FTransform is not allowed to be passed by value on Win32 so pass by const ref instead.

diff --git a/Engine/Plugins/Experimental/GeometryCollectionPlugin/Source/GeometryCollectionComponent/Private/GeometryCollectionUtility.cpp b/Engine/Plugins/Experimental/GeometryCollectionPlugin/Source/GeometryCollectionComponent/Private/GeometryCollectionUtility.cpp
--- a/Engine/Plugins/Experimental/GeometryCollectionPlugin/Source/GeometryCollectionComponent/Private/GeometryCollectionUtility.cpp
+++ b/Engine/Plugins/Experimental/GeometryCollectionPlugin/Source/GeometryCollectionComponent/Private/GeometryCollectionUtility.cpp
@@ -10,7 +10,7 @@ DEFINE_LOG_CATEGORY_STATIC(FGeometryCollectionUtilityLogging, Log, All);
 
 namespace GeometryCollection
 {
-	TSharedPtr<UGeometryCollection> MakeCubeElement(FTransform center, float Scale)
+	TSharedPtr<UGeometryCollection> MakeCubeElement(const FTransform& center, float Scale)
 	{
 		UGeometryCollection * RestCollection = NewObject<UGeometryCollection>();
 
diff --git a/Engine/Plugins/Experimental/GeometryCollectionPlugin/Source/GeometryCollectionComponent/Private/GeometryCollectionUtility.h b/Engine/Plugins/Experimental/GeometryCollectionPlugin/Source/GeometryCollectionComponent/Private/GeometryCollectionUtility.h
--- a/Engine/Plugins/Experimental/GeometryCollectionPlugin/Source/GeometryCollectionComponent/Private/GeometryCollectionUtility.h
+++ b/Engine/Plugins/Experimental/GeometryCollectionPlugin/Source/GeometryCollectionComponent/Private/GeometryCollectionUtility.h
@@ -15,7 +15,7 @@ namespace GeometryCollection
 	* MakeCubeElement
 	*   Utility to create a triangulated unit cube using the UGeometryCollection format.
 	*/
-	TSharedPtr<UGeometryCollection> MakeCubeElement(FTransform center, float Scale);
+	TSharedPtr<UGeometryCollection> MakeCubeElement(const FTransform& center, float Scale);
 
 	/****
 	* SetupCubeGridExample
diff --git a/Engine/Source/Runtime/Experimental/Apeiron/Private/Apeiron/PBDCollisionConstraint.cpp b/Engine/Source/Runtime/Experimental/Apeiron/Private/Apeiron/PBDCollisionConstraint.cpp
--- a/Engine/Source/Runtime/Experimental/Apeiron/Private/Apeiron/PBDCollisionConstraint.cpp
+++ b/Engine/Source/Runtime/Experimental/Apeiron/Private/Apeiron/PBDCollisionConstraint.cpp
@@ -550,7 +550,7 @@ bool NearestPoint(TArray<Pair<TVector<T, d>, TVector<T, d>>>& Points, TVector<T,
 }
 
 template<class T, int d>
-void UpdateLevelsetConstraintHelper(const TPBDRigidParticles<T, d>& InParticles, const int32 j, const TRigidTransform<T, d> LocalToWorld1, const TRigidTransform<T, d> LocalToWorld2, const T Thickness, TRigidBodyContactConstraint<T, d>& Constraint)
+void UpdateLevelsetConstraintHelper(const TPBDRigidParticles<T, d>& InParticles, const int32 j, const TRigidTransform<T, d>& LocalToWorld1, const TRigidTransform<T, d>& LocalToWorld2, const T Thickness, TRigidBodyContactConstraint<T, d>& Constraint)
 {
 	TVector<T, d> WorldSpacePoint = LocalToWorld1.TransformPosition(InParticles.CollisionParticles(Constraint.ParticleIndex).X(j));
 	TVector<T, d> Body2SpacePoint = LocalToWorld2.InverseTransformPosition(WorldSpacePoint);
