Fixed FGeometryCollectionSceneProxy rendering assert.

diff --git a/Engine/Source/Runtime/Experimental/GeometryCollectionEngine/Private/GeometryCollection/GeometryCollectionSceneProxy.cpp b/Engine/Source/Runtime/Experimental/GeometryCollectionEngine/Private/GeometryCollection/GeometryCollectionSceneProxy.cpp
--- a/Engine/Source/Runtime/Experimental/GeometryCollectionEngine/Private/GeometryCollection/GeometryCollectionSceneProxy.cpp
+++ b/Engine/Source/Runtime/Experimental/GeometryCollectionEngine/Private/GeometryCollection/GeometryCollectionSceneProxy.cpp
@@ -287,7 +287,7 @@ void FGeometryCollectionSceneProxy::GetDynamicMeshElements(const TArray<const FS
                     // Collector owns the uniform buffer
                     FDynamicPrimitiveUniformBuffer& DynamicPrimitiveUniformBuffer = Collector.AllocateOneFrameResource<FDynamicPrimitiveUniformBuffer>();
                     DynamicPrimitiveUniformBuffer.Set(GetLocalToWorld(), GetLocalToWorld(), GetBounds(), GetLocalBounds(), true, false, UseEditorDepthTest());
-                    BatchElement.PrimitiveUniformBuffer = DynamicPrimitiveUniformBuffer.UniformBuffer.GetUniformBufferRHI();
+                    BatchElement.PrimitiveUniformBufferResource = &DynamicPrimitiveUniformBuffer.UniformBuffer;
                     
 					BatchElement.FirstIndex = Section.FirstIndex;
 					BatchElement.NumPrimitives = Section.NumTriangles;
@@ -319,7 +319,7 @@ void FGeometryCollectionSceneProxy::GetDynamicMeshElements(const TArray<const FS
                     
                     FDynamicPrimitiveUniformBuffer& DynamicPrimitiveUniformBuffer = Collector.AllocateOneFrameResource<FDynamicPrimitiveUniformBuffer>();
                     DynamicPrimitiveUniformBuffer.Set(GetLocalToWorld(), GetLocalToWorld(), GetBounds(), GetLocalBounds(), true, false, UseEditorDepthTest());
-                    BatchElement.PrimitiveUniformBuffer = DynamicPrimitiveUniformBuffer.UniformBuffer.GetUniformBufferRHI();
+                    BatchElement.PrimitiveUniformBufferResource = &DynamicPrimitiveUniformBuffer.UniformBuffer;
                     
 					BatchElement.FirstIndex = 0;
 					BatchElement.NumPrimitives = GetRequiredIndexCount() / 3;
