Fixed build

#rnx

diff --git a/Engine/Plugins/Runtime/CustomMeshComponent/Source/CustomMeshComponent/Private/CustomMeshComponent.cpp b/Engine/Plugins/Runtime/CustomMeshComponent/Source/CustomMeshComponent/Private/CustomMeshComponent.cpp
--- a/Engine/Plugins/Runtime/CustomMeshComponent/Source/CustomMeshComponent/Private/CustomMeshComponent.cpp
+++ b/Engine/Plugins/Runtime/CustomMeshComponent/Source/CustomMeshComponent/Private/CustomMeshComponent.cpp
@@ -136,7 +136,7 @@ public:
 				GetScene().GetPrimitiveUniformShaderParameters_RenderThread(GetPrimitiveSceneInfo(), bHasPrecomputedVolumetricLightmap, PreviousLocalToWorld, SingleCaptureIndex, bOutputVelocity);
 
 				FDynamicPrimitiveUniformBuffer& DynamicPrimitiveUniformBuffer = Collector.AllocateOneFrameResource<FDynamicPrimitiveUniformBuffer>();
-				DynamicPrimitiveUniformBuffer.Set(GetLocalToWorld(), PreviousLocalToWorld, GetBounds(), GetLocalBounds(), true, bHasPrecomputedVolumetricLightmap, UseEditorDepthTest());
+				DynamicPrimitiveUniformBuffer.Set(GetLocalToWorld(), PreviousLocalToWorld, GetBounds(), GetLocalBounds(), true, bHasPrecomputedVolumetricLightmap, UseEditorDepthTest(), false);
 				BatchElement.PrimitiveUniformBufferResource = &DynamicPrimitiveUniformBuffer.UniformBuffer;
 
 				BatchElement.FirstIndex = 0;
diff --git a/Engine/Source/Runtime/Experimental/GeometryCollectionEngine/Private/GeometryCollection/GeometryCollectionSceneProxy.cpp b/Engine/Source/Runtime/Experimental/GeometryCollectionEngine/Private/GeometryCollection/GeometryCollectionSceneProxy.cpp
--- a/Engine/Source/Runtime/Experimental/GeometryCollectionEngine/Private/GeometryCollection/GeometryCollectionSceneProxy.cpp
+++ b/Engine/Source/Runtime/Experimental/GeometryCollectionEngine/Private/GeometryCollection/GeometryCollectionSceneProxy.cpp
@@ -286,7 +286,7 @@ void FGeometryCollectionSceneProxy::GetDynamicMeshElements(const TArray<const FS
                     
                     // Collector owns the uniform buffer
                     FDynamicPrimitiveUniformBuffer& DynamicPrimitiveUniformBuffer = Collector.AllocateOneFrameResource<FDynamicPrimitiveUniformBuffer>();
-                    DynamicPrimitiveUniformBuffer.Set(GetLocalToWorld(), GetLocalToWorld(), GetBounds(), GetLocalBounds(), true, false, UseEditorDepthTest());
+                    DynamicPrimitiveUniformBuffer.Set(GetLocalToWorld(), GetLocalToWorld(), GetBounds(), GetLocalBounds(), true, false, UseEditorDepthTest(), false);
                     BatchElement.PrimitiveUniformBufferResource = &DynamicPrimitiveUniformBuffer.UniformBuffer;
                     
 					BatchElement.FirstIndex = Section.FirstIndex;
@@ -318,7 +318,7 @@ void FGeometryCollectionSceneProxy::GetDynamicMeshElements(const TArray<const FS
 					Mesh.MaterialRenderProxy = MaterialRenderProxy;
                     
                     FDynamicPrimitiveUniformBuffer& DynamicPrimitiveUniformBuffer = Collector.AllocateOneFrameResource<FDynamicPrimitiveUniformBuffer>();
-                    DynamicPrimitiveUniformBuffer.Set(GetLocalToWorld(), GetLocalToWorld(), GetBounds(), GetLocalBounds(), true, false, UseEditorDepthTest());
+                    DynamicPrimitiveUniformBuffer.Set(GetLocalToWorld(), GetLocalToWorld(), GetBounds(), GetLocalBounds(), true, false, UseEditorDepthTest(), false);
                     BatchElement.PrimitiveUniformBufferResource = &DynamicPrimitiveUniformBuffer.UniformBuffer;
                     
 					BatchElement.FirstIndex = 0;
diff --git a/Engine/Source/Runtime/Renderer/Private/RendererScene.cpp b/Engine/Source/Runtime/Renderer/Private/RendererScene.cpp
--- a/Engine/Source/Runtime/Renderer/Private/RendererScene.cpp
+++ b/Engine/Source/Runtime/Renderer/Private/RendererScene.cpp
@@ -3065,7 +3065,7 @@ bool ShouldForceFullDepthPass(EShaderPlatform ShaderPlatform)
 	const bool bEarlyZMaterialMasking = CVarEarlyZPassOnlyMaterialMasking.GetValueOnAnyThread() != 0;
 
 	// Note: ShouldForceFullDepthPass affects which static draw lists meshes go into, so nothing it depends on can change at runtime, unless you do a FGlobalComponentRecreateRenderStateContext to propagate the cvar change
-	return bDBufferAllowed || bStencilLODDither || bEarlyZMaterialMasking || IsForwardShadingEnabled(ShaderPlatform) || UseSelectiveBasePassOutputs();
+	return bDBufferAllowed || bStencilLODDither || bEarlyZMaterialMasking || IsForwardShadingEnabled(ShaderPlatform) || IsUsingSelectiveBasePassOutputs(ShaderPlatform);
 }
 
 void FScene::UpdateEarlyZPassMode()
