Fix crash from dev-physics.

diff --git a/Engine/Source/Runtime/Experimental/ChaosSolvers/Private/PBDRigidsSolver.cpp b/Engine/Source/Runtime/Experimental/ChaosSolvers/Private/PBDRigidsSolver.cpp
--- a/Engine/Source/Runtime/Experimental/ChaosSolvers/Private/PBDRigidsSolver.cpp
+++ b/Engine/Source/Runtime/Experimental/ChaosSolvers/Private/PBDRigidsSolver.cpp
@@ -1170,7 +1170,7 @@ namespace Chaos
 							BreakingData.AngularVelocity = Particles.W(AllBreakingsArray[Idx].ParticleIndex);
 							BreakingData.Mass = Particles.M(AllBreakingsArray[Idx].ParticleIndex);
 							BreakingData.ParticleIndex = AllBreakingsArray[Idx].ParticleIndex;
-							if (Particles.Geometry(Idx)->HasBoundingBox())
+							if (Particles.Geometry(Idx) && Particles.Geometry(Idx)->HasBoundingBox())
 							{
 								BreakingData.BoundingBox = Particles.Geometry(Idx)->BoundingBox();;
 							}
