Fix warnings from GeomProxy using ispc.

diff --git a/Engine/Source/Runtime/Experimental/GeometryCollectionEngine/Private/GeometryCollection/GeometryCollectionSceneProxy.cpp b/Engine/Source/Runtime/Experimental/GeometryCollectionEngine/Private/GeometryCollection/GeometryCollectionSceneProxy.cpp
--- a/Engine/Source/Runtime/Experimental/GeometryCollectionEngine/Private/GeometryCollection/GeometryCollectionSceneProxy.cpp
+++ b/Engine/Source/Runtime/Experimental/GeometryCollectionEngine/Private/GeometryCollection/GeometryCollectionSceneProxy.cpp
@@ -10,7 +10,25 @@
 #include "GeometryCollection/GeometryCollectionHitProxy.h"
 #endif
 #if INTEL_ISPC
+#if USING_CODE_ANALYSIS
+    MSVC_PRAGMA( warning( push ) )
+    MSVC_PRAGMA( warning( disable : ALL_CODE_ANALYSIS_WARNINGS ) )
+#endif    // USING_CODE_ANALYSIS
+
+#ifdef __clang__
+#pragma clang diagnostic push
+#pragma clang diagnostic ignored "-Wnonportable-include-path"
+#endif
+
 #include "GeometryCollectionSceneProxy.ispc.generated.h"
+
+#ifdef __clang__
+#pragma clang diagnostic pop
+#endif
+
+#if USING_CODE_ANALYSIS
+    MSVC_PRAGMA( warning( pop ) )
+#endif    // USING_CODE_ANALYSIS
 #endif
 
 static int32 GParallelGeometryCollectionBatchSize = 1024;
