Fix ChaosSolverActor Debug Substep that froze when starting with level paused.
Merges:
CL 6956771
CL 6956885
CL 6989030

diff --git a/Engine/Plugins/Experimental/ChaosSolverPlugin/Source/ChaosSolverEditor/Private/Chaos/ChaosSolverEditorDetails.cpp b/Engine/Plugins/Experimental/ChaosSolverPlugin/Source/ChaosSolverEditor/Private/Chaos/ChaosSolverEditorDetails.cpp
--- a/Engine/Plugins/Experimental/ChaosSolverPlugin/Source/ChaosSolverEditor/Private/Chaos/ChaosSolverEditorDetails.cpp
+++ b/Engine/Plugins/Experimental/ChaosSolverPlugin/Source/ChaosSolverEditor/Private/Chaos/ChaosSolverEditorDetails.cpp
@@ -23,26 +23,25 @@ void FChaosDebugSubstepControlCustomization::CustomizeHeader(TSharedRef<IPropert
 	// Hide reset-to-default
 	StructPropertyHandle->MarkResetToDefaultCustomized();
 
-	// Refresh pause button callback
-	const FSimpleDelegate OnPauseChanged = FSimpleDelegate::CreateSP(this, &FChaosDebugSubstepControlCustomization::RefreshPauseButton, PropertyHandlePause);
-	PropertyHandlePause->SetOnPropertyValueChanged(OnPauseChanged);
-
-	// Refresh pause delegate for changes made through code to the actor
-	// The struct property does not refresh when values are set through code, hence the need for this delegate
-	AChaosSolverActor* ChaosSolverActor = nullptr;
+	// Find associated actor
+	bool bHasBegunPlay = false;
 	for (const TWeakObjectPtr<UObject> SelectedObject: StructCustomizationUtils.GetPropertyUtilities()->GetSelectedObjects())
 	{
-		AChaosSolverActor* const SelectedChaosSolverActor = Cast<AChaosSolverActor>(SelectedObject.Get());
-		if (SelectedChaosSolverActor)
+		if (AChaosSolverActor* const ChaosSolverActor = Cast<AChaosSolverActor>(SelectedObject.Get()))
 		{
-			ChaosSolverActor = SelectedChaosSolverActor;
+			// Refresh pause button callback
+			const FSimpleDelegate OnPauseChanged = FSimpleDelegate::CreateSP(this, &FChaosDebugSubstepControlCustomization::RefreshPauseButton, PropertyHandlePause, SelectedObject);
+			PropertyHandlePause->SetOnPropertyValueChanged(OnPauseChanged);
+
+			// Refresh pause delegate for changes made through code to the actor
+			// The struct property does not refresh when values are set through code, hence the need for this delegate
+			ChaosSolverActor->ChaosDebugSubstepControl.OnPauseChanged = OnPauseChanged;
+
+			// Check whether the actor has begun play
+			bHasBegunPlay = ChaosSolverActor->HasActorBegunPlay();
 			break;
 		}
 	}
-	if (ChaosSolverActor)
-	{
-		ChaosSolverActor->ChaosDebugSubstepControl.OnPauseChanged = OnPauseChanged;
-	}
 
 	// Retrieve the current pause state
 	bool bPaused = false;
@@ -83,7 +82,7 @@ void FChaosDebugSubstepControlCustomization::CustomizeHeader(TSharedRef<IPropert
 			.FillWidth(28.0f)
 			[
 				SAssignNew(ButtonSubstep, SButton)
-				.IsEnabled_Lambda([bPaused]() { return bPaused; })
+				.IsEnabled_Lambda([bPaused, bHasBegunPlay]() { return bPaused && bHasBegunPlay; })
 				.ButtonStyle(FEditorStyle::Get(), "HoverHintOnly")
 				.ForegroundColor(FSlateColor::UseForeground())
 				.ContentPadding(4.0f)
@@ -103,7 +102,7 @@ void FChaosDebugSubstepControlCustomization::CustomizeHeader(TSharedRef<IPropert
 			.FillWidth(28.0f)
 			[
 				SAssignNew(ButtonStep, SButton)
-				.IsEnabled_Lambda([bPaused]() { return bPaused; })
+				.IsEnabled_Lambda([bPaused, bHasBegunPlay]() { return bPaused && bHasBegunPlay; })
 				.ButtonStyle(FEditorStyle::Get(), "HoverHintOnly")
 				.ForegroundColor(FSlateColor::UseForeground())
 				.ContentPadding(4.0f)
@@ -132,8 +131,6 @@ FReply FChaosDebugSubstepControlCustomization::OnPause(TSharedRef<IPropertyHandl
 	bPaused = !bPaused;
 	PropertyHandle->SetValue(bPaused);
 
-	RefreshPauseButton(PropertyHandle);
-
 	return FReply::Handled();
 }
 
@@ -149,8 +146,11 @@ FReply FChaosDebugSubstepControlCustomization::OnStep(TSharedRef<IPropertyHandle
 	return FReply::Handled();
 }
 
-void FChaosDebugSubstepControlCustomization::RefreshPauseButton(TSharedRef<IPropertyHandle> PropertyHandle)
+void FChaosDebugSubstepControlCustomization::RefreshPauseButton(TSharedRef<IPropertyHandle> PropertyHandle, TWeakObjectPtr<UObject> Object)
 {
+	const AChaosSolverActor* const ChaosSolverActor = Cast<AChaosSolverActor>(Object.Get());
+	const bool bHasBegunPlay = ChaosSolverActor && ChaosSolverActor->HasActorBegunPlay();
+
 	bool bPaused = false;
 	PropertyHandle->GetValue(bPaused);
 
@@ -160,10 +160,10 @@ void FChaosDebugSubstepControlCustomization::RefreshPauseButton(TSharedRef<IProp
 	}
 	if (ButtonSubstep.IsValid())
 	{
-		ButtonSubstep->SetEnabled(bPaused);
+		ButtonSubstep->SetEnabled(bPaused && bHasBegunPlay);
 	}
 	if (ButtonStep.IsValid())
 	{
-		ButtonStep->SetEnabled(bPaused);
+		ButtonStep->SetEnabled(bPaused && bHasBegunPlay);
 	}
 }
diff --git a/Engine/Plugins/Experimental/ChaosSolverPlugin/Source/ChaosSolverEditor/Private/Chaos/ChaosSolverEditorDetails.h b/Engine/Plugins/Experimental/ChaosSolverPlugin/Source/ChaosSolverEditor/Private/Chaos/ChaosSolverEditorDetails.h
--- a/Engine/Plugins/Experimental/ChaosSolverPlugin/Source/ChaosSolverEditor/Private/Chaos/ChaosSolverEditorDetails.h
+++ b/Engine/Plugins/Experimental/ChaosSolverPlugin/Source/ChaosSolverEditor/Private/Chaos/ChaosSolverEditorDetails.h
@@ -22,7 +22,7 @@ private:
 	FReply OnSubstep(TSharedRef<IPropertyHandle> PropertyHandle);
 	FReply OnStep(TSharedRef<IPropertyHandle> PropertyHandle);
 
-	void RefreshPauseButton(TSharedRef<IPropertyHandle> PropertyHandle);
+	void RefreshPauseButton(TSharedRef<IPropertyHandle> PropertyHandle, TWeakObjectPtr<UObject> Object);
 
 private:
 	TSharedPtr<STextBlock> TextBlockPause;
diff --git a/Engine/Source/Runtime/Engine/Private/PhysicsEngine/Experimental/PhysScene_Chaos.cpp b/Engine/Source/Runtime/Engine/Private/PhysicsEngine/Experimental/PhysScene_Chaos.cpp
--- a/Engine/Source/Runtime/Engine/Private/PhysicsEngine/Experimental/PhysScene_Chaos.cpp
+++ b/Engine/Source/Runtime/Engine/Private/PhysicsEngine/Experimental/PhysScene_Chaos.cpp
@@ -281,10 +281,6 @@ public:
 		FEditorDelegates::PausePIE.AddRaw(this, &FPhysScene_ChaosPauseHandler::PauseSolvers);
 		FEditorDelegates::ResumePIE.AddRaw(this, &FPhysScene_ChaosPauseHandler::ResumeSolvers);
 		FEditorDelegates::SingleStepPIE.AddRaw(this, &FPhysScene_ChaosPauseHandler::SingleStepSolvers);
-#if CHAOS_DEBUG_SUBSTEP
-		// Editor pre-end play handle
-		FEditorDelegates::PrePIEEnded.AddRaw(this, &FPhysScene_ChaosPauseHandler::PreEndSolvers);
-#endif
 	}
 
 	~FPhysScene_ChaosPauseHandler()
@@ -295,19 +291,12 @@ public:
 		FEditorDelegates::PausePIE.RemoveAll(this);
 		FEditorDelegates::ResumePIE.RemoveAll(this);
 		FEditorDelegates::SingleStepPIE.RemoveAll(this);
-#if CHAOS_DEBUG_SUBSTEP
-		// Remove editor pre-end play delegates
-		FEditorDelegates::PrePIEEnded.RemoveAll(this);
-#endif
 	}
 
 private:
 	void PauseSolvers(bool /*bIsSimulating*/) { ChaosModule->PauseSolvers(); }
 	void ResumeSolvers(bool /*bIsSimulating*/) { ChaosModule->ResumeSolvers(); }
 	void SingleStepSolvers(bool /*bIsSimulating*/) { ChaosModule->SingleStepSolvers(); }
-#if CHAOS_DEBUG_SUBSTEP
-	void PreEndSolvers(bool /*bIsSimulating*/) { ChaosModule->PreEndSolvers(); }
-#endif
 
 private:
 	FChaosSolversModule* ChaosModule;
diff --git a/Engine/Source/Runtime/Experimental/Chaos/Private/Chaos/Framework/DebugSubstep.cpp b/Engine/Source/Runtime/Experimental/Chaos/Private/Chaos/Framework/DebugSubstep.cpp
--- a/Engine/Source/Runtime/Experimental/Chaos/Private/Chaos/Framework/DebugSubstep.cpp
+++ b/Engine/Source/Runtime/Experimental/Chaos/Private/Chaos/Framework/DebugSubstep.cpp
@@ -13,16 +13,62 @@ namespace Chaos
 	FDebugSubstep::FDebugSubstep()
 		: bIsEnabled(false)
 		, CommandQueue()
-		, ProgressEvent(FPlatformProcess::GetSynchEventFromPool())
-		, SubstepEvent(FPlatformProcess::GetSynchEventFromPool(true))  // SubstepEvent can be triggered without a matching wait, hence the manual setting
-		, bWaitForStep(false)
+		, ProgressEvent(nullptr)
+		, SubstepEvent(nullptr)
 		, ThreadId(0)
+		, bWaitForStep(false)
+		, bShouldEnable(false)
 	{}
 
 	FDebugSubstep::~FDebugSubstep()
 	{
-		FPlatformProcess::ReturnSynchEventToPool(ProgressEvent);
-		FPlatformProcess::ReturnSynchEventToPool(SubstepEvent);
+		Release();
+	}
+
+	void FDebugSubstep::Initialize()
+	{
+		if (!IsInitialized())
+		{
+			// Allocate sync events
+			ProgressEvent = FPlatformProcess::GetSynchEventFromPool();
+			SubstepEvent = FPlatformProcess::GetSynchEventFromPool(true);  // SubstepEvent can be triggered without a matching wait, hence the manual reset setting
+		}
+	}
+
+	void FDebugSubstep::Release()
+	{
+		if (IsInitialized())
+		{
+			// Exit debug thread
+			if (bIsEnabled)
+			{
+				Stop();
+			}
+
+			// Release sync events
+			FPlatformProcess::ReturnSynchEventToPool(ProgressEvent);
+			FPlatformProcess::ReturnSynchEventToPool(SubstepEvent);
+			ProgressEvent = nullptr;
+			SubstepEvent = nullptr;
+
+			// Reset commands
+			CommandQueue.Empty();
+			bShouldEnable = false;
+		}
+	}
+
+	void FDebugSubstep::Shutdown()
+	{
+		// Exit debug thread
+		if (IsInitialized() && bIsEnabled)
+		{
+			Stop();
+		}
+	}
+
+	void FDebugSubstep::AssumeThisThread()
+	{
+		ThreadId = FPlatformTLS::GetCurrentThreadId();
 	}
 
 	void FDebugSubstep::Enable(bool bEnable)
@@ -43,100 +89,81 @@ namespace Chaos
 		UE_LOG(LogChaosThread, Verbose, TEXT("[Game Thread] Progress"));
 	}
 
-	void FDebugSubstep::Shutdown()
+	bool FDebugSubstep::SyncAdvance(bool bIsSolverEnabled)
 	{
-		if (bIsEnabled)
-		{
-			bIsEnabled = false;
-			UE_LOG(LogChaosThread, Verbose, TEXT("[%s Thread] bIsEnabled changed (true->false)"), IsInGameThread() ? TEXT("Game"): TEXT("Physics"));
-
-			// Trigger progress, with bIsEnabled being false it should go straight to the end of the step
-			SubstepEvent->Reset();     // No race condition with the Add() between these two instructions
-			ProgressEvent->Trigger();  // since this code path is only to be entered while in ProgressEvent->Wait(); state.
-
-			// Wait for the final step event
-			UE_LOG(LogChaosThread, Verbose, TEXT("[%s Thread] Waiting for last step event"), IsInGameThread() ? TEXT("Game"): TEXT("Physics"));
-			SubstepEvent->Wait();  // Last wait event will be triggered on debug thread exit
-			UE_LOG(LogChaosThread, Verbose, TEXT("[%s Thread] Substep event received, wait ended"), IsInGameThread() ? TEXT("Game"): TEXT("Physics"));
-			UE_LOG(LogChaosThread, Log, TEXT("Chaos' debug substep mode is now disengaged. Resuming solver thread at next step."));
-		}
-		// Empty queue
-		CommandQueue.Empty();
-	}
+		// Don't process the commands unless this instance is initialized
+		if (!IsInitialized()) { return false; }
 
-	bool FDebugSubstep::SyncAdvance()
-	{
+		// Process all commands from the queue
+		bool bShouldStep = false;
+		bool bShouldSubstep = false;
 		ECommand Command;
 		while (CommandQueue.Dequeue(Command))
 		{
 			switch (Command)
 			{
 			case ECommand::Enable:
-				if (!bIsEnabled)
-				{
-					bIsEnabled = true;
-					UE_LOG(LogChaosThread, Verbose, TEXT("[%s Thread] bIsEnabled changed (false->true)"), IsInGameThread() ? TEXT("Game"): TEXT("Physics"));
-					UE_LOG(LogChaosThread, Log, TEXT("Chaos' debug substep mode is now engaged. Pausing solver thread at next step."));
-				}
-				return true;  // Wait until thread has started before dequeuing more commands
+				bShouldEnable = true;
+				break;
 
 			case ECommand::Disable:
-				if (bIsEnabled)
-				{
-					bIsEnabled = false;
-					UE_LOG(LogChaosThread, Verbose, TEXT("[%s Thread] bIsEnabled changed (true->false)"), IsInGameThread() ? TEXT("Game"): TEXT("Physics"));
-
-					// Trigger progress, with bIsEnabled being false it should go straight to the end of the step
-					SubstepEvent->Reset();     // No race condition with the Add() between these two instructions
-					ProgressEvent->Trigger();  // since this code path is only to be entered while in ProgressEvent->Wait(); state.
-
-					// Wait for the final step event
-					UE_LOG(LogChaosThread, Verbose, TEXT("[%s Thread] Waiting for last step event"), IsInGameThread() ? TEXT("Game"): TEXT("Physics"));
-					SubstepEvent->Wait();  // Last wait event will be triggered on debug thread exit
-					UE_LOG(LogChaosThread, Verbose, TEXT("[%s Thread] Substep event received, wait ended"), IsInGameThread() ? TEXT("Game"): TEXT("Physics"));
-					UE_LOG(LogChaosThread, Log, TEXT("Chaos' debug substep mode is now disengaged. Resuming solver thread at next step."));
-				}
-				return false;  // Wait until thread has ended before dequeuing more commands
-
-			case ECommand::ProgressToStep:  // Intentional fallthrough
+				bShouldEnable = false;
+				break;
+
+			case ECommand::ProgressToStep:
+				bShouldStep = true;
+				break;
+
 			case ECommand::ProgressToSubstep:
-				if (bIsEnabled)
-				{
-					UE_LOG(LogChaosThread, Verbose, TEXT("[%s Thread] Triggering progress event"), IsInGameThread() ? TEXT("Game"): TEXT("Physics"));
-					bWaitForStep = (Command == ECommand::ProgressToStep);
-					do
-					{
-						// Trigger progress
-						SubstepEvent->Reset();     // No race condition with the Add() between these two instructions
-						ProgressEvent->Trigger();  // since this code path is only to be entered while in ProgressEvent->Wait(); state.
-
-						// Wait for next step/substep event
-						UE_LOG(LogChaosThread, Verbose, TEXT("[%s Thread] Waiting for substep event"), IsInGameThread() ? TEXT("Game"): TEXT("Physics"));
-						SubstepEvent->Wait();
-						UE_LOG(LogChaosThread, Verbose, TEXT("[%s Thread] Substep event received, wait ended"), IsInGameThread() ? TEXT("Game"): TEXT("Physics"));
-					}
-					while (bWaitForStep);
-				}
+				bShouldSubstep = true;
 				break;
 
-			default: 
+			default:
 				UE_LOG(LogChaosThread, Fatal, TEXT("Unknown debug step command."));
 				break;
 			}
 		}
-		return bIsEnabled;
-	}
 
-	void FDebugSubstep::AssumeThisThread()
-	{
-		ThreadId = FPlatformTLS::GetCurrentThreadId();
+		if (!bIsEnabled)
+		{
+			// Currently disabled, check whether it needs enabling
+			if (bShouldEnable && bIsSolverEnabled) // Can only allow the debug thread to run if the solver is already enabled.
+			{
+				// Start debug substep
+				Start();
+			}
+		}
+		else
+		{
+			// Currently enabled, check whether it needs disabling or stepping
+			if (!bShouldEnable || !bIsSolverEnabled) // Should also stop the debug thread if the solver is disabled, otherwise it can't be re-enabled since the HandleSolverCommands won't get executed
+			{
+				// Stop debug substep
+				Stop();
+			}
+			else if (bShouldStep || bShouldSubstep)
+			{
+				// Step/substep
+				Substep(bShouldStep);
+			}
+		}
+
+		return bIsEnabled;
 	}
 
 	void FDebugSubstep::Add(bool bInStep, const TCHAR* Label) const
 	{
+		// Ignore this substep unless this instance is initialized
+		if (!IsInitialized()) { return; }
+
+		// Manage events
 		if (bIsEnabled)
 		{
-			if (bInStep) { bWaitForStep = false; }
+			if (bInStep)
+			{
+				// Signal a step boundary
+				bWaitForStep = false;
+			}
 			checkf(ThreadId == FPlatformTLS::GetCurrentThreadId(), TEXT("Cannot add a substep outside of the solver thread (eg inside a parallel for)."));
 			UE_LOG(LogChaosThread, Log, TEXT("Reached %s '%s'"), bInStep ? TEXT("step"): TEXT("substep"), Label ? Label: TEXT(""));
 			UE_LOG(LogChaosThread, Verbose, TEXT("[Debug Thread] Triggering substep event"));
@@ -147,11 +174,63 @@ namespace Chaos
 		}
 		else if (bInStep)  // Trigger one last event at the step boundary when disabled
 		{
-			UE_LOG(LogChaosThread, Log, TEXT("Reached step '%s'"),  Label ? Label: TEXT(""));
+			UE_LOG(LogChaosThread, Log, TEXT("Reached step '%s'"), Label ? Label: TEXT(""));
 			UE_LOG(LogChaosThread, Verbose, TEXT("[Debug Thread] Triggering substep event"));
 			SubstepEvent->Trigger();
 		}
 	}
+
+	void FDebugSubstep::Start()
+	{
+		check(IsInitialized());
+		check(!bIsEnabled);
+
+		bIsEnabled = true;
+		UE_LOG(LogChaosThread, Verbose, TEXT("[%s Thread] bIsEnabled changed (false->true)"), IsInGameThread() ? TEXT("Game") : TEXT("Physics"));
+		UE_LOG(LogChaosThread, Log, TEXT("Chaos' debug substep mode is now engaged. Pausing solver thread at next step."));
+		// Note: Once enabled the associated task must be created
+	}
+
+	void FDebugSubstep::Stop()
+	{
+		check(IsInitialized());
+		check(bIsEnabled);
+
+		bIsEnabled = false;
+		UE_LOG(LogChaosThread, Verbose, TEXT("[%s Thread] bIsEnabled changed (true->false)"), IsInGameThread() ? TEXT("Game") : TEXT("Physics"));
+
+		// Trigger progress, with bIsEnabled being false it should go straight to the end of the step
+		SubstepEvent->Reset();     // No race condition with the Add() between these two instructions
+		ProgressEvent->Trigger();  // since this code path is only to be entered while in ProgressEvent->Wait(); state.
+
+		// Wait for the final step event
+		UE_LOG(LogChaosThread, Verbose, TEXT("[%s Thread] Waiting for last step event"), IsInGameThread() ? TEXT("Game") : TEXT("Physics"));
+		SubstepEvent->Wait();  // Last wait event will be triggered on debug thread exit
+
+		UE_LOG(LogChaosThread, Verbose, TEXT("[%s Thread] Substep event received, wait ended"), IsInGameThread() ? TEXT("Game") : TEXT("Physics"));
+		UE_LOG(LogChaosThread, Log, TEXT("Chaos' debug substep mode is now disengaged. Resuming solver thread at next step."));
+		// Note: Once disabled the associated task must be deleted
+	}
+
+	void FDebugSubstep::Substep(bool bShouldStep)
+	{
+		check(IsInitialized());
+		check(bIsEnabled);
+
+		UE_LOG(LogChaosThread, Verbose, TEXT("[%s Thread] Triggering progress event"), IsInGameThread() ? TEXT("Game") : TEXT("Physics"));
+		bWaitForStep = bShouldStep;
+		do
+		{
+			// Trigger progress
+			SubstepEvent->Reset();     // No race condition with the Add() between these two instructions
+			ProgressEvent->Trigger();  // since this code path is only to be entered while in ProgressEvent->Wait(); state.
+
+			// Wait for next step/substep event
+			UE_LOG(LogChaosThread, Verbose, TEXT("[%s Thread] Waiting for substep event"), IsInGameThread() ? TEXT("Game") : TEXT("Physics"));
+			SubstepEvent->Wait();
+			UE_LOG(LogChaosThread, Verbose, TEXT("[%s Thread] Substep event received, wait ended"), IsInGameThread() ? TEXT("Game") : TEXT("Physics"));
+		} while (bWaitForStep);
+	}
 }
 
 #endif  // #if CHAOS_DEBUG_SUBSTEP
diff --git a/Engine/Source/Runtime/Experimental/Chaos/Public/Chaos/Framework/DebugSubstep.h b/Engine/Source/Runtime/Experimental/Chaos/Public/Chaos/Framework/DebugSubstep.h
--- a/Engine/Source/Runtime/Experimental/Chaos/Public/Chaos/Framework/DebugSubstep.h
+++ b/Engine/Source/Runtime/Experimental/Chaos/Public/Chaos/Framework/DebugSubstep.h
@@ -42,32 +42,42 @@ namespace Chaos
 		/** Enable/disable substep pause points. */
 		void Enable(bool bEnable);
 
-		/** Allow progress to the next substep (works only after this object is enabled). */
+		/** Allow progress to the next substep (only works once this object is enabled). */
 		void ProgressToSubstep();
 
-		/** Allow progress to the next step (works only after this object is enabled). */
+		/** Allow progress to the next step (only works once this object is enabled). */
 		void ProgressToStep();
 
+	private:
+		FDebugSubstep(const FDebugSubstep&) = delete;
+		FDebugSubstep& operator=(const FDebugSubstep&) = delete;
+
+		/** Initialize sync events. */
+		void Initialize();
+
+		/** Shutdown and release sync events. Once released the debug thread cannot be restarted unless Initialize is called first. */
+		void Release();
+
+		/** Return true when initialized. */
+		FORCEINLINE bool IsInitialized() const { return !!ProgressEvent && !!SubstepEvent; }
+
 		/*
 		 * Disable and wait for the completion of the debug thread.
-		 * Not thread safe. Must be either called from within the physics thread 
-		 * or within the game thread with the physics thread locked.
+		 * Not thread safe. Must be called from within the physics thread or with the physics thread locked.
+		 * Once shutdown, the debug thread can still be restarted when SyncAdvance is called.
 		 */
 		void Shutdown();
 
-	private:
-		FDebugSubstep(const FDebugSubstep&) = delete;
-		FDebugSubstep& operator=(const FDebugSubstep&) = delete;
+		/** Set the id of the thread the debug substepping will be running in. */
+		void AssumeThisThread();
 
 		/**
 		 * Control substepping progress.
 		 * Start substepping, wait until the next substep is reached, or return straightaway if debugging is disabled.
+		 * @param bIsSolverEnabled The state of the solver. The substepping will stop when the solver is disabled and delay starting again until re-enabled.
 		 * @return Whether the debug thread needs running.
 		 */
-		bool SyncAdvance();
-
-		/** Set the id of the thread the debug substepping will be running in. */
-		void AssumeThisThread();
+		bool SyncAdvance(bool bIsSolverEnabled);
 
 		/**
 		 * Add a new step or substep.
@@ -76,6 +86,15 @@ namespace Chaos
 		 */
 		void Add(bool bInStep, const TCHAR* Label) const;
 
+		/** Set sync events to start the substepping process. */
+		void Start();
+
+		/** Trigger/wait for sync events to stop the substepping process. */
+		void Stop();
+
+		/** Trigger/wait for sync events to step/substep. */
+		void Substep(bool bShouldStep);
+
 	private:
 		enum class ECommand { Enable, Disable, ProgressToSubstep, ProgressToStep };
 
@@ -83,8 +102,9 @@ namespace Chaos
 		TQueue<ECommand, EQueueMode::Mpsc> CommandQueue;  // Command queue, thread safe, Multiple-producers single-consumer (MPSC) model.
 		FEvent* ProgressEvent;         // Progress synchronization event.
 		FEvent* SubstepEvent;          // Substep synchronization event.
-		mutable bool bWaitForStep;     // Boolean used to flag the completion of a step. Set within the a const function, hence the mutable.
 		uint32 ThreadId;               // Thread id used to check that the debug substep code is still running within the debug thread.
+		mutable bool bWaitForStep;     // Boolean used to flag the completion of a step. Set within the a const function, hence the mutable.
+		bool bShouldEnable;            // Whether an enable command has been requested. The debugging thread might not be ready yet (e.g. disabled solver), this will ensure it still enable the debug substepping once it is ready.
 	};
 }
 
@@ -108,7 +128,9 @@ namespace Chaos
 		void Add(const TCHAR* /*Label*/ = nullptr) {}
 
 		void Enable(bool /*bEnable*/) {}
-		void Progress() {}
+
+		void ProgressToSubstep() {}
+		void ProgressToStep() {}
 
 	private:
 		FDebugSubstep(const FDebugSubstep&) = delete;
diff --git a/Engine/Source/Runtime/Experimental/ChaosSolvers/Private/ChaosSolversModule.cpp b/Engine/Source/Runtime/Experimental/ChaosSolvers/Private/ChaosSolversModule.cpp
--- a/Engine/Source/Runtime/Experimental/ChaosSolvers/Private/ChaosSolversModule.cpp
+++ b/Engine/Source/Runtime/Experimental/ChaosSolvers/Private/ChaosSolversModule.cpp
@@ -789,21 +789,6 @@ bool FChaosSolversModule::ShouldStepSolver(int32& InOutSingleStepCounter) const
 }
 #endif  // #if WITH_EDITOR
 
-#if WITH_EDITOR && CHAOS_DEBUG_SUBSTEP
-void FChaosSolversModule::PreEndSolvers()
-{
-	// Lock the physics thread prior to shutting down the debug substep threads
-	FChaosScopedPhysicsThreadLock ScopeLock;
-
-	const int32 NumSolvers = Solvers.Num();
-	for(int32 SolverIndex = 0; SolverIndex < NumSolvers; ++SolverIndex)
-	{
-		Chaos::FPBDRigidsSolver* const Solver = Solvers[SolverIndex];
-		Solver->GetDebugSubstep().Shutdown();
-	}
-}
-#endif  // #if WITH_EDITOR && CHAOS_DEBUG_SUBSTEP
-
 void FChaosSolversModule::ChangeBufferMode(Chaos::EMultiBufferMode BufferMode)
 {
 	for (Chaos::FPBDRigidsSolver* Solver : Solvers)
diff --git a/Engine/Source/Runtime/Experimental/ChaosSolvers/Private/Framework/DebugSolverTasks.cpp b/Engine/Source/Runtime/Experimental/ChaosSolvers/Private/Framework/DebugSolverTasks.cpp
--- a/Engine/Source/Runtime/Experimental/ChaosSolvers/Private/Framework/DebugSolverTasks.cpp
+++ b/Engine/Source/Runtime/Experimental/ChaosSolvers/Private/Framework/DebugSolverTasks.cpp
@@ -50,7 +50,7 @@ namespace Chaos
 
 		// Sync/advance debug substep system for this solver
 		FDebugSubstep& DebugSubstep = Solver->GetDebugSubstep();
-		const bool bNeedsDebugThreadRunning = DebugSubstep.SyncAdvance();
+		const bool bNeedsDebugThreadRunning = DebugSubstep.SyncAdvance(Solver->Enabled());
 
 		// Update debug thread status
 		const bool bIsDebugThreadRunning = !!*DebugSolverTask;
@@ -81,6 +81,16 @@ namespace Chaos
 		}
 	}
 
+	void FDebugSolverTasks::Add(FPBDRigidsSolver* Solver)
+	{
+		// Add solver to task map
+		SolverToTaskMap.Add(Solver, nullptr);
+
+		// Reinit debug substep
+		UE_LOG(LogChaosThread, Verbose, TEXT("[Physics Thread] Initializing debug thread"));
+		Solver->GetDebugSubstep().Initialize();
+	}
+
 	void FDebugSolverTasks::Remove(FPBDRigidsSolver* Solver)
 	{
 		// Remove the debug advance task for this solver, if any was created
@@ -90,7 +100,7 @@ namespace Chaos
 		{
 			// Unpause debug thread and exit
 			UE_LOG(LogChaosThread, Verbose, TEXT("[Physics Thread] Shutting down debug thread"));
-			Solver->GetDebugSubstep().Shutdown();
+			Solver->GetDebugSubstep().Release();  // A call to Release() here means it won't debug substep until another call to Initialize() is made.
 
 			// Wait for thread completion and delete the task
 			UE_LOG(LogChaosThread, Verbose, TEXT("[Physics Thread] Deleting debug async task"));
@@ -106,7 +116,7 @@ namespace Chaos
 		{
 			// Shutdown debug task if needed
 			UE_LOG(LogChaosThread, Verbose, TEXT("[Physics Thread] Shutting debug thread"));
-			SolverTask.Key->GetDebugSubstep().Shutdown();
+			SolverTask.Key->GetDebugSubstep().Shutdown();  // A call to Shutdown() here means it will be able to debug substep again , even without another call to Initialize() made (eg. switch back and forth between threading modes).
 
 			// Wait for thread completion and delete the task
 			if (SolverTask.Value)
diff --git a/Engine/Source/Runtime/Experimental/ChaosSolvers/Private/Framework/PersistentTask.cpp b/Engine/Source/Runtime/Experimental/ChaosSolvers/Private/Framework/PersistentTask.cpp
--- a/Engine/Source/Runtime/Experimental/ChaosSolvers/Private/Framework/PersistentTask.cpp
+++ b/Engine/Source/Runtime/Experimental/ChaosSolvers/Private/Framework/PersistentTask.cpp
@@ -58,6 +58,14 @@ namespace Chaos
 		FChaosSolversModule& ChaosModule = FModuleManager::Get().GetModuleChecked<FChaosSolversModule>("ChaosSolvers");
 		Solvers = ChaosModule.GetSolvers();
 
+#if CHAOS_DEBUG_SUBSTEP
+		// Prepare the debug substepping tasks for all existing solvers
+		for (FPBDRigidsSolver* Solver : Solvers)
+		{
+			DebugSolverTasks.Add(Solver);
+		}
+#endif
+
 		bRunning = true;
 		ShutdownEvent->Reset();
 
diff --git a/Engine/Source/Runtime/Experimental/ChaosSolvers/Public/ChaosSolversModule.h b/Engine/Source/Runtime/Experimental/ChaosSolvers/Public/ChaosSolversModule.h
--- a/Engine/Source/Runtime/Experimental/ChaosSolvers/Public/ChaosSolversModule.h
+++ b/Engine/Source/Runtime/Experimental/ChaosSolvers/Public/ChaosSolversModule.h
@@ -35,7 +35,6 @@ public:
 
 #include "Framework/Dispatcher.h"
 #include "Framework/PersistentTask.h"
-#include "Chaos/Framework/DebugSubstep.h"
 
 namespace Chaos
 {
@@ -268,12 +267,6 @@ public:
 	 */
 	bool ShouldStepSolver(int32& InOutSingleStepCounter) const;
 #endif  // #if WITH_EDITOR
-#if WITH_EDITOR && CHAOS_DEBUG_SUBSTEP
-	/**
-	 * Preempt deletion of objects and of solvers. Called from editor handles.
-	 */
-	void PreEndSolvers();
-#endif  // #if WITH_EDITOR && CHAOS_DEBUG_SUBSTEP
 
 	void RegisterSolverActorClassProvider(IChaosSolverActorClassProvider* Provider)
 	{
diff --git a/Engine/Source/Runtime/Experimental/ChaosSolvers/Public/Framework/DebugSolverTasks.h b/Engine/Source/Runtime/Experimental/ChaosSolvers/Public/Framework/DebugSolverTasks.h
--- a/Engine/Source/Runtime/Experimental/ChaosSolvers/Public/Framework/DebugSolverTasks.h
+++ b/Engine/Source/Runtime/Experimental/ChaosSolvers/Public/Framework/DebugSolverTasks.h
@@ -55,7 +55,7 @@ namespace Chaos
 		~FDebugSolverTasks() { Shutdown(); }
 
 		/** Add debug task entry for the specified solver. */
-		FORCEINLINE void Add(FPBDRigidsSolver* Solver) { SolverToTaskMap.Add(Solver, nullptr); }
+		void Add(FPBDRigidsSolver* Solver);
 
 		/** Remove the debug task entry for the specified solver, and delete its task if any was created. */
 		void Remove(FPBDRigidsSolver* Solver);
