fixed bug with invisible GCs

diff --git a/Engine/Source/Runtime/Experimental/GeometryCollectionEngine/Private/GeometryCollection/GeometryCollectionComponent.cpp b/Engine/Source/Runtime/Experimental/GeometryCollectionEngine/Private/GeometryCollection/GeometryCollectionComponent.cpp
--- a/Engine/Source/Runtime/Experimental/GeometryCollectionEngine/Private/GeometryCollection/GeometryCollectionComponent.cpp
+++ b/Engine/Source/Runtime/Experimental/GeometryCollectionEngine/Private/GeometryCollection/GeometryCollectionComponent.cpp
@@ -2078,6 +2078,12 @@ UMaterialInterface* UGeometryCollectionComponent::GetMaterial(int32 MaterialInde
 // #temp HACK for demo, When fracture happens (physics state changes to dynamic) then switch the visible render meshes in a blueprint/actor from static meshes to geometry collections
 void UGeometryCollectionComponent::SwitchRenderModels(const AActor* Actor)
 {
+	// Don't touch visibility if the component is not visible
+	if (!IsVisible())
+	{
+		return;
+	}
+
 	TArray<UActorComponent*> PrimitiveComponents = Actor->GetComponentsByClass(UPrimitiveComponent::StaticClass());
 	for (UActorComponent* PrimitiveComponent : PrimitiveComponents)
 	{
