Fix the geometry collection debug draw component to start with the tick function disabled when the geometry collection is being created while simulating in editor.
Also fix the begin play tick status update, and resave the ChaosExamples_GeometryCollections level with some collision debug draw disabled, since some asset was now showing collision volume after the fix.

diff --git a/Engine/Source/Runtime/Experimental/GeometryCollectionEngine/Private/GeometryCollection/GeometryCollectionDebugDrawComponent.cpp b/Engine/Source/Runtime/Experimental/GeometryCollectionEngine/Private/GeometryCollection/GeometryCollectionDebugDrawComponent.cpp
--- a/Engine/Source/Runtime/Experimental/GeometryCollectionEngine/Private/GeometryCollection/GeometryCollectionDebugDrawComponent.cpp
+++ b/Engine/Source/Runtime/Experimental/GeometryCollectionEngine/Private/GeometryCollection/GeometryCollectionDebugDrawComponent.cpp
@@ -116,7 +116,7 @@ UGeometryCollectionDebugDrawComponent::UGeometryCollectionDebugDrawComponent(con
 #if GEOMETRYCOLLECTION_DEBUG_DRAW
 	PrimaryComponentTick.bCanEverTick = true;
 	PrimaryComponentTick.bTickEvenWhenPaused = true;
-	SetComponentTickEnabled(false);
+	PrimaryComponentTick.bStartWithTickEnabled = false;
 #else  // #if GEOMETRYCOLLECTION_DEBUG_DRAW
 	PrimaryComponentTick.bCanEverTick = false;
 #endif  // #if GEOMETRYCOLLECTION_DEBUG_DRAW #else
@@ -127,8 +127,8 @@ void UGeometryCollectionDebugDrawComponent::BeginPlay()
 	Super::BeginPlay();
 
 #if GEOMETRYCOLLECTION_DEBUG_DRAW
-	// Start with ticking disabled
-	SetComponentTickEnabled(false);
+	// Should always start with ticking disabled
+	check(!PrimaryComponentTick.IsTickFunctionEnabled());
 
 	if (GeometryCollectionComponent)
 	{
@@ -146,6 +146,9 @@ void UGeometryCollectionDebugDrawComponent::BeginPlay()
 			AddTickPrerequisiteActor(GeometryCollectionDebugDrawActor);
 		}
 
+		// Update the visibility and tick status depending on the debug draw properties currently selected
+		OnDebugDrawPropertiesChanged(true);
+
 #if INCLUDE_CHAOS
 		// Find or create level set renderer
 		GeometryCollectionRenderLevelSetActor = AGeometryCollectionRenderLevelSetActor::FindOrCreate(GetWorld());
