Fix multiple successive crashes when a geometry collection is created in SIE due to a missing DynamicCollection transform array update.

diff --git a/Engine/Plugins/Experimental/ChaosEditor/Source/FractureEditor/Private/FractureEditorModeToolkit.cpp b/Engine/Plugins/Experimental/ChaosEditor/Source/FractureEditor/Private/FractureEditorModeToolkit.cpp
--- a/Engine/Plugins/Experimental/ChaosEditor/Source/FractureEditor/Private/FractureEditorModeToolkit.cpp
+++ b/Engine/Plugins/Experimental/ChaosEditor/Source/FractureEditor/Private/FractureEditorModeToolkit.cpp
@@ -747,7 +747,7 @@ int32 FFractureEditorModeToolkit::GetLevelCount()
 	int32 ReturnLevel = -1;
 	for (UGeometryCollectionComponent* Comp : GeomCompSelection)
 	{
-		FGeometryCollectionEdit GCEdit = Comp->EditRestCollection();
+		FGeometryCollectionEdit GCEdit = Comp->EditRestCollection(GeometryCollection::EEditUpdate::None);
 		if (UGeometryCollection* GCObject = GCEdit.GetRestCollection())
 		{
 			TSharedPtr<FGeometryCollection, ESPMode::ThreadSafe> GeometryCollectionPtr = GCObject->GetGeometryCollection();
@@ -1771,7 +1771,7 @@ bool GetValidGeoCenter(const TManagedArray<int32>& TransformToGeometryIndex, con
 
 void FFractureEditorModeToolkit::UpdateExplodedVectors(UGeometryCollectionComponent* GeometryCollectionComponent) const
 {
-	FGeometryCollectionEdit RestCollection = GeometryCollectionComponent->EditRestCollection();
+	FGeometryCollectionEdit RestCollection = GeometryCollectionComponent->EditRestCollection(GeometryCollection::EEditUpdate::RestPhysicsDynamic);
 	UGeometryCollection* GeometryCollection = RestCollection.GetRestCollection();
 
 	TSharedPtr<FGeometryCollection, ESPMode::ThreadSafe> GeometryCollectionPtr = GeometryCollection->GetGeometryCollection();
@@ -1867,7 +1867,7 @@ AGeometryCollectionActor*  FFractureEditorModeToolkit::ConvertStaticMeshToGeomet
 
 	AGeometryCollectionActor* NewActor = CreateNewGeometryActor(InAssetPath, FTransform(), true);
 
-	FGeometryCollectionEdit GeometryCollectionEdit = NewActor->GetGeometryCollectionComponent()->EditRestCollection();
+	FGeometryCollectionEdit GeometryCollectionEdit = NewActor->GetGeometryCollectionComponent()->EditRestCollection(GeometryCollection::EEditUpdate::RestPhysicsDynamic);
 	UGeometryCollection* FracturedGeometryCollection = GeometryCollectionEdit.GetRestCollection();
 
  	for (AActor* Actor : Actors)
diff --git a/Engine/Plugins/Experimental/GeometryCollectionPlugin/Source/GeometryCollectionEditor/Private/GeometryCollection/GeometryCollectionBoneTreeItem.cpp b/Engine/Plugins/Experimental/GeometryCollectionPlugin/Source/GeometryCollectionEditor/Private/GeometryCollection/GeometryCollectionBoneTreeItem.cpp
--- a/Engine/Plugins/Experimental/GeometryCollectionPlugin/Source/GeometryCollectionEditor/Private/GeometryCollection/GeometryCollectionBoneTreeItem.cpp
+++ b/Engine/Plugins/Experimental/GeometryCollectionPlugin/Source/GeometryCollectionEditor/Private/GeometryCollection/GeometryCollectionBoneTreeItem.cpp
@@ -274,7 +274,7 @@ void FGeometryCollectionBoneTreeItem::RenameSubComponent(const FText& InName)
 		return;
 	}
 	UGeometryCollectionComponent* GCComponent = Cast<UGeometryCollectionComponent>(ComponentPtr);
-	FGeometryCollectionEdit ScopedEdit = GCComponent->EditRestCollection(false);
+	FGeometryCollectionEdit ScopedEdit = GCComponent->EditRestCollection(GeometryCollection::EEditUpdate::None);
 	if (UGeometryCollection* GeometryCollectionObject = ScopedEdit.GetRestCollection())
 	{
 		TSharedPtr<FGeometryCollection, ESPMode::ThreadSafe> GeometryCollectionPtr = GeometryCollectionObject->GetGeometryCollection();
diff --git a/Engine/Plugins/Runtime/EditableMesh/Source/EditableMesh/EditableGeometryCollectionAdapter.cpp b/Engine/Plugins/Runtime/EditableMesh/Source/EditableMesh/EditableGeometryCollectionAdapter.cpp
--- a/Engine/Plugins/Runtime/EditableMesh/Source/EditableMesh/EditableGeometryCollectionAdapter.cpp
+++ b/Engine/Plugins/Runtime/EditableMesh/Source/EditableMesh/EditableGeometryCollectionAdapter.cpp
@@ -36,7 +36,7 @@ void UEditableGeometryCollectionAdapter::InitEditableGeometryCollection( UEditab
 	GeometryCollectionComponent = Cast<UGeometryCollectionComponent>( &Component );
 	if( GeometryCollectionComponent != nullptr )
 	{
-		FGeometryCollectionEdit GeometryCollectionEdit = GeometryCollectionComponent->EditRestCollection(false);
+		FGeometryCollectionEdit GeometryCollectionEdit = GeometryCollectionComponent->EditRestCollection(GeometryCollection::EEditUpdate::None);
 		if( UGeometryCollection* GeometryCollectionObject = GeometryCollectionEdit.GetRestCollection() )
 		{
 			TSharedPtr<FGeometryCollection, ESPMode::ThreadSafe> GeometryCollectionPtr = GeometryCollectionObject->GetGeometryCollection();
diff --git a/Engine/Plugins/Runtime/EditableMesh/Source/EditableMesh/GeometryCollectionEditableMeshFormat.cpp b/Engine/Plugins/Runtime/EditableMesh/Source/EditableMesh/GeometryCollectionEditableMeshFormat.cpp
--- a/Engine/Plugins/Runtime/EditableMesh/Source/EditableMesh/GeometryCollectionEditableMeshFormat.cpp
+++ b/Engine/Plugins/Runtime/EditableMesh/Source/EditableMesh/GeometryCollectionEditableMeshFormat.cpp
@@ -26,7 +26,7 @@ void FGeometryCollectionEditableMeshFormat::FillMeshObjectPtr( UPrimitiveCompone
 	{
 		UGeometryCollectionComponent& GeometryCollectionComponent = *GeometryCollectionComponentPtr;
 
-		FGeometryCollectionEdit GeometryCollectionEdit = GeometryCollectionComponent.EditRestCollection(false);
+		FGeometryCollectionEdit GeometryCollectionEdit = GeometryCollectionComponent.EditRestCollection(GeometryCollection::EEditUpdate::None);
 		UGeometryCollection* GeometryCollection = GeometryCollectionEdit.GetRestCollection();
 		if( GeometryCollection != nullptr)
 		{
diff --git a/Engine/Source/Runtime/Experimental/GeometryCollectionEngine/Private/GeometryCollection/GeometryCollectionActor.cpp b/Engine/Source/Runtime/Experimental/GeometryCollectionEngine/Private/GeometryCollection/GeometryCollectionActor.cpp
--- a/Engine/Source/Runtime/Experimental/GeometryCollectionEngine/Private/GeometryCollection/GeometryCollectionActor.cpp
+++ b/Engine/Source/Runtime/Experimental/GeometryCollectionEngine/Private/GeometryCollection/GeometryCollectionActor.cpp
@@ -135,7 +135,7 @@ bool AGeometryCollectionActor::GetReferencedContentObjects(TArray<UObject*>& Obj
 
 	if (GeometryCollectionComponent)
 	{
-		FGeometryCollectionEdit GeometryCollectionEdit = GeometryCollectionComponent->EditRestCollection(false);
+		FGeometryCollectionEdit GeometryCollectionEdit = GeometryCollectionComponent->EditRestCollection(GeometryCollection::EEditUpdate::None);
 		if (UGeometryCollection* GeometryCollection = GeometryCollectionEdit.GetRestCollection())
 		{
 			Objects.Add(GeometryCollection);
diff --git a/Engine/Source/Runtime/Experimental/GeometryCollectionEngine/Private/GeometryCollection/GeometryCollectionComponent.cpp b/Engine/Source/Runtime/Experimental/GeometryCollectionEngine/Private/GeometryCollection/GeometryCollectionComponent.cpp
--- a/Engine/Source/Runtime/Experimental/GeometryCollectionEngine/Private/GeometryCollection/GeometryCollectionComponent.cpp
+++ b/Engine/Source/Runtime/Experimental/GeometryCollectionEngine/Private/GeometryCollection/GeometryCollectionComponent.cpp
@@ -1461,12 +1461,12 @@ void UGeometryCollectionComponent::SetRestCollection(const UGeometryCollection*
 	}
 }
 
-FGeometryCollectionEdit::FGeometryCollectionEdit(UGeometryCollectionComponent* InComponent, bool InUpdate /*= true*/)
+FGeometryCollectionEdit::FGeometryCollectionEdit(UGeometryCollectionComponent* InComponent, GeometryCollection::EEditUpdate InEditUpdate)
 	: Component(InComponent)
-	, bUpdate(InUpdate)
+	, EditUpdate(InEditUpdate)
 {
 	bHadPhysicsState = Component->HasValidPhysicsState();
-	if(bUpdate && bHadPhysicsState)
+	if (EnumHasAnyFlags(EditUpdate, GeometryCollection::EEditUpdate::Physics) && bHadPhysicsState)
 	{
 		Component->DestroyPhysicsState();
 	}
@@ -1475,15 +1475,19 @@ FGeometryCollectionEdit::FGeometryCollectionEdit(UGeometryCollectionComponent* I
 FGeometryCollectionEdit::~FGeometryCollectionEdit()
 {
 #if WITH_EDITOR
-	if (bUpdate)
+	if (!!EditUpdate)
 	{
-		//Component->ResetDynamicCollection();
-		if (GetRestCollection())
+		if (EnumHasAnyFlags(EditUpdate, GeometryCollection::EEditUpdate::Dynamic))
+		{
+			Component->ResetDynamicCollection();
+		}
+
+		if (EnumHasAnyFlags(EditUpdate, GeometryCollection::EEditUpdate::Rest) && GetRestCollection())
 		{
 			GetRestCollection()->Modify();
 		}
 
-		if(bHadPhysicsState)
+		if (EnumHasAnyFlags(EditUpdate, GeometryCollection::EEditUpdate::Physics) && bHadPhysicsState)
 		{
 			Component->RecreatePhysicsState();
 		}
diff --git a/Engine/Source/Runtime/Experimental/GeometryCollectionEngine/Public/GeometryCollection/GeometryCollectionComponent.h b/Engine/Source/Runtime/Experimental/GeometryCollectionEngine/Public/GeometryCollection/GeometryCollectionComponent.h
--- a/Engine/Source/Runtime/Experimental/GeometryCollectionEngine/Public/GeometryCollection/GeometryCollectionComponent.h
+++ b/Engine/Source/Runtime/Experimental/GeometryCollectionEngine/Public/GeometryCollection/GeometryCollectionComponent.h
@@ -141,6 +141,27 @@ struct FGeomComponentCacheParameters
 	float TrailingMinVolumeThreshold;
 };
 
+namespace GeometryCollection
+{
+	/** Type of updates used at the end of an edit operation. */
+	enum class EEditUpdate : uint8
+	{
+		/** No update. */
+		None = 0,
+		/** Mark the rest collection as changed. */
+		Rest = 1,
+		/** Recreate the physics state (proxy). */
+		Physics = 2,
+		/** Reset the dynamic collection. */
+		Dynamic = 4,
+		/** Mark the rest collection as changed, and recreate the physics state (proxy). */
+		RestPhysics = Rest | Physics,
+		/** Reset dynamic collection, mark the rest collection as changed, and recreate the physics state (proxy). */
+		RestPhysicsDynamic = Rest | Physics | Dynamic,
+	};
+	ENUM_CLASS_FLAGS(EEditUpdate);
+}
+
 /**
 *	FGeometryCollectionEdit
 *     Structured RestCollection access where the scope
@@ -149,18 +170,19 @@ struct FGeomComponentCacheParameters
 *
 *	This will force any simulating geometry collection out of the
 *	solver so it can be edited and afterwards will recreate the proxy
+*	The update can also be specified to reset the dynamic collection
 */
 class GEOMETRYCOLLECTIONENGINE_API FGeometryCollectionEdit
 {
 public:
-	FGeometryCollectionEdit(UGeometryCollectionComponent * InComponent, bool InUpdate = true);
+	FGeometryCollectionEdit(UGeometryCollectionComponent* InComponent, GeometryCollection::EEditUpdate EditUpdate = GeometryCollection::EEditUpdate::RestPhysics);
 	~FGeometryCollectionEdit();
 
 	UGeometryCollection* GetRestCollection();
 
 private:
 	UGeometryCollectionComponent* Component;
-	const bool bUpdate;
+	const GeometryCollection::EEditUpdate EditUpdate;
 	bool bHadPhysicsState;
 };
 
@@ -290,7 +312,7 @@ public:
 	/** RestCollection */
 	void SetRestCollection(const UGeometryCollection * RestCollectionIn);
 	FORCEINLINE const UGeometryCollection* GetRestCollection() const { return RestCollection; }
-	FORCEINLINE FGeometryCollectionEdit EditRestCollection(bool Update = true) { return FGeometryCollectionEdit(this, Update); }
+	FORCEINLINE FGeometryCollectionEdit EditRestCollection(GeometryCollection::EEditUpdate EditUpdate = GeometryCollection::EEditUpdate::RestPhysics) { return FGeometryCollectionEdit(this, EditUpdate); }
 #if WITH_EDITOR
 	FORCEINLINE FScopedColorEdit EditBoneSelection() { return FScopedColorEdit(this); }
 #endif
