Fix geometry collection debug draw rigid body selection when picking a rigid body id that exists in different solvers.

diff --git a/Engine/Plugins/Experimental/GeometryCollectionPlugin/Source/GeometryCollectionEditor/Private/GeometryCollection/DetailCustomizations/SelectedRigidBodyCustomization.cpp b/Engine/Plugins/Experimental/GeometryCollectionPlugin/Source/GeometryCollectionEditor/Private/GeometryCollection/DetailCustomizations/SelectedRigidBodyCustomization.cpp
--- a/Engine/Plugins/Experimental/GeometryCollectionPlugin/Source/GeometryCollectionEditor/Private/GeometryCollection/DetailCustomizations/SelectedRigidBodyCustomization.cpp
+++ b/Engine/Plugins/Experimental/GeometryCollectionPlugin/Source/GeometryCollectionEditor/Private/GeometryCollection/DetailCustomizations/SelectedRigidBodyCustomization.cpp
@@ -4,7 +4,9 @@
 
 #include "PropertyHandle.h"
 #include "IPropertyTypeCustomization.h"
+#include "IDetailChildrenBuilder.h"
 #include "IPropertyUtilities.h"
+#include "IDetailPropertyRow.h"
 #include "DetailWidgetRow.h"
 #include "Widgets/Text/STextBlock.h"
 #include "Widgets/Input/SButton.h"
@@ -22,8 +24,9 @@ TSharedRef<IPropertyTypeCustomization> FSelectedRigidBodyCustomization::MakeInst
 
 void FSelectedRigidBodyCustomization::CustomizeHeader(TSharedRef<IPropertyHandle> StructPropertyHandle, FDetailWidgetRow& HeaderRow, IPropertyTypeCustomizationUtils& /*StructCustomizationUtils*/)
 {
-	const TSharedRef<IPropertyHandle> PropertyHandleActor = StructPropertyHandle->GetChildHandle(GET_MEMBER_NAME_CHECKED(FGeometryCollectionDebugDrawActorSelectedRigidBody, GeometryCollectionActor)).ToSharedRef();
 	const TSharedRef<IPropertyHandle> PropertyHandleId = StructPropertyHandle->GetChildHandle(GET_MEMBER_NAME_CHECKED(FGeometryCollectionDebugDrawActorSelectedRigidBody, Id)).ToSharedRef();
+	const TSharedRef<IPropertyHandle> PropertyHandleSolver = StructPropertyHandle->GetChildHandle(GET_MEMBER_NAME_CHECKED(FGeometryCollectionDebugDrawActorSelectedRigidBody, Solver)).ToSharedRef();
+	const TSharedRef<IPropertyHandle> PropertyHandleActor = StructPropertyHandle->GetChildHandle(GET_MEMBER_NAME_CHECKED(FGeometryCollectionDebugDrawActorSelectedRigidBody, GeometryCollection)).ToSharedRef();
 
 	// Add buttons
 	HeaderRow.NameContent()
@@ -33,152 +36,155 @@ void FSelectedRigidBodyCustomization::CustomizeHeader(TSharedRef<IPropertyHandle
 				NSLOCTEXT("ChaosSelectedRigidBody", "ChaosSelectedRigidBody_ToolTip", "Select a Rigid Body here by either entering its Id, or clicking on the Pick button."))
 		]
 	.ValueContent()
-		.MinDesiredWidth(360.f)
+	.MinDesiredWidth(140.f)
+	[
+		SNew(SHorizontalBox)
+		+ SHorizontalBox::Slot()
+		.AutoWidth()
+		.Padding(0.f)
 		[
-			SNew(SHorizontalBox)
-			+ SHorizontalBox::Slot()
-			[
-				PropertyHandleId->CreatePropertyValueWidget()
-			]
-			+ SHorizontalBox::Slot()
-			.AutoWidth()
-			.Padding(2.f, 0.f, 0.f, 0.f)
-			[
-				SAssignNew(CheckBoxPick, SCheckBox)
-				.Style(FEditorStyle::Get(), "ToggleButtonCheckbox")
-				.ToolTipText(NSLOCTEXT("ChaosSelectedRigidBody", "Pick_ToolTip", "Pick a Rigid Body."))
-				.ForegroundColor(FSlateColor::UseForeground())
-				.Padding(6.0f)
-				.IsChecked_Lambda([]()
-				{
-					return FGeometryCollectionSelectRigidBodyEdMode::IsModeActive() ? ECheckBoxState::Checked : ECheckBoxState::Unchecked;
-				})
-				.IsEnabled_Lambda([]()
-				{
-					return FGeometryCollectionSelectRigidBodyEdMode::CanActivateMode();
-				})
-				.OnCheckStateChanged(this, &FSelectedRigidBodyCustomization::OnPick, PropertyHandleId)
-				[ 
-					SNew(SImage)
-					.Image(FEditorStyle::GetBrush("PropertyWindow.Button_PickActorInteractive"))
-					.ColorAndOpacity(FSlateColor::UseForeground())
-				]
+			SAssignNew(CheckBoxPick, SCheckBox)
+			.Style(FEditorStyle::Get(), "ToggleButtonCheckbox")
+			.ToolTipText(NSLOCTEXT("ChaosSelectedRigidBody", "Pick_ToolTip", "Pick a Rigid Body."))
+			.ForegroundColor(FSlateColor::UseForeground())
+			.Padding(6.0f)
+			.IsChecked_Lambda([]()
+			{
+				return FGeometryCollectionSelectRigidBodyEdMode::IsModeActive() ? ECheckBoxState::Checked : ECheckBoxState::Unchecked;
+			})
+			.IsEnabled_Lambda([]()
+			{
+				return FGeometryCollectionSelectRigidBodyEdMode::CanActivateMode();
+			})
+			.OnCheckStateChanged(this, &FSelectedRigidBodyCustomization::OnPick, PropertyHandleId, PropertyHandleSolver)
+			[ 
+				SNew(SImage)
+				.Image(FEditorStyle::GetBrush("PropertyWindow.Button_PickActorInteractive"))
+				.ColorAndOpacity(FSlateColor::UseForeground())
 			]
-			+ SHorizontalBox::Slot()
-			.AutoWidth()
-			.Padding(8.f, 0.f, 0.f, 0.f)
-			[
-				SNew(SButton)
-				.ButtonStyle(FEditorStyle::Get(), "HoverHintOnly")
-				.ToolTipText(NSLOCTEXT("ChaosSelectedRigidBody", "Pick_UpClusterLevel", "Go to parent cluster."))
-				.ForegroundColor(FSlateColor::UseForeground())
-				.IsEnabled_Lambda([PropertyHandleActor, PropertyHandleId]()
-				{
-					return GetParentClusterRigidBodyId(PropertyHandleActor, PropertyHandleId) != INDEX_NONE;
-				})
-				.OnClicked_Lambda([PropertyHandleActor, PropertyHandleId]()
+		]
+		+ SHorizontalBox::Slot()
+		.AutoWidth()
+		.Padding(8.f, 0.f, 0.f, 0.f)
+		[
+			SNew(SButton)
+			.ButtonStyle(FEditorStyle::Get(), "HoverHintOnly")
+			.ToolTipText(NSLOCTEXT("ChaosSelectedRigidBody", "Pick_UpClusterLevel", "Go to parent cluster."))
+			.ForegroundColor(FSlateColor::UseForeground())
+			.IsEnabled_Lambda([PropertyHandleActor, PropertyHandleId]()
+			{
+				return GetParentClusterRigidBodyId(PropertyHandleActor, PropertyHandleId) != INDEX_NONE;
+			})
+			.OnClicked_Lambda([PropertyHandleActor, PropertyHandleId]()
+			{
+				const int32 RigidBodyId = GetParentClusterRigidBodyId(PropertyHandleActor, PropertyHandleId);
+				if (RigidBodyId != INDEX_NONE)
 				{
-					const int32 RigidBodyId = GetParentClusterRigidBodyId(PropertyHandleActor, PropertyHandleId);
-					if (RigidBodyId != INDEX_NONE)
-					{
-						PropertyHandleId->SetValue(RigidBodyId);
-					}
-					return FReply::Handled();
-				})
-				[ 
-					SNew(STextBlock)
-					.Font(FEditorStyle::Get().GetFontStyle("FontAwesome.16"))
-					.Text(FEditorFontGlyphs::Caret_Square_O_Up)
-					.ColorAndOpacity(FSlateColor::UseForeground())
-				]
+					PropertyHandleId->SetValue(RigidBodyId);
+				}
+				return FReply::Handled();
+			})
+			[ 
+				SNew(STextBlock)
+				.Font(FEditorStyle::Get().GetFontStyle("FontAwesome.16"))
+				.Text(FEditorFontGlyphs::Caret_Square_O_Up)
+				.ColorAndOpacity(FSlateColor::UseForeground())
 			]
-			+ SHorizontalBox::Slot()
-			.AutoWidth()
-			[
-				SNew(SButton)
-				.ButtonStyle(FEditorStyle::Get(), "HoverHintOnly")
-				.ToolTipText(NSLOCTEXT("ChaosSelectedRigidBody", "Pick_DownClusterLevel", "Go to child cluster."))
-				.ForegroundColor(FSlateColor::UseForeground())
-				.IsEnabled_Lambda([PropertyHandleActor, PropertyHandleId]()
-				{
-					return GetChildClusterRigidBodyId(PropertyHandleActor, PropertyHandleId) != INDEX_NONE;
-				})
-				.OnClicked_Lambda([PropertyHandleActor, PropertyHandleId]()
+		]
+		+ SHorizontalBox::Slot()
+		.AutoWidth()
+		[
+			SNew(SButton)
+			.ButtonStyle(FEditorStyle::Get(), "HoverHintOnly")
+			.ToolTipText(NSLOCTEXT("ChaosSelectedRigidBody", "Pick_DownClusterLevel", "Go to child cluster."))
+			.ForegroundColor(FSlateColor::UseForeground())
+			.IsEnabled_Lambda([PropertyHandleActor, PropertyHandleId]()
+			{
+				return GetChildClusterRigidBodyId(PropertyHandleActor, PropertyHandleId) != INDEX_NONE;
+			})
+			.OnClicked_Lambda([PropertyHandleActor, PropertyHandleId]()
+			{
+				const int32 RigidBodyId = GetChildClusterRigidBodyId(PropertyHandleActor, PropertyHandleId);
+				if (RigidBodyId != INDEX_NONE)
 				{
-					const int32 RigidBodyId = GetChildClusterRigidBodyId(PropertyHandleActor, PropertyHandleId);
-					if (RigidBodyId != INDEX_NONE)
-					{
-						PropertyHandleId->SetValue(RigidBodyId);
-					}
-					return FReply::Handled();
-				})
-				[ 
-					SNew(STextBlock)
-					.Font(FEditorStyle::Get().GetFontStyle("FontAwesome.16"))
-					.Text(FEditorFontGlyphs::Caret_Square_O_Down)
-					.ColorAndOpacity(FSlateColor::UseForeground())
-				]
+					PropertyHandleId->SetValue(RigidBodyId);
+				}
+				return FReply::Handled();
+			})
+			[ 
+				SNew(STextBlock)
+				.Font(FEditorStyle::Get().GetFontStyle("FontAwesome.16"))
+				.Text(FEditorFontGlyphs::Caret_Square_O_Down)
+				.ColorAndOpacity(FSlateColor::UseForeground())
 			]
-			+ SHorizontalBox::Slot()
-			.AutoWidth()
-			.Padding(8.0f, 0.f, 0.f, 0.f)
-			[
-				SNew(SButton)
-				.ButtonStyle(FEditorStyle::Get(), "HoverHintOnly")
-				.ToolTipText(NSLOCTEXT("ChaosSelectedRigidBody", "Pick_PrevClusterSibling", "Go to previous clustered sibling."))
-				.ForegroundColor(FSlateColor::UseForeground())
-				.IsEnabled_Lambda([PropertyHandleActor, PropertyHandleId]()
-				{
-					return GetPreviousClusteredSiblingRigidBodyId(PropertyHandleActor, PropertyHandleId) != INDEX_NONE;
-				})
-				.OnClicked_Lambda([PropertyHandleActor, PropertyHandleId]()
+		]
+		+ SHorizontalBox::Slot()
+		.AutoWidth()
+		.Padding(8.0f, 0.f, 0.f, 0.f)
+		[
+			SNew(SButton)
+			.ButtonStyle(FEditorStyle::Get(), "HoverHintOnly")
+			.ToolTipText(NSLOCTEXT("ChaosSelectedRigidBody", "Pick_PrevClusterSibling", "Go to previous clustered sibling."))
+			.ForegroundColor(FSlateColor::UseForeground())
+			.IsEnabled_Lambda([PropertyHandleActor, PropertyHandleId]()
+			{
+				return GetPreviousClusteredSiblingRigidBodyId(PropertyHandleActor, PropertyHandleId) != INDEX_NONE;
+			})
+			.OnClicked_Lambda([PropertyHandleActor, PropertyHandleId]()
+			{
+				const int32 RigidBodyId = GetPreviousClusteredSiblingRigidBodyId(PropertyHandleActor, PropertyHandleId);
+				if (RigidBodyId != INDEX_NONE)
 				{
-					const int32 RigidBodyId = GetPreviousClusteredSiblingRigidBodyId(PropertyHandleActor, PropertyHandleId);
-					if (RigidBodyId != INDEX_NONE)
-					{
-						PropertyHandleId->SetValue(RigidBodyId);
-					}
-					return FReply::Handled();
-				})
-				[ 
-					SNew(STextBlock)
-					.Font(FEditorStyle::Get().GetFontStyle("FontAwesome.16"))
-					.Text(FEditorFontGlyphs::Caret_Square_O_Left)
-					.ColorAndOpacity(FSlateColor::UseForeground())
-				]
+					PropertyHandleId->SetValue(RigidBodyId);
+				}
+				return FReply::Handled();
+			})
+			[ 
+				SNew(STextBlock)
+				.Font(FEditorStyle::Get().GetFontStyle("FontAwesome.16"))
+				.Text(FEditorFontGlyphs::Caret_Square_O_Left)
+				.ColorAndOpacity(FSlateColor::UseForeground())
 			]
-			+ SHorizontalBox::Slot()
-			.AutoWidth()
-			[
-				SNew(SButton)
-				.ButtonStyle(FEditorStyle::Get(), "HoverHintOnly")
-				.ToolTipText(NSLOCTEXT("ChaosSelectedRigidBody", "Pick_NextClusterSibling", "Go to next clustered sibling."))
-				.ForegroundColor(FSlateColor::UseForeground())
-				.IsEnabled_Lambda([PropertyHandleActor, PropertyHandleId]()
-				{
-					return GetNextClusteredSiblingRigidBodyId(PropertyHandleActor, PropertyHandleId) != INDEX_NONE;
-				})
-				.OnClicked_Lambda([PropertyHandleActor, PropertyHandleId]()
+		]
+		+ SHorizontalBox::Slot()
+		.AutoWidth()
+		[
+			SNew(SButton)
+			.ButtonStyle(FEditorStyle::Get(), "HoverHintOnly")
+			.ToolTipText(NSLOCTEXT("ChaosSelectedRigidBody", "Pick_NextClusterSibling", "Go to next clustered sibling."))
+			.ForegroundColor(FSlateColor::UseForeground())
+			.IsEnabled_Lambda([PropertyHandleActor, PropertyHandleId]()
+			{
+				return GetNextClusteredSiblingRigidBodyId(PropertyHandleActor, PropertyHandleId) != INDEX_NONE;
+			})
+			.OnClicked_Lambda([PropertyHandleActor, PropertyHandleId]()
+			{
+				const int32 RigidBodyId = GetNextClusteredSiblingRigidBodyId(PropertyHandleActor, PropertyHandleId);
+				if (RigidBodyId != INDEX_NONE)
 				{
-					const int32 RigidBodyId = GetNextClusteredSiblingRigidBodyId(PropertyHandleActor, PropertyHandleId);
-					if (RigidBodyId != INDEX_NONE)
-					{
-						PropertyHandleId->SetValue(RigidBodyId);
-					}
-					return FReply::Handled();
-				})
-				[ 
-					SNew(STextBlock)
-					.Font(FEditorStyle::Get().GetFontStyle("FontAwesome.16"))
-					.Text(FEditorFontGlyphs::Caret_Square_O_Right)
-					.ColorAndOpacity(FSlateColor::UseForeground())
-				]
+					PropertyHandleId->SetValue(RigidBodyId);
+				}
+				return FReply::Handled();
+			})
+			[ 
+				SNew(STextBlock)
+				.Font(FEditorStyle::Get().GetFontStyle("FontAwesome.16"))
+				.Text(FEditorFontGlyphs::Caret_Square_O_Right)
+				.ColorAndOpacity(FSlateColor::UseForeground())
 			]
-		];
+		]
+	];
 }
 
-void FSelectedRigidBodyCustomization::CustomizeChildren(TSharedRef<IPropertyHandle> /*StructPropertyHandle*/, class IDetailChildrenBuilder& /*ChildBuilder*/, IPropertyTypeCustomizationUtils& /*StructCustomizationUtils*/)
+void FSelectedRigidBodyCustomization::CustomizeChildren(TSharedRef<IPropertyHandle> StructPropertyHandle, IDetailChildrenBuilder& ChildBuilder, IPropertyTypeCustomizationUtils& /*StructCustomizationUtils*/)
 {
+	const TSharedRef<IPropertyHandle> PropertyHandleId = StructPropertyHandle->GetChildHandle(GET_MEMBER_NAME_CHECKED(FGeometryCollectionDebugDrawActorSelectedRigidBody, Id)).ToSharedRef();
+	const TSharedRef<IPropertyHandle> PropertyHandleSolver = StructPropertyHandle->GetChildHandle(GET_MEMBER_NAME_CHECKED(FGeometryCollectionDebugDrawActorSelectedRigidBody, Solver)).ToSharedRef();
+	const TSharedRef<IPropertyHandle> PropertyHandleGeometryCollection = StructPropertyHandle->GetChildHandle(GET_MEMBER_NAME_CHECKED(FGeometryCollectionDebugDrawActorSelectedRigidBody, GeometryCollection)).ToSharedRef();
+
+	ChildBuilder.AddProperty(PropertyHandleId);
+	ChildBuilder.AddProperty(PropertyHandleSolver);
+	ChildBuilder.AddProperty(PropertyHandleGeometryCollection);
 }
 
 void FSelectedRigidBodyCustomization::GetSelectedGeometryCollectionCluster(TSharedRef<IPropertyHandle> PropertyHandleActor, TSharedRef<IPropertyHandle> PropertyHandleId, const UGeometryCollectionComponent*& OutGeometryCollectionComponent, int32& OutTransformIndex)
@@ -298,7 +304,7 @@ int32 FSelectedRigidBodyCustomization::GetNextClusteredSiblingRigidBodyId(TShare
 	return INDEX_NONE;
 }
 
-void FSelectedRigidBodyCustomization::OnPick(ECheckBoxState InCheckState, TSharedRef<IPropertyHandle> PropertyHandleId) const
+void FSelectedRigidBodyCustomization::OnPick(ECheckBoxState InCheckState, TSharedRef<IPropertyHandle> PropertyHandleId, TSharedRef<IPropertyHandle> PropertyHandleSolver) const
 {
 	// Set unchecked by default (enter mode is required to set checked)
 	if (CheckBoxPick)
@@ -332,6 +338,6 @@ void FSelectedRigidBodyCustomization::OnPick(ECheckBoxState InCheckState, TShare
 				CheckBoxPickPin->SetIsChecked(ECheckBoxState::Unchecked);
 			}
 		};
-		FGeometryCollectionSelectRigidBodyEdMode::ActivateMode(PropertyHandleId, OnEnterMode, OnExitMode);
+		FGeometryCollectionSelectRigidBodyEdMode::ActivateMode(PropertyHandleId, PropertyHandleSolver, OnEnterMode, OnExitMode);
 	}
 }
diff --git a/Engine/Plugins/Experimental/GeometryCollectionPlugin/Source/GeometryCollectionEditor/Private/GeometryCollection/DetailCustomizations/SelectedRigidBodyCustomization.h b/Engine/Plugins/Experimental/GeometryCollectionPlugin/Source/GeometryCollectionEditor/Private/GeometryCollection/DetailCustomizations/SelectedRigidBodyCustomization.h
--- a/Engine/Plugins/Experimental/GeometryCollectionPlugin/Source/GeometryCollectionEditor/Private/GeometryCollection/DetailCustomizations/SelectedRigidBodyCustomization.h
+++ b/Engine/Plugins/Experimental/GeometryCollectionPlugin/Source/GeometryCollectionEditor/Private/GeometryCollection/DetailCustomizations/SelectedRigidBodyCustomization.h
@@ -26,5 +26,5 @@ private:
 	static int32 GetPreviousClusteredSiblingRigidBodyId(TSharedRef<IPropertyHandle> PropertyHandleActor, TSharedRef<IPropertyHandle> PropertyHandleId);
 	static int32 GetNextClusteredSiblingRigidBodyId(TSharedRef<IPropertyHandle> PropertyHandleActor, TSharedRef<IPropertyHandle> PropertyHandleId);
 
-	void OnPick(ECheckBoxState InCheckState, TSharedRef<IPropertyHandle> PropertyHandleId) const;
+	void OnPick(ECheckBoxState InCheckState, TSharedRef<IPropertyHandle> PropertyHandleId, TSharedRef<IPropertyHandle> PropertyHandleSolver) const;
 };
diff --git a/Engine/Plugins/Experimental/GeometryCollectionPlugin/Source/GeometryCollectionEditor/Private/GeometryCollection/GeometryCollectionSelectRigidBodyEdMode.cpp b/Engine/Plugins/Experimental/GeometryCollectionPlugin/Source/GeometryCollectionEditor/Private/GeometryCollection/GeometryCollectionSelectRigidBodyEdMode.cpp
--- a/Engine/Plugins/Experimental/GeometryCollectionPlugin/Source/GeometryCollectionEditor/Private/GeometryCollection/GeometryCollectionSelectRigidBodyEdMode.cpp
+++ b/Engine/Plugins/Experimental/GeometryCollectionPlugin/Source/GeometryCollectionEditor/Private/GeometryCollection/GeometryCollectionSelectRigidBodyEdMode.cpp
@@ -39,7 +39,7 @@ bool FGeometryCollectionSelectRigidBodyEdMode::CanActivateMode()
 		&& HasValidEditorLevelViewport(GEditor->GetWorldContexts(), GEditor->SlatePlayInEditorMap);  // Not running in new PIE window
 }
 
-void FGeometryCollectionSelectRigidBodyEdMode::ActivateMode(TSharedRef<IPropertyHandle> PropertyHandleId, TFunction<void()> OnEnterMode, TFunction<void()> OnExitMode)
+void FGeometryCollectionSelectRigidBodyEdMode::ActivateMode(TSharedRef<IPropertyHandle> PropertyHandleId, TSharedRef<IPropertyHandle> PropertyHandleSolver, TFunction<void()> OnEnterMode, TFunction<void()> OnExitMode)
 {
 	// Make sure we are playing in an editor window
 	if (CanActivateMode())
@@ -65,6 +65,7 @@ void FGeometryCollectionSelectRigidBodyEdMode::ActivateMode(TSharedRef<IProperty
 			// Set pointers
 			FGeometryCollectionSelectRigidBodyEdMode* const GeometryCollectionSelectRigidBodyEdMode = static_cast<FGeometryCollectionSelectRigidBodyEdMode*>(EdMode);
 			GeometryCollectionSelectRigidBodyEdMode->PropertyHandleId = PropertyHandleId;
+			GeometryCollectionSelectRigidBodyEdMode->PropertyHandleSolver = PropertyHandleSolver;
 			GeometryCollectionSelectRigidBodyEdMode->OnExitMode = OnExitMode;
 
 			// Execute enter mode callback
@@ -167,6 +168,10 @@ bool FGeometryCollectionSelectRigidBodyEdMode::HandleClick(FEditorViewportClient
 					{
 						PropertyHandleIdPin->SetValue(RigidBodyId);
 					}
+					if (const TSharedPtr<IPropertyHandle> PropertyHandleSolverPin = PropertyHandleSolver.Pin())
+					{
+						PropertyHandleSolverPin->SetValue(GeometryCollectionComponent->ChaosSolverActor);
+					}
 				}
 			}
 		}
diff --git a/Engine/Plugins/Experimental/GeometryCollectionPlugin/Source/GeometryCollectionEditor/Private/GeometryCollection/GeometryCollectionSelectRigidBodyEdMode.h b/Engine/Plugins/Experimental/GeometryCollectionPlugin/Source/GeometryCollectionEditor/Private/GeometryCollection/GeometryCollectionSelectRigidBodyEdMode.h
--- a/Engine/Plugins/Experimental/GeometryCollectionPlugin/Source/GeometryCollectionEditor/Private/GeometryCollection/GeometryCollectionSelectRigidBodyEdMode.h
+++ b/Engine/Plugins/Experimental/GeometryCollectionPlugin/Source/GeometryCollectionEditor/Private/GeometryCollection/GeometryCollectionSelectRigidBodyEdMode.h
@@ -17,7 +17,7 @@ public:
 	static bool CanActivateMode();
 
 	/** Activate this editor mode. */
-	static void ActivateMode(TSharedRef<IPropertyHandle> PropertyHandleId, TFunction<void()> OnEnterMode, TFunction<void()> OnExitMode);
+	static void ActivateMode(TSharedRef<IPropertyHandle> PropertyHandleId, TSharedRef<IPropertyHandle> PropertyHandleSolver, TFunction<void()> OnEnterMode, TFunction<void()> OnExitMode);
 
 	/** Deactivate this editor mode. */
 	static void DeactivateMode();
@@ -51,6 +51,7 @@ private:
 private:
 	static const int32 MessageKey;
 	TWeakPtr<IPropertyHandle> PropertyHandleId;  // Handle of the property that will get updated with the selected rigid body id
+	TWeakPtr<IPropertyHandle> PropertyHandleSolver;  // Handle of the property that will get updated with the selected solver actor
 	TFunction<void()> OnExitMode;  // Callback function called when the edit mode is deactivated
 	bool bIsHoveringGeometryCollection;
 };
diff --git a/Engine/Source/Runtime/Experimental/GeometryCollectionEngine/Private/GeometryCollection/GeometryCollectionDebugDrawActor.cpp b/Engine/Source/Runtime/Experimental/GeometryCollectionEngine/Private/GeometryCollection/GeometryCollectionDebugDrawActor.cpp
--- a/Engine/Source/Runtime/Experimental/GeometryCollectionEngine/Private/GeometryCollection/GeometryCollectionDebugDrawActor.cpp
+++ b/Engine/Source/Runtime/Experimental/GeometryCollectionEngine/Private/GeometryCollection/GeometryCollectionDebugDrawActor.cpp
@@ -17,6 +17,7 @@
 #include "UObject/ConstructorHelpers.h"
 #include "Components/BillboardComponent.h"
 #include "GenericPlatform/GenericPlatformMath.h"
+#include "HAL/IConsoleManager.h"
 
 DEFINE_LOG_CATEGORY_STATIC(LogGeometryCollectionDebugDrawActor, Log, All);
 
@@ -36,6 +37,7 @@ namespace GeometryCollectionDebugDrawActorConstants
 	});
 
 	// Defaults
+	static const FString SelectedRigidBodySolverDefault = FName(NAME_None).ToString();
 	static const int32 SelectedRigidBodyIdDefault = INDEX_NONE;
 	static const int32 DebugDrawWholeCollectionDefault = 0;
 	static const int32 DebugDrawHierarchyDefault = 0;
@@ -78,43 +80,49 @@ namespace GeometryCollectionDebugDrawActorConstants
 namespace GeometryCollectionDebugDrawActorCVars
 {
 	// Console variables, also exposed as settings in this actor
-	static TAutoConsoleVariable<int32> SelectedRigidBodyId     (TEXT("p.gc.SelectedRigidBodyId"     ), GeometryCollectionDebugDrawActorConstants::SelectedRigidBodyIdDefault     , TEXT("Geometry Collection debug draw, visualize debug informations for the selected rigid body ids.\nDefault = -1"), ECVF_Cheat);
-	static TAutoConsoleVariable<int32> DebugDrawWholeCollection(TEXT("p.gc.DebugDrawWholeCollection"), GeometryCollectionDebugDrawActorConstants::DebugDrawWholeCollectionDefault, TEXT("Geometry Collection debug draw, show debug visualization for the rest of the geometry collection related to the current rigid body id selection.\nDefault = 0"), ECVF_Cheat);
-	static TAutoConsoleVariable<int32> DebugDrawHierarchy      (TEXT("p.gc.DebugDrawHierarchy"      ), GeometryCollectionDebugDrawActorConstants::DebugDrawHierarchyDefault      , TEXT("Geometry Collection debug draw, show debug visualization for the top level node rather than the bottom leaf nodes of a cluster's hierarchy..\nDefault = 0"), ECVF_Cheat);
-	static TAutoConsoleVariable<int32> DebugDrawClustering     (TEXT("p.gc.DebugDrawClustering"     ), GeometryCollectionDebugDrawActorConstants::DebugDrawClusteringDefault     , TEXT("Geometry Collection debug draw, show debug visualization for all clustered children associated to the current rigid body id selection.\nDefault = 0"), ECVF_Cheat);
-	static TAutoConsoleVariable<int32> HideGeometry            (TEXT("p.gc.HideGeometry"            ), GeometryCollectionDebugDrawActorConstants::HideGeometryDefault            , TEXT("Geometry Collection debug draw, geometry visibility setting, select the part of the geometry to hide in order to better visualize the debug information.\n0: Do not hide any geometries.\n1: Hide the geometry associated to the rigid bodies selected for collision display.\n2: Hide the geometry associated to the selected rigid bodies.\n3: Hide the entire geometry collection associated to the selected rigid bodies.\n4: Hide all geometry collections.\nDefault = 1"), ECVF_Cheat);
-	static TAutoConsoleVariable<int32> ShowRigidBodyId         (TEXT("p.gc.ShowRigidBodyId"         ), GeometryCollectionDebugDrawActorConstants::ShowRigidBodyIdDefault         , TEXT("Geometry Collection debug draw, show the rigid body id(s).\nDefault = 0"), ECVF_Cheat);
-	static TAutoConsoleVariable<int32> ShowRigidBodyCollision  (TEXT("p.gc.ShowRigidBodyCollision"  ), GeometryCollectionDebugDrawActorConstants::ShowRigidBodyCollisionDefault  , TEXT("Geometry Collection debug draw, show the selected's rigid body's collision volume.\nDefault = 0"), ECVF_Cheat);
-	static TAutoConsoleVariable<int32> CollisionAtOrigin       (TEXT("p.gc.CollisionAtOrigin"       ), GeometryCollectionDebugDrawActorConstants::CollisionAtOriginDefault       , TEXT("Geometry Collection debug draw, show any collision volume at the origin, in local space.\nDefault = 0"), ECVF_Cheat);
-	static TAutoConsoleVariable<int32> ShowRigidBodyTransform  (TEXT("p.gc.ShowRigidBodyTransform"  ), GeometryCollectionDebugDrawActorConstants::ShowRigidBodyTransformDefault  , TEXT("Geometry Collection debug draw, show the selected's rigid body's transform.\nDefault = 0"), ECVF_Cheat);
-	static TAutoConsoleVariable<int32> ShowRigidBodyInertia    (TEXT("p.gc.ShowRigidBodyInertia"    ), GeometryCollectionDebugDrawActorConstants::ShowRigidBodyInertiaDefault    , TEXT("Geometry Collection debug draw, show the selected's rigid body's inertia tensor box.\nDefault = 0"), ECVF_Cheat);
-	static TAutoConsoleVariable<int32> ShowRigidBodyVelocity   (TEXT("p.gc.ShowRigidBodyVelocity"   ), GeometryCollectionDebugDrawActorConstants::ShowRigidBodyVelocityDefault   , TEXT("Geometry Collection debug draw, show the selected's rigid body's linear and angular velocities.\nDefault = 0"), ECVF_Cheat);
-	static TAutoConsoleVariable<int32> ShowRigidBodyForce      (TEXT("p.gc.ShowRigidBodyForce"      ), GeometryCollectionDebugDrawActorConstants::ShowRigidBodyForceDefault      , TEXT("Geometry Collection debug draw, show the selected's rigid body's applied force and torque.\nDefault = 0"), ECVF_Cheat);
-	static TAutoConsoleVariable<int32> ShowRigidBodyInfos      (TEXT("p.gc.ShowRigidBodyInfos"      ), GeometryCollectionDebugDrawActorConstants::ShowRigidBodyInfosDefault      , TEXT("Geometry Collection debug draw, show the selected's rigid body's information.\nDefault = 0"), ECVF_Cheat);
-	static TAutoConsoleVariable<int32> ShowTransformIndex      (TEXT("p.gc.ShowTransformIndex"      ), GeometryCollectionDebugDrawActorConstants::ShowTransformIndexDefault      , TEXT("Geometry Collection debug draw, show the transform index for the selected rigid body's associated cluster nodes.\nDefault = 0"), ECVF_Cheat);
-	static TAutoConsoleVariable<int32> ShowTransform           (TEXT("p.gc.ShowTransform"           ), GeometryCollectionDebugDrawActorConstants::ShowTransformDefault           , TEXT("Geometry Collection debug draw, show the transform for the selected rigid body's associated cluster nodes.\nDefault = 0"), ECVF_Cheat);
-	static TAutoConsoleVariable<int32> ShowParent              (TEXT("p.gc.ShowParent"              ), GeometryCollectionDebugDrawActorConstants::ShowParentDefault              , TEXT("Geometry Collection debug draw, show a link from the selected rigid body's associated cluster nodes to their parent's nodes.\nDefault = 0"), ECVF_Cheat);
-	static TAutoConsoleVariable<int32> ShowLevel               (TEXT("p.gc.ShowLevel"               ), GeometryCollectionDebugDrawActorConstants::ShowLevelDefault               , TEXT("Geometry Collection debug draw, show the hierarchical level for the selected rigid body's associated cluster nodes.\nDefault = 0"), ECVF_Cheat);
-	static TAutoConsoleVariable<int32> ShowConnectivityEdges   (TEXT("p.gc.ShowConnectivityEdges"   ), GeometryCollectionDebugDrawActorConstants::ShowConnectivityEdgesDefault   , TEXT("Geometry Collection debug draw, show the connectivity edges for the rigid body's associated cluster nodes.\nDefault = 0"), ECVF_Cheat);
-	static TAutoConsoleVariable<int32> ShowGeometryIndex       (TEXT("p.gc.ShowGeometryIndex"       ), GeometryCollectionDebugDrawActorConstants::ShowGeometryIndexDefault       , TEXT("Geometry Collection debug draw, show the geometry index for the selected rigid body's associated geometries.\nDefault = 0"), ECVF_Cheat);
-	static TAutoConsoleVariable<int32> ShowGeometryTransform   (TEXT("p.gc.ShowGeometryTransform"   ), GeometryCollectionDebugDrawActorConstants::ShowGeometryTransformDefault   , TEXT("Geometry Collection debug draw, show the geometry transform for the selected rigid body's associated geometries.\nDefault = 0"), ECVF_Cheat);
-	static TAutoConsoleVariable<int32> ShowBoundingBox         (TEXT("p.gc.ShowBoundingBox"         ), GeometryCollectionDebugDrawActorConstants::ShowBoundingBoxDefault         , TEXT("Geometry Collection debug draw, show the bounding box for the selected rigid body's associated geometries.\nDefault = 0"), ECVF_Cheat);
-	static TAutoConsoleVariable<int32> ShowFaces               (TEXT("p.gc.ShowFaces"               ), GeometryCollectionDebugDrawActorConstants::ShowFacesDefault               , TEXT("Geometry Collection debug draw, show the faces for the selected rigid body's associated geometries.\nDefault = 0"), ECVF_Cheat);
-	static TAutoConsoleVariable<int32> ShowFaceIndices         (TEXT("p.gc.ShowFaceIndices"         ), GeometryCollectionDebugDrawActorConstants::ShowFaceIndicesDefault         , TEXT("Geometry Collection debug draw, show the face indices for the selected rigid body's associated geometries.\nDefault = 0"), ECVF_Cheat);
-	static TAutoConsoleVariable<int32> ShowFaceNormals         (TEXT("p.gc.ShowFaceNormals"         ), GeometryCollectionDebugDrawActorConstants::ShowFaceNormalsDefault         , TEXT("Geometry Collection debug draw, show the face normals for the selected rigid body's associated geometries.\nDefault = 0"), ECVF_Cheat);
-	static TAutoConsoleVariable<int32> ShowSingleFace          (TEXT("p.gc.ShowSingleFace"          ), GeometryCollectionDebugDrawActorConstants::ShowSingleFaceDefault          , TEXT("Geometry Collection debug draw, enable single face visualization for the selected rigid body's associated geometries.\nDefault = 0"), ECVF_Cheat);
-	static TAutoConsoleVariable<int32> SingleFaceIndex         (TEXT("p.gc.SingleFaceIndex"         ), GeometryCollectionDebugDrawActorConstants::SingleFaceIndexDefault         , TEXT("Geometry Collection debug draw, the index of the single face to visualize.\nDefault = 0"), ECVF_Cheat);
-	static TAutoConsoleVariable<int32> ShowVertices            (TEXT("p.gc.ShowVertices"            ), GeometryCollectionDebugDrawActorConstants::ShowVerticesDefault            , TEXT("Geometry Collection debug draw, show the vertices for the selected rigid body's associated geometries.\nDefault = 0"), ECVF_Cheat);
-	static TAutoConsoleVariable<int32> ShowVertexIndices       (TEXT("p.gc.ShowVertexIndices"       ), GeometryCollectionDebugDrawActorConstants::ShowVertexIndicesDefault       , TEXT("Geometry Collection debug draw, show the vertex index for the selected rigid body's associated geometries.\nDefault = 0"), ECVF_Cheat);
-	static TAutoConsoleVariable<int32> ShowVertexNormals       (TEXT("p.gc.ShowVertexNormals"       ), GeometryCollectionDebugDrawActorConstants::ShowVertexNormalsDefault       , TEXT("Geometry Collection debug draw, show the vertex normals for the selected rigid body's associated geometries.\nDefault = 0"), ECVF_Cheat);
-	static TAutoConsoleVariable<int32> UseActiveVisualization  (TEXT("p.gc.UseActiveVisualization"  ), GeometryCollectionDebugDrawActorConstants::UseActiveVisualizationDefault  , TEXT("Geometry Collection debug draw, adapt visualization depending of the cluster nodes' hierarchical level..\nDefault = 1."), ECVF_Cheat);
-	static TAutoConsoleVariable<float> PointThickness          (TEXT("p.gc.PointThickness"          ), GeometryCollectionDebugDrawActorConstants::PointThicknessDefault          , TEXT("Geometry Collection debug draw, point thickness.\nDefault = 6."), ECVF_Cheat);
-	static TAutoConsoleVariable<float> LineThickness           (TEXT("p.gc.LineThickness"           ), GeometryCollectionDebugDrawActorConstants::LineThicknessDefault           , TEXT("Geometry Collection debug draw, line thickness.\nDefault = 1."), ECVF_Cheat);
-	static TAutoConsoleVariable<int32> TextShadow              (TEXT("p.gc.TextShadow"              ), GeometryCollectionDebugDrawActorConstants::TextShadowDefault              , TEXT("Geometry Collection debug draw, text shadow under indices for better readability.\nDefault = 1."), ECVF_Cheat);
-	static TAutoConsoleVariable<float> TextScale               (TEXT("p.gc.TextScale"               ), GeometryCollectionDebugDrawActorConstants::TextScaleDefault               , TEXT("Geometry Collection debug draw, text scale.\nDefault = 1."), ECVF_Cheat);
-	static TAutoConsoleVariable<float> NormalScale             (TEXT("p.gc.NormalScale"             ), GeometryCollectionDebugDrawActorConstants::NormalScaleDefault             , TEXT("Geometry Collection debug draw, normal size.\nDefault = 10."), ECVF_Cheat);
-	static TAutoConsoleVariable<float> TransformScale          (TEXT("p.gc.TransformScale"          ), GeometryCollectionDebugDrawActorConstants::TransformScaleDefault          , TEXT("Geometry Collection debug draw, transform size.\nDefault = 20."), ECVF_Cheat);
-	static TAutoConsoleVariable<float> ArrowScale              (TEXT("p.gc.ArrowScale"              ), GeometryCollectionDebugDrawActorConstants::ArrowScaleDefault              , TEXT("Geometry Collection debug draw, arrow size for normals.\nDefault = 2.5."), ECVF_Cheat);
+	static TAutoConsoleVariable<FString> SelectedRigidBodySolver (TEXT("p.gc.SelectedRigidBodySolver" ), GeometryCollectionDebugDrawActorConstants::SelectedRigidBodySolverDefault , TEXT("Geometry Collection debug draw, visualize debug informations for the selected rigid body solver.\nDefault = None"), ECVF_Cheat);
+	static TAutoConsoleVariable<int32  > SelectedRigidBodyId     (TEXT("p.gc.SelectedRigidBodyId"     ), GeometryCollectionDebugDrawActorConstants::SelectedRigidBodyIdDefault     , TEXT("Geometry Collection debug draw, visualize debug informations for the selected rigid body ids.\nDefault = -1"), ECVF_Cheat);
+	static TAutoConsoleVariable<int32  > DebugDrawWholeCollection(TEXT("p.gc.DebugDrawWholeCollection"), GeometryCollectionDebugDrawActorConstants::DebugDrawWholeCollectionDefault, TEXT("Geometry Collection debug draw, show debug visualization for the rest of the geometry collection related to the current rigid body id selection.\nDefault = 0"), ECVF_Cheat);
+	static TAutoConsoleVariable<int32  > DebugDrawHierarchy      (TEXT("p.gc.DebugDrawHierarchy"      ), GeometryCollectionDebugDrawActorConstants::DebugDrawHierarchyDefault      , TEXT("Geometry Collection debug draw, show debug visualization for the top level node rather than the bottom leaf nodes of a cluster's hierarchy..\nDefault = 0"), ECVF_Cheat);
+	static TAutoConsoleVariable<int32  > DebugDrawClustering     (TEXT("p.gc.DebugDrawClustering"     ), GeometryCollectionDebugDrawActorConstants::DebugDrawClusteringDefault     , TEXT("Geometry Collection debug draw, show debug visualization for all clustered children associated to the current rigid body id selection.\nDefault = 0"), ECVF_Cheat);
+	static TAutoConsoleVariable<int32  > HideGeometry            (TEXT("p.gc.HideGeometry"            ), GeometryCollectionDebugDrawActorConstants::HideGeometryDefault            , TEXT("Geometry Collection debug draw, geometry visibility setting, select the part of the geometry to hide in order to better visualize the debug information.\n0: Do not hide any geometries.\n1: Hide the geometry associated to the rigid bodies selected for collision display.\n2: Hide the geometry associated to the selected rigid bodies.\n3: Hide the entire geometry collection associated to the selected rigid bodies.\n4: Hide all geometry collections.\nDefault = 1"), ECVF_Cheat);
+	static TAutoConsoleVariable<int32  > ShowRigidBodyId         (TEXT("p.gc.ShowRigidBodyId"         ), GeometryCollectionDebugDrawActorConstants::ShowRigidBodyIdDefault         , TEXT("Geometry Collection debug draw, show the rigid body id(s).\nDefault = 0"), ECVF_Cheat);
+	static TAutoConsoleVariable<int32  > ShowRigidBodyCollision  (TEXT("p.gc.ShowRigidBodyCollision"  ), GeometryCollectionDebugDrawActorConstants::ShowRigidBodyCollisionDefault  , TEXT("Geometry Collection debug draw, show the selected's rigid body's collision volume.\nDefault = 0"), ECVF_Cheat);
+	static TAutoConsoleVariable<int32  > CollisionAtOrigin       (TEXT("p.gc.CollisionAtOrigin"       ), GeometryCollectionDebugDrawActorConstants::CollisionAtOriginDefault       , TEXT("Geometry Collection debug draw, show any collision volume at the origin, in local space.\nDefault = 0"), ECVF_Cheat);
+	static TAutoConsoleVariable<int32  > ShowRigidBodyTransform  (TEXT("p.gc.ShowRigidBodyTransform"  ), GeometryCollectionDebugDrawActorConstants::ShowRigidBodyTransformDefault  , TEXT("Geometry Collection debug draw, show the selected's rigid body's transform.\nDefault = 0"), ECVF_Cheat);
+	static TAutoConsoleVariable<int32  > ShowRigidBodyInertia    (TEXT("p.gc.ShowRigidBodyInertia"    ), GeometryCollectionDebugDrawActorConstants::ShowRigidBodyInertiaDefault    , TEXT("Geometry Collection debug draw, show the selected's rigid body's inertia tensor box.\nDefault = 0"), ECVF_Cheat);
+	static TAutoConsoleVariable<int32  > ShowRigidBodyVelocity   (TEXT("p.gc.ShowRigidBodyVelocity"   ), GeometryCollectionDebugDrawActorConstants::ShowRigidBodyVelocityDefault   , TEXT("Geometry Collection debug draw, show the selected's rigid body's linear and angular velocities.\nDefault = 0"), ECVF_Cheat);
+	static TAutoConsoleVariable<int32  > ShowRigidBodyForce      (TEXT("p.gc.ShowRigidBodyForce"      ), GeometryCollectionDebugDrawActorConstants::ShowRigidBodyForceDefault      , TEXT("Geometry Collection debug draw, show the selected's rigid body's applied force and torque.\nDefault = 0"), ECVF_Cheat);
+	static TAutoConsoleVariable<int32  > ShowRigidBodyInfos      (TEXT("p.gc.ShowRigidBodyInfos"      ), GeometryCollectionDebugDrawActorConstants::ShowRigidBodyInfosDefault      , TEXT("Geometry Collection debug draw, show the selected's rigid body's information.\nDefault = 0"), ECVF_Cheat);
+	static TAutoConsoleVariable<int32  > ShowTransformIndex      (TEXT("p.gc.ShowTransformIndex"      ), GeometryCollectionDebugDrawActorConstants::ShowTransformIndexDefault      , TEXT("Geometry Collection debug draw, show the transform index for the selected rigid body's associated cluster nodes.\nDefault = 0"), ECVF_Cheat);
+	static TAutoConsoleVariable<int32  > ShowTransform           (TEXT("p.gc.ShowTransform"           ), GeometryCollectionDebugDrawActorConstants::ShowTransformDefault           , TEXT("Geometry Collection debug draw, show the transform for the selected rigid body's associated cluster nodes.\nDefault = 0"), ECVF_Cheat);
+	static TAutoConsoleVariable<int32  > ShowParent              (TEXT("p.gc.ShowParent"              ), GeometryCollectionDebugDrawActorConstants::ShowParentDefault              , TEXT("Geometry Collection debug draw, show a link from the selected rigid body's associated cluster nodes to their parent's nodes.\nDefault = 0"), ECVF_Cheat);
+	static TAutoConsoleVariable<int32  > ShowLevel               (TEXT("p.gc.ShowLevel"               ), GeometryCollectionDebugDrawActorConstants::ShowLevelDefault               , TEXT("Geometry Collection debug draw, show the hierarchical level for the selected rigid body's associated cluster nodes.\nDefault = 0"), ECVF_Cheat);
+	static TAutoConsoleVariable<int32  > ShowConnectivityEdges   (TEXT("p.gc.ShowConnectivityEdges"   ), GeometryCollectionDebugDrawActorConstants::ShowConnectivityEdgesDefault   , TEXT("Geometry Collection debug draw, show the connectivity edges for the rigid body's associated cluster nodes.\nDefault = 0"), ECVF_Cheat);
+	static TAutoConsoleVariable<int32  > ShowGeometryIndex       (TEXT("p.gc.ShowGeometryIndex"       ), GeometryCollectionDebugDrawActorConstants::ShowGeometryIndexDefault       , TEXT("Geometry Collection debug draw, show the geometry index for the selected rigid body's associated geometries.\nDefault = 0"), ECVF_Cheat);
+	static TAutoConsoleVariable<int32  > ShowGeometryTransform   (TEXT("p.gc.ShowGeometryTransform"   ), GeometryCollectionDebugDrawActorConstants::ShowGeometryTransformDefault   , TEXT("Geometry Collection debug draw, show the geometry transform for the selected rigid body's associated geometries.\nDefault = 0"), ECVF_Cheat);
+	static TAutoConsoleVariable<int32  > ShowBoundingBox         (TEXT("p.gc.ShowBoundingBox"         ), GeometryCollectionDebugDrawActorConstants::ShowBoundingBoxDefault         , TEXT("Geometry Collection debug draw, show the bounding box for the selected rigid body's associated geometries.\nDefault = 0"), ECVF_Cheat);
+	static TAutoConsoleVariable<int32  > ShowFaces               (TEXT("p.gc.ShowFaces"               ), GeometryCollectionDebugDrawActorConstants::ShowFacesDefault               , TEXT("Geometry Collection debug draw, show the faces for the selected rigid body's associated geometries.\nDefault = 0"), ECVF_Cheat);
+	static TAutoConsoleVariable<int32  > ShowFaceIndices         (TEXT("p.gc.ShowFaceIndices"         ), GeometryCollectionDebugDrawActorConstants::ShowFaceIndicesDefault         , TEXT("Geometry Collection debug draw, show the face indices for the selected rigid body's associated geometries.\nDefault = 0"), ECVF_Cheat);
+	static TAutoConsoleVariable<int32  > ShowFaceNormals         (TEXT("p.gc.ShowFaceNormals"         ), GeometryCollectionDebugDrawActorConstants::ShowFaceNormalsDefault         , TEXT("Geometry Collection debug draw, show the face normals for the selected rigid body's associated geometries.\nDefault = 0"), ECVF_Cheat);
+	static TAutoConsoleVariable<int32  > ShowSingleFace          (TEXT("p.gc.ShowSingleFace"          ), GeometryCollectionDebugDrawActorConstants::ShowSingleFaceDefault          , TEXT("Geometry Collection debug draw, enable single face visualization for the selected rigid body's associated geometries.\nDefault = 0"), ECVF_Cheat);
+	static TAutoConsoleVariable<int32  > SingleFaceIndex         (TEXT("p.gc.SingleFaceIndex"         ), GeometryCollectionDebugDrawActorConstants::SingleFaceIndexDefault         , TEXT("Geometry Collection debug draw, the index of the single face to visualize.\nDefault = 0"), ECVF_Cheat);
+	static TAutoConsoleVariable<int32  > ShowVertices            (TEXT("p.gc.ShowVertices"            ), GeometryCollectionDebugDrawActorConstants::ShowVerticesDefault            , TEXT("Geometry Collection debug draw, show the vertices for the selected rigid body's associated geometries.\nDefault = 0"), ECVF_Cheat);
+	static TAutoConsoleVariable<int32  > ShowVertexIndices       (TEXT("p.gc.ShowVertexIndices"       ), GeometryCollectionDebugDrawActorConstants::ShowVertexIndicesDefault       , TEXT("Geometry Collection debug draw, show the vertex index for the selected rigid body's associated geometries.\nDefault = 0"), ECVF_Cheat);
+	static TAutoConsoleVariable<int32  > ShowVertexNormals       (TEXT("p.gc.ShowVertexNormals"       ), GeometryCollectionDebugDrawActorConstants::ShowVertexNormalsDefault       , TEXT("Geometry Collection debug draw, show the vertex normals for the selected rigid body's associated geometries.\nDefault = 0"), ECVF_Cheat);
+	static TAutoConsoleVariable<int32  > UseActiveVisualization  (TEXT("p.gc.UseActiveVisualization"  ), GeometryCollectionDebugDrawActorConstants::UseActiveVisualizationDefault  , TEXT("Geometry Collection debug draw, adapt visualization depending of the cluster nodes' hierarchical level..\nDefault = 1."), ECVF_Cheat);
+	static TAutoConsoleVariable<float  > PointThickness          (TEXT("p.gc.PointThickness"          ), GeometryCollectionDebugDrawActorConstants::PointThicknessDefault          , TEXT("Geometry Collection debug draw, point thickness.\nDefault = 6."), ECVF_Cheat);
+	static TAutoConsoleVariable<float  > LineThickness           (TEXT("p.gc.LineThickness"           ), GeometryCollectionDebugDrawActorConstants::LineThicknessDefault           , TEXT("Geometry Collection debug draw, line thickness.\nDefault = 1."), ECVF_Cheat);
+	static TAutoConsoleVariable<int32  > TextShadow              (TEXT("p.gc.TextShadow"              ), GeometryCollectionDebugDrawActorConstants::TextShadowDefault              , TEXT("Geometry Collection debug draw, text shadow under indices for better readability.\nDefault = 1."), ECVF_Cheat);
+	static TAutoConsoleVariable<float  > TextScale               (TEXT("p.gc.TextScale"               ), GeometryCollectionDebugDrawActorConstants::TextScaleDefault               , TEXT("Geometry Collection debug draw, text scale.\nDefault = 1."), ECVF_Cheat);
+	static TAutoConsoleVariable<float  > NormalScale             (TEXT("p.gc.NormalScale"             ), GeometryCollectionDebugDrawActorConstants::NormalScaleDefault             , TEXT("Geometry Collection debug draw, normal size.\nDefault = 10."), ECVF_Cheat);
+	static TAutoConsoleVariable<float  > TransformScale          (TEXT("p.gc.TransformScale"          ), GeometryCollectionDebugDrawActorConstants::TransformScaleDefault          , TEXT("Geometry Collection debug draw, transform size.\nDefault = 20."), ECVF_Cheat);
+	static TAutoConsoleVariable<float  > ArrowScale              (TEXT("p.gc.ArrowScale"              ), GeometryCollectionDebugDrawActorConstants::ArrowScaleDefault              , TEXT("Geometry Collection debug draw, arrow size for normals.\nDefault = 2.5."), ECVF_Cheat);
+}
+
+FString FGeometryCollectionDebugDrawActorSelectedRigidBody::GetSolverName() const
+{
+	return !Solver ? FName(NAME_None).ToString() : Solver->GetName();
 }
 
 AGeometryCollectionDebugDrawActor* AGeometryCollectionDebugDrawActor::FindOrCreate(UWorld* World)
@@ -292,41 +300,42 @@ void AGeometryCollectionDebugDrawActor::PostLoad()
 {
 	static const EConsoleVariableFlags SetBy = ECVF_SetByConsole;  // Can't use the default ECVF_SetByCode as otherwise it won't update the global console variable.
 
-	GeometryCollectionDebugDrawActorCVars::SelectedRigidBodyId     ->Set(      SelectedRigidBody.Id      , SetBy);
-	GeometryCollectionDebugDrawActorCVars::DebugDrawWholeCollection->Set(int32(bDebugDrawWholeCollection), SetBy);
-	GeometryCollectionDebugDrawActorCVars::DebugDrawHierarchy      ->Set(int32(bDebugDrawHierarchy      ), SetBy);
-	GeometryCollectionDebugDrawActorCVars::DebugDrawClustering     ->Set(int32(bDebugDrawClustering     ), SetBy);
-	GeometryCollectionDebugDrawActorCVars::HideGeometry            ->Set(int32(HideGeometry             ), SetBy);
-	GeometryCollectionDebugDrawActorCVars::ShowRigidBodyId         ->Set(int32(bShowRigidBodyId         ), SetBy);
-	GeometryCollectionDebugDrawActorCVars::ShowRigidBodyCollision  ->Set(int32(bShowRigidBodyCollision  ), SetBy);
-	GeometryCollectionDebugDrawActorCVars::CollisionAtOrigin       ->Set(int32(bCollisionAtOrigin       ), SetBy);
-	GeometryCollectionDebugDrawActorCVars::ShowRigidBodyTransform  ->Set(int32(bShowRigidBodyTransform  ), SetBy);
-	GeometryCollectionDebugDrawActorCVars::ShowRigidBodyInertia    ->Set(int32(bShowRigidBodyInertia    ), SetBy);
-	GeometryCollectionDebugDrawActorCVars::ShowRigidBodyVelocity   ->Set(int32(bShowRigidBodyVelocity   ), SetBy);
-	GeometryCollectionDebugDrawActorCVars::ShowRigidBodyForce      ->Set(int32(bShowRigidBodyForce      ), SetBy);
-	GeometryCollectionDebugDrawActorCVars::ShowRigidBodyInfos      ->Set(int32(bShowRigidBodyInfos      ), SetBy);
-	GeometryCollectionDebugDrawActorCVars::ShowTransformIndex      ->Set(int32(bShowTransformIndex      ), SetBy);
-	GeometryCollectionDebugDrawActorCVars::ShowTransform           ->Set(int32(bShowTransform           ), SetBy);
-	GeometryCollectionDebugDrawActorCVars::ShowParent              ->Set(int32(bShowParent              ), SetBy);
-	GeometryCollectionDebugDrawActorCVars::ShowLevel               ->Set(int32(bShowLevel               ), SetBy);
-	GeometryCollectionDebugDrawActorCVars::ShowConnectivityEdges   ->Set(int32(bShowConnectivityEdges   ), SetBy);
-	GeometryCollectionDebugDrawActorCVars::ShowGeometryIndex       ->Set(int32(bShowGeometryIndex       ), SetBy);
-	GeometryCollectionDebugDrawActorCVars::ShowGeometryTransform   ->Set(int32(bShowGeometryTransform   ), SetBy);
-	GeometryCollectionDebugDrawActorCVars::ShowBoundingBox         ->Set(int32(bShowBoundingBox         ), SetBy);
-	GeometryCollectionDebugDrawActorCVars::ShowFaces               ->Set(int32(bShowFaces               ), SetBy);
-	GeometryCollectionDebugDrawActorCVars::ShowFaceIndices         ->Set(int32(bShowFaceIndices         ), SetBy);
-	GeometryCollectionDebugDrawActorCVars::ShowFaceNormals         ->Set(int32(bShowFaceNormals         ), SetBy);
-	GeometryCollectionDebugDrawActorCVars::ShowVertices            ->Set(int32(bShowVertices            ), SetBy);
-	GeometryCollectionDebugDrawActorCVars::ShowVertexIndices       ->Set(int32(bShowVertexIndices       ), SetBy);
-	GeometryCollectionDebugDrawActorCVars::ShowVertexNormals       ->Set(int32(bShowVertexNormals       ), SetBy);
-	GeometryCollectionDebugDrawActorCVars::UseActiveVisualization  ->Set(int32(bUseActiveVisualization  ), SetBy);
-	GeometryCollectionDebugDrawActorCVars::PointThickness          ->Set(      PointThickness            , SetBy);
-	GeometryCollectionDebugDrawActorCVars::LineThickness           ->Set(      LineThickness             , SetBy);
-	GeometryCollectionDebugDrawActorCVars::TextShadow              ->Set(int32(bTextShadow              ), SetBy);
-	GeometryCollectionDebugDrawActorCVars::TextScale               ->Set(      TextScale                 , SetBy);
-	GeometryCollectionDebugDrawActorCVars::NormalScale             ->Set(      NormalScale               , SetBy);
-	GeometryCollectionDebugDrawActorCVars::TransformScale          ->Set(      TransformScale            , SetBy);
-	GeometryCollectionDebugDrawActorCVars::ArrowScale              ->Set(      ArrowScale                , SetBy);
+	GeometryCollectionDebugDrawActorCVars::SelectedRigidBodySolver ->Set(*SelectedRigidBody.GetSolverName(), SetBy);
+	GeometryCollectionDebugDrawActorCVars::SelectedRigidBodyId     ->Set( SelectedRigidBody.Id             , SetBy);
+	GeometryCollectionDebugDrawActorCVars::DebugDrawWholeCollection->Set(int32(bDebugDrawWholeCollection  ), SetBy);
+	GeometryCollectionDebugDrawActorCVars::DebugDrawHierarchy      ->Set(int32(bDebugDrawHierarchy        ), SetBy);
+	GeometryCollectionDebugDrawActorCVars::DebugDrawClustering     ->Set(int32(bDebugDrawClustering       ), SetBy);
+	GeometryCollectionDebugDrawActorCVars::HideGeometry            ->Set(int32(HideGeometry               ), SetBy);
+	GeometryCollectionDebugDrawActorCVars::ShowRigidBodyId         ->Set(int32(bShowRigidBodyId           ), SetBy);
+	GeometryCollectionDebugDrawActorCVars::ShowRigidBodyCollision  ->Set(int32(bShowRigidBodyCollision    ), SetBy);
+	GeometryCollectionDebugDrawActorCVars::CollisionAtOrigin       ->Set(int32(bCollisionAtOrigin         ), SetBy);
+	GeometryCollectionDebugDrawActorCVars::ShowRigidBodyTransform  ->Set(int32(bShowRigidBodyTransform    ), SetBy);
+	GeometryCollectionDebugDrawActorCVars::ShowRigidBodyInertia    ->Set(int32(bShowRigidBodyInertia      ), SetBy);
+	GeometryCollectionDebugDrawActorCVars::ShowRigidBodyVelocity   ->Set(int32(bShowRigidBodyVelocity     ), SetBy);
+	GeometryCollectionDebugDrawActorCVars::ShowRigidBodyForce      ->Set(int32(bShowRigidBodyForce        ), SetBy);
+	GeometryCollectionDebugDrawActorCVars::ShowRigidBodyInfos      ->Set(int32(bShowRigidBodyInfos        ), SetBy);
+	GeometryCollectionDebugDrawActorCVars::ShowTransformIndex      ->Set(int32(bShowTransformIndex        ), SetBy);
+	GeometryCollectionDebugDrawActorCVars::ShowTransform           ->Set(int32(bShowTransform             ), SetBy);
+	GeometryCollectionDebugDrawActorCVars::ShowParent              ->Set(int32(bShowParent                ), SetBy);
+	GeometryCollectionDebugDrawActorCVars::ShowLevel               ->Set(int32(bShowLevel                 ), SetBy);
+	GeometryCollectionDebugDrawActorCVars::ShowConnectivityEdges   ->Set(int32(bShowConnectivityEdges     ), SetBy);
+	GeometryCollectionDebugDrawActorCVars::ShowGeometryIndex       ->Set(int32(bShowGeometryIndex         ), SetBy);
+	GeometryCollectionDebugDrawActorCVars::ShowGeometryTransform   ->Set(int32(bShowGeometryTransform     ), SetBy);
+	GeometryCollectionDebugDrawActorCVars::ShowBoundingBox         ->Set(int32(bShowBoundingBox           ), SetBy);
+	GeometryCollectionDebugDrawActorCVars::ShowFaces               ->Set(int32(bShowFaces                 ), SetBy);
+	GeometryCollectionDebugDrawActorCVars::ShowFaceIndices         ->Set(int32(bShowFaceIndices           ), SetBy);
+	GeometryCollectionDebugDrawActorCVars::ShowFaceNormals         ->Set(int32(bShowFaceNormals           ), SetBy);
+	GeometryCollectionDebugDrawActorCVars::ShowVertices            ->Set(int32(bShowVertices              ), SetBy);
+	GeometryCollectionDebugDrawActorCVars::ShowVertexIndices       ->Set(int32(bShowVertexIndices         ), SetBy);
+	GeometryCollectionDebugDrawActorCVars::ShowVertexNormals       ->Set(int32(bShowVertexNormals         ), SetBy);
+	GeometryCollectionDebugDrawActorCVars::UseActiveVisualization  ->Set(int32(bUseActiveVisualization    ), SetBy);
+	GeometryCollectionDebugDrawActorCVars::PointThickness          ->Set(      PointThickness              , SetBy);
+	GeometryCollectionDebugDrawActorCVars::LineThickness           ->Set(      LineThickness               , SetBy);
+	GeometryCollectionDebugDrawActorCVars::TextShadow              ->Set(int32(bTextShadow                ), SetBy);
+	GeometryCollectionDebugDrawActorCVars::TextScale               ->Set(      TextScale                   , SetBy);
+	GeometryCollectionDebugDrawActorCVars::NormalScale             ->Set(      NormalScale                 , SetBy);
+	GeometryCollectionDebugDrawActorCVars::TransformScale          ->Set(      TransformScale              , SetBy);
+	GeometryCollectionDebugDrawActorCVars::ArrowScale              ->Set(      ArrowScale                  , SetBy);
 
 	Super::PostLoad();
 }
@@ -353,6 +362,10 @@ int32 AGeometryCollectionDebugDrawActor::GetLevel(int32 TransformIndex, const TM
 void AGeometryCollectionDebugDrawActor::OnPropertyChanged(bool bForceVisibilityUpdate)
 {
 #if GEOMETRYCOLLECTION_DEBUG_DRAW
+	// Reset selected rigid body's actor
+	SelectedRigidBody.GeometryCollection = nullptr;
+
+	// Update component states
 	UWorld* const World = GetWorld();
 	if (!World || !World->HasBegunPlay() || !HasActorBegunPlay()) { return; }
 
@@ -365,7 +378,7 @@ void AGeometryCollectionDebugDrawActor::OnPropertyChanged(bool bForceVisibilityU
 			const bool bIsSelected = GeometryCollectionDebugDrawComponent->OnDebugDrawPropertiesChanged(bForceVisibilityUpdate);
 			if (bIsSelected)
 			{
-				SelectedRigidBody.GeometryCollectionActor = *ActorIterator;
+				SelectedRigidBody.GeometryCollection = *ActorIterator;
 			}
 		}
 	}
@@ -383,15 +396,41 @@ void AGeometryCollectionDebugDrawActor::UpdatePropertyValue(T1& PropertyValue, c
 	}
 }
 
+template<>
+void AGeometryCollectionDebugDrawActor::UpdatePropertyValue<AChaosSolverActor*, FString>(AChaosSolverActor*& PropertyValue, const TAutoConsoleVariable<FString>& ConsoleVariable, bool& bHasChanged)
+{
+	AChaosSolverActor* NewValue = nullptr;
+
+	if (UWorld* const World = GetWorld())
+	{
+		const FString SolverName = ConsoleVariable.GetValueOnGameThread();
+		for (TActorIterator<AChaosSolverActor> ActorIterator(World); ActorIterator; ++ActorIterator)
+		{
+			if (ActorIterator->GetName() == SolverName)
+			{
+				NewValue = *ActorIterator;
+				break;
+			}
+		}
+	}
+
+	if (PropertyValue != NewValue)
+	{
+		bHasChanged = true;
+		PropertyValue = NewValue;
+	}
+}
+
 void AGeometryCollectionDebugDrawActor::OnCVarsChanged()
 {
 	// Discard callback if this actor isn't in any world
-	UWorld* const World = GetWorld();
-	if (!World) { return; }
+	if (!GetWorld()) { return; }
 
 	// Update properties from cvars
 	bool bHavePropertiesChanged = false;
 	bool bHasDebugDrawClusteringChanged = false;
+
+	UpdatePropertyValue(SelectedRigidBody.Solver , GeometryCollectionDebugDrawActorCVars::SelectedRigidBodySolver , bHavePropertiesChanged);
 	UpdatePropertyValue(SelectedRigidBody.Id     , GeometryCollectionDebugDrawActorCVars::SelectedRigidBodyId     , bHavePropertiesChanged);
 	UpdatePropertyValue(bDebugDrawWholeCollection, GeometryCollectionDebugDrawActorCVars::DebugDrawWholeCollection, bHavePropertiesChanged);
 	UpdatePropertyValue(bDebugDrawHierarchy      , GeometryCollectionDebugDrawActorCVars::DebugDrawHierarchy      , bHavePropertiesChanged);
@@ -442,47 +481,48 @@ void AGeometryCollectionDebugDrawActor::PostEditChangeProperty(FPropertyChangedE
 {
 	// Synchronize the command variables to this Actor's properties if the property name matches.
 	static const EConsoleVariableFlags SetBy = ECVF_SetByConsole;  // Can't use the default ECVF_SetByCode as otherwise changing the UI won't update the global console variable.
-	const FName MemberPropertyName = PropertyChangedEvent.MemberProperty ? PropertyChangedEvent.MemberProperty->GetFName(): NAME_None;
+	const FName PropertyName = PropertyChangedEvent.Property ? PropertyChangedEvent.Property->GetFName() : NAME_None;
 
 	bool bForceVisibilityUpdate = false;
 
-	if      (MemberPropertyName == GET_MEMBER_NAME_CHECKED(AGeometryCollectionDebugDrawActor, SelectedRigidBody        )) { GeometryCollectionDebugDrawActorCVars::SelectedRigidBodyId     ->Set(      SelectedRigidBody.Id      , SetBy); }
-	else if (MemberPropertyName == GET_MEMBER_NAME_CHECKED(AGeometryCollectionDebugDrawActor, bDebugDrawWholeCollection)) { GeometryCollectionDebugDrawActorCVars::DebugDrawWholeCollection->Set(int32(bDebugDrawWholeCollection), SetBy); }
-	else if (MemberPropertyName == GET_MEMBER_NAME_CHECKED(AGeometryCollectionDebugDrawActor, bDebugDrawHierarchy      )) { GeometryCollectionDebugDrawActorCVars::DebugDrawHierarchy      ->Set(int32(bDebugDrawHierarchy      ), SetBy); }
-	else if (MemberPropertyName == GET_MEMBER_NAME_CHECKED(AGeometryCollectionDebugDrawActor, bDebugDrawClustering     )) { GeometryCollectionDebugDrawActorCVars::DebugDrawClustering     ->Set(int32(bDebugDrawClustering     ), SetBy); bForceVisibilityUpdate = true; }
-	else if (MemberPropertyName == GET_MEMBER_NAME_CHECKED(AGeometryCollectionDebugDrawActor, HideGeometry             )) { GeometryCollectionDebugDrawActorCVars::HideGeometry            ->Set(int32(HideGeometry             ), SetBy); }
-	else if (MemberPropertyName == GET_MEMBER_NAME_CHECKED(AGeometryCollectionDebugDrawActor, bShowRigidBodyId         )) { GeometryCollectionDebugDrawActorCVars::ShowRigidBodyId         ->Set(int32(bShowRigidBodyId         ), SetBy); }
-	else if (MemberPropertyName == GET_MEMBER_NAME_CHECKED(AGeometryCollectionDebugDrawActor, bShowRigidBodyCollision  )) { GeometryCollectionDebugDrawActorCVars::ShowRigidBodyCollision  ->Set(int32(bShowRigidBodyCollision  ), SetBy); }
-	else if (MemberPropertyName == GET_MEMBER_NAME_CHECKED(AGeometryCollectionDebugDrawActor, bCollisionAtOrigin       )) { GeometryCollectionDebugDrawActorCVars::CollisionAtOrigin       ->Set(int32(bCollisionAtOrigin       ), SetBy); }
-	else if (MemberPropertyName == GET_MEMBER_NAME_CHECKED(AGeometryCollectionDebugDrawActor, bShowRigidBodyTransform  )) { GeometryCollectionDebugDrawActorCVars::ShowRigidBodyTransform  ->Set(int32(bShowRigidBodyTransform  ), SetBy); }
-	else if (MemberPropertyName == GET_MEMBER_NAME_CHECKED(AGeometryCollectionDebugDrawActor, bShowRigidBodyInertia    )) { GeometryCollectionDebugDrawActorCVars::ShowRigidBodyInertia    ->Set(int32(bShowRigidBodyInertia    ), SetBy); }
-	else if (MemberPropertyName == GET_MEMBER_NAME_CHECKED(AGeometryCollectionDebugDrawActor, bShowRigidBodyVelocity   )) { GeometryCollectionDebugDrawActorCVars::ShowRigidBodyVelocity   ->Set(int32(bShowRigidBodyVelocity   ), SetBy); }
-	else if (MemberPropertyName == GET_MEMBER_NAME_CHECKED(AGeometryCollectionDebugDrawActor, bShowRigidBodyForce      )) { GeometryCollectionDebugDrawActorCVars::ShowRigidBodyForce      ->Set(int32(bShowRigidBodyForce      ), SetBy); }
-	else if (MemberPropertyName == GET_MEMBER_NAME_CHECKED(AGeometryCollectionDebugDrawActor, bShowRigidBodyInfos      )) { GeometryCollectionDebugDrawActorCVars::ShowRigidBodyInfos      ->Set(int32(bShowRigidBodyInfos      ), SetBy); }
-	else if (MemberPropertyName == GET_MEMBER_NAME_CHECKED(AGeometryCollectionDebugDrawActor, bShowTransformIndex      )) { GeometryCollectionDebugDrawActorCVars::ShowTransformIndex      ->Set(int32(bShowTransformIndex      ), SetBy); }
-	else if (MemberPropertyName == GET_MEMBER_NAME_CHECKED(AGeometryCollectionDebugDrawActor, bShowTransform           )) { GeometryCollectionDebugDrawActorCVars::ShowTransform           ->Set(int32(bShowTransform           ), SetBy); }
-	else if (MemberPropertyName == GET_MEMBER_NAME_CHECKED(AGeometryCollectionDebugDrawActor, bShowParent              )) { GeometryCollectionDebugDrawActorCVars::ShowParent              ->Set(int32(bShowParent              ), SetBy); }
-	else if (MemberPropertyName == GET_MEMBER_NAME_CHECKED(AGeometryCollectionDebugDrawActor, bShowLevel               )) { GeometryCollectionDebugDrawActorCVars::ShowLevel               ->Set(int32(bShowLevel               ), SetBy); }
-	else if (MemberPropertyName == GET_MEMBER_NAME_CHECKED(AGeometryCollectionDebugDrawActor, bShowConnectivityEdges   )) { GeometryCollectionDebugDrawActorCVars::ShowConnectivityEdges   ->Set(int32(bShowConnectivityEdges   ), SetBy); }
-	else if (MemberPropertyName == GET_MEMBER_NAME_CHECKED(AGeometryCollectionDebugDrawActor, bShowGeometryIndex       )) { GeometryCollectionDebugDrawActorCVars::ShowGeometryIndex       ->Set(int32(bShowGeometryIndex       ), SetBy); }
-	else if (MemberPropertyName == GET_MEMBER_NAME_CHECKED(AGeometryCollectionDebugDrawActor, bShowGeometryTransform   )) { GeometryCollectionDebugDrawActorCVars::ShowGeometryTransform   ->Set(int32(bShowGeometryTransform   ), SetBy); }
-	else if (MemberPropertyName == GET_MEMBER_NAME_CHECKED(AGeometryCollectionDebugDrawActor, bShowBoundingBox         )) { GeometryCollectionDebugDrawActorCVars::ShowBoundingBox         ->Set(int32(bShowBoundingBox         ), SetBy); }
-	else if (MemberPropertyName == GET_MEMBER_NAME_CHECKED(AGeometryCollectionDebugDrawActor, bShowFaces               )) { GeometryCollectionDebugDrawActorCVars::ShowFaces               ->Set(int32(bShowFaces               ), SetBy); }
-	else if (MemberPropertyName == GET_MEMBER_NAME_CHECKED(AGeometryCollectionDebugDrawActor, bShowFaceIndices         )) { GeometryCollectionDebugDrawActorCVars::ShowFaceIndices         ->Set(int32(bShowFaceIndices         ), SetBy); }
-	else if (MemberPropertyName == GET_MEMBER_NAME_CHECKED(AGeometryCollectionDebugDrawActor, bShowFaceNormals         )) { GeometryCollectionDebugDrawActorCVars::ShowFaceNormals         ->Set(int32(bShowFaceNormals         ), SetBy); }
-	else if (MemberPropertyName == GET_MEMBER_NAME_CHECKED(AGeometryCollectionDebugDrawActor, bShowSingleFace          )) { GeometryCollectionDebugDrawActorCVars::ShowSingleFace          ->Set(int32(bShowSingleFace          ), SetBy); }
-	else if (MemberPropertyName == GET_MEMBER_NAME_CHECKED(AGeometryCollectionDebugDrawActor, SingleFaceIndex          )) { GeometryCollectionDebugDrawActorCVars::SingleFaceIndex         ->Set(      SingleFaceIndex           , SetBy); }
-	else if (MemberPropertyName == GET_MEMBER_NAME_CHECKED(AGeometryCollectionDebugDrawActor, bShowVertices            )) { GeometryCollectionDebugDrawActorCVars::ShowVertices            ->Set(int32(bShowVertices            ), SetBy); }
-	else if (MemberPropertyName == GET_MEMBER_NAME_CHECKED(AGeometryCollectionDebugDrawActor, bShowVertexIndices       )) { GeometryCollectionDebugDrawActorCVars::ShowVertexIndices       ->Set(int32(bShowVertexIndices       ), SetBy); }
-	else if (MemberPropertyName == GET_MEMBER_NAME_CHECKED(AGeometryCollectionDebugDrawActor, bShowVertexNormals       )) { GeometryCollectionDebugDrawActorCVars::ShowVertexNormals       ->Set(int32(bShowVertexNormals       ), SetBy); }
-	else if (MemberPropertyName == GET_MEMBER_NAME_CHECKED(AGeometryCollectionDebugDrawActor, bUseActiveVisualization  )) { GeometryCollectionDebugDrawActorCVars::UseActiveVisualization  ->Set(int32(bUseActiveVisualization  ), SetBy); }
-	else if (MemberPropertyName == GET_MEMBER_NAME_CHECKED(AGeometryCollectionDebugDrawActor, PointThickness           )) { GeometryCollectionDebugDrawActorCVars::PointThickness          ->Set(      PointThickness            , SetBy); }
-	else if (MemberPropertyName == GET_MEMBER_NAME_CHECKED(AGeometryCollectionDebugDrawActor, LineThickness            )) { GeometryCollectionDebugDrawActorCVars::LineThickness           ->Set(      LineThickness             , SetBy); }
-	else if (MemberPropertyName == GET_MEMBER_NAME_CHECKED(AGeometryCollectionDebugDrawActor, bTextShadow              )) { GeometryCollectionDebugDrawActorCVars::TextShadow              ->Set(int32(bTextShadow              ), SetBy); }
-	else if (MemberPropertyName == GET_MEMBER_NAME_CHECKED(AGeometryCollectionDebugDrawActor, TextScale                )) { GeometryCollectionDebugDrawActorCVars::TextScale               ->Set(      TextScale                 , SetBy); }
-	else if (MemberPropertyName == GET_MEMBER_NAME_CHECKED(AGeometryCollectionDebugDrawActor, NormalScale              )) { GeometryCollectionDebugDrawActorCVars::NormalScale             ->Set(      NormalScale               , SetBy); }
-	else if (MemberPropertyName == GET_MEMBER_NAME_CHECKED(AGeometryCollectionDebugDrawActor, TransformScale           )) { GeometryCollectionDebugDrawActorCVars::TransformScale          ->Set(      TransformScale            , SetBy); }
-	else if (MemberPropertyName == GET_MEMBER_NAME_CHECKED(AGeometryCollectionDebugDrawActor, ArrowScale               )) { GeometryCollectionDebugDrawActorCVars::ArrowScale              ->Set(      ArrowScale                , SetBy); }
+	if      (PropertyName == GET_MEMBER_NAME_CHECKED(FGeometryCollectionDebugDrawActorSelectedRigidBody, Solver  )) { GeometryCollectionDebugDrawActorCVars::SelectedRigidBodySolver ->Set(*SelectedRigidBody.GetSolverName(), SetBy); }
+	else if (PropertyName == GET_MEMBER_NAME_CHECKED(FGeometryCollectionDebugDrawActorSelectedRigidBody, Id      )) { GeometryCollectionDebugDrawActorCVars::SelectedRigidBodyId     ->Set( SelectedRigidBody.Id             , SetBy); }
+	else if (PropertyName == GET_MEMBER_NAME_CHECKED(AGeometryCollectionDebugDrawActor, bDebugDrawWholeCollection)) { GeometryCollectionDebugDrawActorCVars::DebugDrawWholeCollection->Set(int32(bDebugDrawWholeCollection  ), SetBy); }
+	else if (PropertyName == GET_MEMBER_NAME_CHECKED(AGeometryCollectionDebugDrawActor, bDebugDrawHierarchy      )) { GeometryCollectionDebugDrawActorCVars::DebugDrawHierarchy      ->Set(int32(bDebugDrawHierarchy        ), SetBy); }
+	else if (PropertyName == GET_MEMBER_NAME_CHECKED(AGeometryCollectionDebugDrawActor, bDebugDrawClustering     )) { GeometryCollectionDebugDrawActorCVars::DebugDrawClustering     ->Set(int32(bDebugDrawClustering       ), SetBy); bForceVisibilityUpdate = true; }
+	else if (PropertyName == GET_MEMBER_NAME_CHECKED(AGeometryCollectionDebugDrawActor, HideGeometry             )) { GeometryCollectionDebugDrawActorCVars::HideGeometry            ->Set(int32(HideGeometry               ), SetBy); }
+	else if (PropertyName == GET_MEMBER_NAME_CHECKED(AGeometryCollectionDebugDrawActor, bShowRigidBodyId         )) { GeometryCollectionDebugDrawActorCVars::ShowRigidBodyId         ->Set(int32(bShowRigidBodyId           ), SetBy); }
+	else if (PropertyName == GET_MEMBER_NAME_CHECKED(AGeometryCollectionDebugDrawActor, bShowRigidBodyCollision  )) { GeometryCollectionDebugDrawActorCVars::ShowRigidBodyCollision  ->Set(int32(bShowRigidBodyCollision    ), SetBy); }
+	else if (PropertyName == GET_MEMBER_NAME_CHECKED(AGeometryCollectionDebugDrawActor, bCollisionAtOrigin       )) { GeometryCollectionDebugDrawActorCVars::CollisionAtOrigin       ->Set(int32(bCollisionAtOrigin         ), SetBy); }
+	else if (PropertyName == GET_MEMBER_NAME_CHECKED(AGeometryCollectionDebugDrawActor, bShowRigidBodyTransform  )) { GeometryCollectionDebugDrawActorCVars::ShowRigidBodyTransform  ->Set(int32(bShowRigidBodyTransform    ), SetBy); }
+	else if (PropertyName == GET_MEMBER_NAME_CHECKED(AGeometryCollectionDebugDrawActor, bShowRigidBodyInertia    )) { GeometryCollectionDebugDrawActorCVars::ShowRigidBodyInertia    ->Set(int32(bShowRigidBodyInertia      ), SetBy); }
+	else if (PropertyName == GET_MEMBER_NAME_CHECKED(AGeometryCollectionDebugDrawActor, bShowRigidBodyVelocity   )) { GeometryCollectionDebugDrawActorCVars::ShowRigidBodyVelocity   ->Set(int32(bShowRigidBodyVelocity     ), SetBy); }
+	else if (PropertyName == GET_MEMBER_NAME_CHECKED(AGeometryCollectionDebugDrawActor, bShowRigidBodyForce      )) { GeometryCollectionDebugDrawActorCVars::ShowRigidBodyForce      ->Set(int32(bShowRigidBodyForce        ), SetBy); }
+	else if (PropertyName == GET_MEMBER_NAME_CHECKED(AGeometryCollectionDebugDrawActor, bShowRigidBodyInfos      )) { GeometryCollectionDebugDrawActorCVars::ShowRigidBodyInfos      ->Set(int32(bShowRigidBodyInfos        ), SetBy); }
+	else if (PropertyName == GET_MEMBER_NAME_CHECKED(AGeometryCollectionDebugDrawActor, bShowTransformIndex      )) { GeometryCollectionDebugDrawActorCVars::ShowTransformIndex      ->Set(int32(bShowTransformIndex        ), SetBy); }
+	else if (PropertyName == GET_MEMBER_NAME_CHECKED(AGeometryCollectionDebugDrawActor, bShowTransform           )) { GeometryCollectionDebugDrawActorCVars::ShowTransform           ->Set(int32(bShowTransform             ), SetBy); }
+	else if (PropertyName == GET_MEMBER_NAME_CHECKED(AGeometryCollectionDebugDrawActor, bShowParent              )) { GeometryCollectionDebugDrawActorCVars::ShowParent              ->Set(int32(bShowParent                ), SetBy); }
+	else if (PropertyName == GET_MEMBER_NAME_CHECKED(AGeometryCollectionDebugDrawActor, bShowLevel               )) { GeometryCollectionDebugDrawActorCVars::ShowLevel               ->Set(int32(bShowLevel                 ), SetBy); }
+	else if (PropertyName == GET_MEMBER_NAME_CHECKED(AGeometryCollectionDebugDrawActor, bShowConnectivityEdges   )) { GeometryCollectionDebugDrawActorCVars::ShowConnectivityEdges   ->Set(int32(bShowConnectivityEdges     ), SetBy); }
+	else if (PropertyName == GET_MEMBER_NAME_CHECKED(AGeometryCollectionDebugDrawActor, bShowGeometryIndex       )) { GeometryCollectionDebugDrawActorCVars::ShowGeometryIndex       ->Set(int32(bShowGeometryIndex         ), SetBy); }
+	else if (PropertyName == GET_MEMBER_NAME_CHECKED(AGeometryCollectionDebugDrawActor, bShowGeometryTransform   )) { GeometryCollectionDebugDrawActorCVars::ShowGeometryTransform   ->Set(int32(bShowGeometryTransform     ), SetBy); }
+	else if (PropertyName == GET_MEMBER_NAME_CHECKED(AGeometryCollectionDebugDrawActor, bShowBoundingBox         )) { GeometryCollectionDebugDrawActorCVars::ShowBoundingBox         ->Set(int32(bShowBoundingBox           ), SetBy); }
+	else if (PropertyName == GET_MEMBER_NAME_CHECKED(AGeometryCollectionDebugDrawActor, bShowFaces               )) { GeometryCollectionDebugDrawActorCVars::ShowFaces               ->Set(int32(bShowFaces                 ), SetBy); }
+	else if (PropertyName == GET_MEMBER_NAME_CHECKED(AGeometryCollectionDebugDrawActor, bShowFaceIndices         )) { GeometryCollectionDebugDrawActorCVars::ShowFaceIndices         ->Set(int32(bShowFaceIndices           ), SetBy); }
+	else if (PropertyName == GET_MEMBER_NAME_CHECKED(AGeometryCollectionDebugDrawActor, bShowFaceNormals         )) { GeometryCollectionDebugDrawActorCVars::ShowFaceNormals         ->Set(int32(bShowFaceNormals           ), SetBy); }
+	else if (PropertyName == GET_MEMBER_NAME_CHECKED(AGeometryCollectionDebugDrawActor, bShowSingleFace          )) { GeometryCollectionDebugDrawActorCVars::ShowSingleFace          ->Set(int32(bShowSingleFace            ), SetBy); }
+	else if (PropertyName == GET_MEMBER_NAME_CHECKED(AGeometryCollectionDebugDrawActor, SingleFaceIndex          )) { GeometryCollectionDebugDrawActorCVars::SingleFaceIndex         ->Set(      SingleFaceIndex             , SetBy); }
+	else if (PropertyName == GET_MEMBER_NAME_CHECKED(AGeometryCollectionDebugDrawActor, bShowVertices            )) { GeometryCollectionDebugDrawActorCVars::ShowVertices            ->Set(int32(bShowVertices              ), SetBy); }
+	else if (PropertyName == GET_MEMBER_NAME_CHECKED(AGeometryCollectionDebugDrawActor, bShowVertexIndices       )) { GeometryCollectionDebugDrawActorCVars::ShowVertexIndices       ->Set(int32(bShowVertexIndices         ), SetBy); }
+	else if (PropertyName == GET_MEMBER_NAME_CHECKED(AGeometryCollectionDebugDrawActor, bShowVertexNormals       )) { GeometryCollectionDebugDrawActorCVars::ShowVertexNormals       ->Set(int32(bShowVertexNormals         ), SetBy); }
+	else if (PropertyName == GET_MEMBER_NAME_CHECKED(AGeometryCollectionDebugDrawActor, bUseActiveVisualization  )) { GeometryCollectionDebugDrawActorCVars::UseActiveVisualization  ->Set(int32(bUseActiveVisualization    ), SetBy); }
+	else if (PropertyName == GET_MEMBER_NAME_CHECKED(AGeometryCollectionDebugDrawActor, PointThickness           )) { GeometryCollectionDebugDrawActorCVars::PointThickness          ->Set(      PointThickness              , SetBy); }
+	else if (PropertyName == GET_MEMBER_NAME_CHECKED(AGeometryCollectionDebugDrawActor, LineThickness            )) { GeometryCollectionDebugDrawActorCVars::LineThickness           ->Set(      LineThickness               , SetBy); }
+	else if (PropertyName == GET_MEMBER_NAME_CHECKED(AGeometryCollectionDebugDrawActor, bTextShadow              )) { GeometryCollectionDebugDrawActorCVars::TextShadow              ->Set(int32(bTextShadow                ), SetBy); }
+	else if (PropertyName == GET_MEMBER_NAME_CHECKED(AGeometryCollectionDebugDrawActor, TextScale                )) { GeometryCollectionDebugDrawActorCVars::TextScale               ->Set(      TextScale                   , SetBy); }
+	else if (PropertyName == GET_MEMBER_NAME_CHECKED(AGeometryCollectionDebugDrawActor, NormalScale              )) { GeometryCollectionDebugDrawActorCVars::NormalScale             ->Set(      NormalScale                 , SetBy); }
+	else if (PropertyName == GET_MEMBER_NAME_CHECKED(AGeometryCollectionDebugDrawActor, TransformScale           )) { GeometryCollectionDebugDrawActorCVars::TransformScale          ->Set(      TransformScale              , SetBy); }
+	else if (PropertyName == GET_MEMBER_NAME_CHECKED(AGeometryCollectionDebugDrawActor, ArrowScale               )) { GeometryCollectionDebugDrawActorCVars::ArrowScale              ->Set(      ArrowScale                  , SetBy); }
 	else
 	{
 		Super::PostEditChangeProperty(PropertyChangedEvent);
diff --git a/Engine/Source/Runtime/Experimental/GeometryCollectionEngine/Private/GeometryCollection/GeometryCollectionDebugDrawComponent.cpp b/Engine/Source/Runtime/Experimental/GeometryCollectionEngine/Private/GeometryCollection/GeometryCollectionDebugDrawComponent.cpp
--- a/Engine/Source/Runtime/Experimental/GeometryCollectionEngine/Private/GeometryCollection/GeometryCollectionDebugDrawComponent.cpp
+++ b/Engine/Source/Runtime/Experimental/GeometryCollectionEngine/Private/GeometryCollection/GeometryCollectionDebugDrawComponent.cpp
@@ -109,6 +109,7 @@ UGeometryCollectionDebugDrawComponent::UGeometryCollectionDebugDrawComponent(con
 	, SelectedRigidBodyId(INDEX_NONE)
 	, SelectedTransformIndex(INDEX_NONE)
 	, HiddenTransformIndex(INDEX_NONE)
+	, SelectedChaosSolver(nullptr)
 #endif  // #if GEOMETRYCOLLECTION_DEBUG_DRAW
 {
 	bNavigationRelevant = false;
@@ -212,7 +213,7 @@ void UGeometryCollectionDebugDrawComponent::TickComponent(float DeltaTime, enum
 			}
 			if (ParentCheckSum != PrevParentCheckSum)
 			{
-				UE_CLOG(GetOwner(), LogGeometryCollectionDebugDraw, Verbose, TEXT("Geometry Collection has broken up."), *GetOwner()->GetName());
+				UE_CLOG(GetOwner(), LogGeometryCollectionDebugDraw, Verbose, TEXT("Geometry Collection has broken up for actor %s."), *GetOwner()->GetName());
 				UpdateGeometryVisibility(true);
 			}
 		}
@@ -404,7 +405,8 @@ void UGeometryCollectionDebugDrawComponent::DebugDrawTick()
 	const bool bIsSelected = (SelectedTransformIndex != INDEX_NONE);
 	const bool bIsOneSelected = bIsSelected && !GeometryCollectionDebugDrawActor->bDebugDrawWholeCollection;
 	const bool bAreAllSelected = (bIsSelected && GeometryCollectionDebugDrawActor->bDebugDrawWholeCollection) ||
-		(GeometryCollectionDebugDrawActor->SelectedRigidBody.Id == INDEX_NONE);
+		(GeometryCollectionDebugDrawActor->SelectedRigidBody.Id == INDEX_NONE && 
+		 GeometryCollectionDebugDrawActor->SelectedRigidBody.Solver == GeometryCollectionComponent->ChaosSolverActor);
 
 	// Clustering
 	if (!bShowTransformIndices && GeometryCollectionDebugDrawActor->bShowTransformIndex && bIsOneSelected)
@@ -545,16 +547,22 @@ void UGeometryCollectionDebugDrawComponent::UpdateSelectedTransformIndex()
 		return;
 	}
 
-	// Check whether the selected rigid body id has changed
-	if (SelectedRigidBodyId == GeometryCollectionDebugDrawActor->SelectedRigidBody.Id) { return; }
+	// Check whether the selected rigid body id, or solver has changed
+	if (SelectedRigidBodyId == GeometryCollectionDebugDrawActor->SelectedRigidBody.Id &&
+		SelectedChaosSolver == GeometryCollectionDebugDrawActor->SelectedRigidBody.Solver)
+	{
+		return;
+	}
 
 	// Default init selected transform index, in case of premature exit
 	SelectedTransformIndex = INDEX_NONE;
 
 	// Simple test to allow for an early exit when nothing has been selected
-	if (GeometryCollectionDebugDrawActor->SelectedRigidBody.Id == INDEX_NONE)
+	if (GeometryCollectionDebugDrawActor->SelectedRigidBody.Id == INDEX_NONE ||
+		GeometryCollectionDebugDrawActor->SelectedRigidBody.Solver != GeometryCollectionComponent->ChaosSolverActor)
 	{
-		SelectedRigidBodyId = INDEX_NONE;
+		SelectedRigidBodyId = GeometryCollectionDebugDrawActor->SelectedRigidBody.Id;
+		SelectedChaosSolver = GeometryCollectionDebugDrawActor->SelectedRigidBody.Solver;
 		return;
 	}
 
@@ -578,11 +586,14 @@ void UGeometryCollectionDebugDrawComponent::UpdateSelectedTransformIndex()
 		}
 	}
 
-	// Update selected rigid body index, only if a selected transform or no bad ids have been found.
-	// As otherwise this would prevent finding the selected transform once the ids hare fully synced.
+	// Update selected rigid body index and solver
+	// This needs to be done to mark the change of selection as already processed, and only when no bad ids have been found.
+	// Otherwise updating these would prevent finding the selected transform at the next call, once the remaining ids have fully synced.
+	// If a selection has been found, the update must be made, even if some invalid ids have been detected.
 	if (SelectedTransformIndex != INDEX_NONE || !bInvalidRigidBodyIdsFound)
 	{
 		SelectedRigidBodyId = GeometryCollectionDebugDrawActor->SelectedRigidBody.Id;
+		SelectedChaosSolver = GeometryCollectionDebugDrawActor->SelectedRigidBody.Solver;
 	}
 }
 
@@ -658,9 +669,11 @@ void UGeometryCollectionDebugDrawComponent::UpdateGeometryVisibility(bool bForce
 	else
 	{
 		// Work out partial changes in visiblity
-		const bool bSelected = (SelectedTransformIndex != INDEX_NONE);
-		const bool bAllSelected = (GeometryCollectionDebugDrawActor->SelectedRigidBody.Id == INDEX_NONE);
-		const bool bAnySelected = bSelected || bAllSelected;
+		const bool bIsSelected = (SelectedTransformIndex != INDEX_NONE);
+		const bool bAreAllSelected = (bIsSelected && GeometryCollectionDebugDrawActor->bDebugDrawWholeCollection) ||
+			(GeometryCollectionDebugDrawActor->SelectedRigidBody.Id == INDEX_NONE && 
+			 GeometryCollectionDebugDrawActor->SelectedRigidBody.Solver == GeometryCollectionComponent->ChaosSolverActor);
+		const bool bAreAnySelected = bIsSelected || bAreAllSelected;
 
 		switch (GeometryCollectionDebugDrawActor->HideGeometry)
 		{
@@ -670,9 +683,9 @@ void UGeometryCollectionDebugDrawComponent::UpdateGeometryVisibility(bool bForce
 			HiddenTransformIndex = INDEX_NONE;
 			break;
 		case EGeometryCollectionDebugDrawActorHideGeometry::HideWithCollision:
-			if (bShowRigidBodyCollisions || (bAnySelected && GeometryCollectionDebugDrawActor->bShowRigidBodyCollision))
+			if (bShowRigidBodyCollisions || (bAreAnySelected && GeometryCollectionDebugDrawActor->bShowRigidBodyCollision))
 			{
-				bIsVisible = !(GeometryCollectionDebugDrawActor->bDebugDrawWholeCollection || bAllSelected);
+				bIsVisible = !bAreAllSelected;
 				HiddenTransformIndex = bIsVisible ? SelectedTransformIndex: INDEX_NONE;
 			}
 			else
@@ -682,11 +695,11 @@ void UGeometryCollectionDebugDrawComponent::UpdateGeometryVisibility(bool bForce
 			}
 			break;
 		case EGeometryCollectionDebugDrawActorHideGeometry::HideSelected:
-			bIsVisible = !bAllSelected;
+			bIsVisible = !bAreAllSelected;
 			HiddenTransformIndex = bIsVisible ? SelectedTransformIndex: INDEX_NONE;
 			break;
 		case EGeometryCollectionDebugDrawActorHideGeometry::HideWholeCollection:
-			bIsVisible = !bAnySelected;
+			bIsVisible = !bAreAnySelected;
 			HiddenTransformIndex = INDEX_NONE;
 			break;
 		case EGeometryCollectionDebugDrawActorHideGeometry::HideAll:
@@ -760,9 +773,9 @@ void UGeometryCollectionDebugDrawComponent::UpdateTickStatus()
 	else
 	{
 		// Check whether anything from this component is selected for debug drawing
-		const bool bIsOneSelected = (SelectedTransformIndex != INDEX_NONE);
-		const bool bAreAllSelected = (GeometryCollectionDebugDrawActor->SelectedRigidBody.Id == INDEX_NONE);
-		const bool bAreAnySelected = bIsOneSelected || bAreAllSelected;
+		const bool bAreAnySelected = (SelectedTransformIndex != INDEX_NONE ||
+			(GeometryCollectionDebugDrawActor->SelectedRigidBody.Id == INDEX_NONE &&
+			 GeometryCollectionDebugDrawActor->SelectedRigidBody.Solver == GeometryCollectionComponent->ChaosSolverActor));
 
 		bIsEnabled =
 			(bAreAnySelected && (
@@ -897,7 +910,8 @@ void UGeometryCollectionDebugDrawComponent::DebugDrawChaosTick()
 	// Visualize other rigid body debug draw informations
 	const bool bIsOneSelected = bIsSelected && !GeometryCollectionDebugDrawActor->bDebugDrawWholeCollection;
 	const bool bAreAllSelected = (bIsSelected && GeometryCollectionDebugDrawActor->bDebugDrawWholeCollection) ||
-		(GeometryCollectionDebugDrawActor->SelectedRigidBody.Id == INDEX_NONE);
+		(GeometryCollectionDebugDrawActor->SelectedRigidBody.Id == INDEX_NONE && 
+		 GeometryCollectionDebugDrawActor->SelectedRigidBody.Solver == GeometryCollectionComponent->ChaosSolverActor);
 
 	if (!bShowRigidBodyIds && GeometryCollectionDebugDrawActor->bShowRigidBodyId && bIsOneSelected)
 	{
@@ -975,9 +989,9 @@ void UGeometryCollectionDebugDrawComponent::DebugDrawChaosTick()
 
 void UGeometryCollectionDebugDrawComponent::UpdateLevelSetVisibility()
 {
-	const bool bSelected = (SelectedTransformIndex != INDEX_NONE);
+	const bool bIsSelected = (SelectedTransformIndex != INDEX_NONE);
 	const bool bShowCollision = bShowRigidBodyCollisions || 
-		(bSelected && GeometryCollectionDebugDrawActor && GeometryCollectionDebugDrawActor->bShowRigidBodyCollision);
+		(bIsSelected && GeometryCollectionDebugDrawActor && GeometryCollectionDebugDrawActor->bShowRigidBodyCollision);
 	
 	if (RenderLevelSetOwner == this && (LastRenderedId == INDEX_NONE || !bShowCollision))
 	{
diff --git a/Engine/Source/Runtime/Experimental/GeometryCollectionEngine/Public/GeometryCollection/GeometryCollectionDebugDrawActor.h b/Engine/Source/Runtime/Experimental/GeometryCollectionEngine/Public/GeometryCollection/GeometryCollectionDebugDrawActor.h
--- a/Engine/Source/Runtime/Experimental/GeometryCollectionEngine/Public/GeometryCollection/GeometryCollectionDebugDrawActor.h
+++ b/Engine/Source/Runtime/Experimental/GeometryCollectionEngine/Public/GeometryCollection/GeometryCollectionDebugDrawActor.h
@@ -7,6 +7,8 @@
 
 #include "GeometryCollectionDebugDrawActor.generated.h"
 
+class AChaosSolverActor;
+class AGeometryCollectionActor;
 class UGeometryCollectionComponent;
 template<class InElementType> class TManagedArray;
 template<class T, int d> class TGeometryCollectionParticlesData;
@@ -41,14 +43,22 @@ struct FGeometryCollectionDebugDrawActorSelectedRigidBody
 {
 	GENERATED_USTRUCT_BODY()
 
-	explicit FGeometryCollectionDebugDrawActorSelectedRigidBody(int32 InId = -1) : Id(InId), GeometryCollectionActor(nullptr) {}
+	explicit FGeometryCollectionDebugDrawActorSelectedRigidBody(int32 InId = -1) : Id(InId), Solver(nullptr), GeometryCollection(nullptr) {}
 	
 	/** Id of the selected rigid body whose to visualize debug informations. Use -1 to visualize all Geometry Collections. */
-	UPROPERTY(EditAnywhere, Category = "Debug Draw", meta = (ClampMin="-1"))
+	UPROPERTY(EditAnywhere, Category = "Selected Rigid Body", meta = (ClampMin="-1"))
 	int32 Id;
 
-	UPROPERTY(EditAnywhere, Category = "Debug Draw")
-	AActor* GeometryCollectionActor;
+	/** Chaos RBD Solver. Will use the world's default solver actor if null. */
+	UPROPERTY(EditAnywhere, Category = "Selected Rigid Body")
+	AChaosSolverActor* Solver;
+
+	/** Currently selected geometry collection. */
+	UPROPERTY(VisibleAnywhere, Category = "Selected Rigid Body")
+	AGeometryCollectionActor* GeometryCollection;
+
+	/** Return the name of selected solver, or "None" if none is selected. */
+	FString GetSolverName() const;
 };
 
 /**
diff --git a/Engine/Source/Runtime/Experimental/GeometryCollectionEngine/Public/GeometryCollection/GeometryCollectionDebugDrawComponent.h b/Engine/Source/Runtime/Experimental/GeometryCollectionEngine/Public/GeometryCollection/GeometryCollectionDebugDrawComponent.h
--- a/Engine/Source/Runtime/Experimental/GeometryCollectionEngine/Public/GeometryCollection/GeometryCollectionDebugDrawComponent.h
+++ b/Engine/Source/Runtime/Experimental/GeometryCollectionEngine/Public/GeometryCollection/GeometryCollectionDebugDrawComponent.h
@@ -17,6 +17,7 @@
 class AGeometryCollectionRenderLevelSetActor;
 class UGeometryCollectionComponent;
 class AGeometryCollectionDebugDrawActor;
+class AChaosSolverActor;
 
 #if INCLUDE_CHAOS && GEOMETRYCOLLECTION_DEBUG_DRAW
 namespace Chaos { template<class T, int d> class TImplicitObject; }
@@ -290,5 +291,6 @@ private:
 	int32 SelectedRigidBodyId;
 	int32 SelectedTransformIndex;
 	int32 HiddenTransformIndex;
+	AChaosSolverActor* SelectedChaosSolver;
 #endif  // #if GEOMETRYCOLLECTION_DEBUG_DRAW
 };
