Fixed an issue with the geometry collection debug draw inital selection when the debug draw components are waiting for some rigid body ids array sync.

diff --git a/Engine/Plugins/Experimental/GeometryCollectionPlugin/Source/GeometryCollectionEditor/Private/GeometryCollection/GeometryCollectionSelectRigidBodyEdMode.cpp b/Engine/Plugins/Experimental/GeometryCollectionPlugin/Source/GeometryCollectionEditor/Private/GeometryCollection/GeometryCollectionSelectRigidBodyEdMode.cpp
--- a/Engine/Plugins/Experimental/GeometryCollectionPlugin/Source/GeometryCollectionEditor/Private/GeometryCollection/GeometryCollectionSelectRigidBodyEdMode.cpp
+++ b/Engine/Plugins/Experimental/GeometryCollectionPlugin/Source/GeometryCollectionEditor/Private/GeometryCollection/GeometryCollectionSelectRigidBodyEdMode.cpp
@@ -157,7 +157,7 @@ bool FGeometryCollectionSelectRigidBodyEdMode::HandleClick(FEditorViewportClient
 
 				// Retrieve the rigid body id
 				const TManagedArray<int32>& RigidBodyIdArray = GeometryCollectionComponent->GetRigidBodyIdArray();
-				const int32 RigidBodyId = (TransformIndex != INDEX_NONE) ? RigidBodyIdArray[TransformIndex]: INDEX_NONE;
+				const int32 RigidBodyId = (TransformIndex != INDEX_NONE && ensure(TransformIndex < RigidBodyIdArray.Num())) ? RigidBodyIdArray[TransformIndex]: INDEX_NONE;
 
 				// Update the rigid body id property
 				if (RigidBodyId != INDEX_NONE)
diff --git a/Engine/Source/Runtime/Experimental/GeometryCollectionEngine/Private/GeometryCollection/GeometryCollectionDebugDrawActor.cpp b/Engine/Source/Runtime/Experimental/GeometryCollectionEngine/Private/GeometryCollection/GeometryCollectionDebugDrawActor.cpp
--- a/Engine/Source/Runtime/Experimental/GeometryCollectionEngine/Private/GeometryCollection/GeometryCollectionDebugDrawActor.cpp
+++ b/Engine/Source/Runtime/Experimental/GeometryCollectionEngine/Private/GeometryCollection/GeometryCollectionDebugDrawActor.cpp
@@ -18,6 +18,10 @@
 #include "Components/BillboardComponent.h"
 #include "GenericPlatform/GenericPlatformMath.h"
 #include "HAL/IConsoleManager.h"
+#if INCLUDE_CHAOS
+#include "PBDRigidsSolver.h"
+#endif  // #if INCLUDE_CHAOS
+
 
 DEFINE_LOG_CATEGORY_STATIC(LogGeometryCollectionDebugDrawActor, Log, All);
 
@@ -255,11 +259,11 @@ void AGeometryCollectionDebugDrawActor::Tick(float DeltaSeconds)
 	// Clear all persistent strings and debug lines.
 	Flush();
 
+	UWorld* const World = GetWorld();
 #if INCLUDE_CHAOS && WITH_EDITOR
 	// Check editor pause status and force a dynamic update on all components to catchup with the physics thread
 	// This can't be done in the GeometryCollectionDebugDrawComponent since it doesn't tick at every frame,
 	// and can't be done in GeometryCollectionComponent either since it doesn't usually tick while paused.
-	UWorld* const World = GetWorld();
 	const bool bIsEditorPaused = World && World->IsPlayInEditor() && World->bDebugPauseExecution;
 	if (bIsEditorPaused && !bWasEditorPaused)
 	{
@@ -277,12 +281,56 @@ void AGeometryCollectionDebugDrawActor::Tick(float DeltaSeconds)
 	}
 	bWasEditorPaused = bIsEditorPaused;
 #endif  // #if INCLUDE_CHAOS && WITH_EDITOR
+
+#if GEOMETRYCOLLECTION_DEBUG_DRAW
+	// Check badly synced collections in case it is still looking for an id match
+	if (World && SelectedRigidBody.Id != INDEX_NONE && !SelectedRigidBody.GeometryCollection)
+	{
+#if INCLUDE_CHAOS
+		// Check the id is within the selected solver range
+		const Chaos::FPBDRigidsSolver* const Solver = 
+			SelectedRigidBody.Solver ? SelectedRigidBody.Solver->GetSolver() :  // Selected solver
+			World->PhysicsScene_Chaos ? World->PhysicsScene_Chaos->GetSolver() :  // Default world solver
+			nullptr;  // No solver
+
+		const bool IsWithinRange = Solver ? (uint32(SelectedRigidBody.Id) < Solver->GetRigidParticles().Size()): false;
+		if (!IsWithinRange)
+		{
+			UE_LOG(LogGeometryCollectionDebugDrawActor, VeryVerbose, TEXT("The selection id is out of range."));
+		}
+		else  // Statement continues below...
+#endif  // #if INCLUDE_CHAOS
+		{
+			UE_LOG(LogGeometryCollectionDebugDrawActor, VeryVerbose, TEXT("The selection couldn't be found. The property update will run on all components still containing any invalid rigid body ids."));
+
+			// Check for delayed Rigid Body Id array initializations
+			for (TActorIterator<AGeometryCollectionActor> ActorIterator(World); ActorIterator; ++ActorIterator)
+			{
+				if (UGeometryCollectionDebugDrawComponent* const GeometryCollectionDebugDrawComponent = ActorIterator->GetGeometryCollectionDebugDrawComponent())
+				{
+					if (GeometryCollectionDebugDrawComponent->GeometryCollectionDebugDrawActor == this &&
+						GeometryCollectionDebugDrawComponent->HasIncompleteRigidBodyIdSync())
+					{
+						const bool bIsSelected = GeometryCollectionDebugDrawComponent->OnDebugDrawPropertiesChanged(false);
+						if (bIsSelected)
+						{
+							SelectedRigidBody.GeometryCollection = *ActorIterator;
+							UE_LOG(LogGeometryCollectionDebugDrawActor, Verbose, TEXT("Selection found. Stopping continuous property update."));
+							break;
+						}
+					}
+				}
+			}
+		}
+	}
+#endif  // #if GEOMETRYCOLLECTION_DEBUG_DRAW
 }
 
 void AGeometryCollectionDebugDrawActor::BeginPlay()
 {
 	Super::BeginPlay();
 #if ENABLE_DRAW_DEBUG
+	// Initialize text renderer
 	const FDebugDrawDelegate DebugDrawTextDelegate = FDebugDrawDelegate::CreateUObject(this, &AGeometryCollectionDebugDrawActor::DebugDrawText);
 	DebugDrawTextDelegateHandle = UDebugDrawService::Register(TEXT("TextRender"), DebugDrawTextDelegate);  // TextRender is an engine show flag that works in both editor and game modes
 #endif  // #if ENABLE_DRAW_DEBUG
@@ -485,7 +533,9 @@ void AGeometryCollectionDebugDrawActor::PostEditChangeProperty(FPropertyChangedE
 
 	bool bForceVisibilityUpdate = false;
 
-	if      (PropertyName == GET_MEMBER_NAME_CHECKED(FGeometryCollectionDebugDrawActorSelectedRigidBody, Solver  )) { GeometryCollectionDebugDrawActorCVars::SelectedRigidBodySolver ->Set(*SelectedRigidBody.GetSolverName(), SetBy); }
+	if      (PropertyName == GET_MEMBER_NAME_CHECKED(AGeometryCollectionDebugDrawActor, SelectedRigidBody        )) { GeometryCollectionDebugDrawActorCVars::SelectedRigidBodySolver ->Set(*SelectedRigidBody.GetSolverName(), SetBy);
+																													  GeometryCollectionDebugDrawActorCVars::SelectedRigidBodyId     ->Set( SelectedRigidBody.Id             , SetBy); }
+	else if (PropertyName == GET_MEMBER_NAME_CHECKED(FGeometryCollectionDebugDrawActorSelectedRigidBody, Solver  )) { GeometryCollectionDebugDrawActorCVars::SelectedRigidBodySolver ->Set(*SelectedRigidBody.GetSolverName(), SetBy); }
 	else if (PropertyName == GET_MEMBER_NAME_CHECKED(FGeometryCollectionDebugDrawActorSelectedRigidBody, Id      )) { GeometryCollectionDebugDrawActorCVars::SelectedRigidBodyId     ->Set( SelectedRigidBody.Id             , SetBy); }
 	else if (PropertyName == GET_MEMBER_NAME_CHECKED(AGeometryCollectionDebugDrawActor, bDebugDrawWholeCollection)) { GeometryCollectionDebugDrawActorCVars::DebugDrawWholeCollection->Set(int32(bDebugDrawWholeCollection  ), SetBy); }
 	else if (PropertyName == GET_MEMBER_NAME_CHECKED(AGeometryCollectionDebugDrawActor, bDebugDrawHierarchy      )) { GeometryCollectionDebugDrawActorCVars::DebugDrawHierarchy      ->Set(int32(bDebugDrawHierarchy        ), SetBy); }
diff --git a/Engine/Source/Runtime/Experimental/GeometryCollectionEngine/Private/GeometryCollection/GeometryCollectionDebugDrawComponent.cpp b/Engine/Source/Runtime/Experimental/GeometryCollectionEngine/Private/GeometryCollection/GeometryCollectionDebugDrawComponent.cpp
--- a/Engine/Source/Runtime/Experimental/GeometryCollectionEngine/Private/GeometryCollection/GeometryCollectionDebugDrawComponent.cpp
+++ b/Engine/Source/Runtime/Experimental/GeometryCollectionEngine/Private/GeometryCollection/GeometryCollectionDebugDrawComponent.cpp
@@ -7,6 +7,7 @@
 #include "GeometryCollection/GeometryCollectionComponent.h"
 #include "GeometryCollection/GeometryCollectionObject.h"
 #include "GeometryCollection/GeometryCollectionDebugDrawActor.h"
+#include "GeometryCollection/GeometryCollection.h"
 #if INCLUDE_CHAOS
 #include "SolverObjects/GeometryCollectionPhysicsObject.h"
 #endif  // #if INCLUDE_CHAOS
@@ -109,6 +110,7 @@ UGeometryCollectionDebugDrawComponent::UGeometryCollectionDebugDrawComponent(con
 	, SelectedRigidBodyId(INDEX_NONE)
 	, SelectedTransformIndex(INDEX_NONE)
 	, HiddenTransformIndex(INDEX_NONE)
+	, bHasIncompleteRigidBodyIdSync(false)
 	, SelectedChaosSolver(nullptr)
 #endif  // #if GEOMETRYCOLLECTION_DEBUG_DRAW
 {
@@ -148,7 +150,7 @@ void UGeometryCollectionDebugDrawComponent::BeginPlay()
 		}
 
 		// Update the visibility and tick status depending on the debug draw properties currently selected
-		OnDebugDrawPropertiesChanged(true);
+		OnDebugDrawPropertiesChanged(false);
 
 #if INCLUDE_CHAOS
 		// Find or create level set renderer
@@ -540,6 +542,7 @@ void UGeometryCollectionDebugDrawComponent::DebugDrawTick()
 void UGeometryCollectionDebugDrawComponent::UpdateSelectedTransformIndex()
 {
 	check(GeometryCollectionComponent);
+
 	// No actor, no selection
 	if (!GeometryCollectionDebugDrawActor)
 	{
@@ -548,7 +551,8 @@ void UGeometryCollectionDebugDrawComponent::UpdateSelectedTransformIndex()
 	}
 
 	// Check whether the selected rigid body id, or solver has changed
-	if (SelectedRigidBodyId == GeometryCollectionDebugDrawActor->SelectedRigidBody.Id &&
+	if (!bHasIncompleteRigidBodyIdSync &&
+		SelectedRigidBodyId == GeometryCollectionDebugDrawActor->SelectedRigidBody.Id &&
 		SelectedChaosSolver == GeometryCollectionDebugDrawActor->SelectedRigidBody.Solver)
 	{
 		return;
@@ -569,32 +573,59 @@ void UGeometryCollectionDebugDrawComponent::UpdateSelectedTransformIndex()
 	// Check rigid body id sync
 	// Note that this test alone isn't enough to ensure that the rigid body ids are valid.
 	const TManagedArray<int32>& RigidBodyIds = GeometryCollectionComponent->RigidBodyIds;
-	if (RigidBodyIds.Num() == 0) { return; }
+	if (RigidBodyIds.Num() == 0)
+	{
+		bHasIncompleteRigidBodyIdSync = !!GeometryCollectionComponent->GetTransformArray().Num();
+		UE_CLOG(bHasIncompleteRigidBodyIdSync && GetOwner(), LogGeometryCollectionDebugDraw, Verbose, TEXT("UpdateSelectedTransformIndex(): Empty RigidBodyIds array for actor %s."), *GetOwner()->GetName());
+		return;
+	}
+
+	// Find the matching transform if any (and also check the sync completion status)
+	bHasIncompleteRigidBodyIdSync = false;
 
-	// Find the matching transform if any
-	bool bInvalidRigidBodyIdsFound = false;
-	for (int32 i = 0; i < RigidBodyIds.Num(); ++i)
+	const TManagedArray<TSet<int32>>& ChildrenRest = GeometryCollectionComponent->GetChildrenArrayRest();
+	const TManagedArray<TSet<int32>>& Children = GeometryCollectionComponent->GetChildrenArray();
+
+	for (int32 TransformIndex = 0; TransformIndex < RigidBodyIds.Num(); ++TransformIndex)
 	{
-		if (RigidBodyIds[i] == INDEX_NONE)
+		// Is this the selected id?
+		if (RigidBodyIds[TransformIndex] == GeometryCollectionDebugDrawActor->SelectedRigidBody.Id)
 		{
-			bInvalidRigidBodyIdsFound = true;
+			SelectedTransformIndex = TransformIndex;
+			bHasIncompleteRigidBodyIdSync = false;  // Found it, the wait for a sync can be canceled
+			break;
 		}
-		else if (RigidBodyIds[i] == GeometryCollectionDebugDrawActor->SelectedRigidBody.Id)
+		// Check the reason behind any invalid index
+		if (RigidBodyIds[TransformIndex] == INDEX_NONE)
 		{
-			SelectedTransformIndex = i;
-			break;
+			// Look for detached clusters in order to differentiate un-synced vs empty cluster rigid body ids.
+			int32 ChildTransformIndex = TransformIndex;
+			while (const TSet<int32>::TConstIterator ChildTransformIterator = Children[ChildTransformIndex].CreateConstIterator())
+			{
+				// Go down to the cluster's leaf level through the first child
+				ChildTransformIndex = *ChildTransformIterator;
+			}
+
+			// If this is a leaf bone, it can not be a detached cluster so it should have a valid rigid body
+			// In which case the sync has yet to happen and it might be worth trying this again later
+			if (!ChildrenRest[ChildTransformIndex].Num())
+			{
+				bHasIncompleteRigidBodyIdSync = true;
+				UE_CLOG(GetOwner(), LogGeometryCollectionDebugDraw, VeryVerbose, TEXT("UpdateSelectedTransformIndex(): Invalid rigid body id for actor %s, TransformIndex %d."), *GetOwner()->GetName(), TransformIndex);
+			}
+			else
+			{
+				// This should match the SimulationType == FST_CLUSTERED or IsClustered(int32 Element)
+				ensure(GeometryCollectionComponent->GetSimulationTypeArrayRest()[ChildTransformIndex] == FGeometryCollection::ESimulationTypes::FST_Clustered);
+				UE_CLOG(GetOwner(), LogGeometryCollectionDebugDraw, VeryVerbose, TEXT("UpdateSelectedTransformIndex(): Found empty cluster for actor %s, TransformIndex %d."), *GetOwner()->GetName(), TransformIndex);
+			}
 		}
 	}
+	UE_CLOG(bHasIncompleteRigidBodyIdSync && GetOwner(), LogGeometryCollectionDebugDraw, Verbose, TEXT("UpdateSelectedTransformIndex(): Invalid RigidBodyIds array elements for actor %s."), *GetOwner()->GetName());
 
 	// Update selected rigid body index and solver
-	// This needs to be done to mark the change of selection as already processed, and only when no bad ids have been found.
-	// Otherwise updating these would prevent finding the selected transform at the next call, once the remaining ids have fully synced.
-	// If a selection has been found, the update must be made, even if some invalid ids have been detected.
-	if (SelectedTransformIndex != INDEX_NONE || !bInvalidRigidBodyIdsFound)
-	{
-		SelectedRigidBodyId = GeometryCollectionDebugDrawActor->SelectedRigidBody.Id;
-		SelectedChaosSolver = GeometryCollectionDebugDrawActor->SelectedRigidBody.Solver;
-	}
+	SelectedRigidBodyId = GeometryCollectionDebugDrawActor->SelectedRigidBody.Id;
+	SelectedChaosSolver = GeometryCollectionDebugDrawActor->SelectedRigidBody.Solver;
 }
 
 int32 UGeometryCollectionDebugDrawComponent::CountFaces(int32 TransformIndex, bool bDebugDrawClustering) const
diff --git a/Engine/Source/Runtime/Experimental/GeometryCollectionEngine/Public/GeometryCollection/GeometryCollectionDebugDrawComponent.h b/Engine/Source/Runtime/Experimental/GeometryCollectionEngine/Public/GeometryCollection/GeometryCollectionDebugDrawComponent.h
--- a/Engine/Source/Runtime/Experimental/GeometryCollectionEngine/Public/GeometryCollection/GeometryCollectionDebugDrawComponent.h
+++ b/Engine/Source/Runtime/Experimental/GeometryCollectionEngine/Public/GeometryCollection/GeometryCollectionDebugDrawComponent.h
@@ -242,6 +242,9 @@ public:
 	/** Update selection and visibility after a change in cluster. Only handled when the debug drawing is active (the component is ticking). */
 	void OnClusterChanged();
 
+	/** Return whether the geometry collection rigid body id array is not completely initialized. This can happen when running the physics multithreaded. */
+	FORCEINLINE bool HasIncompleteRigidBodyIdSync() const { return bHasIncompleteRigidBodyIdSync;  }
+
 private:
 	/** Recursively compute global cluster transforms. Only gives geometry transforms for the leaf nodes, mid-level transforms are those of the clusters. */
 	void ComputeClusterTransforms(int32 Index, TArray<bool>& IsComputed, TArray<FTransform>& InOutGlobalTransforms);
@@ -291,6 +294,7 @@ private:
 	int32 SelectedRigidBodyId;
 	int32 SelectedTransformIndex;
 	int32 HiddenTransformIndex;
+	bool bHasIncompleteRigidBodyIdSync;
 	AChaosSolverActor* SelectedChaosSolver;
 #endif  // #if GEOMETRYCOLLECTION_DEBUG_DRAW
 };
