Fix geometry collection debug draw line flushing when visualizing parents.

diff --git a/Engine/Source/Runtime/Experimental/GeometryCollectionEngine/Private/GeometryCollection/GeometryCollectionDebugDrawActor.cpp b/Engine/Source/Runtime/Experimental/GeometryCollectionEngine/Private/GeometryCollection/GeometryCollectionDebugDrawActor.cpp
--- a/Engine/Source/Runtime/Experimental/GeometryCollectionEngine/Private/GeometryCollection/GeometryCollectionDebugDrawActor.cpp
+++ b/Engine/Source/Runtime/Experimental/GeometryCollectionEngine/Private/GeometryCollection/GeometryCollectionDebugDrawActor.cpp
@@ -1725,6 +1725,8 @@ void AGeometryCollectionDebugDrawActor::DrawParents(const TArray<FTransform>& Gl
 				const FColor ActiveColor = bUseActiveVisualization ? MakeDarker(Color, GetLevel(TransformIndex, ParentArray)) : Color;
 
 				DrawDebugLine(World, ParentPosition, Position, ActiveColor, GeometryCollectionDebugDrawActorConstants::bPersistent, GeometryCollectionDebugDrawActorConstants::LifeTime, GeometryCollectionDebugDrawActorConstants::DepthPriority, LineThickness);
+	
+				bNeedsDebugLinesFlush = true;
 			}
 		}
 	}
@@ -1760,6 +1762,8 @@ void AGeometryCollectionDebugDrawActor::DrawParent(const TArray<FTransform>& Glo
 		const float Scale = ArrowScale * FVector::Dist(ParentPosition, Position);
 
 		DrawDebugLine(World, ParentPosition, Position, Color, GeometryCollectionDebugDrawActorConstants::bPersistent, GeometryCollectionDebugDrawActorConstants::LifeTime, GeometryCollectionDebugDrawActorConstants::DepthPriority, LineThickness);
+
+		bNeedsDebugLinesFlush = true;
 	}
 
 	// Debug draw children if the cluster mode is on, or if there is no geometry attached to this node
