Collapsing the hierarchy will now clear the children list of the to be deleted node avoiding an ensure when it is removed.
Removed old offset transform code

diff --git a/Engine/Source/Runtime/Experimental/GeometryCollectionCore/Private/GeometryCollection/GeometryCollectionClusteringUtility.cpp b/Engine/Source/Runtime/Experimental/GeometryCollectionCore/Private/GeometryCollection/GeometryCollectionClusteringUtility.cpp
--- a/Engine/Source/Runtime/Experimental/GeometryCollectionCore/Private/GeometryCollection/GeometryCollectionClusteringUtility.cpp
+++ b/Engine/Source/Runtime/Experimental/GeometryCollectionCore/Private/GeometryCollection/GeometryCollectionClusteringUtility.cpp
@@ -28,31 +28,9 @@ void FGeometryCollectionClusteringUtility::ClusterBonesUnderNewNode(FGeometryCol
 
 	Transforms[NewBoneIndex] = FTransform::Identity;
 
-	if (GeometryCollection->HasAttribute("ExplodedVector", FGeometryCollection::TransformGroup) &&
-		GeometryCollection->HasAttribute("ExplodedTransform", FGeometryCollection::TransformGroup))
-	{
-		TManagedArray<FVector>& ExplodedVectors = GeometryCollection->GetAttribute<FVector>("ExplodedVector", FGeometryCollection::TransformGroup);
-		TManagedArray<FTransform>& ExplodedTransforms = GeometryCollection->GetAttribute<FTransform>("ExplodedTransform", FGeometryCollection::TransformGroup);
-
-		ExplodedTransforms[NewBoneIndex] = Transforms[NewBoneIndex];
-		ResetSliderTransforms(ExplodedTransforms, Transforms);
-
-		// Selected Bone Setup
-		FVector SumOfOffsets(0, 0, 0);
-		for (int32 SelectedBoneIndex : SelectedBones)
-		{
-			ExplodedTransforms[SelectedBoneIndex] = Transforms[SelectedBoneIndex];
-			SumOfOffsets += ExplodedVectors[SelectedBoneIndex];
-		}
-
-		// This bones offset is the average of all the selected bones
-		ExplodedVectors[NewBoneIndex] = SumOfOffsets / SelectedBones.Num();
-	}
-
 	// re-parent all the geometry nodes under the new shared bone
 	GeometryCollectionAlgo::ParentTransforms(GeometryCollection, NewBoneIndex, SelectedBones);
 
-
 	UpdateHierarchyLevelOfChildren(GeometryCollection, NewBoneIndex);
 
 	// Parent Bone Fixup of Children - add the new node under the first bone selected
@@ -335,8 +313,6 @@ void FGeometryCollectionClusteringUtility::CollapseHierarchyOneLevel(FGeometryCo
 	TManagedArray<TSet<int32>>& Children = GeometryCollection->Children;
 	TManagedArray<FTransform>& Transforms = GeometryCollection->Transform;
 	TManagedArray<FString>& BoneNames = GeometryCollection->BoneName;
-	TManagedArray<FTransform>& ExplodedTransforms = GeometryCollection->GetAttribute<FTransform>("ExplodedTransform", FGeometryCollection::TransformGroup);
-	TManagedArray<FVector>& ExplodedVectors = GeometryCollection->GetAttribute<FVector>("ExplodedVector", FGeometryCollection::TransformGroup);
 	TManagedArray<int32>& Levels = GeometryCollection->GetAttribute<int32>("Level", FGeometryCollection::TransformGroup);
 
 	for (int32 SourceElement : SourceElements)
@@ -351,10 +327,10 @@ void FGeometryCollectionClusteringUtility::CollapseHierarchyOneLevel(FGeometryCo
 				for (int32 ChildElement : Children[DeletedNode])
 				{
 					Children[NewParentElement].Add(ChildElement);
-
 					Levels[ChildElement] -= 1;
 					Parents[ChildElement] = NewParentElement;
 				}
+				Children[DeletedNode].Empty();
 			}
 		}
 	}
@@ -686,7 +662,7 @@ void FGeometryCollectionClusteringUtility::CollapseSelectedHierarchy(int8 Level,
 	if (Level > 0)
 	{
 		TArray<int32> Elements;
-		for (int Element = 0; Element < SelectedBones.Num(); Element++)
+		for (int32 Element = 0; Element < SelectedBones.Num(); Element++)
 		{
 			int32 Index = SelectedBones[Element];
 
