Use Inverse Transform Position for Fields.

diff --git a/Engine/Source/Runtime/Experimental/FieldSystem/Source/FieldSystemCore/Private/Field/FieldSystemNodes.cpp b/Engine/Source/Runtime/Experimental/FieldSystem/Source/FieldSystemCore/Private/Field/FieldSystemNodes.cpp
--- a/Engine/Source/Runtime/Experimental/FieldSystem/Source/FieldSystemCore/Private/Field/FieldSystemNodes.cpp
+++ b/Engine/Source/Runtime/Experimental/FieldSystem/Source/FieldSystemCore/Private/Field/FieldSystemNodes.cpp
@@ -537,14 +537,13 @@ void FBoxFalloff::Evaluator<EFieldFalloffType::Field_FallOff_None>(const FFieldC
 {
 	int32 NumSamples = Context.SampleIndices.Num();
 	FBox UnitBox(FVector(-50), FVector(50));
-	FTransform InverseTransform = Transform.Inverse();
 	for (int32 SampleIndex = 0; SampleIndex < NumSamples; SampleIndex++)
 	{
 		const ContextIndex & Index = Context.SampleIndices[SampleIndex];
 		{
 			Results[Index.Result] = Default;
 
-			FVector LocalPoint = InverseTransform.TransformPosition(Context.Samples[Index.Sample]);
+			FVector LocalPoint = Transform.InverseTransformPosition(Context.Samples[Index.Sample]);
 			if (UnitBox.IsInside(LocalPoint))
 			{
 				Results[Index.Result] = Magnitude;
@@ -560,14 +559,13 @@ void FBoxFalloff::Evaluator<EFieldFalloffType::Field_Falloff_Linear>(const FFiel
 
 	int32 NumSamples = Context.SampleIndices.Num();
 	FBox UnitBox(FVector(-50), FVector(50));
-	FTransform InverseTransform = Transform.Inverse();
 	for (int32 SampleIndex = 0; SampleIndex < NumSamples; SampleIndex++)
 	{
 		const ContextIndex & Index = Context.SampleIndices[SampleIndex];
 		{
 			Results[Index.Result] = Default;
 
-			FVector LocalPoint = InverseTransform.TransformPosition(Context.Samples[Index.Sample]);
+			FVector LocalPoint = Transform.InverseTransformPosition(Context.Samples[Index.Sample]);
 			if (UnitBox.IsInside(LocalPoint))
 			{
 				float Delta2 = Context.Samples[Index.Sample].SizeSquared();
@@ -594,14 +592,13 @@ void FBoxFalloff::Evaluator<EFieldFalloffType::Field_Falloff_Squared>(const FFie
 {
 	int32 NumSamples = Context.SampleIndices.Num();
 	FBox UnitBox(FVector(-50), FVector(50));
-	FTransform InverseTransform = Transform.Inverse();
 	for (int32 SampleIndex = 0; SampleIndex < NumSamples; SampleIndex++)
 	{
 		const ContextIndex & Index = Context.SampleIndices[SampleIndex];
 		{
 			Results[Index.Result] = Default;
 
-			FVector LocalPoint = InverseTransform.TransformPosition(Context.Samples[Index.Sample]);
+			FVector LocalPoint = Transform.InverseTransformPosition(Context.Samples[Index.Sample]);
 			if (UnitBox.IsInside(LocalPoint))
 			{
 				float Distance = UnitBox.ComputeSquaredDistanceToPoint(LocalPoint);
@@ -615,12 +612,11 @@ void FBoxFalloff::Evaluator<EFieldFalloffType::Field_Falloff_Inverse>(const FFie
 {
 	int32 NumSamples = Context.SampleIndices.Num();
 	FBox UnitBox(FVector(-50), FVector(50));
-	FTransform InverseTransform = Transform.Inverse();
 	for (int32 SampleIndex = 0; SampleIndex < NumSamples; SampleIndex++)
 	{
 		const ContextIndex & Index = Context.SampleIndices[SampleIndex];
 		{
-			FVector LocalPoint = InverseTransform.TransformPosition(Context.Samples[Index.Sample]);
+			FVector LocalPoint = Transform.InverseTransformPosition(Context.Samples[Index.Sample]);
 			if (UnitBox.IsInside(LocalPoint))
 			{
 				Results[Index.Result] = Default;
@@ -639,14 +635,13 @@ void FBoxFalloff::Evaluator<EFieldFalloffType::Field_Falloff_Logarithmic>(const
 {
 	int32 NumSamples = Context.SampleIndices.Num();
 	FBox UnitBox(FVector(-50), FVector(50));
-	FTransform InverseTransform = Transform.Inverse();
 	for (int32 SampleIndex = 0; SampleIndex < NumSamples; SampleIndex++)
 	{
 		const ContextIndex & Index = Context.SampleIndices[SampleIndex];
 		{
 			Results[Index.Result] = Default;
 
-			FVector LocalPoint = InverseTransform.TransformPosition(Context.Samples[Index.Sample]);
+			FVector LocalPoint = Transform.InverseTransformPosition(Context.Samples[Index.Sample]);
 			if (UnitBox.IsInside(LocalPoint))
 			{
 				float Distance = UnitBox.ComputeSquaredDistanceToPoint(LocalPoint);
@@ -723,13 +718,12 @@ FNoiseField::Evaluate(const FFieldContext & Context, TArrayView<float> & Results
 	float Scale = 0.5;
 	FVector BoundsSize = 0.1*Bounds.GetSize();
 	FVector Root = Scale * FVector(FMath::Sqrt(BoundsSize.X), FMath::Sqrt(BoundsSize.Y), FMath::Sqrt(BoundsSize.Z));
-	FTransform InverseTransform = Transform.Inverse();
 
 	float MinResult = FLT_MAX, MaxResult = -FLT_MAX;
 	for (int32 SampleIndex = 0; SampleIndex < NumSamples; SampleIndex++)
 	{
 		const ContextIndex & Index = Context.SampleIndices[SampleIndex];
-		FVector Delta = InverseTransform.InverseTransformVector(Context.Samples[Index.Sample] - Bounds.Min);
+		FVector Delta = Transform.TransformVector(Context.Samples[Index.Sample] - Bounds.Min);
 
 		float i = BoundsSize.X? (Delta.X / BoundsSize.X) : 0;
 		float j = BoundsSize.Y? (Delta.Y / BoundsSize.Y) : 0;
