Fix Chaos particles flagged as disabled still being in ActiveIndices array. Replace incorrect usages with Evolution's DisableParticle function.
Fix RigidParticle initializing Island to 0 instead of INDEX_NONE.

diff --git a/Engine/Source/Runtime/Engine/Private/PhysicsEngine/Experimental/PhysScene_Chaos.cpp b/Engine/Source/Runtime/Engine/Private/PhysicsEngine/Experimental/PhysScene_Chaos.cpp
--- a/Engine/Source/Runtime/Engine/Private/PhysicsEngine/Experimental/PhysScene_Chaos.cpp
+++ b/Engine/Source/Runtime/Engine/Private/PhysicsEngine/Experimental/PhysScene_Chaos.cpp
@@ -316,7 +316,7 @@ static void CopyParticleData(Chaos::TPBDRigidParticles<float, 3>& ToParticles, c
 	ToParticles.InvI(ToIndex) = FromParticles.InvI(FromIndex);
 	ToParticles.SetGeometry(ToIndex, FromParticles.Geometry(FromIndex));	//question: do we need to deal with dynamic geometry?
 	ToParticles.CollisionParticles(ToIndex) = MoveTemp(FromParticles.CollisionParticles(FromIndex));
-	ToParticles.Disabled(ToIndex) = FromParticles.Disabled(FromIndex);
+	ToParticles.SetDisabledLowLevel(ToIndex, FromParticles.Disabled(FromIndex));
 	ToParticles.SetSleeping(ToIndex, FromParticles.Sleeping(FromIndex));
 }
 
diff --git a/Engine/Source/Runtime/Experimental/Chaos/Private/Chaos/PBDRigidsEvolutionGBF.cpp b/Engine/Source/Runtime/Experimental/Chaos/Private/Chaos/PBDRigidsEvolutionGBF.cpp
--- a/Engine/Source/Runtime/Experimental/Chaos/Private/Chaos/PBDRigidsEvolutionGBF.cpp
+++ b/Engine/Source/Runtime/Experimental/Chaos/Private/Chaos/PBDRigidsEvolutionGBF.cpp
@@ -213,7 +213,7 @@ void TPBDRigidsEvolutionGBF<T, d>::AdvanceOneTimeStep(const T Dt)
 				if (ParticleDisableCount[Index] > DisableThreshold)
 				{
 					ParticleDisableCount[Index] = 0;
-					Particles.Disabled(Index) = true;
+					Particles.SetDisabledLowLevel(Index, true);
 					DisabledParticles[Island].Add(Index);
 					Particles.V(Index) = TVector<T, d>(0);
 					Particles.W(Index) = TVector<T, d>(0);
@@ -221,7 +221,7 @@ void TPBDRigidsEvolutionGBF<T, d>::AdvanceOneTimeStep(const T Dt)
 
 				if (!(ensure(!FMath::IsNaN(Particles.P(Index)[0])) && ensure(!FMath::IsNaN(Particles.P(Index)[1])) && ensure(!FMath::IsNaN(Particles.P(Index)[2]))))
 				{
-					Particles.Disabled(Index) = true;
+					Particles.SetDisabledLowLevel(Index, true);
 					DisabledParticles[Island].Add(Index);
 				}
 			}
diff --git a/Engine/Source/Runtime/Experimental/Chaos/Private/Chaos/PBDRigidsEvolutionPGS.cpp b/Engine/Source/Runtime/Experimental/Chaos/Private/Chaos/PBDRigidsEvolutionPGS.cpp
--- a/Engine/Source/Runtime/Experimental/Chaos/Private/Chaos/PBDRigidsEvolutionPGS.cpp
+++ b/Engine/Source/Runtime/Experimental/Chaos/Private/Chaos/PBDRigidsEvolutionPGS.cpp
@@ -126,12 +126,12 @@ void TPBDRigidsEvolutionPGS<T, d>::AdvanceOneTimeStep(const T Dt)
 		{
 			if (Particles.ObjectState(Index) != EObjectStateType::Kinematic && Particles.V(Index).SizeSquared() < PhysicsMaterials[Index]->DisabledLinearThreshold && Particles.W(Index).SizeSquared() < PhysicsMaterials[Index]->DisabledAngularThreshold)
 			{
-				Particles.Disabled(Index) = true;
+				Particles.SetDisabledLowLevel(Index, true);
 				DisabledParticles[Island].Add(Index);
 			}
 			if (!(ensure(!FMath::IsNaN(Particles.P(Index)[0])) && ensure(!FMath::IsNaN(Particles.P(Index)[1])) && ensure(!FMath::IsNaN(Particles.P(Index)[2]))))
 			{
-				Particles.Disabled(Index) = true;
+				Particles.SetDisabledLowLevel(Index, true);
 				DisabledParticles[Island].Add(Index);
 			}
 		}
@@ -144,11 +144,13 @@ void TPBDRigidsEvolutionPGS<T, d>::AdvanceOneTimeStep(const T Dt)
 			for (const int32 Index : ConstraintGraph.GetIslandParticles(i))
 			{
 				ActiveIndices.Remove(Index);
+				NonDisabledIndices.Add(Index);
 			}
 		}
 		for (const int32 Index : DisabledParticles[i])
 		{
 			ActiveIndices.Remove(Index);
+			NonDisabledIndices.Add(Index);
 		}
 	}
 
diff --git a/Engine/Source/Runtime/Experimental/Chaos/Public/Chaos/PBDRigidsEvolution.h b/Engine/Source/Runtime/Experimental/Chaos/Public/Chaos/PBDRigidsEvolution.h
--- a/Engine/Source/Runtime/Experimental/Chaos/Public/Chaos/PBDRigidsEvolution.h
+++ b/Engine/Source/Runtime/Experimental/Chaos/Public/Chaos/PBDRigidsEvolution.h
@@ -114,7 +114,7 @@ class TPBDRigidsEvolutionBase
 	 */
 	CHAOS_API void EnableParticle(const int32 ParticleIndex, const int32 ParentParticleIndex)
 	{
-		Particles.Disabled(ParticleIndex) = false;
+		Particles.SetDisabledLowLevel(ParticleIndex, false);
 		NonDisabledIndices.Add(ParticleIndex);
 		ActiveIndices.Add(ParticleIndex);
 
@@ -123,7 +123,7 @@ class TPBDRigidsEvolutionBase
 
 	CHAOS_API void DisableParticle(const int32 ParticleIndex)
 	{
-		Particles.Disabled(ParticleIndex) = true;
+		Particles.SetDisabledLowLevel(ParticleIndex, true);
 		NonDisabledIndices.Remove(ParticleIndex);
 		ActiveIndices.Remove(ParticleIndex);
 
@@ -141,7 +141,7 @@ class TPBDRigidsEvolutionBase
 	{
 		for (int32 ParticleIndex : InParticleIndices)
 		{
-			Particles.Disabled(ParticleIndex) = true;
+			Particles.SetDisabledLowLevel(ParticleIndex, true);
 			NonDisabledIndices.Remove(ParticleIndex);
 			ActiveIndices.Remove(ParticleIndex);
 		}
diff --git a/Engine/Source/Runtime/Experimental/Chaos/Public/Chaos/ParticleHandle.h b/Engine/Source/Runtime/Experimental/Chaos/Public/Chaos/ParticleHandle.h
--- a/Engine/Source/Runtime/Experimental/Chaos/Public/Chaos/ParticleHandle.h
+++ b/Engine/Source/Runtime/Experimental/Chaos/Public/Chaos/ParticleHandle.h
@@ -52,7 +52,7 @@ void PBDRigidParticleDefaultConstruct(FConcrete& Concrete, const TPBDRigidPartic
 {
 	//don't bother calling parent since the call gets made by the corresponding hierarchy in FConcrete
 	Concrete.CollisionGroup() = 0;
-	Concrete.Disabled() = Params.bDisabled;
+	Concrete.SetDisabledLowLevel(Params.bDisabled);
 	Concrete.PreV() = Concrete.V();
 	Concrete.PreW() = Concrete.W();
 	Concrete.P() = Concrete.X();
@@ -63,7 +63,7 @@ void PBDRigidParticleDefaultConstruct(FConcrete& Concrete, const TPBDRigidPartic
 	Concrete.InvM() = 1;
 	Concrete.I() = PMatrix<T, d, d>(1, 1, 1);
 	Concrete.InvI() = PMatrix<T, d, d>(1, 1, 1);
-	Concrete.Island() = 0;
+	Concrete.Island() = INDEX_NONE;
 	Concrete.ToBeRemovedOnFracture() = false;
 	Concrete.SetObjectState(Params.bStartSleeping ? EObjectStateType::Sleeping : EObjectStateType::Dynamic);
 }
@@ -194,7 +194,9 @@ public:
 	int32& CollisionGroup() { return PBDRigidParticles->CollisionGroup(ParticleIdx); }
 
 	bool Disabled() const { return PBDRigidParticles->Disabled(ParticleIdx); }
-	bool& Disabled() { return PBDRigidParticles->Disabled(ParticleIdx); }
+
+	// See Comment on TRigidParticle::SetDisabledLowLevel. State changes in Evolution should accompany this call.
+	void SetDisabledLowLevel(bool disabled) { PBDRigidParticles->SetDisabledLowLevel(ParticleIdx, disabled); }
 
 	const TVector<T, d>& PreV() const { return PBDRigidParticles->PreV(ParticleIdx); }
 	TVector<T, d>& PreV() { return PBDRigidParticles->PreV(ParticleIdx); }
@@ -371,7 +373,10 @@ public:
 	int32& CollisionGroup() { return MCollisionGroup; }
 
 	bool Disabled() const { return MDisabled; }
-	bool& Disabled() { return MDisabled; }
+
+	// Named to match signature of TPBDRigidParticleHandle, as both are used in templated functions.
+	// See its comment for details.
+	bool& SetDisabledLowLevel() { return MDisabled; }
 
 	const TVector<T, d>& PreV() const { return MPreV; }
 	TVector<T, d>& PreV() { return MPreV; }
@@ -553,7 +558,10 @@ public:
 	TVector<T, d>& W() { return NonConstCurSOA->W(ParticleIdx); }
 
 	int32& CollisionGroup() { return NonConstCurSOA->CollisionGroup(ParticleIdx); }
-	bool& Disabled() { return NonConstCurSOA->Disabled(ParticleIdx); }
+
+	// See Comment on TRigidParticle::SetDisabledLowLevel. State changes in Evolution should accompany this call.
+	bool& SetDisabledLowLevel() { return NonConstCurSOA->Disabled(ParticleIdx); }
+
 	TVector<T, d>& PreV() { return NonConstCurSOA->PreV(ParticleIdx); }
 	TVector<T, d>& PreW() { return NonConstCurSOA->PreW(ParticleIdx); }
 	TVector<T, d>& P() { return NonConstCurSOA->P(ParticleIdx); }
diff --git a/Engine/Source/Runtime/Experimental/Chaos/Public/Chaos/RigidParticles.h b/Engine/Source/Runtime/Experimental/Chaos/Public/Chaos/RigidParticles.h
--- a/Engine/Source/Runtime/Experimental/Chaos/Public/Chaos/RigidParticles.h
+++ b/Engine/Source/Runtime/Experimental/Chaos/Public/Chaos/RigidParticles.h
@@ -112,7 +112,12 @@ class CHAOS_API TRigidParticles : public TKinematicGeometryParticles<T, d>
 	int32& CollisionGroup(const int32 Index) { return MCollisionGroup[Index]; }
 
 	const bool Disabled(const int32 Index) const { return MDisabled[Index]; }
-	bool& Disabled(const int32 Index) { return MDisabled[Index]; }
+
+	bool& DisabledRef(const int32 Index) { return MDisabled[Index]; }
+
+	// DisableParticle/EnableParticle on Evolution should be used. Don't disable particles with this.
+    // Using this will break stuff. This is for solver's use only, and possibly some particle construction/copy code.
+	void SetDisabledLowLevel(const int32 Index, bool disabled) { MDisabled[Index] = disabled; }
 
 	const bool ToBeRemovedOnFracture(const int32 Index) const { return MToBeRemovedOnFracture[Index]; }
 	bool& ToBeRemovedOnFracture(const int32 Index) { return MToBeRemovedOnFracture[Index]; }
diff --git a/Engine/Source/Runtime/Experimental/ChaosSolvers/Private/SolverObjects/BodyInstancePhysicsObject.cpp b/Engine/Source/Runtime/Experimental/ChaosSolvers/Private/SolverObjects/BodyInstancePhysicsObject.cpp
--- a/Engine/Source/Runtime/Experimental/ChaosSolvers/Private/SolverObjects/BodyInstancePhysicsObject.cpp
+++ b/Engine/Source/Runtime/Experimental/ChaosSolvers/Private/SolverObjects/BodyInstancePhysicsObject.cpp
@@ -54,9 +54,7 @@ void FBodyInstancePhysicsObject::OnRemoveFromScene()
 
 		for(const int32 Index : InitializedIndices)
 		{
-			CurrSolver->GetRigidParticles().Disabled(Index) = true;
-			CurrSolver->ActiveIndices().Remove(Index);
-			CurrSolver->NonDisabledIndices().Remove(Index);
+			CurrSolver->GetEvolution()->DisableParticle(Index);
 		}
 	}
 
diff --git a/Engine/Source/Runtime/Experimental/ChaosSolvers/Private/SolverObjects/FieldSystemPhysicsObject.cpp b/Engine/Source/Runtime/Experimental/ChaosSolvers/Private/SolverObjects/FieldSystemPhysicsObject.cpp
--- a/Engine/Source/Runtime/Experimental/ChaosSolvers/Private/SolverObjects/FieldSystemPhysicsObject.cpp
+++ b/Engine/Source/Runtime/Experimental/ChaosSolvers/Private/SolverObjects/FieldSystemPhysicsObject.cpp
@@ -244,10 +244,8 @@ void FFieldSystemPhysicsObject::FieldParameterUpdateCallback(Chaos::FPBDRigidsSo
 							if (DynamicStateView[RigidBodyIndex] == 0 && Particles.Disabled(RigidBodyIndex))
 							{
 								ensure(CurrentSolver->GetRigidClustering().GetClusterIdsArray()[RigidBodyIndex].Id == INDEX_NONE);
+								CurrentSolver->GetEvolution()->EnableParticle(RigidBodyIndex, INDEX_NONE);
 								Particles.SetObjectState(RigidBodyIndex, Chaos::EObjectStateType::Dynamic);
-								Particles.Disabled(RigidBodyIndex) = false;
-								CurrentSolver->ActiveIndices().Add(RigidBodyIndex);
-								CurrentSolver->NonDisabledIndices().Add(RigidBodyIndex);
 							}
 						}
 					}
@@ -343,23 +341,14 @@ void FFieldSystemPhysicsObject::FieldParameterUpdateCallback(Chaos::FPBDRigidsSo
 							if (!Particles.Disabled(i) && Results[i] > 0.0)
 							{
 								RemovedParticles.Add((uint32)i);
-								Particles.Disabled(i) = true;
+								CurrentSolver->GetEvolution()->DisableParticle(i);
 							}
 						}
 
-						if (RemovedParticles.Num())
+						if (RemovedParticles.Num() && bHasFloor)
 						{
-							if (bHasFloor)
-							{
-								Particles.Disabled(FloorIndex) = false;
-								Particles.SetObjectState(FloorIndex, Chaos::EObjectStateType::Static);
-							}
-							for (int32 ParticleIndex : RemovedParticles)
-							{
-								CurrentSolver->ActiveIndices().Remove(ParticleIndex);
-								CurrentSolver->NonDisabledIndices().Remove(ParticleIndex);
-							}
-							CurrentSolver->RemoveConstraints(RemovedParticles);
+							CurrentSolver->GetEvolution()->DisableParticle(FloorIndex);
+							Particles.SetObjectState(FloorIndex, Chaos::EObjectStateType::Static);
 						}
 					}
 				}
diff --git a/Engine/Source/Runtime/Experimental/ChaosSolvers/Private/SolverObjects/GeometryCollectionPhysicsObject.cpp b/Engine/Source/Runtime/Experimental/ChaosSolvers/Private/SolverObjects/GeometryCollectionPhysicsObject.cpp
--- a/Engine/Source/Runtime/Experimental/ChaosSolvers/Private/SolverObjects/GeometryCollectionPhysicsObject.cpp
+++ b/Engine/Source/Runtime/Experimental/ChaosSolvers/Private/SolverObjects/GeometryCollectionPhysicsObject.cpp
@@ -185,7 +185,7 @@ void PopulateSimulatedParticle(FGeometryCollectionPhysicsObject::FParticlesType&
 {
 	SCOPE_CYCLE_COUNTER(STAT_PopulateSimulatedParticle);
 
-	Particles.Disabled(RigidBodyIndex) = false;
+	Particles.SetDisabledLowLevel(RigidBodyIndex, false);
 
 	Particles.X(RigidBodyIndex) = WorldTransform.GetTranslation();
 	Particles.V(RigidBodyIndex) = Chaos::TVector<float, 3>( FVector(0) );
@@ -193,6 +193,7 @@ void PopulateSimulatedParticle(FGeometryCollectionPhysicsObject::FParticlesType&
 	Particles.W(RigidBodyIndex) = Chaos::TVector<float, 3>( FVector(0) );
 	Particles.P(RigidBodyIndex) = Particles.X(RigidBodyIndex);
 	Particles.Q(RigidBodyIndex) = Particles.R(RigidBodyIndex);
+	Particles.Island(RigidBodyIndex) = INDEX_NONE;
 
 	//todo: if mass to small use the right inertia
 	ensureMsgf(MassIn >= SharedParams.MinimumMassClamp, TEXT("Mass smaller than minimum mass clamp. Too late to change"));
@@ -501,7 +502,7 @@ void FGeometryCollectionPhysicsObject::StartFrameCallback(const float Dt, const
 				const bool bDisabledNow = Particles.Disabled(ExternalParticleIndex);
 				if (bDisabledNow != bShouldBeDisabled)
 				{
-					Particles.Disabled(ExternalParticleIndex) = bShouldBeDisabled;
+					Particles.SetDisabledLowLevel(ExternalParticleIndex, bShouldBeDisabled);
 					ParticleUpdateLock.Lock();
 					if (!bShouldBeDisabled)
 					{
@@ -1143,7 +1144,7 @@ void FGeometryCollectionPhysicsObject::CreateRigidBodyCallback(FParticlesType& P
 					if (!Particles.Disabled(RigidBodyIndex))
 					{
 						PendingActivationList.Add(TransformGroupIndex);
-						Particles.Disabled(RigidBodyIndex) = true;
+						GetSolver()->GetEvolution()->DisableParticle(RigidBodyIndex);
 					}
 				}
 			}
@@ -1172,10 +1173,13 @@ void FGeometryCollectionPhysicsObject::ActivateBodies()
 	{
 		if (InitializedState == ESimulationInitializationState::Created)
 		{
+			int32 ParentIndex = INDEX_NONE;
+
 			if (Parameters.EnableClustering && Parameters.ClusterGroupIndex)
 			{
 				Chaos::FPBDRigidsSolver::FClusteringType & Clustering = GetSolver()->GetRigidClustering();
 				Clustering.DecrementPendingClusterCounter(Parameters.ClusterGroupIndex);
+				ParentIndex = Parameters.ClusterGroupIndex;
 			}
 
 			Chaos::FPBDRigidsSolver::FParticlesType & Particles = GetSolver()->GetRigidParticles();
@@ -1185,10 +1189,8 @@ void FGeometryCollectionPhysicsObject::ActivateBodies()
 				checkSlow(RigidBodyIndex != INDEX_NONE);
 				if (Particles.Disabled(RigidBodyIndex))
 				{
-					Particles.Disabled(RigidBodyIndex) = false;
-					GetSolver()->NonDisabledIndices().Add(RigidBodyIndex);
+					GetSolver()->GetEvolution()->EnableParticle(RigidBodyIndex, ParentIndex);
 				}
-				GetSolver()->ActiveIndices().Add(RigidBodyIndex);
 			}
 
 			PendingActivationList.Reset(0);
@@ -3144,9 +3146,7 @@ void FGeometryCollectionPhysicsObject::OnRemoveFromScene()
 	{
 		for (int32 ParticleIndex = 0; ParticleIndex < Count; ++ParticleIndex)
 		{
-			Particles.Disabled(Begin + ParticleIndex) = true;
-			GetSolver()->ActiveIndices().Remove(Begin + ParticleIndex);
-			GetSolver()->NonDisabledIndices().Remove(Begin + ParticleIndex);
+			GetSolver()->GetEvolution()->DisableParticle(Begin + ParticleIndex);
 			GetSolver()->GetRigidClustering().GetTopLevelClusterParents().Remove(Begin + ParticleIndex);
 		}
 	}
@@ -3224,7 +3224,7 @@ void FGeometryCollectionPhysicsObject::CacheResults()
 	if(NumParticles > 0)
 	{
 		SCOPE_CYCLE_COUNTER(STAT_CaptureDisabledState);
-		TargetResults.DisabledStates.Append(&Particles.Disabled(BaseParticleIndex), NumParticles);
+		TargetResults.DisabledStates.Append(&Particles.DisabledRef(BaseParticleIndex), NumParticles);
 	}
 
 	{
diff --git a/Engine/Source/Runtime/Experimental/ChaosSolvers/Private/SolverObjects/SkeletalMeshPhysicsObject.cpp b/Engine/Source/Runtime/Experimental/ChaosSolvers/Private/SolverObjects/SkeletalMeshPhysicsObject.cpp
--- a/Engine/Source/Runtime/Experimental/ChaosSolvers/Private/SolverObjects/SkeletalMeshPhysicsObject.cpp
+++ b/Engine/Source/Runtime/Experimental/ChaosSolvers/Private/SolverObjects/SkeletalMeshPhysicsObject.cpp
@@ -207,7 +207,7 @@ void FSkeletalMeshPhysicsObject::CreateRigidBodyCallback(FParticlesType& Particl
 #endif
 			GetSolver()->SetPhysicsMaterial(RigidBodyId, Parameters.PhysicalMaterial);
 
-			Particles.Disabled(RigidBodyId) = false;
+			GetSolver()->GetEvolution()->DisableParticle(RigidBodyId);
 
 			//Particles.X(RigidBodyId) = WorldTransform->TransformPosition(CenterOfMass);
 			Particles.X(RigidBodyId) = WorldTransform->GetTranslation();
@@ -383,9 +383,7 @@ void FSkeletalMeshPhysicsObject::OnRemoveFromScene()
 		{
 			if(static_cast<uint32>(Id) < Particles.Size())
 			{
-				Particles.Disabled(Id) = true;
-				CurrSolver->ActiveIndices().Remove(Id);
-				CurrSolver->NonDisabledIndices().Remove(Id);
+				CurrSolver->GetEvolution()->DisableParticle(Id);
 				CurrSolver->GetRigidClustering().GetTopLevelClusterParents().Remove(Id);
 			}
 		}
diff --git a/Engine/Source/Runtime/Experimental/ChaosSolvers/Private/SolverObjects/StaticMeshPhysicsObject.cpp b/Engine/Source/Runtime/Experimental/ChaosSolvers/Private/SolverObjects/StaticMeshPhysicsObject.cpp
--- a/Engine/Source/Runtime/Experimental/ChaosSolvers/Private/SolverObjects/StaticMeshPhysicsObject.cpp
+++ b/Engine/Source/Runtime/Experimental/ChaosSolvers/Private/SolverObjects/StaticMeshPhysicsObject.cpp
@@ -331,9 +331,7 @@ void FStaticMeshPhysicsObject::OnRemoveFromScene()
 			return;
 		}
 
-		CurrSolver->GetRigidParticles().Disabled(RigidBodyId) = true;
-		CurrSolver->ActiveIndices().Remove(RigidBodyId);
-		CurrSolver->NonDisabledIndices().Remove(RigidBodyId);
+		CurrSolver->GetEvolution()->DisableParticle(RigidBodyId);
 		CurrSolver->GetRigidClustering().GetTopLevelClusterParents().Remove(RigidBodyId);
 	}
 }
diff --git a/Engine/Source/Runtime/PhysicsCore/Private/PhysTestSerializer.cpp b/Engine/Source/Runtime/PhysicsCore/Private/PhysTestSerializer.cpp
--- a/Engine/Source/Runtime/PhysicsCore/Private/PhysTestSerializer.cpp
+++ b/Engine/Source/Runtime/PhysicsCore/Private/PhysTestSerializer.cpp
@@ -220,7 +220,7 @@ void FPhysTestSerializer::CreateChaosData()
 		Particles.W(Idx) + TVector<float, 3>(0);
 		Particles.P(Idx) = Particles.X(Idx);
 		Particles.Q(Idx) = Particles.R(Idx);
-		Particles.Disabled(Idx) = false;
+		Particles.SetDisabledLowLevel(Idx, false);
 
 		PxActorToChaosIdx.Add(Act, Idx);
 
