Fix distance setting on fields.

diff --git a/Engine/Plugins/Experimental/GeometryCollectionPlugin/Source/GeometryCollectionExampleCore/Private/GeometryCollection/GeometryCollectionExampleFields.cpp b/Engine/Plugins/Experimental/GeometryCollectionPlugin/Source/GeometryCollectionExampleCore/Private/GeometryCollection/GeometryCollectionExampleFields.cpp
--- a/Engine/Plugins/Experimental/GeometryCollectionPlugin/Source/GeometryCollectionExampleCore/Private/GeometryCollection/GeometryCollectionExampleFields.cpp
+++ b/Engine/Plugins/Experimental/GeometryCollectionPlugin/Source/GeometryCollectionExampleCore/Private/GeometryCollection/GeometryCollectionExampleFields.cpp
@@ -1141,7 +1141,7 @@ namespace GeometryCollectionExample
 		}
 
 		{
-			FFieldSystemCommand CommandOut("FPlaneFalloff", new FPlaneFalloff(1.f, 3.f, 5.f, 7.f, FVector(9, 11, 13), FVector(17, 19, 23)));
+			FFieldSystemCommand CommandOut("FPlaneFalloff", new FPlaneFalloff(1.f, 3.f, 5.f, 7.f, 100.f, FVector(9, 11, 13), FVector(17, 19, 23)));
 			FFieldSystemCommand CommandIn = SaveAndLoad(CommandOut);
 			R.ExpectTrue(CommandIn == CommandOut);
 		}
@@ -1259,7 +1259,7 @@ namespace GeometryCollectionExample
 			FUniformInteger* UniformInteger = new FUniformInteger(3);
 			FConversionField<int32, float>*  ConversionFieldIF = new FConversionField<int32, float>(UniformInteger);
 
-			FPlaneFalloff* PlaneFalloff = new FPlaneFalloff(1.f, 3.f, 5.f, 7.f, FVector(9, 11, 13), FVector(17, 19, 23));
+			FPlaneFalloff* PlaneFalloff = new FPlaneFalloff(1.f, 3.f, 5.f, 7.f, 100.f, FVector(9, 11, 13), FVector(17, 19, 23));
 			FCullingField<float>* CullingFieldF = new FCullingField<float>(PlaneFalloff, ConversionFieldIF);
 
 
diff --git a/Engine/Source/Runtime/Experimental/FieldSystem/Source/FieldSystemCore/Private/Field/FieldSystemNodes.cpp b/Engine/Source/Runtime/Experimental/FieldSystem/Source/FieldSystemCore/Private/Field/FieldSystemNodes.cpp
--- a/Engine/Source/Runtime/Experimental/FieldSystem/Source/FieldSystemCore/Private/Field/FieldSystemNodes.cpp
+++ b/Engine/Source/Runtime/Experimental/FieldSystem/Source/FieldSystemCore/Private/Field/FieldSystemNodes.cpp
@@ -391,7 +391,8 @@ void FPlaneFalloff::Evaluator<EFieldFalloffType::Field_FallOff_None>(const FFiel
 		const ContextIndex & Index = Context.SampleIndices[SampleIndex];
 		{
 			Results[Index.Result] = Default;
-			if (Plane.PlaneDot(Context.Samples[Index.Sample]) < 0)
+			float LocalDistance = Plane.PlaneDot(Context.Samples[Index.Sample]);
+			if (LocalDistance < 0 && LocalDistance > -Distance)
 			{
 				Results[Index.Result] = Magnitude;
 			}
@@ -404,17 +405,17 @@ void FPlaneFalloff::Evaluator<EFieldFalloffType::Field_Falloff_Linear>(const FFi
 	float Scalar = (MaxRange - MinRange);
 	bool bNonUnitScalar = FMath::Abs(Scalar - 1.f) > KINDA_SMALL_NUMBER;
 
-	float NegMagnitude = -Magnitude;
 	int32 NumSamples = Context.SampleIndices.Num();
 	for (int32 SampleIndex = 0; SampleIndex < NumSamples; SampleIndex++)
 	{
 		const ContextIndex & Index = Context.SampleIndices[SampleIndex];
 		{
 			Results[Index.Result] = Default;
-			float Distance = Plane.PlaneDot(Context.Samples[Index.Sample]);
-			if (Distance < 0)
+			float LocalDistance = Plane.PlaneDot(Context.Samples[Index.Sample]);
+			if (LocalDistance < 0 && LocalDistance > -Distance)
 			{
-				Results[Index.Result] = NegMagnitude * Distance;
+				// 0 at -Distance, Mag at 0
+				Results[Index.Result] = Magnitude * (1.f - (LocalDistance / Distance));
 				if (bNonUnitScalar)
 				{
 					Results[Index.Result] = FMath::Min(FMath::Max(MinRange,Results[Index.Result]),MaxRange);
@@ -432,10 +433,11 @@ void FPlaneFalloff::Evaluator<EFieldFalloffType::Field_Falloff_Squared>(const FF
 		const ContextIndex & Index = Context.SampleIndices[SampleIndex];
 		{
 			Results[Index.Result] = Default;
-			float Distance = Plane.PlaneDot(Context.Samples[Index.Sample]);
-			if (Distance < -SMALL_NUMBER)
+			float LocalDistance = Plane.PlaneDot(Context.Samples[Index.Sample]);
+			if (LocalDistance < -SMALL_NUMBER && LocalDistance > -Distance)
 			{
-				Results[Index.Result] = Magnitude * Distance * Distance;
+				float Fraction = (1.f - (LocalDistance / Distance));
+				Results[Index.Result] = Magnitude * Fraction * Fraction;
 			}
 		}
 	}
@@ -451,10 +453,11 @@ void FPlaneFalloff::Evaluator<EFieldFalloffType::Field_Falloff_Inverse>(const FF
 		const ContextIndex & Index = Context.SampleIndices[SampleIndex];
 		{
 			Results[Index.Result] = Default;
-			float Distance = Plane.PlaneDot(Context.Samples[Index.Sample]);
-			if (Distance < -SMALL_NUMBER)
+			float LocalDistance = Plane.PlaneDot(Context.Samples[Index.Sample]);
+			if (LocalDistance < -SMALL_NUMBER && LocalDistance > -Distance)
 			{
-				Results[Index.Result] = NegMagnitude / Distance;
+				float Fraction = (1.f - (LocalDistance / Distance));
+				Results[Index.Result] = NegMagnitude / (Fraction * Fraction);
 			}
 		}
 	}
@@ -468,10 +471,11 @@ void FPlaneFalloff::Evaluator<EFieldFalloffType::Field_Falloff_Logarithmic>(cons
 		const ContextIndex & Index = Context.SampleIndices[SampleIndex];
 		{
 			Results[Index.Result] = Default;
-			float Distance = Plane.PlaneDot(Context.Samples[Index.Sample]);
-			if (Distance < -SMALL_NUMBER)
+			float LocalDistance = Plane.PlaneDot(Context.Samples[Index.Sample]);
+			if (LocalDistance < -SMALL_NUMBER)
 			{
-				Results[Index.Result] = Magnitude * FMath::LogX(10.f, -Distance + 1.f);
+				float Fraction = (1.f - (LocalDistance / Distance));
+				Results[Index.Result] = Magnitude * FMath::LogX(10.f, Fraction + 1.f);
 			}
 		}
 	}
@@ -506,6 +510,7 @@ void FPlaneFalloff::Serialize(FArchive& Ar)
 	Ar << MinRange;
 	Ar << MaxRange;
 	Ar << Default;
+	Ar << Distance;
 	Ar << Position;
 	Ar << Normal;
 	SerializeInternal<EFieldFalloffType>(Ar, Falloff);
@@ -521,6 +526,7 @@ bool FPlaneFalloff::operator==(const FFieldNodeBase& Node)
 			&& MinRange == Other->MinRange
 			&& MaxRange == Other->MaxRange
 			&& Default == Other->Default
+			&& Distance == Other->Distance
 			&& Position == Other->Position
 			&& Normal == Other->Normal
 			&& Falloff == Other->Falloff;
diff --git a/Engine/Source/Runtime/Experimental/FieldSystem/Source/FieldSystemCore/Public/Field/FieldSystemNodes.h b/Engine/Source/Runtime/Experimental/FieldSystem/Source/FieldSystemCore/Public/Field/FieldSystemNodes.h
--- a/Engine/Source/Runtime/Experimental/FieldSystem/Source/FieldSystemCore/Public/Field/FieldSystemNodes.h
+++ b/Engine/Source/Runtime/Experimental/FieldSystem/Source/FieldSystemCore/Public/Field/FieldSystemNodes.h
@@ -171,6 +171,7 @@ public:
 		float MinRangeIn = 0.f,
 		float MaxRangeIn = 1.f,
 		float DefaultIn = 0.f,
+		float DistanceIn = 0.f,
 		FVector PositionIn = FVector(0, 0, 0),
 		FVector NormalIn = FVector(0, 0, 1),
 		EFieldFalloffType FalloffIn = EFieldFalloffType::Field_Falloff_Linear)
@@ -180,11 +181,12 @@ public:
 		, MinRange(MinRangeIn)
 		, MaxRange(MaxRangeIn)
 		, Default(DefaultIn)
+		, Distance(DistanceIn)
 		, Position(PositionIn)
 		, Normal(NormalIn)
 		, Falloff(FalloffIn)
 	{}
-	virtual FFieldNodeBase * NewCopy() const override { return new FPlaneFalloff(Magnitude, MinRange, MaxRange, Default, Position, Normal, Falloff); }
+	virtual FFieldNodeBase * NewCopy() const override { return new FPlaneFalloff(Magnitude, MinRange, MaxRange, Default, Distance, Position, Normal, Falloff); }
 	virtual ~FPlaneFalloff() {}
 
 	void Evaluate(const FFieldContext &, TArrayView<float> & Results) const override;
@@ -198,6 +200,7 @@ public:
 	float MinRange;
 	float MaxRange;
 	float Default;
+	float Distance;
 	FVector Position;
 	FVector Normal;
 	EFieldFalloffType Falloff;
diff --git a/Engine/Source/Runtime/Experimental/FieldSystem/Source/FieldSystemEngine/Private/Field/FieldSystemObjects.cpp b/Engine/Source/Runtime/Experimental/FieldSystem/Source/FieldSystemEngine/Private/Field/FieldSystemObjects.cpp
--- a/Engine/Source/Runtime/Experimental/FieldSystem/Source/FieldSystemEngine/Private/Field/FieldSystemObjects.cpp
+++ b/Engine/Source/Runtime/Experimental/FieldSystem/Source/FieldSystemEngine/Private/Field/FieldSystemObjects.cpp
@@ -124,7 +124,7 @@ UPlaneFalloff::NewEvaluationGraph(TArray<const UFieldNodeBase*>& Nodes) const
 	if (ensureMsgf(!Nodes.Contains(this), TEXT("Cycle Dependency Error : Graph nodes may not be resued in a single chain.")))
 	{
 		Nodes.Add(this);
-		return new FPlaneFalloff(Magnitude, MinRange, MaxRange, Default, Position, Normal, Falloff);
+		return new FPlaneFalloff(Magnitude, MinRange, MaxRange, Default, Distance, Position, Normal, Falloff);
 	}
 	return nullptr;
 }
diff --git a/Engine/Source/Runtime/Experimental/FieldSystem/Source/FieldSystemEngine/Public/Field/FieldSystemObjects.h b/Engine/Source/Runtime/Experimental/FieldSystem/Source/FieldSystemEngine/Public/Field/FieldSystemObjects.h
--- a/Engine/Source/Runtime/Experimental/FieldSystem/Source/FieldSystemEngine/Public/Field/FieldSystemObjects.h
+++ b/Engine/Source/Runtime/Experimental/FieldSystem/Source/FieldSystemEngine/Public/Field/FieldSystemObjects.h
@@ -291,6 +291,7 @@ public:
 		, MinRange(0.f)
 		, MaxRange(1.f)
 		, Default(0.f)
+		, Distance(0.f)
 		, Position(FVector(0, 0, 0))
 		, Normal(FVector(0, 0, 1))
 		, Falloff(EFieldFalloffType::Field_Falloff_Linear)
@@ -315,6 +316,9 @@ public:
 	UPROPERTY(BlueprintReadWrite, EditAnywhere, Category = "Field")
 	float Default;
 
+	UPROPERTY(BlueprintReadWrite, EditAnywhere, Category = "Field")
+	float Distance;
+
 	UPROPERTY(BlueprintReadWrite, EditAnywhere, Category = "Field")
 	FVector Position;
 
