Fix issue with Geometry Collections that ignore the visible property.

diff --git a/Engine/Source/Runtime/Experimental/GeometryCollectionEngine/Private/GeometryCollection/GeometryCollectionDebugDrawComponent.cpp b/Engine/Source/Runtime/Experimental/GeometryCollectionEngine/Private/GeometryCollection/GeometryCollectionDebugDrawComponent.cpp
--- a/Engine/Source/Runtime/Experimental/GeometryCollectionEngine/Private/GeometryCollection/GeometryCollectionDebugDrawComponent.cpp
+++ b/Engine/Source/Runtime/Experimental/GeometryCollectionEngine/Private/GeometryCollection/GeometryCollectionDebugDrawComponent.cpp
@@ -110,6 +110,7 @@ UGeometryCollectionDebugDrawComponent::UGeometryCollectionDebugDrawComponent(con
 	, SelectedRigidBodyId(INDEX_NONE)
 	, SelectedTransformIndex(INDEX_NONE)
 	, HiddenTransformIndex(INDEX_NONE)
+	, bWasVisible(true)
 	, bHasIncompleteRigidBodyIdSync(false)
 	, SelectedChaosSolver(nullptr)
 #endif  // #if GEOMETRYCOLLECTION_DEBUG_DRAW
@@ -139,6 +140,7 @@ void UGeometryCollectionDebugDrawComponent::BeginPlay()
 		SelectedRigidBodyId = INDEX_NONE;
 		SelectedTransformIndex = INDEX_NONE;
 		HiddenTransformIndex = INDEX_NONE;
+		bWasVisible = true;
 
 		// Find or create global debug draw actor
 		GeometryCollectionDebugDrawActor = AGeometryCollectionDebugDrawActor::FindOrCreate(GetWorld());
@@ -684,9 +686,6 @@ void UGeometryCollectionDebugDrawComponent::UpdateGeometryVisibility(bool bForce
 		return;
 	}
 
-	// Save current component visibility state
-	const bool bWasVisible = GeometryCollectionComponent->IsVisible();
-
 	// Keep old hidden index
 	const int32 PrevHiddenIndex = HiddenTransformIndex;
 
@@ -784,13 +783,18 @@ void UGeometryCollectionDebugDrawComponent::UpdateGeometryVisibility(bool bForce
 		GeometryCollectionComponent->ForceRenderUpdateConstantData();
 	}
 
-	// Update component visibility
-	const bool bVisibilityHasChanged = (bIsVisible != bWasVisible);
-	UE_CLOG(bVisibilityHasChanged, LogGeometryCollectionDebugDraw, Verbose, TEXT("UpdateGeometryVisibility(): Visibility has changed. Old visibility = %d, new visibility = %d."), bWasVisible, bIsVisible);
-	if (bVisibilityHasChanged)
+	// Update component visibility, only if the component visibility hasn't changed in between the last call (or unless the change is in sync with the component visibility)
+	const bool bIsComponentVisible = GeometryCollectionComponent->IsVisible();
+	const bool bAllowVisibilityChange = (bIsComponentVisible || !bWasVisible);
+	if (bAllowVisibilityChange)
 	{
-		UE_LOG(LogGeometryCollectionDebugDraw, Verbose, TEXT("UpdateGeometryVisibility(): Setting visibility."));
-		GeometryCollectionComponent->SetVisibility(bIsVisible);
+		const bool bVisibilityHasChanged = (bIsVisible != bIsComponentVisible);
+		if (bVisibilityHasChanged)
+		{
+			UE_LOG(LogGeometryCollectionDebugDraw, Verbose, TEXT("UpdateGeometryVisibility(): Visibility has changed. Old visibility = %d, new visibility = %d."), bWasVisible, bIsVisible);
+			GeometryCollectionComponent->SetVisibility(bIsVisible);
+		}
+		bWasVisible = bIsVisible;  // Only update when changes are allowed so that the component can stay hidden when visibility is out of sync.
 	}
 }
 
diff --git a/Engine/Source/Runtime/Experimental/GeometryCollectionEngine/Public/GeometryCollection/GeometryCollectionDebugDrawComponent.h b/Engine/Source/Runtime/Experimental/GeometryCollectionEngine/Public/GeometryCollection/GeometryCollectionDebugDrawComponent.h
--- a/Engine/Source/Runtime/Experimental/GeometryCollectionEngine/Public/GeometryCollection/GeometryCollectionDebugDrawComponent.h
+++ b/Engine/Source/Runtime/Experimental/GeometryCollectionEngine/Public/GeometryCollection/GeometryCollectionDebugDrawComponent.h
@@ -294,6 +294,7 @@ private:
 	int32 SelectedRigidBodyId;
 	int32 SelectedTransformIndex;
 	int32 HiddenTransformIndex;
+	bool bWasVisible;
 	bool bHasIncompleteRigidBodyIdSync;
 	AChaosSolverActor* SelectedChaosSolver;
 #endif  // #if GEOMETRYCOLLECTION_DEBUG_DRAW
