Fix GeometryCollectionPhysicsObject initialization of clustered particles allocating unnecessary particles and leaving uninitialized.

diff --git a/Engine/Source/Runtime/Experimental/ChaosSolvers/Private/SolverObjects/GeometryCollectionPhysicsObject.cpp b/Engine/Source/Runtime/Experimental/ChaosSolvers/Private/SolverObjects/GeometryCollectionPhysicsObject.cpp
--- a/Engine/Source/Runtime/Experimental/ChaosSolvers/Private/SolverObjects/GeometryCollectionPhysicsObject.cpp
+++ b/Engine/Source/Runtime/Experimental/ChaosSolvers/Private/SolverObjects/GeometryCollectionPhysicsObject.cpp
@@ -1049,7 +1049,7 @@ void FGeometryCollectionPhysicsObject::CreateRigidBodyCallback(FParticlesType& P
 					TArray<uint32> RigidChildren, CollectionChildren;
 					for (const int32 ChildIndex : Children[TransformGroupIndex])
 					{
-						if(RigidBodyID[ChildIndex] != INDEX_NONE || Children[TransformGroupIndex].Num())
+						if(RigidBodyID[ChildIndex] != INDEX_NONE || Children[ChildIndex].Num())
 						{
 							NumClusters++;
 							break;
