GeometryCollection Cache improvements
- fixed removal of internal faces for non-started caches (GeometryCollection must be static, though)
- fixed behavior of skipping ahead in cache
- fixed rewinding cache in Sequencer editor

diff --git a/Engine/Source/Runtime/Experimental/GeometryCollectionEngine/Private/GeometryCollection/GeometryCollectionComponent.cpp b/Engine/Source/Runtime/Experimental/GeometryCollectionEngine/Private/GeometryCollection/GeometryCollectionComponent.cpp
--- a/Engine/Source/Runtime/Experimental/GeometryCollectionEngine/Private/GeometryCollection/GeometryCollectionComponent.cpp
+++ b/Engine/Source/Runtime/Experimental/GeometryCollectionEngine/Private/GeometryCollection/GeometryCollectionComponent.cpp
@@ -682,19 +682,24 @@ void UGeometryCollectionComponent::InitDynamicData(FGeometryCollectionDynamicDat
 		const TManagedArray<FTransform> &Transform = RestCollection->GetGeometryCollection()->Transform;
 		const TManagedArray<FTransform> &MassToLocal = RestCollection->GetGeometryCollection()->GetAttribute<FTransform>("MassToLocal", FGeometryCollection::TransformGroup);
 
-		DynamicData->IsDynamic = true;
+		// #todo(dmp): find a better place to calculate and store this
+		float CacheDt = CacheParameters.TargetCache->GetData()->GetDt();
+
+		// if we want the state before the first cached frame, then the collection doesn't need to be dynamic since it'll render the pre-fractured geometry
+		DynamicData->IsDynamic = DesiredCacheTime > CacheDt;
 
 		// if we are already on the current cached frame, return
 		if (FMath::IsNearlyEqual(CurrentCacheTime, DesiredCacheTime) && GlobalMatrices.Num() != 0)
 		{
 			DynamicData->PrevTransforms = GlobalMatrices;
 			DynamicData->Transforms = GlobalMatrices;
+
+			// maintaining the cache time means we should consider the transforms equal for dynamic data sending purposes
+			TransformsAreEqual[(TransformsAreEqualIndex++) % TransformsAreEqual.Num()] = true;
+
 			return;
 		}
 
-		// #todo(dmp): find a better place to calculate and store this
-		float CacheDt = CacheParameters.TargetCache->GetData()->GetDt();
-
 		// if the input simulation time to playback is the first frame, reset simulation time
 		if (DesiredCacheTime <= CacheDt ||  FMath::IsNearlyEqual(CurrentCacheTime, FLT_MAX))
 		{
@@ -706,6 +711,9 @@ void UGeometryCollectionComponent::InitDynamicData(FGeometryCollectionDynamicDat
 
 			EventsPlayed.Empty();
 			EventsPlayed.AddDefaulted(CacheParameters.TargetCache->GetData()->Records.Num());
+
+			// reset should send new transforms to the RT
+			TransformsAreEqual[(TransformsAreEqualIndex++) % TransformsAreEqual.Num()] = false;
 		}
 		else if (GlobalMatrices.Num() == 0)
 		{
@@ -717,6 +725,9 @@ void UGeometryCollectionComponent::InitDynamicData(FGeometryCollectionDynamicDat
 			EventsPlayed.Empty();
 			EventsPlayed.AddDefaulted(CacheParameters.TargetCache->GetData()->Records.Num());
 
+			// degenerate case causes and reset should send new transforms to the RT
+			TransformsAreEqual[(TransformsAreEqualIndex++) % TransformsAreEqual.Num()] = false;
+
 			UE_LOG(UGCC_LOG, Warning, TEXT("Cache out of sync - must rewind sequencer to start frame"));
 		}
 		else if (DesiredCacheTime >= CurrentCacheTime)
@@ -727,6 +738,8 @@ void UGeometryCollectionComponent::InitDynamicData(FGeometryCollectionDynamicDat
 			
 			FTransform ActorToWorld = GetComponentTransform();
 
+			bool HasAnyActiveTransforms = false;
+
 			// Jump ahead in increments of CacheDt evaluating the cache until we reach our desired time
 			for (int st = 0; st < NumSteps; ++st)
 			{
@@ -737,7 +750,7 @@ void UGeometryCollectionComponent::InitDynamicData(FGeometryCollectionDynamicDat
 
 				const FRecordedFrame* FirstFrame = nullptr;
 				const FRecordedFrame* SecondFrame = nullptr;
-				CacheParameters.TargetCache->GetData()->GetFramesForTime(DesiredCacheTime, FirstFrame, SecondFrame);
+				CacheParameters.TargetCache->GetData()->GetFramesForTime(CurrentCacheTime, FirstFrame, SecondFrame);
 
 				if (FirstFrame && !SecondFrame)
 				{
@@ -746,6 +759,12 @@ void UGeometryCollectionComponent::InitDynamicData(FGeometryCollectionDynamicDat
 					const TArray<int32> &TransformIndices = FirstFrame->TransformIndices;
 
 					const int32 NumActives = FirstFrame->TransformIndices.Num();
+					
+					if (NumActives > 0)
+					{
+						HasAnyActiveTransforms = true;
+					}
+
 					for (int i = 0; i < NumActives; ++i)
 					{
 						const int32 InternalIndexTmp = TransformIndices[i];
@@ -771,12 +790,18 @@ void UGeometryCollectionComponent::InitDynamicData(FGeometryCollectionDynamicDat
 						GeometryCollectionAlgo::GlobalMatricesFromRoot(InternalIndexTmp, Transform, Children, GlobalMatrices);
 					}
 				}
-				else if (FirstFrame && SecondFrame && DesiredCacheTime > FirstFrame->Timestamp)
+				else if (FirstFrame && SecondFrame && CurrentCacheTime > FirstFrame->Timestamp)
 				{
-					const float Alpha = (DesiredCacheTime - FirstFrame->Timestamp) / (SecondFrame->Timestamp - FirstFrame->Timestamp);
+					const float Alpha = (CurrentCacheTime - FirstFrame->Timestamp) / (SecondFrame->Timestamp - FirstFrame->Timestamp);
 					check(0 <= Alpha && Alpha <= 1.0f);
 
 					const int32 NumActives = SecondFrame->TransformIndices.Num();
+
+					if (NumActives > 0)
+					{
+						HasAnyActiveTransforms = true;
+					}
+
 					for (int Index = 0; Index < NumActives; ++Index)
 					{
 						const int32 InternalIndexTmp = SecondFrame->TransformIndices[Index];
@@ -927,12 +952,18 @@ void UGeometryCollectionComponent::InitDynamicData(FGeometryCollectionDynamicDat
 #endif
 				}								
 			}
+
+			// check if transforms at start of this tick are the same as what is calculated from the cache
+			TransformsAreEqual[(TransformsAreEqualIndex++) % TransformsAreEqual.Num()] = !HasAnyActiveTransforms;
 		}
 		else
 		{			
 			// time is before current cache time so maintain the matrices we have since we can't rewind
 			DynamicData->PrevTransforms = GlobalMatrices;
 			DynamicData->Transforms = GlobalMatrices;
+
+			// reset event means we don't want to consider transforms as being equal between prev and current frame
+			TransformsAreEqual[(TransformsAreEqualIndex++) % TransformsAreEqual.Num()] = true;
 		}
 	}
 	else if (DynamicData->IsDynamic)
@@ -944,6 +975,9 @@ void UGeometryCollectionComponent::InitDynamicData(FGeometryCollectionDynamicDat
 			CalculateGlobalMatrices();		
 			DynamicData->PrevTransforms = GlobalMatrices;		
 			DynamicData->Transforms = GlobalMatrices;
+
+			// reset event means we don't want to consider transforms as being equal between prev and current frame
+			TransformsAreEqual[(TransformsAreEqualIndex++) % TransformsAreEqual.Num()] = false;
 		}
 		else
 		{
@@ -960,11 +994,11 @@ void UGeometryCollectionComponent::InitDynamicData(FGeometryCollectionDynamicDat
 			}
 
 			DynamicData->Transforms = GlobalMatrices;
-		}	
-	}
 
-	// check if previous transforms are the same as current
-	TransformsAreEqual[(TransformsAreEqualIndex++) % TransformsAreEqual.Num()] = IsEqual(DynamicData->PrevTransforms, DynamicData->Transforms);
+			// check if previous transforms are the same as current
+			TransformsAreEqual[(TransformsAreEqualIndex++) % TransformsAreEqual.Num()] = IsEqual(DynamicData->PrevTransforms, DynamicData->Transforms);
+		}
+	}
 }
 
 #if CHAOS_DEBUG_DRAW
@@ -2099,6 +2133,11 @@ void UGeometryCollectionComponent::SwitchRenderModels(const AActor* Actor)
 		}
 		else if (UGeometryCollectionComponent* GeometryCollectionComponent = Cast<UGeometryCollectionComponent>(PrimitiveComponent))
 		{
+			if (!GeometryCollectionComponent->IsVisible())
+			{
+				continue;
+			}
+
 			GeometryCollectionComponent->SetVisibility(true);
 		}
 	}
