Fix use of FTransform by value for Win32.

diff --git a/Engine/Source/Runtime/Experimental/Chaos/Private/Chaos/PBDCollisionConstraintPGS.cpp b/Engine/Source/Runtime/Experimental/Chaos/Private/Chaos/PBDCollisionConstraintPGS.cpp
--- a/Engine/Source/Runtime/Experimental/Chaos/Private/Chaos/PBDCollisionConstraintPGS.cpp
+++ b/Engine/Source/Runtime/Experimental/Chaos/Private/Chaos/PBDCollisionConstraintPGS.cpp
@@ -854,7 +854,7 @@ bool TPBDCollisionConstraintPGS<T, d>::NearestPoint(TArray<Pair<TVector<T, d>, T
 }
 
 template<class T, int d>
-void UpdateLevelsetConstraintHelperCCD(const TRigidParticles<T, d>& InParticles, const int32 j, const TRigidTransform<T, d> LocalToWorld1, const TRigidTransform<T, d> LocalToWorld2, const T Thickness, TRigidBodyContactConstraintPGS<T, d>& Constraint)
+void UpdateLevelsetConstraintHelperCCD(const TRigidParticles<T, d>& InParticles, const int32 j, const TRigidTransform<T, d>& LocalToWorld1, const TRigidTransform<T, d>& LocalToWorld2, const T Thickness, TRigidBodyContactConstraintPGS<T, d>& Constraint)
 {
 	if(InParticles.CollisionParticles(Constraint.ParticleIndex))
 	{
@@ -876,7 +876,7 @@ void UpdateLevelsetConstraintHelperCCD(const TRigidParticles<T, d>& InParticles,
 }
 
 template<class T, int d>
-void UpdateLevelsetConstraintHelper(const TRigidParticles<T, d>& InParticles, const int32 j, const TRigidTransform<T, d> LocalToWorld1, const TRigidTransform<T, d> LocalToWorld2, const T Thickness, TRigidBodyContactConstraintPGS<T, d>& Constraint)
+void UpdateLevelsetConstraintHelper(const TRigidParticles<T, d>& InParticles, const int32 j, const TRigidTransform<T, d>& LocalToWorld1, const TRigidTransform<T, d>& LocalToWorld2, const T Thickness, TRigidBodyContactConstraintPGS<T, d>& Constraint)
 {
 	if(InParticles.CollisionParticles(Constraint.ParticleIndex))
 	{
diff --git a/Engine/Source/Runtime/Experimental/GeometryCollectionSimulationCore/Public/GeometryCollection/GeometryCollectionSimulationCoreTypes.h b/Engine/Source/Runtime/Experimental/GeometryCollectionSimulationCore/Public/GeometryCollection/GeometryCollectionSimulationCoreTypes.h
--- a/Engine/Source/Runtime/Experimental/GeometryCollectionSimulationCore/Public/GeometryCollection/GeometryCollectionSimulationCoreTypes.h
+++ b/Engine/Source/Runtime/Experimental/GeometryCollectionSimulationCore/Public/GeometryCollection/GeometryCollectionSimulationCoreTypes.h
@@ -62,7 +62,7 @@ struct FSimulationParameters
 		, bool OwnsTrackIn
 		, bool SimulatingIn
 		, const FFieldSystem* FieldSystemIn
-		, FTransform WorldTransformIn
+		, FTransform& WorldTransformIn
 		, EObjectTypeEnum ObjectTypeIn
 		, bool EnableClusteringIn
 		, int32 MaxClusterLevelIn
