Workaround for warnings caused by no exported sybmols in ChaosSolvers module when INCLUDE_CHAOS is 0

diff --git a/Engine/Source/Runtime/Experimental/ChaosSolvers/Private/ChaosSolversModule.cpp b/Engine/Source/Runtime/Experimental/ChaosSolvers/Private/ChaosSolversModule.cpp
--- a/Engine/Source/Runtime/Experimental/ChaosSolvers/Private/ChaosSolversModule.cpp
+++ b/Engine/Source/Runtime/Experimental/ChaosSolvers/Private/ChaosSolversModule.cpp
@@ -1,10 +1,11 @@
 // Copyright 1998-2018 Epic Games, Inc. All Rights Reserved.
 
+#include "ChaosSolversModule.h"
+#include "CoreMinimal.h"
+
 #if INCLUDE_CHAOS
 
-#include "ChaosSolversModule.h"
 #include "ChaosStats.h"
-#include "CoreMinimal.h"
 #include "ChaosLog.h"
 #include "HAL/PlatformProcess.h"
 #include "Framework/PersistentTask.h"
@@ -302,4 +303,9 @@ bool FChaosScopedPhysicsThreadLock::DidGetLock() const
 
 IMPLEMENT_MODULE(FChaosSolversModule, ChaosSolvers);
 
+#else
+
+// Workaround for module not having any exported symbols
+CHAOSSOLVERS_API int ChaosSolversExportedSymbol = 0;
+
 #endif
