Fix missing or incorrect copyright notices.

#rnx

diff --git a/Engine/Plugins/Editor/MeshEditor/Source/MeshEditor/MeshEditorGeometryCollectionAdapter.cpp b/Engine/Plugins/Editor/MeshEditor/Source/MeshEditor/MeshEditorGeometryCollectionAdapter.cpp
--- a/Engine/Plugins/Editor/MeshEditor/Source/MeshEditor/MeshEditorGeometryCollectionAdapter.cpp
+++ b/Engine/Plugins/Editor/MeshEditor/Source/MeshEditor/MeshEditorGeometryCollectionAdapter.cpp
@@ -1,4 +1,4 @@
-// Copyright 1998-2015 Epic Games, Inc. All Rights Reserved.
+// Copyright 1998-2018 Epic Games, Inc. All Rights Reserved.
 
 #include "MeshEditorGeometryCollectionAdapter.h"
 #include "EditableMesh.h"
diff --git a/Engine/Plugins/Experimental/BlastPlugin/Source/BlastAuthoring/Private/MeshFractureSettings.cpp b/Engine/Plugins/Experimental/BlastPlugin/Source/BlastAuthoring/Private/MeshFractureSettings.cpp
--- a/Engine/Plugins/Experimental/BlastPlugin/Source/BlastAuthoring/Private/MeshFractureSettings.cpp
+++ b/Engine/Plugins/Experimental/BlastPlugin/Source/BlastAuthoring/Private/MeshFractureSettings.cpp
@@ -1,4 +1,5 @@
-// Copyright 1998-2015 Epic Games, Inc. All Rights Reserved.
+// Copyright 1998-2018 Epic Games, Inc. All Rights Reserved.
+
 #include "MeshFractureSettings.h"
 #include "UnrealEdMisc.h"
 
diff --git a/Engine/Plugins/Experimental/BlastPlugin/Source/BlastAuthoring/Public/MeshFractureSettings.h b/Engine/Plugins/Experimental/BlastPlugin/Source/BlastAuthoring/Public/MeshFractureSettings.h
--- a/Engine/Plugins/Experimental/BlastPlugin/Source/BlastAuthoring/Public/MeshFractureSettings.h
+++ b/Engine/Plugins/Experimental/BlastPlugin/Source/BlastAuthoring/Public/MeshFractureSettings.h
@@ -1,4 +1,4 @@
-// Copyright 1998-2015 Epic Games, Inc. All Rights Reserved.
+// Copyright 1998-2018 Epic Games, Inc. All Rights Reserved.
 
 #pragma once
 
diff --git a/Engine/Plugins/Runtime/Oculus/OculusAvatar/Source/Public/OvrAvatarManager.h b/Engine/Plugins/Runtime/Oculus/OculusAvatar/Source/Public/OvrAvatarManager.h
--- a/Engine/Plugins/Runtime/Oculus/OculusAvatar/Source/Public/OvrAvatarManager.h
+++ b/Engine/Plugins/Runtime/Oculus/OculusAvatar/Source/Public/OvrAvatarManager.h
@@ -1,4 +1,7 @@
+// Copyright 1998-2018 Epic Games, Inc. All Rights Reserved.
+
 #pragma once
+
 #include "Containers/Ticker.h"
 #include "OVR_Avatar.h"
 #include "Containers/Map.h"
diff --git a/Engine/Source/Programs/AutomationTool/Gauntlet/SelfTest/Unreal/Gauntlet.SelfTest.TestUnrealBuildParsing.cs b/Engine/Source/Programs/AutomationTool/Gauntlet/SelfTest/Unreal/Gauntlet.SelfTest.TestUnrealBuildParsing.cs
--- a/Engine/Source/Programs/AutomationTool/Gauntlet/SelfTest/Unreal/Gauntlet.SelfTest.TestUnrealBuildParsing.cs
+++ b/Engine/Source/Programs/AutomationTool/Gauntlet/SelfTest/Unreal/Gauntlet.SelfTest.TestUnrealBuildParsing.cs
@@ -1,4 +1,6 @@
-﻿using Gauntlet;
+﻿// Copyright 1998-2018 Epic Games, Inc. All Rights Reserved.
+
+using Gauntlet;
 using System;
 using System.Collections.Generic;
 using System.Diagnostics;
diff --git a/Engine/Source/Runtime/Experimental/GeometryCollectionCore/Private/GeometryCollection/GeometryCollectionAlgo.cpp b/Engine/Source/Runtime/Experimental/GeometryCollectionCore/Private/GeometryCollection/GeometryCollectionAlgo.cpp
--- a/Engine/Source/Runtime/Experimental/GeometryCollectionCore/Private/GeometryCollection/GeometryCollectionAlgo.cpp
+++ b/Engine/Source/Runtime/Experimental/GeometryCollectionCore/Private/GeometryCollection/GeometryCollectionAlgo.cpp
@@ -1,5 +1,4 @@
 // Copyright 1998-2018 Epic Games, Inc. All Rights Reserved.
-// Copyright 1998-2018 Epic Games, Inc. All Rights Reserved.
 
 /*=============================================================================
 	GeometryCollection.cpp: FGeometryCollection methods.
