Fixed FGeometryCollectionSceneProxy after merge from main.

diff --git a/Engine/Source/Runtime/Experimental/GeometryCollectionEngine/Private/GeometryCollection/GeometryCollectionSceneProxy.cpp b/Engine/Source/Runtime/Experimental/GeometryCollectionEngine/Private/GeometryCollection/GeometryCollectionSceneProxy.cpp
--- a/Engine/Source/Runtime/Experimental/GeometryCollectionEngine/Private/GeometryCollection/GeometryCollectionSceneProxy.cpp
+++ b/Engine/Source/Runtime/Experimental/GeometryCollectionEngine/Private/GeometryCollection/GeometryCollectionSceneProxy.cpp
@@ -216,11 +216,9 @@ void FGeometryCollectionSceneProxy::SetDynamicData_RenderThread(FGeometryCollect
 
 FMaterialRenderProxy* FGeometryCollectionSceneProxy::GetMaterial(FMeshElementCollector& Collector, int32 MaterialIndex) const
 {
-	//todo HLR
-#if 0
 	// material for wireframe
 	auto WireframeMaterialInstance = new FColoredMaterialRenderProxy(
-		GEngine->WireframeMaterial ? GEngine->WireframeMaterial->GetRenderProxy(IsSelected()) : nullptr,
+		GEngine->WireframeMaterial ? GEngine->WireframeMaterial->GetRenderProxy() : nullptr,
 		FLinearColor(0, 0.5f, 1.f)
 	);
 	Collector.RegisterOneFrameMaterialProxy(WireframeMaterialInstance);
@@ -228,7 +226,7 @@ FMaterialRenderProxy* FGeometryCollectionSceneProxy::GetMaterial(FMeshElementCol
 	// material for colored bones
 	UMaterial* VertexColorVisualizationMaterial = GEngine->VertexColorMaterial;
 	auto VertexColorVisualizationMaterialInstance = new FColoredMaterialRenderProxy(
-		VertexColorVisualizationMaterial->GetRenderProxy(false, false),
+		VertexColorVisualizationMaterial->GetRenderProxy(),
 		GetSelectionColor(FLinearColor::White, false, false)
 	);
 	Collector.RegisterOneFrameMaterialProxy(VertexColorVisualizationMaterialInstance);
@@ -241,22 +239,19 @@ FMaterialRenderProxy* FGeometryCollectionSceneProxy::GetMaterial(FMeshElementCol
 	}
 	else
 	{
-		MaterialProxy = Materials[MaterialIndex]->GetRenderProxy(IsSelected());
+		MaterialProxy = Materials[MaterialIndex]->GetRenderProxy();
 	}
 
 	if (MaterialProxy == nullptr)
 	{
-		MaterialProxy = UMaterial::GetDefaultMaterial(MD_Surface)->GetRenderProxy(IsSelected(), IsHovered());
+		MaterialProxy = UMaterial::GetDefaultMaterial(MD_Surface)->GetRenderProxy();
 	}
 
 	return MaterialProxy;
-#endif
-	return nullptr;
 }
 
 void FGeometryCollectionSceneProxy::GetDynamicMeshElements(const TArray<const FSceneView*>& Views, const FSceneViewFamily& ViewFamily, uint32 VisibilityMap, FMeshElementCollector& Collector) const
 {
-#if 0
 	QUICK_SCOPE_CYCLE_COUNTER(STAT_GeometryCollectionSceneProxy_GetDynamicMeshElements);
 	if (GetRequiredVertexCount())
 	{
@@ -305,7 +300,7 @@ void FGeometryCollectionSceneProxy::GetDynamicMeshElements(const TArray<const FS
 			// bone selection is already contained in the rendered colors
 			if (ShowBoneColors||ShowSelectedBones)
 			{
-				FMaterialRenderProxy* MaterialRenderProxy = Materials[BoneSelectionMaterialID]->GetRenderProxy(IsSelected());
+				FMaterialRenderProxy* MaterialRenderProxy = Materials[BoneSelectionMaterialID]->GetRenderProxy();
 
 				if (VisibilityMap & (1 << ViewIndex))
 				{
@@ -337,7 +332,6 @@ void FGeometryCollectionSceneProxy::GetDynamicMeshElements(const TArray<const FS
 #endif
 		}
 	}
-#endif
 }
 
 FPrimitiveViewRelevance FGeometryCollectionSceneProxy::GetViewRelevance(const FSceneView* View) const
