FMeshBatchElement.PrimitiveUniformBuffer doesn't own the RHIResource.

Now using the Collector to allocate a one frame resource so this stays alive.

diff --git a/Engine/Plugins/Experimental/GeometryCache/Source/GeometryCache/Private/GeometryCacheSceneProxy.cpp b/Engine/Plugins/Experimental/GeometryCache/Source/GeometryCache/Private/GeometryCacheSceneProxy.cpp
--- a/Engine/Plugins/Experimental/GeometryCache/Source/GeometryCache/Private/GeometryCacheSceneProxy.cpp
+++ b/Engine/Plugins/Experimental/GeometryCache/Source/GeometryCache/Private/GeometryCacheSceneProxy.cpp
@@ -691,7 +691,13 @@ void FGeometryCacheSceneProxy::GetDynamicMeshElements(const TArray<const FSceneV
 					Mesh.bWireframe = bWireframe;
 					Mesh.VertexFactory = &TrackProxy->VertexFactory;
 					Mesh.MaterialRenderProxy = MaterialProxy;
-					BatchElement.PrimitiveUniformBuffer = CreatePrimitiveUniformBufferImmediate(TrackProxy->WorldMatrix * GetLocalToWorld(), GetBounds(), GetLocalBounds(), true, UseEditorDepthTest());
+                    
+                    const FMatrix& LocalToWorldTransform = TrackProxy->WorldMatrix * GetLocalToWorld();
+                    
+                    FDynamicPrimitiveUniformBuffer& DynamicPrimitiveUniformBuffer = Collector.AllocateOneFrameResource<FDynamicPrimitiveUniformBuffer>();
+                    DynamicPrimitiveUniformBuffer.Set(LocalToWorldTransform, LocalToWorldTransform, GetBounds(), GetLocalBounds(), true, false, UseEditorDepthTest());
+                    BatchElement.PrimitiveUniformBuffer = DynamicPrimitiveUniformBuffer.UniformBuffer.GetUniformBufferRHI();
+                    
 					BatchElement.FirstIndex = BatchInfo.StartIndex;
 					BatchElement.NumPrimitives = BatchInfo.NumTriangles;
 					BatchElement.MinVertexIndex = 0;
@@ -911,4 +917,4 @@ void FGeomCacheVertexBuffer::UpdateSize(int32 NewSizeInBytes)
 		FRHIResourceCreateInfo CreateInfo;
 		VertexBufferRHI = RHICreateVertexBuffer(SizeInBytes, BUF_Static, CreateInfo);
 	}
-}
\ No newline at end of file
+}
diff --git a/Engine/Source/Runtime/Experimental/GeometryCollectionEngine/Private/GeometryCollection/GeometryCollectionSceneProxy.cpp b/Engine/Source/Runtime/Experimental/GeometryCollectionEngine/Private/GeometryCollection/GeometryCollectionSceneProxy.cpp
--- a/Engine/Source/Runtime/Experimental/GeometryCollectionEngine/Private/GeometryCollection/GeometryCollectionSceneProxy.cpp
+++ b/Engine/Source/Runtime/Experimental/GeometryCollectionEngine/Private/GeometryCollection/GeometryCollectionSceneProxy.cpp
@@ -283,7 +283,12 @@ void FGeometryCollectionSceneProxy::GetDynamicMeshElements(const TArray<const FS
 					Mesh.bWireframe = bWireframe;
 					Mesh.VertexFactory = &VertexFactory;
 					Mesh.MaterialRenderProxy = MaterialProxies[SectionIndex];
-					BatchElement.PrimitiveUniformBuffer = CreatePrimitiveUniformBufferImmediate(GetLocalToWorld(), GetBounds(), GetLocalBounds(), true, UseEditorDepthTest());
+                    
+                    // Collector owns the uniform buffer
+                    FDynamicPrimitiveUniformBuffer& DynamicPrimitiveUniformBuffer = Collector.AllocateOneFrameResource<FDynamicPrimitiveUniformBuffer>();
+                    DynamicPrimitiveUniformBuffer.Set(GetLocalToWorld(), GetLocalToWorld(), GetBounds(), GetLocalBounds(), true, false, UseEditorDepthTest());
+                    BatchElement.PrimitiveUniformBuffer = DynamicPrimitiveUniformBuffer.UniformBuffer.GetUniformBufferRHI();
+                    
 					BatchElement.FirstIndex = Section.FirstIndex;
 					BatchElement.NumPrimitives = Section.NumTriangles;
 					BatchElement.MinVertexIndex = Section.MinVertexIndex;
@@ -311,7 +316,11 @@ void FGeometryCollectionSceneProxy::GetDynamicMeshElements(const TArray<const FS
 					Mesh.bWireframe = bWireframe;
 					Mesh.VertexFactory = &VertexFactory;
 					Mesh.MaterialRenderProxy = MaterialRenderProxy;
-					BatchElement.PrimitiveUniformBuffer = CreatePrimitiveUniformBufferImmediate(GetLocalToWorld(), GetBounds(), GetLocalBounds(), true, UseEditorDepthTest());
+                    
+                    FDynamicPrimitiveUniformBuffer& DynamicPrimitiveUniformBuffer = Collector.AllocateOneFrameResource<FDynamicPrimitiveUniformBuffer>();
+                    DynamicPrimitiveUniformBuffer.Set(GetLocalToWorld(), GetLocalToWorld(), GetBounds(), GetLocalBounds(), true, false, UseEditorDepthTest());
+                    BatchElement.PrimitiveUniformBuffer = DynamicPrimitiveUniformBuffer.UniformBuffer.GetUniformBufferRHI();
+                    
 					BatchElement.FirstIndex = 0;
 					BatchElement.NumPrimitives = GetRequiredIndexCount() / 3;
 					BatchElement.MinVertexIndex = 0;
