Non-unity build fixes

diff --git a/Engine/Source/Runtime/Experimental/GeometryCollectionEngine/Private/GeometryCollection/PhysicsAssetSimulation.cpp b/Engine/Source/Runtime/Experimental/GeometryCollectionEngine/Private/GeometryCollection/PhysicsAssetSimulation.cpp
--- a/Engine/Source/Runtime/Experimental/GeometryCollectionEngine/Private/GeometryCollection/PhysicsAssetSimulation.cpp
+++ b/Engine/Source/Runtime/Experimental/GeometryCollectionEngine/Private/GeometryCollection/PhysicsAssetSimulation.cpp
@@ -19,9 +19,7 @@
 #include "Rendering/SkeletalMeshRenderData.h"
 #include "SolverObjects/SkeletalMeshPhysicsObject.h"
 
-#include "PxConvexMesh.h"
-#include "PxMaterial.h"
-#include "PxVec3.h"
+#include "PhysXIncludes.h"
 
 #include "DrawDebugHelpers.h"
 #include "Async/ParallelFor.h"
diff --git a/Engine/Source/Runtime/Experimental/GeometryCollectionEngine/Private/GeometryCollection/SkeletalMeshSimulationComponent.cpp b/Engine/Source/Runtime/Experimental/GeometryCollectionEngine/Private/GeometryCollection/SkeletalMeshSimulationComponent.cpp
--- a/Engine/Source/Runtime/Experimental/GeometryCollectionEngine/Private/GeometryCollection/SkeletalMeshSimulationComponent.cpp
+++ b/Engine/Source/Runtime/Experimental/GeometryCollectionEngine/Private/GeometryCollection/SkeletalMeshSimulationComponent.cpp
@@ -20,8 +20,7 @@
 #include "Rendering/SkeletalMeshRenderData.h"
 #include "SolverObjects/SkeletalMeshPhysicsObject.h"
 
-#include "PxConvexMesh.h"
-#include "PxVec3.h"
+#include "PhysXIncludes.h"
 
 #include "DrawDebugHelpers.h"
 #include "Async/ParallelFor.h"
