Fix set distance on fields.

diff --git a/Engine/Source/Runtime/Experimental/FieldSystem/Source/FieldSystemEngine/Private/Field/FieldSystemObjects.cpp b/Engine/Source/Runtime/Experimental/FieldSystem/Source/FieldSystemEngine/Private/Field/FieldSystemObjects.cpp
--- a/Engine/Source/Runtime/Experimental/FieldSystem/Source/FieldSystemEngine/Private/Field/FieldSystemObjects.cpp
+++ b/Engine/Source/Runtime/Experimental/FieldSystem/Source/FieldSystemEngine/Private/Field/FieldSystemObjects.cpp
@@ -130,13 +130,14 @@ UPlaneFalloff::NewEvaluationGraph(TArray<const UFieldNodeBase*>& Nodes) const
 }
 
 UPlaneFalloff* 
-UPlaneFalloff::SetPlaneFalloff(float InMagnitude, float InMinRange, float InMaxRange, float InDefault, FVector InPosition, FVector InNormal, EFieldFalloffType InFalloff)
+UPlaneFalloff::SetPlaneFalloff(float InMagnitude, float InMinRange, float InMaxRange, float InDefault, float InDistance, FVector InPosition, FVector InNormal, EFieldFalloffType InFalloff)
 {
 	this->Magnitude = InMagnitude;
 	this->MinRange = InMinRange;
 	this->MaxRange = InMaxRange;
 	this->Default = InDefault;
 	this->Position = InPosition;
+	this->Distance = InDistance;
 	this->Normal = InNormal;
 	this->Falloff = InFalloff;
 	return this;
diff --git a/Engine/Source/Runtime/Experimental/FieldSystem/Source/FieldSystemEngine/Public/Field/FieldSystemObjects.h b/Engine/Source/Runtime/Experimental/FieldSystem/Source/FieldSystemEngine/Public/Field/FieldSystemObjects.h
--- a/Engine/Source/Runtime/Experimental/FieldSystem/Source/FieldSystemEngine/Public/Field/FieldSystemObjects.h
+++ b/Engine/Source/Runtime/Experimental/FieldSystem/Source/FieldSystemEngine/Public/Field/FieldSystemObjects.h
@@ -302,7 +302,7 @@ public:
 	virtual FFieldNodeBase* NewEvaluationGraph(TArray<const UFieldNodeBase*>& Nodes) const override;
 
 	UFUNCTION(BlueprintPure, Category = "Field", meta = (Magnitude = "1.0", MinRange = "0.0", MaxRange = "1.0"))
-	UPlaneFalloff* SetPlaneFalloff(float Magnitude, float MinRange, float MaxRange, float Default, FVector Position, FVector Normal, EFieldFalloffType Falloff);
+	UPlaneFalloff* SetPlaneFalloff(float Magnitude, float MinRange, float MaxRange, float Default, float Distance, FVector Position, FVector Normal, EFieldFalloffType Falloff);
 
 	UPROPERTY(BlueprintReadWrite, EditAnywhere, Category = "Field")
 	float Magnitude;
