Fix Clustered GeometryCollection allocating more particles than needed, leaving some uninitialized.

diff --git a/Engine/Source/Runtime/Experimental/ChaosSolvers/Private/SolverObjects/GeometryCollectionPhysicsObject.cpp b/Engine/Source/Runtime/Experimental/ChaosSolvers/Private/SolverObjects/GeometryCollectionPhysicsObject.cpp
--- a/Engine/Source/Runtime/Experimental/ChaosSolvers/Private/SolverObjects/GeometryCollectionPhysicsObject.cpp
+++ b/Engine/Source/Runtime/Experimental/ChaosSolvers/Private/SolverObjects/GeometryCollectionPhysicsObject.cpp
@@ -1042,20 +1042,29 @@ void FGeometryCollectionPhysicsObject::CreateRigidBodyCallback(FParticlesType& P
 
 			// num clusters
 			uint32 NumClusters=0;
+			TArray<bool> SubTreeContainsSimulatableParticle;
+			SubTreeContainsSimulatableParticle.SetNum(RecursiveOrder.Num());
 			for (const int32 TransformGroupIndex : RecursiveOrder)
 			{
 				if (Children[TransformGroupIndex].Num() > 0)
 				{
+					SubTreeContainsSimulatableParticle[TransformGroupIndex] = false;
+
 					TArray<uint32> RigidChildren, CollectionChildren;
 					for (const int32 ChildIndex : Children[TransformGroupIndex])
 					{
-						if(RigidBodyID[ChildIndex] != INDEX_NONE || Children[TransformGroupIndex].Num())
+						if(SubTreeContainsSimulatableParticle[ChildIndex])
 						{
 							NumClusters++;
+							SubTreeContainsSimulatableParticle[TransformGroupIndex] = true;
 							break;
 						}
 					}
 				}
+				else
+				{
+					SubTreeContainsSimulatableParticle[TransformGroupIndex] = (RigidBodyID[TransformGroupIndex] != INDEX_NONE);
+				}
 			}
 
 			const int32 ClusterStartIndex = Particles.Size();
