Change the disabled body id color from black to silver when displaying the geometry collection debug draw visualization, so that the ids can better be read in dark areas.

diff --git a/Engine/Source/Runtime/Experimental/GeometryCollectionEngine/Private/GeometryCollection/GeometryCollectionDebugDrawActor.cpp b/Engine/Source/Runtime/Experimental/GeometryCollectionEngine/Private/GeometryCollection/GeometryCollectionDebugDrawActor.cpp
--- a/Engine/Source/Runtime/Experimental/GeometryCollectionEngine/Private/GeometryCollection/GeometryCollectionDebugDrawActor.cpp
+++ b/Engine/Source/Runtime/Experimental/GeometryCollectionEngine/Private/GeometryCollection/GeometryCollectionDebugDrawActor.cpp
@@ -2028,7 +2028,7 @@ void AGeometryCollectionDebugDrawActor::DrawRigidBodyIdNoChecks(const UGeometryC
 	// Retrieve disabled state
 	const bool bIsDisabled = ParticlesData.IsDisabled(TransformIndex);
 	const TManagedArray<int32>& ParentArray = GeometryCollectionComponent->GetParentArray();
-	const FColor& DisabledColor = bIsDisabled ? FColor::Black: Color;
+	const FColor& DisabledColor = bIsDisabled ? FColor::Silver: Color;
 
 	// Draw rigid body id
 	const FString Text = FString::Printf(TEXT("%d"), RigidBodyId);
