BoneSelection and Highlighted bones reset on GC Componenet registration.
Fixed some metadata capitalization.

diff --git a/Engine/Source/Runtime/Experimental/GeometryCollectionEngine/Private/GeometryCollection/GeometryCollectionComponent.cpp b/Engine/Source/Runtime/Experimental/GeometryCollectionEngine/Private/GeometryCollection/GeometryCollectionComponent.cpp
--- a/Engine/Source/Runtime/Experimental/GeometryCollectionEngine/Private/GeometryCollection/GeometryCollectionComponent.cpp
+++ b/Engine/Source/Runtime/Experimental/GeometryCollectionEngine/Private/GeometryCollection/GeometryCollectionComponent.cpp
@@ -1095,6 +1095,11 @@ void UGeometryCollectionComponent::OnRegister()
 {
 	//UE_LOG(UGCC_LOG, Log, TEXT("GeometryCollectionComponent[%p]::OnRegister()[%p]"), this,RestCollection );
 	ResetDynamicCollection();
+
+ 	FScopedColorEdit ColorEdit(this);
+	ColorEdit.ResetBoneSelection();
+	ColorEdit.ResetHighlightedBones();
+
 	Super::OnRegister();
 }
 
@@ -1867,7 +1872,7 @@ void FScopedColorEdit::UpdateBoneColors()
 		}
 		TManagedArray<FLinearColor>& BoneColors = Collection->BoneColor;
 
-		for (int BoneIndex = 0; BoneIndex < Parents.Num(); BoneIndex++)
+		for (int32 BoneIndex = 0, NumBones = Parents.Num() ; BoneIndex < NumBones; ++BoneIndex)
 		{
 			FLinearColor BoneColor = FLinearColor(FColor::Black);
 
diff --git a/Engine/Source/Runtime/Experimental/GeometryCollectionEngine/Public/GeometryCollection/GeometryCollectionComponent.h b/Engine/Source/Runtime/Experimental/GeometryCollectionEngine/Public/GeometryCollection/GeometryCollectionComponent.h
--- a/Engine/Source/Runtime/Experimental/GeometryCollectionEngine/Public/GeometryCollection/GeometryCollectionComponent.h
+++ b/Engine/Source/Runtime/Experimental/GeometryCollectionEngine/Public/GeometryCollection/GeometryCollectionComponent.h
@@ -280,6 +280,7 @@ public:
 	FORCEINLINE void SetRenderStateDirty() { bRenderStateDirty = true; }
 	virtual void BeginPlay() override;
 	virtual void EndPlay(EEndPlayReason::Type ReasonEnd) override;
+
 	//~ Begin UActorComponent Interface. 
 
 
@@ -596,10 +597,10 @@ private:
 	bool IsObjectLoading;
 
 #if WITH_EDITORONLY_DATA
-	UPROPERTY(transient)
+	UPROPERTY(Transient)
 	TArray<int32> SelectedBones;
 
-	UPROPERTY(transient)
+	UPROPERTY(Transient)
 	TArray<int32> HighlightedBones;
 #endif
 
@@ -625,7 +626,7 @@ private:
 #endif
 
 	// Temporary storage for body setup in order to initialise a dummy body instance
-	UPROPERTY(transient)
+	UPROPERTY(Transient)
 	UBodySetup* DummyBodySetup;
 
 #if WITH_EDITORONLY_DATA
