wrap FColorEdit with WITH_EDITOR #jira UE-77821 #rb trivial

diff --git a/Engine/Source/Runtime/Experimental/GeometryCollectionEngine/Private/GeometryCollection/GeometryCollectionComponent.cpp b/Engine/Source/Runtime/Experimental/GeometryCollectionEngine/Private/GeometryCollection/GeometryCollectionComponent.cpp
--- a/Engine/Source/Runtime/Experimental/GeometryCollectionEngine/Private/GeometryCollection/GeometryCollectionComponent.cpp
+++ b/Engine/Source/Runtime/Experimental/GeometryCollectionEngine/Private/GeometryCollection/GeometryCollectionComponent.cpp
@@ -1096,10 +1096,11 @@ void UGeometryCollectionComponent::OnRegister()
 	//UE_LOG(UGCC_LOG, Log, TEXT("GeometryCollectionComponent[%p]::OnRegister()[%p]"), this,RestCollection );
 	ResetDynamicCollection();
 
- 	FScopedColorEdit ColorEdit(this);
+#if WITH_EDITOR
+	FScopedColorEdit ColorEdit(this);
 	ColorEdit.ResetBoneSelection();
 	ColorEdit.ResetHighlightedBones();
-
+#endif
 	Super::OnRegister();
 }
 
