Ensure Disable/Sleep fields do not overwrite material thresholds that were not outputted by field evaluation.

diff --git a/Engine/Source/Runtime/Experimental/ChaosSolvers/Private/SolverObjects/FieldSystemPhysicsObject.cpp b/Engine/Source/Runtime/Experimental/ChaosSolvers/Private/SolverObjects/FieldSystemPhysicsObject.cpp
--- a/Engine/Source/Runtime/Experimental/ChaosSolvers/Private/SolverObjects/FieldSystemPhysicsObject.cpp
+++ b/Engine/Source/Runtime/Experimental/ChaosSolvers/Private/SolverObjects/FieldSystemPhysicsObject.cpp
@@ -437,8 +437,26 @@ void FFieldSystemPhysicsObject::FieldParameterUpdateCallback(Chaos::FPBDRigidsSo
 						Results.AddUninitialized(Particles.Size());
 						for (const ContextIndex Index : IndicesArray)
 						{
-							Results[Index.Sample] = 0.f;
+							const SolverObjectWrapper& ParticleObjectWrapper = SolverObjectMapping[Index.Result];
+							TSerializablePtr<Chaos::TChaosPhysicsMaterial<float>> Material = CurrentSolver->GetPhysicsMaterial(Index.Result);
+							if (ensure(Material) && ParticleObjectWrapper.SolverObject)
+							{
+								const TUniquePtr<Chaos::TChaosPhysicsMaterial<float>>& InstanceMaterial = CurrentSolver->GetPerParticlePhysicsMaterial(Index.Result);
+								if (InstanceMaterial)
+								{
+									Results[Index.Result] = InstanceMaterial->SleepingLinearThreshold;
+								}
+								else
+								{
+									Results[Index.Result] = Material->SleepingLinearThreshold;
+								}
+							}
+							else
+							{
+								Results[Index.Result] = 0.0f;
+							}
 						}
+
 						TArrayView<float> ResultsView(&(Results[0]), Results.Num());
 						static_cast<const FFieldNode<float> *>(Command.RootNode.Get())->Evaluate(Context, ResultsView);
 
@@ -453,16 +471,29 @@ void FFieldSystemPhysicsObject::FieldParameterUpdateCallback(Chaos::FPBDRigidsSo
 							}
 
 							//per instance override
-							if (!CurrentSolver->GetPerParticlePhysicsMaterial(i))
+							if (!CurrentSolver->GetPerParticlePhysicsMaterial(Index.Result))
 							{
-								//don't already have a per instance material, so make one
-								TUniquePtr<TChaosPhysicsMaterial<float>> NewMaterial = MakeUnique<TChaosPhysicsMaterial<float>>(*Material);
-								CurrentSolver->SetPerParticlePhysicsMaterial(i, MoveTemp(NewMaterial));
+								if (Results[i] != Material->SleepingLinearThreshold)
+								{
+									// value changed from shared material, make unique material.
+									TUniquePtr<Chaos::TChaosPhysicsMaterial<float>> NewMaterial = MakeUnique< Chaos::TChaosPhysicsMaterial<float>>(*Material);
+									CurrentSolver->SetPerParticlePhysicsMaterial(Index.Result, MoveTemp(NewMaterial));
+									const TUniquePtr<TChaosPhysicsMaterial<float>>& InstanceMaterial = CurrentSolver->GetPerParticlePhysicsMaterial(i);
+
+									InstanceMaterial->SleepingLinearThreshold = Results[i];
+									InstanceMaterial->SleepingAngularThreshold = Results[i];
+								}
 							}
+							else
+							{
+								const TUniquePtr<TChaosPhysicsMaterial<float>>& InstanceMaterial = CurrentSolver->GetPerParticlePhysicsMaterial(i);
 
-							const TUniquePtr<TChaosPhysicsMaterial<float>>& InstanceMaterial = CurrentSolver->GetPerParticlePhysicsMaterial(i);
-							InstanceMaterial->SleepingLinearThreshold = Results[i];
-							InstanceMaterial->SleepingAngularThreshold = Results[i];
+								if (InstanceMaterial->SleepingLinearThreshold != Results[i])
+								{
+									InstanceMaterial->SleepingLinearThreshold = Results[i];
+									InstanceMaterial->SleepingAngularThreshold = Results[i];
+								}
+							}
 						}
 					}
 				}
@@ -495,8 +526,26 @@ void FFieldSystemPhysicsObject::FieldParameterUpdateCallback(Chaos::FPBDRigidsSo
 						Results.AddUninitialized(Particles.Size());
 						for (const ContextIndex Index : IndicesArray)
 						{
-							Results[Index.Sample] = 0.f;
+							const SolverObjectWrapper& ParticleObjectWrapper = SolverObjectMapping[Index.Result];
+							TSerializablePtr<Chaos::TChaosPhysicsMaterial<float>> Material = CurrentSolver->GetPhysicsMaterial(Index.Result);
+							if (ensure(Material) && ParticleObjectWrapper.SolverObject)
+							{
+								const TUniquePtr<Chaos::TChaosPhysicsMaterial<float>>& InstanceMaterial = CurrentSolver->GetPerParticlePhysicsMaterial(Index.Result);
+								if (InstanceMaterial)
+								{
+									Results[Index.Result] = InstanceMaterial->DisabledLinearThreshold;
+								}
+								else
+								{
+									Results[Index.Result] = Material->DisabledLinearThreshold;
+								}
+							}
+							else
+							{
+								Results[Index.Result] = 0.0f;
+							}
 						}
+
 						TArrayView<float> ResultsView(&(Results[0]), Results.Num());
 						static_cast<const FFieldNode<float> *>(Command.RootNode.Get())->Evaluate(Context, ResultsView);
 
@@ -505,22 +554,36 @@ void FFieldSystemPhysicsObject::FieldParameterUpdateCallback(Chaos::FPBDRigidsSo
 							const int32 i = Index.Result;
 							const SolverObjectWrapper& ParticleObjectWrapper = SolverObjectMapping[i];
 							TSerializablePtr<Chaos::TChaosPhysicsMaterial<float>> Material = CurrentSolver->GetPhysicsMaterial(i);
+						
 							if (!ensure(Material) || !ParticleObjectWrapper.SolverObject)	//question: do we actually need to check for solver object?
 							{
 								continue;
 							}
 
 							//per instance override
-							if (!CurrentSolver->GetPerParticlePhysicsMaterial(i))
+							if (!CurrentSolver->GetPerParticlePhysicsMaterial(Index.Result))
 							{
-								//don't already have a per instance material, so make one
-								TUniquePtr<Chaos::TChaosPhysicsMaterial<float>> NewMaterial = MakeUnique< Chaos::TChaosPhysicsMaterial<float>>(*Material);
-								CurrentSolver->SetPerParticlePhysicsMaterial(i, MoveTemp(NewMaterial));
+								if (Results[i] != Material->DisabledLinearThreshold)
+								{
+									// value changed from shared material, make unique material.
+									TUniquePtr<Chaos::TChaosPhysicsMaterial<float>> NewMaterial = MakeUnique< Chaos::TChaosPhysicsMaterial<float>>(*Material);
+									CurrentSolver->SetPerParticlePhysicsMaterial(Index.Result, MoveTemp(NewMaterial));
+									const TUniquePtr<TChaosPhysicsMaterial<float>>& InstanceMaterial = CurrentSolver->GetPerParticlePhysicsMaterial(i);
+
+									InstanceMaterial->DisabledLinearThreshold = Results[i];
+									InstanceMaterial->DisabledAngularThreshold = Results[i];
+								}
 							}
+							else
+							{
+								const TUniquePtr<TChaosPhysicsMaterial<float>>& InstanceMaterial = CurrentSolver->GetPerParticlePhysicsMaterial(i);
 
-							const TUniquePtr<Chaos::TChaosPhysicsMaterial<float>>& InstanceMaterial = CurrentSolver->GetPerParticlePhysicsMaterial(i);
-							InstanceMaterial->DisabledLinearThreshold = Results[i];
-							InstanceMaterial->DisabledAngularThreshold = Results[i];
+								if (InstanceMaterial->DisabledLinearThreshold != Results[i])
+								{
+									InstanceMaterial->DisabledLinearThreshold = Results[i];
+									InstanceMaterial->DisabledAngularThreshold = Results[i];
+								}
+							}
 						}
 					}
 				}
