Copying //UE4/Dev-Physics to Dev-Main (//UE4/Dev-Main) [at] 7703071

diff --git a/Engine/Source/Runtime/Experimental/Chaos/Private/Chaos/PBDCollisionConstraint.cpp b/Engine/Source/Runtime/Experimental/Chaos/Private/Chaos/PBDCollisionConstraint.cpp
--- a/Engine/Source/Runtime/Experimental/Chaos/Private/Chaos/PBDCollisionConstraint.cpp
+++ b/Engine/Source/Runtime/Experimental/Chaos/Private/Chaos/PBDCollisionConstraint.cpp
@@ -4,42 +4,24 @@
 
 #include "Chaos/BoundingVolume.h"
 #include "Chaos/BoundingVolumeHierarchy.h"
+#include "Chaos/ChaosPerfTest.h"
 #include "Chaos/Defines.h"
+#include "Chaos/ImplicitObjectUnion.h"
+#include "Chaos/Levelset.h"
 #include "Chaos/Pair.h"
 #include "Chaos/PBDConstraintGraph.h"
 #include "Chaos/Sphere.h"
 #include "Chaos/Transform.h"
-#include "ChaosStats.h"
 #include "ChaosLog.h"
-#include "ProfilingDebugging/ScopedTimers.h"
-#include "Chaos/ImplicitObjectUnion.h"
-#include "Chaos/Levelset.h"
-#include "Chaos/PBDCollisionTypes.h"
-#include "Chaos/ChaosPerfTest.h"
+#include "ChaosStats.h"
 #include "Containers/Queue.h"
+#include "ProfilingDebugging/ScopedTimers.h"
+#include "Chaos/PBDRigidsSOAs.h"
 
 #if INTEL_ISPC
-#if USING_CODE_ANALYSIS
-    MSVC_PRAGMA( warning( push ) )
-    MSVC_PRAGMA( warning( disable : ALL_CODE_ANALYSIS_WARNINGS ) )
-#endif    // USING_CODE_ANALYSIS
-
-#ifdef __clang__
-#pragma clang diagnostic push
-#pragma clang diagnostic ignored "-Wnonportable-include-path"
-#endif
-
 #include "PBDCollisionConstraint.ispc.generated.h"
-
-#ifdef __clang__
-#pragma clang diagnostic pop
 #endif
 
-#if USING_CODE_ANALYSIS
-    MSVC_PRAGMA( warning( pop ) )
-#endif    // USING_CODE_ANALYSIS
-#endif    // INTEL_ISPC
-
 int32 CollisionParticlesBVHDepth = 4;
 FAutoConsoleVariableRef CVarCollisionParticlesBVHDepth(TEXT("p.CollisionParticlesBVHDepth"), CollisionParticlesBVHDepth, TEXT("The maximum depth for collision particles bvh"));
 
@@ -63,13 +45,17 @@ namespace Chaos
 
 namespace Chaos
 {
-template<ECollisionUpdateType, class T_PARTICLES, typename T, int d>
-void UpdateConstraintImp(const T_PARTICLES& InParticles, const TImplicitObject<T, d>& ParticleObject, const TRigidTransform<T, d>& ParticleTM, const TImplicitObject<T, d>& LevelsetObject, const TRigidTransform<T, d>& LevelsetTM, const T Thickness, TRigidBodyContactConstraint<T, d>& Constraint);
+template<ECollisionUpdateType, typename T, int d>
+void UpdateConstraintImp2(const TRigidTransform<T, d>& ParticleTM, const TImplicitObject<T, d>& LevelsetObject, const TRigidTransform<T, d>& LevelsetTM, const T Thickness, TRigidBodyContactConstraint<T, d>& Constraint);
 
 template<typename T, int d>
-TPBDCollisionConstraint<T, d>::TPBDCollisionConstraint(TPBDRigidParticles<T, d>& InParticles, const TArray<int32>& InIndices, TArrayCollectionArray<bool>& Collided, const TArrayCollectionArray<TSerializablePtr<TChaosPhysicsMaterial<T>>>& InPerParticleMaterials, const int32 PairIterations /*= 1*/, const T Thickness /*= (T)0*/)
+TPBDCollisionConstraint<T, d>::TPBDCollisionConstraint(const TPBDRigidsSOAs<T,d>& InParticles, TArrayCollectionArray<bool>& Collided, const TArrayCollectionArray<TSerializablePtr<TChaosPhysicsMaterial<T>>>& InPerParticleMaterials, const int32 PairIterations /*= 1*/, const T Thickness /*= (T)0*/)
+#if CHAOS_PARTICLEHANDLE_TODO
 	: SpatialAccelerationResource(InPlace, InParticles)
 	, SpatialAccelerationResource2(InPlace, InParticles, ConstraintBPBVHDepth)
+#endif
+	: Particles(InParticles)
+	, SpatialAccelerationResource(InPlace, Particles.GetNonDisabledView())
 	, MCollided(Collided)
 	, MPhysicsMaterials(InPerParticleMaterials)
 	, MPairIterations(PairIterations)
@@ -80,7 +66,7 @@ TPBDCollisionConstraint<T, d>::TPBDCollisionConstraint(TPBDRigidParticles<T, d>&
 }
 
 template<typename T, int d>
-void TPBDCollisionConstraint<T, d>::Reset(const TPBDRigidParticles<T, d>& InParticles, const TArray<int32>& InIndices)
+void TPBDCollisionConstraint<T, d>::Reset(/*const TPBDRigidParticles<T, d>& InParticles, const TArray<int32>& InIndices*/)
 {
 	Constraints.Empty();
 	MAngularFriction = 0;
@@ -88,23 +74,24 @@ void TPBDCollisionConstraint<T, d>::Reset(const TPBDRigidParticles<T, d>& InPart
 }
 
 template<typename T, int d>
-void TPBDCollisionConstraint<T, d>::UpdatePositionBasedState(const TPBDRigidParticles<T, d>& InParticles, const TArray<int32>& InIndices, const T Dt)
+void TPBDCollisionConstraint<T, d>::UpdatePositionBasedState(/*const TPBDRigidParticles<T, d>& InParticles, const TArray<int32>& InIndices,*/ const T Dt)
 {
-	Reset(InParticles, InIndices);
+	Reset();
 
 #if !UE_BUILD_SHIPPING
 	if (PendingHierarchyDump)
 	{
-		ComputeConstraints<true>(InParticles, InIndices, Dt);
+		ComputeConstraints<true>(Dt);
 	}
 	else
 #endif
 	{
-		ComputeConstraints(InParticles, InIndices, Dt);
+		ComputeConstraints(Dt);
 	}
 }
 
 
+#if CHAOS_PARTICLEHANDLE_TODO
 template <typename T, int d>
 const ISpatialAcceleration<T, d>& TPBDCollisionConstraint<T, d>::GetSpatialAcceleration() const
 {
@@ -117,10 +104,12 @@ const ISpatialAcceleration<T, d>& TPBDCollisionConstraint<T, d>::GetSpatialAccel
 		return SpatialAccelerationResource.GetRead();
 	}
 }
+#endif
 
 template <typename T, int d>
 void TPBDCollisionConstraint<T, d>::ReleaseSpatialAcceleration() const
 {
+#if CHAOS_PARTICLEHANDLE_TODO
 	if (BPTreeOfGrids)
 	{
 		SpatialAccelerationResource2.ReleaseRead();
@@ -129,11 +118,14 @@ void TPBDCollisionConstraint<T, d>::ReleaseSpatialAcceleration() const
 	{
 		SpatialAccelerationResource.ReleaseRead();
 	}
+#endif
+	SpatialAccelerationResource.ReleaseRead();
 }
 
 template <typename T, int d>
 void TPBDCollisionConstraint<T, d>::SwapSpatialAcceleration()
 {
+#if CHAOS_PARTICLEHANDLE_TODO
 	if (BPTreeOfGrids)
 	{
 		SpatialAccelerationResource2.Swap();
@@ -142,30 +134,32 @@ void TPBDCollisionConstraint<T, d>::SwapSpatialAcceleration()
 	{
 		SpatialAccelerationResource.Swap();
 	}
+#endif
+	SpatialAccelerationResource.Swap();
 }
 
 #if !UE_BUILD_SHIPPING
 
-int32 BucketSizes[] = { 0,1,4,8,16,32,64,128,512,MAX_int32 };
+int32 BucketSizes2[] = { 0,1,4,8,16,32,64,128,512,MAX_int32 };
 
 
 template <bool bGatherStats>
-struct FStatHelper
+struct FStatHelper2
 {
 	int32 MaxCount;
 
-	FStatHelper() {}
+	FStatHelper2() {}
 	void Record(int32 Count) {}
 	FString ToString() const { return TEXT(""); }
 };
 
 template <>
-struct FStatHelper<true>
+struct FStatHelper2<true>
 {
-	int32 BucketCount[ARRAY_COUNT(BucketSizes)];
+	int32 BucketCount[ARRAY_COUNT(BucketSizes2)];
 	int32 MaxCount;
 
-	FStatHelper()
+	FStatHelper2()
 	{
 		FMemory::Memset(BucketCount, 0, sizeof(BucketCount));
 		MaxCount = 0;
@@ -173,9 +167,9 @@ struct FStatHelper<true>
 
 	void Record(int32 Count)
 	{
-		for (int32 BucketIdx = 1; BucketIdx < ARRAY_COUNT(BucketSizes); ++BucketIdx)
+		for (int32 BucketIdx = 1; BucketIdx < ARRAY_COUNT(BucketSizes2); ++BucketIdx)
 		{
-			if (Count >= BucketSizes[BucketIdx - 1] && Count < BucketSizes[BucketIdx])
+			if (Count >= BucketSizes2[BucketIdx - 1] && Count < BucketSizes2[BucketIdx])
 			{
 				++BucketCount[BucketIdx];
 			}
@@ -200,16 +194,16 @@ struct FStatHelper<true>
 		}
 
 		const float CountPerChar = MaxBucketCount / 20.f;
-		for (int32 Idx = 1; Idx < ARRAY_COUNT(BucketSizes); ++Idx)
+		for (int32 Idx = 1; Idx < ARRAY_COUNT(BucketSizes2); ++Idx)
 		{
 			int32 NumChars = BucketCount[Idx] / CountPerChar;
-			if (Idx < ARRAY_COUNT(BucketSizes) - 1)
+			if (Idx < ARRAY_COUNT(BucketSizes2) - 1)
 			{
-				OutLog += FString::Printf(TEXT("\t[%4d - %4d) (%4d) |"), BucketSizes[Idx - 1], BucketSizes[Idx], BucketCount[Idx]);
+				OutLog += FString::Printf(TEXT("\t[%4d - %4d) (%4d) |"), BucketSizes2[Idx - 1], BucketSizes2[Idx], BucketCount[Idx]);
 			}
 			else
 			{
-				OutLog += FString::Printf(TEXT("\t[%4d -  inf) (%4d) |"), BucketSizes[Idx - 1], BucketCount[Idx]);
+				OutLog += FString::Printf(TEXT("\t[%4d -  inf) (%4d) |"), BucketSizes2[Idx - 1], BucketCount[Idx]);
 			}
 			for (int32 Count = 0; Count < NumChars; ++Count)
 			{
@@ -224,58 +218,47 @@ struct FStatHelper<true>
 
 #endif
 
-DECLARE_CYCLE_STAT(TEXT("ComputeConstraints"), STAT_ComputeConstraints, STATGROUP_Chaos);
-DECLARE_CYCLE_STAT(TEXT("ComputeConstraintsNP"), STAT_ComputeConstraintsNP, STATGROUP_Chaos);
+DECLARE_CYCLE_STAT(TEXT("ComputeConstraints2"), STAT_ComputeConstraints2, STATGROUP_Chaos);
+DECLARE_CYCLE_STAT(TEXT("ComputeConstraintsNP2"), STAT_ComputeConstraintsNP2, STATGROUP_Chaos);
 //DECLARE_CYCLE_STAT(TEXT("ComputeConstraintsNP2"), STAT_ComputeConstraintsNP2, STATGROUP_Chaos);
 //DECLARE_CYCLE_STAT(TEXT("ComputeConstraintsNP3"), STAT_ComputeConstraintsNP3, STATGROUP_Chaos);
 
-int32 ComputeConstraintsUseAny = 1;
-FAutoConsoleVariableRef CVarComputeConstraintsUseAny(TEXT("p.ComputeConstraintsUseAny"), ComputeConstraintsUseAny, TEXT(""));
+int32 ComputeConstraintsUseAny2 = 1;
+FAutoConsoleVariableRef CVarComputeConstraintsUseAny2(TEXT("p.ComputeConstraintsUseAny2"), ComputeConstraintsUseAny2, TEXT(""));
+
 
 template <typename T, int d>
 template <typename SPATIAL_ACCELERATION, bool bGatherStats>
-void TPBDCollisionConstraint<T, d>::ComputeConstraintsHelper(const TPBDRigidParticles<T, d>& InParticles, const TArray<int32>& InIndices, T Dt, const SPATIAL_ACCELERATION& SpatialAcceleration)
+void TPBDCollisionConstraint<T, d>::ComputeConstraintsHelper(T Dt, const SPATIAL_ACCELERATION& SpatialAcceleration)
 {
-	SCOPE_CYCLE_COUNTER(STAT_ComputeConstraints);
-	CHAOS_SCOPED_TIMER(ComputeConstraints);
+	SCOPE_CYCLE_COUNTER(STAT_ComputeConstraints2);
+	CHAOS_SCOPED_TIMER(ComputeConstraints2);
 	if (!EnableCollisions) return;
 	// Broad phase
 
 #if !UE_BUILD_SHIPPING
-	FStatHelper<bGatherStats> BroadphasePotentials;
-	FStatHelper<bGatherStats> NarrowPhaseSkipped;
-	FStatHelper<bGatherStats> NarrowPhasePerformed;
-	FStatHelper<bGatherStats> NarrowPhaseRejected;
-	FStatHelper<bGatherStats> BoundsDistribution;
+	FStatHelper2<bGatherStats> BroadphasePotentials;
+	FStatHelper2<bGatherStats> NarrowPhaseSkipped;
+	FStatHelper2<bGatherStats> NarrowPhasePerformed;
+	FStatHelper2<bGatherStats> NarrowPhaseRejected;
+	FStatHelper2<bGatherStats> BoundsDistribution;
 	int32 SimulatedParticles = 0;
 #endif
 
 	{
-		CHAOS_SCOPED_TIMER(ComputeConstraintsBP);
-		SpatialAcceleration.Reinitialize((const TArray<uint32>&)InIndices, true, Dt * BoundsThicknessMultiplier); //todo(ocohen): should we pass MThickness into this structure?
+		CHAOS_SCOPED_TIMER(ComputeConstraintsBP2);
+		SpatialAcceleration.Reinitialize(Particles.GetNonDisabledView(), true, Dt * BoundsThicknessMultiplier); //todo(ocohen): should we pass MThickness into this structure?
 	}
 	{
-		SCOPE_CYCLE_COUNTER(STAT_ComputeConstraintsNP);
+		SCOPE_CYCLE_COUNTER(STAT_ComputeConstraintsNP2);
 
 
 		// Narrow phase
-		CHAOS_SCOPED_TIMER(ComputeConstraints_NP);
+		CHAOS_SCOPED_TIMER(ComputeConstraints_NP2);
 
 		TQueue<TRigidBodyContactConstraint<T, d>, EQueueMode::Mpsc> Queue;	//todo(ocohen): use per thread buffer instead, need better support than ParallelFor for this
-		PhysicsParallelFor(InIndices.Num(), [&](int32 ActiveIdxIdx)
+		Particles.GetNonDisabledDynamicView().ParallelFor([&](auto& Particle1, int32 ActiveIdxIdx)
 		{
-			const int32 Body1Index = InIndices[ActiveIdxIdx];
-			if (!ensure(!InParticles.Disabled(Body1Index)))
-			{
-				//never seen this case but just in case since refactor is fresh. If no one hits this ensure just remove the branch
-				return;
-			}
-
-			if (InParticles.InvM(Body1Index) == 0)
-			{
-				return;
-			}
-
 #if !UE_BUILD_SHIPPING
 			if (bGatherStats)
 			{
@@ -285,15 +268,19 @@ void TPBDCollisionConstraint<T, d>::ComputeConstraintsHelper(const TPBDRigidPart
 
 			//SCOPE_CYCLE_COUNTER(STAT_ComputeConstraintsNP2);
 
-			TArray<int32> PotentialIntersections;
+			TArray<TGeometryParticleHandle<T,d>*> PotentialIntersections;
 
-			const bool bBody1Bounded = HasBoundingBox(InParticles, Body1Index);
+			const bool bBody1Bounded = HasBoundingBox(Particle1);
 
-			TBox<T, d> Box1;
-			const T Box1Thickness = ComputeThickness(InParticles, Dt, Body1Index).Size();
+			const T Box1Thickness = ComputeThickness(Particle1, Dt).Size();
 			if (bBody1Bounded)
 			{
-				Box1 = SpatialAcceleration.GetWorldSpaceBoundingBox(InParticles, Body1Index);
+#if CHAOS_PARTICLEHANDLE_TODO
+				const TBox <T, d> Box1 = SpatialAcceleration.GetWorldSpaceBoundingBox(Particle1);
+#else
+				const TBox <T, d> Box1 = ComputeWorldSpaceBoundingBox(Particle1);	//NOTE: this ignores the velocity expansion which is wrong
+#endif
+
 #if !UE_BUILD_SHIPPING
 				if (bGatherStats)
 				{
@@ -302,8 +289,7 @@ void TPBDCollisionConstraint<T, d>::ComputeConstraintsHelper(const TPBDRigidPart
 					BoundsDistribution.Record(Box1.Extents().GetMax());
 				}
 #endif
-
-				PotentialIntersections = SpatialAcceleration.FindAllIntersections(InParticles, Body1Index);
+				PotentialIntersections = SpatialAcceleration.FindAllIntersections(Box1);
 			}
 			else
 			{
@@ -319,53 +305,53 @@ void TPBDCollisionConstraint<T, d>::ComputeConstraintsHelper(const TPBDRigidPart
 			const int32 NumPotentials = PotentialIntersections.Num();
 			for (int32 i = 0; i < NumPotentials; ++i)
 			{
-				const int32 Body2Index = PotentialIntersections[i];
+				auto& Particle2 = *PotentialIntersections[i];
+				const TGenericParticleHandle<T, d> Particle2Generic(&Particle2);
 				// Collision group culling...
 				// CollisionGroup == 0 : Collide_With_Everything
 				// CollisionGroup == INDEX_NONE : Disabled collisions
 				// CollisionGroup_A != CollisionGroup_B : Skip Check
-				if(InParticles.Disabled(Body2Index))
-				{
-					continue;
-				}
 				
-				if (InParticles.CollisionGroup(Body1Index) == INDEX_NONE || InParticles.CollisionGroup(Body2Index) == INDEX_NONE)
+				if (Particle1.CollisionGroup() == INDEX_NONE || Particle2Generic->CollisionGroup() == INDEX_NONE)
 				{
 					continue;
 				}
-				if (InParticles.CollisionGroup(Body1Index) && InParticles.CollisionGroup(Body2Index) && InParticles.CollisionGroup(Body1Index) != InParticles.CollisionGroup(Body2Index))
+				if (Particle1.CollisionGroup() && Particle2Generic->CollisionGroup() && Particle1.CollisionGroup() != Particle2Generic->CollisionGroup())
 				{
 					continue;
 				}
 
-				if (!InParticles.Geometry(Body1Index) && !InParticles.Geometry(Body2Index))
+				if (!Particle1.Geometry() && !Particle2.Geometry())
 				{
 					continue;
 				}
 
-				const bool bBody2Bounded = HasBoundingBox(InParticles, Body2Index);
+				const bool bBody2Bounded = HasBoundingBox(Particle2);
 
-				if (Body1Index == Body2Index || ((bBody1Bounded == bBody2Bounded) && InParticles.InvM(Body1Index) && InParticles.InvM(Body2Index) && Body2Index > Body1Index))
+				if (Particle1.Handle() == Particle2.Handle())
 				{
-					//if both are dynamic, assume index order matters
 					continue;
 				}
-
-				if (bBody1Bounded && bBody2Bounded)
+				
+				if(bBody1Bounded == bBody2Bounded && Particle2.AsDynamic())
 				{
-					const TBox<T, d>& Box2 = SpatialAcceleration.GetWorldSpaceBoundingBox(InParticles, Body2Index);
-
-					if (!Box1.Intersects(Box2))
+					//if both are dynamic, assume index order matters
+#if CHAOS_DETERMINISTIC
+					if(Particle2.ParticleID() > Particle1.ParticleID())
+#else
+					//not deterministic just use memory address to avoid pair duplication
+					if (Particle2.Handle() > Particle1.Handle())
+#endif
 					{
 						continue;
 					}
 				}
 				//				++TotalTested;
 
-				const TVector<T, d> Box2Thickness = ComputeThickness(InParticles, Dt, Body2Index);
+				const TVector<T, d> Box2Thickness = Particle2.AsDynamic() ? ComputeThickness(*Particle2.AsDynamic(), Dt) : TVector<T, d>(0);
 				const T UseThickness = FMath::Max(Box1Thickness, Box2Thickness.Size());// + MThickness
 
-				auto Constraint = ComputeConstraint(InParticles, Body1Index, Body2Index, UseThickness);
+				auto Constraint = ComputeConstraint(Particle1.Handle(), Particle2.Handle(), UseThickness);
 
 				//if (true || !InParticles.Geometry(Body1Index)->HasBoundingBox() || !InParticles.Geometry(Body2Index)->HasBoundingBox())
 				{
@@ -378,13 +364,13 @@ void TPBDCollisionConstraint<T, d>::ComputeConstraintsHelper(const TPBDRigidPart
 						++CountNP;
 					}
 #endif
-					if (ComputeConstraintsUseAny)
+					if (ComputeConstraintsUseAny2)
 					{
-						UpdateConstraint<ECollisionUpdateType::Any>(InParticles, UseThickness, Constraint);
+						UpdateConstraint<ECollisionUpdateType::Any>(UseThickness, Constraint);
 					}
 					else
 					{
-						UpdateConstraint<ECollisionUpdateType::Deepest>(InParticles, UseThickness, Constraint);
+						UpdateConstraint<ECollisionUpdateType::Deepest>(UseThickness, Constraint);
 					}
 
 					if (Constraint.Phi < UseThickness)
@@ -398,6 +384,7 @@ void TPBDCollisionConstraint<T, d>::ComputeConstraintsHelper(const TPBDRigidPart
 #endif
 					}
 				}
+
 			}
 
 #if !UE_BUILD_SHIPPING
@@ -415,10 +402,15 @@ void TPBDCollisionConstraint<T, d>::ComputeConstraintsHelper(const TPBDRigidPart
 		}
 	}
 
+#if CHAOS_DETERMINISTIC
+	//todo: sort constraints
+#endif
+
 #if !UE_BUILD_SHIPPING
 	FString OutLog;
 	if (bGatherStats)
 	{
+#if CHAOS_PARTICLEHANDLE_TODO
 		const float NumParticles = InParticles.Size();
 		OutLog = FString::Printf(TEXT("ComputeConstraints stats:\n"
 			"Total Particles:%d\nSimulated Particles:%d (%.2f%%)\n"
@@ -455,6 +447,7 @@ void TPBDCollisionConstraint<T, d>::ComputeConstraintsHelper(const TPBDRigidPart
 
 		UE_LOG(LogChaos, Warning, TEXT("%s"), *OutLog);
 		PendingHierarchyDump = 0;
+#endif
 	}
 #endif
 	//	if (TotalTested > 0)
@@ -465,8 +458,9 @@ void TPBDCollisionConstraint<T, d>::ComputeConstraintsHelper(const TPBDRigidPart
 
 template<typename T, int d>
 template <bool bGatherStats>
-void TPBDCollisionConstraint<T, d>::ComputeConstraints(const TPBDRigidParticles<T, d>& InParticles, const TArray<int32>& InIndices, T Dt)
+void TPBDCollisionConstraint<T, d>::ComputeConstraints(/*const TPBDRigidParticles<T, d>& InParticles, const TArray<int32>& InIndices,*/ T Dt)
 {
+#if CHAOS_PARTICLEHANDLE_TODO
 	if (BPTreeOfGrids)
 	{
 		ComputeConstraintsHelper<decltype(SpatialAccelerationResource2.GetWritable()), bGatherStats>(InParticles, InIndices, Dt, SpatialAccelerationResource2.GetWritable());
@@ -475,11 +469,15 @@ void TPBDCollisionConstraint<T, d>::ComputeConstraints(const TPBDRigidParticles<
 	{
 		ComputeConstraintsHelper<decltype(SpatialAccelerationResource.GetWritable()), bGatherStats>(InParticles, InIndices, Dt, SpatialAccelerationResource.GetWritable());
 	}
+#endif
+
+	ComputeConstraintsHelper<decltype(SpatialAccelerationResource.GetWritable()), bGatherStats>(Dt, SpatialAccelerationResource.GetWritable());
 }
 
 template<typename T, int d>
-void TPBDCollisionConstraint<T, d>::RemoveConstraints(const TSet<uint32>& RemovedParticles)
+void TPBDCollisionConstraint<T, d>::RemoveConstraints(const TSet<TGeometryParticleHandle<T, d>*>& RemovedParticles)
 {
+#if CHAOS_PARTICLEHANDLE_TODO
 	SpatialAccelerationResource.GetWritable().RemoveElements(RemovedParticles.Array());
 	for (int32 i = 0; i < Constraints.Num(); ++i)
 	{
@@ -490,15 +488,17 @@ void TPBDCollisionConstraint<T, d>::RemoveConstraints(const TSet<uint32>& Remove
 			i--;
 		}
 	}
+#endif
 }
 
-DECLARE_CYCLE_STAT(TEXT("UpdateConstraints"), STAT_UpdateConstraints, STATGROUP_Chaos);
+DECLARE_CYCLE_STAT(TEXT("UpdateConstraints2"), STAT_UpdateConstraints2, STATGROUP_Chaos);
 
 template <typename T, int d>
 template<typename SPATIAL_ACCELERATION>
-void TPBDCollisionConstraint<T, d>::UpdateConstraintsHelper(const TPBDRigidParticles<T, d>& InParticles, const TArray<int32>& InIndices, T Dt, const TSet<uint32>& AddedParticles, const TArray<uint32>& ActiveParticles, SPATIAL_ACCELERATION& SpatialAcceleration)
+void TPBDCollisionConstraint<T, d>::UpdateConstraintsHelper(/*const TPBDRigidParticles<T, d>& InParticles, const TArray<int32>& InIndices,*/ T Dt, const TSet<TGeometryParticleHandle<T, d>*>& AddedParticles, SPATIAL_ACCELERATION& SpatialAcceleration)
 {
-	SCOPE_CYCLE_COUNTER(STAT_UpdateConstraints);
+#if CHAOS_PARTICLEHANDLE_TODO
+	SCOPE_CYCLE_COUNTER(STAT_UpdateConstraints2);
 	double Time = 0;
 	FDurationTimer Timer(Time);
 
@@ -632,14 +632,16 @@ void TPBDCollisionConstraint<T, d>::UpdateConstraintsHelper(const TPBDRigidParti
 
 	Timer.Stop();
 	UE_LOG(LogChaos, Verbose, TEXT("\tPBDCollisionConstraint Update %d Constraints with Potential Collisions %f"), Constraints.Num(), Time);
+#endif
 }
 
-DECLARE_CYCLE_STAT(TEXT("Reconcile Updated Constraints"), STAT_ReconcileConstraints, STATGROUP_Chaos);
+DECLARE_CYCLE_STAT(TEXT("Reconcile Updated Constraints"), STAT_ReconcileConstraints2, STATGROUP_Chaos);
 template<typename T, int d>
-void TPBDCollisionConstraint<T, d>::UpdateConstraints(const TPBDRigidParticles<T, d>& InParticles, const TArray<int32>& InIndices, T Dt, const TSet<uint32>& AddedParticles, const TArray<uint32>& ActiveParticles)
+void TPBDCollisionConstraint<T, d>::UpdateConstraints(/*const TPBDRigidParticles<T, d>& InParticles, const TArray<int32>& InIndices,*/ T Dt, const TSet<TGeometryParticleHandle<T, d>*>& AddedParticles)
 {
+#if CHAOS_PARTICLEHANDLE_TODO
 	{
-		SCOPE_CYCLE_COUNTER(STAT_ReconcileConstraints);
+		SCOPE_CYCLE_COUNTER(STAT_ReconcileConstraints2);
 
 		// Updating post-clustering, we will have invalid constraints
 		int32 NumRemovedConstraints = 0;
@@ -668,10 +670,12 @@ void TPBDCollisionConstraint<T, d>::UpdateConstraints(const TPBDRigidParticles<T
 	{
 		UpdateConstraintsHelper(InParticles, InIndices, Dt, AddedParticles, ActiveParticles, SpatialAccelerationResource.GetWritable());
 	}
+#endif
 }
 
+// @todo(ccaulfield): This is duplicated in JointConstraints - move to a utility file
 template<typename T>
-PMatrix<T, 3, 3> ComputeFactorMatrix(const TVector<T, 3>& V, const PMatrix<T, 3, 3>& M, const T& Im)
+PMatrix<T, 3, 3> ComputeFactorMatrix3(const TVector<T, 3>& V, const PMatrix<T, 3, 3>& M, const T& Im)
 {
 	// Rigid objects rotational contribution to the impulse.
 	// Vx*M*VxT+Im
@@ -686,25 +690,28 @@ PMatrix<T, 3, 3> ComputeFactorMatrix(const TVector<T, 3>& V, const PMatrix<T, 3,
 }
 
 template<typename T, int d>
-TVector<T, d> GetEnergyClampedImpulse(const TPBDRigidParticles<T, d>& InParticles, const TRigidBodyContactConstraint<T, d>& Constraint, const TVector<T, d>& Impulse, const TVector<T, d>& VectorToPoint1, const TVector<T, d>& VectorToPoint2, const TVector<T, d>& Velocity1, const TVector<T, d>& Velocity2)
+TVector<T, d> GetEnergyClampedImpulse(const TRigidBodyContactConstraint<T, d>& Constraint, const TVector<T, d>& Impulse, const TVector<T, d>& VectorToPoint1, const TVector<T, d>& VectorToPoint2, const TVector<T, d>& Velocity1, const TVector<T, d>& Velocity2)
 {
+	TPBDRigidParticleHandle<T, d>* PBDRigid0 = Constraint.Particle->AsDynamic();
+	TPBDRigidParticleHandle<T, d>* PBDRigid1 = Constraint.Levelset->AsDynamic();
+
 	TVector<T, d> Jr0, Jr1, IInvJr0, IInvJr1;
 	T ImpulseRatioNumerator0 = 0, ImpulseRatioNumerator1 = 0, ImpulseRatioDenom0 = 0, ImpulseRatioDenom1 = 0;
 	T ImpulseSize = Impulse.SizeSquared();
-	TVector<T, d> KinematicVelocity = !InParticles.InvM(Constraint.ParticleIndex) ? Velocity1 : !InParticles.InvM(Constraint.LevelsetIndex) ? Velocity2 : TVector<T, d>(0);
-	if (InParticles.InvM(Constraint.ParticleIndex))
+	TVector<T, d> KinematicVelocity = !PBDRigid0 ? Velocity1 : !PBDRigid1 ? Velocity2 : TVector<T, d>(0);
+	if (PBDRigid0)
 	{
 		Jr0 = TVector<T, d>::CrossProduct(VectorToPoint1, Impulse);
-		IInvJr0 = InParticles.Q(Constraint.ParticleIndex).RotateVector(InParticles.InvI(Constraint.ParticleIndex) * InParticles.Q(Constraint.ParticleIndex).UnrotateVector(Jr0));
-		ImpulseRatioNumerator0 = TVector<T, d>::DotProduct(Impulse, InParticles.V(Constraint.ParticleIndex) - KinematicVelocity) + TVector<T, d>::DotProduct(IInvJr0, InParticles.W(Constraint.ParticleIndex));
-		ImpulseRatioDenom0 = ImpulseSize / InParticles.M(Constraint.ParticleIndex) + TVector<T, d>::DotProduct(Jr0, IInvJr0);
+		IInvJr0 = PBDRigid0->Q().RotateVector(PBDRigid0->InvI() * PBDRigid0->Q().UnrotateVector(Jr0));
+		ImpulseRatioNumerator0 = TVector<T, d>::DotProduct(Impulse, PBDRigid0->V() - KinematicVelocity) + TVector<T, d>::DotProduct(IInvJr0, PBDRigid0->W());
+		ImpulseRatioDenom0 = ImpulseSize / PBDRigid0->M() + TVector<T, d>::DotProduct(Jr0, IInvJr0);
 	}
-	if (InParticles.InvM(Constraint.LevelsetIndex))
+	if (PBDRigid1)
 	{
 		Jr1 = TVector<T, d>::CrossProduct(VectorToPoint2, Impulse);
-		IInvJr1 = InParticles.Q(Constraint.LevelsetIndex).RotateVector(InParticles.InvI(Constraint.LevelsetIndex) * InParticles.Q(Constraint.LevelsetIndex).UnrotateVector(Jr1));
-		ImpulseRatioNumerator1 = TVector<T, d>::DotProduct(Impulse, InParticles.V(Constraint.LevelsetIndex) - KinematicVelocity) + TVector<T, d>::DotProduct(IInvJr1, InParticles.W(Constraint.LevelsetIndex));
-		ImpulseRatioDenom1 = ImpulseSize / InParticles.M(Constraint.LevelsetIndex) + TVector<T, d>::DotProduct(Jr1, IInvJr1);
+		IInvJr1 = PBDRigid1->Q().RotateVector(PBDRigid1->InvI() * PBDRigid1->Q().UnrotateVector(Jr1));
+		ImpulseRatioNumerator1 = TVector<T, d>::DotProduct(Impulse, PBDRigid1->V() - KinematicVelocity) + TVector<T, d>::DotProduct(IInvJr1, PBDRigid1->W());
+		ImpulseRatioDenom1 = ImpulseSize / PBDRigid1->M() + TVector<T, d>::DotProduct(Jr1, IInvJr1);
 	}
 	T Numerator = -2 * (ImpulseRatioNumerator0 - ImpulseRatioNumerator1);
 	if (Numerator < 0)
@@ -716,287 +723,353 @@ TVector<T, d> GetEnergyClampedImpulse(const TPBDRigidParticles<T, d>& InParticle
 	return Numerator < Denominator ? (Impulse * Numerator / Denominator) : Impulse;
 }
 
-DECLARE_CYCLE_STAT(TEXT("Apply"), STAT_Apply, STATGROUP_ChaosWide);
 template<typename T, int d>
-void TPBDCollisionConstraint<T, d>::Apply(TPBDRigidParticles<T, d>& InParticles, const T Dt, const TArray<int32>& InConstraintIndices)
+void TPBDCollisionConstraint<T, d>::Apply(const T Dt, FRigidBodyContactConstraint& Constraint)
 {
-	PhysicsParallelFor(InConstraintIndices.Num(), [&](int32 ConstraintIndex) {
-		FRigidBodyContactConstraint& Constraint = Constraints[InConstraintIndices[ConstraintIndex]];
-		if (InParticles.Sleeping(Constraint.ParticleIndex))
-		{
-			ensure((InParticles.Sleeping(Constraint.LevelsetIndex) || InParticles.InvM(Constraint.LevelsetIndex) == 0));
-			return;
-		}
-		if (InParticles.Sleeping(Constraint.LevelsetIndex))
+	TGeometryParticleHandle<T, d>* Particle0 = Constraint.Particle;
+	TGeometryParticleHandle<T, d>* Particle1 = Constraint.Levelset;
+	TPBDRigidParticleHandle<T, d>* PBDRigid0 = Particle0->AsDynamic();
+	TPBDRigidParticleHandle<T, d>* PBDRigid1 = Particle1->AsDynamic();
+
+	// @todo(ccaulfield): I think we should never get this? Revisit after particle handle refactor
+	if (Particle0->Sleeping())
+	{
+		ensure(!PBDRigid1 || PBDRigid1->Sleeping());
+		return;
+	}
+	if (Particle1->Sleeping())
+	{
+		ensure(!PBDRigid0 || PBDRigid0->Sleeping());
+		return;
+	}
+	UpdateConstraint<ECollisionUpdateType::Deepest>(MThickness, Constraint);
+	if (Constraint.Phi >= MThickness)
+	{
+		return;
+	}
+
+	// @todo(ccaulfield): CHAOS_PARTICLEHANDLE_TODO what's the best way to manage external per-particle data?
+	Particle0->AuxilaryValue(MCollided) = true;
+	Particle1->AuxilaryValue(MCollided) = true;
+
+	// @todo(ccaulfield): CHAOS_PARTICLEHANDLE_TODO split function to avoid ifs
+	const TVector<T, d> ZeroVector = TVector<T, d>(0);
+	const TRotation<T, d>& Q0 = PBDRigid0? PBDRigid0->Q() : Particle0->R();
+	const TRotation<T, d>& Q1 = PBDRigid1? PBDRigid1->Q() : Particle1->R();
+	const TVector<T, d>& P0 = PBDRigid0 ? PBDRigid0->P() : Particle0->X();
+	const TVector<T, d>& P1 = PBDRigid1 ? PBDRigid1->P() : Particle1->X();
+	const TVector<T, d>& V0 = PBDRigid0 ? PBDRigid0->V() : ZeroVector;
+	const TVector<T, d>& V1 = PBDRigid1 ? PBDRigid1->V() : ZeroVector;
+	const TVector<T, d>& W0 = PBDRigid0 ? PBDRigid0->W() : ZeroVector;
+	const TVector<T, d>& W1 = PBDRigid1 ? PBDRigid1->W() : ZeroVector;
+	TSerializablePtr<TChaosPhysicsMaterial<T>> PhysicsMaterial0 = Particle0->AuxilaryValue(MPhysicsMaterials);
+	TSerializablePtr<TChaosPhysicsMaterial<T>> PhysicsMaterial1 = Particle1->AuxilaryValue(MPhysicsMaterials);
+
+	TVector<T, d> VectorToPoint1 = Constraint.Location - P0;
+	TVector<T, d> VectorToPoint2 = Constraint.Location - P1;
+	TVector<T, d> Body1Velocity = V0 + TVector<T, d>::CrossProduct(W0, VectorToPoint1);
+	TVector<T, d> Body2Velocity = V1 + TVector<T, d>::CrossProduct(W1, VectorToPoint2);
+	TVector<T, d> RelativeVelocity = Body1Velocity - Body2Velocity;
+	if (TVector<T, d>::DotProduct(RelativeVelocity, Constraint.Normal) < 0) // ignore separating constraints
+	{
+		PMatrix<T, d, d> WorldSpaceInvI1 = PBDRigid0 ? (Q0 * FMatrix::Identity).GetTransposed() * PBDRigid0->InvI() * (Q0 * FMatrix::Identity) : PMatrix<T, d, d>(0);
+		PMatrix<T, d, d> WorldSpaceInvI2 = PBDRigid1 ? (Q1 * FMatrix::Identity).GetTransposed() * PBDRigid1->InvI() * (Q1 * FMatrix::Identity) : PMatrix<T, d, d>(0);
+		PMatrix<T, d, d> Factor =
+			(PBDRigid0 ? ComputeFactorMatrix3(VectorToPoint1, WorldSpaceInvI1, PBDRigid0->InvM()) : PMatrix<T, d, d>(0)) +
+			(PBDRigid1 ? ComputeFactorMatrix3(VectorToPoint2, WorldSpaceInvI2, PBDRigid1->InvM()) : PMatrix<T, d, d>(0));
+		TVector<T, d> Impulse;
+		TVector<T, d> AngularImpulse(0);
+
+		// Resting contact if very close to the surface
+		T Restitution = (T)0;
+		T Friction = (T)0;
+		bool bApplyRestitution = (RelativeVelocity.Size() > (2 * 980 * Dt));
+		if (PhysicsMaterial0 && PhysicsMaterial1)
 		{
-			ensure((InParticles.Sleeping(Constraint.ParticleIndex) || InParticles.InvM(Constraint.ParticleIndex) == 0));
-			return;
+			if (bApplyRestitution)
+			{
+				Restitution = FMath::Min(PhysicsMaterial0->Restitution, PhysicsMaterial1->Restitution);
+			}
+			Friction = FMath::Max(PhysicsMaterial0->Friction, PhysicsMaterial1->Friction);
 		}
-		const_cast<TPBDCollisionConstraint<T, d>*>(this)->UpdateConstraint<ECollisionUpdateType::Deepest>(InParticles, MThickness, Constraint);
-		if (Constraint.Phi >= MThickness)
+		else if (PhysicsMaterial0)
 		{
-			return;
-		}
-		MCollided[Constraint.LevelsetIndex] = true;
-		MCollided[Constraint.ParticleIndex] = true;
-		TVector<T, d> VectorToPoint1 = Constraint.Location - InParticles.P(Constraint.ParticleIndex);
-		TVector<T, d> VectorToPoint2 = Constraint.Location - InParticles.P(Constraint.LevelsetIndex);
-		TVector<T, d> Body1Velocity = InParticles.V(Constraint.ParticleIndex) + TVector<T, d>::CrossProduct(InParticles.W(Constraint.ParticleIndex), VectorToPoint1);
-		TVector<T, d> Body2Velocity = InParticles.V(Constraint.LevelsetIndex) + TVector<T, d>::CrossProduct(InParticles.W(Constraint.LevelsetIndex), VectorToPoint2);
-		TVector<T, d> RelativeVelocity = Body1Velocity - Body2Velocity;
-		if (TVector<T, d>::DotProduct(RelativeVelocity, Constraint.Normal) < 0) // ignore separating constraints
-		{
-			PMatrix<T, d, d> WorldSpaceInvI1 = (InParticles.Q(Constraint.ParticleIndex) * FMatrix::Identity).GetTransposed() * InParticles.InvI(Constraint.ParticleIndex) * (InParticles.Q(Constraint.ParticleIndex) * FMatrix::Identity);
-			PMatrix<T, d, d> WorldSpaceInvI2 = (InParticles.Q(Constraint.LevelsetIndex) * FMatrix::Identity).GetTransposed() * InParticles.InvI(Constraint.LevelsetIndex) * (InParticles.Q(Constraint.LevelsetIndex) * FMatrix::Identity);
-			PMatrix<T, d, d> Factor =
-				(InParticles.InvM(Constraint.ParticleIndex) > FLT_MIN ? ComputeFactorMatrix(VectorToPoint1, WorldSpaceInvI1, InParticles.InvM(Constraint.ParticleIndex)) : PMatrix<T, d, d>(0)) +
-				(InParticles.InvM(Constraint.LevelsetIndex) > FLT_MIN ? ComputeFactorMatrix(VectorToPoint2, WorldSpaceInvI2, InParticles.InvM(Constraint.LevelsetIndex)) : PMatrix<T, d, d>(0));
-			TVector<T, d> Impulse;
-			TVector<T, d> AngularImpulse(0);
-
-			// Resting contact if very close to the surface
-			T Restitution = (T)0;
-			T Friction = (T)0;
-			bool bApplyRestitution = (RelativeVelocity.Size() > (2 * 980 * Dt));
-			if (MPhysicsMaterials[Constraint.LevelsetIndex] && MPhysicsMaterials[Constraint.ParticleIndex])
+			if (bApplyRestitution)
 			{
-				if (bApplyRestitution)
-				{
-					Restitution = FMath::Min(MPhysicsMaterials[Constraint.ParticleIndex]->Restitution, MPhysicsMaterials[Constraint.LevelsetIndex]->Restitution);
-				}
-				Friction = FMath::Max(MPhysicsMaterials[Constraint.ParticleIndex]->Friction, MPhysicsMaterials[Constraint.LevelsetIndex]->Friction);
+				Restitution = PhysicsMaterial0->Restitution;
 			}
-			else if (MPhysicsMaterials[Constraint.ParticleIndex])
+			Friction = PhysicsMaterial0->Friction;
+		}
+		else if (PhysicsMaterial1)
+		{
+			if (bApplyRestitution)
 			{
-				if (bApplyRestitution)
-				{
-					Restitution = MPhysicsMaterials[Constraint.ParticleIndex]->Restitution;
-				}
-				Friction = MPhysicsMaterials[Constraint.ParticleIndex]->Friction;
+				Restitution = PhysicsMaterial1->Restitution;
 			}
-			else if (MPhysicsMaterials[Constraint.LevelsetIndex])
+			Friction = PhysicsMaterial1->Friction;
+		}
+
+		if (Friction)
+		{
+			T RelativeNormalVelocity = TVector<T, d>::DotProduct(RelativeVelocity, Constraint.Normal);
+			if (RelativeNormalVelocity > 0)
 			{
-				if (bApplyRestitution)
-				{
-					Restitution = MPhysicsMaterials[Constraint.LevelsetIndex]->Restitution;
-				}
-				Friction = MPhysicsMaterials[Constraint.LevelsetIndex]->Friction;
+				RelativeNormalVelocity = 0;
 			}
-
-			if (Friction)
+			TVector<T, d> VelocityChange = -(Restitution * RelativeNormalVelocity * Constraint.Normal + RelativeVelocity);
+			T NormalVelocityChange = TVector<T, d>::DotProduct(VelocityChange, Constraint.Normal);
+			PMatrix<T, d, d> FactorInverse = Factor.Inverse();
+			TVector<T, d> MinimalImpulse = FactorInverse * VelocityChange;
+			const T MinimalImpulseDotNormal = TVector<T, d>::DotProduct(MinimalImpulse, Constraint.Normal);
+			const T TangentialSize = (MinimalImpulse - MinimalImpulseDotNormal * Constraint.Normal).Size();
+			if (TangentialSize <= Friction * MinimalImpulseDotNormal)
 			{
-				T RelativeNormalVelocity = TVector<T, d>::DotProduct(RelativeVelocity, Constraint.Normal);
-				if (RelativeNormalVelocity > 0)
+				//within friction cone so just solve for static friction stopping the object
+				Impulse = MinimalImpulse;
+				if (MAngularFriction)
 				{
-					RelativeNormalVelocity = 0;
-				}
-				TVector<T, d> VelocityChange = -(Restitution * RelativeNormalVelocity * Constraint.Normal + RelativeVelocity);
-				T NormalVelocityChange = TVector<T, d>::DotProduct(VelocityChange, Constraint.Normal);
-				PMatrix<T, d, d> FactorInverse = Factor.Inverse();
-				TVector<T, d> MinimalImpulse = FactorInverse * VelocityChange;
-				const T MinimalImpulseDotNormal = TVector<T, d>::DotProduct(MinimalImpulse, Constraint.Normal);
-				const T TangentialSize = (MinimalImpulse - MinimalImpulseDotNormal * Constraint.Normal).Size();
-				if (TangentialSize <= Friction * MinimalImpulseDotNormal)
-				{
-					//within friction cone so just solve for static friction stopping the object
-					Impulse = MinimalImpulse;
-					if (MAngularFriction)
+					TVector<T, d> RelativeAngularVelocity = W0 - W1;
+					T AngularNormal = TVector<T, d>::DotProduct(RelativeAngularVelocity, Constraint.Normal);
+					TVector<T, d> AngularTangent = RelativeAngularVelocity - AngularNormal * Constraint.Normal;
+					TVector<T, d> FinalAngularVelocity = FMath::Sign(AngularNormal) * FMath::Max((T)0, FMath::Abs(AngularNormal) - MAngularFriction * NormalVelocityChange) * Constraint.Normal + FMath::Max((T)0, AngularTangent.Size() - MAngularFriction * NormalVelocityChange) * AngularTangent.GetSafeNormal();
+					TVector<T, d> Delta = FinalAngularVelocity - RelativeAngularVelocity;
+					if (!PBDRigid0 && PBDRigid1)
 					{
-						TVector<T, d> RelativeAngularVelocity = InParticles.W(Constraint.ParticleIndex) - InParticles.W(Constraint.LevelsetIndex);
-						T AngularNormal = TVector<T, d>::DotProduct(RelativeAngularVelocity, Constraint.Normal);
-						TVector<T, d> AngularTangent = RelativeAngularVelocity - AngularNormal * Constraint.Normal;
-						TVector<T, d> FinalAngularVelocity = FMath::Sign(AngularNormal) * FMath::Max((T)0, FMath::Abs(AngularNormal) - MAngularFriction * NormalVelocityChange) * Constraint.Normal + FMath::Max((T)0, AngularTangent.Size() - MAngularFriction * NormalVelocityChange) * AngularTangent.GetSafeNormal();
-						TVector<T, d> Delta = FinalAngularVelocity - RelativeAngularVelocity;
-						if (!InParticles.InvM(Constraint.ParticleIndex))
-						{
-							PMatrix<T, d, d> WorldSpaceI2 = (InParticles.Q(Constraint.LevelsetIndex) * FMatrix::Identity) * InParticles.I(Constraint.LevelsetIndex) * (InParticles.Q(Constraint.LevelsetIndex) * FMatrix::Identity).GetTransposed();
-							TVector<T, d> ImpulseDelta = InParticles.M(Constraint.LevelsetIndex) * TVector<T, d>::CrossProduct(VectorToPoint2, Delta);
-							Impulse += ImpulseDelta;
-							AngularImpulse += WorldSpaceI2 * Delta - TVector<T, d>::CrossProduct(VectorToPoint2, ImpulseDelta);
-						}
-						else if (!InParticles.InvM(Constraint.LevelsetIndex))
-						{
-							PMatrix<T, d, d> WorldSpaceI1 = (InParticles.Q(Constraint.ParticleIndex) * FMatrix::Identity) * InParticles.I(Constraint.ParticleIndex) * (InParticles.Q(Constraint.ParticleIndex) * FMatrix::Identity).GetTransposed();
-							TVector<T, d> ImpulseDelta = InParticles.M(Constraint.ParticleIndex) * TVector<T, d>::CrossProduct(VectorToPoint1, Delta);
-							Impulse += ImpulseDelta;
-							AngularImpulse += WorldSpaceI1 * Delta - TVector<T, d>::CrossProduct(VectorToPoint1, ImpulseDelta);
-						}
-						else
-						{
-							PMatrix<T, d, d> Cross1(0, VectorToPoint1.Z, -VectorToPoint1.Y, -VectorToPoint1.Z, 0, VectorToPoint1.X, VectorToPoint1.Y, -VectorToPoint1.X, 0);
-							PMatrix<T, d, d> Cross2(0, VectorToPoint2.Z, -VectorToPoint2.Y, -VectorToPoint2.Z, 0, VectorToPoint2.X, VectorToPoint2.Y, -VectorToPoint2.X, 0);
-							PMatrix<T, d, d> CrossI1 = Cross1 * WorldSpaceInvI1;
-							PMatrix<T, d, d> CrossI2 = Cross2 * WorldSpaceInvI2;
-							PMatrix<T, d, d> Diag1 = CrossI1 * Cross1.GetTransposed() + CrossI2 * Cross2.GetTransposed();
-							Diag1.M[0][0] += InParticles.InvM(Constraint.ParticleIndex) + InParticles.InvM(Constraint.LevelsetIndex);
-							Diag1.M[1][1] += InParticles.InvM(Constraint.ParticleIndex) + InParticles.InvM(Constraint.LevelsetIndex);
-							Diag1.M[2][2] += InParticles.InvM(Constraint.ParticleIndex) + InParticles.InvM(Constraint.LevelsetIndex);
-							PMatrix<T, d, d> OffDiag1 = (CrossI1 + CrossI2) * -1;
-							PMatrix<T, d, d> Diag2 = (WorldSpaceInvI1 + WorldSpaceInvI2).Inverse();
-							PMatrix<T, d, d> OffDiag1Diag2 = OffDiag1 * Diag2;
-							TVector<T, d> ImpulseDelta = PMatrix<T, d, d>((Diag1 - OffDiag1Diag2 * OffDiag1.GetTransposed()).Inverse()) * ((OffDiag1Diag2 * -1) * Delta);
-							Impulse += ImpulseDelta;
-							AngularImpulse += Diag2 * (Delta - PMatrix<T, d, d>(OffDiag1.GetTransposed()) * ImpulseDelta);
-						}
+						PMatrix<T, d, d> WorldSpaceI2 = (Q1 * FMatrix::Identity) * PBDRigid1->I() * (Q1 * FMatrix::Identity).GetTransposed();
+						TVector<T, d> ImpulseDelta = PBDRigid1->M() * TVector<T, d>::CrossProduct(VectorToPoint2, Delta);
+						Impulse += ImpulseDelta;
+						AngularImpulse += WorldSpaceI2 * Delta - TVector<T, d>::CrossProduct(VectorToPoint2, ImpulseDelta);
 					}
-				}
-				else
-				{
-					//outside friction cone, solve for normal relative velocity and keep tangent at cone edge
-					TVector<T, d> Tangent = (RelativeVelocity - TVector<T, d>::DotProduct(RelativeVelocity, Constraint.Normal) * Constraint.Normal).GetSafeNormal();
-					TVector<T, d> DirectionalFactor = Factor * (Constraint.Normal - Friction * Tangent);
-					T ImpulseDenominator = TVector<T, d>::DotProduct(Constraint.Normal, DirectionalFactor);
-					if (!ensureMsgf(FMath::Abs(ImpulseDenominator) > SMALL_NUMBER, TEXT("Constraint:%s\n\nParticle:%s\n\nLevelset:%s\n\nDirectionalFactor:%s, ImpulseDenominator:%f"),
-						*Constraint.ToString(),
-						*InParticles.ToString(Constraint.ParticleIndex),
-						*InParticles.ToString(Constraint.LevelsetIndex),
-						*DirectionalFactor.ToString(), ImpulseDenominator))
+					else if (PBDRigid0 && !PBDRigid1)
 					{
-						ImpulseDenominator = (T)1;
+						PMatrix<T, d, d> WorldSpaceI1 = (Q0 * FMatrix::Identity) * PBDRigid0->I() * (Q0 * FMatrix::Identity).GetTransposed();
+						TVector<T, d> ImpulseDelta = PBDRigid0->M() * TVector<T, d>::CrossProduct(VectorToPoint1, Delta);
+						Impulse += ImpulseDelta;
+						AngularImpulse += WorldSpaceI1 * Delta - TVector<T, d>::CrossProduct(VectorToPoint1, ImpulseDelta);
+					}
+					else if (PBDRigid0 && PBDRigid1)
+					{
+						PMatrix<T, d, d> Cross1(0, VectorToPoint1.Z, -VectorToPoint1.Y, -VectorToPoint1.Z, 0, VectorToPoint1.X, VectorToPoint1.Y, -VectorToPoint1.X, 0);
+						PMatrix<T, d, d> Cross2(0, VectorToPoint2.Z, -VectorToPoint2.Y, -VectorToPoint2.Z, 0, VectorToPoint2.X, VectorToPoint2.Y, -VectorToPoint2.X, 0);
+						PMatrix<T, d, d> CrossI1 = Cross1 * WorldSpaceInvI1;
+						PMatrix<T, d, d> CrossI2 = Cross2 * WorldSpaceInvI2;
+						PMatrix<T, d, d> Diag1 = CrossI1 * Cross1.GetTransposed() + CrossI2 * Cross2.GetTransposed();
+						Diag1.M[0][0] += PBDRigid0->InvM() + PBDRigid1->InvM();
+						Diag1.M[1][1] += PBDRigid0->InvM() + PBDRigid1->InvM();
+						Diag1.M[2][2] += PBDRigid0->InvM() + PBDRigid1->InvM();
+						PMatrix<T, d, d> OffDiag1 = (CrossI1 + CrossI2) * -1;
+						PMatrix<T, d, d> Diag2 = (WorldSpaceInvI1 + WorldSpaceInvI2).Inverse();
+						PMatrix<T, d, d> OffDiag1Diag2 = OffDiag1 * Diag2;
+						TVector<T, d> ImpulseDelta = PMatrix<T, d, d>((Diag1 - OffDiag1Diag2 * OffDiag1.GetTransposed()).Inverse())* ((OffDiag1Diag2 * -1) * Delta);
+						Impulse += ImpulseDelta;
+						AngularImpulse += Diag2 * (Delta - PMatrix<T, d, d>(OffDiag1.GetTransposed()) * ImpulseDelta);
 					}
-
-					const T ImpulseMag = -(1 + Restitution) * RelativeNormalVelocity / ImpulseDenominator;
-					Impulse = ImpulseMag * (Constraint.Normal - Friction * Tangent);
 				}
 			}
 			else
 			{
-				T ImpulseDenominator = TVector<T, d>::DotProduct(Constraint.Normal, Factor * Constraint.Normal);
-				TVector<T, d> ImpulseNumerator = -(1 + Restitution) * TVector<T, d>::DotProduct(RelativeVelocity, Constraint.Normal) * Constraint.Normal;
-				if (!ensureMsgf(FMath::Abs(ImpulseDenominator) > SMALL_NUMBER, TEXT("Constraint:%s\n\nParticle:%s\n\nLevelset:%s\n\nFactor*Constraint.Normal:%s, ImpulseDenominator:%f"),
+				//outside friction cone, solve for normal relative velocity and keep tangent at cone edge
+				TVector<T, d> Tangent = (RelativeVelocity - TVector<T, d>::DotProduct(RelativeVelocity, Constraint.Normal) * Constraint.Normal).GetSafeNormal();
+				TVector<T, d> DirectionalFactor = Factor * (Constraint.Normal - Friction * Tangent);
+				T ImpulseDenominator = TVector<T, d>::DotProduct(Constraint.Normal, DirectionalFactor);
+				if (!ensureMsgf(FMath::Abs(ImpulseDenominator) > SMALL_NUMBER, TEXT("Constraint:%s\n\nParticle:%s\n\nLevelset:%s\n\nDirectionalFactor:%s, ImpulseDenominator:%f"),
 					*Constraint.ToString(),
-					*InParticles.ToString(Constraint.ParticleIndex),
-					*InParticles.ToString(Constraint.LevelsetIndex),
-					*(Factor*Constraint.Normal).ToString(), ImpulseDenominator))
+					*Particle0->ToString(),
+					*Particle1->ToString(),
+					*DirectionalFactor.ToString(), ImpulseDenominator))
 				{
 					ImpulseDenominator = (T)1;
 				}
-				Impulse = ImpulseNumerator / ImpulseDenominator;
+
+				const T ImpulseMag = -(1 + Restitution) * RelativeNormalVelocity / ImpulseDenominator;
+				Impulse = ImpulseMag * (Constraint.Normal - Friction * Tangent);
 			}
-			Impulse = GetEnergyClampedImpulse(InParticles, Constraint, Impulse, VectorToPoint1, VectorToPoint2, Body1Velocity, Body2Velocity);
-			Constraint.AccumulatedImpulse += Impulse;
-			TVector<T, d> AngularImpulse1 = TVector<T, d>::CrossProduct(VectorToPoint1, Impulse) + AngularImpulse;
-			TVector<T, d> AngularImpulse2 = TVector<T, d>::CrossProduct(VectorToPoint2, -Impulse) - AngularImpulse;
+		}
+		else
+		{
+			T ImpulseDenominator = TVector<T, d>::DotProduct(Constraint.Normal, Factor * Constraint.Normal);
+			TVector<T, d> ImpulseNumerator = -(1 + Restitution) * TVector<T, d>::DotProduct(RelativeVelocity, Constraint.Normal)* Constraint.Normal;
+			if (!ensureMsgf(FMath::Abs(ImpulseDenominator) > SMALL_NUMBER, TEXT("Constraint:%s\n\nParticle:%s\n\nLevelset:%s\n\nFactor*Constraint.Normal:%s, ImpulseDenominator:%f"),
+				*Constraint.ToString(),
+				*Particle0->ToString(),
+				*Particle1->ToString(),
+				*(Factor * Constraint.Normal).ToString(), ImpulseDenominator))
+			{
+				ImpulseDenominator = (T)1;
+			}
+			Impulse = ImpulseNumerator / ImpulseDenominator;
+		}
+		Impulse = GetEnergyClampedImpulse(Constraint, Impulse, VectorToPoint1, VectorToPoint2, Body1Velocity, Body2Velocity);
+		Constraint.AccumulatedImpulse += Impulse;
+		TVector<T, d> AngularImpulse1 = TVector<T, d>::CrossProduct(VectorToPoint1, Impulse) + AngularImpulse;
+		TVector<T, d> AngularImpulse2 = TVector<T, d>::CrossProduct(VectorToPoint2, -Impulse) - AngularImpulse;
+		if (PBDRigid0)
+		{
 			// Velocity update for next step
-			InParticles.V(Constraint.ParticleIndex) += InParticles.InvM(Constraint.ParticleIndex) * Impulse;
-			InParticles.W(Constraint.ParticleIndex) += WorldSpaceInvI1 * AngularImpulse1;
-			InParticles.V(Constraint.LevelsetIndex) -= InParticles.InvM(Constraint.LevelsetIndex) * Impulse;
-			InParticles.W(Constraint.LevelsetIndex) += WorldSpaceInvI2 * AngularImpulse2;
+			PBDRigid0->V() += PBDRigid0->InvM() * Impulse;
+			PBDRigid0->W() += WorldSpaceInvI1 * AngularImpulse1;
 			// Position update as part of pbd
-			InParticles.P(Constraint.ParticleIndex) += (InParticles.InvM(Constraint.ParticleIndex) * Impulse) * Dt;
-			InParticles.Q(Constraint.ParticleIndex) += TRotation<T, d>(WorldSpaceInvI1 * AngularImpulse1, 0.f) * InParticles.Q(Constraint.ParticleIndex) * Dt * T(0.5);
-			InParticles.Q(Constraint.ParticleIndex).Normalize();
-			InParticles.P(Constraint.LevelsetIndex) -= (InParticles.InvM(Constraint.LevelsetIndex) * Impulse) * Dt;
-			InParticles.Q(Constraint.LevelsetIndex) += TRotation<T, d>(WorldSpaceInvI2 * AngularImpulse2, 0.f) * InParticles.Q(Constraint.LevelsetIndex) * Dt * T(0.5);
-			InParticles.Q(Constraint.LevelsetIndex).Normalize();
+			PBDRigid0->P() += (PBDRigid0->InvM() * Impulse) * Dt;
+			PBDRigid0->Q() += TRotation<T, d>(WorldSpaceInvI1 * AngularImpulse1, 0.f) * Q0 * Dt * T(0.5);
+			PBDRigid0->Q().Normalize();
 		}
+		if (PBDRigid1)
+		{
+			// Velocity update for next step
+			PBDRigid1->V() -= PBDRigid1->InvM() * Impulse;
+			PBDRigid1->W() += WorldSpaceInvI2 * AngularImpulse2;
+			// Position update as part of pbd
+			PBDRigid1->P() -= (PBDRigid1->InvM() * Impulse) * Dt;
+			PBDRigid1->Q() += TRotation<T, d>(WorldSpaceInvI2 * AngularImpulse2, 0.f) * Q0 * Dt * T(0.5);
+			PBDRigid1->Q().Normalize();
+		}
+	}
+}
+
+DECLARE_CYCLE_STAT(TEXT("Apply"), STAT_Apply2, STATGROUP_ChaosWide);
+template<typename T, int d>
+void TPBDCollisionConstraint<T, d>::Apply(const T Dt, const TArray<int32>& InConstraintIndices)
+{
+	PhysicsParallelFor(InConstraintIndices.Num(), [&](int32 ConstraintIndex) {
+		FRigidBodyContactConstraint& Constraint = Constraints[InConstraintIndices[ConstraintIndex]];
+		Apply(Dt, Constraint);
 	});
 }
 
-DECLARE_CYCLE_STAT(TEXT("ApplyPushOut"), STAT_ApplyPushOut, STATGROUP_ChaosWide);
+DECLARE_CYCLE_STAT(TEXT("ApplyPushOut"), STAT_ApplyPushOut2, STATGROUP_ChaosWide);
 template<typename T, int d>
-bool TPBDCollisionConstraint<T, d>::ApplyPushOut(TPBDRigidParticles<T, d>& InParticles, const T Dt, const TArray<int32>& InConstraintIndices, const TSet<int32>& IsTemporarilyStatic, int32 Iteration, int32 NumIterations)
+void TPBDCollisionConstraint<T, d>::ApplyPushOut(const T Dt, FRigidBodyContactConstraint& Constraint, const TSet<TGeometryParticleHandle<T,d>*>& IsTemporarilyStatic, int32 Iteration, int32 NumIterations, bool &NeedsAnotherIteration)
 {
-	SCOPE_CYCLE_COUNTER(STAT_ApplyPushOut);
+	SCOPE_CYCLE_COUNTER(STAT_ApplyPushOut2);
+	TGeometryParticleHandle<T, d>* Particle0 = Constraint.Particle;
+	TGeometryParticleHandle<T, d>* Particle1 = Constraint.Levelset;
+	TPBDRigidParticleHandle<T, d>* PBDRigid0 = Particle0->AsDynamic();
+	TPBDRigidParticleHandle<T, d>* PBDRigid1 = Particle1->AsDynamic();
 
-	bool NeedsAnotherIteration = false;
+	// @todo(ccaulfield): I think we should never get this? Revisit after particle handle refactor
+	if (Particle0->Sleeping())
+	{
+		ensure(!PBDRigid1 || PBDRigid1->Sleeping());
+		return;
+	}
+	if (Particle1->Sleeping())
+	{
+		ensure(!PBDRigid0 || PBDRigid0->Sleeping());
+		return;
+	}
 
-	PhysicsParallelFor(InConstraintIndices.Num(), [&](int32 ConstraintIndex) {
-		FRigidBodyContactConstraint& Constraint = Constraints[InConstraintIndices[ConstraintIndex]];
-		if (InParticles.Sleeping(Constraint.ParticleIndex))
+	const TVector<T, d> ZeroVector = TVector<T, d>(0);
+	const TRotation<T, d>& Q0 = PBDRigid0 ? PBDRigid0->Q() : Particle0->R();
+	const TRotation<T, d>& Q1 = PBDRigid1 ? PBDRigid1->Q() : Particle1->R();
+	const TVector<T, d>& P0 = PBDRigid0 ? PBDRigid0->P() : Particle0->X();
+	const TVector<T, d>& P1 = PBDRigid1 ? PBDRigid1->P() : Particle1->X();
+	const TVector<T, d>& V0 = PBDRigid0 ? PBDRigid0->V() : ZeroVector;
+	const TVector<T, d>& V1 = PBDRigid1 ? PBDRigid1->V() : ZeroVector;
+	const TVector<T, d>& W0 = PBDRigid0 ? PBDRigid0->W() : ZeroVector;
+	const TVector<T, d>& W1 = PBDRigid1 ? PBDRigid1->W() : ZeroVector;
+	TSerializablePtr<TChaosPhysicsMaterial<T>> PhysicsMaterial0 = Particle0->AuxilaryValue(MPhysicsMaterials);
+	TSerializablePtr<TChaosPhysicsMaterial<T>> PhysicsMaterial1 = Particle1->AuxilaryValue(MPhysicsMaterials);
+	const bool IsTemporarilyStatic0 = IsTemporarilyStatic.Contains(Particle0);
+	const bool IsTemporarilyStatic1 = IsTemporarilyStatic.Contains(Particle1);
+
+	for (int32 PairIteration = 0; PairIteration < MPairIterations; ++PairIteration)
+	{
+		UpdateConstraint<ECollisionUpdateType::Deepest>(MThickness, Constraint);
+		if (Constraint.Phi >= MThickness)
 		{
-			ensure((InParticles.Sleeping(Constraint.LevelsetIndex) || InParticles.InvM(Constraint.LevelsetIndex) == 0));
-			return;
+			break;
 		}
-		if (InParticles.Sleeping(Constraint.LevelsetIndex))
+
+		if ((!PBDRigid0 || IsTemporarilyStatic0) && (!PBDRigid1 || IsTemporarilyStatic1))
 		{
-			ensure((InParticles.Sleeping(Constraint.ParticleIndex) || InParticles.InvM(Constraint.ParticleIndex) == 0));
-			return;
+			break;
 		}
-		for (int32 PairIteration = 0; PairIteration < MPairIterations; ++PairIteration)
+
+		NeedsAnotherIteration = true;
+		PMatrix<T, d, d> WorldSpaceInvI1 = PBDRigid0? (Q0 * FMatrix::Identity).GetTransposed() * PBDRigid0->InvI() * (Q0 * FMatrix::Identity) : PMatrix<T, d, d>(0);
+		PMatrix<T, d, d> WorldSpaceInvI2 = PBDRigid1 ? (Q1 * FMatrix::Identity).GetTransposed() * PBDRigid1->InvI() * (Q1 * FMatrix::Identity) : PMatrix<T, d, d>(0);
+		TVector<T, d> VectorToPoint1 = Constraint.Location - P0;
+		TVector<T, d> VectorToPoint2 = Constraint.Location - P1;
+		PMatrix<T, d, d> Factor =
+			(PBDRigid0 ? ComputeFactorMatrix3(VectorToPoint1, WorldSpaceInvI1, PBDRigid0->InvM()) : PMatrix<T, d, d>(0)) +
+			(PBDRigid1 ? ComputeFactorMatrix3(VectorToPoint2, WorldSpaceInvI2, PBDRigid1->InvM()) : PMatrix<T, d, d>(0));
+		T Numerator = FMath::Min((T)(Iteration + 2), (T)NumIterations);
+		T ScalingFactor = Numerator / (T)NumIterations;
+
+		//if pushout is needed we better fix relative velocity along normal. Treat it as if 0 restitution
+		TVector<T, d> Body1Velocity = V0 + TVector<T, d>::CrossProduct(W0, VectorToPoint1);
+		TVector<T, d> Body2Velocity = V1 + TVector<T, d>::CrossProduct(W1, VectorToPoint2);
+		TVector<T, d> RelativeVelocity = Body1Velocity - Body2Velocity;
+		const T RelativeVelocityDotNormal = TVector<T, d>::DotProduct(RelativeVelocity, Constraint.Normal);
+		if (RelativeVelocityDotNormal < 0)
 		{
-			UpdateConstraint<ECollisionUpdateType::Deepest>(InParticles, MThickness, Constraint);
-			if (Constraint.Phi >= MThickness)
+			T ImpulseDenominator = TVector<T, d>::DotProduct(Constraint.Normal, Factor * Constraint.Normal);
+			TVector<T, d> ImpulseNumerator = -TVector<T, d>::DotProduct(RelativeVelocity, Constraint.Normal) * Constraint.Normal * ScalingFactor;
+			if (!ensureMsgf(FMath::Abs(ImpulseDenominator) > SMALL_NUMBER, TEXT("ApplyPushout Constraint:%s\n\nParticle:%s\n\nLevelset:%s\n\nFactor*Constraint.Normal:%s, ImpulseDenominator:%f"),
+				*Constraint.ToString(),
+				*Particle0->ToString(),
+				*Particle1->ToString(),
+				*(Factor*Constraint.Normal).ToString(), ImpulseDenominator))
 			{
-				break;
+				ImpulseDenominator = (T)1;
 			}
 
-			if ((!InParticles.InvM(Constraint.ParticleIndex) || IsTemporarilyStatic.Contains(Constraint.ParticleIndex)) && (!InParticles.InvM(Constraint.LevelsetIndex) || IsTemporarilyStatic.Contains(Constraint.LevelsetIndex)))
+			TVector<T, d> VelocityFixImpulse = ImpulseNumerator / ImpulseDenominator;
+			VelocityFixImpulse = GetEnergyClampedImpulse(Constraint, VelocityFixImpulse, VectorToPoint1, VectorToPoint2, Body1Velocity, Body2Velocity);
+			Constraint.AccumulatedImpulse += VelocityFixImpulse;	//question: should we track this?
+			if (!IsTemporarilyStatic0 && PBDRigid0)
 			{
-				break;
+				TVector<T, d> AngularImpulse = TVector<T, d>::CrossProduct(VectorToPoint1, VelocityFixImpulse);
+				PBDRigid0->V() += PBDRigid0->InvM() * VelocityFixImpulse;
+				PBDRigid0->W() += WorldSpaceInvI1 * AngularImpulse;
+
 			}
 
-			NeedsAnotherIteration = true;
-			PMatrix<T, d, d> WorldSpaceInvI1 = (InParticles.Q(Constraint.ParticleIndex) * FMatrix::Identity).GetTransposed() * InParticles.InvI(Constraint.ParticleIndex) * (InParticles.Q(Constraint.ParticleIndex) * FMatrix::Identity);
-			PMatrix<T, d, d> WorldSpaceInvI2 = (InParticles.Q(Constraint.LevelsetIndex) * FMatrix::Identity).GetTransposed() * InParticles.InvI(Constraint.LevelsetIndex) * (InParticles.Q(Constraint.LevelsetIndex) * FMatrix::Identity);
-			TVector<T, d> VectorToPoint1 = Constraint.Location - InParticles.P(Constraint.ParticleIndex);
-			TVector<T, d> VectorToPoint2 = Constraint.Location - InParticles.P(Constraint.LevelsetIndex);
-			PMatrix<T, d, d> Factor =
-				((InParticles.InvM(Constraint.ParticleIndex) && !IsTemporarilyStatic.Contains(Constraint.ParticleIndex)) ? ComputeFactorMatrix(VectorToPoint1, WorldSpaceInvI1, InParticles.InvM(Constraint.ParticleIndex)) : PMatrix<T, d, d>(0)) +
-				((InParticles.InvM(Constraint.LevelsetIndex) && !IsTemporarilyStatic.Contains(Constraint.LevelsetIndex)) ? ComputeFactorMatrix(VectorToPoint2, WorldSpaceInvI2, InParticles.InvM(Constraint.LevelsetIndex)) : PMatrix<T, d, d>(0));
-			T Numerator = FMath::Min((T)(Iteration + 2), (T)NumIterations);
-			T ScalingFactor = Numerator / (T)NumIterations;
-
-			//if pushout is needed we better fix relative velocity along normal. Treat it as if 0 restitution
-			TVector<T, d> Body1Velocity = InParticles.V(Constraint.ParticleIndex) + TVector<T, d>::CrossProduct(InParticles.W(Constraint.ParticleIndex), VectorToPoint1);
-			TVector<T, d> Body2Velocity = InParticles.V(Constraint.LevelsetIndex) + TVector<T, d>::CrossProduct(InParticles.W(Constraint.LevelsetIndex), VectorToPoint2);
-			TVector<T, d> RelativeVelocity = Body1Velocity - Body2Velocity;
-			const T RelativeVelocityDotNormal = TVector<T, d>::DotProduct(RelativeVelocity, Constraint.Normal);
-			if (RelativeVelocityDotNormal < 0)
+			if (!IsTemporarilyStatic1 && PBDRigid1)
 			{
-				T ImpulseDenominator = TVector<T, d>::DotProduct(Constraint.Normal, Factor * Constraint.Normal);
-				TVector<T, d> ImpulseNumerator = -TVector<T, d>::DotProduct(RelativeVelocity, Constraint.Normal) * Constraint.Normal * ScalingFactor;
-				if (!ensureMsgf(FMath::Abs(ImpulseDenominator) > SMALL_NUMBER, TEXT("ApplyPushout Constraint:%s\n\nParticle:%s\n\nLevelset:%s\n\nFactor*Constraint.Normal:%s, ImpulseDenominator:%f"),
-					*Constraint.ToString(),
-					*InParticles.ToString(Constraint.ParticleIndex),
-					*InParticles.ToString(Constraint.LevelsetIndex),
-					*(Factor*Constraint.Normal).ToString(), ImpulseDenominator))
-				{
-					ImpulseDenominator = (T)1;
-				}
+				TVector<T, d> AngularImpulse = TVector<T, d>::CrossProduct(VectorToPoint2, -VelocityFixImpulse);
+				PBDRigid1->V() -= PBDRigid1->InvM() * VelocityFixImpulse;
+				PBDRigid1->W() += WorldSpaceInvI2 * AngularImpulse;
+			}
 
-				TVector<T, d> VelocityFixImpulse = ImpulseNumerator / ImpulseDenominator;
-				VelocityFixImpulse = GetEnergyClampedImpulse(InParticles, Constraint, VelocityFixImpulse, VectorToPoint1, VectorToPoint2, Body1Velocity, Body2Velocity);
-				Constraint.AccumulatedImpulse += VelocityFixImpulse;	//question: should we track this?
-				if (!IsTemporarilyStatic.Contains(Constraint.ParticleIndex))
-				{
-					TVector<T, d> AngularImpulse = TVector<T, d>::CrossProduct(VectorToPoint1, VelocityFixImpulse);
-					InParticles.V(Constraint.ParticleIndex) += InParticles.InvM(Constraint.ParticleIndex) * VelocityFixImpulse;
-					InParticles.W(Constraint.ParticleIndex) += WorldSpaceInvI1 * AngularImpulse;
+		}
 
-				}
 
-				if (!IsTemporarilyStatic.Contains(Constraint.LevelsetIndex))
-				{
-					TVector<T, d> AngularImpulse = TVector<T, d>::CrossProduct(VectorToPoint2, -VelocityFixImpulse);
-					InParticles.V(Constraint.LevelsetIndex) -= InParticles.InvM(Constraint.LevelsetIndex) * VelocityFixImpulse;
-					InParticles.W(Constraint.LevelsetIndex) += WorldSpaceInvI2 * AngularImpulse;
-				}
+		TVector<T, d> Impulse = PMatrix<T, d, d>(Factor.Inverse()) * ((-Constraint.Phi + MThickness) * ScalingFactor * Constraint.Normal);
+		TVector<T, d> AngularImpulse1 = TVector<T, d>::CrossProduct(VectorToPoint1, Impulse);
+		TVector<T, d> AngularImpulse2 = TVector<T, d>::CrossProduct(VectorToPoint2, -Impulse);
+		if (!IsTemporarilyStatic0 && PBDRigid0)
+		{
+			PBDRigid0->P() += PBDRigid0->InvM() * Impulse;
+			PBDRigid0->Q() = TRotation<T, d>::FromVector(WorldSpaceInvI1 * AngularImpulse1) * Q0;
+			PBDRigid0->Q().Normalize();
+		}
+		if (!IsTemporarilyStatic1 && PBDRigid1)
+		{
+			PBDRigid1->P() -= PBDRigid1->InvM() * Impulse;
+			PBDRigid1->Q() = TRotation<T, d>::FromVector(WorldSpaceInvI2 * AngularImpulse2) * Q1;
+			PBDRigid1->Q().Normalize();
+		}
+	}
+}
 
-			}
+//DECLARE_CYCLE_STAT(TEXT("ApplyPushOut"), STAT_ApplyPushOut, STATGROUP_ChaosWide);
+template<typename T, int d>
+bool TPBDCollisionConstraint<T, d>::ApplyPushOut(const T Dt, const TArray<int32>& InConstraintIndices, const TSet<TGeometryParticleHandle<T,d>*>& IsTemporarilyStatic, int32 Iteration, int32 NumIterations)
+{
+	SCOPE_CYCLE_COUNTER(STAT_ApplyPushOut2);
 
+	bool NeedsAnotherIteration = false;
 
-			TVector<T, d> Impulse = PMatrix<T, d, d>(Factor.Inverse()) * ((-Constraint.Phi + MThickness) * ScalingFactor * Constraint.Normal);
-			TVector<T, d> AngularImpulse1 = TVector<T, d>::CrossProduct(VectorToPoint1, Impulse);
-			TVector<T, d> AngularImpulse2 = TVector<T, d>::CrossProduct(VectorToPoint2, -Impulse);
-			if (!IsTemporarilyStatic.Contains(Constraint.ParticleIndex))
-			{
-				InParticles.P(Constraint.ParticleIndex) += InParticles.InvM(Constraint.ParticleIndex) * Impulse;
-				InParticles.Q(Constraint.ParticleIndex) = TRotation<T, d>::FromVector(WorldSpaceInvI1 * AngularImpulse1) * InParticles.Q(Constraint.ParticleIndex);
-				InParticles.Q(Constraint.ParticleIndex).Normalize();
-			}
-			if (!IsTemporarilyStatic.Contains(Constraint.LevelsetIndex))
-			{
-				InParticles.P(Constraint.LevelsetIndex) -= InParticles.InvM(Constraint.LevelsetIndex) * Impulse;
-				InParticles.Q(Constraint.LevelsetIndex) = TRotation<T, d>::FromVector(WorldSpaceInvI2 * AngularImpulse2) * InParticles.Q(Constraint.LevelsetIndex);
-				InParticles.Q(Constraint.LevelsetIndex).Normalize();
-			}
-		}
-	});
+	PhysicsParallelFor(InConstraintIndices.Num(), [&](int32 ConstraintIndex) {
+		FRigidBodyContactConstraint& Constraint = Constraints[InConstraintIndices[ConstraintIndex]];
+		ApplyPushOut(Dt, Constraint, IsTemporarilyStatic, Iteration, NumIterations, NeedsAnotherIteration);
+		});
 
 	return NeedsAnotherIteration;
 }
 
+
 template<typename T, int d>
 bool TPBDCollisionConstraint<T, d>::NearestPoint(TArray<Pair<TVector<T, d>, TVector<T, d>>>& Points, TVector<T, d>& Direction, TVector<T, d>& ClosestPoint)
 {
@@ -1173,61 +1246,13 @@ bool TPBDCollisionConstraint<T, d>::NearestPoint(TArray<Pair<TVector<T, d>, TVec
 	return false;
 }
 
-template<class T_PARTICLES, typename T, int d>
-TVector<T, d> GetPosition(const T_PARTICLES& InParticles)
-{
-	check(false);
-	return TVector<T, d>();
-}
-
 template<typename T, int d>
-TVector<T, d> GetPosition(const TParticles<T, d>& InParticles, const int32 Index)
+TRigidTransform<T, d> GetTransform(const TGeometryParticleHandle<T, d>* Particle)
 {
-	return InParticles.X(Index);
-}
-
-template<typename T, int d>
-TVector<T, d> GetPosition(const TPBDRigidParticles<T, d>& InParticles, const int32 Index)
-{
-	return InParticles.P(Index);
-}
-
-template<class T_PARTICLES, typename T, int d>
-TRotation<T, d> GetRotation(const T_PARTICLES& InParticles)
-{
-	check(false);
-	return TRotation<T, d>();
-}
-
-template<typename T, int d>
-TRotation<T, d> GetRotation(const TRigidParticles<T, d>& InParticles, const int32 Index)
-{
-	return InParticles.R(Index);
-}
-
-template<typename T, int d>
-TRotation<T, d> GetRotation(const TPBDRigidParticles<T, d>& InParticles, const int32 Index)
-{
-	return InParticles.Q(Index);
-}
-
-template<class T_PARTICLES, typename T, int d>
-TRigidTransform<T, d> GetTransform(const T_PARTICLES& InParticles)
-{
-	check(false);
-	return TRigidTransform<T, d>();
-}
-
-template<typename T, int d>
-TRigidTransform<T, d> GetTransform(const TRigidParticles<T, d>& InParticles, const int32 Index)
-{
-	return TRigidTransform<T, d>(InParticles.X(Index), InParticles.R(Index));
-}
-
-template<typename T, int d>
-TRigidTransform<T, d> GetTransform(const TPBDRigidParticles<T, d>& InParticles, const int32 Index)
-{
-	return TRigidTransform<T, d>(InParticles.P(Index), InParticles.Q(Index));
+	//#if CHAOS_PARTICLEHANDLE_TODO
+	
+	TGenericParticleHandle<T, d> Generic = const_cast<TGeometryParticleHandle<T, d>*>(Particle);	//TODO: give a const version of the generic API
+	return TRigidTransform<T, d>(Generic->P(), Generic->Q());
 }
 
 #if 0
@@ -1255,7 +1280,7 @@ void UpdateLevelsetConstraintHelperCCD(const TRigidParticles<T, d>& InParticles,
 #endif
 
 template <typename T, int d>
-bool SampleObjectHelper(const TImplicitObject<T, d>& Object, const TRigidTransform<T,d>& ObjectTransform, const TRigidTransform<T, d>& SampleToObjectTransform, const TVector<T, d>& SampleParticle, T Thickness, TRigidBodyContactConstraint<T, d>& Constraint)
+bool SampleObjectHelper2(const TImplicitObject<T, d>& Object, const TRigidTransform<T,d>& ObjectTransform, const TRigidTransform<T, d>& SampleToObjectTransform, const TVector<T, d>& SampleParticle, T Thickness, TRigidBodyContactConstraint<T, d>& Constraint)
 {
 	TVector<T, d> LocalPoint = SampleToObjectTransform.TransformPositionNoScale(SampleParticle);
 	TVector<T, d> LocalNormal;
@@ -1271,7 +1296,7 @@ bool SampleObjectHelper(const TImplicitObject<T, d>& Object, const TRigidTransfo
 }
 
 template <typename T, int d>
-bool SampleObjectNoNormal(const TImplicitObject<T, d>& Object, const TRigidTransform<T,d>& ObjectTransform, const TRigidTransform<T, d>& SampleToObjectTransform, const TVector<T, d>& SampleParticle, T Thickness, TRigidBodyContactConstraint<T, d>& Constraint)
+bool SampleObjectNoNormal2(const TImplicitObject<T, d>& Object, const TRigidTransform<T,d>& ObjectTransform, const TRigidTransform<T, d>& SampleToObjectTransform, const TVector<T, d>& SampleParticle, T Thickness, TRigidBodyContactConstraint<T, d>& Constraint)
 {
 	TVector<T, d> LocalPoint = SampleToObjectTransform.TransformPositionNoScale(SampleParticle);
 	TVector<T, d> LocalNormal;
@@ -1285,7 +1310,7 @@ bool SampleObjectNoNormal(const TImplicitObject<T, d>& Object, const TRigidTrans
 }
 
 template <typename T, int d>
-bool SampleObjectNormalAverageHelper(const TImplicitObject<T, d>& Object, const TRigidTransform<T, d>& ObjectTransform, const TRigidTransform<T, d>& SampleToObjectTransform, const TVector<T, d>& SampleParticle, T Thickness, T& TotalThickness, TRigidBodyContactConstraint<T, d>& Constraint)
+bool SampleObjectNormalAverageHelper2(const TImplicitObject<T, d>& Object, const TRigidTransform<T, d>& ObjectTransform, const TRigidTransform<T, d>& SampleToObjectTransform, const TVector<T, d>& SampleParticle, T Thickness, T& TotalThickness, TRigidBodyContactConstraint<T, d>& Constraint)
 {
 	TVector<T, d> LocalPoint = SampleToObjectTransform.TransformPositionNoScale(SampleParticle);
 	TVector<T, d> LocalNormal;
@@ -1300,26 +1325,26 @@ bool SampleObjectNormalAverageHelper(const TImplicitObject<T, d>& Object, const
 	return false;
 }
 
-DECLARE_CYCLE_STAT(TEXT("UpdateLevelsetPartial"), STAT_UpdateLevelsetPartial, STATGROUP_ChaosWide);
-DECLARE_CYCLE_STAT(TEXT("UpdateLevelsetFindParticles"), STAT_UpdateLevelsetFindParticles, STATGROUP_ChaosWide);
-DECLARE_CYCLE_STAT(TEXT("UpdateLevelsetBVHTraversal"), STAT_UpdateLevelsetBVHTraversal, STATGROUP_ChaosWide);
-DECLARE_CYCLE_STAT(TEXT("UpdateLevelsetSignedDistance"), STAT_UpdateLevelsetSignedDistance, STATGROUP_ChaosWide);
-DECLARE_CYCLE_STAT(TEXT("UpdateLevelsetAll"), STAT_UpdateLevelsetAll, STATGROUP_ChaosWide);
-DECLARE_CYCLE_STAT(TEXT("SampleObject"), STAT_SampleObject, STATGROUP_ChaosWide);
+DECLARE_CYCLE_STAT(TEXT("UpdateLevelsetPartial"), STAT_UpdateLevelsetPartial2, STATGROUP_ChaosWide);
+DECLARE_CYCLE_STAT(TEXT("UpdateLevelsetFindParticles"), STAT_UpdateLevelsetFindParticles2, STATGROUP_ChaosWide);
+DECLARE_CYCLE_STAT(TEXT("UpdateLevelsetBVHTraversal"), STAT_UpdateLevelsetBVHTraversal2, STATGROUP_ChaosWide);
+DECLARE_CYCLE_STAT(TEXT("UpdateLevelsetSignedDistance"), STAT_UpdateLevelsetSignedDistance2, STATGROUP_ChaosWide);
+DECLARE_CYCLE_STAT(TEXT("UpdateLevelsetAll"), STAT_UpdateLevelsetAll2, STATGROUP_ChaosWide);
+DECLARE_CYCLE_STAT(TEXT("SampleObject"), STAT_SampleObject2, STATGROUP_ChaosWide);
 
-int32 NormalAveraging = 1;
-FAutoConsoleVariableRef CVarNormalAveraging(TEXT("p.NormalAveraging"), NormalAveraging, TEXT(""));
+int32 NormalAveraging2 = 1;
+FAutoConsoleVariableRef CVarNormalAveraging2(TEXT("p.NormalAveraging2"), NormalAveraging2, TEXT(""));
 
-int32 SampleMinParticlesForAcceleration = 2048;
-FAutoConsoleVariableRef CVarSampleMinParticlesForAcceleration(TEXT("p.SampleMinParticlesForAcceleration"), SampleMinParticlesForAcceleration, TEXT("The minimum number of particles needed before using an acceleration structure when sampling"));
+int32 SampleMinParticlesForAcceleration2 = 2048;
+FAutoConsoleVariableRef CVarSampleMinParticlesForAcceleration2(TEXT("p.SampleMinParticlesForAcceleration2"), SampleMinParticlesForAcceleration2, TEXT("The minimum number of particles needed before using an acceleration structure when sampling"));
 
 template <ECollisionUpdateType UpdateType, typename T, int d>
-void SampleObject(const TImplicitObject<T, d>& Object, const TRigidTransform<T, d>& ObjectTransform, const TBVHParticles<T, d>& SampleParticles, const TRigidTransform<T, d>& SampleParticlesTransform, T Thickness, TRigidBodyContactConstraint<T, d>& Constraint)
+void SampleObject2(const TImplicitObject<T, d>& Object, const TRigidTransform<T, d>& ObjectTransform, const TBVHParticles<T, d>& SampleParticles, const TRigidTransform<T, d>& SampleParticlesTransform, T Thickness, TRigidBodyContactConstraint<T, d>& Constraint)
 {
-	SCOPE_CYCLE_COUNTER(STAT_SampleObject);
+	SCOPE_CYCLE_COUNTER(STAT_SampleObject2);
 	TRigidBodyContactConstraint<T, d> AvgConstraint;
-	AvgConstraint.ParticleIndex = Constraint.ParticleIndex;
-	AvgConstraint.LevelsetIndex = Constraint.LevelsetIndex;
+	AvgConstraint.Particle = Constraint.Particle;
+	AvgConstraint.Levelset = Constraint.Levelset;
 	AvgConstraint.Location = TVector<T, d>::ZeroVector;
 	AvgConstraint.Normal = TVector<T, d>::ZeroVector;
 	AvgConstraint.Phi = Thickness;
@@ -1329,27 +1354,27 @@ void SampleObject(const TImplicitObject<T, d>& Object, const TRigidTransform<T,
 	const int32 NumParticles = SampleParticles.Size();
 
 	const TRigidTransform<T, d> & SampleToObjectTM = SampleParticlesTransform.GetRelativeTransform(ObjectTransform);
-	if (NumParticles > SampleMinParticlesForAcceleration && Object.HasBoundingBox())
+	if (NumParticles > SampleMinParticlesForAcceleration2 && Object.HasBoundingBox())
 	{
-		SCOPE_CYCLE_COUNTER(STAT_UpdateLevelsetPartial);
+		SCOPE_CYCLE_COUNTER(STAT_UpdateLevelsetPartial2);
 		TBox<T, d> ImplicitBox = Object.BoundingBox().TransformedBox(ObjectTransform.GetRelativeTransform(SampleParticlesTransform));
 		ImplicitBox.Thicken(Thickness);
 		TArray<int32> PotentialParticles;
 		{
-			SCOPE_CYCLE_COUNTER(STAT_UpdateLevelsetFindParticles);
+			SCOPE_CYCLE_COUNTER(STAT_UpdateLevelsetFindParticles2);
 			PotentialParticles = SampleParticles.FindAllIntersections(ImplicitBox);
 		}
 		{
-			SCOPE_CYCLE_COUNTER(STAT_UpdateLevelsetSignedDistance);
+			SCOPE_CYCLE_COUNTER(STAT_UpdateLevelsetSignedDistance2);
 			for (int32 i : PotentialParticles)
 			{
-				if (NormalAveraging && UpdateType != ECollisionUpdateType::Any)	//if we just want one don't bother with normal
+				if (NormalAveraging2 && UpdateType != ECollisionUpdateType::Any)	//if we just want one don't bother with normal
 				{
-					SampleObjectNormalAverageHelper(Object, ObjectTransform, SampleToObjectTM, SampleParticles.X(i), Thickness, TotalThickness, AvgConstraint);
+					SampleObjectNormalAverageHelper2(Object, ObjectTransform, SampleToObjectTM, SampleParticles.X(i), Thickness, TotalThickness, AvgConstraint);
 				}
 				else
 				{
-					if (SampleObjectNoNormal(Object, ObjectTransform, SampleToObjectTM, SampleParticles.X(i), Thickness, AvgConstraint))
+					if (SampleObjectNoNormal2(Object, ObjectTransform, SampleToObjectTM, SampleParticles.X(i), Thickness, AvgConstraint))
 					{
 						DeepestParticle = i;
 						if (UpdateType == ECollisionUpdateType::Any)
@@ -1364,16 +1389,16 @@ void SampleObject(const TImplicitObject<T, d>& Object, const TRigidTransform<T,
 	}
 	else
 	{
-		SCOPE_CYCLE_COUNTER(STAT_UpdateLevelsetAll);
+		SCOPE_CYCLE_COUNTER(STAT_UpdateLevelsetAll2);
 		for (int32 i = 0; i < NumParticles; ++i)
 		{
-			if (NormalAveraging && UpdateType != ECollisionUpdateType::Any)	//if we just want one don't bother with normal
+			if (NormalAveraging2 && UpdateType != ECollisionUpdateType::Any)	//if we just want one don't bother with normal
 			{
-				const bool bInside = SampleObjectNormalAverageHelper(Object, ObjectTransform, SampleToObjectTM, SampleParticles.X(i), Thickness, TotalThickness, AvgConstraint);
+				const bool bInside = SampleObjectNormalAverageHelper2(Object, ObjectTransform, SampleToObjectTM, SampleParticles.X(i), Thickness, TotalThickness, AvgConstraint);
 			}
 			else
 			{
-				if (SampleObjectNoNormal(Object, ObjectTransform, SampleToObjectTM, SampleParticles.X(i), Thickness, AvgConstraint))
+				if (SampleObjectNoNormal2(Object, ObjectTransform, SampleToObjectTM, SampleParticles.X(i), Thickness, AvgConstraint))
 				{
 					DeepestParticle = i;
 					if (UpdateType == ECollisionUpdateType::Any)
@@ -1386,7 +1411,7 @@ void SampleObject(const TImplicitObject<T, d>& Object, const TRigidTransform<T,
 		}
 	}
 
-	if (NormalAveraging)
+	if (NormalAveraging2)
 	{
 		if (TotalThickness < -KINDA_SMALL_NUMBER)
 		{
@@ -1419,12 +1444,12 @@ void SampleObject(const TImplicitObject<T, d>& Object, const TRigidTransform<T,
 
 #if INTEL_ISPC
 template<ECollisionUpdateType UpdateType>
-void SampleObject(const TImplicitObject<float, 3>& Object, const TRigidTransform<float, 3>& ObjectTransform, const TBVHParticles<float, 3>& SampleParticles, const TRigidTransform<float, 3>& SampleParticlesTransform, float Thickness, TRigidBodyContactConstraint<float, 3>& Constraint)
+void SampleObject2(const TImplicitObject<float, 3>& Object, const TRigidTransform<float, 3>& ObjectTransform, const TBVHParticles<float, 3>& SampleParticles, const TRigidTransform<float, 3>& SampleParticlesTransform, float Thickness, TRigidBodyContactConstraint<float, 3>& Constraint)
 {
-	SCOPE_CYCLE_COUNTER(STAT_SampleObject);
+	SCOPE_CYCLE_COUNTER(STAT_SampleObject2);
 	TRigidBodyContactConstraint<float, 3> AvgConstraint;
-	AvgConstraint.ParticleIndex = Constraint.ParticleIndex;
-	AvgConstraint.LevelsetIndex = Constraint.LevelsetIndex;
+	AvgConstraint.Particle = Constraint.Particle;
+	AvgConstraint.Levelset = Constraint.Levelset;
 	AvgConstraint.Location = TVector<float, 3>::ZeroVector;
 	AvgConstraint.Normal = TVector<float, 3>::ZeroVector;
 	AvgConstraint.Phi = Thickness;
@@ -1435,18 +1460,18 @@ void SampleObject(const TImplicitObject<float, 3>& Object, const TRigidTransform
 	const TRigidTransform<float, 3>& SampleToObjectTM = SampleParticlesTransform.GetRelativeTransform(ObjectTransform);
 	int32 NumParticles = SampleParticles.Size();
 
-	if (NumParticles > SampleMinParticlesForAcceleration && Object.HasBoundingBox())
+	if (NumParticles > SampleMinParticlesForAcceleration2 && Object.HasBoundingBox())
 	{
-		SCOPE_CYCLE_COUNTER(STAT_UpdateLevelsetPartial);
+		SCOPE_CYCLE_COUNTER(STAT_UpdateLevelsetPartial2);
 		TBox<float, 3> ImplicitBox = Object.BoundingBox().TransformedBox(ObjectTransform.GetRelativeTransform(SampleParticlesTransform));
 		ImplicitBox.Thicken(Thickness);
 		TArray<int32> PotentialParticles;
 		{
-			SCOPE_CYCLE_COUNTER(STAT_UpdateLevelsetFindParticles);
+			SCOPE_CYCLE_COUNTER(STAT_UpdateLevelsetFindParticles2);
 			PotentialParticles = SampleParticles.FindAllIntersections(ImplicitBox);
 		}
 		{
-			SCOPE_CYCLE_COUNTER(STAT_UpdateLevelsetSignedDistance);
+			SCOPE_CYCLE_COUNTER(STAT_UpdateLevelsetSignedDistance2);
 
 			if (Object.GetType(true) == ImplicitObjectType::LevelSet && PotentialParticles.Num() > 0)
 			{
@@ -1454,7 +1479,7 @@ void SampleObject(const TImplicitObject<float, 3>& Object, const TRigidTransform
 				const TLevelSet<float, 3>* LevelSet = Object.GetObject<Chaos::TLevelSet<float, 3>>();
 				const TUniformGrid<float, 3>& Grid = LevelSet->GetGrid();
 
-				if (NormalAveraging && UpdateType != ECollisionUpdateType::Any)
+				if (NormalAveraging2 && UpdateType != ECollisionUpdateType::Any)
 				{
 					ispc::SampleLevelSetNormalAverage(
 						(ispc::FVector&)Grid.MinCorner(),
@@ -1463,7 +1488,7 @@ void SampleObject(const TImplicitObject<float, 3>& Object, const TRigidTransform
 						(ispc::FIntVector&)Grid.Counts(),
 						(ispc::TArrayND*)&LevelSet->GetPhiArray(),
 						(ispc::FTransform&)SampleToObjectTM,
-						(ispc::FVector*)&SampleParticles.XView()[0],
+						(ispc::FVector*)&SampleParticles.XArray()[0],
 						&PotentialParticles[0],
 						Thickness,
 						TotalThickness,
@@ -1479,7 +1504,7 @@ void SampleObject(const TImplicitObject<float, 3>& Object, const TRigidTransform
 						(ispc::FIntVector&)Grid.Counts(),
 						(ispc::TArrayND*)&LevelSet->GetPhiArray(),
 						(ispc::FTransform&)SampleToObjectTM,
-						(ispc::FVector*)&SampleParticles.XView()[0],
+						(ispc::FVector*)&SampleParticles.XArray()[0],
 						&PotentialParticles[0],
 						DeepestParticle,
 						AvgConstraint.Phi,
@@ -1497,13 +1522,13 @@ void SampleObject(const TImplicitObject<float, 3>& Object, const TRigidTransform
 				//QUICK_SCOPE_CYCLE_COUNTER(STAT_Box);
 				const TBox<float, 3>* Box = Object.GetObject<Chaos::TBox<float, 3>>();
 
-				if (NormalAveraging && UpdateType != ECollisionUpdateType::Any)
+				if (NormalAveraging2 && UpdateType != ECollisionUpdateType::Any)
 				{
 					ispc::SampleBoxNormalAverage(
 						(ispc::FVector&)Box->Min(),
 						(ispc::FVector&)Box->Max(),
 						(ispc::FTransform&)SampleToObjectTM,
-						(ispc::FVector*)&SampleParticles.XView()[0],
+						(ispc::FVector*)&SampleParticles.XArray()[0],
 						&PotentialParticles[0],
 						Thickness,
 						TotalThickness,
@@ -1516,7 +1541,7 @@ void SampleObject(const TImplicitObject<float, 3>& Object, const TRigidTransform
 						(ispc::FVector&)Box->Min(),
 						(ispc::FVector&)Box->Max(),
 						(ispc::FTransform&)SampleToObjectTM,
-						(ispc::FVector*)&SampleParticles.XView()[0],
+						(ispc::FVector*)&SampleParticles.XArray()[0],
 						&PotentialParticles[0],
 						DeepestParticle,
 						AvgConstraint.Phi,
@@ -1534,13 +1559,13 @@ void SampleObject(const TImplicitObject<float, 3>& Object, const TRigidTransform
 				//QUICK_SCOPE_CYCLE_COUNTER(STAT_Other);
 				for (int32 i : PotentialParticles)
 				{
-					if (NormalAveraging && UpdateType != ECollisionUpdateType::Any)
+					if (NormalAveraging2 && UpdateType != ECollisionUpdateType::Any)
 					{
-						SampleObjectNormalAverageHelper(Object, ObjectTransform, SampleToObjectTM, SampleParticles.X(i), Thickness, TotalThickness, AvgConstraint);
+						SampleObjectNormalAverageHelper2(Object, ObjectTransform, SampleToObjectTM, SampleParticles.X(i), Thickness, TotalThickness, AvgConstraint);
 					}
 					else
 					{
-						if (SampleObjectNoNormal(Object, ObjectTransform, SampleToObjectTM, SampleParticles.X(i), Thickness, AvgConstraint))
+						if (SampleObjectNoNormal2(Object, ObjectTransform, SampleToObjectTM, SampleParticles.X(i), Thickness, AvgConstraint))
 						{
 							DeepestParticle = i;
 							if (UpdateType == ECollisionUpdateType::Any)
@@ -1556,13 +1581,13 @@ void SampleObject(const TImplicitObject<float, 3>& Object, const TRigidTransform
 	}
 	else
 	{
-		SCOPE_CYCLE_COUNTER(STAT_UpdateLevelsetAll);
+		SCOPE_CYCLE_COUNTER(STAT_UpdateLevelsetAll2);
 		if (Object.GetType(true) == ImplicitObjectType::LevelSet && NumParticles > 0)
 		{
 			const TLevelSet<float, 3>* LevelSet = Object.GetObject<Chaos::TLevelSet<float, 3>>();
 			const TUniformGrid<float, 3>& Grid = LevelSet->GetGrid();
 
-			if (NormalAveraging && UpdateType != ECollisionUpdateType::Any)
+			if (NormalAveraging2 && UpdateType != ECollisionUpdateType::Any)
 			{
 				ispc::SampleLevelSetNormalAverageAll(
 					(ispc::FVector&)Grid.MinCorner(),
@@ -1571,7 +1596,7 @@ void SampleObject(const TImplicitObject<float, 3>& Object, const TRigidTransform
 					(ispc::FIntVector&)Grid.Counts(),
 					(ispc::TArrayND*)&LevelSet->GetPhiArray(),
 					(ispc::FTransform&)SampleToObjectTM,
-					(ispc::FVector*)&SampleParticles.XView()[0],
+					(ispc::FVector*)&SampleParticles.XArray()[0],
 					Thickness,
 					TotalThickness,
 					(ispc::FVector&)AvgConstraint.Location,
@@ -1586,7 +1611,7 @@ void SampleObject(const TImplicitObject<float, 3>& Object, const TRigidTransform
 					(ispc::FIntVector&)Grid.Counts(),
 					(ispc::TArrayND*)&LevelSet->GetPhiArray(),
 					(ispc::FTransform&)SampleToObjectTM,
-					(ispc::FVector*)&SampleParticles.XView()[0],
+					(ispc::FVector*)&SampleParticles.XArray()[0],
 					DeepestParticle,
 					AvgConstraint.Phi,
 					NumParticles);
@@ -1602,13 +1627,13 @@ void SampleObject(const TImplicitObject<float, 3>& Object, const TRigidTransform
 		{
 			const TPlane<float, 3>* Plane = Object.GetObject<Chaos::TPlane<float, 3>>();
 
-			if (NormalAveraging && UpdateType != ECollisionUpdateType::Any)
+			if (NormalAveraging2 && UpdateType != ECollisionUpdateType::Any)
 			{
 				ispc::SamplePlaneNormalAverageAll(
 					(ispc::FVector&)Plane->Normal(),
 					(ispc::FVector&)Plane->X(),
 					(ispc::FTransform&)SampleToObjectTM,
-					(ispc::FVector*)&SampleParticles.XView()[0],
+					(ispc::FVector*)&SampleParticles.XArray()[0],
 					Thickness,
 					TotalThickness,
 					(ispc::FVector&)AvgConstraint.Location,
@@ -1620,7 +1645,7 @@ void SampleObject(const TImplicitObject<float, 3>& Object, const TRigidTransform
 					(ispc::FVector&)Plane->Normal(),
 					(ispc::FVector&)Plane->X(),
 					(ispc::FTransform&)SampleToObjectTM,
-					(ispc::FVector*)&SampleParticles.XView()[0],
+					(ispc::FVector*)&SampleParticles.XArray()[0],
 					DeepestParticle,
 					AvgConstraint.Phi,
 					NumParticles);
@@ -1636,13 +1661,13 @@ void SampleObject(const TImplicitObject<float, 3>& Object, const TRigidTransform
 		{
 			const TBox<float, 3>* Box = Object.GetObject<Chaos::TBox<float, 3>>();
 
-			if (NormalAveraging && UpdateType != ECollisionUpdateType::Any)
+			if (NormalAveraging2 && UpdateType != ECollisionUpdateType::Any)
 			{
 				ispc::SampleBoxNormalAverageAll(
 					(ispc::FVector&)Box->Min(),
 					(ispc::FVector&)Box->Max(),
 					(ispc::FTransform&)SampleToObjectTM,
-					(ispc::FVector*)&SampleParticles.XView()[0],
+					(ispc::FVector*)&SampleParticles.XArray()[0],
 					Thickness,
 					TotalThickness,
 					(ispc::FVector&)AvgConstraint.Location,
@@ -1654,7 +1679,7 @@ void SampleObject(const TImplicitObject<float, 3>& Object, const TRigidTransform
 					(ispc::FVector&)Box->Min(),
 					(ispc::FVector&)Box->Max(),
 					(ispc::FTransform&)SampleToObjectTM,
-					(ispc::FVector*)&SampleParticles.XView()[0],
+					(ispc::FVector*)&SampleParticles.XArray()[0],
 					DeepestParticle,
 					AvgConstraint.Phi,
 					NumParticles);
@@ -1671,13 +1696,13 @@ void SampleObject(const TImplicitObject<float, 3>& Object, const TRigidTransform
 			//QUICK_SCOPE_CYCLE_COUNTER(STAT_Other);
 			for (int32 i = 0; i < NumParticles; ++i)
 			{
-				if (NormalAveraging && UpdateType != ECollisionUpdateType::Any)
+				if (NormalAveraging2 && UpdateType != ECollisionUpdateType::Any)
 				{
-					SampleObjectNormalAverageHelper(Object, ObjectTransform, SampleToObjectTM, SampleParticles.X(i), Thickness, TotalThickness, AvgConstraint);
+					SampleObjectNormalAverageHelper2(Object, ObjectTransform, SampleToObjectTM, SampleParticles.X(i), Thickness, TotalThickness, AvgConstraint);
 				}
 				else
 				{
-					if (SampleObjectNoNormal(Object, ObjectTransform, SampleToObjectTM, SampleParticles.X(i), Thickness, AvgConstraint))
+					if (SampleObjectNoNormal2(Object, ObjectTransform, SampleToObjectTM, SampleParticles.X(i), Thickness, AvgConstraint))
 					{
 						DeepestParticle = i;
 						if (UpdateType == ECollisionUpdateType::Any)
@@ -1691,7 +1716,7 @@ void SampleObject(const TImplicitObject<float, 3>& Object, const TRigidTransform
 		}
 	}
 
-	if (NormalAveraging)
+	if (NormalAveraging2)
 	{
 		if (TotalThickness < -KINDA_SMALL_NUMBER)
 		{
@@ -1786,16 +1811,16 @@ bool UpdateBoxPlaneConstraint(const TBox<T, d>& Box, const TRigidTransform<T, d>
 template <typename T, int d>
 void UpdateSphereConstraint(const TSphere<T, d>& Sphere1, const TRigidTransform<T, d>& Sphere1Transform, const TSphere<T, d>& Sphere2, const TRigidTransform<T, d>& Sphere2Transform, T Thickness, TRigidBodyContactConstraint<T, d>& Constraint)
 {
-	const TVector<T, d> Center1 = Sphere1Transform.TransformPosition(Sphere1.Center());
-	const TVector<T, d> Center2 = Sphere2Transform.TransformPosition(Sphere2.Center());
+	const TVector<T, d> Center1 = Sphere1Transform.TransformPosition(Sphere1.GetCenter());
+	const TVector<T, d> Center2 = Sphere2Transform.TransformPosition(Sphere2.GetCenter());
 	const TVector<T, d> Direction = Center1 - Center2;
 	const T Size = Direction.Size();
-	const T NewPhi = Size - (Sphere1.Radius() + Sphere2.Radius());
+	const T NewPhi = Size - (Sphere1.GetRadius() + Sphere2.GetRadius());
 	if (NewPhi < Constraint.Phi)
 	{
 		Constraint.Normal = Size > SMALL_NUMBER ? Direction / Size : TVector<T, d>(0, 0, 1);
 		Constraint.Phi = NewPhi;
-		Constraint.Location = Center1 - Sphere1.Radius() * Constraint.Normal;
+		Constraint.Location = Center1 - Sphere1.GetRadius() * Constraint.Normal;
 	}
 }
 
@@ -1803,17 +1828,17 @@ template <typename T, int d>
 void UpdateSpherePlaneConstraint(const TSphere<T, d>& Sphere, const TRigidTransform<T, d>& SphereTransform, const TPlane<T, d>& Plane, const TRigidTransform<T, d>& PlaneTransform, T Thickness, TRigidBodyContactConstraint<T, d>& Constraint)
 {
 	const TRigidTransform<T, d> SphereToPlaneTransform(PlaneTransform.Inverse() * SphereTransform);
-	const TVector<T, d> SphereCenter = SphereToPlaneTransform.TransformPosition(Sphere.Center());
+	const TVector<T, d> SphereCenter = SphereToPlaneTransform.TransformPosition(Sphere.GetCenter());
 
 	TVector<T, d> NewNormal;
 	T NewPhi = Plane.PhiWithNormal(SphereCenter, NewNormal);
-	NewPhi -= Sphere.Radius();
+	NewPhi -= Sphere.GetRadius();
 
 	if (NewPhi < Constraint.Phi)
 	{
 		Constraint.Phi = NewPhi;
 		Constraint.Normal = PlaneTransform.TransformVectorNoScale(NewNormal);
-		Constraint.Location = SphereCenter - Constraint.Normal * Sphere.Radius();
+		Constraint.Location = SphereCenter - Constraint.Normal * Sphere.GetRadius();
 	}
 }
 
@@ -1821,28 +1846,28 @@ template <typename T, int d>
 bool UpdateSphereBoxConstraint(const TSphere<T, d>& Sphere, const TRigidTransform<T, d>& SphereTransform, const TBox<T, d>& Box, const TRigidTransform<T, d>& BoxTransform, T Thickness, TRigidBodyContactConstraint<T, d>& Constraint)
 {
 	const TRigidTransform<T, d> SphereToBoxTransform(SphereTransform * BoxTransform.Inverse());
-	const TVector<T, d> SphereCenterInBox = SphereToBoxTransform.TransformPosition(Sphere.Center());
+	const TVector<T, d> SphereCenterInBox = SphereToBoxTransform.TransformPosition(Sphere.GetCenter());
 
 	TVector<T, d> NewNormal;
 	T NewPhi = Box.PhiWithNormal(SphereCenterInBox, NewNormal);
-	NewPhi -= Sphere.Radius();
+	NewPhi -= Sphere.GetRadius();
 
 	if (NewPhi < Constraint.Phi)
 	{
 		Constraint.Phi = NewPhi;
 		Constraint.Normal = BoxTransform.TransformVectorNoScale(NewNormal);
-		Constraint.Location = SphereTransform.TransformPosition(Sphere.Center()) - Constraint.Normal * Sphere.Radius();
+		Constraint.Location = SphereTransform.TransformPosition(Sphere.GetCenter()) - Constraint.Normal * Sphere.GetRadius();
 		return true;
 	}
 
 	return false;
 }
 
-DECLARE_CYCLE_STAT(TEXT("FindRelevantShapes"), STAT_FindRelevantShapes, STATGROUP_ChaosWide);
+DECLARE_CYCLE_STAT(TEXT("FindRelevantShapes2"), STAT_FindRelevantShapes2, STATGROUP_ChaosWide);
 template <typename T, int d>
-TArray<Pair<const TImplicitObject<T, d>*, TRigidTransform<T, d>>> FindRelevantShapes(const TImplicitObject<T,d>* ParticleObj, const TRigidTransform<T,d>& ParticlesTM, const TImplicitObject<T,d>& LevelsetObj, const TRigidTransform<T,d>& LevelsetTM, const T Thickness)
+TArray<Pair<const TImplicitObject<T, d>*, TRigidTransform<T, d>>> FindRelevantShapes2(const TImplicitObject<T,d>* ParticleObj, const TRigidTransform<T,d>& ParticlesTM, const TImplicitObject<T,d>& LevelsetObj, const TRigidTransform<T,d>& LevelsetTM, const T Thickness)
 {
-	SCOPE_CYCLE_COUNTER(STAT_FindRelevantShapes);
+	SCOPE_CYCLE_COUNTER(STAT_FindRelevantShapes2);
 	TArray<Pair<const TImplicitObject<T, d>*, TRigidTransform<T, d>>> RelevantShapes;
 	//find all levelset inner objects
 	if (ParticleObj)
@@ -1870,19 +1895,21 @@ TArray<Pair<const TImplicitObject<T, d>*, TRigidTransform<T, d>>> FindRelevantSh
 	return RelevantShapes;
 }
 
-DECLARE_CYCLE_STAT(TEXT("UpdateUnionUnionConstraint"), STAT_UpdateUnionUnionConstraint, STATGROUP_ChaosWide);
-template<ECollisionUpdateType UpdateType, class T_PARTICLES, typename T, int d>
-void UpdateUnionUnionConstraint(const T_PARTICLES& InParticles, const T Thickness, TRigidBodyContactConstraint<T, d>& Constraint)
+DECLARE_CYCLE_STAT(TEXT("UpdateUnionUnionConstraint"), STAT_UpdateUnionUnionConstraint2, STATGROUP_ChaosWide);
+template<ECollisionUpdateType UpdateType, typename T, int d>
+void UpdateUnionUnionConstraint(const T Thickness, TRigidBodyContactConstraint<T, d>& Constraint)
 {
-	SCOPE_CYCLE_COUNTER(STAT_UpdateUnionUnionConstraint);
-	Constraint.Phi = Thickness;
+	SCOPE_CYCLE_COUNTER(STAT_UpdateUnionUnionConstraint2);
 
-	TRigidTransform<T, d> ParticlesTM = GetTransform(InParticles, Constraint.ParticleIndex);
-	TRigidTransform<T, d> LevelsetTM = GetTransform(InParticles, Constraint.LevelsetIndex);
+	TGenericParticleHandle<T, d> Particle0 = Constraint.Particle;
+	TGenericParticleHandle<T, d> Particle1 = Constraint.Levelset;
 
-	const TImplicitObject<T, d>* ParticleObj = InParticles.Geometry(Constraint.ParticleIndex).Get();
-	const TImplicitObject<T,d>* LevelsetObj = InParticles.Geometry(Constraint.LevelsetIndex).Get();
-	const TArray<Pair<const TImplicitObject<T, d>*, TRigidTransform<T, d>>> LevelsetShapes = FindRelevantShapes(ParticleObj, ParticlesTM, *LevelsetObj, LevelsetTM, Thickness);
+	TRigidTransform<T, d> ParticlesTM = TRigidTransform<T, d>(Particle0->P(), Particle0->Q());
+	TRigidTransform<T, d> LevelsetTM = TRigidTransform<T, d>(Particle1->P(), Particle1->Q());
+
+	const TImplicitObject<T, d>* ParticleObj = Particle0->Geometry().Get();
+	const TImplicitObject<T,d>* LevelsetObj = Particle1->Geometry().Get();
+	const TArray<Pair<const TImplicitObject<T, d>*, TRigidTransform<T, d>>> LevelsetShapes = FindRelevantShapes2(ParticleObj, ParticlesTM, *LevelsetObj, LevelsetTM, Thickness);
 
 	for (const Pair<const TImplicitObject<T, d>*, TRigidTransform<T, d>>& LevelsetObjPair : LevelsetShapes)
 	{
@@ -1890,77 +1917,83 @@ void UpdateUnionUnionConstraint(const T_PARTICLES& InParticles, const T Thicknes
 		const TRigidTransform<T, d>& LevelsetInnerObjTM = LevelsetObjPair.Second * LevelsetTM;
 
 		//now find all particle inner objects
-		const TArray<Pair<const TImplicitObject<T, d>*, TRigidTransform<T, d>>> ParticleShapes = FindRelevantShapes(&LevelsetInnerObj, LevelsetInnerObjTM, *ParticleObj, ParticlesTM, Thickness);
+		const TArray<Pair<const TImplicitObject<T, d>*, TRigidTransform<T, d>>> ParticleShapes = FindRelevantShapes2(&LevelsetInnerObj, LevelsetInnerObjTM, *ParticleObj, ParticlesTM, Thickness);
 
 		//for each inner obj pair, update constraint
 		for (const Pair<const TImplicitObject<T, d>*, TRigidTransform<T, d>>& ParticlePair : ParticleShapes)
 		{
 			const TImplicitObject<T, d>& ParticleInnerObj = *ParticlePair.First;
 			const TRigidTransform<T, d> ParticleInnerObjTM = ParticlePair.Second * ParticlesTM;
-			UpdateConstraintImp<UpdateType>(InParticles, ParticleInnerObj, ParticleInnerObjTM, LevelsetInnerObj, LevelsetInnerObjTM, Thickness, Constraint);
+			UpdateConstraintImp2<UpdateType>(ParticleInnerObj, ParticleInnerObjTM, LevelsetInnerObj, LevelsetInnerObjTM, Thickness, Constraint);
 		}
 	}
 }
 
-DECLARE_CYCLE_STAT(TEXT("UpdateSingleUnionConstraint"), STAT_UpdateSingleUnionConstraint, STATGROUP_ChaosWide);
-template<ECollisionUpdateType UpdateType, class T_PARTICLES, typename T, int d>
-void UpdateSingleUnionConstraint(const T_PARTICLES& InParticles, const T Thickness, TRigidBodyContactConstraint<T, d>& Constraint)
+DECLARE_CYCLE_STAT(TEXT("UpdateSingleUnionConstraint"), STAT_UpdateSingleUnionConstraint2, STATGROUP_ChaosWide);
+template<ECollisionUpdateType UpdateType, typename T, int d>
+void UpdateSingleUnionConstraint(const T Thickness, TRigidBodyContactConstraint<T, d>& Constraint)
 {
-	SCOPE_CYCLE_COUNTER(STAT_UpdateSingleUnionConstraint);
-	Constraint.Phi = Thickness;
+	SCOPE_CYCLE_COUNTER(STAT_UpdateSingleUnionConstraint2);
+
+	TGenericParticleHandle<T, d> Particle0 = Constraint.Particle;
+	TGenericParticleHandle<T, d> Particle1 = Constraint.Levelset;
 
-	TRigidTransform<T, d> ParticlesTM = GetTransform(InParticles, Constraint.ParticleIndex);
-	TRigidTransform<T, d> LevelsetTM = GetTransform(InParticles, Constraint.LevelsetIndex);
+	TRigidTransform<T, d> ParticlesTM = TRigidTransform<T, d>(Particle0->P(), Particle0->Q());
+	TRigidTransform<T, d> LevelsetTM = TRigidTransform<T, d>(Particle1->P(), Particle1->Q());
 
-	const TImplicitObject<T, d>* ParticleObj = InParticles.Geometry(Constraint.ParticleIndex).Get();
-	const TImplicitObject<T, d>* LevelsetObj = InParticles.Geometry(Constraint.LevelsetIndex).Get();
-	const TArray<Pair<const TImplicitObject<T, d>*, TRigidTransform<T, d>>> LevelsetShapes = FindRelevantShapes(ParticleObj, ParticlesTM, *LevelsetObj, LevelsetTM, Thickness);
+	const TImplicitObject<T, d>* ParticleObj = Particle0->Geometry().Get();
+	const TImplicitObject<T, d>* LevelsetObj = Particle1->Geometry().Get();
+	const TArray<Pair<const TImplicitObject<T, d>*, TRigidTransform<T, d>>> LevelsetShapes = FindRelevantShapes2(ParticleObj, ParticlesTM, *LevelsetObj, LevelsetTM, Thickness);
 	
 	for (const Pair<const TImplicitObject<T, d>*, TRigidTransform<T, d>>& LevelsetObjPair : LevelsetShapes)
 	{
 		const TImplicitObject<T, d>& LevelsetInnerObj = *LevelsetObjPair.First;
 		const TRigidTransform<T, d> LevelsetInnerObjTM = LevelsetTM * LevelsetObjPair.Second;
-		UpdateConstraintImp<UpdateType>(InParticles, *ParticleObj, ParticlesTM, LevelsetInnerObj, LevelsetInnerObjTM, Thickness, Constraint);
+		UpdateConstraintImp2<UpdateType>(*ParticleObj, ParticlesTM, LevelsetInnerObj, LevelsetInnerObjTM, Thickness, Constraint);
 	}
 }
 
-DECLARE_CYCLE_STAT(TEXT("UpdateLevelsetConstraint"), STAT_UpdateLevelsetConstraint, STATGROUP_ChaosWide);
+DECLARE_CYCLE_STAT(TEXT("UpdateLevelsetConstraint"), STAT_UpdateLevelsetConstraint2, STATGROUP_ChaosWide);
 template<typename T, int d>
-template<ECollisionUpdateType UpdateType, class T_PARTICLES>
-void TPBDCollisionConstraint<T, d>::UpdateLevelsetConstraint(const T_PARTICLES& InParticles, const T Thickness, TRigidBodyContactConstraint<T, d>& Constraint)
+template<ECollisionUpdateType UpdateType>
+void TPBDCollisionConstraint<T, d>::UpdateLevelsetConstraint(const T Thickness, FRigidBodyContactConstraint& Constraint)
 {
-	SCOPE_CYCLE_COUNTER(STAT_UpdateLevelsetConstraint);
-	Constraint.Phi = Thickness;
-
-	TRigidTransform<T, d> ParticlesTM = GetTransform(InParticles, Constraint.ParticleIndex);
+	SCOPE_CYCLE_COUNTER(STAT_UpdateLevelsetConstraint2);
+	
+	TGenericParticleHandle<T, d> Particle0 = Constraint.Particle;
+	TRigidTransform<T, d> ParticlesTM = TRigidTransform<T, d>(Particle0->P(), Particle0->Q());
 	if (!(ensure(!FMath::IsNaN(ParticlesTM.GetTranslation().X)) && ensure(!FMath::IsNaN(ParticlesTM.GetTranslation().Y)) && ensure(!FMath::IsNaN(ParticlesTM.GetTranslation().Z))))
 	{
 		return;
 	}
-	TRigidTransform<T, d> LevelsetTM = GetTransform(InParticles, Constraint.LevelsetIndex);
+
+	TGenericParticleHandle<T, d> Particle1 = Constraint.Levelset;
+	TRigidTransform<T, d> LevelsetTM = TRigidTransform<T, d>(Particle1->P(), Particle1->Q());
 	if (!(ensure(!FMath::IsNaN(LevelsetTM.GetTranslation().X)) && ensure(!FMath::IsNaN(LevelsetTM.GetTranslation().Y)) && ensure(!FMath::IsNaN(LevelsetTM.GetTranslation().Z))))
 	{
 		return;
 	}
 
 	const TBVHParticles<T, d>* SampleParticles = nullptr;
-	SampleParticles = InParticles.CollisionParticles(Constraint.ParticleIndex).Get();
+	SampleParticles = Particle0->CollisionParticles().Get();
 
 	if(SampleParticles)
 	{
-		SampleObject<UpdateType>(*InParticles.Geometry(Constraint.LevelsetIndex), LevelsetTM, *SampleParticles, ParticlesTM, Thickness, Constraint);
+		SampleObject2<UpdateType>(*Particle1->Geometry(), LevelsetTM, *SampleParticles, ParticlesTM, Thickness, Constraint);
 	}
 }
 
-DECLARE_CYCLE_STAT(TEXT("UpdateUnionLevelsetConstraint"), STAT_UpdateUnionLevelsetConstraint, STATGROUP_ChaosWide);
-template<ECollisionUpdateType UpdateType, class T_PARTICLES, typename T, int d>
-void UpdateUnionLevelsetConstraint(const T_PARTICLES& InParticles, const T Thickness, TRigidBodyContactConstraint<T, d>& Constraint)
+DECLARE_CYCLE_STAT(TEXT("UpdateUnionLevelsetConstraint"), STAT_UpdateUnionLevelsetConstraint2, STATGROUP_ChaosWide);
+template<ECollisionUpdateType UpdateType, typename T, int d>
+void UpdateUnionLevelsetConstraint(const T Thickness, TRigidBodyContactConstraint<T, d>& Constraint)
 {
-	SCOPE_CYCLE_COUNTER(STAT_UpdateUnionLevelsetConstraint);
-	Constraint.Phi = Thickness;
+	SCOPE_CYCLE_COUNTER(STAT_UpdateUnionLevelsetConstraint2);
+
+	TGenericParticleHandle<T, d> Particle0 = Constraint.Particle;
+	TGenericParticleHandle<T, d> Particle1 = Constraint.Levelset;
 
-	TRigidTransform<T, d> ParticlesTM = GetTransform(InParticles, Constraint.ParticleIndex);
-	TRigidTransform<T, d> LevelsetTM = GetTransform(InParticles, Constraint.LevelsetIndex);
+	TRigidTransform<T, d> ParticlesTM = TRigidTransform<T, d>(Particle0->P(), Particle0->Q());
+	TRigidTransform<T, d> LevelsetTM = TRigidTransform<T, d>(Particle1->P(), Particle1->Q());
 
 	if (!(ensure(!FMath::IsNaN(ParticlesTM.GetTranslation().X)) && ensure(!FMath::IsNaN(ParticlesTM.GetTranslation().Y)) && ensure(!FMath::IsNaN(ParticlesTM.GetTranslation().Z))))
 	{
@@ -1972,26 +2005,27 @@ void UpdateUnionLevelsetConstraint(const T_PARTICLES& InParticles, const T Thick
 		return;
 	}
 
-	const TImplicitObject<T, d>* ParticleObj = InParticles.Geometry(Constraint.ParticleIndex).Get();
-	const TImplicitObject<T, d>* LevelsetObj = InParticles.Geometry(Constraint.LevelsetIndex).Get();
-	TArray<Pair<const TImplicitObject<T, d>*, TRigidTransform<T, d>>> LevelsetShapes = FindRelevantShapes(ParticleObj, ParticlesTM, *LevelsetObj, LevelsetTM, Thickness);
+	const TImplicitObject<T, d>* ParticleObj = Particle0->Geometry().Get();
+	const TImplicitObject<T, d>* LevelsetObj = Particle1->Geometry().Get();
+	TArray<Pair<const TImplicitObject<T, d>*, TRigidTransform<T, d>>> LevelsetShapes = FindRelevantShapes2(ParticleObj, ParticlesTM, *LevelsetObj, LevelsetTM, Thickness);
 
-	if (LevelsetShapes.Num() && InParticles.CollisionParticles(Constraint.ParticleIndex).Get())
+	if (LevelsetShapes.Num() && Particle0->CollisionParticles().Get())
 	{
-		const TBVHParticles<T, d>& SampleParticles = *InParticles.CollisionParticles(Constraint.ParticleIndex).Get();
+		const TBVHParticles<T, d>& SampleParticles = *Particle0->CollisionParticles().Get();
 		if (SampleParticles.Size())
 		{
 			for (const Pair<const TImplicitObject<T, d>*, TRigidTransform<T, d>>& LevelsetObjPair : LevelsetShapes)
 			{
 				const TImplicitObject<T, d>* Object = LevelsetObjPair.First;
 				const TRigidTransform<T, d> ObjectTM = LevelsetObjPair.Second * LevelsetTM;
-				SampleObject<UpdateType>(*Object, ObjectTM, SampleParticles, ParticlesTM, Thickness, Constraint);
+				SampleObject2<UpdateType>(*Object, ObjectTM, SampleParticles, ParticlesTM, Thickness, Constraint);
 				if (UpdateType == ECollisionUpdateType::Any && Constraint.Phi < Thickness)
 				{
 					return;
 				}
 			}
 		}
+#if CHAOS_PARTICLEHANDLE_TODO
 		else if (ParticleObj && ParticleObj->IsUnderlyingUnion())
 		{
 			const TImplicitObjectUnion<T, d>* UnionObj = static_cast<const TImplicitObjectUnion<T, d>*>(ParticleObj);
@@ -2001,7 +2035,7 @@ void UpdateUnionLevelsetConstraint(const T_PARTICLES& InParticles, const T Thick
 				const TImplicitObject<T, d>* LevelsetInnerObject = LevelsetObjPair.First;
 				const TRigidTransform<T, d> LevelsetInnerObjectTM = LevelsetObjPair.Second * LevelsetTM;
 
-				TArray<Pair<const TImplicitObject<T, d>*, TRigidTransform<T, d>>> ParticleShapes = FindRelevantShapes(LevelsetInnerObject, LevelsetInnerObjectTM, *ParticleObj, ParticlesTM, Thickness);
+				TArray<Pair<const TImplicitObject<T, d>*, TRigidTransform<T, d>>> ParticleShapes = FindRelevantShapes2(LevelsetInnerObject, LevelsetInnerObjectTM, *ParticleObj, ParticlesTM, Thickness);
 				for (const Pair<const TImplicitObject<T, d>*, TRigidTransform<T, d>>& ParticleObjPair : ParticleShapes)
 				{
 					const TImplicitObject<T, d>* ParticleInnerObject = ParticleObjPair.First;
@@ -2010,7 +2044,7 @@ void UpdateUnionLevelsetConstraint(const T_PARTICLES& InParticles, const T Thick
 					if (const int32* OriginalIdx = UnionObj->MCollisionParticleLookupHack.Find(ParticleInnerObject))
 					{
 						const TBVHParticles<T, d>& InnerSampleParticles = *InParticles.CollisionParticles(*OriginalIdx).Get();
-						SampleObject<UpdateType>(*LevelsetInnerObject, LevelsetInnerObjectTM, InnerSampleParticles, ParticleInnerObjectTM, Thickness, Constraint);
+						SampleObject2<UpdateType>(*LevelsetInnerObject, LevelsetInnerObjectTM, InnerSampleParticles, ParticleInnerObjectTM, Thickness, Constraint);
 						if (UpdateType == ECollisionUpdateType::Any && Constraint.Phi < Thickness)
 						{
 							return;
@@ -2020,19 +2054,25 @@ void UpdateUnionLevelsetConstraint(const T_PARTICLES& InParticles, const T Thick
 
 			}
 		}
+#endif
 	}
 }
 
 
-DECLARE_CYCLE_STAT(TEXT("UpdateLevelsetUnionConstraint"), STAT_UpdateLevelsetUnionConstraint, STATGROUP_ChaosWide);
-template<ECollisionUpdateType UpdateType, class T_PARTICLES, typename T, int d>
-void UpdateLevelsetUnionConstraint(const T_PARTICLES& InParticles, const T Thickness, TRigidBodyContactConstraint<T, d>& Constraint)
+DECLARE_CYCLE_STAT(TEXT("UpdateLevelsetUnionConstraint"), STAT_UpdateLevelsetUnionConstraint2, STATGROUP_ChaosWide);
+template<ECollisionUpdateType UpdateType, typename T, int d>
+void UpdateLevelsetUnionConstraint(const T Thickness, TRigidBodyContactConstraint<T, d>& Constraint)
 {
-	SCOPE_CYCLE_COUNTER(STAT_UpdateLevelsetUnionConstraint);
-	Constraint.Phi = Thickness;
+	SCOPE_CYCLE_COUNTER(STAT_UpdateLevelsetUnionConstraint2);
+
+	TGenericParticleHandle<T, d> Particle0 = Constraint.Particle;
+	TGenericParticleHandle<T, d> Particle1 = Constraint.Levelset;
+
+	TRigidTransform<T, d> ParticlesTM = TRigidTransform<T, d>(Particle0->P(), Particle0->Q());
+	TRigidTransform<T, d> LevelsetTM = TRigidTransform<T, d>(Particle1->P(), Particle1->Q());
 
-	TRigidTransform<T, d> ParticlesTM = GetTransform(InParticles, Constraint.ParticleIndex);
-	TRigidTransform<T, d> LevelsetTM = GetTransform(InParticles, Constraint.LevelsetIndex);
+	const TImplicitObject<T, d>* ParticleObj = Particle0->Geometry().Get();
+	const TImplicitObject<T, d>* LevelsetObj = Particle1->Geometry().Get();
 
 	if (!(ensure(!FMath::IsNaN(ParticlesTM.GetTranslation().X)) && ensure(!FMath::IsNaN(ParticlesTM.GetTranslation().Y)) && ensure(!FMath::IsNaN(ParticlesTM.GetTranslation().Z))))
 	{
@@ -2044,10 +2084,8 @@ void UpdateLevelsetUnionConstraint(const T_PARTICLES& InParticles, const T Thick
 		return;
 	}
 
-	const TImplicitObject<T, d>* ParticleObj = InParticles.Geometry(Constraint.ParticleIndex).Get();
-	const TImplicitObject<T, d>* LevelsetObj = InParticles.Geometry(Constraint.LevelsetIndex).Get();
-
-	TArray<Pair<const TImplicitObject<T, d>*, TRigidTransform<T, d>>> ParticleShapes = FindRelevantShapes(LevelsetObj, LevelsetTM, *ParticleObj, ParticlesTM, Thickness);
+#if CHAOS_PARTICLEHANDLE_TODO
+	TArray<Pair<const TImplicitObject<T, d>*, TRigidTransform<T, d>>> ParticleShapes = FindRelevantShapes2(LevelsetObj, LevelsetTM, *ParticleObj, ParticlesTM, Thickness);
 	check(ParticleObj->IsUnderlyingUnion());
 	const TImplicitObjectUnion<T, d>* UnionObj = static_cast<const TImplicitObjectUnion<T, d>*>(ParticleObj);
 	for (const Pair<const TImplicitObject<T, d>*, TRigidTransform<T, d>>& ParticleObjPair : ParticleShapes)
@@ -2059,33 +2097,41 @@ void UpdateLevelsetUnionConstraint(const T_PARTICLES& InParticles, const T Thick
 			const TBVHParticles<T, d>& SampleParticles = *InParticles.CollisionParticles(*OriginalIdx).Get();
 			const TRigidTransform<T, d> ObjectTM = ParticleObjPair.Second * ParticlesTM;
 
-			SampleObject<UpdateType>(*LevelsetObj, LevelsetTM, SampleParticles, ObjectTM, Thickness, Constraint);
+			SampleObject2<UpdateType>(*LevelsetObj, LevelsetTM, SampleParticles, ObjectTM, Thickness, Constraint);
 			if (UpdateType == ECollisionUpdateType::Any && Constraint.Phi < Thickness)
 			{
 				return;
 			}
 		}
 	}
+#endif
 }
 
 template<typename T, int d>
-template<ECollisionUpdateType UpdateType, class T_PARTICLES>
-void TPBDCollisionConstraint<T, d>::UpdateLevelsetConstraintGJK(const T_PARTICLES& InParticles, const T Thickness, TRigidBodyContactConstraint<T, d>& Constraint)
+template<ECollisionUpdateType UpdateType>
+void TPBDCollisionConstraint<T, d>::UpdateLevelsetConstraintGJK(const T Thickness, FRigidBodyContactConstraint& Constraint)
 {
 	static int32 MaxIterations = 100;
-	Constraint.Phi = Thickness;
-	const TRigidTransform<T, d> LocalToWorld1 = GetTransform(InParticles, Constraint.ParticleIndex);
-	const TRigidTransform<T, d> LocalToWorld2 = GetTransform(InParticles, Constraint.LevelsetIndex);
+
+	TGenericParticleHandle<T, d> Particle0 = Constraint.Particle;
+	TGenericParticleHandle<T, d> Particle1 = Constraint.Levelset;
+
+	const TRigidTransform<T, d> LocalToWorld1 = TRigidTransform<T, d>(Particle0->P(), Particle0->Q());
+	const TRigidTransform<T, d> LocalToWorld2 = TRigidTransform<T, d>(Particle1->P(), Particle1->Q());
+
+	const TSerializablePtr<TImplicitObject<T, d>> Particle0Obj = Particle0->Geometry();
+	const TSerializablePtr<TImplicitObject<T, d>> Particle1Obj = Particle1->Geometry();
+
 	TVector<T, d> Direction = LocalToWorld1.GetTranslation() - LocalToWorld2.GetTranslation();
-	TVector<T, d> SupportA = LocalToWorld1.TransformPosition(InParticles.Geometry(Constraint.ParticleIndex)->Support(LocalToWorld1.InverseTransformVector(-Direction), Thickness));
-	TVector<T, d> SupportB = LocalToWorld2.TransformPosition(InParticles.Geometry(Constraint.LevelsetIndex)->Support(LocalToWorld2.InverseTransformVector(Direction), Thickness));
+	TVector<T, d> SupportA = LocalToWorld1.TransformPosition(Particle0Obj->Support(LocalToWorld1.InverseTransformVector(-Direction), Thickness));
+	TVector<T, d> SupportB = LocalToWorld2.TransformPosition(Particle1Obj->Support(LocalToWorld2.InverseTransformVector(Direction), Thickness));
 	TVector<T, d> Point = SupportB - SupportA;
 	TArray<Pair<TVector<T, d>, TVector<T, d>>> Points = {MakePair(Point, SupportA)};
 	Direction = -Point;
 	for (int32 i = 0; i < MaxIterations; ++i)
 	{
-		SupportA = LocalToWorld1.TransformPosition(InParticles.Geometry(Constraint.ParticleIndex)->Support(LocalToWorld1.InverseTransformVector(-Direction), Thickness));
-		SupportB = LocalToWorld2.TransformPosition(InParticles.Geometry(Constraint.LevelsetIndex)->Support(LocalToWorld2.InverseTransformVector(Direction), Thickness));
+		SupportA = LocalToWorld1.TransformPosition(Particle0Obj->Support(LocalToWorld1.InverseTransformVector(-Direction), Thickness));
+		SupportB = LocalToWorld2.TransformPosition(Particle1Obj->Support(LocalToWorld2.InverseTransformVector(Direction), Thickness));
 		Point = SupportB - SupportA;
 		if (TVector<T, d>::DotProduct(Point, Direction) < 0)
 		{
@@ -2097,11 +2143,11 @@ void TPBDCollisionConstraint<T, d>::UpdateLevelsetConstraintGJK(const T_PARTICLE
 		{
 			TVector<T, d> Body1Location = LocalToWorld1.InverseTransformPosition(ClosestPoint);
 			TVector<T, d> Normal;
-			T Phi = InParticles.Geometry(Constraint.ParticleIndex)->PhiWithNormal(Body1Location, Normal);
+			T Phi = Particle0Obj->PhiWithNormal(Body1Location, Normal);
 			Normal = LocalToWorld1.TransformVector(Normal);
 			Constraint.Location = ClosestPoint - Phi * Normal;
 			TVector<T, d> Body2Location = LocalToWorld2.InverseTransformPosition(Constraint.Location);
-			Constraint.Phi = InParticles.Geometry(Constraint.LevelsetIndex)->PhiWithNormal(Body2Location, Constraint.Normal);
+			Constraint.Phi = Particle1Obj->PhiWithNormal(Body2Location, Constraint.Normal);
 			Constraint.Normal = LocalToWorld2.TransformVector(Constraint.Normal);
 			break;
 		}
@@ -2124,17 +2170,17 @@ void UpdateBoxConstraint(const TBox<T, d>& Box1, const TRigidTransform<T, d>& Bo
 			//If Box1 is overlapping Box2 by this much the signed distance approach will fail (box1 gets sucked into box2). In this case just use two spheres
 			TSphere<T, d> Sphere1(Box1Transform.TransformPosition(Box1.Center()), Box1.Extents().Min() / 2);
 			TSphere<T, d> Sphere2(Box2Transform.TransformPosition(Box2.Center()), Box2.Extents().Min() / 2);
-			const TVector<T, d> Direction = Sphere1.Center() - Sphere2.Center();
+			const TVector<T, d> Direction = Sphere1.GetCenter() - Sphere2.GetCenter();
 			T Size = Direction.Size();
-			if (Size < (Sphere1.Radius() + Sphere2.Radius()))
+			if(Size < (Sphere1.GetRadius() + Sphere2.GetRadius()))
 			{
-				const T NewPhi = Size - (Sphere1.Radius() + Sphere2.Radius());;
+				const T NewPhi = Size - (Sphere1.GetRadius() + Sphere2.GetRadius());;
 				if (NewPhi < Constraint.Phi)
 				{
 					bDeepOverlap = true;
 					Constraint.Normal = Size > SMALL_NUMBER ? Direction / Size : TVector<T, d>(0, 0, 1);
 					Constraint.Phi = NewPhi;
-					Constraint.Location = Sphere1.Center() - Sphere1.Radius() * Constraint.Normal;
+					Constraint.Location = Sphere1.GetCenter() - Sphere1.GetRadius() * Constraint.Normal;
 				}
 			}
 		}
@@ -2151,177 +2197,178 @@ void UpdateBoxConstraint(const TBox<T, d>& Box1, const TRigidTransform<T, d>& Bo
 				int32 NumParticles = SampleParticles.Num();
 				for (int32 i = 0; i < NumParticles; ++i)
 				{
-					SampleObjectHelper(Box2, Box2Transform, Box1ToBox2Transform, SampleParticles[i], Thickness, Constraint);
+					SampleObjectHelper2(Box2, Box2Transform, Box1ToBox2Transform, SampleParticles[i], Thickness, Constraint);
 				}
 			}
 		}
 	}
 }
 
-
 template<typename T, int d>
-TRigidBodyContactConstraint<T, d> ComputeLevelsetConstraint(const TPBDRigidParticles<T, d>& InParticles, int32 ParticleIndex, int32 LevelsetIndex, const T Thickness)
+TRigidBodyContactConstraint<T, d> ComputeLevelsetConstraint(TGeometryParticleHandle<T, d>* Particle0, TGeometryParticleHandle<T, d>* Particle1, const T Thickness)
 {
-	if (!InParticles.Geometry(LevelsetIndex) || (!InParticles.CollisionParticlesSize(ParticleIndex) && InParticles.Geometry(ParticleIndex) && !InParticles.Geometry(ParticleIndex)->IsUnderlyingUnion()))
+	//todo(ocohen):if both have collision particles, use the one with fewer?
+	if (!Particle1->Geometry() || (Particle0->AsDynamic() && !Particle0->AsDynamic()->CollisionParticlesSize() && Particle0->Geometry() && !Particle0->Geometry()->IsUnderlyingUnion()))
 	{
-		int32 TmpIndex = ParticleIndex;
-		ParticleIndex = LevelsetIndex;
-		LevelsetIndex = TmpIndex;
+		TRigidBodyContactConstraint<T, d> Constraint;
+		Constraint.Particle = Particle1;
+		Constraint.Levelset = Particle0;
+		return Constraint;
+	}
+	else
+	{
+		TRigidBodyContactConstraint<T, d> Constraint;
+		Constraint.Particle = Particle0;
+		Constraint.Levelset = Particle1;
+		return Constraint;
 	}
-	check(InParticles.Geometry(LevelsetIndex));
-	//todo(ocohen):if both have collision particles, use the one with fewer?
-	// Find Deepest Point
-	TRigidBodyContactConstraint<T, d> Constraint;
-	Constraint.ParticleIndex = ParticleIndex;
-	Constraint.LevelsetIndex = LevelsetIndex;
-	return Constraint;
 }
 
 template<typename T, int d>
-TRigidBodyContactConstraint<T, d> ComputeLevelsetConstraintGJK(const TPBDRigidParticles<T, d>& InParticles, int32 ParticleIndex, int32 LevelsetIndex, const T Thickness)
+TRigidBodyContactConstraint<T, d> ComputeLevelsetConstraintGJK(TGeometryParticleHandle<T, d>* Particle0, TGeometryParticleHandle<T, d>* Particle1, const T Thickness)
 {
 	TRigidBodyContactConstraint<T, d> Constraint;
-	Constraint.ParticleIndex = ParticleIndex;
-	Constraint.LevelsetIndex = LevelsetIndex;
+	Constraint.Particle = Particle0;
+	Constraint.Levelset = Particle1;
 	return Constraint;
 }
 
 template<typename T, int d>
-TRigidBodyContactConstraint<T, d> ComputeBoxConstraint(const TPBDRigidParticles<T, d>& InParticles, int32 Box1Index, int32 Box2Index, const T Thickness)
+TRigidBodyContactConstraint<T, d> ComputeBoxConstraint(TGeometryParticleHandle<T, d>* Particle0, TGeometryParticleHandle<T, d>* Particle1, const T Thickness)
 {
 	TRigidBodyContactConstraint<T, d> Constraint;
-	Constraint.ParticleIndex = Box1Index;
-	Constraint.LevelsetIndex = Box2Index;
+	Constraint.Particle = Particle0;
+	Constraint.Levelset = Particle1;
 	return Constraint;
 }
 
 template<typename T, int d>
-TRigidBodyContactConstraint<T, d> ComputeBoxPlaneConstraint(const TPBDRigidParticles<T, d>& InParticles, int32 BoxIndex, int32 PlaneIndex, const T Thickness)
+TRigidBodyContactConstraint<T, d> ComputeBoxPlaneConstraint(TGeometryParticleHandle<T, d>* Particle0, TGeometryParticleHandle<T, d>* Particle1, const T Thickness)
 {
 	TRigidBodyContactConstraint<T, d> Constraint;
-	Constraint.ParticleIndex = BoxIndex;
-	Constraint.LevelsetIndex = PlaneIndex;
+	Constraint.Particle = Particle0;
+	Constraint.Levelset = Particle1;
 	return Constraint;
 }
 
 template<typename T, int d>
-TRigidBodyContactConstraint<T, d> ComputeSphereConstraint(const TPBDRigidParticles<T, d>& InParticles, int32 Sphere1Index, int32 Sphere2Index, const T Thickness)
+TRigidBodyContactConstraint<T, d> ComputeSphereConstraint(TGeometryParticleHandle<T, d>* Particle0, TGeometryParticleHandle<T, d>* Particle1, const T Thickness)
 {
 	TRigidBodyContactConstraint<T, d> Constraint;
-	Constraint.ParticleIndex = Sphere1Index;
-	Constraint.LevelsetIndex = Sphere2Index;
+	Constraint.Particle = Particle0;
+	Constraint.Levelset = Particle1;
 	return Constraint;
 }
 
 template<typename T, int d>
-TRigidBodyContactConstraint<T, d> ComputeSpherePlaneConstraint(const TPBDRigidParticles<T, d>& InParticles, int32 SphereIndex, int32 PlaneIndex, const T Thickness)
+TRigidBodyContactConstraint<T, d> ComputeSpherePlaneConstraint(TGeometryParticleHandle<T, d>* Particle0, TGeometryParticleHandle<T, d>* Particle1, const T Thickness)
 {
 	TRigidBodyContactConstraint<T, d> Constraint;
-	Constraint.ParticleIndex = SphereIndex;
-	Constraint.LevelsetIndex = PlaneIndex;
+	Constraint.Particle = Particle0;
+	Constraint.Levelset = Particle1;
 	return Constraint;
 }
 
 template<typename T, int d>
-TRigidBodyContactConstraint<T, d> ComputeSphereBoxConstraint(const TPBDRigidParticles<T, d>& InParticles, int32 SphereIndex, int32 BoxIndex, const T Thickness)
+TRigidBodyContactConstraint<T, d> ComputeSphereBoxConstraint(TGeometryParticleHandle<T, d>* Particle0, TGeometryParticleHandle<T, d>* Particle1, const T Thickness)
 {
 	TRigidBodyContactConstraint<T, d> Constraint;
-	Constraint.ParticleIndex = SphereIndex;
-	Constraint.LevelsetIndex = BoxIndex;
+	Constraint.Particle = Particle0;
+	Constraint.Levelset = Particle1;
 	return Constraint;
 }
 
 template <typename T, int d>
-TRigidBodyContactConstraint<T, d> ComputeSingleUnionConstraint(const TPBDRigidParticles<T, d>& InParticles, int32 NonUnionIndex, int32 UnionIndex, const T Thickness)
+TRigidBodyContactConstraint<T, d> ComputeSingleUnionConstraint(TGeometryParticleHandle<T, d>* Particle0, TGeometryParticleHandle<T, d>* Particle1, const T Thickness)
 {
 	TRigidBodyContactConstraint<T, d> Constraint;
-	Constraint.ParticleIndex = NonUnionIndex;
-	Constraint.LevelsetIndex = UnionIndex;
+	Constraint.Particle = Particle0;
+	Constraint.Levelset = Particle1;
 	return Constraint;
 }
 
 template <typename T, int d>
-TRigidBodyContactConstraint<T, d> ComputeUnionUnionConstraint(const TPBDRigidParticles<T, d>& InParticles, int32 Union1Index, int32 Union2Index, const T Thickness)
+TRigidBodyContactConstraint<T, d> ComputeUnionUnionConstraint(TGeometryParticleHandle<T, d>* Particle0, TGeometryParticleHandle<T, d>* Particle1, const T Thickness)
 {
 	TRigidBodyContactConstraint<T, d> Constraint;
-	Constraint.ParticleIndex = Union1Index;
-	Constraint.LevelsetIndex = Union2Index;
+	Constraint.Particle = Particle0;
+	Constraint.Levelset = Particle1;
 	//todo(ocohen): some heuristic for determining the order?
 	return Constraint;
 }
 
 template<typename T, int d>
-typename TPBDCollisionConstraint<T, d>::FRigidBodyContactConstraint TPBDCollisionConstraint<T, d>::ComputeConstraint(const TPBDRigidParticles<T, d>& InParticles, int32 Body1Index, int32 Body2Index, const T Thickness)
+typename TPBDCollisionConstraint<T, d>::FRigidBodyContactConstraint TPBDCollisionConstraint<T, d>::ComputeConstraint(TGeometryParticleHandle<T, d>* Particle0, TGeometryParticleHandle<T, d>* Particle1, const T Thickness)
 {
-	if (!InParticles.Geometry(Body1Index) || !InParticles.Geometry(Body2Index))
+	if (!Particle0->Geometry() || !Particle1->Geometry())
 	{
-		return ComputeLevelsetConstraint(InParticles, Body1Index, Body2Index, Thickness);
+		return ComputeLevelsetConstraint(Particle0, Particle1, Thickness);
 	}
-	if (InParticles.Geometry(Body1Index)->GetType() == TBox<T, d>::GetType() && InParticles.Geometry(Body2Index)->GetType() == TBox<T, d>::GetType())
+	if (Particle0->Geometry()->GetType() == TBox<T, d>::GetType() && Particle1->Geometry()->GetType() == TBox<T, d>::GetType())
 	{
-		return ComputeBoxConstraint(InParticles, Body1Index, Body2Index, Thickness);
+		return ComputeBoxConstraint(Particle0, Particle1, Thickness);
 	}
-	else if (InParticles.Geometry(Body1Index)->GetType() == TSphere<T, d>::GetType() && InParticles.Geometry(Body2Index)->GetType() == TSphere<T, d>::GetType())
+	else if (Particle0->Geometry()->GetType() == TSphere<T, d>::GetType() && Particle1->Geometry()->GetType() == TSphere<T, d>::GetType())
 	{
-		return ComputeSphereConstraint(InParticles, Body1Index, Body2Index, Thickness);
+		return ComputeSphereConstraint(Particle0, Particle1, Thickness);
 	}
-	else if (InParticles.Geometry(Body1Index)->GetType() == TBox<T, d>::GetType() && InParticles.Geometry(Body2Index)->GetType() == TPlane<T, d>::GetType())
+	else if (Particle0->Geometry()->GetType() == TBox<T, d>::GetType() && Particle1->Geometry()->GetType() == TPlane<T, d>::GetType())
 	{
-		return ComputeBoxPlaneConstraint(InParticles, Body1Index, Body2Index, Thickness);
+		return ComputeBoxPlaneConstraint(Particle0, Particle1, Thickness);
 	}
-	else if (InParticles.Geometry(Body2Index)->GetType() == TPlane<T, d>::GetType() && InParticles.Geometry(Body1Index)->GetType() == TBox<T, d>::GetType())
+	else if (Particle1->Geometry()->GetType() == TPlane<T, d>::GetType() && Particle0->Geometry()->GetType() == TBox<T, d>::GetType())
 	{
-		return ComputeBoxPlaneConstraint(InParticles, Body2Index, Body1Index, Thickness);
+		return ComputeBoxPlaneConstraint(Particle1, Particle0, Thickness);
 	}
-	else if (InParticles.Geometry(Body1Index)->GetType() == TSphere<T, d>::GetType() && InParticles.Geometry(Body2Index)->GetType() == TPlane<T, d>::GetType())
+	else if (Particle0->Geometry()->GetType() == TSphere<T, d>::GetType() && Particle1->Geometry()->GetType() == TPlane<T, d>::GetType())
 	{
-		return ComputeSpherePlaneConstraint(InParticles, Body1Index, Body2Index, Thickness);
+		return ComputeSpherePlaneConstraint(Particle0, Particle1, Thickness);
 	}
-	else if (InParticles.Geometry(Body2Index)->GetType() == TPlane<T, d>::GetType() && InParticles.Geometry(Body1Index)->GetType() == TSphere<T, d>::GetType())
+	else if (Particle1->Geometry()->GetType() == TPlane<T, d>::GetType() && Particle0->Geometry()->GetType() == TSphere<T, d>::GetType())
 	{
-		return ComputeSpherePlaneConstraint(InParticles, Body2Index, Body1Index, Thickness);
+		return ComputeSpherePlaneConstraint(Particle1, Particle0, Thickness);
 	}
-	else if (InParticles.Geometry(Body1Index)->GetType() == TSphere<T, d>::GetType() && InParticles.Geometry(Body2Index)->GetType() == TBox<T, d>::GetType())
+	else if (Particle0->Geometry()->GetType() == TSphere<T, d>::GetType() && Particle1->Geometry()->GetType() == TBox<T, d>::GetType())
 	{
-		return ComputeSphereBoxConstraint(InParticles, Body1Index, Body2Index, Thickness);
+		return ComputeSphereBoxConstraint(Particle0, Particle1, Thickness);
 	}
-	else if (InParticles.Geometry(Body1Index)->GetType() < TImplicitObjectUnion<T, d>::GetType() && InParticles.Geometry(Body2Index)->GetType() == TImplicitObjectUnion<T, d>::GetType())
+	else if (Particle0->Geometry()->GetType() < TImplicitObjectUnion<T, d>::GetType() && Particle1->Geometry()->GetType() == TImplicitObjectUnion<T, d>::GetType())
 	{
-		return ComputeSingleUnionConstraint(InParticles, Body1Index, Body2Index, Thickness);
+		return ComputeSingleUnionConstraint(Particle0, Particle1, Thickness);
 	}
-	else if (InParticles.Geometry(Body1Index)->GetType() == TImplicitObjectUnion<T, d>::GetType() && InParticles.Geometry(Body2Index)->GetType() < TImplicitObjectUnion<T, d>::GetType())
+	else if (Particle0->Geometry()->GetType() == TImplicitObjectUnion<T, d>::GetType() && Particle1->Geometry()->GetType() < TImplicitObjectUnion<T, d>::GetType())
 	{
-		return ComputeSingleUnionConstraint(InParticles, Body2Index, Body1Index, Thickness);
+		return ComputeSingleUnionConstraint(Particle1, Particle0, Thickness);
 	}
-	else if(InParticles.Geometry(Body1Index)->GetType() == TImplicitObjectUnion<T, d>::GetType() && InParticles.Geometry(Body2Index)->GetType() == TImplicitObjectUnion<T, d>::GetType())
+	else if(Particle0->Geometry()->GetType() == TImplicitObjectUnion<T, d>::GetType() && Particle1->Geometry()->GetType() == TImplicitObjectUnion<T, d>::GetType())
 	{
-		return ComputeUnionUnionConstraint(InParticles, Body1Index, Body2Index, Thickness);
+		return ComputeUnionUnionConstraint(Particle0, Particle1, Thickness);
 	}
 #if 0
-	else if (InParticles.Geometry(Body1Index)->IsConvex() && InParticles.Geometry(Body2Index)->IsConvex())
+	else if (Particle0->Geometry()->IsConvex() && Particle1->Geometry()->IsConvex())
 	{
-		return ComputeLevelsetConstraintGJK(InParticles, Body1Index, Body2Index, Thickness);
+		return ComputeLevelsetConstraintGJK(Particle0, Particle1, Thickness);
 	}
 #endif
-	return ComputeLevelsetConstraint(InParticles, Body1Index, Body2Index, Thickness);
+	return ComputeLevelsetConstraint(Particle0, Particle1, Thickness);
 }
 
 // NOTE: UpdateLevelsetConstraintImp and its <float,3> specialization are here for the Linux build. It looks like
 // GCC does not resolve TPBDCollisionConstraint<T, d> correctly and it won't compile without them.
-template<ECollisionUpdateType UpdateType, class T_PARTICLES, typename T, int d>
-void UpdateLevelsetConstraintImp(const T_PARTICLES& InParticles, const T Thickness, TRigidBodyContactConstraint<T, d>& Constraint)
+template<ECollisionUpdateType UpdateType, typename T, int d>
+void UpdateLevelsetConstraintImp(const T Thickness, TRigidBodyContactConstraint<T, d>& Constraint)
 {
-	TPBDCollisionConstraint<T, d>::UpdateLevelsetConstraint<UpdateType>(InParticles, Thickness, Constraint);
+	TPBDCollisionConstraint<T, d>::UpdateLevelsetConstraint<UpdateType>(Thickness, Constraint);
 }
 
 template<ECollisionUpdateType UpdateType>
-void UpdateLevelsetConstraintImp(const TPBDRigidParticles<float, 3>& InParticles, const float Thickness, TRigidBodyContactConstraint<float, 3>& Constraint)
+void UpdateLevelsetConstraintImp(const float Thickness, TRigidBodyContactConstraint<float, 3>& Constraint)
 {
-	TPBDCollisionConstraint<float, 3>::UpdateLevelsetConstraint<UpdateType>(InParticles, Thickness, Constraint);
+	TPBDCollisionConstraint<float, 3>::UpdateLevelsetConstraint<UpdateType>(Thickness, Constraint);
 }
 
-template<ECollisionUpdateType UpdateType, class T_PARTICLES, typename T, int d>
-void UpdateConstraintImp(const T_PARTICLES& InParticles, const TImplicitObject<T, d>& ParticleObject, const TRigidTransform<T, d>& ParticleTM, const TImplicitObject<T, d>& LevelsetObject, const TRigidTransform<T, d>& LevelsetTM, const T Thickness, TRigidBodyContactConstraint<T, d>& Constraint)
+template<ECollisionUpdateType UpdateType, typename T, int d>
+void UpdateConstraintImp2(const TImplicitObject<T, d>& ParticleObject, const TRigidTransform<T, d>& ParticleTM, const TImplicitObject<T, d>& LevelsetObject, const TRigidTransform<T, d>& LevelsetTM, const T Thickness, TRigidBodyContactConstraint<T, d>& Constraint)
 {
 	if (ParticleObject.GetType() == TBox<T, d>::GetType() && LevelsetObject.GetType() == TBox<T, d>::GetType())
 	{
@@ -2375,7 +2422,7 @@ void UpdateConstraintImp(const T_PARTICLES& InParticles, const TImplicitObject<T
 	}
 	else if (ParticleObject.GetType() < TImplicitObjectUnion<T, d>::GetType() && LevelsetObject.GetType() == TImplicitObjectUnion<T, d>::GetType())
 	{
-		return UpdateSingleUnionConstraint<UpdateType>(InParticles, Thickness, Constraint);
+		return UpdateSingleUnionConstraint<UpdateType>(Thickness, Constraint);
 	}
 	else if (ParticleObject.GetType() == TImplicitObjectUnion<T, d>::GetType() && LevelsetObject.GetType() < TImplicitObjectUnion<T, d>::GetType())
 	{
@@ -2383,69 +2430,69 @@ void UpdateConstraintImp(const T_PARTICLES& InParticles, const TImplicitObject<T
 	}
 	else if (ParticleObject.GetType() == TImplicitObjectUnion<T, d>::GetType() && LevelsetObject.GetType() == TImplicitObjectUnion<T, d>::GetType())
 	{
-		return UpdateUnionUnionConstraint<UpdateType>(InParticles, Thickness, Constraint);
+		return UpdateUnionUnionConstraint<UpdateType>(Thickness, Constraint);
 	}
 #if 0
 	else if (ParticleObject.IsConvex() && LevelsetObject.IsConvex())
 	{
-		UpdateLevelsetConstraintGJK<UpdateType>(InParticles, Thickness, Constraint);
+		UpdateLevelsetConstraintGJK<UpdateType>(Thickness, Constraint);
 	}
 #endif
 	else if (LevelsetObject.IsUnderlyingUnion())
 	{
-		UpdateUnionLevelsetConstraint<UpdateType>(InParticles, Thickness, Constraint);
+		UpdateUnionLevelsetConstraint<UpdateType>(Thickness, Constraint);
 	}
 	else if (ParticleObject.IsUnderlyingUnion())
 	{
-		UpdateLevelsetUnionConstraint<UpdateType>(InParticles, Thickness, Constraint);
+		UpdateLevelsetUnionConstraint<UpdateType>(Thickness, Constraint);
 	}
 	else
 	{
-		UpdateLevelsetConstraintImp<UpdateType>(InParticles, Thickness, Constraint);
+		UpdateLevelsetConstraintImp<UpdateType>(Thickness, Constraint);
 	}
 }
 
-DECLARE_CYCLE_STAT(TEXT("UpdateConstraint"), STAT_UpdateConstraint, STATGROUP_ChaosWide);
+DECLARE_CYCLE_STAT(TEXT("UpdateConstraint"), STAT_UpdateConstraint2, STATGROUP_ChaosWide);
 
 template<typename T, int d>
-template<ECollisionUpdateType UpdateType, class T_PARTICLES>
-void TPBDCollisionConstraint<T, d>::UpdateConstraint(const T_PARTICLES& InParticles, const T Thickness, TRigidBodyContactConstraint<T, d>& Constraint)
+template<ECollisionUpdateType UpdateType>
+void TPBDCollisionConstraint<T, d>::UpdateConstraint(const T Thickness, FRigidBodyContactConstraint& Constraint)
 {
-	SCOPE_CYCLE_COUNTER(STAT_UpdateConstraint);
+	SCOPE_CYCLE_COUNTER(STAT_UpdateConstraint2);
 	Constraint.Phi = Thickness;
-	const TRigidTransform<T, d> ParticleTM = GetTransform(InParticles, Constraint.ParticleIndex);
-	const TRigidTransform<T, d> LevelsetTM = GetTransform(InParticles, Constraint.LevelsetIndex);
+	const TRigidTransform<T, d> ParticleTM = GetTransform(Constraint.Particle);
+	const TRigidTransform<T, d> LevelsetTM = GetTransform(Constraint.Levelset);
 
-	if (!InParticles.Geometry(Constraint.ParticleIndex))
+	if (!Constraint.Particle->Geometry())
 	{
-		if (InParticles.Geometry(Constraint.LevelsetIndex))
+		if (Constraint.Levelset->Geometry())
 		{
-			if (!InParticles.Geometry(Constraint.LevelsetIndex)->IsUnderlyingUnion())
+			if (!Constraint.Levelset->Geometry()->IsUnderlyingUnion())
 			{
-				TPBDCollisionConstraint<T, d>::UpdateLevelsetConstraint<UpdateType>(InParticles, Thickness, Constraint);
+				TPBDCollisionConstraint<T, d>::UpdateLevelsetConstraint<UpdateType>(Thickness, Constraint);
 			}
 			else
 			{
-				UpdateUnionLevelsetConstraint<UpdateType>(InParticles, Thickness, Constraint);
+				UpdateUnionLevelsetConstraint<UpdateType>(Thickness, Constraint);
 			}
 		}
 	}
 	else
 	{
-		UpdateConstraintImp<UpdateType>(InParticles, *InParticles.Geometry(Constraint.ParticleIndex), ParticleTM, *InParticles.Geometry(Constraint.LevelsetIndex), LevelsetTM, Thickness, Constraint);
+		UpdateConstraintImp2<UpdateType>(*Constraint.Particle->Geometry(), ParticleTM, *Constraint.Levelset->Geometry(), LevelsetTM, Thickness, Constraint);
 	}
 }
 
 template class TPBDCollisionConstraint<float, 3>;
-template void TPBDCollisionConstraint<float, 3>::UpdateConstraint<ECollisionUpdateType::Any, TPBDRigidParticles<float, 3>>(const TPBDRigidParticles<float, 3>& InParticles, const float Thickness, FRigidBodyContactConstraint& Constraint);
-template void TPBDCollisionConstraint<float, 3>::UpdateConstraint<ECollisionUpdateType::Deepest, TPBDRigidParticles<float, 3>>(const TPBDRigidParticles<float, 3>& InParticles, const float Thickness, FRigidBodyContactConstraint& Constraint);
-template void TPBDCollisionConstraint<float, 3>::UpdateLevelsetConstraint<ECollisionUpdateType::Any, TPBDRigidParticles<float, 3>>(const TPBDRigidParticles<float, 3>& InParticles, const float Thickness, FRigidBodyContactConstraint& Constraint);
-template void TPBDCollisionConstraint<float, 3>::UpdateLevelsetConstraint<ECollisionUpdateType::Deepest, TPBDRigidParticles<float, 3>>(const TPBDRigidParticles<float, 3>& InParticles, const float Thickness, FRigidBodyContactConstraint& Constraint);
-template void TPBDCollisionConstraint<float, 3>::UpdateLevelsetConstraintGJK<ECollisionUpdateType::Any, TPBDRigidParticles<float, 3>>(const TPBDRigidParticles<float, 3>& InParticles, const float Thickness, FRigidBodyContactConstraint& Constraint);
-template void TPBDCollisionConstraint<float, 3>::UpdateLevelsetConstraintGJK<ECollisionUpdateType::Deepest, TPBDRigidParticles<float, 3>>(const TPBDRigidParticles<float, 3>& InParticles, const float Thickness, FRigidBodyContactConstraint& Constraint);
-template void TPBDCollisionConstraint<float, 3>::ComputeConstraints<false>(const TPBDRigidParticles<float, 3>& InParticles, const TArray<int32>& InIndices, float Dt);
-template void TPBDCollisionConstraint<float, 3>::ComputeConstraints<true>(const TPBDRigidParticles<float, 3>& InParticles, const TArray<int32>& InIndices, float Dt);
-
-template void UpdateConstraintImp<ECollisionUpdateType::Any, TPBDRigidParticles<float, 3>, float, 3>(const TPBDRigidParticles<float, 3>& InParticles, const TImplicitObject<float, 3>& ParticleObject, const TRigidTransform<float, 3>& ParticleTM, const TImplicitObject<float, 3>& LevelsetObject, const TRigidTransform<float, 3>& LevelsetTM, const float Thickness, TRigidBodyContactConstraint<float, 3>& Constraint);
-template void UpdateConstraintImp<ECollisionUpdateType::Deepest, TPBDRigidParticles<float, 3>, float, 3>(const TPBDRigidParticles<float, 3>& InParticles, const TImplicitObject<float, 3>& ParticleObject, const TRigidTransform<float, 3>& ParticleTM, const TImplicitObject<float, 3>& LevelsetObject, const TRigidTransform<float, 3>& LevelsetTM, const float Thickness, TRigidBodyContactConstraint<float, 3>& Constraint);
+template void TPBDCollisionConstraint<float, 3>::UpdateConstraint<ECollisionUpdateType::Any>(const float Thickness, FRigidBodyContactConstraint& Constraint);
+template void TPBDCollisionConstraint<float, 3>::UpdateConstraint<ECollisionUpdateType::Deepest>(const float Thickness, FRigidBodyContactConstraint& Constraint);
+template void TPBDCollisionConstraint<float, 3>::UpdateLevelsetConstraint<ECollisionUpdateType::Any>(const float Thickness, FRigidBodyContactConstraint& Constraint);
+template void TPBDCollisionConstraint<float, 3>::UpdateLevelsetConstraint<ECollisionUpdateType::Deepest>(const float Thickness, FRigidBodyContactConstraint& Constraint);
+template void TPBDCollisionConstraint<float, 3>::UpdateLevelsetConstraintGJK<ECollisionUpdateType::Any>(const float Thickness, FRigidBodyContactConstraint& Constraint);
+template void TPBDCollisionConstraint<float, 3>::UpdateLevelsetConstraintGJK<ECollisionUpdateType::Deepest>(const float Thickness, FRigidBodyContactConstraint& Constraint);
+template void TPBDCollisionConstraint<float, 3>::ComputeConstraints<false>(/*const TPBDRigidParticles<float, 3>& InParticles, const TArray<int32>& InIndices,*/ float Dt);
+template void TPBDCollisionConstraint<float, 3>::ComputeConstraints<true>(/*const TPBDRigidParticles<float, 3>& InParticles, const TArray<int32>& InIndices,*/ float Dt);
+
+template void UpdateConstraintImp2<ECollisionUpdateType::Any, float, 3>(const TImplicitObject<float, 3>& ParticleObject, const TRigidTransform<float, 3>& ParticleTM, const TImplicitObject<float, 3>& LevelsetObject, const TRigidTransform<float, 3>& LevelsetTM, const float Thickness, TRigidBodyContactConstraint<float, 3>& Constraint);
+template void UpdateConstraintImp2<ECollisionUpdateType::Deepest, float, 3>(const TImplicitObject<float, 3>& ParticleObject, const TRigidTransform<float, 3>& ParticleTM, const TImplicitObject<float, 3>& LevelsetObject, const TRigidTransform<float, 3>& LevelsetTM, const float Thickness, TRigidBodyContactConstraint<float, 3>& Constraint);
 }
