Copying //UE4/Dev-Physics to Dev-Main (//UE4/Dev-Main) [at] 7703071

diff --git a/Engine/Source/Runtime/Experimental/Chaos/Private/Chaos/PBDCollisionConstraintPGS.cpp b/Engine/Source/Runtime/Experimental/Chaos/Private/Chaos/PBDCollisionConstraintPGS.cpp
--- a/Engine/Source/Runtime/Experimental/Chaos/Private/Chaos/PBDCollisionConstraintPGS.cpp
+++ b/Engine/Source/Runtime/Experimental/Chaos/Private/Chaos/PBDCollisionConstraintPGS.cpp
@@ -12,6 +12,8 @@
 
 #include "ProfilingDebugging/ScopedTimers.h"
 
+#if CHAOS_PARTICLEHANDLE_TODO
+
 using namespace Chaos;
 
 template<class T_PARTICLES, class T, int d>
@@ -1043,14 +1045,14 @@ void TPBDCollisionConstraintPGS<T, d>::UpdateBoxConstraint(const T_PARTICLES& In
 		{
 			TSphere<T, d> Sphere1(Box1Transform.TransformPosition(Box1.Center()), Box1.Extents().Min() / 2);
 			TSphere<T, d> Sphere2(Box2Transform.TransformPosition(Box2.Center()), Box2.Extents().Min() / 2);
-			const TVector<T, d> Direction = Sphere1.Center() - Sphere2.Center();
+			const TVector<T, d> Direction = Sphere1.GetCenter() - Sphere2.GetCenter();
 			T Size = Direction.Size();
-			if (Size < (Sphere1.Radius() + Sphere2.Radius()))
+			if (Size < (Sphere1.GetRadius() + Sphere2.GetRadius()))
 			{
 				TVector<T, d> Normal = Size > SMALL_NUMBER ? Direction / Size : TVector<T, d>(0, 0, 1);
 				Constraint.Normal.Add(Normal);
-				Constraint.Phi.Add(Size - (Sphere1.Radius() + Sphere2.Radius()));
-				Constraint.Location.Add(Sphere1.Center() - Sphere1.Radius() * Normal);
+				Constraint.Phi.Add(Size - (Sphere1.GetRadius() + Sphere2.GetRadius()));
+				Constraint.Location.Add(Sphere1.GetCenter() - Sphere1.GetRadius() * Normal);
 			}
 		}
 		if (!Constraint.Phi.Num())
@@ -1109,16 +1111,16 @@ void TPBDCollisionConstraintPGS<T, d>::UpdateSphereConstraint(const T_PARTICLES&
 	const TRigidTransform<T, d> Sphere2Transform = GetTransformPGS(InParticles, Constraint.LevelsetIndex);
 	const auto& Sphere1 = *InParticles.Geometry(Constraint.ParticleIndex)->template GetObject<TSphere<T, d>>();
 	const auto& Sphere2 = *InParticles.Geometry(Constraint.LevelsetIndex)->template GetObject<TSphere<T, d>>();
-	const TVector<T, d> Center1 = Sphere1Transform.TransformPosition(Sphere1.Center());
-	const TVector<T, d> Center2 = Sphere2Transform.TransformPosition(Sphere2.Center());
+	const TVector<T, d> Center1 = Sphere1Transform.TransformPosition(Sphere1.GetCenter());
+	const TVector<T, d> Center2 = Sphere2Transform.TransformPosition(Sphere2.GetCenter());
 	const TVector<T, d> Direction = Center1 - Center2;
 	const T Size = Direction.Size();
-	if (Size < (Sphere1.Radius() + Sphere2.Radius() + Thickness))
+	if (Size < (Sphere1.GetRadius() + Sphere2.GetRadius() + Thickness))
 	{
 		TVector<T, d> Normal = Size > SMALL_NUMBER ? Direction / Size : TVector<T, d>(0, 0, 1);
 		Constraint.Normal.Add(Normal);
-		Constraint.Phi.Add(Size - (Sphere1.Radius() + Sphere2.Radius()));
-		Constraint.Location.Add(Center1 - Sphere1.Radius() * Normal);
+		Constraint.Phi.Add(Size - (Sphere1.GetRadius() + Sphere2.GetRadius()));
+		Constraint.Location.Add(Center1 - Sphere1.GetRadius() * Normal);
 	}
 }
 
@@ -1134,11 +1136,11 @@ void TPBDCollisionConstraintPGS<T, d>::UpdateSpherePlaneConstraint(const T_PARTI
 	const auto& ObjectSphere = *InParticles.Geometry(Constraint.ParticleIndex)->template GetObject<TSphere<T, d>>();
 	const auto& ObjectPlane = *InParticles.Geometry(Constraint.LevelsetIndex)->template GetObject<TPlane<T, d>>();
 	const TRigidTransform<T, d> SphereToPlaneTransform(PlaneTransform.Inverse() * SphereTransform);
-	const TVector<T, d> SphereCenter = SphereToPlaneTransform.TransformPosition(ObjectSphere.Center());
+	const TVector<T, d> SphereCenter = SphereToPlaneTransform.TransformPosition(ObjectSphere.GetCenter());
 	Constraint.Normal.SetNum(1);
 	Constraint.Phi.Add(ObjectPlane.PhiWithNormal(SphereCenter, Constraint.Normal[0]));
-	Constraint.Phi[0] -= ObjectSphere.Radius();
-	Constraint.Location.Add(SphereCenter - Constraint.Normal[0] * ObjectSphere.Radius());
+	Constraint.Phi[0] -= ObjectSphere.GetRadius();
+	Constraint.Location.Add(SphereCenter - Constraint.Normal[0] * ObjectSphere.GetRadius());
 }
 
 template<class T, int d>
@@ -1153,11 +1155,11 @@ void TPBDCollisionConstraintPGS<T, d>::UpdateSphereBoxConstraint(const T_PARTICL
 	const auto& ObjectSphere = *InParticles.Geometry(Constraint.ParticleIndex)->template GetObject<TSphere<T, d>>();
 	const auto& ObjectBox = *InParticles.Geometry(Constraint.LevelsetIndex)->template GetObject<TBox<T, d>>();
 	const TRigidTransform<T, d> SphereToBoxTransform(SphereTransform * BoxTransform.Inverse());
-	const TVector<T, d> SphereCenter = SphereToBoxTransform.TransformPosition(ObjectSphere.Center());
+	const TVector<T, d> SphereCenter = SphereToBoxTransform.TransformPosition(ObjectSphere.GetCenter());
 	Constraint.Normal.SetNum(1);
 	Constraint.Phi.Add(ObjectBox.PhiWithNormal(SphereCenter, Constraint.Normal[0]));
-	Constraint.Phi[0] -= ObjectSphere.Radius();
-	Constraint.Location.Add(SphereCenter - Constraint.Normal[0] * ObjectSphere.Radius());
+	Constraint.Phi[0] -= ObjectSphere.GetRadius();
+	Constraint.Location.Add(SphereCenter - Constraint.Normal[0] * ObjectSphere.GetRadius());
 }
 
 template<class T, int d>
@@ -1328,3 +1330,5 @@ void TPBDCollisionConstraintPGS<T, d>::UpdateConstraint(const T_PARTICLES& InPar
 }
 
 template class Chaos::TPBDCollisionConstraintPGS<float, 3>;
+
+#endif
\ No newline at end of file
diff --git a/Engine/Source/Runtime/Experimental/Chaos/Private/Chaos/PBDConstraintColor.cpp b/Engine/Source/Runtime/Experimental/Chaos/Private/Chaos/PBDConstraintColor.cpp
--- a/Engine/Source/Runtime/Experimental/Chaos/Private/Chaos/PBDConstraintColor.cpp
+++ b/Engine/Source/Runtime/Experimental/Chaos/Private/Chaos/PBDConstraintColor.cpp
@@ -1,6 +1,7 @@
 // Copyright 1998-2019 Epic Games, Inc. All Rights Reserved.
 
 #include "Chaos/PBDConstraintColor.h"
+#include "Chaos/ParticleHandle.h"
 
 #include "ProfilingDebugging/ScopedTimers.h"
 #include "ChaosStats.h"
@@ -15,9 +16,9 @@
 using namespace Chaos;
 
 template<typename T, int d>
-void TPBDConstraintColor<T, d>::ComputeIslandColoring(const TPBDRigidParticles<T, d>& InParticles, const TArray<int32>& InIndices, const int32 Island, const FConstraintGraph& ConstraintGraph, uint32 ContainerId)
+void TPBDConstraintColor<T, d>::ComputeIslandColoring(const int32 Island, const FConstraintGraph& ConstraintGraph, uint32 ContainerId)
 {
-	const TArray<int32>& IslandParticleIndices = ConstraintGraph.GetIslandParticles(Island);
+	const TArray<TGeometryParticleHandle<T,d>*>& IslandParticles = ConstraintGraph.GetIslandParticles(Island);
 	FLevelToColorToConstraintListMap& LevelToColorToConstraintListMap = IslandData[Island].LevelToColorToConstraintListMap;
 	int32& MaxColor = IslandData[Island].MaxColor;
 
@@ -34,16 +35,16 @@ void TPBDConstraintColor<T, d>::ComputeIslandColoring(const TPBDRigidParticles<T
 	TSet<int32> ProcessedNodes;
 	TArray<int32> NodesToProcess;
 
-	for (const int32 ParticleIndex : IslandParticleIndices)
+	for (const TGeometryParticleHandle<T,d>* Particle :IslandParticles)
 	{
-		if (!ConstraintGraph.ParticleToNodeIndex.Find(ParticleIndex))
+		if (!ConstraintGraph.ParticleToNodeIndex.Find(Particle))
 		{
 			continue;
 		}
 
-		const int32 ParticleNodeIndex = ConstraintGraph.ParticleToNodeIndex[ParticleIndex];
+		const int32 ParticleNodeIndex = ConstraintGraph.ParticleToNodeIndex[Particle];
 
-		if (ProcessedNodes.Contains(ParticleNodeIndex) || !InParticles.InvM(ConstraintGraph.Nodes[ParticleNodeIndex].BodyIndex))
+		if (ProcessedNodes.Contains(ParticleNodeIndex) || Particle->AsDynamic() == nullptr)
 		{
 			continue;
 		}
@@ -98,9 +99,12 @@ void TPBDConstraintColor<T, d>::ComputeIslandColoring(const TPBDRigidParticles<T
 				if (OtherNodeIndex != INDEX_NONE)
 				{
 					FGraphNodeColor& OtherColorNode = Nodes[OtherNodeIndex];
-					while (OtherColorNode.UsedColors.Contains(ColorToUse) || ColorNode.UsedColors.Contains(ColorToUse))
+					if (ConstraintGraph.Nodes[OtherNodeIndex].Particle->AsDynamic() != nullptr)
 					{
-						ColorToUse++;
+						while (OtherColorNode.UsedColors.Contains(ColorToUse) || ColorNode.UsedColors.Contains(ColorToUse))
+						{
+							ColorToUse++;
+						}
 					}
 				}
 
@@ -110,7 +114,7 @@ void TPBDConstraintColor<T, d>::ComputeIslandColoring(const TPBDRigidParticles<T
 				ColorEdge.Color = ColorToUse;
 
 				// Bump color to use next time, but only if we weren't forced to use a different color by the other node
-				if ((ColorToUse == ColorNode.NextColor) && InParticles.InvM(GraphNode.BodyIndex))
+				if ((ColorToUse == ColorNode.NextColor) && Particle->AsDynamic() != nullptr)
 				{
 					ColorNode.NextColor++;
 				}
@@ -137,20 +141,23 @@ void TPBDConstraintColor<T, d>::ComputeIslandColoring(const TPBDRigidParticles<T
 				if (OtherNodeIndex != INDEX_NONE)
 				{
 					const typename FConstraintGraph::FGraphNode& OtherGraphNode = ConstraintGraph.Nodes[OtherNodeIndex];
-					FGraphNodeColor& OtherColorNode = Nodes[OtherNodeIndex];
-
-					// Mark other node as not allowing use of this color
-					if (InParticles.InvM(OtherGraphNode.BodyIndex))
-					{
-						OtherColorNode.UsedColors.Add(ColorEdge.Color);
-					}
-
-					// Queue other node for processing
-					if (!ProcessedNodes.Contains(OtherNodeIndex) && InParticles.InvM(OtherGraphNode.BodyIndex))
+					if (OtherGraphNode.Particle->AsDynamic() != nullptr)
 					{
-						ensure(OtherGraphNode.Island == GraphNode.Island);
-						checkSlow(IslandParticleIndices.Find(OtherGraphNode.BodyIndex) != INDEX_NONE);
-						NodesToProcess.Add(OtherNodeIndex);
+						FGraphNodeColor& OtherColorNode = Nodes[OtherNodeIndex];
+
+						// Mark other node as not allowing use of this color
+						if (Particle->AsDynamic() != nullptr)
+						{
+							OtherColorNode.UsedColors.Add(ColorEdge.Color);
+						}
+
+						// Queue other node for processing
+						if (!ProcessedNodes.Contains(OtherNodeIndex))
+						{
+							ensure(OtherGraphNode.Island == GraphNode.Island);
+							checkSlow(IslandParticles.Find(OtherGraphNode.Particle) != INDEX_NONE);
+							NodesToProcess.Add(OtherNodeIndex);
+						}
 					}
 				}
 			}
@@ -160,17 +167,17 @@ void TPBDConstraintColor<T, d>::ComputeIslandColoring(const TPBDRigidParticles<T
 
 #ifdef USE_CONTACT_LEVELS
 template<typename T, int d>
-void TPBDConstraintColor<T, d>::ComputeContactGraph(const TPBDRigidParticles<T, d>& InParticles, const TArray<int32>& InIndices, const int32 Island, const FConstraintGraph& ConstraintGraph, uint32 ContainerId)
+void TPBDConstraintColor<T, d>::ComputeContactGraph(const int32 Island, const FConstraintGraph& ConstraintGraph, uint32 ContainerId)
 {
 	const TArray<int32>& ConstraintDataIndices = ConstraintGraph.GetIslandConstraintData(Island);
 
 	IslandData[Island].MaxLevel = ConstraintDataIndices.Num() ? 0 : -1;
 
 	std::queue<std::pair<int32, int32>> QueueToProcess;
-	for (const int32 Index : InIndices)
+	for (const TGeometryParticleHandle<T,d>* Particle : ConstraintGraph.GetIslandParticles(Island))
 	{
-		const int32* NodeIndexPtr = ConstraintGraph.ParticleToNodeIndex.Find(Index);
-		if (!InParticles.InvM(Index) && NodeIndexPtr)
+		const int32* NodeIndexPtr = ConstraintGraph.ParticleToNodeIndex.Find(Particle);
+		if (Particle->AsDynamic() == nullptr && NodeIndexPtr)
 		{
 			const int32 NodeIndex = *NodeIndexPtr;
 			QueueToProcess.push(std::make_pair(0, NodeIndex));
@@ -256,16 +263,16 @@ void TPBDConstraintColor<T, d>::InitializeColor(const FConstraintGraph& Constrai
 
 	Nodes.SetNum(ConstraintGraph.Nodes.Num());
 	Edges.SetNum(ConstraintGraph.Edges.Num());
-	IslandData.SetNum(ConstraintGraph.IslandData.Num());
+	IslandData.SetNum(ConstraintGraph.IslandToData.Num());
 }
 
 template<typename T, int d>
-void TPBDConstraintColor<T, d>::ComputeColor(const TPBDRigidParticles<T, d>& InParticles, const TArray<int32>& InIndices, const int32 Island, const FConstraintGraph& ConstraintGraph, uint32 ContainerId)
+void TPBDConstraintColor<T, d>::ComputeColor(const int32 Island, const FConstraintGraph& ConstraintGraph, uint32 ContainerId)
 {
 #ifdef USE_CONTACT_LEVELS
-	ComputeContactGraph(InParticles, InIndices, Island, ConstraintGraph, ContainerId);
+	ComputeContactGraph(Island, ConstraintGraph, ContainerId);
 #endif
-	ComputeIslandColoring(InParticles, InIndices, Island, ConstraintGraph, ContainerId);
+	ComputeIslandColoring(Island, ConstraintGraph, ContainerId);
 }
 
 template<typename T, int d>
diff --git a/Engine/Source/Runtime/Experimental/Chaos/Private/Chaos/PBDConstraintGraph.cpp b/Engine/Source/Runtime/Experimental/Chaos/Private/Chaos/PBDConstraintGraph.cpp
--- a/Engine/Source/Runtime/Experimental/Chaos/Private/Chaos/PBDConstraintGraph.cpp
+++ b/Engine/Source/Runtime/Experimental/Chaos/Private/Chaos/PBDConstraintGraph.cpp
@@ -7,6 +7,8 @@
 #include "Chaos/Framework/Parallel.h"
 #include "Chaos/PBDRigidParticles.h"
 #include "Containers/Queue.h"
+#include "Chaos/ParticleHandle.h"
+#include "Chaos/PBDRigidsSOAs.h"
 
 using namespace Chaos;
 
@@ -16,15 +18,15 @@ TPBDConstraintGraph<T, d>::TPBDConstraintGraph()
 }
 
 template<typename T, int d>
-TPBDConstraintGraph<T, d>::TPBDConstraintGraph(TPBDRigidParticles<T, d>& InParticles, const TArray<int32>& InIndices)
+TPBDConstraintGraph<T, d>::TPBDConstraintGraph(const TParticleView<TGeometryParticles<T, d>>& Particles)
 {
-	InitializeGraph(InParticles, InIndices);
+	InitializeGraph(Particles);
 }
 
 template<typename T, int d>
-void TPBDConstraintGraph<T, d>::InitializeGraph(const TPBDRigidParticles<T, d>& InParticles, const TArray<int32>& InIndices)
+void TPBDConstraintGraph<T, d>::InitializeGraph(const TParticleView<TGeometryParticles<T, d>>& Particles)
 {
-	const int32 NumNonDisabledParticles = InIndices.Num();
+	const int32 NumNonDisabledParticles = Particles.Num();
 
 	Nodes.Reset();
 	Nodes.AddDefaulted(NumNonDisabledParticles);
@@ -33,41 +35,43 @@ void TPBDConstraintGraph<T, d>::InitializeGraph(const TPBDRigidParticles<T, d>&
 
 	ParticleToNodeIndex.Reset();
 	ParticleToNodeIndex.Reserve(NumNonDisabledParticles);
-	for(int32 Index = 0; Index < NumNonDisabledParticles; ++Index)
 	{
-		FGraphNode& Node = Nodes[Index];
-		Node.BodyIndex = InIndices[Index];
-		ParticleToNodeIndex.Add(Node.BodyIndex, Index);
+		int32 Index = 0;
+		for (auto& Particle : Particles)
+		{
+			FGraphNode& Node = Nodes[Index];
+			Node.Particle = Particle.Handle();
+			ParticleToNodeIndex.Add(Node.Particle, Index);
+			++Index;
+		}
 	}
 
 	//@todo(ocohen): Should we reset more than just the edges? What about bIsIslandPersistant?
-	for (TArray<int32>& IslandConstraintList : IslandConstraints)
+	for (TArray<int32>& IslandConstraintList : IslandToConstraints)
 	{
 		IslandConstraintList.Reset();
 	}
 }
 
 template<typename T, int d>
-void TPBDConstraintGraph<T, d>::ResetIslands(const TPBDRigidParticles<T, d>& InParticles, const TArray<int32>& InIndices)
+void TPBDConstraintGraph<T, d>::ResetIslands(const TParticleView<TPBDRigidParticles<T, d>>& PBDRigids)
 {
 	//@todo(ocohen): Should we reset more than just the edges? What about bIsIslandPersistant?
-	for (TArray<int32>& IslandConstraintList : IslandConstraints)
+	for (TArray<int32>& IslandConstraintList : IslandToConstraints)
 	{
 		IslandConstraintList.Reset();
 	}
 
-	const int32 NumBodies = InIndices.Num();
-	for(int32 BodyIdx = 0; BodyIdx < NumBodies; ++BodyIdx)	//@todo(ocohen): could go wide per island if we can get at the sets
+	for(auto& PBDRigid : PBDRigids)	//@todo(ocohen): could go wide per island if we can get at the sets
 	{
-		const int32 ParticleIndex = InIndices[BodyIdx];
-		const int32 Island = InParticles.Island(ParticleIndex);
+		const int32 Island = PBDRigid.Island();
 		if (Island >= 0)
 		{
-			FGraphNode& Node = Nodes[BodyIdx];
+			FGraphNode& Node = Nodes[ParticleToNodeIndex[PBDRigid.Handle()]];
 			Node.Island = Island;
 			for (int32 ConstraintDataIndex : Node.Edges)
 			{
-				IslandConstraints[Island].Add(ConstraintDataIndex);
+				IslandToConstraints[Island].Add(ConstraintDataIndex);
 			}
 		}
 	}
@@ -80,25 +84,25 @@ void TPBDConstraintGraph<T, d>::ReserveConstraints(const int32 NumConstraints)
 }
 
 template<typename T, int d>
-void TPBDConstraintGraph<T, d>::AddConstraint(const uint32 InContainerId, const int32 InConstraintIndex, const TVector<int32, 2>& InConstrainedParticles)
+void TPBDConstraintGraph<T, d>::AddConstraint(const uint32 InContainerId, const int32 InConstraintIndex, const TVector<TGeometryParticleHandle<T,d>*, 2>& ConstrainedParticles)
 {
 	// Must have at least one constrained particle
-	check((InConstrainedParticles[0] != INDEX_NONE) || (InConstrainedParticles[1] != INDEX_NONE));
+	check((ConstrainedParticles[0]) || (ConstrainedParticles[1]));
 
 	const int32 NewEdgeIndex = Edges.Num();
 	FGraphEdge NewEdge;
 	NewEdge.Data = { InContainerId, InConstraintIndex };
 
-	if (InConstrainedParticles[0] != INDEX_NONE)
+	if (ConstrainedParticles[0])
 	{
-		NewEdge.FirstNode = ParticleToNodeIndex[InConstrainedParticles[0]];
-		Nodes[NewEdge.FirstNode].BodyIndex = InConstrainedParticles[0];
+		NewEdge.FirstNode = ParticleToNodeIndex[ConstrainedParticles[0]];
+		Nodes[NewEdge.FirstNode].Particle = ConstrainedParticles[0];
 		Nodes[NewEdge.FirstNode].Edges.Add(NewEdgeIndex);
 	}
-	if (InConstrainedParticles[1] != INDEX_NONE)
+	if (ConstrainedParticles[1])
 	{
-		NewEdge.SecondNode = ParticleToNodeIndex[InConstrainedParticles[1]];
-		Nodes[NewEdge.SecondNode].BodyIndex = InConstrainedParticles[1];
+		NewEdge.SecondNode = ParticleToNodeIndex[ConstrainedParticles[1]];
+		Nodes[NewEdge.SecondNode].Particle = ConstrainedParticles[1];
 		Nodes[NewEdge.SecondNode].Edges.Add(NewEdgeIndex);
 	}
 
@@ -112,43 +116,44 @@ const typename TPBDConstraintGraph<T, d>::FConstraintData& TPBDConstraintGraph<T
 }
 
 template<typename T, int d>
-void TPBDConstraintGraph<T, d>::UpdateIslands(TPBDRigidParticles<T, d>& InParticles, const TArray<int32>& InIndices, TSet<int32>& ActiveIndices)
+void TPBDConstraintGraph<T, d>::UpdateIslands(const TParticleView<TPBDRigidParticles<T,d>>& PBDRigids, TPBDRigidsSOAs<T,d>& Particles)
 {
 	// Maybe expose a memset style function for this instead of iterating
-	for (uint32 i = 0; i < InParticles.Size(); ++i)
+	for (auto& PBDRigid : PBDRigids)
 	{
-		InParticles.Island(i) = INDEX_NONE;
+		PBDRigid.Island() = INDEX_NONE;
 	}
-	ComputeIslands(InParticles, InIndices, ActiveIndices);
+	ComputeIslands(Particles);
 }
 
-DECLARE_CYCLE_STAT(TEXT("IslandGeneration"), STAT_IslandGeneration, STATGROUP_Chaos);
+
+DECLARE_CYCLE_STAT(TEXT("IslandGeneration2"), STAT_IslandGeneration2, STATGROUP_Chaos);
 
 template<typename T, int d>
-void TPBDConstraintGraph<T, d>::ComputeIslands(TPBDRigidParticles<T, d>& InParticles, const TArray<int32>& InIndices, TSet<int32>& ActiveIndices)
+void TPBDConstraintGraph<T, d>::ComputeIslands(TPBDRigidsSOAs<T,d>& Particles)
 {
-	SCOPE_CYCLE_COUNTER(STAT_IslandGeneration);
+	SCOPE_CYCLE_COUNTER(STAT_IslandGeneration2);
 
 	int32 NextIsland = 0;
-	TArray<TSet<int32>> NewIslandParticles;
-	TArray<int32> NewIslandSleepCounts;
+	TArray<TSet<TGeometryParticleHandle<T, d>*>> NewIslandParticles;
+	TArray<int32> NewIslandToSleepCount;
 
-	const int32 NumNodes = InIndices.Num();
-	for(int32 NodeIndex = 0; NodeIndex < NumNodes; ++NodeIndex)
+	const int32 NumNodes = Nodes.Num();
+	for (int32 Idx = 0; Idx < NumNodes; ++Idx)
 	{
-		const int32 ParticleIndex = InIndices[NodeIndex];
+		TGeometryParticleHandle<T, d>* Particle = Nodes[Idx].Particle;
 
-		if (Nodes[NodeIndex].Island >= 0 || InParticles.InvM(ParticleIndex) == 0)
+		if (Nodes[Idx].Island >= 0 || Particle->AsDynamic() == nullptr)
 		{
 			// Island is already known - it was visited in ComputeIsland for a previous node
 			continue;
 		}
 
-		TSet<int32> SingleIslandParticles;
-		TSet<int32> SingleIslandStaticParticles;
-		ComputeIsland(InParticles, InIndices, NodeIndex, NextIsland, SingleIslandParticles, SingleIslandStaticParticles);
-		
-		for (const int32& StaticParticle : SingleIslandStaticParticles)
+		TSet<TGeometryParticleHandle<T, d>*> SingleIslandParticles;
+		TSet<TGeometryParticleHandle<T, d>*> SingleIslandStaticParticles;
+		ComputeIsland(Idx, NextIsland, SingleIslandParticles, SingleIslandStaticParticles);
+
+		for (TGeometryParticleHandle<T, d>* StaticParticle : SingleIslandStaticParticles)
 		{
 			SingleIslandParticles.Add(StaticParticle);
 		}
@@ -161,97 +166,97 @@ void TPBDConstraintGraph<T, d>::ComputeIslands(TPBDRigidParticles<T, d>& InParti
 		}
 	}
 
-	IslandConstraints.SetNum(NextIsland);
-	IslandData.SetNum(NextIsland);
+	IslandToConstraints.SetNum(NextIsland);
+	IslandToData.SetNum(NextIsland);
 
 	for (int32 EdgeIndex = 0; EdgeIndex < Edges.Num(); ++EdgeIndex)
 	{
 		const FGraphEdge& Edge = Edges[EdgeIndex];
-		int32 FirstIsland = (Edge.FirstNode != INDEX_NONE)? Nodes[Edge.FirstNode].Island : INDEX_NONE;
-		int32 SecondIsland = (Edge.SecondNode != INDEX_NONE)? Nodes[Edge.SecondNode].Island : INDEX_NONE;
+		int32 FirstIsland = (Edge.FirstNode != INDEX_NONE) ? Nodes[Edge.FirstNode].Island : INDEX_NONE;
+		int32 SecondIsland = (Edge.SecondNode != INDEX_NONE) ? Nodes[Edge.SecondNode].Island : INDEX_NONE;
 		check(FirstIsland == SecondIsland || FirstIsland == INDEX_NONE || SecondIsland == INDEX_NONE);
 
-		int32 Island = (FirstIsland != INDEX_NONE)? FirstIsland : SecondIsland;
-		check(Island >= 0);
-
-		IslandConstraints[Island].Add(EdgeIndex);
+		int32 Island = (FirstIsland != INDEX_NONE) ? FirstIsland : SecondIsland;
+		
+		// @todo(ccaulfield): should check(Island >= 0) when we disable particles properly
+		if (Island >= 0)
+		{
+			IslandToConstraints[Island].Add(EdgeIndex);
+		}
 	}
 
-	NewIslandSleepCounts.SetNum(NewIslandParticles.Num());
+	NewIslandToSleepCount.SetNum(NewIslandParticles.Num());
 
 	if (NewIslandParticles.Num())
 	{
 		for (int32 Island = 0; Island < NewIslandParticles.Num(); ++Island)
 		{
-			NewIslandSleepCounts[Island] = 0;
+			NewIslandToSleepCount[Island] = 0;
 
-			for (const int32& Index : NewIslandParticles[Island])
+			for (TGeometryParticleHandle<T, d>* Particle : NewIslandParticles[Island])
 			{
-				if (InParticles.InvM(Index))
-				{
-					InParticles.Island(Index) = Island;
-				}
-				else
+				if (TPBDRigidParticleHandle<T,d>* PBDRigid = Particle->AsDynamic())
 				{
-					InParticles.Island(Index) = INDEX_NONE;
+					PBDRigid->Island() = Island;
 				}
 			}
 		}
 		// Force consistent state if no previous islands
-		if (!IslandParticles.Num())
+		if (!IslandToParticles.Num())
 		{
 			for (int32 Island = 0; Island < NewIslandParticles.Num(); ++Island)
 			{
-				IslandData[Island].bIsIslandPersistant = true;
+				IslandToData[Island].bIsIslandPersistant = true;
 				bool bSleepState = true;
 
-				for (const int32& Index : NewIslandParticles[Island])
+				for (TGeometryParticleHandle<T, d>* Particle : NewIslandParticles[Island])
 				{
-					if (!InParticles.Sleeping(Index))
+					if (!Particle->Sleeping())
 					{
 						bSleepState = false;
 						break;
 					}
 				}
 
-				for (const int32& Index : NewIslandParticles[Island])
+				for (TGeometryParticleHandle<T, d>* Particle : NewIslandParticles[Island])
 				{
 					//@todo(DEMO_HACK) : Need to fix, remove the !InParticles.Disabled(Index)
-					if (InParticles.Sleeping(Index) && !bSleepState && !InParticles.Disabled(Index))
+					if (Particle->Sleeping() && !bSleepState/* && !Particle->Disabled()*/)
 					{
-						ActiveIndices.Add(Index);
+						Particles.ActivateParticle(Particle); 	//todo: record state change for array reorder
 					}
 
-					if (!InParticles.Sleeping(Index) && bSleepState && InParticles.InvM(Index))
+					if (TPBDRigidParticleHandle<T, d>* PBDRigid = Particle->AsDynamic())
 					{
-						ActiveIndices.Remove(Index);
-						InParticles.V(Index) = TVector<T, d>(0);
-						InParticles.W(Index) = TVector<T, d>(0);
-					}
+						if (!Particle->Sleeping() && bSleepState)
+						{
+							Particles.DeactivateParticle(Particle); 	//todo: record state change for array reorder
+							PBDRigid->V() = TVector<T, d>(0);
+							PBDRigid->W() = TVector<T, d>(0);
+						}
 
-					if (InParticles.InvM(Index))
-					{
-						InParticles.SetSleeping(Index, bSleepState);
+						PBDRigid->SetSleeping(bSleepState);
 					}
 
-					if ((InParticles.Sleeping(Index) || InParticles.Disabled(Index)) && ActiveIndices.Contains(Index))
+					if ((Particle->Sleeping() /*|| Particle->Disabled()*/))
 					{
-						ActiveIndices.Remove(Index);
+						Particles.DeactivateParticle(Particle);	//todo: record state change for array reorder (function could return true/false)
 					}
 				}
 			}
 		}
 
-		for (int32 Island = 0; Island < IslandParticles.Num(); ++Island)
+		for (int32 Island = 0; Island < IslandToParticles.Num(); ++Island)
 		{
 			bool bIsSameIsland = true;
 
 			// Objects were removed from the island
 			int32 OtherIsland = -1;
 
-			for (const int32 Index : IslandParticles[Island])
+			for (TGeometryParticleHandle<T, d>* Particle : IslandToParticles[Island])
 			{
-				int32 TmpIsland = InParticles.Island(Index);
+				TPBDRigidParticleHandle<T, d>* PBDRigid = Particle->AsDynamic();
+				int32 TmpIsland = PBDRigid ? PBDRigid->Island() : INDEX_NONE; //question: should we even store non dynamics in this array?
 
 				if (OtherIsland == INDEX_NONE && TmpIsland >= 0)
 				{
@@ -268,7 +273,7 @@ void TPBDConstraintGraph<T, d>::ComputeIslands(TPBDRigidParticles<T, d>& InParti
 			}
 
 			// A new object entered the island or the island is entirely new particles
-			if (bIsSameIsland && (OtherIsland == INDEX_NONE || NewIslandParticles[OtherIsland].Num() != IslandParticles[Island].Num()))
+			if (bIsSameIsland && (OtherIsland == INDEX_NONE || NewIslandParticles[OtherIsland].Num() != IslandToParticles[Island].Num()))
 			{
 				bIsSameIsland = false;
 			}
@@ -276,16 +281,23 @@ void TPBDConstraintGraph<T, d>::ComputeIslands(TPBDRigidParticles<T, d>& InParti
 			// Find out if we need to activate island
 			if (bIsSameIsland)
 			{
-				NewIslandSleepCounts[OtherIsland] = IslandSleepCounts[Island];
+				NewIslandToSleepCount[OtherIsland] = IslandToSleepCount[Island];
 			}
 			else
 			{
-				for (const int32 Index : IslandParticles[Island])
+				for (TGeometryParticleHandle<T, d>* Particle : IslandToParticles[Island])
 				{
-					if (!InParticles.Disabled(Index))
+					if (TPBDRigidParticleHandle<T, d>* PBDRigid = Particle->AsDynamic())
+					{
+						if (!PBDRigid->Disabled())	// todo: why is this needed? [we aren't handling enable/disable state changes properly so disabled particles end up in the graph.]
+						{
+							PBDRigid->SetSleeping(false);
+							Particles.ActivateParticle(Particle);
+						}
+					}
+					else
 					{
-						InParticles.SetSleeping(Index, false);
-						ActiveIndices.Add(Index);
+						Particles.ActivateParticle(Particle);
 					}
 				}
 			}
@@ -293,53 +305,29 @@ void TPBDConstraintGraph<T, d>::ComputeIslands(TPBDRigidParticles<T, d>& InParti
 			// #BG Necessary? Should we ever not find an island?
 			if (OtherIsland != INDEX_NONE)
 			{
-				IslandData[OtherIsland].bIsIslandPersistant = bIsSameIsland;
-			}
-		}
-
-		//if any particles are awake, make sure the entire island is awake
-#if 0
-		for (int32 Island = 0; Island < NewIslandParticles.Num(); ++Island)
-		{
-			int32 OtherIsland = -1;
-			for (const int32 Index : NewIslandParticles[Island])
-			{
-				if (!InParticles.Sleeping(Index) && !InParticles.InvM(Index))
-				{
-					for (const int32 Index : NewIslandParticles[Island])
-					{
-						if (!InParticles.InvM(Index))
-						{
-							InParticles.Sleeping(Index) = false;
-							ActiveIndices.Add(Index);
-						}
-					}
-					IslandData[Island].bIsIslandPersistant = false;
-					break;
-				}
+				IslandToData[OtherIsland].bIsIslandPersistant = bIsSameIsland;
 			}
 		}
-#endif
 	}
 
-	IslandParticles.Reset();
-	IslandParticles.Reserve(NewIslandParticles.Num());
+	IslandToParticles.Reset();
+	IslandToParticles.Reserve(NewIslandParticles.Num());
 	for (int32 Island = 0; Island < NewIslandParticles.Num(); ++Island)
 	{
-		IslandParticles.Emplace(NewIslandParticles[Island].Array());
+		IslandToParticles.Emplace(NewIslandParticles[Island].Array());
 	}
-	IslandSleepCounts = MoveTemp(NewIslandSleepCounts);
-	
-	check(IslandParticles.Num() == IslandSleepCounts.Num());
-	check(IslandParticles.Num() == IslandConstraints.Num());
-	check(IslandParticles.Num() == IslandData.Num());
+	IslandToSleepCount = MoveTemp(NewIslandToSleepCount);
+
+	check(IslandToParticles.Num() == IslandToSleepCount.Num());
+	check(IslandToParticles.Num() == IslandToConstraints.Num());
+	check(IslandToParticles.Num() == IslandToData.Num());
 	// @todo(ccaulfield): make a more complex unit test to check island integrity
 	//checkSlow(CheckIslands(InParticles, ActiveIndices));
 }
 
 template<typename T, int d>
-void TPBDConstraintGraph<T, d>::ComputeIsland(const TPBDRigidParticles<T, d>& InParticles, const TArray<int32>& InIndices, int32 InNode, const int32 Island,
-    TSet<int32>& DynamicParticlesInIsland, TSet<int32>& StaticParticlesInIsland)
+void TPBDConstraintGraph<T, d>::ComputeIsland(const int32 InNode, const int32 Island, TSet<TGeometryParticleHandle<T, d> *>& DynamicParticlesInIsland,
+	TSet<TGeometryParticleHandle<T, d> *>& StaticParticlesInIsland)
 {
 	TQueue<int32> NodeQueue;
 	NodeQueue.Enqueue(InNode);
@@ -355,16 +343,24 @@ void TPBDConstraintGraph<T, d>::ComputeIsland(const TPBDRigidParticles<T, d>& In
 			continue;
 		}
 
-		if (!InParticles.InvM(Node.BodyIndex))
+		TPBDRigidParticleHandle<T, d>* RigidHandle = Node.Particle->AsDynamic();
+
+		if (RigidHandle == nullptr)
 		{
-			if (!StaticParticlesInIsland.Contains(Node.BodyIndex))
+			if (!StaticParticlesInIsland.Contains(Node.Particle))
 			{
-				StaticParticlesInIsland.Add(Node.BodyIndex);
+				StaticParticlesInIsland.Add(Node.Particle);
 			}
 			continue;
 		}
 
-		DynamicParticlesInIsland.Add(Node.BodyIndex);
+		// @todo(ccaulfield): we don't handle enable/disable properly so this breaks
+		//if (RigidHandle->Disabled())
+		//{
+		//	continue;
+		//}
+
+		DynamicParticlesInIsland.Add(Node.Particle);
 		Node.Island = Island;
 
 		for (const int32 EdgeIndex : Node.Edges)
@@ -388,19 +384,19 @@ void TPBDConstraintGraph<T, d>::ComputeIsland(const TPBDRigidParticles<T, d>& In
 }
 
 template<typename T, int d>
-bool TPBDConstraintGraph<T, d>::SleepInactive(TPBDRigidParticles<T, d>& InParticles, const int32 Island, const TArrayCollectionArray<TSerializablePtr<TChaosPhysicsMaterial<T>>>& PerParticleMaterialAttributes)
+bool TPBDConstraintGraph<T, d>::SleepInactive(const int32 Island, const TArrayCollectionArray<TSerializablePtr<TChaosPhysicsMaterial<T>>>& PerParticleMaterialAttributes)
 {
 	// @todo(ccaulfield): should be able to eliminate this when island is already sleeping
 
-	const TArray<int32>& IslandParticleIndices = GetIslandParticles(Island);
-	check(IslandParticleIndices.Num());
+	const TArray<TGeometryParticleHandle<T, d>*>& IslandParticles = GetIslandParticles(Island);
+	check(IslandParticles.Num());
 
-	if (!IslandData[Island].bIsIslandPersistant)
+	if (!IslandToData[Island].bIsIslandPersistant)
 	{
 		return false;
 	}
 
-	int32& IslandSleepCount = IslandSleepCounts[Island];
+	int32& IslandSleepCount = IslandToSleepCount[Island];
 
 	TVector<T, d> X(0);
 	TVector<T, d> V(0);
@@ -409,38 +405,36 @@ bool TPBDConstraintGraph<T, d>::SleepInactive(TPBDRigidParticles<T, d>& InPartic
 	T LinearSleepingThreshold = FLT_MAX;
 	T AngularSleepingThreshold = FLT_MAX;
 
-	for (const int32 Index : IslandParticleIndices)
+	for (const TGeometryParticleHandle<T, d>* Particle : IslandToParticles[Island])
 	{
-		if (!InParticles.InvM(Index))
+		if (const TPBDRigidParticleHandle<T, d>* PBDRigid = Particle->AsDynamic())
 		{
-			continue;
-		}
-		X += InParticles.X(Index) * InParticles.M(Index);
-		M += InParticles.M(Index);
-		V += InParticles.V(Index) * InParticles.M(Index);
+			X += PBDRigid->X() * PBDRigid->M();
+			M += PBDRigid->M();
+			V += PBDRigid->V() * PBDRigid->M();
 
-		if (PerParticleMaterialAttributes[Index])
-		{
-			LinearSleepingThreshold = FMath::Min(LinearSleepingThreshold, PerParticleMaterialAttributes[Index]->SleepingLinearThreshold);
-			AngularSleepingThreshold = FMath::Min(AngularSleepingThreshold, PerParticleMaterialAttributes[Index]->SleepingAngularThreshold);
-		}
-		else
-		{
-			LinearSleepingThreshold = (T)0;
-			LinearSleepingThreshold = (T)0;
+			if (TSerializablePtr<TChaosPhysicsMaterial<T>> PhysicsMaterial = Particle->AuxilaryValue(PerParticleMaterialAttributes))
+			{
+				LinearSleepingThreshold = FMath::Min(LinearSleepingThreshold, PhysicsMaterial->SleepingLinearThreshold);
+				AngularSleepingThreshold = FMath::Min(AngularSleepingThreshold, PhysicsMaterial->SleepingAngularThreshold);
+			}
+			else
+			{
+				LinearSleepingThreshold = (T)0;
+				LinearSleepingThreshold = (T)0;
+			}
 		}
 	}
 
 	X /= M;
 	V /= M;
 
-	for (const int32 Index : IslandParticleIndices)
+	for (const TGeometryParticleHandle<T, d>* Particle: IslandParticles)
 	{
-		if (!InParticles.InvM(Index))
+		if (const TPBDRigidParticleHandle<T,d>* PBDRigid = Particle->AsDynamic())
 		{
-			continue;
+			W += /*TVector<T, d>::CrossProduct(PBDRigid->X() - X, PBDRigid->M() * PBDRigid->V()/ +*/ PBDRigid->W() * PBDRigid->M();
 		}
-		W += /*TVector<T, d>::CrossProduct(InParticles.X(Index) - X, InParticles.M(Index) * InParticles.V(Index)/ +*/ InParticles.W(Index) * InParticles.M(Index);
 	}
 
 	W /= M;
@@ -452,15 +446,15 @@ bool TPBDConstraintGraph<T, d>::SleepInactive(TPBDRigidParticles<T, d>& InPartic
 	{
 		if (IslandSleepCount > SleepCountThreshold)
 		{
-			for (const int32 Index : IslandParticleIndices)
+			for (TGeometryParticleHandle<T, d>* Particle : IslandParticles)
 			{
-				if (!InParticles.InvM(Index))
+				if (TPBDRigidParticleHandle<T, d>* PBDRigid = Particle->AsDynamic())
 				{
-					continue;
+					PBDRigid->SetSleeping(true);
+					PBDRigid->V() = TVector<T, d>(0);
+					PBDRigid->W() = TVector<T, d>(0);
 				}
-				InParticles.SetSleeping(Index, true);
-				InParticles.V(Index) = TVector<T, d>(0);
-				InParticles.W(Index) = TVector<T, d>(0);
+				
 			}
 			return true;
 		}
@@ -474,40 +468,43 @@ bool TPBDConstraintGraph<T, d>::SleepInactive(TPBDRigidParticles<T, d>& InPartic
 }
 
 template<typename T, int d>
-void TPBDConstraintGraph<T, d>::WakeIsland(TPBDRigidParticles<T, d>& InParticles, const int32 Island)
+void TPBDConstraintGraph<T, d>::WakeIsland(const int32 Island)
 {
-	for (int32 Particle : IslandParticles[Island])
+	for (TGeometryParticleHandle<T, d>* Particle : IslandToParticles[Island])
 	{
-		if (InParticles.Sleeping(Particle))
+		if (TPBDRigidParticleHandle<T, d>* PBDRigid = Particle->AsDynamic())
 		{
-			InParticles.SetSleeping(Particle, false);
+			if (PBDRigid->Sleeping())
+			{
+				PBDRigid->SetSleeping(false);
+			}
 		}
 	}
-	IslandSleepCounts[Island] = 0;
+	IslandToSleepCount[Island] = 0;
 }
 
 
 template<typename T, int d>
-void TPBDConstraintGraph<T, d>::ReconcileIslands(TPBDRigidParticles<T, d>& InParticles)
+void TPBDConstraintGraph<T, d>::ReconcileIslands()
 {
-	for (int32 i = 0; i < IslandParticles.Num(); ++i)
+	for (int32 Island = 0; Island < IslandToParticles.Num(); ++Island)
 	{
 		bool IsSleeping = true;
 		bool IsSet = false;
-		for (const int32 Index : IslandParticles[i])
+		for (TGeometryParticleHandle<T, d>* Particle : IslandToParticles[Island])
 		{
-			if (InParticles.ObjectState(Index) == EObjectStateType::Static)
+			if (Particle->ObjectState() == EObjectStateType::Static)
 			{
 				continue;
 			}
 			if (!IsSet)
 			{
 				IsSet = true;
-				IsSleeping = InParticles.Sleeping(Index);
+				IsSleeping = Particle->Sleeping();
 			}
-			if (InParticles.Sleeping(Index) != IsSleeping)
+			if (Particle->Sleeping() != IsSleeping)
 			{
-				WakeIsland(InParticles, i);
+				WakeIsland(Island);
 				break;
 			}
 		}
@@ -515,71 +512,84 @@ void TPBDConstraintGraph<T, d>::ReconcileIslands(TPBDRigidParticles<T, d>& InPar
 }
 
 template<typename T, int d>
-void TPBDConstraintGraph<T, d>::EnableParticle(TPBDRigidParticles<T, d>& InParticles, const int32 ParticleIndex, const int32 ParentParticleIndex)
+void TPBDConstraintGraph<T, d>::EnableParticle(TGeometryParticleHandle<T, d>* Particle, const TGeometryParticleHandle<T, d>* ParentParticle)
 {
-	if (ParentParticleIndex != INDEX_NONE)
+	if (ParentParticle)
 	{
-		const int32 Island = InParticles.Island(ParentParticleIndex);
-		InParticles.Island(ParticleIndex) = Island;
-		if (IslandParticles.IsValidIndex(Island))
+		if (const TPBDRigidParticleHandle<T, d>* ParentPBDRigid = ParentParticle->AsDynamic())
 		{
-			IslandParticles[Island].Add(ParticleIndex);
-		}
+			if (TPBDRigidParticleHandle<T, d>* ChildPBDRigid = Particle->AsDynamic())
+			{
+				const int32 Island = ParentPBDRigid->Island();
+				ChildPBDRigid->Island() = Island;
+				if (ensure(IslandToParticles.IsValidIndex(Island)))
+				{
+					IslandToParticles[Island].Add(Particle);
+				}
 
-		const bool SleepState = InParticles.Sleeping(ParentParticleIndex);
-		InParticles.SetSleeping(ParticleIndex, SleepState);
+				const bool SleepState = ParentPBDRigid->Sleeping();
+				ChildPBDRigid->SetSleeping(SleepState);	//todo: need to let evolution know to reorder arrays
+			}
+			else
+			{
+				ensure(false);	//this should never happen
+			}
+		}
 	}
 }
 
 template<typename T, int d>
-void TPBDConstraintGraph<T, d>::DisableParticle(TPBDRigidParticles<T, d>& InParticles, const int32 ParticleIndex)
+void TPBDConstraintGraph<T, d>::DisableParticle(TGeometryParticleHandle<T, d>* Particle)
 {
-	const int32 Island = InParticles.Island(ParticleIndex);
-	if (Island != INDEX_NONE)
+	if (TPBDRigidParticleHandle<T, d>* PBDRigid = Particle->AsDynamic())
 	{
-		InParticles.Island(ParticleIndex) = INDEX_NONE;
-
-		// @todo(ccaulfield): optimize
-		if (IslandParticles.IsValidIndex(Island))
+		const int32 Island = PBDRigid->Island();
+		if (Island != INDEX_NONE)
 		{
-			int32 IslandParticleIndex = IslandParticles[Island].Find(ParticleIndex);
-			check(IslandParticleIndex != INDEX_NONE);
-			IslandParticles[Island].RemoveAtSwap(IslandParticleIndex);
+			PBDRigid->Island() = INDEX_NONE;
+
+			// @todo(ccaulfield): optimize
+			if (ensure(IslandToParticles.IsValidIndex(Island)))
+			{
+				int32 IslandParticleArrayIdx = IslandToParticles[Island].Find(Particle);
+				check(IslandParticleArrayIdx != INDEX_NONE);
+				IslandToParticles[Island].RemoveAtSwap(IslandParticleArrayIdx);
+			}
 		}
 	}
 }
 
 template<typename T, int d>
-void TPBDConstraintGraph<T, d>::DisableParticles(TPBDRigidParticles<T, d>& InParticles, const TSet<int32>& InParticleIndices)
+void TPBDConstraintGraph<T, d>::DisableParticles(const TSet<TGeometryParticleHandle<T, d> *>& Particles)
 {
 	// @todo(ccaulfield): optimize
-	for (int32 ParticleIndex : InParticleIndices)
+	for (TGeometryParticleHandle<T,d>* Particle : Particles)
 	{
-		DisableParticle(InParticles, ParticleIndex);
+		DisableParticle(Particle);
 	}
 }
 
 template<typename T, int d>
-bool TPBDConstraintGraph<T, d>::CheckIslands(TPBDRigidParticles<T, d>& InParticles, const TSet<int32>& InParticleIndices)
+bool TPBDConstraintGraph<T, d>::CheckIslands(const TArray<TGeometryParticleHandle<T, d> *>& Particles)
 {
 	bool bIsValid = true;
 
 	// Check that no particles are in multiple islands
-	TSet<int32> IslandParticlesUnionSet;
-	IslandParticlesUnionSet.Reserve(InParticles.Size());
-	for (int32 Island = 0; Island < IslandParticles.Num(); ++Island)
+	TSet<TGeometryParticleHandle<T, d>*> IslandParticlesUnionSet;
+	IslandParticlesUnionSet.Reserve(Particles.Num());
+	for (int32 Island = 0; Island < IslandToParticles.Num(); ++Island)
 	{
-		TSet<int32> IslandParticlesSet = TSet<int32>(IslandParticles[Island]);
-		TSet<int32> IslandParticlesIntersectSet = IslandParticlesUnionSet.Intersect(IslandParticlesSet);
+		TSet<TGeometryParticleHandle<T, d>*> IslandParticlesSet = TSet<TGeometryParticleHandle<T,d>*>(IslandToParticles[Island]);
+		TSet<TGeometryParticleHandle<T, d>*> IslandParticlesIntersectSet = IslandParticlesUnionSet.Intersect(IslandParticlesSet);
 		if (IslandParticlesIntersectSet.Num() > 0)
 		{
 			// This islands contains particles that were in a previous island.
 			// This is ok only if those particles are static
-			for (const int32 ParticleIndex : IslandParticlesIntersectSet)
+			for (TGeometryParticleHandle<T,d>* Particle : IslandParticlesIntersectSet)
 			{
-				if (!InParticles.HasInfiniteMass(ParticleIndex))
+				if (Particle->AsDynamic())
 				{
-					UE_LOG(LogChaos, Error, TEXT("Island %d contains non-static particle %d that is also in another Island"), Island, ParticleIndex);
+					UE_LOG(LogChaos, Error, TEXT("Island %d contains non-static particle that is also in another Island"), Island);	//todo: add better logging for bad particle
 					bIsValid = false;
 				}
 			}
@@ -590,9 +600,9 @@ bool TPBDConstraintGraph<T, d>::CheckIslands(TPBDRigidParticles<T, d>& InParticl
 	// Check that no constraints refer in the same island
 	TSet<int32> IslandConstraintDataUnionSet;
 	IslandConstraintDataUnionSet.Reserve(Edges.Num());
-	for (int32 Island = 0; Island < IslandConstraints.Num(); ++Island)
+	for (int32 Island = 0; Island < IslandToConstraints.Num(); ++Island)
 	{
-		TSet<int32> IslandConstraintDataSet = TSet<int32>(IslandConstraints[Island]);
+		TSet<int32> IslandConstraintDataSet = TSet<int32>(IslandToConstraints[Island]);
 		TSet<int32> IslandConstraintDataIntersectSet = IslandConstraintDataUnionSet.Intersect(IslandConstraintDataSet);
 		if (IslandConstraintDataIntersectSet.Num() > 0)
 		{
diff --git a/Engine/Source/Runtime/Experimental/Chaos/Private/Chaos/PBDEvolution.cpp b/Engine/Source/Runtime/Experimental/Chaos/Private/Chaos/PBDEvolution.cpp
--- a/Engine/Source/Runtime/Experimental/Chaos/Private/Chaos/PBDEvolution.cpp
+++ b/Engine/Source/Runtime/Experimental/Chaos/Private/Chaos/PBDEvolution.cpp
@@ -16,7 +16,7 @@
 using namespace Chaos;
 
 template<class T, int d>
-TPBDEvolution<T, d>::TPBDEvolution(TPBDParticles<T, d>&& InParticles, TKinematicGeometryParticles<T, d>&& InGeometryParticles, TArray<TVector<int32, 3>>&& CollisionTriangles,
+TPBDEvolution<T, d>::TPBDEvolution(TPBDParticles<T, d>&& InParticles, TKinematicGeometryClothParticles<T, d>&& InGeometryParticles, TArray<TVector<int32, 3>>&& CollisionTriangles,
     int32 NumIterations, T CollisionThickness, T SelfCollisionThickness, T CoefficientOfFriction, T Damping)
     : MParticles(MoveTemp(InParticles))
 	, MCollisionParticles(MoveTemp(InGeometryParticles))
@@ -49,7 +49,7 @@ void TPBDEvolution<T, d>::AdvanceOneTimeStep(const T Dt)
 	TPerParticleEulerStepVelocity<T, d> EulerStepVelocityRule;
 	TPerParticleDampVelocity<T, d> DampVelocityRule(MDamping);
 	TPerParticlePBDEulerStep<T, d> EulerStepRule;
-	TPerParticlePBDCollisionConstraint<T, d> CollisionRule(MCollisionParticles, MCollided, MCollisionThickness, MCoefficientOfFriction);
+	TPerParticlePBDCollisionConstraint<T, d, EGeometryParticlesSimType::Other> CollisionRule(MCollisionParticles, MCollided, MCollisionThickness, MCoefficientOfFriction);
 
 	DampVelocityRule.UpdatePositionBasedState(MParticles);
 
diff --git a/Engine/Source/Runtime/Experimental/Chaos/Private/Chaos/PBDJointConstraintsBase.cpp b/Engine/Source/Runtime/Experimental/Chaos/Private/Chaos/PBDJointConstraintsBase.cpp
--- a/Engine/Source/Runtime/Experimental/Chaos/Private/Chaos/PBDJointConstraintsBase.cpp
+++ b/Engine/Source/Runtime/Experimental/Chaos/Private/Chaos/PBDJointConstraintsBase.cpp
@@ -5,74 +5,165 @@ using namespace Chaos;
 
 // @todo(mlentine): This should be in a utility class somewhere
 template<class T>
-PMatrix<T, 3, 3> ComputeJointFactorMatrix(const TVector<T, 3>& V, const PMatrix<T, 3, 3>& M, const T& Im)
+PMatrix<T, 3, 3> ComputeJointFactorMatrix2(const TVector<T, 3>& V, const PMatrix<T, 3, 3>& M, const T& Im)
 {
-	if (Im > FLT_MIN)
-	{
-		// Rigid objects rotational contribution to the impulse.
-		// Vx*M*VxT+Im
-		return PMatrix<T, 3, 3>(
-			-V[2] * (-V[2] * M.M[1][1] + V[1] * M.M[2][1]) + V[1] * (-V[2] * M.M[2][1] + V[1] * M.M[2][2]) + Im,
-			V[2] * (-V[2] * M.M[1][0] + V[1] * M.M[2][0]) - V[0] * (-V[2] * M.M[2][1] + V[1] * M.M[2][2]),
-			-V[1] * (-V[2] * M.M[1][0] + V[1] * M.M[2][0]) + V[0] * (-V[2] * M.M[1][1] + V[1] * M.M[2][1]),
-			V[2] * (V[2] * M.M[0][0] - V[0] * M.M[2][0]) - V[0] * (V[2] * M.M[2][0] - V[0] * M.M[2][2]) + Im,
-			-V[1] * (V[2] * M.M[0][0] - V[0] * M.M[2][0]) + V[0] * (V[2] * M.M[1][0] - V[0] * M.M[2][1]),
-			-V[1] * (-V[1] * M.M[0][0] + V[0] * M.M[1][0]) + V[0] * (-V[1] * M.M[1][0] + V[0] * M.M[1][1]) + Im);
-	}
-	else
-	{
-		return PMatrix<T, 3, 3>(0);
-	}
+	// Rigid objects rotational contribution to the impulse.
+	// Vx*M*VxT+Im
+	check(Im > FLT_MIN);
+	return PMatrix<T, 3, 3>(
+		-V[2] * (-V[2] * M.M[1][1] + V[1] * M.M[2][1]) + V[1] * (-V[2] * M.M[2][1] + V[1] * M.M[2][2]) + Im,
+		V[2] * (-V[2] * M.M[1][0] + V[1] * M.M[2][0]) - V[0] * (-V[2] * M.M[2][1] + V[1] * M.M[2][2]),
+		-V[1] * (-V[2] * M.M[1][0] + V[1] * M.M[2][0]) + V[0] * (-V[2] * M.M[1][1] + V[1] * M.M[2][1]),
+		V[2] * (V[2] * M.M[0][0] - V[0] * M.M[2][0]) - V[0] * (V[2] * M.M[2][0] - V[0] * M.M[2][2]) + Im,
+		-V[1] * (V[2] * M.M[0][0] - V[0] * M.M[2][0]) + V[0] * (V[2] * M.M[1][0] - V[0] * M.M[2][1]),
+		-V[1] * (-V[1] * M.M[0][0] + V[0] * M.M[1][0]) + V[0] * (-V[1] * M.M[1][0] + V[0] * M.M[1][1]) + Im);
 }
 
 template<class T, int d>
-void TPBDJointConstraintsBase<T, d>::UpdateDistanceInternal(const TRigidParticles<T, d>& InParticles, const TVector<T, d>& InLocation, const int32 InConstraintIndex)
+void TPBDJointConstraintsBase<T, d>::UpdateDistanceInternal(const TVector<T, d>& InLocation, const int32 InConstraintIndex)
 {
-	const int32 ParticleIndex1 = Constraints[InConstraintIndex][0];
-	const int32 ParticleIndex2 = Constraints[InConstraintIndex][1];
-	Distances[InConstraintIndex][0] = InParticles.R(ParticleIndex1).Inverse().RotateVector(InLocation - InParticles.X(ParticleIndex1));
-	Distances[InConstraintIndex][1] = InParticles.R(ParticleIndex2).Inverse().RotateVector(InLocation - InParticles.X(ParticleIndex2));
+	const TGeometryParticleHandle<T,d>* Particle0 = Constraints[InConstraintIndex][0];
+	const TGeometryParticleHandle<T, d>* Particle1 = Constraints[InConstraintIndex][1];
+	Distances[InConstraintIndex][0] = Particle0->R().Inverse().RotateVector(InLocation - Particle0->X());
+	Distances[InConstraintIndex][1] = Particle1->R().Inverse().RotateVector(InLocation - Particle1->X());
 }
 
 template<class T, int d>
-void TPBDJointConstraintsBase<T, d>::UpdateDistance(const TRigidParticles<T, d>& InParticles, const TVector<T, d>& InLocation, const int32 InConstraintIndex)
+void TPBDJointConstraintsBase<T, d>::UpdateDistance(const TVector<T, d>& InLocation, const int32 InConstraintIndex)
 {
 	Distances.SetNum(Constraints.Num());
-	UpdateDistanceInternal(InParticles, InLocation, InConstraintIndex);
+	UpdateDistanceInternal(InLocation, InConstraintIndex);
 }
 
 template<class T, int d>
-void TPBDJointConstraintsBase<T, d>::UpdateDistances(const TRigidParticles<T, d>& InParticles, const TArray<TVector<T, d>>& InLocations)
+void TPBDJointConstraintsBase<T, d>::UpdateDistances(const TArray<TVector<T, d>>& InLocations)
 {
 	const int32 NumConstraints = Constraints.Num();
 	Distances.SetNum(NumConstraints);
 	for (int32 ConstraintIndex = 0; ConstraintIndex < NumConstraints; ++ConstraintIndex)
 	{
-		UpdateDistanceInternal(InParticles, InLocations[ConstraintIndex], ConstraintIndex);
+		UpdateDistanceInternal(InLocations[ConstraintIndex], ConstraintIndex);
 	}
 }
 
 template<class T, int d>
-TVector<T, d> TPBDJointConstraintsBase<T, d>::GetDelta(const TPBDRigidParticles<T, d>& InParticles, const TVector<T, d>& WorldSpaceX1, const TVector<T, d>& WorldSpaceX2, const PMatrix<T, d, d>& WorldSpaceInvI1, const PMatrix<T, d, d>& WorldSpaceInvI2, int32 ConstraintIndex) const
+TVector<T, d> TPBDJointConstraintsBase<T, d>::GetDeltaDynamicDynamic(const TVector<T, d>& P0, const TVector<T, d>& P1, const TVector<T, d>& C0, const TVector<T, d>& C1, const PMatrix<T, d, d>& InvI0, const PMatrix<T, d, d>& InvI1, const T InvM0, const T InvM1) const
 {
-	const TVector<int32, 2>& Constraint = Constraints[ConstraintIndex];
-	int32 ConstraintInnerIndex1 = Constraint[0];
-	int32 ConstraintInnerIndex2 = Constraint[1];
+	PMatrix<T, d, d> Factor = ComputeJointFactorMatrix2(C0 - P0, InvI0, InvM0) + ComputeJointFactorMatrix2(C1 - P1, InvI1, InvM1);
+	PMatrix<T, d, d> FactorInv = Factor.Inverse();
+	FactorInv.M[3][3] = 1;
+	TVector<T, d> Delta = C1 - C0;
+	return FactorInv * Delta;
+}
 
-	if (InParticles.InvM(ConstraintInnerIndex2) == 0 && InParticles.InvM(ConstraintInnerIndex1) == 0)
+template<class T, int d>
+void TPBDJointConstraintsBase<T, d>::ApplyDynamicDynamic(const T Dt, const int32 ConstraintIndex, const int32 PBDRigid0Index, const int32 PBDRigid1Index, const bool bApplyProjection) const
+{
+	check((PBDRigid0Index == 0) || (PBDRigid0Index == 1));
+	check((PBDRigid1Index == 0) || (PBDRigid1Index == 1));
+	check(PBDRigid0Index != PBDRigid1Index);
+
+	TPBDRigidParticleHandle<T, d>* PBDRigid0 = Constraints[ConstraintIndex][PBDRigid0Index]->AsDynamic();
+	TPBDRigidParticleHandle<T, d>* PBDRigid1 = Constraints[ConstraintIndex][PBDRigid1Index]->AsDynamic();
+	check(PBDRigid0 && PBDRigid1 && (PBDRigid0->Island() == PBDRigid1->Island()));
+
+	TRotation<T, d>& Q0 = PBDRigid0->Q();
+	TVector<T, d>& P0 = PBDRigid0->P();
+	TRotation<T, d>& Q1 = PBDRigid1->Q();
+	TVector<T, d>& P1 = PBDRigid1->P();
+
+	// Calculate world-space constraint positions
+	const TVector<T, 3>& Distance0 = Distances[ConstraintIndex][PBDRigid0Index];
+	const TVector<T, 3>& Distance1 = Distances[ConstraintIndex][PBDRigid1Index];
+	const TVector<T, d> C0 = Q0.RotateVector(Distance0) + P0;
+	const TVector<T, d> C1 = Q1.RotateVector(Distance1) + P1;
+
+	// Calculate world-space mass
+	const PMatrix<T, d, d> InvI0 = (Q0 * FMatrix::Identity) * PBDRigid0->InvI() * (Q0 * FMatrix::Identity).GetTransposed();
+	const PMatrix<T, d, d> InvI1 = (Q1 * FMatrix::Identity) * PBDRigid1->InvI() * (Q1 * FMatrix::Identity).GetTransposed();
+	const float InvM0 = PBDRigid0->InvM();
+	const float InvM1 = PBDRigid1->InvM();
+
+	// Calculate mass-weighted correction
+	const TVector<T, d> Delta = GetDeltaDynamicDynamic(P0, P1, C0, C1, InvI0, InvI1, InvM0, InvM1);
+
+	// Apply corrections
+	P0 += InvM0 * Delta;
+	Q0 += TRotation<T, d>(InvI0 * TVector<T, d>::CrossProduct(C0 - P0, Delta), 0.f) * Q0 * T(0.5);
+	Q0.Normalize();
+
+	P1 -= InvM1 * Delta;
+	Q1 += TRotation<T, d>(InvI1 * TVector<T, d>::CrossProduct(C1 - P1, -Delta), 0.f) * Q1 * T(0.5);
+	Q1.Normalize();
+
+	// Correct any remaining error by translating
+	if (bApplyProjection)
 	{
-		return TVector<T, d>(0);
+		const TVector<T, d> C0Proj = Q0.RotateVector(Distance0) + P0;
+		const TVector<T, d> C1Proj = Q1.RotateVector(Distance1) + P1;
+		const TVector<T, d> DeltaProj = (C1Proj - C0Proj) / (InvM0 + InvM1);
+
+		P0 += InvM0 * DeltaProj;
+		P1 -= InvM1 * DeltaProj;
 	}
+}
 
-	PMatrix<T, d, d> Factor =
-		ComputeJointFactorMatrix(WorldSpaceX1 - InParticles.P(ConstraintInnerIndex1), WorldSpaceInvI1, InParticles.InvM(ConstraintInnerIndex1)) +
-		ComputeJointFactorMatrix(WorldSpaceX2 - InParticles.P(ConstraintInnerIndex2), WorldSpaceInvI2, InParticles.InvM(ConstraintInnerIndex2));
+template<class T, int d>
+TVector<T, d> TPBDJointConstraintsBase<T, d>::GetDeltaDynamicKinematic(const TVector<T, d>& P0, const TVector<T, d>& C0, const TVector<T, d>& C1, const PMatrix<T, d, d>& InvI0, const T InvM0) const
+{
+	PMatrix<T, d, d> Factor = ComputeJointFactorMatrix2(C0 - P0, InvI0, InvM0);
 	PMatrix<T, d, d> FactorInv = Factor.Inverse();
 	FactorInv.M[3][3] = 1;
-	TVector<T, d> Delta = WorldSpaceX2 - WorldSpaceX1;
+	TVector<T, d> Delta = C1 - C0;
 	return FactorInv * Delta;
 }
 
+template<class T, int d>
+void TPBDJointConstraintsBase<T, d>::ApplyDynamicStatic(const T Dt, const int32 ConstraintIndex, const int32 PBDRigid0Index, const int32 Static1Index, const bool bApplyProjection) const
+{
+	check((PBDRigid0Index == 0) || (PBDRigid0Index == 1));
+	check((Static1Index == 0) || (Static1Index == 1));
+	check(PBDRigid0Index != Static1Index);
+
+	TPBDRigidParticleHandle<T, d>* PBDRigid0 = Constraints[ConstraintIndex][PBDRigid0Index]->AsDynamic();
+	TGeometryParticleHandle<T, d>* Static1 = Constraints[ConstraintIndex][Static1Index];
+	check(PBDRigid0 && Static1 && !Static1->AsDynamic());
+
+	TRotation<T, d>& Q0 = PBDRigid0->Q();
+	TVector<T, d>& P0 = PBDRigid0->P();
+	const TRotation<T, d>& Q1 = Static1->R();
+	const TVector<T, d>& P1 = Static1->X();
+
+	// Calculate world-space constraint positions
+	const TVector<T, 3>& Distance0 = Distances[ConstraintIndex][PBDRigid0Index];
+	const TVector<T, 3>& Distance1 = Distances[ConstraintIndex][Static1Index];
+	const TVector<T, d> C0 = Q0.RotateVector(Distance0) + P0;
+	const TVector<T, d> C1 = Q1.RotateVector(Distance1) + P1;
+
+	// Calculate world-space mass
+	const PMatrix<T, d, d> InvI0 = (Q0 * FMatrix::Identity) * PBDRigid0->InvI() * (Q0 * FMatrix::Identity).GetTransposed();
+	const float InvM0 = PBDRigid0->InvM();
+
+	// Calculate mass-weighted correction
+	const TVector<T, d> Delta = GetDeltaDynamicKinematic(P0, C0, C1, InvI0, InvM0);
+
+	// Apply correction
+	P0 += InvM0 * Delta;
+	Q0 += TRotation<T, d>(InvI0 * TVector<T, d>::CrossProduct(C0 - P0, Delta), 0.f) * Q0 * T(0.5);
+	Q0.Normalize();
+
+	// Correct any remaining error by translating
+	if (bApplyProjection)
+	{
+		const TVector<T, d> C0Proj = Q0.RotateVector(Distance0) + P0;
+		const TVector<T, d> C1Proj = Q1.RotateVector(Distance1) + P1;
+		const TVector<T, d> DeltaProj = (C1Proj - C0Proj);
+
+		P0 += DeltaProj;
+	}
+}
+
+
 namespace Chaos
 {
 	template class TPBDJointConstraintsBase<float, 3>;
