Copying //UE4/Dev-Physics to Dev-Main (//UE4/Dev-Main) [at] 7703071

diff --git a/Engine/Source/Runtime/Experimental/Chaos/Private/Chaos/PBDRigidClustering.cpp b/Engine/Source/Runtime/Experimental/Chaos/Private/Chaos/PBDRigidClustering.cpp
--- a/Engine/Source/Runtime/Experimental/Chaos/Private/Chaos/PBDRigidClustering.cpp
+++ b/Engine/Source/Runtime/Experimental/Chaos/Private/Chaos/PBDRigidClustering.cpp
@@ -15,6 +15,11 @@
 #include "Containers/Queue.h"
 #include "ProfilingDebugging/ScopedTimers.h"
 #include "Voronoi/Voronoi.h"
+#include "Chaos/PBDRigidsEvolutionGBF.h"
+#include "Chaos/PerParticleInitForce.h"
+#include "Chaos/PerParticleEulerStepVelocity.h"
+#include "Chaos/PerParticleEtherDrag.h"
+#include "Chaos/PerParticlePBDEulerStep.h"
 
 namespace Chaos
 {
@@ -64,34 +69,29 @@ namespace Chaos
 	FAutoConsoleVariableRef CVarDeactivateClusterChildren(TEXT("p.DeactivateClusterChildren"), DeactivateClusterChildren, TEXT("If children should be decativated when broken and put into another cluster."));
 
 	template<class FPBDRigidsEvolution, class FPBDCollisionConstraint, class T, int d>
-	TPBDRigidClustering<FPBDRigidsEvolution, FPBDCollisionConstraint, T, d>::TPBDRigidClustering(FPBDRigidsEvolution& InEvolution, TPBDRigidParticles<T, d>& InParticles)
-	    : MEvolution(InEvolution)
-	    , MParticles(InParticles)
-	    , MCollisionImpulseArrayDirty(true)
-	    , DoGenerateBreakingData(false)
+	TPBDRigidClustering<FPBDRigidsEvolution, FPBDCollisionConstraint, T, d>::TPBDRigidClustering(FPBDRigidsEvolution& InEvolution, TPBDRigidClusteredParticles<T, d>& InParticles)
+		: MEvolution(InEvolution)
+		, MParticles(InParticles)
+		, MCollisionImpulseArrayDirty(true)
+		, DoGenerateBreakingData(false)
 		, MClusterConnectionFactor(1.0)
 		, MClusterUnionConnectionType(FClusterCreationParameters<T>::EConnectionMethod::DelaunayTriangulation)
 	{
-		MParticles.AddArray(&MClusterIds);
-		MParticles.AddArray(&MClusterGroupIndex);
-		MParticles.AddArray(&MCollisionImpulses);
-		MParticles.AddArray(&MInternalCluster);
-		MParticles.AddArray(&MStrains);
-		MParticles.AddArray(&MChildToParent);
-		MParticles.AddArray(&MConnectivityEdges);
-		MParticles.AddArray(&MChildrenSpatial);
-		MParticles.AddArray(&MMultiChildProxyId);
-		MParticles.AddArray(&MMultiChildProxyData);
+	}
+
+	template<class FPBDRigidsEvolution, class FPBDCollisionConstraint, class T, int d>
+	TPBDRigidClustering<FPBDRigidsEvolution, FPBDCollisionConstraint, T, d>::~TPBDRigidClustering()
+	{
 	}
 
 	DECLARE_CYCLE_STAT(TEXT("TPBDRigidClustering<>::CreateClusterParticle"), STAT_CreateClusterParticle, STATGROUP_Chaos);
 	template<class FPBDRigidsEvolution, class FPBDCollisionConstraint, class T, int d>
 	int32 TPBDRigidClustering<FPBDRigidsEvolution, FPBDCollisionConstraint, T, d>::CreateClusterParticle(
-	    int32 ClusterGroupIndex,
-	    const TArray<uint32>& Children,
-	    /*const*/ TSerializablePtr<TImplicitObject<T, d>> ProxyGeometry,
-	    const TRigidTransform<T, d>* ForceMassOrientation,
-	    const FClusterCreationParameters<T>& Parameters)
+		int32 ClusterGroupIndex,
+		TArray<uint32>&& Children,
+		/*const*/ TSerializablePtr<TImplicitObject<T, d>> ProxyGeometry,
+		const TRigidTransform<T, d>* ForceMassOrientation,
+		const FClusterCreationParameters<T>& Parameters)
 	{
 		SCOPE_CYCLE_COUNTER(STAT_CreateClusterParticle);
 
@@ -101,55 +101,61 @@ namespace Chaos
 		if (NewIndex == INDEX_NONE)
 		{
 			NewIndex = MParticles.Size();
-			MParticles.AddParticles(1);
+			MEvolution.CreateClusteredParticles(1);
 		}
+		TGeometryParticleHandle<T, d>* NewParticle = MParticles.Handle(NewIndex);
 
-		MEvolution.EnableParticle(NewIndex, Children[0]);
+		MEvolution.EnableParticle(NewParticle, MParticles.Handle(Children[0]));
 		MParticles.CollisionGroup(NewIndex) = INT_MAX;
 		TopLevelClusterParents.Add(NewIndex);
 
-		MInternalCluster[NewIndex] = false;
-		MClusterIds[NewIndex] = ClusterId(INDEX_NONE, Children.Num());
-		MClusterGroupIndex[NewIndex] = ClusterGroupIndex;
+		MParticles.InternalCluster(NewIndex) = false;
+		MParticles.ClusterIds(NewIndex) = ClusterId(INDEX_NONE, Children.Num());
+		MParticles.ClusterGroupIndex(NewIndex) = ClusterGroupIndex;
 
 		//
 		// Update clustering data structures.
 		//
-		if (MParentToChildren.Contains(NewIndex))
+		if (MParentToChildren.Contains(NewParticle))
 		{
-			MParentToChildren[NewIndex] = MakeUnique<TArray<uint32>>(Children);
+			MParentToChildren[NewParticle] = MoveTemp(Children);
 		}
 		else
 		{
-			MParentToChildren.Add(NewIndex, MakeUnique<TArray<uint32>>(Children));
+			MParentToChildren.Add(NewParticle, MoveTemp(Children));
 		}
 
 		bool bClusterIsAsleep = true;
-		MStrains[NewIndex] = 0.0;
+		MParticles.Strains(NewIndex) = 0.0;
 
-		const TArray<uint32>& ChildrenArray = *MParentToChildren[NewIndex];
-		TSet<uint32> ChildrenSet = TSet<uint32>(ChildrenArray);	// @todo(ccaulfield): try to eliminate set creation (required by RemoveConstraints)
-		MEvolution.DisableParticles(ChildrenSet);
+		const TArray<uint32>& ChildrenArray = MParentToChildren[NewParticle];
+		//todo: avoid this iteration by passing a view into DisableParticles
+		TSet<TGeometryParticleHandle<T, d>*> ChildrenHandles;
+		for (auto Child : ChildrenArray)
+		{
+			ChildrenHandles.Add(MParticles.Handle(Child));
+		}
+		MEvolution.DisableParticles(ChildrenHandles);
 		for (const uint32 Child : ChildrenArray)
 		{
 			TopLevelClusterParents.Remove(Child);
 			bClusterIsAsleep &= MParticles.Sleeping(Child);
-			MClusterGroupIndex[Child] = 0;
-			MClusterIds[Child].Id = NewIndex;
-			MStrains[NewIndex] += MStrains[Child];
+			MParticles.ClusterGroupIndex(Child) = 0;
+			MParticles.ClusterIds(Child).Id = NewIndex;
+			MParticles.Strains(NewIndex) += MParticles.Strains(Child);
 
-			MCollisionImpulses[NewIndex] = FMath::Max(MCollisionImpulses[NewIndex], MCollisionImpulses[Child]);
+			MParticles.CollisionImpulses(NewIndex) = FMath::Max(MParticles.CollisionImpulses(NewIndex), MParticles.CollisionImpulses(Child));
 			MParticles.CollisionGroup(NewIndex) = (MParticles.CollisionGroup(NewIndex) < MParticles.CollisionGroup(Child)) ? MParticles.CollisionGroup(NewIndex) : MParticles.CollisionGroup(Child);
 		}
-		if (Children.Num())
+		if (ChildrenArray.Num())
 		{
-			MStrains[NewIndex] /= Children.Num();
+			MParticles.Strains(NewIndex) /= ChildrenArray.Num();
 		}
 
 		ensureMsgf(!ProxyGeometry || ForceMassOrientation, TEXT("If ProxyGeometry is passed, we must override the mass orientation as they are tied"));
 
-		UpdateMassProperties(Children, NewIndex, ForceMassOrientation);
-		UpdateGeometry(Children, NewIndex, ProxyGeometry, Parameters);
+		UpdateMassProperties(ChildrenArray, NewIndex, ForceMassOrientation);
+		UpdateGeometry(ChildrenArray, NewIndex, ProxyGeometry, Parameters);
 		GenerateConnectionGraph(NewIndex, Parameters);
 		MParticles.SetSleeping(NewIndex, bClusterIsAsleep);
 
@@ -161,7 +167,7 @@ namespace Chaos
 
 	DECLARE_CYCLE_STAT(TEXT("TPBDRigidClustering<>::CreateClusterParticleFromClusterChildren"), STAT_CreateClusterParticleFromClusterChildren, STATGROUP_Chaos);
 	template<class FPBDRigidsEvolution, class FPBDCollisionConstraint, class T, int d>
-	int32 TPBDRigidClustering<FPBDRigidsEvolution, FPBDCollisionConstraint, T, d>::CreateClusterParticleFromClusterChildren(const TArray<uint32>& Children, const int32 ParentIndex, const TRigidTransform<T, d>& ClusterWorldTM, const FClusterCreationParameters<T>& Parameters)
+	int32 TPBDRigidClustering<FPBDRigidsEvolution, FPBDCollisionConstraint, T, d>::CreateClusterParticleFromClusterChildren(TArray<uint32>&& Children, TPBDRigidClusteredParticleHandle<T,d>* Parent, const TRigidTransform<T, d>& ClusterWorldTM, const FClusterCreationParameters<T>& Parameters)
 	{
 		SCOPE_CYCLE_COUNTER(STAT_CreateClusterParticleFromClusterChildren);
 		FClusterMap& MParentToChildren = GetChildrenMap();
@@ -171,48 +177,55 @@ namespace Chaos
 		if (NewIndex == INDEX_NONE)
 		{
 			NewIndex = MParticles.Size();
-			MParticles.AddParticles(1);
+			MEvolution.CreateClusteredParticles(1);
 		}
-
-		MEvolution.EnableParticle(NewIndex, ParentIndex);
+		TGeometryParticleHandle<T, d>* NewParticle = MParticles.Handle(NewIndex);
+		MEvolution.EnableParticle(MParticles.Handle(NewIndex), Parent);
 
 		MParticles.CollisionGroup(NewIndex) = INT_MAX;
 		TopLevelClusterParents.Add(NewIndex);
-		MInternalCluster[NewIndex] = true;
-		MClusterIds[NewIndex] = ClusterId(INDEX_NONE, Children.Num());
+		MParticles.InternalCluster(NewIndex) = true;
+		MParticles.ClusterIds(NewIndex) = ClusterId(INDEX_NONE, Children.Num());
 
 		//
 		// Update clustering data structures.
 		//
-		if (MParentToChildren.Contains(NewIndex))
+		if (MParentToChildren.Contains(NewParticle))
 		{
-			MParentToChildren[NewIndex] = MakeUnique<TArray<uint32>>(Children);
+			MParentToChildren[NewParticle] = MoveTemp(Children);
 		}
 		else
 		{
-			MParentToChildren.Add(NewIndex, MakeUnique<TArray<uint32>>(Children));
+			MParentToChildren.Add(NewParticle, MoveTemp(Children));
 		}
 
+		const TArray<uint32>& ChildrenArray = MParentToChildren[NewParticle];
 		//child transforms are out of date, need to update them. @todo(ocohen): if children transforms are relative we would not need to update this, but would simply have to do a final transform on the new cluster index
 		// TODO(mlentine): Why is this not needed? (Why is it ok to have DeactivateClusterChildren==false?)
 		if (DeactivateClusterChildren)
 		{
-			TSet<uint32> ChildrenSet = TSet<uint32>(Children);	// @todo(ccaulfield): try to eliminate set creation (required by RemoveConstraints)
-			MEvolution.DisableParticles(ChildrenSet);
+			//TODO: avoid iteration just pass in a view
+			//TSet<uint32> ChildrenSet = TSet<uint32>(Children);	// @todo(ccaulfield): try to eliminate set creation (required by RemoveConstraints)
+			TSet<TGeometryParticleHandle<T, d>*> ChildrenHandles;
+			for (auto Child : ChildrenArray)
+			{
+				ChildrenHandles.Add(MParticles.Handle(Child));
+			}
+			MEvolution.DisableParticles(ChildrenHandles);
 		}
-		for (uint32 Child : Children)
+		for (uint32 Child : ChildrenArray)
 		{
-			TRigidTransform<T, d> ChildFrame = MChildToParent[Child] * ClusterWorldTM;
+			TRigidTransform<T, d> ChildFrame = MParticles.ChildToParent(Child) * ClusterWorldTM;
 			MParticles.X(Child) = ChildFrame.GetTranslation();
 			MParticles.R(Child) = ChildFrame.GetRotation();
-			MClusterIds[Child].Id = NewIndex;
-			MClusterGroupIndex[Child] = 0;
+			MParticles.ClusterIds(Child).Id = NewIndex;
+			MParticles.ClusterGroupIndex(Child) = 0;
 			if (DeactivateClusterChildren)
 			{
 				TopLevelClusterParents.Remove(Child);
 			}
 
-			MCollisionImpulses[NewIndex] = FMath::Max(MCollisionImpulses[NewIndex], MCollisionImpulses[Child]);
+			MParticles.CollisionImpulses(NewIndex) = FMath::Max(MParticles.CollisionImpulses(NewIndex), MParticles.CollisionImpulses(Child));
 			MParticles.CollisionGroup(NewIndex) = (MParticles.CollisionGroup(NewIndex) < MParticles.CollisionGroup(Child)) ? MParticles.CollisionGroup(NewIndex) : MParticles.CollisionGroup(Child);
 		}
 
@@ -220,8 +233,8 @@ namespace Chaos
 		NoCleanParams.bCleanCollisionParticles = false;
 		NoCleanParams.bCopyCollisionParticles = !!UnionsHaveCollisionParticles;
 
-		UpdateMassProperties(Children, NewIndex, nullptr);
-		UpdateGeometry(Children, NewIndex, TSerializablePtr<TImplicitObject<T, 3>>(), NoCleanParams);
+		UpdateMassProperties(ChildrenArray, NewIndex, nullptr);
+		UpdateGeometry(ChildrenArray, NewIndex, TSerializablePtr<TImplicitObject<T, 3>>(), NoCleanParams);
 
 		return NewIndex;
 	}
@@ -241,10 +254,10 @@ namespace Chaos
 		FClusterMap& MParentToChildren = GetChildrenMap();
 
 		TMap<int32, TArray<uint32>> GroupMapping;
-		for (int32 i = 0; i < MClusterGroupIndex.Num(); i++) // just loop active clusters here.
+		for (uint32 i = 0; i < MParticles.Size(); i++) // just loop active clusters here.
 		{
 			uint32 ParticleIndex = i;
-			int32 GroupIndex = MClusterGroupIndex[ParticleIndex];
+			int32 GroupIndex = MParticles.ClusterGroupIndex(ParticleIndex);
 			if (GroupIndex > 0)
 			{
 				if (!GroupMapping.Contains(GroupIndex))
@@ -263,19 +276,19 @@ namespace Chaos
 				for (uint32& OriginalRootIdx : Group.Value)
 				{
 					TUniquePtr<TMultiChildProxyData<T, d>> ProxyData;
-					if (MParentToChildren.Contains(OriginalRootIdx))
+					if (MParentToChildren.Contains(MParticles.Handle(OriginalRootIdx)))
 					{
-						if (UseMultiChildProxy && !MParticles.DynamicGeometry(OriginalRootIdx) && MParentToChildren[OriginalRootIdx]->Num() > MinChildrenForMultiProxy) //Don't support dynamic geometry
+						if (UseMultiChildProxy && !MParticles.DynamicGeometry(OriginalRootIdx) && MParentToChildren[MParticles.Handle(OriginalRootIdx)].Num() > MinChildrenForMultiProxy) //Don't support dynamic geometry
 						{
-							if (ensure(MParentToChildren[OriginalRootIdx]->Num()))
+							if (ensure(MParentToChildren[MParticles.Handle(OriginalRootIdx)].Num()))
 							{
 								ProxyData = MakeUnique<TMultiChildProxyData<T, d>>();
-								ProxyData->KeyChild = (*MParentToChildren[OriginalRootIdx])[0];
+								ProxyData->KeyChild = (MParentToChildren[MParticles.Handle(OriginalRootIdx)])[0];
 								ProxyData->RelativeToKeyChild = TRigidTransform<T, d>(MParticles.X(OriginalRootIdx), MParticles.R(OriginalRootIdx)); //store world space of original root. Need to break it up and then compute relative to world space of key child
 							}
 						}
 
-						const TArray<uint32> OriginalRootChildren = DeactivateClusterParticle(OriginalRootIdx).Array();
+						const TArray<uint32> OriginalRootChildren = DeactivateClusterParticle(MParticles.Handle(OriginalRootIdx)).Array();
 						ClusterChildren.Append(OriginalRootChildren);
 
 						if (ProxyData)
@@ -283,12 +296,12 @@ namespace Chaos
 							//now that we have world space updated for key child, compute relative transform for original root
 							const TRigidTransform<T, d> OriginalRootWorldTM = ProxyData->RelativeToKeyChild;
 							ProxyData->RelativeToKeyChild = OriginalRootWorldTM.GetRelativeTransform(TRigidTransform<T, d>(MParticles.X(ProxyData->KeyChild), MParticles.R(ProxyData->KeyChild)));
-							MMultiChildProxyData[OriginalRootIdx] = MoveTemp(ProxyData);
+							MParticles.MultiChildProxyData(OriginalRootIdx) = MoveTemp(ProxyData);
 
 							for (uint32 Child : OriginalRootChildren)
 							{
 								//remember original proxy of child cluster
-								MMultiChildProxyId[Child].Id = OriginalRootIdx;
+								MParticles.MultiChildProxyId(Child).Id = OriginalRootIdx;
 							}
 						}
 					}
@@ -300,9 +313,9 @@ namespace Chaos
 
 				FClusterCreationParameters<T> Parameters(0.3, 100, false, !!UnionsHaveCollisionParticles);
 				Parameters.ConnectionMethod = MClusterUnionConnectionType;
-				int32 NewIndex = CreateClusterParticle(-Group.Key, ClusterChildren, TSerializablePtr<TImplicitObject<T, d>>(), nullptr, Parameters);
-				MInternalCluster[NewIndex] = true;
-				MEvolution.SetPhysicsMaterial(NewIndex, MEvolution.GetPhysicsMaterial(Group.Value[0]));
+				int32 NewIndex = CreateClusterParticle(-Group.Key, MoveTemp(ClusterChildren), TSerializablePtr<TImplicitObject<T, d>>(), nullptr, Parameters);
+				MParticles.InternalCluster(NewIndex) = true;
+				MEvolution.SetPhysicsMaterial(MParticles.Handle(NewIndex), MEvolution.GetPhysicsMaterial(MParticles.Handle(Group.Value[0])));
 
 				PendingClusterCounter.Remove(Group.Key);
 			}
@@ -311,19 +324,16 @@ namespace Chaos
 
 	DECLARE_CYCLE_STAT(TEXT("TPBDRigidClustering<>::DeactivateClusterParticle"), STAT_DeactivateClusterParticle, STATGROUP_Chaos);
 	template<class FPBDRigidsEvolution, class FPBDCollisionConstraint, class T, int d>
-	TSet<uint32> TPBDRigidClustering<FPBDRigidsEvolution, FPBDCollisionConstraint, T, d>::DeactivateClusterParticle(const uint32 ClusterIndex)
+	TSet<uint32> TPBDRigidClustering<FPBDRigidsEvolution, FPBDCollisionConstraint, T, d>::DeactivateClusterParticle(TPBDRigidClusteredParticleHandle<T,d>* ClusteredParticle)
 	{
 		SCOPE_CYCLE_COUNTER(STAT_DeactivateClusterParticle);
 		FClusterMap& MParentToChildren = GetChildrenMap();
 
 		TSet<uint32> ActivatedChildren;
-		if (0 <= ClusterIndex && ClusterIndex < (uint32)MClusterIds.Num())
+		check(!ClusteredParticle->Disabled());
+		if (MParentToChildren.Contains(ClusteredParticle))
 		{
-			check(!MParticles.Disabled(ClusterIndex));
-			if (MParentToChildren.Contains(ClusterIndex))
-			{
-				ActivatedChildren = ReleaseClusterParticles(*MParentToChildren[ClusterIndex]);
-			}
+			ActivatedChildren = ReleaseClusterParticles(MParentToChildren[ClusteredParticle]);
 		}
 		return ActivatedChildren;
 	}
@@ -347,13 +357,13 @@ namespace Chaos
 
 			for (uint32 ChildIdx : ChildrenParticles)
 			{
-				if (MParticles.Disabled(ChildIdx) && MClusterIds[ChildIdx].Id != INDEX_NONE)
+				if (MParticles.Disabled(ChildIdx) && MParticles.ClusterIds(ChildIdx).Id != INDEX_NONE)
 				{
-					if (ensure(!bFound || MClusterIds[ChildIdx].Id == ClusterIdx))
+					if (ensure(!bFound || MParticles.ClusterIds(ChildIdx).Id == ClusterIdx))
 					{
 						bFound = true;
 						FakeStrain[ChildIdx] = FLT_MAX;
-						ClusterIdx = MClusterIds[ChildIdx].Id;
+						ClusterIdx = MParticles.ClusterIds(ChildIdx).Id;
 					}
 					else
 					{
@@ -364,7 +374,7 @@ namespace Chaos
 
 			if (bFound)
 			{
-				ActivatedBodies.Append(ReleaseClusterParticles(ClusterIdx, FakeStrain));
+				ActivatedBodies.Append(ReleaseClusterParticles(MParticles.Handle(ClusterIdx), FakeStrain));
 			}
 
 			DoGenerateBreakingData = bPreDoGenerateData;
@@ -375,7 +385,7 @@ namespace Chaos
 
 	DECLARE_CYCLE_STAT(TEXT("TPBDRigidClustering<>::ReleaseClusterParticles(STRAIN)"), STAT_ReleaseClusterParticles_STRAIN, STATGROUP_Chaos);
 	template<class FPBDRigidsEvolution, class FPBDCollisionConstraint, class T, int d>
-	TSet<uint32> TPBDRigidClustering<FPBDRigidsEvolution, FPBDCollisionConstraint, T, d>::ReleaseClusterParticles(const uint32 ClusterIndex, const TArrayView<T>& StrainArray)
+	TSet<uint32> TPBDRigidClustering<FPBDRigidsEvolution, FPBDCollisionConstraint, T, d>::ReleaseClusterParticles(TPBDRigidClusteredParticleHandle<T,d>* ClusteredParticle, const TArrayView<T>& StrainArray)
 	{
 		SCOPE_CYCLE_COUNTER(STAT_ReleaseClusterParticles_STRAIN);
 		FClusterMap& MParentToChildren = GetChildrenMap();
@@ -383,34 +393,34 @@ namespace Chaos
 		TSet<uint32> ActivatedChildren;
 		const float ClusterDistanceThreshold2 = ClusterDistanceThreshold * ClusterDistanceThreshold;
 
-		if (!ensureMsgf(MParentToChildren.Contains(ClusterIndex), TEXT("Removing Cluster that does not exist!")))
+		if (!ensureMsgf(MParentToChildren.Contains(ClusteredParticle), TEXT("Removing Cluster that does not exist!")))
 		{
 			return ActivatedChildren;
 		}
-		TArray<uint32>& Children = *MParentToChildren[ClusterIndex];
+		TArray<uint32>& Children = MParentToChildren[ClusteredParticle];
 		bool bChildrenChanged = false;
 		const bool bRewindOnDecluster = ChildrenInheritVelocity < 1;
-		const TRigidTransform<T, d> PreSolveTM = bRewindOnDecluster ? TRigidTransform<T, d>(MParticles.X(ClusterIndex), MParticles.R(ClusterIndex)) : TRigidTransform<T, d>(MParticles.P(ClusterIndex), MParticles.Q(ClusterIndex));
+		const TRigidTransform<T, d> PreSolveTM = bRewindOnDecluster ? TRigidTransform<T, d>(ClusteredParticle->X(), ClusteredParticle->R()) : TRigidTransform<T, d>(ClusteredParticle->P(), ClusteredParticle->Q());
 
 		//@todo(ocohen): iterate with all the potential parents at once?
 		//find all children within some distance of contact point
 
 		auto RemoveChildLambda = [&](uint32 Child, int32 ChildIdx) {
 
-			MEvolution.EnableParticle(Child, ClusterIndex);
+			MEvolution.EnableParticle(MParticles.Handle(Child), ClusteredParticle);
 
 			TopLevelClusterParents.Add(Child);
 
 			//make sure to remove multi child proxy if it exists
-			if (MMultiChildProxyId[Child].Id != INDEX_NONE)
+			if (MParticles.MultiChildProxyId(Child).Id != INDEX_NONE)
 			{
-				MMultiChildProxyData[MMultiChildProxyId[Child].Id].Reset();
+				MParticles.MultiChildProxyData(MParticles.MultiChildProxyId(Child).Id).Reset();
 			}
 
-			TRigidTransform<T, d> ChildFrame = MChildToParent[Child] * PreSolveTM;
+			TRigidTransform<T, d> ChildFrame = MParticles.ChildToParent(Child) * PreSolveTM;
 			MParticles.X(Child) = ChildFrame.GetTranslation();
 			MParticles.R(Child) = ChildFrame.GetRotation();
-			MClusterIds[Child].Id = INDEX_NONE;
+			MParticles.ClusterIds(Child).Id = INDEX_NONE;
 
 			if (!bRewindOnDecluster)
 			{
@@ -422,10 +432,10 @@ namespace Chaos
 			//todo(ocohen): in the presence of collisions, this will leave all children with the post-collision
 			// velocity. This should be controlled by material properties so we can allow the broken pieces to
 			// maintain the clusters pre-collision velocity.
-			MParticles.V(Child) = MParticles.V(ClusterIndex);
-			MParticles.W(Child) = MParticles.W(ClusterIndex);
-			MParticles.PreV(Child) = MParticles.PreV(ClusterIndex);
-			MParticles.PreW(Child) = MParticles.PreW(ClusterIndex);
+			MParticles.V(Child) = ClusteredParticle->V();
+			MParticles.W(Child) = ClusteredParticle->W();
+			MParticles.PreV(Child) = ClusteredParticle->PreV();
+			MParticles.PreW(Child) = ClusteredParticle->PreW();
 
 
 			ActivatedChildren.Add(Child);
@@ -440,7 +450,7 @@ namespace Chaos
 		for (int32 ChildIdx = Children.Num() - 1; ChildIdx >= 0; --ChildIdx)
 		{
 			int32 Child = Children[ChildIdx];
-			if (StrainArray[Child] >= MStrains[Child])
+			if (StrainArray[Child] >= MParticles.Strains(Child))
 			{
 				RemoveChildLambda(Child, ChildIdx); //the piece that hits just breaks off - we may want more control by looking at the edges of this piece which would give us cleaner breaks (this approach produces more rubble)
 
@@ -499,7 +509,7 @@ namespace Chaos
 								{
 									ProcessedChildren.Add(Child);
 									ConnectedPieces.Add(Child);
-									for (const TConnectivityEdge<T>& Edge : MConnectivityEdges[Child])
+									for (const TConnectivityEdge<T>& Edge : MParticles.ConnectivityEdges(Child))
 									{
 										if (!ProcessedChildren.Contains(Edge.Sibling))
 										{
@@ -519,18 +529,18 @@ namespace Chaos
 							RemoveChildLambda(Child, INDEX_NONE); //no need to remove child, we'll just empty it at the end
 						}
 					}
-					for (const TArray<uint32>& ConnectedPieces : ConnectedPiecesArray)
+					for (TArray<uint32>& ConnectedPieces : ConnectedPiecesArray)
 					{
 						if (ConnectedPieces.Num() > 1) //now build the remaining pieces
 						{
-							int32 NewClusterIndex = CreateClusterParticleFromClusterChildren(ConnectedPieces, ClusterIndex, PreSolveTM);
-							MStrains[NewClusterIndex] = MStrains[ClusterIndex];
-							MEvolution.SetPhysicsMaterial(NewClusterIndex, MEvolution.GetPhysicsMaterial(ClusterIndex));
-
-							MParticles.V(NewClusterIndex) = MParticles.V(ClusterIndex);
-							MParticles.W(NewClusterIndex) = MParticles.W(ClusterIndex);
-							MParticles.PreV(NewClusterIndex) = MParticles.PreV(ClusterIndex);
-							MParticles.PreW(NewClusterIndex) = MParticles.PreW(ClusterIndex);
+							int32 NewClusterIndex = CreateClusterParticleFromClusterChildren(MoveTemp(ConnectedPieces), ClusteredParticle, PreSolveTM);
+							MParticles.Strains(NewClusterIndex) = ClusteredParticle->Strains();
+							MEvolution.SetPhysicsMaterial(MParticles.Handle(NewClusterIndex), MEvolution.GetPhysicsMaterial(ClusteredParticle));
+
+							MParticles.V(NewClusterIndex) = ClusteredParticle->V();
+							MParticles.W(NewClusterIndex) = ClusteredParticle->W();
+							MParticles.PreV(NewClusterIndex) = ClusteredParticle->PreV();
+							MParticles.PreW(NewClusterIndex) = ClusteredParticle->PreW();
 							MParticles.P(NewClusterIndex) = MParticles.X(NewClusterIndex);
 							MParticles.Q(NewClusterIndex) = MParticles.R(NewClusterIndex);
 
@@ -546,38 +556,38 @@ namespace Chaos
 			}
 
 			//disable cluster
-			DisableCluster(ClusterIndex);
+			DisableCluster(ClusteredParticle);
 		}
 
 		return ActivatedChildren;
 	}
 
 	template<class FPBDRigidsEvolution, class FPBDCollisionConstraint, class T, int d>
-	void TPBDRigidClustering<FPBDRigidsEvolution, FPBDCollisionConstraint, T, d>::DisableCluster(uint32 ClusterIndex)
+	void TPBDRigidClustering<FPBDRigidsEvolution, FPBDCollisionConstraint, T, d>::DisableCluster(TPBDRigidClusteredParticleHandle<T,d>* ClusteredParticle)
 	{
 		// #note: we don't recursively descend to the children
 		FClusterMap& ParentToChildren = GetChildrenMap();
 
-		MEvolution.DisableParticle(ClusterIndex);
+		MEvolution.DisableParticle(ClusteredParticle);
 
 		if (MoveClustersWhenDeactivated)
 		{
-			MParticles.P(ClusterIndex) -= FVector(0.f, 0.f, -10000.f); // HACK : Move them away to prevent reactivation. 
-			MParticles.X(ClusterIndex) -= FVector(0.f, 0.f, -10000.f); // HACK : Move them away to prevent reactivation. 
-			MParticles.V(ClusterIndex) = FVector(0.f);            // HACK : Move them away to prevent reactivation.
+			ClusteredParticle->P() -= FVector(0.f, 0.f, -10000.f); // HACK : Move them away to prevent reactivation. 
+			ClusteredParticle->X() -= FVector(0.f, 0.f, -10000.f); // HACK : Move them away to prevent reactivation. 
+			ClusteredParticle->V() = FVector(0.f);            // HACK : Move them away to prevent reactivation.
 		}
 
-		TopLevelClusterParents.Remove(ClusterIndex);
-		ParentToChildren.Remove(ClusterIndex);
-		MClusterIds[ClusterIndex] = ClusterId();
-		MClusterGroupIndex[ClusterIndex] = 0;
-		MActiveRemovalIndices.Remove(ClusterIndex);
+		TopLevelClusterParents.Remove(ClusteredParticle->TransientParticleIndex());
+		ParentToChildren.Remove(ClusteredParticle);
+		ClusteredParticle->ClusterIds() = ClusterId();
+		ClusteredParticle->ClusterGroupIndex() = 0;
+		MActiveRemovalIndices.Remove(ClusteredParticle->TransientParticleIndex());
 	}
 
 	template<class FPBDRigidsEvolution, class FPBDCollisionConstraint, class T, int d>
 	void TPBDRigidClustering<FPBDRigidsEvolution, FPBDCollisionConstraint, T, d>::DisableParticleWithBreakEvent(uint32 ClusterIndex)
 	{
-		DisableCluster(ClusterIndex);
+		DisableCluster(MParticles.Handle(ClusterIndex));
 
 		if (DoGenerateBreakingData)
 		{
@@ -599,9 +609,9 @@ namespace Chaos
 		SCOPE_CYCLE_COUNTER(STAT_ResetCollisionImpulseArray);
 		if (MCollisionImpulseArrayDirty)
 		{
-			for (int i = 0; i < MCollisionImpulses.Num(); i++)
+			for (uint32 i = 0; i < MParticles.Size(); i++)
 			{
-				MCollisionImpulses[i] = (T)0.0;
+				MParticles.CollisionImpulses(i) = (T)0.0;
 			}
 			MCollisionImpulseArrayDirty = false;
 		}
@@ -626,49 +636,49 @@ namespace Chaos
 				continue;
 			}
 
-			auto ComputeStrainLambda = [&](uint32 ClusterIndex, const TArray<uint32>& ParentToChildren) {
-				const TRigidTransform<T, d> WorldToClusterTM = TRigidTransform<T, d>(MParticles.P(ClusterIndex), MParticles.Q(ClusterIndex));
+			auto ComputeStrainLambda = [&](TPBDRigidClusteredParticleHandle<T, d>* Cluster, const TArray<uint32>& ParentToChildren) {
+				const TRigidTransform<T, d> WorldToClusterTM = TRigidTransform<T, d>(Cluster->P(), Cluster->Q());
 				const TVector<T, d> ContactLocationClusterLocal = WorldToClusterTM.InverseTransformPosition(GetContactLocation(Contact));
 				TBox<T, d> ContactBox(ContactLocationClusterLocal, ContactLocationClusterLocal);
 				ContactBox.Thicken(ClusterDistanceThreshold);
-				if (MChildrenSpatial[ClusterIndex])
+				if (Cluster->ChildrenSpatial())
 				{
-					const TArray<int32> Intersections = MChildrenSpatial[ClusterIndex]->FindAllIntersectingChildren(ContactBox);
+					const TArray<int32> Intersections = Cluster->ChildrenSpatial()->FindAllIntersectingChildren(ContactBox);
 					for (int32 Child : Intersections)
 					{
-						if (const TUniquePtr<TMultiChildProxyData<T, d>>& ProxyData = MMultiChildProxyData[Child])
+						if (const TUniquePtr<TMultiChildProxyData<T, d>>& ProxyData = MParticles.MultiChildProxyData(Child))
 						{
 							//multi child so get its children
 							const int32 KeyChild = ProxyData->KeyChild;
-							const TRigidTransform<T, d> ProxyToCluster = ProxyData->RelativeToKeyChild * MChildToParent[KeyChild];
+							const TRigidTransform<T, d> ProxyToCluster = ProxyData->RelativeToKeyChild * MParticles.ChildToParent(KeyChild);
 							const TVector<T, d> ContactLocationProxyLocal = ProxyToCluster.InverseTransformPosition(ContactLocationClusterLocal);
 							TBox<T, d> ContactBoxProxy(ContactLocationProxyLocal, ContactLocationProxyLocal);
 							ContactBoxProxy.Thicken(ClusterDistanceThreshold);
-							if (MChildrenSpatial[Child])
+							if (MParticles.ChildrenSpatial(Child))
 							{
-								const TArray<int32> SubIntersections = MChildrenSpatial[Child]->FindAllIntersectingChildren(ContactBoxProxy);
+								const TArray<int32> SubIntersections = MParticles.ChildrenSpatial(Child)->FindAllIntersectingChildren(ContactBoxProxy);
 								for (int32 SubChild : SubIntersections)
 								{
-									MCollisionImpulses[SubChild] += Contact.AccumulatedImpulse.Size();
+									MParticles.CollisionImpulses(SubChild) += Contact.AccumulatedImpulse.Size();
 								}
 							}
 						}
 						else
 						{
-							MCollisionImpulses[Child] += Contact.AccumulatedImpulse.Size();
+							MParticles.CollisionImpulses(Child) += Contact.AccumulatedImpulse.Size();
 						}
 					}
 				}
 			};
 
-			if (TUniquePtr<TArray<uint32>>* ChildrenPtrPtr = MParentToChildren.Find(Contact.ParticleIndex))
+			if (auto ChildrenPtr = MParentToChildren.Find(Contact.Particle))
 			{
-				ComputeStrainLambda(Contact.ParticleIndex, **ChildrenPtrPtr);
+				ComputeStrainLambda(Contact.Particle->AsClustered(), *ChildrenPtr);
 			}
 
-			if (TUniquePtr<TArray<uint32>>* ChildrenPtrPtr = MParentToChildren.Find(Contact.LevelsetIndex))
+			if (auto ChildrenPtr = MParentToChildren.Find(Contact.Levelset))
 			{
-				ComputeStrainLambda(Contact.LevelsetIndex, **ChildrenPtrPtr);
+				ComputeStrainLambda(Contact.Levelset->AsClustered(), *ChildrenPtr);
 			}
 
 			MCollisionImpulseArrayDirty = true;
@@ -676,7 +686,7 @@ namespace Chaos
 	}
 
 	template<class T, int d>
-	T CalculatePseudoMomentum(const TPBDRigidParticles<T, d>& InParticles, const uint32 Index)
+	T CalculatePseudoMomentum(const TPBDRigidClusteredParticles<T, d>& InParticles, const uint32 Index)
 	{
 		TVector<T, d> LinearPseudoMomentum = (InParticles.X(Index) - InParticles.P(Index)) * InParticles.M(Index);
 		TRotation<T, d> Delta = InParticles.R(Index) * InParticles.Q(Index).Inverse();
@@ -689,45 +699,48 @@ namespace Chaos
 
 	DECLARE_CYCLE_STAT(TEXT("TPBDRigidClustering<>::RewindAndEvolve<BGF>()"), STAT_RewindAndEvolve_BGF, STATGROUP_Chaos);
 	template<class T, int d>
-	void RewindAndEvolve(TPBDRigidsEvolutionGBF<T, d>& Evolution, TPBDRigidParticles<T, d>& InParticles, const TSet<int32>& IslandsToRecollide, const TSet<uint32>& AllActivatedChildren, const T Dt, TPBDCollisionConstraint<T, d>& CollisionRule)
+	void RewindAndEvolve(TPBDRigidsEvolutionGBF<T, d>& Evolution, TPBDRigidClusteredParticles<T, d>& InParticles, const TSet<int32>& IslandsToRecollide, const TSet<uint32>& AllActivatedChildren, const T Dt, TPBDCollisionConstraint<T, d>& CollisionRule)
 	{
 		SCOPE_CYCLE_COUNTER(STAT_RewindAndEvolve_BGF);
 		// Rewind active particles
 		const TArray<int32> IslandsToRecollideArray = IslandsToRecollide.Array();
 		PhysicsParallelFor(IslandsToRecollideArray.Num(), [&](int32 Idx) {
 			int32 Island = IslandsToRecollideArray[Idx];
-			TArray<int32> ParticleIndices = Evolution.GetIslandParticles(Island); // copy
-			for (int32 ArrayIdx = ParticleIndices.Num() - 1; ArrayIdx >= 0; --ArrayIdx)
+			auto Particles = Evolution.GetIslandParticles(Island); // copy
+			for (int32 ArrayIdx = Particles.Num() - 1; ArrayIdx >= 0; --ArrayIdx)
 			{
-				int32 Index = ParticleIndices[ArrayIdx];
-				if (InParticles.Sleeping(Index) || InParticles.Disabled(Index))
+				if (auto PBDRigid = Particles[ArrayIdx]->AsDynamic())
 				{
-					ParticleIndices.RemoveAtSwap(ArrayIdx);
-				}
-				else
-				{
-					InParticles.P(Index) = InParticles.X(Index);
-					InParticles.Q(Index) = InParticles.R(Index);
-					InParticles.V(Index) = InParticles.PreV(Index);
-					InParticles.W(Index) = InParticles.PreW(Index);
+					if (!PBDRigid->Sleeping() && !PBDRigid->Disabled())
+					{
+						PBDRigid->P() = PBDRigid->X();
+						PBDRigid->Q() = PBDRigid->R();
+						PBDRigid->V() = PBDRigid->PreV();
+						PBDRigid->W() = PBDRigid->PreW();
+						continue;
+					}
 				}
+				Particles.RemoveAtSwap(ArrayIdx);
 			}
-			Evolution.Integrate(ParticleIndices, Dt);
+			Evolution.Integrate(MakeHandleView(Particles), Dt);
 		});
 
-		TSet<uint32> AllIslandParticles;
+		TSet<TGeometryParticleHandle<T, d>*> AllIslandParticles;
 		for (int32 Island = 0; Island < Evolution.NumIslands(); ++Island)
 		{
-			const TArray<int32>& ParticleIndices = Evolution.GetIslandParticles(Island);
-			for (const int32 Index : ParticleIndices)
+			const auto& ParticleIndices = Evolution.GetIslandParticles(Island);
+			for (const auto Particle : ParticleIndices)
 			{
-				bool bDisabled = InParticles.Disabled(Index);
-
-				// #TODO - Have to repeat checking out whether the particle is disabled matching the PFor above.
-				// Move these into shared array so we only process it once
-				if (!AllIslandParticles.Contains(Index) && !bDisabled)
+				if (auto PBDRigid = Particle->AsDynamic())
 				{
-					AllIslandParticles.Add(Index);
+					bool bDisabled = PBDRigid->Disabled();
+
+					// #TODO - Have to repeat checking out whether the particle is disabled matching the PFor above.
+					// Move these into shared array so we only process it once
+					if (!AllIslandParticles.Contains(Particle) && !bDisabled)
+					{
+						AllIslandParticles.Add(Particle);
+					}
 				}
 			}
 		}
@@ -737,7 +750,11 @@ namespace Chaos
 		{
 			// @todo(mlentine): We can precompute internal constraints which can filter some from the narrow phase tests but may not help much
 
+#if CHAOS_PARTICLEHANDLE_TODO
 			CollisionRule.UpdateConstraints(InParticles, Evolution.GetNonDisabledIndices(), Dt, AllActivatedChildren, AllIslandParticles.Array());
+#else
+			CollisionRule.UpdateConstraints(Dt, AllIslandParticles);	//this seems wrong
+#endif
 
 			Evolution.InitializeAccelerationStructures();
 
@@ -754,9 +771,10 @@ namespace Chaos
 		}
 	}
 
+#if CHAOS_PARTICLEHANDLE_TODO
 	DECLARE_CYCLE_STAT(TEXT("TPBDRigidClustering<>::RewindAndEvolve<PSG>()"), STAT_RewindAndEvolve_PSG, STATGROUP_Chaos);
 	template<class T, int d>
-	void RewindAndEvolve(TPBDRigidsEvolutionPGS<T, d>& Evolution, TPBDRigidParticles<T, d>& InParticles, const TSet<int32>& IslandsToRecollide, const TSet<uint32>& AllActivatedChildren, const T Dt, TPBDCollisionConstraintPGS<T, d>& CollisionRule)
+	void RewindAndEvolve(TPBDRigidsEvolutionPGS<T, d>& Evolution, TPBDRigidClusteredParticles<T, d>& InParticles, const TSet<int32>& IslandsToRecollide, const TSet<uint32>& AllActivatedChildren, const T Dt, TPBDCollisionConstraintPGS<T, d>& CollisionRule)
 	{
 		SCOPE_CYCLE_COUNTER(STAT_RewindAndEvolve_PSG);
 		// Rewind active particles
@@ -821,6 +839,7 @@ namespace Chaos
 			Evolution.ApplyPushOut(Dt, Island);
 		});
 	}
+#endif
 
 	DECLARE_CYCLE_STAT(TEXT("TPBDRigidClustering<>::AdvanceClustering"), STAT_AdvanceClustering, STATGROUP_Chaos);
 	DECLARE_CYCLE_STAT(TEXT("TPBDRigidClustering<>::Update Impulse from Strain"), STAT_UpdateImpulseStrain, STATGROUP_Chaos);
@@ -840,17 +859,17 @@ namespace Chaos
 
 		{
 			const float Threshold = 1.0f;
-			TSet<int32> RemovalIndicesCopy = MActiveRemovalIndices; // copy since RemoveOnFracture alters MActiveRemovalIndices
-			for (int32 Idx : RemovalIndicesCopy)
+			auto RemovalIndicesCopy = MActiveRemovalIndices; // copy since RemoveOnFracture alters MActiveRemovalIndices
+			for (uint32 ParticleIdx : RemovalIndicesCopy)
 			{
-				if (MParticles.ToBeRemovedOnFracture(Idx) && MParticles.V(Idx).SizeSquared() > Threshold && MParticles.PreV(Idx).SizeSquared() > Threshold)
+				if (MParticles.ToBeRemovedOnFracture(ParticleIdx) && MParticles.V(ParticleIdx).SizeSquared() > Threshold && MParticles.PreV(ParticleIdx).SizeSquared() > Threshold)
 				{
-					DisableParticleWithBreakEvent(Idx);
+					DisableParticleWithBreakEvent(ParticleIdx);
 				}
 			}
 		}
 
-		if (MCollisionImpulses.Num())
+		if (MParticles.Size())
 		{
 			//
 			//  Grab collision impulses for processing
@@ -870,17 +889,17 @@ namespace Chaos
 			//
 			{
 				SCOPE_CYCLE_COUNTER(STAT_UpdateDirtyImpulses);
-				const TSet<int32>& ActiveIndices = MEvolution.GetActiveIndices();
-				for (const auto& ActiveIndex : ActiveIndices)
+				const auto& ActiveClusteredArray = MEvolution.GetActiveClusteredArray();
+				for (const auto& ActiveCluster : ActiveClusteredArray)
 				{
-					if (MClusterIds[ActiveIndex].NumChildren > 0) //active index is a cluster
+					if (ActiveCluster->ClusterIds().NumChildren > 0) //active index is a cluster
 					{
-						const TArray<uint32>& ParentToChildren = *MParentToChildren[ActiveIndex];
+						const TArray<uint32>& ParentToChildren = MParentToChildren[ActiveCluster];
 						for (const uint32 Child : ParentToChildren)
 						{
-							if (MStrains[Child] <= 0.f)
+							if (MParticles.Strains(Child) <= 0.f)
 							{
-								MCollisionImpulses[Child] = FLT_MAX;
+								MParticles.CollisionImpulses(Child) = FLT_MAX;
 								MCollisionImpulseArrayDirty = true;
 							}
 						}
@@ -891,8 +910,8 @@ namespace Chaos
 			if (MCollisionImpulseArrayDirty)
 			{
 				SCOPE_CYCLE_COUNTER(STAT_UpdateDirtyImpulses);
-				TArrayView<T> CollisionImpulsesView(&(MCollisionImpulses[0]), MCollisionImpulses.Num());
-				TMap<uint32, TSet<uint32>> ClusterToActivatedChildren = BreakingModel(CollisionImpulsesView);
+				TArrayView<T> CollisionImpulsesView(&(MParticles.CollisionImpulses(0)), MParticles.Size());
+				TMap<TPBDRigidClusteredParticleHandle<T, d>*, TSet<uint32>> ClusterToActivatedChildren = BreakingModel(CollisionImpulsesView);
 
 				TSet<uint32> AllActivatedChildren;
 				TSet<int32> IslandsToRecollide;
@@ -928,13 +947,13 @@ namespace Chaos
 					{
 						for (auto Itr : ClusterToActivatedChildren)
 						{
-							const uint32 ClusterId = Itr.Key;
+							auto ClusteredParticle = Itr.Key;
 
 							const TSet<uint32>& ActivatedChildren = Itr.Value;
 							for (uint32 ActiveChild : ActivatedChildren)
 							{
-								MParticles.V(ActiveChild) = MParticles.V(ActiveChild) * (1.f - ChildrenInheritVelocity) + MParticles.V(ClusterId) * ChildrenInheritVelocity;
-								MParticles.W(ActiveChild) = MParticles.W(ActiveChild) * (1.f - ChildrenInheritVelocity) + MParticles.W(ClusterId) * ChildrenInheritVelocity;
+								MParticles.V(ActiveChild) = MParticles.V(ActiveChild) * (1.f - ChildrenInheritVelocity) + ClusteredParticle->V() * ChildrenInheritVelocity;
+								MParticles.W(ActiveChild) = MParticles.W(ActiveChild) * (1.f - ChildrenInheritVelocity) + ClusteredParticle->W() * ChildrenInheritVelocity;
 							}
 						}
 					}
@@ -954,9 +973,9 @@ namespace Chaos
 		FClusterMap& MParentToChildren = GetChildrenMap();
 
 		T Result = 0;
-		if (MParentToChildren.Contains(CurrentNode) && MParentToChildren[CurrentNode])
+		if (MParentToChildren.Contains(MParticles.Handle(CurrentNode)))
 		{
-			for (uint32 Child : *MParentToChildren[CurrentNode])
+			for (uint32 Child : MParentToChildren[MParticles.Handle(CurrentNode)])
 			{
 				Result += PromoteStrains(Child, ExternalStrains);
 			}
@@ -971,39 +990,42 @@ namespace Chaos
 
 	DECLARE_CYCLE_STAT(TEXT("TPBDRigidClustering<>::BreakingModel()"), STAT_BreakingModel, STATGROUP_Chaos);
 	template<class FPBDRigidsEvolution, class FPBDCollisionConstraint, class T, int d>
-	TMap<uint32, TSet<uint32>> TPBDRigidClustering<FPBDRigidsEvolution, FPBDCollisionConstraint, T, d>::BreakingModel(TArrayView<T> ExternalStrain)
+	TMap<TPBDRigidClusteredParticleHandle<T, d>*, TSet<uint32>> TPBDRigidClustering<FPBDRigidsEvolution, FPBDCollisionConstraint, T, d>::BreakingModel(TArrayView<T> ExternalStrain)
 	{
 		SCOPE_CYCLE_COUNTER(STAT_BreakingModel);
 		FClusterMap& MParentToChildren = GetChildrenMap();
 
-		TMap<uint32, TSet<uint32>> AllActivatedChildren;
+		TMap<TPBDRigidClusteredParticleHandle<T,d>*, TSet<uint32>> AllActivatedChildren;
 
-		TArray<int32> NonDisabledIndices = MEvolution.GetNonDisabledIndices(); //make copy because release cluster modifies active indices. We want to iterate over original active indices
-		for (int32 ActiveIdx : NonDisabledIndices)
+		auto NonDisabledClusteredParticles = MEvolution.GetNonDisabledClusteredArray(); //make copy because release cluster modifies active indices. We want to iterate over original active indices
+		for (auto ClusteredParticle : NonDisabledClusteredParticles)
 		{
-			if (MClusterIds[ActiveIdx].NumChildren)
+			if (ClusteredParticle->ClusterIds().NumChildren)
 			{
-				AllActivatedChildren.Add(ActiveIdx, ReleaseClusterParticles(ActiveIdx, ExternalStrain));
+				AllActivatedChildren.Add(ClusteredParticle, ReleaseClusterParticles(ClusteredParticle, ExternalStrain));
 			}
 			else
 			{
 				// there's no children to break but we need to process whether this single piece is to be removed when damaged
-				if (MParticles.ToBeRemovedOnFracture(ActiveIdx))
+				if (ClusteredParticle->ToBeRemovedOnFracture())
 				{
-					if (ExternalStrain[ActiveIdx] >= MStrains[ActiveIdx])
+					if (ExternalStrain[ClusteredParticle->TransientParticleIndex()] >= ClusteredParticle->Strains())
 					{
-						DisableCluster(ActiveIdx);
+						DisableCluster(ClusteredParticle);
 
 						if (DoGenerateBreakingData)
 						{
 							int32 NewIdx = MAllClusterBreakings.Add(TBreakingData<float, 3>());
 
 							TBreakingData<float, 3>& ClusterBreakingItem = MAllClusterBreakings[NewIdx];
-							ClusterBreakingItem.ParticleIndex = ActiveIdx;
-							ClusterBreakingItem.Location = MParticles.X(ActiveIdx);
-							ClusterBreakingItem.Velocity = MParticles.V(ActiveIdx);
-							ClusterBreakingItem.AngularVelocity = MParticles.W(ActiveIdx);
-							ClusterBreakingItem.Mass = MParticles.M(ActiveIdx);
+							ClusterBreakingItem.ParticleIndex = INDEX_NONE;// ClusteredParticle;
+#if CHAOS_PARTICLEHANDLE_TODO
+							ClusterBreakingItem.Particle = ClusteredParticle;
+#endif
+							ClusterBreakingItem.Location = ClusteredParticle->X();
+							ClusterBreakingItem.Velocity = ClusteredParticle->V();
+							ClusterBreakingItem.AngularVelocity = ClusteredParticle->W(); 
+							ClusterBreakingItem.Mass = ClusteredParticle->M();
 						}
 					}
 				}
@@ -1022,10 +1044,10 @@ namespace Chaos
 
 		EObjectStateType ObjectState = EObjectStateType::Dynamic;
 		check(0 <= ClusterIndex && ClusterIndex < MParticles.Size());
-		if (MParentToChildren.Contains(ClusterIndex) && MParentToChildren[ClusterIndex] != nullptr && MParentToChildren[ClusterIndex]->Num())
+		if (MParentToChildren.Contains(MParticles.Handle(ClusterIndex)) && MParentToChildren[MParticles.Handle(ClusterIndex)].Num())
 		{
 			TQueue<int32> Queue;
-			for (const int32 Child : *MParentToChildren[ClusterIndex])
+			for (const int32 Child : MParentToChildren[MParticles.Handle(ClusterIndex)])
 			{
 				Queue.Enqueue(Child);
 			}
@@ -1035,9 +1057,9 @@ namespace Chaos
 			{
 				// @question : Maybe we should just store the leaf node bodies in a
 				// map, that will require Memory(n*log(n))
-				if (MParentToChildren.Contains(CurrentIndex))
+				if (MParentToChildren.Contains(MParticles.Handle(CurrentIndex)))
 				{
-					for (const auto Child : *MParentToChildren[CurrentIndex])
+					for (const auto Child : MParentToChildren[MParticles.Handle(CurrentIndex)])
 					{
 						Queue.Enqueue(Child);
 					}
@@ -1063,7 +1085,7 @@ namespace Chaos
 
 	DECLARE_CYCLE_STAT(TEXT("TPBDRigidClustering<>::UpdateClusterMassProperties()"), STAT_UpdateClusterMassProperties, STATGROUP_Chaos);
 	template<class T, int d>
-	void UpdateClusterMassProperties(TPBDRigidParticles<T, d>& Particles, const TArray<uint32>& Children, const uint32 NewIndex, const TRigidTransform<T, d>* ForceMassOrientation, const TArrayCollectionArray<TUniquePtr<TMultiChildProxyData<T, d>>>* MMultiChildProxyData, const TArrayCollectionArray<FMultiChildProxyId>* MMultiChildProxyId)
+	void UpdateClusterMassProperties(TPBDRigidClusteredParticles<T, d>& Particles, const TArray<uint32>& Children, const uint32 NewIndex, const TRigidTransform<T, d>* ForceMassOrientation, const TArrayCollectionArray<TUniquePtr<TMultiChildProxyData<T, d>>>* MMultiChildProxyData, const TArrayCollectionArray<FMultiChildProxyId>* MMultiChildProxyId)
 	{
 		SCOPE_CYCLE_COUNTER(STAT_UpdateClusterMassProperties);
 		check(0 <= NewIndex && NewIndex < Particles.Size());
@@ -1076,7 +1098,7 @@ namespace Chaos
 		Particles.M(NewIndex) = 0;
 		Particles.I(NewIndex) = PMatrix<T, d, d>(0);
 
-		auto GetMultiChildProxy = [MMultiChildProxyId, MMultiChildProxyData](const int32 OriginalChild, const TMultiChildProxyData<T,d>*& ProxyData) -> int32
+		auto GetMultiChildProxy = [MMultiChildProxyId, MMultiChildProxyData](const int32 OriginalChild, const TMultiChildProxyData<T, d>*& ProxyData) -> int32
 		{
 			const int32 MultiChildProxyId = MassPropertiesFromMultiChildProxy && MMultiChildProxyId ? (*MMultiChildProxyId)[OriginalChild].Id : INDEX_NONE;
 			if (MultiChildProxyId != INDEX_NONE)
@@ -1170,7 +1192,7 @@ namespace Chaos
 		{
 			Particles.M(NewIndex) = 1;
 			Particles.X(NewIndex) = TVector<T, d>(0);
-			Particles.V(NewIndex) =  TVector<T, d>(0);
+			Particles.V(NewIndex) = TVector<T, d>(0);
 			Particles.PreV(NewIndex) = Particles.V(NewIndex);
 			Particles.InvM(NewIndex) = 1;
 			Particles.P(NewIndex) = Particles.X(NewIndex);
@@ -1178,8 +1200,8 @@ namespace Chaos
 			Particles.PreW(NewIndex) = Particles.W(NewIndex);
 			Particles.R(NewIndex) = TRotation<T, d>(FMatrix::Identity);
 			Particles.Q(NewIndex) = Particles.R(NewIndex);
-			Particles.I(NewIndex) = FMatrix::Identity;	
-			Particles.InvI(NewIndex) = FMatrix::Identity;	
+			Particles.I(NewIndex) = FMatrix::Identity;
+			Particles.InvI(NewIndex) = FMatrix::Identity;
 			return;
 		}
 
@@ -1261,7 +1283,7 @@ namespace Chaos
 
 	template<class FPBDRigidsEvolution, class FPBDCollisionConstraint, class T, int d>
 	void TPBDRigidClustering<FPBDRigidsEvolution, FPBDCollisionConstraint, T, d>::UpdateGeometry(const TArray<uint32>& Children,
-	    const uint32 NewIndex, TSerializablePtr<TImplicitObject<T, d>> ProxyGeometry, const FClusterCreationParameters<T>& Parameters)
+		const uint32 NewIndex, TSerializablePtr<TImplicitObject<T, d>> ProxyGeometry, const FClusterCreationParameters<T>& Parameters)
 	{
 		SCOPE_CYCLE_COUNTER(STAT_UpdateGeometry);
 		const bool bUseCollisionPoints = (ProxyGeometry || Parameters.bCopyCollisionParticles) && !Parameters.CollisionParticles;
@@ -1287,17 +1309,17 @@ namespace Chaos
 				int32 UsedGeomChild = Child;
 				if (MParticles.Geometry(Child))
 				{
-					const int32 MultiChildProxyId = MMultiChildProxyId[Child].Id;
-					if (UseLevelsetCollision || MultiChildProxyId == INDEX_NONE || !MMultiChildProxyData[MultiChildProxyId])
+					const int32 MultiChildProxyId = MParticles.MultiChildProxyId(Child).Id;
+					if (UseLevelsetCollision || MultiChildProxyId == INDEX_NONE || !MParticles.MultiChildProxyData(MultiChildProxyId))
 					{
 						Objects.Add(TUniquePtr<TImplicitObject<T, d>>(new TImplicitObjectTransformed<T, d>(MParticles.Geometry(Child), Frame)));
 						Objects2.Add(TUniquePtr<TImplicitObject<T, d>>(new TImplicitObjectTransformed<T, d>(MParticles.Geometry(Child), Frame)));
 						GeomToOriginalParticlesHack.Add(Child);
 					}
-					else if (MMultiChildProxyData[MultiChildProxyId]->KeyChild == Child)
+					else if (MParticles.MultiChildProxyData(MultiChildProxyId)->KeyChild == Child)
 					{
 						//using multi child proxy and this child is the key
-						const TRigidTransform<T, d> ProxyWorldTM = MMultiChildProxyData[MultiChildProxyId]->RelativeToKeyChild * ChildWorldTM;
+						const TRigidTransform<T, d> ProxyWorldTM = MParticles.MultiChildProxyData(MultiChildProxyId)->RelativeToKeyChild * ChildWorldTM;
 						const TRigidTransform<T, d> ProxyRelativeTM = ProxyWorldTM.GetRelativeTransform(ClusterWorldTM);
 						Objects.Add(TUniquePtr<TImplicitObject<T, d>>(new TImplicitObjectTransformed<T, d>(MParticles.Geometry(MultiChildProxyId), ProxyRelativeTM)));
 						Objects2.Add(TUniquePtr<TImplicitObject<T, d>>(new TImplicitObjectTransformed<T, d>(MParticles.Geometry(MultiChildProxyId), ProxyRelativeTM)));
@@ -1308,10 +1330,9 @@ namespace Chaos
 				}
 
 				ensure(MParticles.Disabled(Child) == true);
-				ensure(!MEvolution.GetActiveIndices().Contains(Child));
 
-				check(MClusterIds[Child].Id == NewIndex);
-				MChildToParent[Child] = Frame;
+				check(MParticles.ClusterIds(Child).Id == NewIndex);
+				MParticles.ChildToParent(Child) = Frame;
 
 				if (bUseCollisionPoints)
 				{
@@ -1334,7 +1355,7 @@ namespace Chaos
 
 		{
 			QUICK_SCOPE_CYCLE_COUNTER(SpatialBVH);
-			MChildrenSpatial[NewIndex] = Objects2.Num() ? MakeUnique<TImplicitObjectUnion<T, d>>(MoveTemp(Objects2), GeomToOriginalParticlesHack) : nullptr;
+			MParticles.ChildrenSpatial(NewIndex) = Objects2.Num() ? MakeUnique<TImplicitObjectUnion<T, d>>(MoveTemp(Objects2), GeomToOriginalParticlesHack) : nullptr;
 		}
 
 		TArray<TVector<T, d>> CleanedPoints;
@@ -1460,7 +1481,7 @@ namespace Chaos
 	void TPBDRigidClustering<FPBDRigidsEvolution, FPBDCollisionConstraint, T, d>::GenerateConnectionGraph(int32 ClusterIndex, const FClusterCreationParameters<T>& Parameters)
 	{
 		SCOPE_CYCLE_COUNTER(STAT_GenerateConnectionGraph);
-		if (!GetChildrenMap().Contains(ClusterIndex)) return;
+		if (!GetChildrenMap().Contains(MParticles.Handle(ClusterIndex))) return;
 
 		// Connectivity Graph
 		//    Build a connectivity graph for the cluster. If the PointImplicit is specified
@@ -1471,13 +1492,13 @@ namespace Chaos
 		{
 			typename FClusterCreationParameters<T>::EConnectionMethod LocalConnectionMethod = Parameters.ConnectionMethod;
 
-			if (LocalConnectionMethod == FClusterCreationParameters<T>::EConnectionMethod::None || 
+			if (LocalConnectionMethod == FClusterCreationParameters<T>::EConnectionMethod::None ||
 				(LocalConnectionMethod == FClusterCreationParameters<T>::EConnectionMethod::PointImplicit && !MParticles.CollisionParticles(ClusterIndex)))
 			{
 				LocalConnectionMethod = FClusterCreationParameters<T>::EConnectionMethod::MinimalSpanningSubsetDelaunayTriangulation; // default method
 			}
 
-			if (LocalConnectionMethod == FClusterCreationParameters<T>::EConnectionMethod::PointImplicit || 
+			if (LocalConnectionMethod == FClusterCreationParameters<T>::EConnectionMethod::PointImplicit ||
 				LocalConnectionMethod == FClusterCreationParameters<T>::EConnectionMethod::PointImplicitAugmentedWithMinimalDelaunay)
 			{
 				UpdateConnectivityGraphUsingPointImplicit(ClusterIndex, Parameters);
@@ -1503,7 +1524,7 @@ namespace Chaos
 		// @todo(investigate) : This is trying to set multiple connections and throwing a warning in ConnectNodes
 		SCOPE_CYCLE_COUNTER(STAT_FixConnectivityGraphUsingDelaunayTriangulation);
 
-		const TArray<uint32>& Children = *GetChildrenMap()[ClusterIndex];
+		const TArray<uint32>& Children = GetChildrenMap()[MParticles.Handle(ClusterIndex)];
 
 		// Compute Delaunay neighbor graph on children centers
 		TArray<FVector> Pts;
@@ -1561,7 +1582,7 @@ namespace Chaos
 		};
 		for (int32 i = 0; i < Children.Num(); ++i)
 		{
-			for (const TConnectivityEdge<T>& Edge : MConnectivityEdges[Children[i]])
+			for (const TConnectivityEdge<T>& Edge : MParticles.ConnectivityEdges(Children[i]))
 			{
 				int* nbr = ChildReverseIdx.Find(Edge.Sibling);
 				if (nbr)
@@ -1587,24 +1608,24 @@ namespace Chaos
 				{ // assume we'll get the symmetric connection; don't bother considering this one
 					continue;
 				}
-				
+
 				float DistSq = FVector::DistSquared(Pts[i], Pts[Nbr]);
 				if (DistSq < AlwaysAcceptBelowDistSqThreshold)
 				{ // below always-accept threshold: don't bother adding to candidates array, just merge now
 					MergeGroup(i, Nbr);
 					int32 ChildA = Children[i], ChildB = Children[Nbr];
-					const T AvgStrain = (MStrains[ChildA] + MStrains[ChildB]) * (T)0.5f;
+					const T AvgStrain = (MParticles.Strains(ChildA) + MParticles.Strains(ChildB)) * (T)0.5f;
 					ConnectNodes(ChildA, ChildB, AvgStrain);
 					continue;
 				}
 
-				if (FindGroup(i) == FindGroup(Nbr)) 
+				if (FindGroup(i) == FindGroup(Nbr))
 				{ // already part of the same group so we don't need Delaunay edge  
 					continue;
 				}
 
 				// add to array to sort and add as-needed
-				Candidates.Add({i, Nbr, DistSq});
+				Candidates.Add({ i, Nbr, DistSq });
 			}
 		}
 
@@ -1617,7 +1638,7 @@ namespace Chaos
 			{
 				MergeGroup(A, B);
 				int32 ChildA = Children[A], ChildB = Children[B];
-				const T AvgStrain = (MStrains[ChildA] + MStrains[ChildB]) * (T)0.5f;
+				const T AvgStrain = (MParticles.Strains(ChildA) + MParticles.Strains(ChildB)) * (T)0.5f;
 				ConnectNodes(ChildA, ChildB, AvgStrain);
 			}
 		}
@@ -1629,7 +1650,7 @@ namespace Chaos
 	{
 		SCOPE_CYCLE_COUNTER(STAT_UpdateConnectivityGraphUsingDelaunayTriangulation);
 
-		const TArray<uint32>& Children = *GetChildrenMap()[ClusterIndex];
+		const TArray<uint32>& Children = GetChildrenMap()[MParticles.Handle(ClusterIndex)];
 
 		TArray<FVector> Pts;
 		Pts.AddUninitialized(Children.Num());
@@ -1653,7 +1674,7 @@ namespace Chaos
 				if (!UniqueEdges.Contains(SortedPair))
 				{
 					// this does not use ConnectNodes because Neighbors is bi-direction : as in (1,2),(2,1)
-					const T AvgStrain = (MStrains[First] + MStrains[Sibling]) * (T)0.5f;
+					const T AvgStrain = (MParticles.Strains(First) + MParticles.Strains(Sibling)) * (T)0.5f;
 					//AddUniqueConnection(Children[i], Children[Neighbors[i][j]], AvgStrain);
 					ConnectNodes(First, Sibling, AvgStrain);
 					UniqueEdges.Add(SortedPair);
@@ -1672,7 +1693,7 @@ namespace Chaos
 		if (UseConnectivity)
 		{
 			T Delta = FMath::Min(FMath::Max(Parameters.CoillisionThicknessPercent, (T)0), T(1));
-			const TArray<uint32>& Children = *MParentToChildren[ClusterIndex];
+			const TArray<uint32>& Children = MParentToChildren[MParticles.Handle(ClusterIndex)];
 			for (int32 i = 0; i < Children.Num(); ++i)
 			{
 				const uint32 Child1 = Children[i];
@@ -1721,9 +1742,9 @@ namespace Chaos
 				{
 					for (const auto& Elem : ConnectionList)
 					{
-						if (MConnectivityEdges[Elem.Key].Num() < Parameters.MaxNumConnections)
+						if (MParticles.ConnectivityEdges(Elem.Key).Num() < Parameters.MaxNumConnections)
 						{
-							const T AvgStrain = (MStrains[Elem.Key] + MStrains[Elem.Value]) * (T)0.5f;
+							const T AvgStrain = (MParticles.Strains(Elem.Key) + MParticles.Strains(Elem.Value)) * (T)0.5f;
 							ConnectNodes(Elem.Key, Elem.Value, AvgStrain);
 						}
 					}
@@ -1760,7 +1781,7 @@ namespace Chaos
 	void TPBDRigidClustering<FPBDRigidsEvolution, FPBDCollisionConstraint, T, d>::UpdateMassProperties(const TArray<uint32>& Children, const uint32 ClusterIndex, const TRigidTransform<T, d>* ForceMassOrientation)
 	{
 		SCOPE_CYCLE_COUNTER(STAT_UpdateMassProperties);
-		UpdateClusterMassProperties(MParticles, Children, ClusterIndex, ForceMassOrientation, &MMultiChildProxyData, &MMultiChildProxyId);
+		UpdateClusterMassProperties(MParticles, Children, ClusterIndex, ForceMassOrientation, &MParticles.MultiChildProxyDataArray(), &MParticles.MultiChildProxyIdArray());
 		UpdateKinematicProperties(ClusterIndex);
 	}
 
@@ -1777,7 +1798,7 @@ namespace Chaos
 			{
 				return CurrentIndex;
 			}
-			CurrentIndex = MClusterIds[CurrentIndex].Id;
+			CurrentIndex = MParticles.ClusterIds(CurrentIndex).Id;
 		}
 		return INDEX_NONE;
 	}
@@ -1788,9 +1809,9 @@ namespace Chaos
 		if (Index1 != Index2)
 		{
 			//todo(pref): This can be removed if we are sure there are no duplicate connections generated.
-			for (int32 i = 0; i < MConnectivityEdges[Index1].Num(); i++)
+			for (int32 i = 0; i < MParticles.ConnectivityEdges(Index1).Num(); i++)
 			{
-				if (MConnectivityEdges[Index1][i].Sibling == Index2)
+				if (MParticles.ConnectivityEdges(Index1)[i].Sibling == Index2)
 				{
 					// @todo(duplication connection) : re-enable post GDC.  
 					// FixConnectivityGraphUsingDelaunayTriangulation attempts to add multiple connections.
@@ -1800,7 +1821,7 @@ namespace Chaos
 				}
 			}
 
-			MConnectivityEdges[Index1].Add({Index2, Strain});
+			MParticles.ConnectivityEdges(Index1).Add({ Index2, Strain });
 		}
 	}
 
@@ -1816,21 +1837,21 @@ namespace Chaos
 	void TPBDRigidClustering<FPBDRigidsEvolution, FPBDCollisionConstraint, T, d>::RemoveNodeConnections(uint32 ParticleIndex)
 	{
 		SCOPE_CYCLE_COUNTER(STAT_RemoveNodeConnections);
-		checkSlow(0 <= ParticleIndex && ParticleIndex < (uint32)MClusterIds.Num());
+		checkSlow(0 <= ParticleIndex && ParticleIndex < MParticles.Size());
 
-		for (int32 i = MConnectivityEdges[ParticleIndex].Num() - 1; 0 <= i; i--)
+		for (int32 i = MParticles.ConnectivityEdges(ParticleIndex).Num() - 1; 0 <= i; i--)
 		{
-			const TConnectivityEdge<T>& Edge = MConnectivityEdges[ParticleIndex][i];
-			for (int Idx = MConnectivityEdges[Edge.Sibling].Num() - 1; Idx >= 0; --Idx)
+			const TConnectivityEdge<T>& Edge = MParticles.ConnectivityEdges(ParticleIndex)[i];
+			for (int Idx = MParticles.ConnectivityEdges(Edge.Sibling).Num() - 1; Idx >= 0; --Idx)
 			{
-				const TConnectivityEdge<T>& OtherEdge = MConnectivityEdges[Edge.Sibling][Idx];
+				const TConnectivityEdge<T>& OtherEdge = MParticles.ConnectivityEdges(Edge.Sibling)[Idx];
 				if (OtherEdge.Sibling == ParticleIndex)
 				{
-					MConnectivityEdges[Edge.Sibling].RemoveAtSwap(Idx); //Note: we shouldn't have to keep searching, but sometimes these over subscribe, should fix
+					MParticles.ConnectivityEdges(Edge.Sibling).RemoveAtSwap(Idx); //Note: we shouldn't have to keep searching, but sometimes these over subscribe, should fix
 				}
 			}
 		}
-		MConnectivityEdges[ParticleIndex].SetNum(0);
+		MParticles.ConnectivityEdges(ParticleIndex).SetNum(0);
 	}
 
 	DECLARE_CYCLE_STAT(TEXT("TPBDRigidClustering<>::SwapBufferedData"), STAT_SwapBufferedData, STATGROUP_Chaos);
@@ -1844,14 +1865,15 @@ namespace Chaos
 		const TArray<TSerializablePtr<TImplicitObject<T, d>>>& AllGeom = MParticles.GetAllGeometry();
 		BufferResource.GeometryPtrs.SetNum(AllGeom.Num());
 
-		const TArray<int32>& NonDisabledIndices = MEvolution.GetNonDisabledIndices();
-		for (int32 Idx : NonDisabledIndices)
+		const auto& NonDisabledClusteredParticles = MEvolution.GetNonDisabledClusteredArray();
+		for (auto& ClusteredParticle : NonDisabledClusteredParticles)
 		{
-			const TUniquePtr<TArray<uint32>>* Children = MChildren.Find(Idx);
-			if (Children && MClusterIds[Idx].Id == INDEX_NONE) //root cluster so copy children
+			const TArray<uint32>* Children = MChildren.Find(ClusteredParticle);
+			if (Children && ClusteredParticle->ClusterIds().Id == INDEX_NONE) //root cluster so copy children
 			{
-				BufferResource.MChildren.Add(Idx, **Children);
-				BufferResource.ClusterParentTransforms.Add(Idx, TRigidTransform<float, 3>(MParticles.X(Idx), MParticles.R(Idx)));
+				//TODO: record GT particle pointer instead
+				BufferResource.MChildren.Add(ClusteredParticle->TransientParticleIndex(), *Children);
+				BufferResource.ClusterParentTransforms.Add(ClusteredParticle->TransientParticleIndex(), TRigidTransform<float, 3>(ClusteredParticle->X(), ClusteredParticle->R()));
 			}
 		}
 
@@ -1889,7 +1911,9 @@ namespace Chaos
 
 using namespace Chaos;
 template class Chaos::TPBDRigidClustering<TPBDRigidsEvolutionGBF<float, 3>, TPBDCollisionConstraint<float, 3>, float, 3>;
+#if CHAOS_PARTICLEHANDLE_TODO
 template class Chaos::TPBDRigidClustering<TPBDRigidsEvolutionPGS<float, 3>, TPBDCollisionConstraintPGS<float, 3>, float, 3>;
-template CHAOS_API void Chaos::UpdateClusterMassProperties<float, 3>(TPBDRigidParticles<float, 3>& Particles, const TArray<uint32>& Children, const uint32 ClusterIndex, const TRigidTransform<float, 3>* ForceMassOrientation,
-    const TArrayCollectionArray<TUniquePtr<TMultiChildProxyData<float, 3>>>* MMultiChildProxyData,
-    const TArrayCollectionArray<FMultiChildProxyId>* MMultiChildProxyId);
+#endif
+template CHAOS_API void Chaos::UpdateClusterMassProperties<float, 3>(TPBDRigidClusteredParticles<float, 3>& Particles, const TArray<uint32>& Children, const uint32 ClusterIndex, const TRigidTransform<float, 3>* ForceMassOrientation,
+	const TArrayCollectionArray<TUniquePtr<TMultiChildProxyData<float, 3>>>* MMultiChildProxyData,
+	const TArrayCollectionArray<FMultiChildProxyId>* MMultiChildProxyId);
diff --git a/Engine/Source/Runtime/Experimental/Chaos/Private/Chaos/PBDRigidDynamicSpringConstraints.cpp b/Engine/Source/Runtime/Experimental/Chaos/Private/Chaos/PBDRigidDynamicSpringConstraints.cpp
--- a/Engine/Source/Runtime/Experimental/Chaos/Private/Chaos/PBDRigidDynamicSpringConstraints.cpp
+++ b/Engine/Source/Runtime/Experimental/Chaos/Private/Chaos/PBDRigidDynamicSpringConstraints.cpp
@@ -5,40 +5,47 @@
 using namespace Chaos;
 
 template<class T, int d>
-void TPBDRigidDynamicSpringConstraints<T, d>::ApplyHelper(TPBDRigidParticles<T, d>& InParticles, const T Dt, const TArray<int32>& InConstraintIndices) const
+void TPBDRigidDynamicSpringConstraints<T, d>::ApplyHelper(const T Dt, const TArray<int32>& InConstraintIndices) const
 {
-	for (int32 ConstraintIndex : InConstraintIndices)
+ 	for (int32 ConstraintIndex : InConstraintIndices)
 	{
-		const TVector<int32, 2>& Constraint = Constraints[ConstraintIndex];
-
-		int32 ConstraintInnerIndex1 = Constraint[0];
-		int32 ConstraintInnerIndex2 = Constraint[1];
-
-		check(InParticles.Island(ConstraintInnerIndex1) == InParticles.Island(ConstraintInnerIndex2) || InParticles.Island(ConstraintInnerIndex1) == -1 || InParticles.Island(ConstraintInnerIndex2) == -1);
+		TGeometryParticleHandle<T, d>* Static0 = Constraints[ConstraintIndex][0];
+		TGeometryParticleHandle<T, d>* Static1 = Constraints[ConstraintIndex][1];
+		TPBDRigidParticleHandle<T, d>* PBDRigid0 = Static0->AsDynamic();
+		TPBDRigidParticleHandle<T, d>* PBDRigid1 = Static1->AsDynamic();
+		check(PBDRigid0 || PBDRigid1);
+		check(!PBDRigid0 || !PBDRigid1 || (PBDRigid0->Island() == PBDRigid1->Island()));
+
+		TRotation<T, d>& Q0 = PBDRigid0 ? PBDRigid0->Q() : Static0->R();
+		TRotation<T, d>& Q1 = PBDRigid1 ? PBDRigid1->Q() : Static1->R();
+		TVector<T, d>& P0 = PBDRigid0 ? PBDRigid0->P() : Static0->X();
+		TVector<T, d>& P1 = PBDRigid1 ? PBDRigid1->P() : Static1->X();
 
 		const int32 NumSprings = SpringDistances[ConstraintIndex].Num();
-		const PMatrix<T, d, d> WorldSpaceInvI1 = (InParticles.Q(ConstraintInnerIndex1) * FMatrix::Identity) * InParticles.InvI(ConstraintInnerIndex1) * (InParticles.Q(ConstraintInnerIndex1) * FMatrix::Identity).GetTransposed();
-		const PMatrix<T, d, d> WorldSpaceInvI2 = (InParticles.Q(ConstraintInnerIndex2) * FMatrix::Identity) * InParticles.InvI(ConstraintInnerIndex2) * (InParticles.Q(ConstraintInnerIndex2) * FMatrix::Identity).GetTransposed();
+		const PMatrix<T, d, d> WorldSpaceInvI1 = PBDRigid0? (Q0 * FMatrix::Identity) * PBDRigid0->InvI() * (Q0 * FMatrix::Identity).GetTransposed() : PMatrix<T, d, d>(0);
+		const PMatrix<T, d, d> WorldSpaceInvI2 = PBDRigid1 ? (Q1 * FMatrix::Identity) * PBDRigid1->InvI() * (Q1 * FMatrix::Identity).GetTransposed() : PMatrix<T, d, d>(0);;
 		for (int32 SpringIndex = 0; SpringIndex < NumSprings; ++SpringIndex)
 		{
-			const TVector<T, d> WorldSpaceX1 = InParticles.Q(ConstraintInnerIndex1).RotateVector(Distances[ConstraintIndex][SpringIndex][0]) + InParticles.P(ConstraintInnerIndex1);
-			const TVector<T, d> WorldSpaceX2 = InParticles.Q(ConstraintInnerIndex2).RotateVector(Distances[ConstraintIndex][SpringIndex][1]) + InParticles.P(ConstraintInnerIndex2);
-			const TVector<T, d> Delta = Base::GetDelta(InParticles, WorldSpaceX1, WorldSpaceX2, ConstraintIndex, SpringIndex);
+			const TVector<T, d>& Distance0 = Distances[ConstraintIndex][SpringIndex][0];
+			const TVector<T, d>& Distance1 = Distances[ConstraintIndex][SpringIndex][1];
+			const TVector<T, d> WorldSpaceX1 = Q0.RotateVector(Distance0) + P0;
+			const TVector<T, d> WorldSpaceX2 = Q1.RotateVector(Distance1) + P1;
+			const TVector<T, d> Delta = Base::GetDelta(WorldSpaceX1, WorldSpaceX2, ConstraintIndex, SpringIndex);
 
-			if (InParticles.InvM(ConstraintInnerIndex1) > 0)
+			if (PBDRigid0)
 			{
-				const TVector<T, d> Radius = WorldSpaceX1 - InParticles.P(ConstraintInnerIndex1);
-				InParticles.P(ConstraintInnerIndex1) += InParticles.InvM(ConstraintInnerIndex1) * Delta;
-				InParticles.Q(ConstraintInnerIndex1) += TRotation<T, d>(WorldSpaceInvI1 * TVector<T, d>::CrossProduct(Radius, Delta), 0.f) * InParticles.Q(ConstraintInnerIndex1) * T(0.5);
-				InParticles.Q(ConstraintInnerIndex1).Normalize();
+				const TVector<T, d> Radius = WorldSpaceX1 - P0;
+				P0 += PBDRigid0->InvM() * Delta;
+				Q0 += TRotation<T, d>(WorldSpaceInvI1 * TVector<T, d>::CrossProduct(Radius, Delta), 0.f) * Q0 * T(0.5);
+				Q0.Normalize();
 			}
 
-			if (InParticles.InvM(ConstraintInnerIndex2) > 0)
+			if (PBDRigid1)
 			{
-				const TVector<T, d> Radius = WorldSpaceX2 - InParticles.P(ConstraintInnerIndex2);
-				InParticles.P(ConstraintInnerIndex2) -= InParticles.InvM(ConstraintInnerIndex2) * Delta;
-				InParticles.Q(ConstraintInnerIndex2) += TRotation<T, d>(WorldSpaceInvI2 * TVector<T, d>::CrossProduct(Radius, -Delta), 0.f) * InParticles.Q(ConstraintInnerIndex2) * T(0.5);
-				InParticles.Q(ConstraintInnerIndex2).Normalize();
+				const TVector<T, d> Radius = WorldSpaceX2 - P1;
+				P1 -= PBDRigid1->InvM() * Delta;
+				Q1 += TRotation<T, d>(WorldSpaceInvI2 * TVector<T, d>::CrossProduct(Radius, -Delta), 0.f) * Q1 * T(0.5);
+				Q1.Normalize();
 			}
 		}
 	}
diff --git a/Engine/Source/Runtime/Experimental/Chaos/Private/Chaos/PBDRigidDynamicSpringConstraintsBase.cpp b/Engine/Source/Runtime/Experimental/Chaos/Private/Chaos/PBDRigidDynamicSpringConstraintsBase.cpp
--- a/Engine/Source/Runtime/Experimental/Chaos/Private/Chaos/PBDRigidDynamicSpringConstraintsBase.cpp
+++ b/Engine/Source/Runtime/Experimental/Chaos/Private/Chaos/PBDRigidDynamicSpringConstraintsBase.cpp
@@ -5,26 +5,35 @@
 using namespace Chaos;
 
 template<class T, int d>
-void TPBDRigidDynamicSpringConstraintsBase<T, d>::UpdatePositionBasedState(const TPBDRigidParticles<T, d>& InParticles)
+void TPBDRigidDynamicSpringConstraintsBase2<T, d>::UpdatePositionBasedState()
 {
 	const int32 NumConstraints = Constraints.Num();
 	for (int32 ConstraintIndex = 0; ConstraintIndex < NumConstraints; ++ConstraintIndex)
 	{
-		int32 ConstraintInnerIndex1 = Constraints[ConstraintIndex][0];
-		int32 ConstraintInnerIndex2 = Constraints[ConstraintIndex][1];
+		TGeometryParticleHandle<T, d>* Static0 = Constraints[ConstraintIndex][0];
+		TGeometryParticleHandle<T, d>* Static1 = Constraints[ConstraintIndex][1];
+		TPBDRigidParticleHandle<T, d>* PBDRigid0 = Static0->AsDynamic();
+		TPBDRigidParticleHandle<T, d>* PBDRigid1 = Static1->AsDynamic();
 
 		// Do not create springs between objects with no geometry
-		if (!InParticles.Geometry(ConstraintInnerIndex1) || !InParticles.Geometry(ConstraintInnerIndex2))
+		if (!Static0->Geometry() || !Static1->Geometry())
 		{
 			continue;
 		}
 
+		const TRotation<T, d>& Q0 = PBDRigid0 ? PBDRigid0->Q() : Static0->R();
+		const TRotation<T, d>& Q1 = PBDRigid1 ? PBDRigid1->Q() : Static1->R();
+		const TVector<T, d>& P0 = PBDRigid0 ? PBDRigid0->P() : Static0->X();
+		const TVector<T, d>& P1 = PBDRigid1 ? PBDRigid1->P() : Static1->X();
+
 		// Delete constraints
 		const int32 NumSprings = SpringDistances[ConstraintIndex].Num();
 		for (int32 SpringIndex = NumSprings - 1; SpringIndex >= 0; --SpringIndex)
 		{
-			const TVector<T, d> WorldSpaceX1 = InParticles.Q(ConstraintInnerIndex1).RotateVector(Distances[ConstraintIndex][SpringIndex][0]) + InParticles.P(ConstraintInnerIndex1);
-			const TVector<T, d> WorldSpaceX2 = InParticles.Q(ConstraintInnerIndex2).RotateVector(Distances[ConstraintIndex][SpringIndex][1]) + InParticles.P(ConstraintInnerIndex2);
+			const TVector<T, d>& Distance0 = Distances[ConstraintIndex][SpringIndex][0];
+			const TVector<T, d>& Distance1 = Distances[ConstraintIndex][SpringIndex][1];
+			const TVector<T, d> WorldSpaceX1 = Q0.RotateVector(Distance0) + P0;
+			const TVector<T, d> WorldSpaceX2 = Q1.RotateVector(Distance1) + P1;
 			const TVector<T, d> Difference = WorldSpaceX2 - WorldSpaceX1;
 			float Distance = Difference.Size();
 			if (Distance > CreationThreshold * 2)
@@ -39,28 +48,28 @@ void TPBDRigidDynamicSpringConstraintsBase<T, d>::UpdatePositionBasedState(const
 			continue;
 		}
 
-		TRigidTransform<T, d> Transform1(InParticles.P(ConstraintInnerIndex1), InParticles.Q(ConstraintInnerIndex1));
-		TRigidTransform<T, d> Transform2(InParticles.P(ConstraintInnerIndex2), InParticles.Q(ConstraintInnerIndex2));
+		TRigidTransform<T, d> Transform1(P0, Q0);
+		TRigidTransform<T, d> Transform2(P1, Q1);
 
 		// Create constraints
-		if (InParticles.Geometry(ConstraintInnerIndex1)->HasBoundingBox() && InParticles.Geometry(ConstraintInnerIndex2)->HasBoundingBox())
+		if (Static0->Geometry()->HasBoundingBox() && Static1->Geometry()->HasBoundingBox())
 		{
 			// Matrix multiplication is reversed intentionally to be compatible with unreal
-			TBox<T, d> Box1 = InParticles.Geometry(ConstraintInnerIndex1)->BoundingBox().TransformedBox(Transform1 * Transform2.Inverse());
+			TBox<T, d> Box1 = Static0->Geometry()->BoundingBox().TransformedBox(Transform1 * Transform2.Inverse());
 			Box1.Thicken(CreationThreshold);
-			TBox<T, d> Box2 = InParticles.Geometry(ConstraintInnerIndex2)->BoundingBox();
+			TBox<T, d> Box2 = Static1->Geometry()->BoundingBox();
 			Box2.Thicken(CreationThreshold);
 			if (!Box1.Intersects(Box2))
 			{
 				continue;
 			}
 		}
-		const TVector<T, d> Midpoint = (InParticles.P(ConstraintInnerIndex1) + InParticles.P(ConstraintInnerIndex2)) / (T)2;
+		const TVector<T, d> Midpoint = (P0 + P1) / (T)2;
 		TVector<T, d> Normal1;
-		const T Phi1 = InParticles.Geometry(ConstraintInnerIndex1)->PhiWithNormal(Transform1.InverseTransformPosition(Midpoint), Normal1);
+		const T Phi1 = Static0->Geometry()->PhiWithNormal(Transform1.InverseTransformPosition(Midpoint), Normal1);
 		Normal1 = Transform2.TransformVector(Normal1);
 		TVector<T, d> Normal2;
-		const T Phi2 = InParticles.Geometry(ConstraintInnerIndex2)->PhiWithNormal(Transform2.InverseTransformPosition(Midpoint), Normal2);
+		const T Phi2 = Static1->Geometry()->PhiWithNormal(Transform2.InverseTransformPosition(Midpoint), Normal2);
 		Normal2 = Transform2.TransformVector(Normal2);
 		if ((Phi1 + Phi2) > CreationThreshold)
 		{
@@ -69,38 +78,36 @@ void TPBDRigidDynamicSpringConstraintsBase<T, d>::UpdatePositionBasedState(const
 		TVector<T, d> Location0 = Midpoint - Phi1 * Normal1;
 		TVector<T, d> Location1 = Midpoint - Phi2 * Normal2;
 		TVector<TVector<T, 3>, 2> Distance;
-		Distance[0] = InParticles.Q(ConstraintInnerIndex1).Inverse().RotateVector(Location0 - InParticles.P(ConstraintInnerIndex1));
-		Distance[1] = InParticles.Q(ConstraintInnerIndex2).Inverse().RotateVector(Location1 - InParticles.P(ConstraintInnerIndex2));
+		Distance[0] = Q0.Inverse().RotateVector(Location0 - P0);
+		Distance[1] = Q0.Inverse().RotateVector(Location1 - P1);
 		Distances[ConstraintIndex].Add(MoveTemp(Distance));
 		SpringDistances[ConstraintIndex].Add((Location0 - Location1).Size());
 	}
 }
 
 template<class T, int d>
-TVector<T, d> TPBDRigidDynamicSpringConstraintsBase<T, d>::GetDelta(const TPBDRigidParticles<T, d>& InParticles, const TVector<T, d>& WorldSpaceX1, const TVector<T, d>& WorldSpaceX2, const int32 i, const int32 j) const
+TVector<T, d> TPBDRigidDynamicSpringConstraintsBase2<T, d>::GetDelta(const TVector<T, d>& WorldSpaceX1, const TVector<T, d>& WorldSpaceX2, const int32 ConstraintIndex, const int32 SpringIndex) const
 {
-	const TVector<int32, 2>& Constraint = Constraints[i];
-
-	int32 i1 = Constraint[0];
-	int32 i2 = Constraint[1];
+	TPBDRigidParticleHandle<T, d>* PBDRigid0 = Constraints[ConstraintIndex][0]->AsDynamic();
+	TPBDRigidParticleHandle<T, d>* PBDRigid1 = Constraints[ConstraintIndex][1]->AsDynamic();
 
-	if (InParticles.InvM(i2) == 0 && InParticles.InvM(i1) == 0)
+	if (!PBDRigid0 && !PBDRigid1)
 	{
 		return TVector<T, d>(0);
 	}
 
-	TVector<T, d> Difference = WorldSpaceX2 - WorldSpaceX1;
+	const TVector<T, d> Difference = WorldSpaceX2 - WorldSpaceX1;
 	float Distance = Difference.Size();
-
 	check(Distance > 1e-7);
 
-	TVector<T, d> Direction = Difference / Distance;
-	TVector<T, d> Delta = (Distance - SpringDistances[i][j]) * Direction;
-	T CombinedMass = InParticles.InvM(i2) + InParticles.InvM(i1);
-	return Stiffness * Delta / CombinedMass;
+	const T InvM0 = PBDRigid0 ? PBDRigid0->InvM() : (T)0;
+	const T InvM1 = PBDRigid1 ? PBDRigid1->InvM() : (T)0;
+	const TVector<T, d> Direction = Difference / Distance;
+	const TVector<T, d> Delta = (Distance - SpringDistances[ConstraintIndex][SpringIndex]) * Direction;
+	return Stiffness * Delta / (InvM0 + InvM1);
 }
 
 namespace Chaos
 {
-	template class TPBDRigidDynamicSpringConstraintsBase<float, 3>;
+	template class TPBDRigidDynamicSpringConstraintsBase2<float, 3>;
 }
diff --git a/Engine/Source/Runtime/Experimental/Chaos/Private/Chaos/PBDRigidSpringConstraints.cpp b/Engine/Source/Runtime/Experimental/Chaos/Private/Chaos/PBDRigidSpringConstraints.cpp
--- a/Engine/Source/Runtime/Experimental/Chaos/Private/Chaos/PBDRigidSpringConstraints.cpp
+++ b/Engine/Source/Runtime/Experimental/Chaos/Private/Chaos/PBDRigidSpringConstraints.cpp
@@ -6,5 +6,5 @@ using namespace Chaos;
 
 namespace Chaos
 {
-	template class TPBDRigidSpringConstraints<float, 3>;
+	template class TPBDRigidSpringConstraints2<float, 3>;
 }
diff --git a/Engine/Source/Runtime/Experimental/Chaos/Private/Chaos/PBDRigidSpringConstraintsBase.cpp b/Engine/Source/Runtime/Experimental/Chaos/Private/Chaos/PBDRigidSpringConstraintsBase.cpp
--- a/Engine/Source/Runtime/Experimental/Chaos/Private/Chaos/PBDRigidSpringConstraintsBase.cpp
+++ b/Engine/Source/Runtime/Experimental/Chaos/Private/Chaos/PBDRigidSpringConstraintsBase.cpp
@@ -5,7 +5,7 @@
 using namespace Chaos;
 
 template<class T, int d>
-void TPBDRigidSpringConstraintsBase<T, d>::UpdateDistances(const TRigidParticles<T, d>& InParticles, const TArray<TVector<T, d>>& Locations0, const TArray<TVector<T, d>>& Locations1)
+void TPBDRigidSpringConstraintsBase<T, d>::UpdateDistances(const TArray<TVector<T, d>>& Locations0, const TArray<TVector<T, d>>& Locations1)
 {
 	Distances.SetNum(Constraints.Num());
 	SpringDistances.SetNum(Constraints.Num());
@@ -13,37 +13,39 @@ void TPBDRigidSpringConstraintsBase<T, d>::UpdateDistances(const TRigidParticles
 	const int32 NumConstraints = Constraints.Num();
 	for (int32 ConstraintIndex = 0; ConstraintIndex < NumConstraints; ++ConstraintIndex)
 	{
-		int32 ConstraintInnerIndex1 = Constraints[ConstraintIndex][0];
-		int32 ConstraintInnerIndex2 = Constraints[ConstraintIndex][1];
+		const TVector<TGeometryParticleHandle<T, d>*, 2>& Constraint = Constraints[ConstraintIndex];
+		const TGeometryParticleHandle<T, d>* Particle0 = Constraint[0];
+		const TGeometryParticleHandle<T, d>* Particle1 = Constraint[1];
 
-		Distances[ConstraintIndex][0] = InParticles.R(ConstraintInnerIndex1).Inverse().RotateVector(Locations0[ConstraintIndex] - InParticles.X(ConstraintInnerIndex1));
-		Distances[ConstraintIndex][1] = InParticles.R(ConstraintInnerIndex2).Inverse().RotateVector(Locations1[ConstraintIndex] - InParticles.X(ConstraintInnerIndex2));
+		Distances[ConstraintIndex][0] = Particle0->R().Inverse().RotateVector(Locations0[ConstraintIndex] - Particle0->X());
+		Distances[ConstraintIndex][1] = Particle1->R().Inverse().RotateVector(Locations1[ConstraintIndex] - Particle1->X());
 		SpringDistances[ConstraintIndex] = (Locations0[ConstraintIndex] - Locations1[ConstraintIndex]).Size();
 	}
 }
 
 template<class T, int d>
-TVector<T, d> TPBDRigidSpringConstraintsBase<T, d>::GetDelta(const TPBDRigidParticles<T, d>& InParticles, const TVector<T, d>& WorldSpaceX1, const TVector<T, d>& WorldSpaceX2, int32 i) const
+TVector<T, d> TPBDRigidSpringConstraintsBase<T, d>::GetDelta(const TVector<T, d>& WorldSpaceX1, const TVector<T, d>& WorldSpaceX2, int32 ConstraintIndex) const
 {
-	const TVector<int32, 2>& Constraint = Constraints[i];
+	const TVector<TGeometryParticleHandle<T, d>*, 2>& Constraint = Constraints[ConstraintIndex];
+	const TPBDRigidParticleHandle<T, d>* PBDRigid0 = Constraint[0]->AsDynamic();
+	const TPBDRigidParticleHandle<T, d>* PBDRigid1 = Constraint[1]->AsDynamic();
 
-	int32 i1 = Constraint[0];
-	int32 i2 = Constraint[1];
-
-	if (InParticles.InvM(i2) == 0 && InParticles.InvM(i1) == 0)
+	if (!PBDRigid0 && !PBDRigid1)
 	{
 		return TVector<T, d>(0);
 	}
 
-	TVector<T, d> Difference = WorldSpaceX2 - WorldSpaceX1;
-
-	float Distance = Difference.Size();
+	const TVector<T, d> Difference = WorldSpaceX2 - WorldSpaceX1;
 
+	const float Distance = Difference.Size();
 	check(Distance > 1e-7);
 
-	TVector<T, d> Direction = Difference / Distance;
-	TVector<T, d> Delta = (Distance - SpringDistances[i]) * Direction;
-	T CombinedMass = InParticles.InvM(i2) + InParticles.InvM(i1);
+	const TVector<T, d> Direction = Difference / Distance;
+	const TVector<T, d> Delta = (Distance - SpringDistances[ConstraintIndex]) * Direction;
+	const T InvM0 = (PBDRigid0) ? PBDRigid0->InvM() : (T)0;
+	const T InvM1 = (PBDRigid1) ? PBDRigid1->InvM() : (T)0;
+	const T CombinedMass = InvM0 + InvM1;
+
 	return Stiffness * Delta / CombinedMass;
 }
 
diff --git a/Engine/Source/Runtime/Experimental/Chaos/Private/Chaos/PBDRigidsEvolution.cpp b/Engine/Source/Runtime/Experimental/Chaos/Private/Chaos/PBDRigidsEvolution.cpp
--- a/Engine/Source/Runtime/Experimental/Chaos/Private/Chaos/PBDRigidsEvolution.cpp
+++ b/Engine/Source/Runtime/Experimental/Chaos/Private/Chaos/PBDRigidsEvolution.cpp
@@ -1,6 +1,5 @@
 // Copyright 1998-2019 Epic Games, Inc. All Rights Reserved.
 #include "Chaos/PBDRigidsEvolution.h"
-
 #include "Chaos/Box.h"
 #include "Chaos/Defines.h"
 #include "Chaos/Framework/Parallel.h"
@@ -8,8 +7,6 @@
 #include "Chaos/ImplicitObjectUnion.h"
 #include "Chaos/PBDCollisionConstraint.h"
 #include "Chaos/PBDCollisionSpringConstraints.h"
-#include "Chaos/PBDRigidsEvolutionGBF.h"
-#include "Chaos/PBDRigidsEvolutionPGS.h"
 #include "Chaos/PerParticleEtherDrag.h"
 #include "Chaos/PerParticleEulerStepVelocity.h"
 #include "Chaos/PerParticleGravity.h"
@@ -18,46 +15,26 @@
 #include "Chaos/PerParticlePBDGroundConstraint.h"
 #include "Chaos/PerParticlePBDUpdateFromDeltaPosition.h"
 #include "ChaosStats.h"
-#include "ChaosLog.h"
 
 #include "ProfilingDebugging/ScopedTimers.h"
 #include "Chaos/DebugDrawQueue.h"
+#include "Chaos/Levelset.h"
+#include "Chaos/ChaosPerfTest.h"
+#include "Chaos/PBDRigidsEvolutionGBF.h"
 
-#define LOCTEXT_NAMESPACE "Chaos"
-
-DEFINE_LOG_CATEGORY(LogChaos);
-
-using namespace Chaos;
-
-template<class FPBDRigidsEvolution, class FPBDCollisionConstraint, class T, int d>
-TPBDRigidsEvolutionBase<FPBDRigidsEvolution, FPBDCollisionConstraint, T, d>::TPBDRigidsEvolutionBase(TPBDRigidParticles<T, d>&& InParticles, int32 InNumIterations)
-    : Particles(MoveTemp(InParticles))
-    , Clustering(static_cast<FPBDRigidsEvolution&>(*this), Particles)
-	, Time(0)
-    , DebugSubstep()
-	, NumIterations(InNumIterations)
-{
-	Particles.AddArray(&Collided);
-	Particles.AddArray(&PhysicsMaterials);
-	Particles.AddArray(&PerParticlePhysicsMaterials);
-	Particles.AddArray(&ParticleDisableCount);
-	InitializeFromParticleData(0);
-}
-
-#if !UE_BUILD_SHIPPING
-template<typename FPBDRigidsEvolution, class FPBDCollisionConstraint, class T, int d>
-void TPBDRigidsEvolutionBase<FPBDRigidsEvolution, FPBDCollisionConstraint, T, d>::SerializeForPerfTest(FChaosArchive& Ar)
+namespace Chaos
 {
-	//todo: this is incomplete, but useful for perf testing
-	Ar << Particles;
-	Ar << PhysicsMaterials;
-	Ar << PerParticlePhysicsMaterials;
-	InitializeFromParticleData(0);
+	template<class FPBDRigidsEvolution, class FPBDCollisionConstraint, class T, int d>
+	TPBDRigidsEvolutionBase<FPBDRigidsEvolution, FPBDCollisionConstraint, T, d>::TPBDRigidsEvolutionBase(TPBDRigidsSOAs<T, d>& InParticles, int32 InNumIterations)
+		: Particles(InParticles)
+		, Clustering(static_cast<FPBDRigidsEvolution&>(*this), Particles.GetClusteredParticles())
+		, NumIterations(InNumIterations)
+	{
+		Particles.GetParticleHandles().AddArray(&PhysicsMaterials);
+		Particles.GetParticleHandles().AddArray(&PerParticlePhysicsMaterials);
+		Particles.GetParticleHandles().AddArray(&ParticleDisableCount);
+		Particles.GetParticleHandles().AddArray(&Collided);
+	}
 }
-#endif
-
-template class Chaos::TPBDRigidsEvolutionBase<TPBDRigidsEvolutionGBF<float, 3>, TPBDCollisionConstraint<float, 3>, float, 3>;
-template class Chaos::TPBDRigidsEvolutionBase<TPBDRigidsEvolutionPGS<float, 3>, TPBDCollisionConstraintPGS<float, 3>, float, 3>;
-
 
-#undef LOCTEXT_NAMESPACE
\ No newline at end of file
+template class Chaos::TPBDRigidsEvolutionBase<Chaos::TPBDRigidsEvolutionGBF<float, 3>, Chaos::TPBDCollisionConstraint<float,3>, float, 3>;
\ No newline at end of file
