Copying //UE4/Dev-Physics to Dev-Main (//UE4/Dev-Main) [at] 7703071

diff --git a/Engine/Source/Runtime/Experimental/Chaos/Private/Chaos/PBDRigidsEvolutionGBF.cpp b/Engine/Source/Runtime/Experimental/Chaos/Private/Chaos/PBDRigidsEvolutionGBF.cpp
--- a/Engine/Source/Runtime/Experimental/Chaos/Private/Chaos/PBDRigidsEvolutionGBF.cpp
+++ b/Engine/Source/Runtime/Experimental/Chaos/Private/Chaos/PBDRigidsEvolutionGBF.cpp
@@ -1,6 +1,5 @@
 // Copyright 1998-2019 Epic Games, Inc. All Rights Reserved.
 #include "Chaos/PBDRigidsEvolutionGBF.h"
-
 #include "Chaos/Box.h"
 #include "Chaos/Defines.h"
 #include "Chaos/Framework/Parallel.h"
@@ -20,439 +19,156 @@
 #include "ProfilingDebugging/ScopedTimers.h"
 #include "Chaos/DebugDrawQueue.h"
 #include "Chaos/Levelset.h"
-#include "Chaos/ChaosPerfTest.h"
-
-#define LOCTEXT_NAMESPACE "Chaos"
-
-int32 DisableSim = 0;
-FAutoConsoleVariableRef CVarDisableSim(TEXT("p.DisableSim"), DisableSim, TEXT("Disable Sim"));
 
-using namespace Chaos;
-
-template<class T, int d>
-TPBDRigidsEvolutionGBF<T, d>::TPBDRigidsEvolutionGBF(TPBDRigidParticles<T, d>&& InParticles, int32 InNumIterations)
-    : Base(MoveTemp(InParticles), InNumIterations)
-	, CollisionConstraints(Particles, NonDisabledIndices, Collided, PhysicsMaterials, DefaultNumPushOutPairIterations, (T)0)
-	, CollisionRule(CollisionConstraints, DefaultNumPushOutIterations)
+namespace Chaos
 {
-	SetParticleUpdateVelocityFunction([PBDUpdateRule = TPerParticlePBDUpdateFromDeltaPosition<float, 3>(), this](TPBDRigidParticles<T, d>& ParticlesInput, const T Dt, const TArray<int32>& InActiveIndices) {
-		PhysicsParallelFor(InActiveIndices.Num(), [&](int32 ActiveIndex) {
-			int32 Index = InActiveIndices[ActiveIndex];
-			PBDUpdateRule.Apply(ParticlesInput, Dt, Index);
-		});
-	});
-
-	SetParticleUpdatePositionFunction([this](TPBDRigidParticles<T, d>& ParticlesInput, const T Dt)
-	{
-		const TArray<int32>& ActiveIndicesArray = GetActiveIndicesArray();
-		PhysicsParallelFor(ActiveIndicesArray.Num(), [&](int32 ActiveIndex)
-		{
-			int32 Index = ActiveIndicesArray[ActiveIndex];
-			ParticlesInput.X(Index) = ParticlesInput.P(Index);
-			ParticlesInput.R(Index) = ParticlesInput.Q(Index);
-		});
-	});
-
-	AddConstraintRule(&CollisionRule);
-}
-
-DECLARE_CYCLE_STAT(TEXT("Integrate"), STAT_Integrate, STATGROUP_Chaos);
-
 float HackMaxAngularVelocity = 1000.f;
 FAutoConsoleVariableRef CVarHackMaxAngularVelocity(TEXT("p.HackMaxAngularVelocity"), HackMaxAngularVelocity, TEXT("Max cap on angular velocity: rad/s. This is only a temp solution and should not be relied on as a feature. -1.f to disable"));
 
 float HackMaxVelocity = -1.f;
-FAutoConsoleVariableRef CVarHackMaxVelocity(TEXT("p.HackMaxVelocity"), HackMaxVelocity, TEXT("Max cap on velocity: cm/s. This is only a temp solution and should not be relied on as a feature. -1.f to disable"));
+FAutoConsoleVariableRef CVarHackMaxVelocity(TEXT("p.HackMaxVelocity2"), HackMaxVelocity, TEXT("Max cap on velocity: cm/s. This is only a temp solution and should not be relied on as a feature. -1.f to disable"));
 
 
 float HackLinearDrag = 0.f;
-FAutoConsoleVariableRef CVarHackLinearDrag(TEXT("p.HackLinearDrag"), HackLinearDrag, TEXT("Linear drag used to slow down objects. This is a hack and should not be relied on as a feature."));
+FAutoConsoleVariableRef CVarHackLinearDrag(TEXT("p.HackLinearDrag2"), HackLinearDrag, TEXT("Linear drag used to slow down objects. This is a hack and should not be relied on as a feature."));
 
 float HackAngularDrag = 0.f;
-FAutoConsoleVariableRef CVarHackAngularDrag(TEXT("p.HackAngularDrag"), HackAngularDrag, TEXT("Angular drag used to slow down objects. This is a hack and should not be relied on as a feature."));
+FAutoConsoleVariableRef CVarHackAngularDrag(TEXT("p.HackAngularDrag2"), HackAngularDrag, TEXT("Angular drag used to slow down objects. This is a hack and should not be relied on as a feature."));
 
 int DisableThreshold = 5;
-FAutoConsoleVariableRef CVarDisableThreshold(TEXT("p.DisableThreshold"), DisableThreshold, TEXT("Disable threshold frames to transition to sleeping"));
+FAutoConsoleVariableRef CVarDisableThreshold(TEXT("p.DisableThreshold2"), DisableThreshold, TEXT("Disable threshold frames to transition to sleeping"));
+
+DECLARE_CYCLE_STAT(TEXT("TPBDRigidsEvolutionGBF::AdvanceOneTimeStep"), STAT_AdvanceOneTimeStep, STATGROUP_Chaos);
+DECLARE_CYCLE_STAT(TEXT("TPBDRigidsEvolutionGBF::Integrate"), STAT_Integrate, STATGROUP_Chaos);
+DECLARE_CYCLE_STAT(TEXT("TPBDRigidsEvolutionGBF::UpdateConstraintPositionBasedState"), STAT_UpdateConstraintPositionBasedState, STATGROUP_Chaos);
+DECLARE_CYCLE_STAT(TEXT("TPBDRigidsEvolutionGBF::CreateConstraintGraph"), STAT_CreateConstraintGraph, STATGROUP_Chaos);
+DECLARE_CYCLE_STAT(TEXT("TPBDRigidsEvolutionGBF::CreateIslands"), STAT_CreateIslands, STATGROUP_Chaos);
+DECLARE_CYCLE_STAT(TEXT("TPBDRigidsEvolutionGBF::ParallelSolve"), STAT_ParallelSolve, STATGROUP_Chaos);
+DECLARE_CYCLE_STAT(TEXT("TPBDRigidsEvolutionGBF::DeactivateSleep"), STAT_DeactivateSleep, STATGROUP_Chaos);
 
 template <typename T, int d>
-void TPBDRigidsEvolutionGBF<T, d>::Integrate(const TArray<int32>& InActiveIndices, T Dt)
+void TPBDRigidsEvolutionGBF<T, d>::AdvanceOneTimeStep(T Dt)
 {
-	SCOPE_CYCLE_COUNTER(STAT_Integrate);
-	CHAOS_SCOPED_TIMER(Integrate);
-	double TimerTime = 0.0;
-	FDurationTimer Timer(TimerTime);
-	TPerParticleInitForce<T, d> InitForceRule;
-	TPerParticleEulerStepVelocity<T, d> EulerStepVelocityRule;
-	TPerParticleEtherDrag<T, d> EtherDragRule(HackLinearDrag, HackAngularDrag);
-	TPerParticlePBDEulerStep<T, d> EulerStepRule;
-	Timer.Stop();
-	UE_LOG(LogChaos, Verbose, TEXT("Init Time is %f"), TimerTime);
-
-	TimerTime = 0;
-	Timer.Start();
-	const T MaxAngularSpeed2 = HackMaxAngularVelocity * HackMaxAngularVelocity;
-	const T MaxSpeed2 = HackMaxVelocity * HackMaxVelocity;
-	PhysicsParallelFor(InActiveIndices.Num(), [&](int32 ActiveIndex) {
-		int32 Index = InActiveIndices[ActiveIndex];
-		if (ensure(!Particles.Disabled(Index) && !Particles.Sleeping(Index)))
-		{
-			//save off previous velocities
-			Particles.PreV(Index) = Particles.V(Index);
-			Particles.PreW(Index) = Particles.W(Index);
-
-			InitForceRule.Apply(Particles, Dt, Index);
-			for (FForceRule ForceRule : ForceRules)
-			{
-				ForceRule(Particles, Dt, Index);
-			}
-			EulerStepVelocityRule.Apply(Particles, Dt, Index);
-			EtherDragRule.Apply(Particles, Dt, Index);
-
-			if (HackMaxAngularVelocity >= 0.f)
-			{
-				const T AngularSpeed2 = Particles.W(Index).SizeSquared();
-				if (AngularSpeed2 > MaxAngularSpeed2)
-				{
-					Particles.W(Index) = Particles.W(Index) * (HackMaxAngularVelocity / FMath::Sqrt(AngularSpeed2));
-				}
-			}
-
-			if (HackMaxVelocity >= 0.f)
-			{
-				const T Speed2 = Particles.V(Index).SizeSquared();
-				if (Speed2 > MaxSpeed2)
-				{
-					Particles.V(Index) = Particles.V(Index) * (HackMaxVelocity / FMath::Sqrt(Speed2));
-				}
-			}
-
-
-			EulerStepRule.Apply(Particles, Dt, Index);
-		}
-	});
-	Timer.Stop();
-	UE_LOG(LogChaos, Verbose, TEXT("Per ParticleUpdate Time is %f"), TimerTime);
-}
-
-DECLARE_CYCLE_STAT(TEXT("AdvanceOneTimestep"), STAT_AdvanceOneTimeStep, STATGROUP_Chaos);
-DECLARE_CYCLE_STAT(TEXT("UpdateContactGraph"), STAT_UpdateContactGraph, STATGROUP_Chaos);
-DECLARE_CYCLE_STAT(TEXT("Apply+PushOut"), STAT_ApplyApplyPushOut, STATGROUP_Chaos);
-DECLARE_CYCLE_STAT(TEXT("ParticleUpdateVelocity"), STAT_ParticleUpdateVelocity, STATGROUP_Chaos);
-DECLARE_CYCLE_STAT(TEXT("SleepInactive"), STAT_SleepInactive, STATGROUP_Chaos);
-DECLARE_CYCLE_STAT(TEXT("ParticleUpdatePosition"), STAT_ParticleUpdatePosition, STATGROUP_Chaos);
-DECLARE_DWORD_COUNTER_STAT(TEXT("NumActiveParticles"), STAT_NumActiveParticles, STATGROUP_Chaos);
-DECLARE_DWORD_COUNTER_STAT(TEXT("NumActiveConstraints"), STAT_NumActiveConstraints, STATGROUP_Chaos);
-
-int32 SelectedParticle = 1;
-FAutoConsoleVariableRef CVarSelectedParticle(TEXT("p.SelectedParticle"), SelectedParticle, TEXT("Debug render for a specific particle"));
-int32 ShowCollisionParticles = 0;
-FAutoConsoleVariableRef CVarShowCollisionParticles(TEXT("p.ShowCollisionParticles"), ShowCollisionParticles, TEXT("Debug render the collision particles (can be very slow)"));
-int32 ShowCenterOfMass = 1;
-FAutoConsoleVariableRef CVarShowCenterOfMass(TEXT("p.ShowCenterOfMass"), ShowCenterOfMass, TEXT("Debug render of the center of mass, you will likely need wireframe mode on"));
-int32 ShowClusterConnections = 1;
-FAutoConsoleVariableRef CVarShowClusterConnections(TEXT("p.ShowClusterConnections"), ShowClusterConnections, TEXT("Debug render of the cluster connections"));
-int32 ShowBounds = 1;
-FAutoConsoleVariableRef CVarShowBounds(TEXT("p.ShowBounds"), ShowBounds, TEXT(""));
-int32 ShowLevelSet = 0;
-FAutoConsoleVariableRef CVarShowLevelSet(TEXT("p.ShowLevelSet"), ShowLevelSet, TEXT(""));
-float MaxVisualizePhiDistance = 10.f;
-FAutoConsoleVariableRef CVarMaxPhiDistance(TEXT("p.MaxVisualizePhiDistance"), MaxVisualizePhiDistance, TEXT(""));
-float CullPhiVisualizeDistance = 0.f;
-FAutoConsoleVariableRef CVarCullPhiDistance(TEXT("p.CullPhiVisualizeDistance"), CullPhiVisualizeDistance, TEXT(""));
-
-int32 GatherVerbosePhysicsStats = 0;
-FAutoConsoleVariableRef CVarGatherVerbosePhysicsStats(TEXT("p.GatherVerbosePhysicsStats"), GatherVerbosePhysicsStats, TEXT("If enabled, stat ChaosDedicated will show detailed stats that are more expensive to gather"));
+	SCOPE_CYCLE_COUNTER(STAT_AdvanceOneTimeStep);
 
-template<class T, int d>
-void TPBDRigidsEvolutionGBF<T, d>::AdvanceOneTimeStep(const T Dt)
-{
-	if (DisableSim)
 	{
-		return;
+		SCOPE_CYCLE_COUNTER(STAT_Integrate);
+		Integrate(Particles.GetNonDisabledDynamicView(), Dt);	//Question: should we use an awake view?
 	}
-	SCOPE_CYCLE_COUNTER(STAT_AdvanceOneTimeStep);
-
-	UE_LOG(LogChaos, Verbose, TEXT("START FRAME with Dt %f"), Dt);
-	Integrate(GetActiveIndicesArray(), Dt);
-	GetDebugSubstep().Add(TEXT("TPBDRigidsEvolutionGBF::AdvanceOneTimeStep(): After Integrate()"));
-
-	SET_DWORD_STAT(STAT_NumActiveParticles, ActiveIndices.Num());
-
-	UpdateConstraintPositionBasedState(Dt);
-	CreateConstraintGraph();
-	CreateIslands();
 
+	{
+		SCOPE_CYCLE_COUNTER(STAT_UpdateConstraintPositionBasedState);
+		UpdateConstraintPositionBasedState(Dt);
+	}
+	{
+		SCOPE_CYCLE_COUNTER(STAT_CreateConstraintGraph);
+		CreateConstraintGraph();
+	}
+	{
+		SCOPE_CYCLE_COUNTER(STAT_CreateIslands);
+		CreateIslands();
+	}
+	
 	TArray<bool> SleepedIslands;
 	SleepedIslands.SetNum(ConstraintGraph.NumIslands());
-	TArray<TArray<int32>> DisabledParticles;
+	TArray<TArray<TPBDRigidParticleHandle<T,d>*>> DisabledParticles;
 	DisabledParticles.SetNum(ConstraintGraph.NumIslands());
+	SleepedIslands.SetNum(ConstraintGraph.NumIslands());
 	{
-		SCOPE_CYCLE_COUNTER(STAT_ApplyApplyPushOut);
-		CHAOS_SCOPED_TIMER(ApplyApplyPushOut);
+		SCOPE_CYCLE_COUNTER(STAT_ParallelSolve);
 		PhysicsParallelFor(ConstraintGraph.NumIslands(), [&](int32 Island) {
-			const TArray<int32>& IslandParticleIndices = ConstraintGraph.GetIslandParticles(Island);
+			const TArray<TGeometryParticleHandle<T, d>*>& IslandParticles = ConstraintGraph.GetIslandParticles(Island);
 
 			ApplyConstraints(Dt, Island);
 
-			UpdateVelocities(Dt, Island);
-			
+			if (Dt > 0)
+			{
+				UpdateVelocities(Dt, Island);
+			}
+
 			ApplyPushOut(Dt, Island);
 
-			for (const int32 Index : IslandParticleIndices)
+			for (auto Particle : IslandParticles)
 			{
 				// If a dynamic particle is moving slowly enough for long enough, disable it.
 				// @todo(mlentine): Find a good way of not doing this when we aren't using this functionality
 
 				// increment the disable count for the particle
-				if (Particles.ObjectState(Index) != EObjectStateType::Kinematic && Particles.ObjectState(Index) != EObjectStateType::Static && PhysicsMaterials[Index] && Particles.V(Index).SizeSquared() < PhysicsMaterials[Index]->DisabledLinearThreshold && Particles.W(Index).SizeSquared() < PhysicsMaterials[Index]->DisabledAngularThreshold)
+				if (auto PBDRigid = Particle->AsDynamic())
 				{
-					++ParticleDisableCount[Index];
-				}
+					if (PBDRigid->AuxilaryValue(PhysicsMaterials) && PBDRigid->V().SizeSquared() < PBDRigid->AuxilaryValue(PhysicsMaterials)->DisabledLinearThreshold &&
+						PBDRigid->W().SizeSquared() < PBDRigid->AuxilaryValue(PhysicsMaterials)->DisabledAngularThreshold)
+					{
+						++PBDRigid->AuxilaryValue(ParticleDisableCount);
+					}
 
-				// check if we're over the disable count threshold
-				if (ParticleDisableCount[Index] > DisableThreshold)
-				{
-					ParticleDisableCount[Index] = 0;
-					Particles.SetDisabledLowLevel(Index, true);
-					DisabledParticles[Island].Add(Index);
-					Particles.V(Index) = TVector<T, d>(0);
-					Particles.W(Index) = TVector<T, d>(0);
-				}
+					// check if we're over the disable count threshold
+					if (PBDRigid->AuxilaryValue(ParticleDisableCount) > DisableThreshold)
+					{
+						PBDRigid->AuxilaryValue(ParticleDisableCount) = 0;
+						//Particles.Disabled(Index) = true;
+						DisabledParticles[Island].Add(PBDRigid);
+						//Particles.V(Index) = TVector<T, d>(0);
+						//Particles.W(Index) = TVector<T, d>(0);
+					}
 
-				if (!(ensure(!FMath::IsNaN(Particles.P(Index)[0])) && ensure(!FMath::IsNaN(Particles.P(Index)[1])) && ensure(!FMath::IsNaN(Particles.P(Index)[2]))))
-				{
-					Particles.SetDisabledLowLevel(Index, true);
-					DisabledParticles[Island].Add(Index);
+					if (!(ensure(!FMath::IsNaN(PBDRigid->P()[0])) && ensure(!FMath::IsNaN(PBDRigid->P()[1])) && ensure(!FMath::IsNaN(PBDRigid->P()[2]))))
+					{
+						//Particles.Disabled(Index) = true;
+						DisabledParticles[Island].Add(PBDRigid);
+					}
 				}
 			}
 
 			// Turn off if not moving
-			SleepedIslands[Island] = ConstraintGraph.SleepInactive(Particles, Island, PhysicsMaterials);
+			SleepedIslands[Island] = ConstraintGraph.SleepInactive(Island, PhysicsMaterials);
 		});
-
-		GatherStats();
-
+	}
+	{
+		SCOPE_CYCLE_COUNTER(STAT_DeactivateSleep);
 		for (int32 Island = 0; Island < ConstraintGraph.NumIslands(); ++Island)
 		{
 			if (SleepedIslands[Island])
 			{
-				for (const int32 Index : ConstraintGraph.GetIslandParticles(Island))
-				{
-					ActiveIndices.Remove(Index);
-				}
+				Particles.DeactivateParticles(ConstraintGraph.GetIslandParticles(Island));
 			}
-			for (const int32 Index : DisabledParticles[Island])
+			for (const auto Particle : DisabledParticles[Island])
 			{
-				ActiveIndices.Remove(Index);
-				NonDisabledIndices.Remove(Index);
+				Particles.DisableParticle(Particle);
 			}
 		}
-
-		GetDebugSubstep().Add(TEXT("TPBDRigidsEvolutionGBF::AdvanceOneTimeStep(): Before AdvanceClustering"));
-		Clustering.AdvanceClustering(Dt, GetCollisionConstraints());
-		GetDebugSubstep().Add(TEXT("TPBDRigidsEvolutionGBF::AdvanceOneTimeStep(): After AdvanceClustering"));
-
-		{
-			SCOPE_CYCLE_COUNTER(STAT_ParticleUpdatePosition);
-			ParticleUpdatePosition(Particles, Dt);
-		}
-
-		Time += Dt;
 	}
-}
 
-template<class T, int d>
-void TPBDRigidsEvolutionGBF<T, d>::GatherStats()
-{
-	EvolutionStats.Reset();
-	if (GatherVerbosePhysicsStats)
-	{
-		auto GatherLambda = [&](const TArray<int32>& Indices, const TArray<int32>& Count, int32& NumCollisionParticles, int32& NumShapes)
-		{
-			for (int32 i = 0; i < Indices.Num(); ++i)
-			{
-				const int32 Index = Indices[i];
-				if (Particles.Geometry(Index))
-				{
-					if (Particles.Geometry(Index)->IsUnderlyingUnion())
-					{
-						const bool bUseSubCollision = Particles.CollisionParticlesSize(Index) == 0;
-						const TImplicitObjectUnion<T, d>* Union = static_cast<const TImplicitObjectUnion<T, d>*>(Particles.Geometry(Index).Get());
-						const TArray<TUniquePtr<TImplicitObject<T, d>>>& SubObjects = Union->GetObjects();
-						NumShapes += SubObjects.Num();
-						if (bUseSubCollision)
-						{
-							for (const TUniquePtr <TImplicitObject<T, d>>& Obj : SubObjects)
-							{
-								if (ensure(Obj->GetType() == ImplicitObjectType::Transformed))
-								{
-									const TImplicitObjectTransformed<T, d>* Transformed = static_cast<const TImplicitObjectTransformed<T, d>*>(Obj.Get());
+	Clustering.AdvanceClustering(Dt, GetCollisionConstraints());
 
-									const int32 OriginalIdx = Union->MCollisionParticleLookupHack[Transformed->GetTransformedObject()];
-									NumCollisionParticles += Particles.CollisionParticlesSize(OriginalIdx) * Count[i];
-								}
-							}
-						}
-						else
-						{
-							NumCollisionParticles += Particles.CollisionParticlesSize(Index) * Count[i];
-						}
-					}
-					else
-					{
-						NumShapes++;
-						NumCollisionParticles += Particles.CollisionParticlesSize(Index) * Count[i];
-					}
-				}
-				else
-				{
-					NumCollisionParticles += Particles.CollisionParticlesSize(Index) * Count[i];
-				}
-			}
-		};
-
-		const TArray<int32>& ActiveIndicesArray = GetActiveIndicesArray();
-		TArray<int32> Ones;
-		Ones.Init(1, ActiveIndicesArray.Num());
-		GatherLambda(ActiveIndicesArray, Ones, EvolutionStats.ActiveCollisionPoints, EvolutionStats.ActiveShapes);
-	}
+	ParticleUpdatePosition(Particles.GetActiveParticlesView(), Dt);
+	CollisionConstraints.SwapSpatialAcceleration();	//TODO_SQ_IMPLEMENTATION: there's a better place for this
 }
 
-template<class T, int d>
-void TPBDRigidsEvolutionGBF<T, d>::DebugDraw()
+template <typename T, int d>
+TPBDRigidsEvolutionGBF<T, d>::TPBDRigidsEvolutionGBF(TPBDRigidsSOAs<T, d>& InParticles, int32 InNumIterations)
+	: Base(InParticles, InNumIterations)
+	, CollisionConstraints(InParticles, Collided, PhysicsMaterials, DefaultNumPushOutPairIterations)
+	, CollisionRule(CollisionConstraints, DefaultNumPushOutIterations)
 {
-#if CHAOS_DEBUG_DRAW
-	if (FDebugDrawQueue::IsDebugDrawingEnabled())
+	SetParticleUpdateVelocityFunction([PBDUpdateRule = TPerParticlePBDUpdateFromDeltaPosition<float, 3>(), this](const TParticleView<TPBDRigidParticles<T, d>>& ParticlesInput, const T Dt) {
+		ParticlesInput.ParallelFor([&](auto& Particle, int32 Index) {
+			PBDUpdateRule.Apply(Particle, Dt);
+		});
+	});
+
+	SetParticleUpdatePositionFunction([this](const TParticleView<TPBDRigidParticles<T, d>>& ParticlesInput, const T Dt)
 	{
-		if (1)
+		ParticlesInput.ParallelFor([&](auto& Particle, int32 Index)
 		{
-			if (ShowClusterConnections)
-			{
-				for (const typename FRigidClustering::FClusterMap::ElementType& Cluster : Clustering.GetChildrenMap())
-				{
-					if (!Particles.Disabled(Cluster.Key))
-					{
-						for (const int32 ChildIdx : *Cluster.Value)
-						{
-							const FString Text = FString::Printf(TEXT("%d"), ChildIdx);
-							for (const TConnectivityEdge<T>& Edge : Clustering.GetConnectivityEdges()[ChildIdx])
-							{
-								FDebugDrawQueue::GetInstance().DrawDebugLine(Particles.X(ChildIdx), Particles.X(Edge.Sibling), FColor::Blue, false, 1e-4, 0, 2.f);
-							}
-						}
-					}
-				}
-			}
-
-			for (uint32 Idx = 0; Idx < Particles.Size(); ++Idx)
-			{
-				if (Particles.Disabled(Idx)) { continue; }
-				if (ShowCollisionParticles && (SelectedParticle == Idx || ShowCollisionParticles == -1))
-				{
-					if (Particles.CollisionParticles(Idx))
-					{
-						bool bDrawUnionCollision = Particles.Geometry(Idx) && Particles.Geometry(Idx)->IsUnderlyingUnion() && Particles.CollisionParticlesSize(Idx) == 0;
-						if (bDrawUnionCollision)
-						{
-							const TImplicitObjectUnion<T, d>* Union = Particles.Geometry(Idx)->template GetObject<const TImplicitObjectUnion<T, d>>();
-							const TArray<TUniquePtr<TImplicitObject<T, d>>>& SubObjects = Union->GetObjects();
-							for (const TUniquePtr <TImplicitObject<T, d>>& Obj : SubObjects)
-							{
-								if (ensure(Obj->GetType() == ImplicitObjectType::Transformed))
-								{
-									const TImplicitObjectTransformed<T, d>* Transformed = static_cast<const TImplicitObjectTransformed<T, d>*>(Obj.Get());
-
-									const int32 OriginalIdx = Union->MCollisionParticleLookupHack[Transformed->GetTransformedObject()];
-									for (uint32 CollisionIdx = 0; CollisionIdx < Particles.CollisionParticles(OriginalIdx)->Size(); ++CollisionIdx)
-									{
-										const TVector<T, d>& X = Particles.CollisionParticles(OriginalIdx)->X(CollisionIdx);
-										const TVector<T, d> WorldX = (Transformed->GetTransform()* TRigidTransform<T, d>(Particles.X(Idx), Particles.R(Idx))).TransformPosition(X);
-										FDebugDrawQueue::GetInstance().DrawDebugPoint(WorldX, FColor::Purple, false, 1e-4, 0, 10.f);
-									}
-								}
-							}
-						}
-						else
-						{
-							for (uint32 CollisionIdx = 0; CollisionIdx < Particles.CollisionParticles(Idx)->Size(); ++CollisionIdx)
-							{
-								const TVector<T, d>& X = Particles.CollisionParticles(Idx)->X(CollisionIdx);
-								const TVector<T, d> WorldX = TRigidTransform<T, d>(Particles.X(Idx), Particles.R(Idx)).TransformPosition(X);
-								FDebugDrawQueue::GetInstance().DrawDebugPoint(WorldX, FColor::Purple, false, 1e-4, 0, 10.f);
-							}
-						}
-					}
-				}
-
-				if (ShowCenterOfMass && (SelectedParticle == Idx || ShowCenterOfMass == -1))
-				{
-					FColor AxisColors[] = { FColor::Red, FColor::Green, FColor::Blue };
-					T MaxInertia = KINDA_SMALL_NUMBER;
-					for (int i = 0; i < d; ++i)
-					{
-						MaxInertia = FMath::Max(Particles.I(Idx).M[i][i], MaxInertia);
-					}
-
-					for (int i = 0; i < d; ++i)
-					{
-						const TVector<T, d> WorldDir = Particles.R(Idx) * TVector<T, d>::AxisVector(i) * 100 * Particles.I(Idx).M[i][i] / MaxInertia;
-						FDebugDrawQueue::GetInstance().DrawDebugDirectionalArrow(Particles.X(Idx), Particles.X(Idx) + WorldDir, 3, AxisColors[i], false, 1e-4, 0, 2.f);
-
-					}
-					FDebugDrawQueue::GetInstance().DrawDebugSphere(Particles.X(Idx), 20.f, 16, FColor::Yellow, false, 1e-4);
-				}
-
-				if (ShowBounds && (SelectedParticle == Idx || ShowBounds == -1) && Particles.Geometry(Idx)->HasBoundingBox())
-				{
-					const TBox<T, d>& Bounds = Particles.Geometry(Idx)->BoundingBox();
-					const TRigidTransform<T, d> TM(Particles.X(Idx), Particles.R(Idx));
-					const TVector<T, d> Center = TM.TransformPosition(Bounds.Center());
-					FDebugDrawQueue::GetInstance().DrawDebugBox(Center, Bounds.Extents() * 0.5f, TM.GetRotation(), FColor::Yellow, false, 1e-4, 0, 2.f);
-				}
-
-				if (ShowLevelSet && (SelectedParticle == Idx || ShowLevelSet == -1))
-				{
-					auto RenderLevelSet = [](const TRigidTransform<T, d>& LevelSetToWorld, const TLevelSet<T, d>& LevelSet)
-					{
-						const TUniformGrid<T, d>& Grid = LevelSet.GetGrid();
-						const int32 NumCells = Grid.GetNumCells();
-						const TArrayND<T, 3>& PhiArray = LevelSet.GetPhiArray();
-						for (int32 CellIdx = 0; CellIdx < NumCells; ++CellIdx)
-						{
-							const TVector<T, d> GridSpaceLocation = Grid.Center(CellIdx);
-							const TVector<T, d> WorldSpaceLocation = LevelSetToWorld.TransformPosition(GridSpaceLocation);
-							const T Phi = PhiArray(Grid.GetIndex(CellIdx));
-							//FDebugDrawQueue::GetInstance().DrawDebugPoint(WorldSpaceLocation, Phi > 0 ? FColor::Purple : FColor::Green, false, 1e-4, 0, 10.f);
-							//FDebugDrawQueue::GetInstance().DrawDebugSphere(WorldSpaceLocation, Grid.Dx().GetAbsMin(), 16, Phi > 0 ? FColor::Purple : FColor::Green, false, 1e-4);
-							if (Phi <= CullPhiVisualizeDistance)
-							{
-								const T LocalPhi = Phi - CullPhiVisualizeDistance;
-								const T MaxPhi = (-LocalPhi / MaxVisualizePhiDistance) * 255;
-								const uint8 MaxPhiInt = MaxPhi > 255 ? 255 : MaxPhi;
-								FDebugDrawQueue::GetInstance().DrawDebugPoint(WorldSpaceLocation, FColor(255, MaxPhiInt, 255, 255), false, 1e-4, 0, 30.f);
-							}
-						}
-					};
+			Particle.X() = Particle.P();
+			Particle.R() = Particle.Q();
+		});
+	});
 
-					if (const TLevelSet<T, d>* LevelSet = Particles.Geometry(Idx)->template GetObject<TLevelSet<T, d>>())
-					{
-						RenderLevelSet(TRigidTransform<T, d>(Particles.X(Idx), Particles.R(Idx)), *LevelSet);
-					}
-					else if (const TImplicitObjectTransformed<T, d>* Transformed = Particles.Geometry(Idx)->template GetObject<TImplicitObjectTransformed<T, d>>())
-					{
-						if (const TLevelSet<T, d>* InnerLevelSet = Transformed->GetTransformedObject()->template GetObject<TLevelSet<T, d>>())
-						{
-							RenderLevelSet(Transformed->GetTransform() * TRigidTransform<T, d>(Particles.X(Idx), Particles.R(Idx)), *InnerLevelSet);
-						}
-					}
-				}
-			}
-		}
-		FDebugDrawQueue::GetInstance().Flush();
-	}
-#endif
+	AddConstraintRule(&CollisionRule);
 }
 
-template class Chaos::TPBDRigidsEvolutionGBF<float, 3>;
+}
 
-#undef LOCTEXT_NAMESPACE
+template class Chaos::TPBDRigidsEvolutionGBF<float, 3>;
\ No newline at end of file
diff --git a/Engine/Source/Runtime/Experimental/Chaos/Private/Chaos/PBDRigidsEvolutionPGS.cpp b/Engine/Source/Runtime/Experimental/Chaos/Private/Chaos/PBDRigidsEvolutionPGS.cpp
--- a/Engine/Source/Runtime/Experimental/Chaos/Private/Chaos/PBDRigidsEvolutionPGS.cpp
+++ b/Engine/Source/Runtime/Experimental/Chaos/Private/Chaos/PBDRigidsEvolutionPGS.cpp
@@ -25,6 +25,7 @@
 
 using namespace Chaos;
 
+#if CHAOS_PARTICLEHANDLE_TODO
 template<class T, int d>
 TPBDRigidsEvolutionPGS<T, d>::TPBDRigidsEvolutionPGS(TPBDRigidParticles<T, d>&& InParticles, int32 NumIterations)
     : Base(MoveTemp(InParticles), NumIterations)
@@ -180,5 +181,6 @@ void TPBDRigidsEvolutionPGS<T, d>::AdvanceOneTimeStep(const T Dt)
 }
 
 template class Chaos::TPBDRigidsEvolutionPGS<float, 3>;
+#endif
 
 #undef LOCTEXT_NAMESPACE
diff --git a/Engine/Source/Runtime/Experimental/Chaos/Private/Chaos/ParticleHandle.cpp b/Engine/Source/Runtime/Experimental/Chaos/Private/Chaos/ParticleHandle.cpp
new file mode 100644
--- /dev/null
+++ b/Engine/Source/Runtime/Experimental/Chaos/Private/Chaos/ParticleHandle.cpp
@@ -0,0 +1,17 @@
+// Copyright 1998-2019 Epic Games, Inc. All Rights Reserved.
+
+#include "Chaos/ParticleHandle.h"
+
+namespace Chaos
+{
+
+	template class CHAOS_API TGeometryParticleData<float, 3>;
+	template class CHAOS_API TGeometryParticle<float, 3>;
+
+	template class CHAOS_API TKinematicGeometryParticleData<float, 3>;
+	template class CHAOS_API TKinematicGeometryParticle<float, 3>;
+
+	template class CHAOS_API TPBDRigidParticleData<float, 3>;
+	template class CHAOS_API TPBDRigidParticle<float, 3>;
+
+}
diff --git a/Engine/Source/Runtime/Experimental/Chaos/Private/Chaos/RigidParticles.cpp b/Engine/Source/Runtime/Experimental/Chaos/Private/Chaos/RigidParticles.cpp
--- a/Engine/Source/Runtime/Experimental/Chaos/Private/Chaos/RigidParticles.cpp
+++ b/Engine/Source/Runtime/Experimental/Chaos/Private/Chaos/RigidParticles.cpp
@@ -4,28 +4,20 @@
 using namespace Chaos;
 
 //Note this has to be in the cpp to avoid allocating/freeing across DLLs
+
 template<typename T, int d>
-void TRigidParticles<T, d>::CollisionParticlesInitIfNeeded(const int32 Index, TArray<TVector<T, d>>* Points)
+void TRigidParticles<T, d>::CollisionParticlesInitIfNeeded(const int32 Index)
 {
 	if (MCollisionParticles[Index] == nullptr)
 	{
-		if (Points)
-		{
-			MCollisionParticles[Index] = MakeUnique<TBVHParticles<T, d>>(*Points);
-		}
-		else
-		{
-			MCollisionParticles[Index] = MakeUnique<TBVHParticles<T, d>>();
-		}
-	}
-	else
-	{
-		// Make sure that if we were supposed to construct particles with a view to Points
-		// that the existing particles already use the Points array.  Otherwise, maybe we
-		// should make a copy?  Not sure...  At this point, we'll leave it an error so that
-		// folks clean up their logic.
-		check(!Points || MCollisionParticles[Index]->X().GetData() == Points->GetData());
+		MCollisionParticles[Index] = MakeUnique<TBVHParticles<T, d>>();
 	}
 }
 
+template<typename T, int d>
+void TRigidParticles<T, d>::SetCollisionParticles(const int32 Index, TParticles<T, d>&& Points)
+{
+	MCollisionParticles[Index] = MakeUnique<TBVHParticles<T, d>>(MoveTemp(Points));
+}
+
 template class Chaos::TRigidParticles<float, 3>;
\ No newline at end of file
diff --git a/Engine/Source/Runtime/Experimental/Chaos/Private/Chaos/TriangleMeshImplicitObject.cpp b/Engine/Source/Runtime/Experimental/Chaos/Private/Chaos/TriangleMeshImplicitObject.cpp
--- a/Engine/Source/Runtime/Experimental/Chaos/Private/Chaos/TriangleMeshImplicitObject.cpp
+++ b/Engine/Source/Runtime/Experimental/Chaos/Private/Chaos/TriangleMeshImplicitObject.cpp
@@ -13,35 +13,15 @@ TTriangleMeshImplicitObject<T>::TTriangleMeshImplicitObject(TParticles<T, 3>&& P
 	, MParticles(MoveTemp(Particles))
 	, MElements(MoveTemp(Elements))
 	, MLocalBoundingBox(MParticles.X(0), MParticles.X(0))
-	, BVH(ConvexHulls)
+	, BVH(MakeParticleView(&ConvexHulls))
 {
 	for (uint32 Idx = 1; Idx < MParticles.Size(); ++Idx)
 	{
 		MLocalBoundingBox.GrowToInclude(MParticles.X(Idx));
 	}
 
-	ConvexHulls.AddParticles(MElements.Num());
-
-	for (int32 TriIdx = 0; TriIdx < MElements.Num(); ++TriIdx)
-	{
-		const TVector<T, 3>& A = MParticles.X(MElements[TriIdx][0]);
-		const TVector<T, 3>& B = MParticles.X(MElements[TriIdx][1]);
-		const TVector<T, 3>& C = MParticles.X(MElements[TriIdx][2]);
-
-		TParticles<T, 3> Pts;
-		Pts.AddParticles(3);
-		Pts.X(0) = A;
-		Pts.X(1) = B;
-		Pts.X(2) = C;
-
-		TUniquePtr<TConvex<T, 3>> TriConvex = MakeUnique<TConvex<T,3>>(Pts);
-		ConvexHulls.X(TriIdx) = TVector<T,3>(0);
-		ConvexHulls.R(TriIdx) = TRotation<T, 3>::Identity;
-		ConvexHulls.SetDynamicGeometry(TriIdx, MoveTemp(TriConvex));
-	}
-
-	//BVH = BVHType(ConvexHulls, 2);
-	BVH = BVHType(ConvexHulls, false, 0, 64);
+	RebuildConvexHulls();
+	RebuildBV();
 }
 
 template <typename T>
@@ -94,7 +74,8 @@ struct TTriangleMeshRaycastVisitor
 		T Time;
 
 		//Check if we even intersect with triangle plane
-		if (TriPlane.Raycast(StartPoint, Dir, CurLength, Thickness, Time, RaycastPosition, RaycastNormal))
+		int32 DummyFaceIndex;
+		if (TriPlane.Raycast(StartPoint, Dir, CurLength, Thickness, Time, RaycastPosition, RaycastNormal, DummyFaceIndex))
 		{
 			TVector<T, 3> IntersectionPosition = RaycastPosition;
 			TVector<T, 3> IntersectionNormal = RaycastNormal;
@@ -107,6 +88,7 @@ struct TTriangleMeshRaycastVisitor
 				if (DistToTriangle2 <= R2)
 				{
 					OutTime = 0;
+					OutFaceIndex = TriIdx;
 					return false; //no one will beat Time == 0
 				}
 			}
@@ -126,13 +108,13 @@ struct TTriangleMeshRaycastVisitor
 				bool bBorderIntersections[3];
 
 				const TCapsule<T> ABCapsule(A, B, Thickness);
-				bBorderIntersections[0] = ABCapsule.Raycast(StartPoint, Dir, CurLength, 0, BorderTimes[0], BorderPositions[0], BorderNormals[0]);
+				bBorderIntersections[0] = ABCapsule.Raycast(StartPoint, Dir, CurLength, 0, BorderTimes[0], BorderPositions[0], BorderNormals[0], DummyFaceIndex);
 
 				const TCapsule<T> BCCapsule(B, C, Thickness);
-				bBorderIntersections[1] = BCCapsule.Raycast(StartPoint, Dir, CurLength, 0, BorderTimes[1], BorderPositions[1], BorderNormals[1]);
+				bBorderIntersections[1] = BCCapsule.Raycast(StartPoint, Dir, CurLength, 0, BorderTimes[1], BorderPositions[1], BorderNormals[1], DummyFaceIndex);
 
 				const TCapsule<T> ACCapsule(A, C, Thickness);
-				bBorderIntersections[2] = ACCapsule.Raycast(StartPoint, Dir, CurLength, 0, BorderTimes[2], BorderPositions[2], BorderNormals[2]);
+				bBorderIntersections[2] = ACCapsule.Raycast(StartPoint, Dir, CurLength, 0, BorderTimes[2], BorderPositions[2], BorderNormals[2], DummyFaceIndex);
 
 				int32 MinBorderIdx = INDEX_NONE;
 				T MinBorderTime = 0;	//initialization not needed, but fixes warning
@@ -175,6 +157,7 @@ struct TTriangleMeshRaycastVisitor
 					OutNormal = RaycastNormal;	//We use the plane normal even when hitting triangle edges. This is to deal with triangles that approximate a single flat surface.
 					OutTime = Time;
 					CurLength = Time;	//prevent future rays from going any farther
+					OutFaceIndex = TriIdx;
 				}
 			}
 		}
@@ -200,10 +183,11 @@ struct TTriangleMeshRaycastVisitor
 	T OutTime;
 	TVector<T, 3> OutPosition;
 	TVector<T, 3> OutNormal;
+	int32 OutFaceIndex;
 };
 
 template <typename T>
-bool TTriangleMeshImplicitObject<T>::Raycast(const TVector<T, 3>& StartPoint, const TVector<T, 3>& Dir, const T Length, const T Thickness, T& OutTime, TVector<T, 3>& OutPosition, TVector<T, 3>& OutNormal) const
+bool TTriangleMeshImplicitObject<T>::Raycast(const TVector<T, 3>& StartPoint, const TVector<T, 3>& Dir, const T Length, const T Thickness, T& OutTime, TVector<T, 3>& OutPosition, TVector<T, 3>& OutNormal, int32& OutFaceIndex) const
 {
 	TTriangleMeshRaycastVisitor<T> SQVisitor(StartPoint, Dir, Thickness, MParticles, MElements);
 
@@ -221,6 +205,7 @@ bool TTriangleMeshImplicitObject<T>::Raycast(const TVector<T, 3>& StartPoint, co
 		OutTime = SQVisitor.OutTime;
 		OutPosition = SQVisitor.OutPosition;
 		OutNormal = SQVisitor.OutNormal;
+		OutFaceIndex = SQVisitor.OutFaceIndex;
 		return true;
 	}
 	else
@@ -233,7 +218,7 @@ template <typename T>
 bool TTriangleMeshImplicitObject<T>::Overlap(const TVector<T, 3>& Point, const T Thickness) const
 {
 	TBox<T, 3> QueryBounds(Point, Point);
-	QueryBounds.Thicken(TVector<T, 3>(Thickness));
+	QueryBounds.Thicken(Thickness);
 	const TArray<int32> PotentialIntersections = BVH.FindAllIntersections(QueryBounds);
 
 	const T Epsilon = 1e-4;
@@ -272,7 +257,7 @@ bool TTriangleMeshImplicitObject<T>::OverlapGeom(const TImplicitObject<T, 3>& Qu
 {
 	bool bResult = false;
 	TBox<T, 3> QueryBounds = QueryGeom.BoundingBox();
-	QueryBounds.Thicken(TVector<T, 3>(Thickness));
+	QueryBounds.Thicken(Thickness);
 	QueryBounds = QueryBounds.TransformedBox(QueryTM);
 	const TArray<int32> PotentialIntersections = BVH.FindAllIntersections(QueryBounds);
 
@@ -317,7 +302,7 @@ struct TTriangleMeshSweepVisitor
 		T Time;
 		TVector<T, 3> HitPosition;
 		TVector<T, 3> HitNormal;
-		if (GJKRaycast<T>(TriConvex, QueryGeom, StartTM, StartTM.GetLocation(), Dir, CurLength, Time, HitPosition, HitNormal, Thickness))
+		if (GJKRaycast<T>(TriConvex, QueryGeom, StartTM, Dir, CurLength, Time, HitPosition, HitNormal, Thickness))
 		{
 			if (Time < OutTime)
 			{
@@ -325,6 +310,7 @@ struct TTriangleMeshSweepVisitor
 				OutPosition = HitPosition;
 				OutTime = Time;
 				CurLength = Time;
+				OutFaceIndex = TriIdx;
 
 				if (Time == 0)
 				{
@@ -347,10 +333,11 @@ struct TTriangleMeshSweepVisitor
 	T OutTime;
 	TVector<T, 3> OutPosition;
 	TVector<T, 3> OutNormal;
+	int32 OutFaceIndex;
 };
 
 template <typename T>
-bool TTriangleMeshImplicitObject<T>::SweepGeom(const TImplicitObject<T, 3>& QueryGeom, const TRigidTransform<T, 3>& StartTM, const TVector<T, 3>& Dir, const T Length, T& OutTime, TVector<T, 3>& OutPosition, TVector<T, 3>& OutNormal, const T Thickness, const TVector<T,3> Scale) const
+bool TTriangleMeshImplicitObject<T>::SweepGeom(const TImplicitObject<T, 3>& QueryGeom, const TRigidTransform<T, 3>& StartTM, const TVector<T, 3>& Dir, const T Length, T& OutTime, TVector<T, 3>& OutPosition, TVector<T, 3>& OutNormal, int32& OutFaceIndex, const T Thickness, const TVector<T,3> Scale) const
 {
 	bool bHit = false;
 	TTriangleMeshSweepVisitor<T> SQVisitor(*this, QueryGeom, StartTM, Dir, Thickness, Scale);
@@ -364,11 +351,49 @@ bool TTriangleMeshImplicitObject<T>::SweepGeom(const TImplicitObject<T, 3>& Quer
 		OutTime = SQVisitor.OutTime;
 		OutPosition = SQVisitor.OutPosition;
 		OutNormal = SQVisitor.OutNormal;
+		OutFaceIndex = SQVisitor.OutFaceIndex;
 		bHit = true;
 	}
 	return bHit;
 }
 
+template <typename T>
+int32 TTriangleMeshImplicitObject<T>::FindMostOpposingFace(const TVector<T, 3>& Position, const TVector<T, 3>& UnitDir, int32 HintFaceIndex) const
+{
+	ensure(HintFaceIndex != INDEX_NONE);
+	return HintFaceIndex;	//for now we just return what they gave us. todo: use adjacency information to compute this
+}
+
+
+template<typename T>
+void Chaos::TTriangleMeshImplicitObject<T>::RebuildBV()
+{
+	BVH = BVHType(MakeParticleView(&ConvexHulls), false, 0, 64);
+}
+
+template <typename T>
+void Chaos::TTriangleMeshImplicitObject<T>::RebuildConvexHulls()
+{
+	ConvexHulls.Resize(MElements.Num());
+
+	for (int32 TriIdx = 0; TriIdx < MElements.Num(); ++TriIdx)
+	{
+		const TVector<T, 3>& A = MParticles.X(MElements[TriIdx][0]);
+		const TVector<T, 3>& B = MParticles.X(MElements[TriIdx][1]);
+		const TVector<T, 3>& C = MParticles.X(MElements[TriIdx][2]);
+
+		TParticles<T, 3> Pts;
+		Pts.AddParticles(3);
+		Pts.X(0) = A;
+		Pts.X(1) = B;
+		Pts.X(2) = C;
+
+		TUniquePtr<TConvex<T, 3>> TriConvex = MakeUnique<TConvex<T, 3>>(Pts);
+		ConvexHulls.X(TriIdx) = TVector<T, 3>(0);
+		ConvexHulls.R(TriIdx) = TRotation<T, 3>::Identity;
+		ConvexHulls.SetDynamicGeometry(TriIdx, MoveTemp(TriConvex));
+	}
+}
 
 }
 
diff --git a/Engine/Source/Runtime/Experimental/Chaos/Private/ChaosModule.cpp b/Engine/Source/Runtime/Experimental/Chaos/Private/ChaosModule.cpp
--- a/Engine/Source/Runtime/Experimental/Chaos/Private/ChaosModule.cpp
+++ b/Engine/Source/Runtime/Experimental/Chaos/Private/ChaosModule.cpp
@@ -4,19 +4,12 @@
 #include "CoreMinimal.h"
 #include "Modules/ModuleManager.h"
 
-#if INCLUDE_CHAOS
-
-#include "ChaosLog.h"
-
-DEFINE_LOG_CATEGORY(LogChaosGeneral);
-DEFINE_LOG_CATEGORY(LogChaosThread);
-DEFINE_LOG_CATEGORY(LogChaosSimulation);
-DEFINE_LOG_CATEGORY(LogChaosDebug);
+//#if INCLUDE_CHAOS
 
 IMPLEMENT_MODULE(FChaosEngineModule, Chaos);
 
-#else
-
-IMPLEMENT_MODULE(FDefaultModuleImpl, Chaos);
-
-#endif
+//#else
+//
+//IMPLEMENT_MODULE(FDefaultModuleImpl, Chaos);
+//
+//#endif
diff --git a/Engine/Source/Runtime/Experimental/Chaos/Private/ChaosStats.cpp b/Engine/Source/Runtime/Experimental/Chaos/Private/ChaosStats.cpp
--- a/Engine/Source/Runtime/Experimental/Chaos/Private/ChaosStats.cpp
+++ b/Engine/Source/Runtime/Experimental/Chaos/Private/ChaosStats.cpp
@@ -28,7 +28,7 @@ DEFINE_STAT(STAT_SkelMeshUpdateAnim);
 DEFINE_STAT(STAT_DispatchEventNotifies);
 DEFINE_STAT(STAT_DispatchCollisionEvents);
 DEFINE_STAT(STAT_DispatchBreakEvents);
-DEFINE_STAT(STAT_CacheResults);
+DEFINE_STAT(STAT_BufferPhysicsResults);
 DEFINE_STAT(STAT_FlipResults);
 DEFINE_STAT(STAT_CacheResultGeomCollection);
 DEFINE_STAT(STAT_CacheResultStaticMesh);
@@ -42,6 +42,7 @@ DEFINE_STAT(STAT_DisableCollisions);
 DEFINE_STAT(STAT_EvolutionAndKinematicUpdate);
 DEFINE_STAT(STAT_AdvanceEventWaits);
 DEFINE_STAT(STAT_ResetCollisionRule);
+DEFINE_STAT(STAT_EventDataGathering)
 
 DEFINE_STAT(STAT_ParamUpdateObject);
 DEFINE_STAT(STAT_ParamUpdateField);
diff --git a/Engine/Source/Runtime/Experimental/Chaos/Public/Chaos/BVHParticles.h b/Engine/Source/Runtime/Experimental/Chaos/Public/Chaos/BVHParticles.h
--- a/Engine/Source/Runtime/Experimental/Chaos/Public/Chaos/BVHParticles.h
+++ b/Engine/Source/Runtime/Experimental/Chaos/Public/Chaos/BVHParticles.h
@@ -4,6 +4,7 @@
 #include "Chaos/BoundingVolumeHierarchy.h"
 #include "Chaos/Particles.h"
 #include "ChaosArchive.h"
+#include "Chaos/ParticleHandleFwd.h"
 
 extern int32 CHAOS_API CollisionParticlesBVHDepth;
 
@@ -15,14 +16,16 @@ namespace Chaos
 	public:
 		static constexpr bool IsSerializablePtr = true;
 		using TArrayCollection::Size;
-		using TParticles<T, d>::AddElements;
 		using TParticles<T, d>::X;
+		using TParticles<T, d>::AddParticles;
 
 		TBVHParticles()
-		    : TParticles<T, d>(), MBVH(*this, CollisionParticlesBVHDepth)
+		    : TParticles<T, d>()
+			, MBVH(*this, CollisionParticlesBVHDepth)
 		{}
 		TBVHParticles(TBVHParticles<T, d>&& Other)
-		    : TParticles<T, d>(MoveTemp(Other)), MBVH(MoveTemp(Other.MBVH))
+		    : TParticles<T, d>(MoveTemp(Other))
+			, MBVH(MoveTemp(Other.MBVH))
 		{}
 		TBVHParticles(TParticles<T, d>&& Other)
 		    : TParticles<T, d>(MoveTemp(Other))
@@ -35,17 +38,6 @@ namespace Chaos
 			return *this;
 		}
 
-		/**
-		 * Constructor that replaces the positions array with a view of @param Points.
-		 */
-		TBVHParticles(TArray<TVector<T, d>>& Points)
-		    : TParticles<T, d>(Points)
-		    , MBVH(*this, CollisionParticlesBVHDepth)
-		{}
-
-		~TBVHParticles()
-		{}
-
 		TBVHParticles& operator=(TBVHParticles<T, d>&& Other)
 		{
 			MBVH = MoveTemp(Other.MBVH);
@@ -95,9 +87,10 @@ namespace Chaos
 
 	private:
 		TBVHParticles(const TBVHParticles<T, d>& Other)
-			: TParticles<T, d>(), MBVH(*this, CollisionParticlesBVHDepth)
+			: TParticles<T, d>()
+			, MBVH(*this, CollisionParticlesBVHDepth)
 		{
-			AddElements(Other.Size());
+			AddParticles(Other.Size());
 			for (int32 i = Other.Size() - 1; 0 <= i; i--)
 			{
 				X(i) = Other.X(i);
diff --git a/Engine/Source/Runtime/Experimental/Chaos/Public/Chaos/BoundingVolume.h b/Engine/Source/Runtime/Experimental/Chaos/Public/Chaos/BoundingVolume.h
--- a/Engine/Source/Runtime/Experimental/Chaos/Public/Chaos/BoundingVolume.h
+++ b/Engine/Source/Runtime/Experimental/Chaos/Public/Chaos/BoundingVolume.h
@@ -11,6 +11,8 @@
 #include "Chaos/UniformGrid.h"
 #include "Chaos/ISpatialAcceleration.h"
 #include "HAL/IConsoleManager.h"
+#include "Chaos/ParticleHandleFwd.h"
+#include "Templates/Models.h"
 
 #include <memory>
 #include <unordered_set>
@@ -20,6 +22,49 @@
 //#include "Engine/World.h"
 //#include "DrawDebugHelpers.h"
 
+template <typename TPayloadType>
+struct TPayloadHelper
+{
+	template <typename TParticle>
+	static TPayloadType GetPayload(TParticle& Particle, int32 Idx)
+	{
+		return Particle.Handle();
+	}
+};
+
+template <>
+struct TPayloadHelper<int32>
+{
+	template <typename TParticle>
+	static int32 GetPayload(const TParticle& Particle, int32 Idx)
+	{
+		return Idx;
+	}
+};
+
+template <typename T, bool>
+struct TSpatialAccelerationTraits
+{
+};
+
+template <typename TSOA>
+struct TSpatialAccelerationTraits<TSOA, true>
+{
+	using TPayloadType = typename TSOA::THandleType;
+};
+
+template <typename T>
+struct TSpatialAccelerationTraits<T, false>
+{
+	using TPayloadType = int32;
+};
+
+struct CParticleView
+{
+	template <typename T>
+	auto Requires(typename T::THandleType) ->void;
+};
+
 struct FBoundingVolumeCVars
 {
 	static int32 FilterFarBodies;
@@ -29,459 +74,133 @@
 namespace Chaos
 {
 
-template <typename OBJECT_ARRAY, typename T, int d>
-class TBoundingVolume;
-
-template <typename OBJECT_ARRAY, typename T, int d>
-void FixupLeafObj(const OBJECT_ARRAY& Objects, TArray<TBoundingVolume<OBJECT_ARRAY, T, d>>& Leafs);
-
-template<class OBJECT_ARRAY, class T, int d>
-class TBoundingVolume final : public ISpatialAcceleration<T,d>
+template<typename TSOA, typename InPayloadType, typename T, int d>
+class TBoundingVolume final : public ISpatialAcceleration<InPayloadType, T,d>
 {
   public:
-	  static constexpr int32 DefaultMaxCells = 15;
+	using TPayloadType = InPayloadType;
+	static constexpr int32 DefaultMaxCells = 15;
 	TBoundingVolume()
-		: MObjects(nullptr)
 	{
 	}
 
-	TBoundingVolume(const OBJECT_ARRAY& Objects, const bool bUseVelocity = false, const T Dt = 0, const int32 MaxCells = DefaultMaxCells)
-	    : MObjects(&Objects)
+	template <typename ParticleView>
+	TBoundingVolume(const ParticleView& Particles, const bool bUseVelocity = false, const T Dt = 0, const int32 MaxCells = DefaultMaxCells)
 	{
-		Reinitialize(bUseVelocity, Dt, MaxCells);
+		Reinitialize(Particles, bUseVelocity, Dt, MaxCells);
 	}
 
-	TBoundingVolume(const OBJECT_ARRAY& Objects, const TArray<uint32>& ActiveIndices, const bool bUseVelocity = false, const T Dt = 0, const int32 MaxCells = DefaultMaxCells)
-	    : MObjects(&Objects)
-	{
-		check(GetObjectCount(Objects) > 0);
-		Reinitialize(ActiveIndices, bUseVelocity, Dt, MaxCells);
-	}
-
-	TBoundingVolume(TBoundingVolume<OBJECT_ARRAY, T, d>&& Other)
-		: MObjects(Other.MObjects), MGlobalObjects(MoveTemp(Other.MGlobalObjects)), MAllObjects(MoveTemp(Other.MAllObjects)), MWorldSpaceBoxes(MoveTemp(Other.MWorldSpaceBoxes)), MGrid(MoveTemp(Other.MGrid)), MElements(MoveTemp(Other.MElements)), bIsEmpty(Other.bIsEmpty)
+	TBoundingVolume(TBoundingVolume<TSOA, TPayloadType, T, d>&& Other)
+		: MGlobalPayloads(MoveTemp(Other.MGlobalPayloads)), MGrid(MoveTemp(Other.MGrid)), MElements(MoveTemp(Other.MElements)), bIsEmpty(Other.bIsEmpty)
 	{
 	}
 
 private:
-	TBoundingVolume(const TBoundingVolume<OBJECT_ARRAY, T, d>& Other)
-		: MObjects(Other.MObjects), MGlobalObjects(Other.MGlobalObjects), MAllObjects(Other.MAllObjects), MWorldSpaceBoxes(Other.MWorldSpaceBoxes), MGrid(Other.MGrid), MElements(Other.MElements.Copy()), bIsEmpty(Other.bIsEmpty)
+	TBoundingVolume(const TBoundingVolume<TSOA, TPayloadType, T, d>& Other)
+		: MGlobalPayloads(Other.MGlobalPayloads), MGrid(Other.MGrid), MElements(Other.MElements.Copy()), bIsEmpty(Other.bIsEmpty)
 	{
 	}
 
 public:
-	TBoundingVolume<OBJECT_ARRAY, T, d>& operator=(const TBoundingVolume<OBJECT_ARRAY, T, d>& Other) = delete;
-	TBoundingVolume<OBJECT_ARRAY, T, d>& operator=(TBoundingVolume<OBJECT_ARRAY, T, d>&& Other)
+	TBoundingVolume<TSOA, TPayloadType, T, d>& operator=(const TBoundingVolume<TSOA, TPayloadType, T, d>& Other) = delete;
+	TBoundingVolume<TSOA, TPayloadType, T, d>& operator=(TBoundingVolume<TSOA, TPayloadType, T, d>&& Other)
 	{
-		const_cast<OBJECT_ARRAY*&>(MObjects) = const_cast<OBJECT_ARRAY*>(Other.MObjects);
-		MGlobalObjects = MoveTemp(Other.MGlobalObjects);
-		MAllObjects = MoveTemp(Other.MAllObjects);
-		MWorldSpaceBoxes = MoveTemp(Other.MWorldSpaceBoxes);
+		MGlobalPayloads = MoveTemp(Other.MGlobalPayloads);
 		MGrid = MoveTemp(Other.MGrid);
 		MElements = MoveTemp(Other.MElements);
 		bIsEmpty = Other.bIsEmpty;
 		return *this;
 	}
 
-	TBoundingVolume<OBJECT_ARRAY, T, d> Copy() const
+	TBoundingVolume<TSOA, TPayloadType, T, d> Copy() const
 	{
-		return TBoundingVolume<OBJECT_ARRAY, T, d>(*this);
+		return TBoundingVolume<TSOA, TPayloadType, T, d>(*this);
 	}
 
-	void Reinitialize(const bool bUseVelocity = false, const T Dt = 0, const int32 MaxCells = DefaultMaxCells)
+	template <typename ParticleView>
+	void Reinitialize(const ParticleView& Particles, const bool bUseVelocity = false, const T Dt = 0, const int32 MaxCells = DefaultMaxCells)
 	{
-		MGlobalObjects.Reset();
-		MAllObjects.Reset();
-		for (int32 i = 0; i < GetObjectCount(*MObjects); ++i)
+		MGlobalPayloads.Reset();
+		bIsEmpty = true;
+		TArray<bool> ComputeBounds;
+		ComputeBounds.AddUninitialized(Particles.Num());
+		int32 Count = 0;
+
+		for (auto& Particle : Particles)
 		{
-			if (IsDisabled(*MObjects, i))
-				continue;
-			if (HasBoundingBox(*MObjects, i))
+			if (!HasBoundingBox(Particle))
 			{
-				MAllObjects.Add(i);
+				MGlobalPayloads.Add(TPayloadHelper<TPayloadType>::GetPayload(Particle, Count));
+				ComputeBounds[Count++] = false;
 			}
 			else
 			{
-				MGlobalObjects.Add(i);
+				ComputeBounds[Count++] = true;
+				bIsEmpty = false;
 			}
 		}
-		bIsEmpty = MAllObjects.Num() == 0;
-		GenerateTree(*MObjects, MAllObjects, bUseVelocity, Dt, MaxCells);
+		GenerateTree(Particles, MoveTemp(ComputeBounds), bUseVelocity, Dt, MaxCells);
 		check(bIsEmpty || MGrid.GetNumCells() > 0);
 	}
 
-	void Reinitialize(const TArray<uint32>& ActiveIndices, const bool bUseVelocity = false, const T Dt = 0, const int32 MaxCells = DefaultMaxCells)
-	{
-		MGlobalObjects.Reset();
-		MAllObjects.Reset();
-		for (const auto& Index : ActiveIndices)
-		{
-			check(!IsDisabled(*MObjects, Index));
-			if (HasBoundingBox(*MObjects, Index))
-			{
-				MAllObjects.Add(Index);
-			}
-			else
-			{
-				MGlobalObjects.Add(Index);
-			}
-		}
-		bIsEmpty = MAllObjects.Num() == 0;
-		GenerateTree(*MObjects, MAllObjects, bUseVelocity, Dt, MaxCells);
-		check(bIsEmpty || MGrid.GetNumCells() > 0);
-	}
-
-	void GenerateTree(const OBJECT_ARRAY& Objects, const TArray<int32>& AllObjects, const bool bUseVelocity, const T Dt, const int32 MaxCells)
-	{
-		if (!AllObjects.Num())
-		{
-			return;
-		}
-		ComputeAllWorldSpaceBoundingBoxes(Objects, AllObjects, bUseVelocity, Dt, MWorldSpaceBoxes);
-		TBox<T, d> GlobalBox = Chaos::GetWorldSpaceBoundingBox(Objects, AllObjects[0], MWorldSpaceBoxes);
-		T Dx = TVector<T, d>::DotProduct(GlobalBox.Extents(), TVector<T, d>(1)) / d;
-		for (int32 i = 1; i < AllObjects.Num(); ++i)
-		{
-			const auto& WorldBox = Chaos::GetWorldSpaceBoundingBox(Objects, AllObjects[i], MWorldSpaceBoxes);
-			Dx += TVector<T, d>::DotProduct(WorldBox.Extents(), TVector<T, d>(1)) / d;
-			GlobalBox.GrowToInclude(WorldBox);
-		}
-		Dx /= AllObjects.Num();
-		if (FBoundingVolumeCVars::FilterFarBodies)
-		{
-			TArray<int32> ObjectsToRemove;
-			for (int32 i = 1; i < AllObjects.Num(); ++i)
-			{
-				const auto& WorldBox = Chaos::GetWorldSpaceBoundingBox(Objects, AllObjects[i], MWorldSpaceBoxes);
-				const TVector<T, d> MinToDXRatio = WorldBox.Min() / Dx;
-				for (int32 Axis = 0; Axis < d; ++Axis)
-				{
-					if (FMath::Abs(MinToDXRatio[Axis]) > 1e7)
-					{
-						ObjectsToRemove.Add(i);
-						MGlobalObjects.Add(AllObjects[i]);
-						break;
-					}
-				}
-			}
-			if (ObjectsToRemove.Num())
-			{
-				for (int32 i = ObjectsToRemove.Num() - 1; i >= 0; --i)
-				{
-					const_cast<TArray<int32>&>(AllObjects).RemoveAtSwap(ObjectsToRemove[i]);
-				}
-				GlobalBox = Chaos::GetWorldSpaceBoundingBox(Objects, AllObjects[0], MWorldSpaceBoxes);
-				Dx = TVector<T, d>::DotProduct(GlobalBox.Extents(), TVector<T, d>(1)) / d;
-				for (int32 i = 1; i < AllObjects.Num(); ++i)
-				{
-					const auto& WorldBox = Chaos::GetWorldSpaceBoundingBox(Objects, AllObjects[i], MWorldSpaceBoxes);
-					Dx += TVector<T, d>::DotProduct(WorldBox.Extents(), TVector<T, d>(1)) / d;
-					GlobalBox.GrowToInclude(WorldBox);
-				}
-				Dx /= AllObjects.Num();
-			}
-		}
-		TVector<int32, d> Cells = Dx > 0 ? GlobalBox.Extents() / Dx : TVector<int32, d>(MaxCells);
-		Cells += TVector<int32, d>(1);
-		for (int32 Axis = 0; Axis < d; ++Axis)
-		{
-			if (Cells[Axis] > MaxCells)
-				Cells[Axis] = MaxCells;
-			if (!(ensure(Cells[Axis] >=0 )))	//seeing this because GlobalBox is huge leading to int overflow. Need to investigate why bounds get so big
-			{
-				Cells[Axis] = MaxCells;
-			}
-		}
-		MGrid = TUniformGrid<T, d>(GlobalBox.Min(), GlobalBox.Max(), Cells);
-		MElements = TArrayND<TArray<int32>, d>(MGrid);
-		for (int32 i = 0; i < AllObjects.Num(); ++i)
-		{
-			const auto& ObjectBox = Chaos::GetWorldSpaceBoundingBox(Objects, AllObjects[i], MWorldSpaceBoxes);
-			const auto StartIndex = MGrid.ClampIndex(MGrid.Cell(ObjectBox.Min()));
-			const auto EndIndex = MGrid.ClampIndex(MGrid.Cell(ObjectBox.Max()));
-			for (int32 x = StartIndex[0]; x <= EndIndex[0]; ++x)
-			{
-				for (int32 y = StartIndex[1]; y <= EndIndex[1]; ++y)
-				{
-					for (int32 z = StartIndex[2]; z <= EndIndex[2]; ++z)
-					{
-						MElements(x, y, z).Add(AllObjects[i]);
-					}
-				}
-			}
-		}
-		T NumObjects = 0;
-		for (int32 i = 0; i < MGrid.GetNumCells(); ++i)
-		{
-			NumObjects += MElements(MGrid.GetIndex(i)).Num();
-		}
-		NumObjects /= AllObjects.Num();
-		UE_LOG(LogChaos, Verbose, TEXT("Generated Tree with (%d, %d, %d) Nodes and %f Per Cell"), MGrid.Counts()[0], MGrid.Counts()[1], MGrid.Counts()[2], NumObjects);
-	}
-
-	void RemoveElements(const TArray<uint32>& RemovedIndices)
-	{
-		for (const auto& Index : RemovedIndices)
-		{
-			if (!HasBoundingBox(*MObjects, Index))
-			{
-				MGlobalObjects.RemoveSwap(Index);
-			}
-			else
-			{
-				if (MWorldSpaceBoxes.Contains(Index))
-				{
-					const auto& ObjectBox = Chaos::GetWorldSpaceBoundingBox(*MObjects, Index, MWorldSpaceBoxes);
-					const auto StartIndex = MGrid.ClampIndex(MGrid.Cell(ObjectBox.Min()));
-					const auto EndIndex = MGrid.ClampIndex(MGrid.Cell(ObjectBox.Max()));
-					for (int32 x = StartIndex[0]; x <= EndIndex[0]; ++x)
-					{
-						for (int32 y = StartIndex[1]; y <= EndIndex[1]; ++y)
-						{
-							for (int32 z = StartIndex[2]; z <= EndIndex[2]; ++z)
-							{
-								MElements(x, y, z).RemoveSwap(Index);
-							}
-						}
-					}
-				}
-			}
-		}
-	}
-
-	void AddElements(const TArray<uint32>& AddedIndices)
-	{
-		if (!AddedIndices.Num())
-		{
-			return;
-		}
-		for (const auto& Index : AddedIndices)
-		{
-			// Compute and store world space box
-			MWorldSpaceBoxes.FindOrAdd(Index) = ComputeWorldSpaceBoundingBox(*MObjects, Index);
-		}
-		// Compute how many cells need to be added on
-		{
-			TBox<T, d> AddedBox = Chaos::GetWorldSpaceBoundingBox(*MObjects, AddedIndices[0], MWorldSpaceBoxes);
-			for (int32 i = 1; i < AddedIndices.Num(); ++i)
-			{
-				const TBox<T, d>& WorldBox = Chaos::GetWorldSpaceBoundingBox(*MObjects, AddedIndices[i], MWorldSpaceBoxes);
-				AddedBox.GrowToInclude(WorldBox);
-			}
-			TVector<int32, d> NewCells = MGrid.Counts();
-			TVector<T, d> NewMinCorner = MGrid.MinCorner();
-			TVector<T, d> NewMaxCorner = MGrid.MaxCorner();
-			bool bChanged = false;
-			for (int32 Axis = 0; Axis < d; ++Axis)
-			{
-				if (AddedBox.Min()[Axis] < MGrid.MinCorner()[Axis])
-				{
-					int32 NumNewCells = static_cast<int32>((MGrid.MinCorner()[Axis] - AddedBox.Min()[Axis]) / MGrid.Dx()[Axis]);
-					NewCells[Axis] += NumNewCells;
-					NewMinCorner[Axis] -= NumNewCells * MGrid.Dx()[Axis];
-					bChanged = true;
-				}
-				if (AddedBox.Max()[Axis] > MGrid.MaxCorner()[Axis])
-				{
-					int32 NumNewCells = static_cast<int32>((AddedBox.Max()[Axis] - MGrid.MaxCorner()[Axis]) / MGrid.Dx()[Axis]);
-					NewCells[Axis] += NumNewCells;
-					NewMaxCorner[Axis] += NumNewCells * MGrid.Dx()[Axis];
-					bChanged = true;
-				}
-			}
-			if (bChanged)
-			{
-				TUniformGrid<T, d> NewGrid(NewMinCorner, NewMaxCorner, NewCells);
-				TArrayND<TArray<int32>, d> NewElements(NewGrid);
-				for (int32 i = 0; i < MGrid.GetNumCells(); ++i)
-				{
-					TVector<int32, d> GridIndex = NewGrid.Cell(MGrid.Center(i));
-					NewElements(GridIndex) = MElements[i];
-				}
-				MGrid = NewGrid;
-				MElements = std::move(NewElements);
-			}
-		}
-		for (const auto& Index : AddedIndices)
-		{
-			if (!HasBoundingBox(*MObjects, Index))
-			{
-				MGlobalObjects.Add(Index);
-			}
-			else
-			{
-				// Add elem
-				const auto& ObjectBox = Chaos::GetWorldSpaceBoundingBox(*MObjects, Index, MWorldSpaceBoxes);
-				const auto StartIndex = MGrid.ClampIndex(MGrid.Cell(ObjectBox.Min()));
-				const auto EndIndex = MGrid.ClampIndex(MGrid.Cell(ObjectBox.Max()));
-				for (int32 x = StartIndex[0]; x <= EndIndex[0]; ++x)
-				{
-					for (int32 y = StartIndex[1]; y <= EndIndex[1]; ++y)
-					{
-						for (int32 z = StartIndex[2]; z <= EndIndex[2]; ++z)
-						{
-							MElements(x, y, z).Add(Index);
-						}
-					}
-				}
-			}
-		}
-	}
-
 	template<class T_INTERSECTION>
-	TArray<int32> FindAllIntersectionsImp(const T_INTERSECTION& Intersection) const
+	TArray<TPayloadType> FindAllIntersectionsImp(const T_INTERSECTION& Intersection) const
 	{
-		if (!MObjects)
-		{
-			return TArray<int32>();
-		}
 		if (bIsEmpty)
 		{
-			return TArray<int32>(MGlobalObjects);
+			return MGlobalPayloads;
 		}
-		TArray<int32> IntersectionList = FindAllIntersectionsHelper(Intersection);
-		IntersectionList.Append(MGlobalObjects);
+		TArray<TPayloadType> IntersectionList = FindAllIntersectionsHelper(Intersection);
+		IntersectionList.Append(MGlobalPayloads);
 		return IntersectionList;
 	}
 
 	// Begin ISpatialAcceleration interface
-	virtual TArray<int32> FindAllIntersections(const TBox<T, d>& Box) const override { return FindAllIntersectionsImp(Box); }
-	virtual TArray<int32> FindAllIntersections(const TSpatialRay<T, d>& Ray) const override { return FindAllIntersectionsImp(Ray); }
-	virtual TArray<int32> FindAllIntersections(const TVector<T, d>& Point) const override { return FindAllIntersectionsImp(Point); }
-	virtual TArray<int32> FindAllIntersections(const TGeometryParticles<T, d>& InParticles, const int32 i) const override
+	virtual TArray<TPayloadType> FindAllIntersections(const TBox<T, d>& Box) const override { return FindAllIntersectionsImp(Box); }
+	virtual TArray<TPayloadType> FindAllIntersections(const TVector<T, d>& Point) const override
 	{
-		return FindAllIntersectionsImp(Chaos::GetWorldSpaceBoundingBox(InParticles, i, MWorldSpaceBoxes));
+#if CHAOS_PARTICLEHANDLE_TODO
+		return FindAllIntersectionsImp(Box);
+#else
+		check(false);
+		return TArray<TPayloadType>();
+#endif
 	}
-	// End ISpatialAcceleration interface
 
-	template <bool bPruneDuplicates = true>
-	struct TSpatialRayIterator
+	virtual TArray<TPayloadType> FindAllIntersections(const TSpatialRay<T, d>& Ray) const override
 	{
-		TSpatialRayIterator(const TVector<T, d>& InStart, const TVector<T, d>& InDir, const T Length, const TBoundingVolume<OBJECT_ARRAY, T, d>& InBV)
-			: BV(InBV)
-			, Start(InStart)
-			, Dir(InDir)
-			, OriginalLength(Length)
-		{
-			TBox<T, d> GlobalBounds(BV.GetGrid().MinCorner(), BV.GetGrid().MaxCorner());
+#if CHAOS_PARTICLEHANDLE_TODO
+		return FindAllIntersectionsImp(Box);
+#else
+		check(false);
+		return TArray<TPayloadType>();
+#endif
+	}
 
-			T TOI;
+	virtual TArray<TPayloadType> FindAllIntersections(const TGeometryParticles<T, d>& InParticles, const int32 i) const override
+	{
+#if CHAOS_PARTICLEHANDLE_TODO
+		return FindAllIntersectionsImp(Box);
+#else
+		check(false);
+		return TArray<TPayloadType>();
+#endif
+	}
+	
 
-			InvLength = 1 / OriginalLength;
-			for (int Axis = 0; Axis < d; ++Axis)
-			{
-				bParallel[Axis] = Dir[Axis] == 0;
-				InvDir[Axis] = bParallel[Axis] ? 0 : 1 / Dir[Axis];
-			}
+	const TArray<TPayloadType>& GlobalObjects() const
+	{
+		return MGlobalPayloads;
+	}
 
-			bCellsLeft = GlobalBounds.RaycastFast(Start, Dir, InvDir, bParallel, Length, InvLength, TOI, NextStart);
-			if (bCellsLeft)
-			{
-				CellIdx = BV.GetGrid().Cell(NextStart);
-				CellIdx = BV.GetGrid().ClampIndex(CellIdx);	//raycast may have ended slightly outside of grid
-				End = Start + InDir * Length;
-			}
-		}
-
-		TArray<int32> GetNextIntersections()
-		{
-			TArray<int32> Results;
-			if (bCellsLeft)
-			{
-				do
-				{
-					//gather all instances in current cell whose bounds intersect with ray
-					const TArray<int32>& Instances = BV.GetElements()(CellIdx);
-					Results.Reserve(Instances.Num());
-					TVector<T, d> TmpPosition;
-					T TOI;
-
-					for (int32 Instance : Instances)
-					{
-						if (bPruneDuplicates)
-						{
-							if (InstancesSeen.Contains(Instance))
-							{
-								continue;
-							}
-							InstancesSeen.Add(Instance);
-						}
-						const TBox<T, d>& InstanceBounds = BV.GetWorldSpaceBoxes()[Instance];
-						if (InstanceBounds.RaycastFast(Start, Dir, InvDir, bParallel, OriginalLength, InvLength, TOI, TmpPosition))
-						{
-							Results.Add(Instance);
-						}
-					}
-
-					//find next cell
-
-					//We want to know which plane we used to cross into next cell
-					const TVector<T, d> CellCenter = BV.GetGrid().Location(CellIdx);
-					const TVector<T, d>& Dx = BV.GetGrid().Dx();
-
-					T Times[3];
-					T BestTime = TNumericLimits<T>::Max();
-					for (int Axis = 0; Axis < d; ++Axis)
-					{
-						if (!bParallel[Axis])
-						{
-							const T CrossPoint = Dir[Axis] > 0 ? CellCenter[Axis] + Dx[Axis] / 2 : CellCenter[Axis] - Dx[Axis] / 2;
-							const T Distance = CrossPoint - NextStart[Axis];	//note: CellCenter already has /2, we probably want to use the corner instead
-							const T Time = Distance * InvDir[Axis];
-							Times[Axis] = Time;
-							if (Time < BestTime)
-							{
-								BestTime = Time;
-							}
-						}
-						else
-						{
-							Times[Axis] = TNumericLimits<T>::Max();
-						}
-					}
-
-					for (int Axis = 0; Axis < d; ++Axis)
-					{
-						constexpr T Epsilon = 1e-2;	//if raycast is slightly off we still count it as hitting the cell surface
-						CellIdx[Axis] += (Times[Axis] <= BestTime + Epsilon) ? (Dir[Axis] > 0 ? 1 : -1) : 0;
-						if (CellIdx[Axis] < 0 || CellIdx[Axis] >= BV.GetGrid().Counts()[Axis])
-						{
-							bCellsLeft = false;
-							break;
-						}
-					}
-
-
-					NextStart = NextStart + Dir * BestTime;
-				} while (Results.Num() == 0 && bCellsLeft);	//if results are empty keep traveling
-			}
-
-			return Results;
-		}
-
-	private:
-		TSet<int32> InstancesSeen;
-		const TBoundingVolume<OBJECT_ARRAY, T, d>& BV;
-		const TVector<T, d> Start;
-		const TVector<T, d> Dir;
-		const T OriginalLength;
-		T InvLength;
-		TVector<T, d> End;
-		TVector<T, d> NextStart;
-		TVector<int32, d> CellIdx;
-		TVector<T, d> InvDir;
-		bool bParallel[d];
-		bool bCellsLeft;
-	};
-
-	virtual void Raycast(const TVector<T, d>& Start, const TVector<T, d>& Dir, const T OriginalLength, ISpatialVisitor<T>& Visitor) const override
+	virtual void Raycast(const TVector<T, d>& Start, const TVector<T, d>& Dir, const T OriginalLength, ISpatialVisitor<TPayloadType, T>& Visitor) const override
 	{
-		TSpatialVisitor<T> ProxyVisitor(Visitor);
+		TSpatialVisitor<TPayloadType, T> ProxyVisitor(Visitor);
 		Raycast(Start, Dir, OriginalLength, ProxyVisitor);
 	}
 
 	template <typename SQVisitor, bool bPruneDuplicates = true>
-	void Raycast(const TVector<T,d>& Start, const TVector<T, d>& Dir, const T OriginalLength, SQVisitor& Visitor) const
+	void Raycast(const TVector<T, d>& Start, const TVector<T, d>& Dir, const T OriginalLength, SQVisitor& Visitor) const
 	{
 		TBox<T, d> GlobalBounds(MGrid.MinCorner(), MGrid.MaxCorner());
 		bool bParallel[d];
@@ -497,8 +216,8 @@
 		T TOI;
 		TVector<T, d> NextStart;
 		TVector<int32, d> CellIdx;
-		TSet<int32> InstancesSeen;
-		bool bCellsLeft = GlobalBounds.RaycastFast(Start, Dir, InvDir, bParallel, OriginalLength, InvOriginalLength, TOI, NextStart);
+		TSet<TPayloadType> InstancesSeen;
+		bool bCellsLeft = MElements.Num() && GlobalBounds.RaycastFast(Start, Dir, InvDir, bParallel, OriginalLength, InvOriginalLength, TOI, NextStart);
 		if (bCellsLeft)
 		{
 			CellIdx = MGrid.Cell(NextStart);
@@ -509,24 +228,24 @@
 			do
 			{
 				//gather all instances in current cell whose bounds intersect with ray
-				const TArray<int32>& Instances = MElements(CellIdx);
+				const auto& Elems = MElements(CellIdx);
 				//should we let callback know about max potential?
 				TVector<T, d> TmpPosition;
 
-				for (int32 Instance : Instances)
+				for (const auto& Elem : Elems)
 				{
 					if (bPruneDuplicates)
 					{
-						if (InstancesSeen.Contains(Instance))
+						if (InstancesSeen.Contains(Elem.Payload))
 						{
 							continue;
 						}
-						InstancesSeen.Add(Instance);
+						InstancesSeen.Add(Elem.Payload);
 					}
-					const TBox<T, d>& InstanceBounds = MWorldSpaceBoxes[Instance];
+					const auto& InstanceBounds = Elem.Bounds;
 					if (InstanceBounds.RaycastFast(Start, Dir, InvDir, bParallel, CurrentLength, InvCurrentLength, TOI, TmpPosition))
 					{
-						const bool bContinue = Visitor.VisitRaycast(Instance, CurrentLength);
+						const bool bContinue = Visitor.VisitRaycast(Elem.Payload, CurrentLength);
 						if (!bContinue)
 						{
 							return;
@@ -584,15 +303,20 @@
 		}
 	}
 
-	void Sweep(const TVector<T, d>& Start, const TVector<T, d>& Dir, T OriginalLength, const TVector<T, d> QueryHalfExtents, ISpatialVisitor<T>& Visitor, const TVector<T, d>& Scale = TVector<T, d>(1)) const override
+	void Sweep(const TVector<T, d>& Start, const TVector<T, d>& Dir, T OriginalLength, const TVector<T, d> QueryHalfExtents, ISpatialVisitor<TPayloadType, T>& Visitor, const TVector<T, d>& Scale = TVector<T, d>(1)) const override
 	{
-		TSpatialVisitor<T> ProxyVisitor(Visitor);
+		TSpatialVisitor<TPayloadType, T> ProxyVisitor(Visitor);
 		Sweep(Start, Dir, OriginalLength, QueryHalfExtents, ProxyVisitor, Scale);
 	}
 
 	template <typename SQVisitor, bool bPruneDuplicates = true>
-	void Sweep(const TVector<T, d>& Start, const TVector<T, d>& Dir, T OriginalLength, const TVector<T, d> QueryHalfExtents, SQVisitor& Visitor, const TVector<T,d>& Scale = TVector<T,d>(1)) const
+	void Sweep(const TVector<T, d>& Start, const TVector<T, d>& Dir, T OriginalLength, const TVector<T, d> QueryHalfExtents, SQVisitor& Visitor, const TVector<T, d>& Scale = TVector<T, d>(1)) const
 	{
+		if (MElements.Num() == 0)
+		{
+			return;
+		}
+
 		const TVector<T, d> ScaledMin = MGrid.MinCorner() * Scale;
 		const TVector<T, d> ScaledMax = MGrid.MaxCorner() * Scale;
 
@@ -608,17 +332,16 @@
 			InvDir[Axis] = bParallel[Axis] ? 0 : 1 / Dir[Axis];
 		}
 
-
 		struct FCellIntersection
 		{
 			TVector<int32, d> CellIdx;
 			T TOI;
 		};
 
-		T TOI;
+		T TOI = 0;	//not needed, but fixes compiler warning
 		TVector<T, d> HitPoint;
-		TSet<int32> InstancesSeen;
-		TSet<TVector<int32,d>> IdxsSeen;
+		TSet<TPayloadType> InstancesSeen;
+		TSet<TVector<int32, d>> IdxsSeen;
 		const bool bInitialHit = GlobalBounds.RaycastFast(Start, Dir, InvDir, bParallel, OriginalLength, InvOriginalLength, TOI, HitPoint);
 		if (bInitialHit)
 		{
@@ -632,7 +355,7 @@
 			IdxsQueue.Add({ HitCellIdx, TOI });
 
 			int32 QueueIdx = 0;	//FIFO because early cells are more likely to block later cells we can skip
-			while(QueueIdx < IdxsQueue.Num())
+			while (QueueIdx < IdxsQueue.Num())
 			{
 				const FCellIntersection CellIntersection = IdxsQueue[QueueIdx++];
 				if (CellIntersection.TOI > CurrentLength)
@@ -686,23 +409,23 @@
 				}
 
 				//check if any instances in the cell are hit
-				const TArray<int32>& Instances = MElements(CellIntersection.CellIdx);
-				for (int32 Instance : Instances)
+				const auto& Elems = MElements(CellIntersection.CellIdx);
+				for (const auto& Elem : Elems)
 				{
 					if (bPruneDuplicates)
 					{
-						if (InstancesSeen.Contains(Instance))
+						if (InstancesSeen.Contains(Elem.Payload))
 						{
 							continue;
 						}
-						InstancesSeen.Add(Instance);
+						InstancesSeen.Add(Elem.Payload);
 					}
-						
-					const TBox<T,d>& InstanceBounds = MWorldSpaceBoxes[Instance];
+
+					const TBox<T, d>& InstanceBounds = Elem.Bounds;
 					const TBox<T, d> InflatedScaledInstanceBounds(InstanceBounds.Min() * Scale - QueryHalfExtents, InstanceBounds.Max() * Scale + QueryHalfExtents);
 					if (InflatedScaledInstanceBounds.RaycastFast(Start, Dir, InvDir, bParallel, CurrentLength, InvCurrentLength, TOI, HitPoint))
 					{
-						const bool bContinue = Visitor.VisitSweep(Instance, CurrentLength);
+						const bool bContinue = Visitor.VisitSweep(Elem.Payload, CurrentLength);
 						if (!bContinue)
 						{
 							return;
@@ -714,24 +437,24 @@
 		}
 	}
 
-	void Overlap(const TBox<T, d>& QueryBounds, ISpatialVisitor<T>& Visitor, const TVector<T, d>& Scale = TVector<T, d>(1)) const override
+	void Overlap(const TBox<T, d>& QueryBounds, ISpatialVisitor<TPayloadType, T>& Visitor, const TVector<T, d>& Scale = TVector<T, d>(1)) const override
 	{
-		TSpatialVisitor<T> ProxyVisitor(Visitor);
+		TSpatialVisitor<TPayloadType, T> ProxyVisitor(Visitor);
 		return Overlap(QueryBounds, ProxyVisitor, Scale);
 	}
 
 	template <typename SQVisitor, bool bPruneDuplicates = true>
-	void Overlap(const TBox<T,d>& QueryBounds, SQVisitor& Visitor, const TVector<T, d>& Scale = TVector<T, d>(1)) const
+	void Overlap(const TBox<T, d>& QueryBounds, SQVisitor& Visitor, const TVector<T, d>& Scale = TVector<T, d>(1)) const
 	{
 		const TVector<T, d> ScaledMin = MGrid.MinCorner() * Scale;
 		const TVector<T, d> ScaledMax = MGrid.MaxCorner() * Scale;
 
 		TUniformGrid<T, d> ScaledGrid(ScaledMin, ScaledMax, MGrid.Counts(), 0);
 		TBox<T, d> GlobalBounds(ScaledMin, ScaledMax);
-		
+
 		const TVector<int32, d> StartIndex = ScaledGrid.ClampIndex(ScaledGrid.Cell(QueryBounds.Min()));
 		const TVector<int32, d> EndIndex = ScaledGrid.ClampIndex(ScaledGrid.Cell(QueryBounds.Max()));
-		TSet<int32> InstancesSeen;
+		TSet<TPayloadType> InstancesSeen;
 
 		for (int32 X = StartIndex[0]; X <= EndIndex[0]; ++X)
 		{
@@ -739,7 +462,378 @@
 			{
 				for (int32 Z = StartIndex[2]; Z <= EndIndex[2]; ++Z)
 				{
-					const TArray<int32>& Instances = MElements(X, Y, Z);
+					const auto& Elems = MElements(X, Y, Z);
+					for (const auto& Elem : Elems)
+					{
+						if (bPruneDuplicates)
+						{
+							if (InstancesSeen.Contains(Elem.Payload))
+							{
+								continue;
+							}
+							InstancesSeen.Add(Elem.Payload);
+						}
+						const TBox<T, d>& InstanceBounds = Elem.Bounds;
+						const TBox<T, d> InflatedScaledInstanceBounds(InstanceBounds.Min() * Scale, InstanceBounds.Max() * Scale);
+						if (QueryBounds.Intersects(InflatedScaledInstanceBounds))
+						{
+							if (Visitor.VisitOverlap(Elem.Payload) == false)
+							{
+								return;
+							}
+						}
+					}
+				}
+			}
+		}
+	}
+
+private:
+
+	template <typename ParticleView>
+	void GenerateTree(const ParticleView& Particles, TArray<bool>&& HasBounds, const bool bUseVelocity, const T Dt, const int32 MaxCells)
+	{
+		if (bIsEmpty)
+		{
+			return;
+		}
+		TArray<TBox<T, d>> WorldSpaceBoxes;
+		ComputeAllWorldSpaceBoundingBoxes(Particles, HasBounds, bUseVelocity, Dt, WorldSpaceBoxes);
+
+		auto ComputeBoxAndDx = [&HasBounds, &Particles, &WorldSpaceBoxes](TBox<T, d>& GlobalBox) -> T
+		{
+			T Dx = 0;
+			bool bFoundValidBounds = false;
+			T NumElements = 0;
+			//todo: add const copy
+			for (int32 Idx = 0; Idx < Particles.Num(); ++Idx)
+			{
+				if (HasBounds[Idx])
+				{
+					NumElements += 1;
+					const auto& WorldSpaceBox = WorldSpaceBoxes[Idx];
+					if (!bFoundValidBounds)
+					{
+						GlobalBox = WorldSpaceBox;
+						bFoundValidBounds = true;
+					}
+					else
+					{
+						GlobalBox.GrowToInclude(WorldSpaceBox);
+					}
+					Dx += TVector<T, d>::DotProduct(WorldSpaceBox.Extents(), TVector<T, d>(1)) / d;
+				}
+			}
+
+			Dx /= NumElements;
+			return Dx;
+		};
+
+		TBox<T, d> GlobalBox;
+		T Dx = ComputeBoxAndDx(GlobalBox);
+
+		if (FBoundingVolumeCVars::FilterFarBodies)
+		{
+			bool bRecomputeBoxAndDx = false;
+			int32 Idx = 0;
+			for (auto& Particle : Particles)
+			{
+				if (HasBounds[Idx])
+				{
+					const auto& WorldSpaceBox = WorldSpaceBoxes[Idx];
+					const TVector<T, d> MinToDXRatio = WorldSpaceBox.Min() / Dx;
+					for (int32 Axis = 0; Axis < d; ++Axis)
+					{
+						if (FMath::Abs(MinToDXRatio[Axis]) > 1e7)
+						{
+							bRecomputeBoxAndDx = true;
+							HasBounds[Idx] = false;
+							MGlobalPayloads.Add(TPayloadHelper<TPayloadType>::GetPayload(Particle, Idx));
+							break;
+						}
+					}
+				}
+				++Idx;
+			}
+
+			if (bRecomputeBoxAndDx)
+			{
+				Dx = ComputeBoxAndDx(GlobalBox);
+			}
+		}
+
+		TVector<int32, d> Cells = Dx > 0 ? GlobalBox.Extents() / Dx : TVector<int32, d>(MaxCells);
+		Cells += TVector<int32, d>(1);
+		for (int32 Axis = 0; Axis < d; ++Axis)
+		{
+			if (Cells[Axis] > MaxCells)
+				Cells[Axis] = MaxCells;
+			if (!(ensure(Cells[Axis] >= 0)))	//seeing this because GlobalBox is huge leading to int overflow. Need to investigate why bounds get so big
+			{
+				Cells[Axis] = MaxCells;
+			}
+		}
+		MGrid = TUniformGrid<T, d>(GlobalBox.Min(), GlobalBox.Max(), Cells);
+		MElements = TArrayND<TArray<FCellElement>, d>(MGrid);
+
+		T NumObjectsWithBounds = 0;
+		T NumObjectsInCells = 0;
+		{
+			int32 Idx = 0;
+			for (auto& Particle : Particles)
+			{
+				if (HasBounds[Idx])
+				{
+					NumObjectsWithBounds += 1;
+					const auto& ObjectBox = WorldSpaceBoxes[Idx];
+					const auto StartIndex = MGrid.ClampIndex(MGrid.Cell(ObjectBox.Min()));
+					const auto EndIndex = MGrid.ClampIndex(MGrid.Cell(ObjectBox.Max()));
+					for (int32 x = StartIndex[0]; x <= EndIndex[0]; ++x)
+					{
+						for (int32 y = StartIndex[1]; y <= EndIndex[1]; ++y)
+						{
+							for (int32 z = StartIndex[2]; z <= EndIndex[2]; ++z)
+							{
+								MElements(x, y, z).Add({ ObjectBox, TPayloadHelper<TPayloadType>::GetPayload(Particle, Idx) });
+								NumObjectsInCells += 1;
+							}
+						}
+					}
+				}
+				++Idx;
+			}
+		}
+
+		UE_LOG(LogChaos, Verbose, TEXT("Generated Tree with (%d, %d, %d) Nodes and %f Per Cell"), MGrid.Counts()[0], MGrid.Counts()[1], MGrid.Counts()[2], NumObjectsInCells / NumObjectsWithBounds);
+	}
+
+	TArray<TPayloadType> FindAllIntersectionsHelper(const TBox<T, d>& ObjectBox) const
+	{
+		TArray<TPayloadType> Intersections;
+		const auto StartIndex = MGrid.ClampIndex(MGrid.Cell(ObjectBox.Min()));
+		const auto EndIndex = MGrid.ClampIndex(MGrid.Cell(ObjectBox.Max()));
+		for (int32 x = StartIndex[0]; x <= EndIndex[0]; ++x)
+		{
+			for (int32 y = StartIndex[1]; y <= EndIndex[1]; ++y)
+			{
+				for (int32 z = StartIndex[2]; z <= EndIndex[2]; ++z)
+				{
+					const TArray<FCellElement>& CellElements = MElements(x, y, z);
+					Intersections.Reserve(Intersections.Num() + CellElements.Num());
+					for (const FCellElement& Elem : CellElements)
+					{
+						if (ObjectBox.Intersects(Elem.Bounds))
+						{
+							Intersections.Add(Elem.Payload);
+						}
+					}
+				}
+			}
+		}
+
+
+		Algo::Sort(Intersections);
+
+		for (int32 i = Intersections.Num() - 1; i > 0; i--)
+		{
+			if (Intersections[i] == Intersections[i - 1])
+			{
+				Intersections.RemoveAtSwap(i, 1, false);
+			}
+		}
+
+		return Intersections;
+	}
+
+	struct FCellElement
+	{
+		TBox<T, d> Bounds;
+		TPayloadType Payload;
+	};
+
+	TArray<TPayloadType> MGlobalPayloads;
+	TUniformGrid<T, d> MGrid;
+	TArrayND<TArray<FCellElement>, d> MElements;
+	bool bIsEmpty;
+
+#if 0
+	void RemoveElements(const TArray<uint32>& RemovedIndices)
+	{
+		for (const auto& Index : RemovedIndices)
+		{
+			if (!HasBoundingBox(*MObjects, Index))
+			{
+				MGlobalPayloads.RemoveSwap(Index);
+			}
+			else
+			{
+				if (MWorldSpaceBoxes.Contains(Index))
+				{
+					const auto& ObjectBox = Chaos::GetWorldSpaceBoundingBox(*MObjects, Index, MWorldSpaceBoxes);
+					const auto StartIndex = MGrid.ClampIndex(MGrid.Cell(ObjectBox.Min()));
+					const auto EndIndex = MGrid.ClampIndex(MGrid.Cell(ObjectBox.Max()));
+					for (int32 x = StartIndex[0]; x <= EndIndex[0]; ++x)
+					{
+						for (int32 y = StartIndex[1]; y <= EndIndex[1]; ++y)
+						{
+							for (int32 z = StartIndex[2]; z <= EndIndex[2]; ++z)
+							{
+								MElements(x, y, z).RemoveSwap(Index);
+							}
+						}
+					}
+				}
+			}
+		}
+	}
+
+	void AddElements(const TArray<uint32>& AddedIndices)
+	{
+		if (!AddedIndices.Num())
+		{
+			return;
+		}
+		for (const auto& Index : AddedIndices)
+		{
+			// Compute and store world space box
+			MWorldSpaceBoxes.FindOrAdd(Index) = ComputeWorldSpaceBoundingBox(*MObjects, Index);
+		}
+		// Compute how many cells need to be added on
+		{
+			TBox<T, d> AddedBox = Chaos::GetWorldSpaceBoundingBox(*MObjects, AddedIndices[0], MWorldSpaceBoxes);
+			for (int32 i = 1; i < AddedIndices.Num(); ++i)
+			{
+				const TBox<T, d>& WorldBox = Chaos::GetWorldSpaceBoundingBox(*MObjects, AddedIndices[i], MWorldSpaceBoxes);
+				AddedBox.GrowToInclude(WorldBox);
+			}
+			TVector<int32, d> NewCells = MGrid.Counts();
+			TVector<T, d> NewMinCorner = MGrid.MinCorner();
+			TVector<T, d> NewMaxCorner = MGrid.MaxCorner();
+			bool bChanged = false;
+			for (int32 Axis = 0; Axis < d; ++Axis)
+			{
+				if (AddedBox.Min()[Axis] < MGrid.MinCorner()[Axis])
+				{
+					int32 NumNewCells = static_cast<int32>((MGrid.MinCorner()[Axis] - AddedBox.Min()[Axis]) / MGrid.Dx()[Axis]);
+					NewCells[Axis] += NumNewCells;
+					NewMinCorner[Axis] -= NumNewCells * MGrid.Dx()[Axis];
+					bChanged = true;
+				}
+				if (AddedBox.Max()[Axis] > MGrid.MaxCorner()[Axis])
+				{
+					int32 NumNewCells = static_cast<int32>((AddedBox.Max()[Axis] - MGrid.MaxCorner()[Axis]) / MGrid.Dx()[Axis]);
+					NewCells[Axis] += NumNewCells;
+					NewMaxCorner[Axis] += NumNewCells * MGrid.Dx()[Axis];
+					bChanged = true;
+				}
+			}
+			if (bChanged)
+			{
+				TUniformGrid<T, d> NewGrid(NewMinCorner, NewMaxCorner, NewCells);
+				TArrayND<TArray<int32>, d> NewElements(NewGrid);
+				for (int32 i = 0; i < MGrid.GetNumCells(); ++i)
+				{
+					TVector<int32, d> GridIndex = NewGrid.Cell(MGrid.Center(i));
+					NewElements(GridIndex) = MElements[i];
+				}
+				MGrid = NewGrid;
+				MElements = std::move(NewElements);
+			}
+		}
+		for (const auto& Index : AddedIndices)
+		{
+			if (!HasBoundingBox(*MObjects, Index))
+			{
+				MGlobalPayloads.Add(Index);
+			}
+			else
+			{
+				// Add elem
+				const auto& ObjectBox = Chaos::GetWorldSpaceBoundingBox(*MObjects, Index, MWorldSpaceBoxes);
+				const auto StartIndex = MGrid.ClampIndex(MGrid.Cell(ObjectBox.Min()));
+				const auto EndIndex = MGrid.ClampIndex(MGrid.Cell(ObjectBox.Max()));
+				for (int32 x = StartIndex[0]; x <= EndIndex[0]; ++x)
+				{
+					for (int32 y = StartIndex[1]; y <= EndIndex[1]; ++y)
+					{
+						for (int32 z = StartIndex[2]; z <= EndIndex[2]; ++z)
+						{
+							MElements(x, y, z).Add(Index);
+						}
+					}
+				}
+			}
+		}
+	}
+
+	template<class T_INTERSECTION>
+	TArray<int32> FindAllIntersectionsImp(const T_INTERSECTION& Intersection) const
+	{
+		if (!MObjects)
+		{
+			return TArray<int32>();
+		}
+		if (bIsEmpty)
+		{
+			return TArray<int32>(MGlobalPayloads);
+		}
+		TArray<int32> IntersectionList = FindAllIntersectionsHelper(Intersection);
+		IntersectionList.Append(MGlobalPayloads);
+		return IntersectionList;
+	}
+
+	// Begin ISpatialAcceleration interface
+	virtual TArray<int32> FindAllIntersections(const TBox<T, d>& Box) const override { return FindAllIntersectionsImp(Box); }
+	virtual TArray<int32> FindAllIntersections(const TSpatialRay<T, d>& Ray) const override { return FindAllIntersectionsImp(Ray); }
+	virtual TArray<int32> FindAllIntersections(const TVector<T, d>& Point) const override { return FindAllIntersectionsImp(Point); }
+	virtual TArray<int32> FindAllIntersections(const TGeometryParticles<T, d>& InParticles, const int32 i) const override
+	{
+		return FindAllIntersectionsImp(Chaos::GetWorldSpaceBoundingBox(InParticles, i, MWorldSpaceBoxes));
+	}
+	// End ISpatialAcceleration interface
+
+	template <bool bPruneDuplicates = true>
+	struct TSpatialRayIterator
+	{
+		TSpatialRayIterator(const TVector<T, d>& InStart, const TVector<T, d>& InDir, const T Length, const TBoundingVolume<OBJECT_ARRAY, T, d>& InBV)
+			: BV(InBV)
+			, Start(InStart)
+			, Dir(InDir)
+			, OriginalLength(Length)
+		{
+			TBox<T, d> GlobalBounds(BV.GetGrid().MinCorner(), BV.GetGrid().MaxCorner());
+
+			T TOI;
+
+			InvLength = 1 / OriginalLength;
+			for (int Axis = 0; Axis < d; ++Axis)
+			{
+				bParallel[Axis] = Dir[Axis] == 0;
+				InvDir[Axis] = bParallel[Axis] ? 0 : 1 / Dir[Axis];
+			}
+
+			bCellsLeft = GlobalBounds.RaycastFast(Start, Dir, InvDir, bParallel, Length, InvLength, TOI, NextStart);
+			if (bCellsLeft)
+			{
+				CellIdx = BV.GetGrid().Cell(NextStart);
+				CellIdx = BV.GetGrid().ClampIndex(CellIdx);	//raycast may have ended slightly outside of grid
+				End = Start + InDir * Length;
+			}
+		}
+
+		TArray<int32> GetNextIntersections()
+		{
+			TArray<int32> Results;
+			if (bCellsLeft)
+			{
+				do
+				{
+					//gather all instances in current cell whose bounds intersect with ray
+					const TArray<int32>& Instances = BV.GetElements()(CellIdx);
+					Results.Reserve(Instances.Num());
+					TVector<T, d> TmpPosition;
+					T TOI;
+
 					for (int32 Instance : Instances)
 					{
 						if (bPruneDuplicates)
@@ -750,25 +844,73 @@
 							}
 							InstancesSeen.Add(Instance);
 						}
-						const TBox<T, d>& InstanceBounds = MWorldSpaceBoxes[Instance];
-						const TBox<T, d> InflatedScaledInstanceBounds(InstanceBounds.Min() * Scale, InstanceBounds.Max() * Scale);
-						if (QueryBounds.Intersects(InflatedScaledInstanceBounds))
+						const TBox<T, d>& InstanceBounds = BV.GetWorldSpaceBoxes()[Instance];
+						if (InstanceBounds.RaycastFast(Start, Dir, InvDir, bParallel, OriginalLength, InvLength, TOI, TmpPosition))
 						{
-							if (Visitor.VisitOverlap(Instance) == false)
+							Results.Add(Instance);
+						}
+					}
+
+					//find next cell
+
+					//We want to know which plane we used to cross into next cell
+					const TVector<T, d> CellCenter = BV.GetGrid().Location(CellIdx);
+					const TVector<T, d>& Dx = BV.GetGrid().Dx();
+
+					T Times[3];
+					T BestTime = TNumericLimits<T>::Max();
+					for (int Axis = 0; Axis < d; ++Axis)
+					{
+						if (!bParallel[Axis])
+						{
+							const T CrossPoint = Dir[Axis] > 0 ? CellCenter[Axis] + Dx[Axis] / 2 : CellCenter[Axis] - Dx[Axis] / 2;
+							const T Distance = CrossPoint - NextStart[Axis];	//note: CellCenter already has /2, we probably want to use the corner instead
+							const T Time = Distance * InvDir[Axis];
+							Times[Axis] = Time;
+							if (Time < BestTime)
 							{
-								return;
+								BestTime = Time;
 							}
 						}
+						else
+						{
+							Times[Axis] = TNumericLimits<T>::Max();
+						}
 					}
-				}
+
+					for (int Axis = 0; Axis < d; ++Axis)
+					{
+						constexpr T Epsilon = 1e-2;	//if raycast is slightly off we still count it as hitting the cell surface
+						CellIdx[Axis] += (Times[Axis] <= BestTime + Epsilon) ? (Dir[Axis] > 0 ? 1 : -1) : 0;
+						if (CellIdx[Axis] < 0 || CellIdx[Axis] >= BV.GetGrid().Counts()[Axis])
+						{
+							bCellsLeft = false;
+							break;
+						}
+					}
+
+
+					NextStart = NextStart + Dir * BestTime;
+				} while (Results.Num() == 0 && bCellsLeft);	//if results are empty keep traveling
 			}
+
+			return Results;
 		}
-	}
 
-	const TArray<int32>& GlobalObjects() const
-	{
-		return MGlobalObjects;
-	}
+	private:
+		TSet<int32> InstancesSeen;
+		const TBoundingVolume<OBJECT_ARRAY, T, d>& BV;
+		const TVector<T, d> Start;
+		const TVector<T, d> Dir;
+		const T OriginalLength;
+		T InvLength;
+		TVector<T, d> End;
+		TVector<T, d> NextStart;
+		TVector<int32, d> CellIdx;
+		TVector<T, d> InvDir;
+		bool bParallel[d];
+		bool bCellsLeft;
+	};
 
 	const TMap<int32, TBox<T, d>>& GetWorldSpaceBoxes() const
 	{
@@ -793,17 +935,17 @@
 	
 	void Serialize(FArchive& Ar)
 	{
-		Ar << MGlobalObjects << MWorldSpaceBoxes << MGrid << MElements << bIsEmpty;
+		Ar << MGlobalPayloads << MWorldSpaceBoxes << MGrid << MElements << bIsEmpty;
 	}
 
-  private:
-
 	//Needed for serialization
 	void SetObjects(const OBJECT_ARRAY& Object)
 	{
 		MObjects = &Object;
 	}
 
+  private:
+
 	template <typename OBJECT_ARRAY2, typename T2, int d2>
 	friend void FixupLeafObj(const OBJECT_ARRAY2& Objects, TArray<TBoundingVolume<OBJECT_ARRAY2, T2, d2>>& Leafs);
 
@@ -1064,20 +1206,34 @@
 	}
 
 	const OBJECT_ARRAY* MObjects;
-	TArray<int32> MGlobalObjects;
+	TArray<int32> MGlobalPayloads;
 	TArray<int32> MAllObjects;
 	TMap<int32, TBox<T, d>> MWorldSpaceBoxes;
 	TUniformGrid<T, d> MGrid;
 	TArrayND<TArray<int32>, d> MElements;
 	bool bIsEmpty;
 	FCriticalSection CriticalSection;
+#endif
 };
 
+#if 0
 template<class OBJECT_ARRAY, class T, int d>
 FArchive& operator<<(FArchive& Ar, TBoundingVolume<OBJECT_ARRAY, T, d>& BoundingVolume)
 {
 	BoundingVolume.Serialize(Ar);
 	return Ar;
 }
+#endif
+
+template<class OBJECT_ARRAY, typename TPayloadType, class T, int d>
+FArchive& operator<<(FArchive& Ar, TBoundingVolume<OBJECT_ARRAY, TPayloadType, T, d>& BoundingVolume)
+{
+	check(false);
+#if CHAOS_PARTICLEHANDLE_TODO
+	BoundingVolume.Serialize(Ar);
+#endif
+	return Ar;
+}
 
 }
+
diff --git a/Engine/Source/Runtime/Experimental/Chaos/Public/Chaos/BoundingVolumeHierarchy.h b/Engine/Source/Runtime/Experimental/Chaos/Public/Chaos/BoundingVolumeHierarchy.h
--- a/Engine/Source/Runtime/Experimental/Chaos/Public/Chaos/BoundingVolumeHierarchy.h
+++ b/Engine/Source/Runtime/Experimental/Chaos/Public/Chaos/BoundingVolumeHierarchy.h
@@ -9,6 +9,30 @@
 #include "Chaos/Transform.h"
 #include "ChaosLog.h"
 #include "Chaos/ISpatialAcceleration.h"
+#include "Templates/Models.h"
+
+template <typename T, bool>
+struct TBVHLeafTraits
+{
+};
+
+template <typename T>
+struct TBVHLeafTraits<T, true>
+{
+	using TPayloadType = typename T::TPayloadType;
+};
+
+template <typename T>
+struct TBVHLeafTraits<T, false>
+{
+	using TPayloadType = typename T::ElementType;
+};
+
+struct CComplexBVHLeaf
+{
+	template <typename T>
+	auto Requires(typename T::TPayloadType) ->void;
+};
 
 #define MIN_NUM_OBJECTS 5
 
@@ -33,13 +57,14 @@ FArchive& operator<<(FArchive& Ar, TBVHNode<T,d>& LeafNode)
 }
 
 template<class OBJECT_ARRAY, class LEAF_TYPE, class T, int d>
-class CHAOS_API TBoundingVolumeHierarchy final : public ISpatialAcceleration<T,d>
+class CHAOS_API TBoundingVolumeHierarchy final : public ISpatialAcceleration<int32, T,d>
 {
   public:
 	static constexpr int32 DefaultMaxLevels = 12;
 	static constexpr bool DefaultAllowMultipleSplitting = false;
 	static constexpr bool DefaultUseVelocity = false;
 	static constexpr T DefaultDt = 0;
+	using TPayloadType = typename TBVHLeafTraits<LEAF_TYPE, TModels<CComplexBVHLeaf, LEAF_TYPE>::Value>::TPayloadType;
 
 	TBoundingVolumeHierarchy(const OBJECT_ARRAY& Objects, const int32 MaxLevels = DefaultMaxLevels, const bool bUseVelocity = DefaultUseVelocity, const T Dt = DefaultDt);
 	TBoundingVolumeHierarchy(const OBJECT_ARRAY& Objects, const TArray<uint32>& ActiveIndices, const int32 MaxLevels = DefaultMaxLevels, const bool bUseVelocity = DefaultUseVelocity, const T Dt = DefaultDt);
@@ -147,4 +172,5 @@ FArchive& operator<<(FArchive& Ar, TBoundingVolumeHierarchy<OBJECT_ARRAY, LEAF_T
 	BVH.Serialize(Ar);
 	return Ar;
 }
+
 }
diff --git a/Engine/Source/Runtime/Experimental/Chaos/Public/Chaos/BoundingVolumeUtilities.h b/Engine/Source/Runtime/Experimental/Chaos/Public/Chaos/BoundingVolumeUtilities.h
--- a/Engine/Source/Runtime/Experimental/Chaos/Public/Chaos/BoundingVolumeUtilities.h
+++ b/Engine/Source/Runtime/Experimental/Chaos/Public/Chaos/BoundingVolumeUtilities.h
@@ -7,6 +7,7 @@
 #include "Chaos/Framework/Parallel.h"
 #include "Chaos/GeometryParticles.h"
 #include "Chaos/ImplicitObject.h"
+#include "Chaos/ParticleHandle.h"
 #include "Chaos/PBDRigidParticles.h"
 #include "Chaos/Particles.h"
 #include "Chaos/Sphere.h"
@@ -42,6 +43,21 @@ bool HasBoundingBox(const TPBDRigidParticles<T, d>& Objects, const int32 i)
 	return Objects.CollisionParticles(i) != nullptr && Objects.CollisionParticles(i)->Size() > 0;
 }
 
+template<typename THandle>
+bool HasBoundingBox(const THandle& Handle)
+{
+	if (Handle.Geometry())
+	{
+		return Handle.Geometry()->HasBoundingBox();
+	}
+	
+	if (const auto PBDRigid = Handle.AsDynamic())
+	{
+		return PBDRigid->CollisionParticles() && PBDRigid->CollisionParticles()->Size() > 0;
+	}
+	return false;
+}
+
 template<class OBJECT_ARRAY, class T, int d>
 const TBox<T, d>& GetWorldSpaceBoundingBox(const OBJECT_ARRAY& Objects, const int32 i, const TMap<int32, TBox<T, d>>& WorldSpaceBoxes)
 {
@@ -98,6 +114,30 @@ TBox<T, d> ComputeWorldSpaceBoundingBox(const TPBDRigidParticles<T, d>& Objects,
 	return LocalBoundingBox.TransformedBox(LocalToWorld);
 }
 
+template<typename THandle>
+TBox<typename THandle::TType, THandle::D> ComputeWorldSpaceBoundingBox(const THandle& Handle)
+{
+	using T = typename THandle::TType;
+	constexpr int d = THandle::D;
+
+	const auto PBDRigid = Handle.AsDynamic();
+	TRigidTransform<T, d> LocalToWorld = PBDRigid ? TRigidTransform<T,d>(PBDRigid->P(), PBDRigid->Q()) : TRigidTransform<T,d>(Handle.X(), Handle.R());
+	if (Handle.Geometry())
+	{
+		const auto& LocalBoundingBox = Handle.Geometry()->BoundingBox();
+		return LocalBoundingBox.TransformedBox(LocalToWorld);
+	}
+
+	check(PBDRigid);
+	check(PBDRigid->CollisionParticles() && PBDRigid->CollisionParticles()->Size());
+	TBox<T, d> LocalBoundingBox(PBDRigid->CollisionParticles()->X(0), PBDRigid->CollisionParticles()->X(0));
+	for (uint32 j = 1; j < PBDRigid->CollisionParticles()->Size(); ++j)
+	{
+		LocalBoundingBox.GrowToInclude(PBDRigid->CollisionParticles()->X(j));
+	}
+	return LocalBoundingBox.TransformedBox(LocalToWorld);
+}
+
 template<typename OBJECT_ARRAY, typename T, int d>
 const TBox<T, d> ComputeGlobalBoxAndSplitAxis(const OBJECT_ARRAY& Objects, const TArray<int32>& AllObjects, const TMap<int32, TBox<T, d>>& WorldSpaceBoxes, bool bAllowMultipleSplitting, int32& OutAxis)
 {
@@ -221,6 +261,18 @@ TVector<T, d> ComputeThickness(const TPBDRigidParticles<T, d>& InParticles, T Dt
 	return AbsVelocity;
 }
 
+template <typename THandle, typename T>
+TVector<T, THandle::D> ComputeThickness(const THandle& PBDRigid, T Dt)
+{
+	auto AbsVelocity = PBDRigid.V().GetAbs();
+	for (int i = 0; i < THandle::D; ++i)
+	{
+		AbsVelocity[i] = FMath::Max(MinBoundsThickness, AbsVelocity[i] * Dt * BoundsThicknessMultiplier);//todo(ocohen): ignoring MThickness for now
+	}
+
+	return AbsVelocity;
+}
+
 template<class T, int d>
 void ComputeAllWorldSpaceBoundingBoxes(const TPBDRigidParticles<T, d>& Objects, const TArray<int32>& AllObjects, const bool bUseVelocity, const T Dt, TMap<int32, TBox<T, d>>& WorldSpaceBoxes)
 {
@@ -233,12 +285,33 @@ void ComputeAllWorldSpaceBoundingBoxes(const TPBDRigidParticles<T, d>& Objects,
 		WorldSpaceBox = ComputeWorldSpaceBoundingBox(Objects, BodyIndex);
 		if (bUseVelocity)
 		{
-			WorldSpaceBox.Thicken(ComputeThickness(Objects, Dt, BodyIndex));
+			WorldSpaceBox.ThickenSymmetrically(ComputeThickness(Objects, Dt, BodyIndex));
 		}
 	}
 	//});
 }
 
+//todo: how do we protect ourselves and make it const?
+template<typename ParticleView, typename T, int d>
+void ComputeAllWorldSpaceBoundingBoxes(const ParticleView& Particles, const TArray<bool>& RequiresBounds, const bool bUseVelocity, const T Dt, TArray<TBox<T, d>>& WorldSpaceBoxes)
+{
+	WorldSpaceBoxes.AddUninitialized(Particles.Num());
+	ParticlesParallelFor(Particles, [&RequiresBounds, &WorldSpaceBoxes, bUseVelocity, Dt](const auto& Particle, int32 Index)
+	{
+		if (RequiresBounds[Index])
+		{
+			WorldSpaceBoxes[Index] = ComputeWorldSpaceBoundingBox(Particle);
+			if (bUseVelocity)
+			{
+				if (const auto PBDRigid = Particle.AsDynamic())
+				{
+					WorldSpaceBoxes.Last().ThickenSymmetrically(ComputeThickness(*PBDRigid, Dt));
+				}
+			}
+		}
+	});
+}
+
 template<class OBJECT_ARRAY>
 int32 GetObjectCount(const OBJECT_ARRAY& Objects)
 {
diff --git a/Engine/Source/Runtime/Experimental/Chaos/Public/Chaos/Box.h b/Engine/Source/Runtime/Experimental/Chaos/Public/Chaos/Box.h
--- a/Engine/Source/Runtime/Experimental/Chaos/Public/Chaos/Box.h
+++ b/Engine/Source/Runtime/Experimental/Chaos/Public/Chaos/Box.h
@@ -207,7 +207,7 @@ namespace Chaos
 			}
 		}
 
-		virtual bool CHAOS_API Raycast(const TVector<T, d>& StartPoint, const TVector<T, d>& Dir, const T Length, const T Thickness, T& OutTime, TVector<T, d>& OutPosition, TVector<T, d>& OutNormal) const override;
+		virtual bool CHAOS_API Raycast(const TVector<T, d>& StartPoint, const TVector<T, d>& Dir, const T Length, const T Thickness, T& OutTime, TVector<T, d>& OutPosition, TVector<T, d>& OutNormal, int32& OutFaceIndex) const override;
 
 		bool RaycastFast(const TVector<T, d>& StartPoint, const TVector<T, d>& Dir, const TVector<T, d>& InvDir, const bool* bParallel, const T Length, const T InvLength, T& OutTime, TVector<T, d>& OutPosition) const
 		{
@@ -261,7 +261,7 @@ namespace Chaos
 				return false;
 			}
 
-			OutTime = LatestStartTime * InvLength;
+			OutTime = LatestStartTime;
 			OutPosition = StartPoint + LatestStartTime * Dir;
 			return true;
 		}
@@ -389,10 +389,11 @@ namespace Chaos
 			MMax += TVector<T, d>(Thickness);
 		}
 
-		FORCEINLINE void Thicken(const TVector<T, d> Thickness)
+		//Grows (or shrinks) the box by this vector symmetrically - Changed name because previous Thicken had different semantics which caused several bugs
+		FORCEINLINE void ThickenSymmetrically(const TVector<T,d>& Thickness)
 		{
-			GrowToInclude(MMin + Thickness);
-			GrowToInclude(MMax + Thickness);
+			MMin -= Thickness;
+			MMax += Thickness;
 		}
 
 		FORCEINLINE TVector<T, d> Center() const { return (MMax - MMin) / (T)2 + MMin; }
@@ -454,6 +455,11 @@ namespace Chaos
 		static TBox<T, d> EmptyBox() { return TBox<T, d>(TVector<T, d>(TNumericLimits<T>::Max()), TVector<T, d>(-TNumericLimits<T>::Max())); }
 		static TBox<T, d> ZeroBox() { return TBox<T, d>(TVector<T, d>((T)0), TVector<T, d>((T)0)); }
 
+		virtual uint32 GetTypeHash() const override 
+		{
+			return HashCombine(::GetTypeHash(MMin), ::GetTypeHash(MMax));
+		}
+
 	private:
 		TVector<T, d> MMin, MMax;
 	};
diff --git a/Engine/Source/Runtime/Experimental/Chaos/Public/Chaos/Capsule.h b/Engine/Source/Runtime/Experimental/Chaos/Public/Chaos/Capsule.h
--- a/Engine/Source/Runtime/Experimental/Chaos/Public/Chaos/Capsule.h
+++ b/Engine/Source/Runtime/Experimental/Chaos/Public/Chaos/Capsule.h
@@ -127,7 +127,7 @@ namespace Chaos
 		TArray<TVector<T, 3>> ComputeSamplePoints(const T PointsPerUnitArea, const int32 MinPoints = 0, const int32 MaxPoints = 1000) const
 		{ return ComputeSamplePoints(FMath::Clamp(static_cast<int32>(ceil(PointsPerUnitArea * GetArea())), MinPoints, MaxPoints)); }
 
-		T PhiWithNormal(const TVector<T, 3>& x, TVector<T, 3>& Normal) const
+		virtual T PhiWithNormal(const TVector<T, 3>& x, TVector<T, 3>& Normal) const override
 		{
 			auto Dot = FMath::Clamp(TVector<T, 3>::DotProduct(x - MPoint, MVector), (T)0., MHeight);
 			TVector<T, 3> ProjectedPoint = Dot * MVector + MPoint;
@@ -137,7 +137,7 @@ namespace Chaos
 
 		virtual const TBox<T, 3>& BoundingBox() const override { return MLocalBoundingBox; }
 
-		virtual bool Raycast(const TVector<T, 3>& StartPoint, const TVector<T, 3>& Dir, const T Length, const T Thickness, T& OutTime, TVector<T, 3>& OutPosition, TVector<T, 3>& OutNormal) const override
+		virtual bool Raycast(const TVector<T, 3>& StartPoint, const TVector<T, 3>& Dir, const T Length, const T Thickness, T& OutTime, TVector<T, 3>& OutPosition, TVector<T, 3>& OutNormal, int32& OutFaceIndex) const override
 		{
 			ensure(FMath::IsNearlyEqual(MVector.SizeSquared(), 1, KINDA_SMALL_NUMBER));
 			ensure(FMath::IsNearlyEqual(Dir.SizeSquared(), 1, KINDA_SMALL_NUMBER));
@@ -145,6 +145,7 @@ namespace Chaos
 
 			const T R = MRadius + Thickness;
 			const T R2 = R * R;
+			OutFaceIndex = INDEX_NONE;
 
 			//First check if we are initially overlapping
 			//Find closest point to cylinder core and check if it's inside the inflated capsule
@@ -224,7 +225,7 @@ namespace Chaos
 					const T PositionLengthOnCoreCylinder = TVector<T, 3>::DotProduct(CylinderToSpherePosition, MVector);
 					if (PositionLengthOnCoreCylinder >= 0 && PositionLengthOnCoreCylinder < MHeight)
 					{
-						OutTime = Time / Length;
+						OutTime = Time;
 						OutNormal = (CylinderToSpherePosition - MVector * PositionLengthOnCoreCylinder) / R;
 						OutPosition = SpherePosition - OutNormal * Thickness;
 						return true;
@@ -247,8 +248,8 @@ namespace Chaos
 				T Time1, Time2;
 				TVector<T, 3> Position1, Position2;
 				TVector<T, 3> Normal1, Normal2;
-				bool bHitX1 = X1Sphere.Raycast(StartPoint, Dir, Length, Thickness, Time1, Position1, Normal1);
-				bool bHitX2 = X2Sphere.Raycast(StartPoint, Dir, Length, Thickness, Time2, Position2, Normal2);
+				bool bHitX1 = X1Sphere.Raycast(StartPoint, Dir, Length, Thickness, Time1, Position1, Normal1, OutFaceIndex);
+				bool bHitX2 = X2Sphere.Raycast(StartPoint, Dir, Length, Thickness, Time2, Position2, Normal2, OutFaceIndex);
 
 				if (bHitX1 && bHitX2)
 				{
@@ -338,7 +339,14 @@ namespace Chaos
 		}
 
 		static TRotation<T, 3> GetRotationOfMass()
-		{ return TRotation<T, 3>(TVector<T, 3>(0), 1); }
+		{
+			return TRotation<T, 3>(TVector<T, 3>(0), 1); 
+		}
+
+		virtual uint32 GetTypeHash() const override
+		{
+			return HashCombine(::GetTypeHash(MPoint), ::GetTypeHash(MVector));
+		}
 
 	private:
 		void InitUnionedObjects()
@@ -357,7 +365,6 @@ namespace Chaos
 			return MUnionedObjects->FindClosestIntersection(StartPoint, EndPoint, Thickness);
 		}
 
-	private:
 		TVector<T, 3> MPoint, MVector;
 		T MHeight, MRadius;
 		TBox<T, 3> MLocalBoundingBox;
