Copying //UE4/Dev-Physics to Dev-Main (//UE4/Dev-Main) [at] 7703071

diff --git a/Engine/Source/Runtime/Experimental/Chaos/Public/Chaos/CollisionConvexMesh.h b/Engine/Source/Runtime/Experimental/Chaos/Public/Chaos/CollisionConvexMesh.h
new file mode 100644
--- /dev/null
+++ b/Engine/Source/Runtime/Experimental/Chaos/Public/Chaos/CollisionConvexMesh.h
@@ -0,0 +1,577 @@
+// Copyright 1998-2019 Epic Games, Inc. All Rights Reserved.
+
+#pragma once
+
+#include "CoreMinimal.h"
+#include "Chaos/Box.h"
+#include "TriangleMesh.h"
+#include "Particles.h"
+
+namespace Chaos
+{
+	template <typename T>
+	class TConvexBuilder
+	{
+	public:
+		static void Build(const TParticles<T, 3>& InParticles, TArray <TPlane<T, 3>>& OutPlanes, TParticles<T, 3>& OutSurfaceParticles, TBox<T, 3>& OutLocalBounds)
+		{
+			OutPlanes.Reset();
+			OutSurfaceParticles.Resize(0);
+			OutLocalBounds = TBox<T, 3>::EmptyBox();
+
+			const uint32 NumParticles = InParticles.Size();
+			if(NumParticles == 0)
+			{
+				return;
+			}
+
+			OutLocalBounds = TBox<T, 3>(InParticles.X(0), InParticles.X(0));
+			for(uint32 ParticleIndex = 0; ParticleIndex < NumParticles; ++ParticleIndex)
+			{
+				OutLocalBounds.GrowToInclude(InParticles.X(ParticleIndex));
+			}
+
+			if(NumParticles >= 4)
+			{
+				TArray<TVector<int32, 3>> Indices;
+				BuildConvexHull(InParticles, Indices);
+				OutPlanes.Reserve(Indices.Num());
+				TSet<int32> AllIndices;
+
+				for(const TVector<int32, 3>& Idx : Indices)
+				{
+					TVector<T, 3> Vs[3] = {InParticles.X(Idx[0]), InParticles.X(Idx[1]), InParticles.X(Idx[2])};
+					const TVector<T, 3> Normal = TVector<T, 3>::CrossProduct(Vs[1] - Vs[0], Vs[2] - Vs[0]).GetUnsafeNormal();
+					OutPlanes.Add(TPlane<T, 3>(Vs[0], Normal));
+					AllIndices.Add(Idx[0]);
+					AllIndices.Add(Idx[1]);
+					AllIndices.Add(Idx[2]);
+				}
+
+				OutSurfaceParticles.AddParticles(AllIndices.Num());
+				int32 NewIdx = 0;
+				for(int32 Idx : AllIndices)
+				{
+					OutSurfaceParticles.X(NewIdx++) = InParticles.X(Idx);
+				}
+			}
+			else if(NumParticles == 3)
+			{
+				//special support for triangle
+				TVector<T, 3> Normal = TVector<T, 3>::CrossProduct(InParticles.X(1) - InParticles.X(0), InParticles.X(2) - InParticles.X(0));
+				const T NormalLength = Normal.SafeNormalize();
+
+				TVector<T, 3> A = InParticles.X(1) - InParticles.X(0);
+				TVector<T, 3> B = InParticles.X(2) - InParticles.X(0);
+				TVector<T, 3> CheckNormal = TVector<T, 3>::CrossProduct(A, B);
+
+				//TODO_SQ_IMPLEMENTATION: should do proper cleanup to avoid this
+				//if(ensure(NormalLength > 1e-4))
+				if (NormalLength > 1e-4)
+				{
+					OutPlanes.Add(TPlane<T, 3>(InParticles.X(0), Normal));
+					OutSurfaceParticles.AddParticles(3);
+					OutSurfaceParticles.X(0) = InParticles.X(0);
+					OutSurfaceParticles.X(1) = InParticles.X(1);
+					OutSurfaceParticles.X(2) = InParticles.X(2);
+				}
+
+			}
+		}
+
+		static void BuildConvexHull(const TParticles<T, 3>& InParticles, TArray<TVector<int32, 3>>& OutIndices)
+		{
+			OutIndices.Reset();
+			FConvexFace* Faces = BuildInitialHull(InParticles);
+			if(Faces == nullptr)
+			{
+				return;
+			}
+			FConvexFace DummyFace(Faces->Plane);
+			DummyFace.Prev = nullptr;
+			DummyFace.Next = Faces;
+			Faces->Prev = &DummyFace;
+
+			FHalfEdge* ConflictV = FindConflictVertex(InParticles, DummyFace.Next);
+			while(ConflictV)
+			{
+				AddVertex(InParticles, ConflictV);
+				ConflictV = FindConflictVertex(InParticles, DummyFace.Next);
+			}
+
+			FConvexFace* Cur = DummyFace.Next;
+			while(Cur)
+			{
+				//todo(ocohen): this assumes faces are triangles, not true once face merging is added
+				OutIndices.Add(TVector<int32, 3>(Cur->FirstEdge->Vertex, Cur->FirstEdge->Next->Vertex, Cur->FirstEdge->Next->Next->Vertex));
+				delete Cur->FirstEdge->Next;
+				delete Cur->FirstEdge->Prev;
+				delete Cur->FirstEdge;
+				FConvexFace* Next = Cur->Next;
+				delete Cur;
+				Cur = Next;
+			}
+		}
+
+		static TTriangleMesh<T> BuildConvexHullTriMesh(const TParticles<T, 3>& InParticles)
+		{
+			TArray<TVector<int32, 3>> Indices;
+			BuildConvexHull(InParticles, Indices);
+			return TTriangleMesh<T>(MoveTemp(Indices));
+		}
+
+	private:
+
+		struct FHalfEdge;
+		struct FConvexFace
+		{
+			FConvexFace(const TPlane<T, 3>& FacePlane)
+				: ConflictList(nullptr)
+				, Plane(FacePlane)
+			{
+			}
+
+			FHalfEdge* FirstEdge;
+			FHalfEdge* ConflictList; //Note that these half edges are really just free verts grouped together
+			TPlane<T, 3> Plane;
+			FConvexFace* Prev;
+			FConvexFace* Next; //these have no geometric meaning, just used for book keeping
+		};
+
+		struct FHalfEdge
+		{
+			FHalfEdge(int32 InVertex)
+				: Vertex(InVertex) {}
+			int32 Vertex;
+			FHalfEdge* Prev;
+			FHalfEdge* Next;
+			FHalfEdge* Twin;
+			FConvexFace* Face;
+		};
+
+		static TVector<T, 3> ComputeFaceNormal(const TVector<T, 3>& A, const TVector<T, 3>& B, const TVector<T, 3>& C)
+		{
+			return TVector<T, 3>::CrossProduct((B - A), (C - A));
+		}
+
+		static FConvexFace* CreateFace(const TParticles<T, 3>& InParticles, FHalfEdge* RS, FHalfEdge* ST, FHalfEdge* TR)
+		{
+			RS->Prev = TR;
+			RS->Next = ST;
+			ST->Prev = RS;
+			ST->Next = TR;
+			TR->Prev = ST;
+			TR->Next = RS;
+			TVector<T, 3> RSTNormal = ComputeFaceNormal(InParticles.X(RS->Vertex), InParticles.X(ST->Vertex), InParticles.X(TR->Vertex));
+			const T RSTNormalSize = RSTNormal.Size();
+			check(RSTNormalSize > 1e-4);
+			RSTNormal = RSTNormal * (1 / RSTNormalSize);
+			FConvexFace* RST = new FConvexFace(TPlane<T, 3>(InParticles.X(RS->Vertex), RSTNormal));
+			RST->FirstEdge = RS;
+			RS->Face = RST;
+			ST->Face = RST;
+			TR->Face = RST;
+			return RST;
+		}
+
+		static void StealConflictList(const TParticles<T, 3>& InParticles, FHalfEdge* OldList, FConvexFace** Faces, int32 NumFaces)
+		{
+			FHalfEdge* Cur = OldList;
+			while(Cur)
+			{
+				T MaxD = 1e-4;
+				int32 MaxIdx = -1;
+				for(int32 Idx = 0; Idx < NumFaces; ++Idx)
+				{
+					T Distance = Faces[Idx]->Plane.SignedDistance(InParticles.X(Cur->Vertex));
+					if(Distance > MaxD)
+					{
+						MaxD = Distance;
+						MaxIdx = Idx;
+					}
+				}
+
+				bool bDeleteVertex = MaxIdx == -1;
+				if(!bDeleteVertex)
+				{
+					//let's make sure faces created with this new conflict vertex will be valid. The plane check above is not sufficient because long thin triangles will have a plane with its point at one of these. Combined with normal and precision we can have errors
+					auto PretendNormal = [&InParticles](FHalfEdge* A, FHalfEdge* B, FHalfEdge* C) {
+						return TVector<T, 3>::CrossProduct(InParticles.X(B->Vertex) - InParticles.X(A->Vertex), InParticles.X(C->Vertex) - InParticles.X(A->Vertex)).SizeSquared();
+					};
+					FHalfEdge* Edge = Faces[MaxIdx]->FirstEdge;
+					do
+					{
+						if(PretendNormal(Edge->Prev, Edge, Cur) < 1e-4)
+						{
+							bDeleteVertex = true;
+							break;
+						}
+						Edge = Edge->Next;
+					} while(Edge != Faces[MaxIdx]->FirstEdge);
+				}
+
+				if(!bDeleteVertex)
+				{
+					FHalfEdge* Next = Cur->Next;
+					FHalfEdge*& ConflictList = Faces[MaxIdx]->ConflictList;
+					if(ConflictList)
+					{
+						ConflictList->Prev = Cur;
+					}
+					Cur->Next = ConflictList;
+					ConflictList = Cur;
+					Cur->Prev = nullptr;
+					Cur = Next;
+				}
+				else
+				{
+					//point is contained, we can delete it
+					FHalfEdge* Next = Cur->Next;
+					delete Cur; //todo(ocohen): better allocation strategy
+					Cur = Next;
+				}
+			}
+		}
+
+		static FConvexFace* BuildInitialHull(const TParticles<T, 3>& InParticles)
+		{
+			if(InParticles.Size() < 4) //not enough points
+			{
+				return nullptr;
+			}
+
+			const T Epsilon = 1e-4;
+
+			const int32 NumParticles = InParticles.Size();
+
+			//We store the vertex directly in the half-edge. We use its next to group free vertices by context list
+			//create a starting triangle by finding min/max on X and max on Y
+			T MinX = TNumericLimits<T>::Max();
+			T MaxX = TNumericLimits<T>::Lowest();
+			FHalfEdge* A = nullptr; //min x
+			FHalfEdge* B = nullptr; //max x
+			FHalfEdge DummyHalfEdge(-1);
+			DummyHalfEdge.Next = nullptr;
+			DummyHalfEdge.Prev = nullptr;
+			FHalfEdge* Prev = &DummyHalfEdge;
+
+			for(int32 i = 0; i < NumParticles; ++i)
+			{
+				FHalfEdge* VHalf = new FHalfEdge(i); //todo(ocohen): preallocate these
+				VHalf->Vertex = i;
+				Prev->Next = VHalf;
+				VHalf->Prev = Prev;
+				VHalf->Next = nullptr;
+				const TVector<T, 3>& V = InParticles.X(i);
+
+				if(V[0] < MinX)
+				{
+					MinX = V[0];
+					A = VHalf;
+				}
+				if(V[0] > MaxX)
+				{
+					MaxX = V[0];
+					B = VHalf;
+				}
+
+				Prev = VHalf;
+			}
+
+			check(A && B);
+			if(A == B || (InParticles.X(A->Vertex) - InParticles.X(B->Vertex)).SizeSquared() < Epsilon) //infinitely thin
+			{
+				return nullptr;
+			}
+
+			//remove A and B from conflict list
+			A->Prev->Next = A->Next;
+			if(A->Next)
+			{
+				A->Next->Prev = A->Prev;
+			}
+			B->Prev->Next = B->Next;
+			if(B->Next)
+			{
+				B->Next->Prev = B->Prev;
+			}
+
+			//find C so that we get the biggest base triangle
+			T MaxTriSize = Epsilon;
+			const TVector<T, 3> AToB = InParticles.X(B->Vertex) - InParticles.X(A->Vertex);
+			FHalfEdge* C = nullptr;
+			for(FHalfEdge* V = DummyHalfEdge.Next; V; V = V->Next)
+			{
+				T TriSize = TVector<T, 3>::CrossProduct(AToB, InParticles.X(V->Vertex) - InParticles.X(A->Vertex)).SizeSquared();
+				if(TriSize > MaxTriSize)
+				{
+					MaxTriSize = TriSize;
+					C = V;
+				}
+			}
+
+			if(C == nullptr) //biggest triangle is tiny
+			{
+				return nullptr;
+			}
+
+			//remove C from conflict list
+			C->Prev->Next = C->Next;
+			if(C->Next)
+			{
+				C->Next->Prev = C->Prev;
+			}
+
+			//find farthest D along normal
+			const TVector<T, 3> AToC = InParticles.X(C->Vertex) - InParticles.X(A->Vertex);
+			const TVector<T, 3> Normal = TVector<T, 3>::CrossProduct(AToB, AToC);
+
+			T MaxPosDistance = Epsilon;
+			T MaxNegDistance = Epsilon;
+			FHalfEdge* PosD = nullptr;
+			FHalfEdge* NegD = nullptr;
+			for(FHalfEdge* V = DummyHalfEdge.Next; V; V = V->Next)
+			{
+				T Dot = TVector<T, 3>::DotProduct(InParticles.X(V->Vertex) - InParticles.X(A->Vertex), Normal);
+				if(Dot > MaxPosDistance)
+				{
+					MaxPosDistance = Dot;
+					PosD = V;
+				}
+				if(-Dot > MaxNegDistance)
+				{
+					MaxNegDistance = -Dot;
+					NegD = V;
+				}
+			}
+
+			if(MaxNegDistance == Epsilon && MaxPosDistance == Epsilon)
+			{
+				return nullptr; //plane
+			}
+
+			const bool bPositive = MaxNegDistance < MaxPosDistance;
+			FHalfEdge* D = bPositive ? PosD : NegD;
+
+			//remove D from conflict list
+			D->Prev->Next = D->Next;
+			if(D->Next)
+			{
+				D->Next->Prev = D->Prev;
+			}
+
+			//we must now create the 3 faces. Face must be oriented CCW around normal and positive normal should face out
+			//Note we are now using A,B,C,D as edges. We can only use one edge per face so once they're used we'll need new ones
+			FHalfEdge* Edges[4] = {A, B, C, D};
+
+			//The base is a plane with Edges[0], Edges[1], Edges[2]. The order depends on which side D is on
+			if(bPositive)
+			{
+				//D is on the positive side of Edges[0], Edges[1], Edges[2] so we must reorder it
+				FHalfEdge* Tmp = Edges[0];
+				Edges[0] = Edges[1];
+				Edges[1] = Tmp;
+			}
+
+			FConvexFace* Faces[4];
+			Faces[0] = CreateFace(InParticles, Edges[0], Edges[1], Edges[2]); //base
+			Faces[1] = CreateFace(InParticles, new FHalfEdge(Edges[1]->Vertex), new FHalfEdge(Edges[0]->Vertex), Edges[3]);
+			Faces[2] = CreateFace(InParticles, new FHalfEdge(Edges[0]->Vertex), new FHalfEdge(Edges[2]->Vertex), new FHalfEdge(Edges[3]->Vertex));
+			Faces[3] = CreateFace(InParticles, new FHalfEdge(Edges[2]->Vertex), new FHalfEdge(Edges[1]->Vertex), new FHalfEdge(Edges[3]->Vertex));
+
+			auto MakeTwins = [](FHalfEdge* E1, FHalfEdge* E2) {
+				E1->Twin = E2;
+				E2->Twin = E1;
+			};
+			//mark twins so half edge can cross faces
+			MakeTwins(Edges[0], Faces[1]->FirstEdge); //0-1 1-0
+			MakeTwins(Edges[1], Faces[3]->FirstEdge); //1-2 2-1
+			MakeTwins(Edges[2], Faces[2]->FirstEdge); //2-0 0-2
+			MakeTwins(Faces[1]->FirstEdge->Next, Faces[2]->FirstEdge->Prev); //0-3 3-0
+			MakeTwins(Faces[1]->FirstEdge->Prev, Faces[3]->FirstEdge->Next); //3-1 1-3
+			MakeTwins(Faces[2]->FirstEdge->Next, Faces[3]->FirstEdge->Prev); //2-3 3-2
+
+			Faces[0]->Prev = nullptr;
+			for(int i = 1; i < 4; ++i)
+			{
+				Faces[i - 1]->Next = Faces[i];
+				Faces[i]->Prev = Faces[i - 1];
+			}
+			Faces[3]->Next = nullptr;
+
+			//split up the conflict list
+			StealConflictList(InParticles, DummyHalfEdge.Next, Faces, 4);
+			return Faces[0];
+		}
+
+		static FHalfEdge* FindConflictVertex(const TParticles<T, 3>& InParticles, FConvexFace* FaceList)
+		{
+			for(FConvexFace* CurFace = FaceList; CurFace; CurFace = CurFace->Next)
+			{
+				T MaxD = TNumericLimits<T>::Lowest();
+				FHalfEdge* MaxV = nullptr;
+				for(FHalfEdge* CurFaceVertex = CurFace->ConflictList; CurFaceVertex; CurFaceVertex = CurFaceVertex->Next)
+				{
+					//is it faster to cache this from stealing stage?
+					T Dist = TVector<T, 3>::DotProduct(InParticles.X(CurFaceVertex->Vertex), CurFace->Plane.Normal());
+					if(Dist > MaxD)
+					{
+						MaxD = Dist;
+						MaxV = CurFaceVertex;
+					}
+				}
+				check(CurFace->ConflictList == nullptr || MaxV);
+				if(MaxV)
+				{
+					if(MaxV->Prev)
+					{
+						MaxV->Prev->Next = MaxV->Next;
+					}
+					if(MaxV->Next)
+					{
+						MaxV->Next->Prev = MaxV->Prev;
+					}
+					if(MaxV == CurFace->ConflictList)
+					{
+						CurFace->ConflictList = MaxV->Next;
+					}
+					MaxV->Face = CurFace;
+					return MaxV;
+				}
+			}
+
+			return nullptr;
+		}
+
+		static void BuildHorizon(const TParticles<T, 3>& InParticles, FHalfEdge* ConflictV, TArray<FHalfEdge*>& HorizonEdges, TArray<FConvexFace*>& FacesToDelete)
+		{
+			//We must flood fill from the initial face and mark edges of faces the conflict vertex cannot see
+			//In order to return a CCW ordering we must traverse each face in CCW order from the edge we crossed over
+			//This should already be the ordering in the half edge
+			const T Epsilon = 1e-1;
+			const TVector<T, 3> V = InParticles.X(ConflictV->Vertex);
+			TSet<FConvexFace*> Processed;
+			TArray<FHalfEdge*> Queue;
+			check(ConflictV->Face);
+			Queue.Add(ConflictV->Face->FirstEdge->Prev); //stack pops so reverse order
+			Queue.Add(ConflictV->Face->FirstEdge->Next);
+			Queue.Add(ConflictV->Face->FirstEdge);
+			FacesToDelete.Add(ConflictV->Face);
+			while(Queue.Num())
+			{
+				FHalfEdge* Edge = Queue.Pop(/*bAllowShrinking=*/false);
+				Processed.Add(Edge->Face);
+				FHalfEdge* Twin = Edge->Twin;
+				FConvexFace* NextFace = Twin->Face;
+				if(Processed.Contains(NextFace))
+				{
+					continue;
+				}
+				const T Distance = NextFace->Plane.SignedDistance(V);
+				if(Distance > Epsilon)
+				{
+					Queue.Add(Twin->Prev); //stack pops so reverse order
+					Queue.Add(Twin->Next);
+					FacesToDelete.Add(NextFace);
+				}
+				else
+				{
+					HorizonEdges.Add(Edge);
+				}
+			}
+		}
+
+		static void BuildFaces(const TParticles<T, 3>& InParticles, const FHalfEdge* ConflictV, const TArray<FHalfEdge*>& HorizonEdges, const TArray<FConvexFace*> OldFaces, TArray<FConvexFace*>& NewFaces)
+		{
+			//The HorizonEdges are in CCW order. We must make new faces and edges to join from ConflictV to these edges
+			check(HorizonEdges.Num() >= 3);
+			NewFaces.Reserve(HorizonEdges.Num());
+			FHalfEdge* PrevEdge = nullptr;
+			for(int32 HorizonIdx = 0; HorizonIdx < HorizonEdges.Num(); ++HorizonIdx)
+			{
+				FHalfEdge* OriginalEdge = HorizonEdges[HorizonIdx];
+				FHalfEdge* NewHorizonEdge = new FHalfEdge(OriginalEdge->Vertex);
+				NewHorizonEdge->Twin = OriginalEdge->Twin; //swap edges
+				NewHorizonEdge->Twin->Twin = NewHorizonEdge;
+				FHalfEdge* HorizonNext = new FHalfEdge(OriginalEdge->Next->Vertex);
+				check(HorizonNext->Vertex == HorizonEdges[(HorizonIdx + 1) % HorizonEdges.Num()]->Vertex); //should be ordered properly
+				FHalfEdge* V = new FHalfEdge(ConflictV->Vertex);
+				V->Twin = PrevEdge;
+				if(PrevEdge)
+				{
+					PrevEdge->Twin = V;
+				}
+				PrevEdge = HorizonNext;
+
+				//link new faces together
+				FConvexFace* NewFace = CreateFace(InParticles, NewHorizonEdge, HorizonNext, V);
+				if(NewFaces.Num() > 0)
+				{
+					NewFace->Prev = NewFaces[NewFaces.Num() - 1];
+					NewFaces[NewFaces.Num() - 1]->Next = NewFace;
+				}
+				else
+				{
+					NewFace->Prev = nullptr;
+				}
+				NewFaces.Add(NewFace);
+			}
+
+			check(PrevEdge);
+			NewFaces[0]->FirstEdge->Prev->Twin = PrevEdge;
+			PrevEdge->Twin = NewFaces[0]->FirstEdge->Prev;
+			NewFaces[NewFaces.Num() - 1]->Next = nullptr;
+
+			//redistribute conflict list
+			for(FConvexFace* OldFace : OldFaces)
+			{
+				StealConflictList(InParticles, OldFace->ConflictList, &NewFaces[0], NewFaces.Num());
+			}
+
+			//insert all new faces after conflict vertex face
+			FConvexFace* OldFace = ConflictV->Face;
+			FConvexFace* StartFace = NewFaces[0];
+			FConvexFace* EndFace = NewFaces[NewFaces.Num() - 1];
+			if(OldFace->Next)
+			{
+				OldFace->Next->Prev = EndFace;
+			}
+			EndFace->Next = OldFace->Next;
+			OldFace->Next = StartFace;
+			StartFace->Prev = OldFace;
+		}
+
+		static void AddVertex(const TParticles<T, 3>& InParticles, FHalfEdge* ConflictV)
+		{
+			TArray<FHalfEdge*> HorizonEdges;
+			TArray<FConvexFace*> FacesToDelete;
+			BuildHorizon(InParticles, ConflictV, HorizonEdges, FacesToDelete);
+
+			TArray<FConvexFace*> NewFaces;
+			BuildFaces(InParticles, ConflictV, HorizonEdges, FacesToDelete, NewFaces);
+
+			for(FConvexFace* Face : FacesToDelete)
+			{
+				FHalfEdge* Edge = Face->FirstEdge;
+				do
+				{
+					FHalfEdge* Next = Edge->Next;
+					delete Edge;
+					Edge = Next;
+				} while(Edge != Face->FirstEdge);
+				if(Face->Prev)
+				{
+					Face->Prev->Next = Face->Next;
+				}
+				if(Face->Next)
+				{
+					Face->Next->Prev = Face->Prev;
+				}
+				delete Face;
+			}
+
+			//todo(ocohen): need to explicitly test for merge failures. Coplaner, nonconvex, etc...
+			//getting this in as is for now to unblock other systems
+		}
+	};
+}
\ No newline at end of file
diff --git a/Engine/Source/Runtime/Experimental/Chaos/Public/Chaos/CollisionFilterData.h b/Engine/Source/Runtime/Experimental/Chaos/Public/Chaos/CollisionFilterData.h
new file mode 100644
--- /dev/null
+++ b/Engine/Source/Runtime/Experimental/Chaos/Public/Chaos/CollisionFilterData.h
@@ -0,0 +1,19 @@
+// Copyright 1998-2019 Epic Games, Inc. All Rights Reserved.
+#pragma once
+
+#if INCLUDE_CHAOS
+
+struct FCollisionFilterData
+{
+	uint32 Word0;
+	uint32 Word1;
+	uint32 Word2;
+	uint32 Word3;
+
+	FORCEINLINE FCollisionFilterData()
+	{
+		Word0 = Word1 = Word2 = Word3 = 0;
+	}
+};
+
+#endif
\ No newline at end of file
diff --git a/Engine/Source/Runtime/Experimental/Chaos/Public/Chaos/Convex.h b/Engine/Source/Runtime/Experimental/Chaos/Public/Chaos/Convex.h
--- a/Engine/Source/Runtime/Experimental/Chaos/Public/Chaos/Convex.h
+++ b/Engine/Source/Runtime/Experimental/Chaos/Public/Chaos/Convex.h
@@ -1,9 +1,10 @@
 // Copyright 1998-2019 Epic Games, Inc. All Rights Reserved.
 #pragma once
 
-#include "Chaos/Box.h"
 #include "Chaos/ImplicitObject.h"
+#include "Chaos/Box.h"
 #include "TriangleMesh.h"
+#include "CollisionConvexMesh.h"
 
 namespace Chaos
 {
@@ -11,15 +12,16 @@ namespace Chaos
 	class TConvex final : public TImplicitObject<T, d>
 	{
 	public:
+		IMPLICIT_OBJECT_SERIALIZER(TConvex);
 		TConvex()
 		    : TImplicitObject<T,3>(EImplicitObject::IsConvex | EImplicitObject::HasBoundingBox, ImplicitObjectType::Convex)
 		{}
 		TConvex(const TConvex&) = delete;
 		TConvex(TConvex&& Other)
 		    : TImplicitObject<T,3>(EImplicitObject::IsConvex | EImplicitObject::HasBoundingBox, ImplicitObjectType::Convex)
-			, MPlanes(MoveTemp(Other.MPlanes))
-		    , MSurfaceParticles(MoveTemp(Other.MSurfaceParticles))
-		    , MLocalBoundingBox(MoveTemp(Other.MLocalBoundingBox))
+			, Planes(MoveTemp(Other.Planes))
+		    , SurfaceParticles(MoveTemp(Other.SurfaceParticles))
+		    , LocalBoundingBox(MoveTemp(Other.LocalBoundingBox))
 		{}
 		TConvex(const TParticles<T, 3>& InParticles)
 		    : TImplicitObject<T, d>(EImplicitObject::IsConvex | EImplicitObject::HasBoundingBox, ImplicitObjectType::Convex)
@@ -30,68 +32,22 @@ namespace Chaos
 				return;
 			}
 
-			MLocalBoundingBox = TBox<T, 3>(InParticles.X(0), InParticles.X(0));
-
-			for (uint32 i = 0; i < NumParticles; ++i)
-			{
-				MLocalBoundingBox.GrowToInclude(InParticles.X(i));
-			}
-
-			if(NumParticles >= 4)
-			{
-				TArray<TVector<int32, 3>> Indices;
-				BuildConvexHull(InParticles, Indices);
-				MPlanes.Reserve(Indices.Num());
-				TSet<int32> AllIndices;
-
-				for (const TVector<int32, 3>& Idx : Indices)
-				{
-					TVector<T, d> Vs[3] = { InParticles.X(Idx[0]), InParticles.X(Idx[1]), InParticles.X(Idx[2]) };
-					const TVector<T, d> Normal = TVector<T, d>::CrossProduct(Vs[1] - Vs[0], Vs[2] - Vs[0]).GetUnsafeNormal();
-					MPlanes.Add(TPlane<T, d>(Vs[0], Normal));
-					AllIndices.Add(Idx[0]);
-					AllIndices.Add(Idx[1]);
-					AllIndices.Add(Idx[2]);
-				}
-
-				MSurfaceParticles.AddParticles(AllIndices.Num());
-				int32 NewIdx = 0;
-				for (int32 Idx : AllIndices)
-				{
-					MSurfaceParticles.X(NewIdx++) = InParticles.X(Idx);
-				}
-			}
-			else if (NumParticles == 3)
-			{
-				//special support for triangle
-				TVector<T, d> Normal = TVector<T, d>::CrossProduct(InParticles.X(1) - InParticles.X(0), InParticles.X(2) - InParticles.X(0));
-				const T NormalLength = Normal.SafeNormalize();
-
-				if (ensure(NormalLength > 1e-4))
-				{
-					MPlanes.Add(TPlane<T, d>(InParticles.X(0), Normal));
-					MSurfaceParticles.AddParticles(3);
-					MSurfaceParticles.X(0) = InParticles.X(0);
-					MSurfaceParticles.X(1) = InParticles.X(1);
-					MSurfaceParticles.X(2) = InParticles.X(2);
-				}
-
-			}
+			TConvexBuilder<T>::Build(InParticles, Planes, SurfaceParticles, LocalBoundingBox);
 		}
 
 		static ImplicitObjectType GetType()
 		{
-			return ImplicitObjectType::Unknown;
+			return ImplicitObjectType::Convex;
 		}
 
 		virtual const TBox<T, d>& BoundingBox() const override
 		{
-			return MLocalBoundingBox;
+			return LocalBoundingBox;
 		}
 
 		virtual T PhiWithNormal(const TVector<T, d>& x, TVector<T, d>& Normal) const override
 		{
-			const int32 NumPlanes = MPlanes.Num();
+			const int32 NumPlanes = Planes.Num();
 			if (NumPlanes == 0)
 			{
 				return FLT_MAX;
@@ -103,7 +59,7 @@ namespace Chaos
 
 			for (int32 Idx = 0; Idx < NumPlanes; ++Idx)
 			{
-				const T Phi = MPlanes[Idx].SignedDistance(x);
+				const T Phi = Planes[Idx].SignedDistance(x);
 				if (Phi > MaxPhi)
 				{
 					MaxPhi = Phi;
@@ -111,17 +67,17 @@ namespace Chaos
 				}
 			}
 
-			return MPlanes[MaxPlane].PhiWithNormal(x, Normal);
+			return Planes[MaxPlane].PhiWithNormal(x, Normal);
 		}
 
 		virtual Pair<TVector<T, d>, bool> FindClosestIntersectionImp(const TVector<T, d>& StartPoint, const TVector<T, d>& EndPoint, const T Thickness) const override
 		{
-			const int32 NumPlanes = MPlanes.Num();
+			const int32 NumPlanes = Planes.Num();
 			TArray<Pair<T, TVector<T, 3>>> Intersections;
 			Intersections.Reserve(NumPlanes);
 			for (int32 Idx = 0; Idx < NumPlanes; ++Idx)
 			{
-				auto PlaneIntersection = MPlanes[Idx].FindClosestIntersection(StartPoint, EndPoint, Thickness);
+				auto PlaneIntersection = Planes[Idx].FindClosestIntersection(StartPoint, EndPoint, Thickness);
 				if (PlaneIntersection.Second)
 				{
 					Intersections.Add(MakePair((PlaneIntersection.First - StartPoint).Size(), PlaneIntersection.First));
@@ -142,12 +98,12 @@ namespace Chaos
 		{
 			T MaxDot = TNumericLimits<T>::Lowest();
 			int32 MaxVIdx = 0;
-			const int32 NumVertices = MSurfaceParticles.Size();
+			const int32 NumVertices = SurfaceParticles.Size();
 
 			check(NumVertices > 0);
 			for (int32 Idx = 0; Idx < NumVertices; ++Idx)
 			{
-				const T Dot = TVector<T, d>::DotProduct(MSurfaceParticles.X(Idx), Direction);
+				const T Dot = TVector<T, d>::DotProduct(SurfaceParticles.X(Idx), Direction);
 				if (Dot > MaxDot)
 				{
 					MaxDot = Dot;
@@ -155,7 +111,7 @@ namespace Chaos
 				}
 			}
 
-			return MSurfaceParticles.X(MaxVIdx);
+			return SurfaceParticles.X(MaxVIdx);
 		}
 
 		virtual FString ToString() const
@@ -163,509 +119,44 @@ namespace Chaos
 			return FString::Printf(TEXT("Convex"));
 		}
 
-		const TParticles<T, d>& SurfaceParticles() const
+		const TParticles<T, d>& GetSurfaceParticles() const
 		{
-			return MSurfaceParticles;
-		}
-
-	private:
-		TArray<TPlane<T, d>> MPlanes;
-		TParticles<T, d> MSurfaceParticles;	//copy of the vertices that are just on the convex hull boundary
-		TBox<T, d> MLocalBoundingBox;
-
-	private:
-		struct FHalfEdge;
-		struct FConvexFace
-		{
-			FConvexFace(const TPlane<T, 3>& FacePlane)
-			    : ConflictList(nullptr)
-			    , Plane(FacePlane)
-			{
-			}
-
-			FHalfEdge* FirstEdge;
-			FHalfEdge* ConflictList; //Note that these half edges are really just free verts grouped together
-			TPlane<T, 3> Plane;
-			FConvexFace* Prev;
-			FConvexFace* Next; //these have no geometric meaning, just used for book keeping
-		};
-
-		struct FHalfEdge
-		{
-			FHalfEdge(int32 InVertex)
-			    : Vertex(InVertex) {}
-			int32 Vertex;
-			FHalfEdge* Prev;
-			FHalfEdge* Next;
-			FHalfEdge* Twin;
-			FConvexFace* Face;
-		};
-
-		static TVector<T, 3> ComputeFaceNormal(const TVector<T, 3>& A, const TVector<T, 3>& B, const TVector<T, 3>& C)
-		{
-			return TVector<T, 3>::CrossProduct((B - A), (C - A));
-		}
-
-		static FConvexFace* CreateFace(const TParticles<T, 3>& InParticles, FHalfEdge* RS, FHalfEdge* ST, FHalfEdge* TR)
-		{
-			RS->Prev = TR;
-			RS->Next = ST;
-			ST->Prev = RS;
-			ST->Next = TR;
-			TR->Prev = ST;
-			TR->Next = RS;
-			TVector<T, 3> RSTNormal = ComputeFaceNormal(InParticles.X(RS->Vertex), InParticles.X(ST->Vertex), InParticles.X(TR->Vertex));
-			const T RSTNormalSize = RSTNormal.Size();
-			check(RSTNormalSize > 1e-4);
-			RSTNormal = RSTNormal * (1 / RSTNormalSize);
-			FConvexFace* RST = new FConvexFace(TPlane<T, 3>(InParticles.X(RS->Vertex), RSTNormal));
-			RST->FirstEdge = RS;
-			RS->Face = RST;
-			ST->Face = RST;
-			TR->Face = RST;
-			return RST;
-		}
-
-		static void StealConflictList(const TParticles<T, 3>& InParticles, FHalfEdge* OldList, FConvexFace** Faces, int32 NumFaces)
-		{
-			FHalfEdge* Cur = OldList;
-			while (Cur)
-			{
-				T MaxD = 1e-4;
-				int32 MaxIdx = -1;
-				for (int32 Idx = 0; Idx < NumFaces; ++Idx)
-				{
-					T Distance = Faces[Idx]->Plane.SignedDistance(InParticles.X(Cur->Vertex));
-					if (Distance > MaxD)
-					{
-						MaxD = Distance;
-						MaxIdx = Idx;
-					}
-				}
-
-				bool bDeleteVertex = MaxIdx == -1;
-				if (!bDeleteVertex)
-				{
-					//let's make sure faces created with this new conflict vertex will be valid. The plane check above is not sufficient because long thin triangles will have a plane with its point at one of these. Combined with normal and precision we can have errors
-					auto PretendNormal = [&InParticles](FHalfEdge* A, FHalfEdge* B, FHalfEdge* C) {
-						return TVector<T, d>::CrossProduct(InParticles.X(B->Vertex) - InParticles.X(A->Vertex), InParticles.X(C->Vertex) - InParticles.X(A->Vertex)).SizeSquared();
-					};
-					FHalfEdge* Edge = Faces[MaxIdx]->FirstEdge;
-					do
-					{
-						if (PretendNormal(Edge->Prev, Edge, Cur) < 1e-4)
-						{
-							bDeleteVertex = true;
-							break;
-						}
-						Edge = Edge->Next;
-					} while (Edge != Faces[MaxIdx]->FirstEdge);
-				}
-
-				if (!bDeleteVertex)
-				{
-					FHalfEdge* Next = Cur->Next;
-					FHalfEdge*& ConflictList = Faces[MaxIdx]->ConflictList;
-					if (ConflictList)
-					{
-						ConflictList->Prev = Cur;
-					}
-					Cur->Next = ConflictList;
-					ConflictList = Cur;
-					Cur->Prev = nullptr;
-					Cur = Next;
-				}
-				else
-				{
-					//point is contained, we can delete it
-					FHalfEdge* Next = Cur->Next;
-					delete Cur; //todo(ocohen): better allocation strategy
-					Cur = Next;
-				}
-			}
+			return SurfaceParticles;
 		}
 
-		static FConvexFace* BuildInitialHull(const TParticles<T, 3>& InParticles)
+		virtual uint32 GetTypeHash() const override
 		{
-			if (InParticles.Size() < 4) //not enough points
-			{
-				return nullptr;
-			}
-
-			const T Epsilon = 1e-4;
-
-			const int32 NumParticles = InParticles.Size();
-
-			//We store the vertex directly in the half-edge. We use its next to group free vertices by context list
-			//create a starting triangle by finding min/max on X and max on Y
-			T MinX = TNumericLimits<T>::Max();
-			T MaxX = TNumericLimits<T>::Lowest();
-			FHalfEdge* A = nullptr; //min x
-			FHalfEdge* B = nullptr; //max x
-			FHalfEdge DummyHalfEdge(-1);
-			DummyHalfEdge.Next = nullptr;
-			DummyHalfEdge.Prev = nullptr;
-			FHalfEdge* Prev = &DummyHalfEdge;
-
-			for (int32 i = 0; i < NumParticles; ++i)
-			{
-				FHalfEdge* VHalf = new FHalfEdge(i); //todo(ocohen): preallocate these
-				VHalf->Vertex = i;
-				Prev->Next = VHalf;
-				VHalf->Prev = Prev;
-				VHalf->Next = nullptr;
-				const TVector<T, 3>& V = InParticles.X(i);
-
-				if (V[0] < MinX)
-				{
-					MinX = V[0];
-					A = VHalf;
-				}
-				if (V[0] > MaxX)
-				{
-					MaxX = V[0];
-					B = VHalf;
-				}
-
-				Prev = VHalf;
-			}
-
-			check(A && B);
-			if (A == B || (InParticles.X(A->Vertex) - InParticles.X(B->Vertex)).SizeSquared() < Epsilon) //infinitely thin
-			{
-				return nullptr;
-			}
-
-			//remove A and B from conflict list
-			A->Prev->Next = A->Next;
-			if (A->Next)
-			{
-				A->Next->Prev = A->Prev;
-			}
-			B->Prev->Next = B->Next;
-			if (B->Next)
-			{
-				B->Next->Prev = B->Prev;
-			}
-
-			//find C so that we get the biggest base triangle
-			T MaxTriSize = Epsilon;
-			const TVector<T, 3> AToB = InParticles.X(B->Vertex) - InParticles.X(A->Vertex);
-			FHalfEdge* C = nullptr;
-			for (FHalfEdge* V = DummyHalfEdge.Next; V; V = V->Next)
-			{
-				T TriSize = TVector<T, 3>::CrossProduct(AToB, InParticles.X(V->Vertex) - InParticles.X(A->Vertex)).SizeSquared();
-				if (TriSize > MaxTriSize)
-				{
-					MaxTriSize = TriSize;
-					C = V;
-				}
-			}
-
-			if (C == nullptr) //biggest triangle is tiny
-			{
-				return nullptr;
-			}
-
-			//remove C from conflict list
-			C->Prev->Next = C->Next;
-			if (C->Next)
-			{
-				C->Next->Prev = C->Prev;
-			}
+			uint32 Result = LocalBoundingBox.GetTypeHash();
 
-			//find farthest D along normal
-			const TVector<T, 3> AToC = InParticles.X(C->Vertex) - InParticles.X(A->Vertex);
-			const TVector<T, d> Normal = TVector<T, 3>::CrossProduct(AToB, AToC);
+			Result = HashCombine(Result, SurfaceParticles.GetTypeHash());
 
-			T MaxPosDistance = Epsilon;
-			T MaxNegDistance = Epsilon;
-			FHalfEdge* PosD = nullptr;
-			FHalfEdge* NegD = nullptr;
-			for (FHalfEdge* V = DummyHalfEdge.Next; V; V = V->Next)
+			for(const TPlane<T, d>& Plane : Planes)
 			{
-				T Dot = TVector<T, 3>::DotProduct(InParticles.X(V->Vertex) - InParticles.X(A->Vertex), Normal);
-				if (Dot > MaxPosDistance)
-				{
-					MaxPosDistance = Dot;
-					PosD = V;
-				}
-				if (-Dot > MaxNegDistance)
-				{
-					MaxNegDistance = -Dot;
-					NegD = V;
-				}
+				Result = HashCombine(Result, Plane.GetTypeHash());
 			}
 
-			if (MaxNegDistance == Epsilon && MaxPosDistance == Epsilon)
-			{
-				return nullptr; //plane
-			}
-
-			const bool bPositive = MaxNegDistance < MaxPosDistance;
-			FHalfEdge* D = bPositive ? PosD : NegD;
-
-			//remove D from conflict list
-			D->Prev->Next = D->Next;
-			if (D->Next)
-			{
-				D->Next->Prev = D->Prev;
-			}
-
-			//we must now create the 3 faces. Face must be oriented CCW around normal and positive normal should face out
-			//Note we are now using A,B,C,D as edges. We can only use one edge per face so once they're used we'll need new ones
-			FHalfEdge* Edges[4] = {A, B, C, D};
-
-			//The base is a plane with Edges[0], Edges[1], Edges[2]. The order depends on which side D is on
-			if (bPositive)
-			{
-				//D is on the positive side of Edges[0], Edges[1], Edges[2] so we must reorder it
-				FHalfEdge* Tmp = Edges[0];
-				Edges[0] = Edges[1];
-				Edges[1] = Tmp;
-			}
-
-			FConvexFace* Faces[4];
-			Faces[0] = CreateFace(InParticles, Edges[0], Edges[1], Edges[2]); //base
-			Faces[1] = CreateFace(InParticles, new FHalfEdge(Edges[1]->Vertex), new FHalfEdge(Edges[0]->Vertex), Edges[3]);
-			Faces[2] = CreateFace(InParticles, new FHalfEdge(Edges[0]->Vertex), new FHalfEdge(Edges[2]->Vertex), new FHalfEdge(Edges[3]->Vertex));
-			Faces[3] = CreateFace(InParticles, new FHalfEdge(Edges[2]->Vertex), new FHalfEdge(Edges[1]->Vertex), new FHalfEdge(Edges[3]->Vertex));
-
-			auto MakeTwins = [](FHalfEdge* E1, FHalfEdge* E2) {
-				E1->Twin = E2;
-				E2->Twin = E1;
-			};
-			//mark twins so half edge can cross faces
-			MakeTwins(Edges[0], Faces[1]->FirstEdge); //0-1 1-0
-			MakeTwins(Edges[1], Faces[3]->FirstEdge); //1-2 2-1
-			MakeTwins(Edges[2], Faces[2]->FirstEdge); //2-0 0-2
-			MakeTwins(Faces[1]->FirstEdge->Next, Faces[2]->FirstEdge->Prev); //0-3 3-0
-			MakeTwins(Faces[1]->FirstEdge->Prev, Faces[3]->FirstEdge->Next); //3-1 1-3
-			MakeTwins(Faces[2]->FirstEdge->Next, Faces[3]->FirstEdge->Prev); //2-3 3-2
-
-			Faces[0]->Prev = nullptr;
-			for (int i = 1; i < 4; ++i)
-			{
-				Faces[i - 1]->Next = Faces[i];
-				Faces[i]->Prev = Faces[i - 1];
-			}
-			Faces[3]->Next = nullptr;
-
-			//split up the conflict list
-			StealConflictList(InParticles, DummyHalfEdge.Next, Faces, 4);
-			return Faces[0];
-		}
-
-		static FHalfEdge* FindConflictVertex(const TParticles<T, 3>& InParticles, FConvexFace* FaceList)
-		{
-			for (FConvexFace* CurFace = FaceList; CurFace; CurFace = CurFace->Next)
-			{
-				T MaxD = TNumericLimits<T>::Lowest();
-				FHalfEdge* MaxV = nullptr;
-				for (FHalfEdge* CurFaceVertex = CurFace->ConflictList; CurFaceVertex; CurFaceVertex = CurFaceVertex->Next)
-				{
-					//is it faster to cache this from stealing stage?
-					T Dist = TVector<T, 3>::DotProduct(InParticles.X(CurFaceVertex->Vertex), CurFace->Plane.Normal());
-					if (Dist > MaxD)
-					{
-						MaxD = Dist;
-						MaxV = CurFaceVertex;
-					}
-				}
-				check(CurFace->ConflictList == nullptr || MaxV);
-				if (MaxV)
-				{
-					if (MaxV->Prev)
-					{
-						MaxV->Prev->Next = MaxV->Next;
-					}
-					if (MaxV->Next)
-					{
-						MaxV->Next->Prev = MaxV->Prev;
-					}
-					if (MaxV == CurFace->ConflictList)
-					{
-						CurFace->ConflictList = MaxV->Next;
-					}
-					MaxV->Face = CurFace;
-					return MaxV;
-				}
-			}
-
-			return nullptr;
-		}
-
-		static void BuildHorizon(const TParticles<T, 3>& InParticles, FHalfEdge* ConflictV, TArray<FHalfEdge*>& HorizonEdges, TArray<FConvexFace*>& FacesToDelete)
-		{
-			//We must flood fill from the initial face and mark edges of faces the conflict vertex cannot see
-			//In order to return a CCW ordering we must traverse each face in CCW order from the edge we crossed over
-			//This should already be the ordering in the half edge
-			const T Epsilon = 1e-1;
-			const TVector<T, 3> V = InParticles.X(ConflictV->Vertex);
-			TSet<FConvexFace*> Processed;
-			TArray<FHalfEdge*> Queue;
-			check(ConflictV->Face);
-			Queue.Add(ConflictV->Face->FirstEdge->Prev); //stack pops so reverse order
-			Queue.Add(ConflictV->Face->FirstEdge->Next);
-			Queue.Add(ConflictV->Face->FirstEdge);
-			FacesToDelete.Add(ConflictV->Face);
-			while (Queue.Num())
-			{
-				FHalfEdge* Edge = Queue.Pop(/*bAllowShrinking=*/false);
-				Processed.Add(Edge->Face);
-				FHalfEdge* Twin = Edge->Twin;
-				FConvexFace* NextFace = Twin->Face;
-				if (Processed.Contains(NextFace))
-				{
-					continue;
-				}
-				const T Distance = NextFace->Plane.SignedDistance(V);
-				if (Distance > Epsilon)
-				{
-					Queue.Add(Twin->Prev); //stack pops so reverse order
-					Queue.Add(Twin->Next);
-					FacesToDelete.Add(NextFace);
-				}
-				else
-				{
-					HorizonEdges.Add(Edge);
-				}
-			}
+			return Result;
 		}
 
-		static void BuildFaces(const TParticles<T, 3>& InParticles, const FHalfEdge* ConflictV, const TArray<FHalfEdge*>& HorizonEdges, const TArray<FConvexFace*> OldFaces, TArray<FConvexFace*>& NewFaces)
+		FORCEINLINE void SerializeImp(FArchive& Ar)
 		{
-			//The HorizonEdges are in CCW order. We must make new faces and edges to join from ConflictV to these edges
-			check(HorizonEdges.Num() >= 3);
-			NewFaces.Reserve(HorizonEdges.Num());
-			FHalfEdge* PrevEdge = nullptr;
-			for (int32 HorizonIdx = 0; HorizonIdx < HorizonEdges.Num(); ++HorizonIdx)
-			{
-				FHalfEdge* OriginalEdge = HorizonEdges[HorizonIdx];
-				FHalfEdge* NewHorizonEdge = new FHalfEdge(OriginalEdge->Vertex);
-				NewHorizonEdge->Twin = OriginalEdge->Twin; //swap edges
-				NewHorizonEdge->Twin->Twin = NewHorizonEdge;
-				FHalfEdge* HorizonNext = new FHalfEdge(OriginalEdge->Next->Vertex);
-				check(HorizonNext->Vertex == HorizonEdges[(HorizonIdx + 1) % HorizonEdges.Num()]->Vertex); //should be ordered properly
-				FHalfEdge* V = new FHalfEdge(ConflictV->Vertex);
-				V->Twin = PrevEdge;
-				if (PrevEdge)
-				{
-					PrevEdge->Twin = V;
-				}
-				PrevEdge = HorizonNext;
-
-				//link new faces together
-				FConvexFace* NewFace = CreateFace(InParticles, NewHorizonEdge, HorizonNext, V);
-				if (NewFaces.Num() > 0)
-				{
-					NewFace->Prev = NewFaces[NewFaces.Num() - 1];
-					NewFaces[NewFaces.Num() - 1]->Next = NewFace;
-				}
-				else
-				{
-					NewFace->Prev = nullptr;
-				}
-				NewFaces.Add(NewFace);
-			}
-
-			check(PrevEdge);
-			NewFaces[0]->FirstEdge->Prev->Twin = PrevEdge;
-			PrevEdge->Twin = NewFaces[0]->FirstEdge->Prev;
-			NewFaces[NewFaces.Num() - 1]->Next = nullptr;
-
-			//redistribute conflict list
-			for (FConvexFace* OldFace : OldFaces)
-			{
-				StealConflictList(InParticles, OldFace->ConflictList, &NewFaces[0], NewFaces.Num());
-			}
-
-			//insert all new faces after conflict vertex face
-			FConvexFace* OldFace = ConflictV->Face;
-			FConvexFace* StartFace = NewFaces[0];
-			FConvexFace* EndFace = NewFaces[NewFaces.Num() - 1];
-			if (OldFace->Next)
-			{
-				OldFace->Next->Prev = EndFace;
-			}
-			EndFace->Next = OldFace->Next;
-			OldFace->Next = StartFace;
-			StartFace->Prev = OldFace;
+			TImplicitObject<T, 3>::SerializeImp(Ar);
+			Ar << Planes << SurfaceParticles << LocalBoundingBox;
 		}
 
-		static void AddVertex(const TParticles<T, 3>& InParticles, FHalfEdge* ConflictV)
+		virtual void Serialize(FChaosArchive& Ar) override
 		{
-			TArray<FHalfEdge*> HorizonEdges;
-			TArray<FConvexFace*> FacesToDelete;
-			BuildHorizon(InParticles, ConflictV, HorizonEdges, FacesToDelete);
-
-			TArray<FConvexFace*> NewFaces;
-			BuildFaces(InParticles, ConflictV, HorizonEdges, FacesToDelete, NewFaces);
-
-			for (FConvexFace* Face : FacesToDelete)
-			{
-				FHalfEdge* Edge = Face->FirstEdge;
-				do
-				{
-					FHalfEdge* Next = Edge->Next;
-					delete Edge;
-					Edge = Next;
-				} while (Edge != Face->FirstEdge);
-				if (Face->Prev)
-				{
-					Face->Prev->Next = Face->Next;
-				}
-				if (Face->Next)
-				{
-					Face->Next->Prev = Face->Prev;
-				}
-				delete Face;
-			}
-
-			//todo(ocohen): need to explicitly test for merge failures. Coplaner, nonconvex, etc...
-			//getting this in as is for now to unblock other systems
+			SerializeImp(Ar);
 		}
 
-	public:
-		static void BuildConvexHull(const TParticles<T, 3>& InParticles, TArray<TVector<int32, 3>>& OutIndices)
+		virtual void Serialize(FArchive& Ar) override
 		{
-			OutIndices.SetNum(0, false);
-			FConvexFace* Faces = BuildInitialHull(InParticles);
-			if (Faces == nullptr)
-			{
-				return;
-			}
-			FConvexFace DummyFace(Faces->Plane);
-			DummyFace.Prev = nullptr;
-			DummyFace.Next = Faces;
-			Faces->Prev = &DummyFace;
-
-			FHalfEdge* ConflictV = FindConflictVertex(InParticles, DummyFace.Next);
-			while (ConflictV)
-			{
-				AddVertex(InParticles, ConflictV);
-				ConflictV = FindConflictVertex(InParticles, DummyFace.Next);
-			}
-
-			FConvexFace* Cur = DummyFace.Next;
-			while (Cur)
-			{
-				//todo(ocohen): this assumes faces are triangles, not true once face merging is added
-				OutIndices.Add(TVector<int32, 3>(Cur->FirstEdge->Vertex, Cur->FirstEdge->Next->Vertex, Cur->FirstEdge->Next->Next->Vertex));
-				delete Cur->FirstEdge->Next;
-				delete Cur->FirstEdge->Prev;
-				delete Cur->FirstEdge;
-				FConvexFace* Next = Cur->Next;
-				delete Cur;
-				Cur = Next;
-			}
+			SerializeImp(Ar);
 		}
 
-		static TTriangleMesh<T> BuildConvexHullTriMesh(const TParticles<T, 3>& InParticles)
-		{
-			TArray<TVector<int32, 3>> Indices;
-			BuildConvexHull(InParticles, Indices);
-			return TTriangleMesh<T>(MoveTemp(Indices));
-		}
+	private:
+		TArray<TPlane<T, d>> Planes;
+		TParticles<T, d> SurfaceParticles;	//copy of the vertices that are just on the convex hull boundary
+		TBox<T, d> LocalBoundingBox;
 	};
 }
diff --git a/Engine/Source/Runtime/Experimental/Chaos/Public/Chaos/Cylinder.h b/Engine/Source/Runtime/Experimental/Chaos/Public/Chaos/Cylinder.h
--- a/Engine/Source/Runtime/Experimental/Chaos/Public/Chaos/Cylinder.h
+++ b/Engine/Source/Runtime/Experimental/Chaos/Public/Chaos/Cylinder.h
@@ -241,6 +241,14 @@ namespace Chaos
 			return MakePair(TVector<T, 3>(0), false);
 		}
 
+		virtual uint32 GetTypeHash() const override
+		{
+			const uint32 PlaneHashes = HashCombine(MPlane1.GetTypeHash(), MPlane2.GetTypeHash());
+			const uint32 PropertyHash = HashCombine(::GetTypeHash(MHeight), ::GetTypeHash(MRadius));
+
+			return HashCombine(PlaneHashes, PropertyHash);
+		}
+
 	private:
 		TPlane<T, 3> MPlane1, MPlane2;
 		T MHeight, MRadius;
diff --git a/Engine/Source/Runtime/Experimental/Chaos/Public/Chaos/DynamicParticles.h b/Engine/Source/Runtime/Experimental/Chaos/Public/Chaos/DynamicParticles.h
--- a/Engine/Source/Runtime/Experimental/Chaos/Public/Chaos/DynamicParticles.h
+++ b/Engine/Source/Runtime/Experimental/Chaos/Public/Chaos/DynamicParticles.h
@@ -27,7 +27,6 @@ class TDynamicParticles : public TParticles<T, d>
 		TArrayCollection::AddArray(&MM);
 		TArrayCollection::AddArray(&MInvM);
 	}
-	~TDynamicParticles() {}
 
 	const TVector<T, d>& V(const int32 Index) const { return MV[Index]; }
 	TVector<T, d>& V(const int32 Index) { return MV[Index]; }
diff --git a/Engine/Source/Runtime/Experimental/Chaos/Public/Chaos/Framework/PhysicsProxy.h b/Engine/Source/Runtime/Experimental/Chaos/Public/Chaos/Framework/PhysicsProxy.h
new file mode 100644
--- /dev/null
+++ b/Engine/Source/Runtime/Experimental/Chaos/Public/Chaos/Framework/PhysicsProxy.h
@@ -0,0 +1,211 @@
+// Copyright 1998-2019 Epic Games, Inc. All Rights Reserved.
+
+#pragma once
+
+#if INCLUDE_CHAOS
+
+#include "Chaos/PBDRigidParticles.h"
+#include "Chaos/PBDCollisionConstraint.h"
+#include "Chaos/ArrayCollection.h"
+#include "Chaos/ArrayCollectionArray.h"
+#include "Chaos/Declares.h"
+#include "UObject/GCObject.h"
+
+struct FKinematicProxy;
+class FFieldSystemCommand;
+struct FBodyInstance;
+
+enum class EPhysicsProxyType
+{
+	NoneType = 0,
+	StaticMeshType = 1,
+	GeometryCollectionType = 2,
+	FieldType = 3,
+	SkeletalMeshType = 4,
+	FloatGeometryParticleType = 5
+};
+
+class IPhysicsProxyBase
+{
+public:
+	IPhysicsProxyBase()
+		: Solver(nullptr)
+	{}
+
+	virtual UObject* GetOwner() const = 0;
+
+	Chaos::FPhysicsSolver* GetSolver() const { return Solver; }
+
+	//Should this be in the public API? probably not
+	void SetSolver(Chaos::FPhysicsSolver* InSolver) { Solver = InSolver; }
+
+
+protected:
+	// Intended to be non-virtual and protected. This ensures that derived classes can successfully call this destructor
+	// but no one can delete using a IPhysicsProxyBase*
+	~IPhysicsProxyBase() = default;
+	
+	/** The solver that owns the solver object */
+	Chaos::FPhysicsSolver* Solver;
+};
+
+struct PhysicsProxyWrapper
+{
+	IPhysicsProxyBase* PhysicsProxy;
+	EPhysicsProxyType Type;
+};
+
+/**
+ * Base object interface for solver objects. Defines the expected API for objects
+ * uses CRTP for static dispatch, entire API considered "pure-virtual" and must be* defined.
+ * Forgetting to implement any of the interface functions will give errors regarding
+ * recursion on all control paths for TPhysicsProxy<T> where T will be the type
+ * that has not correctly implemented the API.
+ *
+ * PersistentTask uses IPhysicsProxyBase, so when implementing a new specialized type
+ * it is necessary to include its header file in PersistentTask.cpp allowing the linker
+ * to properly resolve the new type. 
+ *
+ * May not be necessary overall once the engine has solidified - we can just use the
+ * final concrete objects but this gives us almost the same flexibility as the old
+ * callbacks while solving most of the drawbacks (virtual dispatch, cross-object interaction)
+ *
+ * #BG TODO - rename the callbacks functions, document for the base solver object
+ */
+template<class Concrete, class ConcreteData>
+class TPhysicsProxy : public IPhysicsProxyBase
+{
+
+public:
+	using FParticleType = typename Concrete;
+	using FParticleData = typename ConcreteData;
+
+	using FParticlesType = Chaos::TPBDRigidParticles<float, 3>;
+	using FCollisionConstraintsType = Chaos::TPBDCollisionConstraint<float, 3>;
+	using FIntArray = Chaos::TArrayCollectionArray<int32>;
+
+	TPhysicsProxy()
+		: IPhysicsProxyBase()
+		, Owner(nullptr)
+	{
+	}
+
+	explicit TPhysicsProxy(UObject* InOwner)
+		: IPhysicsProxyBase()
+		, Owner(InOwner)
+	{
+	}
+
+	/** Virtual destructor for derived objects, ideally no other virtuals should exist in this chain */
+	virtual ~TPhysicsProxy() {}
+
+	/**
+	 * The following functions are to be implemented by all solver objects as we're using CRTP / F-Bound to
+	 * statically dispatch the calls. Any common functions should be added here and to the derived solver objects
+	 */
+
+	// Previously callback related functions, all called in the context of the physics thread if enabled.
+	bool IsSimulating() const { return static_cast<const Concrete*>(this)->IsSimulating(); }
+	void UpdateKinematicBodiesCallback(const FParticlesType& InParticles, const float InDt, const float InTime, FKinematicProxy& InKinematicProxy) { static_cast<Concrete*>(this)->UpdateKinematicBodiesCallback(InParticles, InDt, InTime, InKinematicProxy); }
+	void StartFrameCallback(const float InDt, const float InTime) { static_cast<Concrete*>(this)->StartFrameCallback(InDt, InTime); }
+	void EndFrameCallback(const float InDt) { static_cast<Concrete*>(this)->EndFrameCallback(InDt); }
+	void CreateRigidBodyCallback(FParticlesType& InOutParticles) { static_cast<Concrete*>(this)->CreateRigidBodyCallback(InOutParticles); }
+	void ParameterUpdateCallback(FParticlesType& InParticles, const float InTime) { static_cast<Concrete*>(this)->ParameterUpdateCallback(InParticles, InTime); }
+	void DisableCollisionsCallback(TSet<TTuple<int32, int32>>& InPairs) { static_cast<Concrete*>(this)->DisableCollisionsCallback(InPairs); }
+	void AddForceCallback(FParticlesType& InParticles, const float InDt, const int32 InIndex) { static_cast<Concrete*>(this)->AddForceCallback(InParticles, InDt, InIndex); }
+	void FieldForcesUpdateCallback(Chaos::FPhysicsSolver* InSolver, FParticlesType& Particles, Chaos::TArrayCollectionArray<FVector> & Force, Chaos::TArrayCollectionArray<FVector> & Torque, const float Time) { static_cast<Concrete*>(this)->FieldForcesUpdateCallback(InSolver, Particles, Force, Torque, Time); }
+
+	/** The Particle Binding creates a connection between the particles in the simulation and the solver objects dataset. */
+	void BindParticleCallbackMapping(Chaos::TArrayCollectionArray<PhysicsProxyWrapper> & PhysicsProxyReverseMap, Chaos::TArrayCollectionArray<int32> & ParticleIDReverseMap) {static_cast<Concrete*>(this)->BindParticleCallbackMapping(PhysicsProxyReverseMap, ParticleIDReverseMap);}
+
+	/** Called to buffer a command to be processed at the next available safe opportunity */
+	void BufferCommand(Chaos::FPhysicsSolver* InSolver, const FFieldSystemCommand& InCommand) { static_cast<Concrete*>(this)->BufferCommand(InSolver, InCommand); }
+
+	/** Returns the concrete type of the derived class*/
+	EPhysicsProxyType ConcreteType() { return static_cast<Concrete*>(this)->ConcreteType(); }
+
+	/**
+	 * CONTEXT: GAMETHREAD
+	* Returns a new unmanaged allocation of the data saved on the handle, otherwise nullptr
+	*/
+	Chaos::FParticleData* NewData() { return static_cast<Concrete*>(this)->NewData(); }
+
+	/**
+	* CONTEXT: GAMETHREAD -> to -> PHYSICSTHREAD
+	* Called on the game thread when the solver is about to advance forward. This
+	* callback should Enqueue commands on the PhysicsThread to update the state of
+	* the solver
+	*/
+	void PushToPhysicsState(const Chaos::FParticleData* InData) { static_cast<Concrete*>(this)->PushToPhysicsState(InData); }
+
+	/**
+	 * CONTEXT: PHYSICSTHREAD
+	 * Called per-tick after the simulation has completed. The proxy should cache the results of their
+	 * simulation into the local buffer. 
+	 */
+	void BufferPhysicsResults() { static_cast<Concrete*>(this)->BufferPhysicsResults(); }
+
+	/**
+	 * CONTEXT: PHYSICSTHREAD (Write Locked)
+	 * Called by the physics thread to signal that it is safe to perform any double-buffer flips here.
+	 * The physics thread has pre-locked an RW lock for this operation so the game thread won't be reading
+	 * the data
+	 */
+	void FlipBuffer() { static_cast<Concrete*>(this)->FlipBuffer(); }
+
+	/**
+	 * CONTEXT: GAMETHREAD (Read Locked)
+	 * Perform a similar operation to Sync, but take the data from a gamethread-safe buffer. This will be called
+	 * from the game thread when it cannot sync to the physics thread. The simulation is very likely to be running
+	 * when this happens so never read any physics thread data here!
+	 *
+	 * Note: A read lock will have been acquired for this - so the physics thread won't force a buffer flip while this
+	 * sync is ongoing
+	 */
+	void PullFromPhysicsState() { static_cast<Concrete*>(this)->PullFromPhysicsState(); }
+
+	/**
+	 * CONTEXT: GAMETHREAD
+	 * Called during the gamethread sync after the proxy has been removed from its solver
+	 * intended for final handoff of any data the proxy has that the gamethread may
+	 * be interested in
+	*/
+	void SyncBeforeDestroy() { static_cast<Concrete*>(this)->SyncBeforeDestroy(); }
+
+	/**
+	 * CONTEXT: PHYSICSTHREAD
+	 * Called on the physics thread when the engine is shutting down the proxy and we need to remove it from
+	 * any active simulations. Proxies are expected to entirely clean up their simulation
+	 * state within this method. This is run in the task command step by the scene
+	 * so the simulation will currently be idle
+	 */
+	void OnRemoveFromScene() { static_cast<Concrete*>(this)->OnRemoveFromScene(); }
+
+
+
+	/**
+	 * End PhysicsProxy Common interface
+	 */
+
+	/** Get/Set the solver this object belongs to */
+	void SetSolver(Chaos::FPhysicsSolver* InSolver) { Solver = InSolver; }
+	Chaos::FPhysicsSolver* GetSolver() const { return Solver; }
+
+	/** Gets the owning external object for this solver object, never used internally */
+	virtual UObject* GetOwner() const override { return Owner; }
+
+	void* GetUserData() const { return nullptr; }
+
+	Chaos::TRigidTransform<float, 3> GetTransform() const { return Chaos::TRigidTransform<float, 3>(); }
+
+
+private:
+
+	/** 
+	 * The owner for this solver object, essentially user-data managed by the caller 
+	 * @see GetOwner
+	 */
+	UObject* Owner;
+};
+
+#endif
diff --git a/Engine/Source/Runtime/Experimental/Chaos/Public/Chaos/GJK.h b/Engine/Source/Runtime/Experimental/Chaos/Public/Chaos/GJK.h
--- a/Engine/Source/Runtime/Experimental/Chaos/Public/Chaos/GJK.h
+++ b/Engine/Source/Runtime/Experimental/Chaos/Public/Chaos/GJK.h
@@ -418,11 +418,10 @@ namespace Chaos
 	/** Sweeps one geometry against the other
 	 @A The first geometry
 	 @B The second geometry
-	 @BToATM The transform of B in A's local space
-	 @StartPoint The ray's starting point in A's local space
+	 @StartTM B's starting configuration in A's local space
 	 @RayDir The ray's direction (normalized)
 	 @RayLength The ray's length
-	 @OutTime The normalized time along the ray when the objects first overlap
+	 @OutTime The time along the ray when the objects first overlap
 	 @OutPosition The first point of impact (in A's local space) when the objects first overlap. Invalid if time of impact is 0
 	 @OutNormal The impact normal (in A's local space) when the objects first overlap. Invalid if time of impact is 0
 	 @Thickness The amount of geometry inflation (for example if the surface distance of two geometries with thickness 0 would be 2, a thickness of 0.5 would give a distance of 1.5)
@@ -430,12 +429,13 @@ namespace Chaos
 	 @return True if the geometries overlap during the sweep, False otherwise */
 
 	template <typename T, typename TGeometryA, typename TGeometryB>
-	bool GJKRaycast(const TGeometryA& A, const TGeometryB& B, const TRigidTransform<T, 3>& BToATM, const TVector<T,3>& StartPoint, const TVector<T,3>& RayDir, const T RayLength,
+	bool GJKRaycast(const TGeometryA& A, const TGeometryB& B, const TRigidTransform<T, 3>& StartTM, const TVector<T,3>& RayDir, const T RayLength,
 		T& OutTime, TVector<T,3>& OutPosition, TVector<T,3>& OutNormal, const T Thickness = 0, const TVector<T, 3>& InitialDir = TVector<T, 3>(1, 0, 0))
 	{
 		ensure(FMath::IsNearlyEqual(RayDir.SizeSquared(), 1, KINDA_SMALL_NUMBER));
 		ensure(RayLength > 0);
 		check(A.IsConvex() && B.IsConvex());
+		const TVector<T, 3> StartPoint = StartTM.GetLocation();
 
 		TVector<T, 3> Simplex[4] = { TVector<T,3>(0), TVector<T,3>(0), TVector<T,3>(0), TVector<T,3>(0) };
 		TVector<T, 3> As[4] = { TVector<T,3>(0), TVector<T,3>(0), TVector<T,3>(0), TVector<T,3>(0) };
@@ -444,13 +444,14 @@ namespace Chaos
 		T Barycentric[4];
 
 		FSimplex SimplexIDs;
-		const TRotation<T, 3> AToBRotation = BToATM.GetRotation().Inverse();
+		const TRotation<T, 3> BToARotation = StartTM.GetRotation();
+		const TRotation<T, 3> AToBRotation = BToARotation.Inverse();
 		TVector<T,3> SupportA = A.Support(InitialDir, Thickness);	//todo: use Thickness on quadratic geometry
 		As[0] = SupportA;
 
 		const TVector<T, 3> InitialDirInB = AToBRotation * (-InitialDir);
 		const TVector<T, 3> InitialSupportBLocal = B.Support(InitialDirInB, 0);
-		TVector<T, 3> SupportB = BToATM.TransformVectorNoScale(InitialSupportBLocal);
+		TVector<T, 3> SupportB = BToARotation * InitialSupportBLocal;
 		Bs[0] = SupportB;
 
 		T Lambda = 0;
@@ -473,7 +474,7 @@ namespace Chaos
 			SupportA = A.Support(V, Thickness);	//todo: add thickness to quadratic geometry to avoid quadratic vs quadratic when possible
 			const TVector<T, 3> VInB = AToBRotation * (-V);
 			const TVector<T, 3> SupportBLocal = B.Support(VInB, 0);
-			SupportB = BToATM.TransformVectorNoScale(SupportBLocal);
+			SupportB = BToARotation * SupportBLocal;
 			const TVector<T, 3> P = SupportA - SupportB;
 			const TVector<T, 3> W = X - P;
 			SimplexIDs[SimplexIDs.NumVerts] = SimplexIDs.NumVerts;	//is this needed?
@@ -527,7 +528,7 @@ namespace Chaos
 
 		} while (!bTerminate);
 
-		OutTime = Lambda / RayLength;
+		OutTime = Lambda;
 
 		if (Lambda > 0)
 		{
diff --git a/Engine/Source/Runtime/Experimental/Chaos/Public/Chaos/GeometryParticles.h b/Engine/Source/Runtime/Experimental/Chaos/Public/Chaos/GeometryParticles.h
--- a/Engine/Source/Runtime/Experimental/Chaos/Public/Chaos/GeometryParticles.h
+++ b/Engine/Source/Runtime/Experimental/Chaos/Public/Chaos/GeometryParticles.h
@@ -5,56 +5,139 @@
 #include "Chaos/ImplicitObject.h"
 #include "Chaos/Particles.h"
 #include "Chaos/Rotation.h"
+#include "Chaos/ParticleHandleFwd.h"
+#include "Chaos/CollisionFilterData.h"
+
+#ifndef CHAOS_DETERMINISTIC
+#define CHAOS_DETERMINISTIC 1
+#endif
 
 namespace Chaos
 {
+	/** Data that is associated with geometry. If a union is used an entry is created per internal geometry */
 	template <typename T, int d>
-	class TGeometryParticleHandle;
+	class TPerShapeData
+	{
+	public:
+#if INCLUDE_CHAOS
+		FCollisionFilterData QueryData;
+		FCollisionFilterData SimData;
+#endif
+		void* UserData;
+		const TImplicitObject<T, d>* Geometry;
+		
+		TPerShapeData()
+			: UserData(nullptr)
+		{
+		}
+	};
 
-	template<class T, int d>
-	class TGeometryParticles : public TParticles<T, d>
+	template <typename T, int d>
+	using TShapesArray = TArray<TUniquePtr<TPerShapeData<T, d>>, TInlineAllocator<1>>;
+
+	template <typename T, int d>
+	void UpdateShapesArrayFromGeometry(TShapesArray<T, d>& ShapesArray, const TImplicitObject<T, d>* Geometry);
+
+	extern template void CHAOS_API UpdateShapesArrayFromGeometry(TShapesArray<float, 3>& ShapesArray, const TImplicitObject<float, 3>* Geometry);
+
+	#if CHAOS_DETERMINISTIC
+	using FParticleID = int32;	//Used to break ties when determinism is needed. Should not be used for anything else
+#endif
+	//Used for down casting when iterating over multiple SOAs.
+	enum class EParticleType : uint8
+	{
+		Static,
+		Kinematic,
+		Dynamic,
+		Clustered,	//only applicable on physics thread side
+		StaticMesh,
+		SkeletalMesh,
+		GeometryCollection
+	};
+
+	enum class EGeometryParticlesSimType
+	{
+		RigidBodySim,
+		Other
+	};
+	
+	template<class T, int d, EGeometryParticlesSimType SimType>
+	class TGeometryParticlesImp : public TParticles<T, d>
 	{
 	public:
 		using TArrayCollection::Size;
 
-		TGeometryParticles()
+		TGeometryParticlesImp()
 		    : TParticles<T, d>()
 		{
+			MParticleType = EParticleType::Static;
 			TArrayCollection::AddArray(&MR);
 			TArrayCollection::AddArray(&MGeometry);
+			TArrayCollection::AddArray(&MSharedGeometry);
 			TArrayCollection::AddArray(&MDynamicGeometry);
-			TArrayCollection::AddArray(&MGeometryParticleHandle);
-		}
-		TGeometryParticles(const TGeometryParticles<T, d>& Other) = delete;
-		TGeometryParticles(TGeometryParticles<T, d>&& Other)
-		    : TParticles<T, d>(MoveTemp(Other)), MR(MoveTemp(Other.MR)), MGeometry(MoveTemp(Other.MGeometry)), MDynamicGeometry(MoveTemp(Other.MDynamicGeometry))
-		{
-			TArrayCollection::AddArray(&MR);
-			TArrayCollection::AddArray(&MGeometry);
-			TArrayCollection::AddArray(&MDynamicGeometry);
-			TArrayCollection::AddArray(&MGeometryParticleHandle);
+			TArrayCollection::AddArray(&MParticleIDs);
+			TArrayCollection::AddArray(&MShapesArray);
+
+			if (IsRigidBodySim())
+			{
+				TArrayCollection::AddArray(&MGeometryParticleHandle);
+				TArrayCollection::AddArray(&MGeometryParticle);
+			}
+
 		}
-		TGeometryParticles(TParticles<T, d>&& Other)
+		TGeometryParticlesImp(const TGeometryParticlesImp<T, d, SimType>& Other) = delete;
+		TGeometryParticlesImp(TGeometryParticlesImp<T, d, SimType>&& Other)
 		    : TParticles<T, d>(MoveTemp(Other))
+			, MR(MoveTemp(Other.MR))
+			, MGeometry(MoveTemp(Other.MGeometry))
+			, MSharedGeometry(MoveTemp(Other.MSharedGeometry))
+			, MDynamicGeometry(MoveTemp(Other.MDynamicGeometry))
+			, MGeometryParticleHandle(MoveTemp(Other.MGeometryParticleHandle))
+			, MGeometryParticle(MoveTemp(Other.MGeometryParticle))
+			, MShapesArray(MoveTemp(Other.MShapesArray))
+#if CHAOS_DETERMINISTIC
+			, MParticleIDs(MoveTemp(Other.MParticleIDs))
+#endif
 		{
+			MParticleType = EParticleType::Static;
 			TArrayCollection::AddArray(&MR);
 			TArrayCollection::AddArray(&MGeometry);
+			TArrayCollection::AddArray(&MSharedGeometry);
 			TArrayCollection::AddArray(&MDynamicGeometry);
-			TArrayCollection::AddArray(&MGeometryParticleHandle);
+			TArrayCollection::AddArray(&MShapesArray);
+#if CHAOS_DETERMINISTIC
+			TArrayCollection::AddArray(&MParticleIDs);
+#endif
+
+			if (IsRigidBodySim())
+			{
+				TArrayCollection::AddArray(&MGeometryParticleHandle);
+				TArrayCollection::AddArray(&MGeometryParticle);
+			}
 		}
-			
-		/**
-		 * Constructor that replaces the positions array with a view of @param Points.
-		 */
-		TGeometryParticles(TArray<TVector<T, d>>& Points)
-		    : TParticles<T, d>(Points)
+
+		static constexpr bool IsRigidBodySim() { return SimType == EGeometryParticlesSimType::RigidBodySim; }
+
+		TGeometryParticlesImp(TParticles<T, d>&& Other)
+		    : TParticles<T, d>(MoveTemp(Other))
 		{
+			MParticleType = EParticleType::Static;
 			TArrayCollection::AddArray(&MR);
 			TArrayCollection::AddArray(&MGeometry);
+			TArrayCollection::AddArray(&MSharedGeometry);
 			TArrayCollection::AddArray(&MDynamicGeometry);
-			TArrayCollection::AddArray(&MGeometryParticleHandle);
+			TArrayCollection::AddArray(&MShapesArray);
+#if CHAOS_DETERMINISTIC
+			TArrayCollection::AddArray(&MParticleIDs);
+#endif
+
+			if (IsRigidBodySim())
+			{
+				TArrayCollection::AddArray(&MGeometryParticleHandle);
+				TArrayCollection::AddArray(&MGeometryParticle);
+			}
 		}
-		
+
 		const TRotation<T, d>& R(const int32 Index) const { return MR[Index]; }
 		TRotation<T, d>& R(const int32 Index) { return MR[Index]; }
 
@@ -62,16 +145,37 @@ namespace Chaos
 
 		const TUniquePtr<TImplicitObject<T, d>>& DynamicGeometry(const int32 Index) const { return MDynamicGeometry[Index]; }
 
-		void SetDynamicGeometry(const int32 Index, TUniquePtr<TImplicitObject<T, d>>&& Unique)
+		const TSharedPtr<TImplicitObject<T, d>, ESPMode::ThreadSafe>& SharedGeometry(const int32 Index) const { return MSharedGeometry[Index]; }
+
+		const TShapesArray<T, d>& ShapesArray(const int32 Index) const { return MShapesArray[Index]; }
+
+#if CHAOS_DETERMINISTIC
+		FParticleID ParticleID(const int32 Idx) const { return MParticleIDs[Idx]; }
+		FParticleID& ParticleID(const int32 Idx) { return MParticleIDs[Idx]; }
+#endif
+
+		void SetDynamicGeometry(const int32 Index, TUniquePtr<TImplicitObject<T, d>>&& InUnique)
 		{
-			MDynamicGeometry[Index] = MoveTemp(Unique);
-			MGeometry[Index] = MakeSerializable(DynamicGeometry(Index));
+			check(!SharedGeometry(Index));	// If shared geometry exists we should not be setting dynamic geometry on top
+			MGeometry[Index] = MakeSerializable(InUnique);
+			MDynamicGeometry[Index] = MoveTemp(InUnique);
+			UpdateShapesArray(Index);
 		}
 
-		void SetGeometry(const int32 Index, TSerializablePtr<TImplicitObject<T, d>> Serializable)
+		void SetSharedGeometry(const int32 Index, TSharedPtr<TImplicitObject<T, d>, ESPMode::ThreadSafe> InShared)
+		{
+			check(!DynamicGeometry(Index));	// If dynamic geometry exists we should not be setting shared geometry on top
+			MGeometry[Index] = MakeSerializable(InShared);
+			MSharedGeometry[Index] = InShared;
+			UpdateShapesArray(Index);
+		}
+		
+		void SetGeometry(const int32 Index, TSerializablePtr<TImplicitObject<T, d>> InGeometry)
 		{
-			check(!DynamicGeometry(Index));	//If dynamic geometry exists we should not be setting other geometry on top
-			MGeometry[Index] = Serializable;
+			check(!DynamicGeometry(Index));
+			check(!SharedGeometry(Index));
+			MGeometry[Index] = InGeometry;
+			UpdateShapesArray(Index);
 		}
 
 		const TArray<TSerializablePtr<TImplicitObject<T, d>>>& GetAllGeometry() const { return MGeometry; }
@@ -81,6 +185,9 @@ namespace Chaos
 
 		THandleType*& Handle(int32 Index){ return MGeometryParticleHandle[Index]; }
 
+		TGeometryParticle<T, d>* GTGeometryParticle(const int32 Index) const { return MGeometryParticle[Index]; }
+		TGeometryParticle<T, d>*& GTGeometryParticle(const int32 Index) { return MGeometryParticle[Index]; }
+
 		FString ToString(int32 index) const
 		{
 			FString BaseString = TParticles<T, d>::ToString(index);
@@ -94,17 +201,46 @@ namespace Chaos
 			//todo: serialize handles
 		}
 
+		EParticleType ParticleType() const { return MParticleType; }
+
+	protected:
+		EParticleType MParticleType;
+
 	private:
 		TArrayCollectionArray<TRotation<T, d>> MR;
+		// MGeometry contains raw ptrs to every entry in both MSharedGeometry and MDynamicGeometry.
+		// It may also contain raw ptrs to geometry which is managed outside of Chaos.
 		TArrayCollectionArray<TSerializablePtr<TImplicitObject<T, d>>> MGeometry;
+		// MSharedGeometry entries are owned by the solver, shared between *representations* of a particle.
+		// This is NOT for sharing geometry resources between particle's A and B, this is for sharing the
+		// geometry between particle A's various representations.
+		TArrayCollectionArray<TSharedPtr<TImplicitObject<T, d>, ESPMode::ThreadSafe>> MSharedGeometry;
+		// MDynamicGeometry entries are used for geo which is by the evolution. It is not set from the game side.
 		TArrayCollectionArray<TUniquePtr<TImplicitObject<T, d>>> MDynamicGeometry;
 		TArrayCollectionArray<TGeometryParticleHandle<T, d>*> MGeometryParticleHandle;
+		TArrayCollectionArray<TGeometryParticle<T, d>*> MGeometryParticle;
+		TArrayCollectionArray<TShapesArray<T,d>> MShapesArray;
+
+		void UpdateShapesArray(const int32 Index)
+		{
+			UpdateShapesArrayFromGeometry(MShapesArray[Index], MGeometry[Index].Get());
+		}
+
+#if CHAOS_DETERMINISTIC
+		TArrayCollectionArray<FParticleID> MParticleIDs;
+#endif
 	};
 
-	template <typename T, int d>
-	FChaosArchive& operator<<(FChaosArchive& Ar, TGeometryParticles<T, d>& Particles)
+	template <typename T, int d, EGeometryParticlesSimType SimType>
+	FChaosArchive& operator<<(FChaosArchive& Ar, TGeometryParticlesImp<T, d, SimType>& Particles)
 	{
 		Particles.Serialize(Ar);
 		return Ar;
 	}
+
+	template <typename T, int d>
+	using TGeometryParticles = TGeometryParticlesImp<T, d, EGeometryParticlesSimType::RigidBodySim>;
+
+	template <typename T, int d>
+	using TGeometryClothParticles = TGeometryParticlesImp<T, d, EGeometryParticlesSimType::Other>;
 }
diff --git a/Engine/Source/Runtime/Experimental/Chaos/Public/Chaos/GeometryQueries.h b/Engine/Source/Runtime/Experimental/Chaos/Public/Chaos/GeometryQueries.h
--- a/Engine/Source/Runtime/Experimental/Chaos/Public/Chaos/GeometryQueries.h
+++ b/Engine/Source/Runtime/Experimental/Chaos/Public/Chaos/GeometryQueries.h
@@ -41,8 +41,8 @@ namespace Chaos
 		if(BType == ImplicitObjectType::Sphere)
 		{
 			const TSphere<T, d>& BSphere = static_cast<const TSphere<T, d>&>(B);
-			const TVector<T, d> PtInA = BToATM.TransformPositionNoScale(BSphere.Center());
-			return A.Overlap(PtInA, Thickness + BSphere.Radius());
+			const TVector<T, d> PtInA = BToATM.TransformPositionNoScale(BSphere.GetCenter());
+			return A.Overlap(PtInA, Thickness + BSphere.GetRadius());
 		}
 		//todo: A is a sphere
 		else if (A.IsConvex())
@@ -82,7 +82,7 @@ namespace Chaos
 	}
 
 	template <typename T, int d>
-	bool SweepQuery(const TImplicitObject<T, d>& A, const TRigidTransform<T,d>& ATM, const TImplicitObject<T, d>& B, const TRigidTransform<T, d>& BTM, const TVector<T,d>& Dir, const T Length, T& OutTime, TVector<T,d>& OutPosition, TVector<T,d>& OutNormal, const T Thickness = 0)
+	bool SweepQuery(const TImplicitObject<T, d>& A, const TRigidTransform<T,d>& ATM, const TImplicitObject<T, d>& B, const TRigidTransform<T, d>& BTM, const TVector<T,d>& Dir, const T Length, T& OutTime, TVector<T,d>& OutPosition, TVector<T,d>& OutNormal, int32& OutFaceIndex, const T Thickness = 0)
 	{
 		const ImplicitObjectType AType = A.GetType(true);
 		const ImplicitObjectType BType = B.GetType(true);
@@ -93,17 +93,18 @@ namespace Chaos
 		{
 			const TImplicitObjectTransformed<T, d>& TransformedA = static_cast<const TImplicitObjectTransformed<T, d>&>(A);
 			const TRigidTransform<T, d> NewATM = TransformedA.GetTransform() * ATM;
-			return SweepQuery(*TransformedA.GetTransformedObject(), NewATM, B, BTM, Dir, Length, OutTime, OutPosition, OutNormal, Thickness);
+			return SweepQuery(*TransformedA.GetTransformedObject(), NewATM, B, BTM, Dir, Length, OutTime, OutPosition, OutNormal, OutFaceIndex, Thickness);
 		}
 
 		if (BType == ImplicitObjectType::Transformed)
 		{
 			const TImplicitObjectTransformed<T, d>& TransformedB = static_cast<const TImplicitObjectTransformed<T, d>&>(B);
 			const TRigidTransform<T, d> NewBTM = TransformedB.GetTransform() * BTM;
-			return SweepQuery(A, ATM, *TransformedB.GetTransformedObject(), NewBTM, Dir, Length, OutTime, OutPosition, OutNormal, Thickness);
+			return SweepQuery(A, ATM, *TransformedB.GetTransformedObject(), NewBTM, Dir, Length, OutTime, OutPosition, OutNormal, OutFaceIndex, Thickness);
 		}
 
 		check(B.IsConvex());	//Object being swept must be convex
+		OutFaceIndex = INDEX_NONE;
 		
 		TVector<T, d> LocalPosition;
 		TVector<T, d> LocalNormal;
@@ -111,17 +112,34 @@ namespace Chaos
 		const TRigidTransform<T, d> BToATM = BTM.GetRelativeTransform(ATM);
 		const TVector<T, d> LocalDir = ATM.InverseTransformVectorNoScale(Dir);
 
-		if (BType == ImplicitObjectType::Sphere)
+		bool bSweepAsRaycast = BType == ImplicitObjectType::Sphere;
+		if (bSweepAsRaycast && AType == ImplicitObjectType::Scaled)
+		{
+			const auto& Scaled = A. template GetObjectChecked<TImplicitObjectScaled<T, d>>();
+			const TVector<T, d>& Scale = Scaled.GetScale();
+			bSweepAsRaycast = FMath::IsNearlyEqual(Scale[0], Scale[1]) && FMath::IsNearlyEqual(Scale[0], Scale[2]);
+		}
+
+		if (bSweepAsRaycast)
 		{
 			const TSphere<T, d>& BSphere = static_cast<const TSphere<T, d>&>(B);
-			const TVector<T, d> Start = BToATM.TransformPositionNoScale(BSphere.Center());
-			bResult = B.Raycast(Start, LocalDir, Length, Thickness + BSphere.Radius(), OutTime, LocalPosition, LocalNormal);
+			const TVector<T, d> Start = BToATM.TransformPositionNoScale(BSphere.GetCenter());
+			bResult = A.Raycast(Start, LocalDir, Length, Thickness + BSphere.GetRadius(), OutTime, LocalPosition, LocalNormal, OutFaceIndex);
 		}
 		//todo: handle case where A is a sphere
 		else if (A.IsConvex())
 		{
 			const TVector<T, d> Offset = ATM.GetLocation() - BTM.GetLocation();
-			bResult = GJKRaycast<T>(A, B, BToATM, BToATM.GetLocation(), LocalDir, Length, OutTime, LocalPosition, LocalNormal, Thickness, Offset);
+			bResult = GJKRaycast<T>(A, B, BToATM, LocalDir, Length, OutTime, LocalPosition, LocalNormal, Thickness, Offset);
+
+			if (AType == ImplicitObjectType::Convex)
+			{
+				//todo: find face index
+			}
+			else if (AType == ImplicitObjectType::Scaled)
+			{
+				//todo: find face index if convex hull
+			}
 		}
 		else
 		{
@@ -130,13 +148,13 @@ namespace Chaos
 			case ImplicitObjectType::HeightField:
 			{
 				const THeightField<T>& AHeightField = static_cast<const THeightField<T>&>(A);
-				bResult = AHeightField.SweepGeom(B, BToATM, LocalDir, Length, OutTime, LocalPosition, LocalNormal, Thickness);
+				bResult = AHeightField.SweepGeom(B, BToATM, LocalDir, Length, OutTime, LocalPosition, LocalNormal, OutFaceIndex, Thickness);
 				break;
 			}
 			case ImplicitObjectType::TriangleMesh:
 			{
 				const TTriangleMeshImplicitObject<T>& ATriangleMesh = static_cast<const TTriangleMeshImplicitObject<T>&>(A);
-				bResult = ATriangleMesh.SweepGeom(B, BToATM, LocalDir, Length, OutTime, LocalPosition, LocalNormal, Thickness);
+				bResult = ATriangleMesh.SweepGeom(B, BToATM, LocalDir, Length, OutTime, LocalPosition, LocalNormal, OutFaceIndex, Thickness);
 				break;
 			}
 			case ImplicitObjectType::Scaled:
@@ -145,7 +163,7 @@ namespace Chaos
 				const TImplicitObject<T, d>& UnscaledObj = *AScaled.GetUnscaledObject();
 				check(UnscaledObj.GetType(true) == ImplicitObjectType::TriangleMesh);
 				const TTriangleMeshImplicitObject<T>& ATriangleMesh = static_cast<const TTriangleMeshImplicitObject<T>&>(UnscaledObj);
-				bResult = ATriangleMesh.SweepGeom(B, BToATM, LocalDir, Length, OutTime, LocalPosition, LocalNormal, Thickness, AScaled.GetScale());
+				bResult = ATriangleMesh.SweepGeom(B, BToATM, LocalDir, Length, OutTime, LocalPosition, LocalNormal, OutFaceIndex, Thickness, AScaled.GetScale());
 				break;
 			}
 			default:
@@ -155,8 +173,11 @@ namespace Chaos
 		}
 
 		//put back into world space
-		OutNormal = ATM.TransformVectorNoScale(LocalNormal);
-		OutPosition = ATM.TransformPositionNoScale(LocalPosition);
+		if (OutTime > 0)
+		{
+			OutNormal = ATM.TransformVectorNoScale(LocalNormal);
+			OutPosition = ATM.TransformPositionNoScale(LocalPosition);
+		}
 
 		return bResult;
 	}
diff --git a/Engine/Source/Runtime/Experimental/Chaos/Public/Chaos/HeightField.h b/Engine/Source/Runtime/Experimental/Chaos/Public/Chaos/HeightField.h
--- a/Engine/Source/Runtime/Experimental/Chaos/Public/Chaos/HeightField.h
+++ b/Engine/Source/Runtime/Experimental/Chaos/Public/Chaos/HeightField.h
@@ -27,16 +27,27 @@ namespace Chaos
 			return (T)0;
 		}
 
-		virtual bool Raycast(const TVector<T, 3>& StartPoint, const TVector<T, 3>& Dir, const T Length, const T Thickness, T& OutTime, TVector<T,3>& OutPosition, TVector<T,3>& OutNormal) const override;
+		virtual bool Raycast(const TVector<T, 3>& StartPoint, const TVector<T, 3>& Dir, const T Length, const T Thickness, T& OutTime, TVector<T,3>& OutPosition, TVector<T,3>& OutNormal, int32& OutFaceIndex) const override;
 		virtual bool Overlap(const TVector<T, 3>& Point, const T Thickness) const override;
 		bool OverlapGeom(const TImplicitObject<T, 3>& QueryGeom, const TRigidTransform<T, 3>& QueryTM, const T Thickness) const;
-		bool SweepGeom(const TImplicitObject<T, 3>& QueryGeom, const TRigidTransform<T, 3>& StartTM, const TVector<T, 3>& Dir, const T Length, T& OutTime, TVector<T, 3>& OutPosition, TVector<T, 3>& OutNormal, const T Thickness = 0) const;
+		bool SweepGeom(const TImplicitObject<T, 3>& QueryGeom, const TRigidTransform<T, 3>& StartTM, const TVector<T, 3>& Dir, const T Length, T& OutTime, TVector<T, 3>& OutPosition, TVector<T, 3>& OutNormal, int32& OutFaceIndex, const T Thickness = 0) const;
+		virtual int32 FindMostOpposingFace(const TVector<T, 3>& Position, const TVector<T, 3>& UnitDir, int32 HintFaceIndex) const override;
 
 		virtual const TBox<T, 3>& BoundingBox() const
 		{
 			return MTriangleMeshImplicitObject->BoundingBox();
 		}
 
+		virtual uint32 GetTypeHash() const override
+		{
+			return MTriangleMeshImplicitObject->GetTypeHash();
+		}
+
+		static ImplicitObjectType GetType()
+		{
+			return ImplicitObjectType::HeightField;
+		}
+
 	private:
 		//For now just convert into a trimesh
 		TUniquePtr<TTriangleMeshImplicitObject<T>> MTriangleMeshImplicitObject;
diff --git a/Engine/Source/Runtime/Experimental/Chaos/Public/Chaos/ISpatialAcceleration.h b/Engine/Source/Runtime/Experimental/Chaos/Public/Chaos/ISpatialAcceleration.h
--- a/Engine/Source/Runtime/Experimental/Chaos/Public/Chaos/ISpatialAcceleration.h
+++ b/Engine/Source/Runtime/Experimental/Chaos/Public/Chaos/ISpatialAcceleration.h
@@ -5,6 +5,9 @@
 namespace Chaos
 {
 
+template <typename T, int d>
+class TGeometryParticle;
+
 template <typename T, int d>
 class TSpatialRay
 {
@@ -27,7 +30,7 @@ public:
 	This class is responsible for gathering any information it wants (for example narrow phase query results).
 	This class determines whether the acceleration structure should continue to iterate through potential instances
 */
-template <typename T>
+template <typename TPayloadType, typename T>
 class CHAOS_API ISpatialVisitor
 {
 public:
@@ -37,21 +40,21 @@ public:
 		@Instance - the instance we are potentially overlapping
 		Returns true to continue iterating through the acceleration structure
 	*/
-	virtual bool Overlap(const int32 Instance) = 0;
+	virtual bool Overlap(const TPayloadType Instance) = 0;
 
 	/** Called whenever an instance in the acceleration structure may intersect with a raycast
 		@Instance - the instance we are potentially intersecting with a raycast
 		@CurLength - the length all future intersection tests will use. A blocking intersection should update this
 		Returns true to continue iterating through the acceleration structure
 	*/
-	virtual bool Raycast(const int32 Instance, T& CurLength) = 0;
+	virtual bool Raycast(const TPayloadType Instance, T& CurLength) = 0;
 
 	/** Called whenever an instance in the acceleration structure may intersect with a sweep
 		@Instance - the instance we are potentially intersecting with a sweep
 		@CurLength - the length all future intersection tests will use. A blocking intersection should update this
 		Returns true to continue iterating through the acceleration structure
 	*/
-	virtual bool Sweep(const int32 Instance, T& CurLength) = 0;
+	virtual bool Sweep(const TPayloadType Instance, T& CurLength) = 0;
 };
 
 enum class ESpatialAccelerationType
@@ -79,20 +82,20 @@ public:
 
 };
 
-template <typename T, int d>
+template <typename TPayloadType, typename T, int d>
 class CHAOS_API ISpatialAcceleration
 {
 public:
 	virtual ~ISpatialAcceleration() = default;
 
-	virtual TArray<int32> FindAllIntersections(const TBox<T, d>& Box) const = 0;
-	virtual TArray<int32> FindAllIntersections(const TSpatialRay<T,d>& Ray) const = 0;
-	virtual TArray<int32> FindAllIntersections(const TVector<T, d>& Point) const = 0;
-	virtual TArray<int32> FindAllIntersections(const TGeometryParticles<T, d>& InParticles, const int32 i) const = 0;
+	virtual TArray<TPayloadType> FindAllIntersections(const TBox<T, d>& Box) const = 0;
+	virtual TArray<TPayloadType> FindAllIntersections(const TSpatialRay<T,d>& Ray) const = 0;
+	virtual TArray<TPayloadType> FindAllIntersections(const TVector<T, d>& Point) const = 0;
+	virtual TArray<TPayloadType> FindAllIntersections(const TGeometryParticles<T, d>& InParticles, const int32 i) const = 0;
 
-	virtual void Raycast(const TVector<T, d>& Start, const TVector<T, d>& Dir, const T OriginalLength, ISpatialVisitor<T>& Visitor) const {}
-	virtual void Sweep(const TVector<T, d>& Start, const TVector<T, d>& Dir, T OriginalLength, const TVector<T, d> QueryHalfExtents, ISpatialVisitor<T>& Visitor, const TVector<T, d>& Scale = TVector<T, d>(1)) const {}
-	virtual void Overlap(const TBox<T, d>& QueryBounds, ISpatialVisitor<T>& Visitor, const TVector<T, d>& Scale = TVector<T, d>(1)) const {}
+	virtual void Raycast(const TVector<T, d>& Start, const TVector<T, d>& Dir, const T OriginalLength, ISpatialVisitor<TPayloadType, T>& Visitor) const {}
+	virtual void Sweep(const TVector<T, d>& Start, const TVector<T, d>& Dir, T OriginalLength, const TVector<T, d> QueryHalfExtents, ISpatialVisitor<TPayloadType, T>& Visitor, const TVector<T, d>& Scale = TVector<T, d>(1)) const {}
+	virtual void Overlap(const TBox<T, d>& QueryBounds, ISpatialVisitor<TPayloadType, T>& Visitor, const TVector<T, d>& Scale = TVector<T, d>(1)) const {}
 
 #if !UE_BUILD_SHIPPING
 	virtual void DebugDraw(ISpacialDebugDrawInterface<T>* InInterface) const {}
@@ -102,29 +105,29 @@ public:
 };
 
 /** Helper class used to bridge virtual to template implementation of acceleration structures */
-template <typename T>
+template <typename TPayloadType, typename T>
 class TSpatialVisitor
 {
 public:
-	TSpatialVisitor(ISpatialVisitor<T>& InVisitor)
+	TSpatialVisitor(ISpatialVisitor<TPayloadType, T>& InVisitor)
 		: Visitor(InVisitor) {}
-	FORCEINLINE bool VisitOverlap(const int32 Instance)
+	FORCEINLINE bool VisitOverlap(const TPayloadType Instance)
 	{
 		return Visitor.Overlap(Instance);
 	}
 
-	FORCEINLINE bool VisitRaycast(const int32 Instance, T& CurLength)
+	FORCEINLINE bool VisitRaycast(const TPayloadType Instance, T& CurLength)
 	{
 		return Visitor.Raycast(Instance, CurLength);
 	}
 
-	FORCEINLINE bool VisitSweep(const int32 Instance, T& CurLength)
+	FORCEINLINE bool VisitSweep(const TPayloadType Instance, T& CurLength)
 	{
 		return Visitor.Sweep(Instance, CurLength);
 	}
 
 private:
-	ISpatialVisitor<T>& Visitor;
+	ISpatialVisitor<TPayloadType, T>& Visitor;
 };
 
 }
diff --git a/Engine/Source/Runtime/Experimental/Chaos/Public/Chaos/ImplicitObject.h b/Engine/Source/Runtime/Experimental/Chaos/Public/Chaos/ImplicitObject.h
--- a/Engine/Source/Runtime/Experimental/Chaos/Public/Chaos/ImplicitObject.h
+++ b/Engine/Source/Runtime/Experimental/Chaos/Public/Chaos/ImplicitObject.h
@@ -88,6 +88,20 @@ public:
 		return nullptr;
 	}
 
+	template<class T_DERIVED>
+	const T_DERIVED& GetObjectChecked() const
+	{
+		check(T_DERIVED::GetType() == Type);
+		return static_cast<const T_DERIVED&>(*this);
+	}
+
+	template<class T_DERIVED>
+	T_DERIVED& GetObjectChecked()
+	{
+		check(T_DERIVED::GetType() == Type);
+		return static_cast<const T_DERIVED&>(*this);
+	}
+
 	ImplicitObjectType GetType(bool bGetTrueType = false) const
 	{
 		if (bIgnoreAnalyticCollisions && !bGetTrueType)
@@ -128,20 +142,32 @@ public:
 	Pair<TVector<T, d>, bool> FindClosestIntersection(const TVector<T, d>& StartPoint, const TVector<T, d>& EndPoint, const T Thickness) const;
 
 	//This gives derived types a way to avoid calling PhiWithNormal todo: this api is confusing
-	virtual bool Raycast(const TVector<T, d>& StartPoint, const TVector<T,d>& Dir, const T Length, const T Thickness, T& OutTime, TVector<T,d>& OutPosition, TVector<T,d>& OutNormal) const
+	virtual bool Raycast(const TVector<T, d>& StartPoint, const TVector<T,d>& Dir, const T Length, const T Thickness, T& OutTime, TVector<T,d>& OutPosition, TVector<T,d>& OutNormal, int32& OutFaceIndex) const
 	{
+		OutFaceIndex = INDEX_NONE;
 		const TVector<T, d> EndPoint = StartPoint + Dir * Length;
 		Pair<TVector<T,d>, bool> Result = FindClosestIntersection(StartPoint, EndPoint, Thickness);
 		if (Result.Second)
 		{
 			OutPosition = Result.First;
 			OutNormal = Normal(Result.First);
-			OutTime = Length > 0 ? (OutPosition - StartPoint).Size() / Length : 0.f;
+			OutTime = Length > 0 ? (OutPosition - StartPoint).Size() : 0.f;
 			return true;
 		}
 		return false;
 	}
 
+	/** Returns the most opposing face.
+		@param Position - local position to search around (for example an edge of a convex hull)
+		@param UnitDir - the direction we want to oppose (for example a ray moving into the edge of a convex hull would get the face with the most negative dot(FaceNormal, UnitDir)
+		@param HintFaceIndex - for certain geometry we can use this to accelerate the search.
+	*/
+	virtual int32 FindMostOpposingFace(const TVector<T, d>& Position, const TVector<T, d>& UnitDir, int32 HintFaceIndex) const
+	{
+		//Many objects have no concept of a face
+		return INDEX_NONE;
+	}
+
 	//This gives derived types a way to do an overlap check without calling PhiWithNormal todo: this api is confusing
 	virtual bool Overlap(const TVector<T, d>& Point, const T Thickness) const
 	{
@@ -176,6 +202,8 @@ public:
 	
 	static FArchive& SerializeLegacyHelper(FArchive& Ar, TUniquePtr<TImplicitObject<T, d>>& Value);
 
+	virtual uint32 GetTypeHash() const = 0;
+
 protected:
 	ImplicitObjectType Type;
 	bool bIsConvex;
@@ -206,5 +234,11 @@ static void StaticSerialize(FChaosArchive& Ar, TSerializablePtr<Type>& Obj)\
 	TImplicitObject<T, d>::StaticSerialize(Ar, (TSerializablePtr<TImplicitObject<T, d>>&)Obj);\
 }
 
+#define IMPLICIT_OBJECT_SERIALIZER_DIM(Type, Dim)\
+static void StaticSerialize(FChaosArchive& Ar, TSerializablePtr<Type>& Obj)\
+{\
+	TImplicitObject<T, Dim>::StaticSerialize(Ar, (TSerializablePtr<TImplicitObject<T, Dim>>&)Obj);\
+}
+
 typedef TImplicitObject<float, 3> FImplicitObject3;
 }
\ No newline at end of file
diff --git a/Engine/Source/Runtime/Experimental/Chaos/Public/Chaos/ImplicitObjectIntersection.h b/Engine/Source/Runtime/Experimental/Chaos/Public/Chaos/ImplicitObjectIntersection.h
--- a/Engine/Source/Runtime/Experimental/Chaos/Public/Chaos/ImplicitObjectIntersection.h
+++ b/Engine/Source/Runtime/Experimental/Chaos/Public/Chaos/ImplicitObjectIntersection.h
@@ -57,6 +57,19 @@ class TImplicitObjectIntersection : public TImplicitObject<T, d>
 
 	virtual const TBox<T,d>& BoundingBox() const { return MLocalBoundingBox; }
 
+
+	virtual uint32 GetTypeHash() const override
+	{
+		uint32 OutHash = MObjects.Num() > 0 ? MObjects[0]->GetTypeHash() : 0;
+
+		for(const TUniquePtr<TImplicitObject<T, d>>& Ptr : MObjects)
+		{
+			HashCombine(Ptr->GetTypeHash(), OutHash);
+		}
+
+		return OutHash;
+	}
+
 private:
 	virtual Pair<TVector<T, d>, bool> FindClosestIntersectionImp(const TVector<T, d>& StartPoint, const TVector<T, d>& EndPoint, const T Thickness) const
 	{
diff --git a/Engine/Source/Runtime/Experimental/Chaos/Public/Chaos/ImplicitObjectScaled.h b/Engine/Source/Runtime/Experimental/Chaos/Public/Chaos/ImplicitObjectScaled.h
--- a/Engine/Source/Runtime/Experimental/Chaos/Public/Chaos/ImplicitObjectScaled.h
+++ b/Engine/Source/Runtime/Experimental/Chaos/Public/Chaos/ImplicitObjectScaled.h
@@ -73,11 +73,11 @@ public:
 		return ScaledPhi;
 	}
 
-	virtual bool Raycast(const TVector<T, d>& StartPoint, const TVector<T, d>& Dir, const T Length, const T Thickness, T& OutTime, TVector<T, d>& OutPosition, TVector<T, d>& OutNormal) const override
+	virtual bool Raycast(const TVector<T, d>& StartPoint, const TVector<T, d>& Dir, const T Length, const T Thickness, T& OutTime, TVector<T, d>& OutPosition, TVector<T, d>& OutNormal, int32& OutFaceIndex) const override
 	{
 		ensure(Length > 0);
 		ensure(FMath::IsNearlyEqual(Dir.SizeSquared(), 1, KINDA_SMALL_NUMBER));
-		ensure(Thickness == 0 || (MScale[0] == MScale[1] && MScale[0] == MScale[2]));	//non uniform turns sphere into an ellipsoid so no longer a raycast and requires a more expensive sweep
+		ensure(Thickness == 0 || (FMath::IsNearlyEqual(MScale[0], MScale[1]) && FMath::IsNearlyEqual(MScale[0], MScale[2])));	//non uniform turns sphere into an ellipsoid so no longer a raycast and requires a more expensive sweep
 
 		const TVector<T, d> UnscaledStart = MInvScale * StartPoint;
 		TVector<T, d> UnscaledDir = MInvScale * Dir * Length;
@@ -86,8 +86,9 @@ public:
 		TVector<T, d> UnscaledPosition;
 		TVector<T, d> UnscaledNormal;
 
-		if (MObject->Raycast(UnscaledStart, UnscaledDir, UnscaledLength, Thickness * MInvScale[0], OutTime, UnscaledPosition, UnscaledNormal))
+		if (MObject->Raycast(UnscaledStart, UnscaledDir, UnscaledLength, Thickness * MInvScale[0], OutTime, UnscaledPosition, UnscaledNormal, OutFaceIndex))
 		{
+			OutTime = UnscaledLength ? OutTime * (Length / UnscaledLength) : 0;	//is this needed? SafeNormalize above seems bad
 			OutPosition = MScale * UnscaledPosition;
 			OutNormal = (MScale * UnscaledNormal).GetSafeNormal();
 			return true;
@@ -156,6 +157,11 @@ public:
 		Ar << MObject << MScale << MInvScale << MLocalBoundingBox;
 	}
 
+	virtual uint32 GetTypeHash() const override
+	{
+		return HashCombine(MObject->GetTypeHash(), ::GetTypeHash(MScale));
+	}
+
 private:
 	ObjectType MObject;
 	TVector<T, d> MScale;
diff --git a/Engine/Source/Runtime/Experimental/Chaos/Public/Chaos/ImplicitObjectTransformed.h b/Engine/Source/Runtime/Experimental/Chaos/Public/Chaos/ImplicitObjectTransformed.h
--- a/Engine/Source/Runtime/Experimental/Chaos/Public/Chaos/ImplicitObjectTransformed.h
+++ b/Engine/Source/Runtime/Experimental/Chaos/Public/Chaos/ImplicitObjectTransformed.h
@@ -5,7 +5,6 @@
 #include "Chaos/ImplicitObject.h"
 #include "Chaos/Transform.h"
 #include "ChaosArchive.h"
-#include <type_traits>
 #include "Templates/EnableIf.h"
 
 namespace Chaos
@@ -38,7 +37,7 @@ void TImplicitObjectTransformAccumulateSerializableHelper(TArray<Pair<TSerializa
 template<class T, int d, bool bSerializable = true>
 class TImplicitObjectTransformed final : public TImplicitObject<T, d>
 {
-using ObjectType = typename std::conditional<bSerializable, TSerializablePtr<TImplicitObject<T, d>>, const TImplicitObject<T, d>*>::type;
+using ObjectType = typename TChooseClass<bSerializable, TSerializablePtr<TImplicitObject<T, d>>, const TImplicitObject<T, d>*>::Result;
 
 public:
 	IMPLICIT_OBJECT_SERIALIZER(TImplicitObjectTransformed)
@@ -89,14 +88,14 @@ public:
 		return Phi;
 	}
 
-	virtual bool Raycast(const TVector<T, d>& StartPoint, const TVector<T, d>& Dir, const T Length, const T Thickness, T& OutTime, TVector<T, d>& OutPosition, TVector<T, d>& OutNormal) const override
+	virtual bool Raycast(const TVector<T, d>& StartPoint, const TVector<T, d>& Dir, const T Length, const T Thickness, T& OutTime, TVector<T, d>& OutPosition, TVector<T, d>& OutNormal, int32& OutFaceIndex) const override
 	{
 		const TVector<T, d> LocalStart = MTransform.InverseTransformPosition(StartPoint);
 		const TVector<T, d> LocalDir = MTransform.InverseTransformVector(Dir);
 		TVector<T, d> LocalPosition;
 		TVector<T, d> LocalNormal;
 
-		if (MObject->Raycast(LocalStart, LocalDir, Length, Thickness, OutTime, LocalPosition, LocalNormal))
+		if (MObject->Raycast(LocalStart, LocalDir, Length, Thickness, OutTime, LocalPosition, LocalNormal, OutFaceIndex))
 		{
 			OutPosition = MTransform.TransformPosition(LocalPosition);
 			OutNormal = MTransform.TransformVector(LocalNormal);
@@ -172,6 +171,12 @@ public:
 		Ar << MTransform << MLocalBoundingBox;
 	}
 
+	virtual uint32 GetTypeHash() const override
+	{
+		// Combine the hash from the inner, non transformed object with our transform
+		return HashCombine(MObject->GetTypeHash(), ::GetTypeHash(MTransform));
+	}
+
 private:
 	ObjectType MObject;
 	TUniquePtr<Chaos::TImplicitObject<T, d>> MObjectOwner;
diff --git a/Engine/Source/Runtime/Experimental/Chaos/Public/Chaos/ImplicitObjectUnion.h b/Engine/Source/Runtime/Experimental/Chaos/Public/Chaos/ImplicitObjectUnion.h
--- a/Engine/Source/Runtime/Experimental/Chaos/Public/Chaos/ImplicitObjectUnion.h
+++ b/Engine/Source/Runtime/Experimental/Chaos/Public/Chaos/ImplicitObjectUnion.h
@@ -19,7 +19,7 @@ class TImplicitObjectUnion : public TImplicitObject<T, d>
 	    : TImplicitObject<T, d>(EImplicitObject::HasBoundingBox, ImplicitObjectType::Union)
 	    , MObjects(MoveTemp(Objects))
 		, Hierarchy(GeomParticles)
-	    , MLocalBoundingBox(MObjects[0]->BoundingBox())
+	    , MLocalBoundingBox()
 		, bHierarchyBuilt(false)
 		, MOriginalParticleLookupHack(OriginalParticleLookupHack)
 	{
@@ -29,6 +29,10 @@ class TImplicitObjectUnion : public TImplicitObject<T, d>
 			{
 				MLocalBoundingBox.GrowToInclude(MObjects[i]->BoundingBox());
 			}
+			else
+			{
+				MLocalBoundingBox = MObjects[i]->BoundingBox();
+			}
 			check(MOriginalParticleLookupHack.Num() == 0 || MOriginalParticleLookupHack.Num() == MObjects.Num());
 			if (MOriginalParticleLookupHack.Num() > 0)
 			{
@@ -229,7 +233,7 @@ class TImplicitObjectUnion : public TImplicitObject<T, d>
 		}
 	}
 
-	virtual bool Raycast(const TVector<T, d>& StartPoint, const TVector<T, d>& Dir, const T Length, const T Thickness, T& OutTime, TVector<T, d>& OutPosition, TVector<T, d>& OutNormal) const override
+	virtual bool Raycast(const TVector<T, d>& StartPoint, const TVector<T, d>& Dir, const T Length, const T Thickness, T& OutTime, TVector<T, d>& OutPosition, TVector<T, d>& OutNormal, int32& OutFaceIndex) const override
 	{
 		T MinTime = 0;	//initialization not needed, but doing it to avoid warning
 		bool bFound = false;
@@ -239,7 +243,8 @@ class TImplicitObjectUnion : public TImplicitObject<T, d>
 			TVector<T, d> Position;
 			TVector<T, d> Normal;
 			T Time;
-			if (Obj->Raycast(StartPoint, Dir, Length, Thickness, Time, Position, Normal))
+			int32 FaceIdx;
+			if (Obj->Raycast(StartPoint, Dir, Length, Thickness, Time, Position, Normal, FaceIdx))
 			{
 				if (!bFound || Time < MinTime)
 				{
@@ -247,6 +252,7 @@ class TImplicitObjectUnion : public TImplicitObject<T, d>
 					OutTime = Time;
 					OutPosition = Position;
 					OutNormal = Normal;
+					OutFaceIndex = FaceIdx;
 					bFound = true;
 				}
 			}
@@ -277,6 +283,19 @@ class TImplicitObjectUnion : public TImplicitObject<T, d>
 
 	const TArray<TUniquePtr<TImplicitObject<T, d>>>& GetObjects() const { return MObjects; }
 
+	virtual uint32 GetTypeHash() const override
+	{
+		uint32 Result = 0;
+
+		// Union hash is just the hash of all internal objects
+		for(const TUniquePtr<TImplicitObject<T, d>>& InnerObj : MObjects)
+		{
+			Result = HashCombine(Result, InnerObj->GetTypeHash());
+		}
+
+		return Result;
+	}
+
 private:
 	virtual Pair<TVector<T, d>, bool> FindClosestIntersectionImp(const TVector<T, d>& StartPoint, const TVector<T, d>& EndPoint, const T Thickness) const override
 	{
diff --git a/Engine/Source/Runtime/Experimental/Chaos/Public/Chaos/KinematicGeometryParticles.h b/Engine/Source/Runtime/Experimental/Chaos/Public/Chaos/KinematicGeometryParticles.h
--- a/Engine/Source/Runtime/Experimental/Chaos/Public/Chaos/KinematicGeometryParticles.h
+++ b/Engine/Source/Runtime/Experimental/Chaos/Public/Chaos/KinematicGeometryParticles.h
@@ -6,27 +6,26 @@
 
 namespace Chaos
 {
-template <typename T, int d>
-class TKinematicGeometryParticleHandle;
-
-template<class T, int d>
-class CHAOS_API TKinematicGeometryParticles : public TGeometryParticles<T, d>
+template<class T, int d, EGeometryParticlesSimType SimType>
+class TKinematicGeometryParticlesImp : public TGeometryParticlesImp<T, d, SimType>
 {
   public:
-	TKinematicGeometryParticles()
-	    : TGeometryParticles<T, d>()
+	TKinematicGeometryParticlesImp()
+	    : TGeometryParticlesImp<T, d, SimType>()
 	{
+		this->MParticleType = EParticleType::Kinematic;
 		TArrayCollection::AddArray(&MV);
 		TArrayCollection::AddArray(&MW);
 	}
-	TKinematicGeometryParticles(const TKinematicGeometryParticles<T, d>& Other) = delete;
-	TKinematicGeometryParticles(TKinematicGeometryParticles<T, d>&& Other)
-	    : TGeometryParticles<T, d>(MoveTemp(Other)), MV(MoveTemp(Other.MV)), MW(MoveTemp(Other.MW))
+	TKinematicGeometryParticlesImp(const TKinematicGeometryParticlesImp<T, d, SimType>& Other) = delete;
+	TKinematicGeometryParticlesImp(TKinematicGeometryParticlesImp<T, d, SimType>&& Other)
+	    : TGeometryParticlesImp<T, d, SimType>(MoveTemp(Other)), MV(MoveTemp(Other.MV)), MW(MoveTemp(Other.MW))
 	{
+		this->MParticleType = EParticleType::Kinematic;
 		TArrayCollection::AddArray(&MV);
 		TArrayCollection::AddArray(&MW);
 	}
-	~TKinematicGeometryParticles() {}
+	~TKinematicGeometryParticlesImp() {}
 
 	const TVector<T, d>& V(const int32 Index) const { return MV[Index]; }
 	TVector<T, d>& V(const int32 Index) { return MV[Index]; }
@@ -36,7 +35,7 @@ class CHAOS_API TKinematicGeometryParticles : public TGeometryParticles<T, d>
 
 	FString ToString(int32 index) const
 	{
-		FString BaseString = TGeometryParticles<T, d>::ToString(index);
+		FString BaseString = TGeometryParticlesImp<T, d, SimType>::ToString(index);
 		return FString::Printf(TEXT("%s, MV:%s, MW:%s"), *BaseString, *V(index).ToString(), *W(index).ToString());
 	}
 
@@ -48,7 +47,7 @@ class CHAOS_API TKinematicGeometryParticles : public TGeometryParticles<T, d>
 
 	void Serialize(FChaosArchive& Ar)
 	{
-		TGeometryParticles<T, d>::Serialize(Ar);
+		TGeometryParticlesImp<T, d, SimType>::Serialize(Ar);
 		Ar << MV << MW;
 	}
 
@@ -57,10 +56,19 @@ class CHAOS_API TKinematicGeometryParticles : public TGeometryParticles<T, d>
 	TArrayCollectionArray<TVector<T, d>> MW;
 };
 
-template <typename T, int d>
-FChaosArchive& operator<<(FChaosArchive& Ar, TKinematicGeometryParticles<T, d>& Particles)
+template <typename T, int d, EGeometryParticlesSimType SimType>
+FChaosArchive& operator<<(FChaosArchive& Ar, TKinematicGeometryParticlesImp<T, d, SimType>& Particles)
 {
 	Particles.Serialize(Ar);
 	return Ar;
 }
+
+extern template class Chaos::TKinematicGeometryParticlesImp<float, 3, Chaos::EGeometryParticlesSimType::RigidBodySim>;
+extern template class Chaos::TKinematicGeometryParticlesImp<float, 3, Chaos::EGeometryParticlesSimType::Other>;
+
+template <typename T, int d>
+using TKinematicGeometryParticles = TKinematicGeometryParticlesImp<T, d, EGeometryParticlesSimType::RigidBodySim>;
+
+template <typename T, int d>
+using TKinematicGeometryClothParticles = TKinematicGeometryParticlesImp<T, d, EGeometryParticlesSimType::Other>;
 }
diff --git a/Engine/Source/Runtime/Experimental/Chaos/Public/Chaos/Levelset.h b/Engine/Source/Runtime/Experimental/Chaos/Public/Chaos/Levelset.h
--- a/Engine/Source/Runtime/Experimental/Chaos/Public/Chaos/Levelset.h
+++ b/Engine/Source/Runtime/Experimental/Chaos/Public/Chaos/Levelset.h
@@ -80,6 +80,20 @@ class CHAOS_API TLevelSet final : public TImplicitObject<T, d>
 
 	bool CheckData(FErrorReporter& ErrorReporter, const TParticles<T, d>& InParticles, const TTriangleMesh<T>& Mesh, const TArray<TVector<T, 3>> &Normals);
 
+	virtual uint32 GetTypeHash() const override
+	{
+		uint32 Result = 0;
+
+		const int32 NumValues = MPhi.Num();
+
+		for(int32 Index = 0; Index < NumValues; ++Index)
+		{
+			Result = HashCombine(Result, ::GetTypeHash(MPhi[Index]));
+		}
+
+		return Result;
+	}
+
   private:
 	bool ComputeDistancesNearZeroIsocontour(FErrorReporter& ErrorReporter, const TParticles<T, d>& InParticles, const TArray<TVector<T, 3>> &Normals, const TTriangleMesh<T>& Mesh, TArrayND<bool, d>& BlockedFaceX, TArrayND<bool, d>& BlockedFaceY, TArrayND<bool, d>& BlockedFaceZ, TArray<TVector<int32, d>>& InterfaceIndices);
 	void ComputeDistancesNearZeroIsocontour(const TImplicitObject<T, d>& Object, const TArrayND<T, d>& ObjectPhi, TArray<TVector<int32, d>>& InterfaceIndices);
diff --git a/Engine/Source/Runtime/Experimental/Chaos/Public/Chaos/OrientedParticles.h b/Engine/Source/Runtime/Experimental/Chaos/Public/Chaos/OrientedParticles.h
--- a/Engine/Source/Runtime/Experimental/Chaos/Public/Chaos/OrientedParticles.h
+++ b/Engine/Source/Runtime/Experimental/Chaos/Public/Chaos/OrientedParticles.h
@@ -22,7 +22,6 @@ class TOrientedParticles : public TParticles<T, d>
 	{
 		TArrayCollection::AddArray(&MR);
 	}
-	~TOrientedParticles() {}
 
 	const TRotation<T, d>& R(const int32 Index) const { return MR[Index]; }
 	TRotation<T, d>& R(const int32 Index) { return MR[Index]; }
