Copying //UE4/Dev-Physics to Dev-Main (//UE4/Dev-Main) [at] 7703071

diff --git a/Engine/Source/Runtime/Experimental/Chaos/Public/Chaos/PBDCollisionConstraint.h b/Engine/Source/Runtime/Experimental/Chaos/Public/Chaos/PBDCollisionConstraint.h
--- a/Engine/Source/Runtime/Experimental/Chaos/Public/Chaos/PBDCollisionConstraint.h
+++ b/Engine/Source/Runtime/Experimental/Chaos/Public/Chaos/PBDCollisionConstraint.h
@@ -4,10 +4,12 @@
 #include "Chaos/BoundingVolume.h"
 #include "Chaos/PBDCollisionTypes.h"
 #include "Framework/BufferedData.h"
+#include "Chaos/PBDCollisionConstraint.h"
 
 #include <memory>
 #include <queue>
 #include <sstream>
+#include "BoundingVolume.h"
 
 namespace Chaos
 {
@@ -29,6 +31,9 @@ class TBox;
 template<class T>
 class TChaosPhysicsMaterial;
 
+template <typename T, int d>
+class TPBDRigidsSOAs;
+
 /** Specifies the type of work we should do*/
 enum class ECollisionUpdateType
 {
@@ -36,6 +41,7 @@ enum class ECollisionUpdateType
 	Deepest	//find the deepest penetration. Compute location and normal
 };
 
+
 /**
  * Manages a set of contact constraints:
  *	- Performs collision detection to generate constraints.
@@ -50,43 +56,43 @@ class CHAOS_API TPBDCollisionConstraint
 
 	typedef TRigidBodyContactConstraint<T, d> FRigidBodyContactConstraint;
 
-	TPBDCollisionConstraint(TPBDRigidParticles<T, d>& InParticles, const TArray<int32>& InIndices, TArrayCollectionArray<bool>& Collided, const TArrayCollectionArray<TSerializablePtr<TChaosPhysicsMaterial<T>>>& PerParticleMaterials, const int32 PairIterations = 1, const T Thickness = (T)0);
+	TPBDCollisionConstraint(const TPBDRigidsSOAs<T,d>& InParticles, TArrayCollectionArray<bool>& Collided, const TArrayCollectionArray<TSerializablePtr<TChaosPhysicsMaterial<T>>>& PerParticleMaterials, const int32 PairIterations = 1, const T Thickness = (T)0);
 	virtual ~TPBDCollisionConstraint() {}
 
 	/**
 	 * Generate all contact constraints.
 	 */
-	void UpdatePositionBasedState(const TPBDRigidParticles<T, d>& InParticles, const TArray<int32>& InIndices, const T Dt);
+	void UpdatePositionBasedState(/*const TPBDRigidParticles<T, d>& InParticles, const TArray<int32>& InIndices,*/ const T Dt);
 
 	/**
 	 * Apply corrections for the specified list of constraints.
 	 */
 	// @todo(ccaulfield): this runs wide. The serial/parallel decision should be in the ConstraintRule
-	void Apply(TPBDRigidParticles<T, d>& InParticles, const T Dt, const TArray<int32>& InConstraintIndices);
+	void Apply(const T Dt, const TArray<int32>& InConstraintIndices);
 
 	/** 
 	 * Apply push out for the specified list of constraints.
 	 * Return true if we need another iteration 
 	 */
 	 // @todo(ccaulfield): this runs wide. The serial/parallel decision should be in the ConstraintRule
-	bool ApplyPushOut(TPBDRigidParticles<T, d>& InParticles, const T Dt, const TArray<int32>& InConstraintIndices, const TSet<int32>& IsTemporarilyStatic, int32 Iteration, int32 NumIterations);
+	bool ApplyPushOut(const T Dt, const TArray<int32>& InConstraintIndices, const TSet<TGeometryParticleHandle<T,d>*>& IsTemporarilyStatic, int32 Iteration, int32 NumIterations);
 
-	void RemoveConstraints(const TSet<uint32>& RemovedParticles);
-	void UpdateConstraints(const TPBDRigidParticles<T, d>& InParticles, const TArray<int32>& InIndices, T Dt, const TSet<uint32>& AddedParticles, const TArray<uint32>& ActiveParticles);
+	void RemoveConstraints(const TSet<TGeometryParticleHandle<T, d>*>& RemovedParticles);
+	void UpdateConstraints(/*const TPBDRigidParticles<T, d>& InParticles, const TArray<int32>& InIndices,*/ T Dt, const TSet<TGeometryParticleHandle<T, d>*>& AddedParticles);
 
 	static bool NearestPoint(TArray<Pair<TVector<T, d>, TVector<T, d>>>& Points, TVector<T, d>& Direction, TVector<T, d>& ClosestPoint);
 
 	const TArray<FRigidBodyContactConstraint>& GetAllConstraints() const { return Constraints; }
-	const ISpatialAcceleration<T, d>& GetSpatialAcceleration() const;// { return SpatialAccelerationResource.GetRead(); }
+	const ISpatialAcceleration<TGeometryParticleHandle<T, d>*, T, d>& GetSpatialAcceleration() const { return SpatialAccelerationResource.GetRead(); }
 	void ReleaseSpatialAcceleration() const;
 	void SwapSpatialAcceleration();
 
 	// @todo(ccaulfield) We should probably wrap the LevelsetConstraint functions in a utility class and remove them from here.
 	// They are currently public for the Linux build, and member functions for Headless Chaos (TPBDCollisionConstraintAccessor). 
-	template<ECollisionUpdateType, class T_PARTICLES>
-	static void UpdateLevelsetConstraint(const T_PARTICLES& InParticles, const T Thickness, TRigidBodyContactConstraint<T, d>& Constraint);
-	template<ECollisionUpdateType, class T_PARTICLES>
-	static void UpdateLevelsetConstraintGJK(const T_PARTICLES& InParticles, const T Thickness, TRigidBodyContactConstraint<T, d>& Constraint);
+	template<ECollisionUpdateType>
+	static void UpdateLevelsetConstraint(const T Thickness, FRigidBodyContactConstraint& Constraint);
+	template<ECollisionUpdateType>
+	static void UpdateLevelsetConstraintGJK(const T Thickness, FRigidBodyContactConstraint& Constraint);
 
 
 	int32 NumConstraints() const { return Constraints.Num(); }
@@ -94,30 +100,38 @@ class CHAOS_API TPBDCollisionConstraint
 	/**
 	 * Get the particles that are affected by the specified constraint.
 	 */
-	TVector<int32, 2> ConstraintParticleIndices(int32 ContactIndex) const { return { Constraints[ContactIndex].ParticleIndex, Constraints[ContactIndex].LevelsetIndex }; }
+	TVector<TGeometryParticleHandle<T, d>*, 2> ConstraintParticles(int32 ContactIndex) const { return { Constraints[ContactIndex].Particle, Constraints[ContactIndex].Levelset }; }
 
 	void SetPushOutPairIterations(int32 InPushOutPairIterations) { MPairIterations = InPushOutPairIterations; }
 
 private:
-	void Reset(const TPBDRigidParticles<T, d>& InParticles, const TArray<int32>& InIndices);
+	void Reset(/*const TPBDRigidParticles<T, d>& InParticles, const TArray<int32>& InIndices*/);
+
+	void Apply(const T Dt, FRigidBodyContactConstraint& Constraint);
+	void ApplyPushOut(const T Dt, FRigidBodyContactConstraint& Constraint, const TSet<TGeometryParticleHandle<T,d>*>& IsTemporarilyStatic, int32 Iteration, int32 NumIterations, bool& NeedsAnotherIteration);
 
 	template <bool bGatherStats = false>
-	void ComputeConstraints(const TPBDRigidParticles<T, d>& InParticles, const TArray<int32>& InIndices, T Dt);
+	void ComputeConstraints(/*const TPBDRigidParticles<T, d>& InParticles, const TArray<int32>& InIndices,*/ T Dt);
 
-	template<ECollisionUpdateType, class T_PARTICLES>
-	void UpdateConstraint(const T_PARTICLES& InParticles, const T Thickness, TRigidBodyContactConstraint<T, d>& Constraint);
+	template<ECollisionUpdateType>
+	void UpdateConstraint(const T Thickness, FRigidBodyContactConstraint& Constraint);
 
-	FRigidBodyContactConstraint ComputeConstraint(const TPBDRigidParticles<T, d>& InParticles, int32 Body1Index, int32 Body2Index, const T Thickness);
+	FRigidBodyContactConstraint ComputeConstraint(TGeometryParticleHandle<T, d>* Particle0, TGeometryParticleHandle<T, d>* Particle1, const T Thickness);
 
 	template <typename SPATIAL_ACCELERATION, bool bGatherStats>
-	void ComputeConstraintsHelper(const TPBDRigidParticles<T, d>& InParticles, const TArray<int32>& InIndices, T Dt, const SPATIAL_ACCELERATION& SpatialAcceleration);
+	void ComputeConstraintsHelper(T Dt, const SPATIAL_ACCELERATION& SpatialAcceleration);
 
 	template<typename SPATIAL_ACCELERATION>
-	void UpdateConstraintsHelper(const TPBDRigidParticles<T, d>& InParticles, const TArray<int32>& InIndices, T Dt, const TSet<uint32>& AddedParticles, const TArray<uint32>& ActiveParticles, SPATIAL_ACCELERATION& SpatialAcceleration);
+	void UpdateConstraintsHelper(/*const TPBDRigidParticles<T, d>& InParticles, const TArray<int32>& InIndices,*/ T Dt, const TSet<TGeometryParticleHandle<T, d>*>& AddedParticles, SPATIAL_ACCELERATION& SpatialAcceleration);
 
 	// @todo(ccaulfield): move spatial acceleration out of constraint container and make shareable
+#if CHAOS_PARTICLEHANDLE_TODO
 	TChaosReadWriteResource<TBoundingVolume<TPBDRigidParticles<T, d>, T, d>> SpatialAccelerationResource;
 	TChaosReadWriteResource<TBoundingVolumeHierarchy<TPBDRigidParticles<T, d>, TBoundingVolume<TPBDRigidParticles<T, d>, T, d>, T, d>> SpatialAccelerationResource2;
+#endif
+
+	const TPBDRigidsSOAs<T,d>& Particles;
+	TChaosReadWriteResource<TBoundingVolume<TGeometryParticles<T, d>, TGeometryParticleHandle<T,d>*, T, d>> SpatialAccelerationResource;
 
 	TArray<FRigidBodyContactConstraint> Constraints;
 	TArrayCollectionArray<bool>& MCollided;
@@ -127,4 +141,13 @@ private:
 	T MAngularFriction;
 	bool bUseCCD;
 };
+
+extern template void TPBDCollisionConstraint<float, 3>::UpdateConstraint<ECollisionUpdateType::Any>(const float Thickness, FRigidBodyContactConstraint& Constraint);
+extern template void TPBDCollisionConstraint<float, 3>::UpdateConstraint<ECollisionUpdateType::Deepest>(const float Thickness, FRigidBodyContactConstraint& Constraint);
+extern template void TPBDCollisionConstraint<float, 3>::UpdateLevelsetConstraint<ECollisionUpdateType::Any>(const float Thickness, FRigidBodyContactConstraint& Constraint);
+extern template void TPBDCollisionConstraint<float, 3>::UpdateLevelsetConstraint<ECollisionUpdateType::Deepest>(const float Thickness, FRigidBodyContactConstraint& Constraint);
+extern template void TPBDCollisionConstraint<float, 3>::UpdateLevelsetConstraintGJK<ECollisionUpdateType::Any>(const float Thickness, FRigidBodyContactConstraint& Constraint);
+extern template void TPBDCollisionConstraint<float, 3>::UpdateLevelsetConstraintGJK<ECollisionUpdateType::Deepest>(const float Thickness, FRigidBodyContactConstraint& Constraint);
+extern template void TPBDCollisionConstraint<float, 3>::ComputeConstraints<false>(float Dt);
+extern template void TPBDCollisionConstraint<float, 3>::ComputeConstraints<true>(float Dt);
 }
diff --git a/Engine/Source/Runtime/Experimental/Chaos/Public/Chaos/PBDCollisionConstraintPGS.h b/Engine/Source/Runtime/Experimental/Chaos/Public/Chaos/PBDCollisionConstraintPGS.h
--- a/Engine/Source/Runtime/Experimental/Chaos/Public/Chaos/PBDCollisionConstraintPGS.h
+++ b/Engine/Source/Runtime/Experimental/Chaos/Public/Chaos/PBDCollisionConstraintPGS.h
@@ -9,6 +9,8 @@
 #include <queue>
 #include <sstream>
 
+#if CHAOS_PARTICLEHANDLE_TODO
+
 namespace ChaosTest
 {
 	template<class T> void CollisionPGS();
@@ -87,3 +89,5 @@ class TPBDCollisionConstraintPGS
 	const bool bUseCCD;
 };
 }
+
+#endif
\ No newline at end of file
diff --git a/Engine/Source/Runtime/Experimental/Chaos/Public/Chaos/PBDCollisionTypes.h b/Engine/Source/Runtime/Experimental/Chaos/Public/Chaos/PBDCollisionTypes.h
--- a/Engine/Source/Runtime/Experimental/Chaos/Public/Chaos/PBDCollisionTypes.h
+++ b/Engine/Source/Runtime/Experimental/Chaos/Public/Chaos/PBDCollisionTypes.h
@@ -1,9 +1,9 @@
 // Copyright 1998-2019 Epic Games, Inc. All Rights Reserved.
 #pragma once
-
 #include "Chaos/Vector.h"
 #include "CoreMinimal.h"
 #include "Box.h"
+#include "Chaos/ParticleHandle.h"
 
 class UPhysicalMaterial;
 
@@ -15,7 +15,8 @@ template<class T, int d>
 struct TRigidBodyContactConstraint
 {
 	TRigidBodyContactConstraint() : AccumulatedImpulse(0.f) {}
-	int32 ParticleIndex, LevelsetIndex;
+	TGeometryParticleHandle<T, d>* Particle;
+	TGeometryParticleHandle<T, d>* Levelset;
 	TVector<T, d> Normal;
 	TVector<T, d> Location;
 	T Phi;
@@ -23,7 +24,7 @@ struct TRigidBodyContactConstraint
 
 	FString ToString() const
 	{
-		return FString::Printf(TEXT("ParticleIndex:%d, LevelsetIndex:%d, Normal:%s, Location:%s, Phi:%f, AccumulatedImpulse:%s"), ParticleIndex, LevelsetIndex, *Normal.ToString(), *Location.ToString(), Phi, *AccumulatedImpulse.ToString());
+		return FString::Printf(TEXT("Particle:%s, Levelset:%s, Normal:%s, Location:%s, Phi:%f, AccumulatedImpulse:%s"), *Particle->ToString(), *Levelset->ToString(), *Normal.ToString(), *Location.ToString(), Phi, *AccumulatedImpulse.ToString());
 	}
 };
 
@@ -31,7 +32,8 @@ template<class T, int d>
 struct TRigidBodyContactConstraintPGS
 {
 	TRigidBodyContactConstraintPGS() : AccumulatedImpulse(0.f) {}
-	int32 ParticleIndex, LevelsetIndex;
+	TGeometryParticleHandle<T, d>* Particle;
+	TGeometryParticleHandle<T, d>* Levelset;
 	TArray<TVector<T, d>> Normal;
 	TArray<TVector<T, d>> Location;
 	TArray<T> Phi;
@@ -55,8 +57,8 @@ struct TCollisionData
 		, Mass1((T)0.0)
 		, Mass2((T)0.0)
 		, PenetrationDepth((T)0.0)
-		, ParticleIndex(INDEX_NONE)
-		, LevelsetIndex(INDEX_NONE)
+		, Particle(nullptr)
+		, Levelset(nullptr)
 		, ParticleIndexMesh(INDEX_NONE)
 		, LevelsetIndexMesh(INDEX_NONE)
 	{}
@@ -71,8 +73,8 @@ struct TCollisionData
 		, T InMass1
 		, T InMass2
 		, T InPenetrationDepth
-		, int32 InParticleIndex
-		, int32 InLevelsetIndex
+		, TGeometryParticle<T, d>* InParticle
+		, TGeometryParticle<T, d>* InLevelset
 		, int32 InParticleIndexMesh
 		, int32 InLevelsetIndexMesh)
 		: Location(InLocation)
@@ -85,8 +87,8 @@ struct TCollisionData
 		, Mass1(InMass1)
 		, Mass2(InMass2)
 		, PenetrationDepth(InPenetrationDepth)
-		, ParticleIndex(InParticleIndex)
-		, LevelsetIndex(InLevelsetIndex)
+		, Particle(InParticle)
+		, Levelset(InLevelset)
 		, ParticleIndexMesh(InParticleIndexMesh)
 		, LevelsetIndexMesh(InLevelsetIndexMesh)
 	{}
@@ -98,12 +100,15 @@ struct TCollisionData
 	TVector<T, d> AngularVelocity1, AngularVelocity2;
 	T Mass1, Mass2;
 	T PenetrationDepth;
-	int32 ParticleIndex, LevelsetIndex;
+	TGeometryParticle<T, d>* Particle;
+	TGeometryParticle<T, d>* Levelset;
+	// @todo(ccaulfield): CHAOS_PARTICLEHANDLE_TODO
 	int32 ParticleIndexMesh, LevelsetIndexMesh; // If ParticleIndex points to a cluster then this index will point to an actual mesh in the cluster
 												// It is important to be able to get extra data from the component
 
 };
 
+
 /*
 CollisionData used in the subsystems
 */
@@ -120,8 +125,8 @@ struct TCollisionDataExt
 		, AngularVelocity2(TVector<T, d>((T)0.0))
 		, Mass1((T)0.0)
 		, Mass2((T)0.0)
-		, ParticleIndex(INDEX_NONE)
-		, LevelsetIndex(INDEX_NONE)
+		, Particle(nullptr)
+		, Levelset(nullptr)
 		, ParticleIndexMesh(INDEX_NONE)
 		, LevelsetIndexMesh(INDEX_NONE)
 		, BoundingboxVolume((T)-1.0)
@@ -139,8 +144,8 @@ struct TCollisionDataExt
 		, TVector<T, d> InAngularVelocity2
 		, T InMass1
 		, T InMass2
-		, int32 InParticleIndex
-		, int32 InLevelsetIndex
+		, TGeometryParticle<T, d>* InParticle
+		, TGeometryParticle<T, d>* InLevelset
 		, int32 InParticleIndexMesh
 		, int32 InLevelsetIndexMesh
 		, float InBoundingboxVolume
@@ -156,8 +161,8 @@ struct TCollisionDataExt
 		, AngularVelocity2(InAngularVelocity2)
 		, Mass1(InMass1)
 		, Mass2(InMass2)
-		, ParticleIndex(InParticleIndex)
-		, LevelsetIndex(InLevelsetIndex)
+		, Particle(InParticle)
+		, Levelset(InLevelset)
 		, ParticleIndexMesh(InParticleIndexMesh)
 		, LevelsetIndexMesh(InLevelsetIndexMesh)
 		, BoundingboxVolume(InBoundingboxVolume)
@@ -176,8 +181,8 @@ struct TCollisionDataExt
 		, AngularVelocity2(InCollisionData.AngularVelocity2)
 		, Mass1(InCollisionData.Mass1)
 		, Mass2(InCollisionData.Mass2)
-		, ParticleIndex(InCollisionData.ParticleIndex)
-		, LevelsetIndex(InCollisionData.LevelsetIndex)
+		, Particle(InCollisionData.Particle)
+		, Levelset(InCollisionData.Levelset)
 		, ParticleIndexMesh(InCollisionData.ParticleIndexMesh)
 		, LevelsetIndexMesh(InCollisionData.LevelsetIndexMesh)
 		, BoundingboxVolume((T)-1.0)
@@ -192,7 +197,8 @@ struct TCollisionDataExt
 	TVector<T, d> Velocity1, Velocity2;
 	TVector<T, d> AngularVelocity1, AngularVelocity2;
 	T Mass1, Mass2;
-	int32 ParticleIndex, LevelsetIndex;
+	TGeometryParticle<T, d>* Particle;
+	TGeometryParticle<T, d>* Levelset;
 	int32 ParticleIndexMesh, LevelsetIndexMesh;
 	float BoundingboxVolume;
 	float BoundingboxExtentMin, BoundingboxExtentMax;
@@ -210,32 +216,18 @@ struct TBreakingData
 		, Velocity(TVector<T, d>((T)0.0))
 		, AngularVelocity(TVector<T, d>((T)0.0))
 		, Mass((T)0.0)
+		, Particle(nullptr)
 		, ParticleIndex(INDEX_NONE)
 		, ParticleIndexMesh(INDEX_NONE)
 		, BoundingBox(TBox<T, d>(TVector<T, d>((T)0.0), TVector<T, d>((T)0.0)))
 	{}
 
-	TBreakingData(TVector<T, d> InLocation
-		, TVector<T, d> InVelocity
-		, TVector<T, d> InAngularVelocity
-		, T InMass
-		, int32 InParticleIndex
-		, int32 InParticleIndexMesh
-		, Chaos::TBox<T, d>& InBoundingBox)
-		: Location(InLocation)
-		, Velocity(InVelocity)
-		, AngularVelocity(InAngularVelocity)
-		, Mass(InMass)
-		, ParticleIndex(InParticleIndex)
-		, ParticleIndexMesh(InParticleIndexMesh)
-	    , BoundingBox(InBoundingBox)
-	{}
-
 	TVector<T, d> Location;
 	TVector<T, d> Velocity;
 	TVector<T, d> AngularVelocity;
 	T Mass;
-	int32 ParticleIndex;
+	TGeometryParticle<T, d>* Particle;
+	int32 ParticleIndex; //#todo: remove this in favor of TGeometryParticle?
 	int32 ParticleIndexMesh; // If ParticleIndex points to a cluster then this index will point to an actual mesh in the cluster
 							 // It is important to be able to get extra data from the component
 	Chaos::TBox<T, d> BoundingBox;
@@ -331,7 +323,7 @@ struct TTrailingData
 		, Velocity(TVector<T, d>((T)0.0))
 		, AngularVelocity(TVector<T, d>((T)0.0))
 		, Mass((T)0.0)
-		, ParticleIndex(INDEX_NONE)
+		, Particle(nullptr)
 		, ParticleIndexMesh(INDEX_NONE)
 		, BoundingBox(TBox<T, d>(TVector<T, d>((T)0.0), TVector<T, d>((T)0.0)))
 	{}
@@ -340,7 +332,7 @@ struct TTrailingData
 		, TVector<T, d> InVelocity
 		, TVector<T, d> InAngularVelocity
 		, T InMass
-		, int32 InParticleIndex
+		, TGeometryParticle<T, d>* InParticle
 		, int32 InParticleIndexMesh
 		, float InBoundingboxVolume
 		, float InBoundingboxExtentMin
@@ -351,7 +343,7 @@ struct TTrailingData
 		, Velocity(InVelocity)
 		, AngularVelocity(InAngularVelocity)
 		, Mass(InMass)
-		, ParticleIndex(InParticleIndex)
+		, Particle(InParticle)
 		, ParticleIndexMesh(InParticleIndexMesh)
 		, BoundingBox(InBoundingBox)
 	{}
@@ -360,19 +352,19 @@ struct TTrailingData
 	TVector<T, d> Velocity;
 	TVector<T, d> AngularVelocity;
 	T Mass;
-	int32 ParticleIndex;
+	TGeometryParticle<T, d>* Particle;
 	int32 ParticleIndexMesh; // If ParticleIndex points to a cluster then this index will point to an actual mesh in the cluster
 							 // It is important to be able to get extra data from the component
 	Chaos::TBox<T, d> BoundingBox;
 
 	friend inline uint32 GetTypeHash(const TTrailingData& Other)
 	{
-		return ::GetTypeHash(Other.ParticleIndex);
+		return ::GetTypeHash(Other.Particle);
 	}
 
 	friend bool operator==(const TTrailingData& A, const TTrailingData& B)
 	{
-		return A.ParticleIndex == B.ParticleIndex;
+		return A.Particle == B.Particle;
 	}
 };
 
@@ -387,7 +379,7 @@ struct TTrailingDataExt
 		, Velocity(TVector<T, d>((T)0.0))
 		, AngularVelocity(TVector<T, d>((T)0.0))
 		, Mass((T)0.0)
-		, ParticleIndex(INDEX_NONE)
+		, Particle(nullptr)
 		, ParticleIndexMesh(INDEX_NONE)
 		, BoundingboxVolume((T)-1.0)
 		, BoundingboxExtentMin((T)-1.0)
@@ -399,7 +391,7 @@ struct TTrailingDataExt
 		, TVector<T, d> InVelocity
 		, TVector<T, d> InAngularVelocity
 		, T InMass
-		, int32 InParticleIndex
+		, TGeometryParticle<T, d>* InParticle
 		, int32 InParticleIndexMesh
 		, float InBoundingboxVolume
 		, float InBoundingboxExtentMin
@@ -409,7 +401,7 @@ struct TTrailingDataExt
 		, Velocity(InVelocity)
 		, AngularVelocity(InAngularVelocity)
 		, Mass(InMass)
-		, ParticleIndex(InParticleIndex)
+		, Particle(InParticle)
 		, ParticleIndexMesh(InParticleIndexMesh)
 		, BoundingboxVolume(InBoundingboxVolume)
 		, BoundingboxExtentMin(InBoundingboxExtentMin)
@@ -422,7 +414,7 @@ struct TTrailingDataExt
 		, Velocity(InTrailingData.Velocity)
 		, AngularVelocity(InTrailingData.AngularVelocity)
 		, Mass(InTrailingData.Mass)
-		, ParticleIndex(InTrailingData.ParticleIndex)
+		, Particle(InTrailingData.Particle)
 		, ParticleIndexMesh(InTrailingData.ParticleIndexMesh)
 		, BoundingboxVolume((T)-1.0)
 		, BoundingboxExtentMin((T)-1.0)
@@ -434,7 +426,7 @@ struct TTrailingDataExt
 	TVector<T, d> Velocity;
 	TVector<T, d> AngularVelocity;
 	T Mass;
-	int32 ParticleIndex;
+	TGeometryParticle<T, d>* Particle;
 	int32 ParticleIndexMesh; // If ParticleIndex points to a cluster then this index will point to an actual mesh in the cluster
 							 // It is important to be able to get extra data from the component
 	float BoundingboxVolume;
@@ -443,12 +435,12 @@ struct TTrailingDataExt
 
 	friend inline uint32 GetTypeHash(const TTrailingDataExt& Other)
 	{
-		return ::GetTypeHash(Other.ParticleIndex);
+		return ::GetTypeHash(Other.Particle);
 	}
 
 	friend bool operator==(const TTrailingDataExt& A, const TTrailingDataExt& B)
 	{
-		return A.ParticleIndex == B.ParticleIndex;
+		return A.Particle == B.Particle;
 	}
 };
 
diff --git a/Engine/Source/Runtime/Experimental/Chaos/Public/Chaos/PBDConstraintColor.h b/Engine/Source/Runtime/Experimental/Chaos/Public/Chaos/PBDConstraintColor.h
--- a/Engine/Source/Runtime/Experimental/Chaos/Public/Chaos/PBDConstraintColor.h
+++ b/Engine/Source/Runtime/Experimental/Chaos/Public/Chaos/PBDConstraintColor.h
@@ -3,13 +3,13 @@
 
 #include "Chaos/PBDCollisionTypes.h"
 #include "Chaos/PBDConstraintGraph.h"
+#include "Chaos/ParticleHandle.h"
 
 // @todo(ccaulfield): can we get rid of this now?
 #define USE_CONTACT_LEVELS 1
 
 namespace Chaos
 {
-
 	/**
 	 * Generates color information for a single constraint rule in a connection graph.
 	 * Edges with the same color are non-interacting and can safely be processed in parallel.
@@ -32,7 +32,7 @@ namespace Chaos
 		/**
 		 * Calculate the color information for the specified island.
 		 */
-		void ComputeColor(const TPBDRigidParticles<T, d>& InParticles, const TArray<int32>& InIndices, const int32 Island, const FConstraintGraph& ConstraintGraph, uint32 ContainerId);
+		void ComputeColor(const int32 Island, const FConstraintGraph& ConstraintGraph, uint32 ContainerId);
 
 		/**
 		 * Get the Level-Color-ConstraintList map for the specified island.
@@ -50,8 +50,8 @@ namespace Chaos
 		int GetIslandMaxLevel(int32 Island) const;
 
 	private:
-		void ComputeContactGraph(const TPBDRigidParticles<T, d>& InParticles, const TArray<int32>& NodeIndices, const int32 Island, const FConstraintGraph& ConstraintGraph, uint32 ContainerId);
-		void ComputeIslandColoring(const TPBDRigidParticles<T, d>& InParticles, const TArray<int32>& InIndices, const int32 Island, const FConstraintGraph& ConstraintGraph, uint32 ContainerId);
+		void ComputeContactGraph(const int32 Island, const FConstraintGraph& ConstraintGraph, uint32 ContainerId);
+		void ComputeIslandColoring(const int32 Island, const FConstraintGraph& ConstraintGraph, uint32 ContainerId);
 
 		struct FGraphNodeColor
 		{
diff --git a/Engine/Source/Runtime/Experimental/Chaos/Public/Chaos/PBDConstraintGraph.h b/Engine/Source/Runtime/Experimental/Chaos/Public/Chaos/PBDConstraintGraph.h
--- a/Engine/Source/Runtime/Experimental/Chaos/Public/Chaos/PBDConstraintGraph.h
+++ b/Engine/Source/Runtime/Experimental/Chaos/Public/Chaos/PBDConstraintGraph.h
@@ -2,6 +2,7 @@
 #pragma once
 
 #include "Chaos/PBDCollisionTypes.h"
+#include "ParticleHandle.h"
 
 
 // @todo(ccaulfield): can we get rid of this now?
@@ -9,197 +10,209 @@
 
 namespace Chaos
 {
-template<typename T, int d>
-class TPBDRigidParticles;
+	template<typename T, int d>
+	class TPBDRigidParticles;
 
-template<typename T, int d>
-class TPBDConstraintColor;
+	template<typename T, int d>
+	class TPBDConstraintColor;
 
-template<class T>
-class TChaosPhysicsMaterial;
+	template<class T>
+	class TChaosPhysicsMaterial;
 
-template <typename T>
-class TSerializablePtr;
+	template <typename T>
+	class TSerializablePtr;
 
-template<typename T>
-class TArrayCollectionArray;
+	template<typename T>
+	class TArrayCollectionArray;
 
-
-/**
- * Build a graph of connected particles, and then a set of independent islands.
- * Particles/constraints in different islands do not interact, so islands can be updated in parallel.
- * This is also where particle sleeping is controlled.
- */
-// @todo(ccaulfield): Break out the Island and Sleep stuff from the particle graph
-// @todo(ccaulfield): The level system should be part of the Island system, and not embedded with the color system
-template<typename T, int d>
-class CHAOS_API TPBDConstraintGraph
-{
-public:
-	friend class TPBDConstraintColor<T, d>;
-
-	/** Number of sleep threshold passes before sleep mode is enabled */
-	static const int32 SleepCountThreshold = 5;
-
-	/** Information required to map a graph edge back to its constraint */
-	struct FConstraintData
-	{
-		uint32 ContainerId;
-		int32 ConstraintIndex;
-	};
-
-	TPBDConstraintGraph();
-	TPBDConstraintGraph(TPBDRigidParticles<T, d>& InParticles, const TArray<int32>& InActiveIndices);
-	virtual ~TPBDConstraintGraph() {}
-
-	/** 
-	 * Set up the particle-to-graph-node mapping for the specified active particle indices, and clear the graph.
-	 * Should be called before AddConstraint.
-	 */
-	void InitializeGraph(const TPBDRigidParticles<T, d>& InParticles, const TArray<int32>& InIndices);
-
-	/**
-	 * Reserve space in the graph for NumConstraints additional constraints.
-	 */
-	void ReserveConstraints(const int32 NumConstraints);
-
-	/**
-	 * Add a constraint to the graph for each constraint in the container.
-	 */
-	void AddConstraint(const uint32 InContainerId, const int32 InConstraintIndex, const TVector<int32, 2>& InConstrainedParticles);
-
-	/**
-	 * Add particles/constraints to their particle's already-assigned islands (if applicable).
-	 */
-	// @todo(ccaulfield): InitializeIslands and ResetIslands are a bit confusing. Try to come up with better names.
-	void ResetIslands(const TPBDRigidParticles<T, d>& InParticles, const TArray<int32>& InIndices);
-
-	/**
-	 * Generate the simulation islands of connected particles (AddConstraints must have already been called). There are no constraints connecting particles in different islands.
-	 */
-	void UpdateIslands(TPBDRigidParticles<T, d>& InParticles, const TArray<int32>& InIndices, TSet<int32>& ActiveIndices);
-
-	/**
-	 * Put particles in inactive islands to sleep.
-	 */
-	bool SleepInactive(TPBDRigidParticles<T, d>& InParticles, const int32 Island, const TArrayCollectionArray<TSerializablePtr<TChaosPhysicsMaterial<T>>>& PerParticleMaterialAttributes);
-
-	/**
-	 * Wake all particles in an Island.
-	 */
-	void WakeIsland(TPBDRigidParticles<T, d>& InParticles, const int32 Island);
-
-	/**
-	 * Ensure that the particles in each island have consistent sleep states - if any are awake, wake all.
-	 */
-	// @todo(ccaulfield): Do we really need this? It implies some behind-the-scenes state manipulation.
-	void ReconcileIslands(TPBDRigidParticles<T, d>& InParticles);
-
-	/**
-	 * Get the list of ConstraintData indices associated with the specified island. NOTE: ConstraintDataIndex != ConstraintIndex.
-	 * Indices returned are into the ConstraintData array in the ConstraintGraph to get to the Constraint Index and Container Id
-	 * @see GetConstraintsData().
-	 */
-	const TArray<int32>& GetIslandConstraintData(int32 Island) const
-	{
-		return IslandConstraints[Island];
-	}
-
-	/**
-	 * Get the list of Particle indices associated with the specified island.
-	 */
-	const TArray<int32>& GetIslandParticles(int32 Island) const
-	{
-		return IslandParticles[Island];
-	}
+	template <typename T, int d>
+	class TPBDRigidsSOAs;
 
 	/**
-	 * Get the list of Particle indices associated with the specified island.
-	 */
-	int32 GetIslandSleepCount(int32 Island) const
+	 * Build a graph of connected particles, and then a set of independent islands.
+	 * Particles/constraints in different islands do not interact, so islands can be updated in parallel.
+	 * This is also where particle sleeping is controlled.
+	 * Note that at the moment this graph is rebuilt every frame. When this changes some of the underlying data will need to be changed. I.e. move from indices to pointers to handles
+	 */
+	 // @todo(ccaulfield): Break out the Island and Sleep stuff from the particle graph
+	 // @todo(ccaulfield): The level system should be part of the Island system, and not embedded with the color system
+	template<typename T, int d>
+	class CHAOS_API TPBDConstraintGraph
 	{
-		return IslandSleepCounts[Island];
-	}
-
-	/**
-	 * The number of islands in the graph.
-	 */
-	int32 NumIslands() const
-	{
-		return IslandParticles.Num();
-	}
-
-	/**
-	 * Information mapping the edge back to a rule's constraint
-	 */
-	const FConstraintData& GetConstraintData(int32 ConstraintDataIndex) const;
+	public:
+		template <typename T2, int d2>
+		friend class TPBDConstraintColor;
 
-	/**
-	 * Enable a particle after initialization. Adds the particle to the island of the parent particle if supplied.
-	 */
-	void EnableParticle(TPBDRigidParticles<T, d>& InParticles, const int32 ParticleIndex, const int32 ParentParticleIndex);
-
-	/**
-	 * Disable a particle and remove it from its island.
-	 * @note: this does not remove any attached constraints - constraints need to be re-added and islands will need to be rebuilt.
-	 */
-	void DisableParticle(TPBDRigidParticles<T, d>& InParticles, const int32 ParticleIndex);
-
-	/**
-	 * Disable a set of particles and remove them from their island.
-	 * @note: this does not remove any attached constraints - constraints need to be re-added and islands will need to be rebuilt.
-	 */
-	void DisableParticles(TPBDRigidParticles<T, d>& InParticles, const TSet<int32>& InParticleIndices);
+		/** Number of sleep threshold passes before sleep mode is enabled */
+		static const int32 SleepCountThreshold = 5;
 
-private:
-	struct FGraphNode
-	{
-		FGraphNode()
-			: BodyIndex(INDEX_NONE)
-			, Island(INDEX_NONE)
+		/** Information required to map a graph edge back to its constraint */
+		struct FConstraintData
+		{
+			uint32 ContainerId;
+			int32 ConstraintIndex;
+		};
+
+		TPBDConstraintGraph();
+		TPBDConstraintGraph(const TParticleView<TGeometryParticles<T,d>>& Particles);
+		virtual ~TPBDConstraintGraph() {}
+
+		/**
+		 * Clear the graph and set up the particle-to-graph-node mapping for the specified particles
+		 * Should be called before AddConstraint.
+		 */
+		void InitializeGraph(const TParticleView<TGeometryParticles<T, d>>& Particles);
+
+		/**
+		 * Reserve space in the graph for NumConstraints additional constraints.
+		 */
+		void ReserveConstraints(const int32 NumConstraints);
+
+		/**
+		 * Add a constraint to the graph for each constraint in the container.
+		 */
+		void AddConstraint(const uint32 InContainerId, const int32 InConstraintIndex, const TVector<TGeometryParticleHandle<T,d>*, 2>& InConstrainedParticles);
+
+		/**
+		 * Add particles/constraints to their particle's already-assigned islands (if applicable).
+		 */
+		 // @todo(ccaulfield): InitializeIslands and ResetIslands are a bit confusing. Try to come up with better names.
+		void ResetIslands(const TParticleView<TPBDRigidParticles<T, d>>& PBDRigids);
+
+		/**
+		 * Generate the simulation islands of connected particles (AddConstraints must have already been called). There are no constraints connecting particles in different islands.
+		 */
+		void UpdateIslands(const TParticleView<TPBDRigidParticles<T, d>>& confusing, TPBDRigidsSOAs<T, d>& Particles);
+
+		/**
+		 * Put particles in inactive islands to sleep.
+		 */
+		bool SleepInactive(const int32 Island, const TArrayCollectionArray<TSerializablePtr<TChaosPhysicsMaterial<T>>>& PerParticleMaterialAttributes);
+
+		/**
+		 * Wake all particles in an Island.
+		 */
+		void WakeIsland(const int32 Island);
+
+		/**
+		 * Ensure that the particles in each island have consistent sleep states - if any are awake, wake all.
+		 */
+		 // @todo(ccaulfield): Do we really need this? It implies some behind-the-scenes state manipulation.
+		void ReconcileIslands();
+
+		/**
+		 * Get the list of ConstraintData indices associated with the specified island. NOTE: ConstraintDataIndex != ConstraintIndex.
+		 * Indices returned are into the ConstraintData array in the ConstraintGraph to get to the Constraint Index and Container Id
+		 * @see GetConstraintsData().
+		 */
+		const TArray<int32>& GetIslandConstraintData(int32 Island) const
 		{
+			return IslandToConstraints[Island];
 		}
 
-		TArray<int32> Edges;
-		int32 BodyIndex;
-		int32 Island;
-	};
-
-	struct FGraphEdge
-	{
-		FGraphEdge()
-			: FirstNode(INDEX_NONE)
-			, SecondNode(INDEX_NONE)
+		/**
+		 * Get the list of Particle indices associated with the specified island.
+		 */
+		const TArray<TGeometryParticleHandle<T,d>*>& GetIslandParticles(int32 Island) const
 		{
+			return IslandToParticles[Island];
 		}
 
-		int32 FirstNode;
-		int32 SecondNode;
-		FConstraintData Data;
-	};
+		/**
+		 * Get the list of Particle indices associated with the specified island.
+		 */
+		int32 GetIslandSleepCount(int32 Island) const
+		{
+			return IslandToSleepCount[Island];
+		}
 
-	struct FIslandData
-	{
-		FIslandData()
-			: bIsIslandPersistant(false)
+		/**
+		 * The number of islands in the graph.
+		 */
+		int32 NumIslands() const
 		{
+			return IslandToParticles.Num();
 		}
 
-		bool bIsIslandPersistant;
-	};
+		/**
+		 * Information mapping the edge back to a rule's constraint
+		 */
+		const FConstraintData& GetConstraintData(int32 ConstraintDataIndex) const;
+
+		/**
+		 * Enable a particle after initialization. Adds the particle to the island of the parent particle if supplied.
+		 */
+		void EnableParticle(TGeometryParticleHandle<T,d>* Particle, const TGeometryParticleHandle<T, d>* ParentParticle);
+
+		/**
+		 * Disable a particle and remove it from its island.
+		 * @note: this does not remove any attached constraints - constraints need to be re-added and islands will need to be rebuilt.
+		 */
+		void DisableParticle(TGeometryParticleHandle<T, d>* Particle);
 
-	void ComputeIslands(TPBDRigidParticles<T, d>& InParticles, const TArray<int32>& InIndices, TSet<int32>& ActiveIndices);
-	void ComputeIsland(const TPBDRigidParticles<T, d>& InParticles, const TArray<int32>& InIndices, const int32 Node, const int32 Island, TSet<int32>& DynamicParticlesInIsland, TSet<int32>& StaticParticlesInIsland);
-	bool CheckIslands(TPBDRigidParticles<T, d>& InParticles, const TSet<int32>& InParticleIndices);
+		//Remove particle from constraint, maybe rethink some of these names
+		void RemoveParticle(TGeometryParticleHandle<T, d>* Particle)
+		{
+			//we know it just removes, probably rename later
+			DisableParticle(Particle);
+		}
 
-	TArray<FGraphNode> Nodes;
-	TArray<FGraphEdge> Edges;
-	TArray<FIslandData> IslandData;
-	TMap<int32, int32> ParticleToNodeIndex;
+		/**
+		 * Disable a set of particles and remove them from their island.
+		 * @note: this does not remove any attached constraints - constraints need to be re-added and islands will need to be rebuilt.
+		 */
+		void DisableParticles(const TSet<TGeometryParticleHandle<T,d>*>& Particles);
 
-	TArray<TArray<int32>> IslandParticles;
-	TArray<TArray<int32>> IslandConstraints;
-	TArray<int32> IslandSleepCounts;
-};
+	private:
+		struct FGraphNode
+		{
+			FGraphNode()
+				: Particle(nullptr)
+				, Island(INDEX_NONE)
+			{
+			}
+
+			TArray<int32> Edges;
+			TGeometryParticleHandle<T,d>* Particle;
+			int32 Island;
+		};
+
+		struct FGraphEdge
+		{
+			FGraphEdge()
+				: FirstNode(INDEX_NONE)
+				, SecondNode(INDEX_NONE)
+			{
+			}
+
+			int32 FirstNode;
+			int32 SecondNode;
+			FConstraintData Data;
+		};
+
+		struct FIslandData
+		{
+			FIslandData()
+				: bIsIslandPersistant(false)
+			{
+			}
+
+			bool bIsIslandPersistant;
+		};
+
+		void ComputeIslands(TPBDRigidsSOAs<T,d>& Particles);
+		void ComputeIsland(const int32 Node, const int32 Island,
+			TSet<TGeometryParticleHandle<T, d>*>& DynamicParticlesInIsland, TSet<TGeometryParticleHandle<T, d>*>& StaticParticlesInIsland);
+		bool CheckIslands(const TArray<TGeometryParticleHandle<T, d>*>& Particles);
+
+		TArray<FGraphNode> Nodes;
+		TArray<FGraphEdge> Edges;
+		TArray<FIslandData> IslandToData;
+		TMap<TGeometryParticleHandle<T,d>*, int32> ParticleToNodeIndex;
+
+		TArray<TArray<TGeometryParticleHandle<T,d>*>> IslandToParticles;
+		TArray<TArray<int32>> IslandToConstraints;
+		TArray<int32> IslandToSleepCount;
+	};
 
 }
diff --git a/Engine/Source/Runtime/Experimental/Chaos/Public/Chaos/PBDConstraintRule.h b/Engine/Source/Runtime/Experimental/Chaos/Public/Chaos/PBDConstraintRule.h
--- a/Engine/Source/Runtime/Experimental/Chaos/Public/Chaos/PBDConstraintRule.h
+++ b/Engine/Source/Runtime/Experimental/Chaos/Public/Chaos/PBDConstraintRule.h
@@ -1,6 +1,7 @@
 // Copyright 1998-2019 Epic Games, Inc. All Rights Reserved.
 #pragma once
 
+#include "Chaos/PBDConstraintRule.h"
 #include "Chaos/PBDCollisionTypes.h"
 #include "Chaos/PBDConstraintGraph.h"
 #include "Chaos/PBDConstraintColor.h"
@@ -39,13 +40,13 @@ namespace Chaos
 		virtual void BindToGraph(FConstraintGraph& InContactGraph, uint32 InContainerId) {}
 
 		/** Called once per tick to allow constraint containers to create/alter their constraints based on particle position */
-		virtual void UpdatePositionBasedState(const TPBDRigidParticles<T, d>& InParticles, const TArray<int32>& InIndices, const T Dt) {}
+		virtual void UpdatePositionBasedState(const T Dt) {}
 
 		/** Apply all corrections for constraints in the specified island */
-		virtual void ApplyConstraints(TPBDRigidParticles<T, d>& InParticles, const T Dt, int32 Island) {};
+		virtual void ApplyConstraints(const T Dt, int32 Island) {};
 
 		/** Apply push out for constraints in the specified island */
-		virtual void ApplyPushOut(TPBDRigidParticles<T, d>& InParticles, const T Dt, int32 Island) {};
+		virtual void ApplyPushOut(const T Dt, int32 Island) {};
 
 		/** Add all constraints to the connectivity graph */
 		virtual void AddToGraph() {}
@@ -54,12 +55,11 @@ namespace Chaos
 		virtual void InitializeAccelerationStructures() {}
 
 		/** Set up the perf-acceleration structures for the specified island. May be called in parallel for islands */
-		virtual void UpdateAccelerationStructures(const TPBDRigidParticles<T, d>& InParticles, const TArray<int32>& InIndices, const int32 Island) {}
+		virtual void UpdateAccelerationStructures(const int32 Island) {}
 
 		/** Remove all constraints associated with the specified particles */
 		// @todo(ccaulfield): remove uint version
-		void RemoveConstraints(const TSet<int32>& RemovedParticles) { RemoveConstraints(reinterpret_cast<const TSet<uint32>&>(RemovedParticles)); }
-		virtual void RemoveConstraints(const TSet<uint32>& RemovedParticles) {}
+		virtual void RemoveConstraints(const TSet<TGeometryParticleHandle<T,d>*>& RemovedParticles) { }
 
 		/** The number of constraints in the collection */
 		virtual int32 NumConstraints() const { return 0; }
@@ -89,9 +89,9 @@ namespace Chaos
 			ContainerId = InContainerId;
 		}
 
-		virtual void UpdatePositionBasedState(const TPBDRigidParticles<T, d>& InParticles, const TArray<int32>& InIndices, const T Dt) override
+		virtual void UpdatePositionBasedState(const T Dt) override
 		{
-			Constraints.UpdatePositionBasedState(InParticles, InIndices, Dt);
+			Constraints.UpdatePositionBasedState(Dt);
 		}
 
 		virtual void AddToGraph() override
@@ -99,11 +99,11 @@ namespace Chaos
 			ConstraintGraph->ReserveConstraints(Constraints.NumConstraints());
 			for (int32 ConstraintIndex = 0; ConstraintIndex < Constraints.NumConstraints(); ++ConstraintIndex)
 			{
-				ConstraintGraph->AddConstraint(ContainerId, ConstraintIndex, Constraints.ConstraintParticleIndices(ConstraintIndex));
+				ConstraintGraph->AddConstraint(ContainerId, ConstraintIndex, Constraints.ConstraintParticles(ConstraintIndex));
 			}
 		}
 
-		virtual void RemoveConstraints(const TSet<uint32>& RemovedParticles) override
+		virtual void RemoveConstraints(const TSet<TGeometryParticleHandle<T,d>*>& RemovedParticles) override
 		{
 			// @todo(ccaulfield): should also remove graph edges here too. Currently the edges are rebuilt every tick so not a problem...
 			//ConstraintGraph->RemoveEdges(...)
@@ -135,19 +135,19 @@ namespace Chaos
 		{
 		}
 
-		virtual void ApplyConstraints(TPBDRigidParticles<T, d>& InParticles, const T Dt, int32 Island) override
+		virtual void ApplyConstraints(const T Dt, int32 Island) override
 		{
 			if (IslandConstraintLists[Island].Num())
 			{
-				Constraints.Apply(InParticles, Dt, IslandConstraintLists[Island]);
+				Constraints.Apply(Dt, IslandConstraintLists[Island]);
 			}
 		}
 
-		virtual void ApplyPushOut(TPBDRigidParticles<T, d>& InParticles, const T Dt, int32 Island) override
+		virtual void ApplyPushOut(const T Dt, int32 Island) override
 		{
 			if (IslandConstraintLists[Island].Num())
 			{
-				Constraints.ApplyPushOut(InParticles, Dt, IslandConstraintLists[Island]);
+				Constraints.ApplyPushOut(Dt, IslandConstraintLists[Island]);
 			}
 		}
 
@@ -160,7 +160,7 @@ namespace Chaos
 			}
 		}
 		
-		virtual void UpdateAccelerationStructures(const TPBDRigidParticles<T, d>& InParticles, const TArray<int32>& InIndices, const int32 Island) override
+		virtual void UpdateAccelerationStructures(const int32 Island) override
 		{
 			const TArray<int32>& ConstraintDataIndices = ConstraintGraph->GetIslandConstraintData(Island);
 			TArray<int32>& IslandConstraintList = IslandConstraintLists[Island];
@@ -209,12 +209,12 @@ namespace Chaos
 		}
 
 
-		virtual void UpdatePositionBasedState(const TPBDRigidParticles<T, d>& InParticles, const TArray<int32>& InIndices, const T Dt) override
+		virtual void UpdatePositionBasedState(const T Dt) override
 		{
-			Constraints.UpdatePositionBasedState(InParticles, InIndices, Dt);
+			Constraints.UpdatePositionBasedState(Dt);
 		}
 
-		virtual void ApplyConstraints(TPBDRigidParticles<T, d>& InParticles, const T Dt, int32 Island) override
+		virtual void ApplyConstraints(const T Dt, int32 Island) override
 		{
 			const typename FConstraintColor::FLevelToColorToConstraintListMap& LevelToColorToConstraintListMap = GraphColor.GetIslandLevelToColorToConstraintListMap(Island);
 			int32 MaxColor = GraphColor.GetIslandMaxColor(Island);
@@ -225,20 +225,20 @@ namespace Chaos
 				{
 					if (LevelToColorToConstraintListMap[Level].Contains(Color) && LevelToColorToConstraintListMap[Level][Color].Num())
 					{
-						Constraints.Apply(InParticles, Dt, LevelToColorToConstraintListMap[Level][Color]);
+						Constraints.Apply(Dt, LevelToColorToConstraintListMap[Level][Color]);
 					}
 				}
 			}
 		}
 
-		virtual void ApplyPushOut(TPBDRigidParticles<T, d>& InParticles, const T Dt, int32 Island) override
+		virtual void ApplyPushOut(const T Dt, int32 Island) override
 		{
 			const typename FConstraintColor::FLevelToColorToConstraintListMap& LevelToColorToConstraintListMap = GraphColor.GetIslandLevelToColorToConstraintListMap(Island);
 			int32 MaxColor = GraphColor.GetIslandMaxColor(Island);
 			int32 MaxLevel = GraphColor.GetIslandMaxLevel(Island);
 
 			bool NeedsAnotherIteration = false;
-			TSet<int32> IsTemporarilyStatic;	// Also needs to be per-constraint type
+			TSet<TGeometryParticleHandle<T,d>*> IsTemporarilyStatic;	// Also needs to be per-constraint type
 			for (int32 Iteration = 0; (Iteration == 0 || NeedsAnotherIteration) && Iteration < PushOutIterations; ++Iteration)
 			{
 				NeedsAnotherIteration = false;
@@ -248,7 +248,7 @@ namespace Chaos
 					{
 						if (LevelToColorToConstraintListMap[Level].Contains(Color) && LevelToColorToConstraintListMap[Level][Color].Num())
 						{
-							NeedsAnotherIteration = Constraints.ApplyPushOut(InParticles, Dt, LevelToColorToConstraintListMap[Level][Color], IsTemporarilyStatic, Iteration, PushOutIterations);
+							NeedsAnotherIteration = Constraints.ApplyPushOut(Dt, LevelToColorToConstraintListMap[Level][Color], IsTemporarilyStatic, Iteration, PushOutIterations);
 						}
 					}
 #if USE_SHOCK_PROPOGATION
@@ -259,18 +259,16 @@ namespace Chaos
 							for (int32 Edge = 0; Edge < LevelToColorToConstraintListMap[Level][Color].Num(); ++Edge)
 							{
 								const int32 ConstraintIndex = LevelToColorToConstraintListMap[Level][Color][Edge];
-								const TVector<int32, 2> ParticleIndices = Constraints.ConstraintParticleIndices(ConstraintIndex);
-								const int32 ParticleIndex = ParticleIndices[0];
-								const int32 LevelsetIndex = ParticleIndices[1];
+								const TVector<TGeometryParticleHandle<T,d>*, 2> Particles = Constraints.ConstraintParticles(ConstraintIndex);
 								if (Iteration == PushOutIterations - 1)
 								{
-									if (!InParticles.InvM(ParticleIndex) || IsTemporarilyStatic.Contains(ParticleIndex))
+									if (Particles[0]->AsDynamic() == nullptr || IsTemporarilyStatic.Contains(Particles[0]))
 									{
-										IsTemporarilyStatic.Add(LevelsetIndex);
+										IsTemporarilyStatic.Add(Particles[1]);
 									}
-									else if (!InParticles.InvM(LevelsetIndex) || IsTemporarilyStatic.Contains(LevelsetIndex))
+									else if (Particles[1]->AsDynamic() == nullptr || IsTemporarilyStatic.Contains(Particles[1]))
 									{
-										IsTemporarilyStatic.Add(ParticleIndex);
+										IsTemporarilyStatic.Add(Particles[0]);
 									}
 								}
 							}
@@ -286,9 +284,9 @@ namespace Chaos
 			GraphColor.InitializeColor(*ConstraintGraph);
 		}
 
-		virtual void UpdateAccelerationStructures(const TPBDRigidParticles<T, d>& InParticles, const TArray<int32>& InIndices, const int32 Island) override
+		virtual void UpdateAccelerationStructures(const int32 Island) override
 		{
-			GraphColor.ComputeColor(InParticles, InIndices, Island, *ConstraintGraph, ContainerId);
+			GraphColor.ComputeColor(Island, *ConstraintGraph, ContainerId);
 		}
 
 		void SetPushOutIterations(const int32 InPushOutIterations)
@@ -320,7 +318,7 @@ namespace Chaos
 	 * update loop and therefore know what algorithms should be applied.
 	 */
 	template<typename T, int d>
-	struct CHAOS_API FConstraintRuleFactory
+	struct CHAOS_API TConstraintRuleFactory
 	{
 	public:
 		typedef TPBDConstraintGraph<T, d> FConstraintGraph;
@@ -338,4 +336,4 @@ namespace Chaos
 		}
 	};
 
-}
+}
\ No newline at end of file
diff --git a/Engine/Source/Runtime/Experimental/Chaos/Public/Chaos/PBDEvolution.h b/Engine/Source/Runtime/Experimental/Chaos/Public/Chaos/PBDEvolution.h
--- a/Engine/Source/Runtime/Experimental/Chaos/Public/Chaos/PBDEvolution.h
+++ b/Engine/Source/Runtime/Experimental/Chaos/Public/Chaos/PBDEvolution.h
@@ -13,13 +13,13 @@ class CHAOS_API TPBDEvolution
 {
   public:
 	// TODO(mlentine): Init particles from some type of input
-	TPBDEvolution(TPBDParticles<T, d>&& InParticles, TKinematicGeometryParticles<T, d>&& InGeometryParticles, TArray<TVector<int32, 3>>&& CollisionTriangles, int32 NumIterations = 1, T CollisionThickness = 0, T SelfCollisionsThickness = 0, T CoefficientOfFriction = 0, T Damping = 0.04);
+	TPBDEvolution(TPBDParticles<T, d>&& InParticles, TKinematicGeometryClothParticles<T, d>&& InGeometryParticles, TArray<TVector<int32, 3>>&& CollisionTriangles, int32 NumIterations = 1, T CollisionThickness = 0, T SelfCollisionsThickness = 0, T CoefficientOfFriction = 0, T Damping = 0.04);
 	~TPBDEvolution() {}
 
 	void AdvanceOneTimeStep(const T dt);
 
 	void SetKinematicUpdateFunction(TFunction<void(TPBDParticles<T, d>&, const T, const T, const int32)> KinematicUpdate) { MKinematicUpdate = KinematicUpdate; }
-	void SetCollisionKinematicUpdateFunction(TFunction<void(TKinematicGeometryParticles<T, d>&, const T, const T, const int32)> KinematicUpdate) { MCollisionKinematicUpdate = KinematicUpdate; }
+	void SetCollisionKinematicUpdateFunction(TFunction<void(TKinematicGeometryClothParticles<T, d>&, const T, const T, const int32)> KinematicUpdate) { MCollisionKinematicUpdate = KinematicUpdate; }
 	void SetParticleUpdateFunction(TFunction<void(TPBDParticles<T, d>&, const T)> ParticleUpdate) { MParticleUpdate = ParticleUpdate; }
 	void AddPBDConstraintFunction(TFunction<void(TPBDParticles<T, d>&, const T)> ConstraintFunction) { MConstraintRules.Add(ConstraintFunction); }
 	void AddForceFunction(TFunction<void(TPBDParticles<T, d>&, const T, const int32)> ForceFunction) { MForceRules.Add(ForceFunction); }
@@ -27,8 +27,8 @@ class CHAOS_API TPBDEvolution
 	const TPBDParticles<T, d>& Particles() const { return MParticles; }
 	TPBDParticles<T, d>& Particles() { return MParticles; }
 
-	const TGeometryParticles<T, d>& CollisionParticles() const { return MCollisionParticles; }
-	TGeometryParticles<T, d>& CollisionParticles() { return MCollisionParticles; }
+	const TGeometryClothParticles<T, d>& CollisionParticles() const { return MCollisionParticles; }
+	TGeometryClothParticles<T, d>& CollisionParticles() { return MCollisionParticles; }
 	const bool Collided(int32 index) { return MCollided[index]; }
 
 	TArray<TVector<int32, 3>>& CollisionTriangles() { return MCollisionTriangles; }
@@ -41,7 +41,7 @@ class CHAOS_API TPBDEvolution
 
   private:
 	TPBDParticles<T, d> MParticles;
-	TKinematicGeometryParticles<T, d> MCollisionParticles;
+	TKinematicGeometryClothParticles<T, d> MCollisionParticles;
 	TArray<TVector<int32, 3>> MCollisionTriangles;
 	TSet<TVector<int32, 2>> MDisabledCollisionElements;
 	TArrayCollectionArray<bool> MCollided;
@@ -56,6 +56,6 @@ class CHAOS_API TPBDEvolution
 	TArray<TFunction<void(TPBDParticles<T, d>&, const T)>> MConstraintRules;
 	TFunction<void(TPBDParticles<T, d>&, const T)> MParticleUpdate;
 	TFunction<void(TPBDParticles<T, d>&, const T, const T, const int32)> MKinematicUpdate;
-	TFunction<void(TKinematicGeometryParticles<T, d>&, const T, const T, const int32)> MCollisionKinematicUpdate;
+	TFunction<void(TKinematicGeometryClothParticles<T, d>&, const T, const T, const int32)> MCollisionKinematicUpdate;
 };
 }
diff --git a/Engine/Source/Runtime/Experimental/Chaos/Public/Chaos/PBDJointConstraints.h b/Engine/Source/Runtime/Experimental/Chaos/Public/Chaos/PBDJointConstraints.h
--- a/Engine/Source/Runtime/Experimental/Chaos/Public/Chaos/PBDJointConstraints.h
+++ b/Engine/Source/Runtime/Experimental/Chaos/Public/Chaos/PBDJointConstraints.h
@@ -12,110 +12,40 @@ template<class T, int d>
 class CHAOS_API TPBDJointConstraints : public TPBDJointConstraintsBase<T, d>, public TPBDConstraintContainer<T, d>
 {
 	typedef TPBDJointConstraintsBase<T, d> Base;
-	using Base::Constraints;
-	using Base::Distances;
 
   public:
+	using Base::AddConstraint;
+	using Base::ApplySingle;
+	using Base::ConstraintParticles;
+	using Base::NumConstraints;
+	using Base::RemoveConstraints;
+
 	TPBDJointConstraints(const T InStiffness = (T)1)
 	    : TPBDJointConstraintsBase<T, d>(InStiffness) 
 	{}
 
-	TPBDJointConstraints(const TRigidParticles<T, d>& InParticles, const TArray<TVector<T, 3>>& Locations, TArray<TVector<int32, 2>>&& InConstraints, const T InStiffness = (T)1)
-	    : TPBDJointConstraintsBase<T, d>(InParticles, Locations, MoveTemp(InConstraints), InStiffness) 
+	TPBDJointConstraints(const TArray<TVector<T, 3>>& Locations, TArray<TVector<TGeometryParticleHandle<T,d>*, 2>>&& InConstraints, const T InStiffness = (T)1)
+	    : TPBDJointConstraintsBase<T, d>(Locations, MoveTemp(InConstraints), InStiffness) 
 	{}
 
 	virtual ~TPBDJointConstraints()
 	{}
 
-	TArray<TVector<int32, 2>>& GetConstraints()
-	{
-		return Constraints;
-	}
-
-	int32 NumConstraints() const
-	{
-		return Constraints.Num();
-	}
-
-	TVector<int32, 2> ConstraintParticleIndices(int32 ConstraintIndex) const
-	{
-		return Constraints[ConstraintIndex];
-	}
 
-	void UpdatePositionBasedState(const TPBDRigidParticles<T, d>& InParticles, const TArray<int32>& InIndices, const T Dt)
+	void UpdatePositionBasedState(const T Dt)
 	{
 	}
 
-	void ApplyHelper(TPBDRigidParticles<T, d>& InParticles, const T Dt, const TArray<int32>& InConstraintIndices) const
+	void Apply(const T Dt, const TArray<int32>& InConstraintIndices) const
 	{
-		// @todo(ccaulfield): should be an option. Per-constraint or per-container?
-		const bool bApplyProjection = true;
-
 		for (int32 ConstraintIndex : InConstraintIndices)
 		{
-			const TVector<int32, 2>& Constraint = Constraints[ConstraintIndex];
-			int32 ConstraintInnerIndex1 = Constraint[0];
-			int32 ConstraintInnerIndex2 = Constraint[1];
-
-			check(InParticles.Island(ConstraintInnerIndex1) == InParticles.Island(ConstraintInnerIndex2) || InParticles.Island(ConstraintInnerIndex1) == INDEX_NONE || InParticles.Island(ConstraintInnerIndex2) == INDEX_NONE);
-
-			const TVector<T, d> WorldSpaceX1 = InParticles.Q(ConstraintInnerIndex1).RotateVector(Distances[ConstraintIndex][0]) + InParticles.P(ConstraintInnerIndex1);
-			const TVector<T, d> WorldSpaceX2 = InParticles.Q(ConstraintInnerIndex2).RotateVector(Distances[ConstraintIndex][1]) + InParticles.P(ConstraintInnerIndex2);
-			const PMatrix<T, d, d> WorldSpaceInvI1 = (InParticles.Q(ConstraintInnerIndex1) * FMatrix::Identity) * InParticles.InvI(ConstraintInnerIndex1) * (InParticles.Q(ConstraintInnerIndex1) * FMatrix::Identity).GetTransposed();
-			const PMatrix<T, d, d> WorldSpaceInvI2 = (InParticles.Q(ConstraintInnerIndex2) * FMatrix::Identity) * InParticles.InvI(ConstraintInnerIndex2) * (InParticles.Q(ConstraintInnerIndex2) * FMatrix::Identity).GetTransposed();
-			const TVector<T, d> Delta = Base::GetDelta(InParticles, WorldSpaceX1, WorldSpaceX2, WorldSpaceInvI1, WorldSpaceInvI2, ConstraintIndex);
-
-			if (InParticles.InvM(ConstraintInnerIndex1) > 0)
-			{
-				const TVector<T, d> Radius = WorldSpaceX1 - InParticles.P(ConstraintInnerIndex1);
-				const TVector<T, d> Dp = InParticles.InvM(ConstraintInnerIndex1) * Delta;
-				const TRotation<T, d> Dq = TRotation<T, d>(WorldSpaceInvI1 * TVector<T, d>::CrossProduct(Radius, Delta), 0.f)* InParticles.Q(ConstraintInnerIndex1)* T(0.5);
-				InParticles.P(ConstraintInnerIndex1) += Dp;
-				InParticles.Q(ConstraintInnerIndex1) += Dq;
-				InParticles.Q(ConstraintInnerIndex1).Normalize();
-			}
-
-			if (InParticles.InvM(ConstraintInnerIndex2) > 0)
-			{
-				const TVector<T, d> Radius = WorldSpaceX2 - InParticles.P(ConstraintInnerIndex2);
-				const TVector<T, d> Dp = InParticles.InvM(ConstraintInnerIndex2) * Delta;
-				const TRotation<T, d> Dq = TRotation<T, d>(WorldSpaceInvI2 * TVector<T, d>::CrossProduct(Radius, -Delta), 0.f) * InParticles.Q(ConstraintInnerIndex2) * T(0.5);
-				InParticles.P(ConstraintInnerIndex2) -= Dp;
-				InParticles.Q(ConstraintInnerIndex2) += Dq;
-				InParticles.Q(ConstraintInnerIndex2).Normalize();
+			ApplySingle(Dt, ConstraintIndex);
 			}
-
-			if (bApplyProjection)
-			{
-				const TVector<T, d> WorldSpaceX12 = InParticles.Q(ConstraintInnerIndex1).RotateVector(Distances[ConstraintIndex][0]) + InParticles.P(ConstraintInnerIndex1);
-				const TVector<T, d> WorldSpaceX22 = InParticles.Q(ConstraintInnerIndex2).RotateVector(Distances[ConstraintIndex][1]) + InParticles.P(ConstraintInnerIndex2);
-				const TVector<T, d> Delta2 = (WorldSpaceX22 - WorldSpaceX12) / (InParticles.InvM(ConstraintInnerIndex1) + InParticles.InvM(ConstraintInnerIndex2));
-				if (InParticles.InvM(ConstraintInnerIndex1) > 0)
-				{
-					const TVector<T, d> Dp = InParticles.InvM(ConstraintInnerIndex1) * Delta2;
-					InParticles.P(ConstraintInnerIndex1) += Dp;
-				}
-				if (InParticles.InvM(ConstraintInnerIndex2) > 0)
-				{
-					const TVector<T, d> Dp = InParticles.InvM(ConstraintInnerIndex2) * Delta2;
-					InParticles.P(ConstraintInnerIndex2) -= Dp;
-				}
-			}
-		}
-	}
-
-	void Apply(TPBDRigidParticles<T, d>& InParticles, const T Dt, const TArray<int32>& InConstraintIndices) const
-	{
-		ApplyHelper(InParticles, Dt, InConstraintIndices);
-	}
-
-	void ApplyPushOut(TPBDRigidParticles<T, d>& InParticles, const T Dt, const TArray<int32>& InConstraintIndices) const
-	{
 	}
 
-	void RemoveConstraints(const TSet<uint32>& RemovedParticles)
+	void ApplyPushOut(const T Dt, const TArray<int32>& InConstraintIndices) const
 	{
-		// @todo(ccaulfield): constraint management
 	}
 };
 }
diff --git a/Engine/Source/Runtime/Experimental/Chaos/Public/Chaos/PBDJointConstraintsBase.h b/Engine/Source/Runtime/Experimental/Chaos/Public/Chaos/PBDJointConstraintsBase.h
--- a/Engine/Source/Runtime/Experimental/Chaos/Public/Chaos/PBDJointConstraintsBase.h
+++ b/Engine/Source/Runtime/Experimental/Chaos/Public/Chaos/PBDJointConstraintsBase.h
@@ -3,8 +3,8 @@
 
 #include "Chaos/Array.h"
 #include "Chaos/Map.h"
-#include "Chaos/PBDParticles.h"
 #include "Chaos/PerParticleRule.h"
+#include "ParticleHandle.h"
 
 namespace Chaos
 {
@@ -17,33 +17,86 @@ namespace Chaos
 		{
 		}
 
-		TPBDJointConstraintsBase(const TRigidParticles<T, d>& InParticles, const TArray<TVector<T, d>>& Locations, TArray<TVector<int32, 2>>&& InConstraints, const T InStiffness = (T)1)
+		TPBDJointConstraintsBase(const TArray<TVector<T, d>>& Locations, TArray<TVector<TGeometryParticleHandle<T,d>*, 2>>&& InConstraints, const T InStiffness = (T)1)
 		    : Constraints(MoveTemp(InConstraints)), Stiffness(InStiffness)
 		{
-			UpdateDistances(InParticles, Locations);
+			UpdateDistances(Locations);
 		}
 		virtual ~TPBDJointConstraintsBase() {}
 	
+		int32 NumConstraints() const
+		{
+			return Constraints.Num();
+		}
+
 		// @todo(ccaulfield): optimize and generalize the constraint allocation/free api. We still want packed arrays, so we need handles not direct indices.
-		int32 AddConstraint(const TRigidParticles<T, d>& InParticles, const TVector<int32, 2>& InConstrainedParticleIndices, const TVector<T, d>& InLocation)
+		int32 AddConstraint(const TVector<TGeometryParticleHandle<T,d>*, 2>& InConstrainedParticles, const TVector<T, d>& InLocation)
+		{
+			int32 ConstraintIndex = Constraints.Emplace(InConstrainedParticles);
+			UpdateDistance(InLocation, ConstraintIndex);
+			return ConstraintIndex;
+		}
+
+		int32 AddConstraintLocal(const TVector<TGeometryParticleHandle<T, d>*, 2>& InConstrainedParticles, const TVector<TVector<T, 3>, 2>& InLocations)
 		{
-			int32 ConstraintIndex = Constraints.Emplace(InConstrainedParticleIndices);
-			UpdateDistance(InParticles, InLocation, ConstraintIndex);
+			int32 ConstraintIndex = Constraints.Emplace(InConstrainedParticles);
+			Distances.SetNum(Constraints.Num());
+			Distances[ConstraintIndex] = InLocations;
 			return ConstraintIndex;
 		}
 
+		void RemoveConstraints(const TSet<TGeometryParticleHandle<T, d>*>& RemovedParticles)
+		{
+			// @todo(ccaulfield): constraint management
+		}
+
+		TArray<TVector<TGeometryParticleHandle<T, d>*, 2>>& GetConstraints()
+		{
+			return Constraints;
+		}
+
+		TVector<TGeometryParticleHandle<T, d>*, 2> ConstraintParticles(int32 ConstraintIndex) const
+		{
+			return Constraints[ConstraintIndex];
+		}
 
 	protected:
-		void UpdateDistance(const TRigidParticles<T, d>& InParticles, const TVector<T, d>& InLocation, const int32 InConstraintIndex);
-		void UpdateDistances(const TRigidParticles<T, d>& InParticles, const TArray<TVector<T, d>>& InLocations);
-		TVector<T, d> GetDelta(const TPBDRigidParticles<T, d>& InParticles, const TVector<T, d>& WorldSpaceX1, const TVector<T, d>& WorldSpaceX2, const PMatrix<T, d, d>& WorldSpaceInvI1, const PMatrix<T, d, d>& WorldSpaceInvI2, const int32 i) const;
+		void ApplySingle(const T Dt, const int32 ConstraintIndex) const
+		{
+			// @todo(ccaulfield): bApplyProjection should be an option. Either per-constraint or per-container...
+			const bool bApplyProjection = true;
+
+			const TVector<TGeometryParticleHandle<T, d>*, 2>& Constraint = Constraints[ConstraintIndex];
+			if (Constraint[0]->AsDynamic() && Constraint[1]->AsDynamic())
+			{
+				ApplyDynamicDynamic(Dt, ConstraintIndex, 0, 1, bApplyProjection);
+			}
+			else if (Constraint[0]->AsDynamic())
+			{
+				ApplyDynamicStatic(Dt, ConstraintIndex, 0, 1, bApplyProjection);
+			}
+			else
+			{
+				ApplyDynamicStatic(Dt, ConstraintIndex, 1, 0, bApplyProjection);
+			}
+		}
 
-		TArray<TVector<int32, 2>> Constraints;
-		TArray<TVector<TVector<T, 3>, 2>> Distances;
-	
 	private:
-		void UpdateDistanceInternal(const TRigidParticles<T, d>& InParticles, const TVector<T, d>& InLocation, const int32 InConstraintIndex);
+		void UpdateDistanceInternal(const TVector<T, d>& InLocation, const int32 InConstraintIndex);
+		void UpdateDistance(const TVector<T, d>& InLocation, const int32 InConstraintIndex);
+		void UpdateDistances(const TArray<TVector<T, d>>& InLocations);
+
+		// Double dynamic body solve
+		void ApplyDynamicDynamic(const T Dt, const int32 ConstraintIndex, const int32 PBDRigid0Index, const int32 PBDRigid1Index, const bool bApplyProjection) const;
+		TVector<T, d> GetDeltaDynamicDynamic(const TVector<T, d>& P0, const TVector<T, d>& P1, const TVector<T, d>& WorldSpaceX0, const TVector<T, d>& WorldSpaceX1, const PMatrix<T, d, d>& WorldSpaceInvI0, const PMatrix<T, d, d>& WorldSpaceInvI1, const T InvM0, const T InvM1) const;
 
+		// Single dynamic body solve
+		void ApplyDynamicStatic(const T Dt, const int32 ConstraintIndex, const int32 PBDRigid0Index, const int32 Static1Index, const bool bApplyProjection) const;
+		TVector<T, d> GetDeltaDynamicKinematic(const TVector<T, d>& P0, const TVector<T, d>& WorldSpaceX0, const TVector<T, d>& WorldSpaceX1, const PMatrix<T, d, d>& WorldSpaceInvI0, const T InvM0) const;
+
+		// @todo(ccaulfield): we never iterate over these separately. Do we still want SoA?
+		TArray<TVector<TGeometryParticleHandle<T, d>*, 2>> Constraints;
+		TArray<TVector<TVector<T, 3>, 2>> Distances;
 		T Stiffness;
 	};
 }
diff --git a/Engine/Source/Runtime/Experimental/Chaos/Public/Chaos/PBDPositionConstraints.h b/Engine/Source/Runtime/Experimental/Chaos/Public/Chaos/PBDPositionConstraints.h
--- a/Engine/Source/Runtime/Experimental/Chaos/Public/Chaos/PBDPositionConstraints.h
+++ b/Engine/Source/Runtime/Experimental/Chaos/Public/Chaos/PBDPositionConstraints.h
@@ -4,6 +4,7 @@
 #include "Chaos/Array.h"
 #include "Chaos/Map.h"
 #include "Chaos/PBDConstraintContainer.h"
+#include "ParticleHandle.h"
 
 namespace Chaos
 {
@@ -15,57 +16,54 @@ class TPBDPositionConstraints : public TPBDConstraintContainer<T, d>
 		: Stiffness(InStiffness)
 	{}
 
-	TPBDPositionConstraints(const TRigidParticles<T, d>& InParticles, TArray<TVector<T, d>>&& Locations, TArray<int32>&& InConstraints, const T InStiffness = (T)1)
-		: Targets(MoveTemp(Locations)), Constraints(MoveTemp(InConstraints)), Stiffness(InStiffness)
+	TPBDPositionConstraints(TArray<TVector<T, d>>&& Locations, TArray<TPBDRigidParticleHandle<T,d>*>&& InConstrainedParticles, const T InStiffness = (T)1)
+		: Targets(MoveTemp(Locations)), ConstrainedParticles(MoveTemp(InConstrainedParticles)), Stiffness(InStiffness)
 	{}
 
 	virtual ~TPBDPositionConstraints() {}
 
 	int32 NumConstraints() const
 	{
-		return Constraints.Num();
+		return ConstrainedParticles.Num();
 	}
 
-	TVector<int32, 2> ConstraintParticleIndices(int32 ConstraintIndex) const
+	TVector<TGeometryParticleHandle<T,d>*, 2> ConstraintParticles(int32 ConstraintIndex) const
 	{ 
-		return { Constraints[ConstraintIndex], INDEX_NONE };
+		return { ConstrainedParticles[ConstraintIndex], nullptr };
 	}
 
-	void UpdatePositionBasedState(const TPBDRigidParticles<T, d>& InParticles, const TArray<int32>& InIndices, const T Dt)
+	void UpdatePositionBasedState(const T Dt)
 	{
 	}
 
-	void ApplyHelper(TPBDRigidParticles<T, d>& InParticles, const T Dt, const TArray<int32>& InConstraintIndices) const
+	void ApplyHelper(const T Dt, const TArray<int32>& InConstraintIndices) const
 	{
 		for (int32 ConstraintIndex : InConstraintIndices)
 		{
-			const int32 Constraint = Constraints[ConstraintIndex];
-
-			if (InParticles.InvM(Constraint) == 0)
+			if (TPBDRigidParticleHandle<T, d>* PBDRigid = ConstrainedParticles[ConstraintIndex])
 			{
-				continue;
+				const TVector<T, d>& P1 = PBDRigid->P();
+				const TVector<T, d>& P2 = Targets[ConstraintIndex];
+				TVector<T, d> Difference = P1 - P2;
+				PBDRigid->P() -= Stiffness * Difference;
 			}
-			const TVector<T, d>& P1 = InParticles.P(Constraint);
-			const TVector<T, d>& P2 = Targets[ConstraintIndex];
-			TVector<T, d> Difference = P1 - P2;
-			InParticles.P(Constraint) -= Stiffness * Difference;
 		}
 	}
 
-	void Apply(TPBDRigidParticles<T, d>& InParticles, const T Dt, const TArray<int32>& InConstraintIndices) const
+	void Apply(const T Dt, const TArray<int32>& InConstraintIndices) const
 	{
-		ApplyHelper(InParticles, Dt, InConstraintIndices);
+		ApplyHelper(Dt, InConstraintIndices);
 	}
 
-	void ApplyPushOut(TPBDRigidParticles<T, d>& InParticles, const T Dt, const TArray<int32>& InConstraintIndices) const
+	void ApplyPushOut(const T Dt, const TArray<int32>& InConstraintIndices) const
 	{
 	}
 
-	int32 Add(const int32 Index, const TVector<T, d>& Position)
+	int32 Add(TGeometryParticleHandle<T,d>* Particle, const TVector<T, d>& Position)
 	{
 		int32 NewIndex = Targets.Num();
 		Targets.Add(Position);
-		Constraints.Add(Index);
+		ConstrainedParticles.Add(Particle);
 		return NewIndex;
 	}
 
@@ -74,14 +72,14 @@ class TPBDPositionConstraints : public TPBDConstraintContainer<T, d>
 		Targets[Index] = Position;
 	}
 
-	void RemoveConstraints(const TSet<uint32>& RemovedParticles)
+	void RemoveConstraints(const TSet<TGeometryParticleHandle<T,d>*>& RemovedParticles)
 	{
 		// @todo(ccaulfield): constraint management
 	}
 
 private:
 	TArray<TVector<T, d>> Targets;
-	TArray<int32> Constraints;
+	TArray<TPBDRigidParticleHandle<T,d>*> ConstrainedParticles;
 	const T Stiffness;
 };
 }
diff --git a/Engine/Source/Runtime/Experimental/Chaos/Public/Chaos/PBDRigidClusteredParticles.h b/Engine/Source/Runtime/Experimental/Chaos/Public/Chaos/PBDRigidClusteredParticles.h
new file mode 100644
--- /dev/null
+++ b/Engine/Source/Runtime/Experimental/Chaos/Public/Chaos/PBDRigidClusteredParticles.h
@@ -0,0 +1,174 @@
+// Copyright 1998-2019 Epic Games, Inc. All Rights Reserved.
+#pragma once
+
+#include "Chaos/ArrayCollectionArray.h"
+#include "Chaos/PBDRigidParticles.h"
+#include "Chaos/ImplicitObjectUnion.h"
+
+namespace Chaos
+{
+// ClusterId 
+//    Used within the clustering system to describe the clustering hierarchy. The ID will
+//  store the children IDs, and a Parent ID. When Id==NONE_INDEX the cluster is not
+//  controlled by another body. 
+struct ClusterId
+{
+	ClusterId() : Id(INDEX_NONE), NumChildren(0) {}
+	ClusterId(int32 NewId, int NumChildrenIn)
+		: Id(NewId)
+		, NumChildren(NumChildrenIn) {}
+	int32 Id;
+	int32 NumChildren;
+};
+
+/** When multiple children are active and can share one collision proxy. Only valid if all original children are still in the cluster*/
+template <typename T, int d>
+struct CHAOS_API TMultiChildProxyData
+{
+	TRigidTransform<T, d> RelativeToKeyChild;	//Use one child's transform to determine where to place the geometry. Needed for partial fracture where all children are still present and can therefore use proxy
+	uint32 KeyChild;
+};
+
+/** Used with TMultiChildProxyData. INDEX_NONE indicates no proxy data available */
+struct FMultiChildProxyId
+{
+	FMultiChildProxyId() : Id(INDEX_NONE) {}
+	int32 Id;
+};
+
+template <typename T>
+struct TConnectivityEdge
+{
+	TConnectivityEdge() {}
+	TConnectivityEdge(uint32 InSibling, T InStrain)
+		: Sibling(InSibling)
+		, Strain(InStrain) {}
+
+	TConnectivityEdge(const TConnectivityEdge& Other)
+		: Sibling(Other.Sibling)
+		, Strain(Other.Strain) {}
+
+	uint32 Sibling;
+	T Strain;
+};
+
+template<class T, int d>
+class TPBDRigidClusteredParticles : public TPBDRigidParticles<T, d>
+{
+  public:
+
+	TPBDRigidClusteredParticles()
+	: TPBDRigidParticles<T, d>()
+	{
+		InitHelper();
+	}
+	TPBDRigidClusteredParticles(const TPBDRigidParticles<T, d>& Other) = delete;
+	TPBDRigidClusteredParticles(TPBDRigidParticles<T, d>&& Other)
+	: TPBDRigidParticles<T, d>(MoveTemp(Other))
+	, MClusterIds(MoveTemp(Other.MClusterIds))
+	, MChildToParent(MoveTemp(Other.MChildToParent))
+	, MClusterGroupIndex(MoveTemp(Other.MClusterGroupIndex))
+	, MInternalCluster(MoveTemp(Other.MInternalCluster))
+	, MMultiChildProxyId(MoveTemp(Other.MMultiChildProxyId))
+	, MMultiChildProxyData(MoveTemp(Other.MMultiChildProxyData))
+	, MCollisionImpulses(MoveTemp(Other.MCollisionImpulses))
+	, MStrains(MoveTemp(Other.MStrains))
+	, MConnectivityEdges(MoveTemp(Other.MConnectivityEdges))
+
+	{
+		InitHelper();
+	}
+	~TPBDRigidClusteredParticles() {}
+
+	const auto& ClusterIds(int32 Idx) const { return MClusterIds[Idx]; }
+	auto& ClusterIds(int32 Idx) { return MClusterIds[Idx]; }
+
+	const auto& ChildToParent(int32 Idx) const { return MChildToParent[Idx]; }
+	auto& ChildToParent(int32 Idx) { return MChildToParent[Idx]; }
+
+	const auto& ClusterGroupIndex(int32 Idx) const { return MClusterGroupIndex[Idx]; }
+	auto& ClusterGroupIndex(int32 Idx) { return MClusterGroupIndex[Idx]; }
+
+	const auto& InternalCluster(int32 Idx) const { return MInternalCluster[Idx]; }
+	auto& InternalCluster(int32 Idx) { return MInternalCluster[Idx]; }
+
+	const auto& ChildrenSpatial(int32 Idx) const { return MChildrenSpatial[Idx]; }
+	auto& ChildrenSpatial(int32 Idx) { return MChildrenSpatial[Idx]; }
+
+	const auto& MultiChildProxyId(int32 Idx) const { return MMultiChildProxyId[Idx]; }
+	auto& MultiChildProxyId(int32 Idx) { return MMultiChildProxyId[Idx]; }
+
+	const auto& MultiChildProxyData(int32 Idx) const { return MMultiChildProxyData[Idx]; }
+	auto& MultiChildProxyData(int32 Idx) { return MMultiChildProxyData[Idx]; }
+
+	const auto& CollisionImpulses(int32 Idx) const { return MCollisionImpulses[Idx]; }
+	auto& CollisionImpulses(int32 Idx) { return MCollisionImpulses[Idx]; }
+
+	const auto& Strains(int32 Idx) const { return MStrains[Idx]; }
+	auto& Strains(int32 Idx) { return MStrains[Idx]; }
+
+	const auto& ConnectivityEdges(int32 Idx) const { return MConnectivityEdges[Idx]; }
+	auto& ConnectivityEdges(int32 Idx) { return MConnectivityEdges[Idx]; }
+
+	//TODO_CHAOSPARTICLE_HANDLE: is this really needed?
+	const auto& ConnectivityEdgesArray() const { return MConnectivityEdges; }
+	const auto& MultiChildProxyDataArray() const { return MMultiChildProxyData; }
+	const auto& MultiChildProxyIdArray() const { return MMultiChildProxyId; }
+
+	const auto& ClusterIdsArray() const { return MClusterIds; }
+	auto& ClusterIdsArray() { return MClusterIds; }
+
+	const auto& ChildToParentArray() const { return MChildToParent; }
+	auto& ChildToParentArray() { return MChildToParent; }
+
+	const auto& StrainsArray() const { return MStrains; }
+	auto& StrainsArray() { return MStrains; }
+
+	const auto& ClusterGroupIndexArray() const { return MClusterGroupIndex; }
+	auto& ClusterGroupIndexArray() { return MClusterGroupIndex; }
+
+	const auto& InternalClusterArray() const { return MInternalCluster; }
+	auto& InternalClusterArray() { return MInternalCluster; }
+
+	typedef TPBDRigidClusteredParticleHandle<T, d> THandleType;
+	const THandleType* Handle(int32 Index) const { return static_cast<const THandleType*>(TGeometryParticles<T,d>::Handle(Index)); }
+
+	//cannot be reference because double pointer would allow for badness, but still useful to have non const access to handle
+	THandleType* Handle(int32 Index) { return static_cast<THandleType*>(TGeometryParticles<T, d>::Handle(Index)); }
+
+	
+  private:
+
+	  void InitHelper()
+	  {
+		  this->MParticleType = EParticleType::Clustered;
+		  TArrayCollection::AddArray(&MClusterIds);
+		  TArrayCollection::AddArray(&MChildToParent);
+		  TArrayCollection::AddArray(&MClusterGroupIndex);
+		  TArrayCollection::AddArray(&MInternalCluster);
+		  TArrayCollection::AddArray(&MChildrenSpatial);
+		  TArrayCollection::AddArray(&MMultiChildProxyId);
+		  TArrayCollection::AddArray(&MMultiChildProxyData);
+		  TArrayCollection::AddArray(&MCollisionImpulses);
+		  TArrayCollection::AddArray(&MStrains);
+		  TArrayCollection::AddArray(&MConnectivityEdges);
+	  }
+
+	  TArrayCollectionArray<ClusterId> MClusterIds;
+	  TArrayCollectionArray<TRigidTransform<T, d>> MChildToParent;
+	  TArrayCollectionArray<int32> MClusterGroupIndex;
+	  TArrayCollectionArray<bool> MInternalCluster;
+	  TArrayCollectionArray<TUniquePtr<TImplicitObjectUnion<T, d>>> MChildrenSpatial;
+	  TArrayCollectionArray<FMultiChildProxyId> MMultiChildProxyId;
+	  TArrayCollectionArray<TUniquePtr<TMultiChildProxyData<T, d>>> MMultiChildProxyData;
+
+	  // Collision Impulses
+	  TArrayCollectionArray<T> MCollisionImpulses;
+
+	  // User set parameters
+	  TArrayCollectionArray<T> MStrains;
+
+	  TArrayCollectionArray<TArray<TConnectivityEdge<T>>> MConnectivityEdges;
+};
+
+}
diff --git a/Engine/Source/Runtime/Experimental/Chaos/Public/Chaos/PBDRigidClustering.h b/Engine/Source/Runtime/Experimental/Chaos/Public/Chaos/PBDRigidClustering.h
--- a/Engine/Source/Runtime/Experimental/Chaos/Public/Chaos/PBDRigidClustering.h
+++ b/Engine/Source/Runtime/Experimental/Chaos/Public/Chaos/PBDRigidClustering.h
@@ -1,60 +1,17 @@
 // Copyright 1998-2019 Epic Games, Inc. All Rights Reserved.
 #pragma once
 
-#include "Chaos/PBDRigidParticles.h"
+#include "Chaos/PBDRigidClusteredParticles.h"
 #include "Chaos/Transform.h"
 #include "Chaos/PBDCollisionTypes.h"
 #include "Framework/BufferedData.h"
 #include "BoundingVolume.h"
 #include "ImplicitObjectUnion.h"
 
-namespace Chaos
-{
-
-// ClusterId 
-//    Used within the clustering system to describe the clustering hierarchy. The ID will
-//  store the children IDs, and a Parent ID. When Id==NONE_INDEX the cluster is not
-//  controlled by another body. 
-struct ClusterId
-{
-	ClusterId() : Id(INDEX_NONE), NumChildren(0) {}
-	ClusterId(int32 NewId, int NumChildrenIn)
-	    : Id(NewId)
-		, NumChildren(NumChildrenIn) {}
-	int32 Id;
-	int32 NumChildren;
-};
+#define TODO_CONVERT_GEOMETRY_COLLECTION_PARTICLE_INDICES_TO_PARTICLE_POINTERS 0
 
-/** When multiple children are active and can share one collision proxy. Only valid if all original children are still in the cluster*/
-template <typename T, int d>
-struct CHAOS_API TMultiChildProxyData
-{
-	TRigidTransform<T, d> RelativeToKeyChild;	//Use one child's transform to determine where to place the geometry. Needed for partial fracture where all children are still present and can therefore use proxy
-	uint32 KeyChild;
-};
-
-/** Used with TMultiChildProxyData. INDEX_NONE indicates no proxy data available */
-struct FMultiChildProxyId
-{
-	FMultiChildProxyId() : Id(INDEX_NONE) {}
-	int32 Id;
-};
-
-template <typename T>
-struct TConnectivityEdge
+namespace Chaos
 {
-	TConnectivityEdge() {}
-	TConnectivityEdge(uint32 InSibling, T InStrain)
-	: Sibling(InSibling)
-	, Strain(InStrain) {}
-
-	TConnectivityEdge(const TConnectivityEdge& Other)
-	: Sibling(Other.Sibling)
-	, Strain(Other.Strain) {}
-
-	uint32 Sibling;
-	T Strain;
-};
 
 template <typename T>
 struct FClusterCreationParameters
@@ -117,10 +74,10 @@ class CHAOS_API TPBDRigidClustering
 	typedef typename FPBDCollisionConstraint::FRigidBodyContactConstraint FRigidBodyContactConstraint;
 
 public:
-	typedef TMap<uint32, TUniquePtr<TArray<uint32>> > FClusterMap;
+	typedef TMap<TGeometryParticleHandle<T,d>*, TArray<uint32> > FClusterMap;
 
-	TPBDRigidClustering(FPBDRigidEvolution& InEvolution, TPBDRigidParticles<T, d>& InParticles);
-	~TPBDRigidClustering() {}
+	TPBDRigidClustering(FPBDRigidEvolution& InEvolution, TPBDRigidClusteredParticles<T, d>& InParticles);
+	~TPBDRigidClustering();
 
 	//
 	// Initialization
@@ -133,7 +90,7 @@ public:
 	*    ProxyGeometry : Collision default for the cluster, automatically generated otherwise.
 	*    ForceMassOrientation : Inertial alignment into mass space.
 	*/
-	int32 CreateClusterParticle(int32 ClusterGroupIndex, const TArray<uint32>& Children,
+	int32 CreateClusterParticle(int32 ClusterGroupIndex, TArray<uint32>&& Children,
 		TSerializablePtr<TImplicitObject<T, d>> ProxyGeometry = TSerializablePtr<TImplicitObject<T, 3>>(),
 		const TRigidTransform<T, d>* ForceMassOrientation = nullptr, const FClusterCreationParameters<T>& Parameters = FClusterCreationParameters<T>());
 
@@ -141,7 +98,7 @@ public:
 	*  CreateClusterParticleFromClusterChildren
 	*    Children : Rigid body ID to include in the cluster.
 	*/
-	int32 CreateClusterParticleFromClusterChildren(const TArray<uint32>& Children, const int32 ParentIndex,
+	int32 CreateClusterParticleFromClusterChildren(TArray<uint32>&& Children, TPBDRigidClusteredParticleHandle<T,d>* Parent,
 		const TRigidTransform<T, d>& ClusterWorldTM, const FClusterCreationParameters<T>& Parameters = FClusterCreationParameters<T>());
 
 	/*
@@ -159,7 +116,7 @@ public:
 	*  DeactivateClusterParticle
 	*    Release all the particles within the cluster particle
 	*/
-	TSet<uint32> DeactivateClusterParticle(const uint32 ClusterIndex);
+	TSet<uint32> DeactivateClusterParticle(TPBDRigidClusteredParticleHandle<T,d>* ClusteredParticle);
 
 
 	/*
@@ -168,7 +125,7 @@ public:
 	*    have a MStrain value less than its entry in the StrainArray will be
 	*    released from the cluster.
 	*/
-	TSet<uint32> ReleaseClusterParticles(const uint32 ClusterIndex, const TArrayView<T>& StrainArray);
+	TSet<uint32> ReleaseClusterParticles(TPBDRigidClusteredParticleHandle<T, d>* ClusteredParticle, const TArrayView<T>& StrainArray);
 
 	/*
 	*  ReleaseClusterParticles
@@ -197,7 +154,7 @@ public:
 	*    encoded strain. The remainder strains are propagated back down to
 	*    the children clusters.
 	*/
-	TMap<uint32, TSet<uint32>> BreakingModel(TArrayView<T> ExternalStrain);
+	TMap<TPBDRigidClusteredParticleHandle<T, d>*, TSet<uint32>> BreakingModel(TArrayView<T> ExternalStrain);
 
 	/**
 	*  PromoteStrains
@@ -240,21 +197,21 @@ public:
 	*    active id, see the GetActiveClusterIndex to find the active cluster.
 	*    INDEX_NONE represents a non-clustered body.
 	*/
-	TArrayCollectionArray<ClusterId>&             GetClusterIdsArray() { return MClusterIds; }
-	const TArrayCollectionArray<ClusterId>&             GetClusterIdsArray() const { return MClusterIds; }
+	TArrayCollectionArray<ClusterId>&             GetClusterIdsArray() { return MParticles.ClusterIdsArray(); }
+	const TArrayCollectionArray<ClusterId>&             GetClusterIdsArray() const { return MParticles.ClusterIdsArray(); }
 
 	/*
 	*  GetInternalClusterArray
 	*    The internal cluster array indicates if this cluster was generated internally
 	*    and would no be owned by an external source.
 	*/
-	const TArrayCollectionArray<bool>&                  GetInternalClusterArray() const { return MInternalCluster; }
+	const TArrayCollectionArray<bool>&                  GetInternalClusterArray() const { return MParticles.InternalClusterArray(); }
 
 	/*
 	*  GetChildToParentMap
 	*    This map stores the relative transform from a child to its cluster parent.
 	*/
-	const TArrayCollectionArray<TRigidTransform<T, d>>& GetChildToParentMap() const { return MChildToParent; }
+	const TArrayCollectionArray<TRigidTransform<T, d>>& GetChildToParentMap() const { return MParticles.ChildToParentArray(); }
 
 	/*
 	*  GetStrainArray
@@ -262,7 +219,7 @@ public:
 	*    body in the simulation. This attribute is initialized during the creation of
 	*    the cluster body, can be updated during the evaluation of the simulation.
 	*/
-	TArrayCollectionArray<T>&                           GetStrainArray() { return MStrains; }
+	TArrayCollectionArray<T>&                           GetStrainArray() { return MParticles.StrainsArray(); }
 
 	/**
 	*  GetParentToChildren
@@ -281,13 +238,13 @@ public:
 	*    is called during AdvanceClustering) the positive bodies are joined with a negative pre-existing
 	*    body, then set negative. Zero entries are ignored within the union.
 	*/
-	TArrayCollectionArray<int32>&                       GetClusterGroupIndexArray() { return MClusterGroupIndex; }
+	TArrayCollectionArray<int32>&                       GetClusterGroupIndexArray() { return MParticles.ClusterGroupIndexArray(); }
 
 	/** Indicates if the child geometry is approximated by a single proxy */
-	const TArrayCollectionArray<FMultiChildProxyId>& GetMultiChildProxyIdArray() const { return MMultiChildProxyId; }
+	const TArrayCollectionArray<FMultiChildProxyId>& GetMultiChildProxyIdArray() const { return MParticles.MultiChildProxyIdArray(); }
 
 	/** If multi child proxy is used, this is the data needed */
-	const TArrayCollectionArray<TUniquePtr<TMultiChildProxyData<T, d>>>& GetMultiChildProxyDataArray() const { return MMultiChildProxyData; }
+	const TArrayCollectionArray<TUniquePtr<TMultiChildProxyData<T, d>>>& GetMultiChildProxyDataArray() const { return MParticles.MultiChildProxyDataArray(); }
 
 	/*
 	*  GetClusterGroupIndexArray
@@ -310,7 +267,9 @@ public:
 	* GetConnectivityEdges
 	*    Provides a list of each rigid body's current siblings and associated strain within the cluster.
 	*/
-	const TArrayCollectionArray<TArray<TConnectivityEdge<T>>>& GetConnectivityEdges() const { return MConnectivityEdges; }
+//#if CHAOS_PARTICLEHANDLE_TODO: is this really needed?
+	const TArrayCollectionArray<TArray<TConnectivityEdge<T>>>& GetConnectivityEdges() const { return MParticles.ConnectivityEdgesArray(); }
+//#endif
 
 	/**
 	* GenerateConnectionGraph
@@ -331,7 +290,7 @@ public:
 		}
 		for (uint32 i = StartIndex; i < MParticles.Size(); ++i)
 		{
-			if (MClusterIds[i].Id == INDEX_NONE && !MParticles.Disabled(i))
+			if (MParticles.ClusterIds(i).Id == INDEX_NONE && !MParticles.Disabled(i))
 			{
 				TopLevelClusterParents.Add(i);
 			}
@@ -346,7 +305,7 @@ public:
 	void UpdateGeometry(const TArray<uint32>& Children, const uint32 NewIndex, TSerializablePtr<TImplicitObject<T, d>> ProxyGeometry, const FClusterCreationParameters<T>& Parameters);
 	void ComputeStrainFromCollision(const FPBDCollisionConstraint& CollisionRule);
 	void ResetCollisionImpulseArray();
-	void DisableCluster(uint32 ClusterIndex);
+	void DisableCluster(TPBDRigidClusteredParticleHandle<T, d>* ClusteredParticle);
 	void DisableParticleWithBreakEvent(uint32 ClusterIndex);
 
 	/*
@@ -362,7 +321,7 @@ public:
 private:
 
 	FPBDRigidEvolution& MEvolution;
-	TPBDRigidParticles<T, d>& MParticles;
+	TPBDRigidClusteredParticles<T, d>& MParticles;
 	TSet<int32> TopLevelClusterParents;
 	TSet<int32> MActiveRemovalIndices;
 
@@ -370,34 +329,22 @@ private:
 	mutable FRWLock ResourceLock;
 	TClusterBuffer<T, d> BufferResource;
 	FClusterMap MChildren;
-	TArrayCollectionArray<ClusterId> MClusterIds;
 	TMap<int32, int32> PendingClusterCounter;
 
-	TArrayCollectionArray<TRigidTransform<T, d>> MChildToParent;
-	TArrayCollectionArray<int32> MClusterGroupIndex;
-	TArrayCollectionArray<bool> MInternalCluster;
-	TArrayCollectionArray<TUniquePtr<TImplicitObjectUnion<T,d>>> MChildrenSpatial;
-	TArrayCollectionArray<FMultiChildProxyId> MMultiChildProxyId;
-	TArrayCollectionArray<TUniquePtr<TMultiChildProxyData<T, d>>> MMultiChildProxyData;
 
 	// Collision Impulses
 	bool MCollisionImpulseArrayDirty;
-	TArrayCollectionArray<T> MCollisionImpulses;
-
-	// User set parameters
-	TArrayCollectionArray<T> MStrains;
-
+	
 	// Breaking data
 	bool DoGenerateBreakingData;
 	TArray<TBreakingData<float, 3>> MAllClusterBreakings;
 
 	float MClusterConnectionFactor;
 	typename FClusterCreationParameters<T>::EConnectionMethod MClusterUnionConnectionType;
-	TArrayCollectionArray<TArray<TConnectivityEdge<T>>> MConnectivityEdges;
 };
 
 template <typename T, int d>
-void UpdateClusterMassProperties(TPBDRigidParticles<T, d>& Particles, const TArray<uint32>& Children,
+void UpdateClusterMassProperties(TPBDRigidClusteredParticles<T, d>& Particles, const TArray<uint32>& Children,
 	const uint32 NewIndex, const TRigidTransform<T, d>* ForceMassOrientation = nullptr,
 	const TArrayCollectionArray<TUniquePtr<TMultiChildProxyData<T, d>>>* MMultiChildProxyData = nullptr,
 	const TArrayCollectionArray<FMultiChildProxyId>* MMultiChildProxyId = nullptr);
diff --git a/Engine/Source/Runtime/Experimental/Chaos/Public/Chaos/PBDRigidDynamicSpringConstraints.h b/Engine/Source/Runtime/Experimental/Chaos/Public/Chaos/PBDRigidDynamicSpringConstraints.h
--- a/Engine/Source/Runtime/Experimental/Chaos/Public/Chaos/PBDRigidDynamicSpringConstraints.h
+++ b/Engine/Source/Runtime/Experimental/Chaos/Public/Chaos/PBDRigidDynamicSpringConstraints.h
@@ -9,9 +9,9 @@
 namespace Chaos
 {
 template<class T, int d>
-class TPBDRigidDynamicSpringConstraints : public TPBDRigidDynamicSpringConstraintsBase<T, d>, public TPBDConstraintContainer<T, d>
+class TPBDRigidDynamicSpringConstraints : public TPBDRigidDynamicSpringConstraintsBase2<T, d>, public TPBDConstraintContainer<T, d>
 {
-	typedef TPBDRigidDynamicSpringConstraintsBase<T, d> Base;
+	typedef TPBDRigidDynamicSpringConstraintsBase2<T, d> Base;
 
 	using Base::Constraints;
 	using Base::Distances;
@@ -21,29 +21,29 @@ public:
 	using Base::UpdatePositionBasedState;
 
 	TPBDRigidDynamicSpringConstraints(const T InStiffness = (T)1)
-	    : TPBDRigidDynamicSpringConstraintsBase<T, d>(InStiffness) 
+	    : TPBDRigidDynamicSpringConstraintsBase2<T, d>(InStiffness) 
 	{}
 
-	TPBDRigidDynamicSpringConstraints(TArray<TVector<int32, 2>>&& InConstraints, const T InCreationThreshold = (T)1, const int32 InMaxSprings = 1, const T InStiffness = (T)1)
-	    : TPBDRigidDynamicSpringConstraintsBase<T, d>(MoveTemp(InConstraints), InCreationThreshold, InMaxSprings, InStiffness)
+	TPBDRigidDynamicSpringConstraints(TArray<TVector<TGeometryParticleHandle<T, d>*, 2>>&& InConstraints, const T InCreationThreshold = (T)1, const int32 InMaxSprings = 1, const T InStiffness = (T)1)
+	    : TPBDRigidDynamicSpringConstraintsBase2<T, d>(MoveTemp(InConstraints), InCreationThreshold, InMaxSprings, InStiffness)
 	{}
 
 	virtual ~TPBDRigidDynamicSpringConstraints() {}
 
-	CHAOS_API void UpdatePositionBasedState(const TPBDRigidParticles<T, d>& InParticles, const TArray<int32>& InIndices, const T Dt)
+	CHAOS_API void UpdatePositionBasedState(const T Dt)
 	{
-		Base::UpdatePositionBasedState(InParticles);
+		Base::UpdatePositionBasedState();
 	}
 
 	// @todo(mlentine): Why is this needed?
-	CHAOS_API void ApplyHelper(TPBDRigidParticles<T, d>& InParticles, const T Dt, const TArray<int32>& InConstraintIndices) const;
+	CHAOS_API void ApplyHelper(const T Dt, const TArray<int32>& InConstraintIndices) const;
 
-	CHAOS_API void Apply(TPBDRigidParticles<T, d>& InParticles, const T Dt, const TArray<int32>& InConstraintIndices) const
+	CHAOS_API void Apply(const T Dt, const TArray<int32>& InConstraintIndices) const
 	{
-		ApplyHelper(InParticles, Dt, InConstraintIndices);
+		ApplyHelper(Dt, InConstraintIndices);
 	}
 
-	CHAOS_API void ApplyPushOut(TPBDRigidParticles<T, d>& InParticles, const T Dt, const TArray<int32>& InConstraintIndices)
+	CHAOS_API void ApplyPushOut(const T Dt, const TArray<int32>& InConstraintIndices)
 	{
 
 	}
diff --git a/Engine/Source/Runtime/Experimental/Chaos/Public/Chaos/PBDRigidDynamicSpringConstraintsBase.h b/Engine/Source/Runtime/Experimental/Chaos/Public/Chaos/PBDRigidDynamicSpringConstraintsBase.h
--- a/Engine/Source/Runtime/Experimental/Chaos/Public/Chaos/PBDRigidDynamicSpringConstraintsBase.h
+++ b/Engine/Source/Runtime/Experimental/Chaos/Public/Chaos/PBDRigidDynamicSpringConstraintsBase.h
@@ -3,66 +3,65 @@
 #pragma once
 
 #include "Chaos/Array.h"
-#include "Chaos/Map.h"
 #include "Chaos/PBDParticles.h"
-#include "Chaos/PerParticleRule.h"
-
-#include <functional>
+#include "Chaos/ParticleHandle.h"
 
 namespace Chaos
 {
 	template<class T, int d>
-	class TPBDRigidDynamicSpringConstraintsBase
+	class TPBDRigidDynamicSpringConstraintsBase2
 	{
-	  public:
-		TPBDRigidDynamicSpringConstraintsBase(const T InStiffness = (T)1)
+	public:
+		TPBDRigidDynamicSpringConstraintsBase2(const T InStiffness = (T)1)
 		    : Stiffness(InStiffness)
 		{
 		}
-		TPBDRigidDynamicSpringConstraintsBase(TArray<TVector<int32, 2>>&& InConstraints, const T InCreationThreshold = (T)1, const int32 InMaxSprings = 1, const T InStiffness = (T)1)
+		TPBDRigidDynamicSpringConstraintsBase2(TArray<TVector<TGeometryParticleHandle<T, d>*, 2>>&& InConstraints, const T InCreationThreshold = (T)1, const int32 InMaxSprings = 1, const T InStiffness = (T)1)
 		    : Constraints(MoveTemp(InConstraints)), CreationThreshold(InCreationThreshold), MaxSprings(InMaxSprings), Stiffness(InStiffness)
 		{
 			Distances.SetNum(Constraints.Num());
 			SpringDistances.SetNum(Constraints.Num());
 		}
 	
-		virtual ~TPBDRigidDynamicSpringConstraintsBase() {}
+		virtual ~TPBDRigidDynamicSpringConstraintsBase2() {}
 
 		int32 NumConstraints() const
 		{
 			return Constraints.Num();
 		}
 
-		TVector<int32, 2> ConstraintParticleIndices(int32 ConstraintIndex) const
+		const TVector<TGeometryParticleHandle<T, d>*, 2>& ConstraintParticles(int32 ConstraintIndex) const
 		{
 			return Constraints[ConstraintIndex];
 		}
 
-		CHAOS_API void UpdatePositionBasedState(const TPBDRigidParticles<T, d>& InParticles);
-		TVector<T, d> GetDelta(const TPBDRigidParticles<T, d>& InParticles, const TVector<T, d>& WorldSpaceX1, const TVector<T, d>& WorldSpaceX2, const int32 i, const int32 j) const;
+		CHAOS_API void UpdatePositionBasedState();
+
+		TVector<T, d> GetDelta(const TVector<T, d>& WorldSpaceX1, const TVector<T, d>& WorldSpaceX2, const int32 ConstraintIndex, const int32 SpringIndex) const;
 
-		void Add(const int32 Index1, const int32 Index2)
+		void Add(TGeometryParticleHandle<T, d>* Particle0, TGeometryParticleHandle<T, d>* Particle1)
 		{
-			Constraints.Add(TVector<int32, 2>(Index1, Index2));
-			Distances.Add(TArray<TVector<TVector<T, 3>, 2>>());
-			SpringDistances.Add(TArray<T>());
+			Constraints.Add({ Particle0, Particle1 });
+			Distances.Add({});
+			SpringDistances.Add({});
 		}
+
 		void SetDistance(const T Threshold)
 		{
 			CreationThreshold = Threshold;
 		}
 
-		void RemoveConstraints(const TSet<uint32>& RemovedParticles)
+		void RemoveConstraints(const TSet<TGeometryParticleHandle<T, d>*>& RemovedParticles)
 		{
 			// @todo(ccaulfield): constraint management
 		}
 
-	  protected:
-		TArray<TVector<int32, 2>> Constraints;
+	protected:
+		TArray<TVector<TGeometryParticleHandle<T, d>*, 2>> Constraints;
 		TArray<TArray<TVector<TVector<T, 3>, 2>>> Distances;
 		TArray<TArray<T>> SpringDistances;
 	
-	  private:
+	private:
 		T CreationThreshold;
 		int32 MaxSprings;
 		T Stiffness;
diff --git a/Engine/Source/Runtime/Experimental/Chaos/Public/Chaos/PBDRigidParticles.h b/Engine/Source/Runtime/Experimental/Chaos/Public/Chaos/PBDRigidParticles.h
--- a/Engine/Source/Runtime/Experimental/Chaos/Public/Chaos/PBDRigidParticles.h
+++ b/Engine/Source/Runtime/Experimental/Chaos/Public/Chaos/PBDRigidParticles.h
@@ -10,9 +10,6 @@ namespace Chaos
 template<class T, int d>
 class TPBDRigidsEvolution;
 
-template <typename T, int d>
-class TPBDRigidParticleHandle;
-
 template<class T, int d>
 class TPBDRigidParticles : public TRigidParticles<T, d>
 {
@@ -24,6 +21,7 @@ class TPBDRigidParticles : public TRigidParticles<T, d>
 	TPBDRigidParticles()
 	    : TRigidParticles<T, d>()
 	{
+		this->MParticleType = EParticleType::Dynamic;
 		TArrayCollection::AddArray(&MP);
 		TArrayCollection::AddArray(&MQ);
 		TArrayCollection::AddArray(&MPreV);
@@ -31,14 +29,18 @@ class TPBDRigidParticles : public TRigidParticles<T, d>
 	}
 	TPBDRigidParticles(const TPBDRigidParticles<T, d>& Other) = delete;
 	TPBDRigidParticles(TPBDRigidParticles<T, d>&& Other)
-	    : TRigidParticles<T, d>(MoveTemp(Other)), MP(MoveTemp(Other.MP)), MQ(MoveTemp(Other.MQ))
+	    : TRigidParticles<T, d>(MoveTemp(Other))
+		, MP(MoveTemp(Other.MP))
+		, MQ(MoveTemp(Other.MQ))
+		, MPreV(MoveTemp(Other.MPreV))
+		, MPreW(MoveTemp(Other.MPreW))
 	{
+		this->MParticleType = EParticleType::Dynamic;
 		TArrayCollection::AddArray(&MP);
 		TArrayCollection::AddArray(&MQ);
 		TArrayCollection::AddArray(&MPreV);
 		TArrayCollection::AddArray(&MPreW);
 	}
-	~TPBDRigidParticles() {}
 
 	const TVector<T, d>& P(const int32 index) const { return MP[index]; }
 	TVector<T, d>& P(const int32 index) { return MP[index]; }
diff --git a/Engine/Source/Runtime/Experimental/Chaos/Public/Chaos/PBDRigidSpringConstraints.h b/Engine/Source/Runtime/Experimental/Chaos/Public/Chaos/PBDRigidSpringConstraints.h
--- a/Engine/Source/Runtime/Experimental/Chaos/Public/Chaos/PBDRigidSpringConstraints.h
+++ b/Engine/Source/Runtime/Experimental/Chaos/Public/Chaos/PBDRigidSpringConstraints.h
@@ -1,79 +1,84 @@
 // Copyright 1998-2019 Epic Games, Inc. All Rights Reserved.
 #pragma once
 
-#include "Chaos/Array.h"
-#include "Chaos/Map.h"
 #include "Chaos/PBDRigidSpringConstraintsBase.h"
 #include "Chaos/PBDConstraintContainer.h"
 
 namespace Chaos
 {
 template<class T, int d>
-class TPBDRigidSpringConstraints : public TPBDRigidSpringConstraintsBase<T, d>, public TPBDConstraintContainer<T, d>
+class TPBDRigidSpringConstraints2 : public TPBDRigidSpringConstraintsBase<T, d>, public TPBDConstraintContainer<T, d>
 {
 	typedef TPBDRigidSpringConstraintsBase<T, d> Base;
 
 	using Base::Constraints;
 	using Base::Distances;
 
-  public:
-	TPBDRigidSpringConstraints(const T InStiffness = (T)1)
+public:
+	using FConstrainedParticlePair = TVector<TGeometryParticleHandle<T, d>*, 2>;
+
+
+	TPBDRigidSpringConstraints2(const T InStiffness = (T)1)
 	    : TPBDRigidSpringConstraintsBase<T, d>(InStiffness) 
 	{}
 
-	TPBDRigidSpringConstraints(const TRigidParticles<T, d>& InParticles, const TArray<TVector<T, 3>>& Locations0, const TArray<TVector<T, 3>>& Locations1, TArray<TVector<int32, 2>>&& InConstraints, const T InStiffness = (T)1)
-	    : TPBDRigidSpringConstraintsBase<T, d>(InParticles, Locations0, Locations1, MoveTemp(InConstraints), InStiffness)
+	TPBDRigidSpringConstraints2(const TArray<TVector<T, 3>>& Locations0, const TArray<TVector<T, 3>>& Locations1, TArray<FConstrainedParticlePair>&& InConstraints, const T InStiffness = (T)1)
+	    : TPBDRigidSpringConstraintsBase<T, d>(Locations0, Locations1, MoveTemp(InConstraints), InStiffness)
 	{}
 
-	virtual ~TPBDRigidSpringConstraints() {}
+	virtual ~TPBDRigidSpringConstraints2() {}
 
-	void UpdatePositionBasedState(const TPBDRigidParticles<T, d>& InParticles, const TArray<int32>& InIndices, const T Dt)
+	void UpdatePositionBasedState(const T Dt)
 	{
 	}
 
-	void ApplyHelper(TPBDRigidParticles<T, d>& InParticles, const T Dt, const TArray<int32>& InConstraintIndices) const
+	void ApplyHelper(const T Dt, const TArray<int32>& InConstraintIndices) const
 	{
 		for (int32 ConstraintIndex : InConstraintIndices)
 		{
-			const TVector<int32, 2>& Constraint = Constraints[ConstraintIndex];
+			const FConstrainedParticlePair& Constraint = Constraints[ConstraintIndex];
 
-			int32 ConstraintInnerIndex1 = Constraint[0];
-			int32 ConstraintInnerIndex2 = Constraint[1];
+			TPBDRigidParticleHandle<T, d>* PBDRigid0 = Constraint[0]->AsDynamic();
+			TPBDRigidParticleHandle<T, d>* PBDRigid1 = Constraint[1]->AsDynamic();
+			check((PBDRigid0 && PBDRigid1 && PBDRigid0->Island() == PBDRigid1->Island()) || (!PBDRigid0 && PBDRigid1) || (PBDRigid0 && PBDRigid1));
 
-			check(InParticles.Island(ConstraintInnerIndex1) == InParticles.Island(ConstraintInnerIndex2) || InParticles.Island(ConstraintInnerIndex1) == -1 || InParticles.Island(ConstraintInnerIndex2) == -1);
+			// @todo(ccaulfield): see if we can eliminate the need for all these ifs
+			const TRotation<T, d> & Q0 = PBDRigid0 ? PBDRigid0->Q() : Constraint[0]->R();
+			const TRotation<T, d> & Q1 = PBDRigid1 ? PBDRigid1->Q() : Constraint[1]->R();
+			const TVector<T, d> & P0 = PBDRigid0 ? PBDRigid0->P() : Constraint[0]->X();
+			const TVector<T, d> & P1 = PBDRigid1 ? PBDRigid1->P() : Constraint[1]->X();
 
-			const TVector<T, d> WorldSpaceX1 = InParticles.Q(ConstraintInnerIndex1).RotateVector(Distances[ConstraintIndex][0]) + InParticles.P(ConstraintInnerIndex1);
-			const TVector<T, d> WorldSpaceX2 = InParticles.Q(ConstraintInnerIndex2).RotateVector(Distances[ConstraintIndex][1]) + InParticles.P(ConstraintInnerIndex2);
-			const PMatrix<T, d, d> WorldSpaceInvI1 = (InParticles.Q(ConstraintInnerIndex1) * FMatrix::Identity) * InParticles.InvI(ConstraintInnerIndex1) * (InParticles.Q(ConstraintInnerIndex1) * FMatrix::Identity).GetTransposed();
-			const PMatrix<T, d, d> WorldSpaceInvI2 = (InParticles.Q(ConstraintInnerIndex2) * FMatrix::Identity) * InParticles.InvI(ConstraintInnerIndex2) * (InParticles.Q(ConstraintInnerIndex2) * FMatrix::Identity).GetTransposed();
-			const TVector<T, d> Delta = Base::GetDelta(InParticles, WorldSpaceX1, WorldSpaceX2, ConstraintIndex);
+			const TVector<T, d> WorldSpaceX1 = Q0.RotateVector(Distances[ConstraintIndex][0]) + P0;
+			const TVector<T, d> WorldSpaceX2 = Q1.RotateVector(Distances[ConstraintIndex][1]) + P1;
+			const PMatrix<T, d, d> WorldSpaceInvI1 = PBDRigid0 ? (PBDRigid0->Q() * FMatrix::Identity) * PBDRigid0->InvI() * (PBDRigid0->Q() * FMatrix::Identity).GetTransposed() : PMatrix<T, d, d>(0);
+			const PMatrix<T, d, d> WorldSpaceInvI2 = PBDRigid1 ? (PBDRigid1->Q() * FMatrix::Identity) * PBDRigid1->InvI() * (PBDRigid1->Q() * FMatrix::Identity).GetTransposed() : PMatrix<T, d, d>(0);
+			const TVector<T, d> Delta = Base::GetDelta(WorldSpaceX1, WorldSpaceX2, ConstraintIndex);
 
-			if (InParticles.InvM(ConstraintInnerIndex1) > 0)
+			if (PBDRigid0)
 			{
-				const TVector<T, d> Radius = WorldSpaceX1 - InParticles.P(ConstraintInnerIndex1);
-				InParticles.P(ConstraintInnerIndex1) += InParticles.InvM(ConstraintInnerIndex1) * Delta;
-				InParticles.Q(ConstraintInnerIndex1) += TRotation<T, d>(WorldSpaceInvI1 * TVector<T, d>::CrossProduct(Radius, Delta), 0.f) * InParticles.Q(ConstraintInnerIndex1) * T(0.5);
-				InParticles.Q(ConstraintInnerIndex1).Normalize();
+				const TVector<T, d> Radius = WorldSpaceX1 - PBDRigid0->P();
+				PBDRigid0->P() += PBDRigid0->InvM() * Delta;
+				PBDRigid0->Q() += TRotation<T, d>(WorldSpaceInvI1 * TVector<T, d>::CrossProduct(Radius, Delta), 0.f) * PBDRigid0->Q() * T(0.5);
+				PBDRigid0->Q().Normalize();
 			}
 
-			if (InParticles.InvM(ConstraintInnerIndex2) > 0)
+			if (PBDRigid1)
 			{
-				const TVector<T, d> Radius = WorldSpaceX2 - InParticles.P(ConstraintInnerIndex2);
-				InParticles.P(ConstraintInnerIndex2) -= InParticles.InvM(ConstraintInnerIndex2) * Delta;
-				InParticles.Q(ConstraintInnerIndex2) += TRotation<T, d>(WorldSpaceInvI2 * TVector<T, d>::CrossProduct(Radius, -Delta), 0.f) * InParticles.Q(ConstraintInnerIndex2) * T(0.5);
-				InParticles.Q(ConstraintInnerIndex2).Normalize();
+				const TVector<T, d> Radius = WorldSpaceX2 - PBDRigid1->P();
+				PBDRigid1->P() -= PBDRigid1->InvM() * Delta;
+				PBDRigid1->Q() += TRotation<T, d>(WorldSpaceInvI2 * TVector<T, d>::CrossProduct(Radius, -Delta), 0.f) * PBDRigid1->Q() * T(0.5);
+				PBDRigid1->Q().Normalize();
 			}
 		}
 	}
 
-	void Apply(TPBDRigidParticles<T, d>& InParticles, const T Dt, const TArray<int32>& InConstraintIndices) const
+	void Apply(const T Dt, const TArray<int32>& InConstraintIndices) const
 	{
-		ApplyHelper(InParticles, Dt, InConstraintIndices);
+		ApplyHelper(Dt, InConstraintIndices);
 	}
 
-	CHAOS_API void ApplyPushOut(TPBDRigidParticles<T, d>& InParticles, const T Dt, const TArray<int32>& InConstraintIndices)
+	CHAOS_API void ApplyPushOut(const T Dt, const TArray<int32>& InConstraintIndices)
 	{
-
 	}
 };
 }
