Copying //UE4/Dev-Physics to Dev-Main (//UE4/Dev-Main) [at] 7703071

diff --git a/Engine/Source/Runtime/Experimental/Chaos/Public/Chaos/PBDRigidSpringConstraintsBase.h b/Engine/Source/Runtime/Experimental/Chaos/Public/Chaos/PBDRigidSpringConstraintsBase.h
--- a/Engine/Source/Runtime/Experimental/Chaos/Public/Chaos/PBDRigidSpringConstraintsBase.h
+++ b/Engine/Source/Runtime/Experimental/Chaos/Public/Chaos/PBDRigidSpringConstraintsBase.h
@@ -3,26 +3,25 @@
 #pragma once
 
 #include "Chaos/Array.h"
-#include "Chaos/Map.h"
-#include "Chaos/PBDParticles.h"
-#include "Chaos/PerParticleRule.h"
-
-#include <functional>
+#include "Chaos/ParticleHandle.h"
 
 namespace Chaos
 {
 	template<class T, int d>
 	class TPBDRigidSpringConstraintsBase
 	{
-	  public:
+	public:
+		using FConstrainedParticlePair = TVector<TGeometryParticleHandle<T, d>*, 2>;
+
 		TPBDRigidSpringConstraintsBase(const T InStiffness = (T)1)
 		    : Stiffness(InStiffness)
 		{
 		}
-		TPBDRigidSpringConstraintsBase(const TRigidParticles<T, d>& InParticles, const TArray<TVector<T, d>>& Locations0,  const TArray<TVector<T, d>>& Locations1, TArray<TVector<int32, 2>>&& InConstraints, const T InStiffness = (T)1)
+
+		TPBDRigidSpringConstraintsBase(const TArray<TVector<T, d>>& Locations0,  const TArray<TVector<T, d>>& Locations1, TArray<FConstrainedParticlePair>&& InConstraints, const T InStiffness = (T)1)
 		    : Constraints(MoveTemp(InConstraints)), Stiffness(InStiffness)
 		{
-			UpdateDistances(InParticles, Locations0, Locations1);
+			UpdateDistances(Locations0, Locations1);
 		}
 	
 		virtual ~TPBDRigidSpringConstraintsBase() {}
@@ -32,25 +31,25 @@ namespace Chaos
 			return Constraints.Num();
 		}
 
-		TVector<int32, 2> ConstraintParticleIndices(int32 ConstraintIndex) const
+		const FConstrainedParticlePair& ConstraintParticles(int32 ConstraintIndex) const
 		{
 			return Constraints[ConstraintIndex];
 		}
 
-		void UpdateDistances(const TRigidParticles<T, d>& InParticles, const TArray<TVector<T, d>>& Locations0, const TArray<TVector<T, d>>& Locations1);
-		TVector<T, d> GetDelta(const TPBDRigidParticles<T, d>& InParticles, const TVector<T, d>& WorldSpaceX1, const TVector<T, d>& WorldSpaceX2, const int32 i) const;
+		void UpdateDistances(const TArray<TVector<T, d>>& Locations0, const TArray<TVector<T, d>>& Locations1);
+		TVector<T, d> GetDelta(const TVector<T, d>& WorldSpaceX1, const TVector<T, d>& WorldSpaceX2, const int32 ConstraintIndex) const;
 	
-		void RemoveConstraints(const TSet<uint32>& RemovedParticles)
+		void RemoveConstraints(const TSet<TGeometryParticleHandle<T, d>*>& RemovedParticles)
 		{
 			// @todo(ccaulfield): constraint management
 		}
 
-	  protected:
-		TArray<TVector<int32, 2>> Constraints;
+	protected:
+		TArray<FConstrainedParticlePair> Constraints;
 		TArray<TVector<TVector<T, 3>, 2>> Distances;
 		TArray<T> SpringDistances;
 	
-	  private:
+	private:
 		T Stiffness;
 	};
 }
diff --git a/Engine/Source/Runtime/Experimental/Chaos/Public/Chaos/PBDRigidsEvolution.h b/Engine/Source/Runtime/Experimental/Chaos/Public/Chaos/PBDRigidsEvolution.h
--- a/Engine/Source/Runtime/Experimental/Chaos/Public/Chaos/PBDRigidsEvolution.h
+++ b/Engine/Source/Runtime/Experimental/Chaos/Public/Chaos/PBDRigidsEvolution.h
@@ -7,9 +7,11 @@
 #include "Chaos/PBDRigidClustering.h"
 #include "Chaos/PBDRigidParticles.h"
 #include "Chaos/PBDConstraintRule.h"
+#include "Chaos/ParticleHandle.h"
 #include "Chaos/Transform.h"
 #include "Chaos/Framework/DebugSubstep.h"
 #include "HAL/Event.h"
+#include "Chaos/PBDRigidsSOAs.h"
 
 namespace Chaos
 {
@@ -45,149 +47,94 @@ struct CHAOS_API FEvolutionStats
 	}
 };
 
-
-
 template<class FPBDRigidsEvolution, class FPBDCollisionConstraint, class T, int d>
 class TPBDRigidsEvolutionBase
 {
   public:
-	typedef typename FPBDCollisionConstraint::FRigidBodyContactConstraint FRigidBodyContactConstraint;
-	typedef TPBDRigidClustering<FPBDRigidsEvolution, FPBDCollisionConstraint, T, d> FRigidClustering;
-	typedef TPBDConstraintGraph<T, d> FConstraintGraph;
-	typedef TPBDConstraintGraphRule<T, d> FConstraintRule;
-	typedef TFunction<void(TPBDRigidParticles<T, d>&, const T, const int32)> FForceRule;
-	typedef TFunction<void(TPBDRigidParticles<T, d>&, const T, const TArray<int32>&)> FUpdateVelocityRule;
-	typedef TFunction<void(TPBDRigidParticles<T, d>&, const T)> FUpdatePositionRule;
+	typedef TFunction<void(TTransientPBDRigidParticleHandle<T,d>& Particle, const T)> FForceRule;
+	typedef TFunction<void(const TParticleView<TPBDRigidParticles<T, d>>&, const T)> FUpdateVelocityRule;
+	typedef TFunction<void(const TParticleView<TPBDRigidParticles<T, d>>&, const T)> FUpdatePositionRule;
 	typedef TFunction<void(TPBDRigidParticles<T, d>&, const T, const T, const int32)> FKinematicUpdateRule;
 
-	// TODO(mlentine): Init particles from some type of input
-	CHAOS_API TPBDRigidsEvolutionBase(TPBDRigidParticles<T, d>&& InParticles, int32 NumIterations = 1);
-	CHAOS_API virtual ~TPBDRigidsEvolutionBase() {}
-
-	CHAOS_API inline void InitializeFromParticleData(const int32 StartIndex)
+	CHAOS_API TPBDRigidsEvolutionBase(TPBDRigidsSOAs<T, d>& InParticles, int32 InNumIterations = 1);
+	CHAOS_API virtual ~TPBDRigidsEvolutionBase()
 	{
-		if (!StartIndex)
-		{
-			ActiveIndices.Reset();
-			NonDisabledIndices.Reset();
-		}
-		for (uint32 i = StartIndex; i < Particles.Size(); ++i)
-		{
-			if(Particles.Disabled(i))
-			{
-				continue;
-			}
+		Particles.GetParticleHandles().RemoveArray(&PhysicsMaterials);
+		Particles.GetParticleHandles().RemoveArray(&PerParticlePhysicsMaterials);
+		Particles.GetParticleHandles().RemoveArray(&ParticleDisableCount);
+		Particles.GetParticleHandles().RemoveArray(&Collided);
+	}
 
-			if(Particles.Sleeping(i))
-			{
-				NonDisabledIndices.Add(i);
-				continue;
-			}
+	CHAOS_API TArray<TGeometryParticleHandle<T, d>*> CreateStaticParticles(int32 NumParticles, const TGeometryParticleParameters<T, d>& Params = TGeometryParticleParameters<T, d>()) { return Particles.CreateStaticParticles(NumParticles, Params); }
+	CHAOS_API TArray<TKinematicGeometryParticleHandle<T, d>*> CreateKinematicParticles(int32 NumParticles, const TKinematicGeometryParticleParameters<T, d>& Params = TKinematicGeometryParticleParameters<T, d>()) { return Particles.CreateKinematicParticles(NumParticles, Params); }
+	CHAOS_API TArray<TPBDRigidParticleHandle<T, d>*> CreateDynamicParticles(int32 NumParticles, const TPBDRigidParticleParameters<T, d>& Params = TPBDRigidParticleParameters<T, d>()) { return Particles.CreateDynamicParticles(NumParticles, Params); }
+	CHAOS_API TArray<TPBDRigidClusteredParticleHandle<T, d>*> CreateClusteredParticles(int32 NumParticles, const TPBDRigidParticleParameters<T, d>& Params = TPBDRigidParticleParameters<T, d>()) { return Particles.CreateClusteredParticles(NumParticles, Params); }
 
-			NonDisabledIndices.Add(i);
-			ActiveIndices.Add(i);
-		}
+	CHAOS_API void AddForceFunction(FForceRule ForceFunction) { ForceRules.Add(ForceFunction); }
+	CHAOS_API void SetParticleUpdateVelocityFunction(FUpdateVelocityRule ParticleUpdate) { ParticleUpdateVelocity = ParticleUpdate; }
+	CHAOS_API void SetParticleUpdatePositionFunction(FUpdatePositionRule ParticleUpdate) { ParticleUpdatePosition = ParticleUpdate; }
 
-		if (StartIndex && NonDisabledIndices.Num())
-		{
-			//Hack: we may have inserted double indices, just sort and remove duplicates. Long term should use set
-			NonDisabledIndices.Sort();
-			int32 Prev = NonDisabledIndices[NonDisabledIndices.Num() - 1];
-			for (int32 Idx = NonDisabledIndices.Num() - 2; Idx >= 0; --Idx)
-			{
-				if (NonDisabledIndices[Idx] == Prev)
-				{
-					NonDisabledIndices.RemoveAtSwap(Idx);
-				}
-				else
-				{
-					Prev = NonDisabledIndices[Idx];
-				}
-			}
-		}
+	CHAOS_API TGeometryParticleHandles<T, d>& GetParticleHandles() { return Particles.GetParticleHandles(); }
+	CHAOS_API const TGeometryParticleHandles<T, d>& GetParticleHandles() const { return Particles.GetParticleHandles(); }
 
-		Clustering.InitTopLevelClusterParents(StartIndex);
-	}
+	CHAOS_API TPBDRigidsSOAs<T,d>& GetParticles() { return Particles; }
+	CHAOS_API const TPBDRigidsSOAs<T, d>& GetParticles() const { return Particles; }
 
-	/**
-	 * Enable a particle in the same island and state as another particle
-	 */
-	CHAOS_API void EnableParticle(const int32 ParticleIndex, const int32 ParentParticleIndex)
-	{
-		Particles.SetDisabledLowLevel(ParticleIndex, false);
-		NonDisabledIndices.Add(ParticleIndex);
-		ActiveIndices.Add(ParticleIndex);
+	typedef TPBDConstraintGraph<T, d> FConstraintGraph;
+	typedef TPBDConstraintGraphRule<T, d> FConstraintRule;
 
-		ConstraintGraph.EnableParticle(Particles, ParticleIndex, ParentParticleIndex);
+	CHAOS_API void AddConstraintRule(FConstraintRule* ConstraintRule)
+	{
+		uint32 ContainerId = (uint32)ConstraintRules.Num();
+		ConstraintRules.Add(ConstraintRule);
+		ConstraintRule->BindToGraph(ConstraintGraph, ContainerId);
 	}
 
-	CHAOS_API void DisableParticle(const int32 ParticleIndex)
+	CHAOS_API void EnableParticle(TGeometryParticleHandle<T,d>* Particle, const TGeometryParticleHandle<T, d>* ParentParticle)
 	{
-		Particles.SetDisabledLowLevel(ParticleIndex, true);
-		NonDisabledIndices.Remove(ParticleIndex);
-		ActiveIndices.Remove(ParticleIndex);
+		Particles.EnableParticle(Particle);
+		ConstraintGraph.EnableParticle(Particle, ParentParticle);
+	}
 
-		ConstraintGraph.DisableParticle(Particles, ParticleIndex);
+	CHAOS_API void DisableParticle(TGeometryParticleHandle<T,d>* Particle)
+	{
+		Particles.DisableParticle(Particle);
+		ConstraintGraph.DisableParticle(Particle);
 
-		RemoveConstraints(TSet<int32>({ ParticleIndex }));
+		RemoveConstraints(TSet<TGeometryParticleHandle<T,d>*>({ Particle }));
 	}
 
-	// @todo(ccaulfield): remove uint version
-	CHAOS_API void DisableParticles(const TSet<uint32>& InParticleIndices)
+	CHAOS_API void DestroyParticle(TGeometryParticleHandle<T, d>* Particle)
 	{
-		DisableParticles(reinterpret_cast<const TSet<int32>&>(InParticleIndices));
+		ConstraintGraph.RemoveParticle(Particle);
+		RemoveConstraints(TSet<TGeometryParticleHandle<T, d>*>({ Particle }));
+		Particles.DestroyParticle(Particle);
 	}
-	CHAOS_API void DisableParticles(const TSet<int32>& InParticleIndices)
+
+	CHAOS_API void DisableParticles(const TSet<TGeometryParticleHandle<T,d>*>& InParticles)
 	{
-		for (int32 ParticleIndex : InParticleIndices)
+		for (TGeometryParticleHandle<T, d>* Particle : InParticles)
 		{
-			Particles.SetDisabledLowLevel(ParticleIndex, true);
-			NonDisabledIndices.Remove(ParticleIndex);
-			ActiveIndices.Remove(ParticleIndex);
+			Particles.DisableParticle(Particle);
 		}
 
-		ConstraintGraph.DisableParticles(Particles, InParticleIndices);
+		ConstraintGraph.DisableParticles(InParticles);
 
-		RemoveConstraints(InParticleIndices);
+		RemoveConstraints(InParticles);
 	}
 
 	CHAOS_API void WakeIsland(const int32 Island)
 	{
-		ConstraintGraph.WakeIsland(Particles, Island);
-		for (int32 Particle : ConstraintGraph.GetIslandParticles(Island))
+		ConstraintGraph.WakeIsland(Island);
+		//Update Particles SOAs
+		/*for (auto Particle : ContactGraph.GetIslandParticles(Island))
 		{
 			ActiveIndices.Add(Particle);
-		}
-	}
-
-	// @todo(ccaulfield): Remove uint version
-	CHAOS_API void WakeIslands(const TSet<uint32>& InIslandIndices)
-	{
-		WakeIslands(reinterpret_cast<const TSet<int32>&>(InIslandIndices));
-	}
-	CHAOS_API void WakeIslands(const TSet<int32>& InIslandIndices)
-	{
-		for (int32 Island : InIslandIndices)
-		{
-			ConstraintGraph.WakeIsland(Particles, Island);
-			for (int32 Particle : ConstraintGraph.GetIslandParticles(Island))
-			{
-				ActiveIndices.Add(Particle);
-			}
-		}
-	}
-
-	CHAOS_API void ReconcileIslands()
-	{
-		ConstraintGraph.ReconcileIslands(Particles);
+		}*/
 	}
 
 	// @todo(ccaulfield): Remove the uint version
-	CHAOS_API void RemoveConstraints(const TSet<uint32>& RemovedParticles)
-	{
-		RemoveConstraints(reinterpret_cast<const TSet<int32>&>(RemovedParticles));
-	}
-	CHAOS_API void RemoveConstraints(const TSet<int32>& RemovedParticles)
+	CHAOS_API void RemoveConstraints(const TSet<TGeometryParticleHandle<T, d>*>& RemovedParticles)
 	{
 		for (FConstraintRule* ConstraintRule : ConstraintRules)
 		{
@@ -195,112 +142,42 @@ class TPBDRigidsEvolutionBase
 		}
 	}
 
-	//CHAOS_API void SetKinematicUpdateFunction(TFunction<void(TPBDRigidParticles<T, d>&, const T, const T, const int32)> KinematicUpdate) { MKinematicUpdate = KinematicUpdate; }
-	CHAOS_API void SetParticleUpdateVelocityFunction(TFunction<void(TPBDRigidParticles<T, d>&, const T, const TArray<int32>& InActiveIndices)> ParticleUpdate) { ParticleUpdateVelocity = ParticleUpdate; }
-	CHAOS_API void SetParticleUpdatePositionFunction(TFunction<void(TPBDRigidParticles<T, d>&, const T)> ParticleUpdate) { ParticleUpdatePosition = ParticleUpdate; }
-	CHAOS_API void AddForceFunction(TFunction<void(TPBDRigidParticles<T, d>&, const T, const int32)> ForceFunction) { ForceRules.Add(ForceFunction); }
-	CHAOS_API void AddConstraintRule(FConstraintRule* ConstraintRule) 
-	{ 
-		uint32 ContainerId = (uint32)ConstraintRules.Num();
-		ConstraintRules.Add(ConstraintRule); 
-		ConstraintRule->BindToGraph(ConstraintGraph, ContainerId);
-	}
-
-	// @todo(ccaulfield): Disallow direct write access to containers and provide methods to add/remove enable/disable sleep/wake particles
-
-	/**/
-	CHAOS_API TPBDRigidParticles<T, d>& GetParticles() { return Particles; }
-	CHAOS_API const TPBDRigidParticles<T, d>& GetParticles() const { return Particles; }
+	//TEMP: this is only needed while clustering continues to use indices directly
+	const auto& GetActiveClusteredArray() const { return Particles.GetActiveClusteredArray(); }
+	const auto& GetNonDisabledClusteredArray() const { return Particles.GetNonDisabledClusteredArray(); }
 
-	/**/
-	CHAOS_API const TArray<int32>& GetIslandParticles(const int32 Island) const { return ConstraintGraph.GetIslandParticles(Island); }
-	CHAOS_API int32 NumIslands() const { return ConstraintGraph.NumIslands(); }
-
-	/**/
-	CHAOS_API TSet<int32>& GetActiveIndices() { return ActiveIndices; }
-	CHAOS_API const TSet<int32>& GetActiveIndices() const { return ActiveIndices; }
-
-	// @todo(ccaulfield): optimize this (member array with dirty flag) when we have removed public write access to ActiveIndices
-	CHAOS_API const TArray<int32> GetActiveIndicesArray() const { return ActiveIndices.Array(); }
-
-	/**/
-	CHAOS_API TArray<int32>& GetNonDisabledIndices() { return NonDisabledIndices; }
-	CHAOS_API const TArray<int32>& GetNonDisabledIndices() const { return NonDisabledIndices; }
-
-	/**/
-	CHAOS_API TSet<TTuple<int32, int32>>& GetDisabledCollisions() { return DisabledCollisions; }
-	CHAOS_API const TSet<TTuple<int32, int32>>& GetDisabledCollisions() const { return DisabledCollisions; }
-
-	/**/
-	FRigidClustering& GetRigidClustering() { return Clustering; }
-	const FRigidClustering& GetRigidClustering() const { return Clustering; }
-
-	/**/
-	// @todo(ccaulfield): make sure these are hooked up
-	CHAOS_API inline void SetIterations(int32 Iterations) { NumIterations = Iterations; }
-	CHAOS_API virtual void SetPushOutIterations(int32 PushOutIterations) {}
-	CHAOS_API virtual void SetPushOutPairIterations(int32 PushOutPairIterations) {}
-
-	CHAOS_API TSerializablePtr<TChaosPhysicsMaterial<T>> GetPhysicsMaterial(const int32 Index) const { return PhysicsMaterials[Index]; }
-	CHAOS_API void SetPhysicsMaterial(const int32 Index, TSerializablePtr<TChaosPhysicsMaterial<T>> InMaterial)
+	CHAOS_API TSerializablePtr<TChaosPhysicsMaterial<T>> GetPhysicsMaterial(const TGeometryParticleHandle<T, d>* Particle) const { return Particle->AuxilaryValue(PhysicsMaterials); }
+	CHAOS_API void SetPhysicsMaterial(TGeometryParticleHandle<T,d>* Particle, TSerializablePtr<TChaosPhysicsMaterial<T>> InMaterial)
 	{
-		check(!PerParticlePhysicsMaterials[Index]);	//shouldn't be setting non unique material if a unique one already exists
-		PhysicsMaterials[Index] = InMaterial;
+		check(!Particle->AuxilaryValue(PerParticlePhysicsMaterials)); //shouldn't be setting non unique material if a unique one already exists
+		Particle->AuxilaryValue(PhysicsMaterials) = InMaterial;
 	}
 
-	CHAOS_API const TUniquePtr<TChaosPhysicsMaterial<T>>& GetPerParticlePhysicsMaterial(const int32 Index) const { return PerParticlePhysicsMaterials[Index]; }
-	CHAOS_API void SetPerParticlePhysicsMaterial(const int32 Index, TUniquePtr<TChaosPhysicsMaterial<T>>&& PerParticleMaterial)
-	{
-		PhysicsMaterials[Index] = MakeSerializable(PerParticleMaterial);
-		PerParticlePhysicsMaterials[Index] = MoveTemp(PerParticleMaterial);
-	}
-
-
-#if !UE_BUILD_SHIPPING
-	CHAOS_API void SerializeForPerfTest(FChaosArchive& Ar);
-#endif
-
-	/* Return the instance of the debug substep object that manages the debug pause/progress to step/substep commands for this solver. */
-	CHAOS_API inline FDebugSubstep& GetDebugSubstep() const { return DebugSubstep; }
-
-	CHAOS_API const FEvolutionStats& GetEvolutionStats() const { return EvolutionStats; }
-
+	CHAOS_API const TArray<TGeometryParticleHandle<T,d>*>& GetIslandParticles(const int32 Island) const { return ConstraintGraph.GetIslandParticles(Island); }
+	CHAOS_API int32 NumIslands() const { return ConstraintGraph.NumIslands(); }
 
-	int32 NumConstraints() const
+	void InitializeAccelerationStructures()
 	{
-		int32 NumConstraints = 0;
-		for (const FConstraintRule* ConstraintRule : ConstraintRules)
-		{
-			NumConstraints += ConstraintRule->NumConstraints();
-		}
-		return NumConstraints;
-	}
+		ConstraintGraph.InitializeGraph(Particles.GetNonDisabledView());
 
-	void UpdateConstraintPositionBasedState(T Dt)
-	{
 		for (FConstraintRule* ConstraintRule : ConstraintRules)
 		{
-			ConstraintRule->UpdatePositionBasedState(Particles, NonDisabledIndices, Dt);
+			ConstraintRule->AddToGraph();
 		}
-	}
 
-	void CreateConstraintGraph()
-	{
-		ConstraintGraph.InitializeGraph(Particles, NonDisabledIndices);
+		ConstraintGraph.ResetIslands(Particles.GetNonDisabledDynamicView());
 
 		for (FConstraintRule* ConstraintRule : ConstraintRules)
 		{
-			ConstraintRule->AddToGraph();
+			ConstraintRule->InitializeAccelerationStructures();
 		}
 	}
 
-	void CreateIslands()
+	void UpdateAccelerationStructures(int32 Island)
 	{
-		ConstraintGraph.UpdateIslands(Particles, NonDisabledIndices, ActiveIndices);
-
 		for (FConstraintRule* ConstraintRule : ConstraintRules)
 		{
-			ConstraintRule->InitializeAccelerationStructures();
+			ConstraintRule->UpdateAccelerationStructures(Island);
 		}
 	}
 
@@ -308,227 +185,90 @@ class TPBDRigidsEvolutionBase
 	{
 		for (FConstraintRule* ConstraintRule : ConstraintRules)
 		{
-			ConstraintRule->UpdateAccelerationStructures(Particles, NonDisabledIndices, Island);
+			ConstraintRule->UpdateAccelerationStructures(Island);
 		}
 
 		for (int i = 0; i < NumIterations; ++i)
 		{
 			for (FConstraintRule* ConstraintRule : ConstraintRules)
 			{
-				ConstraintRule->ApplyConstraints(Particles, Dt, Island);
+				ConstraintRule->ApplyConstraints(Dt, Island);
 			}
 		}
 	}
 
-	void UpdateVelocities(const T Dt, int32 Island)
+	const auto& GetRigidClustering() const { return Clustering; }
+	auto& GetRigidClustering() { return Clustering; }
+
+protected:
+	int32 NumConstraints() const
 	{
-		ParticleUpdateVelocity(Particles, Dt, ConstraintGraph.GetIslandParticles(Island));
+		int32 NumConstraints = 0;
+		for (const FConstraintRule* ConstraintRule : ConstraintRules)
+		{
+			NumConstraints += ConstraintRule->NumConstraints();
+		}
+		return NumConstraints;
 	}
 
-	void ApplyPushOut(const T Dt, int32 Island)
+	void UpdateConstraintPositionBasedState(T Dt)
 	{
 		for (FConstraintRule* ConstraintRule : ConstraintRules)
 		{
-			ConstraintRule->ApplyPushOut(Particles, Dt, Island);
+			ConstraintRule->UpdatePositionBasedState(Dt);
 		}
 	}
 
-	void InitializeAccelerationStructures()
+	void CreateConstraintGraph()
 	{
-		const TArray<int32>& ActiveIndicesArray = GetActiveIndicesArray();
-
-		ConstraintGraph.InitializeGraph(Particles, ActiveIndicesArray);
+		ConstraintGraph.InitializeGraph(Particles.GetNonDisabledView());
 
 		for (FConstraintRule* ConstraintRule : ConstraintRules)
 		{
 			ConstraintRule->AddToGraph();
 		}
+	}
 
-		ConstraintGraph.ResetIslands(Particles, ActiveIndicesArray);
+	void CreateIslands()
+	{
+		ConstraintGraph.UpdateIslands(Particles.GetNonDisabledDynamicView(), Particles);
 
 		for (FConstraintRule* ConstraintRule : ConstraintRules)
 		{
 			ConstraintRule->InitializeAccelerationStructures();
 		}
 	}
+	
+	void UpdateVelocities(const T Dt, int32 Island)
+	{
+		ParticleUpdateVelocity(Particles.GetActiveParticlesView(), Dt);
+	}
 
-	void UpdateAccelerationStructures(int32 Island)
+	void ApplyPushOut(const T Dt, int32 Island)
 	{
 		for (FConstraintRule* ConstraintRule : ConstraintRules)
 		{
-			ConstraintRule->UpdateAccelerationStructures(Particles, NonDisabledIndices, Island);
+			ConstraintRule->ApplyPushOut(Dt, Island);
 		}
 	}
 
-protected:
-
-	TPBDRigidParticles<T, d> Particles;
-	FRigidClustering Clustering;
-	FConstraintGraph ConstraintGraph;
-
-	TSet<int32> ActiveIndices;
-	TArray<int32> NonDisabledIndices;
-	TSet<TTuple<int32, int32>> DisabledCollisions;
-	T Time;
-	mutable FDebugSubstep DebugSubstep;
 
-	// User query data
-	TArrayCollectionArray<bool> Collided;
 	TArray<FForceRule> ForceRules;
-	TArray<FConstraintRule*> ConstraintRules;
 	FUpdateVelocityRule ParticleUpdateVelocity;
 	FUpdatePositionRule ParticleUpdatePosition;
 	FKinematicUpdateRule KinematicUpdate;
-
-	int32 NumIterations;
-
+	TArray<FConstraintRule*> ConstraintRules;
+	FConstraintGraph ConstraintGraph;
 	TArrayCollectionArray<TSerializablePtr<TChaosPhysicsMaterial<T>>> PhysicsMaterials;
 	TArrayCollectionArray<TUniquePtr<TChaosPhysicsMaterial<T>>> PerParticlePhysicsMaterials;
 	TArrayCollectionArray<int32> ParticleDisableCount;
-	FEvolutionStats EvolutionStats;
-};
+	TArrayCollectionArray<bool> Collided;
 
-template<class T, int d>
-class TPBDRigidsEvolutionPGS : public TPBDRigidsEvolutionBase<TPBDRigidsEvolutionPGS<T, d>, TPBDCollisionConstraintPGS<T, d>, T, d>
-{
-	typedef TPBDRigidsEvolutionBase<TPBDRigidsEvolutionPGS<T, d>, TPBDCollisionConstraintPGS<T, d>, T, d> Base;
-public:
-	typedef typename Base::FRigidBodyContactConstraint FRigidBodyContactConstraint;
-	typedef typename Base::FRigidClustering FRigidClustering;
-	typedef typename Base::FConstraintGraph FConstraintGraph;
-	typedef typename Base::FConstraintRule FConstraintRule;
-	typedef typename Base::FForceRule FForceRule;
-	typedef typename Base::FUpdateVelocityRule FUpdateVelocityRule;
-	typedef typename Base::FUpdatePositionRule FUpdatePositionRule;
-	typedef typename Base::FKinematicUpdateRule FKinematicUpdateRule;
-	typedef TPBDCollisionConstraintPGS<T, d> FCollisionConstraints;
-	typedef TPBDConstraintIslandRule<FCollisionConstraints, T, d> FCollisionConstraintRule;
-
-	using Base::AddConstraintRule;
-	using Base::ApplyConstraints;
-	using Base::ApplyPushOut;
-	using Base::CreateConstraintGraph;
-	using Base::CreateIslands;
-	using Base::GetActiveIndices;
-	using Base::GetActiveIndicesArray;
-	using Base::GetDebugSubstep;
-	using Base::GetIslandParticles;
-	using Base::GetNonDisabledIndices;
-	using Base::SetParticleUpdatePositionFunction;
-	using Base::SetParticleUpdateVelocityFunction;
-	using Base::UpdateAccelerationStructures;
-	using Base::UpdateConstraintPositionBasedState;
-	using Base::UpdateVelocities;
-
-
-	// TODO(mlentine): Init particles from some type of input
-	CHAOS_API TPBDRigidsEvolutionPGS(TPBDRigidParticles<T, d>&& InParticles, int32 NumIterations = 1);
-	CHAOS_API ~TPBDRigidsEvolutionPGS() {}
-
-	CHAOS_API void IntegrateV(const TArray<int32>& InActiveIndices, const T dt);
-	CHAOS_API void IntegrateX(const TArray<int32>& InActiveIndices, const T dt);
-	CHAOS_API void AdvanceOneTimeStep(const T dt);
-
-private:
-	using Base::ActiveIndices;
-	using Base::Clustering;
-	using Base::ConstraintRules;
-	using Base::ConstraintGraph;
-	using Base::Collided;
-	using Base::EvolutionStats;
-	using Base::ForceRules;
-	using Base::KinematicUpdate;
-	using Base::NumIterations;
-	using Base::Particles;
-	using Base::ParticleDisableCount;
-	using Base::ParticleUpdatePosition;
-	using Base::ParticleUpdateVelocity;
-	using Base::Time;
-	using Base::NonDisabledIndices;
-	using Base::PhysicsMaterials;
-	using Base::PerParticlePhysicsMaterials;
-
-	FCollisionConstraints CollisionConstraints;
-	FCollisionConstraintRule CollisionRule;
-};
+	TPBDRigidsSOAs<T, d>& Particles;
 
-template<class T, int d>
-class TPBDRigidsEvolutionGBF : public TPBDRigidsEvolutionBase<TPBDRigidsEvolutionGBF<T, d>, TPBDCollisionConstraint<T, d>, T, d>
-{
-	typedef TPBDRigidsEvolutionBase<TPBDRigidsEvolutionGBF<T, d>, TPBDCollisionConstraint<T, d>, T, d> Base;
-public:
-	typedef typename Base::FRigidBodyContactConstraint FRigidBodyContactConstraint;
-	typedef typename Base::FRigidClustering FRigidClustering;
-	typedef typename Base::FConstraintGraph FConstraintGraph;
-	typedef typename Base::FConstraintRule FConstraintRule;
-	typedef typename Base::FForceRule FForceRule;
-	typedef typename Base::FUpdateVelocityRule FUpdateVelocityRule;
-	typedef typename Base::FUpdatePositionRule FUpdatePositionRule;
-	typedef typename Base::FKinematicUpdateRule FKinematicUpdateRule;
-	typedef TPBDCollisionConstraint<T, d> FCollisionConstraints;
-	typedef TPBDConstraintColorRule<FCollisionConstraints, T, d> FCollisionConstraintRule;
-
-	static const int32 DefaultNumIterations = 1;
-	static const int32 DefaultNumPushOutIterations = 5;
-	static const int32 DefaultNumPushOutPairIterations = 2;
-
-	using Base::AddConstraintRule;
-	using Base::ApplyConstraints;
-	using Base::ApplyPushOut;
-	using Base::CreateConstraintGraph;
-	using Base::CreateIslands;
-	using Base::GetActiveIndices;
-	using Base::GetActiveIndicesArray;
-	using Base::GetDebugSubstep;
-	using Base::GetIslandParticles;
-	using Base::GetNonDisabledIndices;
-	using Base::SetParticleUpdatePositionFunction;
-	using Base::SetParticleUpdateVelocityFunction;
-	using Base::UpdateAccelerationStructures;
-	using Base::UpdateConstraintPositionBasedState;
-	using Base::UpdateVelocities;
-
-	// TODO(mlentine): Init particles from some type of input
-	CHAOS_API TPBDRigidsEvolutionGBF(TPBDRigidParticles<T, d>&& InParticles, int32 NumIterations = 1);
-	CHAOS_API ~TPBDRigidsEvolutionGBF() {}
-
-	FCollisionConstraints& GetCollisionConstraints() { return CollisionConstraints; }
-	const FCollisionConstraints& GetCollisionConstraints() const { return CollisionConstraints; }
-
-	CHAOS_API void AdvanceOneTimeStep(const T dt);
-
-	CHAOS_API void Integrate(const TArray<int32>& InActiveIndices, const T dt);
-
-	CHAOS_API const ISpatialAcceleration<T, d>& GetSpatialAcceleration() const { return GetCollisionConstraints().GetSpatialAcceleration(); }
-	CHAOS_API void ReleaseSpatialAcceleration() const { GetCollisionConstraints().ReleaseSpatialAcceleration(); }
-
-	CHAOS_API virtual void SetPushOutIterations(int32 PushOutIterations) override { CollisionRule.SetPushOutIterations(PushOutIterations); }
-	CHAOS_API virtual void SetPushOutPairIterations(int32 PushOutPairIterations) override { CollisionConstraints.SetPushOutPairIterations(PushOutPairIterations); }
-
-	void GatherStats();
-	void DebugDraw();
-
-private:
-	using Base::ActiveIndices;
-	using Base::Clustering;
-	using Base::ConstraintRules;
-	using Base::ConstraintGraph;
-	using Base::Collided;
-	using Base::EvolutionStats;
-	using Base::ForceRules;
-	using Base::KinematicUpdate;
-	using Base::NumIterations;
-	using Base::Particles;
-	using Base::ParticleDisableCount;
-	using Base::ParticleUpdatePosition;
-	using Base::ParticleUpdateVelocity;
-	using Base::Time;
-	using Base::NonDisabledIndices;
-	using Base::PhysicsMaterials;
-	using Base::PerParticlePhysicsMaterials;
-
-
-	FCollisionConstraints CollisionConstraints;
-	FCollisionConstraintRule CollisionRule;
+	TPBDRigidClustering<FPBDRigidsEvolution, FPBDCollisionConstraint, T, d> Clustering;
+
+	int32 NumIterations;
 };
+
 }
diff --git a/Engine/Source/Runtime/Experimental/Chaos/Public/Chaos/PBDRigidsEvolutionGBF.h b/Engine/Source/Runtime/Experimental/Chaos/Public/Chaos/PBDRigidsEvolutionGBF.h
--- a/Engine/Source/Runtime/Experimental/Chaos/Public/Chaos/PBDRigidsEvolutionGBF.h
+++ b/Engine/Source/Runtime/Experimental/Chaos/Public/Chaos/PBDRigidsEvolutionGBF.h
@@ -1,4 +1,121 @@
 // Copyright 1998-2019 Epic Games, Inc. All Rights Reserved.
 #pragma once
 
-#include "Chaos/PBDRigidsEvolution.h"
\ No newline at end of file
+#include "Chaos/PBDRigidsEvolution.h"
+#include "Chaos/PBDCollisionConstraint.h"
+#include "Chaos/ChaosPerfTest.h"
+#include "Chaos/PerParticleInitForce.h"
+#include "Chaos/PerParticleEulerStepVelocity.h"
+#include "Chaos/PerParticleEtherDrag.h"
+#include "Chaos/PerParticlePBDEulerStep.h"
+
+namespace Chaos
+{
+class FChaosArchive;
+
+CHAOS_API extern float HackMaxAngularVelocity;
+CHAOS_API extern float HackMaxVelocity;
+
+CHAOS_API extern float HackLinearDrag;
+CHAOS_API extern float HackAngularDrag;
+
+template<typename T, int d>
+class TPBDRigidsEvolutionGBF : public TPBDRigidsEvolutionBase<TPBDRigidsEvolutionGBF<T, d>, TPBDCollisionConstraint<T,d>, T, d>
+{
+public:
+	using Base = TPBDRigidsEvolutionBase<TPBDRigidsEvolutionGBF<T, d>, TPBDCollisionConstraint<T, d>, T, d>;
+	using Base::Particles;
+	using Base::ForceRules;
+	using Base::ParticleUpdatePosition;
+	using Base::SetParticleUpdateVelocityFunction;
+	using Base::SetParticleUpdatePositionFunction;
+	using Base::AddConstraintRule;
+	using Base::Clustering;
+	using typename Base::FForceRule;
+	using FCollisionConstraints = TPBDCollisionConstraint<T, d>;
+	using FCollisionConstraintRule = TPBDConstraintColorRule<FCollisionConstraints, T, d>;
+
+	static constexpr int32 DefaultNumIterations = 1;
+	static constexpr int32 DefaultNumPushOutIterations = 5;
+	static constexpr int32 DefaultNumPushOutPairIterations = 2;
+
+	CHAOS_API TPBDRigidsEvolutionGBF(TPBDRigidsSOAs<T, d>& InParticles, int32 InNumIterations = DefaultNumIterations);
+	CHAOS_API ~TPBDRigidsEvolutionGBF() {}
+
+	CHAOS_API void AdvanceOneTimeStep(const T dt);
+
+	using Base::ApplyConstraints;
+	using Base::ApplyPushOut;
+
+	FCollisionConstraints& GetCollisionConstraints() { return CollisionConstraints; }
+	const FCollisionConstraints& GetCollisionConstraints() const { return CollisionConstraints; }
+
+	template <typename TParticleView>
+	void Integrate(const TParticleView& InParticles, T Dt)
+	{
+		//SCOPE_CYCLE_COUNTER(STAT_Integrate);
+		CHAOS_SCOPED_TIMER(Integrate);
+		TPerParticleInitForce<T, d> InitForceRule;
+		TPerParticleEulerStepVelocity<T, d> EulerStepVelocityRule;
+		TPerParticleEtherDrag<T, d> EtherDragRule(HackLinearDrag, HackAngularDrag);
+		TPerParticlePBDEulerStep<T, d> EulerStepRule;
+
+		const T MaxAngularSpeedSq = HackMaxAngularVelocity * HackMaxAngularVelocity;
+		const T MaxSpeedSq = HackMaxVelocity * HackMaxVelocity;
+		InParticles.ParallelFor([&](auto& GeomParticle, int32 Index)
+		{
+			//question: can we enforce this at the API layer? Right now islands contain non dynamic which makes this hard
+			if (auto PBDParticle = GeomParticle.AsDynamic())
+			{
+				auto& Particle = *PBDParticle;
+
+				//save off previous velocities
+				Particle.PreV() = Particle.V();
+				Particle.PreW() = Particle.W();
+
+				InitForceRule.Apply(Particle, Dt);
+				for (FForceRule ForceRule : ForceRules)
+				{
+					ForceRule(Particle, Dt);
+				}
+				EulerStepVelocityRule.Apply(Particle, Dt);
+				EtherDragRule.Apply(Particle, Dt);
+
+				if (HackMaxAngularVelocity >= 0.f)
+				{
+					const T AngularSpeedSq = Particle.W().SizeSquared();
+					if (AngularSpeedSq > MaxAngularSpeedSq)
+					{
+						Particle.W() = Particle.W() * (HackMaxAngularVelocity / FMath::Sqrt(AngularSpeedSq));
+					}
+				}
+
+				if (HackMaxVelocity >= 0.f)
+				{
+					const T SpeedSq = Particle.V().SizeSquared();
+					if (SpeedSq > MaxSpeedSq)
+					{
+						Particle.V() = Particle.V() * (HackMaxVelocity / FMath::Sqrt(SpeedSq));
+					}
+				}
+
+
+				EulerStepRule.Apply(Particle, Dt);
+			}
+		});
+	}
+
+protected:
+	using Base::UpdateConstraintPositionBasedState;
+	using Base::CreateConstraintGraph;
+	using Base::CreateIslands;
+	using Base::ConstraintGraph;
+	using Base::UpdateVelocities;
+	using Base::PhysicsMaterials;
+	using Base::ParticleDisableCount;
+	using Base::Collided;
+
+	FCollisionConstraints CollisionConstraints;
+	FCollisionConstraintRule CollisionRule;
+};
+}
diff --git a/Engine/Source/Runtime/Experimental/Chaos/Public/Chaos/PBDRigidsSOAs.h b/Engine/Source/Runtime/Experimental/Chaos/Public/Chaos/PBDRigidsSOAs.h
new file mode 100644
--- /dev/null
+++ b/Engine/Source/Runtime/Experimental/Chaos/Public/Chaos/PBDRigidsSOAs.h
@@ -0,0 +1,349 @@
+// Copyright 1998-2019 Epic Games, Inc. All Rights Reserved.
+#pragma once
+
+#include "Chaos/ParticleHandle.h"
+#include "Chaos/Framework/PhysicsProxy.h"
+
+namespace Chaos
+{
+
+template <typename T, int d>
+class TPBDRigidsSOAs
+{
+public:
+	TPBDRigidsSOAs()
+	{
+#if CHAOS_DETERMINISTIC
+		BiggestParticleID = 0;
+#endif
+
+		ClusteredParticles.RemoveParticleBehavior() = ERemoveParticleBehavior::Remove;	//clustered particles maintain relative ordering
+		UpdateViews();
+	}
+
+	TPBDRigidsSOAs(const TPBDRigidsSOAs<T,d>&) = delete;
+	TPBDRigidsSOAs(TPBDRigidsSOAs<T, d>&& Other) = delete;
+
+	void Reset()
+	{
+		check(0);
+	}
+	
+	TArray<TGeometryParticleHandle<T, d>*> CreateStaticParticles(int32 NumParticles, const TGeometryParticleParameters<T, d>& Params = TGeometryParticleParameters<T, d>())
+	{
+		auto Results =  CreateParticlesHelper<TGeometryParticleHandle<T, d>>(NumParticles, Params.bDisabled ? StaticDisabledParticles : StaticParticles, Params);
+		UpdateViews();
+		return Results;
+	}
+	TArray<TKinematicGeometryParticleHandle<T, d>*> CreateKinematicParticles(int32 NumParticles, const TKinematicGeometryParticleParameters<T, d>& Params = TKinematicGeometryParticleParameters<T, d>())
+	{
+		auto Results = CreateParticlesHelper<TKinematicGeometryParticleHandle<T, d>>(NumParticles, Params.bDisabled ? KinematicDisabledParticles : KinematicParticles, Params);
+		UpdateViews();
+		return Results;
+	}
+	TArray<TPBDRigidParticleHandle<T, d>*> CreateDynamicParticles(int32 NumParticles, const TPBDRigidParticleParameters<T, d>& Params = TPBDRigidParticleParameters<T, d>())
+	{
+		auto Results = CreateParticlesHelper<TPBDRigidParticleHandle<T, d>>(NumParticles, Params.bDisabled ? DynamicDisabledParticles : DynamicParticles, Params);;
+
+		if (!Params.bStartSleeping)
+		{
+			InsertToMapAndArray(Results, ActiveParticlesToIndex, ActiveParticlesArray);
+		}
+		UpdateViews();
+		return Results;
+	}
+
+	/** Used specifically by PBDRigidClustering. These have special properties for maintaining relative order, efficiently switching from kinematic to dynamic, disable to enable, etc... */
+	TArray<TPBDRigidClusteredParticleHandle<T, d>*> CreateClusteredParticles(int32 NumParticles, const TPBDRigidParticleParameters<T, d>& Params = TPBDRigidParticleParameters<T, d>())
+	{
+		auto NewClustered = CreateParticlesHelper<TPBDRigidClusteredParticleHandle<T, d>>(NumParticles, ClusteredParticles, Params);
+		
+			if (!Params.bDisabled)
+			{
+			InsertToMapAndArray(NewClustered, NonDisabledClusteredToIndex, NonDisabledClusteredArray);
+			}
+
+			if (!Params.bStartSleeping)
+			{
+			InsertToMapAndArray(reinterpret_cast<TArray<TPBDRigidParticleHandle<T,d>*>&>(NewClustered), ActiveParticlesToIndex, ActiveParticlesArray);
+			InsertToMapAndArray(NewClustered, ActiveClusteredToIndex, ActiveClusteredArray);
+		}
+
+		UpdateViews();
+		
+		return NewClustered;
+	}
+
+	void DestroyParticle(TGeometryParticleHandle<T, d>* Particle)
+	{
+		check(Particle->AsClustered() == nullptr);	//not supported
+
+		if (auto PBDRigid = Particle->AsDynamic())
+		{
+			RemoveFromMapAndArray(PBDRigid, ActiveParticlesToIndex, ActiveParticlesArray);
+		}
+
+		ParticleHandles.DestroyHandleSwap(Particle);
+		UpdateViews();
+	}
+
+	void DisableParticle(TGeometryParticleHandle<T, d>* Particle)
+	{
+		if (auto PBDRigid = Particle->AsDynamic())
+		{
+			PBDRigid->Disabled() = true;
+			PBDRigid->V() = TVector<T, d>(0);
+			PBDRigid->W() = TVector<T, d>(0);
+
+			if (auto PBDRigidClustered = Particle->AsClustered())
+			{
+				RemoveFromMapAndArray(PBDRigidClustered, NonDisabledClusteredToIndex, NonDisabledClusteredArray);
+				RemoveFromMapAndArray(PBDRigidClustered, ActiveClusteredToIndex, ActiveClusteredArray);
+			}
+			else
+			{
+				Particle->MoveToSOA(DynamicDisabledParticles);
+			}
+			RemoveFromMapAndArray(PBDRigid, ActiveParticlesToIndex, ActiveParticlesArray);
+		}
+		else if (Particle->AsKinematic())
+		{
+			Particle->MoveToSOA(KinematicDisabledParticles);
+		}
+		else
+		{
+			Particle->MoveToSOA(StaticDisabledParticles);
+		}
+		UpdateViews();
+	}
+
+	void EnableParticle(TGeometryParticleHandle<T, d>* Particle)
+	{
+		if (auto PBDRigid = Particle->AsDynamic())
+		{
+			if (auto PBDRigidClustered = Particle->AsClustered())
+			{
+				InsertToMapAndArray(PBDRigidClustered, NonDisabledClusteredToIndex, NonDisabledClusteredArray);
+				if (!PBDRigid->Sleeping())
+				{
+					InsertToMapAndArray(PBDRigidClustered, ActiveClusteredToIndex, ActiveClusteredArray);
+				}
+			}
+			else
+			{
+				Particle->MoveToSOA(DynamicParticles);
+			}
+
+			if (!PBDRigid->Sleeping())
+			{
+				InsertToMapAndArray(PBDRigid, ActiveParticlesToIndex, ActiveParticlesArray);
+			}
+
+			PBDRigid->Disabled() = false;
+		}
+		else if (Particle->AsKinematic())
+		{
+			Particle->MoveToSOA(KinematicParticles);
+		}
+		else
+		{
+			Particle->MoveToSOA(StaticParticles);
+		}
+		UpdateViews();
+	}
+
+	void ActivateParticle(TGeometryParticleHandle<T, d>* Particle)
+	{
+		if (auto PBDRigid = Particle->AsDynamic())
+		{
+			check(!PBDRigid->Disabled());
+			if (auto PBDRigidClustered = Particle->AsClustered())
+			{
+				InsertToMapAndArray(PBDRigidClustered, ActiveClusteredToIndex, ActiveClusteredArray);
+			}
+
+			InsertToMapAndArray(PBDRigid, ActiveParticlesToIndex, ActiveParticlesArray);
+		}
+		
+		UpdateViews();
+	}
+
+	void DeactivateParticle(TGeometryParticleHandle<T, d>* Particle)
+	{
+		if (auto PBDRigid = Particle->AsDynamic())
+		{
+			check(!PBDRigid->Disabled());
+			if (auto PBDRigidClustered = Particle->AsClustered())
+			{
+				RemoveFromMapAndArray(PBDRigidClustered, ActiveClusteredToIndex, ActiveClusteredArray);
+			}
+
+			RemoveFromMapAndArray(PBDRigid, ActiveParticlesToIndex, ActiveParticlesArray);
+		}
+
+		UpdateViews();
+	}
+
+	void DeactivateParticles(const TArray<TGeometryParticleHandle<T, d>*>& Particles)
+	{
+		for (auto Particle : Particles)
+		{
+			DeactivateParticle(Particle);
+		}
+	}
+
+
+	const TParticleView<TGeometryParticles<T, d>>& GetNonDisabledView() const { return NonDisabledView; }
+
+	const TParticleView<TPBDRigidParticles<T, d>>& GetNonDisabledDynamicView() const { return NonDisabledDynamicView; }
+
+	const TParticleView<TPBDRigidParticles<T, d>>& GetActiveParticlesView() const { return ActiveParticlesView; }
+
+	const TParticleView<TGeometryParticles<T, d>>& GetAllParticlesView() const { return AllParticlesView; }
+
+	const TGeometryParticleHandles<T, d>& GetParticleHandles() const { return ParticleHandles; }
+	TGeometryParticleHandles<T, d>& GetParticleHandles() { return ParticleHandles; }
+
+	const TPBDRigidParticles<T, d>& GetDynamicParticles() const { return DynamicParticles; }
+	TPBDRigidParticles<T, d>& GetDynamicParticles() { return DynamicParticles; }
+
+	const TGeometryParticles<T, d>& GetNonDisabledStaticParticles() const { return StaticParticles; }
+	TGeometryParticles<T, d>& GetNonDisabledStaticParticles() { return StaticParticles; }
+
+	//TEMP: only needed while clustering code continues to use direct indices
+	const auto& GetActiveClusteredArray() const { return ActiveClusteredArray; }
+	const auto& GetNonDisabledClusteredArray() const { return NonDisabledClusteredArray; }
+
+	const auto& GetClusteredParticles() const { return ClusteredParticles; }
+	auto& GetClusteredParticles() { return ClusteredParticles; }
+
+private:
+	template <typename TParticleHandleType, typename TParticles>
+	TArray<TParticleHandleType*> CreateParticlesHelper(int32 NumParticles, TParticles& Particles, const TGeometryParticleParameters<T, d>& Params)
+	{
+		const int32 ParticlesStartIdx = Particles.Size();
+		Particles.AddParticles(NumParticles);
+		TArray<TParticleHandleType*> ReturnHandles;
+		ReturnHandles.AddUninitialized(NumParticles);
+
+		const int32 HandlesStartIdx = ParticleHandles.Size();
+		ParticleHandles.AddHandles(NumParticles);
+
+		for (int32 Count = 0; Count < NumParticles; ++Count)
+		{
+			const int32 ParticleIdx = Count + ParticlesStartIdx;
+			const int32 HandleIdx = Count + HandlesStartIdx;
+
+			TParticleHandleType* NewParticleHandle = new TParticleHandleType(&Particles, ParticleIdx, HandleIdx);
+#if CHAOS_DETERMINISTIC
+			NewParticleHandle->ParticleID() = BiggestParticleID++;
+#endif
+			ParticleHandles.Handle(HandleIdx).Reset(NewParticleHandle);
+			ReturnHandles[Count] = NewParticleHandle;
+		}
+
+		return ReturnHandles;
+	}
+
+	template <typename TParticle>
+	void InsertToMapAndArray(const TArray<TParticle*>& ParticlesToInsert, TMap<TParticle*, int32>& ParticleToIndex, TArray<TParticle*>& ParticleArray)
+	{
+		int32 NextIdx = ParticleArray.Num();
+		for (auto Particle : ParticlesToInsert)
+		{
+			ParticleToIndex.Add(Particle, NextIdx++);
+		}
+		ParticleArray.Append(ParticlesToInsert);
+	}
+
+	template <typename TParticle>
+	static void InsertToMapAndArray(TParticle* Particle, TMap<TParticle*, int32>& ParticleToIndex, TArray<TParticle*>& ParticleArray)
+	{
+		if (ParticleToIndex.Contains(Particle) == false)
+			{
+			ParticleToIndex.Add(Particle, ParticleArray.Num());
+			ParticleArray.Add(Particle);
+			}
+		}
+
+	template <typename TParticle>
+	static void RemoveFromMapAndArray(TParticle* Particle, TMap<TParticle*, int32>& ParticleToIndex, TArray<TParticle*>& ParticleArray)
+	{
+		if (int32* IdxPtr = ParticleToIndex.Find(Particle))
+		{
+			int32 Idx = *IdxPtr;
+			ParticleArray.RemoveAtSwap(Idx);
+			if (Idx < ParticleArray.Num())
+			{
+				//update swapped element with new index
+				ParticleToIndex[ParticleArray[Idx]] = Idx;
+			}
+			ParticleToIndex.Remove(Particle);
+		}
+	}
+
+	//should be called whenever particles are added / removed / reordered
+	void UpdateViews()
+	{
+		//build various views. Group SOA types together for better branch prediction
+		{
+			TArray<TSOAView<TGeometryParticles<T, d>>> TmpArray = { &StaticParticles, &KinematicParticles, &DynamicParticles, {&NonDisabledClusteredArray} };
+			NonDisabledView = MakeParticleView(MoveTemp(TmpArray));
+		}
+		{
+			TArray<TSOAView<TPBDRigidParticles<T, d>>> TmpArray = { &DynamicParticles, {&NonDisabledClusteredArray} };
+			NonDisabledDynamicView = MakeParticleView(MoveTemp(TmpArray));
+		}
+		{
+			TArray<TSOAView<TPBDRigidParticles<T, d>>> TmpArray = { {&ActiveParticlesArray} };
+			ActiveParticlesView = MakeParticleView(MoveTemp(TmpArray));
+		}
+
+		{
+			TArray<TSOAView<TGeometryParticles<T, d>>> TmpArray = { &StaticParticles, &StaticDisabledParticles, &KinematicParticles, &KinematicDisabledParticles,
+			&DynamicParticles, &DynamicDisabledParticles, &ClusteredParticles };
+			AllParticlesView = MakeParticleView(MoveTemp(TmpArray));
+		}
+	}
+
+	//Organized by SOA type
+	TGeometryParticles<T, d> StaticParticles;
+	TGeometryParticles<T, d> StaticDisabledParticles;
+
+	TKinematicGeometryParticles<T, d> KinematicParticles;
+	TKinematicGeometryParticles<T, d> KinematicDisabledParticles;
+
+	TPBDRigidParticles<T, d> DynamicParticles;
+	TPBDRigidParticles<T, d> DynamicDisabledParticles;
+	
+	TPBDRigidClusteredParticles<T, d> ClusteredParticles;
+
+	//Utility structures for maintaining an Active particles view
+	TMap<TPBDRigidParticleHandle<T, d>*, int32> ActiveParticlesToIndex;
+	TArray<TPBDRigidParticleHandle<T, d>*> ActiveParticlesArray;
+	TMap<TPBDRigidClusteredParticleHandle<T, d>*, int32> ActiveClusteredToIndex;
+	TArray<TPBDRigidClusteredParticleHandle<T, d>*> ActiveClusteredArray;
+
+	//Utility structures for maintaining a NonDisabled particle view
+	TMap<TPBDRigidClusteredParticleHandle<T, d>*, int32> NonDisabledClusteredToIndex;
+	TArray<TPBDRigidClusteredParticleHandle<T, d>*> NonDisabledClusteredArray;
+
+	//Particle Views
+	TParticleView<TGeometryParticles<T, d>> NonDisabledView;	//all particles that are not disabled
+	TParticleView<TPBDRigidParticles<T, d>> NonDisabledDynamicView;	//all dynamic particles that are not disabled
+	TParticleView<TPBDRigidParticles<T, d>> ActiveParticlesView;	//all particles that are active
+	TParticleView<TGeometryParticles<T, d>> AllParticlesView;	//all particles
+
+	//Auxiliary data synced with particle handles
+	TGeometryParticleHandles<T, d> ParticleHandles;
+	TSet<TGeometryParticleHandle<T, d>*> ActiveParticles;
+	TSet<TPBDRigidClusteredParticleHandle<T, d>*> NonDisabledRigidClustered;
+	TArray<TPBDRigidClusteredParticleHandle<T,d>*> NonDisabledRigidClusteredArray;
+	TSet<TPBDRigidClusteredParticleHandle<T, d>*> ActiveRigidClustered;
+	TArray<TPBDRigidClusteredParticleHandle<T, d>*> ActiveRigidClusteredArray;
+
+#if CHAOS_DETERMINISTIC
+	int32 BiggestParticleID;
+#endif
+};
+}
diff --git a/Engine/Source/Runtime/Experimental/Chaos/Public/Chaos/Pair.h b/Engine/Source/Runtime/Experimental/Chaos/Public/Chaos/Pair.h
--- a/Engine/Source/Runtime/Experimental/Chaos/Public/Chaos/Pair.h
+++ b/Engine/Source/Runtime/Experimental/Chaos/Public/Chaos/Pair.h
@@ -6,6 +6,11 @@ namespace Chaos
 template<class T1, class T2>
 struct Pair
 {
+	Pair()
+	{}
+	Pair(const T1& First, const T2& Second)
+		: First(First), Second(Second)
+	{}
 	T1 First;
 	T2 Second;
 };
diff --git a/Engine/Source/Runtime/Experimental/Chaos/Public/Chaos/ParticleDirtyFlags.h b/Engine/Source/Runtime/Experimental/Chaos/Public/Chaos/ParticleDirtyFlags.h
new file mode 100644
--- /dev/null
+++ b/Engine/Source/Runtime/Experimental/Chaos/Public/Chaos/ParticleDirtyFlags.h
@@ -0,0 +1,66 @@
+// Copyright 1998-2019 Epic Games, Inc. All Rights Reserved.
+
+#pragma once
+
+namespace Chaos
+{
+	// There is a dirty flag for every user-settable particle property.
+	// Dirty property values will get copied from game to physics thread buffers,
+	// but clean property values will get overridden with physics thread results.
+	enum EParticleFlags : int32
+	{
+		X =				1 << 0,
+		R =				1 << 1,
+		V =				1 << 2,
+		W =				1 << 3,
+		CollisionGroup =1 << 4,
+		Disabled =		1 << 5,
+		PreV =			1 << 6,
+		PreW =			1 << 7,
+		P =				1 << 8,
+		Q =				1 << 9,
+		F =				1 << 10,
+		Torque =		1 << 11,
+		I =				1 << 12,
+		InvI =			1 << 13,
+		M =				1 << 14,
+		InvM =			1 << 15,
+		ObjectState =	1 << 16,
+		Geometry =		1 << 17
+	};
+
+	class FParticleDirtyFlags
+	{
+	public:
+		FParticleDirtyFlags() : Bits(0) { }
+		~FParticleDirtyFlags() { }
+
+		bool IsDirty() const
+		{
+			return Bits != 0;
+		}
+
+		bool IsDirty(const EParticleFlags CheckBits) const
+		{
+			return (Bits & (int32)CheckBits) != 0;
+		}
+
+		void MarkDirty(const EParticleFlags DirtyBits)
+		{
+			Bits |= (int32)DirtyBits;
+		}
+
+		void MarkClean(const EParticleFlags CleanBits)
+		{
+			Bits |= ~(int32)CleanBits;
+		}
+
+		void Clear()
+		{
+			Bits = 0;
+		}
+
+	private:
+		int32 Bits;
+	};
+}
