Copying //UE4/Dev-Physics to Dev-Main (//UE4/Dev-Main) [at] 7703071

diff --git a/Engine/Source/Runtime/Experimental/Chaos/Public/Chaos/ParticleHandle.h b/Engine/Source/Runtime/Experimental/Chaos/Public/Chaos/ParticleHandle.h
--- a/Engine/Source/Runtime/Experimental/Chaos/Public/Chaos/ParticleHandle.h
+++ b/Engine/Source/Runtime/Experimental/Chaos/Public/Chaos/ParticleHandle.h
@@ -1,7 +1,12 @@
 // Copyright 1998-2019 Epic Games, Inc. All Rights Reserved.
 #pragma once
 
-#include "Chaos/PBDRigidParticles.h"
+#include "Chaos/PBDRigidClusteredParticles.h"
+#include "Chaos/ParticleHandleFwd.h"
+#include "Chaos/ParticleIterator.h"
+#include "Chaos/ParticleDirtyFlags.h"
+
+class IPhysicsProxyBase;
 
 namespace Chaos
 {
@@ -36,38 +41,42 @@ struct TPBDRigidParticleParameters : public TKinematicGeometryParticleParameters
 template <typename T, int d, typename FConcrete>
 void GeometryParticleDefaultConstruct(FConcrete& Concrete, const TGeometryParticleParameters<T,d>& Params)
 {
-	Concrete.X() = TVector<T, d>(0);
-	Concrete.R() = TRotation<T, d>::Identity;
+	Concrete.SetX(TVector<T, d>(0));
+	Concrete.SetR(TRotation<T, d>::Identity);
 }
 
 template <typename T, int d, typename FConcrete>
 void KinematicGeometryParticleDefaultConstruct(FConcrete& Concrete, const TKinematicGeometryParticleParameters<T, d>& Params)
 {
-	Concrete.V() = TVector<T, d>(0);
-	Concrete.W() = TVector<T, d>(0);
+	Concrete.SetV(TVector<T, d>(0));
+	Concrete.SetW(TVector<T, d>(0));
 }
 
 template <typename T, int d, typename FConcrete>
 void PBDRigidParticleDefaultConstruct(FConcrete& Concrete, const TPBDRigidParticleParameters<T, d>& Params)
 {
 	//don't bother calling parent since the call gets made by the corresponding hierarchy in FConcrete
-	Concrete.CollisionGroup() = 0;
-	Concrete.SetDisabledLowLevel(Params.bDisabled);
-	Concrete.PreV() = Concrete.V();
-	Concrete.PreW() = Concrete.W();
-	Concrete.P() = Concrete.X();
-	Concrete.Q() = Concrete.R();
-	Concrete.F() = TVector<T, d>(0);
-	Concrete.Torque() = TVector<T, d>(0);
-	Concrete.M() = 1;
-	Concrete.InvM() = 1;
-	Concrete.I() = PMatrix<T, d, d>(1, 1, 1);
-	Concrete.InvI() = PMatrix<T, d, d>(1, 1, 1);
-	Concrete.Island() = INDEX_NONE;
-	Concrete.ToBeRemovedOnFracture() = false;
+	Concrete.SetCollisionGroup(0);
+	Concrete.SetDisabled(Params.bDisabled);
+	Concrete.SetPreV(Concrete.V());
+	Concrete.SetPreW(Concrete.W());
+	Concrete.SetP(Concrete.X());
+	Concrete.SetQ(Concrete.R());
+	Concrete.SetF(TVector<T, d>(0));
+	Concrete.SetTorque(TVector<T, d>(0));
+	Concrete.SetM(1);
+	Concrete.SetInvM(1);
+	Concrete.SetI(PMatrix<T, d, d>(1, 1, 1));
+	Concrete.SetInvI(PMatrix<T, d, d>(1, 1, 1));
 	Concrete.SetObjectState(Params.bStartSleeping ? EObjectStateType::Sleeping : EObjectStateType::Dynamic);
 }
 
+template <typename T, int d, typename FConcrete>
+void PBDRigidClusteredParticleDefaultConstruct(FConcrete& Concrete, const TPBDRigidParticleParameters<T, d>& Params)
+{
+	//don't bother calling parent since the call gets made by the corresponding hierarchy in FConcrete
+}
+
 template <typename FConcrete>
 bool GeometryParticleSleeping(const FConcrete& Concrete)
 {
@@ -75,195 +84,515 @@ bool GeometryParticleSleeping(const FConcrete& Concrete)
 }
 
 template <typename T, int d>
-class TGeometryParticleHandle;
+class TParticleHandleBase
+{
+public:
+	using TType = T;
+	static constexpr int D = d;
 
-template <typename T, int d>
-class TKinematicGeometryParticleHandle;
+	TParticleHandleBase(TGeometryParticles<T, d>* InParticles = nullptr, int32 InParticleIdx = 0)
+	: GeometryParticles(InParticles)
+	, ParticleIdx(InParticleIdx)
+	, Type(InParticles ? InParticles->ParticleType() : EParticleType::Static)
+	{
+	}
 
-template <typename T, int d>
-class TPBDRigidParticleHandle;
+	//NOTE: this is not virtual and only acceptable because we know the children have no extra data that requires destruction. 
+	//You must modify the union to extend this class and do not add any member variables
+	~TParticleHandleBase()
+	{
+	}
+
+
+protected:
+	union
+	{
+		TGeometryParticles<T, d>* GeometryParticles;
+		TKinematicGeometryParticles<T, d>* KinematicGeometryParticles;
+		TPBDRigidParticles<T, d>* PBDRigidParticles;
+		TPBDRigidClusteredParticles<T, d>* PBDRigidClusteredParticles;
+	};
+
+	template <typename TSOA>
+	friend class TConstParticleIterator;
+	//todo: maybe make private?
+	int32 ParticleIdx;	//Index into the particle struct of arrays. Note the index can change
+	EParticleType Type;
 
-enum class EParticleType : uint8
-{
-	Static,
-	Kinematic,
-	Dynamic
 };
 
+template <typename T, int d, bool bPersistent>
+class TKinematicGeometryParticleHandleImp;
+
+template <typename T, int d, bool bPersistent>
+class TPBDRigidParticleHandleImp;
+
+template <typename T, int d>
+TGeometryParticleHandle<T, d>* GetHandleHelper(TGeometryParticleHandle<T, d>* Handle) { return Handle; }
+template <typename T, int d>
+const TGeometryParticleHandle<T, d>* GetHandleHelper(const TGeometryParticleHandle<T, d>* Handle) { return Handle; }
+template <typename T, int d>
+TGeometryParticleHandle<T, d>* GetHandleHelper(TTransientGeometryParticleHandle<T, d>* Handle);
+template <typename T, int d>
+const TGeometryParticleHandle<T, d>* GetHandleHelper(const TTransientGeometryParticleHandle<T, d>* Handle);
+
 template <typename T, int d>
-class TGeometryParticleHandle
+class TGeometryParticleHandles;
+
+template <typename T, int d, bool bPersistent>
+class TGeometryParticleHandleImp : public TParticleHandleBase<T,d>
 {
 public:
-	TGeometryParticleHandle(TGeometryParticles<T, d>* InParticles, int32 InParticleIdx, int32 InHandleIdx, const TGeometryParticleParameters<T, d>& Params = TGeometryParticleParameters<T, d>())
-		: GeometryParticles(InParticles)
-		, ParticleIdx(InParticleIdx)
+	using TTransientHandle = TTransientGeometryParticleHandle<T, d>;
+	using THandleBase = TParticleHandleBase<T, d>;
+	using THandleBase::GeometryParticles;
+	using THandleBase::ParticleIdx;
+	using THandleBase::Type;
+	using TSOAType = TGeometryParticles<T,d>;
+
+	TGeometryParticleHandleImp(TGeometryParticles<T, d>* InParticles, int32 InParticleIdx, int32 InHandleIdx, const TGeometryParticleParameters<T, d>& Params = TGeometryParticleParameters<T, d>())
+		: TParticleHandleBase<T,d>(InParticles, InParticleIdx)
 		, HandleIdx(InHandleIdx)
-		, Type(EParticleType::Static)
 	{
 		GeometryParticles->Handle(ParticleIdx) = this;
 		GeometryParticleDefaultConstruct<T, d>(*this, Params);
 	}
 
-	//NOTE: this is not virtual and only acceptable because we know the children have no extra data. 
-	//You must modify the union to extend this class and do not add any member variables
-	~TGeometryParticleHandle()	
+	~TGeometryParticleHandleImp()
 	{
+		if (bPersistent)
+		{
+			GeometryParticles->DestroyParticle(ParticleIdx);
+			if (static_cast<uint32>(ParticleIdx) < GeometryParticles->Size())
+			{
+				if (GeometryParticles->RemoveParticleBehavior() == ERemoveParticleBehavior::RemoveAtSwap)
+				{
+					GeometryParticles->Handle(ParticleIdx)->ParticleIdx = ParticleIdx;
+				}
+				else
+				{
+					//need to update all handles >= ParticleIdx
+					for (int32 Idx = ParticleIdx; static_cast<uint32>(Idx) < GeometryParticles->Size(); ++Idx)
+					{
+						GeometryParticles->Handle(Idx)->ParticleIdx -= 1;
+					}
+				}
+			}
+		}
 	}
 
-	TGeometryParticleHandle(const TGeometryParticleHandle&) = delete;
+	template <typename T2, int d2>
+	friend TGeometryParticleHandle<T2, d2>* GetHandleHelper(TTransientGeometryParticleHandle<T2, d2>* Handle);
+	template <typename T2, int d2>
+	friend const TGeometryParticleHandle<T2, d2>* GetHandleHelper(const TTransientGeometryParticleHandle<T2, d2>* Handle);
+
+	TGeometryParticleHandleImp(const TGeometryParticleHandleImp&) = delete;
 
 	const TVector<T, d>& X() const { return GeometryParticles->X(ParticleIdx); }
 	TVector<T, d>& X() { return GeometryParticles->X(ParticleIdx); }
+	void SetX(const TVector<T, d>& InX) { GeometryParticles->X(ParticleIdx) = InX; }
 
 	const TRotation<T, d>& R() const { return GeometryParticles->R(ParticleIdx); }
 	TRotation<T, d>& R() { return GeometryParticles->R(ParticleIdx); }
+	void SetR(const TRotation<T, d>& InR) { GeometryParticles->R(ParticleIdx) = InR; }
 
 	TSerializablePtr<TImplicitObject<T, d>> Geometry() const { return GeometryParticles->Geometry(ParticleIdx); }
+	void SetGeometry(TSerializablePtr<TImplicitObject<T, d>> InGeometry) { GeometryParticles->SetGeometry(ParticleIdx, InGeometry); }
+
+	TSerializablePtr<TImplicitObject<T, d>> SharedGeometry() const { return GeometryParticles->SharedGeometry(ParticleIdx); }
+	void SetSharedGeometry(TSharedPtr<TImplicitObject<T, d>, ESPMode::ThreadSafe> InGeometry) { GeometryParticles->SetSharedGeometry(ParticleIdx, InGeometry); }
+
+	const TUniquePtr<TImplicitObject<T, d>>& DynamicGeometry() const { return GeometryParticles->DynamicGeometry(ParticleIdx); }
+	void SetDynamicGeometry(TUniquePtr<TImplicitObject<T, d>>&& Unique) { GeometryParticles->SetDynamicGeometry(ParticleIdx, MoveTemp(Unique)); }
+
+	const TShapesArray<T,d>& ShapesArray() const { return GeometryParticles->ShapesArray(ParticleIdx); }
 	
 	EObjectStateType ObjectState() const;
 
-	const TKinematicGeometryParticleHandle<T, d>* ToKinematic() const;
-	TKinematicGeometryParticleHandle<T, d>& ToKinematic();
+	TGeometryParticle<T, d>* GTGeometryParticle() const { return GeometryParticles->GTGeometryParticle(ParticleIdx); }
+	TGeometryParticle<T, d>*& GTGeometryParticle() { return GeometryParticles->GTGeometryParticle(ParticleIdx); }
+
+	const TKinematicGeometryParticleHandleImp<T, d, bPersistent>* AsKinematic() const;
+	TKinematicGeometryParticleHandleImp<T, d, bPersistent>* AsKinematic();
+
+	const TPBDRigidParticleHandleImp<T, d, bPersistent>* AsDynamic() const;
+	TPBDRigidParticleHandleImp<T, d, bPersistent>* AsDynamic();
 
-	const TPBDRigidParticleHandle<T, d>* ToDynamic() const;
-	TPBDRigidParticleHandle<T, d>* ToDynamic();
+	const TPBDRigidClusteredParticleHandleImp<T, d, bPersistent>* AsClustered() const;
+	TPBDRigidClusteredParticleHandleImp<T, d, bPersistent>* AsClustered();
+
+	const TGeometryParticleHandle<T, d>* Handle() const { return GetHandleHelper(this);}
+	TGeometryParticleHandle<T, d>* Handle() { return GetHandleHelper(this); }
 
 	bool Sleeping() const { return GeometryParticleSleeping(*this); }
 
-	int32 TransientHandleIdx() const { return HandleIdx; }
+	template <typename Container>
+	const auto& AuxilaryValue(const Container& AuxContainer) const
+	{
+		return AuxContainer[HandleIdx];
+	}
+
+	template <typename Container>
+	auto& AuxilaryValue(Container& AuxContainer)
+	{
+		return AuxContainer[HandleIdx];
+	}
+
+#if CHAOS_DETERMINISTIC
+	FParticleID ParticleID() const { return GeometryParticles->ParticleID(ParticleIdx); }
+	FParticleID& ParticleID() { return GeometryParticles->ParticleID(ParticleIdx); }
+#endif
+
+	void MoveToSOA(TGeometryParticles<T, d>& ToSOA)
+	{
+		static_assert(bPersistent, "Cannot move particles from a transient handle");
+		check(ToSOA.ParticleType() == Type);
+		if (GeometryParticles != &ToSOA)
+		{
+			GeometryParticles->MoveToOtherParticles(ParticleIdx, ToSOA);
+			if (static_cast<uint32>(ParticleIdx) < GeometryParticles->Size())
+			{
+				GeometryParticles->Handle(ParticleIdx)->ParticleIdx = ParticleIdx;
+			}
+			const int32 NewParticleIdx = ToSOA.Size() - 1;
+			ParticleIdx = NewParticleIdx;
+			GeometryParticles = &ToSOA;
+		}
+	}
 
 	static constexpr EParticleType StaticType() { return EParticleType::Static; }
 
+	FString ToString() const;
+
 protected:
 
-	union
-	{
-		TGeometryParticles<T, d>* GeometryParticles;
-		TKinematicGeometryParticles<T, d>* KinematicGeometryParticles;
-		TPBDRigidParticles<T, d>* PBDRigidParticles;
-	};
-	int32 ParticleIdx;	//Index into the particle struct of arrays. Note the index can change
-	int32 HandleIdx;	//Index into the handles array. This is useful for binding external attributes. Note the index can change
-	EParticleType Type;
+	friend TGeometryParticleHandles<T, d>;
+	
+	struct FInvalidFromTransient {};
+	typename std::conditional<bPersistent, int32, FInvalidFromTransient>::type HandleIdx;	//Index into the handles array. This is useful for binding external attributes. Note the index can change
 };
 
 template <typename T, int d>
-class TKinematicGeometryParticleHandle : public TGeometryParticleHandle<T,d>
+TGeometryParticleHandle<T, d>* GetHandleHelper(TTransientGeometryParticleHandle<T, d>* Handle)
 {
-public:
-	using TGeometryParticleHandle<T, d>::ParticleIdx;
-	using TGeometryParticleHandle<T, d>::KinematicGeometryParticles;
-	using TGeometryParticleHandle<T, d>::Type;
-	using TGeometryParticleHandle<T, d>::ToDynamic;
+	return Handle->GeometryParticles->Handle(Handle->ParticleIdx);
+}
+template <typename T, int d>
+const TGeometryParticleHandle<T, d>* GetHandleHelper(const TTransientGeometryParticleHandle<T, d>* Handle)
+{
+	return Handle->GeometryParticles->Handle(Handle->ParticleIdx);
+}
 
-	TKinematicGeometryParticleHandle(TKinematicGeometryParticles<T,d>* Particles, int32 InIdx, int32 InGlobalIdx, const TKinematicGeometryParticleParameters<T, d>& Params = TKinematicGeometryParticleParameters<T,d>())
-		: TGeometryParticleHandle<T, d>(Particles, InIdx, InGlobalIdx, Params)
+template <typename T, int d, bool bPersistent>
+class TKinematicGeometryParticleHandleImp : public TGeometryParticleHandleImp<T,d, bPersistent>
+{
+public:
+	using TGeometryParticleHandleImp<T, d, bPersistent>::ParticleIdx;
+	using TGeometryParticleHandleImp<T, d, bPersistent>::KinematicGeometryParticles;
+	using TGeometryParticleHandleImp<T, d, bPersistent>::Type;
+	using TGeometryParticleHandleImp<T, d, bPersistent>::AsDynamic;
+	using TTransientHandle = TTransientKinematicGeometryParticleHandle<T, d>;
+	using TSOAType = TKinematicGeometryParticles<T, d>;
+
+	TKinematicGeometryParticleHandleImp(TKinematicGeometryParticles<T,d>* Particles, int32 InIdx, int32 InGlobalIdx, const TKinematicGeometryParticleParameters<T, d>& Params = TKinematicGeometryParticleParameters<T,d>())
+		: TGeometryParticleHandleImp<T, d, bPersistent>(Particles, InIdx, InGlobalIdx, Params)
 	{
-		Type = EParticleType::Kinematic;
 		KinematicGeometryParticleDefaultConstruct<T, d>(*this, Params);
 	}
 
 	const TVector<T, d>& V() const { return KinematicGeometryParticles->V(ParticleIdx); }
 	TVector<T, d>& V() { return KinematicGeometryParticles->V(ParticleIdx); }
+	void SetV(const TVector<T, d>& InV) { KinematicGeometryParticles->V(ParticleIdx) = InV; }
 
 	const TVector<T, d>& W() const { return KinematicGeometryParticles->W(ParticleIdx); }
 	TVector<T, d>& W() { return KinematicGeometryParticles->W(ParticleIdx); }
+	void SetW(const TVector<T, d>& InW) { KinematicGeometryParticles->W(ParticleIdx) = InW; }
+
+	//Really only useful when using a transient handle
+	const TKinematicGeometryParticleHandleImp<T, d, true>* Handle() const { return KinematicGeometryParticles->Handle(ParticleIdx); }
+	TKinematicGeometryParticleHandleImp<T, d, true>* Handle() { return KinematicGeometryParticles->Handle(ParticleIdx); }
 
 	EObjectStateType ObjectState() const;
 	static constexpr EParticleType StaticType() { return EParticleType::Kinematic; }
 };
 
-template <typename T, int d>
-class TPBDRigidParticleHandle : public TKinematicGeometryParticleHandle<T, d>
+template <typename T, int d, bool bPersistent>
+class TPBDRigidParticleHandleImp : public TKinematicGeometryParticleHandleImp<T, d, bPersistent>
 {
 public:
-	using TGeometryParticleHandle<T, d>::ParticleIdx;
-	using TGeometryParticleHandle<T, d>::PBDRigidParticles;
-	using TGeometryParticleHandle<T, d>::Type;
+	using TGeometryParticleHandleImp<T, d, bPersistent>::ParticleIdx;
+	using TGeometryParticleHandleImp<T, d, bPersistent>::PBDRigidParticles;
+	using TGeometryParticleHandleImp<T, d, bPersistent>::Type;
+	using TTransientHandle = TTransientPBDRigidParticleHandle<T, d>;
+	using TSOAType = TPBDRigidParticles<T, d>;
 
-	TPBDRigidParticleHandle<T, d>(TPBDRigidParticles<T,d>* Particles, int32 InIdx, int32 InGlobalIdx, const TPBDRigidParticleParameters<T, d>& Params = TPBDRigidParticleParameters<T, d>())
-		: TKinematicGeometryParticleHandle<T, d>(Particles, InIdx, InGlobalIdx, Params)
+
+	TPBDRigidParticleHandleImp<T, d, bPersistent>(TPBDRigidParticles<T,d>* Particles, int32 InIdx, int32 InGlobalIdx, const TPBDRigidParticleParameters<T, d>& Params = TPBDRigidParticleParameters<T, d>())
+		: TKinematicGeometryParticleHandleImp<T, d, bPersistent>(Particles, InIdx, InGlobalIdx, Params)
 	{
-		Type = EParticleType::Dynamic;
 		PBDRigidParticleDefaultConstruct<T, d>(*this, Params);
+		SetIsland(INDEX_NONE);
+		SetToBeRemovedOnFracture(false);
 	}
 
+	operator TPBDRigidParticleHandleImp<T, d, false>& () { return reinterpret_cast<TPBDRigidParticleHandleImp<T, d, false>&>(*this); }
+
 	const TUniquePtr<TBVHParticles<T, d>>& CollisionParticles() const { return PBDRigidParticles->CollisionParticles(ParticleIdx); }
+	TUniquePtr<TBVHParticles<T, d>>& CollisionParticles() { return PBDRigidParticles->CollisionParticles(ParticleIdx); }
+
+	int32 CollisionParticlesSize() const { return PBDRigidParticles->CollisionParticlesSize(ParticleIdx); }
+	void CollisionParticlesInitIfNeeded() { PBDRigidParticles->CollisionParticlesInitIfNeeded(ParticleIdx); }
+	void SetCollisionParticles(TParticles<T, d>&& Points) { PBDRigidParticles->SetCollisionParticles(ParticleIdx, MoveTemp(Points)); }
 
 	int32 CollisionGroup() const { return PBDRigidParticles->CollisionGroup(ParticleIdx); }
 	int32& CollisionGroup() { return PBDRigidParticles->CollisionGroup(ParticleIdx); }
+	void SetCollisionGroup(const int32 InCollisionGroup) { PBDRigidParticles->CollisionGroup(ParticleIdx) = InCollisionGroup; }
 
 	bool Disabled() const { return PBDRigidParticles->Disabled(ParticleIdx); }
+	bool& Disabled() { return PBDRigidParticles->DisabledRef(ParticleIdx); }
 
 	// See Comment on TRigidParticle::SetDisabledLowLevel. State changes in Evolution should accompany this call.
 	void SetDisabledLowLevel(bool disabled) { PBDRigidParticles->SetDisabledLowLevel(ParticleIdx, disabled); }
+	void SetDisabled(const bool InDisabled) { PBDRigidParticles->DisabledRef(ParticleIdx) = InDisabled; }
 
 	const TVector<T, d>& PreV() const { return PBDRigidParticles->PreV(ParticleIdx); }
 	TVector<T, d>& PreV() { return PBDRigidParticles->PreV(ParticleIdx); }
+	void SetPreV(const TVector<T, d>& InPreV) { PBDRigidParticles->PreV(ParticleIdx) = InPreV; }
+
 	const TVector<T, d>& PreW() const { return PBDRigidParticles->PreW(ParticleIdx); }
 	TVector<T, d>& PreW() { return PBDRigidParticles->PreW(ParticleIdx); }
+	void SetPreW(const TVector<T, d>& InPreW) { PBDRigidParticles->PreW(ParticleIdx) = InPreW; }
 
 	const TVector<T, d>& P() const { return PBDRigidParticles->P(ParticleIdx); }
 	TVector<T, d>& P() { return PBDRigidParticles->P(ParticleIdx); }
+	void SetP(const TVector<T, d>& InP) { PBDRigidParticles->P(ParticleIdx) = InP; }
 
 	const TRotation<T, d>& Q() const { return PBDRigidParticles->Q(ParticleIdx); }
 	TRotation<T, d>& Q() { return PBDRigidParticles->Q(ParticleIdx); }
+	void SetQ(const TRotation<T, d>& InQ) { PBDRigidParticles->Q(ParticleIdx) = InQ; }
 
 	const TVector<T, d>& F() const { return PBDRigidParticles->F(ParticleIdx); }
 	TVector<T, d>& F() { return PBDRigidParticles->F(ParticleIdx); }
+	void SetF(const TVector<T, d>& InF) { PBDRigidParticles->F(ParticleIdx) = InF; }
 
 	const TVector<T, d>& Torque() const { return PBDRigidParticles->Torque(ParticleIdx); }
 	TVector<T, d>& Torque() { return PBDRigidParticles->Torque(ParticleIdx); }
+	void SetTorque(const TVector<T, d>& InTorque) { PBDRigidParticles->Torque(ParticleIdx) = InTorque; }
 
 	const PMatrix<T, d, d>& I() const { return PBDRigidParticles->I(ParticleIdx); }
 	PMatrix<T, d, d>& I() { return PBDRigidParticles->I(ParticleIdx); }
+	void SetI(const PMatrix<T, d, d>& InI) { PBDRigidParticles->I(ParticleIdx) = InI; }
 
 	const PMatrix<T, d, d>& InvI() const { return PBDRigidParticles->InvI(ParticleIdx); }
 	PMatrix<T, d, d>& InvI() { return PBDRigidParticles->InvI(ParticleIdx); }
+	void SetInvI(const PMatrix<T, d, d>& InInvI) { PBDRigidParticles->InvI(ParticleIdx) = InInvI; }
 
 	T M() const { return PBDRigidParticles->M(ParticleIdx); }
 	T& M() { return PBDRigidParticles->M(ParticleIdx); }
+	void SetM(const T& InM) { PBDRigidParticles->M(ParticleIdx) = InM; }
 
 	T InvM() const { return PBDRigidParticles->InvM(ParticleIdx); }
 	T& InvM() { return PBDRigidParticles->InvM(ParticleIdx); }
+	void SetInvM(const T& InInvM) { PBDRigidParticles->InvM(ParticleIdx) = InInvM; }
 
 	int32 Island() const { return PBDRigidParticles->Island(ParticleIdx); }
 	int32& Island() { return PBDRigidParticles->Island(ParticleIdx); }
+	void SetIsland(const int32 InIsland) { PBDRigidParticles->Island(ParticleIdx) = InIsland; }
 
 	bool ToBeRemovedOnFracture() const { return PBDRigidParticles->ToBeRemovedOnFracture(ParticleIdx); }
 	bool& ToBeRemovedOnFracture() { return PBDRigidParticles->ToBeRemovedOnFracture(ParticleIdx); }
+	void SetToBeRemovedOnFracture(const bool bToBeRemovedOnFracture) { PBDRigidParticles->ToBeRemovedOnFracture(ParticleIdx) = bToBeRemovedOnFracture; }
 
 	EObjectStateType ObjectState() const { return PBDRigidParticles->ObjectState(ParticleIdx); }
 	void SetObjectState(EObjectStateType InState) { PBDRigidParticles->SetObjectState(ParticleIdx, InState); }
 	void SetSleeping(bool bSleeping) { PBDRigidParticles->SetSleeping(ParticleIdx, bSleeping); }
 
+	//Really only useful when using a transient handle
+	const TPBDRigidParticleHandleImp<T, d, true>* Handle() const { return PBDRigidParticles->Handle(ParticleIdx); }
+	TPBDRigidParticleHandleImp<T, d, true>* Handle() { return PBDRigidParticles->Handle(ParticleIdx); }
+
 	static constexpr EParticleType StaticType() { return EParticleType::Dynamic; }
 };
 
-template <typename T, int d>
-const TKinematicGeometryParticleHandle<T, d>* TGeometryParticleHandle<T,d>::ToKinematic() const { return Type >= EParticleType::Kinematic ? static_cast<const TKinematicGeometryParticleHandle<T, d>*>(this) : nullptr; }
-template <typename T, int d>
-TKinematicGeometryParticleHandle<T, d>& TGeometryParticleHandle<T, d>::ToKinematic() { return Type >= EParticleType::Kinematic ? static_cast<TKinematicGeometryParticleHandle<T, d>*>(this) : nullptr; }
+template <typename T, int d, bool bPersistent>
+class TPBDRigidClusteredParticleHandleImp : public TPBDRigidParticleHandleImp<T, d, bPersistent>
+{
+public:
+	using TGeometryParticleHandleImp<T, d, bPersistent>::ParticleIdx;
+	using TGeometryParticleHandleImp<T, d, bPersistent>::PBDRigidClusteredParticles;
+	using TGeometryParticleHandleImp<T, d, bPersistent>::Type;
+	using TTransientHandle = TTransientPBDRigidParticleHandle<T, d>;
+	using TSOAType = TPBDRigidClusteredParticles<T, d>;
 
-template <typename T, int d>
-const TPBDRigidParticleHandle<T, d>* TGeometryParticleHandle<T, d>::ToDynamic() const { return Type >= EParticleType::Dynamic ? static_cast<const TPBDRigidParticleHandle<T, d>*>(this) : nullptr; }
-template <typename T, int d>
-TPBDRigidParticleHandle<T, d>* TGeometryParticleHandle<T, d>::ToDynamic() { return Type >= EParticleType::Dynamic ? static_cast<TPBDRigidParticleHandle<T, d>*>(this) : nullptr; }
 
-template <typename T, int d>
-EObjectStateType TGeometryParticleHandle<T,d>::ObjectState() const
+	TPBDRigidClusteredParticleHandleImp<T, d, bPersistent>(TPBDRigidClusteredParticles<T, d>* Particles, int32 InIdx, int32 InGlobalIdx, const TPBDRigidParticleParameters<T, d>& Params = TPBDRigidParticleParameters<T, d>())
+		: TPBDRigidParticleHandleImp<T, d, bPersistent>(Particles, InIdx, InGlobalIdx, Params)
+	{
+		PBDRigidClusteredParticleDefaultConstruct<T, d>(*this, Params);
+	}
+
+	const auto& ClusterIds() const { return PBDRigidClusteredParticles->ClusterIds(ParticleIdx); }
+	auto& ClusterIds() { return PBDRigidClusteredParticles->ClusterIds(ParticleIdx); }
+
+	const auto& ChildToParent() const { return PBDRigidClusteredParticles->ChildToParent(ParticleIdx); }
+	auto& ChildToParent() { return PBDRigidClusteredParticles->ChildToParent(ParticleIdx); }
+
+	const auto& ClusterGroupIndex() const { return PBDRigidClusteredParticles->ClusterGroupIndex(ParticleIdx); }
+	auto& ClusterGroupIndex() { return PBDRigidClusteredParticles->ClusterGroupIndex(ParticleIdx); }
+
+	const auto& InternalCluster() const { return PBDRigidClusteredParticles->InternalCluster(ParticleIdx); }
+	auto& InternalCluster() { return PBDRigidClusteredParticles->InternalCluster(ParticleIdx); }
+
+	const auto& ChildrenSpatial() const { return PBDRigidClusteredParticles->ChildrenSpatial(ParticleIdx); }
+	auto& ChildrenSpatial() { return PBDRigidClusteredParticles->ChildrenSpatial(ParticleIdx); }
+
+	//const auto& MultiChildProxyId() const { return PBDRigidClusteredParticles->MultiChildProxyId(ParticleIdx); }
+	//auto& MultiChildProxyId() { return PBDRigidClusteredParticles->MultiChildProxyId(ParticleIdx); }
+
+	//const auto& MultiChildProxyData() const { return PBDRigidClusteredParticles->MultiChildProxyData(ParticleIdx); }
+	//auto& MultiChildProxyData() { return PBDRigidClusteredParticles->MultiChildProxyData(ParticleIdx); }
+
+	const auto& CollisionImpulses() const { return PBDRigidClusteredParticles->CollisionImpulses(ParticleIdx); }
+	auto& CollisionImpulses() { return PBDRigidClusteredParticles->CollisionImpulses(ParticleIdx); }
+
+	const auto& Strains() const { return PBDRigidClusteredParticles->Strains(ParticleIdx); }
+	auto& Strains() { return PBDRigidClusteredParticles->Strains(ParticleIdx); }
+
+	const TPBDRigidClusteredParticleHandleImp<T, d, true>* Handle() const { return PBDRigidClusteredParticles->Handle(ParticleIdx); }
+	TPBDRigidClusteredParticleHandleImp<T, d, true>* Handle() { return PBDRigidClusteredParticles->Handle(ParticleIdx); }
+
+	static constexpr EParticleType StaticType() { return EParticleType::Dynamic; }
+
+	int32 TransientParticleIndex() const { return ParticleIdx; }
+};
+
+template <typename T, int d, bool bPersistent>
+const TKinematicGeometryParticleHandleImp<T, d, bPersistent>* TGeometryParticleHandleImp<T,d, bPersistent>::AsKinematic() const { return Type >= EParticleType::Kinematic ? static_cast<const TKinematicGeometryParticleHandleImp<T, d, bPersistent>*>(this) : nullptr; }
+template <typename T, int d, bool bPersistent>
+TKinematicGeometryParticleHandleImp<T, d, bPersistent>* TGeometryParticleHandleImp<T, d, bPersistent>::AsKinematic() { return Type >= EParticleType::Kinematic ? static_cast<TKinematicGeometryParticleHandleImp<T, d, bPersistent>*>(this) : nullptr; }
+
+template <typename T, int d, bool bPersistent>
+const TPBDRigidParticleHandleImp<T, d, bPersistent>* TGeometryParticleHandleImp<T, d, bPersistent>::AsDynamic() const { return Type >= EParticleType::Dynamic ? static_cast<const TPBDRigidParticleHandleImp<T, d, bPersistent>*>(this) : nullptr; }
+template <typename T, int d, bool bPersistent>
+TPBDRigidParticleHandleImp<T, d, bPersistent>* TGeometryParticleHandleImp<T, d, bPersistent>::AsDynamic() { return Type >= EParticleType::Dynamic ? static_cast<TPBDRigidParticleHandleImp<T, d, bPersistent>*>(this) : nullptr; }
+
+template <typename T, int d, bool bPersistent>
+const TPBDRigidClusteredParticleHandleImp<T, d, bPersistent>* TGeometryParticleHandleImp<T, d, bPersistent>::AsClustered() const { return Type >= EParticleType::Clustered ? static_cast<const TPBDRigidClusteredParticleHandleImp<T, d, bPersistent>*>(this) : nullptr; }
+template <typename T, int d, bool bPersistent>
+TPBDRigidClusteredParticleHandleImp<T, d, bPersistent>* TGeometryParticleHandleImp<T, d, bPersistent>::AsClustered() { return Type >= EParticleType::Clustered ? static_cast<TPBDRigidClusteredParticleHandleImp<T, d, bPersistent>*>(this) : nullptr; }
+
+template <typename T, int d, bool bPersistent>
+EObjectStateType TGeometryParticleHandleImp<T,d, bPersistent>::ObjectState() const
 {
-	const TKinematicGeometryParticleHandle<T, d>* Kin = ToKinematic();
+	const TKinematicGeometryParticleHandleImp<T, d, bPersistent>* Kin = AsKinematic();
 	return Kin ? Kin->ObjectState() : EObjectStateType::Static;
 }
 
-template <typename T, int d>
-EObjectStateType TKinematicGeometryParticleHandle<T, d>::ObjectState() const
+template <typename T, int d, bool bPersistent>
+EObjectStateType TKinematicGeometryParticleHandleImp<T, d, bPersistent>::ObjectState() const
 {
-	const TPBDRigidParticleHandle<T, d>* Dyn = ToDynamic();
+	const TPBDRigidParticleHandleImp<T, d, bPersistent>* Dyn = AsDynamic();
 	return Dyn ? Dyn->ObjectState() : EObjectStateType::Kinematic;
 }
 
+template <typename T, int d, bool bPersistent>
+FString TGeometryParticleHandleImp<T, d, bPersistent>::ToString() const
+{
+	switch (Type)
+	{
+	case EParticleType::Static:
+		return FString::Printf(TEXT("Static[%d]"), ParticleIdx);
+		break;
+	case EParticleType::Kinematic:
+		return FString::Printf(TEXT("Kinemmatic[%d]"), ParticleIdx);
+		break;
+	case EParticleType::Dynamic:
+		return FString::Printf(TEXT("Dynamic[%d]"), ParticleIdx);
+		break;
+	}
+	return FString();
+}
+
+/**
+ * A wrapper around any type of particle handle to provide a consistent (read-only) API for all particle types.
+ * This can make code simpler because you can write code that is type-agnostic, but it
+ * has a cost. Where possible it is better to write code that is specific to the type(s)
+ * of particles being operated on. TGenericParticleHandle has pointer symatics, so you can use one wherever
+ * you have a particle handle pointer;
+ *
+ */
+template <typename T, int d>
+class TGenericParticleHandle
+{
+public:
+	TGenericParticleHandle(TGeometryParticleHandle<T, d>* InHandle) : Imp(InHandle) {}
+
+	class HandleImp
+	{
+	public:
+		HandleImp(TGeometryParticleHandle<T, d>* InHandle) { Handle = InHandle; }
+
+		// Check for the exact type of particle (see also AsKinematic etc, which will work on derived types)
+		bool IsStatic() const { return (Handle->ObjectState() == EObjectStateType::Static); }
+		bool IsKinematic() const { return (Handle->ObjectState() == EObjectStateType::Kinematic); }
+		bool IsDynamic() const { return (Handle->ObjectState() == EObjectStateType::Dynamic) || (Handle->ObjectState() == EObjectStateType::Sleeping); }
+
+		// Static Particles
+		const TVector<T, d>& X() const { return Handle->X(); }
+		const TRotation<T, d>& R() const { return Handle->R(); }
+		TSerializablePtr<TImplicitObject<T, d>> Geometry() const { return Handle->Geometry(); }
+		const TUniquePtr<TImplicitObject<T, d>>& DynamicGeometry() const { return Handle->DynamicGeometry(); }
+
+		// Kinematic Particles
+		const TVector<T, d>& V() const { return (Handle->AsKinematic()) ? Handle->AsKinematic()->V() : ZeroVector; }
+		const TVector<T, d>& W() const { return (Handle->AsKinematic()) ? Handle->AsKinematic()->W() : ZeroVector; }
+
+		// Dynamic Particles
+		int32 CollisionParticlesSize() const { return (Handle->AsDynamic()) ? Handle->AsDynamic()->CollisionParticlesSize() : 0; }
+		const TUniquePtr<TBVHParticles<T, d>>& CollisionParticles() const { return (Handle->AsDynamic()) ? Handle->AsDynamic()->CollisionParticles() : NullBVHParticles; }
+		int32 CollisionGroup() const { return (Handle->AsDynamic()) ? Handle->AsDynamic()->CollisionGroup() : 0; }
+		bool Disabled() const { return (Handle->AsDynamic()) ? Handle->AsDynamic()->Disabled() : false; }				// @todo(ccaulfield): should be available on all types?
+		const TVector<T, d>& PreV() const { return (Handle->AsDynamic()) ? Handle->AsDynamic()->PreV() : ZeroVector; }	// @todo(ccaulfield): should be available on kinematics?
+		const TVector<T, d>& PreW() const { return (Handle->AsDynamic()) ? Handle->AsDynamic()->PreW() : ZeroVector; }	// @todo(ccaulfield): should be available on kinematics?
+		const TVector<T, d>& P() const { return (Handle->AsDynamic()) ? Handle->AsDynamic()->P() : X(); }
+		const TRotation<T, d>& Q() const { return (Handle->AsDynamic()) ? Handle->AsDynamic()->Q() : R(); }
+		const TVector<T, d>& F() const { return (Handle->AsDynamic()) ? Handle->AsDynamic()->F() : ZeroVector; }
+		const TVector<T, d>& Torque() const { return (Handle->AsDynamic()) ? Handle->AsDynamic()->Torque() : ZeroVector; }
+		const PMatrix<T, d, d>& I() const { return (Handle->AsDynamic()) ? Handle->AsDynamic()->I() : ZeroMatrix; }
+		const PMatrix<T, d, d>& InvI() const { return (Handle->AsDynamic()) ? Handle->AsDynamic()->InvI() : ZeroMatrix; }
+		T M() const { return (Handle->AsDynamic()) ? Handle->AsDynamic()->M() : (T)0; }
+		T InvM() const { return (Handle->AsDynamic()) ? Handle->AsDynamic()->InvM() : (T)0; }
+		int32 Island() const { return (Handle->AsDynamic()) ? Handle->AsDynamic()->Island() : INDEX_NONE; }
+		bool ToBeRemovedOnFracture() const { return (Handle->AsDynamic()) ? Handle->AsDynamic()->ToBeRemovedOnFracture() : false; }
+
+	private:
+		TGeometryParticleHandle<T, d>* Handle;
+	};
+
+	HandleImp* operator->() { return &Imp; }
+	const HandleImp* operator->() const { return &Imp; }
+
+private:
+	HandleImp Imp;
+
+	static const TVector<T, d> ZeroVector;
+	static const TRotation<T, d> IdentityRotation;
+	static const PMatrix<T, d, d> ZeroMatrix;
+	static const TUniquePtr<TBVHParticles<T, d>> NullBVHParticles;
+};
+
+template <typename T, int d>
+const TVector<T, d> TGenericParticleHandle<T, d>::ZeroVector = TVector<T, d>(0);
+
+template <typename T, int d>
+const TRotation<T, d> TGenericParticleHandle<T, d>::IdentityRotation = TRotation<T, d>(0, 0, 0, 1);
+
+template <typename T, int d>
+const PMatrix<T, d, d> TGenericParticleHandle<T, d>::ZeroMatrix = PMatrix<T, d, d>(0);
+
+template <typename T, int d>
+const TUniquePtr<TBVHParticles<T, d>> TGenericParticleHandle<T, d>::NullBVHParticles = TUniquePtr<TBVHParticles<T, d>>();
+
 template <typename T, int d>
 class TGeometryParticleHandles : public TArrayCollection
 {
@@ -278,10 +607,25 @@ public:
 		AddElementsHelper(NumHandles);
 	}
 
+	void Reset()
+	{
+		ResizeHelper(0);
+	}
+
+	void DestroyHandleSwap(TGeometryParticleHandle<T,d>* Handle)
+	{
+		const int32 UnstableIdx = Handle->HandleIdx;
+		RemoveAtSwapHelper(UnstableIdx);
+		if (static_cast<uint32>(UnstableIdx) < Size())
+		{
+			Handles[UnstableIdx]->HandleIdx = UnstableIdx;
+		}
+	}
+
 	const TUniquePtr<TGeometryParticleHandle<T, d>>& Handle(int32 Idx) const { return Handles[Idx];}
 	TUniquePtr<TGeometryParticleHandle<T, d>>& Handle(int32 Idx) { return Handles[Idx]; }
 private:
-	TArrayCollectionArray<TUniquePtr<TGeometryParticleHandle<T, d>>> Handles;
+	TArrayCollectionArray<TUniquePtr<TGeometryParticleHandleImp<T, d, true>>> Handles;
 };
 
 template <typename T, int d>
@@ -290,13 +634,43 @@ class TKinematicGeometryParticle;
 template <typename T, int d>
 class TPBDRigidParticle;
 
+
+class FParticleData
+{
+public:
+	FParticleData(EParticleType InType = EParticleType::Static)
+		: Type(InType)
+	{}
+
+	void Reset() { Type = EParticleType::Static; }
+
+	EParticleType Type;
+};
+
+template <typename T, int d> 
+class TGeometryParticleData;
+
+template <typename T, int d>
+class TKinematicGeometryParticleData;
+
+template <typename T, int d>
+class TPBDRigidParticleData;
+
+
 template <typename T, int d>
 class TGeometryParticle
 {
 public:
+	typedef TGeometryParticleData<T, d> FData;
+	typedef TGeometryParticleHandle<T, d> FHandle;
+
+	friend FData;
+
 	TGeometryParticle(const TGeometryParticleParameters<T, d>& StaticParams = TGeometryParticleParameters<T, d>())
+		: MUserData(nullptr)
 	{
 		Type = EParticleType::Static;
+		Proxy = nullptr;
 		GeometryParticleDefaultConstruct<T, d>(*this, StaticParams);
 	}
 
@@ -304,35 +678,169 @@ public:
 
 	TGeometryParticle(const TGeometryParticle&) = delete;
 
+	TGeometryParticle& operator=(const TGeometryParticle&) = delete;
+
 	const TVector<T, d>& X() const { return this->MX; }
-	TVector<T, d>& X() { return this->MX; }
+	void SetX(const TVector<T, d>& InX)
+	{
+		this->MarkDirty(EParticleFlags::X);
+		this->MX = InX;
+	}
 
 	const TRotation<T, d>& R() const { return this->MR; }
-	TRotation<T, d>& R() { return this->MR; }
+	void SetR(const TRotation<T, d>& InR)
+	{
+		this->MarkDirty(EParticleFlags::R);
+		this->MR = InR;
+	}
+
+	//todo: geometry should not be owned by particle
+	void SetGeometry(TUniquePtr<TImplicitObject<T, d>>&& UniqueGeometry)
+	{
+		// Take ownership of the geometry, putting it into a shared ptr.
+		// This is necessary because we cannot be sure whether the particle
+		// will be destroyed on the game thread or physics thread first,
+		// but geometry data is shared between them.
+		TImplicitObject<T, d>* RawGeometry = UniqueGeometry.Release();
+		SetGeometry(TSharedPtr<TImplicitObject<T, d>, ESPMode::ThreadSafe>(RawGeometry));
+	}
+
+	// TODO: Right now this method exists so we can do things like FPhysTestSerializer::CreateChaosData.
+	//       We should replace this with a method for supporting SetGeometry(RawGeometry).
+	void SetGeometry(TSharedPtr<TImplicitObject<T, d>, ESPMode::ThreadSafe> SharedGeometry)
+	{
+		this->MarkDirty(EParticleFlags::Geometry);
+		this->MGeometry = SharedGeometry;
+		UpdateShapesArray();
+	}
 
-	TSerializablePtr<TImplicitObject<T, d>> Geometry() const { return this->MGeometry; }
+	void SetGeometry(TSerializablePtr<TImplicitObject<T, d>> RawGeometry)
+	{
+		// Ultimately this method should replace SetGeometry(SharedPtr).
+		// We don't really want people making shared ptrs to geometry everywhere.
+		check(false);
+	}
+
+	void* UserData() const { return this->MUserData; }
+	void SetUserData(void* InUserData)
+	{
+		this->MUserData = InUserData;
+	}
+
+	//Note: this must be called after setting geometry. This API seems bad. Should probably be part of setting geometry
+	void SetShapesArray(TShapesArray<T, d>&& InShapesArray)
+	{
+		ensure(InShapesArray.Num() == MShapesArray.Num());
+		MShapesArray = MoveTemp(InShapesArray);
+	}
+
+	TSerializablePtr<TImplicitObject<T, d>> Geometry() const { return MakeSerializable(MGeometry); }
+
+	const TShapesArray<T,d>& ShapesArray() const { return MShapesArray; }
 
 	EObjectStateType ObjectState() const;
 
-	const TKinematicGeometryParticle<T, d>* ToKinematic() const;
-	TKinematicGeometryParticle<T, d>& ToKinematic();
+	EParticleType ObjectType() const
+	{
+		return Type;
+	}
+
+	const TKinematicGeometryParticle<T, d>* AsKinematic() const;
+	TKinematicGeometryParticle<T, d>* AsKinematic();
 
-	const TPBDRigidParticle<T, d>* ToDynamic() const;
-	TPBDRigidParticle<T, d>* ToDynamic();
+	const TPBDRigidParticle<T, d>* AsDynamic() const;
+	TPBDRigidParticle<T, d>* AsDynamic();
+
+	FParticleData* NewData() const { return new TGeometryParticleData<T, d>( *this ); }
+
+	bool IsDirty() const
+	{
+		return this->MDirtyFlags.IsDirty();
+	}
+
+	bool IsDirty(const EParticleFlags CheckBits) const
+	{
+		return this->MDirtyFlags.IsDirty(CheckBits);
+	}
+
+	// Pointer to any data that the solver wants to associate with this particle
+	// TODO: It's important to eventually hide this!
+	// Right now it's exposed to lubricate the creation of the whole proxy system.
+	class IPhysicsProxyBase* Proxy;
 
 private:
 	TVector<T, d> MX;
 	TRotation<T, d> MR;
+	TSharedPtr<TImplicitObject<T, d>, ESPMode::ThreadSafe> MGeometry;	//TODO: geometry should live in bodysetup
+	TShapesArray<T,d> MShapesArray;
+
+	// Pointer to some arbitrary data associated with the particle, but not used by Chaos. External systems may use this for whatever.
+	void* MUserData;
+
+	// This is only for use by ParticleData. This should be called only in one place,
+	// when the geometry is being copied from GT to PT.
+	TSharedPtr<TImplicitObject<T, d>, ESPMode::ThreadSafe> GeometrySharedLowLevel() const
+	{
+		return MGeometry;
+	}
+
 protected:
 	EParticleType Type;
+	FParticleDirtyFlags MDirtyFlags;
+
+	void MarkDirty(const EParticleFlags DirtyBits)
+	{
+		this->MDirtyFlags.MarkDirty(DirtyBits);
+	}
+
+	void UpdateShapesArray()
+	{
+		UpdateShapesArrayFromGeometry<T, d>(MShapesArray, MGeometry.Get());
+	}
+};
+
+template <typename T, int d>
+class TGeometryParticleData : public FParticleData
+{
+	typedef FParticleData Base;
+public:
+
+	TGeometryParticleData(EParticleType InType = EParticleType::Static)
+		: FParticleData(InType)
+		, X(TVector<T, d>(0))
+		, R(TRotation<T, d>())
+	{}
+
+	TGeometryParticleData(const TGeometryParticle<T,d>& InParticle)
+		: FParticleData(EParticleType::Static)
+		, X(InParticle.X())
+		, R(InParticle.R())
+		, Geometry(InParticle.GeometrySharedLowLevel())
+	{}
+
+	void Reset() 
+	{ 
+		FParticleData::Reset();  
+		X = TVector<T, d>(0); 
+		R = TRotation<T, d>(); 
+		Geometry = TSharedPtr<TImplicitObjectUnion<T, d>, ESPMode::ThreadSafe>();
+	}
+
+	TVector<T, d> X;
+	TRotation<T, d> R;
+	TSharedPtr<TImplicitObject<T, d>, ESPMode::ThreadSafe> Geometry;
 };
 
+
 template <typename T, int d>
 class TKinematicGeometryParticle : public TGeometryParticle<T, d>
 {
 public:
+	typedef TKinematicGeometryParticleData<T, d> FData;
+	typedef TKinematicGeometryParticleHandle<T, d> FHandle;
+
 	using TGeometryParticle<T, d>::Type;
-	using TGeometryParticle<T, d>::ToDynamic;
+	using TGeometryParticle<T, d>::AsDynamic;
 
 	TKinematicGeometryParticle(const TKinematicGeometryParticleParameters<T, d>& KinematicParams = TKinematicGeometryParticleParameters<T,d>())
 		: TGeometryParticle<T, d>(KinematicParams)
@@ -342,79 +850,199 @@ public:
 	}
 
 	const TVector<T, d>& V() const { return MV; }
-	TVector<T, d>& V() { return MV; }
+	void SetV(const TVector<T, d>& InV)
+	{
+		this->MarkDirty(EParticleFlags::V);
+		this->MV = InV;
+	}
 
 	const TVector<T, d>& W() const { return MW; }
-	TVector<T, d>& W() { return MW; }
+	void SetW(const TVector<T, d>& InW)
+	{
+		this->MarkDirty(EParticleFlags::W);
+		this->MW = InW;
+	}
 
 	EObjectStateType ObjectState() const;
 
+	FParticleData* NewData() const
+	{
+		return new TKinematicGeometryParticleData<T, d>(*this);
+	}
+
 private:
 	TVector<T, d> MV;
 	TVector<T, d> MW;
 };
 
+template <typename T, int d>
+class TKinematicGeometryParticleData : public TGeometryParticleData<T, d>
+{
+	typedef TGeometryParticleData<T, d> Base;
+public:
+
+	using TGeometryParticleData<T, d>::Type;
+
+	TKinematicGeometryParticleData(EParticleType InType = EParticleType::Kinematic)
+		: Base(InType)
+		, MV(TVector<T, d>(0))
+		, MW(TVector<T, d>(0)) {}
+
+	TKinematicGeometryParticleData(const TKinematicGeometryParticle<T, d>& InParticle)
+		: Base(InParticle)
+		, MV(InParticle.V())
+		, MW(InParticle.W()) 
+	{
+		Type = EParticleType::Kinematic;
+	}
+
+
+	void Reset() {
+		TGeometryParticleData<T, d>::Reset();
+		Type = EParticleType::Kinematic;
+		MV = TVector<T, d>(0); MW = TVector<T, d>();
+	}
+
+	TVector<T, d> MV;
+	TVector<T, d> MW;
+};
+
+
 template <typename T, int d>
 class TPBDRigidParticle : public TKinematicGeometryParticle<T, d>
 {
 public:
+	typedef TPBDRigidParticleData<T, d> FData;
+	typedef TPBDRigidParticleHandle<T, d> FHandle;
+
 	using TGeometryParticle<T, d>::Type;
 
 	TPBDRigidParticle<T, d>(const TPBDRigidParticleParameters<T, d>& DynamicParams = TPBDRigidParticleParameters<T, d>())
 		: TKinematicGeometryParticle<T, d>(DynamicParams)
 	{
 		Type = EParticleType::Dynamic;
+		MIsland = INDEX_NONE;
+		MToBeRemovedOnFracture = false;
 		PBDRigidParticleDefaultConstruct<T, d>(*this, DynamicParams);
 	}
 
 	const TUniquePtr<TBVHParticles<T, d>>& CollisionParticles() const { return MCollisionParticles; }
 
 	int32 CollisionGroup() const { return MCollisionGroup; }
-	int32& CollisionGroup() { return MCollisionGroup; }
+	void SetCollisionGroup(const int32 InCollisionGroup)
+	{
+		this->MarkDirty(EParticleFlags::CollisionGroup);
+		this->MCollisionGroup = InCollisionGroup;
+	}
 
 	bool Disabled() const { return MDisabled; }
+	void SetDisabled(const bool InDisabled)
+	{
+		this->MarkDirty(EParticleFlags::CollisionGroup);
+		this->MDisabled = InDisabled;
+	}
 
 	// Named to match signature of TPBDRigidParticleHandle, as both are used in templated functions.
 	// See its comment for details.
 	bool& SetDisabledLowLevel() { return MDisabled; }
 
 	const TVector<T, d>& PreV() const { return MPreV; }
-	TVector<T, d>& PreV() { return MPreV; }
+	void SetPreV(const TVector<T, d>& InPreV)
+	{
+		this->MarkDirty(EParticleFlags::PreV);
+		this->MPreV = InPreV;
+	}
+
 	const TVector<T, d>& PreW() const { return MPreW; }
-	TVector<T, d>& PreW() { return MPreW; }
+	void SetPreW(const TVector<T, d>& InPreW)
+	{
+		this->MarkDirty(EParticleFlags::PreW);
+		this->MPreW = InPreW;
+	}
 
 	const TVector<T, d>& P() const { return MP; }
-	TVector<T, d>& P() { return MP; }
+	void SetP(const TVector<T, d>& InP)
+	{
+		this->MarkDirty(EParticleFlags::P);
+		this->MP = InP;
+	}
 
-	const TVector<T, d>& Q() const { return MQ; }
-	TRotation<T, d>& Q() { return MQ; }
+	const TRotation<T, d>& Q() const { return MQ; }
+	void SetQ(const TRotation<T, d>& InQ)
+	{
+		this->MarkDirty(EParticleFlags::Q);
+		this->MQ = InQ;
+	}
 
 	const TVector<T, d>& F() const { return MF; }
-	TVector<T, d>& F() { return MF; }
+	void SetF(const TVector<T, d>& InF)
+	{
+		this->MarkDirty(EParticleFlags::F);
+		this->MF = InF;
+	}
 
 	const TVector<T, d>& Torque() const { return MTorque; }
-	TVector<T, d>& Torque() { return MTorque; }
+	void SetTorque(const TVector<T, d>& InTorque)
+	{
+		this->MarkDirty(EParticleFlags::Torque);
+		this->MTorque = InTorque;
+	}
 
 	const PMatrix<T, d, d>& I() const { return MI; }
-	PMatrix<T, d, d>& I() { return MI; }
+	void SetI(const PMatrix<T, d, d>& InI)
+	{
+		this->MarkDirty(EParticleFlags::I);
+		this->MI = InI;
+	}
 
 	const PMatrix<T, d, d>& InvI() const { return MInvI; }
-	PMatrix<T, d, d>& InvI() { return MInvI; }
+	void SetInvI(const PMatrix<T, d, d>& InInvI)
+	{
+		this->MarkDirty(EParticleFlags::InvI);
+		this->MInvI = InInvI;
+	}
 
 	T M() const { return MM; }
-	T& M() { return MM; }
+	void SetM(const T& InM)
+	{
+		this->MarkDirty(EParticleFlags::M);
+		this->MM = InM;
+	}
 
 	T InvM() const { return MInvM; }
-	T& InvM() { return MInvM; }
+	void SetInvM(const T& InInvM)
+	{
+		this->MarkDirty(EParticleFlags::InvM);
+		this->MInvM = InInvM;
+	}
 
 	int32 Island() const { return MIsland; }
-	int32& Island() { return MIsland; }
+	// TODO(stett): Make the setter private. It is public right now to provide access to proxies.
+	void SetIsland(const int32 InIsland)
+	{
+		this->MIsland = InIsland;
+	}
 
 	bool ToBeRemovedOnFracture() const { return MToBeRemovedOnFracture; }
-	bool& ToBeRemovedOnFracture() { return MToBeRemovedOnFracture; }
+	// TODO(stett): Make the setter private. It is public right now to provide access to proxies.
+	void SetToBeRemovedOnFracture(const bool bToBeRemovedOnFracture)
+	{
+		this->MToBeRemovedOnFracture = bToBeRemovedOnFracture;
+	}
 
 	EObjectStateType ObjectState() const { return MObjectState; }
-	void SetObjectState(EObjectStateType InState) { MObjectState = InState; }  //todo: look at physics thread logic 
+	void SetObjectState(EObjectStateType InState)
+	{
+		//todo: look at physics thread logic
+		MObjectState = InState;
+		this->MarkDirty(EParticleFlags::ObjectState);
+	}
+
+	FParticleData* NewData() const
+	{
+		return new TPBDRigidParticleData<T, d>(*this);
+	}
+
 
 private:
 	TRotation<T, d> MQ;
@@ -433,162 +1061,122 @@ private:
 	EObjectStateType MObjectState;
 	bool MDisabled;
 	bool MToBeRemovedOnFracture;
-};
-
-template <typename T, int d>
-const TKinematicGeometryParticle<T, d>* TGeometryParticle<T, d>::ToKinematic() const { return Type >= EParticleType::Kinematic ? static_cast<const TKinematicGeometryParticle<T, d>*>(this) : nullptr; }
-template <typename T, int d>
-TKinematicGeometryParticle<T, d>& TGeometryParticle<T, d>::ToKinematic() { return Type >= EParticleType::Kinematic ? static_cast<TKinematicGeometryParticle<T, d>*>(this) : nullptr; }
 
-template <typename T, int d>
-const TPBDRigidParticle<T, d>* TGeometryParticle<T, d>::ToDynamic() const { return Type >= EParticleType::Dynamic ? static_cast<const TPBDRigidParticle<T, d>*>(this) : nullptr; }
-template <typename T, int d>
-TPBDRigidParticle<T, d>* TGeometryParticle<T, d>::ToDynamic() { return Type >= EParticleType::Dynamic ? static_cast<TPBDRigidParticle<T, d>*>(this) : nullptr; }
-
-template <typename T, int d>
-EObjectStateType TGeometryParticle<T, d>::ObjectState() const
-{
-	const TKinematicGeometryParticle<T, d>* Kin = ToKinematic();
-	return Kin ? Kin->ObjectState() : EObjectStateType::Static;
-}
+};
 
 template <typename T, int d>
-EObjectStateType TKinematicGeometryParticle<T, d>::ObjectState() const
-{
-	const TPBDRigidParticle<T, d>* Dyn = ToDynamic();
-	return Dyn ? Dyn->ObjectState() : EObjectStateType::Kinematic;
-}
-
-template <typename T, int d, typename TSOA>
-class TConstParticleIterator
+class TPBDRigidParticleData : public TKinematicGeometryParticleData<T, d>
 {
+	typedef TKinematicGeometryParticleData<T, d> Base;
 public:
-	TConstParticleIterator()
-		: ParticleIdx(0)
-		, SOAIdx(0)
-		, CurSOA(nullptr)
-	{
-	}
-	
-	TConstParticleIterator(TArray<TSOA*>&& InSOAs)
-		: SOAs(MoveTemp(InSOAs))
-		, ParticleIdx(0)
-		, SOAIdx(0)
-		, CurSOA(SOAs.Num() ? SOAs[0] : nullptr)
-	{
-	}
-	
-	operator bool() const { return CurSOA && ParticleIdx < CurSOA->Size(); }
-	TConstParticleIterator<T, d, TSOA>& operator++()
-	{
-		++ParticleIdx;
-		if (ParticleIdx >= CurSOA->Size())
-		{
-			++SOAIdx;
-			CurSOA = SOAIdx < SOAs.Num() ? SOAs[SOAIdx] : nullptr;
-			ParticleIdx = 0;
-		}
-		return *this;
-	}
-
-	const typename TSOA::THandleType* Handle() const { return CurSOA->Handle(ParticleIdx); }
-
-	//Full PBDRigid API, we rely on templates only instantiating the functions they use to avoid compiler errors on more base particle types
-	const TVector<T, d>& X() const { return CurSOA->X(ParticleIdx); }
-	const TRotation<T, d>& R() const { return CurSOA->R(ParticleIdx); }
-	TSerializablePtr<TImplicitObject<T, d>> Geometry() const { return CurSOA->Geometry(ParticleIdx); }
-
-	const TVector<T, d>& V() const { return CurSOA->V(ParticleIdx); }
-	const TVector<T, d>& W() const { return CurSOA->W(ParticleIdx); }
-
-	const TUniquePtr<TBVHParticles<T, d>>& CollisionParticles() const { return CurSOA->CollisionParticles(ParticleIdx); }
-	int32 CollisionGroup() const { return CurSOA->CollisionGroup(ParticleIdx); }
-	bool Disabled() const { return CurSOA->Disabled(ParticleIdx); }
-	const TVector<T, d>& PreW() const { return CurSOA->PreW(ParticleIdx); }
-	const TVector<T, d>& P() const { return CurSOA->P(ParticleIdx); }
-	const TRotation<T, d>& Q() const { return CurSOA->Q(ParticleIdx); }
-	const TVector<T, d>& F() const { return CurSOA->F(ParticleIdx); }
-	const TVector<T, d>& Torque() const { return CurSOA->Torque(ParticleIdx); }
-	const PMatrix<T, d, d>& I() const { return CurSOA->I(ParticleIdx); }
-	const PMatrix<T, d, d>& InvI() const { return CurSOA->InvI(ParticleIdx); }
-	T M() const { return CurSOA->M(ParticleIdx); }
-	T InvM() const { return CurSOA->InvM(ParticleIdx); }
-	int32 Island() const { return CurSOA->Island(ParticleIdx); }
-	bool ToBeRemovedOnFracture() const { return CurSOA->ToBeRemovedOnFracture(ParticleIdx); }
-
-	//Note: this won't compile for anything but dynamic even though static handles support it. This should not be necessary even at the handle level
-	EObjectStateType ObjectState() const { return CurSOA->ObjectState(ParticleIdx); }
 
-protected:
-	TArray<TSOA*> SOAs;
-	uint32 ParticleIdx;
-	int32 SOAIdx;
-	union
-	{
-		const TSOA* CurSOA;
-		TSOA* NonConstCurSOA;	//only here for the non const iterator
-	};
-};
-
-template <typename T, int d, typename TSOA>
-class TParticleIterator : public TConstParticleIterator<T,d,TSOA>
-{
-public:
-	using Base = TConstParticleIterator<T, d, TSOA>;
-	using Base::NonConstCurSOA;
-	using Base::ParticleIdx;
+	using TKinematicGeometryParticleData<T, d>::Type;
+
+	TPBDRigidParticleData(EParticleType InType = EParticleType::Dynamic)
+		: Base(InType)
+		, MQ(TRotation<T, d>())
+		, MPreV(TVector<T, d>(0))
+		, MPreW(TVector<T, d>(0))
+		, MP(TVector<T, d>(0))
+		, MF(TVector<T, d>(0))
+		, MTorque(TVector<T, d>(0))
+		, MI(PMatrix<T, d, d>(0))
+		, MInvI(PMatrix<T, d, d>(0))
+		, MCollisionParticles(nullptr)
+		, MM(T(0))
+		, MInvM(T(0))
+		, MIsland(INDEX_NONE)
+		, MCollisionGroup(0)
+		, MObjectState(EObjectStateType::Uninitialized)
+		, MDisabled(false)
+		, MToBeRemovedOnFracture(false)
+	{}
 
-	TParticleIterator(TSOA* InSOA = nullptr)
-		: Base(InSOA)
-	{
-	}
-	TParticleIterator(TArray<TSOA*>&& InSOAs)
-		: Base(MoveTemp(InSOAs))
+	TPBDRigidParticleData(const TPBDRigidParticle<T, d>& InParticle)
+		: Base(InParticle)
+		, MQ(InParticle.Q())
+		, MPreV(InParticle.PreV())
+		, MPreW(InParticle.PreW())
+		, MP(InParticle.P())
+		, MF(InParticle.F())
+		, MTorque(InParticle.Torque())
+		, MI(InParticle.I())
+		, MInvI(InParticle.InvI())
+		, MCollisionParticles(nullptr)
+		, MM(InParticle.M())
+		, MInvM(InParticle.InvM())
+		, MIsland(InParticle.Island())
+		, MCollisionGroup(InParticle.CollisionGroup())
+		, MObjectState(InParticle.ObjectState())
+		, MDisabled(InParticle.Disabled())
+		, MToBeRemovedOnFracture(InParticle.ToBeRemovedOnFracture())
 	{
+		Type = EParticleType::Dynamic;
 	}
 
 
-	typename TSOA::THandleType* Handle() { return NonConstCurSOA->Handle(ParticleIdx); }
+	TRotation<T, d> MQ;
+	TVector<T, d> MPreV;
+	TVector<T, d> MPreW;
+	TVector<T, d> MP;
+	TVector<T, d> MF;
+	TVector<T, d> MTorque;
+	PMatrix<T, d, d> MI;
+	PMatrix<T, d, d> MInvI;
+	const TBVHParticles<T, d> * MCollisionParticles;
+	T MM;
+	T MInvM;
+	int32 MIsland;
+	int32 MCollisionGroup;
+	EObjectStateType MObjectState;
+	bool MDisabled;
+	bool MToBeRemovedOnFracture;
 
-	//Full PBDRigid API, we rely on templates only instantiating the functions they use to avoid compiler errors on more base particle types
-	TVector<T, d>& X() { return NonConstCurSOA->X(ParticleIdx); }
-	TRotation<T, d>& R() { return NonConstCurSOA->R(ParticleIdx); }
+	void Reset() {
+		TKinematicGeometryParticleData<T, d>::Reset();
+		Type = EParticleType::Dynamic;
+		MQ = TRotation<T, d>();
+		MPreV = TVector<T, d>(0);
+		MPreW = TVector<T, d>(0);
+		MP = TVector<T, d>(0);
+		MF = TVector<T, d>(0);
+		MTorque = TVector<T, d>(0);
+		MI = PMatrix<T, d, d>(0);
+		MInvI = PMatrix<T, d, d>(0);
+		MCollisionParticles = nullptr;
+		MM = T(0);
+		MInvM = T(0);
+		MIsland = INDEX_NONE;
+		MCollisionGroup = 0;
+		MObjectState = EObjectStateType::Uninitialized;
+		MDisabled = false;
+		MToBeRemovedOnFracture = false;
+	}
+};
 
-	TVector<T, d>& V() { return NonConstCurSOA->V(ParticleIdx); }
-	TVector<T, d>& W() { return NonConstCurSOA->W(ParticleIdx); }
 
-	int32& CollisionGroup() { return NonConstCurSOA->CollisionGroup(ParticleIdx); }
+template <typename T, int d>
+const TKinematicGeometryParticle<T, d>* TGeometryParticle<T, d>::AsKinematic() const { return Type >= EParticleType::Kinematic ? static_cast<const TKinematicGeometryParticle<T, d>*>(this) : nullptr; }
+template <typename T, int d>
+TKinematicGeometryParticle<T, d>* TGeometryParticle<T, d>::AsKinematic() { return Type >= EParticleType::Kinematic ? static_cast<TKinematicGeometryParticle<T, d>*>(this) : nullptr; }
 
-	// See Comment on TRigidParticle::SetDisabledLowLevel. State changes in Evolution should accompany this call.
-	bool& SetDisabledLowLevel() { return NonConstCurSOA->Disabled(ParticleIdx); }
-
-	TVector<T, d>& PreV() { return NonConstCurSOA->PreV(ParticleIdx); }
-	TVector<T, d>& PreW() { return NonConstCurSOA->PreW(ParticleIdx); }
-	TVector<T, d>& P() { return NonConstCurSOA->P(ParticleIdx); }
-	TRotation<T, d>& Q() { return NonConstCurSOA->Q(ParticleIdx); }
-	TVector<T, d>& F() { return NonConstCurSOA->F(ParticleIdx); }
-	TVector<T, d>& Torque() { return NonConstCurSOA->Torque(ParticleIdx); }
-	PMatrix<T, d, d>& I() { return NonConstCurSOA->I(ParticleIdx); }
-	PMatrix<T, d, d>& InvI() { return NonConstCurSOA->InvI(ParticleIdx); }
-	T& M() { return NonConstCurSOA->M(ParticleIdx); }
-	T& InvM() { return NonConstCurSOA->InvM(ParticleIdx); }
-	int32& Island() { return NonConstCurSOA->Island(ParticleIdx); }
-	bool& ToBeRemovedOnFracture() { return NonConstCurSOA->ToBeRemovedOnFracture(ParticleIdx); }
-	void SetObjectState(EObjectStateType InState) { NonConstCurSOA->SetObjectState(ParticleIdx, InState); }
-	void SetSleeping(bool bSleeping) { NonConstCurSOA->SetSleeping(ParticleIdx, bSleeping); }
-};
+template <typename T, int d>
+const TPBDRigidParticle<T, d>* TGeometryParticle<T, d>::AsDynamic() const { return Type >= EParticleType::Dynamic ? static_cast<const TPBDRigidParticle<T, d>*>(this) : nullptr; }
+template <typename T, int d>
+TPBDRigidParticle<T, d>* TGeometryParticle<T, d>::AsDynamic() { return Type >= EParticleType::Dynamic ? static_cast<TPBDRigidParticle<T, d>*>(this) : nullptr; }
 
-template <typename T, int d, typename TSOA>
-TParticleIterator<T, d, TSOA> MakeParticleIterator(TArray<TSOA*>&& SOAs)
+template <typename T, int d>
+EObjectStateType TGeometryParticle<T, d>::ObjectState() const
 {
-	return TParticleIterator<T, d, TSOA>(MoveTemp(SOAs));
+	const TKinematicGeometryParticle<T, d>* Kin = AsKinematic();
+	return Kin ? Kin->ObjectState() : EObjectStateType::Static;
 }
 
-template <typename T, int d, typename TSOA>
-TConstParticleIterator<T, d, TSOA> MakeConstParticleIterator(TArray<TSOA*>&& SOAs)
+template <typename T, int d>
+EObjectStateType TKinematicGeometryParticle<T, d>::ObjectState() const
 {
-	return TConstParticleIterator<T, d, TSOA>(MoveTemp(SOAs));
+	const TPBDRigidParticle<T, d>* Dyn = AsDynamic();
+	return Dyn ? Dyn->ObjectState() : EObjectStateType::Kinematic;
 }
 
-
-}
\ No newline at end of file
+}
diff --git a/Engine/Source/Runtime/Experimental/Chaos/Public/Chaos/Particles.h b/Engine/Source/Runtime/Experimental/Chaos/Public/Chaos/Particles.h
--- a/Engine/Source/Runtime/Experimental/Chaos/Public/Chaos/Particles.h
+++ b/Engine/Source/Runtime/Experimental/Chaos/Public/Chaos/Particles.h
@@ -6,166 +6,140 @@
 #include "Chaos/Vector.h"
 #include "ChaosArchive.h"
 
+#if !(UE_BUILD_SHIPPING || UE_BUILD_TEST)
+#define PARTICLE_ITERATOR_RANGED_FOR_CHECK 1
+#else
+#define PARTICLE_ITERATOR_RANGED_FOR_CHECK 0
+#endif
+
 namespace Chaos
 {
 	template<class T, int d>
 	class TPBDRigidsEvolution;
 
+	enum class ERemoveParticleBehavior : uint8
+	{
+		RemoveAtSwap,	//O(1) but reorders particles relative to one another
+		Remove			//Keeps particles relative to one another, but O(n)
+	};
+
 	template<class T, int d>
 	class TParticles : public TArrayCollection
 	{
 	public:
 		TParticles()
-		    : MXView(MX.GetData(), MX.Num())
+			: MRemoveParticleBehavior(ERemoveParticleBehavior::RemoveAtSwap)
 		{
 			AddArray(&MX);
-			MXView = TArrayCollectionArrayView<TVector<T, d>>(MX.GetData(), MX.Num());
+#if PARTICLE_ITERATOR_RANGED_FOR_CHECK
+			MDirtyValidationCount = 0;
+#endif
 		}
 		TParticles(const TParticles<T, d>& Other) = delete;
 		TParticles(TParticles<T, d>&& Other)
-		    : TArrayCollection(), MX(MoveTemp(Other.MX)), MXView(MoveTemp(Other.MXView))
+		    : TArrayCollection(), MX(MoveTemp(Other.MX))
+			, MRemoveParticleBehavior(Other.MRemoveParticleBehavior)
 		{
 			AddParticles(Other.Size());
 			AddArray(&MX);
 			Other.MSize = 0;
+#if PARTICLE_ITERATOR_RANGED_FOR_CHECK
+			MDirtyValidationCount = 0;
+#endif
 		}
 
-		/**
-		 * Constructor that replaces the positions array with a view of @param Points.
-		 */
-		TParticles(TArray<TVector<T, d>>& Points)
-		    : TArrayCollection()
-		    , MXView(Points.GetData(), Points.Num())
+		TParticles(TArray<TVector<T, d>>&& Positions)
+			: TArrayCollection(), MX(MoveTemp(Positions))
+			, MRemoveParticleBehavior(ERemoveParticleBehavior::RemoveAtSwap)
 		{
-			AddElementsHelper(MXView.Num());
-			AddArray(&MXView);
-		}
-		TParticles(TVector<T, d>* Data, const int32 Num)
-		    : TArrayCollection()
-		    , MXView(Data, Num)
-		{
-			AddElementsHelper(MXView.Num());
-			AddArray(&MXView);
+			AddParticles(MX.Num());
+			AddArray(&MX);
+
+#if PARTICLE_ITERATOR_RANGED_FOR_CHECK
+			MDirtyValidationCount = 0;
+#endif
 		}
 
 		~TParticles()
 		{}
 
-		/**
-		 * Add (or remove) particles.  If this class is currently using a non-local view for 
-		 * positions, they are copied into a local positions array and resized.
-		 */
-		void AddElements(const int32 Num)
+		void AddParticles(const int32 Num)
 		{
-			if (Num == 0) // Can be negative
-			{
-				return;
-			}
-			const bool UsingLocalXArray = IsUsingLocalXArray();
-			if (!UsingLocalXArray)
+			AddElementsHelper(Num);
+			IncrementDirtyValidation();
+		}
+
+		void DestroyParticle(const int32 Idx)
+		{
+			if (MRemoveParticleBehavior == ERemoveParticleBehavior::RemoveAtSwap)
 			{
-				// Don't add MX yet, so we avoid an extra memory allocation
-				RemoveArray(&MXView);
+				RemoveAtSwapHelper(Idx);
 			}
-			// Add elements, which may allocate new memory
-			AddElementsHelper(Num);
-			if (!UsingLocalXArray)
+			else
 			{
-				AddArray(&MX); // Reuses MXView's position in the TArrayCollection
-				for (int32 It = 0, EndIt = FMath::Min(MX.Num(), MXView.Num()); It < EndIt; ++It)
-				{
-					MX[It] = MXView[It];
-				}
+				RemoveAtHelper(Idx, 1);
 			}
-			// MX's base pointer may have changed
-			MXView = TArrayCollectionArrayView<TVector<T, d>>(MX.GetData(), MX.Num());
-		}
-		void AddParticles(const int32 Num)
-		{
-			AddElements(Num);
+			IncrementDirtyValidation();
 		}
 
-		void Resize(const int32 Num)
+		ERemoveParticleBehavior RemoveParticleBehavior() const { return MRemoveParticleBehavior; }
+		ERemoveParticleBehavior& RemoveParticleBehavior() { return MRemoveParticleBehavior; }
+
+		void MoveToOtherParticles(const int32 Idx, TParticles<T,d>& Other)
 		{
-			AddElements(Num - Size());
+			MoveToOtherArrayCollection(Idx, Other);
+			IncrementDirtyValidation();
 		}
 
-		void RemoveAt(const int32 Index, const int32 Count)
+		void Resize(const int32 Num)
 		{
-			if (Count <= 0 || Index < 0 || Index >= static_cast<int32>(Size()))
-			{
-				return;
-			}
-			const bool UsingLocalXArray = IsUsingLocalXArray();
-			if (!UsingLocalXArray)
-			{
-				RemoveArray(&MXView);
-			}
-			RemoveAtHelper(Index, Count);
-			if (!UsingLocalXArray)
-			{
-				AddArray(&MX);
-				int32 It = 0;
-				for (int32 EndIt = FMath::Min(Index, MXView.Num()); It < EndIt; ++It)
-				{
-					MX[It] = MXView[It];
-				}
-				for (int32 EndIt=MXView.Num(); It < EndIt; ++It)
-				{
-					MX[It] = MXView[It + Count];
-				}
-			}
-			MXView = TArrayCollectionArrayView<TVector<T, d>>(MX.GetData(), MX.Num());
+			AddParticles(Num - Size());
+			IncrementDirtyValidation();
 		}
 
 		TParticles& operator=(TParticles<T, d>&& Other)
 		{
 			MX = MoveTemp(Other.MX);
-			MXView = MoveTemp(Other.MXView);
 			AddParticles(Other.Size());
 			Other.MSize = 0;
+
+#if PARTICLE_ITERATOR_RANGED_FOR_CHECK
+			MDirtyValidationCount = 0;
+			++Other.MDirtyValidationCount;
+#endif
 			return *this;
 		}
 
-		const TArrayCollectionArrayView<TVector<T, d>>& X() const
+		inline const TArrayCollectionArray<TVector<T, d>>& X() const
 		{
-			return MXView;
+			return MX;
 		}
 		
-		void Serialize(FChaosArchive& Ar)
+		void Serialize(FArchive& Ar)
 		{
-			bool bSerialize = IsUsingLocalXArray();
+			bool bSerialize = true;	//leftover from when we had view support
 			Ar << bSerialize;
-			if (ensureMsgf(bSerialize, TEXT("Cannot serialize shared views. Refactor needed to reduce memory")))	//todo(ocohen): give better way to re-use simplicial
+			if (ensureMsgf(bSerialize, TEXT("Cannot serialize shared views. Refactor needed to reduce memory")))
 			{
 				Ar << MX;
 				ResizeHelper(MX.Num());
-				MXView = TArrayCollectionArrayView<TVector<T, d>>(MX.GetData(), MX.Num());
 			}
+			IncrementDirtyValidation();
 		}
+
 		const TArrayCollectionArray<TVector<T, d>>& XArray() const
 		{
 			return MX;
 		}
-		const TArrayCollectionArrayView<TVector<T, d>>& XView() const
-		{
-			return MXView;
-		}
-		const bool IsUsingLocalXArray() const
-		{
-			return MXView.GetData() == MX.GetData();
-		}
-
+		
 		const TVector<T, d>& X(const int32 Index) const
 		{
-			check(IsUsingLocalXArray() && MXView.Num() == MX.Num() || !IsUsingLocalXArray());
-			return MXView[Index];
+			return MX[Index];
 		}
 
 		TVector<T, d>& X(const int32 Index)
 		{
-			check(IsUsingLocalXArray() && MXView.Num() == MX.Num() || !IsUsingLocalXArray());
-			return MXView[Index];
+			return MX[Index];
 		}
 
 		FString ToString(int32 index) const
@@ -173,8 +147,56 @@ namespace Chaos
 			return FString::Printf(TEXT("MX:%s"), *X(index).ToString());
 		}
 
+		uint32 GetTypeHash() const
+		{
+			uint32 OutHash = 0;
+			const int32 NumXEntries = MX.Num();
+
+			if(NumXEntries > 0)
+			{
+				OutHash = ::GetTypeHash(MX[0]);
+
+				for(int32 XIndex = 1; XIndex < NumXEntries; ++XIndex)
+				{
+					OutHash = HashCombine(OutHash, ::GetTypeHash(MX[XIndex]));
+				}
+			}
+
+			return OutHash;
+		}
+
+#if PARTICLE_ITERATOR_RANGED_FOR_CHECK
+		int32 DirtyValidationCount() const { return MDirtyValidationCount; }
+#endif
+
 	private:
 		TArrayCollectionArray<TVector<T, d>> MX;
-		TArrayCollectionArrayView<TVector<T, d>> MXView;
+
+		ERemoveParticleBehavior MRemoveParticleBehavior;
+
+#if PARTICLE_ITERATOR_RANGED_FOR_CHECK
+		int32 MDirtyValidationCount;
+#endif
+
+		void IncrementDirtyValidation()
+		{
+#if PARTICLE_ITERATOR_RANGED_FOR_CHECK
+			++MDirtyValidationCount;
+#endif
+		}
+
+		template<typename operator_T, int operator_d>
+		friend FArchive& operator<<(FArchive& Ar, TParticles<operator_T, operator_d>& InParticles)
+		{
+			InParticles.Serialize(Ar);
+			return Ar;
+		}
 	};
+
+	template<typename T, int d>
+	static uint32 GetTypeHash(const TParticles<T, d>& InParticles)
+	{
+		return InParticles.GetTypeHash();
+	}
+	
 } // namespace Chaos
diff --git a/Engine/Source/Runtime/Experimental/Chaos/Public/Chaos/PerParticleEtherDrag.h b/Engine/Source/Runtime/Experimental/Chaos/Public/Chaos/PerParticleEtherDrag.h
--- a/Engine/Source/Runtime/Experimental/Chaos/Public/Chaos/PerParticleEtherDrag.h
+++ b/Engine/Source/Runtime/Experimental/Chaos/Public/Chaos/PerParticleEtherDrag.h
@@ -41,6 +41,11 @@ class TPerParticleEtherDrag : public TPerParticleRule<T, d>
 		InParticles.W(Index) -= MAngularCoefficient * InParticles.W(Index);
 	}
 
+	inline void Apply(TTransientPBDRigidParticleHandle<T, d>& Particle, const T Dt) const override
+	{
+		Particle.V() -= MCoefficient * Particle.V();
+	}
+
   private:
 	T MCoefficient;
 	T MAngularCoefficient;
diff --git a/Engine/Source/Runtime/Experimental/Chaos/Public/Chaos/PerParticleEulerStepVelocity.h b/Engine/Source/Runtime/Experimental/Chaos/Public/Chaos/PerParticleEulerStepVelocity.h
--- a/Engine/Source/Runtime/Experimental/Chaos/Public/Chaos/PerParticleEulerStepVelocity.h
+++ b/Engine/Source/Runtime/Experimental/Chaos/Public/Chaos/PerParticleEulerStepVelocity.h
@@ -33,5 +33,10 @@ class TPerParticleEulerStepVelocity : public TPerParticleRule<T, d>
 		TVector<T, d> L = InParticles.I(Index) * InParticles.W(Index);
 		InParticles.W(Index) += InParticles.InvI(Index) * (InParticles.Torque(Index) - TVector<T, d>::CrossProduct(InParticles.W(Index), L)) * Dt;
 	}
+	
+	inline void Apply(TTransientPBDRigidParticleHandle<T, d>& Particle, const T Dt) const override //-V762
+	{
+		Particle.V() += Particle.F() * Particle.InvM() * Dt;
+	}
 };
 }
diff --git a/Engine/Source/Runtime/Experimental/Chaos/Public/Chaos/PerParticleGravity.h b/Engine/Source/Runtime/Experimental/Chaos/Public/Chaos/PerParticleGravity.h
--- a/Engine/Source/Runtime/Experimental/Chaos/Public/Chaos/PerParticleGravity.h
+++ b/Engine/Source/Runtime/Experimental/Chaos/Public/Chaos/PerParticleGravity.h
@@ -3,6 +3,7 @@
 
 #include "Chaos/DynamicParticles.h"
 #include "Chaos/PerParticleRule.h"
+#include "Chaos/ParticleHandle.h"
 
 namespace Chaos
 {
@@ -30,6 +31,11 @@ class PerParticleGravity : public TPerParticleRule<T, d>
 		ApplyHelper(InParticles, Dt, Index);
 	}
 
+	inline void Apply(TTransientPBDRigidParticleHandle<T, d>& Particle, const T Dt) const override //-V762
+	{
+		Particle.F() += MAcceleration * Particle.M();
+	}
+
 	void SetAcceleration(const TVector<T, d>& Acceleration)
 	{
 		MAcceleration = Acceleration;
diff --git a/Engine/Source/Runtime/Experimental/Chaos/Public/Chaos/PerParticleInitForce.h b/Engine/Source/Runtime/Experimental/Chaos/Public/Chaos/PerParticleInitForce.h
--- a/Engine/Source/Runtime/Experimental/Chaos/Public/Chaos/PerParticleInitForce.h
+++ b/Engine/Source/Runtime/Experimental/Chaos/Public/Chaos/PerParticleInitForce.h
@@ -22,5 +22,11 @@ class TPerParticleInitForce : public TPerParticleRule<T, d>
 		InParticles.F(Index) = TVector<T,d>(0);
 		InParticles.Torque(Index) = TVector<T,d>(0);
 	}
+
+	inline void Apply(TTransientPBDRigidParticleHandle<T, d>& Particle, const T Dt) const override //-V762
+	{
+		Particle.F() = TVector<T, d>(0);
+		Particle.Torque() = TVector<T, d>(0);
+	}
 };
 }
diff --git a/Engine/Source/Runtime/Experimental/Chaos/Public/Chaos/PerParticlePBDCollisionConstraint.h b/Engine/Source/Runtime/Experimental/Chaos/Public/Chaos/PerParticlePBDCollisionConstraint.h
--- a/Engine/Source/Runtime/Experimental/Chaos/Public/Chaos/PerParticlePBDCollisionConstraint.h
+++ b/Engine/Source/Runtime/Experimental/Chaos/Public/Chaos/PerParticlePBDCollisionConstraint.h
@@ -9,7 +9,7 @@
 
 namespace Chaos
 {
-template<class T, int d>
+template<class T, int d, EGeometryParticlesSimType SimType>
 class TPerParticlePBDCollisionConstraint : public TPerParticleRule<T, d>
 {
 	struct VelocityConstraint
@@ -19,7 +19,7 @@ class TPerParticlePBDCollisionConstraint : public TPerParticleRule<T, d>
 	};
 
   public:
-	TPerParticlePBDCollisionConstraint(const TKinematicGeometryParticles<T, d>& InParticles, TArrayCollectionArray<bool>& Collided, const T Thickness = (T)0, const T Friction = (T)0)
+	TPerParticlePBDCollisionConstraint(const TKinematicGeometryParticlesImp<T, d, SimType>& InParticles, TArrayCollectionArray<bool>& Collided, const T Thickness = (T)0, const T Friction = (T)0)
 	    : MParticles(InParticles), MCollided(Collided), MThickness(Thickness), MFriction(Friction) {}
 	virtual ~TPerParticlePBDCollisionConstraint() {}
 
@@ -67,7 +67,7 @@ class TPerParticlePBDCollisionConstraint : public TPerParticleRule<T, d>
 
   private:
 	// TODO(mlentine): Need a bb hierarchy
-	const TKinematicGeometryParticles<T, d>& MParticles;
+	const TKinematicGeometryParticlesImp<T, d, SimType>& MParticles;
 	TArrayCollectionArray<bool>& MCollided;
 	mutable TMap<int32, VelocityConstraint> MVelocityConstraints;
 	const T MThickness, MFriction;
diff --git a/Engine/Source/Runtime/Experimental/Chaos/Public/Chaos/PerParticlePBDEulerStep.h b/Engine/Source/Runtime/Experimental/Chaos/Public/Chaos/PerParticlePBDEulerStep.h
--- a/Engine/Source/Runtime/Experimental/Chaos/Public/Chaos/PerParticlePBDEulerStep.h
+++ b/Engine/Source/Runtime/Experimental/Chaos/Public/Chaos/PerParticlePBDEulerStep.h
@@ -2,6 +2,7 @@
 #pragma once
 
 #include "Chaos/PerParticleRule.h"
+#include "Chaos/ParticleHandle.h"
 
 namespace Chaos
 {
@@ -29,5 +30,12 @@ class TPerParticlePBDEulerStep : public TPerParticleRule<T, d>
 		InParticles.Q(Index) = InParticles.R(Index) + TRotation<T, d>(InParticles.W(Index), 0.f) * InParticles.R(Index) * Dt * T(0.5);
 		InParticles.Q(Index).Normalize();
 	}
+
+	inline void Apply(TTransientPBDRigidParticleHandle<T, d>& Particle, const T Dt) const override
+	{
+		Particle.P() = Particle.X() + Particle.V() * Dt;
+		Particle.Q() = Particle.R() + TRotation<T, d>(Particle.W(), 0.f) * Particle.R() * Dt * T(0.5);
+		Particle.Q().Normalize();
+	}
 };
 }
diff --git a/Engine/Source/Runtime/Experimental/Chaos/Public/Chaos/PerParticlePBDUpdateFromDeltaPosition.h b/Engine/Source/Runtime/Experimental/Chaos/Public/Chaos/PerParticlePBDUpdateFromDeltaPosition.h
--- a/Engine/Source/Runtime/Experimental/Chaos/Public/Chaos/PerParticlePBDUpdateFromDeltaPosition.h
+++ b/Engine/Source/Runtime/Experimental/Chaos/Public/Chaos/PerParticlePBDUpdateFromDeltaPosition.h
@@ -37,5 +37,15 @@ class TPerParticlePBDUpdateFromDeltaPosition : public TPerParticleRule<T, d>
 		InParticles.W(Index) = Axis * Angle / Dt;
 		//InParticles.R(Index) = InParticles.Q(Index);
 	}
+
+	inline void Apply(TTransientPBDRigidParticleHandle<T, d>& Particle, const T Dt) const override //-V762
+	{
+		Particle.V() = (Particle.P() - Particle.X()) / Dt;
+		TRotation<T, d> Delta = Particle.Q() * Particle.R().Inverse();
+		TVector<T, d> Axis;
+		T Angle;
+		Delta.ToAxisAndAngle(Axis, Angle);
+		Particle.W() = Axis * Angle / Dt;
+	}
 };
 }
diff --git a/Engine/Source/Runtime/Experimental/Chaos/Public/Chaos/PerParticleRule.h b/Engine/Source/Runtime/Experimental/Chaos/Public/Chaos/PerParticleRule.h
--- a/Engine/Source/Runtime/Experimental/Chaos/Public/Chaos/PerParticleRule.h
+++ b/Engine/Source/Runtime/Experimental/Chaos/Public/Chaos/PerParticleRule.h
@@ -46,5 +46,7 @@ class TPerParticleRule : public TParticleRule<T, d>
 	virtual void Apply(TPBDParticles<T, d>& InParticles, const T Dt, const int Index) const { Apply(static_cast<TDynamicParticles<T, d>&>(InParticles), Dt, Index); };
 	virtual void Apply(TRigidParticles<T, d>& InParticles, const T Dt, const int Index) const { Apply(static_cast<TParticles<T, d>&>(InParticles), Dt, Index); };
 	virtual void Apply(TPBDRigidParticles<T, d>& InParticles, const T Dt, const int Index) const { Apply(static_cast<TRigidParticles<T, d>&>(InParticles), Dt, Index); };
+
+	virtual void Apply(TTransientPBDRigidParticleHandle<T, d>& Particle, const T Dt) const { check(0); }
 };
 }
diff --git a/Engine/Source/Runtime/Experimental/Chaos/Public/Chaos/Plane.h b/Engine/Source/Runtime/Experimental/Chaos/Public/Chaos/Plane.h
--- a/Engine/Source/Runtime/Experimental/Chaos/Public/Chaos/Plane.h
+++ b/Engine/Source/Runtime/Experimental/Chaos/Public/Chaos/Plane.h
@@ -54,10 +54,11 @@ class TPlane final : public TImplicitObject<T, d>
 		return x - TVector<T, d>(Dist * MNormal);
 	}
 
-	virtual bool Raycast(const TVector<T, d>& StartPoint, const TVector<T, d>& Dir, const T Length, const T Thickness, T& OutTime, TVector<T, d>& OutPosition, TVector<T, d>& OutNormal) const override
+	virtual bool Raycast(const TVector<T, d>& StartPoint, const TVector<T, d>& Dir, const T Length, const T Thickness, T& OutTime, TVector<T, d>& OutPosition, TVector<T, d>& OutNormal, int32& OutFaceIndex) const override
 	{
 		ensure(FMath::IsNearlyEqual(Dir.SizeSquared(),1, KINDA_SMALL_NUMBER));
 		ensure(Length > 0);
+		OutFaceIndex = INDEX_NONE;
 
 		const T SignedDist = TVector<T, d>::DotProduct(StartPoint - MX, MNormal);
 		if (FMath::Abs(SignedDist) < Thickness)
@@ -86,7 +87,7 @@ class TPlane final : public TImplicitObject<T, d>
 			return false;	//never reach
 		}
 
-		OutTime = LengthAlongRay / Length;
+		OutTime = LengthAlongRay;
 		OutPosition = StartPoint + (LengthAlongRay + Thickness) * Dir;
 		OutNormal = -DirTowardsPlane;
 		return true;
@@ -138,6 +139,11 @@ class TPlane final : public TImplicitObject<T, d>
 		SerializeImp(Ar);
 	}
 
+	virtual uint32 GetTypeHash() const override
+	{
+		return HashCombine(::GetTypeHash(MX), ::GetTypeHash(MNormal));
+	}
+
   private:
 	TVector<T, d> MX;
 	TVector<T, d> MNormal;
diff --git a/Engine/Source/Runtime/Experimental/Chaos/Public/Chaos/RigidParticles.h b/Engine/Source/Runtime/Experimental/Chaos/Public/Chaos/RigidParticles.h
--- a/Engine/Source/Runtime/Experimental/Chaos/Public/Chaos/RigidParticles.h
+++ b/Engine/Source/Runtime/Experimental/Chaos/Public/Chaos/RigidParticles.h
@@ -40,7 +40,7 @@ constexpr int8 LowBitsMask(const int8 NumBits)
 static constexpr int8 ObjectStateBitCount = NumBitsNeeded((int8)EObjectStateType::Count - 1);
 
 template<class T, int d>
-class CHAOS_API TRigidParticles : public TKinematicGeometryParticles<T, d>
+class TRigidParticles : public TKinematicGeometryParticles<T, d>
 {
   public:
 	using TArrayCollection::Size;
@@ -80,9 +80,6 @@ class CHAOS_API TRigidParticles : public TKinematicGeometryParticles<T, d>
 		TArrayCollection::AddArray(&MIsland);
 		TArrayCollection::AddArray(&MToBeRemovedOnFracture);
 	}
-	~TRigidParticles()
-	{
-	}
 
 	const TVector<T, d>& Torque(const int32 Index) const { return MT[Index]; }
 	TVector<T, d>& Torque(const int32 Index) { return MT[Index]; }
@@ -103,7 +100,8 @@ class CHAOS_API TRigidParticles : public TKinematicGeometryParticles<T, d>
 	T& InvM(const int32 Index) { return MInvM[Index]; }
 
 	int32 CollisionParticlesSize(int32 Index) const { return MCollisionParticles[Index] == nullptr ? 0 : MCollisionParticles[Index]->Size(); }
-	void CollisionParticlesInitIfNeeded(const int32 Index, TArray<TVector<T, d>>* Points=nullptr);
+	void CHAOS_API CollisionParticlesInitIfNeeded(const int32 Index);
+	void CHAOS_API SetCollisionParticles(const int32 Index, TParticles<T, d>&& Particles);
 	
 	const TUniquePtr<TBVHParticles<T, d>>& CollisionParticles(const int32 Index) const { return MCollisionParticles[Index]; }
 	TUniquePtr<TBVHParticles<T, d>>& CollisionParticles(const int32 Index) { return MCollisionParticles[Index]; }
@@ -170,4 +168,7 @@ FChaosArchive& operator<<(FChaosArchive& Ar, TRigidParticles<T, d>& Particles)
 	Particles.Serialize(Ar);
 	return Ar;
 }
+
+extern template class TRigidParticles<float, 3>;
+
 }
diff --git a/Engine/Source/Runtime/Experimental/Chaos/Public/Chaos/Serializable.h b/Engine/Source/Runtime/Experimental/Chaos/Public/Chaos/Serializable.h
--- a/Engine/Source/Runtime/Experimental/Chaos/Public/Chaos/Serializable.h
+++ b/Engine/Source/Runtime/Experimental/Chaos/Public/Chaos/Serializable.h
@@ -13,6 +13,9 @@ public:
 	explicit TSerializablePtr(const TUniquePtr<T>& Unique) : Ptr(Unique.Get()) {}
 	explicit TSerializablePtr(TUniquePtr<T>&& Unique) = delete;
 	
+	template<ESPMode TESPMode>
+	explicit TSerializablePtr(const TSharedPtr<T, TESPMode>& Shared) : Ptr(Shared.Get()) {}
+	
 	const T* operator->() const { return Ptr; }
 	const T* Get() const { return Ptr; }
 	const T& operator*() const { return *Ptr; }
@@ -44,7 +47,6 @@ private:
 template <typename T>
 inline uint32 GetTypeHash(const TSerializablePtr<T>& Ptr) { return ::GetTypeHash(Ptr.Get()); }
 
-
 template <typename T>
 TSerializablePtr<T> MakeSerializable(const TUniquePtr<T>& Unique)
 {
@@ -54,6 +56,12 @@ TSerializablePtr<T> MakeSerializable(const TUniquePtr<T>& Unique)
 template <typename T>
 TSerializablePtr<T> MakeSerializable(const TUniquePtr<T>&& Unique) = delete;
 
+template<typename T, ESPMode TESPMode>
+TSerializablePtr<T> MakeSerializable(const TSharedPtr<T, TESPMode>& Shared)
+{
+	return TSerializablePtr<T>(Shared);
+}
+
 class FChaosArchive;
 
 }
\ No newline at end of file
diff --git a/Engine/Source/Runtime/Experimental/Chaos/Public/Chaos/Sphere.h b/Engine/Source/Runtime/Experimental/Chaos/Public/Chaos/Sphere.h
--- a/Engine/Source/Runtime/Experimental/Chaos/Public/Chaos/Sphere.h
+++ b/Engine/Source/Runtime/Experimental/Chaos/Public/Chaos/Sphere.h
@@ -21,27 +21,27 @@ namespace Chaos
 	public:
 		IMPLICIT_OBJECT_SERIALIZER(TSphere)
 
-		TSphere(const TVector<T, d>& Center, const T Radius)
+		TSphere(const TVector<T, d>& InCenter, const T InRadius)
 		    : TImplicitObject<T, d>(EImplicitObject::IsConvex | EImplicitObject::HasBoundingBox, ImplicitObjectType::Sphere)
-		    , MCenter(Center)
-		    , MRadius(Radius)
-		    , MLocalBoundingBox(MCenter - MRadius, MCenter + MRadius)
+		    , Center(InCenter)
+		    , Radius(InRadius)
+		    , LocalBoundingBox(Center - Radius, Center + Radius)
 		{
 		}
 
 		TSphere(const TSphere<T, d>& Other)
 		    : TImplicitObject<T, d>(EImplicitObject::IsConvex | EImplicitObject::HasBoundingBox, ImplicitObjectType::Sphere)
-		    , MCenter(Other.MCenter)
-		    , MRadius(Other.MRadius)
-		    , MLocalBoundingBox(Other.MLocalBoundingBox)
+		    , Center(Other.Center)
+		    , Radius(Other.Radius)
+		    , LocalBoundingBox(Other.LocalBoundingBox)
 		{
 		}
 
 		TSphere(TSphere<T, d>&& Other)
 		    : TImplicitObject<T, d>(EImplicitObject::IsConvex | EImplicitObject::HasBoundingBox, ImplicitObjectType::Sphere)
-		    , MCenter(MoveTemp(Other.MCenter))
-		    , MRadius(Other.MRadius)
-		    , MLocalBoundingBox(MoveTemp(Other.MLocalBoundingBox))
+		    , Center(MoveTemp(Other.Center))
+		    , Radius(Other.Radius)
+		    , LocalBoundingBox(MoveTemp(Other.LocalBoundingBox))
 		{
 		}
 
@@ -51,44 +51,48 @@ namespace Chaos
 			this->bIsConvex = InSteal.bIsConvex;
 			this->bIgnoreAnalyticCollisions = InSteal.bIgnoreAnalyticCollisions;
 			this->bHasBoundingBox = InSteal.bHasBoundingBox;
-			MCenter = MoveTemp(InSteal.MCenter);
-			MRadius = InSteal.MRadius;
-			MLocalBoundingBox = MoveTemp(InSteal.MLocalBoundingBox);
+			Center = MoveTemp(InSteal.Center);
+			Radius = InSteal.Radius;
+			LocalBoundingBox = MoveTemp(InSteal.LocalBoundingBox);
 
 			return *this;
 		}
 
 		virtual ~TSphere() {}
 
-		static ImplicitObjectType GetType() { return ImplicitObjectType::Sphere; }
+		static ImplicitObjectType GetType() 
+		{ 
+			return ImplicitObjectType::Sphere; 
+		}
 
-		virtual T PhiWithNormal(const TVector<T, d>& x, TVector<T, d>& Normal) const override
+		virtual T PhiWithNormal(const TVector<T, d>& InSamplePoint, TVector<T, d>& OutNormal) const override
 		{
-			Normal = x - MCenter;
-			return Normal.SafeNormalize() - MRadius;
+			OutNormal = InSamplePoint - Center;
+			return OutNormal.SafeNormalize() - Radius;
 		}
 
 		bool Intersects(const TSphere<T, d>& Other) const
 		{
-			T CenterDistance = FVector::DistSquared(Other.Center(), Center());
-			T RadialSum = Other.Radius() + Radius();
+			T CenterDistance = FVector::DistSquared(Other.GetCenter(), GetCenter());
+			T RadialSum = Other.GetRadius() + GetRadius();
 			return RadialSum >= CenterDistance;
 		}
 
 		TVector<T, d> FindClosestPoint(const TVector<T, d>& StartPoint, const T Thickness = (T)0) const
 		{
-			TVector<T, 3> Result = MCenter + (StartPoint - MCenter).GetSafeNormal() * (MRadius + Thickness);
+			TVector<T, 3> Result = Center + (StartPoint - Center).GetSafeNormal() * (Radius + Thickness);
 			return Result;
 		}
 
-		virtual bool Raycast(const TVector<T, d>& StartPoint, const TVector<T, d>& Dir, const T Length, const T Thickness, T& OutTime, TVector<T, d>& OutPosition, TVector<T, d>& OutNormal) const override
+		virtual bool Raycast(const TVector<T, d>& StartPoint, const TVector<T, d>& Dir, const T Length, const T Thickness, T& OutTime, TVector<T, d>& OutPosition, TVector<T, d>& OutNormal, int32& OutFaceIndex) const override
 		{
 			ensure(FMath::IsNearlyEqual(Dir.SizeSquared(),1, KINDA_SMALL_NUMBER));
 			ensure(Length > 0);
+			OutFaceIndex = INDEX_NONE;
 
-			const T EffectiveRadius = Thickness + MRadius;
+			const T EffectiveRadius = Thickness + Radius;
 			const T EffectiveRadius2 = EffectiveRadius * EffectiveRadius;
-			const TVector<T, d> Offset = MCenter - StartPoint;
+			const TVector<T, d> Offset = Center - StartPoint;
 			const T OffsetSize2 = Offset.SizeSquared();
 			if (OffsetSize2 < EffectiveRadius2)
 			{
@@ -117,10 +121,10 @@ namespace Chaos
 			if (FirstTime >= 0 && FirstTime <= Length)
 			{
 				const TVector<T, d> FinalSpherePosition = StartPoint + FirstTime * Dir;
-				const TVector<T, d> FinalNormal = (FinalSpherePosition - MCenter) / EffectiveRadius;
+				const TVector<T, d> FinalNormal = (FinalSpherePosition - Center) / EffectiveRadius;
 				const TVector<T, d> IntersectionPosition = FinalSpherePosition - FinalNormal * Thickness;
 				
-				OutTime = FirstTime / Length;
+				OutTime = FirstTime;
 				OutPosition = IntersectionPosition;
 				OutNormal = FinalNormal;
 				return true;
@@ -134,9 +138,9 @@ namespace Chaos
 			TVector<T, d> Direction = EndPoint - StartPoint;
 			T Length = Direction.Size();
 			Direction = Direction.GetSafeNormal();
-			TVector<T, d> SphereToStart = StartPoint - MCenter;
+			TVector<T, d> SphereToStart = StartPoint - Center;
 			T DistanceProjected = TVector<T, d>::DotProduct(Direction, SphereToStart);
-			T EffectiveRadius = MRadius + Thickness;
+			T EffectiveRadius = Radius + Thickness;
 			T UnderRoot = DistanceProjected * DistanceProjected - SphereToStart.SizeSquared() + EffectiveRadius * EffectiveRadius;
 			if (UnderRoot < 0)
 			{
@@ -178,52 +182,81 @@ namespace Chaos
 			T SizeSqr = Direction.SizeSquared();
 			if (SizeSqr <= TNumericLimits<T>::Min())
 			{
-				return MCenter;
+				return Center;
 			}
 			const TVector<T,d> Normalized = Direction / sqrt(SizeSqr);
 
-			return MCenter + Normalized * (MRadius + Thickness);
+			return Center + Normalized * (Radius + Thickness);
 		}
 
-		virtual const TBox<T, d>& BoundingBox() const { return MLocalBoundingBox; }
+		virtual const TBox<T, d>& BoundingBox() const 
+		{
+			return LocalBoundingBox; 
+		}
 
-		T GetArea() const { return GetArea(MRadius); }
-		static T GetArea(const T Radius)
+		T GetArea() const 
+		{ 
+			return GetArea(Radius); 
+		}
+		
+		static T GetArea(const T InRadius)
 		{
 			static const T FourPI = PI * 4;
 			static const T TwoPI = PI * 2;
-			return d == 3 ? FourPI * Radius * Radius : TwoPI * Radius;
+			return d == 3 ? FourPI * InRadius * InRadius : TwoPI * InRadius;
+		}
+
+		T GetVolume() const 
+		{
+			return GetVolume(Radius); 
 		}
 
-		T GetVolume() const { return GetVolume(MRadius); }
-		static T GetVolume(const T Radius)
+		static T GetVolume(const T InRadius)
 		{
 			check(d == 3);
 			static const T FourThirdsPI = 4. / 3 * PI;
-			return FourThirdsPI * Radius * Radius * Radius;
+			return FourThirdsPI * InRadius * InRadius * InRadius;
 		}
 
-		const TVector<T, d>& Center() const { return MCenter; }
-		const TVector<T, d>& GetCenterOfMass() const { return Center(); }
-		T Radius() const { return MRadius; }
+		const TVector<T, d>& GetCenter() const 
+		{ 
+			return Center; 
+		}
+
+		const TVector<T, d>& GetCenterOfMass() const 
+		{ 
+			return Center; 
+		}
+
+		T GetRadius() const 
+		{ 
+			return Radius; 
+		}
 
 		virtual FString ToString() const
 		{
-			return FString::Printf(TEXT("TSphere Center:%s, Radius:%f"), *Center().ToString(), Radius());
+			return FString::Printf(TEXT("TSphere Center:%s, Radius:%f"), *Center.ToString(), GetRadius());
 		}
 
 		FORCEINLINE void SerializeImp(FArchive& Ar)
 		{
 			TImplicitObject<T, d>::SerializeImp(Ar);
-			Ar << MCenter << MRadius;
+			Ar << Center << Radius;
 			if (Ar.IsLoading())
 			{
-				MLocalBoundingBox = TBox<T, d>(MCenter - MRadius, MCenter + MRadius);
+				LocalBoundingBox = TBox<T, d>(Center - Radius, Center + Radius);
 			}
 		}
 
-		virtual void Serialize(FChaosArchive& Ar) override { SerializeImp(Ar); }
-		virtual void Serialize(FArchive& Ar) override { SerializeImp(Ar); }
+		virtual void Serialize(FChaosArchive& Ar) override 
+		{ 
+			SerializeImp(Ar); 
+		}
+
+		virtual void Serialize(FArchive& Ar) override
+		{ 
+			SerializeImp(Ar); 
+		}
 
 		/**
 		 * Returns sample points centered about the origin.
@@ -231,12 +264,15 @@ namespace Chaos
 		TArray<TVector<T, d>> ComputeLocalSamplePoints(const int NumPoints) const
 		{
 			TArray<TVector<T, d>> Points;
-			TSphere<T, d> LocalSphere(TVector<T, d>(0.0), MRadius);
+			TSphere<T, d> LocalSphere(TVector<T, d>(0.0), Radius);
 			TSphereSpecializeSamplingHelper<T, d>::ComputeSamplePoints(Points, LocalSphere, NumPoints);
 			return Points;
 		}
+
 		TArray<TVector<T, d>> ComputeLocalSamplePoints(const T PointsPerUnitArea, const int32 MinPoints = 0, const int32 MaxPoints = 1000) const
-		{ return ComputeLocalSamplePoints(FMath::Clamp(static_cast<int32>(ceil(PointsPerUnitArea * GetArea())), MinPoints, MaxPoints)); }
+		{ 
+			return ComputeLocalSamplePoints(FMath::Clamp(static_cast<int32>(ceil(PointsPerUnitArea * GetArea())), MinPoints, MaxPoints)); 
+		}
 
 		/**
 		 * Returns sample points at the current location of the sphere.
@@ -247,30 +283,51 @@ namespace Chaos
 			TSphereSpecializeSamplingHelper<T, d>::ComputeSamplePoints(Points, *this, NumPoints);
 			return Points;
 		}
+
 		TArray<TVector<T, d>> ComputeSamplePoints(const T PointsPerUnitArea, const int32 MinPoints = 0, const int32 MaxPoints = 1000) const
-		{ return ComputeSamplePoints(FMath::Clamp(static_cast<int32>(ceil(PointsPerUnitArea * GetArea())), MinPoints, MaxPoints)); }
+		{ 
+			return ComputeSamplePoints(FMath::Clamp(static_cast<int32>(ceil(PointsPerUnitArea * GetArea())), MinPoints, MaxPoints)); 
+		}
+
+		PMatrix<T, d, d> GetInertiaTensor(const T InMass, const bool bInThinShell = false) const 
+		{ 
+			return GetInertiaTensor(InMass, Radius, bInThinShell); 
+		}
 
-		PMatrix<T, d, d> GetInertiaTensor(const T Mass, const bool ThinShell = false) const { return GetInertiaTensor(Mass, MRadius, ThinShell); }
-		static PMatrix<T, d, d> GetInertiaTensor(const T Mass, const T Radius, const bool ThinShell = false)
+		static PMatrix<T, d, d> GetInertiaTensor(const T InMass, const T InRadius, const bool bInThinShell = false)
 		{
 			static const T TwoThirds = 2. / 3;
 			static const T TwoFifths = 2. / 5;
-			const T Diagonal = ThinShell ? TwoThirds * Mass * Radius * Radius : TwoFifths * Mass * Radius * Radius;
+			const T Diagonal = bInThinShell ? TwoThirds * InMass * InRadius * InRadius : TwoFifths * InMass * InRadius * InRadius;
 			return PMatrix<T, d, d>(Diagonal, Diagonal, Diagonal);
 		}
 
 		static TRotation<T, d> GetRotationOfMass()
-		{ return TRotation<T, d>(TVector<T, d>(0), 1); }
+		{ 
+			return TRotation<T, d>(TVector<T, d>(0), 1); 
+		}
+
+		virtual uint32 GetTypeHash() const override
+		{
+			const uint32 CenterHash = ::GetTypeHash(Center);
+			const uint32 RadiusHash = ::GetTypeHash(Radius);
+			const uint32 BoundsHash = LocalBoundingBox.GetTypeHash();
+			return HashCombine(CenterHash, HashCombine(RadiusHash, BoundsHash));
+		}
 
 	private:
-		TVector<T, d> MCenter;
-		T MRadius;
-		TBox<T, d> MLocalBoundingBox;
+		TVector<T, d> Center;
+		T Radius;
+		TBox<T, d> LocalBoundingBox;
 
 	private:
+
+		// TImplicitObject requires ability to default construct when deserializing shapes
+		friend TImplicitObject<T, d>;
 		TSphere()
-		    : TImplicitObject<T, d>(EImplicitObject::IsConvex | EImplicitObject::HasBoundingBox, ImplicitObjectType::Sphere) {} //needed for serialization
-		friend TImplicitObject<T, d>; //needed for serialization
+		    : TImplicitObject<T, d>(EImplicitObject::IsConvex | EImplicitObject::HasBoundingBox, ImplicitObjectType::Sphere) 
+		{}
+
 	};
 
 	template<typename T>
@@ -278,7 +335,7 @@ namespace Chaos
 	{
 		static FORCEINLINE void ComputeSamplePoints(TArray<TVector<T, 2>>& Points, const TSphere<T, 2>& Sphere, const int32 NumPoints)
 		{
-			if (NumPoints <= 1 || Sphere.Radius() < KINDA_SMALL_NUMBER)
+			if (NumPoints <= 1 || Sphere.GetRadius() < KINDA_SMALL_NUMBER)
 			{
 				const int32 Offset = Points.AddUninitialized(1);
 				Points[Offset] = Sphere.Center();
@@ -292,7 +349,7 @@ namespace Chaos
 		 */
 		static FORCEINLINE void ComputeGoldenSpiralPoints(TArray<TVector<T, 2>>& Points, const TSphere<T, 2>& Sphere, const int32 NumPoints)
 		{
-			ComputeGoldenSpiralPoints(Points, Sphere.Center(), Sphere.Radius(), NumPoints);
+			ComputeGoldenSpiralPoints(Points, Sphere.Center(), Sphere.GetRadius(), NumPoints);
 		}
 
 		static FORCEINLINE void ComputeGoldenSpiralPoints(
@@ -336,10 +393,10 @@ namespace Chaos
 	{
 		static FORCEINLINE void ComputeSamplePoints(TArray<TVector<T, 3>>& Points, const TSphere<T, 3>& Sphere, const int32 NumPoints)
 		{
-			if (NumPoints <= 1 || Sphere.Radius() < KINDA_SMALL_NUMBER)
+			if (NumPoints <= 1 || Sphere.GetRadius() < KINDA_SMALL_NUMBER)
 			{
 				const int32 Offset = Points.AddUninitialized(1);
-				Points[Offset] = Sphere.Center();
+				Points[Offset] = Sphere.GetCenter();
 				return;
 			}
 			ComputeGoldenSpiralPoints(Points, Sphere, NumPoints);
@@ -352,7 +409,7 @@ namespace Chaos
 		    TArray<TVector<T, 3>>& Points, const TSphere<T, 3>& Sphere, const int32 NumPoints,
 		    const bool FirstHalf = true, const bool SecondHalf = true, const int32 SpiralSeed = 0)
 		{
-			ComputeGoldenSpiralPoints(Points, Sphere.Center(), Sphere.Radius(), NumPoints, FirstHalf, SecondHalf, SpiralSeed);
+			ComputeGoldenSpiralPoints(Points, Sphere.GetCenter(), Sphere.GetRadius(), NumPoints, FirstHalf, SecondHalf, SpiralSeed);
 		}
 
 		/**
diff --git a/Engine/Source/Runtime/Experimental/Chaos/Public/Chaos/TaperedCylinder.h b/Engine/Source/Runtime/Experimental/Chaos/Public/Chaos/TaperedCylinder.h
--- a/Engine/Source/Runtime/Experimental/Chaos/Public/Chaos/TaperedCylinder.h
+++ b/Engine/Source/Runtime/Experimental/Chaos/Public/Chaos/TaperedCylinder.h
@@ -330,6 +330,14 @@ namespace Chaos
 		static TRotation<T, 3> GetRotationOfMass()
 		{ return TRotation<T, 3>(TVector<T, 3>(0), 1); }
 
+		virtual uint32 GetTypeHash() const override
+		{
+			const uint32 PlaneHashes = HashCombine(MPlane1.GetTypeHash(), MPlane2.GetTypeHash());
+			const uint32 PropertyHash = HashCombine(::GetTypeHash(MHeight), HashCombine(::GetTypeHash(MRadius1), ::GetTypeHash(MRadius2)));
+
+			return HashCombine(PlaneHashes, PropertyHash);
+		}
+
 	private:
 		//Phi is distance from closest point on plane1
 		T GetRadius(const T& Phi) const
diff --git a/Engine/Source/Runtime/Experimental/Chaos/Public/Chaos/Transform.h b/Engine/Source/Runtime/Experimental/Chaos/Public/Chaos/Transform.h
--- a/Engine/Source/Runtime/Experimental/Chaos/Public/Chaos/Transform.h
+++ b/Engine/Source/Runtime/Experimental/Chaos/Public/Chaos/Transform.h
@@ -148,5 +148,10 @@ class TRigidTransform<float, 3> : public FTransform
 };
 }
 
+inline uint32 GetTypeHash(const Chaos::TRigidTransform<float, 3>& InTransform)
+{
+	return HashCombine(GetTypeHash(InTransform.GetTranslation()), HashCombine(GetTypeHash(InTransform.GetRotation().Euler()), GetTypeHash(InTransform.GetScale3D())));
+}
+
 CHAOS_API Chaos::PMatrix<float, 4, 4> operator*(const Chaos::TRigidTransform<float, 3>& Transform, const Chaos::PMatrix<float, 4, 4>& Matrix);
 CHAOS_API Chaos::PMatrix<float, 4, 4> operator*(const Chaos::PMatrix<float, 4, 4>& Matrix, const Chaos::TRigidTransform<float, 3>& Transform);
\ No newline at end of file
diff --git a/Engine/Source/Runtime/Experimental/Chaos/Public/Chaos/TriangleMesh.h b/Engine/Source/Runtime/Experimental/Chaos/Public/Chaos/TriangleMesh.h
--- a/Engine/Source/Runtime/Experimental/Chaos/Public/Chaos/TriangleMesh.h
+++ b/Engine/Source/Runtime/Experimental/Chaos/Public/Chaos/TriangleMesh.h
@@ -200,7 +200,7 @@ namespace Chaos
 		static void InitEquilateralTriangleXY(TTriangleMesh<T>& TriMesh, TParticles<T, 3>& Particles)
 		{
 			const int32 Idx = Particles.Size();
-			Particles.AddElements(3);
+			Particles.AddParticles(3);
 			// Left handed
 			Particles.X(Idx + 0) = TVector<T, 3>(0., 0.8083, 0.);
 			Particles.X(Idx + 1) = TVector<T, 3>(0.7, -0.4041, 0.);
@@ -215,7 +215,7 @@ namespace Chaos
 		static void InitEquilateralTriangleYZ(TTriangleMesh<T>& TriMesh, TParticles<T, 3>& Particles)
 		{
 			const int32 Idx = Particles.Size();
-			Particles.AddElements(3);
+			Particles.AddParticles(3);
 			// Left handed
 			Particles.X(Idx + 0) = TVector<T, 3>(0., 0., 0.8083);
 			Particles.X(Idx + 1) = TVector<T, 3>(0., 0.7, -0.4041);
diff --git a/Engine/Source/Runtime/Experimental/Chaos/Public/Chaos/TriangleMeshImplicitObject.h b/Engine/Source/Runtime/Experimental/Chaos/Public/Chaos/TriangleMeshImplicitObject.h
--- a/Engine/Source/Runtime/Experimental/Chaos/Public/Chaos/TriangleMeshImplicitObject.h
+++ b/Engine/Source/Runtime/Experimental/Chaos/Public/Chaos/TriangleMeshImplicitObject.h
@@ -16,7 +16,7 @@ namespace Chaos
 	class CHAOS_API TTriangleMeshImplicitObject final : public TImplicitObject<T,3>
 	{
 	public:
-		//IMPLICIT_OBJECT_SERIALIZER(TTriangleMeshImplicitObject)
+		IMPLICIT_OBJECT_SERIALIZER_DIM(TTriangleMeshImplicitObject, 3)
 		TTriangleMeshImplicitObject(TParticles<T,3>&& Particles, TArray<TVector<int32, 3>>&& Elements);
 		TTriangleMeshImplicitObject(const TTriangleMeshImplicitObject& Other) = delete;
 		TTriangleMeshImplicitObject(TTriangleMeshImplicitObject&& Other) = default;
@@ -28,18 +28,75 @@ namespace Chaos
 			return (T)0;
 		}
 
-		virtual bool Raycast(const TVector<T, 3>& StartPoint, const TVector<T, 3>& Dir, const T Length, const T Thickness, T& OutTime, TVector<T,3>& OutPosition, TVector<T,3>& OutNormal) const override;
+		virtual bool Raycast(const TVector<T, 3>& StartPoint, const TVector<T, 3>& Dir, const T Length, const T Thickness, T& OutTime, TVector<T,3>& OutPosition, TVector<T,3>& OutNormal, int32& OutFaceIndex) const override;
 		virtual bool Overlap(const TVector<T, 3>& Point, const T Thickness) const override;
 
 		bool OverlapGeom(const TImplicitObject<T,3>& QueryGeom, const TRigidTransform<T, 3>& QueryTM, const T Thickness, const TVector<T,3> Scale = TVector<T,3>(1)) const;
-		bool SweepGeom(const TImplicitObject<T, 3>& QueryGeom, const TRigidTransform<T, 3>& StartTM, const TVector<T,3>& Dir, const T Length, T& OutTime, TVector<T,3>& OutPosition, TVector<T,3>& OutNormal, const T Thickness = 0, const TVector<T,3> Scale = TVector<T,3>(1)) const;
+		bool SweepGeom(const TImplicitObject<T, 3>& QueryGeom, const TRigidTransform<T, 3>& StartTM, const TVector<T,3>& Dir, const T Length, T& OutTime, TVector<T,3>& OutPosition, TVector<T,3>& OutNormal, int32& OutFaceIndex, const T Thickness = 0, const TVector<T,3> Scale = TVector<T,3>(1)) const;
+		virtual int32 FindMostOpposingFace(const TVector<T, 3>& Position, const TVector<T, 3>& UnitDir, int32 HintFaceIndex) const override;
 
 		virtual const TBox<T, 3>& BoundingBox() const
 		{
 			return MLocalBoundingBox;
 		}
 
+		static ImplicitObjectType GetType()
+		{
+			return ImplicitObjectType::TriangleMesh;
+		}
+
+		void SerializeImp(FArchive& Ar)
+		{
+			TImplicitObject<T, 3>::SerializeImp(Ar);
+			Ar << MParticles << MElements << MLocalBoundingBox;
+
+#if 0
+			// Disabled during 2-1 replacement. Replace when BV is actually working again
+			Ar << ConvexHulls;
+			Ar << BVH;
+			if(Ar.IsLoading())
+			{
+				// Re-link the object array when we load this back in
+				BVH.SetObjects(ConvexHulls);
+			}
+#else
+			if(Ar.IsLoading())
+			{
+				RebuildConvexHulls();
+				RebuildBV();
+			}
+#endif
+		}
+
+		virtual void Serialize(FChaosArchive& Ar) override
+		{
+			SerializeImp(Ar);
+		}
+
+		virtual void Serialize(FArchive& Ar) override
+		{
+			SerializeImp(Ar);
+		}
+
+		virtual uint32 GetTypeHash() const override
+		{
+			uint32 Result = MParticles.GetTypeHash();
+			Result = HashCombine(Result, MLocalBoundingBox.GetTypeHash());
+
+			for(TVector<int32, 3> Tri : MElements)
+			{
+				uint32 TriHash = HashCombine(::GetTypeHash(Tri[0]), HashCombine(::GetTypeHash(Tri[1]), ::GetTypeHash(Tri[2])));
+				Result = HashCombine(Result, TriHash);
+			}
+
+			return Result;
+		}
+
 	private:
+
+		void RebuildConvexHulls();
+		void RebuildBV();
+
 		TParticles<T, 3> MParticles;
 		TArray<TVector<int32, 3>> MElements;
 		TBox<T, 3> MLocalBoundingBox;
@@ -47,10 +104,18 @@ namespace Chaos
 		TGeometryParticles<T, 3> ConvexHulls;
 		//using BVHType = TBoundingVolumeHierarchy<TGeometryParticles<T, 3>, TArray<int32>, T, 3>;
 		//using BVHType = TBoundingVolumeHierarchy<TGeometryParticles<T, 3>, TBoundingVolume<TGeometryParticles<T,3>, T, 3>, T, 3>;
-		using BVHType = TBoundingVolume<TGeometryParticles<T, 3>, T, 3>;
+		using BVHType = TBoundingVolume<TGeometryParticles<T, 3>, int32, T, 3>;
 		BVHType BVH;
 
 		template <typename R>
 		friend struct TTriangleMeshSweepVisitor;
+
+		// Required by implicit object serialization, disabled for general use.
+		template<typename InnerT, int InnerD>
+		friend class TImplicitObject;
+
+		TTriangleMeshImplicitObject() 
+			: TImplicitObject<T, 3>(EImplicitObject::HasBoundingBox, ImplicitObjectType::TriangleMesh)
+		{};
 	};
 }
