Copying //UE4/Dev-Physics to Dev-Main (//UE4/Dev-Main) [at] 7703071

diff --git a/Engine/Source/Runtime/Experimental/Chaos/Public/Chaos/Utilities.h b/Engine/Source/Runtime/Experimental/Chaos/Public/Chaos/Utilities.h
--- a/Engine/Source/Runtime/Experimental/Chaos/Public/Chaos/Utilities.h
+++ b/Engine/Source/Runtime/Experimental/Chaos/Public/Chaos/Utilities.h
@@ -27,9 +27,11 @@ namespace Chaos
 			return GetGravityFunction<TPBDParticles<T, d>, TDynamicParticles<T, d>, T, d>(Direction, Magnitude);
 		}
 		template<class T, int d>
-		inline TFunction<void(TPBDRigidParticles<T, d>&, const T, const int32)> GetRigidsGravityFunction(const TVector<T, d>& Direction, const T Magnitude)
+		inline auto GetRigidsGravityFunction(const TVector<T, d>& Direction, const T Magnitude)
 		{
-			return GetGravityFunction<TPBDRigidParticles<T, d>, TRigidParticles<T, d>, T, d>(Direction, Magnitude);
+			return[Gravity = PerParticleGravity<T, d>(Direction, Magnitude)](TTransientPBDRigidParticleHandle<T,d>& Particle, const T Dt) {
+				Gravity.Apply(Particle, Dt);
+			};
 		}
 
 		// Take the factorial of \p Num, which should be of integral type.
diff --git a/Engine/Source/Runtime/Experimental/Chaos/Public/ChaosCheck.h b/Engine/Source/Runtime/Experimental/Chaos/Public/ChaosCheck.h
new file mode 100644
--- /dev/null
+++ b/Engine/Source/Runtime/Experimental/Chaos/Public/ChaosCheck.h
@@ -0,0 +1,11 @@
+// Copyright 1998-2019 Epic Games, Inc. All Rights Reserved.
+
+#pragma once
+
+#include "Misc/AssertionMacros.h"
+
+#if CHAOS_CHECKED
+#define CHAOS_CHECK(Condition) check(Condition)
+#else
+#define CHAOS_CHECK(Condition)
+#endif
diff --git a/Engine/Source/Runtime/Experimental/Chaos/Public/ChaosModule.h b/Engine/Source/Runtime/Experimental/Chaos/Public/ChaosModule.h
--- a/Engine/Source/Runtime/Experimental/Chaos/Public/ChaosModule.h
+++ b/Engine/Source/Runtime/Experimental/Chaos/Public/ChaosModule.h
@@ -1,6 +1,6 @@
 // Copyright 1998-2019 Epic Games, Inc. All Rights Reserved.
 
-#if INCLUDE_CHAOS
+//#if INCLUDE_CHAOS
 
 #pragma once
 
@@ -14,4 +14,4 @@ public:
 	virtual void ShutdownModule() override {};
 };
 
-#endif
\ No newline at end of file
+//#endif
\ No newline at end of file
diff --git a/Engine/Source/Runtime/Experimental/Chaos/Public/ChaosStats.h b/Engine/Source/Runtime/Experimental/Chaos/Public/ChaosStats.h
--- a/Engine/Source/Runtime/Experimental/Chaos/Public/ChaosStats.h
+++ b/Engine/Source/Runtime/Experimental/Chaos/Public/ChaosStats.h
@@ -37,10 +37,10 @@ DECLARE_CYCLE_STAT_EXTERN(TEXT("Skeletal Mesh Update Anim"), STAT_SkelMeshUpdate
 DECLARE_CYCLE_STAT_EXTERN(TEXT("Dispatch Event Notifies"), STAT_DispatchEventNotifies, STATGROUP_Chaos, CHAOS_API);
 DECLARE_CYCLE_STAT_EXTERN(TEXT("Dispatch Collision Events"), STAT_DispatchCollisionEvents, STATGROUP_Chaos, CHAOS_API);
 DECLARE_CYCLE_STAT_EXTERN(TEXT("Dispatch Break Events"), STAT_DispatchBreakEvents, STATGROUP_Chaos, CHAOS_API);
-DECLARE_CYCLE_STAT_EXTERN(TEXT("Cache Results"), STAT_CacheResults, STATGROUP_Chaos, CHAOS_API);
+DECLARE_CYCLE_STAT_EXTERN(TEXT("BufferPhysicsResults"), STAT_BufferPhysicsResults, STATGROUP_Chaos, CHAOS_API);
 DECLARE_CYCLE_STAT_EXTERN(TEXT("Flip Results"), STAT_FlipResults, STATGROUP_Chaos, CHAOS_API);
-DECLARE_CYCLE_STAT_EXTERN(TEXT("[CacheResults] - Geometry Collection"), STAT_CacheResultGeomCollection, STATGROUP_ChaosWide, CHAOS_API);
-DECLARE_CYCLE_STAT_EXTERN(TEXT("[CacheResults] - StaticMesh"), STAT_CacheResultStaticMesh, STATGROUP_Chaos, CHAOS_API);
+DECLARE_CYCLE_STAT_EXTERN(TEXT("[BufferPhysicsResults] - Geometry Collection"), STAT_CacheResultGeomCollection, STATGROUP_ChaosWide, CHAOS_API);
+DECLARE_CYCLE_STAT_EXTERN(TEXT("[BufferPhysicsResults] - StaticMesh"), STAT_CacheResultStaticMesh, STATGROUP_Chaos, CHAOS_API);
 DECLARE_CYCLE_STAT_EXTERN(TEXT("Capture Disabled State"), STAT_CaptureDisabledState, STATGROUP_Chaos, CHAOS_API);
 DECLARE_CYCLE_STAT_EXTERN(TEXT("Calc Global Matrices"), STAT_CalcGlobalGCMatrices, STATGROUP_Chaos, CHAOS_API);
 DECLARE_CYCLE_STAT_EXTERN(TEXT("Calc Global Bounds"), STAT_CalcGlobalGCBounds, STATGROUP_Chaos, CHAOS_API);
@@ -51,6 +51,7 @@ DECLARE_CYCLE_STAT_EXTERN(TEXT("Parameter Update"), STAT_UpdateParams, STATGROUP
 DECLARE_CYCLE_STAT_EXTERN(TEXT("Disable collisions"), STAT_DisableCollisions, STATGROUP_Chaos, CHAOS_API);
 DECLARE_CYCLE_STAT_EXTERN(TEXT("Evolution/Kinematic update and forces"), STAT_EvolutionAndKinematicUpdate, STATGROUP_Chaos, CHAOS_API);
 DECLARE_CYCLE_STAT_EXTERN(TEXT("AdvanceOneTimestep Event Waits"), STAT_AdvanceEventWaits, STATGROUP_Chaos, CHAOS_API);
+DECLARE_CYCLE_STAT_EXTERN(TEXT("Gathering Event Data"), STAT_EventDataGathering, STATGROUP_Chaos, CHAOS_API);
 
 DECLARE_CYCLE_STAT_EXTERN(TEXT("Reset Collision Rule"), STAT_ResetCollisionRule, STATGROUP_Chaos, CHAOS_API);
 
diff --git a/Engine/Source/Runtime/Experimental/ChaosCore/ChaosCore.Build.cs b/Engine/Source/Runtime/Experimental/ChaosCore/ChaosCore.Build.cs
--- a/Engine/Source/Runtime/Experimental/ChaosCore/ChaosCore.Build.cs
+++ b/Engine/Source/Runtime/Experimental/ChaosCore/ChaosCore.Build.cs
@@ -25,6 +25,15 @@ namespace UnrealBuildTool.Rules
             {
                 PublicDefinitions.Add("INCLUDE_CHAOS=0");
             }
+
+            if (Target.bUseChaosChecked == true)
+            {
+                PublicDefinitions.Add("CHAOS_CHECKED=1");
+            }
+            else
+            {
+                PublicDefinitions.Add("CHAOS_CHECKED=0");
+            }
         }
     }
 }
diff --git a/Engine/Source/Runtime/Experimental/ChaosCore/Public/Chaos/ArrayCollection.h b/Engine/Source/Runtime/Experimental/ChaosCore/Public/Chaos/ArrayCollection.h
--- a/Engine/Source/Runtime/Experimental/ChaosCore/Public/Chaos/ArrayCollection.h
+++ b/Engine/Source/Runtime/Experimental/ChaosCore/Public/Chaos/ArrayCollection.h
@@ -89,7 +89,37 @@ protected:
 				Array->RemoveAt(Index, Count);
 			}
 		}
-		MSize = FMath::Min(static_cast<int32>(MSize) - Index, Count);
+		const int32 AvailableToRemove = MSize - Index;
+		MSize -= FMath::Min(AvailableToRemove, Count);
+	}
+
+	void RemoveAtSwapHelper(const int32 Index)
+	{
+		check(static_cast<uint32>(Index) < MSize);
+		for (TArrayCollectionArrayBase* Array : MArrays)
+		{
+			if (Array)
+			{
+				Array->RemoveAtSwap(Index);
+			}
+		}
+		MSize--;
+	}
+
+	void MoveToOtherArrayCollection(const int32 Index, TArrayCollection& Other)
+	{
+		check(MArrays.Num() == Other.MArrays.Num());
+		check(static_cast<uint32>(Index) < MSize);
+
+		for (int32 ArrayIdx = 0; ArrayIdx < MArrays.Num(); ++ArrayIdx)
+		{
+			if (TArrayCollectionArrayBase* Array = MArrays[ArrayIdx])
+			{
+				Array->MoveToOtherArray(Index, *Other.MArrays[ArrayIdx]);
+			}
+		}
+		++Other.MSize;
+		--MSize;
 	}
 
 private:
diff --git a/Engine/Source/Runtime/Experimental/ChaosCore/Public/Chaos/ArrayCollectionArray.h b/Engine/Source/Runtime/Experimental/ChaosCore/Public/Chaos/ArrayCollectionArray.h
--- a/Engine/Source/Runtime/Experimental/ChaosCore/Public/Chaos/ArrayCollectionArray.h
+++ b/Engine/Source/Runtime/Experimental/ChaosCore/Public/Chaos/ArrayCollectionArray.h
@@ -13,6 +13,8 @@ class TArrayCollectionArray : public TArrayCollectionArrayBase, public TArray<T>
 {
 	using TArray<T>::SetNum;
 	using TArray<T>::RemoveAt;
+	using TArray<T>::RemoveAtSwap;
+	using TArray<T>::Emplace;
 
   public:
 	TArrayCollectionArray()
@@ -25,6 +27,12 @@ class TArrayCollectionArray : public TArrayCollectionArrayBase, public TArray<T>
 		TArray<T>::operator=(MoveTemp(Other));
 		return *this;
 	}
+
+	TArrayCollectionArray(TArray<T>&& Other)
+	: TArray<T>(MoveTemp(Other))
+	{
+	}
+
 	virtual ~TArrayCollectionArray() {}
 
 	TArrayCollectionArray<T> Clone()
@@ -39,46 +47,27 @@ class TArrayCollectionArray : public TArrayCollectionArrayBase, public TArray<T>
 		SetNum(Num);
 	}
 
-	void RemoveAt(const int Num, const int Count) override
+	void RemoveAt(const int Idx, const int Count) override
 	{
-		TArray<T>::RemoveAt(Num, Count);
+		TArray<T>::RemoveAt(Idx, Count);
 	}
 
-	uint64 SizeOfElem() const override
+	void RemoveAtSwap(const int Idx) override
 	{
-		return sizeof(T);
+		TArray<T>::RemoveAtSwap(Idx);
 	}
-};
 
-template<class T>
-class TArrayCollectionArrayView : public TArrayCollectionArrayBase, public TArrayView<T>
-{
-  public:
-	TArrayCollectionArrayView(T* BasePtr, const int32 Num)
-	    : TArrayView<T>(BasePtr, Num) {}
-	TArrayCollectionArrayView(const TArrayCollectionArrayView<T>& Other) = delete;
-	TArrayCollectionArrayView(TArrayCollectionArrayView<T>&& Other)
-	    : TArrayView<T>(MoveTemp(Other)) {}
-	TArrayCollectionArrayView& operator=(TArrayCollectionArrayView<T>&& Other)
+	void MoveToOtherArray(const int Idx, TArrayCollectionArrayBase& Other)
 	{
-		TArrayView<T>::operator=(MoveTemp(Other));
-		return *this;
-	}
-	virtual ~TArrayCollectionArrayView() {}
-
-	void Resize(const int Num) override
-	{
-		check(Num == TArrayView<T>::Num());
-	}
-
-	void RemoveAt(const int Num, const int Count) override
-	{
-		check(Count == 0);
+		//todo: add developer check to make sure this is ok?
+		auto& OtherTArray = static_cast<TArrayCollectionArray<T>&>(Other);
+		OtherTArray.Emplace(MoveTemp(TArray<T>::operator [](Idx)));
+		TArray<T>::RemoveAtSwap(Idx);
 	}
 
 	uint64 SizeOfElem() const override
 	{
-		return 0;
+		return sizeof(T);
 	}
 };
 }
diff --git a/Engine/Source/Runtime/Experimental/ChaosCore/Public/Chaos/ArrayCollectionArrayBase.h b/Engine/Source/Runtime/Experimental/ChaosCore/Public/Chaos/ArrayCollectionArrayBase.h
--- a/Engine/Source/Runtime/Experimental/ChaosCore/Public/Chaos/ArrayCollectionArrayBase.h
+++ b/Engine/Source/Runtime/Experimental/ChaosCore/Public/Chaos/ArrayCollectionArrayBase.h
@@ -7,7 +7,9 @@ class TArrayCollectionArrayBase
 {
   public:
 	virtual void Resize(const int Num) = 0;
-	virtual void RemoveAt(const int Num, const int Count) = 0;
+	virtual void RemoveAt(const int Idx, const int Count) = 0;
+	virtual void RemoveAtSwap(const int Idx) = 0;
 	virtual uint64 SizeOfElem() const = 0;
+	virtual void MoveToOtherArray(const int Idx, TArrayCollectionArrayBase& Other) = 0;
 };
 }
diff --git a/Engine/Source/Runtime/Experimental/ChaosSolverEngine/Private/Chaos/ChaosEventListenerComponent.cpp b/Engine/Source/Runtime/Experimental/ChaosSolverEngine/Private/Chaos/ChaosEventListenerComponent.cpp
--- a/Engine/Source/Runtime/Experimental/ChaosSolverEngine/Private/Chaos/ChaosEventListenerComponent.cpp
+++ b/Engine/Source/Runtime/Experimental/ChaosSolverEngine/Private/Chaos/ChaosEventListenerComponent.cpp
@@ -1,7 +1,7 @@
 // Copyright 1998-2019 Epic Games, Inc. All Rights Reserved.
 
 #include "Chaos/ChaosEventListenerComponent.h"
-#include "PBDRigidsSolver.h"
+#include "PhysicsSolver.h"
 #include "Chaos/ChaosSolverActor.h"
 
 UChaosEventListenerComponent::UChaosEventListenerComponent()
@@ -22,7 +22,7 @@ const AChaosSolverActor* UChaosEventListenerComponent::GetSolverActor() const
 }
 
 #if INCLUDE_CHAOS
-const Chaos::FPBDRigidsSolver* UChaosEventListenerComponent::GetSolver() const
+const Chaos::FPhysicsSolver* UChaosEventListenerComponent::GetSolver() const
 {
 	const AChaosSolverActor* A = GetSolverActor();
 	return A ? A->GetSolver() : nullptr;
diff --git a/Engine/Source/Runtime/Experimental/ChaosSolverEngine/Private/Chaos/ChaosGameplayEventDispatcher.cpp b/Engine/Source/Runtime/Experimental/ChaosSolverEngine/Private/Chaos/ChaosGameplayEventDispatcher.cpp
--- a/Engine/Source/Runtime/Experimental/ChaosSolverEngine/Private/Chaos/ChaosGameplayEventDispatcher.cpp
+++ b/Engine/Source/Runtime/Experimental/ChaosSolverEngine/Private/Chaos/ChaosGameplayEventDispatcher.cpp
@@ -3,8 +3,8 @@
 #include "Chaos/ChaosGameplayEventDispatcher.h"
 #include "Components/PrimitiveComponent.h"
 #include "PhysicsEngine/BodyInstance.h"
-#include "SolverObjects/SolverObject.h"
-#include "PBDRigidsSolver.h"
+#include "Chaos/Framework/PhysicsProxy.h"
+#include "PhysicsSolver.h"
 #include "Chaos/PBDCollisionTypes.h"
 #include "Physics/Experimental/PhysScene_Chaos.h"
 #include "Engine/World.h"
@@ -15,6 +15,22 @@
 
 // PRAGMA_DISABLE_OPTIMIZATION
 
+void UChaosGameplayEventDispatcher::BeginPlay()
+{
+	Super::BeginPlay();
+#if INCLUDE_CHAOS
+	RegisterChaosEvents();
+#endif
+}
+
+void UChaosGameplayEventDispatcher::EndPlay(const EEndPlayReason::Type EndPlayReason)
+{
+#if INCLUDE_CHAOS
+	UnregisterChaosEvents();
+#endif
+	Super::EndPlay(EndPlayReason);
+}
+
 // internal
 static void DispatchPendingBreakEvents(TArray<FChaosBreakEvent> const& Events, TMap<UPrimitiveComponent*, FBreakEventCallbackWrapper> const& Registrations)
 {
@@ -51,255 +67,6 @@ static void SetCollisionInfoFromComp(FRigidBodyCollisionInfo& Info, UPrimitiveCo
 	}
 }
 
-void UChaosGameplayEventDispatcher::TickComponent(float DeltaTime, enum ELevelTick TickType, FActorComponentTickFunction *ThisTickFunction)
-{
-	SCOPE_CYCLE_COUNTER(STAT_DispatchEventNotifies);
-
-	Super::TickComponent(DeltaTime, TickType, ThisTickFunction);
-
-#if INCLUDE_CHAOS
-	// #todo: determine which way is better -- iterate the whole collision list or iterate the list of components that want notifies
-	// collision list could be very small or zero for lots of frames
-	// but collision list could be very large for certain frames
-	// so do we optimize for the steady state or for the spikes?
-
-	FPhysScene_Chaos const* const Scene = GetPhysicsScene().Get();
-	Chaos::FPBDRigidsSolver const* const Solver = GetSolver();
-	if (Scene && Solver)
-	{
-		Chaos::FPBDRigidsSolver::FScopedGetEventsData ScopedAccess = Solver->ScopedGetEventsData();
-
-		SCOPE_CYCLE_COUNTER(STAT_DispatchCollisionEvents);
-
-		// COLLISION EVENTS
-		{
-			PendingChaosCollisionNotifies.Reset();
-			ContactPairToPendingNotifyMap.Reset();
-
-			// get collision data from the solver
-			const Chaos::FPBDRigidsSolver::FAllCollisionDataMaps& AllCollisionData_Maps = ScopedAccess.GetAllCollisions_Maps();
-
-			if (AllCollisionData_Maps.AllCollisionData && AllCollisionData_Maps.AllCollisionsIndicesBySolverObject && AllCollisionData_Maps.SolverObjectReverseMapping)
-			{
-				const float CollisionDataTimestamp = AllCollisionData_Maps.AllCollisionData->TimeCreated;
-				if (CollisionDataTimestamp > LastCollisionDataTime)
-				{
-					LastCollisionDataTime = CollisionDataTimestamp;
-
-					TMap<ISolverObjectBase*, TArray<int32>> const& SolverObjectToCollisionIndicesMap = AllCollisionData_Maps.AllCollisionsIndicesBySolverObject->AllCollisionsIndicesBySolverObjectMap;
-					Chaos::FPBDRigidsSolver::FCollisionDataArray const& CollisionData = AllCollisionData_Maps.AllCollisionData->AllCollisionsArray;
-
-					int32 NumCollisions = CollisionData.Num();		// is this zero enough frames to be worth the check?
-					if (NumCollisions > 0)
-					{
-						TArray<SolverObjectWrapper> const& SolverObjectReverseMappingArray = AllCollisionData_Maps.SolverObjectReverseMapping->SolverObjectReverseMappingArray;
-						TArray<int32> const& ParticleIndexReverseMappingArray = AllCollisionData_Maps.ParticleIndexReverseMapping->ParticleIndexReverseMappingArray;
-
-						// look through all the components that someone is interested in, and see if they had a collision
-						// note that we only need to care about the interaction from the POV of the registered component,
-						// since if anyone wants notifications for the other component it hit, it's also registered and we'll get to that elsewhere in the list
-						for (TMap<UPrimitiveComponent*, FChaosHandlerSet>::TIterator It(CollisionEventRegistrations); It; ++It)
-						{
-							const FChaosHandlerSet& HandlerSet = It.Value();
-
-							UPrimitiveComponent* const Comp0 = Cast<UPrimitiveComponent>(It.Key());
-							ISolverObjectBase* const SolverObject0 = Scene->GetOwnedSolverObject(Comp0);
-							TArray<int32> const* const CollisionIndices = SolverObjectToCollisionIndicesMap.Find(SolverObject0);
-							if (CollisionIndices)
-							{
-								for (int32 EncodedCollisionIdx : *CollisionIndices)
-								{
-									bool bSwapOrder;
-									int32 CollisionIdx = Chaos::FPBDRigidsSolver::DecodeCollisionIndex(EncodedCollisionIdx, bSwapOrder);
-
-									Chaos::TCollisionData<float, 3> const& CollisionDataItem = CollisionData[CollisionIdx];
-									const int32 ParticleIndex1 = bSwapOrder ? CollisionDataItem.ParticleIndex : CollisionDataItem.LevelsetIndex;
-									ISolverObjectBase* const SolverObject1 = SolverObjectReverseMappingArray[ParticleIndex1].SolverObject;
-
-									if (HandlerSet.bLegacyComponentNotify)
-									{
-										bool bNewEntry = false;
-										FCollisionNotifyInfo& NotifyInfo = GetPendingCollisionForContactPair(SolverObject0, SolverObject1, bNewEntry);
-
-										// #note: we only notify on the first contact, though we will still accumulate the impulse data from subsequent contacts
-										const FVector NormalImpulse = FVector::DotProduct(CollisionDataItem.AccumulatedImpulse, CollisionDataItem.Normal) * CollisionDataItem.Normal;	// project impulse along normal
-										const FVector FrictionImpulse = FVector(CollisionDataItem.AccumulatedImpulse) - NormalImpulse; // friction is component not along contact normal
-										NotifyInfo.RigidCollisionData.TotalNormalImpulse += NormalImpulse;
-										NotifyInfo.RigidCollisionData.TotalFrictionImpulse += FrictionImpulse;
-
-										if (bNewEntry)
-										{
-											UPrimitiveComponent* const Comp1 = Scene->GetOwningComponent<UPrimitiveComponent>(SolverObject1);
-
-											// fill in legacy contact data
-											NotifyInfo.bCallEvent0 = true;
-											// if Comp1 wants this event too, it will get its own pending collision entry, so we leave it false
-
-											SetCollisionInfoFromComp(NotifyInfo.Info0, Comp0);
-											SetCollisionInfoFromComp(NotifyInfo.Info1, Comp1);
-
-											FRigidBodyContactInfo& NewContact = NotifyInfo.RigidCollisionData.ContactInfos.AddZeroed_GetRef();
-											NewContact.ContactNormal = CollisionDataItem.Normal;
-											NewContact.ContactPosition = CollisionDataItem.Location;
-											NewContact.ContactPenetration = CollisionDataItem.PenetrationDepth; 
-											// NewContact.PhysMaterial[1] UPhysicalMaterial required here
-										}
-
-									}
-
-									if (HandlerSet.ChaosHandlers.Num() > 0)
-									{
-										bool bNewEntry = false;
-										FChaosPendingCollisionNotify& ChaosNotifyInfo = GetPendingChaosCollisionForContactPair(SolverObject0, SolverObject1, bNewEntry);
-
-										// #note: we only notify on the first contact, though we will still accumulate the impulse data from subsequent contacts
-										ChaosNotifyInfo.CollisionInfo.AccumulatedImpulse += CollisionDataItem.AccumulatedImpulse;
-
-										if (bNewEntry)
-										{
-											UPrimitiveComponent* const Comp1 = Scene->GetOwningComponent<UPrimitiveComponent>(SolverObject1);
-
-											// fill in Chaos contact data
-											ChaosNotifyInfo.CollisionInfo.Component = Comp0;
-											ChaosNotifyInfo.CollisionInfo.OtherComponent = Comp1;
-											ChaosNotifyInfo.CollisionInfo.Location = CollisionDataItem.Location;
-											ChaosNotifyInfo.NotifyRecipients = HandlerSet.ChaosHandlers;
-
-											if (bSwapOrder)
-											{
-												ChaosNotifyInfo.CollisionInfo.AccumulatedImpulse = -CollisionDataItem.AccumulatedImpulse;
-												ChaosNotifyInfo.CollisionInfo.Normal = -CollisionDataItem.Normal;
-
-												ChaosNotifyInfo.CollisionInfo.Velocity = CollisionDataItem.Velocity2;
-												ChaosNotifyInfo.CollisionInfo.OtherVelocity = CollisionDataItem.Velocity1;
-												ChaosNotifyInfo.CollisionInfo.AngularVelocity = CollisionDataItem.AngularVelocity2;
-												ChaosNotifyInfo.CollisionInfo.OtherAngularVelocity = CollisionDataItem.AngularVelocity1;
-												ChaosNotifyInfo.CollisionInfo.Mass = CollisionDataItem.Mass2;
-												ChaosNotifyInfo.CollisionInfo.OtherMass = CollisionDataItem.Mass1;
-											}
-											else
-											{
-												ChaosNotifyInfo.CollisionInfo.AccumulatedImpulse = CollisionDataItem.AccumulatedImpulse;
-												ChaosNotifyInfo.CollisionInfo.Normal = CollisionDataItem.Normal;
-
-												ChaosNotifyInfo.CollisionInfo.Velocity = CollisionDataItem.Velocity1;
-												ChaosNotifyInfo.CollisionInfo.OtherVelocity = CollisionDataItem.Velocity2;
-												ChaosNotifyInfo.CollisionInfo.AngularVelocity = CollisionDataItem.AngularVelocity1;
-												ChaosNotifyInfo.CollisionInfo.OtherAngularVelocity = CollisionDataItem.AngularVelocity2;
-												ChaosNotifyInfo.CollisionInfo.Mass = CollisionDataItem.Mass1;
-												ChaosNotifyInfo.CollisionInfo.OtherMass = CollisionDataItem.Mass2;
-											}
-										}
-									}
-								}
-							}
-						}
-					}
-				}
-			}
-
-			// Tell the world and actors about the collisions
-			DispatchPendingCollisionNotifies();
-		}
-
-		SCOPE_CYCLE_COUNTER(STAT_DispatchBreakEvents);
-
-		// BREAK EVENTS
-		// Same dilemma as above, do we iterate all the breaks or all the components registered for notifications? Or switch based on heuristic?
-		{
-			TArray<FChaosBreakEvent> PendingBreakEvents;
-
-			// get break data from the solver
-			const Chaos::FPBDRigidsSolver::FAllBreakingDataMaps& AllBreakingData_Maps = ScopedAccess.GetAllBreakings_Maps();
-
-			if (AllBreakingData_Maps.AllBreakingData && AllBreakingData_Maps.SolverObjectReverseMapping)
-			{
-				const float BreakingDataTimestamp = AllBreakingData_Maps.AllBreakingData->TimeCreated;
-				if (BreakingDataTimestamp > LastBreakingDataTime)
-				{
-					LastBreakingDataTime = BreakingDataTimestamp;
-
-					Chaos::FPBDRigidsSolver::FBreakingDataArray const& BreakingData = AllBreakingData_Maps.AllBreakingData->AllBreakingsArray;
-
-					// let's assume breaks are very rare, so we will iterate breaks instead of registered components for now
-					const int32 NumBreaks = BreakingData.Num();
-					if (NumBreaks > 0)
-					{
-						// Array[ParticleIndex] = SolverObject
-						TArray<SolverObjectWrapper> const& ParticleIdxToSolverObjectArray = AllBreakingData_Maps.SolverObjectReverseMapping->SolverObjectReverseMappingArray;
-
-						for (Chaos::TBreakingData<float, 3> const& BreakingDataItem : BreakingData)
-						{
-							ISolverObjectBase* const SolverObject = ParticleIdxToSolverObjectArray[BreakingDataItem.ParticleIndex].SolverObject;
-							UPrimitiveComponent* const PrimComp = Scene->GetOwningComponent<UPrimitiveComponent>(SolverObject);
-							if (PrimComp && BreakEventRegistrations.Contains(PrimComp))
-							{
-								// queue them up so we can release the physics data before trigging BP events
-								FChaosBreakEvent& BreakEvent = PendingBreakEvents.AddZeroed_GetRef();
-								BreakEvent.Component = PrimComp;
-								BreakEvent.Location = BreakingDataItem.Location;
-								BreakEvent.Velocity = BreakingDataItem.Velocity;
-								BreakEvent.AngularVelocity = BreakingDataItem.AngularVelocity;
-								BreakEvent.Mass = BreakingDataItem.Mass;
-							}
-						}
-
-						DispatchPendingBreakEvents(PendingBreakEvents, BreakEventRegistrations);
-					}
-				}
-			}
-		}
-	}
-
-	// old iterate-the-entire-list way
-#if 0
-	const FPhysScene_Chaos* Scene = GetPhysicsScene().Get();
-	if (Scene)
-	{
-		// Populate the pending notify list from the data in the solver
-		PendingCollisionNotifies.Reset();
-		ContactPairToPendingNotifyMap.Reset();
-
-		const Chaos::PBDRigidsSolver* Solver = GetSolver();
-		const Chaos::PBDRigidsSolver::FCollisionData& CollisionData = Solver->GetCollisionData_GameThread();
-
-		for (Chaos::TCollisionData<float, 3> const& CollisionDataItem : CollisionData.CollisionDataArray)
-		{
-// 			Chaos::TPBDRigidParticles<float, 3>& Particles = Solver->GetRigidParticles();
-// 			Particles[CollisionDataItem.ParticleIndex].
-
-			ISolverObjectBase* const SolverObject0 = nullptr;		// #todo Solver->GetSolverObject(CollisionDataItem.ParticleIndex)
-			UPrimitiveComponent* Comp0 = Scene->GetOwner<UPrimitiveComponent>(SolverObject0);
-			FBodyInstance const* const BI0 = Comp0 ? Comp0->GetBodyInstance() : nullptr;
-
-			ISolverObjectBase* const SolverObject1 = nullptr;		// #todo Solver->GetSolverObject(CollisionDataItem.LevelsetIndex)
-			UPrimitiveComponent* Comp1 = Scene->GetOwner<UPrimitiveComponent>(SolverObject1);
-			FBodyInstance const* const BI1 = Comp1 ? Comp1->GetBodyInstance() : nullptr;
-
-			const bool bWantsNotify0 = BI0 && BI0->bNotifyRigidBodyCollision;
-			const bool bWantsNotify1 = BI1 && BI1->bNotifyRigidBodyCollision;
-
-			if (bWantsNotify0 || bWantsNotify1)
-			{
-//				const FUniqueContactPairKey Pair = { Comp0, Comp1 };
-				FCollisionNotifyInfo& NotifyInfo = GetPendingCollisionForContactPair(BI0, BI1);
-
-				// #todo, use a new contact info struct that contains the more detailed contact info we have available
-				FRigidBodyContactInfo& NewContact = NotifyInfo.RigidCollisionData.ContactInfos.AddZeroed_GetRef();
-				NewContact.ContactNormal = CollisionDataItem.Normal;
-				//NewContact.ContactPenetration = 0.f;		// #todo? 
-				NewContact.ContactPosition = CollisionDataItem.Location;
-				//NewContact.PhysMaterial = ;				// #todo?
-			}
-		}
-
-		// Tell the world and actors about the collisions
-		DispatchPendingNotifies();
-	}
-#endif		// 0
-
-#endif		// INCLUDE_CHAOS
-}
 
 FCollisionNotifyInfo& UChaosGameplayEventDispatcher::GetPendingCollisionForContactPair(const void* P0, const void* P1, bool& bNewEntry)
 {
@@ -430,4 +197,197 @@ void UChaosGameplayEventDispatcher::UnRegisterForBreakEvents(UPrimitiveComponent
 	}
 }
 
+#if INCLUDE_CHAOS
+
+void UChaosGameplayEventDispatcher::RegisterChaosEvents()
+{
+#if WITH_CHAOS
+	FPhysScene* Scene = GetWorld()->GetPhysicsScene();
+	Scene->RegisterEventHandler<Chaos::FCollisionEventData>(Chaos::EEventType::Collision, this, &UChaosGameplayEventDispatcher::HandleCollisionEvents);
+	//Scene->RegisterEventHandler<Chaos::FBreakingEventData>(Chaos::EEventType::Breaking, this, &UChaosGameplayEventDispatcher::HandleBreakingEvents);
+#endif
+}
+
+void UChaosGameplayEventDispatcher::UnregisterChaosEvents()
+{
+#if WITH_CHAOS
+	if (GetWorld())
+	{
+		FPhysScene* Scene = GetWorld()->GetPhysicsScene();
+		Scene->UnregisterEventHandler(Chaos::EEventType::Collision, this);
+		//Scene->UnregisterEventHandler(Chaos::EEventType::Breaking, this);
+	}
+#endif
+}
+
+void UChaosGameplayEventDispatcher::HandleCollisionEvents(const Chaos::FCollisionEventData& Event)
+{
+	SCOPE_CYCLE_COUNTER(STAT_DispatchCollisionEvents);
+
+	FPhysScene_Chaos const* const Scene = GetPhysicsScene().Get();
+
+	PendingChaosCollisionNotifies.Reset();
+	ContactPairToPendingNotifyMap.Reset();
+
+	{
+		TMap<IPhysicsProxyBase*, TArray<int32>> const& PhysicsProxyToCollisionIndicesMap = Event.PhysicsProxyToCollisionIndices.PhysicsProxyToIndicesMap;
+		Chaos::FCollisionDataArray const& CollisionData = Event.CollisionData.AllCollisionsArray;
+
+		int32 NumCollisions = CollisionData.Num();
+		if (NumCollisions > 0)
+		{
+			// look through all the components that someone is interested in, and see if they had a collision
+			// note that we only need to care about the interaction from the POV of the registered component,
+			// since if anyone wants notifications for the other component it hit, it's also registered and we'll get to that elsewhere in the list
+			for (TMap<UPrimitiveComponent*, FChaosHandlerSet>::TIterator It(CollisionEventRegistrations); It; ++It)
+			{
+				const FChaosHandlerSet& HandlerSet = It.Value();
+
+				UPrimitiveComponent* const Comp0 = Cast<UPrimitiveComponent>(It.Key());
+				IPhysicsProxyBase* const PhysicsProxy0 = Scene->GetOwnedPhysicsProxy(Comp0);
+				TArray<int32> const* const CollisionIndices = PhysicsProxyToCollisionIndicesMap.Find(PhysicsProxy0);
+				if (CollisionIndices)
+				{
+					for (int32 EncodedCollisionIdx : *CollisionIndices)
+					{
+						bool bSwapOrder;
+						int32 CollisionIdx = Chaos::FEventManager::DecodeCollisionIndex(EncodedCollisionIdx, bSwapOrder);
+
+						Chaos::TCollisionData<float, 3> const& CollisionDataItem = CollisionData[CollisionIdx];
+						const Chaos::TGeometryParticle<float, 3>* Particle1 = bSwapOrder ? CollisionDataItem.Particle : CollisionDataItem.Levelset;
+						UPrimitiveComponent* NewComp1 = Cast<UPrimitiveComponent>(Particle1->Proxy->GetOwner());
+						IPhysicsProxyBase* const PhysicsProxy1 = NewComp1 ? Scene->GetOwnedPhysicsProxy(NewComp1) : nullptr;
+
+						if (HandlerSet.bLegacyComponentNotify)
+						{
+							bool bNewEntry = false;
+							FCollisionNotifyInfo& NotifyInfo = GetPendingCollisionForContactPair(PhysicsProxy0, PhysicsProxy1, bNewEntry);
+
+							// #note: we only notify on the first contact, though we will still accumulate the impulse data from subsequent contacts
+							const FVector NormalImpulse = FVector::DotProduct(CollisionDataItem.AccumulatedImpulse, CollisionDataItem.Normal) * CollisionDataItem.Normal;	// project impulse along normal
+							const FVector FrictionImpulse = FVector(CollisionDataItem.AccumulatedImpulse) - NormalImpulse; // friction is component not along contact normal
+							NotifyInfo.RigidCollisionData.TotalNormalImpulse += NormalImpulse;
+							NotifyInfo.RigidCollisionData.TotalFrictionImpulse += FrictionImpulse;
+
+							if (bNewEntry)
+							{
+								UPrimitiveComponent* const Comp1 = Scene->GetOwningComponent<UPrimitiveComponent>(PhysicsProxy1);
+
+								// fill in legacy contact data
+								NotifyInfo.bCallEvent0 = true;
+								// if Comp1 wants this event too, it will get its own pending collision entry, so we leave it false
+
+								SetCollisionInfoFromComp(NotifyInfo.Info0, Comp0);
+								SetCollisionInfoFromComp(NotifyInfo.Info1, Comp1);
+
+								FRigidBodyContactInfo& NewContact = NotifyInfo.RigidCollisionData.ContactInfos.AddZeroed_GetRef();
+								NewContact.ContactNormal = CollisionDataItem.Normal;
+								NewContact.ContactPosition = CollisionDataItem.Location;
+								NewContact.ContactPenetration = CollisionDataItem.PenetrationDepth;
+								// NewContact.PhysMaterial[1] UPhysicalMaterial required here
+							}
+
+						}
+
+						if (HandlerSet.ChaosHandlers.Num() > 0)
+						{
+							bool bNewEntry = false;
+							FChaosPendingCollisionNotify& ChaosNotifyInfo = GetPendingChaosCollisionForContactPair(PhysicsProxy0, PhysicsProxy1, bNewEntry);
+
+							// #note: we only notify on the first contact, though we will still accumulate the impulse data from subsequent contacts
+							ChaosNotifyInfo.CollisionInfo.AccumulatedImpulse += CollisionDataItem.AccumulatedImpulse;
+
+							if (bNewEntry)
+							{
+								UPrimitiveComponent* const Comp1 = Scene->GetOwningComponent<UPrimitiveComponent>(PhysicsProxy1);
+
+								// fill in Chaos contact data
+								ChaosNotifyInfo.CollisionInfo.Component = Comp0;
+								ChaosNotifyInfo.CollisionInfo.OtherComponent = Comp1;
+								ChaosNotifyInfo.CollisionInfo.Location = CollisionDataItem.Location;
+								ChaosNotifyInfo.NotifyRecipients = HandlerSet.ChaosHandlers;
+
+								if (bSwapOrder)
+								{
+									ChaosNotifyInfo.CollisionInfo.AccumulatedImpulse = -CollisionDataItem.AccumulatedImpulse;
+									ChaosNotifyInfo.CollisionInfo.Normal = -CollisionDataItem.Normal;
+
+									ChaosNotifyInfo.CollisionInfo.Velocity = CollisionDataItem.Velocity2;
+									ChaosNotifyInfo.CollisionInfo.OtherVelocity = CollisionDataItem.Velocity1;
+									ChaosNotifyInfo.CollisionInfo.AngularVelocity = CollisionDataItem.AngularVelocity2;
+									ChaosNotifyInfo.CollisionInfo.OtherAngularVelocity = CollisionDataItem.AngularVelocity1;
+									ChaosNotifyInfo.CollisionInfo.Mass = CollisionDataItem.Mass2;
+									ChaosNotifyInfo.CollisionInfo.OtherMass = CollisionDataItem.Mass1;
+								}
+								else
+								{
+									ChaosNotifyInfo.CollisionInfo.AccumulatedImpulse = CollisionDataItem.AccumulatedImpulse;
+									ChaosNotifyInfo.CollisionInfo.Normal = CollisionDataItem.Normal;
+
+									ChaosNotifyInfo.CollisionInfo.Velocity = CollisionDataItem.Velocity1;
+									ChaosNotifyInfo.CollisionInfo.OtherVelocity = CollisionDataItem.Velocity2;
+									ChaosNotifyInfo.CollisionInfo.AngularVelocity = CollisionDataItem.AngularVelocity1;
+									ChaosNotifyInfo.CollisionInfo.OtherAngularVelocity = CollisionDataItem.AngularVelocity2;
+									ChaosNotifyInfo.CollisionInfo.Mass = CollisionDataItem.Mass1;
+									ChaosNotifyInfo.CollisionInfo.OtherMass = CollisionDataItem.Mass2;
+								}
+							}
+						}
+					}
+				}
+			}
+		}
+	}
+
+	// Tell the world and actors about the collisions
+	DispatchPendingCollisionNotifies();
+}
+
+void UChaosGameplayEventDispatcher::HandleBreakingEvents(const Chaos::FBreakingEventData& Event)
+{
+
+	SCOPE_CYCLE_COUNTER(STAT_DispatchBreakEvents);
+
+	// BREAK EVENTS
+
+	TArray<FChaosBreakEvent> PendingBreakEvents;
+
+	const float BreakingDataTimestamp = Event.BreakingData.TimeCreated;
+	if (BreakingDataTimestamp > LastBreakingDataTime)
+	{
+		LastBreakingDataTime = BreakingDataTimestamp;
+
+		Chaos::FBreakingDataArray const& BreakingData = Event.BreakingData.AllBreakingsArray;
+
+		// let's assume breaks are very rare, so we will iterate breaks instead of registered components for now
+		const int32 NumBreaks = BreakingData.Num();
+		if (NumBreaks > 0)
+		{
+			for (Chaos::TBreakingData<float, 3> const& BreakingDataItem : BreakingData)
+			{	
+				if (BreakingDataItem.Particle && (BreakingDataItem.Particle->Proxy))
+				{
+					UPrimitiveComponent* const PrimComp = Cast<UPrimitiveComponent>(BreakingDataItem.Particle->Proxy->GetOwner());
+					if (PrimComp && BreakEventRegistrations.Contains(PrimComp))
+					{
+						// queue them up so we can release the physics data before trigging BP events
+						FChaosBreakEvent& BreakEvent = PendingBreakEvents.AddZeroed_GetRef();
+							BreakEvent.Component = PrimComp;
+							BreakEvent.Location = BreakingDataItem.Location;
+							BreakEvent.Velocity = BreakingDataItem.Velocity;
+							BreakEvent.AngularVelocity = BreakingDataItem.AngularVelocity;
+							BreakEvent.Mass = BreakingDataItem.Mass;
+					}
+				}
+			}
+
+			DispatchPendingBreakEvents(PendingBreakEvents, BreakEventRegistrations);
+		}
+	}
+
+}
+
+#endif // INCLUDE_CHAOS
+
+
 // PRAGMA_ENABLE_OPTIMIZATION
\ No newline at end of file
diff --git a/Engine/Source/Runtime/Experimental/ChaosSolverEngine/Private/Chaos/ChaosSolverActor.cpp b/Engine/Source/Runtime/Experimental/ChaosSolverEngine/Private/Chaos/ChaosSolverActor.cpp
--- a/Engine/Source/Runtime/Experimental/ChaosSolverEngine/Private/Chaos/ChaosSolverActor.cpp
+++ b/Engine/Source/Runtime/Experimental/ChaosSolverEngine/Private/Chaos/ChaosSolverActor.cpp
@@ -2,7 +2,7 @@
 
 #include "Chaos/ChaosSolverActor.h"
 #include "UObject/ConstructorHelpers.h"
-#include "PBDRigidsSolver.h"
+#include "PhysicsSolver.h"
 #include "ChaosModule.h"
 
 #include "Components/BillboardComponent.h"
@@ -23,6 +23,7 @@
 
 #if CHAOS_DEBUG_SUBSTEP
 #include "HAL/IConsoleManager.h"
+#include "Chaos/Utilities.h"
 
 class FChaosSolverActorConsoleObjects final
 {
@@ -72,11 +73,10 @@ private:
 	{
 		for (auto SolverActor: SolverActors)
 		{
-			const Chaos::FPBDRigidsSolver* const Solver = SolverActor.Value->GetSolver();
+			const Chaos::FPhysicsSolver* const Solver = SolverActor.Value->GetSolver();
 			if (Solver)
 			{
-				const FSolverObjectStorage& SolverObjectStorage = Solver->GetObjectStorage_GameThread();
-				UE_LOG(LogChaosDebug, Display, TEXT("%s (%d objects)"), *SolverActor.Key, SolverObjectStorage.GetNumObjects());
+				UE_LOG(LogChaosDebug, Display, TEXT("%s (%d objects)"), *SolverActor.Key, Solver->GetNumPhysicsProxies());
 			}
 		}
 	}
@@ -84,7 +84,7 @@ private:
 	void Pause(const TArray<FString>& Args)
 	{
 		AChaosSolverActor* const* SolverActor;
-		Chaos::FPBDRigidsSolver* Solver;
+		Chaos::FPhysicsSolver* Solver;
 		switch (Args.Num())
 		{
 		default:
@@ -101,6 +101,7 @@ private:
 			if ((SolverActor = SolverActors.Find(Args[0])) != nullptr &&
 				(Solver = (*SolverActor)->GetSolver()) != nullptr)
 			{
+#if TODO_REIMPLEMENT_DEBUG_SUBSTEP
 				if (Args[1] == TEXT("0"))
 				{
 					Solver->GetDebugSubstep().Enable(false);
@@ -119,6 +120,7 @@ private:
 #endif
 					return;
 				}
+#endif
 			}
 			break;  // Invalid arguments
 		}
@@ -132,8 +134,9 @@ private:
 
 	void Step(const TArray<FString>& Args)
 	{
+#if TODO_REIMPLEMENT_DEBUG_SUBSTEP
 		AChaosSolverActor* const* SolverActor;
-		Chaos::FPBDRigidsSolver* Solver;
+		Chaos::FPhysicsSolver* Solver;
 		switch (Args.Num())
 		{
 		default:
@@ -152,12 +155,14 @@ private:
 		UE_LOG(LogChaosDebug, Display, TEXT("  p.Chaos.Solver.Step [SolverName]"));
 		UE_LOG(LogChaosDebug, Display, TEXT("  SolverName  The Id name of the solver as shown by p.Chaos.Solver.List"));
 		UE_LOG(LogChaosDebug, Display, TEXT("Example: p.Chaos.Solver.Step ChaosSolverActor_3"));
+#endif
 	}
 
 	void Substep(const TArray<FString>& Args)
 	{
+#if TODO_REIMPLEMENT_DEBUG_SUBSTEP
 		AChaosSolverActor* const* SolverActor;
-		Chaos::FPBDRigidsSolver* Solver;
+		Chaos::FPhysicsSolver* Solver;
 		switch (Args.Num())
 		{
 		default:
@@ -176,6 +181,7 @@ private:
 		UE_LOG(LogChaosDebug, Display, TEXT("  p.Chaos.Solver.Substep [SolverName]"));
 		UE_LOG(LogChaosDebug, Display, TEXT("  SolverName  The Id name of the solver as shown by p.Chaos.Solver.List"));
 		UE_LOG(LogChaosDebug, Display, TEXT("Example: p.Chaos.Solver.Substep ChaosSolverActor_3"));
+#endif
 	}
 
 private:
@@ -205,6 +211,8 @@ AChaosSolverActor::AChaosSolverActor(const FObjectInitializer& ObjectInitializer
 	// @question(Benn) : Does this need to be created on the Physics thread using a queued command?
 	PhysScene = MakeShareable(new FPhysScene_Chaos(this));
 	Solver = PhysScene->GetSolver();
+	// @todo: This should be gated by a parameter
+	Solver->GetEvolution()->AddForceFunction(Chaos::Utilities::GetRigidsGravityFunction(Chaos::TVector<float, 3>(0, 0, -1), 980.f));
 #endif
 	// Ticking setup for collision/breaking notifies
 	PrimaryActorTick.TickGroup = TG_PostPhysics;
@@ -265,7 +273,7 @@ void AChaosSolverActor::BeginPlay()
 	Chaos::IDispatcher* PhysDispatcher = PhysScene->GetDispatcher();
 	if (PhysDispatcher)
 	{
-		PhysDispatcher->EnqueueCommand(Solver,
+		PhysDispatcher->EnqueueCommandImmediate(Solver,
 			[InTimeStepMultiplier = TimeStepMultiplier
 			, InCollisionIterations = CollisionIterations
 			, InPushOutIterations = PushOutIterations
@@ -281,14 +289,16 @@ void AChaosSolverActor::BeginPlay()
 			, InHasFloor = bHasFloor
 			, InFloorHeight = FloorHeight
 			, InMassScale = MassScale]
-		(Chaos::FPBDRigidsSolver* InSolver)
+		(Chaos::FPhysicsSolver* InSolver)
 		{
+#if TODO_REIMPLEMENT_SOLVER_SETTINGS_ACCESSORS
 			InSolver->SetTimeStepMultiplier(InTimeStepMultiplier);
 			InSolver->SetIterations(InCollisionIterations);
 			InSolver->SetPushOutIterations(InPushOutIterations);
 			InSolver->SetPushOutPairIterations(InPushOutPairIterations);
 			InSolver->SetClusterConnectionFactor(InClusterConnectionFactor);
 			InSolver->SetClusterUnionConnectionType((Chaos::FClusterCreationParameters<float>::EConnectionMethod)InClusterUnionConnectionType);
+#endif
 			InSolver->SetGenerateCollisionData(InDoGenerateCollisionData);
 			InSolver->SetGenerateBreakingData(InDoGenerateBreakingData);
 			InSolver->SetGenerateTrailingData(InDoGenerateTrailingData);
@@ -297,14 +307,20 @@ void AChaosSolverActor::BeginPlay()
 			InSolver->SetTrailingFilterSettings(InTrailingFilterSettings);
 			InSolver->SetHasFloor(InHasFloor);
 			InSolver->SetFloorHeight(InFloorHeight);
+#if TODO_REIMPLEMENT_SOLVER_SETTINGS_ACCESSORS
 			InSolver->SetMassScale(InMassScale);
+#endif
 			InSolver->SetEnabled(true);
+#if TODO_REIMPLEMENT_SOLVER_SETTINGS_ACCESSORS
 #if CHAOS_WITH_PAUSABLE_SOLVER
 			InSolver->SetPaused(false);
 #endif  // #if CHAOS_WITH_PAUSABLE_SOLVER
+#endif  // #if TODO_REIMPLEMENT_SOLVER_SETTINGS_ACCESSORS
 		});
 	}
 #endif
+
+#if TODO_REIMPLEMENT_DEBUG_SUBSTEP
 #if CHAOS_DEBUG_SUBSTEP
 	if (!ChaosSolverActorConsoleObjects)
 	{
@@ -318,6 +334,7 @@ void AChaosSolverActor::BeginPlay()
 	}
 #endif  // #if WITH_EDITOR
 #endif  // #if CHAOS_DEBUG_SUBSTEP
+#endif  // #if TODO_REIMPLEMENT_DEBUG_SUBSTEP
 }
 
 void AChaosSolverActor::EndPlay(EEndPlayReason::Type ReasonEnd)
@@ -326,7 +343,7 @@ void AChaosSolverActor::EndPlay(EEndPlayReason::Type ReasonEnd)
 	Chaos::IDispatcher* PhysDispatcher = PhysScene->GetDispatcher();
 	if (PhysDispatcher)
 	{
-		PhysDispatcher->EnqueueCommand(Solver, [](Chaos::FPBDRigidsSolver* InSolver)
+		PhysDispatcher->EnqueueCommandImmediate(Solver, [](Chaos::FPhysicsSolver* InSolver)
 		{
 			// #TODO BG - We should really reset the solver here but the current reset function
 			// is really heavy handed and clears out absolutely everything. Ideally we want to keep
@@ -335,9 +352,11 @@ void AChaosSolverActor::EndPlay(EEndPlayReason::Type ReasonEnd)
 		});
 	}
 #endif
+#if TODO_REIMPLEMENT_DEBUG_SUBSTEP
 #if CHAOS_DEBUG_SUBSTEP
 	ChaosSolverActorConsoleObjects->RemoveSolver(GetName());
 #endif  // #if CHAOS_DEBUG_SUBSTEP
+#endif  // #if TODO_REIMPLEMENT_DEBUG_SUBSTEP
 }
 
 void AChaosSolverActor::PostRegisterAllComponents()
@@ -373,7 +392,7 @@ void AChaosSolverActor::SetSolverActive(bool bActive)
 
 		if(Dispatcher)
 		{
-			Dispatcher->EnqueueCommand(Solver, [bShouldBeActive = bActive](Chaos::FPBDRigidsSolver* InSolver)
+			Dispatcher->EnqueueCommandImmediate(Solver, [bShouldBeActive = bActive](Chaos::FPhysicsSolver* InSolver)
 			{
 				InSolver->SetEnabled(bShouldBeActive);
 			});
@@ -393,114 +412,121 @@ void AChaosSolverActor::PostEditChangeProperty(struct FPropertyChangedEvent& Pro
 		Chaos::IDispatcher* PhysDispatcher = PhysScene->GetDispatcher();
 		if (PhysDispatcher)
 		{
+#if TODO_REIMPLEMENT_TIMESTEP_MULTIPLIER
 			if (PropertyChangedEvent.Property->GetFName() == GET_MEMBER_NAME_CHECKED(AChaosSolverActor, TimeStepMultiplier))
 			{
-				PhysDispatcher->EnqueueCommand(Solver, [InTimeStepMultiplier = TimeStepMultiplier]
-				(Chaos::FPBDRigidsSolver* InSolver)
+				PhysDispatcher->EnqueueCommandImmediate(Solver, [InTimeStepMultiplier = TimeStepMultiplier]
+				(Chaos::FPhysicsSolver* InSolver)
 				{
 					InSolver->SetTimeStepMultiplier(InTimeStepMultiplier);
 				});
 			}
 			else if (PropertyChangedEvent.Property->GetFName() == GET_MEMBER_NAME_CHECKED(AChaosSolverActor, CollisionIterations))
 			{
-				PhysDispatcher->EnqueueCommand(Solver, [InCollisionIterations = CollisionIterations]
-				(Chaos::FPBDRigidsSolver* InSolver)
+				PhysDispatcher->EnqueueCommandImmediate(Solver, [InCollisionIterations = CollisionIterations]
+				(Chaos::FPhysicsSolver* InSolver)
 				{
 					InSolver->SetIterations(InCollisionIterations);
 				});
 			}
 			else if (PropertyChangedEvent.Property->GetFName() == GET_MEMBER_NAME_CHECKED(AChaosSolverActor, PushOutIterations))
 			{
-				PhysDispatcher->EnqueueCommand(Solver, [InPushOutIterations = PushOutIterations]
-				(Chaos::FPBDRigidsSolver* InSolver)
+				PhysDispatcher->EnqueueCommandImmediate(Solver, [InPushOutIterations = PushOutIterations]
+				(Chaos::FPhysicsSolver* InSolver)
 				{
 					InSolver->SetPushOutIterations(InPushOutIterations);
 				});
 			}
 			else if (PropertyChangedEvent.Property->GetFName() == GET_MEMBER_NAME_CHECKED(AChaosSolverActor, PushOutPairIterations))
 			{
-				PhysDispatcher->EnqueueCommand(Solver, [InPushOutPairIterations = PushOutPairIterations]
-				(Chaos::FPBDRigidsSolver* InSolver)
+				PhysDispatcher->EnqueueCommandImmediate(Solver, [InPushOutPairIterations = PushOutPairIterations]
+				(Chaos::FPhysicsSolver* InSolver)
 				{
 					InSolver->SetPushOutPairIterations(InPushOutPairIterations);
 				});
 			}
-			else if (PropertyChangedEvent.Property->GetFName() == GET_MEMBER_NAME_CHECKED(AChaosSolverActor, DoGenerateCollisionData))
+			else
+#endif
+			if (PropertyChangedEvent.Property->GetFName() == GET_MEMBER_NAME_CHECKED(AChaosSolverActor, DoGenerateCollisionData))
 			{
-				PhysDispatcher->EnqueueCommand(Solver, [InDoGenerateCollisionData = DoGenerateCollisionData]
-				(Chaos::FPBDRigidsSolver* InSolver)
+				PhysDispatcher->EnqueueCommandImmediate(Solver, [InDoGenerateCollisionData = DoGenerateCollisionData]
+				(Chaos::FPhysicsSolver* InSolver)
 				{
 					InSolver->SetGenerateCollisionData(InDoGenerateCollisionData);
 				});
 			}
 			else if (PropertyChangedEvent.MemberProperty->GetFName() == GET_MEMBER_NAME_CHECKED(AChaosSolverActor, CollisionFilterSettings))
 			{
-				PhysDispatcher->EnqueueCommand(Solver, [InCollisionFilterSettings = CollisionFilterSettings]
-				(Chaos::FPBDRigidsSolver* InSolver)
+				PhysDispatcher->EnqueueCommandImmediate(Solver, [InCollisionFilterSettings = CollisionFilterSettings]
+				(Chaos::FPhysicsSolver* InSolver)
 				{
 					InSolver->SetCollisionFilterSettings(InCollisionFilterSettings);
 				});
 			}
 			else if (PropertyChangedEvent.Property->GetFName() == GET_MEMBER_NAME_CHECKED(AChaosSolverActor, DoGenerateBreakingData))
 			{
-				PhysDispatcher->EnqueueCommand(Solver, [InDoGenerateBreakingData = DoGenerateBreakingData]
-				(Chaos::FPBDRigidsSolver* InSolver)
+				PhysDispatcher->EnqueueCommandImmediate(Solver, [InDoGenerateBreakingData = DoGenerateBreakingData]
+				(Chaos::FPhysicsSolver* InSolver)
 				{
 					InSolver->SetGenerateBreakingData(InDoGenerateBreakingData);
 				});
 			}
 			else if (PropertyChangedEvent.MemberProperty->GetFName() == GET_MEMBER_NAME_CHECKED(AChaosSolverActor, BreakingFilterSettings))
 			{
-				PhysDispatcher->EnqueueCommand(Solver, [InBreakingFilterSettings = BreakingFilterSettings]
-				(Chaos::FPBDRigidsSolver* InSolver)
+				PhysDispatcher->EnqueueCommandImmediate(Solver, [InBreakingFilterSettings = BreakingFilterSettings]
+				(Chaos::FPhysicsSolver* InSolver)
 				{
 					InSolver->SetBreakingFilterSettings(InBreakingFilterSettings);
 				});
 			}
 			else if (PropertyChangedEvent.Property->GetFName() == GET_MEMBER_NAME_CHECKED(AChaosSolverActor, DoGenerateTrailingData))
 			{
-			PhysDispatcher->EnqueueCommand(Solver, [InDoGenerateTrailingData = DoGenerateTrailingData]
-			(Chaos::FPBDRigidsSolver* InSolver)
-			{
-				InSolver->SetGenerateTrailingData(InDoGenerateTrailingData);
-			});
+				PhysDispatcher->EnqueueCommandImmediate(Solver, [InDoGenerateTrailingData = DoGenerateTrailingData]
+				(Chaos::FPhysicsSolver* InSolver)
+				{
+					InSolver->SetGenerateTrailingData(InDoGenerateTrailingData);
+				});
 			}
 			else if (PropertyChangedEvent.MemberProperty->GetFName() == GET_MEMBER_NAME_CHECKED(AChaosSolverActor, TrailingFilterSettings))
 			{
-				PhysDispatcher->EnqueueCommand(Solver, [InTrailingFilterSettings = TrailingFilterSettings]
-				(Chaos::FPBDRigidsSolver* InSolver)
+				PhysDispatcher->EnqueueCommandImmediate(Solver, [InTrailingFilterSettings = TrailingFilterSettings]
+				(Chaos::FPhysicsSolver* InSolver)
 				{
 					InSolver->SetTrailingFilterSettings(InTrailingFilterSettings);
 				});
 			}
 			else if (PropertyChangedEvent.Property->GetFName() == GET_MEMBER_NAME_CHECKED(AChaosSolverActor, bHasFloor))
 			{
-				PhysDispatcher->EnqueueCommand(Solver, [InHasFloor = bHasFloor]
-				(Chaos::FPBDRigidsSolver* InSolver)
+				PhysDispatcher->EnqueueCommandImmediate(Solver, [InHasFloor = bHasFloor]
+				(Chaos::FPhysicsSolver* InSolver)
 				{
 					InSolver->SetHasFloor(InHasFloor);
 				});
 			}
 			else if (PropertyChangedEvent.Property->GetFName() == GET_MEMBER_NAME_CHECKED(AChaosSolverActor, FloorHeight))
 			{
-				PhysDispatcher->EnqueueCommand(Solver, [InFloorHeight = FloorHeight]
-				(Chaos::FPBDRigidsSolver* InSolver)
+				PhysDispatcher->EnqueueCommandImmediate(Solver, [InFloorHeight = FloorHeight]
+				(Chaos::FPhysicsSolver* InSolver)
 				{
 					InSolver->SetFloorHeight(InFloorHeight);
 				});
 			}
+#if TODO_REIMPLEMENT_TIMESTEP_MULTIPLIER
 			else if (PropertyChangedEvent.Property->GetFName() == GET_MEMBER_NAME_CHECKED(AChaosSolverActor, MassScale))
 			{
-				PhysDispatcher->EnqueueCommand(Solver, [InMassScale = MassScale]
-				(Chaos::FPBDRigidsSolver* InSolver)
+				PhysDispatcher->EnqueueCommandImmediate(Solver, [InMassScale = MassScale]
+				(Chaos::FPhysicsSolver* InSolver)
 				{
 					InSolver->SetMassScale(InMassScale);
 				});
 			}
+#endif
 		}
 	}
 #endif
+
 #if CHAOS_DEBUG_SUBSTEP
+#if TODO_REIMPLEMENT_DEBUG_SUBSTEP
 	if (Solver && PropertyChangedEvent.Property)
 	{
 		if (PropertyChangedEvent.Property->GetFName() == GET_MEMBER_NAME_CHECKED(FChaosDebugSubstepControl, bPause))
@@ -528,8 +554,10 @@ void AChaosSolverActor::PostEditChangeProperty(struct FPropertyChangedEvent& Pro
 		}
 	}
 #endif
+#endif
 }
 
+#if TODO_REIMPLEMENT_SERIALIZATION_FOR_PERF_TEST
 #if !UE_BUILD_SHIPPING
 void SerializeForPerfTest(const TArray< FString >&, UWorld* World, FOutputDevice&)
 {
@@ -540,14 +568,15 @@ void SerializeForPerfTest(const TArray< FString >&, UWorld* World, FOutputDevice
 	//todo(mlentine): use this once chaos solver actors are in
 	for (TActorIterator<AChaosSolverActor> Itr(World); Itr; ++Itr)
 	{
-		Chaos::FPBDRigidsSolver* Solver = Itr->GetSolver();
-		FChaosSolversModule::GetModule()->GetDispatcher()->EnqueueCommand(Solver, [FileName](Chaos::FPBDRigidsSolver* SolverIn) { SolverIn->SerializeForPerfTest(FileName); });
+		Chaos::FPhysicsSolver* Solver = Itr->GetSolver();
+		FChaosSolversModule::GetModule()->GetDispatcher()->EnqueueCommandImmediate(Solver, [FileName](Chaos::FPhysicsSolver* SolverIn) { SolverIn->SerializeForPerfTest(FileName); });
 	}
 #endif
 }
 
 FAutoConsoleCommand SerializeForPerfTestCommand(TEXT("p.SerializeForPerfTest"), TEXT(""), FConsoleCommandWithWorldArgsAndOutputDeviceDelegate::CreateStatic(&SerializeForPerfTest));
 #endif
+#endif
 
 #endif
 
diff --git a/Engine/Source/Runtime/Experimental/ChaosSolverEngine/Public/Chaos/ChaosEventListenerComponent.h b/Engine/Source/Runtime/Experimental/ChaosSolverEngine/Public/Chaos/ChaosEventListenerComponent.h
--- a/Engine/Source/Runtime/Experimental/ChaosSolverEngine/Public/Chaos/ChaosEventListenerComponent.h
+++ b/Engine/Source/Runtime/Experimental/ChaosSolverEngine/Public/Chaos/ChaosEventListenerComponent.h
@@ -10,9 +10,9 @@
 // #include "UObject/ObjectMacros.h"
 
 #include "Components/ActorComponent.h"
+#include "Chaos/Declares.h"
 #include "ChaosEventListenerComponent.generated.h"
 
-namespace Chaos { class FPBDRigidsSolver; }
 class FPhysScene_Chaos;
 class AChaosSolverActor;
 
@@ -37,7 +37,7 @@ public:
 protected:
 
 #if INCLUDE_CHAOS
-	const Chaos::FPBDRigidsSolver* GetSolver() const;
+	const Chaos::FPhysicsSolver* GetSolver() const;
 	const TSharedPtr<FPhysScene_Chaos> GetPhysicsScene() const;
 #endif
 
diff --git a/Engine/Source/Runtime/Experimental/ChaosSolverEngine/Public/Chaos/ChaosGameplayEventDispatcher.h b/Engine/Source/Runtime/Experimental/ChaosSolverEngine/Public/Chaos/ChaosGameplayEventDispatcher.h
--- a/Engine/Source/Runtime/Experimental/ChaosSolverEngine/Public/Chaos/ChaosGameplayEventDispatcher.h
+++ b/Engine/Source/Runtime/Experimental/ChaosSolverEngine/Public/Chaos/ChaosGameplayEventDispatcher.h
@@ -5,6 +5,7 @@
 #include "ChaosEventListenerComponent.h"
 #include "PhysicsPublic.h"
 #include "ChaosNotifyHandlerInterface.h"
+#include "EventsData.h"
 #include "ChaosGameplayEventDispatcher.generated.h"
 
 struct FBodyInstance;
@@ -71,7 +72,8 @@ class CHAOSSOLVERENGINE_API UChaosGameplayEventDispatcher : public UChaosEventLi
 
 public:
 
-	virtual void TickComponent(float DeltaTime, enum ELevelTick TickType, FActorComponentTickFunction *ThisTickFunction) override;
+	virtual void BeginPlay() override;
+	virtual void EndPlay(const EEndPlayReason::Type EndPlayReason) override;
 
 private:
 
@@ -132,6 +134,16 @@ private:
 	float LastBreakingDataTime = -1.f;
 
 	void DispatchPendingCollisionNotifies();
+
+#if INCLUDE_CHAOS
+	void RegisterChaosEvents();
+	void UnregisterChaosEvents();
+
+	// Chaos Event Handlers
+	void HandleCollisionEvents(const Chaos::FCollisionEventData& CollisionData);
+	void HandleBreakingEvents(const Chaos::FBreakingEventData& BreakingData);
+#endif
+
 };
 
 
diff --git a/Engine/Source/Runtime/Experimental/ChaosSolverEngine/Public/Chaos/ChaosSolverActor.h b/Engine/Source/Runtime/Experimental/ChaosSolverEngine/Public/Chaos/ChaosSolverActor.h
--- a/Engine/Source/Runtime/Experimental/ChaosSolverEngine/Public/Chaos/ChaosSolverActor.h
+++ b/Engine/Source/Runtime/Experimental/ChaosSolverEngine/Public/Chaos/ChaosSolverActor.h
@@ -183,7 +183,7 @@ public:
 
 #if INCLUDE_CHAOS
 	TSharedPtr<FPhysScene_Chaos> GetPhysicsScene() const { return PhysScene; }
-	Chaos::FPBDRigidsSolver* GetSolver() const { return Solver; }
+	Chaos::FPhysicsSolver* GetSolver() const { return Solver; }
 #endif
 
 #if WITH_EDITOR
@@ -197,7 +197,7 @@ public:
 private:
 #if INCLUDE_CHAOS
 	TSharedPtr<FPhysScene_Chaos> PhysScene;
-	Chaos::FPBDRigidsSolver* Solver;
+	Chaos::FPhysicsSolver* Solver;
 #endif
 
 	/** Component responsible for harvesting and triggering physics-related gameplay events (hits, breaks, etc) */
diff --git a/Engine/Source/Runtime/Experimental/ChaosSolvers/Private/ChaosSolversModule.cpp b/Engine/Source/Runtime/Experimental/ChaosSolvers/Private/ChaosSolversModule.cpp
--- a/Engine/Source/Runtime/Experimental/ChaosSolvers/Private/ChaosSolversModule.cpp
+++ b/Engine/Source/Runtime/Experimental/ChaosSolvers/Private/ChaosSolversModule.cpp
@@ -12,13 +12,14 @@
 #include "Framework/PersistentTask.h"
 #include "Misc/CoreDelegates.h"
 #include "HAL/IConsoleManager.h"
-#include "PBDRigidsSolver.h"
+#include "PhysicsSolver.h"
 #include "FramePro/FramePro.h"
 #include "Chaos/BoundingVolume.h"
 #include "Chaos/PBDRigidParticles.h"
 #include "Chaos/UniformGrid.h"
 #include "UObject/Class.h"
 #include "Framework/Dispatcher.h"
+#include "Framework/DispatcherImpl.h"
 
 TAutoConsoleVariable<int32> CVarChaosThreadEnabled(
 	TEXT("p.Chaos.DedicatedThreadEnabled"),
@@ -44,16 +45,7 @@ namespace Chaos
 	{
 		FChaosSolversModule* ChaosModule = FModuleManager::Get().GetModulePtr<FChaosSolversModule>("ChaosSolvers");
 
-		EThreadingMode CurrentMode = ChaosModule->GetDispatcher()->GetMode();
-
-		if(InMode != EThreadingMode::Invalid && InMode != CurrentMode)
-		{
-			// Handle shutdown of current threading model
-			ChaosModule->ShutdownThreadingMode();
-			
-			// Handle entering new threading model
-			ChaosModule->InitializeThreadingMode(InMode);
-		}
+		ChaosModule->ChangeThreadingMode(InMode);
 	}
 
 	void ChangeBufferingMode()
@@ -103,7 +95,7 @@ namespace Chaos
 
 	Chaos::EThreadingMode FInternalDefaultSettings::GetDefaultThreadingMode() const
 	{
-		return Chaos::EThreadingMode::DedicatedThread;
+		return Chaos::EThreadingMode::TaskGraph;
 	}
 
 	EChaosSolverTickMode FInternalDefaultSettings::GetDedicatedThreadTickMode() const
@@ -131,7 +123,7 @@ void FChaosConsoleSinks::OnCVarsChanged()
 		if(ChaosModule->IsPersistentTaskRunning())
 		{
 			float NewHz = CVarDedicatedThreadDesiredHz.GetValueOnGameThread();
-			ChaosModule->GetDispatcher()->EnqueueCommand([NewHz](Chaos::FPersistentPhysicsTask* Thread)
+			ChaosModule->GetDispatcher()->EnqueueCommandImmediate([NewHz](Chaos::FPersistentPhysicsTask* Thread)
 			{
 				if(Thread)
 				{
@@ -199,6 +191,7 @@ void FChaosSolversModule::StartupModule()
 	{
 		FModuleManager::Get().LoadModule("FieldSystemEngine");
 	}
+	Initialize();
 }
 
 void FChaosSolversModule::ShutdownModule()
@@ -214,26 +207,22 @@ void FChaosSolversModule::Initialize()
 	{
 		const Chaos::EThreadingMode DefaultThreadingMode = GetSettingsProvider().GetDefaultThreadingMode();
 
-		if(DefaultThreadingMode == Chaos::EThreadingMode::DedicatedThread)
+		switch(DefaultThreadingMode)
 		{
+		case Chaos::EThreadingMode::DedicatedThread:
 			StartPhysicsTask();
-		}
-		else
-		{
-			switch(DefaultThreadingMode)
-			{
-			case Chaos::EThreadingMode::SingleThread:
-				Dispatcher = new Chaos::FDispatcher<Chaos::EThreadingMode::SingleThread>(this);
-				break;
-			case Chaos::EThreadingMode::TaskGraph:
-				Dispatcher = new Chaos::FDispatcher<Chaos::EThreadingMode::TaskGraph>(this);
-				break;
-
-			default:
-				// Must have a dispatcher! Add handling for new threading models above
-				check(false);
-				break;
-			}
+			break;
+		case Chaos::EThreadingMode::SingleThread:
+			Dispatcher = new Chaos::FDispatcher<Chaos::EThreadingMode::SingleThread>(this);
+			break;
+		case Chaos::EThreadingMode::TaskGraph:
+			Dispatcher = new Chaos::FDispatcher<Chaos::EThreadingMode::TaskGraph>(this);
+			break;
+
+		default:
+			// Must have a dispatcher! Add handling for new threading models above
+			check(false);
+			break;
 		}
 
 		bModuleInitialized = true;
@@ -284,8 +273,31 @@ void FChaosSolversModule::ShutdownThreadingMode()
 	}
 	break;
 	
-	case EThreadingMode::SingleThread:
 	case EThreadingMode::TaskGraph:
+	{
+		if(Dispatcher)
+		{
+			// we need to flush out any commands currently waiting in the taskgraph dispatcher.
+			// Dedicated will wait for execution to end, and single thread runs immediately so we only
+			// need to handle this for task graph dispatchers
+			Dispatcher->Execute();
+
+			for(FPBDRigidsSolver* Solver : Solvers)
+			{
+				IDispatcher::FSolverCommand Command;
+				while(Solver->CommandQueue.Dequeue(Command))
+				{
+					Command(Solver);
+				}
+			}
+
+			delete Dispatcher;
+			Dispatcher = nullptr;
+		}
+	}
+	break;
+
+	case EThreadingMode::SingleThread:
 	{
 		if(Dispatcher)
 		{
@@ -332,6 +344,23 @@ void FChaosSolversModule::InitializeThreadingMode(Chaos::EThreadingMode InNewMod
 	}
 }
 
+void FChaosSolversModule::ChangeThreadingMode(Chaos::EThreadingMode InNewMode)
+{
+	EChaosThreadingMode CurrentMode = GetDispatcher()->GetMode();
+
+	if(InNewMode != EChaosThreadingMode::Invalid && InNewMode != CurrentMode)
+	{
+		// Handle shutdown of current threading model
+		ShutdownThreadingMode();
+
+		// Handle entering new threading model
+		InitializeThreadingMode(InNewMode);
+
+		// Buffering mode may change when the threading mode changes
+		//ChangeBufferingMode();
+	}
+}
+
 bool FChaosSolversModule::IsPersistentTaskEnabled() const
 {
 	return CVarChaosThreadEnabled.GetValueOnGameThread() == 1;
@@ -415,7 +444,7 @@ void FChaosSolversModule::SyncTask(bool bForceBlockingSync /*= false*/)
 	UpdateStats();
 }
 
-Chaos::FPBDRigidsSolver* FChaosSolversModule::CreateSolver(bool bStandalone /*= false*/)
+Chaos::FPhysicsSolver* FChaosSolversModule::CreateSolver(bool bStandalone /*= false*/)
 {
 	FChaosScopeSolverLock SolverScopeLock;
 	
@@ -425,13 +454,13 @@ Chaos::FPBDRigidsSolver* FChaosSolversModule::CreateSolver(bool bStandalone /*=
 		SolverBufferMode = GetBufferModeFromThreadingModel(GetDispatcher()->GetMode());
 	}
 
-	Solvers.Add(Chaos::FPBDRigidsSolver::FAccessor::CreateSolver(SolverBufferMode));
-	Chaos::FPBDRigidsSolver* NewSolver = Solvers.Last();
+	Solvers.Add(new Chaos::FPhysicsSolver(SolverBufferMode));
+	Chaos::FPhysicsSolver* NewSolver = Solvers.Last();
 
 	if(!bStandalone && IsPersistentTaskRunning() && Dispatcher)
 	{
 		// Need to let the thread know there's a new solver to care about
-		Dispatcher->EnqueueCommand([NewSolver](Chaos::FPersistentPhysicsTask* PhysThread)
+		Dispatcher->EnqueueCommandImmediate([NewSolver](Chaos::FPersistentPhysicsTask* PhysThread)
 		{
 			PhysThread->AddSolver(NewSolver);
 		});
@@ -455,7 +484,7 @@ void FChaosSolversModule::SetDedicatedThreadTickMode(EChaosSolverTickMode InTick
 {
 	check(Dispatcher);
 
-	Dispatcher->EnqueueCommand([InTickMode](Chaos::FPersistentPhysicsTask* InThread)
+	Dispatcher->EnqueueCommandImmediate([InTickMode](Chaos::FPersistentPhysicsTask* InThread)
 	{
 		if(InThread)
 		{
@@ -464,7 +493,7 @@ void FChaosSolversModule::SetDedicatedThreadTickMode(EChaosSolverTickMode InTick
 	});
 }
 
-void FChaosSolversModule::DestroySolver(Chaos::FPBDRigidsSolver* InSolver)
+void FChaosSolversModule::DestroySolver(Chaos::FPhysicsSolver* InSolver)
 {
 	FChaosScopeSolverLock SolverScopeLock;
 
@@ -472,7 +501,7 @@ void FChaosSolversModule::DestroySolver(Chaos::FPBDRigidsSolver* InSolver)
 	{
 		if(Dispatcher)
 		{
-			Dispatcher->EnqueueCommand([InSolver](Chaos::FPersistentPhysicsTask* PhysThread)
+			Dispatcher->EnqueueCommandImmediate([InSolver](Chaos::FPersistentPhysicsTask* PhysThread)
 			{
 				if(PhysThread)
 				{
@@ -527,7 +556,8 @@ void FChaosSolversModule::DumpHierarchyStats(int32* OutOptMaxCellElements)
 	const int32 NumSolvers = Solvers.Num();
 	for(int32 SolverIndex = 0; SolverIndex < NumSolvers; ++SolverIndex)
 	{
-		Chaos::FPBDRigidsSolver* Solver = Solvers[SolverIndex];
+		Chaos::FPhysicsSolver* Solver = Solvers[SolverIndex];
+#if TODO_REIMPLEMENT_SPATIAL_ACCELERATION_ACCESS
 		if(const Chaos::ISpatialAcceleration<float,3>* SpatialAcceleration = Solver->GetSpatialAcceleration())
 		{
 #if !UE_BUILD_SHIPPING
@@ -535,6 +565,7 @@ void FChaosSolversModule::DumpHierarchyStats(int32* OutOptMaxCellElements)
 #endif
 			Solver->ReleaseSpatialAcceleration();
 		}
+#endif
 #if 0
 
 		const TArray<Chaos::TBox<float, 3>>& Boxes = Hierarchy->GetWorldSpaceBoxes();
@@ -644,8 +675,11 @@ void FChaosSolversModule::DumpHierarchyStats(int32* OutOptMaxCellElements)
 			
 			UE_LOG(LogChaos, Warning, TEXT("%s"), *OutputString);
 		}
-		#endif
+#endif
+
+#if TODO_REIMPLEMENT_SPATIAL_ACCELERATION_ACCESS
 		Solver->ReleaseSpatialAcceleration();
+#endif
 
 #if !UE_BUILD_SHIPPING
 		Chaos::PendingHierarchyDump = 1;	//mark solver pending dump to get more info
@@ -791,12 +825,12 @@ bool FChaosSolversModule::ShouldStepSolver(int32& InOutSingleStepCounter) const
 
 void FChaosSolversModule::ChangeBufferMode(Chaos::EMultiBufferMode BufferMode)
 {
-	for (Chaos::FPBDRigidsSolver* Solver : Solvers)
+	for (Chaos::FPhysicsSolver* Solver : Solvers)
 	{
 		if (Dispatcher)
 		{
-			Dispatcher->EnqueueCommand(Solver, [InBufferMode = BufferMode]
-			(Chaos::FPBDRigidsSolver* InSolver)
+			Dispatcher->EnqueueCommandImmediate(Solver, [InBufferMode = BufferMode]
+			(Chaos::FPhysicsSolver* InSolver)
 			{
 				InSolver->ChangeBufferMode(InBufferMode);
 			});
@@ -836,7 +870,7 @@ FChaosScopedPhysicsThreadLock::FChaosScopedPhysicsThreadLock(uint32 InMsToWait)
 		PTStallEvent = FPlatformProcess::GetSynchEventFromPool(false);
 
 		// Request a halt on the physics thread
-		PhysDispatcher->EnqueueCommand([PTStall = PTStallEvent, GTSync = CompleteEvent](Chaos::FPersistentPhysicsTask* PhysThread)
+		PhysDispatcher->EnqueueCommandImmediate([PTStall = PTStallEvent, GTSync = CompleteEvent](Chaos::FPersistentPhysicsTask* PhysThread)
 		{
 			PTStall->Trigger();
 			GTSync->Wait();
diff --git a/Engine/Source/Runtime/Experimental/ChaosSolvers/Private/EventManager.cpp b/Engine/Source/Runtime/Experimental/ChaosSolvers/Private/EventManager.cpp
new file mode 100644
--- /dev/null
+++ b/Engine/Source/Runtime/Experimental/ChaosSolvers/Private/EventManager.cpp
@@ -0,0 +1,111 @@
+// Copyright 1998-2019 Epic Games, Inc. All Rights Reserved.
+
+#include "EventManager.h"
+
+namespace Chaos
+{
+	void FEventManager::Reset()
+	{
+		for (FEventContainerPtr Container : EventContainers)
+		{
+			delete Container;
+			Container = nullptr;
+		}
+		EventContainers.Reset();
+	}
+
+	void FEventManager::UnregisterEvent(const FEventID& EventID)
+	{
+		if (EventID < EventContainers.Num())
+		{
+			delete EventContainers[EventID];
+			EventContainers.RemoveAt(EventID);
+		}
+	}
+
+	void FEventManager::UnregisterHandler(const FEventID& EventID, const void* InHandler)
+	{
+		checkf(EventID < EventContainers.Num(), TEXT("Unregistering event Handler for an event ID that does not exist"));
+		EventContainers[EventID]->UnregisterHandler(InHandler);
+	}
+
+	void FEventManager::FillProducerData(const Chaos::FPBDRigidsSolver* Solver)
+	{
+		for (FEventContainerPtr EventContainer : EventContainers)
+		{
+			if (EventContainer)
+			{
+				EventContainer->InjectProducerData(Solver);
+			}
+		}
+	}
+
+	void FEventManager::FlipBuffersIfRequired()
+	{
+		if (BufferMode == EMultiBufferMode::Double)
+		{
+			ResourceLock.WriteLock();
+		}
+
+		for (FEventContainerPtr EventContainer : EventContainers)
+		{
+			if (EventContainer)
+			{
+				EventContainer->FlipBufferIfRequired();
+			}
+		}
+
+		if (BufferMode == EMultiBufferMode::Double)
+		{
+			ResourceLock.WriteUnlock();
+		}
+	}
+
+	void FEventManager::DispatchEvents()
+	{
+		if (BufferMode == EMultiBufferMode::Double)
+		{
+			ResourceLock.ReadLock();
+		}
+
+		for (FEventContainerPtr EventContainer : EventContainers)
+		{
+			if (EventContainer)
+			{
+				EventContainer->DispatchConsumerData();
+			}
+		}
+
+		if (BufferMode == EMultiBufferMode::Double)
+		{
+			ResourceLock.ReadUnlock();
+		}
+
+	}
+
+	void FEventManager::InternalRegisterInjector(const FEventID& EventID, const FEventContainerPtr& Container)
+	{
+		if (EventID > EventContainers.Num())
+		{
+			for (int i = EventContainers.Num(); i < EventID; i++)
+			{
+				EventContainers.Push(nullptr);
+			}
+		}
+
+		EventContainers.EmplaceAt(EventID, Container);
+	}
+
+	int32 FEventManager::EncodeCollisionIndex(int32 ActualCollisionIndex, bool bSwapOrder)
+	{
+		return bSwapOrder ? (ActualCollisionIndex | (1 << 31)) : ActualCollisionIndex;
+	}
+
+	int32 FEventManager::DecodeCollisionIndex(int32 EncodedCollisionIdx, bool& bSwapOrder)
+	{
+		bSwapOrder = EncodedCollisionIdx & (1 << 31);
+		return EncodedCollisionIdx & ~(1 << 31);
+	}
+
+
+}
diff --git a/Engine/Source/Runtime/Experimental/ChaosSolvers/Private/Framework/CommandBuffer.cpp b/Engine/Source/Runtime/Experimental/ChaosSolvers/Private/Framework/CommandBuffer.cpp
new file mode 100644
--- /dev/null
+++ b/Engine/Source/Runtime/Experimental/ChaosSolvers/Private/Framework/CommandBuffer.cpp
@@ -0,0 +1,66 @@
+// Copyright 1998-2019 Epic Games, Inc. All Rights Reserved.
+
+#if INCLUDE_CHAOS
+
+#include "Framework/CommandBuffer.h"
+#include "ChaosSolversModule.h"
+#include "Framework/Dispatcher.h"
+
+bool Chaos::FCommandListData::HasCommands() const
+{
+	const int32 NumGlobalCommands = GlobalCommands.Num();
+	const int32 NumTaskCommands = TaskCommands.Num();
+	const int32 NumSolverCommands = SolverCommands.Num();
+
+	return (NumGlobalCommands + NumTaskCommands + NumSolverCommands) > 0;
+}
+
+Chaos::FCommandList::FCommandList()
+	: Data(nullptr)
+{
+	AllocData();
+}
+
+void Chaos::FCommandList::Flush()
+{
+	if(!Data || !Data->HasCommands())
+	{
+		// Nothing to flush here
+		return;
+	}
+
+	Chaos::IDispatcher* Dispatcher = FChaosSolversModule::GetModule()->GetDispatcher();
+
+	// Send the commands, the dispatcher will steal the pointer
+	Dispatcher->SubmitCommandList(MoveTemp(Data));
+
+	// Create new data object for subsequent commands as we've just handed the other
+	// data ptr to the dispatcher
+	AllocData();
+}
+
+void Chaos::FCommandList::Enqueue(IDispatcher::FGlobalCommand&& InCommand)
+{
+	Data->GlobalCommands.Add(InCommand);
+}
+
+void Chaos::FCommandList::Enqueue(FPhysicsSolver* InSolver, IDispatcher::FSolverCommand&& InCommand)
+{
+	Data->SolverCommands.Add(FCommandListData::FSolverCommandTuple(InSolver, InCommand));
+}
+
+void Chaos::FCommandList::Enqueue(IDispatcher::FTaskCommand&& InCommand)
+{
+	Data->TaskCommands.Add(InCommand);
+}
+
+void Chaos::FCommandList::AllocData()
+{
+	// Enable construction for command list data for the make call below
+	class FCommandList_Constructor : public FCommandListData
+	{};
+
+	Data = MakeUnique<FCommandList_Constructor>();
+}
+
+#endif
diff --git a/Engine/Source/Runtime/Experimental/ChaosSolvers/Private/Framework/DebugSolverTasks.cpp b/Engine/Source/Runtime/Experimental/ChaosSolvers/Private/Framework/DebugSolverTasks.cpp
--- a/Engine/Source/Runtime/Experimental/ChaosSolvers/Private/Framework/DebugSolverTasks.cpp
+++ b/Engine/Source/Runtime/Experimental/ChaosSolvers/Private/Framework/DebugSolverTasks.cpp
@@ -5,7 +5,7 @@
 #include "Framework/DebugSolverTasks.h"
 
 #if CHAOS_DEBUG_SUBSTEP
-#include "PBDRigidsSolver.h"
+#include "PhysicsSolver.h"
 #include "ChaosLog.h"
 
 namespace Chaos
@@ -38,7 +38,7 @@ namespace Chaos
 		DebugSubstep.Add(true, TEXT("Debug thread exit"));
 	}
 
-	void FDebugSolverTasks::DebugStep(FPBDRigidsSolver* Solver, TFunction<void()> StepFunction)
+	void FDebugSolverTasks::DebugStep(FPhysicsSolver* Solver, TFunction<void()> StepFunction)
 	{
 		// Retrieve debug thread pointer for this solver.
 		// It must have been already added unless the single thread mode has just been switched on.
@@ -48,6 +48,7 @@ namespace Chaos
 			DebugSolverTask = &SolverToTaskMap.Add(Solver, nullptr);
 		}
 
+#if TODO_REIMPLEMENT_DEBUG_SUBSTEP
 		// Sync/advance debug substep system for this solver
 		FDebugSubstep& DebugSubstep = Solver->GetDebugSubstep();
 		const bool bNeedsDebugThreadRunning = DebugSubstep.SyncAdvance(Solver->Enabled());
@@ -79,20 +80,24 @@ namespace Chaos
 		{
 			StepFunction();
 		}
+#endif
 	}
 
-	void FDebugSolverTasks::Add(FPBDRigidsSolver* Solver)
+	void FDebugSolverTasks::Add(FPhysicsSolver* Solver)
 	{
 		// Add solver to task map
 		SolverToTaskMap.Add(Solver, nullptr);
 
+#if TODO_REIMPLEMENT_DEBUG_SUBSTEP
 		// Reinit debug substep
 		UE_LOG(LogChaosThread, Verbose, TEXT("[Physics Thread] Initializing debug thread"));
 		Solver->GetDebugSubstep().Initialize();
+#endif
 	}
 
-	void FDebugSolverTasks::Remove(FPBDRigidsSolver* Solver)
+	void FDebugSolverTasks::Remove(FPhysicsSolver* Solver)
 	{
+#if TODO_REIMPLEMENT_DEBUG_SUBSTEP
 		// Remove the debug advance task for this solver, if any was created
 		FAsyncTask<FDebugSolverTask>* DebugSolverTask;
 		const bool bFound = SolverToTaskMap.RemoveAndCopyValue(Solver, DebugSolverTask);
@@ -107,10 +112,12 @@ namespace Chaos
 			DebugSolverTask->EnsureCompletion(false);
 			delete DebugSolverTask;
 		}
+#endif
 	}
 
 	void FDebugSolverTasks::Shutdown()
 	{
+#if TODO_REIMPLEMENT_DEBUG_SUBSTEP
 		// Iterate through all solver tasks in the map
 		for (auto& SolverTask: SolverToTaskMap)
 		{
@@ -127,6 +134,7 @@ namespace Chaos
 				SolverTask.Value = nullptr;
 			}
 		}
+#endif
 	}
 
 }  // namespace Chaos
diff --git a/Engine/Source/Runtime/Experimental/ChaosSolvers/Private/Framework/Dispatcher.cpp b/Engine/Source/Runtime/Experimental/ChaosSolvers/Private/Framework/Dispatcher.cpp
--- a/Engine/Source/Runtime/Experimental/ChaosSolvers/Private/Framework/Dispatcher.cpp
+++ b/Engine/Source/Runtime/Experimental/ChaosSolvers/Private/Framework/Dispatcher.cpp
@@ -3,77 +3,216 @@
 #if INCLUDE_CHAOS
 
 #include "Framework/Dispatcher.h"
+#include "DispatcherImpl.h"
 
-#include "PBDRigidsSolver.h"
+#include "PhysicsSolver.h"
+#include "Framework/CommandBuffer.h"
+#include "ChaosSolversModule.h"
+#include "ChaosStats.h"
+#include "PhysicsCoreTypes.h"
 
 namespace Chaos
 {
 	template<>
-	void FDispatcher<EThreadingMode::DedicatedThread>::EnqueueCommand(FPBDRigidsSolver* InSolver, TFunction<void(FPBDRigidsSolver *)> InCommand) const
+	void FDispatcher<EThreadingMode::DedicatedThread>::EnqueueCommandImmediate(FPhysicsSolver* InSolver, FSolverCommand InCommand)
 	{
 		check(InSolver && InCommand);
 		InSolver->GetCommandQueue().Enqueue(InCommand);
 	}
 
 	template<>
-	void FDispatcher<EThreadingMode::DedicatedThread>::EnqueueCommand(TFunction<void(FPersistentPhysicsTask*)> InCommand)
+	void FDispatcher<EThreadingMode::DedicatedThread>::EnqueueCommandImmediate(FTaskCommand InCommand)
 	{
 		check(Owner);
 		TaskCommandQueue.Enqueue(InCommand);
 	}
 
 	template<>
-	void FDispatcher<EThreadingMode::DedicatedThread>::EnqueueCommand(TFunction<void()> InCommand)
+	void FDispatcher<EThreadingMode::DedicatedThread>::EnqueueCommandImmediate(FGlobalCommand InCommand)
 	{
 		check(Owner);
 		GlobalCommandQueue.Enqueue(InCommand);
 	}
 
+	template<>
+	void FDispatcher<EThreadingMode::DedicatedThread>::SubmitCommandList(TUniquePtr<FCommandListData>&& InCommandData)
+	{
+		CommandLists.Enqueue(MoveTemp(InCommandData));
+	}
+
+	template<>
+	void FDispatcher<EThreadingMode::DedicatedThread>::Execute()
+	{
+		check(!IsInGameThread());
+
+		TUniquePtr<FCommandListData> Data = nullptr;
+		FPersistentPhysicsTask* Task = Owner->GetDedicatedTask();
+
+		// Enqueue all pending command lists
+		while(CommandLists.Dequeue(Data))
+		{
+			for(FGlobalCommand& C : Data->GlobalCommands)
+			{
+				EnqueueCommandImmediate(C);
+			}
+
+			for(FTaskCommand& C : Data->TaskCommands)
+			{
+				EnqueueCommandImmediate(C);
+			}
+
+			for(TTuple<FPhysicsSolver*, FSolverCommand> Pair : Data->SolverCommands)
+			{
+				EnqueueCommandImmediate(Pair.Get<0>(), Pair.Get<1>());
+			}
+		}
+
+		// Execute global and task commands.
+		{
+			SCOPE_CYCLE_COUNTER(STAT_PhysCommands);
+			TFunction<void()> GlobalCommand;
+			while(GlobalCommandQueue.Dequeue(GlobalCommand))
+			{
+				GlobalCommand();
+			}
+		}
+
+		{
+			SCOPE_CYCLE_COUNTER(STAT_TaskCommands);
+			TFunction<void(FPersistentPhysicsTask*)> TaskCommand;
+			while(TaskCommandQueue.Dequeue(TaskCommand))
+			{
+				// No dedicated thread task in this threading mode.
+				TaskCommand(Task);
+			}
+		}
+	}
+
 	//////////////////////////////////////////////////////////////////////////
 
 	template<>
-	void FDispatcher<EThreadingMode::SingleThread>::EnqueueCommand(FPBDRigidsSolver* InSolver, TFunction<void(FPBDRigidsSolver *)> InCommand) const
+	void FDispatcher<EThreadingMode::SingleThread>::EnqueueCommandImmediate(FPhysicsSolver* InSolver, FSolverCommand InCommand)
 	{
 		check(InSolver && InCommand);
 		InCommand(InSolver);
 	}
 
 	template<>
-	void FDispatcher<EThreadingMode::SingleThread>::EnqueueCommand(TFunction<void(FPersistentPhysicsTask*)> InCommand)
+	void FDispatcher<EThreadingMode::SingleThread>::EnqueueCommandImmediate(FTaskCommand InCommand)
 	{
 		check(Owner);
 		InCommand(nullptr);
 	}
 
 	template<>
-	void FDispatcher<EThreadingMode::SingleThread>::EnqueueCommand(TFunction<void()> InCommand)
+	void FDispatcher<EThreadingMode::SingleThread>::EnqueueCommandImmediate(FGlobalCommand InCommand)
 	{
 		check(Owner);
 		InCommand();
 	}
 
+	template<>
+	void FDispatcher<EThreadingMode::SingleThread>::SubmitCommandList(TUniquePtr<FCommandListData>&& InCommandData)
+	{
+		// Steal the ptr from the external caller still to emulate the same experience under all dispatchers
+		TUniquePtr<FCommandListData> Data = MoveTemp(InCommandData);
+
+		// Just pass to enqueue
+		for(FGlobalCommand& C : Data->GlobalCommands)
+		{
+			EnqueueCommandImmediate(C);
+		}
+
+		for(FTaskCommand& C : Data->TaskCommands)
+		{
+			EnqueueCommandImmediate(C);
+		}
+
+		for(const TTuple<FPhysicsSolver*, FSolverCommand>& Pair : Data->SolverCommands)
+		{
+			EnqueueCommandImmediate(Pair.Get<0>(), Pair.Get<1>());
+		}
+	}
+
+	template<>
+	void FDispatcher<EThreadingMode::SingleThread>::Execute()
+	{
+		ensureMsgf(false, TEXT("Single threaded dispatcher should never be executed as commands are processed immediately."));
+	}
+
 	//////////////////////////////////////////////////////////////////////////
 
 	template<>
-	void FDispatcher<EThreadingMode::TaskGraph>::EnqueueCommand(FPBDRigidsSolver* InSolver, TFunction<void(FPBDRigidsSolver *)> InCommand) const
+	void FDispatcher<EThreadingMode::TaskGraph>::EnqueueCommandImmediate(FPhysicsSolver* InSolver, FSolverCommand InCommand)
 	{
 		check(InSolver && InCommand);
 		InSolver->GetCommandQueue().Enqueue(InCommand);
 	}
 
 	template<>
-	void FDispatcher<EThreadingMode::TaskGraph>::EnqueueCommand(TFunction<void(FPersistentPhysicsTask*)> InCommand)
+	void FDispatcher<EThreadingMode::TaskGraph>::EnqueueCommandImmediate(FTaskCommand InCommand)
 	{
 		check(Owner);
 		TaskCommandQueue.Enqueue(InCommand);
 	}
 
 	template<>
-	void FDispatcher<EThreadingMode::TaskGraph>::EnqueueCommand(TFunction<void()> InCommand)
+	void FDispatcher<EThreadingMode::TaskGraph>::EnqueueCommandImmediate(FGlobalCommand InCommand)
 	{
 		check(Owner);
 		GlobalCommandQueue.Enqueue(InCommand);
 	}
+
+	template<>
+	void FDispatcher<EThreadingMode::TaskGraph>::SubmitCommandList(TUniquePtr<FCommandListData>&& InCommandData)
+	{
+		CommandLists.Enqueue(MoveTemp(InCommandData));
+	}
+
+	template<>
+	void FDispatcher<EThreadingMode::TaskGraph>::Execute()
+	{
+		TUniquePtr<FCommandListData> Data = nullptr;
+
+		// Enqueue all pending command lists
+		while(CommandLists.Dequeue(Data))
+		{
+			for(FGlobalCommand& C : Data->GlobalCommands)
+			{
+				EnqueueCommandImmediate(C);
+			}
+
+			for(FTaskCommand& C : Data->TaskCommands)
+			{
+				EnqueueCommandImmediate(C);
+			}
+
+			for(TTuple<FPhysicsSolver*, FSolverCommand> Pair : Data->SolverCommands)
+			{
+				EnqueueCommandImmediate(Pair.Get<0>(), Pair.Get<1>());
+			}
+		}
+
+		// Execute global and task commands.
+		{
+			SCOPE_CYCLE_COUNTER(STAT_PhysCommands);
+			TFunction<void()> GlobalCommand;
+			while(GlobalCommandQueue.Dequeue(GlobalCommand))
+			{
+				GlobalCommand();
+			}
+		}
+
+		{
+			SCOPE_CYCLE_COUNTER(STAT_TaskCommands);
+			TFunction<void(FPersistentPhysicsTask*)> TaskCommand;
+			while(TaskCommandQueue.Dequeue(TaskCommand))
+			{
+				// No dedicated thread task in this threading mode.
+				TaskCommand(nullptr);
+			}
+		}
+	}
 }
 
 void LexFromString(Chaos::EThreadingMode& OutValue, const TCHAR* InString) 
@@ -110,4 +249,9 @@ FString LexToString(const Chaos::EThreadingMode InValue)
 
 	return TEXT("");
 }
+
+template class Chaos::FDispatcher<EChaosThreadingMode::DedicatedThread>;
+template class Chaos::FDispatcher<EChaosThreadingMode::SingleThread>;
+template class Chaos::FDispatcher<EChaosThreadingMode::TaskGraph>;
+
 #endif
diff --git a/Engine/Source/Runtime/Experimental/ChaosSolvers/Private/Framework/DispatcherImpl.h b/Engine/Source/Runtime/Experimental/ChaosSolvers/Private/Framework/DispatcherImpl.h
new file mode 100644
--- /dev/null
+++ b/Engine/Source/Runtime/Experimental/ChaosSolvers/Private/Framework/DispatcherImpl.h
@@ -0,0 +1,68 @@
+// Copyright 1998-2019 Epic Games, Inc. All Rights Reserved.
+
+#if INCLUDE_CHAOS
+
+#pragma once
+
+#include "CoreMinimal.h"
+#include "Containers/Queue.h"
+#include "Framework/Threading.h"
+#include "Chaos/Declares.h"
+#include "Framework/CommandBuffer.h"
+
+class FChaosSolversModule;
+class FPhysicsCommandsTask;
+
+namespace Chaos
+{
+	class FPersistentPhysicsTask;
+	class FCommandList;
+	class FCommandListData;
+}
+
+namespace Chaos
+{
+	template<EThreadingMode Mode>
+	class FDispatcher : public IDispatcher
+	{
+	public:
+		friend class FPersistentPhysicsTask;
+		friend class ::FPhysicsCommandsTask;
+		friend class FPhysScene_ChaosInterface;
+
+		FDispatcher() = delete;
+		FDispatcher(const FDispatcher& InCopy) = delete;
+		FDispatcher(FDispatcher&& InSteal) = delete;
+		FDispatcher& operator=(const FDispatcher& InCopy) = delete;
+		FDispatcher& operator=(FDispatcher&& InSteal) = delete;
+
+		explicit FDispatcher(FChaosSolversModule* InOwnerModule)
+			: Owner(InOwnerModule)
+		{
+		}
+
+		virtual void EnqueueCommandImmediate(FGlobalCommand InCommand) final override;
+		virtual void EnqueueCommandImmediate(FTaskCommand InCommand) final override;
+		virtual void EnqueueCommandImmediate(FPhysicsSolver* InSolver, FSolverCommand InCommand) final override;
+
+		virtual EThreadingMode GetMode() const final override { return Mode; }
+
+		virtual void SubmitCommandList(TUniquePtr<FCommandListData>&& InCommandData) override;
+
+		virtual void Execute() override;
+
+	private:
+
+		FChaosSolversModule* Owner;
+
+		TQueue<TFunction<void()>, EQueueMode::Mpsc> GlobalCommandQueue;
+		TQueue<TFunction<void(FPersistentPhysicsTask*)>, EQueueMode::Mpsc> TaskCommandQueue;
+		TQueue<TUniquePtr<FCommandListData>, EQueueMode::Mpsc> CommandLists;
+	};
+
+	extern template class FDispatcher<EThreadingMode::DedicatedThread>;
+	extern template class FDispatcher<EThreadingMode::SingleThread>;
+	extern template class FDispatcher<EThreadingMode::TaskGraph>;
+}
+
+#endif
\ No newline at end of file
diff --git a/Engine/Source/Runtime/Experimental/ChaosSolvers/Private/Framework/PersistentTask.cpp b/Engine/Source/Runtime/Experimental/ChaosSolvers/Private/Framework/PersistentTask.cpp
--- a/Engine/Source/Runtime/Experimental/ChaosSolvers/Private/Framework/PersistentTask.cpp
+++ b/Engine/Source/Runtime/Experimental/ChaosSolvers/Private/Framework/PersistentTask.cpp
@@ -12,14 +12,14 @@
 
 #include "Chaos/Framework/Parallel.h"
 #include "Framework/TimeStep.h"
-#include "PBDRigidsSolver.h"
+#include "PhysicsSolver.h"
 #include "Field/FieldSystem.h"
 
-#include "SolverObjects/SkeletalMeshPhysicsObject.h"
-#include "SolverObjects/StaticMeshPhysicsObject.h"
-#include "SolverObjects/BodyInstancePhysicsObject.h"
-#include "SolverObjects/GeometryCollectionPhysicsObject.h"
-#include "SolverObjects/FieldSystemPhysicsObject.h"
+#include "PhysicsProxy/SkeletalMeshPhysicsProxy.h"
+#include "PhysicsProxy/StaticMeshPhysicsProxy.h"
+#include "Chaos/Framework/SingleParticlePhysicsProxy.h"
+#include "PhysicsProxy/GeometryCollectionPhysicsProxy.h"
+#include "PhysicsProxy/FieldSystemPhysicsProxy.h"
 
 #ifndef CHAOS_WITH_PAUSABLE_SOLVER
 #define CHAOS_WITH_PAUSABLE_SOLVER 1
@@ -27,7 +27,7 @@
 
 namespace Chaos
 {
-	FPersistentPhysicsTask::FPersistentPhysicsTask(float InTargetDt, bool bInAvoidSpiral, FDispatcher<EThreadingMode::DedicatedThread>* InDispatcher)
+	FPersistentPhysicsTask::FPersistentPhysicsTask(float InTargetDt, bool bInAvoidSpiral, IDispatcher* InDispatcher)
 		: TickMode(EChaosSolverTickMode::VariableCappedWithTarget)
 		, CommandDispatcher(InDispatcher)
 		, Timestep(nullptr)
@@ -60,7 +60,7 @@ namespace Chaos
 
 #if CHAOS_DEBUG_SUBSTEP
 		// Prepare the debug substepping tasks for all existing solvers
-		for (FPBDRigidsSolver* Solver : Solvers)
+		for (FPhysicsSolver* Solver : Solvers)
 		{
 			DebugSolverTasks.Add(Solver);
 		}
@@ -75,30 +75,13 @@ namespace Chaos
 		Timestep->Reset();
 
 		// Scratch array for active solver list
-		TArray<Chaos::FPBDRigidsSolver*> ActiveSolverList;
+		TArray<Chaos::FPhysicsSolver*> ActiveSolverList;
 
 		while(bRunning)
 		{
 			SCOPE_CYCLE_COUNTER(STAT_PhysicsAdvance);
 
-			// Run global and task commands
-			{
-				SCOPE_CYCLE_COUNTER(STAT_PhysCommands);
-				TFunction<void()> GlobalCommand;
-				while(CommandDispatcher->GlobalCommandQueue.Dequeue(GlobalCommand))
-				{
-					GlobalCommand();
-				}
-			}
-
-			{
-				SCOPE_CYCLE_COUNTER(STAT_TaskCommands);
-				TFunction<void(FPersistentPhysicsTask*)> TaskCommand;
-				while(CommandDispatcher->TaskCommandQueue.Dequeue(TaskCommand))
-				{
-					TaskCommand(this);
-				}
-			}
+			CommandDispatcher->Execute();
 
 			Dt = Timestep->GetCalculatedDt();
 
@@ -109,9 +92,9 @@ namespace Chaos
 			{
 				// Go wide if possible on the solvers
 				ActiveSolverList.Reset(Solvers.Num());
-				for(Chaos::FPBDRigidsSolver* Solver : Solvers)
+				for(Chaos::FPhysicsSolver* Solver : Solvers)
 				{
-					if(Solver->HasActiveObjects())
+					if(Solver->HasActiveParticles())
 					{
 						ActiveSolverList.Add(Solver);
 					}
@@ -123,7 +106,7 @@ namespace Chaos
 				{
 					SCOPE_CYCLE_COUNTER(STAT_SolverAdvance);
 				
-					FPBDRigidsSolver* Solver = ActiveSolverList[Index];
+					FPhysicsSolver* Solver = ActiveSolverList[Index];
 
 					// Execute Step function either on this thread or in a pausable side debug thread
 					DebugSolverTasks.DebugStep(Solver, [this, Solver, Dt]()
@@ -161,23 +144,27 @@ namespace Chaos
 
 				for(int32 SolverIndex = 0; SolverIndex < NumSolvers; ++SolverIndex)
 				{
-					FPBDRigidsSolver* Solver = Solvers[SolverIndex];
+					FPhysicsSolver* Solver = Solvers[SolverIndex];
 
 					FPersistentPhysicsTaskStatistics::FPerSolverStatistics& SolverStat = CurrStats.SolverStats[SolverIndex];
-					FPBDRigidsEvolution& Evolution = *Solver->MEvolution;
+					FPhysicsSolver::FPBDRigidsEvolution& Evolution = *Solver->MEvolution;
 
+#if TODO_REIMPLEMENT_SOLVER_ENABLING
 					if(Solver->Enabled())
 					{
+#if TODO_REIMPLMEENT_EVOLUTION_ACCESSORS
 						SolverStat.NumActiveParticles = Evolution.GetActiveIndices().Num();
 						SolverStat.NumActiveConstraints = Evolution.NumConstraints();
 						SolverStat.NumAllocatedParticles = Evolution.GetParticles().Size();
 						SolverStat.NumParticleIslands = Evolution.NumIslands();
 						SolverStat.EvolutionStats = Evolution.GetEvolutionStats();
+#endif
 					}
 					else
 					{
 						SolverStat.Reset();
 					}
+#endif
 				}
 #endif
 
@@ -190,31 +177,31 @@ namespace Chaos
 		ShutdownEvent->Trigger();
 	}
 
-	void FPersistentPhysicsTask::StepSolver(FPBDRigidsSolver* InSolver, float Dt)
+	void FPersistentPhysicsTask::StepSolver(FPhysicsSolver* InSolver, float Dt)
 	{
 		HandleSolverCommands(InSolver);
 
 		// Check whether this solver is paused (changes in pause states usually happens during HandleSolverCommands)
 #if CHAOS_WITH_PAUSABLE_SOLVER
+#if TODO_REIMPLEMENT_SOLVER_PAUSING
 		if (!InSolver->Paused())
+#endif
 #endif
 		{
 			if (InSolver->bEnabled)
 			{
-				FSolverObjectStorage& Objects = InSolver->GetObjectStorage();
-					
 				// Only process if we have something to actually simulate
-				if(Objects.GetNumObjects() > 0)
+				if(InSolver->HasActiveParticles())
 				{
 					AdvanceSolver(InSolver, Dt);
 
 					{
-						SCOPE_CYCLE_COUNTER(STAT_CacheResults);
+						SCOPE_CYCLE_COUNTER(STAT_BufferPhysicsResults);
 						CacheLock.ReadLock();
 
-						Objects.ForEachSolverObjectParallel([](auto* Object)
+						InSolver->ForEachPhysicsProxyParallel([](auto* Object)
 						{
-							Object->CacheResults();
+							Object->BufferPhysicsResults();
 						});
 
 						CacheLock.ReadUnlock();
@@ -224,9 +211,9 @@ namespace Chaos
 						SCOPE_CYCLE_COUNTER(STAT_FlipResults);
 						CacheLock.WriteLock();
 
-						Objects.ForEachSolverObject([](auto* Object)
+						InSolver->ForEachPhysicsProxy([](auto* Object)
 						{
-							Object->FlipCache();
+							Object->FlipBuffer();
 						});
 
 						CacheLock.WriteUnlock();
@@ -236,13 +223,13 @@ namespace Chaos
 		}
 	}
 
-	void FPersistentPhysicsTask::AddSolver(FPBDRigidsSolver* InSolver)
+	void FPersistentPhysicsTask::AddSolver(FPhysicsSolver* InSolver)
 	{
 		Solvers.Add(InSolver);
 		DebugSolverTasks.Add(InSolver);
 	}
 
-	void FPersistentPhysicsTask::RemoveSolver(FPBDRigidsSolver* InSolver)
+	void FPersistentPhysicsTask::RemoveSolver(FPhysicsSolver* InSolver)
 	{
 		DebugSolverTasks.Remove(InSolver);
 		Solvers.Remove(InSolver);
@@ -257,21 +244,22 @@ namespace Chaos
 
 		CacheLock.ReadLock();
 
+
 		if(bFullSync)
 		{
-			TArray<FFieldSystemPhysicsObject*> FieldsToDelete;
+			TArray<FFieldSystemPhysicsProxy*> FieldsToDelete;
 
-			for(FPBDRigidsSolver* Solver : Solvers)
+			for(FPhysicsSolver* Solver : Solvers)
 			{
-				FSolverObjectStorage& Objects = Solver->GetObjectStorage();
-				FSolverObjectStorage& RemovedObjects = Solver->GetRemovedObjectStorage();
-				
-				Objects.ForEachSolverObject([](auto* Object)
+				Solver->ForEachPhysicsProxy([](auto* Object)
 				{
-					Object->SyncToCache();
+					Object->PullFromPhysicsState();
 				});
 
-				RemovedObjects.ForEachSolverObject([](auto* Object)
+#if TODO_REIMPLEMENT_REMOVED_PROXY_STORAGE
+				FPhysicsProxyStorage& RemovedObjects = Solver->GetRemovedObjectStorage();
+
+				RemovedObjects.ForEachPhysicsProxy([](auto* Object)
 				{
 					if (ensure(Object))
 					{
@@ -280,7 +268,7 @@ namespace Chaos
 					}
 				});
 
-				RemovedObjects.ForEachFieldSolverObject([Solver, &FieldsToDelete](auto* Object)
+				RemovedObjects.ForEachFieldPhysicsProxy([Solver, &FieldsToDelete](auto* Object)
 				{
 					if(Object->GetSolver() == Solver)
 					{
@@ -289,31 +277,31 @@ namespace Chaos
 				});
 
 				RemovedObjects.Reset();
+#endif
 			}
 
-			for(FFieldSystemPhysicsObject* FieldObj : FieldsToDelete)
+			// @todo(question) : Why is there a separate delete here for fields. [brice]
+			for(FFieldSystemPhysicsProxy* FieldObj : FieldsToDelete)
 			{
 				delete FieldObj;
 			}
 
-			for (FPBDRigidsSolver* Solver : Solvers)
+			for (FPhysicsSolver* Solver : Solvers)
 			{
 				Solver->SyncEvents_GameThread();
 			}
 		}
 		else
 		{ 
-			for(FPBDRigidsSolver* Solver : Solvers)
+			for(FPhysicsSolver* Solver : Solvers)
 			{
-				FSolverObjectStorage& Objects_GameThread = Solver->GetObjectStorage_GameThread();
-
-				Objects_GameThread.ForEachSolverObject([](auto* Object)
+				Solver->ForEachPhysicsProxy([](auto* Object)
 				{
-					Object->SyncToCache();
+					Object->PullFromPhysicsState();
 				});
 			}
 
-			for (FPBDRigidsSolver* Solver : Solvers)
+			for (FPhysicsSolver* Solver : Solvers)
 			{
 				Solver->SyncEvents_GameThread();
 			}
@@ -390,20 +378,20 @@ namespace Chaos
 		return Stats.GetGameDataForRead();
 	}
 
-	void FPersistentPhysicsTask::HandleSolverCommands(FPBDRigidsSolver* InSolver)
+	void FPersistentPhysicsTask::HandleSolverCommands(FPhysicsSolver* InSolver)
 	{
 		SCOPE_CYCLE_COUNTER(STAT_HandleSolverCommands);
 
 		check(InSolver);
-		TQueue<TFunction<void(FPBDRigidsSolver*)>, EQueueMode::Mpsc>& Queue = InSolver->CommandQueue;
-		TFunction<void(FPBDRigidsSolver*)> Command;
+		TQueue<TFunction<void(FPhysicsSolver*)>, EQueueMode::Mpsc>& Queue = InSolver->CommandQueue;
+		TFunction<void(FPhysicsSolver*)> Command;
 		while(Queue.Dequeue(Command))
 		{
 			Command(InSolver);
 		}
 	}
 
-	void FPersistentPhysicsTask::AdvanceSolver(FPBDRigidsSolver* InSolver, float InDt)
+	void FPersistentPhysicsTask::AdvanceSolver(FPhysicsSolver* InSolver, float InDt)
 	{
 		SCOPE_CYCLE_COUNTER(STAT_IntegrateSolver);
 
diff --git a/Engine/Source/Runtime/Experimental/ChaosSolvers/Private/Framework/PhysicsTickTask.cpp b/Engine/Source/Runtime/Experimental/ChaosSolvers/Private/Framework/PhysicsTickTask.cpp
--- a/Engine/Source/Runtime/Experimental/ChaosSolvers/Private/Framework/PhysicsTickTask.cpp
+++ b/Engine/Source/Runtime/Experimental/ChaosSolvers/Private/Framework/PhysicsTickTask.cpp
@@ -7,7 +7,8 @@
 #include "ChaosSolversModule.h"
 #include "Framework/Dispatcher.h"
 #include "ChaosStats.h"
-#include "PBDRigidsSolver.h"
+#include "PhysicsSolver.h"
+#include "PhysicsCoreTypes.h"
 
 FAutoConsoleTaskPriority CPrio_FPhysicsTickTask(
 	TEXT("TaskGraph.TaskPriorities.PhysicsTickTask"),
@@ -51,15 +52,15 @@ void FPhysicsTickTask::DoTask(ENamedThreads::Type CurrentThread, const FGraphEve
 	// per-solver commands and the solver advance
 	FGraphEventRef CommandsTask = TGraphTask<FPhysicsCommandsTask>::CreateTask().ConstructAndDispatchWhenReady();
 
-	const TArray<FPBDRigidsSolver*>& SolverList = Module->GetSolvers();
+	const TArray<FPhysicsSolver*>& SolverList = Module->GetSolvers();
 
 	// List of active solvers (assume all are active for single alloc)
-	TArray<FPBDRigidsSolver*> ActiveSolvers;
+	TArray<FPhysicsSolver*> ActiveSolvers;
 	ActiveSolvers.Reserve(SolverList.Num());
 
-	for(FPBDRigidsSolver* Solver : SolverList)
+	for(FPhysicsSolver* Solver : SolverList)
 	{
-		if(Solver->HasActiveObjects())
+		if(Solver->HasActiveParticles() || Solver->HasPendingCommands())
 		{
 			ActiveSolvers.Add(Solver);
 		}
@@ -76,7 +77,7 @@ void FPhysicsTickTask::DoTask(ENamedThreads::Type CurrentThread, const FGraphEve
 
 	// For each solver spawn a new solver advance task (which will run the per-solver command buffer and then advance the solver)
 	// Record the task reference as a prerequisite for the completion
-	for(FPBDRigidsSolver* Solver : ActiveSolvers)
+	for(FPhysicsSolver* Solver : ActiveSolvers)
 	{
 		FGraphEventRef SolverTaskRef = TGraphTask<FPhysicsSolverAdvanceTask>::CreateTask(&SolverTaskPrerequisites).ConstructAndDispatchWhenReady(Solver, Dt);
 		CompletionTaskPrerequisites.Add(SolverTaskRef);
@@ -96,9 +97,8 @@ FPhysicsCommandsTask::FPhysicsCommandsTask()
 	Module = FChaosSolversModule::GetModule();
 	check(Module);
 
-	Chaos::IDispatcher* BaseDispatcher = Module->GetDispatcher();
-	Dispatcher = BaseDispatcher && BaseDispatcher->GetMode() == EChaosThreadingMode::TaskGraph ? static_cast<Chaos::FDispatcher<EChaosThreadingMode::TaskGraph>*>(BaseDispatcher) : nullptr;
-	check(Dispatcher);
+	Dispatcher = Module->GetDispatcher();
+	check(Dispatcher->GetMode() == EChaosThreadingMode::TaskGraph);
 }
 
 TStatId FPhysicsCommandsTask::GetStatId() const
@@ -120,29 +120,12 @@ void FPhysicsCommandsTask::DoTask(ENamedThreads::Type CurrentThread, const FGrap
 {
 	using namespace Chaos;
 
-	// Global and task-level commands (in this threading mode these are analogous as there is no task)
-	{
-		SCOPE_CYCLE_COUNTER(STAT_PhysCommands);
-		TFunction<void()> GlobalCommand;
-		while(Dispatcher->GlobalCommandQueue.Dequeue(GlobalCommand))
-		{
-			GlobalCommand();
-		}
-	}
-	
-	{
-		SCOPE_CYCLE_COUNTER(STAT_TaskCommands);
-		TFunction<void(FPersistentPhysicsTask*)> TaskCommand;
-		while(Dispatcher->TaskCommandQueue.Dequeue(TaskCommand))
-		{
-			TaskCommand(nullptr);
-		}
-	}
+	Dispatcher->Execute();
 }
 
 //////////////////////////////////////////////////////////////////////////
 
-FPhysicsSolverAdvanceTask::FPhysicsSolverAdvanceTask(Chaos::FPBDRigidsSolver* InSolver, float InDt)
+FPhysicsSolverAdvanceTask::FPhysicsSolverAdvanceTask(Chaos::FPhysicsSolver* InSolver, float InDt)
 	: Solver(InSolver)
 	, Dt(InDt)
 {
@@ -170,7 +153,7 @@ void FPhysicsSolverAdvanceTask::DoTask(ENamedThreads::Type CurrentThread, const
 	StepSolver(Solver, Dt);
 }
 
-void FPhysicsSolverAdvanceTask::StepSolver(Chaos::FPBDRigidsSolver* InSolver, float InDt)
+void FPhysicsSolverAdvanceTask::StepSolver(Chaos::FPhysicsSolver* InSolver, float InDt)
 {
 	using namespace Chaos;
 
@@ -180,8 +163,8 @@ void FPhysicsSolverAdvanceTask::StepSolver(Chaos::FPBDRigidsSolver* InSolver, fl
 	{
 		SCOPE_CYCLE_COUNTER(STAT_HandleSolverCommands);
 
-		TQueue<TFunction<void(FPBDRigidsSolver*)>, EQueueMode::Mpsc>& Queue = InSolver->CommandQueue;
-		TFunction<void(FPBDRigidsSolver*)> Command;
+		TQueue<TFunction<void(FPhysicsSolver*)>, EQueueMode::Mpsc>& Queue = InSolver->CommandQueue;
+		TFunction<void(FPhysicsSolver*)> Command;
 		while(Queue.Dequeue(Command))
 		{
 			Command(InSolver);
@@ -190,10 +173,8 @@ void FPhysicsSolverAdvanceTask::StepSolver(Chaos::FPBDRigidsSolver* InSolver, fl
 
 	if(InSolver->bEnabled)
 	{
-		FSolverObjectStorage& Objects = InSolver->GetObjectStorage();
-
 		// Only process if we have something to actually simulate
-		if(Objects.GetNumObjects() > 0)
+		if(Solver->HasActiveParticles())
 		{
 			InSolver->AdvanceSolverBy(InDt);
 		}
