Copying //UE4/Dev-Physics to Dev-Main (//UE4/Dev-Main) [at] 7703071

diff --git a/Engine/Source/Runtime/Experimental/ChaosSolvers/Private/PBDRigidsSolver.cpp b/Engine/Source/Runtime/Experimental/ChaosSolvers/Private/PBDRigidsSolver.cpp
--- a/Engine/Source/Runtime/Experimental/ChaosSolvers/Private/PBDRigidsSolver.cpp
+++ b/Engine/Source/Runtime/Experimental/ChaosSolvers/Private/PBDRigidsSolver.cpp
@@ -5,26 +5,26 @@
 
 #include "PBDRigidsSolver.h"
 
-#include "Chaos/Utilities.h"
-#include "Chaos/PBDCollisionConstraintUtil.h"
 #include "Async/AsyncWork.h"
-#include "Misc/ScopeLock.h"
+#include "Chaos/ChaosArchive.h"
+#include "Chaos/Framework/SingleParticlePhysicsProxy.h"
+#include "Chaos/PBDCollisionConstraintUtil.h"
+#include "Chaos/Utilities.h"
 #include "ChaosStats.h"
-#include "SolverObjects/BodyInstancePhysicsObject.h"
-#include "SolverObjects/SkeletalMeshPhysicsObject.h"
-#include "SolverObjects/StaticMeshPhysicsObject.h"
-#include "SolverObjects/GeometryCollectionPhysicsObject.h"
-#include "SolverObjects/FieldSystemPhysicsObject.h"
+#include "ChaosSolversModule.h"
 #include "HAL/FileManager.h"
-#include "Chaos/ChaosArchive.h"
+#include "Misc/ScopeLock.h"
+#include "PhysicsProxy/SkeletalMeshPhysicsProxy.h"
+#include "PhysicsProxy/StaticMeshPhysicsProxy.h"
+#include "PhysicsProxy/GeometryCollectionPhysicsProxy.h"
+#include "PhysicsProxy/FieldSystemPhysicsProxy.h"
+#include "EventDefaults.h"
 
 DEFINE_LOG_CATEGORY_STATIC(LogPBDRigidsSolverSolver, Log, All);
 
 namespace Chaos
 {
 
-	int8 FPBDRigidsSolver::Invalid = -1;
-
 	class AdvanceOneTimeStepTask : public FNonAbandonableTask
 	{
 		friend class FAutoDeleteAsyncTask<AdvanceOneTimeStepTask>;
@@ -43,78 +43,44 @@ namespace Chaos
 			UE_LOG(LogPBDRigidsSolverSolver, Verbose, TEXT("AdvanceOneTimeStepTask::DoWork()"));
 
 			{
-				SCOPE_CYCLE_COUNTER(STAT_CreateBodies);
-				MSolver->CreateRigidBodyCallback(MSolver->MEvolution->GetParticles());
-			}
-
-			if(MSolver->GetRigidClustering().NumberOfPendingClusters()!=0)
-			{
-				return;
-			}
-
-			{
-				SCOPE_CYCLE_COUNTER(STAT_UpdateReverseMapping);
-				MSolver->BindParticleCallbackMappingPart1();
-			}
-
-			{
-				SCOPE_CYCLE_COUNTER(STAT_UpdateParams);
-				MSolver->ParameterUpdateCallback(MSolver->MEvolution->GetParticles(), MSolver->MTime);
-			}
-
-			{
-				SCOPE_CYCLE_COUNTER(STAT_DisableCollisions);
-				MSolver->DisableCollisionsCallback(MSolver->MEvolution->GetDisabledCollisions());
+				//SCOPE_CYCLE_COUNTER(STAT_UpdateParams);
+				//MSolver->ParameterUpdateCallback(MSolver->GetSolverTime());
 			}
 
 			{
-				SCOPE_CYCLE_COUNTER(STAT_BeginFrame);
-				MSolver->StartFrameCallback(MDeltaTime, MSolver->MTime);
+				//SCOPE_CYCLE_COUNTER(STAT_BeginFrame);
+				//MSolver->StartFrameCallback(MDeltaTime, MSolver->GetSolverTime());
 			}
 
 			{
 				SCOPE_CYCLE_COUNTER(STAT_EvolutionAndKinematicUpdate);
-				while(MDeltaTime > MSolver->MMaxDeltaTime)
+				while (MDeltaTime > MSolver->GetMaxDeltaTime())
 				{
-					MSolver->ForceUpdateCallback(MSolver->MEvolution->GetParticles(), MSolver->MTime);
-					MSolver->MEvolution->ReconcileIslands();
-					MSolver->KinematicUpdateCallback(MSolver->MEvolution->GetParticles(), MSolver->MMaxDeltaTime, MSolver->MTime);
-					MSolver->MEvolution->AdvanceOneTimeStep(MSolver->MMaxDeltaTime);
-					MDeltaTime -= MSolver->MMaxDeltaTime;
+					//MSolver->ForceUpdateCallback(MSolver->GetSolverTime());
+					//MSolver->GetEvolution()->ReconcileIslands();
+					//MSolver->KinematicUpdateCallback(MSolver->GetMaxDeltaTime(), MSolver->GetSolverTime());
+					MSolver->GetEvolution()->AdvanceOneTimeStep(MSolver->GetMaxDeltaTime());
+					MDeltaTime -= MSolver->GetMaxDeltaTime();
 				}
-				MSolver->ForceUpdateCallback(MSolver->MEvolution->GetParticles(), MSolver->MTime);
-				MSolver->MEvolution->ReconcileIslands();
-				MSolver->KinematicUpdateCallback(MSolver->MEvolution->GetParticles(), MDeltaTime, MSolver->MTime);
-				MSolver->MEvolution->AdvanceOneTimeStep(MDeltaTime);
-			}
-
-			{
-				SCOPE_CYCLE_COUNTER(STAT_UpdateReverseMapping);
-				MSolver->BindParticleCallbackMappingPart2();
-			}
-
-			{
-				SCOPE_CYCLE_COUNTER(STAT_CollisionContactsCallback);
-				MSolver->CollisionContactsCallback(MSolver->MEvolution->GetParticles(), MSolver->MEvolution->GetCollisionConstraints());
-			}
-
-			{
-				SCOPE_CYCLE_COUNTER(STAT_BreakingCallback);
-				MSolver->BreakingCallback(MSolver->MEvolution->GetParticles());
+				//MSolver->ForceUpdateCallback(MSolver->GetSolverTime());
+				//MSolver->GetEvolution()->ReconcileIslands();
+				//MSolver->KinematicUpdateCallback(MDeltaTime, MSolver->GetSolverTime());
+				MSolver->GetEvolution()->AdvanceOneTimeStep(MDeltaTime);
 			}
 
 			{
-				SCOPE_CYCLE_COUNTER(STAT_TrailingCallback);
-				MSolver->TrailingCallback(MSolver->MEvolution->GetParticles());
+				SCOPE_CYCLE_COUNTER(STAT_EventDataGathering);
+				MSolver->GetEventManager()->FillProducerData(MSolver);
+				MSolver->GetEventManager()->FlipBuffersIfRequired();
 			}
 
 			{
-				SCOPE_CYCLE_COUNTER(STAT_EndFrame);
-				MSolver->EndFrameCallback(MDeltaTime);
+				//SCOPE_CYCLE_COUNTER(STAT_EndFrame);
+				//MSolver->EndFrameCallback(MDeltaTime);
 			}
 
-			MSolver->MTime += MDeltaTime;
-			MSolver->CurrentFrame++;
+			MSolver->GetSolverTime() += MDeltaTime;
+			MSolver->GetCurrentFrame()++;
 		}
 
 	protected:
@@ -130,1233 +96,328 @@ namespace Chaos
 		TSharedPtr<FEvent> PrevEvent, CurrentEvent;
 	};
 
+
+
+
 	FPBDRigidsSolver::FPBDRigidsSolver(const EMultiBufferMode BufferingModeIn)
-		: TimeStepMultiplier(1.f)
-		, bDoGenerateCollisionData(false)
-		, bDoGenerateBreakingData(false)
-		, bDoGenerateTrailingData(false)
-#if CHAOS_WITH_PAUSABLE_SOLVER
-		, bPaused(false)
-#endif
+		: CurrentFrame(0)
+		, MTime(0.0)
+		, MLastDt(0.0)
+		, MMaxDeltaTime(0.0)
+		, TimeStepMultiplier(1.0)
 		, bEnabled(false)
 		, bHasFloor(true)
 		, bIsFloorAnalytic(false)
 		, FloorHeight(0.f)
-		, FloorIndex(INDEX_NONE)
-		, MassScale(1.0f)
-		, MCurrentEvent(nullptr)
-		, MCurrentLock(nullptr)
-		, PositionTarget()
-		, PositionConstraintsRule(PositionTarget)
-		, DynamicConstraints()
-		, DynamicConstraintsRule(DynamicConstraints)
-		, JointConstraints()
-		, JointConstraintsRule(JointConstraints)
-		, GameThreadHasSynced(false)
+		, MEvolution(new FPBDRigidsEvolution(Particles))
+		, MEventManager(new FEventManager(BufferingModeIn))
+		, MSolverEventFilters(new FSolverEventFilters())
 		, BufferMode(BufferingModeIn)
-		, AllCollisions_FromSequencerCache(new FAllCollisionData)
-		, AllBreakings_FromSequencerCache(new FAllBreakingData)
-		, AllTrailings_FromSequencerCache(new FAllTrailingData)
-		, NumCollisionsPrevFrame(0)
-		, NumBreakingsPrevFrame(0)
-		, NumTrailingsPrevFrame(0)
+		, MCurrentLock(new FCriticalSection())
 	{
 		UE_LOG(LogPBDRigidsSolverSolver, Verbose, TEXT("PBDRigidsSolver::PBDRigidsSolver()"));
 		Reset();
-
-		// default filter functions
-		SolverCollisionEventFilter = MakeUnique<FSolverCollisionEventFilter>(CollisionFilterSettings);
-		SolverBreakingEventFilter = MakeUnique<FSolverBreakingEventFilter>(BreakingFilterSettings);
-		SolverTrailingEventFilter = MakeUnique<FSolverTrailingEventFilter>(TrailingFilterSettings);
-	}
-
-	void FPBDRigidsSolver::RegisterObject(TSolverObject<FGeometryCollectionPhysicsObject>* InObject)
-	{
-		if(KinematicProxiesForObjects.Num() > Objects.GeometryCollectionObjects.Num())
-		{
-			KinematicProxiesForObjects.Insert(FKinematicProxy(), Objects.GeometryCollectionObjects.Num());
-		}
-		else
-		{
-			KinematicProxiesForObjects.Add(FKinematicProxy());
-		}
-
-		Objects.GeometryCollectionObjects.Add(InObject);
-	}
-
-	void FPBDRigidsSolver::RegisterObject(TSolverObject<FSkeletalMeshPhysicsObject>* InObject)
-	{
-		Objects.SkeletalMeshObjects.Add(InObject);
-		KinematicProxiesForObjects.Add(FKinematicProxy());
-	}
-
-	void FPBDRigidsSolver::RegisterObject(TSolverObject<FStaticMeshPhysicsObject>* InObject)
-	{
-		Objects.StaticMeshObjects.Add(InObject);
-		KinematicProxiesForObjects.Add(FKinematicProxy());
-	}
-
-	void FPBDRigidsSolver::RegisterObject(TSolverObject<FBodyInstancePhysicsObject>* InObject)
-	{
-		Objects.BodyInstanceObjects.Add(InObject);
-		KinematicProxiesForObjects.Add(FKinematicProxy());
-	}
-
-	void FPBDRigidsSolver::RegisterObject(FFieldSystemPhysicsObject* InObject)
-	{
-		Objects.FieldSystemObjects.Add(InObject);
 	}
 
-	void FPBDRigidsSolver::UnregisterObject(TSolverObject<FGeometryCollectionPhysicsObject>* InObject)
+	void FPBDRigidsSolver::RegisterObject(TGeometryParticle<float, 3>* GTParticle)
 	{
-		int32 Index = INDEX_NONE;
-		if(Objects.GeometryCollectionObjects.Find(InObject, Index))
-		{
-			Objects.GeometryCollectionObjects.RemoveAt(Index);
-			if (KinematicProxiesForObjects.Num() > Index)
-			{
-				KinematicProxiesForObjects.RemoveAt(Index);
-			}
-			RemovedObjects.GeometryCollectionObjects.Add(InObject);
-
-			// Because we removed the collection, we need to fix up the object reverse map
-			{
-				FGeometryCollectionPhysicsObject* Obj = (FGeometryCollectionPhysicsObject*)InObject;
-				int32 BaseIdx = Obj->GetBaseParticleIndex();
-				int32 EndIdx = BaseIdx + Obj->GetNumAddedParticles();
-
-				SolverObjectReverseMappingLock.WriteLock();
-				for(int32 ParticleIdx = BaseIdx; ParticleIdx < EndIdx; ++ParticleIdx)
-				{
-					if(ensure(SolverObjectReverseMappingBuffer->GetConsumerBuffer()->SolverObjectReverseMappingArray.IsValidIndex(ParticleIdx)))
-					{
-						SolverObjectWrapper& Wrapper = SolverObjectReverseMappingBuffer->AccessProducerBuffer()->SolverObjectReverseMappingArray[ParticleIdx];
-						Wrapper.SolverObject = nullptr;
-						Wrapper.Type = ESolverObjectType::NoneType;
-					}
-				}
-				SolverObjectReverseMappingLock.WriteUnlock();
-			}
-		}
-	}
-
-	void FPBDRigidsSolver::UnregisterObject(TSolverObject<FSkeletalMeshPhysicsObject>* InObject)
-	{
-		int32 Index = INDEX_NONE;
-		if(Objects.SkeletalMeshObjects.Find(InObject, Index))
-		{
-			Objects.SkeletalMeshObjects.RemoveAt(Index);
-
-			// Poking the Stop button in the editor causes the geometry collection component to call 
-			// Reset() on this class, which clears this array.
-			if (KinematicProxiesForObjects.Num() > Index)
-			{
-				KinematicProxiesForObjects.RemoveAt(Index);
-			}
-			RemovedObjects.SkeletalMeshObjects.Add(InObject);
-		}
-	}
-
-	void FPBDRigidsSolver::UnregisterObject(TSolverObject<FStaticMeshPhysicsObject>* InObject)
-	{
-		int32 Index = INDEX_NONE;
-		if(Objects.StaticMeshObjects.Find(InObject, Index))
-		{
-			Objects.StaticMeshObjects.RemoveAt(Index);
-			if (KinematicProxiesForObjects.Num() > Index)
-			{
-				KinematicProxiesForObjects.RemoveAt(Index);
-			}
-			RemovedObjects.StaticMeshObjects.Add(InObject);
-		}
-	}
+		UE_LOG(LogPBDRigidsSolverSolver, Verbose, TEXT("FPBDRigidsSolver::RegisterObject()"));
 
-	void FPBDRigidsSolver::UnregisterObject(TSolverObject<FBodyInstancePhysicsObject>* InObject)
-	{
-		int32 Index = INDEX_NONE;
-		if (Objects.BodyInstanceObjects.Find(InObject, Index))
-		{
-			Objects.BodyInstanceObjects.RemoveAt(Index);
-			if (KinematicProxiesForObjects.Num() > Index)
-			{
-				KinematicProxiesForObjects.RemoveAt(Index);
-			}
-			RemovedObjects.BodyInstanceObjects.Add(InObject);
-		}
-	} 
-
-	void FPBDRigidsSolver::UnregisterObject(FFieldSystemPhysicsObject* InObject)
-	{
-		// We can just remove fields, there's no kinematic proxy to manage
-		if(Objects.FieldSystemObjects.Remove(InObject) > 0)
-		{
-			RemovedObjects.FieldSystemObjects.Add(InObject);
-		}
-	}
-
-	void FPBDRigidsSolver::Reset()
-	{
-
-		UE_LOG(LogPBDRigidsSolverSolver, Verbose, TEXT("PBDRigidsSolver::Reset()"));
-
-#if CHAOS_WITH_PAUSABLE_SOLVER
-		bPaused = false;
-#endif
-		MTime = 0;
-		MLastDt = 0.0f;
-		bEnabled = false;
-		CurrentFrame = 0;
-		MMaxDeltaTime = 1;
-		FieldForceNum = 0;
-		LastSwapMappingSyncSize = 0;
-		LastMappingSyncSize = 0;
-
-		FParticlesType TRigidParticles;
-		MEvolution = TUniquePtr<FPBDRigidsEvolution>(new FPBDRigidsEvolution(MoveTemp(TRigidParticles)));
-		GetRigidParticles().AddArray(&FieldForce);
-		GetRigidParticles().AddArray(&FieldTorque);
-		GetRigidParticles().AddArray(&SolverObjectReverseMapping);
-		GetRigidParticles().AddArray(&ParticleIndexReverseMapping);
-
-		MEvolution->AddConstraintRule(&JointConstraintsRule);
-		MEvolution->AddConstraintRule(&DynamicConstraintsRule);
-		MEvolution->AddConstraintRule(&PositionConstraintsRule);
-		MEvolution->AddForceFunction([&](FParticlesType& Particles, const float Time, const int32 Index)
-		{
-			this->AddForceCallback(Particles, Time, Index);
-		});
-		MEvolution->AddForceFunction([&](FParticlesType& Particles, const float Time, const int32 Index)
-		{
-			if(Index < FieldForceNum)
-			{
-				Particles.F(Index) += FieldForce[Index];
-				Particles.Torque(Index) += FieldTorque[Index];
-			}
-		});
+		// Make sure this particle doesn't already have a proxy
+		checkSlow(GTParticle->Proxy == nullptr);
 
-		GetRigidClustering().SetGenerateClusterBreaking(bDoGenerateBreakingData);
+		// Grab the particle's type
+		const EParticleType InParticleType = GTParticle->ObjectType();
 
-		KinematicProxies.Reset();
-
-		// #BGallagher TODO could these ever leak? 
-		Objects.Reset();
-		RemovedObjects.Reset();
-		KinematicProxiesForObjects.Reset();
-
-		AllCollisionsBuffer = FMultiBufferFactory<FAllCollisionData>::CreateBuffer(BufferMode);
-		AllCollisionsIndicesBySolverObjectBuffer = FMultiBufferFactory<FAllCollisionsIndicesBySolverObject>::CreateBuffer(BufferMode);
-		AllBreakingsBuffer = FMultiBufferFactory<FAllBreakingData>::CreateBuffer(BufferMode);
-		AllBreakingsIndicesBySolverObjectBuffer = FMultiBufferFactory<FAllBreakingsIndicesBySolverObject>::CreateBuffer(BufferMode);
-		AllTrailingsBuffer = FMultiBufferFactory<FAllTrailingData>::CreateBuffer(BufferMode);
-		AllTrailingsIndicesBySolverObjectBuffer = FMultiBufferFactory<FAllTrailingsIndicesBySolverObject>::CreateBuffer(BufferMode);
-
-		// Note: the way these were being used there was no double buffering going on - all access was through game thread buffer 
-		// with dedicated locks. Keeping same behavior just now but will need to fix
-		SolverObjectReverseMappingBuffer = FMultiBufferFactory<FSolverObjectReverseMapping>::CreateBuffer(EMultiBufferMode::Single);
-		ParticleIndexReverseMappingBuffer = FMultiBufferFactory<FParticleIndexReverseMapping>::CreateBuffer(EMultiBufferMode::Single);
-	}
-
-	void FPBDRigidsSolver::ChangeBufferMode(Chaos::EMultiBufferMode InBufferMode)
-	{
-		if (BufferMode != InBufferMode)
-		{
-			check(BufferMode == AllCollisionsBuffer->GetBufferMode());
-			BufferMode = InBufferMode;
-			AllCollisionsBuffer.Reset();
-			AllCollisionsIndicesBySolverObjectBuffer.Reset();
-			AllBreakingsBuffer.Reset();
-			AllBreakingsIndicesBySolverObjectBuffer.Reset();
-			AllTrailingsBuffer.Reset();
-			AllTrailingsIndicesBySolverObjectBuffer.Reset();
-
-			AllCollisionsBuffer = FMultiBufferFactory<FAllCollisionData>::CreateBuffer(BufferMode);
-			AllCollisionsIndicesBySolverObjectBuffer = FMultiBufferFactory<FAllCollisionsIndicesBySolverObject>::CreateBuffer(BufferMode);
-			AllBreakingsBuffer = FMultiBufferFactory<FAllBreakingData>::CreateBuffer(BufferMode);
-			AllBreakingsIndicesBySolverObjectBuffer = FMultiBufferFactory<FAllBreakingsIndicesBySolverObject>::CreateBuffer(BufferMode);
-			AllTrailingsBuffer = FMultiBufferFactory<FAllTrailingData>::CreateBuffer(BufferMode);
-			AllTrailingsIndicesBySolverObjectBuffer = FMultiBufferFactory<FAllTrailingsIndicesBySolverObject>::CreateBuffer(BufferMode);
-		}
-	}
-
-#if !UE_BUILD_SHIPPING
-	void FPBDRigidsSolver::SerializeForPerfTest(const FString& FileName)
-	{
-		int32 Tries = 0;
-		FString UseFileName;
-		do
-		{
-			UseFileName = FString::Printf(TEXT("%s_%d.bin"), *FileName, Tries++);
-		} while (IFileManager::Get().FileExists(*UseFileName));
+		IPhysicsProxyBase* ProxyBase = nullptr;
+		// NOTE: Do we really need this list of proxies if we can just
+		// access them through the GTParticles list?
+		
+		Chaos::FParticleData* ProxyData;
 
-		//this is not actually file safe but oh well, very unlikely someone else is trying to create this file at the same time
-		TUniquePtr<FArchive> File(IFileManager::Get().CreateFileWriter(*UseFileName));
-		if (File)
+		// Make a physics proxy, giving it our particle and particle handle
+		if (InParticleType == EParticleType::Dynamic)
 		{
-			UE_LOG(LogChaos, Log, TEXT("ChaosPerfTestFile: %s"), *UseFileName);
-			FChaosArchive Ar(*File);
-			SerializeForPerfTest(Ar);
+			auto Proxy = new FRigidParticlePhysicsProxy(GTParticle->AsDynamic(), nullptr);
+			RigidParticlePhysicsProxies.Add((FRigidParticlePhysicsProxy*)Proxy);
+			ProxyData = Proxy->NewData();
+			ProxyBase = Proxy;
 		}
-		else
+		else if (InParticleType == EParticleType::Kinematic)
 		{
-			UE_LOG(LogChaos, Warning, TEXT("Could not create perf file(%s)"), *UseFileName);
+			auto Proxy = new FKinematicGeometryParticlePhysicsProxy(GTParticle->AsKinematic(), nullptr);
+			KinematicGeometryParticlePhysicsProxies.Add((FKinematicGeometryParticlePhysicsProxy*)Proxy);
+			ProxyData = Proxy->NewData();
+			ProxyBase = Proxy;
 		}
-	}
-#endif
-
-	void FPBDRigidsSolver::AdvanceSolverBy(float DeltaTime)
-	{
-		UE_LOG(LogPBDRigidsSolverSolver, Verbose, TEXT("PBDRigidsSolver::Tick(%3.5f)"), DeltaTime);
-		if(bEnabled)
+		else // Assume it's a static (geometry) if it's not dynamic or kinematic
 		{
-			MLastDt = DeltaTime;
-
-			// #TODO verify needed, should be handled by phys scene prior to registration now
-			Objects.ForEachSolverObject([this](auto* Obj) 
-			{
-				Obj->SetSolver(this);
-			});
-
-			Objects.ForEachFieldSolverObject([this](auto* Obj)
-			{
-				Obj->SetSolver(this);
-			});
-
-			int32 NumTimeSteps = (int32)(1.f*TimeStepMultiplier);
-			float dt = FMath::Min(DeltaTime, float(5.f / 30.f)) / (float)NumTimeSteps;
-			for(int i = 0; i < NumTimeSteps; i++)
-			{
-				AdvanceOneTimeStepTask(this, DeltaTime).DoWork();
-			}
+			auto Proxy = new FGeometryParticlePhysicsProxy(GTParticle, nullptr);
+			GeometryParticlePhysicsProxies.Add((FGeometryParticlePhysicsProxy*)Proxy);
+			ProxyData = Proxy->NewData();
+			ProxyBase = Proxy;
 		}
 
-	}
-
+		ProxyBase->SetSolver(this);
 
-	void FPBDRigidsSolver::CreateRigidBodyCallback(FPBDRigidsSolver::FParticlesType& Particles)
-	{
-		UE_LOG(LogPBDRigidsSolverSolver, Verbose, TEXT("PBDRigidsSolver::CreateRigidBodyCallback()"));
-		int32 NumParticles = Particles.Size();
-		if(!Particles.Size())
-		{
-			UE_LOG(LogPBDRigidsSolverSolver, Verbose, TEXT("... creating particles"));
-			if(bHasFloor)
-			{
-				UE_LOG(LogPBDRigidsSolverSolver, Verbose, TEXT("... creating floor"));
-				FloorIndex = Particles.Size();
-				Particles.AddParticles(1);
-				Particles.SetObjectState(FloorIndex, Chaos::EObjectStateType::Static);
-				Particles.CollisionGroup(0);
-				Particles.X(FloorIndex) = Chaos::TVector<float, 3>(0.f, 0.f, FloorHeight);
-				Particles.V(FloorIndex) = Chaos::TVector<float, 3>(0.f, 0.f, 0.f);
-				Particles.R(FloorIndex) = Chaos::TRotation<float, 3>::MakeFromEuler(Chaos::TVector<float, 3>(0.f, 0.f, 0.f));
-				Particles.W(FloorIndex) = Chaos::TVector<float, 3>(0.f, 0.f, 0.f);
-				Particles.P(FloorIndex) = Particles.X(0);
-				Particles.Q(FloorIndex) = Particles.R(0);
-				Particles.M(FloorIndex) = 1.f;
-				Particles.InvM(FloorIndex) = 0.f;
-				Particles.I(FloorIndex) = Chaos::PMatrix<float, 3, 3>(1.f, 0.f, 0.f, 0.f, 1.f, 0.f, 0.f, 0.f, 1.f);
-				Particles.InvI(FloorIndex) = Chaos::PMatrix<float, 3, 3>(0.f, 0.f, 0.f, 0.f, 0.f, 0.f, 0.f, 0.f, 0.f);
-				Particles.SetDynamicGeometry(FloorIndex, MakeUnique<Chaos::TPlane<float,3>>(Chaos::TVector<float, 3>(0.f, 0.f, FloorHeight), Chaos::TVector<float, 3>(0.f, 0.f, 1.f)));
-				Particles.DynamicGeometry(FloorIndex)->IgnoreAnalyticCollisions(!bIsFloorAnalytic);
-			}
-		}
+		// Associate the proxy with the particle
+		GTParticle->Proxy = ProxyBase;
 
-		Objects.ForEachSolverObject([&Particles](auto* Obj) 
+		//enqueue onto physics thread for finalizing registration
+		FChaosSolversModule::GetModule()->GetDispatcher()->EnqueueCommandImmediate(this, [GTParticle, InParticleType, ProxyBase, ProxyData](FPBDRigidsSolver* Solver)
 		{
-			if(Obj && Obj->IsSimulating())
-			{
-				Obj->CreateRigidBodyCallback(Particles);
-			}
-		});
+			UE_LOG(LogPBDRigidsSolverSolver, Verbose, TEXT("FPBDRigidsSolver::RegisterObject() ~ Dequeue"));
 
-		GetRigidClustering().UnionClusterGroups();
+			// Create a handle for the new particle
+			TGeometryParticleHandle<float, 3>* Handle;
 
-		if(NumParticles != Particles.Size())
-		{
-			InitializeFromParticleData(NumParticles);
-		}
-	}
+			// Insert the particle ptr into the auxiliary array
+			//
+			// TODO: Is this array even necessary?
+			//       Proxy should be able to map back to Particle.
+			//
 
-	bool FPBDRigidsSolver::Enabled() const
-	{
-		UE_LOG(LogPBDRigidsSolverSolver, Verbose, TEXT("PBDRigidsSolver::Enabled()"));
-		if(bEnabled)
-		{
-			// Look for simulating objects, could this be done cheaper? (no ForEachSolverObject as we're earlying out)
-			for(TSolverObject<FGeometryCollectionPhysicsObject>* GeomObject : Objects.GeometryCollectionObjects)
+			if (InParticleType == EParticleType::Dynamic)
 			{
-				if(GeomObject && GeomObject->IsSimulating())
-				{
-					return true;
-				}
+				Handle = Solver->Particles.CreateDynamicParticles(1)[0];
+				auto Proxy = static_cast<FRigidParticlePhysicsProxy*>(ProxyBase);
+				Proxy->SetHandle(Handle->AsDynamic());
+				Proxy->PushToPhysicsState(ProxyData);
 			}
-
-			for(TSolverObject<FSkeletalMeshPhysicsObject>* SkeletalMeshObject : Objects.SkeletalMeshObjects)
+			else if (InParticleType == EParticleType::Kinematic)
 			{
-				if(SkeletalMeshObject && SkeletalMeshObject->IsSimulating())
-				{
-					return true;
-				}
+				Handle = Solver->Particles.CreateKinematicParticles(1)[0];
+				auto Proxy = static_cast<FKinematicGeometryParticlePhysicsProxy*>(ProxyBase);
+				Proxy->SetHandle(Handle->AsKinematic());
+				Proxy->PushToPhysicsState(ProxyData);
 			}
-
-			for(TSolverObject<FStaticMeshPhysicsObject>* StaticMeshObject : Objects.StaticMeshObjects)
+			else // Assume it's a static (geometry) if it's not dynamic or kinematic
 			{
-				if(StaticMeshObject && StaticMeshObject->IsSimulating())
-				{
-					return true;
-				}
-			}
+				Handle = Solver->Particles.CreateStaticParticles(1)[0];
 
-			for (TSolverObject<FBodyInstancePhysicsObject>* BodyInstanceObject : Objects.BodyInstanceObjects)
-			{
-				if (BodyInstanceObject && BodyInstanceObject->IsSimulating())
-				{
-					return true;
-				}
+				auto Proxy = static_cast<FGeometryParticlePhysicsProxy*>(ProxyBase);
+				Proxy->SetHandle(Handle);
+				Proxy->PushToPhysicsState(ProxyData);
 			}
-		}
-		return false;
-	}
-
-	const Chaos::FPBDRigidsSolver::FAllCollisionDataMaps& FPBDRigidsSolver::GetAllCollisions_Maps_GameThread_NEEDSLOCK() const
-	{
-		QUICK_SCOPE_CYCLE_COUNTER(STAT_GetAllCollisions_Maps_GameThread);
 
-		FPBDRigidsSolver* const MutableThis = const_cast<FPBDRigidsSolver*>(this);
-		
-		if (AllCollisions_FromSequencerCache->AllCollisionsArray.Num() > 0)
-		{
-			MutableThis->AllCollisionDataMaps.bIsDataAndMapsInSync = true;
-			MutableThis->AllCollisionDataMaps.AllCollisionData = AllCollisions_FromSequencerCache.Get();
-			MutableThis->AllCollisionDataMaps.SolverObjectReverseMapping = SolverObjectReverseMappingBuffer->GetSyncConsumerBuffer();
-			MutableThis->AllCollisionDataMaps.ParticleIndexReverseMapping = ParticleIndexReverseMappingBuffer->GetSyncConsumerBuffer();
-			MutableThis->AllCollisionDataMaps.AllCollisionsIndicesBySolverObject = AllCollisionsIndicesBySolverObjectBuffer->GetSyncConsumerBuffer();
-		}
-		else if (AllCollisionsBuffer->GetConsumerBuffer()->AllCollisionsArray.Num() > 0)
-		{
-			MutableThis->AllCollisionDataMaps.bIsDataAndMapsInSync = true;
-			MutableThis->AllCollisionDataMaps.AllCollisionData = AllCollisionsBuffer->GetSyncConsumerBuffer();
-			MutableThis->AllCollisionDataMaps.SolverObjectReverseMapping =  SolverObjectReverseMappingBuffer->GetSyncConsumerBuffer();
-			MutableThis->AllCollisionDataMaps.ParticleIndexReverseMapping =  ParticleIndexReverseMappingBuffer->GetSyncConsumerBuffer();
-			MutableThis->AllCollisionDataMaps.AllCollisionsIndicesBySolverObject =  AllCollisionsIndicesBySolverObjectBuffer->GetSyncConsumerBuffer();
-		}
-		else
-		{
-			MutableThis->AllCollisionDataMaps.AllCollisionData = nullptr;
-			MutableThis->AllCollisionDataMaps.SolverObjectReverseMapping = nullptr;
-			MutableThis->AllCollisionDataMaps.ParticleIndexReverseMapping = nullptr;
-			MutableThis->AllCollisionDataMaps.AllCollisionsIndicesBySolverObject = nullptr;
-			MutableThis->AllCollisionDataMaps.bIsDataAndMapsInSync = false;
-		}
+			Handle->GTGeometryParticle() = GTParticle;
 
-		return AllCollisionDataMaps;
+			delete ProxyData;
+		});
 	}
 
-	
-	const Chaos::FPBDRigidsSolver::FAllBreakingDataMaps& FPBDRigidsSolver::GetAllBreakings_Maps_GameThread_NEEDSLOCK() const
+	void FPBDRigidsSolver::UnregisterObject(TGeometryParticle<float, 3>* GTParticle)
 	{
-		FPBDRigidsSolver* const MutableThis = const_cast<FPBDRigidsSolver*>(this);
+		UE_LOG(LogPBDRigidsSolverSolver, Verbose, TEXT("FPBDRigidsSolver::UnregisterObject()"));
 
-		if (AllBreakings_FromSequencerCache->AllBreakingsArray.Num() > 0)
-		{
-			MutableThis->AllBreakingDataMaps.bIsDataAndMapsInSync = true;
-			MutableThis->AllBreakingDataMaps.AllBreakingData = AllBreakings_FromSequencerCache.Get();
-			MutableThis->AllBreakingDataMaps.SolverObjectReverseMapping = SolverObjectReverseMappingBuffer->GetSyncConsumerBuffer();
-			MutableThis->AllBreakingDataMaps.ParticleIndexReverseMapping = ParticleIndexReverseMappingBuffer->GetSyncConsumerBuffer();
-			MutableThis->AllBreakingDataMaps.AllBreakingsIndicesBySolverObject = AllBreakingsIndicesBySolverObjectBuffer->GetSyncConsumerBuffer();
-		}
-		else if (AllBreakingsBuffer->GetSyncConsumerBuffer()->AllBreakingsArray.Num() > 0)
-		{
-			MutableThis->AllBreakingDataMaps.bIsDataAndMapsInSync = true;
-			MutableThis->AllBreakingDataMaps.AllBreakingData = AllBreakingsBuffer->GetSyncConsumerBuffer();
-			MutableThis->AllBreakingDataMaps.SolverObjectReverseMapping = SolverObjectReverseMappingBuffer->GetSyncConsumerBuffer();
-			MutableThis->AllBreakingDataMaps.ParticleIndexReverseMapping = ParticleIndexReverseMappingBuffer->GetSyncConsumerBuffer();
-			MutableThis->AllBreakingDataMaps.AllBreakingsIndicesBySolverObject = AllBreakingsIndicesBySolverObjectBuffer->GetSyncConsumerBuffer();
-		}
-		else
-		{
-			MutableThis->AllBreakingDataMaps.bIsDataAndMapsInSync = false;
-			MutableThis->AllBreakingDataMaps.AllBreakingData = nullptr;
-			MutableThis->AllBreakingDataMaps.SolverObjectReverseMapping = nullptr;
-			MutableThis->AllBreakingDataMaps.ParticleIndexReverseMapping = nullptr;
-			MutableThis->AllBreakingDataMaps.AllBreakingsIndicesBySolverObject = nullptr;
-		}
+		// Get the proxy associated with this particle
+		IPhysicsProxyBase* InProxy = GTParticle->Proxy;
+		check(InProxy);
 
-		return AllBreakingDataMaps;
-	}
+		// Grab the particle's type
+		const EParticleType InParticleType = GTParticle->ObjectType();
 
-	const Chaos::FPBDRigidsSolver::FAllTrailingDataMaps& FPBDRigidsSolver::GetAllTrailings_Maps_GameThread_NEEDSLOCK() const
-	{
-		FPBDRigidsSolver* const MutableThis = const_cast<FPBDRigidsSolver*>(this);
+		// Null out the particle's proxy pointer
+		GTParticle->Proxy = nullptr;
 
-		if (AllTrailings_FromSequencerCache->AllTrailingsArray.Num() > 0)
+		// Remove the proxy from the GT proxy map
+		if (InParticleType == EParticleType::Dynamic)
 		{
-			MutableThis->AllTrailingDataMaps.bIsDataAndMapsInSync = true;
-			MutableThis->AllTrailingDataMaps.AllTrailingData = AllTrailings_FromSequencerCache.Get();
-			MutableThis->AllTrailingDataMaps.SolverObjectReverseMapping = SolverObjectReverseMappingBuffer->GetSyncConsumerBuffer();
-			MutableThis->AllTrailingDataMaps.ParticleIndexReverseMapping = ParticleIndexReverseMappingBuffer->GetSyncConsumerBuffer();
-			MutableThis->AllTrailingDataMaps.AllTrailingsIndicesBySolverObject = AllTrailingsIndicesBySolverObjectBuffer->GetSyncConsumerBuffer();
+			RigidParticlePhysicsProxies.Remove((FRigidParticlePhysicsProxy*)InProxy);
 		}
-		else if (AllTrailingsBuffer->GetSyncConsumerBuffer()->AllTrailingsArray.Num() > 0)
+		else if (InParticleType == EParticleType::Kinematic)
 		{
-			MutableThis->AllTrailingDataMaps.bIsDataAndMapsInSync = true;
-			MutableThis->AllTrailingDataMaps.AllTrailingData = AllTrailingsBuffer->GetSyncConsumerBuffer();
-			MutableThis->AllTrailingDataMaps.SolverObjectReverseMapping = SolverObjectReverseMappingBuffer->GetSyncConsumerBuffer();
-			MutableThis->AllTrailingDataMaps.ParticleIndexReverseMapping = ParticleIndexReverseMappingBuffer->GetSyncConsumerBuffer();
-			MutableThis->AllTrailingDataMaps.AllTrailingsIndicesBySolverObject = AllTrailingsIndicesBySolverObjectBuffer->GetSyncConsumerBuffer();
+			KinematicGeometryParticlePhysicsProxies.Remove((FKinematicGeometryParticlePhysicsProxy*)InProxy);
 		}
 		else
 		{
-			MutableThis->AllTrailingDataMaps.bIsDataAndMapsInSync = false;
-			MutableThis->AllTrailingDataMaps.AllTrailingData = nullptr;
-			MutableThis->AllTrailingDataMaps.SolverObjectReverseMapping = nullptr;
-			MutableThis->AllTrailingDataMaps.ParticleIndexReverseMapping = nullptr;
-			MutableThis->AllTrailingDataMaps.AllTrailingsIndicesBySolverObject = nullptr;
+			GeometryParticlePhysicsProxies.Remove((FGeometryParticlePhysicsProxy*)InProxy);
 		}
 
-		return AllTrailingDataMaps;
-	}
-
-
-	void FPBDRigidsSolver::ParameterUpdateCallback(FPBDRigidsSolver::FParticlesType& Particles, const float Time)
-	{
-		UE_LOG(LogPBDRigidsSolverSolver, Verbose, TEXT("PBDRigidsSolver::ParameterUpdateCallback()"));
-		
-		MEvolution->GetRigidClustering().ResetAllClusterBreakings();
-
-		{
-			SCOPE_CYCLE_COUNTER(STAT_ParamUpdateObject);
-
-			Objects.ForEachSolverObject([&Particles, &Time](auto* Obj)
-			{
-				if(Obj && Obj->IsSimulating())
-				{
-					Obj->ParameterUpdateCallback(Particles, Time);
-				}
-			});
-		}
-
-		DynamicConstraints.UpdatePositionBasedState(Particles);
-		TArrayCollectionArray<float>& StrainArray = MEvolution->GetRigidClustering().GetStrainArray();
-		Chaos::TPBDPositionConstraints<float, 3>& LocalPositionTarget = PositionTarget;
-		TMap<int32, int32>& LocalPositionTargetedParticles = PositionTargetedParticles;
-		TArray<FKinematicProxy>& LocalKinematicProxies = KinematicProxiesForObjects;
-		
+		// Enqueue a command to remove the particle and delete the proxy
+		FChaosSolversModule::GetModule()->GetDispatcher()->EnqueueCommandImmediate(this, [InProxy, InParticleType](FPBDRigidsSolver* Solver)
 		{
-			SCOPE_CYCLE_COUNTER(STAT_ParamUpdateField);
-
-			Objects.ForEachFieldSolverObject([this, &Particles, &StrainArray, &LocalPositionTarget, &LocalPositionTargetedParticles, &LocalKinematicProxies, &Time](auto* Obj)
-			{
-				if(Obj && Obj->IsSimulating())
-				{
-					Obj->FieldParameterUpdateCallback(this, Particles, StrainArray, LocalPositionTarget, LocalPositionTargetedParticles, LocalKinematicProxies, Time);
-				}
-			});
-		}
-	}
-
-	void FPBDRigidsSolver::ForceUpdateCallback(FPBDRigidsSolver::FParticlesType& Particles, const float Time)
-	{
-		UE_LOG(LogPBDRigidsSolverSolver, Verbose, TEXT("PBDRigidsSolver::ParameterUpdateCallback()"));
-
-		// reset the FieldForces
-		FieldForceNum = FieldForce.Num();
-		for(int32 i = 0; i < FieldForce.Num(); i++)
-		{
-			FieldForce[i] = FVector(0);
-			FieldTorque[i] = FVector(0);
-		}
-
-		Objects.ForEachSolverObject([this, &Forces = FieldForce, &Torques = FieldTorque, Time, &Particles](auto* Obj)
-		{
-			if (Obj && Obj->IsSimulating())
-			{
-				Obj->FieldForcesUpdateCallback(this, Particles, Forces, Torques, Time);
-			}
-		});
+			UE_LOG(LogPBDRigidsSolverSolver, Verbose, TEXT("FPBDRigidsSolver::UnregisterObject() ~ Dequeue"));
 
-		Objects.ForEachFieldSolverObject([this, &Forces = FieldForce, &Torques = FieldTorque, Time, &Particles](auto* Obj)
-		{
-			if (Obj && Obj->IsSimulating())
+			// Get the physics thread-handle from the proxy, and then delete the proxy.
+			//
+			// NOTE: We have to delete the proxy from its derived version, because the
+			// base destructor is protected. This makes everything just a bit uglier,
+			// maybe that extra safety is not needed if we continue to contain all
+			// references to proxy instances entirely in Chaos?
+			TGeometryParticleHandle<float, 3>* Handle;
+			if (InParticleType == Chaos::EParticleType::Dynamic)
 			{
-				Obj->FieldForcesUpdateCallback(this, Particles, Forces, Torques, Time);
+				auto Proxy = (FRigidParticlePhysicsProxy*)InProxy;
+				Handle = Proxy->GetHandle();
+				delete Proxy;
 			}
-		});
-	}
-
-
-	void FPBDRigidsSolver::DisableCollisionsCallback(TSet<TTuple<int32, int32>>& CollisionPairs)
-	{
-		UE_LOG(LogPBDRigidsSolverSolver, Verbose, TEXT("PBDRigidsSolver::DisableCollisionsCallback()"));
-
-		Objects.ForEachSolverObject([&CollisionPairs](auto* Obj)
-		{
-			if(Obj && Obj->IsSimulating())
+			else if (InParticleType == Chaos::EParticleType::Kinematic)
 			{
-				Obj->DisableCollisionsCallback(CollisionPairs);
+				auto Proxy = (FKinematicGeometryParticlePhysicsProxy*)InProxy;
+				Handle = Proxy->GetHandle();
+				delete Proxy;
 			}
-		});
-	}
-
-	void FPBDRigidsSolver::StartFrameCallback(const float Dt, const float Time)
-	{
-		UE_LOG(LogPBDRigidsSolverSolver, Verbose, TEXT("PBDRigidsSolver::StartFrameCallback()"));
-
-
-		int32 KinematicIndex = 0;
-		Objects.ForEachSolverObject([this, &Dt, &Time, &KinematicIndex](auto* Obj)
-		{
-			if(Obj)
+			else
 			{
-				Obj->StartFrameCallback(Dt, Time);
-
-				if(Obj->IsSimulating())
-				{
-					Obj->UpdateKinematicBodiesCallback(MEvolution->GetParticles(), Dt, Time, KinematicProxiesForObjects[KinematicIndex]);
-				}
+				auto Proxy = (FGeometryParticlePhysicsProxy*)InProxy;
+				Handle = Proxy->GetHandle();
+				delete Proxy;
 			}
 
-			++KinematicIndex;
+			// Use the handle to destroy the particle data
+			Solver->Particles.DestroyParticle(Handle);
 		});
-
-		GetRigidClustering().SetGenerateClusterBreaking(bDoGenerateBreakingData);
-	}
-
-	// Required because we have multiple game side accessors vs a free running physics thread
-	// In the case of events it's not good enough to just supply the latest data as we would
-	// miss events if the physics thread happens to tick multiple times before a game system
-	// has time to read it
-	void FPBDRigidsSolver::SyncEvents_GameThread()
-	{
-		{
-			SCOPE_CYCLE_COUNTER(STAT_SyncEvents_GameThread);
-
-			if (BufferMode == EMultiBufferMode::Double)
-			{
-				SolverResourceLock.ReadLock();
-			}
-
-			if (bDoGenerateCollisionData)
-			{
-				AllCollisionsBuffer->SyncGameThread();
-				AllCollisionsIndicesBySolverObjectBuffer->SyncGameThread();
-			}
-			if (bDoGenerateBreakingData)
-			{
-				AllBreakingsBuffer->SyncGameThread();
-				AllBreakingsIndicesBySolverObjectBuffer->SyncGameThread();
-			}
-			if (bDoGenerateTrailingData)
-			{
-				AllTrailingsBuffer->SyncGameThread();
-				AllTrailingsIndicesBySolverObjectBuffer->SyncGameThread();
-			}
-			SolverObjectReverseMappingBuffer->SyncGameThread();
-			ParticleIndexReverseMappingBuffer->SyncGameThread();
-
-			GameThreadHasSynced.Store(true);
-
-			if (BufferMode == EMultiBufferMode::Double)
-			{
-				SolverResourceLock.ReadUnlock();
-			}
-
-		}
-
 	}
 
-	void FPBDRigidsSolver::EndFrameCallback(const float EndFrame)
-	{
-		UE_LOG(LogPBDRigidsSolverSolver, Verbose, TEXT("PBDRigidsSolver::EndFrameCallback()"));
-
-		GetRigidClustering().SwapBufferedData();
-		GetCollisionConstraints().SwapSpatialAcceleration();
-
-		Objects.ForEachSolverObjectParallel([EndFrame](auto* Obj)
-		{
-			if (Obj && Obj->IsSimulating())
-			{
-				Obj->EndFrameCallback(EndFrame);
-			}
-		});
-		Objects.ForEachFieldSolverObject([EndFrame](auto* Obj)
-		{
-			if (Obj && Obj->IsSimulating())
-			{
-				Obj->EndFrameCallback(EndFrame);
-			}
-		});
-
-		auto& AllCollisionsDataArray = GetAllCollisionsDataArray();
-		auto& AllBreakingsDataArray = GetAllBreakingsDataArray();
-		auto& AllTrailingsDataArray = GetAllTrailingsDataArray();
-
-		bool SwapCollisionData = AllCollisionsDataArray.Num() > 0 || (AllCollisionsDataArray.Num() == 0 && NumCollisionsPrevFrame > 0);
-		bool SwapBreakingData = AllBreakingsDataArray.Num() > 0 || (AllBreakingsDataArray.Num() == 0 && NumBreakingsPrevFrame > 0);
-		bool SwapTrailingData = AllTrailingsDataArray.Num() > 0 || (AllTrailingsDataArray.Num() == 0 && NumTrailingsPrevFrame > 0);
-
-		if (BufferMode == EMultiBufferMode::Double)
-		{
-			SolverResourceLock.WriteLock();
-		}
-
-		// if the game has not synced to the last physics update then skip the buffer swap and just append the new events
-		// to the existing data; so we don't lose the events for this frame
-		bool GameHasReadLastResults = GameThreadHasSynced.Load();
-
-		if (SwapCollisionData)
-		{
-			AllCollisionsBuffer->AccessProducerBuffer()->TimeCreated = MTime;
-			AllCollisionsIndicesBySolverObjectBuffer->AccessProducerBuffer()->TimeCreated = MTime;
-
-			if (GameHasReadLastResults)
-			{
-				AllCollisionsBuffer->FlipProducer();
-				AllCollisionsIndicesBySolverObjectBuffer->FlipProducer();
-			}
-		}
-
-		if (SwapBreakingData)
-		{
-			AllBreakingsBuffer->AccessProducerBuffer()->TimeCreated = MTime;
-			AllBreakingsIndicesBySolverObjectBuffer->AccessProducerBuffer()->TimeCreated = MTime;
-			if (GameHasReadLastResults)
-			{
-				AllBreakingsBuffer->FlipProducer();
-				AllBreakingsIndicesBySolverObjectBuffer->FlipProducer();
-			}
-		}
-
-		if (SwapTrailingData)
-		{
-			AllTrailingsBuffer->AccessProducerBuffer()->TimeCreated = MTime;
-			AllTrailingsIndicesBySolverObjectBuffer->AccessProducerBuffer()->TimeCreated = MTime;
-			if (GameHasReadLastResults)
-			{
-				AllTrailingsBuffer->FlipProducer();
-				AllTrailingsIndicesBySolverObjectBuffer->FlipProducer();
-			}
-		}
-
-		GameThreadHasSynced.Store(false);
-
-		//Update Time stamps no matter what?
-		if (GameHasReadLastResults)
-		{
-			SolverObjectReverseMappingLock.WriteLock();
-			SolverObjectReverseMappingBuffer->AccessProducerBuffer()->TimeCreated = MTime;
-			SolverObjectReverseMappingLock.WriteUnlock();
-
-			ParticleIndexReverseMappingLock.WriteLock();
-			ParticleIndexReverseMappingBuffer->AccessProducerBuffer()->TimeCreated = MTime;
-			ParticleIndexReverseMappingLock.WriteUnlock();
-		}
-
-		if (BufferMode == EMultiBufferMode::Double)
-		{
-			SolverResourceLock.WriteUnlock();
-		}
-
-		NumCollisionsPrevFrame = AllCollisionsDataArray.Num();
-		NumBreakingsPrevFrame = AllBreakingsDataArray.Num();
-		NumTrailingsPrevFrame = AllTrailingsDataArray.Num();
+	bool FPBDRigidsSolver::IsSimulating() const
+	{
+		for (FGeometryParticlePhysicsProxy* Obj : GeometryParticlePhysicsProxies)
+			if (Obj->IsSimulating())
+				return true;
+		for (FKinematicGeometryParticlePhysicsProxy* Obj : KinematicGeometryParticlePhysicsProxies)
+			if (Obj->IsSimulating())
+				return true;
+		for (FRigidParticlePhysicsProxy* Obj : RigidParticlePhysicsProxies)
+			if (Obj->IsSimulating())
+				return true;
+		for (FSkeletalMeshPhysicsProxy* Obj : SkeletalMeshPhysicsProxies)
+			if (Obj->IsSimulating())
+				return true;
+		for (FStaticMeshPhysicsProxy* Obj : StaticMeshPhysicsProxies)
+			if (Obj->IsSimulating())
+				return true;
+		for (FGeometryCollectionPhysicsProxy* Obj : GeometryCollectionPhysicsProxies)
+			if (Obj->IsSimulating())
+				return true;
+		for (FFieldSystemPhysicsProxy* Obj : FieldSystemPhysicsProxies)
+			if (Obj->IsSimulating())
+				return true;
+		return false;
 	}
 
-	void FPBDRigidsSolver::BindParticleCallbackMappingPart1()
-	{
-		UE_LOG(LogPBDRigidsSolverSolver, Verbose, TEXT("PBDRigidsSolver::BindParticleCallbackMapping()"));
-		if (LastMappingSyncSize != MEvolution->GetParticles().Size())
-		{
-			for (int32 i = LastMappingSyncSize; i < SolverObjectReverseMapping.Num(); i++)
-			{
-				SolverObjectReverseMapping[i] = { nullptr, ESolverObjectType::NoneType };
-				ParticleIndexReverseMapping[i] = INDEX_NONE;
-			}
-
-			Objects.ForEachSolverObject([&SolverMapping = SolverObjectReverseMapping, &ParticleMapping = ParticleIndexReverseMapping](auto* Obj)
-			{
-				if (Obj && Obj->IsSimulating())
-				{
-					Obj->BindParticleCallbackMapping(SolverMapping, ParticleMapping);
-				}
-			});
 
-			LastMappingSyncSize = MEvolution->GetParticles().Size();
-		}
-	}
 
-	void FPBDRigidsSolver::BindParticleCallbackMappingPart2()
+	void FPBDRigidsSolver::Reset()
 	{
-		UE_LOG(LogPBDRigidsSolverSolver, Verbose, TEXT("PBDRigidsSolver::BindParticleCallbackMapping()"));
-		if (LastMappingSyncSize != MEvolution->GetParticles().Size())
-		{
-			for (int32 i = LastMappingSyncSize; i < SolverObjectReverseMapping.Num(); i++)
-			{
-				SolverObjectReverseMapping[i] = { nullptr, ESolverObjectType::NoneType };
-				ParticleIndexReverseMapping[i] = INDEX_NONE;
-			}
-			
-			LastMappingSyncSize = MEvolution->GetParticles().Size();
-		}
-		if (LastSwapMappingSyncSize != MEvolution->GetParticles().Size())
-		{
-			// Copy the data to the game thread
-			SolverObjectReverseMappingLock.WriteLock();
-			SolverObjectReverseMappingBuffer->AccessProducerBuffer()->TimeCreated = MTime;
-
-			SolverObjectReverseMappingBuffer->AccessProducerBuffer()->SolverObjectReverseMappingArray.Reserve(SolverObjectReverseMapping.Num());
-			for (int32 Idx = LastSwapMappingSyncSize; Idx < SolverObjectReverseMapping.Num(); ++Idx)
-			{
-				SolverObjectReverseMappingBuffer->AccessProducerBuffer()->SolverObjectReverseMappingArray.Add(SolverObjectReverseMapping[Idx]);
-			}
-			SolverObjectReverseMappingLock.WriteUnlock();
 
-			ParticleIndexReverseMappingLock.WriteLock();
-			ParticleIndexReverseMappingBuffer->AccessProducerBuffer()->TimeCreated = MTime;
-
-			ParticleIndexReverseMappingBuffer->AccessProducerBuffer()->ParticleIndexReverseMappingArray.Reserve(ParticleIndexReverseMapping.Num());
-			for (int32 Idx = LastSwapMappingSyncSize; Idx < ParticleIndexReverseMapping.Num(); ++Idx)
-			{
-				ParticleIndexReverseMappingBuffer->AccessProducerBuffer()->ParticleIndexReverseMappingArray.Add(ParticleIndexReverseMapping[Idx]);
-			}
-			ParticleIndexReverseMappingLock.WriteUnlock();
-
-			LastSwapMappingSyncSize = MEvolution->GetParticles().Size();
-		}
-	}
+		UE_LOG(LogPBDRigidsSolverSolver, Verbose, TEXT("PBDRigidsSolver::Reset()"));
 
-	void FPBDRigidsSolver::KinematicUpdateCallback(FPBDRigidsSolver::FParticlesType& Particles, const float Dt, const float Time)
-	{
-		UE_LOG(LogPBDRigidsSolverSolver, Verbose, TEXT("PBDRigidsSolver::KinematicUpdateCallback()"));
-		SCOPE_CYCLE_COUNTER(STAT_KinematicUpdate);
+		MTime = 0;
+		MLastDt = 0.0f;
+		bEnabled = false;
+		CurrentFrame = 0;
+		MMaxDeltaTime = 1;
+		TimeStepMultiplier = 1;
+		MEvolution = TUniquePtr<FPBDRigidsEvolution>(new FPBDRigidsEvolution(Particles));
 
-		PhysicsParallelFor(KinematicProxiesForObjects.Num(), [&](int32 i)
-		{
-			FKinematicProxy& KinematicProxy = KinematicProxiesForObjects[i];
-			for(int32 ProxyIndex = 0; ProxyIndex < KinematicProxy.Ids.Num(); ++ProxyIndex)
-			{
-				const int32 Index = KinematicProxy.Ids[ProxyIndex];
-				if(Index < 0 || Particles.InvM(Index) != 0 || Particles.Disabled(Index))
-				{
-					continue;
-				}
-				Particles.X(Index) = KinematicProxy.Position[ProxyIndex];
-				Particles.R(Index) = KinematicProxy.Rotation[ProxyIndex];
-				Particles.V(Index) = (KinematicProxy.NextPosition[ProxyIndex] - KinematicProxy.Position[ProxyIndex]) / Dt;
-				Particles.W(Index) = TRotation<float, 3>::CalculateAngularVelocity(KinematicProxy.Rotation[ProxyIndex], KinematicProxy.NextRotation[ProxyIndex], Dt);
-			}
-		});
+		FEventDefaults::RegisterSystemEvents(*GetEventManager());
 	}
 
-	void FPBDRigidsSolver::AddForceCallback(FPBDRigidsSolver::FParticlesType& Particles, const float Dt, const int32 Index)
+	void FPBDRigidsSolver::ChangeBufferMode(Chaos::EMultiBufferMode InBufferMode)
 	{
-		// @todo : The index based callbacks need to change. This should be based on the indices
-		//         managed by the specific Callback. 
-		UE_LOG(LogPBDRigidsSolverSolver, Verbose, TEXT("PBDRigidsSolver::AddForceCallback()"));
-		Chaos::PerParticleGravity<float, 3>(Chaos::TVector<float, 3>(0, 0, -1.f), 980.f).Apply(Particles, Dt, Index);
+		// This seems unused inside the solver? #BH
+		BufferMode = InBufferMode;
 	}
 
-	int32 FPBDRigidsSolver::GetParticleIndexMesh(const TMap<uint32, TUniquePtr<TArray<uint32>>>& ParentToChildrenMap, int32 ParticleIndex)
+	void FPBDRigidsSolver::AdvanceSolverBy(float DeltaTime)
 	{
-		ensure(ParentToChildrenMap.Contains(ParticleIndex));
-
-		while (ParentToChildrenMap[ParticleIndex]->Num())
+		UE_LOG(LogPBDRigidsSolverSolver, Verbose, TEXT("PBDRigidsSolver::Tick(%3.5f)"), DeltaTime);
+		if (bEnabled)
 		{
-			ParticleIndex = (*ParentToChildrenMap[ParticleIndex])[0];
-			if (!ParentToChildrenMap.Contains(ParticleIndex))
+			MLastDt = DeltaTime;
+
+			int32 NumTimeSteps = (int32)(1.f*TimeStepMultiplier);
+			float dt = FMath::Min(DeltaTime, float(5.f / 30.f)) / (float)NumTimeSteps;
+			for (int i = 0; i < NumTimeSteps; i++)
 			{
-				break;
+				AdvanceOneTimeStepTask(this, DeltaTime).DoWork();
 			}
 		}
 
-		return ParticleIndex;
 	}
 
-	int32 FPBDRigidsSolver::EncodeCollisionIndex(int32 ActualCollisionIndex, bool bSwapOrder)
+	void FPBDRigidsSolver::UpdatePhysicsThreadStructures()
 	{
-		return bSwapOrder ? (ActualCollisionIndex | (1 << 31)) : ActualCollisionIndex;
-	}
-
-	int32 FPBDRigidsSolver::DecodeCollisionIndex(int32 EncodedCollisionIdx, bool& bSwapOrder)
-	{
-		bSwapOrder = EncodedCollisionIdx & (1 << 31);
-		return EncodedCollisionIdx & ~(1 << 31);
-	}
-
-	int32 CollisionFilteringThreshold = -1;
-	FAutoConsoleVariableRef CVarCollisionFilteringThreshold(TEXT("p.CollisionFilteringThreshold"), CollisionFilteringThreshold, TEXT("Minimum number of collisions to turn on filtering"));
-
-	int32 CollisionFilteringSkipN = 0;
-	FAutoConsoleVariableRef CVarCollisionFilteringSkipN(TEXT("p.CollisionFilteringSkipN"), CollisionFilteringSkipN, TEXT("N number of collisions will be skipped"));
-
-	DECLARE_DWORD_COUNTER_STAT(TEXT("NumAllCollisionConstraints"), STAT_NumAllCollisionConstraints, STATGROUP_Chaos);
-	DECLARE_DWORD_COUNTER_STAT(TEXT("NumAllValidCollisions"), STAT_NumAllValidCollisions, STATGROUP_Chaos);
-	DECLARE_DWORD_COUNTER_STAT(TEXT("NumAllFilteredCollisions"), STAT_NumAllFilteredCollisions, STATGROUP_Chaos);
+		//ensure(IsInPhyscisThread<BufferMode>());
 
-	void FPBDRigidsSolver::CollisionContactsCallback(FPBDRigidsSolver::FParticlesType& Particles, FPBDRigidsSolver::FCollisionConstraints& CollisionConstraints)
-	{
-		if (bDoGenerateCollisionData && MTime > 0.f)
+		if (FChaosSolversModule* ChaosModule = FModuleManager::Get().GetModulePtr<FChaosSolversModule>("ChaosSolvers"))
 		{
-			const TArrayCollectionArray<ClusterId>& ClusterIdsArray = GetRigidClustering().GetClusterIdsArray();
-			const TMap<uint32, TUniquePtr<TArray<uint32>>>& ParentToChildrenMap = GetRigidClustering().GetChildrenMap();
-
-			const TArray<Chaos::TPBDCollisionConstraint<float, 3>::FRigidBodyContactConstraint>& AllConstraintsArray = CollisionConstraints.GetAllConstraints();
-
-			INC_DWORD_STAT_BY(STAT_NumAllCollisionConstraints, AllConstraintsArray.Num());
-
-			if (AllConstraintsArray.Num() > 0)
+			// #todo(benn.gallagher) : why is the dispatcher failing here...
+			if (Chaos::IDispatcher* Dispatcher = ChaosModule->GetDispatcher())
 			{
-				// Add the keys to AllCollisionsIndicesBySolverObject map
-				Objects.ForEachSolverObject([this](auto* Obj)
+				this->ForEachPhysicsProxyParallel([Dispatcher](auto* Proxy)
 				{
-					if (Obj && Obj->IsSimulating())
+					if (Chaos::FParticleData* ProxyData = Proxy->NewData())
 					{
-						auto& CollisionsIndicesBySolverObject = GetAllCollisionsIndicesBySolverObject();
-						if (!CollisionsIndicesBySolverObject.Contains(Obj))
+						Dispatcher->EnqueueCommandImmediate([Proxy , ProxyData](Chaos::FPersistentPhysicsTask* PhysThread)
 						{
-							CollisionsIndicesBySolverObject.Add(Obj, TArray<int32>());
-						}
+							Proxy->PushToPhysicsState(ProxyData);
+							delete ProxyData;
+						});
+						//Proxy->Particle()->ClearDirtyFlags()
 					}
 				});
-
-				// Get the number of valid constraints (AccumulatedImpulse != 0.f and Phi < 0.f) from AllConstraintsArray
-				TArray<int32> ValidCollisionIndices;
-				ValidCollisionIndices.SetNumUninitialized(AllConstraintsArray.Num());
-				int32 NumValidCollisions = 0;
-
-				for (int32 Idx = 0; Idx < AllConstraintsArray.Num(); ++Idx)
-				{
-					Chaos::TPBDCollisionConstraint<float, 3>::FRigidBodyContactConstraint const& Constraint = AllConstraintsArray[Idx];
-
-					if (Particles.ObjectState(Constraint.ParticleIndex) == Chaos::EObjectStateType::Dynamic)
-					{
-						// Since Clustered GCs can be unioned the particleIndex representing the union 
-						// is not associated with a PhysicsObject
-						if (SolverObjectReverseMapping[Constraint.ParticleIndex].SolverObject != nullptr)
-						{
-							if (ensure(!Constraint.AccumulatedImpulse.ContainsNaN() && FMath::IsFinite(Constraint.Phi)))
-							{
-								if (!Constraint.AccumulatedImpulse.IsZero())
-								{
-									if (ensure(!Constraint.Location.ContainsNaN() &&
-										!Constraint.Normal.ContainsNaN()) &&
-										!Particles.V(Constraint.ParticleIndex).ContainsNaN() &&
-										!Particles.V(Constraint.LevelsetIndex).ContainsNaN() &&
-										!Particles.W(Constraint.ParticleIndex).ContainsNaN() &&
-										!Particles.W(Constraint.LevelsetIndex).ContainsNaN())
-									{
-										ValidCollisionIndices[NumValidCollisions] = Idx;
-										NumValidCollisions++;
-									}
-								}
-							}
-						}
-					}
-				}
-
-				ValidCollisionIndices.SetNum(NumValidCollisions);
-
-				INC_DWORD_STAT_BY(STAT_NumAllValidCollisions, ValidCollisionIndices.Num());
-
-				auto& AllCollisionsDataArray = GetAllCollisionsDataArray();
-				auto& AllCollisionsIndicesBySolverObject = GetAllCollisionsIndicesBySolverObject();
-
-				if (ValidCollisionIndices.Num() > 0)
-				{
-					for (int32 IdxCollision = 0; IdxCollision < ValidCollisionIndices.Num(); ++IdxCollision)
-					{
-						if ((CollisionFilteringThreshold < 0 || CollisionFilteringSkipN < 1) || 
-							(CollisionFilteringThreshold >= 0 && CollisionFilteringSkipN >= 1 && ValidCollisionIndices.Num() <= CollisionFilteringThreshold) ||
-							(CollisionFilteringThreshold >= 0 && CollisionFilteringSkipN >= 1 && ValidCollisionIndices.Num() > CollisionFilteringThreshold && IdxCollision % (CollisionFilteringSkipN + 1) == 0))
-						{
-							Chaos::TPBDCollisionConstraint<float, 3>::FRigidBodyContactConstraint const& Constraint = AllConstraintsArray[ValidCollisionIndices[IdxCollision]];
-
-							TCollisionData<float, 3> Data;
-							Data.Location = Constraint.Location;
-							Data.AccumulatedImpulse = Constraint.AccumulatedImpulse;
-							Data.Normal = Constraint.Normal;
-							Data.Velocity1 = Particles.V(Constraint.ParticleIndex);
-							Data.Velocity2 = Particles.V(Constraint.LevelsetIndex);
-							Data.AngularVelocity1 = Particles.W(Constraint.ParticleIndex);
-							Data.AngularVelocity2 = Particles.W(Constraint.LevelsetIndex);
-							Data.Mass1 = Particles.M(Constraint.ParticleIndex);
-							Data.Mass2 = Particles.M(Constraint.LevelsetIndex);
-							Data.PenetrationDepth = Constraint.Phi;
-							Data.ParticleIndex = Constraint.ParticleIndex;
-							Data.LevelsetIndex = Constraint.LevelsetIndex;
-
-							if (!SolverCollisionEventFilter->Enabled() || SolverCollisionEventFilter->Pass(Data))
-							{
-								const int32 NewIdx = AllCollisionsDataArray.Add(TCollisionData<float, 3>());
-								TCollisionData<float, 3>& CollisionDataArrayItem = AllCollisionsDataArray[NewIdx];
-
-								CollisionDataArrayItem = Data;
-
-								// If Constraint.ParticleIndex is a cluster store an index for a mesh in this cluster
-								if (ClusterIdsArray[Constraint.ParticleIndex].NumChildren > 0)
-								{
-									int32 ParticleIndexMesh = GetParticleIndexMesh(ParentToChildrenMap, Constraint.ParticleIndex);
-									ensure(ParticleIndexMesh != INDEX_NONE);
-									CollisionDataArrayItem.ParticleIndexMesh = ParticleIndexMesh;
-								}
-								// If Constraint.LevelsetIndex is a cluster store an index for a mesh in this cluster
-								if (ClusterIdsArray[Constraint.LevelsetIndex].NumChildren > 0)
-								{
-									int32 LevelsetIndexMesh = GetParticleIndexMesh(ParentToChildrenMap, Constraint.LevelsetIndex);
-									ensure(LevelsetIndexMesh != INDEX_NONE);
-									CollisionDataArrayItem.LevelsetIndexMesh = LevelsetIndexMesh;
-								}
-
-								// Add to AllCollisionsIndicesBySolverObject
-								ISolverObjectBase* const SolverObject = SolverObjectReverseMapping[Constraint.ParticleIndex].SolverObject;
-								AllCollisionsIndicesBySolverObject.FindOrAdd(SolverObject).Add(EncodeCollisionIndex(NewIdx, false));
-
-								ISolverObjectBase* const OtherSolverObject = SolverObjectReverseMapping[Constraint.LevelsetIndex].SolverObject;
-								if (OtherSolverObject && OtherSolverObject != SolverObject)
-								{
-									AllCollisionsIndicesBySolverObject.FindOrAdd(OtherSolverObject).Add(EncodeCollisionIndex(NewIdx, true));
-								}
-							}
-						}
-					}
-				}
-
-				INC_DWORD_STAT_BY(STAT_NumAllFilteredCollisions, AllCollisionsDataArray.Num());
 			}
-		}		
-	}
-
-	void FPBDRigidsSolver::BreakingCallback(FPBDRigidsSolver::FParticlesType& Particles)
-	{
-		if (bDoGenerateBreakingData && MTime > 0.f)
-		{
-			const TArrayCollectionArray<ClusterId>& ClusterIdsArray = GetRigidClustering().GetClusterIdsArray();
-			const TMap<uint32, TUniquePtr<TArray<uint32>>>& ParentToChildrenMap = GetRigidClustering().GetChildrenMap();
-
-			const TArray<TBreakingData<float, 3>>& AllBreakingsArray = GetRigidClustering().GetAllClusterBreakings();
-
-			auto& AllBreakingsDataArray = GetAllBreakingsDataArray();
-			auto& AllBreakingsIndicesBySolverObject = GetAllBreakingsIndicesBySolverObject();
-
-			if (AllBreakingsArray.Num() > 0)
+			else // no threading
 			{
-				// Add the keys to AllCollisionsIndicesBySolverObject map
-				Objects.ForEachSolverObject([this](auto* Obj)
+				this->ForEachPhysicsProxy([](auto* Proxy)
 				{
-					if (Obj && Obj->IsSimulating())
+					if (Chaos::FParticleData* ProxyData = Proxy->NewData())
 					{
-						auto& BreakingsIndicesBySolverObject = GetAllBreakingsIndicesBySolverObject();
-						if (!BreakingsIndicesBySolverObject.Contains(Obj))
-						{
-							BreakingsIndicesBySolverObject.Add(Obj, TArray<int32>());
-						}
+						Proxy->PushToPhysicsState(ProxyData);
+						delete ProxyData;
+						//Proxy->Particle()->ClearDirtyFlags()
 					}
 				});
-
-				// Get the number of valid constraints (AccumulatedImpulse != 0.f and Phi < 0.f) from AllConstraintsArray
-				for (int32 Idx = 0; Idx < AllBreakingsArray.Num(); ++Idx)
-				{
-					// Since Clustered GCs can be unioned the particleIndex representing the union 
-					// is not associated with a PhysicsObject
-					if (SolverObjectReverseMapping[AllBreakingsArray[Idx].ParticleIndex].SolverObject != nullptr)
-					{
-						if (ensure(!AllBreakingsArray[Idx].Location.ContainsNaN() &&
-							!Particles.V(AllBreakingsArray[Idx].ParticleIndex).ContainsNaN() &&
-							!Particles.W(AllBreakingsArray[Idx].ParticleIndex).ContainsNaN()))
-						{
-							TBreakingData<float, 3> BreakingData;
-							BreakingData.Location = AllBreakingsArray[Idx].Location;
-							BreakingData.Velocity = Particles.V(AllBreakingsArray[Idx].ParticleIndex);
-							BreakingData.AngularVelocity = Particles.W(AllBreakingsArray[Idx].ParticleIndex);
-							BreakingData.Mass = Particles.M(AllBreakingsArray[Idx].ParticleIndex);
-							BreakingData.ParticleIndex = AllBreakingsArray[Idx].ParticleIndex;
-							if (Particles.Geometry(Idx) && Particles.Geometry(Idx)->HasBoundingBox())
-							{
-								BreakingData.BoundingBox = Particles.Geometry(Idx)->BoundingBox();;
-							}
-
-							if (!SolverBreakingEventFilter->Enabled() || SolverBreakingEventFilter->Pass(BreakingData))
-							{
-								int32 NewIdx = AllBreakingsDataArray.Add(TBreakingData<float, 3>());
-								TBreakingData<float, 3>& BreakingDataArrayItem = AllBreakingsDataArray[NewIdx];
-								BreakingDataArrayItem = BreakingData;
-
-								// If AllBreakingsArray[Idx].ParticleIndex is a cluster store an index for a mesh in this cluster
-								if (ClusterIdsArray[AllBreakingsArray[Idx].ParticleIndex].NumChildren > 0)
-								{
-									int32 ParticleIndexMesh = GetParticleIndexMesh(ParentToChildrenMap, AllBreakingsArray[Idx].ParticleIndex);
-									ensure(ParticleIndexMesh != INDEX_NONE);
-									BreakingDataArrayItem.ParticleIndexMesh = ParticleIndexMesh;
-								}
-
-								// Add to AllBreakingsIndicesBySolverObject
-								ISolverObjectBase* SolverObject = SolverObjectReverseMapping[AllBreakingsArray[Idx].ParticleIndex].SolverObject;
-								AllBreakingsIndicesBySolverObject.FindOrAdd(SolverObject).Add(NewIdx);
-							}
-						}
-					}
-				}
 			}
 		}
 	}
 
-	void FPBDRigidsSolver::TrailingCallback(FPBDRigidsSolver::FParticlesType& Particles)
+	void FPBDRigidsSolver::UpdateGameThreadStructures()
 	{
-		if (bDoGenerateTrailingData && Particles.Size() > 0 && MTime > 0.f)
-		{
-			const TArrayCollectionArray<ClusterId>& ClusterIdsArray = GetRigidClustering().GetClusterIdsArray();
-			const TMap<uint32, TUniquePtr<TArray<uint32>>>& ParentToChildrenMap = GetRigidClustering().GetChildrenMap();
+		//ensure(IsInGameThread());
 
-			auto& AllTrailingsDataArray = GetAllTrailingsDataArray();
-			auto& AllTrailingsIndicesBySolverObject = GetAllTrailingsIndicesBySolverObject();
-
-			// Add the keys to AllTrailingsIndicesBySolverObject map
-			Objects.ForEachSolverObject([this](auto* Obj)
-			{
-				if (Obj && Obj->IsSimulating())
-				{
-					auto& TrailingsIndicesBySolverObject = GetAllTrailingsIndicesBySolverObject();
-					if (!TrailingsIndicesBySolverObject.Contains(Obj))
-					{
-						TrailingsIndicesBySolverObject.Add(Obj, TArray<int32>());
-					}
-				}
-			});
+		// on game thread
+		this->ForEachPhysicsProxy([](auto* Object)
+		{
+			Object->PullFromPhysicsState();
+		});
 
-			for (uint32 IdxParticle : ActiveIndices())
-			{
-				// Since Clustered GCs can be unioned the particleIndex representing the union 
-				// is not associated with a PhysicsObject
-				if (SolverObjectReverseMapping[IdxParticle].SolverObject != nullptr)
-				{
-					if (ensure(FMath::IsFinite(Particles.InvM(IdxParticle))))
-					{
-						if (Particles.InvM(IdxParticle) != 0.f &&
-							Particles.Geometry(IdxParticle) &&
-							Particles.Geometry(IdxParticle)->HasBoundingBox())
-						{
-							if (ensure(!Particles.X(IdxParticle).ContainsNaN() &&
-								!Particles.V(IdxParticle).ContainsNaN() &&
-								!Particles.W(IdxParticle).ContainsNaN() &&
-								FMath::IsFinite(Particles.M(IdxParticle))))
-							{
-								TTrailingData<float, 3> TrailingData;
-								TrailingData.Location = Particles.X(IdxParticle);
-								TrailingData.Velocity = Particles.V(IdxParticle);
-								TrailingData.AngularVelocity = Particles.W(IdxParticle);
-								TrailingData.Mass = Particles.M(IdxParticle);
-								TrailingData.ParticleIndex = (int32)IdxParticle;
-								if (Particles.Geometry(IdxParticle)->HasBoundingBox())
-								{
-									TrailingData.BoundingBox = Particles.Geometry(IdxParticle)->BoundingBox();
-								}
-
-								if (!SolverTrailingEventFilter->Enabled() || SolverTrailingEventFilter->Pass(TrailingData))
-								{
-									int32 NewIdx = AllTrailingsDataArray.Add(TTrailingData<float, 3>());
-									TTrailingData<float, 3>& TrailingDataArrayItem = AllTrailingsDataArray[NewIdx];
-									TrailingDataArrayItem = TrailingData;
-
-									// If IdxParticle is a cluster store an index for a mesh in this cluster
-									if (ClusterIdsArray[IdxParticle].NumChildren > 0)
-									{
-										int32 ParticleIndexMesh = GetParticleIndexMesh(ParentToChildrenMap, IdxParticle);
-										ensure(ParticleIndexMesh != INDEX_NONE);
-										TrailingDataArrayItem.ParticleIndexMesh = ParticleIndexMesh;
-									}
-
-									// Add to AllTrailingsIndicesBySolverObject
-									ISolverObjectBase* SolverObject = SolverObjectReverseMapping[IdxParticle].SolverObject;
-									AllTrailingsIndicesBySolverObject.FindOrAdd(SolverObject).Add(NewIdx);
-								}
-							}
-						}
-					}
-				}
-			}
-		}
 	}
-}
-
-void FSolverObjectStorage::Reset()
-{
-	GeometryCollectionObjects.Reset();
-	SkeletalMeshObjects.Reset();
-	StaticMeshObjects.Reset();
-	FieldSystemObjects.Reset();
-}
-
-int32 FSolverObjectStorage::GetNumObjects() const
-{
-	return GetNumObjectsOfType<FGeometryCollectionPhysicsObject>() +
-		GetNumObjectsOfType<FSkeletalMeshPhysicsObject>() +
-		GetNumObjectsOfType<FStaticMeshPhysicsObject>() +
-		GetNumObjectsOfType<FBodyInstancePhysicsObject>();
-}
-
-int32 FSolverObjectStorage::GetNumFieldObjects() const
-{
-	return FieldSystemObjects.Num();
-}
-
-template<>
-TArray<TSolverObject<FGeometryCollectionPhysicsObject>*>& FSolverObjectStorage::GetObjectStorage()
-{
-	return GeometryCollectionObjects;
-}
+	
+	void FPBDRigidsSolver::BufferPhysicsResults()
+	{
+		//ensure(IsInGameThread());
 
-template<>
-TArray<TSolverObject<FSkeletalMeshPhysicsObject>*>& FSolverObjectStorage::GetObjectStorage()
-{
-	return SkeletalMeshObjects;
-}
+		// on game thread
+		this->ForEachPhysicsProxy([](auto* Object)
+		{
+			Object->BufferPhysicsResults();
+		});
 
-template<>
-TArray<TSolverObject<FStaticMeshPhysicsObject>*>& FSolverObjectStorage::GetObjectStorage()
-{
-	return StaticMeshObjects;
-}
+	}
 
-template<>
-TArray<TSolverObject<FBodyInstancePhysicsObject>*>& FSolverObjectStorage::GetObjectStorage()
-{
-	return BodyInstanceObjects;
-}
 
-template<>
-TArray<FFieldSystemPhysicsObject*>& FSolverObjectStorage::GetFieldObjectStorage()
-{
-	return FieldSystemObjects;
-}
+		
+	void FPBDRigidsSolver::FlipBuffers()
+	{
+		//ensure(IsInGameThread());
 
-template<>
-int32 FSolverObjectStorage::GetNumObjectsOfType<FGeometryCollectionPhysicsObject>() const
-{
-	return GeometryCollectionObjects.Num();
-}
+		// on game thread
+		this->ForEachPhysicsProxy([](auto* Object)
+		{
+			Object->FlipBuffer();
+		});
 
-template<>
-int32 FSolverObjectStorage::GetNumObjectsOfType<FSkeletalMeshPhysicsObject>() const
-{
-	return SkeletalMeshObjects.Num();
-}
+	}
 
-template<>
-int32 FSolverObjectStorage::GetNumObjectsOfType<FStaticMeshPhysicsObject>() const
-{
-	return StaticMeshObjects.Num();
-}
+	void FPBDRigidsSolver::SyncEvents_GameThread()
+	{
+		GetEventManager()->DispatchEvents();
+	}
 
-template<>
-int32 FSolverObjectStorage::GetNumObjectsOfType<FBodyInstancePhysicsObject>() const
-{
-	return BodyInstanceObjects.Num();
-}
+}; // namespace Chaos
 
-template<>
-int32 FSolverObjectStorage::GetNumObjectsOfType<FFieldSystemPhysicsObject>() const
-{
-	return FieldSystemObjects.Num();
-}
 
 #endif
