Copying //UE4/Dev-Physics to Dev-Main (//UE4/Dev-Main) [at] 7703071

diff --git a/Engine/Source/Runtime/Experimental/ChaosSolvers/Private/PhysicsProxy/FieldSystemPhysicsProxy.cpp b/Engine/Source/Runtime/Experimental/ChaosSolvers/Private/PhysicsProxy/FieldSystemPhysicsProxy.cpp
new file mode 100644
--- /dev/null
+++ b/Engine/Source/Runtime/Experimental/ChaosSolvers/Private/PhysicsProxy/FieldSystemPhysicsProxy.cpp
@@ -0,0 +1,1098 @@
+// Copyright 1998-2019 Epic Games, Inc. All Rights Reserved.
+
+#if INCLUDE_CHAOS
+
+#include "PhysicsProxy/FieldSystemPhysicsProxy.h"
+#include "PhysicsProxy/GeometryCollectionPhysicsProxy.h"
+#include "Field/FieldSystem.h"
+#include "Chaos/PBDRigidClustering.h"
+#include "GeometryCollection/GeometryCollectionSimulationTypes.h"
+#include "PhysicsSolver.h"
+#include "ChaosStats.h"
+
+void ResetIndicesArray(TArray<int32> & IndicesArray, int32 Size)
+{
+	if(IndicesArray.Num() != Size)
+	{
+		IndicesArray.SetNum(Size);
+		for(int32 i = 0; i < IndicesArray.Num(); ++i)
+		{
+			IndicesArray[i] = i;
+		}
+	}
+}
+
+
+FFieldSystemPhysicsProxy::FFieldSystemPhysicsProxy(UObject* InOwner)
+	: Base(InOwner)
+{
+
+}
+
+FFieldSystemPhysicsProxy::~FFieldSystemPhysicsProxy()
+{
+	// Need to delete any command lists we have hanging around
+	FScopeLock Lock(&CommandLock);
+	for(TTuple<Chaos::FPhysicsSolver*, TArray<FFieldSystemCommand>*>& Pair : Commands)
+	{
+		delete Pair.Get<1>();
+	}
+	Commands.Reset();
+}
+
+bool FFieldSystemPhysicsProxy::IsSimulating() const
+{
+	return true; // #todo Actually start gating this?
+}
+
+void FFieldSystemPhysicsProxy::FieldParameterUpdateCallback(Chaos::FPhysicsSolver* InSolver, FParticlesType& Particles, Chaos::TArrayCollectionArray<float>& Strains, Chaos::TPBDPositionConstraints<float, 3>& PositionTarget, TMap<int32, int32>& PositionTargetedParticles, const TArray<FKinematicProxy>& AnimatedPosition, const float InTime)
+{
+	using namespace Chaos;
+	SCOPE_CYCLE_COUNTER(STAT_ParamUpdateField_Object);
+	
+	Chaos::FPhysicsSolver* CurrentSolver = InSolver;
+
+	if (Commands.Num() && InSolver)
+	{
+#if TODO_REIMPLEMENT_RIGID_CLUSTERING
+		const Chaos::TArrayCollectionArray<Chaos::ClusterId> & ClusterIDs = CurrentSolver->GetRigidClustering().GetClusterIdsArray();
+#endif
+
+		// @todo: This seems like a waste if we just want to get everything
+		TArray<ContextIndex> IndicesArray;
+
+		TArray<FFieldSystemCommand>* CommandListPtr = GetSolverCommandList(InSolver);
+		
+		if(!CommandListPtr)
+		{
+			// No command list present for this solver, bail out
+			return;
+		}
+
+		const int32 NumCommands = CommandListPtr->Num();
+
+		TArray<int32> CommandsToRemove;
+		for (int32 CommandIndex = 0; CommandIndex < NumCommands; CommandIndex++)
+		{
+			const FFieldSystemCommand& Command = (*CommandListPtr)[CommandIndex];
+			EFieldResolutionType ResolutionType = EFieldResolutionType::Field_Resolution_Minimal;
+			if (Command.MetaData.Contains(FFieldSystemMetaData::EMetaType::ECommandData_ProcessingResolution))
+			{
+				check(Command.MetaData[FFieldSystemMetaData::EMetaType::ECommandData_ProcessingResolution] != nullptr);
+				ResolutionType = static_cast<FFieldSystemMetaDataProcessingResolution*>(Command.MetaData[FFieldSystemMetaData::EMetaType::ECommandData_ProcessingResolution].Get())->ProcessingResolution;
+			}
+
+			if (Command.TargetAttribute == GetFieldPhysicsName(EFieldPhysicsType::Field_DynamicState))
+			{
+				SCOPE_CYCLE_COUNTER(STAT_ParamUpdateField_DynamicState);
+
+				FFieldSystemPhysicsProxy::ContiguousIndices(IndicesArray, CurrentSolver, ResolutionType, IndicesArray.Num() != Particles.Size());
+				if (IndicesArray.Num())
+				{
+					TArrayView<ContextIndex> IndexView(&(IndicesArray[0]), IndicesArray.Num());
+
+					FVector * tptr = &(Particles.X(0));
+					TArrayView<FVector> SamplesView(tptr, int32(Particles.Size()));
+
+					FFieldContext Context{
+						IndexView, // @todo(brice) important: an empty index array should evaluate everything
+						SamplesView,
+						Command.MetaData
+					};
+
+					//
+					//  Sample the dynamic state array in the field
+					//
+					TArray<int32> DynamicState;
+					DynamicState.AddUninitialized(Particles.Size());
+					for(const ContextIndex& Index : IndicesArray)
+					{
+						DynamicState[Index.Sample] = 0;	//is this needed?
+					}
+					
+					for (const ContextIndex& Index : IndicesArray)
+					{
+						const int32 i = Index.Sample;
+						if (Particles.ObjectState(i) == Chaos::EObjectStateType::Kinematic)
+						{
+							DynamicState[i] = (int)EObjectStateTypeEnum::Chaos_Object_Kinematic;
+						}
+						else if (Particles.ObjectState(i) == Chaos::EObjectStateType::Static)
+						{
+							DynamicState[i] = (int)EObjectStateTypeEnum::Chaos_Object_Static;
+						}
+						else
+						{
+							DynamicState[i] = (int)EObjectStateTypeEnum::Chaos_Object_Dynamic;
+						}
+					}
+					TArrayView<int32> DynamicStateView(&(DynamicState[0]), DynamicState.Num());
+
+					if (ensureMsgf(Command.RootNode->Type() == FFieldNode<int32>::StaticType(),
+						TEXT("Field based evaluation of the simulations 'ObjectType' parameter expects int32 field inputs.")))
+					{
+						static_cast<const FFieldNode<int32> *>(Command.RootNode.Get())->Evaluate(Context, DynamicStateView);
+					}
+
+#if TODO_REIMPLEMENT_GETFLOORINDEX
+					// transfer results to rigid system.
+					int32 FloorIndex = CurrentSolver->GetFloorIndex();
+					int32 NumSamples = Context.SampleIndices.Num();
+					for (int32 SampleIndex = 0; SampleIndex < NumSamples; SampleIndex++)
+					{
+						int32 RigidBodyIndex = Context.SampleIndices[SampleIndex].Result;
+						if (RigidBodyIndex != FloorIndex) // ignore the floor
+						{ 
+							if (DynamicStateView[RigidBodyIndex] == (int32)EObjectStateTypeEnum::Chaos_Object_Dynamic
+								&& Particles.ObjectState(RigidBodyIndex) == Chaos::EObjectStateType::Static
+								&& FLT_EPSILON < Particles.M(RigidBodyIndex))
+							{
+								Particles.SetObjectState(RigidBodyIndex, Chaos::EObjectStateType::Dynamic);
+							}
+							else if ((DynamicStateView[RigidBodyIndex] == (int32)EObjectStateTypeEnum::Chaos_Object_Kinematic)
+								&& Particles.ObjectState(RigidBodyIndex) == Chaos::EObjectStateType::Dynamic)
+							{
+								Particles.SetObjectState(RigidBodyIndex, Chaos::EObjectStateType::Kinematic);
+								Particles.V(RigidBodyIndex) = Chaos::TVector<float, 3>(0);
+								Particles.W(RigidBodyIndex) = Chaos::TVector<float, 3>(0);
+							}
+							else if ((DynamicStateView[RigidBodyIndex] == (int32)EObjectStateTypeEnum::Chaos_Object_Static)
+								&& Particles.ObjectState(RigidBodyIndex) == Chaos::EObjectStateType::Dynamic)
+							{
+								Particles.SetObjectState(RigidBodyIndex, Chaos::EObjectStateType::Static);
+								Particles.V(RigidBodyIndex) = Chaos::TVector<float, 3>(0);
+								Particles.W(RigidBodyIndex) = Chaos::TVector<float, 3>(0);
+							}
+							else if ((DynamicStateView[RigidBodyIndex] == (int32)EObjectStateTypeEnum::Chaos_Object_Dynamic)
+								&& Particles.ObjectState(RigidBodyIndex) == Chaos::EObjectStateType::Sleeping)
+							{
+								Particles.SetObjectState(RigidBodyIndex, Chaos::EObjectStateType::Dynamic);
+								CurrentSolver->ActiveIndices().Add(RigidBodyIndex);
+							}
+							else if ((DynamicStateView[RigidBodyIndex] == (int32)EObjectStateTypeEnum::Chaos_Object_Sleeping)
+								&& Particles.ObjectState(RigidBodyIndex) == Chaos::EObjectStateType::Dynamic)
+							{
+								Particles.SetObjectState(RigidBodyIndex, Chaos::EObjectStateType::Sleeping);
+								CurrentSolver->ActiveIndices().Remove(RigidBodyIndex);
+							}
+						}
+					}
+#endif
+
+#if TODO_REIMPLEMENT_RIGID_CLUSTERING
+					//  Update all cluster bodies based on the changes in the kinematic state.
+					const Chaos::TArrayCollectionArray<Chaos::ClusterId>& ClusterIdArray = CurrentSolver->GetRigidClustering().GetClusterIdsArray();
+					for (int32 ActiveParticleIndex : CurrentSolver->ActiveIndices())
+					{
+						if (ClusterIdArray[ActiveParticleIndex].NumChildren)
+						{
+							CurrentSolver->GetRigidClustering().UpdateKinematicProperties(ActiveParticleIndex);
+						}
+					}
+#endif
+
+				}
+				CommandsToRemove.Add(CommandIndex);
+			}
+			else if (Command.TargetAttribute == GetFieldPhysicsName(EFieldPhysicsType::Field_ActivateDisabled))
+			{
+				FFieldSystemPhysicsProxy::ContiguousIndices(IndicesArray, CurrentSolver, ResolutionType, IndicesArray.Num() != Particles.Size());
+				if (IndicesArray.Num())
+				{
+					TArrayView<ContextIndex> IndexView(&(IndicesArray[0]), IndicesArray.Num());
+
+					FVector * tptr = &(Particles.X(0));
+					TArrayView<FVector> SamplesView(tptr, int32(Particles.Size()));
+
+					FFieldContext Context{
+						IndexView, // @todo(brice) important: an empty index array should evaluate everything
+						SamplesView,
+						Command.MetaData
+					};
+
+					//
+					//  Sample the dynamic state array in the field
+					//
+					TArray<int32> DynamicState;
+					DynamicState.AddUninitialized(Particles.Size());
+					for(const ContextIndex& Index : IndicesArray)
+					{
+						DynamicState[Index.Sample] = 0;	//is this needed?
+					}
+					
+					for (const ContextIndex& Index : IndicesArray)
+					{
+						const int32 i = Index.Sample;
+						if (Particles.Disabled(i))
+						{
+							DynamicState[i] = 1;
+						}
+						else
+						{
+							DynamicState[i] = 0;
+						}
+					}
+					TArrayView<int32> DynamicStateView(&(DynamicState[0]), DynamicState.Num());
+
+					if (ensureMsgf(Command.RootNode->Type() == FFieldNode<int32>::StaticType(),
+						TEXT("Field based evaluation of the simulations 'ObjectType' parameter expects int32 field inputs.")))
+					{
+						static_cast<const FFieldNode<int32> *>(Command.RootNode.Get())->Evaluate(Context, DynamicStateView);
+					}
+
+#if TODO_REIMPLEMENT_RIGID_CLUSTERING
+					// transfer results to rigid system.
+					int32 FloorIndex = CurrentSolver->GetFloorIndex();
+					int32 NumSamples = Context.SampleIndices.Num();
+					for (int32 SampleIndex = 0; SampleIndex < NumSamples; SampleIndex++)
+					{
+						int32 RigidBodyIndex = Context.SampleIndices[SampleIndex].Result;
+						if (RigidBodyIndex != FloorIndex) // ignore the floor
+						{ 
+							if (DynamicStateView[RigidBodyIndex] == 0 && Particles.Disabled(RigidBodyIndex))
+							{
+								ensure(CurrentSolver->GetRigidClustering().GetClusterIdsArray()[RigidBodyIndex].Id == INDEX_NONE);
+								CurrentSolver->GetEvolution()->EnableParticle(RigidBodyIndex, INDEX_NONE);
+								Particles.SetObjectState(RigidBodyIndex, Chaos::EObjectStateType::Dynamic);
+							}
+						}
+					}
+#endif
+				}
+				CommandsToRemove.Add(CommandIndex);
+			}
+#if TODO_REIMPLEMENT_RIGID_CLUSTERING
+			else if (Command.TargetAttribute == GetFieldPhysicsName(EFieldPhysicsType::Field_ExternalClusterStrain))
+			{
+				SCOPE_CYCLE_COUNTER(STAT_ParamUpdateField_ExternalClusterStrain);
+
+				if (ensureMsgf(Command.RootNode->Type() == FFieldNode<float>::StaticType(),
+					TEXT("Field based evaluation of the simulations 'Strain' parameter expects float field inputs.")))
+				{
+					FFieldSystemPhysicsProxy::ContiguousIndices(IndicesArray, CurrentSolver, ResolutionType, IndicesArray.Num() != Particles.Size());
+					if (IndicesArray.Num())
+					{
+						TArrayView<ContextIndex> IndexView(&(IndicesArray[0]), IndicesArray.Num());
+
+						FVector * tptr = &(Particles.X(0));
+						TArrayView<FVector> SamplesView(tptr, int32(Particles.Size()));
+
+						FFieldContext Context{
+							IndexView,
+							SamplesView,
+							Command.MetaData
+						};
+
+						TArray<float> StrainSamples;
+						StrainSamples.AddUninitialized(SamplesView.Num());
+						for (const ContextIndex& Index : IndicesArray)
+						{
+							StrainSamples[Index.Sample] = 0.f;
+						}
+						TArrayView<float> FloatBuffer(&StrainSamples[0], StrainSamples.Num());
+						static_cast<const FFieldNode<float> *>(Command.RootNode.Get())->Evaluate(Context, FloatBuffer);
+
+						int32 Iterations = 1;
+						if (Command.MetaData.Contains(FFieldSystemMetaData::EMetaType::ECommandData_Iteration))
+						{
+							Iterations = static_cast<FFieldSystemMetaDataIteration*>(Command.MetaData[FFieldSystemMetaData::EMetaType::ECommandData_Iteration].Get())->Iterations;
+						}
+
+						if (StrainSamples.Num())
+						{
+							CurrentSolver->GetRigidClustering().BreakingModel(StrainSamples);
+						}
+					}
+				}
+				CommandsToRemove.Add(CommandIndex);
+			}
+#endif
+			else if (Command.TargetAttribute == GetFieldPhysicsName(EFieldPhysicsType::Field_Kill))
+			{
+				SCOPE_CYCLE_COUNTER(STAT_ParamUpdateField_Kill);
+
+				if (ensureMsgf(Command.RootNode->Type() == FFieldNode<float>::StaticType(),
+					TEXT("Field based evaluation of the simulations 'Disabled' parameter expects float field inputs.")))
+				{
+					FFieldSystemPhysicsProxy::ContiguousIndices(IndicesArray, CurrentSolver, ResolutionType, IndicesArray.Num() != Particles.Size());
+					if (IndicesArray.Num())
+					{
+						TArrayView<ContextIndex> IndexView(&(IndicesArray[0]), IndicesArray.Num());
+
+						FVector * tptr = &(Particles.X(0));
+						TArrayView<FVector> SamplesView(tptr, int32(Particles.Size()));
+
+						FFieldContext Context{
+							IndexView,
+							SamplesView,
+							Command.MetaData
+						};
+
+						// @todo(better) : Dont copy, support native type conversion. 
+						TArray<float> Results;
+						Results.AddUninitialized(Particles.Size());
+						for (const ContextIndex& Index : IndicesArray)
+						{
+							Results[Index.Sample] = 0.f;
+						}
+						TArrayView<float> ResultsView(&(Results[0]), Results.Num());
+						static_cast<const FFieldNode<float> *>(Command.RootNode.Get())->Evaluate(Context, ResultsView);
+
+#if TODO_REIMPLEMENT_GETFLOORINDEX
+						bool bHasFloor = false;
+						int32 FloorIndex = CurrentSolver->GetFloorIndex();
+						if (FloorIndex != INDEX_NONE)
+						{
+							bHasFloor = !Particles.Disabled(FloorIndex);
+						}
+
+						TSet<uint32> RemovedParticles;
+						for (const ContextIndex& Index : IndicesArray)
+						{
+							const int32 i = Index.Result;
+							if (!Particles.Disabled(i) && Results[i] > 0.0)
+							{
+								RemovedParticles.Add((uint32)i);
+								CurrentSolver->GetEvolution()->DisableParticle(i);
+							}
+						}
+
+						if (RemovedParticles.Num() && bHasFloor)
+						{
+							CurrentSolver->GetEvolution()->DisableParticle(FloorIndex);
+							Particles.SetObjectState(FloorIndex, Chaos::EObjectStateType::Static);
+						}
+#endif
+					}
+				}
+				CommandsToRemove.Add(CommandIndex);
+			}
+			else if (Command.TargetAttribute == GetFieldPhysicsName(EFieldPhysicsType::Field_LinearVelocity))
+			{
+				SCOPE_CYCLE_COUNTER(STAT_ParamUpdateField_LinearVelocity);
+
+				if (ensureMsgf(Command.RootNode->Type() == FFieldNode<FVector>::StaticType(),
+					TEXT("Field based evaluation of the simulations 'LinearVelocity' parameter expects FVector field inputs.")))
+				{
+					FFieldSystemPhysicsProxy::ContiguousIndices(IndicesArray, CurrentSolver, ResolutionType, IndicesArray.Num() != Particles.Size());
+					if (IndicesArray.Num())
+					{
+						TArrayView<ContextIndex> IndexView(&(IndicesArray[0]), IndicesArray.Num());
+
+						FVector * tptr = &(Particles.X(0));
+						TArrayView<FVector> SamplesView(tptr, int32(Particles.Size()));
+
+						FFieldContext Context{
+							IndexView,
+							SamplesView,
+							Command.MetaData
+						};
+
+						FVector * vptr = &(Particles.V(0));
+						TArrayView<FVector> ResultsView(vptr, Particles.Size());
+						static_cast<const FFieldNode<FVector> *>(Command.RootNode.Get())->Evaluate(Context, ResultsView);
+					}
+				}
+				CommandsToRemove.Add(CommandIndex);
+			}
+			else if (Command.TargetAttribute == GetFieldPhysicsName(EFieldPhysicsType::Field_AngularVelociy))
+			{
+				SCOPE_CYCLE_COUNTER(STAT_ParamUpdateField_AngularVelocity);
+
+				if (ensureMsgf(Command.RootNode->Type() == FFieldNode<FVector>::StaticType(),
+					TEXT("Field based evaluation of the simulations 'AngularVelocity' parameter expects FVector field inputs.")))
+				{
+					FFieldSystemPhysicsProxy::ContiguousIndices(IndicesArray, CurrentSolver, ResolutionType, IndicesArray.Num() != Particles.Size());
+					if (IndicesArray.Num())
+					{
+						TArrayView<ContextIndex> IndexView(&(IndicesArray[0]), IndicesArray.Num());
+
+						FVector * tptr = &(Particles.X(0));
+						TArrayView<FVector> SamplesView(tptr, int32(Particles.Size()));
+
+						FFieldContext Context{
+							IndexView,
+							SamplesView,
+							Command.MetaData
+						};
+
+						FVector * vptr = &(Particles.W(0));
+						TArrayView<FVector> ResultsView(vptr, Particles.Size());
+						static_cast<const FFieldNode<FVector> *>(Command.RootNode.Get())->Evaluate(Context, ResultsView);
+					}
+				}
+				CommandsToRemove.Add(CommandIndex);
+			}
+			else if (Command.TargetAttribute == GetFieldPhysicsName(EFieldPhysicsType::Field_SleepingThreshold))
+			{
+				SCOPE_CYCLE_COUNTER(STAT_ParamUpdateField_SleepingThreshold);
+
+				if (ensureMsgf(Command.RootNode->Type() == FFieldNode<float>::StaticType(),
+					TEXT("Field based evaluation of the simulations 'Disable' parameter expects scale field inputs.")))
+				{
+#if TODO_REIMPLEMENT_PHYSICS_PROXY_REVERSE_MAPPING
+					const Chaos::TArrayCollectionArray<PhysicsProxyWrapper>& PhysicsProxyMapping = CurrentSolver->GetPhysicsProxyReverseMapping();
+					
+					FFieldSystemPhysicsProxy::ContiguousIndices(IndicesArray, CurrentSolver, ResolutionType, IndicesArray.Num() != Particles.Size());
+					if (IndicesArray.Num())
+					{
+						TArrayView<ContextIndex> IndexView(&(IndicesArray[0]), IndicesArray.Num());
+
+						FVector * tptr = &(Particles.X(0));
+						TArrayView<FVector> SamplesView(tptr, int32(Particles.Size()));
+
+						FFieldContext Context{
+							IndexView,
+							SamplesView,
+							Command.MetaData
+						};
+
+						TArray<float> Results;
+						Results.AddUninitialized(Particles.Size());
+						for (const ContextIndex Index : IndicesArray)
+						{
+							const SolverObjectWrapper& ParticleObjectWrapper = SolverObjectMapping[Index.Result];
+							TSerializablePtr<Chaos::TChaosPhysicsMaterial<float>> Material = CurrentSolver->GetPhysicsMaterial(Index.Result);
+							if (ensure(Material) && ParticleObjectWrapper.SolverObject)
+							{
+								const TUniquePtr<Chaos::TChaosPhysicsMaterial<float>>& InstanceMaterial = CurrentSolver->GetPerParticlePhysicsMaterial(Index.Result);
+								if (InstanceMaterial)
+								{
+									Results[Index.Result] = InstanceMaterial->SleepingLinearThreshold;
+								}
+								else
+								{
+									Results[Index.Result] = Material->SleepingLinearThreshold;
+								}
+							}
+							else
+							{
+								Results[Index.Result] = 0.0f;
+							}
+						}
+
+						TArrayView<float> ResultsView(&(Results[0]), Results.Num());
+						static_cast<const FFieldNode<float> *>(Command.RootNode.Get())->Evaluate(Context, ResultsView);
+
+						for (const ContextIndex Index : IndicesArray)
+						{
+							const int32 i = Index.Result;
+							const PhysicsProxyWrapper& ParticleObjectWrapper = PhysicsProxyMapping[i];
+							TSerializablePtr<TChaosPhysicsMaterial<float>> Material = CurrentSolver->GetPhysicsMaterial(i);
+							if (!ensure(Material) || !ParticleObjectWrapper.PhysicsProxy)	//question: do we actually need to check for solver object?
+							{
+								continue;
+							}
+
+							//per instance override
+							if (!CurrentSolver->GetPerParticlePhysicsMaterial(Index.Result))
+							{
+								if (Results[i] != Material->SleepingLinearThreshold)
+								{
+									// value changed from shared material, make unique material.
+									TUniquePtr<Chaos::TChaosPhysicsMaterial<float>> NewMaterial = MakeUnique< Chaos::TChaosPhysicsMaterial<float>>(*Material);
+									CurrentSolver->SetPerParticlePhysicsMaterial(Index.Result, MoveTemp(NewMaterial));
+									const TUniquePtr<TChaosPhysicsMaterial<float>>& InstanceMaterial = CurrentSolver->GetPerParticlePhysicsMaterial(i);
+
+									InstanceMaterial->SleepingLinearThreshold = Results[i];
+									InstanceMaterial->SleepingAngularThreshold = Results[i];
+								}
+							}
+							else
+							{
+								const TUniquePtr<TChaosPhysicsMaterial<float>>& InstanceMaterial = CurrentSolver->GetPerParticlePhysicsMaterial(i);
+
+								if (InstanceMaterial->SleepingLinearThreshold != Results[i])
+								{
+									InstanceMaterial->SleepingLinearThreshold = Results[i];
+									InstanceMaterial->SleepingAngularThreshold = Results[i];
+								}
+							}
+						}
+					}
+#endif
+				}
+				CommandsToRemove.Add(CommandIndex);
+			}
+			else if (Command.TargetAttribute == GetFieldPhysicsName(EFieldPhysicsType::Field_DisableThreshold))
+			{
+				SCOPE_CYCLE_COUNTER(STAT_ParamUpdateField_DisableThreshold);
+
+				if (ensureMsgf(Command.RootNode->Type() == FFieldNode<float>::StaticType(),
+					TEXT("Field based evaluation of the simulations 'Disable' parameter expects scale field inputs.")))
+				{
+#if TODO_REIMPLEMENT_PHYSICS_PROXY_REVERSE_MAPPING
+					const Chaos::TArrayCollectionArray<PhysicsProxyWrapper>& PhysicsProxyMapping = CurrentSolver->GetPhysicsProxyReverseMapping();
+					
+					FFieldSystemPhysicsProxy::ContiguousIndices(IndicesArray, CurrentSolver, ResolutionType, IndicesArray.Num() != Particles.Size());
+					if (IndicesArray.Num())
+					{
+						TArrayView<ContextIndex> IndexView(&(IndicesArray[0]), IndicesArray.Num());
+
+						FVector * tptr = &(Particles.X(0));
+						TArrayView<FVector> SamplesView(tptr, int32(Particles.Size()));
+
+						FFieldContext Context{
+							IndexView,
+							SamplesView,
+							Command.MetaData
+						};
+
+						TArray<float> Results;
+						Results.AddUninitialized(Particles.Size());
+						for (const ContextIndex Index : IndicesArray)
+						{
+							const SolverObjectWrapper& ParticleObjectWrapper = SolverObjectMapping[Index.Result];
+							TSerializablePtr<Chaos::TChaosPhysicsMaterial<float>> Material = CurrentSolver->GetPhysicsMaterial(Index.Result);
+							if (ensure(Material) && ParticleObjectWrapper.SolverObject)
+							{
+								const TUniquePtr<Chaos::TChaosPhysicsMaterial<float>>& InstanceMaterial = CurrentSolver->GetPerParticlePhysicsMaterial(Index.Result);
+								if (InstanceMaterial)
+								{
+									Results[Index.Result] = InstanceMaterial->DisabledLinearThreshold;
+								}
+								else
+								{
+									Results[Index.Result] = Material->DisabledLinearThreshold;
+								}
+							}
+							else
+							{
+								Results[Index.Result] = 0.0f;
+							}
+						}
+
+						TArrayView<float> ResultsView(&(Results[0]), Results.Num());
+						static_cast<const FFieldNode<float> *>(Command.RootNode.Get())->Evaluate(Context, ResultsView);
+
+						for (const ContextIndex Index : IndicesArray)
+						{
+							const int32 i = Index.Result;
+							const PhysicsProxyWrapper& ParticleObjectWrapper = PhysicsProxyMapping[i];
+							TSerializablePtr<Chaos::TChaosPhysicsMaterial<float>> Material = CurrentSolver->GetPhysicsMaterial(i);
+							if (!ensure(Material) || !ParticleObjectWrapper.PhysicsProxy)	//question: do we actually need to check for solver object?
+							{
+								continue;
+							}
+
+							//per instance override
+							if (!CurrentSolver->GetPerParticlePhysicsMaterial(Index.Result))
+							{
+								if (Results[i] != Material->DisabledLinearThreshold)
+								{
+									// value changed from shared material, make unique material.
+									TUniquePtr<Chaos::TChaosPhysicsMaterial<float>> NewMaterial = MakeUnique< Chaos::TChaosPhysicsMaterial<float>>(*Material);
+									CurrentSolver->SetPerParticlePhysicsMaterial(Index.Result, MoveTemp(NewMaterial));
+									const TUniquePtr<TChaosPhysicsMaterial<float>>& InstanceMaterial = CurrentSolver->GetPerParticlePhysicsMaterial(i);
+
+									InstanceMaterial->DisabledLinearThreshold = Results[i];
+									InstanceMaterial->DisabledAngularThreshold = Results[i];
+								}
+							}
+							else
+							{
+								const TUniquePtr<TChaosPhysicsMaterial<float>>& InstanceMaterial = CurrentSolver->GetPerParticlePhysicsMaterial(i);
+
+								if (InstanceMaterial->DisabledLinearThreshold != Results[i])
+								{
+									InstanceMaterial->DisabledLinearThreshold = Results[i];
+									InstanceMaterial->DisabledAngularThreshold = Results[i];
+								}
+							}
+						}
+					}
+#endif
+				}
+				CommandsToRemove.Add(CommandIndex);
+			}
+#if TODO_REIMPLEMENT_RIGID_CLUSTERING
+			else if (Command.TargetAttribute == GetFieldPhysicsName(EFieldPhysicsType::Field_InternalClusterStrain))
+			{
+				SCOPE_CYCLE_COUNTER(STAT_ParamUpdateField_InternalClusterStrain);
+				if (ensureMsgf(Command.RootNode->Type() == FFieldNode<float>::StaticType(),
+					TEXT("Field based evaluation of the simulations 'ExternalClusterStrain' parameter expects scalar field inputs.")))
+				{
+					FFieldSystemPhysicsProxy::ContiguousIndices(IndicesArray, CurrentSolver, ResolutionType, IndicesArray.Num() != Particles.Size());
+					if (IndicesArray.Num())
+					{
+						TArrayView<ContextIndex> IndexView(&(IndicesArray[0]), IndicesArray.Num());
+
+						FVector * tptr = &(Particles.X(0));
+						TArrayView<FVector> SamplesView(tptr, int32(Particles.Size()));
+
+						FFieldContext Context{
+							IndexView,
+							SamplesView,
+							Command.MetaData
+						};
+
+						float * vptr = &(Strains[0]);
+						TArrayView<float> ResultsView(vptr, Particles.Size());
+						static_cast<const FFieldNode<float> *>(Command.RootNode.Get())->Evaluate(Context, ResultsView);
+					}
+				}
+				CommandsToRemove.Add(CommandIndex);
+			}
+#endif
+			else if (Command.TargetAttribute == GetFieldPhysicsName(EFieldPhysicsType::Field_CollisionGroup))
+			{
+				if (ensureMsgf(Command.RootNode->Type() == FFieldNode<int32>::StaticType(),
+					TEXT("Field based evaluation of the simulations 'CollisionGroup' parameter expects int field inputs.")))
+				{
+					FFieldSystemPhysicsProxy::ContiguousIndices(IndicesArray, CurrentSolver, ResolutionType, IndicesArray.Num() != Particles.Size());
+					if (IndicesArray.Num())
+					{
+						TArrayView<ContextIndex> IndexView(&(IndicesArray[0]), IndicesArray.Num());
+
+						FVector * tptr = &(Particles.X(0));
+						TArrayView<FVector> SamplesView(tptr, int32(Particles.Size()));
+
+						FFieldContext Context{
+							IndexView,
+							SamplesView,
+							Command.MetaData
+						};
+
+						int32 * cptr = &(Particles.CollisionGroup(0));
+						TArrayView<int32> ResultsView(cptr, Particles.Size());
+						static_cast<const FFieldNode<int32> *>(Command.RootNode.Get())->Evaluate(Context, ResultsView);
+					}
+				}
+				CommandsToRemove.Add(CommandIndex);
+			}
+			else if (Command.TargetAttribute == GetFieldPhysicsName(EFieldPhysicsType::Field_PositionStatic))
+			{
+				SCOPE_CYCLE_COUNTER(STAT_ParamUpdateField_PositionStatic);
+
+#if TODO_REIMPLEMENT_FIELDS_TO_USE_PARTICLEHANDLES
+				if (ensureMsgf(Command.RootNode->Type() == FFieldNode<int32>::StaticType(),
+					TEXT("Field based evaluation of the simulations 'Position' parameter expects integer field inputs.")))
+				{
+					FFieldSystemPhysicsProxy::ContiguousIndices(IndicesArray, CurrentSolver, ResolutionType, IndicesArray.Num() != Particles.Size());
+					if (IndicesArray.Num())
+					{
+						TArrayView<ContextIndex> IndexView(&(IndicesArray[0]), IndicesArray.Num());
+
+						FVector * tptr = &(Particles.X(0));
+						TArrayView<FVector> SamplesView(tptr, int32(Particles.Size()));
+
+						FFieldContext Context{
+							IndexView,
+							SamplesView,
+							Command.MetaData
+						};
+
+						TArray<int32> Results;
+						Results.AddUninitialized(Particles.Size());
+						for (const ContextIndex& Index : IndicesArray)
+						{
+							Results[Index.Sample] = false;
+						}
+						TArrayView<int32> ResultsView(&(Results[0]), Results.Num());
+						static_cast<const FFieldNode<int32> *>(Command.RootNode.Get())->Evaluate(Context, ResultsView);
+
+						for (const ContextIndex& CIndex : IndicesArray)
+						{
+							const int32 i = CIndex.Result;
+							if (Results[i])
+							{
+								if (PositionTargetedParticles.Contains(i))
+								{
+									int32 Index = PositionTargetedParticles[i];
+									PositionTarget.Replace(Index, Particles.X(i));
+								}
+								else
+								{
+									int32 Index = PositionTarget.Add(i, Particles.X(i));
+									PositionTargetedParticles.Add(i, Index);
+								}
+							}
+						}
+					}
+				}
+#endif
+				CommandsToRemove.Add(CommandIndex);
+			}
+			else if (Command.TargetAttribute == GetFieldPhysicsName(EFieldPhysicsType::Field_PositionTarget))
+			{
+				SCOPE_CYCLE_COUNTER(STAT_ParamUpdateField_PositionTarget);
+
+#if TODO_REIMPLEMENT_FIELDS_TO_USE_PARTICLEHANDLES
+
+				if (ensureMsgf(Command.RootNode->Type() == FFieldNode<FVector>::StaticType(),
+					TEXT("Field based evaluation of the simulations 'PositionTarget' parameter expects vector field inputs.")))
+				{
+					FFieldSystemPhysicsProxy::ContiguousIndices(IndicesArray, CurrentSolver, ResolutionType, IndicesArray.Num() != Particles.Size());
+					if (IndicesArray.Num())
+					{
+						TArrayView<ContextIndex> IndexView(&(IndicesArray[0]), IndicesArray.Num());
+
+						FVector * tptr = &(Particles.X(0));
+						TArrayView<FVector> SamplesView(tptr, int32(Particles.Size()));
+
+						FFieldContext Context{
+							IndexView,
+							SamplesView,
+							Command.MetaData
+						};
+
+						TArray<FVector> Results;
+						Results.AddUninitialized(Particles.Size());
+						for (const ContextIndex& Index : IndicesArray)
+						{
+							Results[Index.Sample] = FVector(FLT_MAX);
+						}
+						TArrayView<FVector> ResultsView(&(Results[0]), Results.Num());
+						static_cast<const FFieldNode<FVector> *>(Command.RootNode.Get())->Evaluate(Context, ResultsView);
+
+						for (const ContextIndex& CIndex : IndicesArray)
+						{
+							const int32 i = CIndex.Result;
+							if (Results[i] != FVector(FLT_MAX))
+							{
+								if (PositionTargetedParticles.Contains(i))
+								{
+									int32 Index = PositionTargetedParticles[i];
+									PositionTarget.Replace(Index, Results[i]);
+								}
+								else
+								{
+									int32 Index = PositionTarget.Add(i, Results[i]);
+									PositionTargetedParticles.Add(i, Index);
+								}
+							}
+						}
+					}
+				}
+#endif
+
+				CommandsToRemove.Add(CommandIndex);
+			}
+			else if (Command.TargetAttribute == GetFieldPhysicsName(EFieldPhysicsType::Field_PositionAnimated))
+			{
+				SCOPE_CYCLE_COUNTER(STAT_ParamUpdateField_PositionAnimated);
+
+				if (ensureMsgf(Command.RootNode->Type() == FFieldNode<int32>::StaticType(),
+					TEXT("Field based evaluation of the simulations 'Position' parameter expects integer field inputs.")))
+				{
+					FFieldSystemPhysicsProxy::ContiguousIndices(IndicesArray, CurrentSolver, ResolutionType, IndicesArray.Num() != Particles.Size());
+					if (IndicesArray.Num())
+					{
+						TArrayView<ContextIndex> IndexView(&(IndicesArray[0]), IndicesArray.Num());
+
+						FVector * tptr = &(Particles.X(0));
+						TArrayView<FVector> SamplesView(tptr, int32(Particles.Size()));
+
+						FFieldContext Context{
+							IndexView,
+							SamplesView,
+							Command.MetaData
+						};
+
+						TArray<int32> Results;
+						Results.AddUninitialized(Particles.Size());
+						for (const ContextIndex& Index : IndicesArray)
+						{
+							Results[Index.Sample] = false;
+						}
+						TArrayView<int32> ResultsView(&(Results[0]), Results.Num());
+						static_cast<const FFieldNode<int32> *>(Command.RootNode.Get())->Evaluate(Context, ResultsView);
+
+#if TODO_REIMPLEMENT_KINEMATIC_PROXY
+						for (int32 i = 0; i < AnimatedPosition.Num(); ++i)
+						{
+							for (int32 j = 0; j < AnimatedPosition[i].Ids.Num(); ++j)
+							{
+								int32 Index = AnimatedPosition[i].Ids[j];
+								if (Results[Index])
+								{
+									if (PositionTargetedParticles.Contains(Index))
+									{
+										int32 PosIndex = PositionTargetedParticles[i];
+										PositionTarget.Replace(PosIndex, AnimatedPosition[i].Position[j]);
+									}
+									else
+									{
+										int32 PosIndex = PositionTarget.Add(i, AnimatedPosition[i].Position[j]);
+										PositionTargetedParticles.Add(i, PosIndex);
+									}
+								}
+							}
+						}
+#endif
+					}
+				}
+				CommandsToRemove.Add(CommandIndex);
+			}
+			else if (Command.TargetAttribute == GetFieldPhysicsName(EFieldPhysicsType::Field_DynamicConstraint))
+			{
+				SCOPE_CYCLE_COUNTER(STAT_ParamUpdateField_DynamicConstraint);
+
+				if (ensureMsgf(Command.RootNode->Type() == FFieldNode<float>::StaticType(),
+					TEXT("Field based evaluation of the simulations 'DynamicConstraint' parameter expects scalar field inputs.")))
+				{
+#if TODO_REIMPLEMENT_DYNAMIC_CONSTRAINT_ACCESSORS
+					Chaos::TPBDRigidDynamicSpringConstraints<float, 3>& DynamicConstraints = FPhysicsSolver::FAccessor(CurrentSolver).DynamicConstraints();
+					TSet<int32>& DynamicConstraintParticles = FPhysicsSolver::FAccessor(CurrentSolver).DynamicConstraintParticles();
+
+					FFieldSystemPhysicsProxy::ContiguousIndices(IndicesArray, CurrentSolver, ResolutionType, IndicesArray.Num() != Particles.Size());
+					if (IndicesArray.Num())
+					{
+						TArrayView<ContextIndex> IndexView(&(IndicesArray[0]), IndicesArray.Num());
+
+						FVector * tptr = &(Particles.X(0));
+						TArrayView<FVector> SamplesView(tptr, int32(Particles.Size()));
+
+						FFieldContext Context{
+							IndexView,
+							SamplesView,
+							Command.MetaData
+						};
+
+						TArray<float> Results;
+						Results.AddUninitialized(Particles.Size());
+						for (const ContextIndex& CIndex : IndicesArray)
+						{
+							Results[CIndex.Sample] = FLT_MAX;
+						}
+						TArrayView<float> ResultsView(&(Results[0]), Results.Num());
+						static_cast<const FFieldNode<float> *>(Command.RootNode.Get())->Evaluate(Context, ResultsView);
+
+						for (const ContextIndex& CIndex : IndicesArray)
+						{
+							const int32 i = CIndex.Result;
+							if (Results[i] != FLT_MAX)
+							{
+								if (!DynamicConstraintParticles.Contains(i))
+								{
+									DynamicConstraints.SetDistance(Results[i]);
+									for (const int32 Index : DynamicConstraintParticles)
+									{
+										DynamicConstraints.Add(Index, i);
+									}
+									DynamicConstraintParticles.Add(i);
+								}
+							}
+						}
+					}
+#endif
+				}
+				CommandsToRemove.Add(CommandIndex);
+			}
+		}
+		for (int32 Index = CommandsToRemove.Num() - 1; Index >= 0; --Index)
+		{
+			CommandListPtr->RemoveAt(CommandsToRemove[Index]);
+		}
+	}
+}
+
+void FFieldSystemPhysicsProxy::FieldForcesUpdateCallback(Chaos::FPhysicsSolver* InSolver, FParticlesType& Particles, Chaos::TArrayCollectionArray<FVector> & Force,
+	Chaos::TArrayCollectionArray<FVector> & Torque, const float Time)
+{
+	if (Commands.Num() && InSolver)
+	{
+		Chaos::FPhysicsSolver* CurrentSolver = InSolver;
+		TArray<ContextIndex> IndicesArray;
+
+		TArray<FFieldSystemCommand>* CommandListPtr = GetSolverCommandList(InSolver);
+
+		if(!CommandListPtr)
+		{
+			return;
+		}
+
+		TArray<int32> CommandsToRemove;
+		const int32 NumCommands = CommandListPtr->Num();
+		for(int32 CommandIndex = 0; CommandIndex < NumCommands; CommandIndex++)
+		{
+			const FFieldSystemCommand & Command = (*CommandListPtr)[CommandIndex];
+			EFieldResolutionType ResolutionType = EFieldResolutionType::Field_Resolution_Minimal;
+			if (Command.MetaData.Contains(FFieldSystemMetaData::EMetaType::ECommandData_ProcessingResolution))
+			{
+				check(Command.MetaData[FFieldSystemMetaData::EMetaType::ECommandData_ProcessingResolution] != nullptr);
+				ResolutionType = static_cast<FFieldSystemMetaDataProcessingResolution*>(Command.MetaData[FFieldSystemMetaData::EMetaType::ECommandData_ProcessingResolution].Get())->ProcessingResolution;
+			}
+
+			if (Command.TargetAttribute == GetFieldPhysicsName(EFieldPhysicsType::Field_LinearForce))
+			{
+				if (ensureMsgf(Command.RootNode->Type() == FFieldNode<FVector>::StaticType(),
+					TEXT("Field based evaluation of the simulations 'Force' parameter expects FVector field inputs.")))
+				{
+					FFieldSystemPhysicsProxy::ContiguousIndices(IndicesArray, CurrentSolver, ResolutionType, IndicesArray.Num() != Particles.Size());
+					if (IndicesArray.Num())
+					{
+						TArrayView<ContextIndex> IndexView(&(IndicesArray[0]), IndicesArray.Num());
+
+						FVector * tptr = &(Particles.X(0));
+						TArrayView<FVector> SamplesView(tptr, int32(Particles.Size()));
+
+						FFieldContext Context{
+							IndexView,
+							SamplesView,
+							Command.MetaData
+						};
+						TArrayView<FVector> ForceView(&(Force[0]), Force.Num());
+						static_cast<const FFieldNode<FVector> *>(Command.RootNode.Get())->Evaluate(Context, ForceView);
+
+						// @todo(ccaulfield): encapsulation: add WakeParticles (and therefore islands) functionality to Evolution
+						TSet<int32> IslandsToActivate;
+						for (const ContextIndex& CIndex : IndicesArray)
+						{
+							const int32 i = CIndex.Result;
+							if (ForceView[i] != FVector(0) && Particles.ObjectState(i) == Chaos::EObjectStateType::Sleeping && !Particles.Disabled(i) && IslandsToActivate.Find(Particles.Island(i)) == nullptr)
+							{
+								IslandsToActivate.Add(Particles.Island(i));
+							}
+						}
+#if TODO_REIMPLEMENT_WAKE_ISLANDS
+						InSolver->WakeIslands(IslandsToActivate);
+#endif
+					}
+				}
+				CommandsToRemove.Add(CommandIndex);
+			}
+			else if (Command.TargetAttribute == GetFieldPhysicsName(EFieldPhysicsType::Field_AngularTorque))
+			{
+				if (ensureMsgf(Command.RootNode->Type() == FFieldNode<FVector>::StaticType(),
+					TEXT("Field based evaluation of the simulations 'Torque' parameter expects FVector field inputs.")))
+				{
+					FFieldSystemPhysicsProxy::ContiguousIndices(IndicesArray, CurrentSolver, ResolutionType, IndicesArray.Num() != Particles.Size());
+					if (IndicesArray.Num())
+					{
+						TArrayView<ContextIndex> IndexView(&(IndicesArray[0]), IndicesArray.Num());
+
+						FVector * tptr = &(Particles.X(0));
+						TArrayView<FVector> SamplesView(tptr, int32(Particles.Size()));
+
+						FFieldContext Context{
+							IndexView,
+							SamplesView,
+							Command.MetaData
+						};
+						TArrayView<FVector> TorqueView(&(Torque[0]), Torque.Num());
+						static_cast<const FFieldNode<FVector> *>(Command.RootNode.Get())->Evaluate(Context, TorqueView);
+
+						// @todo(ccaulfield): encapsulation: add WakeParticles (and therefore islands) functionality to Evolution
+						TSet<int32> IslandsToActivate;
+						for (const ContextIndex& CIndex : IndicesArray)
+						{
+							const int32 i = CIndex.Result;
+							if (TorqueView[i] != FVector(0) && Particles.ObjectState(i) == Chaos::EObjectStateType::Sleeping && !Particles.Disabled(i) && IslandsToActivate.Find(Particles.Island(i)) == nullptr)
+							{
+								IslandsToActivate.Add(Particles.Island(i));
+							}
+						}
+#if TODO_REIMPLEMENT_WAKE_ISLANDS
+						InSolver->WakeIslands(IslandsToActivate);
+#endif
+					}
+				}
+				CommandsToRemove.Add(CommandIndex);
+			}
+		}
+		for (int32 Index = CommandsToRemove.Num() - 1; Index >= 0; --Index)
+		{
+			CommandListPtr->RemoveAt(CommandsToRemove[Index]);
+		}
+	}
+}
+
+void FFieldSystemPhysicsProxy::EndFrameCallback(const float InDt)
+{
+}
+
+void FFieldSystemPhysicsProxy::BufferCommand(Chaos::FPhysicsSolver* InSolver, const FFieldSystemCommand& InCommand)
+{
+	TArray<FFieldSystemCommand>** ExistingList = nullptr;
+	{
+		FScopeLock Lock(&CommandLock);
+		ExistingList = Commands.Find(InSolver);
+
+		if(!ExistingList)
+		{
+			Commands.Add(InSolver);
+			ExistingList = Commands.Find(InSolver);
+			check(ExistingList);
+
+			(*ExistingList) = new TArray<FFieldSystemCommand>();
+		}
+	}
+
+	(*ExistingList)->Add(InCommand);
+}
+
+void FFieldSystemPhysicsProxy::ContiguousIndices(TArray<ContextIndex>& Array, const Chaos::FPhysicsSolver* RigidSolver, EFieldResolutionType ResolutionType, bool bForce = true)
+{
+#if TODO_REIMPLEMENT_RIGID_CLUSTERING
+	if (bForce)
+	{
+		const Chaos::FPhysicsSolver::FParticlesType & Particles = RigidSolver->GetRigidParticles();
+		if (ResolutionType == EFieldResolutionType::Field_Resolution_Minimal)
+		{
+			Array.SetNum(0, false);
+
+			int32 IndexCount = 0;
+			int32 FloorIndex = RigidSolver->GetFloorIndex();
+			const Chaos::FPhysicsSolver::FClusteringType& Clustering = RigidSolver->GetRigidClustering();
+			const Chaos::TArrayCollectionArray<Chaos::ClusterId>& ClusterIdArray = Clustering.GetClusterIdsArray();
+			const Chaos::FPhysicsSolver::FClusteringType::FClusterMap &  ClusterMap = Clustering.GetChildrenMap();
+
+			for (int32 ActiveParticleIndex : RigidSolver->NonDisabledIndices())
+			{
+				if (ClusterIdArray[ActiveParticleIndex].NumChildren)
+				{
+					for (uint32 ClusterChild : *ClusterMap[ActiveParticleIndex])
+					{
+						Array.Add(ContextIndex(ClusterChild, ClusterChild));
+					}
+				}
+
+				if( ActiveParticleIndex != FloorIndex )
+				{
+					Array.Add(ContextIndex(ActiveParticleIndex, ActiveParticleIndex));
+				}
+			}
+		}
+		if (ResolutionType == EFieldResolutionType::Field_Resolution_DisabledParents)
+		{
+			Array.SetNum(0, false);
+
+			int32 IndexCount = 0;
+			int32 FloorIndex = RigidSolver->GetFloorIndex();
+			const Chaos::FPhysicsSolver::FClusteringType& Clustering = RigidSolver->GetRigidClustering();
+			const Chaos::TArrayCollectionArray<Chaos::ClusterId>& ClusterIdArray = Clustering.GetClusterIdsArray();
+
+			for (int32 TopLevelParent : Clustering.GetTopLevelClusterParents())
+			{
+				if( TopLevelParent != FloorIndex )
+				{
+					Array.Add(ContextIndex(TopLevelParent, TopLevelParent));
+				}
+			}
+		}
+		else if (ResolutionType == EFieldResolutionType::Field_Resolution_Maximum)
+		{
+			Array.SetNum(Particles.Size());
+			for (int32 i = 0; i < Array.Num(); ++i)
+			{
+				Array[i].Result = i;
+				Array[i].Sample = i;
+			}
+		}
+	}
+#endif
+}
+
+TArray<FFieldSystemCommand>* FFieldSystemPhysicsProxy::GetSolverCommandList(const Chaos::FPhysicsSolver* InSolver)
+{
+	FScopeLock Lock(&CommandLock);
+	TArray<FFieldSystemCommand>** ExistingList = Commands.Find(InSolver);
+	return ExistingList ? *ExistingList : nullptr;
+}
+
+void FFieldSystemPhysicsProxy::OnRemoveFromScene()
+{
+}
+
+#endif 
