Copying //UE4/Dev-Physics to Dev-Main (//UE4/Dev-Main) [at] 7703071

diff --git a/Engine/Source/Runtime/Experimental/ChaosSolvers/Private/PhysicsProxy/SkeletalMeshPhysicsProxy.cpp b/Engine/Source/Runtime/Experimental/ChaosSolvers/Private/PhysicsProxy/SkeletalMeshPhysicsProxy.cpp
new file mode 100644
--- /dev/null
+++ b/Engine/Source/Runtime/Experimental/ChaosSolvers/Private/PhysicsProxy/SkeletalMeshPhysicsProxy.cpp
@@ -0,0 +1,505 @@
+// Copyright 1998-2019 Epic Games, Inc. All Rights Reserved.
+
+#include "PhysicsProxy/SkeletalMeshPhysicsProxy.h"
+#include "GeometryCollection/GeometryCollectionCollisionStructureManager.h"
+#include "ChaosStats.h"
+#include "PhysicsSolver.h"
+
+#if INCLUDE_CHAOS
+
+FSkeletalMeshPhysicsProxy::FSkeletalMeshPhysicsProxy(UObject* InOwner, const FInitFunc& InInitFunc)
+	: Base(InOwner)
+	, JointConstraintsRule(JointConstraints)
+	, NextInputProducerBuffer(nullptr)
+	, CurrentOutputConsumerBuffer(nullptr)
+	, bInitializedState(false)
+	, InitFunc(InInitFunc)
+{
+	check(IsInGameThread())
+}
+
+FSkeletalMeshPhysicsProxy::~FSkeletalMeshPhysicsProxy()
+{
+}
+
+void FSkeletalMeshPhysicsProxy::Initialize()
+{
+	check(IsInGameThread());
+
+	// InitFunc builds a bone hierarchy from the UPhysicsAsset on a USkeletalMeshComponent.
+	// Transforms will be set from the FReferenceSkeleton reference bone poses.
+	InitFunc(Parameters);
+
+	// Set InitializedState to false so that rigid bodies will now be created.
+	Reset();
+}
+
+void FSkeletalMeshPhysicsProxy::Reset()
+{
+	bInitializedState = false;
+}
+
+bool FSkeletalMeshPhysicsProxy::IsSimulating() const
+{
+	return Parameters.bSimulating;
+}
+
+void FSkeletalMeshPhysicsProxy::UpdateKinematicBodiesCallback(const FParticlesType& Particles, const float Dt, const float Time, FKinematicProxy& Proxy)
+{
+	// @todo(ccaulfield): Handle changes in body state
+
+	if (bInitializedState && Parameters.bSimulating)
+	{
+		// The triple buffer will return null if there's no new data for the consumer thread.
+		const FSkeletalMeshPhysicsProxyInputs* InputBuffer = InputBuffers.ExchangeConsumerBuffer();
+		if (InputBuffer == nullptr)
+		{
+			return;
+		}
+
+#if TODO_REIMPLEMENT_KINEMATIC_PROXY
+
+		// @todo(ccaulfield): Optimize. Keep a list of kinematic particles?
+		const int32 MaxRigidBodies = RigidBodyIds.Num();
+		Proxy.Ids.Reset(MaxRigidBodies);
+		Proxy.Position.Reset(MaxRigidBodies);
+		Proxy.NextPosition.Reset(MaxRigidBodies);
+		Proxy.Rotation.Reset(MaxRigidBodies);
+		Proxy.NextRotation.Reset(MaxRigidBodies);
+
+		int32 RigidBodyIndex = 0;
+		TArray<TUniquePtr<FAnalyticImplicitGroup>>& Groups = Parameters.BoneHierarchy.GetAnalyticShapeGroups();
+		for (TUniquePtr<FAnalyticImplicitGroup>& Group : Groups)
+		{
+			if (!Group->NumStructures())
+			{
+				continue;
+			}
+
+			const int32 RigidBodyId = Group->GetRigidBodyId();
+			const int32 BoneIndex = Group->GetBoneIndex();
+			const int32 TransformIndex = Parameters.BoneHierarchy.GetTransformIndex(BoneIndex);
+			bool bIsKinematic = (Particles.ObjectState(RigidBodyId) == Chaos::EObjectStateType::Kinematic);
+			if (bIsKinematic)
+			{
+				const FTransform& Transform = InputBuffer->Transforms[TransformIndex];
+				const FVector& LinearVelocity = InputBuffer->LinearVelocities[TransformIndex];
+				const FVector& AngularVelocity = InputBuffer->AngularVelocities[TransformIndex];
+
+				Proxy.Ids.Add(RigidBodyId);
+				Proxy.Position.Add(Chaos::TVector<float, 3>(Transform.GetTranslation()));
+				Proxy.NextPosition.Add(Proxy.Position[RigidBodyIndex] + Chaos::TVector<float, 3>(LinearVelocity) * Dt);
+				Proxy.Rotation.Add(Chaos::TRotation<float, 3>(Transform.GetRotation().GetNormalized()));
+				Proxy.NextRotation.Add(Chaos::TRotation<float, 3>::IntegrateRotationWithAngularVelocity(Proxy.Rotation[RigidBodyIndex], AngularVelocity, Dt));
+				++RigidBodyIndex;
+			}
+		}
+
+#endif
+
+	}
+}
+
+void FSkeletalMeshPhysicsProxy::StartFrameCallback(const float InDt, const float InTime)
+{
+}
+
+void FSkeletalMeshPhysicsProxy::EndFrameCallback(const float InDt)
+{
+}
+
+void FSkeletalMeshPhysicsProxy::CreateRigidBodyCallback(FParticlesType& Particles)
+{
+	if (!bInitializedState && Parameters.bSimulating)
+	{
+#if TODO_REIMPLEMENT_RIGID_CLUSTERING
+		// @todo(ccaulfield): should not have to do this cluster thing here...
+		Chaos::TArrayCollectionArray<int32>& ClusterGroupIndex = GetSolver()->GetRigidClustering().GetClusterGroupIndexArray();
+#endif
+#if 0
+		Chaos::TArrayCollectionArray<float>& StrainArray = GetSolver()->GetRigidClustering().GetStrainArray();
+#endif
+		Parameters.BoneHierarchy.PrepareAnimWorldSpaceTransforms();
+		RigidBodyIds.Reset();
+		RigidBodyIds.Reserve(Parameters.BoneHierarchy.GetBoneIndices().Num());
+		for (const int32 BoneIndex : Parameters.BoneHierarchy.GetBoneIndices())
+		{
+			FAnalyticImplicitGroup* Group = Parameters.BoneHierarchy.GetAnalyticShapeGroup(BoneIndex);
+			if (!Group || !Group->NumStructures())
+			{
+				continue;
+			}
+
+			const FTransform* WorldTransform = Parameters.BoneHierarchy.GetAnimWorldSpaceTransformsForBone(BoneIndex);
+			check(WorldTransform);
+
+			float TotalMass = 0.f;
+			const float DensityKgCm3 = Parameters.Density / 1000.0f;
+			Chaos::TMassProperties<float, 3> MassProperties = Group->BuildMassProperties(DensityKgCm3, TotalMass);
+
+			const bool DoCollisionGeom = Parameters.CollisionType == ECollisionTypeEnum::Chaos_Surface_Volumetric;
+			TArray<Chaos::TVector<float, 3>>* SamplePoints = nullptr;
+			Chaos::TBox<float, 3> SamplePointsBBox = Chaos::TBox<float, 3>::EmptyBox();
+			if (DoCollisionGeom)
+			{
+				SamplePoints =
+					Group->BuildSamplePoints(
+						Parameters.ParticlesPerUnitArea,
+						Parameters.MinNumParticles,
+						Parameters.MaxNumParticles);
+				for (const Chaos::TVector<float, 3>& Point : *SamplePoints)
+					SamplePointsBBox.GrowToInclude(Point);
+			}
+			//Chaos::TTriangleMesh<T> TriMesh = Chaos::TTriangleMesh<T>::GetConvexHullFromParticles(
+			//	TArrayView<const Chaos::TVector<T, 3>>(SamplePoints));
+
+			TUniquePtr<Chaos::TImplicitObject<float, 3>> ImplicitObject(Group->BuildSimImplicitObject());
+			check(ImplicitObject);
+
+			if (SamplePoints)
+			{
+				auto PointVolumeRegistrationCheck = 
+					[&](const TArray<Chaos::TVector<float, 3>> &InSamplePoints, const Chaos::TImplicitObject<float, 3> &InImplicitObject, const float InTolerance) -> bool
+					{
+						for (int32 i = 0; i < InSamplePoints.Num(); i++)
+						{
+							const auto& Pt = InSamplePoints[i];
+							const auto Phi = InImplicitObject.SignedDistance(Pt);
+							if (FMath::Abs(Phi) > InTolerance)
+							{
+								return false;
+							}
+						}
+						return true; 
+					};
+				// Precision gets worse in opt builds, so only do the phi test in debug.
+				checkSlow(PointVolumeRegistrationCheck(*SamplePoints, *ImplicitObject, 0.01));
+				auto PointBBoxCheck =
+					[&](const auto &InSamplePoints, const Chaos::TBox<float,3> &BBox) -> bool
+					{
+						for (int32 i = 0; i < InSamplePoints.Num(); i++)
+						{
+							const auto& Pt = InSamplePoints[i];
+							if (!BBox.Contains(Pt))
+							{
+								return false;
+							}
+						}
+						return true;
+					};
+				// BBox tests should still work though.
+				check(PointBBoxCheck(*SamplePoints, ImplicitObject->BoundingBox()));
+			}
+
+			const Chaos::TBox<float, 3>& BBox = ImplicitObject->BoundingBox();
+
+			const FVector Scale = WorldTransform->GetScale3D();
+			const FVector &CenterOfMass = MassProperties.CenterOfMass;
+
+			FBox Bounds(BBox.Min(), BBox.Max());
+			Bounds = Bounds.InverseTransformBy(FTransform(CenterOfMass));
+			Bounds.Min *= Scale;
+			Bounds.Max *= Scale;
+			checkSlow(Bounds.GetVolume() > FLT_EPSILON);
+
+			const int32 RigidBodyId = Particles.Size();
+			Particles.AddParticles(1);
+			RigidBodyIds.Add(RigidBodyId);
+			Group->SetRigidBodyId(RigidBodyId);
+#if 0
+			ClusterGroupIndex[RigidBodyId] = Parameters.ClusterGroupIndex;
+			StrainArray[RigidBodyId] = Parameters.DamageThreshold;
+#else
+#if TODO_REIMPLEMENT_RIGID_CLUSTERING
+			ClusterGroupIndex[RigidBodyId] = 0;
+#endif
+#endif
+
+#if TODO_REIMPLEMENT_SET_PHYSICS_MATERIAL
+			GetSolver()->SetPhysicsMaterial(RigidBodyId, Parameters.PhysicalMaterial);
+#endif
+			Particles.DisabledRef(RigidBodyId) = false;
+
+			//Particles.X(RigidBodyId) = WorldTransform->TransformPosition(CenterOfMass);
+			Particles.X(RigidBodyId) = WorldTransform->GetTranslation();
+			check(!FMath::IsNaN(Particles.X(RigidBodyId)[0]) && !FMath::IsNaN(Particles.X(RigidBodyId)[1]) && !FMath::IsNaN(Particles.X(RigidBodyId)[2]));
+			Particles.V(RigidBodyId) = Chaos::TVector<float, 3>(0.f);
+			//Particles.R(RigidBodyId) = (WorldTransform->TransformRotation(MassProperties.RotationOfMass)).GetNormalized();
+			Particles.R(RigidBodyId) = WorldTransform->GetRotation().GetNormalized();
+			Particles.W(RigidBodyId) = Chaos::TVector<float, 3>(0.f);
+			Particles.P(RigidBodyId) = Particles.X(RigidBodyId);
+			Particles.Q(RigidBodyId) = Particles.R(RigidBodyId);
+
+			Particles.M(RigidBodyId) = FMath::Clamp(TotalMass, Parameters.MinMass, Parameters.MaxMass);
+			Particles.InvM(RigidBodyId) = Particles.M(RigidBodyId) < KINDA_SMALL_NUMBER ?  BIG_NUMBER/2 : 1.f / Particles.M(RigidBodyId);
+			Particles.I(RigidBodyId) = MassProperties.InertiaTensor;
+			Particles.InvI(RigidBodyId) = MassProperties.InertiaTensor.Inverse();
+
+			Particles.CollisionGroup(RigidBodyId) = Parameters.CollisionGroup;
+			
+			// Set implicit surface
+			const Chaos::ImplicitObjectType ShapeType = ImplicitObject->GetType();
+			if (ShapeType == Chaos::ImplicitObjectType::LevelSet)
+			{
+				FCollisionStructureManager::UpdateImplicitFlags(
+					ImplicitObject.Get(), ECollisionTypeEnum::Chaos_Surface_Volumetric);
+				Particles.SetDynamicGeometry(RigidBodyId, MoveTemp(ImplicitObject));
+			}
+			else
+			{
+				Particles.SetDynamicGeometry(RigidBodyId, MoveTemp(ImplicitObject));
+			}
+
+			if (Parameters.CollisionType == ECollisionTypeEnum::Chaos_Surface_Volumetric)
+			{
+				if (SamplePoints)
+				{
+					Particles.CollisionParticlesInitIfNeeded(RigidBodyId);
+
+					Particles.CollisionParticles(RigidBodyId)->Resize(0);
+					Particles.CollisionParticles(RigidBodyId)->AddParticles(SamplePoints->Num());
+					for (int Vi = 0; Vi < SamplePoints->Num(); Vi++)
+					{
+						Particles.CollisionParticles(RigidBodyId)->X(Vi) = (*SamplePoints)[Vi];
+					}
+					if (Particles.CollisionParticles(RigidBodyId)->Size())
+					{
+						Particles.CollisionParticles(RigidBodyId)->UpdateAccelerationStructures();
+					}
+				}
+			}
+
+			// Set the simulation state of the body.
+			// If set to "user defined" we inherit the owner's state
+			// @todo(ccaulfield): formalize the state inheritence - currently hijacking UserDefined...
+			Chaos::EObjectStateType RigidBodyState = Chaos::EObjectStateType::Kinematic;
+			EObjectStateTypeEnum SourceState = Group->GetRigidBodyState();
+			if (SourceState == EObjectStateTypeEnum::Chaos_Object_UserDefined)
+			{
+				SourceState = Parameters.ObjectType;
+			}
+			switch (SourceState)
+			{
+			case EObjectStateTypeEnum::Chaos_Object_Sleeping:
+				RigidBodyState = Chaos::EObjectStateType::Sleeping;
+				break;
+			case EObjectStateTypeEnum::Chaos_Object_Kinematic:
+				RigidBodyState = Chaos::EObjectStateType::Kinematic;
+				break;
+			case EObjectStateTypeEnum::Chaos_Object_Static:
+				RigidBodyState = Chaos::EObjectStateType::Static;
+				break;
+			default:
+				RigidBodyState = Chaos::EObjectStateType::Dynamic;
+				break;
+			}
+			Particles.SetObjectState(RigidBodyId, RigidBodyState);
+
+			// Create joints
+			if (RigidBodyState == Chaos::EObjectStateType::Dynamic)
+			{
+				const int32 ParentBoneIndex = Group->GetParentBoneIndex();
+				if (ParentBoneIndex != INDEX_NONE)
+				{
+					const FAnalyticImplicitGroup* ParentGroup = Parameters.BoneHierarchy.GetAnalyticShapeGroup(ParentBoneIndex);
+					if (ParentGroup && (ParentGroup->GetRigidBodyId() != INDEX_NONE))
+					{
+						const int32 ParentRigidBodyId = ParentGroup->GetRigidBodyId();
+						const FTransform& ParentWorldTransform = *Parameters.BoneHierarchy.GetAnimWorldSpaceTransformsForBone(ParentBoneIndex);
+
+						const Chaos::TVector<int32, 2> ConstrainedParticleIndices = { RigidBodyId, ParentRigidBodyId };
+						const Chaos::TVector<float, 3> JointPosition = ParentWorldTransform.GetTranslation();
+#if TODO_REWRITE_ALL_CONSTRAINT_ADDS_TO_USE_HANDLES
+						JointConstraints.AddConstraint(Particles, ConstrainedParticleIndices, JointPosition);
+#endif
+					}
+				}
+			}
+
+
+			UE_LOG(USkeletalMeshSimulationComponentLogging, Log,
+				TEXT("FSkeletalMeshPhysicsProxy::CreateRigidBodyCallback() - this: %p - "
+					"Created rigid body %d for bone %d.\n"
+					"    center of mass: (%g, %g, %g), Bone pos: (%g, %g, %g), RB position: (%g, %g, %g)\n"
+					"    implicit shape domain: [(%g %g %g), (%g %g %g)] - volume: %g.\n"
+					"    %d local space collision points, domain: [(%g %g %g), (%g %g %g)] - volume: %g."),
+
+				this, 
+				RigidBodyId, Group->GetBoneIndex(),
+
+				CenterOfMass[0], CenterOfMass[1], CenterOfMass[2], 
+				WorldTransform->GetTranslation()[0], WorldTransform->GetTranslation()[1], WorldTransform->GetTranslation()[2], 
+				Particles.X(RigidBodyId)[0], Particles.X(RigidBodyId)[2], Particles.X(RigidBodyId)[2], 
+
+				SamplePointsBBox.Min()[0], SamplePointsBBox.Min()[1], SamplePointsBBox.Min()[2],
+				SamplePointsBBox.Max()[0], SamplePointsBBox.Max()[1], SamplePointsBBox.Max()[2],
+				SamplePointsBBox.Extents().Product(),
+
+				SamplePoints?SamplePoints->Num():0,
+				BBox.Min()[0], BBox.Min()[1], BBox.Min()[2],
+				BBox.Max()[0], BBox.Max()[1], BBox.Max()[2],
+				BBox.Extents().Product());
+		}
+
+		// Add joints to the scene
+		if (JointConstraints.NumConstraints() > 0)
+		{
+			GetSolver()->GetEvolution()->AddConstraintRule(&JointConstraintsRule);
+		}
+
+		bInitializedState = true;
+	}
+}
+
+void FSkeletalMeshPhysicsProxy::ParameterUpdateCallback(FParticlesType& InParticles, const float InTime)
+{}
+
+void FSkeletalMeshPhysicsProxy::DisableCollisionsCallback(TSet<TTuple<int32, int32>>& InPairs)
+{}
+
+void FSkeletalMeshPhysicsProxy::AddForceCallback(FParticlesType& InParticles, const float InDt, const int32 InIndex)
+{}
+
+void FSkeletalMeshPhysicsProxy::BindParticleCallbackMapping(Chaos::TArrayCollectionArray<PhysicsProxyWrapper>& PhysicsProxyReverseMap, Chaos::TArrayCollectionArray<int32>& ParticleIDReverseMap)
+{
+	if (bInitializedState)
+	{
+		for (int32 Index = 0; Index < RigidBodyIds.Num(); Index++)
+		{
+			if (RigidBodyIds[Index] != INDEX_NONE)
+			{
+				PhysicsProxyReverseMap[RigidBodyIds[Index]] = { this, EPhysicsProxyType::SkeletalMeshType };
+				ParticleIDReverseMap[RigidBodyIds[Index]] = Index;
+			}
+		}
+	}
+}
+
+void FSkeletalMeshPhysicsProxy::SyncBeforeDestroy() 
+{}
+
+void FSkeletalMeshPhysicsProxy::OnRemoveFromScene()
+{
+	// Disable the particle we added
+	Chaos::FPhysicsSolver* CurrSolver = GetSolver();
+	if(CurrSolver && RigidBodyIds.Num())
+	{
+#if TODO_REIMPLEMENT_GET_RIGID_PARTICLES
+		Chaos::FPhysicsSolver::FParticlesType &Particles = CurrSolver->GetRigidParticles();
+		// #BG TODO Special case here because right now we reset/realloc the evolution per geom component
+		// in endplay which clears this out. That needs to not happen and be based on world shutdown
+		if (Particles.Size() == 0)
+		{
+			return;
+		}
+
+		for (const int32 Id : RigidBodyIds)
+		{
+			if(static_cast<uint32>(Id) < Particles.Size())
+			{
+				Particles.DisabledRef(Id) = true;
+				CurrSolver->GetRigidClustering().GetTopLevelClusterParents().Remove(Id);
+			}
+		}
+#endif
+	}
+}
+
+void FSkeletalMeshPhysicsProxy::BufferPhysicsResults()
+{
+	using namespace Chaos;
+	// PHYSICS THREAD
+
+	// @todo(ccaulfield)
+	if (GetSolver() != nullptr)
+	{
+		// Copy the world-space state of the simulated bones into the output buffer
+		int NumBones = Parameters.BoneHierarchy.GetBoneIndices().Num();
+		FSkeletalMeshPhysicsProxyOutputs& Outputs = OutputBuffers.GetPhysicsDataForWrite();
+		Outputs.Transforms.SetNumZeroed(NumBones);
+		Outputs.LinearVelocities.SetNumZeroed(NumBones);
+		Outputs.AngularVelocities.SetNumZeroed(NumBones);
+
+#if TODO_REIMPLEMENT_GET_RIGID_PARTICLES
+		FPhysicsSolver::FParticlesType& Particles = GetSolver()->GetRigidParticles();
+		for (const int32 BoneIndex : Parameters.BoneHierarchy.GetBoneIndices())
+		{
+			FAnalyticImplicitGroup* ShapeGroup = Parameters.BoneHierarchy.GetAnalyticShapeGroup(BoneIndex);
+			if (ShapeGroup != nullptr)
+			{
+				int32 RigidBodyId = ShapeGroup->GetRigidBodyId();
+				if (RigidBodyId >= 0)
+				{
+					const FTransform BodyTransform = FTransform(Particles.Q(RigidBodyId), Particles.P(RigidBodyId));
+					int32 TransformIndex = Parameters.BoneHierarchy.GetTransformIndex(BoneIndex);
+					Outputs.Transforms[TransformIndex] = BodyTransform;
+					Outputs.LinearVelocities[TransformIndex] = Particles.V(RigidBodyId);
+					Outputs.AngularVelocities[TransformIndex] = Particles.W(RigidBodyId);
+				}
+			}
+		}
+#endif
+	}
+}
+
+void FSkeletalMeshPhysicsProxy::FlipBuffer()
+{
+	// PHYSICS THREAD
+
+	OutputBuffers.Flip();
+}
+
+void FSkeletalMeshPhysicsProxy::PullFromPhysicsState()
+{
+	// GAME THREAD
+	// @todo(ccaulfield): this needs to update the Component-Space pose in BoneHierarchy and update non-simulated children
+
+	CurrentOutputConsumerBuffer = &OutputBuffers.GetGameDataForRead();
+}
+
+void FSkeletalMeshPhysicsProxy::CaptureInputs(const float Dt, const FInputFunc& InputFunc)
+{
+	// GAME THREAD
+	SCOPE_CYCLE_COUNTER(STAT_SkelMeshUpdateAnim);
+
+	// Get a new buffer to write the input data to
+	if (NextInputProducerBuffer == nullptr)
+	{
+		NextInputProducerBuffer = InputBuffers.ExchangeProducerBuffer();
+	}
+
+	// Fill in the input data and switch the buffers if we had new data
+	if (InputFunc(Dt, Parameters))
+	{
+		const FBoneHierarchy& BoneHierarchy = Parameters.BoneHierarchy;
+		FSkeletalMeshPhysicsProxyInputs& PhysicsInputs = *NextInputProducerBuffer;
+		PhysicsInputs.Transforms.SetNumZeroed(BoneHierarchy.GetBoneIndices().Num());
+		PhysicsInputs.LinearVelocities.SetNumZeroed(BoneHierarchy.GetBoneIndices().Num());
+		PhysicsInputs.AngularVelocities.SetNumZeroed(BoneHierarchy.GetBoneIndices().Num());
+		for (int32 BoneIndex : BoneHierarchy.GetBoneIndices())
+		{
+			const int32 TransformIndex = BoneHierarchy.GetTransformIndex(BoneIndex);
+			if (TransformIndex != INDEX_NONE)
+			{
+				const FTransform& BoneWorldTransform = *BoneHierarchy.GetAnimWorldSpaceTransformsForBone(BoneIndex);
+				const FTransform* PrevBoneWorldTransform = BoneHierarchy.GetPrevAnimWorldSpaceTransformForBone(BoneIndex);
+				if ((Dt > SMALL_NUMBER) && PrevBoneWorldTransform)
+				{
+					PhysicsInputs.Transforms[TransformIndex] = *PrevBoneWorldTransform;
+					PhysicsInputs.LinearVelocities[TransformIndex] = (BoneWorldTransform.GetTranslation() - PrevBoneWorldTransform->GetTranslation()) / Dt;
+					PhysicsInputs.AngularVelocities[TransformIndex] = Chaos::TRotation<float, 3>::CalculateAngularVelocity(PrevBoneWorldTransform->GetRotation(), BoneWorldTransform.GetRotation(), Dt);
+				}
+				else
+				{
+					PhysicsInputs.Transforms[TransformIndex] = BoneWorldTransform;
+					PhysicsInputs.LinearVelocities[TransformIndex] = FVector::ZeroVector;
+					PhysicsInputs.AngularVelocities[TransformIndex] = FVector::ZeroVector;
+				}
+			}
+		}
+
+		NextInputProducerBuffer = InputBuffers.ExchangeProducerBuffer();
+	}
+}
+
+#endif // INCLUDE_CHAOS
diff --git a/Engine/Source/Runtime/Experimental/ChaosSolvers/Private/PhysicsProxy/SolverEventFilters.cpp b/Engine/Source/Runtime/Experimental/ChaosSolvers/Private/PhysicsProxy/SolverEventFilters.cpp
new file mode 100644
--- /dev/null
+++ b/Engine/Source/Runtime/Experimental/ChaosSolvers/Private/PhysicsProxy/SolverEventFilters.cpp
@@ -0,0 +1,69 @@
+// Copyright 1998-2019 Epic Games, Inc. All Rights Reserved.
+
+#include "SolverEventFilters.h"
+
+namespace Chaos
+{
+
+	bool FSolverCollisionEventFilter::Pass(const TCollisionData<float, 3>& InData) const
+	{
+		const float MinVelocitySquared = FMath::Square(Settings.MinSpeed);
+		const float MinImpulseSquared = FMath::Square(Settings.MinImpulse);
+
+		if (Settings.MinMass > 0.0f && InData.Mass1 < Settings.MinMass && InData.Mass2 < Settings.MinMass)
+			return false;
+
+		if (Settings.MinSpeed > 0.0f && InData.Velocity1.SizeSquared() < MinVelocitySquared && InData.Velocity2.SizeSquared() < MinVelocitySquared)
+			return false;
+
+		if (Settings.MinImpulse > 0.0f && InData.AccumulatedImpulse.SizeSquared() < MinImpulseSquared)
+			return false;
+
+		return true;
+	}
+
+	bool FSolverTrailingEventFilter::Pass(const TTrailingData<float, 3>& InData) const
+	{
+		const float MinSpeedThresholdSquared = Settings.MinSpeed * Settings.MinSpeed;
+
+		if (Settings.MinMass > 0.0f && InData.Mass < Settings.MinMass)
+			return false;
+
+		if (Settings.MinSpeed > 0 && InData.Velocity.SizeSquared() < MinSpeedThresholdSquared)
+			return false;
+
+		if (Settings.MinVolume > 0)
+		{
+			TVector<float, 3> Extents = InData.BoundingBox.Extents();
+			float Volume = Extents[0] * Extents[1] * Extents[2];
+
+			if (Volume < Settings.MinVolume)
+				return false;
+		}
+
+		return true;
+	}
+
+	bool FSolverBreakingEventFilter::Pass(const TBreakingData<float, 3>& InData) const
+	{
+		const float MinSpeedThresholdSquared = Settings.MinSpeed * Settings.MinSpeed;
+
+		if (Settings.MinMass > 0.0f && InData.Mass < Settings.MinMass)
+			return false;
+
+		if (Settings.MinSpeed > 0 && InData.Velocity.SizeSquared() < MinSpeedThresholdSquared)
+			return false;
+
+		if (Settings.MinVolume > 0)
+		{
+			TVector<float, 3> Extents = InData.BoundingBox.Extents();
+			float Volume = Extents[0] * Extents[1] * Extents[2];
+
+			if (Volume < Settings.MinVolume)
+				return false;
+		}
+
+		return true;
+	}
+
+} // namespace Chaos
diff --git a/Engine/Source/Runtime/Experimental/ChaosSolvers/Private/PhysicsProxy/StaticMeshPhysicsObject.cpp b/Engine/Source/Runtime/Experimental/ChaosSolvers/Private/PhysicsProxy/StaticMeshPhysicsObject.cpp
new file mode 100644
--- /dev/null
+++ b/Engine/Source/Runtime/Experimental/ChaosSolvers/Private/PhysicsProxy/StaticMeshPhysicsObject.cpp
@@ -0,0 +1,369 @@
+// Copyright 1998-2019 Epic Games, Inc. All Rights Reserved.
+
+#include "PhysicsProxy/StaticMeshPhysicsProxy.h"
+#include "PhysicsSolver.h"
+#include "GeometryCollection/GeometryCollectionCollisionStructureManager.h"
+#include "Chaos/Box.h"
+#include "Chaos/Capsule.h"
+#include "Chaos/Sphere.h"
+#include "Chaos/ErrorReporter.h"
+#include "Chaos/Serializable.h"
+#include "ChaosStats.h"
+
+#if INCLUDE_CHAOS
+
+FStaticMeshPhysicsProxy::FStaticMeshPhysicsProxy(UObject* InOwner, FCallbackInitFunc InInitFunc, FSyncDynamicFunc InSyncFunc)
+	: Base(InOwner)
+	, bInitializedState(false)
+	, RigidBodyId(INDEX_NONE)
+	, CenterOfMass(FVector::ZeroVector)
+	, Scale(FVector::ZeroVector)
+	, SimTransform(FTransform::Identity)
+	, InitialiseCallbackParamsFunc(InInitFunc)
+	, SyncDynamicTransformFunc(InSyncFunc)
+{
+	check(IsInGameThread());
+
+	Results.Get(0) = FTransform::Identity;
+	Results.Get(1) = FTransform::Identity;
+}
+
+void FStaticMeshPhysicsProxy::Initialize()
+{
+	check(IsInGameThread());
+
+	// Safe here - we're not in the solver yet if we're creating callbacks
+	Results.Get(0) = FTransform::Identity;
+	Results.Get(1) = FTransform::Identity;
+
+	InitialiseCallbackParamsFunc(Parameters);
+	Parameters.TargetTransform = &SimTransform;
+
+	Reset();
+}
+
+void FStaticMeshPhysicsProxy::Reset()
+{
+	bInitializedState = false;
+}
+
+void FStaticMeshPhysicsProxy::BufferKinematicUpdate(const FPhysicsProxyKinematicUpdate& InParamUpdate)
+{
+	BufferedKinematicUpdate = InParamUpdate;
+	bPendingKinematicUpdate = true;
+};
+
+bool FStaticMeshPhysicsProxy::IsSimulating() const
+{
+	return Parameters.bSimulating;
+}
+
+void FStaticMeshPhysicsProxy::UpdateKinematicBodiesCallback(const FParticlesType& Particles, const float Dt, const float Time, FKinematicProxy& Proxy)
+{
+	bool IsControlled = Parameters.ObjectType == EObjectStateTypeEnum::Chaos_Object_Kinematic;
+	if(IsControlled && Parameters.bSimulating)
+	{
+#if TODO_REIMPLEMENT_KINEMATIC_PROXY
+		bool bFirst = !Proxy.Ids.Num();
+		if(bFirst)
+		{
+			Proxy.Ids.Add(RigidBodyId);
+			Proxy.Position.SetNum(1);
+			Proxy.NextPosition.SetNum(1);
+			Proxy.Rotation.SetNum(1);
+			Proxy.NextRotation.SetNum(1);
+
+			const FTransform& Transform = Parameters.InitialTransform;
+			Proxy.Position[0] = Chaos::TVector<float, 3>(Transform.GetTranslation());
+			Proxy.NextPosition[0] = Proxy.Position[0] + Chaos::TVector<float, 3>(Parameters.InitialLinearVelocity) * Dt;
+			Proxy.Rotation[0] = Chaos::TRotation<float, 3>(Transform.GetRotation().GetNormalized());
+			Proxy.NextRotation[0] = Proxy.Rotation[0];
+		}
+
+		if (bPendingKinematicUpdate)
+		{
+			Proxy.Position[0] = Particles.X(RigidBodyId);
+			Proxy.NextPosition[0] = BufferedKinematicUpdate.NewTransform.GetTranslation();
+			Proxy.Rotation[0] = Particles.R(RigidBodyId);
+			Proxy.NextRotation[0] = BufferedKinematicUpdate.NewTransform.GetRotation().GetNormalized();
+
+			bPendingKinematicUpdate = false;
+		}
+#endif
+	}
+}
+
+void FStaticMeshPhysicsProxy::StartFrameCallback(const float InDt, const float InTime)
+{
+
+}
+
+void FStaticMeshPhysicsProxy::EndFrameCallback(const float InDt)
+{
+	bool IsControlled = Parameters.ObjectType == EObjectStateTypeEnum::Chaos_Object_Kinematic;
+	if(Parameters.bSimulating && !IsControlled && Parameters.TargetTransform)
+	{
+#if TODO_REIMPLEMENT_GET_RIGID_PARTICLES
+		const Chaos::TPBDRigidParticles<float, 3>& Particles = GetSolver()->GetRigidParticles();
+
+		Parameters.TargetTransform->SetTranslation((FVector)Particles.X(RigidBodyId));
+		Parameters.TargetTransform->SetRotation((FQuat)Particles.R(RigidBodyId));
+#endif
+	}
+}
+
+void FStaticMeshPhysicsProxy::BindParticleCallbackMapping(Chaos::TArrayCollectionArray<PhysicsProxyWrapper> & PhysicsProxyReverseMap, Chaos::TArrayCollectionArray<int32> & ParticleIDReverseMap)
+{
+	if (bInitializedState)
+	{
+		PhysicsProxyReverseMap[RigidBodyId] = { this, EPhysicsProxyType::StaticMeshType };
+		ParticleIDReverseMap[RigidBodyId] = 0;
+	}
+}
+
+void FStaticMeshPhysicsProxy::CreateRigidBodyCallback(FParticlesType& Particles)
+{
+	if(!bInitializedState && Parameters.bSimulating)
+	{
+		RigidBodyId = Particles.Size();
+		Particles.AddParticles(1);
+
+		FBox Bounds(ForceInitToZero);
+		if (Parameters.ShapeType == EImplicitTypeEnum::Chaos_Implicit_LevelSet)
+		{
+			// Build implicit surface
+			const ECollisionTypeEnum CollisionType = ECollisionTypeEnum::Chaos_Surface_Volumetric;
+			check(sizeof(Chaos::TVector<float, 3>) == sizeof(FVector));
+
+			//TODO: this should all be done on the asset instead of in the callback. Duplicates all mesh vertices every time
+			Chaos::TParticles<float, 3> MeshParticles(MoveTemp(Parameters.MeshVertexPositions));
+
+			TUniquePtr<Chaos::TTriangleMesh<float>> TriangleMesh(new Chaos::TTriangleMesh<float>(MoveTemp(Parameters.TriIndices)));
+			Chaos::FErrorReporter ErrorReporter(Parameters.Name + " | RigidBodyId: " + FString::FromInt(RigidBodyId));;
+			Particles.SetDynamicGeometry(
+				RigidBodyId,
+				TUniquePtr<Chaos::TImplicitObject<float, 3>>(
+					FCollisionStructureManager::NewImplicit(
+						ErrorReporter,
+						MeshParticles, *TriangleMesh,
+						Bounds,
+						FVector::Distance(FVector(0.f, 0.f, 0.f), Bounds.GetExtent())*0.5,
+						Parameters.MinRes, Parameters.MaxRes, 0.f,
+						CollisionType, Parameters.ShapeType)));
+			if (!ensure(Parameters.MeshVertexPositions.Size()))
+			{
+				Parameters.MeshVertexPositions.AddParticles(1);
+				Parameters.MeshVertexPositions.X(0) = Chaos::TVector<float, 3>(0, 0, 0);
+			}
+
+			for(uint32 i=0; i < Parameters.MeshVertexPositions.Size(); ++i)
+			{
+				Bounds += Parameters.MeshVertexPositions.X(i);
+			}
+		}
+		else if (Parameters.ShapeType == EImplicitTypeEnum::Chaos_Implicit_Sphere)
+		{
+			Chaos::TSphere<float,3>* Sphere = new Chaos::TSphere<float, 3>(Chaos::TVector<float, 3>(0), Parameters.ShapeParams.SphereRadius);
+			const Chaos::TBox<float, 3> BBox = Sphere->BoundingBox();
+			Bounds.Min = BBox.Min();
+			Bounds.Max = BBox.Max();
+			Particles.SetDynamicGeometry(RigidBodyId, TUniquePtr<Chaos::TImplicitObject<float, 3>>(Sphere));
+			if (!Parameters.MeshVertexPositions.Size())
+			{
+				float Radius = Parameters.ShapeParams.SphereRadius;
+
+				Parameters.MeshVertexPositions.AddParticles(6);
+				Parameters.MeshVertexPositions.X(0) = Chaos::TVector<float, 3>(Radius, 0, 0);
+				Parameters.MeshVertexPositions.X(1) = Chaos::TVector<float, 3>(-Radius, 0, 0);
+				Parameters.MeshVertexPositions.X(2) = Chaos::TVector<float, 3>(0, Radius, Radius);
+				Parameters.MeshVertexPositions.X(3) = Chaos::TVector<float, 3>(0, -Radius, Radius);
+				Parameters.MeshVertexPositions.X(4) = Chaos::TVector<float, 3>(0, -Radius, -Radius);
+				Parameters.MeshVertexPositions.X(5) = Chaos::TVector<float, 3>(0, Radius, -Radius);
+			}
+		}
+		else if (Parameters.ShapeType == EImplicitTypeEnum::Chaos_Implicit_Box)
+		{
+			Chaos::TVector<float, 3> HalfExtents = Parameters.ShapeParams.BoxExtents / 2;
+			Chaos::TBox<float,3>* Box = new Chaos::TBox<float, 3>(-HalfExtents, HalfExtents);
+			Bounds.Min = Box->Min();
+			Bounds.Max = Box->Max();
+			Particles.SetDynamicGeometry(RigidBodyId, TUniquePtr<Chaos::TImplicitObject<float, 3>>(Box));
+			if (!Parameters.MeshVertexPositions.Size())
+			{
+				Chaos::TVector<float, 3> x1(-HalfExtents);
+				Chaos::TVector<float, 3> x2(HalfExtents);
+
+				Parameters.MeshVertexPositions.AddParticles(8);
+				Parameters.MeshVertexPositions.X(0) = Chaos::TVector<float, 3>(x1.X, x1.Y, x1.Z);
+				Parameters.MeshVertexPositions.X(1) = Chaos::TVector<float, 3>(x1.X, x1.Y, x2.Z);
+				Parameters.MeshVertexPositions.X(2) = Chaos::TVector<float, 3>(x1.X, x2.Y, x1.Z);
+				Parameters.MeshVertexPositions.X(3) = Chaos::TVector<float, 3>(x2.X, x1.Y, x1.Z);
+				Parameters.MeshVertexPositions.X(4) = Chaos::TVector<float, 3>(x2.X, x2.Y, x2.Z);
+				Parameters.MeshVertexPositions.X(5) = Chaos::TVector<float, 3>(x2.X, x2.Y, x1.Z);
+				Parameters.MeshVertexPositions.X(6) = Chaos::TVector<float, 3>(x2.X, x1.Y, x2.Z);
+				Parameters.MeshVertexPositions.X(7) = Chaos::TVector<float, 3>(x1.X, x2.Y, x2.Z);
+			}
+		}
+		else if (Parameters.ShapeType == EImplicitTypeEnum::Chaos_Implicit_Capsule)
+		{
+			Chaos::TVector<float, 3> x1(0, -Parameters.ShapeParams.CapsuleHalfHeightAndRadius.X, 0);
+			Chaos::TVector<float, 3> x2(0, Parameters.ShapeParams.CapsuleHalfHeightAndRadius.X, 0);
+			Chaos::TCapsule<float>* Capsule = new Chaos::TCapsule<float>(x1, x2, Parameters.ShapeParams.CapsuleHalfHeightAndRadius.Y);
+			const Chaos::TBox<float, 3> BBox = Capsule->BoundingBox();
+			Bounds.Min = BBox.Min();
+			Bounds.Max = BBox.Max();
+			Particles.SetDynamicGeometry(RigidBodyId, TUniquePtr<Chaos::TImplicitObject<float, 3>>(Capsule));
+			if (!Parameters.MeshVertexPositions.Size())
+			{
+				float HalfHeight = Parameters.ShapeParams.CapsuleHalfHeightAndRadius.X;
+				float Radius = Parameters.ShapeParams.CapsuleHalfHeightAndRadius.Y;
+
+				Parameters.MeshVertexPositions.AddParticles(14);
+				Parameters.MeshVertexPositions.X(0) = Chaos::TVector<float, 3>(HalfHeight + Radius, 0, 0);
+				Parameters.MeshVertexPositions.X(1) = Chaos::TVector<float, 3>(-HalfHeight - Radius, 0, 0);
+				Parameters.MeshVertexPositions.X(2) = Chaos::TVector<float, 3>(HalfHeight, Radius, Radius);
+				Parameters.MeshVertexPositions.X(3) = Chaos::TVector<float, 3>(HalfHeight, -Radius, Radius);
+				Parameters.MeshVertexPositions.X(4) = Chaos::TVector<float, 3>(HalfHeight, -Radius, -Radius);
+				Parameters.MeshVertexPositions.X(5) = Chaos::TVector<float, 3>(HalfHeight, Radius, -Radius);
+				Parameters.MeshVertexPositions.X(6) = Chaos::TVector<float, 3>(0, Radius, Radius);
+				Parameters.MeshVertexPositions.X(7) = Chaos::TVector<float, 3>(0, -Radius, Radius);
+				Parameters.MeshVertexPositions.X(8) = Chaos::TVector<float, 3>(0, -Radius, -Radius);
+				Parameters.MeshVertexPositions.X(9) = Chaos::TVector<float, 3>(0, Radius, -Radius);
+				Parameters.MeshVertexPositions.X(10) = Chaos::TVector<float, 3>(-HalfHeight, Radius, Radius);
+				Parameters.MeshVertexPositions.X(11) = Chaos::TVector<float, 3>(-HalfHeight, -Radius, Radius);
+				Parameters.MeshVertexPositions.X(12) = Chaos::TVector<float, 3>(-HalfHeight, -Radius, -Radius);
+				Parameters.MeshVertexPositions.X(13) = Chaos::TVector<float, 3>(-HalfHeight, Radius, -Radius);
+			}
+		}
+		else
+		{
+			Bounds.Min = Chaos::TVector<float, 3>(-KINDA_SMALL_NUMBER);
+			Bounds.Max = Chaos::TVector<float, 3>(KINDA_SMALL_NUMBER);
+			Particles.SetSharedGeometry(RigidBodyId, TSharedPtr<Chaos::TImplicitObject<float, 3>, ESPMode::ThreadSafe>());
+		}
+
+		FTransform WorldTransform = Parameters.InitialTransform;
+
+		Scale = WorldTransform.GetScale3D();
+		CenterOfMass = Bounds.GetCenter();
+		Bounds = Bounds.InverseTransformBy(FTransform(CenterOfMass));
+		Bounds.Min *= Scale;
+		Bounds.Max *= Scale;
+		checkSlow((Bounds.Max + Bounds.Min).Size() < KINDA_SMALL_NUMBER);
+
+		Particles.InvM(RigidBodyId) = 0.f;
+		ensure(Parameters.Mass >= 0.f);
+		Particles.M(RigidBodyId) = Parameters.Mass;
+		if(Parameters.Mass < FLT_EPSILON)
+		{
+			Particles.InvM(RigidBodyId) = 1.f;
+		}
+		else
+		{
+			Particles.InvM(RigidBodyId) = 1.f / Parameters.Mass;
+		}
+
+		Particles.X(RigidBodyId) = WorldTransform.TransformPosition(CenterOfMass);
+		Particles.V(RigidBodyId) = Chaos::TVector<float, 3>(Parameters.InitialLinearVelocity);
+		Particles.R(RigidBodyId) = WorldTransform.GetRotation().GetNormalized();
+		Particles.W(RigidBodyId) = Chaos::TVector<float, 3>(Parameters.InitialAngularVelocity);
+		Particles.P(RigidBodyId) = Particles.X(RigidBodyId);
+		Particles.Q(RigidBodyId) = Particles.R(RigidBodyId);
+
+		const FVector SideSquared(Bounds.GetSize()[0] * Bounds.GetSize()[0], Bounds.GetSize()[1] * Bounds.GetSize()[1], Bounds.GetSize()[2] * Bounds.GetSize()[2]);
+		const FVector Inertia(Parameters.Mass * (SideSquared.Y + SideSquared.Z) / 12.f, Parameters.Mass * (SideSquared.X + SideSquared.Z) / 12.f, Parameters.Mass * (SideSquared.X + SideSquared.Y) / 12.f);
+		Particles.I(RigidBodyId) = Chaos::PMatrix<float, 3, 3>(Inertia.X, 0.f, 0.f, 0.f, Inertia.Y, 0.f, 0.f, 0.f, Inertia.Z);
+		Particles.InvI(RigidBodyId) = Chaos::PMatrix<float, 3, 3>(1.f / Inertia.X, 0.f, 0.f, 0.f, 1.f / Inertia.Y, 0.f, 0.f, 0.f, 1.f / Inertia.Z);
+
+		if(Parameters.ObjectType == EObjectStateTypeEnum::Chaos_Object_Sleeping)
+		{
+			Particles.SetObjectState(RigidBodyId, Chaos::EObjectStateType::Sleeping);
+			Particles.SetSleeping(RigidBodyId, true);
+		}
+		else if(Parameters.ObjectType != EObjectStateTypeEnum::Chaos_Object_Dynamic)
+		{
+			Particles.InvM(RigidBodyId) = 0.f;
+			Particles.InvI(RigidBodyId) = Chaos::PMatrix<float, 3, 3>(0);
+			Particles.SetObjectState(RigidBodyId, Chaos::EObjectStateType::Kinematic);
+		}
+		else
+		{
+			Particles.SetObjectState(RigidBodyId, Chaos::EObjectStateType::Dynamic);
+		}
+
+		if (ensure(Parameters.MeshVertexPositions.Size()))
+		{
+			// Add collision vertices
+			check(Particles.CollisionParticles(RigidBodyId) == nullptr);
+			Particles.SetCollisionParticles(RigidBodyId, MoveTemp(Parameters.MeshVertexPositions));
+			if (Particles.CollisionParticles(RigidBodyId)->Size())
+			{
+				Particles.CollisionParticles(RigidBodyId)->UpdateAccelerationStructures();
+			}
+		}
+
+#if TODO_REIMPLEMENT_SET_PHYSICS_MATERIAL
+		GetSolver()->SetPhysicsMaterial(RigidBodyId, Parameters.PhysicalMaterial);
+#endif
+
+		bInitializedState = true;
+	}
+}
+
+void FStaticMeshPhysicsProxy::ParameterUpdateCallback(FParticlesType& InParticles, const float InTime)
+{
+
+}
+
+void FStaticMeshPhysicsProxy::DisableCollisionsCallback(TSet<TTuple<int32, int32>>& InPairs)
+{
+
+}
+
+void FStaticMeshPhysicsProxy::AddForceCallback(FParticlesType& InParticles, const float InDt, const int32 InIndex)
+{
+
+}
+
+void FStaticMeshPhysicsProxy::OnRemoveFromScene()
+{
+	// Disable the particle we added
+	Chaos::FPhysicsSolver* CurrSolver = GetSolver();
+
+	if(CurrSolver && RigidBodyId != INDEX_NONE)
+	{
+#if TODO_REIMPLEMENT_GET_RIGID_PARTICLES
+		// #BG TODO Special case here because right now we reset/realloc the evolution per geom component
+		// in endplay which clears this out. That needs to not happen and be based on world shutdown
+		if(CurrSolver->GetRigidParticles().Size() == 0)
+		{
+			return;
+		}
+
+		CurrSolver->GetEvolution()->DisableParticle(RigidBodyId);
+		CurrSolver->GetRigidClustering().GetTopLevelClusterParents().Remove(RigidBodyId);
+#endif
+	}
+}
+
+void FStaticMeshPhysicsProxy::BufferPhysicsResults()
+{
+	SCOPE_CYCLE_COUNTER(STAT_CacheResultStaticMesh);
+	Results.GetPhysicsDataForWrite() = SimTransform;
+}
+
+void FStaticMeshPhysicsProxy::FlipBuffer()
+{
+	Results.Flip();
+}
+
+void FStaticMeshPhysicsProxy::PullFromPhysicsState()
+{
+	if(Parameters.ObjectType == EObjectStateTypeEnum::Chaos_Object_Dynamic && Parameters.bSimulating && SyncDynamicTransformFunc)
+	{
+		// Send transform to update callable
+		SyncDynamicTransformFunc(Results.GetGameDataForRead());
+	}
+}
+
+#endif
diff --git a/Engine/Source/Runtime/Experimental/ChaosSolvers/Public/BoneHierarchy.h b/Engine/Source/Runtime/Experimental/ChaosSolvers/Public/BoneHierarchy.h
--- a/Engine/Source/Runtime/Experimental/ChaosSolvers/Public/BoneHierarchy.h
+++ b/Engine/Source/Runtime/Experimental/ChaosSolvers/Public/BoneHierarchy.h
@@ -2,7 +2,7 @@
 #pragma once
 #if INCLUDE_CHAOS
 
-#include "SolverObjects/AnalyticImplicitGroup.h"
+#include "PhysicsProxy/AnalyticImplicitGroup.h"
 
 DEFINE_LOG_CATEGORY_STATIC(USkeletalMeshSimulationComponentLogging, NoLogging, All);
 //DEFINE_LOG_CATEGORY_STATIC(USkeletalMeshSimulationComponentLogging, Log, All);
@@ -301,7 +301,7 @@ protected:
 		check(Roots.Num() != 0);
 		check(NumTransforms == ImplicitGroups.Num());
 		UE_LOG(USkeletalMeshSimulationComponentLogging, Log,
-			TEXT("USkeletalMeshPhysicsObject::InitHierarchy() - this: %p - "
+			TEXT("USkeletalMeshPhysicsProxy::InitHierarchy() - this: %p - "
 				"Implicit groups: %d, num transforms: %d, num roots: %d"),
 			this,
 			ImplicitGroups.Num(),
diff --git a/Engine/Source/Runtime/Experimental/ChaosSolvers/Public/ChaosSolversModule.h b/Engine/Source/Runtime/Experimental/ChaosSolvers/Public/ChaosSolversModule.h
--- a/Engine/Source/Runtime/Experimental/ChaosSolvers/Public/ChaosSolversModule.h
+++ b/Engine/Source/Runtime/Experimental/ChaosSolvers/Public/ChaosSolversModule.h
@@ -11,7 +11,8 @@
 #include "UObject/ObjectMacros.h"
 #include "Framework/Threading.h"
 #include "PhysicsCoreTypes.h"
-#include "MultiBufferResource.h"
+#include "Chaos/Framework/MultiBufferResource.h"
+#include "Chaos/Declares.h"
 
 /** Classes that want to set the solver actor class can implement this. */
 class IChaosSolverActorClassProvider
@@ -40,7 +41,6 @@ namespace Chaos
 {
 	class FPersistentPhysicsTask;
 	class FPhysicsProxy;
-	class FPBDRigidsSolver;
 }
 
 // Default settings implementation
@@ -66,7 +66,7 @@ struct FSolverStateStorage
 {
 	friend class FChaosSolversModule;
 
-	Chaos::FPBDRigidsSolver* Solver;
+	Chaos::FPhysicsSolver* Solver;
 	TArray<Chaos::FPhysicsProxy*> ActiveProxies;
 	TArray<Chaos::FPhysicsProxy*> ActiveProxies_GameThread;
 
@@ -99,6 +99,7 @@ public:
 
 	void ShutdownThreadingMode();
 	void InitializeThreadingMode(Chaos::EThreadingMode InNewMode);
+	void ChangeThreadingMode(Chaos::EThreadingMode InNewMode);
 
 	/**
 	 * Queries for multithreaded configurations
@@ -152,7 +153,7 @@ public:
 	 *
 	 * @param bStandalone Whether the solver is standalone (not sent to physics thread - updating left to caller)
 	 */
-	Chaos::FPBDRigidsSolver* CreateSolver(bool bStandalone = false);
+	Chaos::FPhysicsSolver* CreateSolver(bool bStandalone = false);
 
 	Chaos::EMultiBufferMode GetBufferModeFromThreadingModel(Chaos::EThreadingMode ThreadingMode)
 	{
@@ -217,13 +218,13 @@ public:
 	 *
 	 * Should be called on whichever thread currently owns the solver state
 	 */
-	void DestroySolver(Chaos::FPBDRigidsSolver* InState);
+	void DestroySolver(Chaos::FPhysicsSolver* InState);
 
 	/**
 	 * Read access to the current solver-state objects, be aware which thread owns this data when
 	 * attempting to use this. Physics thread will query when spinning up to get current world state
 	 */
-	const TArray<Chaos::FPBDRigidsSolver*>& GetSolvers() const { return Solvers; }
+	const TArray<Chaos::FPhysicsSolver*>& GetSolvers() const { return Solvers; }
 
 	/**
 	 * Outputs statistics for the solver hierarchies. Currently engine calls into this
@@ -314,7 +315,7 @@ private:
 	// Where these objects are valid for interaction depends on the current threading mode. Use IsPersistentTaskRunning to
 	// check whether the physics thread owns these before manipulating. When adding/removing solver or proxy items in
 	// multi threaded mode the physics thread must also be notified of the change.
-	TArray<Chaos::FPBDRigidsSolver*> Solvers;
+	TArray<Chaos::FPhysicsSolver*> Solvers;
 
 	// Lock for the above list to ensure we don't delete solvers out from underneath other threads
 	// or mess up the solvers array during use.
diff --git a/Engine/Source/Runtime/Experimental/ChaosSolvers/Public/EventDefaults.h b/Engine/Source/Runtime/Experimental/ChaosSolvers/Public/EventDefaults.h
new file mode 100644
--- /dev/null
+++ b/Engine/Source/Runtime/Experimental/ChaosSolvers/Public/EventDefaults.h
@@ -0,0 +1,417 @@
+// Copyright 1998-2019 Epic Games, Inc. All Rights Reserved.
+
+#pragma once
+
+#include "EventManager.h"
+#include "EventsData.h"
+#include "Chaos/PBDRigidsEvolution.h"
+#include "Chaos/PBDRigidClustering.h"
+#include "Chaos/PBDCollisionConstraint.h"
+#include "PBDRigidsSolver.h"
+
+namespace Chaos
+{
+	class FEventDefaults
+	{
+	public:
+
+		/**
+		 * Register default event types
+		 */
+		static void RegisterSystemEvents(FEventManager& EventManager)
+		{
+			RegisterCollisionEvent(EventManager);
+			RegisterBreakingEvent(EventManager);
+			RegisterTrailingEvent(EventManager);
+			RegisterSleepingEvent(EventManager);
+		}
+
+	private:
+
+		/**
+		 * Register collision event gathering function & data type
+		 */
+		static void RegisterCollisionEvent(FEventManager& EventManager)
+		{
+			EventManager.RegisterEvent<FCollisionEventData>(EEventType::Collision, []
+			(const Chaos::FPBDRigidsSolver* Solver, FCollisionEventData& CollisionEventData)
+			{
+				check(Solver);
+
+				// #todo: This isn't working - SolverActor parameters are set on a solver but it is currently a different solver that is simulating!!
+				//if (!Solver->GetEventFilters()->IsCollisionEventEnabled())
+				//	return;
+
+				FCollisionDataArray& AllCollisionsDataArray = CollisionEventData.CollisionData.AllCollisionsArray;
+				TMap<IPhysicsProxyBase*, TArray<int32>>& AllCollisionsIndicesByPhysicsProxy = CollisionEventData.PhysicsProxyToCollisionIndices.PhysicsProxyToIndicesMap;
+
+				AllCollisionsDataArray.Reset();
+				AllCollisionsIndicesByPhysicsProxy.Reset();
+
+				CollisionEventData.CollisionData.TimeCreated = Solver->MTime;
+				CollisionEventData.PhysicsProxyToCollisionIndices.TimeCreated = Solver->MTime;
+
+				const FPBDRigidsSolver::FPBDRigidsEvolution* Evolution = Solver->GetEvolution();
+
+				const FPhysicsSolver::FPBDCollisionConstraints& CollisionRule = Evolution->GetCollisionConstraints();
+				const TArray<FPhysicsSolver::FPBDCollisionConstraints::FRigidBodyContactConstraint>& AllConstraintsArray = CollisionRule.GetAllConstraints();
+
+				const TPBDRigidParticles<float, 3>& Particles = Evolution->GetParticles().GetDynamicParticles();
+				const TArrayCollectionArray<ClusterId>& ClusterIdsArray = Evolution->GetRigidClustering().GetClusterIdsArray();
+#if TODO_REIMPLEMENT_RIGID_CLUSTERING
+				const Chaos::FPBDRigidsSolver::FClusteringType::FClusterMap& ParentToChildrenMap = Evolution->GetRigidClustering().GetChildrenMap();
+#endif
+				const Chaos::TPBDRigidClustering<FPhysicsSolver::FPBDRigidsEvolution, FPhysicsSolver::FPBDCollisionConstraints, float, 3>::FClusterMap& ParentToChildrenMap = Evolution->GetRigidClustering().GetChildrenMap();
+
+				if (AllConstraintsArray.Num() > 0)
+				{
+					// There's a couple of nasty const correctness issues here
+					Chaos::FPBDRigidsSolver* NonConstSolver = (Chaos::FPBDRigidsSolver*)(Solver);
+					TMap<IPhysicsProxyBase*, TArray<int32>>* ProxyMap = (TMap<IPhysicsProxyBase*, TArray<int32>>*)(&AllCollisionsIndicesByPhysicsProxy);
+					// Add the keys to AllCollisionsIndicesByPhysicsProxy map
+					NonConstSolver->ForEachPhysicsProxy([ProxyMap](auto* Obj)
+					{
+						if (Obj && Obj->IsSimulating())
+						{
+							if (!ProxyMap->Contains(Obj))
+							{
+								ProxyMap->Reset();
+								ProxyMap->Add(Obj, TArray<int32>());
+							}
+						}
+					});
+
+					// Get the number of valid constraints (AccumulatedImpulse != 0.f and Phi < 0.f) from AllConstraintsArray
+					TArray<int32> ValidCollisionIndices;
+					ValidCollisionIndices.SetNumUninitialized(AllConstraintsArray.Num());
+					int32 NumValidCollisions = 0;
+
+					for (int32 Idx = 0; Idx < AllConstraintsArray.Num(); ++Idx)
+					{
+						const TRigidBodyContactConstraint<float, 3>& Constraint = AllConstraintsArray[Idx];
+
+						// Since Clustered GCs can be unioned the particleIndex representing the union 
+						// is not associated with a PhysicsProxy
+						if (Constraint.Particle->Handle()->GTGeometryParticle()->Proxy != nullptr)
+
+						{
+							if (ensure(!Constraint.AccumulatedImpulse.ContainsNaN() && FMath::IsFinite(Constraint.Phi)))
+							{
+								TGeometryParticleHandle<float, 3>* Particle0 = Constraint.Particle;
+								TGeometryParticleHandle<float, 3>* Particle1 = Constraint.Levelset;
+								TKinematicGeometryParticleHandle<float, 3>* Body0 = Particle0->AsKinematic();
+
+								// presently when a rigidbody or kinematic hits static geometry then Body1 is null
+								TKinematicGeometryParticleHandle<float, 3>* Body1 = Particle1->AsKinematic();
+
+								if (!Constraint.AccumulatedImpulse.IsZero() && Body0)
+								{
+									if (ensure(!Constraint.Location.ContainsNaN() &&
+										!Constraint.Normal.ContainsNaN()) &&
+										!Body0->V().ContainsNaN() &&
+										!Body0->W().ContainsNaN() &&
+										(Body1==nullptr || (!Body1->V().ContainsNaN()) && !Body1->W().ContainsNaN()) )
+									{
+										ValidCollisionIndices[NumValidCollisions] = Idx;
+										NumValidCollisions++;
+									}
+								}
+							}
+						}
+					}
+
+					ValidCollisionIndices.SetNum(NumValidCollisions);
+
+					if (ValidCollisionIndices.Num() > 0)
+					{
+						for (int32 IdxCollision = 0; IdxCollision < ValidCollisionIndices.Num(); ++IdxCollision)
+						{
+							Chaos::TPBDCollisionConstraint<float, 3>::FRigidBodyContactConstraint const& Constraint = AllConstraintsArray[ValidCollisionIndices[IdxCollision]];
+
+							TGeometryParticleHandle<float, 3>* Particle0 = Constraint.Particle;
+							TGeometryParticleHandle<float, 3>* Particle1 = Constraint.Levelset;
+
+							TCollisionData<float, 3> Data;
+							Data.Location = Constraint.Location;
+							Data.AccumulatedImpulse = Constraint.AccumulatedImpulse;
+							Data.Normal = Constraint.Normal;
+							Data.PenetrationDepth = Constraint.Phi;
+							Data.Particle = Particle0->Handle()->GTGeometryParticle();
+							Data.Levelset = Particle1->Handle()->GTGeometryParticle();
+
+							// todo: do we need these anymore now we are storing the particles you can access all of this stuff from there
+							// do we still need these now we have pointers to particles returned?
+							if (TPBDRigidParticleHandle<float, 3>* PBDRigid0 = Particle0->AsDynamic())
+							{
+								Data.Velocity1 = PBDRigid0->V();
+								Data.AngularVelocity1 = PBDRigid0->W();
+								Data.Mass1 = PBDRigid0->M();
+							}
+
+							if (TPBDRigidParticleHandle<float, 3>* PBDRigid1 = Particle1->AsDynamic())
+							{
+								Data.Velocity2 = PBDRigid1->V();
+								Data.AngularVelocity2 = PBDRigid1->W();
+								Data.Mass2 = PBDRigid1->M();
+							}
+
+							IPhysicsProxyBase* const PhysicsProxy = Data.Particle->Proxy;
+							IPhysicsProxyBase* const OtherPhysicsProxy = Data.Levelset->Proxy;
+							//Data.Material1 = nullptr; // #todo: provide UPhysicalMaterial for Particle
+							//Data.Material2 = nullptr; // #todo: provide UPhysicalMaterial for Levelset
+
+							const FSolverCollisionEventFilter* SolverCollisionEventFilter = Solver->GetEventFilters()->GetCollisionFilter();
+							if (!SolverCollisionEventFilter->Enabled() || SolverCollisionEventFilter->Pass(Data))
+
+							{
+								const int32 NewIdx = AllCollisionsDataArray.Add(TCollisionData<float, 3>());
+								TCollisionData<float, 3>& CollisionDataArrayItem = AllCollisionsDataArray[NewIdx];
+
+								CollisionDataArrayItem = Data;
+
+#if TODO_REIMPLEMENT_RIGID_CLUSTERING
+								// If Constraint.ParticleIndex is a cluster store an index for a mesh in this cluster
+								if (ClusterIdsArray[Constraint.ParticleIndex].NumChildren > 0)
+								{
+									int32 ParticleIndexMesh = GetParticleIndexMesh(ParentToChildrenMap, Constraint.ParticleIndex);
+									ensure(ParticleIndexMesh != INDEX_NONE);
+									CollisionDataArrayItem.ParticleIndexMesh = ParticleIndexMesh;
+								}
+								// If Constraint.LevelsetIndex is a cluster store an index for a mesh in this cluster
+								if (ClusterIdsArray[Constraint.LevelsetIndex].NumChildren > 0)
+								{
+									int32 LevelsetIndexMesh = GetParticleIndexMesh(ParentToChildrenMap, Constraint.LevelsetIndex);
+									ensure(LevelsetIndexMesh != INDEX_NONE);
+									CollisionDataArrayItem.LevelsetIndexMesh = LevelsetIndexMesh;
+								}
+#endif
+
+								// Add to AllCollisionsIndicesByPhysicsProxy
+								AllCollisionsIndicesByPhysicsProxy.FindOrAdd(PhysicsProxy).Add(FEventManager::EncodeCollisionIndex(NewIdx, false));
+
+								if (OtherPhysicsProxy && OtherPhysicsProxy != PhysicsProxy)
+								{
+									AllCollisionsIndicesByPhysicsProxy.FindOrAdd(OtherPhysicsProxy).Add(FEventManager::EncodeCollisionIndex(NewIdx, true));
+								}
+							}
+						}
+					}
+				}
+
+			});
+		}
+
+		/**
+		 * Register breaking event gathering function & data type
+		 */
+		static void RegisterBreakingEvent(FEventManager& EventManager)
+		{
+			EventManager.RegisterEvent<FBreakingEventData>(EEventType::Breaking, []
+			(const Chaos::FPBDRigidsSolver* Solver, FBreakingEventData& BreakingEventData)
+			{
+				check(Solver);
+
+				// #todo: This isn't working - SolverActor parameters are set on a solver but it is currently a different solver that is simulating!!
+				//if (!Solver->GetEventFilters()->IsBreakingEventEnabled())
+				//	return;
+
+				FBreakingDataArray& AllBreakingDataArray = BreakingEventData.BreakingData.AllBreakingsArray;
+				TMap<IPhysicsProxyBase*, TArray<int32>>& AllBreakingIndicesByPhysicsProxy = BreakingEventData.PhysicsProxyToBreakingIndices.PhysicsProxyToIndicesMap;
+
+				AllBreakingDataArray.Reset();
+				AllBreakingIndicesByPhysicsProxy.Reset();
+
+				BreakingEventData.BreakingData.TimeCreated = Solver->MTime;
+				BreakingEventData.PhysicsProxyToBreakingIndices.TimeCreated = Solver->MTime;
+
+				const FPBDRigidsSolver::FPBDRigidsEvolution* Evolution = Solver->GetEvolution();
+				const TPBDRigidParticles<float, 3>& Particles = Evolution->GetParticles().GetDynamicParticles();
+				const TArray<TBreakingData<float, 3>>& AllBreakingsArray = Evolution->GetRigidClustering().GetAllClusterBreakings();
+				const TArrayCollectionArray<ClusterId>& ClusterIdsArray = Evolution->GetRigidClustering().GetClusterIdsArray();
+
+#if TODO_REIMPLEMENT_RIGID_CLUSTERING
+				const Chaos::FPBDRigidsSolver::FClusteringType::FClusterMap& ParentToChildrenMap = Evolution->GetRigidClustering().GetChildrenMap();
+#endif
+
+				if (AllBreakingsArray.Num() > 0)
+				{
+					// There's a couple of nasty const correctness issues here
+					Chaos::FPBDRigidsSolver* NonConstSolver = (Chaos::FPBDRigidsSolver*)(Solver);
+					TMap<IPhysicsProxyBase*, TArray<int32>>* ProxyMap = (TMap<IPhysicsProxyBase*, TArray<int32>>*)(&AllBreakingIndicesByPhysicsProxy);
+					NonConstSolver->ForEachPhysicsProxy([ProxyMap](auto* Obj)
+					{
+						if (Obj && Obj->IsSimulating())
+						{
+							if (!ProxyMap->Contains(Obj))
+							{
+								ProxyMap->Add(Obj, TArray<int32>());
+							}
+						}
+					});
+
+					for (int32 Idx = 0; Idx < AllBreakingsArray.Num(); ++Idx)
+					{
+						// Since Clustered GCs can be unioned the particleIndex representing the union 
+						// is not associated with a PhysicsProxy
+						if (AllBreakingsArray[Idx].Particle->Proxy != nullptr)
+						{
+							if (ensure(!AllBreakingsArray[Idx].Location.ContainsNaN() &&
+								!Particles.V(AllBreakingsArray[Idx].ParticleIndex).ContainsNaN() &&
+								!Particles.W(AllBreakingsArray[Idx].ParticleIndex).ContainsNaN()))
+							{
+								TBreakingData<float, 3> BreakingData;
+								BreakingData.Location = AllBreakingsArray[Idx].Location;
+								BreakingData.Velocity = Particles.V(AllBreakingsArray[Idx].ParticleIndex);
+								BreakingData.AngularVelocity = Particles.W(AllBreakingsArray[Idx].ParticleIndex);
+								BreakingData.Mass = Particles.M(AllBreakingsArray[Idx].ParticleIndex);
+								BreakingData.ParticleIndex = AllBreakingsArray[Idx].ParticleIndex;
+								if (Particles.Geometry(Idx)->HasBoundingBox())
+								{
+									BreakingData.BoundingBox = Particles.Geometry(Idx)->BoundingBox();;
+								}
+
+								const FSolverBreakingEventFilter* SolverBreakingEventFilter = Solver->GetEventFilters()->GetBreakingFilter();
+								if (!SolverBreakingEventFilter->Enabled() || SolverBreakingEventFilter->Pass(BreakingData))
+								{
+									int32 NewIdx = AllBreakingDataArray.Add(TBreakingData<float, 3>());
+									TBreakingData<float, 3>& BreakingDataArrayItem = AllBreakingDataArray[NewIdx];
+									BreakingDataArrayItem = BreakingData;
+
+									// If AllBreakingsArray[Idx].ParticleIndex is a cluster store an index for a mesh in this cluster
+									if (ClusterIdsArray[AllBreakingsArray[Idx].ParticleIndex].NumChildren > 0)
+									{
+#if 0 // #todo
+										int32 ParticleIndexMesh = GetParticleIndexMesh(ParentToChildrenMap, AllBreakingsArray[Idx].ParticleIndex);
+										ensure(ParticleIndexMesh != INDEX_NONE);
+										BreakingDataArrayItem.ParticleIndexMesh = ParticleIndexMesh;
+#endif
+									}
+
+									// Add to AllBreakingsIndicesByPhysicsProxy
+									IPhysicsProxyBase* PhysicsProxy = AllBreakingsArray[Idx].Particle->Proxy;
+									AllBreakingIndicesByPhysicsProxy.FindOrAdd(PhysicsProxy).Add(NewIdx);
+								}
+							}
+						}
+					}
+				}		
+			});
+		}
+
+		/**
+		 * Register trailing event gathering function & data type
+		 */
+		static void RegisterTrailingEvent(FEventManager& EventManager)
+		{
+			EventManager.RegisterEvent<FTrailingEventData>(EEventType::Trailing, []
+			(const Chaos::FPBDRigidsSolver* Solver, FTrailingEventData& TrailingEventData)
+			{
+				check(Solver);
+
+				// #todo: This isn't working - SolverActor parameters are set on a solver but it is currently a different solver that is simulating!!
+				//if (!Solver->GetEventFilters()->IsTrailingEventEnabled())
+				//	return;
+
+				const FPBDRigidsSolver::FPBDRigidsEvolution* Evolution = Solver->GetEvolution();
+
+				const TArrayCollectionArray<ClusterId>& ClusterIdsArray = Evolution->GetRigidClustering().GetClusterIdsArray();
+#if TODO_REIMPLEMENT_RIGID_CLUSTERING
+				const TMap<uint32, TUniquePtr<TArray<uint32>>>& ParentToChildrenMap = Evolution->GetRigidClustering().GetChildrenMap();
+#endif
+				auto& AllTrailingsDataArray = TrailingEventData.TrailingData.AllTrailingsArray;
+				auto& AllTrailingsIndicesByPhysicsProxy = TrailingEventData.PhysicsProxyToTrailingIndices.PhysicsProxyToIndicesMap;
+
+				AllTrailingsDataArray.Reset();
+				AllTrailingsIndicesByPhysicsProxy.Reset();
+
+				TrailingEventData.TrailingData.TimeCreated = Solver->MTime;
+				TrailingEventData.PhysicsProxyToTrailingIndices.TimeCreated = Solver->MTime;
+
+				// There's a couple of nasty const correctness issues here
+				Chaos::FPBDRigidsSolver* NonConstSolver = (Chaos::FPBDRigidsSolver*)(Solver);
+				TMap<IPhysicsProxyBase*, TArray<int32>>* ProxyMap = (TMap<IPhysicsProxyBase*, TArray<int32>>*)(&AllTrailingsIndicesByPhysicsProxy);
+				// Add the keys to AllTrailingsIndicesByPhysicsProxy map
+				NonConstSolver->ForEachPhysicsProxy([ProxyMap](auto* Obj)
+				{
+					if (Obj && Obj->IsSimulating())
+					{
+						if (!ProxyMap->Contains(Obj))
+						{
+							ProxyMap->Add(Obj, TArray<int32>());
+						}
+					}
+				});
+
+				for (auto& ActiveParticle :  Evolution->GetParticles().GetActiveParticlesView())
+				{
+					TPBDRigidParticle<float, 3>* GTParticle = ActiveParticle.Handle()->GTGeometryParticle()->AsDynamic();
+					// Since Clustered GCs can be unioned the particleIndex representing the union 
+					// is not associated with a PhysicsProxy
+					if (GTParticle && GTParticle->Proxy != nullptr)
+					{
+						if (ensure(FMath::IsFinite(GTParticle->InvM())))
+						{
+							if (GTParticle->InvM() != 0.f &&
+								ActiveParticle.Geometry() &&
+								ActiveParticle.Geometry()->HasBoundingBox())
+							{
+								if (ensure(!GTParticle->X().ContainsNaN() &&
+									!GTParticle->V().ContainsNaN() &&
+									!GTParticle->W().ContainsNaN() &&
+									FMath::IsFinite(GTParticle->M())))
+								{
+									TTrailingData<float, 3> TrailingData;
+									TrailingData.Location = GTParticle->X();
+									TrailingData.Velocity = GTParticle->V();
+									TrailingData.AngularVelocity = GTParticle->W();
+									TrailingData.Mass = GTParticle->M();
+									TrailingData.Particle = GTParticle;
+									if (ActiveParticle.Geometry()->HasBoundingBox())
+									{
+										TrailingData.BoundingBox = ActiveParticle.Geometry()->BoundingBox();
+									}
+
+									const FSolverTrailingEventFilter* SolverTrailingEventFilter = Solver->GetEventFilters()->GetTrailingFilter();
+									if (!SolverTrailingEventFilter->Enabled() || SolverTrailingEventFilter->Pass(TrailingData))
+									{
+										int32 NewIdx = AllTrailingsDataArray.Add(TTrailingData<float, 3>());
+										TTrailingData<float, 3>& TrailingDataArrayItem = AllTrailingsDataArray[NewIdx];
+										TrailingDataArrayItem = TrailingData;
+
+										// If IdxParticle is a cluster store an index for a mesh in this cluster
+#if 0
+										if (ClusterIdsArray[IdxParticle].NumChildren > 0)
+										{
+											int32 ParticleIndexMesh = GetParticleIndexMesh(ParentToChildrenMap, IdxParticle);
+											ensure(ParticleIndexMesh != INDEX_NONE);
+											TrailingDataArrayItem.ParticleIndexMesh = ParticleIndexMesh;
+										}
+#endif
+
+										// Add to AllTrailingsIndicesByPhysicsProxy
+										IPhysicsProxyBase* PhysicsProxy = GTParticle->Proxy;
+										AllTrailingsIndicesByPhysicsProxy.FindOrAdd(PhysicsProxy).Add(NewIdx);
+									}
+								}
+							}
+						}
+					}
+				}
+			});
+		}
+
+		static void RegisterSleepingEvent(FEventManager& EventManager)
+		{
+			//#todo: implement sleeping events
+			//EventManager.RegisterEvent<FSleepingEventData>(EEventType::Sleeping, []
+			//(const Chaos::FPBDRigidsSolver* Solver, FSleepingEventData& SleepingEventData)
+			//{
+			//	check(Solver);
+
+			//});
+		}
+
+
+
+	};
+}
\ No newline at end of file
diff --git a/Engine/Source/Runtime/Experimental/ChaosSolvers/Public/EventManager.h b/Engine/Source/Runtime/Experimental/ChaosSolvers/Public/EventManager.h
new file mode 100644
--- /dev/null
+++ b/Engine/Source/Runtime/Experimental/ChaosSolvers/Public/EventManager.h
@@ -0,0 +1,321 @@
+// Copyright 1998-2019 Epic Games, Inc. All Rights Reserved.
+
+#pragma once
+
+#include "CoreMinimal.h"
+#include "Chaos/Framework/MultiBufferResource.h"
+#include "EventManager.h"
+
+namespace Chaos
+{
+
+	//#todo : need to enable disable events, or not register them if disabled.
+	//#todo : add timers
+	//#todo : warning if trying to add same event ID twice
+	//#todo : need sparse array for EventID -> EventContainer?
+
+	class FPBDRigidsSolver;
+
+	/**
+	 * Predefined System Event Types
+	 */
+	enum EEventType : int32
+	{
+		Collision = 0,
+		Breaking = 1,
+		Trailing = 2,
+		Sleeping = 3
+	};
+
+	typedef int32 FEventID;
+
+
+	/**
+	 * Interface for event handler 
+	 */
+	class IEventHandler
+	{
+		template<typename PayloadType>
+		friend class TEventContainer;
+
+	public:
+		virtual ~IEventHandler() {}
+		virtual void HandleEvent(const void* EventData) const = 0;
+
+	protected:
+		// internal use only
+		virtual void* GetHandler() const = 0;
+	};
+
+	/** Instance event handler */
+	template<typename PayloadType, typename HandlerType>
+	class TRawEventHandler : public IEventHandler
+	{
+	public:
+		typedef void (HandlerType::*FHandlerFunction)(const PayloadType&);
+
+		TRawEventHandler(HandlerType* InHandler, FHandlerFunction InFunction)
+			: Handler(InHandler)
+			, HandlerFunction(InFunction)
+		{
+			check(Handler);
+			check(HandlerFunction);
+		}
+
+		virtual void HandleEvent(const void* EventData) const override
+		{
+			(Handler->*HandlerFunction)(*(const PayloadType*)EventData);
+		}
+
+	protected:
+		void* GetHandler() const override
+		{
+			return Handler;
+		}
+
+	private:
+		HandlerType* Handler;
+		FHandlerFunction HandlerFunction;
+	};
+
+	/**
+	 * Pointer to the event handler
+	 */
+	typedef IEventHandler* FEventHandlerPtr;
+
+	/**
+	 * Interface for the injected producer function and associated data buffer
+	 */
+	class IEventContainer
+	{
+	public:
+		virtual ~IEventContainer() {}
+		/**
+		 * Register the delegate function that will handle the events on the game thread
+		 */
+		virtual void RegisterHandler(const FEventHandlerPtr& Handler) = 0;
+
+		/**
+		 * Unregister the delegate function that handles the events on the game thread
+		 */
+		virtual void UnregisterHandler(const void* Handler) = 0;
+
+		/*
+		 * Inject data from the physics solver into the producer side of the buffer
+		 */
+		virtual void InjectProducerData(const FPBDRigidsSolver* Solver) = 0;
+
+		/**
+		 * Flips the buffer if the buffer type is double or triple
+		 */
+		virtual void FlipBufferIfRequired() = 0;
+
+		/**
+		 * Dispatch events to the registered handlers
+		 */
+		virtual void DispatchConsumerData() = 0;
+	};
+
+	/**
+	 * Class that owns the injected producer function and its associated data buffer
+	 */
+	template<typename PayloadType>
+	class TEventContainer : public IEventContainer
+	{
+	public:
+		/**
+		 * Regular constructor
+		 */
+		TEventContainer(const Chaos::EMultiBufferMode& BufferMode, TFunction<void(const FPBDRigidsSolver* Solver, PayloadType& EventDataInOut)> InFunction)
+			: InjectedFunction(InFunction)
+			, EventBuffer(Chaos::FMultiBufferFactory<PayloadType>::CreateBuffer(BufferMode))
+		{
+		}
+
+		/**
+		 * Copy constructor
+		 */
+		TEventContainer(TEventContainer& Other)
+		{
+			InjectedFunction = Other.InjectedFunction;
+			EventBuffer = MoveTemp(Other.EventBuffer);
+		}
+
+		/**
+		 * Destructor cleans up memory
+		 */
+		~TEventContainer()
+		{
+			for (FEventHandlerPtr Handler : HandlerArray)
+			{
+				delete Handler;
+				Handler = nullptr;
+			}
+		}
+#
+		/**
+		 * Register the delegate function that will handle the events on the game thread
+		 */
+		virtual void RegisterHandler(const FEventHandlerPtr& Handler)
+		{
+			HandlerArray.AddUnique(Handler);
+		}
+
+		/**
+		 * Unregister the delegate function that handles the events on the game thread
+		 */
+		virtual void UnregisterHandler(const void* InHandler)
+		{
+			for (int i = 0; i < HandlerArray.Num(); i++)
+			{
+				if (HandlerArray[i]->GetHandler() == InHandler)
+				{
+					DeleteHandler(HandlerArray[i]);
+					HandlerArray.RemoveAt(i);
+					break;
+				}
+			}
+		}
+
+		/*
+		 * Inject data from the physics solver into the producer side of the buffer
+		 */
+		virtual void InjectProducerData(const FPBDRigidsSolver* Solver)
+		{
+			InjectedFunction(Solver, *EventBuffer->AccessProducerBuffer());
+		}
+
+		/**
+		 * Flips the buffer if the buffer type is double or triple
+		 */
+		virtual void FlipBufferIfRequired()
+		{
+			EventBuffer->FlipProducer();
+		}
+
+		/**
+		 * Dispatch events to the registered handlers
+		 */
+		virtual void DispatchConsumerData()
+		{
+			for (FEventHandlerPtr Handler : HandlerArray)
+			{
+				Handler->HandleEvent(EventBuffer.Get()->GetConsumerBuffer());
+			}
+		}
+
+	private:
+
+		void DeleteHandler(FEventHandlerPtr& HandlerPtr)
+		{
+			delete HandlerPtr;
+			HandlerPtr = nullptr;
+		}
+
+		/**
+		 * The function that handles filling the event data buffer
+		 */
+		TFunction<void(const FPBDRigidsSolver* Solver, PayloadType& EventData)> InjectedFunction;
+
+		/**
+		 * The data buffer that is filled by the producer and read by the consumer
+		 */
+		TUniquePtr<IBufferResource<PayloadType>> EventBuffer;
+
+		/**
+		 * Delegate function registered to handle this event when it is dispatched
+		 */
+		TArray<FEventHandlerPtr> HandlerArray;
+	};
+
+	/**
+	 * Pointer to event data buffer & injector functionality
+	 */
+	typedef IEventContainer* FEventContainerPtr;
+
+
+	class CHAOSSOLVERS_API FEventManager
+	{
+		friend class FPBDRigidsSolver;
+
+	public:
+
+		FEventManager(const Chaos::EMultiBufferMode& BufferModeIn) : BufferMode(BufferModeIn) {}
+
+		~FEventManager()
+		{
+			Reset();
+		}
+
+		/**
+		 * Clears out every handler and container calling destructors on held items
+		 */
+		void Reset();
+
+		/**
+		 * Set the buffer mode to be used within the event containers
+		 */
+		void SetBufferMode(const Chaos::EMultiBufferMode& BufferModeIn)
+		{
+			BufferMode = BufferModeIn;
+		}
+
+		/**
+		 * Register a new event into the system, providing the function that will fill the producer side of the event buffer
+		 */
+		template<typename PayloadType>
+		void RegisterEvent(const FEventID& EventID, TFunction<void(const Chaos::FPBDRigidsSolver* Solver, PayloadType& EventData)> InFunction)
+		{
+			InternalRegisterInjector(EventID, new TEventContainer<PayloadType>(BufferMode, InFunction));
+		}
+
+		/**
+		 * Unregister specified event from system
+		 */
+		void UnregisterEvent(const FEventID& EventID);
+
+		/**
+		 * Register a handler that will receive the dispatched events
+		 */
+		template<typename PayloadType, typename HandlerType>
+		void RegisterHandler(const FEventID& EventID, HandlerType* Handler, typename TRawEventHandler<PayloadType, HandlerType>::FHandlerFunction HandlerFunction)
+		{
+			checkf(EventID < EventContainers.Num(), TEXT("Registering event Handler for an event ID that does not exist"));
+			EventContainers[EventID]->RegisterHandler(new TRawEventHandler<PayloadType, HandlerType>(Handler, HandlerFunction));
+		}
+
+		/**
+		 * Unregister the specified event handler
+		 */
+		void UnregisterHandler(const FEventID& EventID, const void* InHandler);
+
+		/**
+		 * Called by the solver to invoke the functions that fill the producer side of all the event data buffers
+		 */
+		void FillProducerData(const Chaos::FPBDRigidsSolver* Solver);
+
+		/**
+		 * Flips the event data buffer if it is of double or triple buffer type
+		 */
+		void FlipBuffersIfRequired();
+
+		/**
+		 * // Dispatch events to the registered handlers
+		 */
+		void DispatchEvents();
+
+		/** Returns encoded collision index. */
+		static int32 EncodeCollisionIndex(int32 ActualCollisionIndex, bool bSwapOrder);
+		/** Returns decoded collision index. */
+		static int32 DecodeCollisionIndex(int32 EncodedCollisionIdx, bool& bSwapOrder);
+
+	private:
+
+		void InternalRegisterInjector(const FEventID& EventID, const FEventContainerPtr& Container);
+
+		Chaos::EMultiBufferMode BufferMode;			// specifies the buffer type to be constructed, single, double, triple
+		TArray<FEventContainerPtr> EventContainers;	// Array of event types
+		FRWLock ResourceLock;
+	};
+
+}
\ No newline at end of file
diff --git a/Engine/Source/Runtime/Experimental/ChaosSolvers/Public/Framework/CommandBuffer.h b/Engine/Source/Runtime/Experimental/ChaosSolvers/Public/Framework/CommandBuffer.h
new file mode 100644
--- /dev/null
+++ b/Engine/Source/Runtime/Experimental/ChaosSolvers/Public/Framework/CommandBuffer.h
@@ -0,0 +1,117 @@
+// Copyright 1998-2019 Epic Games, Inc. All Rights Reserved.
+
+#pragma once
+
+#if INCLUDE_CHAOS
+
+#include "Dispatcher.h"
+#include "Threading.h"
+
+namespace Chaos
+{
+	/**
+	 * \class Chaos::FCommandListData
+	 * \brief Command list internal data
+	 *
+	 * Command list internal data, private data for a command list that is submitted to
+	 * the physics dispatcher when the owning FCommandList is flushed (See FCommandList below)
+	 */
+	class FCommandListData
+	{
+	public:
+
+		friend class FCommandList;
+
+		template<enum Chaos::EThreadingMode>
+		friend class FDispatcher;
+		
+		// Element type for solver command storage
+		using FSolverCommandTuple = TTuple<FPhysicsSolver*, IDispatcher::FSolverCommand>;
+
+		// Whether we have any pending commands waiting to submit
+		bool HasCommands() const;
+
+		// Can't copy or directly move list data (command list hands over ptrs to list data)
+		FCommandListData(const FCommandListData& InCopy) = delete;
+		FCommandListData(FCommandListData&& InSteal) = delete;
+		FCommandListData& operator=(const FCommandListData& InCopy) = delete;
+		FCommandListData& operator=(FCommandListData&& InSteal) = delete;
+
+	private:
+
+		// Only friends can create a data entry - only intended to be used by command lists
+		FCommandListData() = default;
+
+		// Storage for pending commands
+		TArray<IDispatcher::FGlobalCommand> GlobalCommands;
+		TArray<IDispatcher::FTaskCommand> TaskCommands;
+		TArray<FSolverCommandTuple> SolverCommands;
+	};
+
+	/**
+	 * \class Chaos::FCommandList
+	 * \brief Physics command list
+	 *
+	 * When performing complex interactions with the physics engine it may be desirable
+	 * to have a batch of commands in a way that ensures all commands are executed together
+	 * before a physics tick. For this case use a command list, enqueueing all commands to the
+	 * list and calling flush when finished.
+	 *
+	 * If you have an object that does this often just keep an FCommandList as a member as each
+	 * call to flush will submit the commands to the physics dispatcher and prepare for a new
+	 * batch, for repeated use this is better than creating a temporary to submit commands.
+	 *
+	 * Example usage:
+	 * FCommandList MyList;
+	 * MyList.Enqueue([](){
+	 *     // Do command body 1
+	 * });
+	 * MyList.Enqueue([](){
+	 *     // Do command body 2
+	 * });
+	 * MyList.Enqueue([](){
+	 *     // Do command body 3
+	 * });
+	 * MyList.Flush();
+	 */
+	class FCommandList
+	{
+	public:
+
+		friend class FChaosSolversModule;
+
+		FCommandList();
+		FCommandList(const FCommandList& InCopy) = delete;
+		FCommandList(FCommandList&& InSteal) = delete;
+		FCommandList& operator=(const FCommandList& InCopy) = delete;
+		FCommandList& operator=(FCommandList&& InSteal) = delete;
+
+		/**
+		 * Enqueue functions for adding commands to be processed
+		 */
+		void Enqueue(IDispatcher::FGlobalCommand&& InCommand);
+		void Enqueue(IDispatcher::FTaskCommand&& InCommand);
+		void Enqueue(FPhysicsSolver* InSolver, IDispatcher::FSolverCommand&& InCommand);
+
+		/**
+		 * Submit this command list to the physics system. This will clear out the command queues in this
+		 * object as they are moved out to the owner to handle the commands. A new empty list will be in
+		 * place to accept further commands
+		 *
+		 * If there are no commands to be executed, no submission or allocation of a new list will be
+		 * performed.
+		 */
+		void Flush();
+
+	private:
+
+		// Sets up the data ptr, used on init and after each flush
+		void AllocData();
+		
+		// The actual commands, on flush this ptr is moved into the dispatcher for executioand a new
+		// list data allocated to service future commands.
+		TUniquePtr<FCommandListData> Data;
+	};
+}
+
+#endif
diff --git a/Engine/Source/Runtime/Experimental/ChaosSolvers/Public/Framework/DebugSolverTasks.h b/Engine/Source/Runtime/Experimental/ChaosSolvers/Public/Framework/DebugSolverTasks.h
--- a/Engine/Source/Runtime/Experimental/ChaosSolvers/Public/Framework/DebugSolverTasks.h
+++ b/Engine/Source/Runtime/Experimental/ChaosSolvers/Public/Framework/DebugSolverTasks.h
@@ -6,11 +6,7 @@
 
 #include "Chaos/Framework/DebugSubstep.h"
 #include "Templates/Function.h"
-
-namespace Chaos
-{
-	class FPBDRigidsSolver;
-}  // namespace Chaos
+#include "Chaos/Declares.h"
 
 #if CHAOS_DEBUG_SUBSTEP
 #include "Async/AsyncWork.h"
@@ -55,19 +51,19 @@ namespace Chaos
 		~FDebugSolverTasks() { Shutdown(); }
 
 		/** Add debug task entry for the specified solver. */
-		void Add(FPBDRigidsSolver* Solver);
+		void Add(FPhysicsSolver* Solver);
 
 		/** Remove the debug task entry for the specified solver, and delete its task if any was created. */
-		void Remove(FPBDRigidsSolver* Solver);
+		void Remove(FPhysicsSolver* Solver);
 
 		/** Run the specified step function in one go within the current thread, or in a debug thread substep by substep depending on the Solver's DebugSustep status. */
-		void DebugStep(FPBDRigidsSolver* Solver, TFunction<void()> StepFunction);
+		void DebugStep(FPhysicsSolver* Solver, TFunction<void()> StepFunction);
 
 		/** Shutdown all debug threads. */
 		void Shutdown();
 
 	private:
-		TMap<FPBDRigidsSolver*, FAsyncTask<FDebugSolverTask>*> SolverToTaskMap;
+		TMap<FPhysicsSolver*, FAsyncTask<FDebugSolverTask>*> SolverToTaskMap;
 	};
 }  // namespace Chaos
 
@@ -84,10 +80,10 @@ namespace Chaos
 		FDebugSolverTasks() {}
 		~FDebugSolverTasks() {}
 
-		void Add(FPBDRigidsSolver* /*Solver*/) {}
-		void Remove(FPBDRigidsSolver* /*Solver*/) {}
+		void Add(FPhysicsSolver* /*Solver*/) {}
+		void Remove(FPhysicsSolver* /*Solver*/) {}
 
-		FORCEINLINE void DebugStep(FPBDRigidsSolver* /*Solver*/, TFunction<void()> StepFunction) { StepFunction(); }
+		FORCEINLINE void DebugStep(FPhysicsSolver* /*Solver*/, TFunction<void()> StepFunction) { StepFunction(); }
 
 		void Shutdown() {}
 	};
diff --git a/Engine/Source/Runtime/Experimental/ChaosSolvers/Public/Framework/Dispatcher.h b/Engine/Source/Runtime/Experimental/ChaosSolvers/Public/Framework/Dispatcher.h
--- a/Engine/Source/Runtime/Experimental/ChaosSolvers/Public/Framework/Dispatcher.h
+++ b/Engine/Source/Runtime/Experimental/ChaosSolvers/Public/Framework/Dispatcher.h
@@ -7,73 +7,88 @@
 #include "CoreMinimal.h"
 #include "Containers/Queue.h"
 #include "Threading.h"
+#include "Chaos/Declares.h"
 
 class FChaosSolversModule;
 class FPhysicsCommandsTask;
 
 namespace Chaos
 {
-	class FPBDRigidsSolver;
 	class FPersistentPhysicsTask;
+	class FCommandList;
+	class FCommandListData;
 }
 
 namespace Chaos
 {
+	/**
+	 * Physics command dispatcher
+	 *
+	 * Used to safely interact with Chaos physics data. When performing any operation that needs to affect live
+	 * physics data such as modifying simulation data (positions, velocities etc.) make sure that this code
+	 * is placed with in an appropriate command.
+	 *
+	 * Command Types:
+	 *
+	 *		* Global -  First commands to run in a batch, no parameters, if your command doesn't interact
+	 *					with a solver, this is probably what you need. This will run at the beginning of
+	 *					a physics tick before task or solver commands
+	 *
+	 *		* Task -	Ran immediately after the global commands, takes a ptr to the actual physics task
+	 *					that is running if we are in dedicated thread mode, nullptr otherwise. Although this
+	 *					command type knows about the underlying threading model care should be taken when using
+	 *					this command type to make sure it also accomplishes its goal in non-dedicated threading modes.
+	 *					Only use if absolutely necessary
+	 *
+	 *		* Solver -	Another useful command, is bound to a solver. Runs after the task commands have all run.
+	 *					Maybe all be ran in different threads. The solver pointer passed to the command is safe
+	 *					to read and write from (as should be any captured object handles) but do not attempt to
+	 *					access other solvers, or other off-thread data as this will not be safe when running in
+	 *					any threading model that is not single-threaded
+	 *
+	 * A note on batched commands and command lists:
+	 *		Batched commands and command lists are consistent with the above ordering of commands but only within
+	 *		that specific batch. If two batched A and B are submitted then the execution order of the commands is:
+	 *				A.Global -> A.Task -> A.Solver -> B.Global -> B.Task -> B.Solver
+	 */
 	class IDispatcher
 	{
 	public:
 		virtual ~IDispatcher() {}
 
-		virtual void EnqueueCommand(TFunction<void()> InCommand) = 0;
-		virtual void EnqueueCommand(TFunction<void(FPersistentPhysicsTask*)> InCommand) = 0;
-		virtual void EnqueueCommand(FPBDRigidsSolver* InSolver, TFunction<void(FPBDRigidsSolver*)> InCommand) const = 0;
+		using FGlobalCommand = TFunction<void()>;
+		using FTaskCommand = TFunction<void(FPersistentPhysicsTask*)>;
+		using FSolverCommand = TFunction<void(FPhysicsSolver*)>;
+
+		/**
+		 * Immediate commands:
+		 * Enqueueing an immediate command will run that command at the next available opportunity when
+		 * the physics scene is next ticked. Note that in some threading models commands enqueued immediately
+		 * after one another on a different thread may get executed in different physics frames.
+		 * If this is not desirable consider either batched commands or submitting a custom command list
+		 */
+		virtual void EnqueueCommandImmediate(FGlobalCommand InCommand) = 0;
+		virtual void EnqueueCommandImmediate(FTaskCommand InCommand) = 0;
+		virtual void EnqueueCommandImmediate(FPhysicsSolver* InSolver, FSolverCommand InCommand) = 0;
+
+		/**
+		 * Get the current threading mode for this dispatcher
+		 */
 		virtual EThreadingMode GetMode() const = 0;
-	};
-
-	template<EThreadingMode Mode>
-	class CHAOSSOLVERS_API FDispatcher : public IDispatcher
-	{
-	public:
-		friend class FPersistentPhysicsTask;
-		friend class ::FPhysicsCommandsTask;
-		friend class FPhysScene_ChaosInterface;
-
-		explicit FDispatcher(FChaosSolversModule* InOwnerModule)
-			: Owner(InOwnerModule)
-		{}
 
-		virtual void EnqueueCommand(TFunction<void()> InCommand) final override;
-		virtual void EnqueueCommand(TFunction<void(FPersistentPhysicsTask*)> InCommand) final override;
-		virtual void EnqueueCommand(FPBDRigidsSolver* InSolver, TFunction<void(FPBDRigidsSolver*)> InCommand) const final override;
-		virtual EThreadingMode GetMode() const final override { return Mode; }
-
-	private:
-		FChaosSolversModule* Owner;
-
-		TQueue<TFunction<void()>, EQueueMode::Mpsc> GlobalCommandQueue;
-		TQueue<TFunction<void(FPersistentPhysicsTask*)>, EQueueMode::Mpsc> TaskCommandQueue;
+		/**
+		 * Given a command list, submit it as a batch for execution on the next physics frame
+		 */
+		virtual void SubmitCommandList(TUniquePtr<FCommandListData>&& InCommandData) = 0;
+
+		/**
+		 * Execute any pending submitted command lists along with global and task commands
+		 * Intended to be called by whatever code is responsible for updating the physics scene
+		 * #BG - Make solver commands execute here too instead of 2 execution sites (Needs parallel consideration).
+		 */
+		virtual void Execute() = 0;
 	};
 
-	template<>
-	void FDispatcher<EThreadingMode::DedicatedThread>::EnqueueCommand(TFunction<void()> InCommand);
-	template<>
-	void FDispatcher<EThreadingMode::DedicatedThread>::EnqueueCommand(TFunction<void(FPersistentPhysicsTask*)> InCommand);
-	template<>
-	void FDispatcher<EThreadingMode::DedicatedThread>::EnqueueCommand(FPBDRigidsSolver* InSolver, TFunction<void(FPBDRigidsSolver *)> InCommand) const;
-
-	template<>
-	void FDispatcher<EThreadingMode::SingleThread>::EnqueueCommand(TFunction<void()> InCommand);
-	template<>
-	void FDispatcher<EThreadingMode::SingleThread>::EnqueueCommand(TFunction<void(FPersistentPhysicsTask*)> InCommand);
-	template<>
-	void FDispatcher<EThreadingMode::SingleThread>::EnqueueCommand(FPBDRigidsSolver* InSolver, TFunction<void(FPBDRigidsSolver *)> InCommand) const;
-
-	template<>
-	void FDispatcher<EThreadingMode::DedicatedThread>::EnqueueCommand(TFunction<void()> InCommand);
-	template<>
-	void FDispatcher<EThreadingMode::DedicatedThread>::EnqueueCommand(TFunction<void(FPersistentPhysicsTask*)> InCommand);
-	template<>
-	void FDispatcher<EThreadingMode::DedicatedThread>::EnqueueCommand(FPBDRigidsSolver* InSolver, TFunction<void(FPBDRigidsSolver *)> InCommand) const;
 }
 
 void LexFromString(Chaos::EThreadingMode& OutValue, const TCHAR* InString);
diff --git a/Engine/Source/Runtime/Experimental/ChaosSolvers/Public/Framework/PersistentTask.h b/Engine/Source/Runtime/Experimental/ChaosSolvers/Public/Framework/PersistentTask.h
--- a/Engine/Source/Runtime/Experimental/ChaosSolvers/Public/Framework/PersistentTask.h
+++ b/Engine/Source/Runtime/Experimental/ChaosSolvers/Public/Framework/PersistentTask.h
@@ -11,6 +11,7 @@
 #include "Chaos/Framework/BufferedData.h"
 #include "Framework/DebugSolverTasks.h"
 #include "Chaos/PBDRigidsEvolution.h"
+#include "Chaos/Declares.h"
 
 #define CHAOSTHREADSTATS_ENABLED 1
 #define CHAOSTHREADSTATS_PERSOLVER 1
@@ -21,7 +22,6 @@ struct FSolverStateStorage;
 
 namespace Chaos
 {
-	class FPBDRigidsSolver;
 	class FPhysicsProxy;
 	class ITimeStep;
 }
@@ -114,7 +114,7 @@ namespace Chaos
 		friend class FAsyncTask<FPersistentPhysicsTask>;
 
 	public:
-		FPersistentPhysicsTask(float InTargetDt, bool bInAvoidSpiral, FDispatcher<EThreadingMode::DedicatedThread>* InDispatcher);
+		FPersistentPhysicsTask(float InTargetDt, bool bInAvoidSpiral, IDispatcher* InDispatcher);
 		virtual ~FPersistentPhysicsTask();
 
 		/**
@@ -129,12 +129,12 @@ namespace Chaos
 		 * Once the solver has been added to this task the game thread should never
 		 * touch the internal state again unless performing a sync of the data
 		 */
-		void AddSolver(FPBDRigidsSolver* InSolver);
+		void AddSolver(FPhysicsSolver* InSolver);
 
 		/**
 		 * Removes a solver from the internal list of solvers to run on the async task
 		 */
-		void RemoveSolver(FPBDRigidsSolver* InSolver);
+		void RemoveSolver(FPhysicsSolver* InSolver);
 
 		/**
 		 * Synchronize proxies to their most recent gamethread readable results
@@ -201,12 +201,12 @@ namespace Chaos
 			RETURN_QUICK_DECLARE_CYCLE_STAT(FPersistentPhysicsTask, STATGROUP_ThreadPoolAsyncTasks);
 		}
 
-		void StepSolver(FPBDRigidsSolver* InSolver, float InDt);
-		void HandleSolverCommands(FPBDRigidsSolver* InSolver);
-		void AdvanceSolver(FPBDRigidsSolver* InSolver, float InDt);
+		void StepSolver(FPhysicsSolver* InSolver, float InDt);
+		void HandleSolverCommands(FPhysicsSolver* InSolver);
+		void AdvanceSolver(FPhysicsSolver* InSolver, float InDt);
 
 		// List of solvers we'll advance in this task
-		TArray<FPBDRigidsSolver*> Solvers;
+		TArray<FPhysicsSolver*> Solvers;
 
 		// Debug threads used to debug substep solver advance.
 		FDebugSolverTasks DebugSolverTasks;
@@ -222,7 +222,7 @@ namespace Chaos
 		FThreadSafeBool bRunning;
 
 		// The dispatcher made by the Chaos module to enable the gamethread to communicate with this one.
-		FDispatcher<EThreadingMode::DedicatedThread>* CommandDispatcher;
+		IDispatcher* CommandDispatcher;
 
 		// Event to fire after we've broken from the running physics loop as the thread shuts down
 		FEvent* ShutdownEvent;
diff --git a/Engine/Source/Runtime/Experimental/ChaosSolvers/Public/Framework/PhysicsTickTask.h b/Engine/Source/Runtime/Experimental/ChaosSolvers/Public/Framework/PhysicsTickTask.h
--- a/Engine/Source/Runtime/Experimental/ChaosSolvers/Public/Framework/PhysicsTickTask.h
+++ b/Engine/Source/Runtime/Experimental/ChaosSolvers/Public/Framework/PhysicsTickTask.h
@@ -4,14 +4,10 @@
 
 #include "Async/TaskGraphInterfaces.h"
 #include "Dispatcher.h"
+#include "Chaos/Declares.h"
 
 class FChaosSolversModule;
 
-namespace Chaos
-{
-	class FPBDRigidsSolver;
-}
-
 /**
  * Task responsible for handling a full frame update for physics under the TaskGraph threading
  * mode.
@@ -52,7 +48,7 @@ public:
 private:
 
 	FChaosSolversModule* Module;
-	Chaos::FDispatcher<EChaosThreadingMode::TaskGraph>* Dispatcher;
+	Chaos::IDispatcher* Dispatcher;
 };
 
 /**
@@ -63,7 +59,7 @@ class CHAOSSOLVERS_API FPhysicsSolverAdvanceTask
 {
 public:
 
-	FPhysicsSolverAdvanceTask(Chaos::FPBDRigidsSolver* InSolver, float InDt);
+	FPhysicsSolverAdvanceTask(Chaos::FPhysicsSolver* InSolver, float InDt);
 
 	TStatId GetStatId() const;
 	static ENamedThreads::Type GetDesiredThread();
@@ -72,9 +68,9 @@ public:
 
 private:
 
-	void StepSolver(Chaos::FPBDRigidsSolver* InSolver, float InDt);
+	void StepSolver(Chaos::FPhysicsSolver* InSolver, float InDt);
 
-	Chaos::FPBDRigidsSolver* Solver;
+	Chaos::FPhysicsSolver* Solver;
 	float Dt;
 };
 
@@ -86,7 +82,7 @@ class CHAOSSOLVERS_API FPhysicsSolverAdvanceSubsteppingTask
 {
 public:
 
-	FPhysicsSolverAdvanceSubsteppingTask(Chaos::FPBDRigidsSolver* InSolver, int32 NumIterations, float InDtPerIteration);
+	FPhysicsSolverAdvanceSubsteppingTask(Chaos::FPhysicsSolver* InSolver, int32 NumIterations, float InDtPerIteration);
 
 	TStatId GetStatId() const;
 	static ENamedThreads::Type GetDesiredThread();
@@ -95,7 +91,7 @@ public:
 
 private:
 
-	void StepSolver(Chaos::FPBDRigidsSolver* InSolver, float Dt);
+	void StepSolver(Chaos::FPhysicsSolver* InSolver, float Dt);
 
 	FGraphEventRef CompletionEvent;
 	FChaosSolversModule* Module;
