Copying //UE4/Dev-Physics to Dev-Main (//UE4/Dev-Main) [at] 7703071

diff --git a/Engine/Source/Runtime/Experimental/ChaosSolvers/Public/PBDRigidsSolver.h b/Engine/Source/Runtime/Experimental/ChaosSolvers/Public/PBDRigidsSolver.h
--- a/Engine/Source/Runtime/Experimental/ChaosSolvers/Public/PBDRigidsSolver.h
+++ b/Engine/Source/Runtime/Experimental/ChaosSolvers/Public/PBDRigidsSolver.h
@@ -4,879 +4,292 @@
 
 #include "Chaos/ArrayCollectionArray.h"
 #include "Chaos/Defines.h"
+#include "Chaos/Framework/PhysicsProxy.h"
+#include "Chaos/Framework/MultiBufferResource.h"
+#include "Chaos/Framework/SingleParticlePhysicsProxy.h"
 #include "Chaos/PBDRigidParticles.h"
-#include "Chaos/PBDRigidsEvolution.h"
+#include "Chaos/PBDRigidsEvolutionGBF.h"
 #include "Chaos/PBDCollisionConstraint.h"
 #include "Chaos/PBDRigidDynamicSpringConstraints.h"
 #include "Chaos/PBDPositionConstraints.h"
 #include "Chaos/PBDJointConstraints.h"
-#include "Chaos/PBDSpringConstraints.h"
 #include "Chaos/PBDConstraintRule.h"
 #include "Chaos/PerParticleGravity.h"
+#include "Chaos/ParticleHandle.h"
 #include "Chaos/Transform.h"
 #include "CoreMinimal.h"
 #include "Containers/Queue.h"
 #include "Framework/Dispatcher.h"
 #include "Field/FieldSystem.h"
-#include "SolverObjects/SolverObject.h"
-#include "MultiBufferResource.h"
+#include "EventManager.h"
 #include "SolverEventFilters.h"
 
 #ifndef CHAOS_WITH_PAUSABLE_SOLVER
 #define CHAOS_WITH_PAUSABLE_SOLVER 1
 #endif
 
-class FSkeletalMeshPhysicsObject;
-class FStaticMeshPhysicsObject;
-class FBodyInstancePhysicsObject;
-class FGeometryCollectionPhysicsObject;
-class FFieldSystemPhysicsObject;
 class FPhysicsSolverAdvanceTask;
-
-#define USE_PGS 0
-
-#if USE_PGS
-typedef Chaos::TPBDRigidsEvolutionPGS<float, 3> FPBDRigidsEvolution;
-typedef Chaos::TPBDCollisionConstraintPGS<float, 3> FPBDCollisionConstraints;
-#else
-typedef Chaos::TPBDRigidsEvolutionGBF<float, 3> FPBDRigidsEvolution;
-typedef Chaos::TPBDCollisionConstraint<float, 3> FPBDCollisionConstraints;
-#endif
-
 class FPhysInterface_Chaos;
-namespace Chaos
-{
-	class FPBDRigidsSolver;
-	class AdvanceOneTimeStepTask;
-	class FPersistentPhysicsTask;
-	class FPhysicsCommand;
-};
 
-struct FKinematicProxy
-{
-	TArray<int32> Ids;
-	TArray<FVector> Position;
-	TArray<FQuat> Rotation;
-	TArray<FVector> NextPosition;
-	TArray<FQuat> NextRotation;
-};
-
-struct CHAOSSOLVERS_API FSolverObjectStorage
-{
-	void Reset();
-
-	// Call helper for TSolverObject items, use with a generic lambda of form [](auto* Object){}
-	template<typename Lambda>
-	void ForEachSolverObject(Lambda InCallable)
-	{
-		for(TSolverObject<FGeometryCollectionPhysicsObject>* Obj : GeometryCollectionObjects)
-		{
-			InCallable(Obj);
-		}
-		
-		for(TSolverObject<FSkeletalMeshPhysicsObject>* Obj : SkeletalMeshObjects)
-		{
-			InCallable(Obj);
-		}
-		
-		for(TSolverObject<FStaticMeshPhysicsObject>* Obj : StaticMeshObjects)
-		{
-			InCallable(Obj);
-		}
-
-		for (TSolverObject<FBodyInstancePhysicsObject>* Obj : BodyInstanceObjects)
-		{
-			InCallable(Obj);
-		}
-
-	}
-
-	template<typename Lambda>
-	void ForEachSolverObjectParallel(Lambda InCallable)
-	{
-		Chaos::PhysicsParallelFor(GeometryCollectionObjects.Num(), [this, &InCallable](const int32 Index)
-		{
-			TSolverObject<FGeometryCollectionPhysicsObject>* Obj = GeometryCollectionObjects[Index];
-			InCallable(Obj);
-		});
-
-		Chaos::PhysicsParallelFor(SkeletalMeshObjects.Num(), [this, &InCallable](const int32 Index)
-		{
-			TSolverObject<FSkeletalMeshPhysicsObject>* Obj = SkeletalMeshObjects[Index];
-			InCallable(Obj);
-		});
-
-		Chaos::PhysicsParallelFor(StaticMeshObjects.Num(), [this, &InCallable](const int32 Index)
-		{
-			TSolverObject<FStaticMeshPhysicsObject>* Obj = StaticMeshObjects[Index];
-			InCallable(Obj);
-		});
-
-		Chaos::PhysicsParallelFor(BodyInstanceObjects.Num(), [this, &InCallable](const int32 Index)
-		{
-			TSolverObject<FBodyInstancePhysicsObject>* Obj = BodyInstanceObjects[Index];
-			InCallable(Obj);
-		});
-	}
-
-	template<typename Lambda>
-	void ForEachFieldSolverObject(Lambda InCallable)
-	{
-		for(FFieldSystemPhysicsObject* Obj : FieldSystemObjects)
-		{
-			InCallable(Obj);
-		}
-	}
-
-	template<typename ObjectType>
-	TArray<TSolverObject<ObjectType>*>& GetObjectStorage();
-
-	template<typename ObjectType>
-	TArray<ObjectType*>& GetFieldObjectStorage();
-
-	int32 GetNumObjects() const;
-
-	int32 GetNumFieldObjects() const;
-
-	template<typename ObjectType>
-	int32 GetNumObjectsOfType() const;
-
-	TArray<TSolverObject<FGeometryCollectionPhysicsObject>*> GeometryCollectionObjects;
-	TArray<TSolverObject<FSkeletalMeshPhysicsObject>*> SkeletalMeshObjects;
-	TArray<TSolverObject<FStaticMeshPhysicsObject>*> StaticMeshObjects;
-	TArray<TSolverObject<FBodyInstancePhysicsObject>*> BodyInstanceObjects;
-	TArray<FFieldSystemPhysicsObject*> FieldSystemObjects;
-};
-
-template<> TArray<TSolverObject<FGeometryCollectionPhysicsObject>*>& FSolverObjectStorage::GetObjectStorage();
-template<> TArray<TSolverObject<FSkeletalMeshPhysicsObject>*>& FSolverObjectStorage::GetObjectStorage();
-template<> TArray<TSolverObject<FStaticMeshPhysicsObject>*>& FSolverObjectStorage::GetObjectStorage();
-template<> TArray<TSolverObject<FBodyInstancePhysicsObject>*>& FSolverObjectStorage::GetObjectStorage();
-template<> TArray<FFieldSystemPhysicsObject*>& FSolverObjectStorage::GetFieldObjectStorage();
-template<> int32 FSolverObjectStorage::GetNumObjectsOfType<FGeometryCollectionPhysicsObject>() const;
-template<> int32 FSolverObjectStorage::GetNumObjectsOfType<FSkeletalMeshPhysicsObject>() const;
-template<> int32 FSolverObjectStorage::GetNumObjectsOfType<FStaticMeshPhysicsObject>() const;
-template<> int32 FSolverObjectStorage::GetNumObjectsOfType<FBodyInstancePhysicsObject>() const;
-template<> int32 FSolverObjectStorage::GetNumObjectsOfType<FFieldSystemPhysicsObject>() const;
+class FSkeletalMeshPhysicsProxy;
+class FStaticMeshPhysicsProxy;
+class FGeometryCollectionPhysicsProxy;
+class FFieldSystemPhysicsProxy;
 
 /**
 *
 */
 namespace Chaos
 {
+	class AdvanceOneTimeStepTask;
+	class FPersistentPhysicsTask;
+	class FPhysicsCommand;
 	class FChaosArchive;
 	/**
 	*
 	*/
 	class CHAOSSOLVERS_API FPBDRigidsSolver
 	{
-	public:
-		typedef TPBDRigidParticles<float, 3> FParticlesType;
-
-		typedef FPBDCollisionConstraints FCollisionConstraints;
-		typedef TPBDJointConstraints<float, 3> FJointConstraints;
-		typedef TPBDRigidDynamicSpringConstraints<float, 3> FRigidDynamicSpringConstraints;
-		typedef TPBDPositionConstraints<float, 3> FPositionConstraints;
 
-		typedef TPBDConstraintIslandRule<FJointConstraints, float, 3> FJointConstraintsRule;
-		typedef TPBDConstraintIslandRule<FRigidDynamicSpringConstraints, float, 3> FRigidDynamicSpringConstraintsRule;
-		typedef TPBDConstraintIslandRule<FPositionConstraints, float, 3> FPositionConstraintsRule;
+		FPBDRigidsSolver(const EMultiBufferMode BufferingModeIn);
 
-		friend class AdvanceOneTimeStepTask;
+	public:
 		friend class FPersistentPhysicsTask;
 		friend class ::FPhysicsSolverAdvanceTask;
+		friend class ::FChaosSolversModule;
 
 		template<EThreadingMode Mode>
 		friend class FDispatcher;
+		friend class FEventDefaults;
 
 		friend class FPhysInterface_Chaos;
 		friend class FPhysScene_ChaosInterface;
 
-		static int8 Invalid;
-	
-		friend struct FAccessor;
-		friend class FScopedGetEventsData;
+		void* PhysSceneHack;	//This is a total hack for now to get at the owning scene
 
-		/* Access to details */
-		// This is required for GCC which doesn't allow friends across dlls
-		struct CHAOSSOLVERS_API FAccessor
-		{
-			FAccessor(FPBDRigidsSolver* InSolver) { Solver = InSolver; }
-
-			static FPBDRigidsSolver* CreateSolver(const EMultiBufferMode BufferingModeIn = EMultiBufferMode::Double)
-			{
-				return new FPBDRigidsSolver(BufferingModeIn);
-			}
+		typedef TPBDRigidsSOAs<float, 3> FParticlesType;
 
-			TPBDRigidDynamicSpringConstraints<float, 3>& DynamicConstraints() { return Solver->DynamicConstraints; }
-			TSet<int32>& DynamicConstraintParticles() { return Solver->DynamicConstraintParticles; }
+		typedef Chaos::TGeometryParticle<float, 3> FParticle;
+		typedef Chaos::TGeometryParticleHandle<float, 3> FHandle;
+		typedef Chaos::TPBDRigidsEvolutionGBF<float, 3> FPBDRigidsEvolution;
+		typedef Chaos::TPBDCollisionConstraint<float, 3> FPBDCollisionConstraints;
 
-		private:
-			FPBDRigidsSolver* Solver;
-		};
-
-		/* General maps reverse lookups */
-		struct FTimeResource
-		{
-			FTimeResource() : TimeCreated(-FLT_MAX) {}
-			float TimeCreated;
-		};
-
-		// SolverObjectReverseMapping for one frame time stamped with the time for that frame  
-		struct FSolverObjectReverseMapping : FTimeResource
-		{
-			FSolverObjectReverseMapping()
-				: SolverObjectReverseMappingArray(TArray<SolverObjectWrapper>())
-			{}
+		typedef FPBDCollisionConstraints FCollisionConstraints;
+		typedef TPBDJointConstraints<float, 3> FJointConstraints;
+		typedef TPBDRigidDynamicSpringConstraints<float, 3> FRigidDynamicSpringConstraints;
+		typedef TPBDPositionConstraints<float, 3> FPositionConstraints;
 
-			void Reset()
-			{
-				SolverObjectReverseMappingArray.Reset();
-			}
+		typedef TPBDConstraintIslandRule<FJointConstraints, float, 3> FJointConstraintsRule;
+		typedef TPBDConstraintIslandRule<FRigidDynamicSpringConstraints, float, 3> FRigidDynamicSpringConstraintsRule;
+		typedef TPBDConstraintIslandRule<FPositionConstraints, float, 3> FPositionConstraintsRule;
 
-			/** Maps ParticleIndex to SolverObject.  SolverObjectReverseMappingArray[ParticleIndex] = SolverObject */
-			TArray<SolverObjectWrapper> SolverObjectReverseMappingArray; 
-		};
+		//
+		// Execution API
+		//
 
-		// ParticleIndexReverseMapping for one frame time stamped with the time for that frame
-		struct FParticleIndexReverseMapping : FTimeResource
-		{
-			FParticleIndexReverseMapping()
-				: ParticleIndexReverseMappingArray(TArray<int32>())
-			{}
+		void ChangeBufferMode(Chaos::EMultiBufferMode InBufferMode);
+		TQueue<TFunction<void(FPBDRigidsSolver*)>, EQueueMode::Mpsc>& GetCommandQueue() { return CommandQueue; }
 
-			void Reset()
-			{
-				ParticleIndexReverseMappingArray.Reset();
-			}
 
-			TArray<int32> ParticleIndexReverseMappingArray; // ParticleIndex -> TransformIndex
-		};
+		//
+		//  Object API
+		//
 
-		/* Collision */
+		void RegisterObject(Chaos::TGeometryParticle<float, 3>* GTParticle);
+		void UnregisterObject(Chaos::TGeometryParticle<float, 3>* GTParticle);
 
-		typedef TArray<TCollisionData<float, 3>> FCollisionDataArray;
+		// TODO: Set up an interface for registering fields and geometry collections
+		//void RegisterObject(FGeometryCollectionPhysicsProxy* InProxy);
+		//int UnregisterObject(FGeometryCollectionPhysicsProxy* InProxy);
+		//void RegisterObject(FFieldSystemPhysicsProxy* InProxy);
+		//int UnregisterObject(FFieldSystemPhysicsProxy* InProxy);
 
-		struct FCollisionData
-		{
-			float TimeCreated;
-			int32 NumCollisions;
-			FCollisionDataArray CollisionDataArray;
-		};
+		bool IsSimulating() const;
 
-		// All the collisions for one frame time stamped with the time for that frame  
-		struct FAllCollisionData : public FTimeResource
+		template<typename Lambda>
+		void ForEachPhysicsProxy(Lambda InCallable)
 		{
-			FAllCollisionData()
-				: AllCollisionsArray(FCollisionDataArray())
-			{}
-
-			void Reset()
+			for (FGeometryParticlePhysicsProxy* Obj : GeometryParticlePhysicsProxies)
 			{
-				AllCollisionsArray.Reset();
+				InCallable(Obj);
 			}
-
-			FCollisionDataArray AllCollisionsArray;
-		};
-
-		// AllCollisionsIndicesBySolverObject for one frame time stamped with the time for that frame
-		struct FAllCollisionsIndicesBySolverObject : public FTimeResource
-		{
-			FAllCollisionsIndicesBySolverObject()
-				: AllCollisionsIndicesBySolverObjectMap(TMap<ISolverObjectBase*, TArray<int32>>())
-			{}
-
-			void Reset()
+			for (FKinematicGeometryParticlePhysicsProxy* Obj : KinematicGeometryParticlePhysicsProxies)
 			{
-				AllCollisionsIndicesBySolverObjectMap.Reset();
+				InCallable(Obj);
 			}
-
-			TMap<ISolverObjectBase*, TArray<int32>> AllCollisionsIndicesBySolverObjectMap; // SolverObject -> Indices in AllCollisions array
-		};
-
-		// All the collisions with all the maps
-		struct FAllCollisionDataMaps
-		{
-			FAllCollisionDataMaps()
-				: bIsDataAndMapsInSync(false)
-				, AllCollisionData(nullptr)
-				, SolverObjectReverseMapping(nullptr)
-				, ParticleIndexReverseMapping(nullptr)
-				, AllCollisionsIndicesBySolverObject(nullptr)
-			{}
-
-			bool IsValid() const 
+			for (FRigidParticlePhysicsProxy* Obj : RigidParticlePhysicsProxies)
 			{
-				return bIsDataAndMapsInSync && AllCollisionData && SolverObjectReverseMapping && ParticleIndexReverseMapping && AllCollisionsIndicesBySolverObject; 
+				InCallable(Obj);
 			}
-
-			bool bIsDataAndMapsInSync;
-			const FAllCollisionData* AllCollisionData;
-			const FSolverObjectReverseMapping* SolverObjectReverseMapping;
-			const FParticleIndexReverseMapping* ParticleIndexReverseMapping;
-			const FAllCollisionsIndicesBySolverObject* AllCollisionsIndicesBySolverObject;
-		};
-
-		/* Breaking */
-
-		typedef TArray<TBreakingData<float, 3>> FBreakingDataArray;
-		struct FBreakingData
-		{
-			float TimeCreated;
-			int32 NumBreakings;
-			FBreakingDataArray BreakingDataArray;
-		};
-
-		// All the breakings for one frame time stamped with the time for that frame  
-		struct FAllBreakingData : public FTimeResource
-		{
-			FAllBreakingData()
-				: AllBreakingsArray(FBreakingDataArray())
-			{}
-
-			void Reset()
+			for (FSkeletalMeshPhysicsProxy* Obj : SkeletalMeshPhysicsProxies)
 			{
-				AllBreakingsArray.Reset();
+				InCallable(Obj);
 			}
-
-			FBreakingDataArray AllBreakingsArray;
-		};
-
-		// AllBreakingsIndicesBySolverObject for one frame time stamped with the time for that frame
-		struct FAllBreakingsIndicesBySolverObject : public FTimeResource
-		{
-			FAllBreakingsIndicesBySolverObject()
-				: AllBreakingsIndicesBySolverObjectMap(TMap<ISolverObjectBase*, TArray<int32>>())
-			{}
-
-			void Reset()
+			for (FStaticMeshPhysicsProxy* Obj : StaticMeshPhysicsProxies)
 			{
-				AllBreakingsIndicesBySolverObjectMap.Reset();
+				InCallable(Obj);
 			}
-
-			TMap<ISolverObjectBase*, TArray<int32>> AllBreakingsIndicesBySolverObjectMap; // SolverObject -> Indices in AllBreakings array
-		};
-
-		// All the Breakings with all the maps
-		struct FAllBreakingDataMaps
-		{
-			FAllBreakingDataMaps()
-				: bIsDataAndMapsInSync(false)
-				, AllBreakingData(nullptr)
-				, SolverObjectReverseMapping(nullptr)
-				, ParticleIndexReverseMapping(nullptr)
-				, AllBreakingsIndicesBySolverObject(nullptr)
-			{}
-
-			bool IsValid() const { return bIsDataAndMapsInSync; }
-
-			bool bIsDataAndMapsInSync;
-			const FAllBreakingData* AllBreakingData;
-			const FSolverObjectReverseMapping* SolverObjectReverseMapping;
-			const FParticleIndexReverseMapping* ParticleIndexReverseMapping;
-			const FAllBreakingsIndicesBySolverObject* AllBreakingsIndicesBySolverObject;
-		};
-
-		/* Trailing */
-
-		typedef TSet<TTrailingData<float, 3>> FTrailingDataSet;
-		typedef TArray<TTrailingData<float, 3>> FTrailingDataArray;
-		struct FTrailingData
-		{
-			float TimeLastUpdated;
-			FTrailingDataSet TrailingDataSet;
-		};
-		typedef TPBDRigidClustering<FPBDRigidsEvolution, FCollisionConstraints, float, 3> FClusteringType;
-
-		// All the trailings for one frame time stamped with the time for that frame  
-		struct FAllTrailingData : FTimeResource
-		{
-			FAllTrailingData()
-				: AllTrailingsArray(FTrailingDataArray())
-			{}
-
-			void Reset()
+			for (FGeometryCollectionPhysicsProxy* Obj : GeometryCollectionPhysicsProxies)
 			{
-				AllTrailingsArray.Reset();
+				InCallable(Obj);
 			}
-
-			FTrailingDataArray AllTrailingsArray;
-		};
-
-		// AllTrailingsIndicesBySolverObject for one frame time stamped with the time for that frame
-		struct FAllTrailingsIndicesBySolverObject : public FTimeResource
-		{
-			FAllTrailingsIndicesBySolverObject()
-				: AllTrailingsIndicesBySolverObjectMap(TMap<ISolverObjectBase*, TArray<int32>>())
-			{}
-
-			void Reset()
+			for (FFieldSystemPhysicsProxy* Obj : FieldSystemPhysicsProxies)
 			{
-				AllTrailingsIndicesBySolverObjectMap.Reset();
+				InCallable(Obj);
 			}
+		}
 
-			TMap<ISolverObjectBase*, TArray<int32>> AllTrailingsIndicesBySolverObjectMap; // SolverObject -> Indices in AllTrailings array
-		};
-
-		// All the Trailings with all the maps
-		struct FAllTrailingDataMaps
-		{
-			FAllTrailingDataMaps()
-				: bIsDataAndMapsInSync(false)
-				, AllTrailingData(nullptr)
-				, SolverObjectReverseMapping(nullptr)
-				, ParticleIndexReverseMapping(nullptr)
-				, AllTrailingsIndicesBySolverObject(nullptr)
-			{}
-
-			bool IsValid() const { return bIsDataAndMapsInSync; }
-
-			bool bIsDataAndMapsInSync;
-			const FAllTrailingData* AllTrailingData;
-			const FSolverObjectReverseMapping* SolverObjectReverseMapping;
-			const FParticleIndexReverseMapping* ParticleIndexReverseMapping;
-			const FAllTrailingsIndicesBySolverObject* AllTrailingsIndicesBySolverObject;
-		};
-
-
-		class CHAOSSOLVERS_API FScopedGetEventsData
+		template<typename Lambda>
+		void ForEachPhysicsProxyParallel(Lambda InCallable)
 		{
-		public:
-
-			FScopedGetEventsData(const FPBDRigidsSolver* SolverIn) : Solver(SolverIn)
+			Chaos::PhysicsParallelFor(GeometryParticlePhysicsProxies.Num(), [this, &InCallable](const int32 Index)
 			{
-				check(Solver);
-				if (Solver->BufferMode == EMultiBufferMode::Double)
-				{
-					Solver->SolverResourceLock.ReadLock();
-				}
-			}
-
-			~FScopedGetEventsData()
+				FGeometryParticlePhysicsProxy* Obj = GeometryParticlePhysicsProxies[Index];
+				InCallable(Obj);
+			});
+			Chaos::PhysicsParallelFor(KinematicGeometryParticlePhysicsProxies.Num(), [this, &InCallable](const int32 Index)
 			{
-				check(Solver);
-
-				if (Solver->BufferMode == EMultiBufferMode::Double)
-				{
-					Solver->SolverResourceLock.ReadUnlock();
-				}
-			}
-
-			const Chaos::FPBDRigidsSolver::FAllCollisionDataMaps& GetAllCollisions_Maps() const
+				FKinematicGeometryParticlePhysicsProxy* Obj = KinematicGeometryParticlePhysicsProxies[Index];
+				InCallable(Obj);
+			});
+			Chaos::PhysicsParallelFor(RigidParticlePhysicsProxies.Num(), [this, &InCallable](const int32 Index)
 			{
-				return Solver->GetAllCollisions_Maps_GameThread_NEEDSLOCK();
-			}
-
-			const Chaos::FPBDRigidsSolver::FAllBreakingDataMaps& GetAllBreakings_Maps() const
+				FRigidParticlePhysicsProxy* Obj = RigidParticlePhysicsProxies[Index];
+				InCallable(Obj);
+			});
+			Chaos::PhysicsParallelFor(SkeletalMeshPhysicsProxies.Num(), [this, &InCallable](const int32 Index)
 			{
-				return Solver->GetAllBreakings_Maps_GameThread_NEEDSLOCK();
-			}
-
-			const Chaos::FPBDRigidsSolver::FAllTrailingDataMaps& GetAllTrailings_Maps() const
+				FSkeletalMeshPhysicsProxy* Obj = SkeletalMeshPhysicsProxies[Index];
+				InCallable(Obj);
+			});
+			Chaos::PhysicsParallelFor(StaticMeshPhysicsProxies.Num(), [this, &InCallable](const int32 Index)
 			{
-				return Solver->GetAllTrailings_Maps_GameThread_NEEDSLOCK();
-			}
-
-
-		private:
-
-			const FPBDRigidsSolver* Solver;
-		};
-
-
-		/* ---------------- */
-
-		private:
-		FPBDRigidsSolver(const EMultiBufferMode BufferingModeIn);
-		public:
-
-
-		/* Object registration */
-		void RegisterObject(TSolverObject<FGeometryCollectionPhysicsObject>* InObject);
-		void RegisterObject(TSolverObject<FSkeletalMeshPhysicsObject>* InObject);
-		void RegisterObject(TSolverObject<FStaticMeshPhysicsObject>* InObject);
-		void RegisterObject(TSolverObject<FBodyInstancePhysicsObject>* InObject);
-		void RegisterObject(FFieldSystemPhysicsObject* InObject);
-
-		void UnregisterObject(TSolverObject<FGeometryCollectionPhysicsObject>* InObject);
-		void UnregisterObject(TSolverObject<FSkeletalMeshPhysicsObject>* InObject);
-		void UnregisterObject(TSolverObject<FStaticMeshPhysicsObject>* InObject);
-		void UnregisterObject(TSolverObject<FBodyInstancePhysicsObject>* InObject);
-		void UnregisterObject(FFieldSystemPhysicsObject* InObject);
-		/*---------------------*/
+				FStaticMeshPhysicsProxy* Obj = StaticMeshPhysicsProxies[Index];
+				InCallable(Obj);
+			});
+			Chaos::PhysicsParallelFor(GeometryCollectionPhysicsProxies.Num(), [this, &InCallable](const int32 Index)
+			{
+				FGeometryCollectionPhysicsProxy* Obj = GeometryCollectionPhysicsProxies[Index];
+				InCallable(Obj);
+			});
+			Chaos::PhysicsParallelFor(FieldSystemPhysicsProxies.Num(), [this, &InCallable](const int32 Index)
+			{
+				FFieldSystemPhysicsProxy* Obj = FieldSystemPhysicsProxies[Index];
+				InCallable(Obj);
+			});
+		}
 
-		const FSolverObjectStorage& GetObjectStorage() const { return Objects; }
-		FSolverObjectStorage& GetObjectStorage() { return Objects; }
-		
-		const FSolverObjectStorage& GetObjectStorage_GameThread() const { return Objects_GameThread; }
-		FSolverObjectStorage& GetObjectStorage_GameThread() { return Objects_GameThread; }
+		bool GetNumPhysicsProxies() const {
+			return GeometryParticlePhysicsProxies.Num() + KinematicGeometryParticlePhysicsProxies.Num() + RigidParticlePhysicsProxies.Num()
+				+ SkeletalMeshPhysicsProxies.Num() + StaticMeshPhysicsProxies.Num()
+				+ GeometryCollectionPhysicsProxies.Num() + FieldSystemPhysicsProxies.Num() ;
+		}
 
-		const FSolverObjectStorage& GetRemovedObjectStorage() const { return RemovedObjects; }
-		FSolverObjectStorage& GetRemovedObjectStorage() { return RemovedObjects; }
+		//
+		//  Simulation API
+		//
 
-		bool HasActiveObjects() const { return !!Objects.GetNumObjects(); }
+		/**/
+		bool Enabled() const { if (bEnabled) return this->IsSimulating(); return false; }
+		void SetEnabled(bool bEnabledIn) { bEnabled = bEnabledIn; }
+		bool HasActiveParticles() const { return !!GetNumPhysicsProxies(); }
+		bool HasPendingCommands() const { return !CommandQueue.IsEmpty(); }
 
 		/**/
 		void Reset();
-		void ChangeBufferMode(Chaos::EMultiBufferMode InBufferMode);
 
 		/**/
 		void AdvanceSolverBy(float DeltaTime);
-
-		/* Particle Update and access*/
-		void InitializeFromParticleData(const int32 StartIndex) { MEvolution->InitializeFromParticleData(StartIndex); }
-		const FParticlesType& GetRigidParticles() const { return MEvolution->GetParticles(); }
-		FParticlesType& GetRigidParticles() { return MEvolution->GetParticles(); }
-
-		/**/
-		const TSet<int32>& ActiveIndices() const { return MEvolution->GetActiveIndices(); }
-		TSet<int32>& ActiveIndices() { return MEvolution->GetActiveIndices(); }
-		const TArray<int32>& NonDisabledIndices() const { return MEvolution->GetNonDisabledIndices(); }
-		TArray<int32>& NonDisabledIndices() { return MEvolution->GetNonDisabledIndices(); }
-
-		void WakeIslands(const TSet<int32>& InIslandIndices) { MEvolution->WakeIslands(InIslandIndices); }
-
-
-		TSerializablePtr<TChaosPhysicsMaterial<float>> GetPhysicsMaterial(const int32 Index) const { return MEvolution->GetPhysicsMaterial(Index);  }
-		void SetPhysicsMaterial(const int32 Index, TSerializablePtr<TChaosPhysicsMaterial<float>> InMaterial) const { return MEvolution->SetPhysicsMaterial(Index, InMaterial); }
-
-		const TUniquePtr<TChaosPhysicsMaterial<float>>& GetPerParticlePhysicsMaterial(const int32 Index) const { return MEvolution->GetPerParticlePhysicsMaterial(Index); }
-		void SetPerParticlePhysicsMaterial(const int32 Index, TUniquePtr<TChaosPhysicsMaterial<float>>&& InMaterial) const { return MEvolution->SetPerParticlePhysicsMaterial(Index, MoveTemp(InMaterial)); }
-
-		FCollisionConstraints& GetCollisionConstraints() const { return MEvolution->GetCollisionConstraints(); }
-
-		const TSet<TTuple<int32, int32>>& GetDisabledCollisionPairs() const { return MEvolution->GetDisabledCollisions(); }
-		void RemoveConstraints(const TSet<uint32>& RemovedParticles) { MEvolution->RemoveConstraints(RemovedParticles); }
+		void UpdatePhysicsThreadStructures();
+		void UpdateGameThreadStructures();
+		void BufferPhysicsResults();
+		void FlipBuffers();
 
 		/**/
 		void SetCurrentFrame(const int32 CurrentFrameIn) { CurrentFrame = CurrentFrameIn; }
-		int32 GetCurrentFrame() { return CurrentFrame; }
-
-		/* Pause */
-#if CHAOS_WITH_PAUSABLE_SOLVER
-		bool Paused() const { return bPaused; }
-		void SetPaused(bool bPausedIn) { bPaused = bPausedIn; }
-#endif
-
-		bool Enabled() const;
-		void SetEnabled(bool bEnabledIn) { bEnabled = bEnabledIn; }
-
-		/* Clustering Access */
-		FClusteringType& GetRigidClustering() { return MEvolution->GetRigidClustering(); }
-		const FClusteringType& GetRigidClustering() const { return MEvolution->GetRigidClustering(); }
+		int32& GetCurrentFrame() { return CurrentFrame; }
 
 		/**/
-		const float GetSolverTime() { return MTime; }
-		const float GetLastDt() { return MLastDt; }
-		const bool GetGenerateCollisionData() { return bDoGenerateCollisionData; }
-		const bool GetGenerateBreakingData() { return bDoGenerateBreakingData; }
-		const bool GetGenerateTrailingData() { return bDoGenerateTrailingData; }
-
-		const TArray<TBreakingData<float, 3>>& GetAllClusterBreakings() const { return MEvolution->GetRigidClustering().GetAllClusterBreakings(); }
+		float& GetSolverTime() { return MTime; }
+		const float GetSolverTime() const { return MTime; }
 
 		/**/
-		const ISpatialAcceleration<float,3>* GetSpatialAcceleration() const { return &GetCollisionConstraints().GetSpatialAcceleration(); }
-		void ReleaseSpatialAcceleration() const { GetCollisionConstraints().ReleaseSpatialAcceleration(); }
-
-		/* General maps for reverse lookups */
-		const FSolverObjectReverseMapping& GetSolverObjectReverseMapping_GameThread() const
-		{ 
-			SolverObjectReverseMappingLock.ReadLock();
-			return *SolverObjectReverseMappingBuffer->GetConsumerBuffer();
-		}
-		void ReleaseSolverObjectReverseMapping() const { SolverObjectReverseMappingLock.ReadUnlock(); }
-
-		void SyncEvents_GameThread();
-
-private:
-		/* Collision */
-
-		/** 
-		 * Returns the most recent gamethread-accessible collision data. 
-		 */
-		const FAllCollisionDataMaps& GetAllCollisions_Maps_GameThread_NEEDSLOCK() const;
-
-		/* Breaking */
-
-		/**
-		 * Returns the most recent gamethread-accessible breaking data, incl various data maps
-		 */
-		const FAllBreakingDataMaps& GetAllBreakings_Maps_GameThread_NEEDSLOCK() const;
-
-		/* Trailing */
-
-		/**
-		 * Returns the most recent gamethread-accessible trailing data, incl various data maps
-		 */
-		const FAllTrailingDataMaps& GetAllTrailings_Maps_GameThread_NEEDSLOCK() const;
-public:
+		const float GetLastDt() const { return MLastDt; }
+		const float GetMaxDeltaTime() const { return MMaxDeltaTime; }
 
 		/**/
-		void SetTimeStepMultiplier(float TimeStepMultiplierIn) { ensure(TimeStepMultiplierIn > 0); TimeStepMultiplier = TimeStepMultiplierIn; }
-		void SetIterations(int32 Iterations) { MEvolution->SetIterations(Iterations); }
-		void SetPushOutIterations(int32 PushOutIterations) { MEvolution->SetPushOutIterations(PushOutIterations); }
-		void SetPushOutPairIterations(int32 PushOutPairIterations) { MEvolution->SetPushOutPairIterations(PushOutPairIterations); }
-		void SetClusterConnectionFactor(float ClusterConnectionFactor) { GetRigidClustering().SetClusterConnectionFactor(ClusterConnectionFactor); }
-		void SetClusterUnionConnectionType(FClusterCreationParameters<float>::EConnectionMethod ClusterConnectionType) {GetRigidClustering().SetClusterUnionConnectionType(ClusterConnectionType);}
-		void SetGenerateCollisionData(bool bDoGenerate) { bDoGenerateCollisionData = bDoGenerate; }
-		void SetGenerateBreakingData(bool bDoGenerate) 
-		{
-			bDoGenerateBreakingData = bDoGenerate; 
-			GetRigidClustering().SetGenerateClusterBreaking(bDoGenerateBreakingData);
-		}
-		void SetGenerateTrailingData(bool bDoGenerate) { bDoGenerateTrailingData = bDoGenerate; }
-		void SetCollisionFilterSettings(const FSolverCollisionFilterSettings InCollisionFilterSettings)
+		void SetGenerateCollisionData(bool bDoGenerate) { GetEventFilters()->SetGenerateCollisionEvents(bDoGenerate); }
+		void SetGenerateBreakingData(bool bDoGenerate)
 		{
-			CollisionFilterSettings = InCollisionFilterSettings;
-		}
-		const FSolverCollisionFilterSettings& GetCollisionFilterSettings() const { return CollisionFilterSettings; }
-
-		void SetBreakingFilterSettings(const FSolverBreakingFilterSettings InBreakingFilterSettings)
-		{
-			BreakingFilterSettings = InBreakingFilterSettings;
-		}
-		const FSolverBreakingFilterSettings& GetBreakingFilterSettings() const { return BreakingFilterSettings; }
-
-		void SetTrailingFilterSettings(const FSolverTrailingFilterSettings InTrailingFilterSettings)
-		{
-			TrailingFilterSettings = InTrailingFilterSettings;
+			GetEventFilters()->SetGenerateBreakingEvents(bDoGenerate);
+#if TODO_REIMPLEMENT_RIGID_CLUSTERING
+			GetRigidClustering().SetGenerateClusterBreaking(GenerateBreakingEventsEnabled);
+#endif
 		}
-		const FSolverTrailingFilterSettings& GetTrailingFilterSettings() const { return TrailingFilterSettings; }
-
-		/**/
+		void SetGenerateTrailingData(bool bDoGenerate) { GetEventFilters()->SetGenerateTrailingEvents(bDoGenerate); }
+		void SetCollisionFilterSettings(const FSolverCollisionFilterSettings& InCollisionFilterSettings) { GetEventFilters()->GetCollisionFilter()->UpdateFilterSettings(InCollisionFilterSettings); }
+		void SetBreakingFilterSettings(const FSolverBreakingFilterSettings& InBreakingFilterSettings) { GetEventFilters()->GetBreakingFilter()->UpdateFilterSettings(InBreakingFilterSettings); }
+		void SetTrailingFilterSettings(const FSolverTrailingFilterSettings& InTrailingFilterSettings) { GetEventFilters()->GetTrailingFilter()->UpdateFilterSettings(InTrailingFilterSettings); }
 		void SetHasFloor(bool bHasFloorIn) { bHasFloor = bHasFloorIn; }
 		void SetIsFloorAnalytic(bool bIsAnalytic) { bIsFloorAnalytic = bIsAnalytic; }
 		void SetFloorHeight(float Height) { FloorHeight = Height; }
-		void SetMassScale(float MassScaleIn) { MassScale = MassScaleIn; }
-		int32 GetFloorIndex() const { return FloorIndex; }
-
-		const Chaos::TArrayCollectionArray<SolverObjectWrapper>& GetSolverObjectReverseMapping() const { return SolverObjectReverseMapping; }
-		const Chaos::TArrayCollectionArray<int32>& GetParticleIndexReverseMapping() const { return ParticleIndexReverseMapping; }
-
-		FCollisionDataArray& GetAllCollisionsDataArray() { return AllCollisionsBuffer->AccessProducerBuffer()->AllCollisionsArray; }
-		TMap<ISolverObjectBase*, TArray<int32>>& GetAllCollisionsIndicesBySolverObject() { return AllCollisionsIndicesBySolverObjectBuffer->AccessProducerBuffer()->AllCollisionsIndicesBySolverObjectMap; }
-		FBreakingDataArray& GetAllBreakingsDataArray() { return AllBreakingsBuffer->AccessProducerBuffer()->AllBreakingsArray; }
-		TMap<ISolverObjectBase*, TArray<int32>>& GetAllBreakingsIndicesBySolverObject() { return AllBreakingsIndicesBySolverObjectBuffer->AccessProducerBuffer()->AllBreakingsIndicesBySolverObjectMap; }
-		FTrailingDataArray& GetAllTrailingsDataArray() { return AllTrailingsBuffer->AccessProducerBuffer()->AllTrailingsArray; }
-		TMap<ISolverObjectBase*, TArray<int32>>& GetAllTrailingsIndicesBySolverObject() { return AllTrailingsIndicesBySolverObjectBuffer->AccessProducerBuffer()->AllTrailingsIndicesBySolverObjectMap; }
-
-		int32 GetParticleIndexMesh(const TMap<uint32, TUniquePtr<TArray<uint32>>>& ParentToChildrenMap, int32 ParticleIndex);
-
-		/** Returns encoded collision index. */
-		static int32 EncodeCollisionIndex(int32 ActualCollisionIndex, bool bSwapOrder);
-		/** Returns decoded collision index. */
-		static int32 DecodeCollisionIndex(int32 EncodedCollisionIdx, bool& bSwapOrder);
-
-
-		TQueue<TFunction<void(FPBDRigidsSolver*)>, EQueueMode::Mpsc>& GetCommandQueue() { return CommandQueue; }
-
-#if !UE_BUILD_SHIPPING
-		void SerializeForPerfTest(FChaosArchive& Ar) { MEvolution->SerializeForPerfTest(Ar); }
-		void SerializeForPerfTest(const FString& FileName);
-#endif
-
-		/* Return the debug substep object to allow pausing/stepping/substepping operations on this solver. */
-		FDebugSubstep& GetDebugSubstep() const { return MEvolution->GetDebugSubstep(); }
-
-		const float GetMassScale() const { return MassScale; }
-
-		FPBDRigidsEvolution* GetEvolution() { return MEvolution.Get(); }
-
-	protected:
-		/**/
-		void CreateRigidBodyCallback(FParticlesType& Particles);
 
 		/**/
-		void ParameterUpdateCallback(FParticlesType& Particles, const float Time);
-
-		/**/
-		void ForceUpdateCallback(FParticlesType& Particles, const float Time);
-
-		/**/
-		void DisableCollisionsCallback(TSet<TTuple<int32, int32>>& CollisionPairs);
-
-		/**/
-		void StartFrameCallback(const float Dt, const float Time);
-
-		/**/
-		void EndFrameCallback(const float EndFrame);
-
-		/**/
-		void KinematicUpdateCallback(FParticlesType& Particles, const float Dt, const float Time);
-
-		/**/
-		void AddForceCallback(FParticlesType& Particles, const float Time, const int32 Index);
+		FPBDRigidsEvolution* GetEvolution() { return MEvolution.Get(); }
+		FPBDRigidsEvolution* GetEvolution() const { return MEvolution.Get(); }
 
-		/**/
-		void CollisionContactsCallback(FParticlesType& Particles, FCollisionConstraints& CollisionConstraints);
+		const FParticlesType& GetParticles() const { return Particles; }
 
 		/**/
-		void BreakingCallback(FParticlesType& Particles);
+		FEventManager* GetEventManager() { return MEventManager.Get(); }
 
 		/**/
-		void TrailingCallback(FParticlesType& Particles);
+		FSolverEventFilters* GetEventFilters() { return MSolverEventFilters.Get(); }
+		FSolverEventFilters* GetEventFilters() const { return MSolverEventFilters.Get(); }
 
 		/**/
-		void BindParticleCallbackMappingPart1( );
-		void BindParticleCallbackMappingPart2( );
+		void SyncEvents_GameThread();
 
 	private:
 
+		//
+		// Solver Data
+		//
 		int32 CurrentFrame;
 		float MTime;
 		float MLastDt;
 		float MMaxDeltaTime;
 		float TimeStepMultiplier;
 
-		bool bDoGenerateCollisionData;
-		bool bDoGenerateBreakingData;
-		bool bDoGenerateTrailingData;
-
-#if CHAOS_WITH_PAUSABLE_SOLVER
-		bool bPaused;
-#endif
 		bool bEnabled;
 		bool bHasFloor;
 		bool bIsFloorAnalytic;
 		float FloorHeight;
-		int32 FloorIndex;
-
-		float MassScale;
-
-		FSolverCollisionFilterSettings CollisionFilterSettings;
-		FSolverBreakingFilterSettings BreakingFilterSettings;
-		FSolverTrailingFilterSettings TrailingFilterSettings;
-
-		TSharedPtr<FEvent> MCurrentEvent;
-		TSharedPtr<FCriticalSection> MCurrentLock;
 
+		FParticlesType Particles;
 		TUniquePtr<FPBDRigidsEvolution> MEvolution;
+		TUniquePtr<FEventManager> MEventManager;
+		TUniquePtr<FSolverEventFilters> MSolverEventFilters;
 
-		int32 FieldForceNum;
-		TArrayCollectionArray<FVector> FieldForce;
-		TArrayCollectionArray<FVector> FieldTorque;
-		TArray<FKinematicProxy> KinematicProxies;
-
-		TArray<FKinematicProxy> KinematicProxiesForObjects;
-
-		FPositionConstraints PositionTarget;
-		FPositionConstraintsRule PositionConstraintsRule;
-		TMap<int32, int32> PositionTargetedParticles;
-
-		FRigidDynamicSpringConstraints DynamicConstraints;
-		FRigidDynamicSpringConstraintsRule DynamicConstraintsRule;
-		TSet<int32> DynamicConstraintParticles;
-
-		FJointConstraints JointConstraints;
-		FJointConstraintsRule JointConstraintsRule;
-
-		int32 LastSwapMappingSyncSize;
-		int32 LastMappingSyncSize;
-
-		// Mappings
-		TArrayCollectionArray<SolverObjectWrapper> SolverObjectReverseMapping; // ParticleIndex -> SolverObject
-		TArrayCollectionArray<int32> ParticleIndexReverseMapping; // ParticleIndex -> TransformIndex
-
+		//
+		// Commands
+		//
 		TQueue<TFunction<void(FPBDRigidsSolver*)>, EQueueMode::Mpsc> CommandQueue;
 
-		// Storage for objects
-		//  - Objects -> Physics thread side objects. Never read from any other thread
-		//  - Objects_GameThread -> Game thread side objects. This can be iterated and
-		//                          then GT only data read from the objects. This doesn't
-		//                          guarantee all data in the object is safe, only that the
-		//                          size of the storage won't change off the game thread
-		//  - RemovedObjects -> Temporary storage when objects are removed until the runtime
-		//                      hits a final sync opportunity
-		FSolverObjectStorage Objects;
-		FSolverObjectStorage Objects_GameThread;
-		FSolverObjectStorage RemovedObjects;
-
-		TAtomic<bool> GameThreadHasSynced;
-
-	public:
-		mutable FRWLock SolverResourceLock;
-		mutable FRWLock SolverObjectReverseMappingLock;
-		mutable FRWLock ParticleIndexReverseMappingLock;
-
-		const FScopedGetEventsData ScopedGetEventsData() const { return FScopedGetEventsData(this); }
+		//
+		// Proxies
+		//
 		EMultiBufferMode BufferMode;
-
-	private:
-		TUniquePtr<IBufferResource<FSolverObjectReverseMapping>> SolverObjectReverseMappingBuffer;
-		TUniquePtr<IBufferResource<FParticleIndexReverseMapping>> ParticleIndexReverseMappingBuffer;
-
-	public:
-		const FAllCollisionData* GetAllCollisions_FromSequencerCache_NEEDSLOCK() const { return AllCollisions_FromSequencerCache.Get(); }
-		const FAllBreakingData* GetAllBreakings_FromSequencerCache_NEEDSLOCK() const { return AllBreakings_FromSequencerCache.Get(); }
-		const FAllTrailingData* GetAllTrailings_FromSequencerCache_NEEDSLOCK() const { return AllTrailings_FromSequencerCache.Get(); }
-	private:
-
-		TUniquePtr<FAllCollisionData> AllCollisions_FromSequencerCache;
-		TUniquePtr<FAllBreakingData> AllBreakings_FromSequencerCache;
-		TUniquePtr<FAllTrailingData> AllTrailings_FromSequencerCache;
-		//mutable FCriticalSection SequencerCacheLock;
-	//private:
-		TUniquePtr<IBufferResource<FAllCollisionData>> AllCollisionsBuffer;
-		TUniquePtr<IBufferResource<FAllCollisionsIndicesBySolverObject>> AllCollisionsIndicesBySolverObjectBuffer;
-		FAllCollisionDataMaps AllCollisionDataMaps;
-
-
-		/* Breakings */
-		TUniquePtr<IBufferResource<FAllBreakingData>> AllBreakingsBuffer;
-		TUniquePtr<IBufferResource<FAllBreakingsIndicesBySolverObject>> AllBreakingsIndicesBySolverObjectBuffer;
-
-		FAllBreakingDataMaps AllBreakingDataMaps;
-
-		TUniquePtr<IBufferResource<FAllTrailingData>> AllTrailingsBuffer;
-		TUniquePtr<IBufferResource<FAllTrailingsIndicesBySolverObject>> AllTrailingsIndicesBySolverObjectBuffer;
-
-		FAllTrailingDataMaps AllTrailingDataMaps;
-
-		TUniquePtr<FSolverCollisionEventFilter> SolverCollisionEventFilter;
-		TUniquePtr<FSolverBreakingEventFilter> SolverBreakingEventFilter;
-		TUniquePtr<FSolverTrailingEventFilter> SolverTrailingEventFilter;
-		
-		// Previous frame
-		int32 NumCollisionsPrevFrame;
-		int32 NumBreakingsPrevFrame;
-		int32 NumTrailingsPrevFrame;
-	};
-
-	struct FSolverReadLock : public FNoncopyable
-	{
-		FSolverReadLock() = delete;
-		FSolverReadLock(const FPBDRigidsSolver* InSolver)
-		{
-			CachedLock = &(InSolver->SolverResourceLock);
-			CachedLock->ReadLock();
-		}
-
-		~FSolverReadLock()
-		{
-			CachedLock->ReadUnlock();
-		}
-
-	private:
-		FRWLock* CachedLock;
-	};
-
-	struct FSolverWriteLock : public FNoncopyable
-	{
-		FSolverWriteLock() = delete;
-		FSolverWriteLock(FPBDRigidsSolver* InSolver)
-		{
-			CachedLock = &(InSolver->SolverResourceLock);
-			CachedLock->WriteLock();
-		}
-
-		~FSolverWriteLock()
-		{
-			CachedLock->WriteUnlock();
-		}
-
-	private:
-		FRWLock* CachedLock;
+		TSharedPtr<FCriticalSection> MCurrentLock;
+		TArray< FGeometryParticlePhysicsProxy* > GeometryParticlePhysicsProxies;
+		TArray< FKinematicGeometryParticlePhysicsProxy* > KinematicGeometryParticlePhysicsProxies;
+		TArray< FRigidParticlePhysicsProxy* > RigidParticlePhysicsProxies;
+		TArray< FSkeletalMeshPhysicsProxy* > SkeletalMeshPhysicsProxies; // dep
+		TArray< FStaticMeshPhysicsProxy* > StaticMeshPhysicsProxies; // dep
+		TArray< FGeometryCollectionPhysicsProxy* > GeometryCollectionPhysicsProxies;
+		TArray< FFieldSystemPhysicsProxy* > FieldSystemPhysicsProxies;
+				
 	};
-}
 
+}; // namespace Chaos
 #endif
diff --git a/Engine/Source/Runtime/Experimental/ChaosSolvers/Public/PhysicsProxy/AnalyticImplicitGroup.h b/Engine/Source/Runtime/Experimental/ChaosSolvers/Public/PhysicsProxy/AnalyticImplicitGroup.h
new file mode 100644
--- /dev/null
+++ b/Engine/Source/Runtime/Experimental/ChaosSolvers/Public/PhysicsProxy/AnalyticImplicitGroup.h
@@ -0,0 +1,661 @@
+// Copyright 1998-2019 Epic Games, Inc. All Rights Reserved.
+#pragma once
+#if INCLUDE_CHAOS
+
+#include "Chaos/Box.h"
+#include "Chaos/Capsule.h"
+#include "Chaos/Convex.h"
+#include "Chaos/ImplicitObject.h"
+#include "Chaos/ImplicitObjectTransformed.h"
+#include "Chaos/ImplicitObjectUnion.h"
+#include "Chaos/Levelset.h"
+#include "Chaos/MassProperties.h"
+#include "Chaos/Sphere.h"
+#include "Chaos/TaperedCylinder.h"
+#include "GeometryCollection/GeometryCollectionSimulationTypes.h"
+
+/**
+ * A collection of analytic implicit shapes parented to a single transform, 
+ * in a hierarchy of transforms.
+ *
+ * Currently we use this class to represent bones even if they don't have any 
+ * implicit shapes.
+ */
+class FAnalyticImplicitGroup
+{
+public:
+	FAnalyticImplicitGroup()
+		: BoneName(NAME_None)
+		, BoneIndex(INDEX_NONE)
+		, ParentBoneIndex(INDEX_NONE)
+		, RigidBodyId(INDEX_NONE)
+		, RigidBodyState(EObjectStateTypeEnum::Chaos_Object_Kinematic)
+		, Parent(nullptr)
+	{}
+	FAnalyticImplicitGroup(const FName &InBoneName, const int32 InBoneIndex)
+		: BoneName(InBoneName)
+		, BoneIndex(InBoneIndex)
+		, ParentBoneIndex(INDEX_NONE)
+		, RigidBodyId(INDEX_NONE)
+		, RigidBodyState(EObjectStateTypeEnum::Chaos_Object_Kinematic)
+		, Parent(nullptr)
+	{}
+	FAnalyticImplicitGroup(const FAnalyticImplicitGroup &) = delete;
+	FAnalyticImplicitGroup(FAnalyticImplicitGroup &&Other)
+		: BoneName(MoveTemp(Other.BoneName))
+		, BoneIndex(MoveTemp(Other.BoneIndex))
+		, ParentBoneIndex(MoveTemp(Other.ParentBoneIndex))
+
+		, RigidBodyId(MoveTemp(Other.RigidBodyId))
+		, RigidBodyState(MoveTemp(Other.RigidBodyState))
+
+		, Spheres(MoveTemp(Other.Spheres))
+		, Boxes(MoveTemp(Other.Boxes))
+		, Capsules(MoveTemp(Other.Capsules))
+		, TaperedCylinders(MoveTemp(Other.TaperedCylinders))
+		, ConvexHulls(MoveTemp(Other.ConvexHulls))
+		, LevelSets(MoveTemp(Other.LevelSets))
+
+		, Transforms(MoveTemp(Other.Transforms))
+		, Parent(MoveTemp(Other.Parent))
+		, Children(MoveTemp(Other.Children))
+	{}
+	~FAnalyticImplicitGroup()
+	{
+		for (Chaos::TSphere<float, 3>* Sphere : Spheres) if (Sphere) delete Sphere;
+		for (Chaos::TBox<float, 3>* Box : Boxes) if (Box) delete Box;
+		for (Chaos::TCapsule<float>* Capsule : Capsules) if (Capsule) delete Capsule;
+		for (Chaos::TTaperedCylinder<float>* TaperedCylinder : TaperedCylinders) if (TaperedCylinder) delete TaperedCylinder;
+		for (Chaos::TConvex<float, 3>* ConvexHull : ConvexHulls) if (ConvexHull) delete ConvexHull;
+		for (Chaos::TLevelSet<float, 3>* LevelSet : LevelSets) if (LevelSet) delete LevelSet;
+	}
+
+	void Init(const int32 NumStructures, const bool DoCollGeom=true)
+	{
+		Transforms.Reserve(NumStructures);
+		if (DoCollGeom)
+		{
+			CollisionPoints.SetNum(NumStructures);
+			CollisionTriangles.SetNum(NumStructures);
+		}
+	}
+
+	/** The number of analytic shapes in this group. */
+	int32 NumStructures() const { return Transforms.Num(); }
+
+	bool IsValid() { return BoneIndex != INDEX_NONE; }
+
+	const FName& GetBoneName() const { return BoneName; }
+	const int32 GetBoneIndex() const { return BoneIndex; }
+
+	void SetParentBoneIndex(const int32 InParentBoneIndex) { ParentBoneIndex = InParentBoneIndex; }
+	const int32 GetParentBoneIndex() const { return ParentBoneIndex; }
+
+	void SetRigidBodyId(const int32 InRigidBodyId) { RigidBodyId = InRigidBodyId; }
+	int32 GetRigidBodyId() const { return RigidBodyId; }
+
+	void SetRigidBodyState(const EObjectStateTypeEnum State) { RigidBodyState = State; }
+	EObjectStateTypeEnum GetRigidBodyState() const { return RigidBodyState; }
+
+	int32 Add(const FTransform &InitialXf, Chaos::TSphere<float, 3> *Sphere) { Spheres.Add(Sphere); return Transforms.Insert(InitialXf, Spheres.Num()-1); }
+	int32 Add(const FTransform &InitialXf, Chaos::TBox<float, 3> *Box) { Boxes.Add(Box); return Transforms.Insert(InitialXf, Spheres.Num()+Boxes.Num()-1); }
+	int32 Add(const FTransform &InitialXf, Chaos::TCapsule<float> *Capsule) { Capsules.Add(Capsule); return Transforms.Insert(InitialXf, Spheres.Num()+Boxes.Num()+Capsules.Num()-1); }
+	int32 Add(const FTransform &InitialXf, Chaos::TTaperedCylinder<float> *TaperedCylinder) { TaperedCylinders.Add(TaperedCylinder); return Transforms.Insert(InitialXf, Spheres.Num()+Boxes.Num()+Capsules.Num()+TaperedCylinders.Num()-1); }
+	int32 Add(const FTransform &InitialXf, Chaos::TConvex<float, 3> *ConvexHull) { ConvexHulls.Add(ConvexHull); return Transforms.Insert(InitialXf, Spheres.Num()+Boxes.Num()+Capsules.Num()+TaperedCylinders.Num()+ConvexHulls.Num()-1); }
+	int32 Add(const FTransform &InitialXf, Chaos::TLevelSet<float, 3> *LevelSet) { LevelSets.Add(LevelSet); return Transforms.Add(InitialXf); }
+
+	void SetCollisionTopology(
+		const int32 Index,
+		TArray<Chaos::TVector<float,3>>&& Points,
+		TArray<Chaos::TVector<int32,3>>&& Triangles)
+	{ CollisionPoints[Index] = MoveTemp(Points); CollisionTriangles[Index] = MoveTemp(Triangles); }
+
+	const TArray<FTransform>& GetInitialStructureTransforms() const { return Transforms; }
+	void ResetTransforms() { Transforms.Init(FTransform::Identity, Transforms.Num()); }
+
+	Chaos::TMassProperties<float, 3> BuildMassProperties(
+		const float Density,
+		float& TotalMass)
+	{
+		// Make sure this function is being called before we've transfered ownership 
+		// of our implicit shapes to the simulator.
+		checkSlow(!Spheres.Contains(nullptr) && !Boxes.Contains(nullptr) && !Capsules.Contains(nullptr) && 
+			!TaperedCylinders.Contains(nullptr) && !ConvexHulls.Contains(nullptr) && !LevelSets.Contains(nullptr));
+
+		const int32 Num = NumStructures();
+		TArray<Chaos::TMassProperties<float, 3>> MPArray;
+		TArray<Chaos::TBox<float, 3>> BBoxes;
+		MPArray.SetNum(Num);
+		BBoxes.SetNum(Num);
+
+		int32 TransformIndex = 0;
+		for (Chaos::TSphere<float, 3>* Sphere : Spheres)
+		{
+			const FTransform& Xf = Transforms[TransformIndex];
+			BBoxes[TransformIndex] = Sphere->BoundingBox().TransformedBox(Xf);
+			Chaos::TMassProperties<float, 3> &MP = MPArray[TransformIndex++];
+			MP.Volume = Sphere->GetVolume();
+			MP.CenterOfMass = Xf.TransformPositionNoScale(Sphere->GetCenterOfMass());
+			MP.RotationOfMass = Xf.TransformRotation(Sphere->GetRotationOfMass());
+		}
+		for (Chaos::TBox<float, 3>* Box : Boxes)
+		{
+			const FTransform& Xf = Transforms[TransformIndex];
+			BBoxes[TransformIndex] = Box->BoundingBox().TransformedBox(Xf);
+			Chaos::TMassProperties<float, 3> &MP = MPArray[TransformIndex++];
+			MP.Volume = Box->GetVolume();
+			MP.CenterOfMass = Xf.TransformPositionNoScale(Box->GetCenterOfMass());
+			MP.RotationOfMass = Xf.TransformRotation(Box->GetRotationOfMass());
+		}
+		for (Chaos::TCapsule<float>* Capsule : Capsules)
+		{
+			const FTransform& Xf = Transforms[TransformIndex];
+			BBoxes[TransformIndex] = Capsule->BoundingBox().TransformedBox(Xf);
+			Chaos::TMassProperties<float, 3> &MP = MPArray[TransformIndex++];
+			MP.Volume = Capsule->GetVolume();
+			MP.CenterOfMass = Xf.TransformPositionNoScale(Capsule->GetCenterOfMass());
+			MP.RotationOfMass = Xf.TransformRotation(Capsule->GetRotationOfMass());
+		}
+		for (Chaos::TTaperedCylinder<float>* TaperedCylinder : TaperedCylinders)
+		{
+			const FTransform& Xf = Transforms[TransformIndex];
+			BBoxes[TransformIndex] = TaperedCylinder->BoundingBox().TransformedBox(Xf);
+			Chaos::TMassProperties<float, 3> &MP = MPArray[TransformIndex++];
+			MP.Volume = TaperedCylinder->GetVolume();
+			MP.CenterOfMass = Xf.TransformPositionNoScale(TaperedCylinder->GetCenterOfMass());
+			MP.RotationOfMass = Xf.TransformRotation(TaperedCylinder->GetRotationOfMass());
+		}
+		for (Chaos::TConvex<float, 3>* Convex : ConvexHulls)
+		{
+			const FTransform& Xf = Transforms[TransformIndex];
+			BBoxes[TransformIndex] = Convex->BoundingBox().TransformedBox(Xf);
+			Chaos::TMassProperties<float, 3> &MP = MPArray[TransformIndex++];
+			MP.Volume = Convex->BoundingBox().GetVolume();
+			MP.CenterOfMass = Xf.TransformPositionNoScale(Convex->BoundingBox().Center());
+			MP.RotationOfMass = Xf.TransformRotation(Convex->BoundingBox().GetRotationOfMass());
+		}
+		for (Chaos::TLevelSet<float, 3>* LevelSet : LevelSets)
+		{
+			const FTransform& Xf = Transforms[TransformIndex];
+			BBoxes[TransformIndex] = LevelSet->BoundingBox().TransformedBox(Xf);
+			Chaos::TMassProperties<float, 3> &MP = MPArray[TransformIndex++];
+			MP.Volume = LevelSet->BoundingBox().GetVolume();
+			MP.CenterOfMass = Xf.TransformPositionNoScale(LevelSet->BoundingBox().Center());
+			MP.RotationOfMass = Xf.TransformRotation(LevelSet->BoundingBox().GetRotationOfMass());
+		}
+
+		// Find overlap and adjust volumes accordingly.  We do this by determining 
+		// how much their AABB's overlap for sake of speed and simplicity.  Ideally
+		// we'd do something more accurate...
+		for (int32 i=0; i < Num-1; i++)
+		{
+			const Chaos::TBox<float, 3>& BoxI = BBoxes[i];
+			for (int32 j = i+1; j < Num; j++)
+			{
+				const Chaos::TBox<float, 3>& BoxJ = BBoxes[j];
+				if (BoxI.Intersects(BoxJ))
+				{
+					Chaos::TBox<float, 3> BoxIJ = BoxI.GetIntersection(BoxJ);
+					const float VolIJ = BoxIJ.GetVolume();
+					if (VolIJ > KINDA_SMALL_NUMBER)
+					{
+						const float VolI = BoxI.GetVolume();
+						const float VolJ = BoxJ.GetVolume();
+						const float PctOverlapI = VolI > KINDA_SMALL_NUMBER ? VolIJ / VolI : 0.;
+						const float PctOverlapJ = VolJ > KINDA_SMALL_NUMBER ? VolIJ / VolJ : 0.;
+						// Split the overlapping volume between i and j
+						MPArray[i].Volume *= (1.0 - PctOverlapI / 2);
+						MPArray[j].Volume *= (1.0 - PctOverlapJ / 2);
+					}
+				}
+			}
+		}
+
+		TotalMass = 0.f;
+		TransformIndex = 0;
+		for (Chaos::TSphere<float, 3>* Sphere : Spheres)
+		{
+			Chaos::TMassProperties<float, 3> &MP = MPArray[TransformIndex++];
+			float Mass = Density * MP.Volume;
+			TotalMass += Mass;
+			MP.InertiaTensor = Sphere->GetInertiaTensor(Mass);
+		}
+		for (Chaos::TBox<float, 3>* Box : Boxes)
+		{
+			Chaos::TMassProperties<float, 3> &MP = MPArray[TransformIndex++];
+			float Mass = Density * MP.Volume;
+			TotalMass += Mass;
+			MP.InertiaTensor = Box->GetInertiaTensor(Mass);
+		}
+		for (Chaos::TCapsule<float>* Capsule : Capsules)
+		{
+			Chaos::TMassProperties<float, 3> &MP = MPArray[TransformIndex++];
+			float Mass = Density * MP.Volume;
+			TotalMass += Mass;
+			MP.InertiaTensor = Capsule->GetInertiaTensor(Mass);
+		}
+		for (Chaos::TTaperedCylinder<float>* TaperedCylinder : TaperedCylinders)
+		{
+			Chaos::TMassProperties<float, 3> &MP = MPArray[TransformIndex++];
+			float Mass = Density * MP.Volume;
+			TotalMass += Mass;
+			MP.InertiaTensor = TaperedCylinder->GetInertiaTensor(Mass);
+		}
+		for (Chaos::TConvex<float, 3>* Convex : ConvexHulls)
+		{
+			Chaos::TMassProperties<float, 3> &MP = MPArray[TransformIndex++];
+			float Mass = Density * MP.Volume;
+			TotalMass += Mass;
+			MP.InertiaTensor = Convex->BoundingBox().GetInertiaTensor(Mass);
+		}
+		for (Chaos::TLevelSet<float, 3>* LevelSet : LevelSets)
+		{
+			Chaos::TMassProperties<float, 3> &MP = MPArray[TransformIndex++];
+			float Mass = Density * MP.Volume;
+			TotalMass += Mass;
+			MP.InertiaTensor = LevelSet->BoundingBox().GetInertiaTensor(Mass);
+		}
+
+		return Chaos::Combine(MPArray);
+	}
+
+	TArray<Chaos::TVector<float, 3>>* BuildSamplePoints(
+		const float ParticlesPerUnitArea,
+		const int32 MinParticles,
+		const int32 MaxParticles)
+	{
+		// Make sure this function is being called before we've transfered ownership 
+		// of our implicit shapes to the simulator.
+		checkSlow(!Spheres.Contains(nullptr) && !Boxes.Contains(nullptr) && !Capsules.Contains(nullptr) && 
+			!TaperedCylinders.Contains(nullptr) && !ConvexHulls.Contains(nullptr) && !LevelSets.Contains(nullptr));
+
+		ContiguousCollisionPoints.Reset();
+		const int32 Num = NumStructures();
+		if (Num == 0)
+		{
+			return &ContiguousCollisionPoints;
+		}
+
+		int32 TransformIndex = 0;
+		for (Chaos::TSphere<float, 3>* Sphere : Spheres)
+		{
+			TArray<Chaos::TVector<float, 3>>& Points = CollisionPoints[TransformIndex];
+			if (!Points.Num())
+				Points = Sphere->ComputeSamplePoints(ParticlesPerUnitArea, MinParticles, MaxParticles);
+			CullDeepPoints(Points, TransformIndex);
+			TransformIndex++;
+		}
+		for (Chaos::TBox<float, 3>* Box : Boxes)
+		{
+			TArray<Chaos::TVector<float, 3>>& Points = CollisionPoints[TransformIndex];
+			if (!Points.Num())
+				Points = Box->ComputeSamplePoints();
+			CullDeepPoints(Points, TransformIndex);
+			TransformIndex++;
+		}
+		for (Chaos::TCapsule<float>* Capsule : Capsules)
+		{
+			TArray<Chaos::TVector<float, 3>>& Points = CollisionPoints[TransformIndex];
+			if (!Points.Num())
+				Points = Capsule->ComputeSamplePoints(ParticlesPerUnitArea, MinParticles, MaxParticles);
+			CullDeepPoints(Points, TransformIndex);
+			TransformIndex++;
+		}
+		for (Chaos::TTaperedCylinder<float>* TaperedCylinder : TaperedCylinders)
+		{
+			TArray<Chaos::TVector<float, 3>>& Points = CollisionPoints[TransformIndex];
+			if (!Points.Num())
+				Points = TaperedCylinder->ComputeSamplePoints(ParticlesPerUnitArea, false, MinParticles, MaxParticles);
+			CullDeepPoints(Points, TransformIndex);
+			TransformIndex++;
+		}
+		for (Chaos::TConvex<float, 3>* Convex : ConvexHulls)
+		{
+			TArray<Chaos::TVector<float, 3>>& Points = CollisionPoints[TransformIndex];
+			if (!Points.Num())
+			{
+				const Chaos::TBox<float, 3>& BBox = Convex->BoundingBox();
+				Chaos::TSphere<float, 3> Sphere(BBox.Center(), BBox.Extents().Size() / 2);
+				Points = Sphere.ComputeSamplePoints(ParticlesPerUnitArea, MinParticles, MaxParticles);
+				Chaos::TVector<float, 3> Normal;
+				for (Chaos::TVector<float, 3> &Pt : Points)
+				{
+					const float Phi = Convex->PhiWithNormal(Pt, Normal);
+					Pt += Normal * -Phi;
+					//check(FMath::Abs(Convex->SignedDistance(Pt)) <= KINDA_SMALL_NUMBER);
+				}
+			}
+			CullDeepPoints(Points, TransformIndex);
+			TransformIndex++;
+		}
+		for (Chaos::TLevelSet<float, 3>* LevelSet : LevelSets)
+		{
+			TArray<Chaos::TVector<float, 3>>& Points = CollisionPoints[TransformIndex];
+			if (!Points.Num())
+			{
+				const Chaos::TBox<float, 3>& BBox = LevelSet->BoundingBox();
+				Chaos::TSphere<float, 3> Sphere(BBox.Center(), BBox.Extents().Size() / 2);
+				Points = Sphere.ComputeSamplePoints(ParticlesPerUnitArea, MinParticles, MaxParticles);
+				Chaos::TVector<float, 3> Normal;
+				for (Chaos::TVector<float, 3> &Pt : Points)
+				{
+					const float Phi = LevelSet->PhiWithNormal(Pt, Normal);
+					Pt += Normal * -Phi;
+					//check(FMath::Abs(LevelSet->SignedDistance(Pt)) <= KINDA_SMALL_NUMBER);
+				}
+			}
+			CullDeepPoints(Points, TransformIndex);
+			TransformIndex++;
+		}
+
+		if (Num > 1)
+		{
+			int32 NumPoints = 0;
+			for (const auto& PtArray : CollisionPoints)
+				NumPoints += PtArray.Num();
+			ContiguousCollisionPoints.Reserve(NumPoints);
+		}
+		for (TransformIndex = 0; TransformIndex < Num; TransformIndex++)
+		{
+			TArray<Chaos::TVector<float, 3>> &PtArray = CollisionPoints[TransformIndex];
+			const FTransform& Xf = Transforms[TransformIndex];
+			if(!Xf.Equals(FTransform::Identity))
+			{
+				for (Chaos::TVector<float, 3>& Pt : PtArray)
+					Pt = Xf.TransformPosition(Pt);
+			}
+			if (Num == 1)
+			{
+				// Free memory we're not going to use
+				ContiguousCollisionPoints.Empty();
+				return &PtArray;
+			}
+			ContiguousCollisionPoints.Append(PtArray);	
+		}
+
+		// Free memory we're not going to use
+		CollisionPoints.Empty();
+		return &ContiguousCollisionPoints;
+	}
+
+	TArray<Chaos::TVector<int32, 3>> BuildSampleTopology() const
+	{
+		int32 NumTris = 0;
+		for (const TArray<Chaos::TVector<int32,3>>& Tris : CollisionTriangles)
+		{
+			NumTris += Tris.Num();
+		}
+
+		TArray<Chaos::TVector<int32, 3>> AllTriangles;
+		AllTriangles.Reserve(NumTris);
+		int32 Offset = 0;
+		for (int32 Index=0; Index < CollisionTriangles.Num(); Index++)
+		{
+			for (const Chaos::TVector<int32, 3>& Tri : CollisionTriangles[Index])
+			{
+				AllTriangles.Add(Tri + Offset);
+			}
+			Offset += CollisionPoints[Index].Num();
+		}
+
+		return AllTriangles;
+	}
+
+	/**
+	 * Build the implicit object representation of this object.
+	 *
+	 * Transfers ownership of sub structures to the returned implicit object.
+	 */
+	Chaos::TImplicitObject<float, 3>* BuildSimImplicitObject()
+	{
+		// TODO: We make copies of implicit objects owned by this class, so we
+		// give the solver memory it can own.  It'd be nice if we could transfer
+		// or share the memory this class owns to/with the solver.
+
+		const int32 Num = NumStructures();
+		if (Num == 0)
+		{
+			return nullptr;
+		}
+		else if (Num == 1)
+		{
+			if (Transforms[0].Equals(FTransform::Identity))
+			{
+				return TransferImplicitObj(0);
+			}
+			else
+			{
+				// Make a copy and transfer ownership to the transformed implicit.
+				TUniquePtr<Chaos::TImplicitObject<float, 3>> ObjPtr(TransferImplicitObj(0));
+				return new Chaos::TImplicitObjectTransformed<float, 3, true>(
+					MakeSerializable(ObjPtr),
+					MoveTemp(ObjPtr),
+					Chaos::TRigidTransform<float, 3>(Transforms[0]));
+			}
+		}
+		else
+		{
+			// Make copies of the implicits owned by transformed immplicits, and 
+			// transfer ownership of the transformed implicits to the implicit union.
+			TArray<TUniquePtr<Chaos::TImplicitObject<float, 3>>> ImplicitObjects;
+			ImplicitObjects.Reserve(Num);
+			for (int i = 0; i < Num; i++)
+			{
+				TUniquePtr<Chaos::TImplicitObject<float, 3>> ObjPtr(TransferImplicitObj(i));
+
+				const FTransform &Xf = Transforms[i];
+				if (Xf.Equals(FTransform::Identity))
+				{
+					// If we ever support animated substructures, we'll need all
+					// of them wrapped by a transform.
+					ImplicitObjects.Add(MoveTemp(ObjPtr));
+				}
+				else
+				{
+					ImplicitObjects.Add(
+						TUniquePtr<Chaos::TImplicitObject<float, 3>>(
+							new Chaos::TImplicitObjectTransformed<float, 3, true>(
+								MakeSerializable(ObjPtr),
+								MoveTemp(ObjPtr),
+								Chaos::TRigidTransform<float, 3>(Xf))));
+				}
+			}
+			return new Chaos::TImplicitObjectUnion<float, 3>(MoveTemp(ImplicitObjects));
+		}
+	}
+
+protected:
+	friend class FBoneHierarchy;
+
+	void SetParent(FAnalyticImplicitGroup *InParent)
+	{ Parent = InParent; }
+	const FAnalyticImplicitGroup* GetParent() const
+	{ return Parent; }
+
+	void AddChild(FAnalyticImplicitGroup *Child)
+	{ Children.Add(Child); }
+	const TArray<FAnalyticImplicitGroup*>& GetChildren() const
+	{ return Children; }
+
+	void ClearHierarchy()
+	{ Parent = nullptr; Children.Reset(); }
+
+	template<class TImplicitShape>
+	void CullDeepPoints(TArray<Chaos::TVector<float,3>>& Points, const TImplicitShape& Shape, const FTransform& Xf)
+	{
+		const Chaos::TBox<float, 3>& BBox = Shape.BoundingBox();
+		const float Tolerance = -BBox.Extents().Max() / 100.f; // -1/100th the largest dimension
+		if (Xf.Equals(FTransform::Identity))
+		{
+			for (int32 i = Points.Num(); i--; )
+			{
+				const Chaos::TVector<float, 3>& LocalPoint = Points[i];
+				const float Phi = Shape.SignedDistance(LocalPoint);
+				if (Phi < Tolerance)
+				{
+					Points.RemoveAt(i);
+				}
+			}
+		}
+		else
+		{
+			const FTransform InvXf = Xf.Inverse();
+			for (int32 i = Points.Num(); i--; )
+			{
+				const Chaos::TVector<float, 3> LocalPoint = InvXf.TransformPosition(Points[i]);
+				const float Phi = Shape.SignedDistance(LocalPoint);
+				if (Phi < Tolerance)
+				{
+					Points.RemoveAt(i);
+				}
+			}
+		}
+	}
+
+	void CullDeepPoints(TArray<Chaos::TVector<float, 3>>& Points, const int32 SkipIndex)
+	{
+		int32 TransformIndex = 0;
+		for (Chaos::TSphere<float, 3>* Sphere : Spheres)
+		{
+			if (TransformIndex != SkipIndex)
+			{
+				CullDeepPoints(Points, *Sphere, Transforms[TransformIndex]);
+			}
+			TransformIndex++;
+		}
+		for (Chaos::TBox<float, 3>* Box : Boxes)
+		{
+			if (TransformIndex != SkipIndex)
+			{
+				CullDeepPoints(Points, *Box, Transforms[TransformIndex]);
+			}
+			TransformIndex++;
+		}
+		for (Chaos::TCapsule<float>* Capsule : Capsules)
+		{
+			if (TransformIndex != SkipIndex)
+			{
+				CullDeepPoints(Points, *Capsule, Transforms[TransformIndex]);
+			}
+			TransformIndex++;
+		}
+		for (Chaos::TTaperedCylinder<float>* TaperedCylinder : TaperedCylinders)
+		{
+			if (TransformIndex != SkipIndex)
+			{
+				CullDeepPoints(Points, *TaperedCylinder, Transforms[TransformIndex]);
+			}
+			TransformIndex++;
+		}
+		for (Chaos::TConvex<float, 3>* Convex : ConvexHulls)
+		{
+			if (TransformIndex != SkipIndex)
+			{
+				CullDeepPoints(Points, *Convex, Transforms[TransformIndex]);
+			}
+			TransformIndex++;
+		}
+		for (Chaos::TLevelSet<float, 3>* LevelSet : LevelSets)
+		{
+			if (TransformIndex != SkipIndex)
+			{
+				CullDeepPoints(Points, *LevelSet, Transforms[TransformIndex]);
+			}
+			TransformIndex++;
+		}
+	}
+
+	Chaos::TImplicitObject<float, 3>* TransferImplicitObj(int32 Idx)
+	{
+		Chaos::TImplicitObject<float, 3>* Obj = nullptr;
+		if (Idx < Spheres.Num())
+		{
+			Obj = Spheres[Idx]; 
+			Spheres[Idx] = nullptr;
+			return Obj;
+		}
+		Idx -= Spheres.Num();
+
+		if (Idx < Boxes.Num())
+		{
+			Obj = Boxes[Idx];
+			Boxes[Idx] = nullptr;
+			return Obj;
+		}
+		Idx -= Boxes.Num();
+
+		if (Idx < Capsules.Num())
+		{
+			Obj = Capsules[Idx];
+			Capsules[Idx] = nullptr;
+			return Obj;
+		}
+		Idx -= Capsules.Num();
+
+		if (Idx < TaperedCylinders.Num())
+		{
+			Obj = TaperedCylinders[Idx];
+			TaperedCylinders[Idx] = nullptr;
+			return Obj;
+		}
+		Idx -= TaperedCylinders.Num();
+
+		if (Idx < ConvexHulls.Num())
+		{
+			Obj = ConvexHulls[Idx];
+			ConvexHulls[Idx] = nullptr;
+			return Obj;
+		}
+		Idx -= ConvexHulls.Num();
+
+		if (Idx < LevelSets.Num())
+		{
+			Obj = LevelSets[Idx];
+			LevelSets[Idx] = nullptr;
+			return Obj;
+		}
+
+		check(false); // Index out of range!
+		return Obj;
+	}
+
+protected:
+	friend class USkeletalMeshSimulationComponent;
+	friend struct FPhysicsAssetSimulationUtil;
+
+	FName BoneName;
+	int32 BoneIndex;
+	int32 ParentBoneIndex;
+
+	int32 RigidBodyId;
+	EObjectStateTypeEnum RigidBodyState;
+
+	// FkSphereElem and FKTaperedCapsuleElem ends
+	TArray<Chaos::TSphere<float, 3>*> Spheres;
+	// FKBoxElem
+	TArray<Chaos::TBox<float, 3>*> Boxes;
+	// FKSphylElem - Z axis is capsule axis
+	TArray<Chaos::TCapsule<float>*> Capsules;
+	// FKTaperedCapsuleElem - Z axis is the capsule axis
+	TArray<Chaos::TTaperedCylinder<float>*> TaperedCylinders;
+	// FKConvexElem
+	TArray<Chaos::TConvex<float, 3>*> ConvexHulls;
+	// Chaos::TConvex replacement
+	TArray<Chaos::TLevelSet<float, 3>*> LevelSets;
+
+	TArray<Chaos::TVector<float, 3>> ContiguousCollisionPoints;
+	TArray<TArray<Chaos::TVector<float, 3>>> CollisionPoints;
+	TArray<TArray<Chaos::TVector<int32, 3>>> CollisionTriangles;
+
+	TArray<FTransform> Transforms;
+	FTransform RefBoneXf;
+
+	FAnalyticImplicitGroup* Parent;
+	TArray<FAnalyticImplicitGroup*> Children;
+};
+
+
+
+
+
+#endif // INCLUDE_CHAOS
\ No newline at end of file
diff --git a/Engine/Source/Runtime/Experimental/ChaosSolvers/Public/PhysicsProxy/FieldSystemPhysicsProxy.h b/Engine/Source/Runtime/Experimental/ChaosSolvers/Public/PhysicsProxy/FieldSystemPhysicsProxy.h
new file mode 100644
--- /dev/null
+++ b/Engine/Source/Runtime/Experimental/ChaosSolvers/Public/PhysicsProxy/FieldSystemPhysicsProxy.h
@@ -0,0 +1,67 @@
+// Copyright 1998-2019 Epic Games, Inc. All Rights Reserved.
+
+#pragma once
+
+#if INCLUDE_CHAOS
+
+#include "Chaos/Framework/PhysicsProxy.h"
+#include "Field/FieldSystem.h"
+#include "Chaos/PBDPositionConstraints.h"
+
+struct FKinematicProxy;
+
+class FStubFieldData : public Chaos::FParticleData {
+	void Reset() {};
+};
+
+class CHAOSSOLVERS_API FFieldSystemPhysicsProxy : public TPhysicsProxy<FFieldSystemPhysicsProxy, FStubFieldData>
+{
+	typedef TPhysicsProxy<FFieldSystemPhysicsProxy, FStubFieldData> Base;
+
+public:
+	FFieldSystemPhysicsProxy() = delete;
+	FFieldSystemPhysicsProxy(UObject* InOwner);
+	virtual ~FFieldSystemPhysicsProxy();
+
+	// Callbacks
+	bool IsSimulating() const;
+	void FieldParameterUpdateCallback(Chaos::FPhysicsSolver* InSolver, FParticlesType& InParticles, Chaos::TArrayCollectionArray<float>& Strains, Chaos::TPBDPositionConstraints<float, 3>& PositionTarget, TMap<int32, int32>& PositionTargetedParticles, const TArray<FKinematicProxy>& AnimatedPositions, const float InTime);
+	void FieldForcesUpdateCallback(Chaos::FPhysicsSolver* InSolver, FParticlesType& Particles, Chaos::TArrayCollectionArray<FVector> & Force, Chaos::TArrayCollectionArray<FVector> & Torque, const float Time);
+	void EndFrameCallback(const float InDt);
+	void BufferCommand(Chaos::FPhysicsSolver* InSolver, const FFieldSystemCommand& InCommand);
+
+	// Inactive Callbacks
+	void ParameterUpdateCallback(FParticlesType& InParticles, const float InTime){}
+	void UpdateKinematicBodiesCallback(const FParticlesType& InParticles, const float InDt, const float InTime, FKinematicProxy& InKinematicProxy){}
+	void BindParticleCallbackMapping(Chaos::TArrayCollectionArray<PhysicsProxyWrapper> & PhysicsProxyReverseMap, Chaos::TArrayCollectionArray<int32> & ParticleIDReverseMap) {}
+	void StartFrameCallback(const float InDt, const float InTime){}
+	void CreateRigidBodyCallback(FParticlesType& InOutParticles){}
+	void DisableCollisionsCallback(TSet<TTuple<int32, int32>>& InPairs){}
+	void AddForceCallback(FParticlesType& InParticles, const float InDt, const int32 InIndex){}
+
+	Chaos::FParticleData* NewData() { return nullptr; }
+	void SyncBeforeDestroy(){}
+	void OnRemoveFromScene();
+	void PushToPhysicsState(const Chaos::FParticleData*) {};
+	void BufferPhysicsResults(){}
+	void FlipBuffer(){}
+	void PullFromPhysicsState(){}
+	EPhysicsProxyType ConcreteType() {return EPhysicsProxyType::FieldType;}
+
+	/*
+	* ContiguousIndices
+	*   Generates a mapping between the Position array and the results array. When EFieldResolutionType is set to Maximum the complete
+	*   particle mapping is provided from the Particles.X to Particles.Attribute, when Minimum is set only the ActiveIndices and the
+	*   direct children of the active clusters are set in the IndicesArray.
+	*/
+	static void ContiguousIndices(TArray<ContextIndex>& IndicesArray, const Chaos::FPhysicsSolver* RigidSolver, EFieldResolutionType ResolutionType, bool bForce);
+
+private:
+
+	TArray<FFieldSystemCommand>* GetSolverCommandList(const Chaos::FPhysicsSolver* InSolver);
+
+	FCriticalSection CommandLock;
+	TMap<Chaos::FPhysicsSolver*, TArray<FFieldSystemCommand>*> Commands;
+};
+
+#endif
diff --git a/Engine/Source/Runtime/Experimental/ChaosSolvers/Public/PhysicsProxy/GeometryCollectionPhysicsProxy.h b/Engine/Source/Runtime/Experimental/ChaosSolvers/Public/PhysicsProxy/GeometryCollectionPhysicsProxy.h
new file mode 100644
--- /dev/null
+++ b/Engine/Source/Runtime/Experimental/ChaosSolvers/Public/PhysicsProxy/GeometryCollectionPhysicsProxy.h
@@ -0,0 +1,231 @@
+// Copyright 1998-2019 Epic Games, Inc. All Rights Reserved.
+
+#pragma once
+
+#include "GeometryCollection/GeometryCollection.h"
+
+#if INCLUDE_CHAOS
+
+#include "Chaos/Framework/PhysicsProxy.h"
+
+#include "GeometryCollection/ManagedArray.h"
+#include "GeometryCollection/GeometryCollectionSimulationCoreTypes.h"
+#include "GeometryCollection/GeometryCollectionCollisionStructureManager.h"
+#include "Chaos/Framework/BufferedData.h"
+
+struct FGeometryCollectionResults
+{
+	FGeometryCollectionResults();
+
+	int32 BaseIndex;
+	int32 NumParticlesAdded;
+	TArray<bool> DisabledStates;
+	TManagedArray<FTransform> Transforms;
+	TArray<FMatrix> GlobalTransforms;
+	TManagedArray<int32> RigidBodyIds;
+	TArray<FTransform> ParticleToWorldTransforms;
+	TManagedArray<int32> Level;
+	TManagedArray<int32> Parent;
+	TManagedArray<TSet<int32>> Children;
+	TManagedArray<int32> SimulationType;
+	TManagedArray<int32> StatusFlags;
+
+	bool IsObjectDynamic;
+	bool IsObjectLoading;
+
+	FBoxSphereBounds WorldBounds;
+};
+
+namespace Chaos
+{
+template <typename T>
+class TSerializablePtr;
+class FErrorReporter;
+}
+
+class FStubGeometryCollectionData : public Chaos::FParticleData {
+	void Reset() {};
+};
+
+class CHAOSSOLVERS_API FGeometryCollectionPhysicsProxy : public TPhysicsProxy<FGeometryCollectionPhysicsProxy, FStubGeometryCollectionData>
+{
+	typedef TPhysicsProxy<FGeometryCollectionPhysicsProxy, FStubGeometryCollectionData> Base;
+
+public:
+	// collection attributes
+	static FName SimplicialsAttribute;
+	static FName ImplicitsAttribute;
+
+	typedef FCollisionStructureManager::FSimplicial FSimplicial;
+
+	/** Proxy publics */
+	using FInitFunc = TFunction<void(FSimulationParameters&)>;
+	using FCacheSyncFunc = TFunction<void(const FGeometryCollectionResults&)>;
+	using FFinalSyncFunc = TFunction<void(const FRecordedTransformTrack&)>;
+	/** ----------------------- */
+
+	FGeometryCollectionPhysicsProxy() = delete;
+	FGeometryCollectionPhysicsProxy(UObject* InOwner, FGeometryDynamicCollection* InDynamicCollection, FInitFunc InInitFunc, FCacheSyncFunc InCacheSyncFunc, FFinalSyncFunc InFinalSyncFunc);
+	virtual ~FGeometryCollectionPhysicsProxy();
+
+	void Initialize();
+	void Reset();
+
+	/** Solver Object interface */
+	bool IsSimulating() const;
+	void UpdateKinematicBodiesCallback(const FParticlesType& InParticles, const float InDt, const float InTime, FKinematicProxy& InKinematicProxy);
+	void StartFrameCallback(const float InDt, const float InTime);
+	void EndFrameCallback(const float InDt);
+	void BindParticleCallbackMapping(Chaos::TArrayCollectionArray<PhysicsProxyWrapper> & PhysicsProxyReverseMap, Chaos::TArrayCollectionArray<int32> & ParticleIDReverseMap);
+	void CreateRigidBodyCallback(FParticlesType& InOutParticles);
+	void ParameterUpdateCallback(FParticlesType& InParticles, const float InTime);
+	void DisableCollisionsCallback(TSet<TTuple<int32, int32>>& InPairs);
+	void AddForceCallback(FParticlesType& InParticles, const float InDt, const int32 InIndex);
+	void FieldForcesUpdateCallback(Chaos::FPhysicsSolver* InSolver, FParticlesType& Particles, Chaos::TArrayCollectionArray<FVector> & Force, Chaos::TArrayCollectionArray<FVector> & Torque, const float Time);
+	void ActivateBodies();
+
+	Chaos::FParticleData* NewData() { return nullptr; }
+	void SyncBeforeDestroy();
+	void OnRemoveFromScene();
+	void PushToPhysicsState(const Chaos::FParticleData*) {};
+	void BufferPhysicsResults();
+	void FlipBuffer();
+	void PullFromPhysicsState();
+	EPhysicsProxyType ConcreteType() { return EPhysicsProxyType::GeometryCollectionType; }
+	/** ----------------------- */
+
+	void MergeRecordedTracks(const FRecordedTransformTrack& A, const FRecordedTransformTrack& B, FRecordedTransformTrack& Target);
+	FRecordedFrame& InsertRecordedFrame(FRecordedTransformTrack& InTrack, float InTime);
+	void BuildClusters(uint32 CollectionClusterIndex, const TArray<uint32>& CollectionChildIDs, const TArray<uint32>& ChildIDs, const Chaos::FClusterCreationParameters<float> & Parameters);
+	void BufferCommand(Chaos::FPhysicsSolver* InSolver, const FFieldSystemCommand& Command) { Commands.Add(Command); }
+
+	static void InitializeSharedCollisionStructures(Chaos::FErrorReporter& ErrorReporter, FGeometryCollection& RestCollection, const FSharedSimulationParameters& SharedParams);
+
+	static void InitRemoveOnFracture(FGeometryCollection& RestCollection, const FSharedSimulationParameters& SharedParams);
+	const Chaos::TBufferedData<FGeometryCollectionResults>& GetPhysicsResults() const { return Results; }
+
+	// Game thread and physics thread access to the dynamic collection. 
+	const FGeometryDynamicCollection* GetGeometryDynamicCollection_PhysicsThread() { return SimulationCollection.Get(); }
+	const FGeometryDynamicCollection* GetGeometryDynamicCollection_GameThread() { return GTDynamicCollection; }
+
+	const FSimulationParameters& GetSimulationParameters() { return Parameters; }
+
+	int32 GetBaseParticleIndex() const { return BaseParticleIndex; }
+	int32 GetNumAddedParticles() const { return NumParticles; }
+
+	/*
+	* ContiguousIndices
+	*   Generates a mapping between the Position array and the results array. When EFieldResolutionType is set to Maximum the complete
+	*   particle mapping is provided from the Particles.X to Particles.Attribute, when Minimum is set only the ActiveIndices and the
+	*   direct children of the active clusters are set in the IndicesArray.
+	*/
+	void ContiguousIndices(TArray<ContextIndex>& IndicesArray, const Chaos::FPhysicsSolver* RigidSolver, EFieldResolutionType ResolutionType, bool bForce);
+
+	//
+	void SetCollisionParticlesPerObjectFraction(float CollisionParticlesPerObjectFractionIn) {CollisionParticlesPerObjectFraction = CollisionParticlesPerObjectFractionIn;}
+
+	/*
+	* For Testing ONLY!
+	*/
+	const TManagedArray<int32>& RigidBodyIDArray_TestingAccess() const { return RigidBodyID; }
+	const TManagedArray<FTransform>& MassToLocal_GameThreadAccess() const { return MassToLocal; }
+
+protected:
+	int32 CalculateHierarchyLevel(FGeometryDynamicCollection* GeometryCollection, int32 TransformIndex);
+	void CreateDynamicAttributes();
+	void InitializeKinematics(FParticlesType& Particles, const TManagedArray<int32>& ObjectType);
+	void InitializeRemoveOnFracture(FParticlesType& Particles, const TManagedArray<int32>& DynamicState);
+	void ProcessCommands(FParticlesType& Particles, const float Time);
+	void PushKinematicStateToSolver(FParticlesType& Particles);
+
+
+
+	
+private:
+	ESimulationInitializationState InitializedState;
+
+	TManagedArray<FTransform> MassToLocal;
+	TManagedArray<int32> CollisionMask;
+	TManagedArray<int32> CollisionStructureID;
+	TManagedArray<int32> RigidBodyID;
+	TManagedArray<int32> SolverClusterID;
+	TManagedArray<FVector> InitialAngularVelocity;
+	TManagedArray<FVector> InitialLinearVelocity;
+	TManagedArray<bool> SimulatableParticles;
+	TManagedArray<TUniquePtr<FSimplicial> > Simplicials;
+	TManagedArray<Chaos::TSerializablePtr<Chaos::TImplicitObject<float, 3>>> Implicits;
+	TArray<int32> EndFrameUnparentingBuffer;
+
+	// This is a subset of the geometry group that are used in the transform hierarchy to represent geometry
+	TArray<FBox> ValidGeometryBoundingBoxes;
+	TArray<int32> ValidGeometryTransformIndices;
+
+	FSimulationParameters Parameters;
+	TArray<FFieldSystemCommand> Commands;
+
+	TFunction<void(void)> ResetAnimationCacheCallback;
+	TFunction<void(const TArrayView<FTransform> &)> UpdateTransformsCallback;
+	TFunction<void(const int32 & CurrentFrame, const TManagedArray<int32> & RigidBodyID, const TManagedArray<int32>& Level, const TManagedArray<int32>& Parent, const TManagedArray<TSet<int32>>& Children, const TManagedArray<uint32>& SimulationType, const TManagedArray<uint32>& StatusFlags, const FParticlesType& Particles)> UpdateRestStateCallback;
+	TFunction<void(float SolverTime, const TManagedArray<int32> & RigidBodyID, const FParticlesType& Particles, const FCollisionConstraintsType& CollisionRule)> UpdateRecordedStateCallback;
+	TFunction<void(FRecordedTransformTrack& InTrack)> CommitRecordedStateCallback;
+
+	// Index of the first particles for this collection in the larger particle array
+	int32 BaseParticleIndex;
+
+	// Number of particles added by this collection
+	int32 NumParticles;
+
+	// Time since this object started simulating
+	float ProxySimDuration;
+
+	void UpdateRecordedState(float SolverTime, const TManagedArray<int32>& RigidBodyID, const TManagedArray<int32>& CollectionClusterID, const Chaos::TArrayCollectionArray<bool>& InternalCluster, const FParticlesType& Particles, const FCollisionConstraintsType& CollisionRule);
+
+	void AddCollisionToCollisionData(FRecordedFrame* ExistingFrame, const FParticlesType& Particles, const Chaos::TPBDCollisionConstraint<float, 3>::FRigidBodyContactConstraint& Constraint);
+
+	void UpdateCollisionData(const FParticlesType& Particles, const FCollisionConstraintsType& CollisionRule, FRecordedFrame* ExistingFrame);
+
+	void AddBreakingToBreakingData(FRecordedFrame* ExistingFrame, const FParticlesType& Particles, const Chaos::TBreakingData<float, 3>& Breaking);
+
+	void UpdateBreakingData(const FParticlesType& Particles, FRecordedFrame* ExistingFrame);
+
+	void UpdateTrailingData(const FParticlesType& Particles, FRecordedFrame* ExistingFrame);
+
+	// Duplicated dynamic collection for use on the physics thread, copied to the game thread on sync
+	TUniquePtr<FGeometryDynamicCollection> SimulationCollection;
+
+	// Dynamic collection on the game thread - used to populate the simulated collection
+	FGeometryDynamicCollection* GTDynamicCollection;
+
+	// bodies waiting to activate on initialization
+	TArray<uint32> PendingActivationList;
+
+	// Storage for the recorded frame information when we're caching the geometry component results.
+	// Synced back to the component with SyncBeforeDestroy
+	FRecordedTransformTrack RecordedTracks;
+
+	// Functions to handle engine-side events
+	FInitFunc InitFunc;
+	FCacheSyncFunc CacheSyncFunc;
+	FFinalSyncFunc FinalSyncFunc;
+
+	// Sync frame numbers so we don't do many syncs when physics is running behind
+	uint32 LastSyncCountGT;
+
+	// Records current dynamic state
+	bool IsObjectDynamic;
+
+	// Indicate when loaded
+	bool IsObjectLoading;
+
+	// Per object collision fraction.
+	float CollisionParticlesPerObjectFraction;
+
+	// Double buffer of geom collection result data
+	Chaos::TBufferedData<FGeometryCollectionResults> Results;
+
+	/** ----------------------- */
+
+};
+
+CHAOSSOLVERS_API void BuildSimulationData(Chaos::FErrorReporter& ErrorReporter, FGeometryCollection& GeometryCollection, const FSharedSimulationParameters& SharedParams);
+#endif
diff --git a/Engine/Source/Runtime/Experimental/ChaosSolvers/Public/PhysicsProxy/SkeletalMeshPhysicsProxy.h b/Engine/Source/Runtime/Experimental/ChaosSolvers/Public/PhysicsProxy/SkeletalMeshPhysicsProxy.h
new file mode 100644
--- /dev/null
+++ b/Engine/Source/Runtime/Experimental/ChaosSolvers/Public/PhysicsProxy/SkeletalMeshPhysicsProxy.h
@@ -0,0 +1,187 @@
+// Copyright 1998-2019 Epic Games, Inc. All Rights Reserved.
+#pragma once
+#if INCLUDE_CHAOS
+
+#include "Chaos/Framework/PhysicsProxy.h"
+#include "BoneHierarchy.h"
+#include "GeometryCollection/GeometryCollectionSimulationTypes.h"
+
+#include "Chaos/PBDJointConstraints.h"
+#include "Chaos/PBDConstraintRule.h"
+#include "ChaosSolvers/Public/Framework/TripleBufferedData.h"
+
+
+struct CHAOSSOLVERS_API FSkeletalMeshPhysicsProxyParams
+{
+	FSkeletalMeshPhysicsProxyParams()
+		: Name("")
+		, InitialTransform(FTransform::Identity)
+		, InitialLinearVelocity(FVector::ZeroVector)
+		, InitialAngularVelocity(FVector::ZeroVector)
+
+		, ObjectType(EObjectStateTypeEnum::Chaos_Object_Kinematic)
+
+		, CollisionType(ECollisionTypeEnum::Chaos_Volumetric)
+		, ParticlesPerUnitArea(0.1)
+		, MinNumParticles(0)
+		, MaxNumParticles(50)
+		, MinRes(5)
+		, MaxRes(10)
+		, CollisionGroup(0)
+#if 0
+		, bEnableClustering(false)
+		, ClusterGroupIndex(0)
+		, MaxClusterLevel(100)
+		, DamageThreshold(250.)
+#endif
+		, Density(2.4)
+		, MinMass(0.001)
+		, MaxMass(1.e6)
+
+		, bSimulating(false)
+	{}
+
+	FString Name;
+
+	//
+	// Analytic implicit representation
+	//
+
+	FBoneHierarchy BoneHierarchy;
+
+	//
+	// Mesh
+	//
+
+	TArray<FVector> MeshVertexPositions;
+	TArray<FIntVector> Triangles;
+
+	//
+	// Transform hierarchy
+	//
+
+	FTransform InitialTransform;
+	FTransform LocalToWorld;
+	FVector InitialLinearVelocity;
+	FVector InitialAngularVelocity;
+
+	Chaos::TSerializablePtr<Chaos::TChaosPhysicsMaterial<float>> PhysicalMaterial;	// @todo(ccaulfield): should be per-shape
+	EObjectStateTypeEnum ObjectType;												// @todo(ccaulfield): should be per-body
+
+	ECollisionTypeEnum CollisionType;
+	float ParticlesPerUnitArea;
+	int32 MinNumParticles;
+	int32 MaxNumParticles;
+	int32 MinRes;
+	int32 MaxRes;
+	int32 CollisionGroup;
+#if 0
+	bool bEnableClustering;
+	int32 ClusterGroupIndex;
+	int32 MaxClusterLevel;
+	float DamageThreshold;
+#endif
+	float Density;
+	float MinMass;
+	float MaxMass;
+
+	bool bSimulating;
+};
+
+// @todo(ccaulfield): make the IO structures private again - only the hierarchy should be required outside the PhysicsProxy
+struct FSkeletalMeshPhysicsProxyInputs
+{
+	TArray<FTransform> Transforms;
+	TArray<FVector> LinearVelocities;
+	TArray<FVector> AngularVelocities;
+};
+
+struct FSkeletalMeshPhysicsProxyOutputs : public Chaos::FParticleData
+{
+	TArray<FTransform> Transforms;
+	TArray<FVector> LinearVelocities;
+	TArray<FVector> AngularVelocities;
+	void Reset() {
+		Transforms = TArray<FTransform>();
+		LinearVelocities = TArray<FVector>();
+		AngularVelocities = TArray<FVector>();
+	}
+};
+
+
+class CHAOSSOLVERS_API FSkeletalMeshPhysicsProxy : public TPhysicsProxy<FSkeletalMeshPhysicsProxy, FSkeletalMeshPhysicsProxyOutputs>
+{
+	typedef TPhysicsProxy<FSkeletalMeshPhysicsProxy, FSkeletalMeshPhysicsProxyOutputs> Base;
+public:
+
+
+	using FInitFunc = TFunction<void(FSkeletalMeshPhysicsProxyParams& OutParams)>;
+	using FInputFunc = TFunction<bool(const float Dt, FSkeletalMeshPhysicsProxyParams& OutParams)>;
+
+
+	FSkeletalMeshPhysicsProxy() = delete;
+	FSkeletalMeshPhysicsProxy(UObject* InOwner, const FInitFunc& InitFunc);
+	~FSkeletalMeshPhysicsProxy();
+
+	/** Solver Object interface */
+	void Initialize();
+	bool IsSimulating() const;
+	void UpdateKinematicBodiesCallback(const FParticlesType& Particles, const float Dt, const float Time, FKinematicProxy& Proxy);
+	void StartFrameCallback(const float InDt, const float InTime);
+	void EndFrameCallback(const float InDt);
+	void CreateRigidBodyCallback(FParticlesType& InOutParticles);
+	void ParameterUpdateCallback(FParticlesType& InParticles, const float InTime);
+	void DisableCollisionsCallback(TSet<TTuple<int32, int32>>& InPairs);
+	void AddForceCallback(FParticlesType& InParticles, const float InDt, const int32 InIndex);
+	void FieldForcesUpdateCallback(Chaos::FPhysicsSolver* InSolver, FParticlesType& Particles, Chaos::TArrayCollectionArray<FVector>& Force, Chaos::TArrayCollectionArray<FVector>& Torque, const float Time) {}
+
+	void BindParticleCallbackMapping(Chaos::TArrayCollectionArray<PhysicsProxyWrapper>& PhysicsProxyReverseMap, Chaos::TArrayCollectionArray<int32>& ParticleIDReverseMap);
+
+	void BufferCommand(Chaos::FPhysicsSolver* InSolver, const FFieldSystemCommand& InCommmand) {}
+
+	Chaos::FParticleData* NewData() { return nullptr; }
+	void SyncBeforeDestroy();
+	void OnRemoveFromScene();
+	void PushToPhysicsState(const Chaos::FParticleData*) {};
+	void BufferPhysicsResults();
+	void FlipBuffer();
+	void PullFromPhysicsState();
+	EPhysicsProxyType ConcreteType() { return EPhysicsProxyType::SkeletalMeshType; }
+	/** ----------------------- */
+
+	/**
+	 *
+	 */
+	void Reset();
+
+	/**
+	 * Capture the current animation pose for use by the physics.
+	 * Called by game thread via the owning component's tick.
+	 */
+	void CaptureInputs(const float Dt, const FInputFunc& InputFunc);
+
+	/** 
+	 */
+	const FSkeletalMeshPhysicsProxyOutputs* GetOutputs() const { return CurrentOutputConsumerBuffer; }
+
+	const FBoneHierarchy& GetBoneHierarchy() const { return Parameters.BoneHierarchy; }
+
+private:
+	typedef Chaos::TPBDJointConstraints<float, 3> FJointConstraints;
+	typedef Chaos::TPBDConstraintIslandRule<FJointConstraints, float, 3> FJointConstraintsRule;
+
+	FSkeletalMeshPhysicsProxyParams Parameters;
+	TArray<int32> RigidBodyIds;
+	FJointConstraints JointConstraints;
+	FJointConstraintsRule JointConstraintsRule;
+	// @todo(ccaulfield): sort out the IO buffer stuff
+	Chaos::TTripleBufferedData<FSkeletalMeshPhysicsProxyInputs> InputBuffers;
+	Chaos::TBufferedData<FSkeletalMeshPhysicsProxyOutputs> OutputBuffers;
+	FSkeletalMeshPhysicsProxyInputs* NextInputProducerBuffer;				// Buffer for the game to write to next
+	const FSkeletalMeshPhysicsProxyOutputs* CurrentOutputConsumerBuffer;	// Buffer for the game to read from next
+	bool bInitializedState;
+
+	FInitFunc InitFunc;
+};
+
+#endif // INCLUDE_CHAOS
diff --git a/Engine/Source/Runtime/Experimental/ChaosSolvers/Public/PhysicsProxy/StaticMeshPhysicsProxy.h b/Engine/Source/Runtime/Experimental/ChaosSolvers/Public/PhysicsProxy/StaticMeshPhysicsProxy.h
new file mode 100644
--- /dev/null
+++ b/Engine/Source/Runtime/Experimental/ChaosSolvers/Public/PhysicsProxy/StaticMeshPhysicsProxy.h
@@ -0,0 +1,134 @@
+// Copyright 1998-2019 Epic Games, Inc. All Rights Reserved.
+
+#pragma once
+
+#if INCLUDE_CHAOS
+
+#include "Chaos/Framework/PhysicsProxy.h"
+#include "GeometryCollection/GeometryCollectionSimulationTypes.h"
+#include "Chaos/Framework/BufferedData.h"
+
+struct FPhysicsProxyKinematicUpdate
+{
+	FTransform NewTransform;
+	FVector NewVelocity;
+};
+
+struct FStubSkeletalMeshData : public Chaos::FParticleData {
+	void Reset() { };
+};
+
+class CHAOSSOLVERS_API FStaticMeshPhysicsProxy : public TPhysicsProxy<FStaticMeshPhysicsProxy, FStubSkeletalMeshData>
+{
+	typedef  TPhysicsProxy<FStaticMeshPhysicsProxy, FStubSkeletalMeshData> Base;
+public:
+
+	struct FShapeParams
+	{
+		FVector BoxExtents;
+		FVector2D CapsuleHalfHeightAndRadius;
+		float SphereRadius;
+
+		FShapeParams() 
+		{ 
+			FMemory::Memset(this, 0, sizeof(FShapeParams));
+		}
+	};
+
+	struct Params
+	{
+		Params()
+			: Name("")
+			, InitialTransform(FTransform::Identity)
+			, InitialLinearVelocity(FVector::ZeroVector)
+			, InitialAngularVelocity(FVector::ZeroVector)
+			, ObjectType(EObjectStateTypeEnum::Chaos_Object_Dynamic)
+			, ShapeType(EImplicitTypeEnum::Chaos_Max)
+			, bSimulating(false)
+			, TargetTransform(nullptr)
+			, Mass(0.0f)
+			, MinRes(5)
+			, MaxRes(10)
+		{}
+
+		FString Name;
+		Chaos::TParticles<float,3> MeshVertexPositions;
+		TArray<Chaos::TVector<int32, 3> > TriIndices;
+		FShapeParams ShapeParams;
+		FTransform InitialTransform;
+		FVector InitialLinearVelocity;
+		FVector InitialAngularVelocity;
+		EObjectStateTypeEnum ObjectType;
+		EImplicitTypeEnum ShapeType;
+		bool bSimulating;
+		FTransform* TargetTransform;
+		Chaos::TSerializablePtr<Chaos::TChaosPhysicsMaterial<float>> PhysicalMaterial;
+		float Mass;
+		int32 MinRes;
+		int32 MaxRes;
+	};
+
+	// Engine interface functions
+	using FCallbackInitFunc = TFunction<void(FStaticMeshPhysicsProxy::Params&)>;
+	using FSyncDynamicFunc = TFunction<void(const FTransform&)>;
+
+	FStaticMeshPhysicsProxy() = delete;
+	FStaticMeshPhysicsProxy(UObject* InOwner, FCallbackInitFunc InInitFunc, FSyncDynamicFunc InSyncFunc);
+
+	void Initialize();
+	void Reset();
+
+	/** Stores latest update, to be applied at next opportunity (via UpdateKinematicBodiesCallback). */
+	void BufferKinematicUpdate(const FPhysicsProxyKinematicUpdate& InParamUpdate);
+
+	/** Solver Object interface */
+	bool IsSimulating() const;
+	void UpdateKinematicBodiesCallback(const FParticlesType& Particles, const float Dt, const float Time, FKinematicProxy& Proxy);
+	void StartFrameCallback(const float InDt, const float InTime);
+	void EndFrameCallback(const float InDt);
+	void BindParticleCallbackMapping(Chaos::TArrayCollectionArray<PhysicsProxyWrapper> & PhysicsProxyReverseMap, Chaos::TArrayCollectionArray<int32> & ParticleIDReverseMap);
+	void CreateRigidBodyCallback(FParticlesType& InOutParticles);
+	void ParameterUpdateCallback(FParticlesType& InParticles, const float InTime);
+	void DisableCollisionsCallback(TSet<TTuple<int32, int32>>& InPairs);
+	void AddForceCallback(FParticlesType& InParticles, const float InDt, const int32 InIndex);
+	void FieldForcesUpdateCallback(Chaos::FPhysicsSolver* InSolver, FParticlesType& Particles, Chaos::TArrayCollectionArray<FVector> & Force, Chaos::TArrayCollectionArray<FVector> & Torque, const float Time) {};
+	void BufferCommand(Chaos::FPhysicsSolver* InSolver, const FFieldSystemCommand& InCommand) {};
+
+	void SyncBeforeDestroy() {};
+	void OnRemoveFromScene();
+	void PushToPhysicsState(const Chaos::FParticleData*) {};
+	void BufferPhysicsResults();
+	void FlipBuffer();
+	void PullFromPhysicsState();
+	Chaos::FParticleData* NewData() { return nullptr; }
+	EPhysicsProxyType ConcreteType() { return EPhysicsProxyType::StaticMeshType; }
+	/** ----------------------- */
+
+private:
+
+	Params Parameters;
+
+	bool bInitializedState;
+	int32 RigidBodyId;
+	FVector CenterOfMass;
+	FVector Scale;
+
+	// Transform that the callback object will write into during simulation.
+	// During sync this will be pushed back to the component
+	FTransform SimTransform;
+
+	// Double buffered result data
+	Chaos::TBufferedData<FTransform> Results;
+
+	/**
+	 *	External functions for setup and sync, called on the game thread during callback creation and syncing
+	 */
+	FCallbackInitFunc InitialiseCallbackParamsFunc;
+	FSyncDynamicFunc SyncDynamicTransformFunc;
+	//////////////////////////////////////////////////////////////////////////
+
+	bool bPendingKinematicUpdate;
+	FPhysicsProxyKinematicUpdate BufferedKinematicUpdate;
+};
+
+#endif
diff --git a/Engine/Source/Runtime/Experimental/ChaosSolvers/Public/PhysicsSolver.h b/Engine/Source/Runtime/Experimental/ChaosSolvers/Public/PhysicsSolver.h
new file mode 100644
--- /dev/null
+++ b/Engine/Source/Runtime/Experimental/ChaosSolvers/Public/PhysicsSolver.h
@@ -0,0 +1,5 @@
+// Copyright 1998-2019 Epic Games, Inc. All Rights Reserved.
+
+#pragma once
+#include "Chaos/Declares.h"
+#include "PBDRigidsSolver.h"
\ No newline at end of file
diff --git a/Engine/Source/Runtime/Experimental/ChaosSolvers/Public/SolverEventFilters.h b/Engine/Source/Runtime/Experimental/ChaosSolvers/Public/SolverEventFilters.h
--- a/Engine/Source/Runtime/Experimental/ChaosSolvers/Public/SolverEventFilters.h
+++ b/Engine/Source/Runtime/Experimental/ChaosSolvers/Public/SolverEventFilters.h
@@ -101,34 +101,85 @@ namespace Chaos
 	class CHAOSSOLVERS_API FSolverCollisionEventFilter
 	{
 	public:
+		FSolverCollisionEventFilter() {}
 		FSolverCollisionEventFilter(const FSolverCollisionFilterSettings& InSettings) : Settings(InSettings) {}
 
-		bool Pass(const Chaos::TCollisionData<float, 3>& InData);
-		bool Enabled() { return Settings.FilterEnabled; }
+		bool Pass(const Chaos::TCollisionData<float, 3>& InData) const;
+		bool Enabled() const { return Settings.FilterEnabled; }
+		void UpdateFilterSettings(const FSolverCollisionFilterSettings& InSettings) { Settings = InSettings; }
 
-		const FSolverCollisionFilterSettings& Settings;
+		FSolverCollisionFilterSettings Settings;
 	};
 
 	class CHAOSSOLVERS_API FSolverTrailingEventFilter
 	{
 	public:
+		FSolverTrailingEventFilter() {}
 		FSolverTrailingEventFilter(const FSolverTrailingFilterSettings &InSettings) : Settings(InSettings) {}
 
-		bool Pass(const Chaos::TTrailingData<float, 3>& InData);
-		bool Enabled() { return Settings.FilterEnabled; }
+		bool Pass(const Chaos::TTrailingData<float, 3>& InData) const;
+		bool Enabled() const { return Settings.FilterEnabled; }
+		void UpdateFilterSettings(const FSolverTrailingFilterSettings& InSettings) { Settings = InSettings; }
 
-		const FSolverTrailingFilterSettings& Settings;
+		FSolverTrailingFilterSettings Settings;
 	};
 
 	class CHAOSSOLVERS_API FSolverBreakingEventFilter
 	{
 	public:
+		FSolverBreakingEventFilter() {}
 		FSolverBreakingEventFilter(const FSolverBreakingFilterSettings& InSettings) : Settings(InSettings) {}
 
-		bool Pass(const Chaos::TBreakingData<float, 3>& InData);
-		bool Enabled() { return Settings.FilterEnabled; }
+		bool Pass(const Chaos::TBreakingData<float, 3>& InData) const;
+		bool Enabled() const { return Settings.FilterEnabled; }
+		void UpdateFilterSettings(const FSolverBreakingFilterSettings& InSettings) { Settings = InSettings; }
 
-		const FSolverBreakingFilterSettings& Settings;
+		FSolverBreakingFilterSettings Settings;
+	};
+
+
+	/**
+	 * Container for the Solver Event Filters that have settings exposed through the Solver Actor
+	 */
+	class FSolverEventFilters
+	{
+	public:
+		FSolverEventFilters()
+			: CollisionFilter(new FSolverCollisionEventFilter())
+			, BreakingFilter(new FSolverBreakingEventFilter())
+			, TrailingFilter(new FSolverTrailingEventFilter())
+			, CollisionEventsEnabled(false)
+			, BreakingEventsEnabled(false)
+			, TrailingEventsEnabled(false)
+		{}
+
+		void SetGenerateCollisionEvents(bool bDoGenerate) { CollisionEventsEnabled = bDoGenerate; }
+		void SetGenerateBreakingEvents(bool bDoGenerate) { BreakingEventsEnabled = bDoGenerate; }
+		void SetGenerateTrailingEvents(bool bDoGenerate) { TrailingEventsEnabled = bDoGenerate; }
+
+		/* Const access */
+		FSolverCollisionEventFilter* GetCollisionFilter() const { return CollisionFilter.Get(); }
+		FSolverBreakingEventFilter* GetBreakingFilter() const { return BreakingFilter.Get(); }
+		FSolverTrailingEventFilter* GetTrailingFilter() const { return TrailingFilter.Get(); }
+
+		/* non-const access */
+		FSolverCollisionEventFilter* GetCollisionFilter() { return CollisionFilter.Get(); }
+		FSolverBreakingEventFilter* GetBreakingFilter() { return BreakingFilter.Get(); }
+		FSolverTrailingEventFilter* GetTrailingFilter() { return TrailingFilter.Get(); }
+
+		bool IsCollisionEventEnabled() const { return CollisionEventsEnabled; }
+		bool IsBreakingEventEnabled() const { return BreakingEventsEnabled; }
+		bool IsTrailingEventEnabled() const { return TrailingEventsEnabled; }
+
+	private:
+
+		TUniquePtr<FSolverCollisionEventFilter> CollisionFilter;
+		TUniquePtr<FSolverBreakingEventFilter> BreakingFilter;
+		TUniquePtr<FSolverTrailingEventFilter> TrailingFilter;
+
+		bool CollisionEventsEnabled;
+		bool BreakingEventsEnabled;
+		bool TrailingEventsEnabled;
 	};
 
 
diff --git a/Engine/Source/Runtime/Experimental/FieldSystem/Source/FieldSystemEngine/Private/Field/FieldSystemComponent.cpp b/Engine/Source/Runtime/Experimental/FieldSystem/Source/FieldSystemEngine/Private/Field/FieldSystemComponent.cpp
--- a/Engine/Source/Runtime/Experimental/FieldSystem/Source/FieldSystemEngine/Private/Field/FieldSystemComponent.cpp
+++ b/Engine/Source/Runtime/Experimental/FieldSystem/Source/FieldSystemEngine/Private/Field/FieldSystemComponent.cpp
@@ -11,7 +11,8 @@
 #include "Modules/ModuleManager.h"
 #include "Misc/CoreMiscDefines.h"
 #include "Physics/Experimental/PhysScene_Chaos.h"
-#include "SolverObjects/FieldSystemPhysicsObject.h"
+#include "PhysicsProxy/FieldSystemPhysicsProxy.h"
+#include "PBDRigidsSolver.h"
 
 DEFINE_LOG_CATEGORY_STATIC(FSC_Log, NoLogging, All);
 
@@ -19,7 +20,7 @@ UFieldSystemComponent::UFieldSystemComponent(const FObjectInitializer& ObjectIni
 	: Super(ObjectInitializer)
 	, FieldSystem(nullptr)
 #if INCLUDE_CHAOS
-	, PhysicsObject(nullptr)
+	, PhysicsProxy(nullptr)
 	, ChaosModule(nullptr)
 #endif
 	, bHasPhysicsState(false)
@@ -50,9 +51,9 @@ void UFieldSystemComponent::OnCreatePhysicsState()
 		ChaosModule = FModuleManager::Get().GetModulePtr<FChaosSolversModule>("ChaosSolvers");
 		check(ChaosModule);
 
-		PhysicsObject = new FFieldSystemPhysicsObject(this);
+		PhysicsProxy = new FFieldSystemPhysicsProxy(this);
 		TSharedPtr<FPhysScene_Chaos> Scene = GetOwner()->GetWorld()->PhysicsScene_Chaos;
-		Scene->AddObject(this, PhysicsObject);
+		Scene->AddObject(this, PhysicsProxy);
 #endif
 
 		bHasPhysicsState = true;
@@ -63,17 +64,17 @@ void UFieldSystemComponent::OnDestroyPhysicsState()
 {
 	UActorComponent::OnDestroyPhysicsState();
 #if INCLUDE_CHAOS
-	if (!PhysicsObject)
+	if (!PhysicsProxy)
 	{
 		check(!bHasPhysicsState);
 		return;
 	}
 
 	TSharedPtr<FPhysScene_Chaos> Scene = GetOwner()->GetWorld()->PhysicsScene_Chaos;
-	Scene->RemoveObject(PhysicsObject);
+	Scene->RemoveObject(PhysicsProxy);
 
 	ChaosModule = nullptr;
-	PhysicsObject = nullptr;
+	PhysicsProxy = nullptr;
 #endif
 
 	bHasPhysicsState = false;
@@ -100,7 +101,7 @@ void UFieldSystemComponent::DispatchCommand(const FFieldSystemCommand& InCommand
 		checkSlow(PhysicsDispatcher); // Should always have one of these
 
 		// Assemble a list of compatible solvers
-		TArray<Chaos::FPBDRigidsSolver*> SolverList;
+		TArray<Chaos::FPhysicsSolver*> SolverList;
 		if(SupportedSolvers.Num() > 0)
 		{
 			for(TSoftObjectPtr<AChaosSolverActor>& SolverActorPtr : SupportedSolvers)
@@ -112,18 +113,18 @@ void UFieldSystemComponent::DispatchCommand(const FFieldSystemCommand& InCommand
 			}
 		}
 
-		PhysicsDispatcher->EnqueueCommand([PhysicsObject = this->PhysicsObject, NewCommand = InCommand, ChaosModule = this->ChaosModule, SolverList]()
+		PhysicsDispatcher->EnqueueCommandImmediate([PhysicsProxy = this->PhysicsProxy, NewCommand = InCommand, ChaosModule = this->ChaosModule, SolverList]()
 		{
 			const int32 NumFilterSolvers = SolverList.Num();
-			const TArray<Chaos::FPBDRigidsSolver*>& Solvers = ChaosModule->GetSolvers();
+			const TArray<Chaos::FPhysicsSolver*>& Solvers = ChaosModule->GetSolvers();
 
-			for(Chaos::FPBDRigidsSolver* Solver : Solvers)
+			for(Chaos::FPhysicsSolver* Solver : Solvers)
 			{
 				const bool bSolverValid = NumFilterSolvers == 0 || SolverList.Contains(Solver);
 
-				if(Solver->Enabled() && Solver->HasActiveObjects() && bSolverValid)
+				if(Solver->Enabled() && Solver->HasActiveParticles() && bSolverValid)
 				{
-					PhysicsObject->BufferCommand(Solver, NewCommand);
+					PhysicsProxy->BufferCommand(Solver, NewCommand);
 				}
 			}
 		});
diff --git a/Engine/Source/Runtime/Experimental/FieldSystem/Source/FieldSystemEngine/Public/Field/FieldSystemComponent.h b/Engine/Source/Runtime/Experimental/FieldSystem/Source/FieldSystemEngine/Public/Field/FieldSystemComponent.h
--- a/Engine/Source/Runtime/Experimental/FieldSystem/Source/FieldSystemEngine/Public/Field/FieldSystemComponent.h
+++ b/Engine/Source/Runtime/Experimental/FieldSystem/Source/FieldSystemEngine/Public/Field/FieldSystemComponent.h
@@ -12,7 +12,7 @@
 #include "FieldSystemComponent.generated.h"
 
 struct FFieldSystemSampleData;
-class FFieldSystemPhysicsObject;
+class FFieldSystemPhysicsProxy;
 class FChaosSolversModule;
 
 /**
@@ -189,7 +189,7 @@ protected:
 
 #if INCLUDE_CHAOS
 
-	FFieldSystemPhysicsObject* PhysicsObject;
+	FFieldSystemPhysicsProxy* PhysicsProxy;
 	FChaosSolversModule* ChaosModule;
 
 #endif
diff --git a/Engine/Source/Runtime/Experimental/GeometryCollectionCore/Public/GeometryCollection/GeometryCollection.h b/Engine/Source/Runtime/Experimental/GeometryCollectionCore/Public/GeometryCollection/GeometryCollection.h
--- a/Engine/Source/Runtime/Experimental/GeometryCollectionCore/Public/GeometryCollection/GeometryCollection.h
+++ b/Engine/Source/Runtime/Experimental/GeometryCollectionCore/Public/GeometryCollection/GeometryCollection.h
@@ -24,8 +24,8 @@ public:
 	FGeometryCollection();
 	FGeometryCollection(FGeometryCollection &) = delete;
 	FGeometryCollection& operator=(const FGeometryCollection &) = delete;
-	FGeometryCollection(FGeometryCollection &&) = delete;
-	FGeometryCollection& operator=(FGeometryCollection &&) = delete;
+	FGeometryCollection(FGeometryCollection &&) = default;
+	FGeometryCollection& operator=(FGeometryCollection &&) = default;
 
 	typedef FTransformCollection Super;
 		/***
@@ -290,4 +290,10 @@ public:
 	void UpdateOldAttributeNames();
 
 
-};
\ No newline at end of file
+};
+
+FORCEINLINE Chaos::FChaosArchive& operator<<(Chaos::FChaosArchive& Ar, FGeometryCollection& Value)
+{
+	Value.Serialize(Ar);
+	return Ar;
+}
