Copying //UE4/Dev-Physics to Dev-Main (//UE4/Dev-Main) [at] 7703071

diff --git a/Engine/Source/Runtime/Experimental/GeometryCollectionCore/Public/GeometryCollection/ManagedArray.h b/Engine/Source/Runtime/Experimental/GeometryCollectionCore/Public/GeometryCollection/ManagedArray.h
--- a/Engine/Source/Runtime/Experimental/GeometryCollectionCore/Public/GeometryCollection/ManagedArray.h
+++ b/Engine/Source/Runtime/Experimental/GeometryCollectionCore/Public/GeometryCollection/ManagedArray.h
@@ -206,12 +206,19 @@ public:
 	/**
 	* Move Constructor
 	*/
-	FORCEINLINE TManagedArrayBase(TManagedArrayBase<ElementType>&& Other) = delete;
+	FORCEINLINE TManagedArrayBase(TManagedArrayBase<ElementType>&& Other)
+		: Array(MoveTemp(Other.Array))
+	{}
 
 	/**
 	* Assignment operator
 	*/
-	FORCEINLINE TManagedArrayBase& operator=(TManagedArrayBase<ElementType>&& Other) = delete;
+	FORCEINLINE TManagedArrayBase& operator=(TManagedArrayBase<ElementType>&& Other)
+	{
+		Array = MoveTemp(Other.Array);
+		return *this;
+	}
+
 
 
 	/**
@@ -481,9 +488,17 @@ public:
 		: TManagedArrayBase<InElementType>(Other)
 	{}
 
+	FORCEINLINE TManagedArray(TManagedArray<InElementType>&& Other)
+		: TManagedArrayBase<InElementType>(MoveTemp(Other))
+	{}
+
+	FORCEINLINE TManagedArray& operator=(TManagedArray<InElementType>&& Other)
+	{
+		TManagedArrayBase<InElementType>::operator=(MoveTemp(Other));
+		return *this;
+	}
+
 	FORCEINLINE TManagedArray(const TManagedArray<InElementType>& Other) = delete;
-	FORCEINLINE TManagedArray(TManagedArray<InElementType>&& Other) = delete;
-	FORCEINLINE TManagedArray& operator=(TManagedArray<InElementType>&& Other) = delete;
 
 	virtual ~TManagedArray()
 	{}
@@ -503,8 +518,8 @@ public:
 	{}
 
 	FORCEINLINE TManagedArray(const TManagedArray<int32>& Other) = delete;
-	FORCEINLINE TManagedArray(TManagedArray<int32>&& Other) = delete;
-	FORCEINLINE TManagedArray& operator=(TManagedArray<int32>&& Other) = delete;
+	FORCEINLINE TManagedArray(TManagedArray<int32>&& Other) = default;
+	FORCEINLINE TManagedArray& operator=(TManagedArray<int32>&& Other) = default;
 
 	virtual ~TManagedArray()
 	{}
@@ -555,8 +570,8 @@ public:
 	{}
 
 	FORCEINLINE TManagedArray(const TManagedArray<TSet<int32>>& Other) = delete;
-	FORCEINLINE TManagedArray(TManagedArray<TSet<int32>>&& Other) = delete;
-	FORCEINLINE TManagedArray& operator=(TManagedArray<TSet<int32>>&& Other) = delete;
+	FORCEINLINE TManagedArray(TManagedArray<TSet<int32>>&& Other) = default;
+	FORCEINLINE TManagedArray& operator=(TManagedArray<TSet<int32>>&& Other) = default;
 
 	virtual ~TManagedArray()
 	{}
@@ -622,8 +637,8 @@ public:
 	{}
 
 	FORCEINLINE TManagedArray(const TManagedArray<FIntVector>& Other) = delete;
-	FORCEINLINE TManagedArray(TManagedArray<FIntVector>&& Other) = delete;
-	FORCEINLINE TManagedArray& operator=(TManagedArray<FIntVector>&& Other) = delete;
+	FORCEINLINE TManagedArray(TManagedArray<FIntVector>&& Other) = default;
+	FORCEINLINE TManagedArray& operator=(TManagedArray<FIntVector>&& Other) = default;
 
 	virtual ~TManagedArray()
 	{}
diff --git a/Engine/Source/Runtime/Experimental/GeometryCollectionCore/Public/GeometryCollection/TransformCollection.h b/Engine/Source/Runtime/Experimental/GeometryCollectionCore/Public/GeometryCollection/TransformCollection.h
--- a/Engine/Source/Runtime/Experimental/GeometryCollectionCore/Public/GeometryCollection/TransformCollection.h
+++ b/Engine/Source/Runtime/Experimental/GeometryCollectionCore/Public/GeometryCollection/TransformCollection.h
@@ -21,8 +21,8 @@ public:
 	FTransformCollection();
 	FTransformCollection(FTransformCollection &) = delete;
 	FTransformCollection& operator=(const FTransformCollection&) = delete;
-	FTransformCollection(FTransformCollection&&) = delete;
-	FTransformCollection& operator=(FTransformCollection&&) = delete;
+	FTransformCollection(FTransformCollection&&) = default;
+	FTransformCollection& operator=(FTransformCollection&&) = default;
 
 
 	/***
diff --git a/Engine/Source/Runtime/Experimental/GeometryCollectionEngine/Private/ChaosBlueprint.cpp b/Engine/Source/Runtime/Experimental/GeometryCollectionEngine/Private/ChaosBlueprint.cpp
--- a/Engine/Source/Runtime/Experimental/GeometryCollectionEngine/Private/ChaosBlueprint.cpp
+++ b/Engine/Source/Runtime/Experimental/GeometryCollectionEngine/Private/ChaosBlueprint.cpp
@@ -1,9 +1,11 @@
 // Copyright 1998-2019 Epic Games, Inc. All Rights Reserved.
 
 #include "ChaosBlueprint.h"
-#include "PBDRigidsSolver.h"
+#include "PhysicsSolver.h"
 #include "Async/Async.h"
-#include "SolverObjects/GeometryCollectionPhysicsObject.h"
+#include "PhysicsProxy/GeometryCollectionPhysicsProxy.h"
+
+#define DISPATCH_BLUEPRINTS_IMMEDIATE 1
 
 UChaosDestructionListener::UChaosDestructionListener(FObjectInitializer const& ObjectInitializer)
 	: Super(ObjectInitializer), LastCollisionDataTimeStamp(-1.f), LastBreakingDataTimeStamp(-1.f), LastTrailingDataTimeStamp(-1.f)
@@ -12,7 +14,11 @@ UChaosDestructionListener::UChaosDestructionListener(FObjectInitializer const& O
 	bAutoActivate = true;	
 	bNeverNeedsRenderUpdate = true;
 
+#ifdef DISPATCH_BLUEPRINTS_IMMEDIATE
+	PrimaryComponentTick.bCanEverTick = false;
+#else
 	PrimaryComponentTick.bCanEverTick = true;
+#endif
 
 #if INCLUDE_CHAOS
 	SetCollisionFilter(MakeShareable(new FChaosCollisionEventFilter(&CollisionEventRequestSettings)));
@@ -21,275 +27,58 @@ UChaosDestructionListener::UChaosDestructionListener(FObjectInitializer const& O
 #endif
 }
 
-void UChaosDestructionListener::UpdateSolvers()
+void UChaosDestructionListener::ClearEvents()
+{
+#if INCLUDE_CHAOS
+#if 0 // solver actors no longer functional, using GetWorld()->GetPhysicsScene() instead
+	for (AChaosSolverActor* ChaosSolverActorObject : ChaosSolverActors)
+	{
+		UnregisterChaosEvents(ChaosSolverActorObject->GetPhysicsScene());
+	}
+#else
+	UnregisterChaosEvents(GetWorld()->GetPhysicsScene());
+#endif
+#endif
+}
+
+void UChaosDestructionListener::UpdateEvents()
 {
 #if INCLUDE_CHAOS
-	// Reset the solvers
-	Solvers.Reset();
 
-	if (!ChaosSolverActors.Num())
+	if (!ChaosSolverActors.Num() && !GeometryCollectionActors.Num())
 	{
-		if (TSharedPtr<FPhysScene_Chaos> WorldSolver = GetWorld()->PhysicsScene_Chaos)
+		if (FPhysScene* PhysScene = GetWorld()->GetPhysicsScene())
 		{
-			if (Chaos::FPBDRigidsSolver* Solver = WorldSolver->GetSolver())
-			{
-				Solvers.Add(Solver);
-			}
+			RegisterChaosEvents(PhysScene);
 		}
 	}
 	else
 	{
+
+#if 0 // solver actors no longer functional, using GetWorld()->GetPhysicsScene() instead
 		for (AChaosSolverActor* ChaosSolverActorObject : ChaosSolverActors)
 		{
 			if (ChaosSolverActorObject)
 			{
-				if (Chaos::FPBDRigidsSolver* Solver = ChaosSolverActorObject->GetSolver())
-				{
-					Solvers.Add(Solver);
-				}
+				RegisterChaosEvents(ChaosSolverActorObject->GetPhysicsScene());
 			}
 		}
-	}
-#endif
-}
-
-void UChaosDestructionListener::UpdateGeometryCollectionPhysicsObjects()
-{
-#if INCLUDE_CHAOS
-	GeometryCollectionPhysicsObjects.Reset();
-
-	if (GeometryCollectionActors.Num() > 0)
-	{
-		for (AGeometryCollectionActor* GeometryCollectionActorObject : GeometryCollectionActors)
+#else
+		if (FPhysScene* PhysScene = GetWorld()->GetPhysicsScene())
 		{
-			if (GeometryCollectionActorObject)
-			{
-				// Get GeometryCollectionComponent
-				if (const UGeometryCollectionComponent* GeometryCollectionComponent = GeometryCollectionActorObject->GetGeometryCollectionComponent())
-				{
-					// Get GeometryCollectionPhysicsObject
-					if (const FGeometryCollectionPhysicsObject* GeometryCollectionPhysicsObject = GeometryCollectionComponent->GetPhysicsObject())
-					{
-						if (Chaos::FPBDRigidsSolver* Solver = GeometryCollectionPhysicsObject->GetSolver())
-						{
-							if (!Solvers.Contains(Solver))
-							{
-								GeometryCollectionPhysicsObjects.Add(GeometryCollectionPhysicsObject);
-							}
-						}
-					}
-				}
-			}
+			RegisterChaosEvents(PhysScene);
 		}
-	}
 #endif
-}
 
-void UChaosDestructionListener::GetDataFromSolvers()
-{
-#if INCLUDE_CHAOS
-	// Loop through each solver and build up the array of data before passing off to the task to sort and analyze
-	// Note: we currently need to do this on the GT since the Solver getter itself is not thread safe so can't be retrieved on the task itself.
-	for (Chaos::FPBDRigidsSolver* Solver : Solvers)
-	{
-		if (Solver)
+		for (AGeometryCollectionActor* GeometryCollectionActor : GeometryCollectionActors)
 		{
-			Chaos::FPBDRigidsSolver::FScopedGetEventsData ScopedAccess = Solver->ScopedGetEventsData();
-
-			if (bIsCollisionEventListeningEnabled)
+			if (GeometryCollectionActor)
 			{
-				if (Solver->GetGenerateCollisionData() && Solver->GetSolverTime() > 0.f)
+				if (const UGeometryCollectionComponent* GeometryCollectionComponent = GeometryCollectionActor->GetGeometryCollectionComponent())
 				{
-					// ----------------------------------------------------------------------------------------------------------
-					// GETTING DATA FROM PBDRIGIDSOLVER
-					// ----------------------------------------------------------------------------------------------------------
-					const Chaos::FPBDRigidsSolver::FAllCollisionDataMaps& AllCollisionData_Maps = ScopedAccess.GetAllCollisions_Maps();
-					Chaos::FPBDRigidsSolver::FAllCollisionData AllCollisions;
-					Chaos::FPBDRigidsSolver::FSolverObjectReverseMapping SolverObjectReverseMapping;
-					Chaos::FPBDRigidsSolver::FParticleIndexReverseMapping ParticleIndexReverseMapping;
-					Chaos::FPBDRigidsSolver::FAllCollisionsIndicesBySolverObject AllCollisionsIndicesBySolverObject;
-					if (AllCollisionData_Maps.IsValid())
+					if (const FGeometryCollectionPhysicsProxy* GeometryCollectionPhysicsObject = GeometryCollectionComponent->GetPhysicsProxy())
 					{
-						const float DataTimeStamp = AllCollisionData_Maps.AllCollisionData->TimeCreated;
-						if (DataTimeStamp > LastCollisionDataTimeStamp)
-						{
-							LastCollisionDataTimeStamp = DataTimeStamp;
-
-							if (AllCollisionData_Maps.AllCollisionData)
-							{
-								AllCollisions = *AllCollisionData_Maps.AllCollisionData;
-							}
-							else
-							{
-								return;
-							}
-							if (AllCollisionData_Maps.SolverObjectReverseMapping)
-							{
-								SolverObjectReverseMapping = *AllCollisionData_Maps.SolverObjectReverseMapping;
-							}
-							else
-							{
-								return;
-							}
-							if (AllCollisionData_Maps.ParticleIndexReverseMapping)
-							{
-								ParticleIndexReverseMapping = *AllCollisionData_Maps.ParticleIndexReverseMapping;
-							}
-							else
-							{
-								return;
-							}
-							if (AllCollisionData_Maps.AllCollisionsIndicesBySolverObject)
-							{
-								AllCollisionsIndicesBySolverObject = *AllCollisionData_Maps.AllCollisionsIndicesBySolverObject;
-							}
-							else
-							{
-								return;
-							}
-
-							// ----------------------------------------------------------------------------------------------------------
-							// END OF GETTING DATA FROM PBDRIGIDSOLVER
-							// ----------------------------------------------------------------------------------------------------------
-
-							int32 NumCollisions = AllCollisions.AllCollisionsArray.Num();
-							if (NumCollisions > 0)
-							{
-								RawCollisionDataArray.Append(AllCollisions.AllCollisionsArray.GetData(), NumCollisions);
-							}
-						}
-					}
-				}
-			}
-
-			if (bIsBreakingEventListeningEnabled)
-			{
-				if (Solver->GetGenerateBreakingData() && Solver->GetSolverTime() > 0.f)
-				{
-					// ----------------------------------------------------------------------------------------------------------
-					// GETTING DATA FROM PBDRIGIDSOLVER
-					// ----------------------------------------------------------------------------------------------------------
-					const Chaos::FPBDRigidsSolver::FAllBreakingDataMaps& AllBreakingData_Maps = ScopedAccess.GetAllBreakings_Maps();
-					Chaos::FPBDRigidsSolver::FAllBreakingData AllBreakings;
-					Chaos::FPBDRigidsSolver::FSolverObjectReverseMapping SolverObjectReverseMapping;
-					Chaos::FPBDRigidsSolver::FParticleIndexReverseMapping ParticleIndexReverseMapping;
-					Chaos::FPBDRigidsSolver::FAllBreakingsIndicesBySolverObject AllBreakingsIndicesBySolverObject;
-					if (AllBreakingData_Maps.IsValid())
-					{
-						const float DataTimeStamp = AllBreakingData_Maps.AllBreakingData->TimeCreated;
-						if (DataTimeStamp > LastBreakingDataTimeStamp)
-						{
-							LastBreakingDataTimeStamp = DataTimeStamp;
-
-							if (AllBreakingData_Maps.AllBreakingData)
-							{
-								AllBreakings = *AllBreakingData_Maps.AllBreakingData;
-							}
-							else
-							{
-								return;
-							}
-							if (AllBreakingData_Maps.SolverObjectReverseMapping)
-							{
-								SolverObjectReverseMapping = *AllBreakingData_Maps.SolverObjectReverseMapping;
-							}
-							else
-							{
-								return;
-							}
-							if (AllBreakingData_Maps.ParticleIndexReverseMapping)
-							{
-								ParticleIndexReverseMapping = *AllBreakingData_Maps.ParticleIndexReverseMapping;
-							}
-							else
-							{
-								return;
-							}
-							if (AllBreakingData_Maps.AllBreakingsIndicesBySolverObject)
-							{
-								AllBreakingsIndicesBySolverObject = *AllBreakingData_Maps.AllBreakingsIndicesBySolverObject;
-							}
-							else
-							{
-								return;
-							}
-
-							// ----------------------------------------------------------------------------------------------------------
-							// END OF GETTING DATA FROM PBDRIGIDSOLVER
-							// ----------------------------------------------------------------------------------------------------------
-
-							int32 NumBreakings = AllBreakings.AllBreakingsArray.Num();
-							if (NumBreakings > 0)
-							{
-								RawBreakingDataArray.Append(AllBreakings.AllBreakingsArray.GetData(), NumBreakings);
-							}
-						}
-					}
-				}
-			}
-
-			if (bIsTrailingEventListeningEnabled)
-			{
-				if (Solver->GetGenerateTrailingData() && Solver->GetSolverTime() > 0.f)
-				{
-					// ----------------------------------------------------------------------------------------------------------
-					// GETTING DATA FROM PBDRIGIDSOLVER
-					// ----------------------------------------------------------------------------------------------------------
-					const Chaos::FPBDRigidsSolver::FAllTrailingDataMaps& AllTrailingData_Maps = ScopedAccess.GetAllTrailings_Maps();
-					Chaos::FPBDRigidsSolver::FAllTrailingData AllTrailings;
-					Chaos::FPBDRigidsSolver::FSolverObjectReverseMapping SolverObjectReverseMapping;
-					Chaos::FPBDRigidsSolver::FParticleIndexReverseMapping ParticleIndexReverseMapping;
-					Chaos::FPBDRigidsSolver::FAllTrailingsIndicesBySolverObject AllTrailingsIndicesBySolverObject;
-					if (AllTrailingData_Maps.IsValid())
-					{
-						const float DataTimeStamp = AllTrailingData_Maps.AllTrailingData->TimeCreated;
-						if (DataTimeStamp > LastTrailingDataTimeStamp)
-						{
-							LastTrailingDataTimeStamp = DataTimeStamp;
-
-							if (AllTrailingData_Maps.AllTrailingData)
-							{
-								AllTrailings = *AllTrailingData_Maps.AllTrailingData;
-							}
-							else
-							{
-								return;
-							}
-							if (AllTrailingData_Maps.SolverObjectReverseMapping)
-							{
-								SolverObjectReverseMapping = *AllTrailingData_Maps.SolverObjectReverseMapping;
-							}
-							else
-							{
-								return;
-							}
-							if (AllTrailingData_Maps.ParticleIndexReverseMapping)
-							{
-								ParticleIndexReverseMapping = *AllTrailingData_Maps.ParticleIndexReverseMapping;
-							}
-							else
-							{
-								return;
-							}
-							if (AllTrailingData_Maps.AllTrailingsIndicesBySolverObject)
-							{
-								AllTrailingsIndicesBySolverObject = *AllTrailingData_Maps.AllTrailingsIndicesBySolverObject;
-							}
-							else
-							{
-								return;
-							}
-
-							// ----------------------------------------------------------------------------------------------------------
-							// END OF GETTING DATA FROM PBDRIGIDSOLVER
-							// ----------------------------------------------------------------------------------------------------------
-
-							int32 NumTrailings = AllTrailings.AllTrailingsArray.Num();
-							if (NumTrailings > 0)
-							{
-								RawTrailingDataArray.Append(AllTrailings.AllTrailingsArray.GetData(), NumTrailings);
-							}
-						}
+						RegisterChaosEvents(GeometryCollectionComponent->GetPhysicsScene());
 					}
 				}
 			}
@@ -298,226 +87,29 @@ void UChaosDestructionListener::GetDataFromSolvers()
 #endif
 }
 
-void UChaosDestructionListener::GetDataFromGeometryCollectionPhysicsObjects()
+#if 0 // #todo: No longer required?
+void UChaosDestructionListener::UpdateGeometryCollectionPhysicsProxies()
 {
 #if INCLUDE_CHAOS
-	for (const FGeometryCollectionPhysicsObject* GeometryCollectionPhysicsObject : GeometryCollectionPhysicsObjects)
+	GeometryCollectionPhysicsProxies.Reset();
+
+	if (GeometryCollectionActors.Num() > 0)
 	{
-		if (GeometryCollectionPhysicsObject)
+		for (AGeometryCollectionActor* GeometryCollectionActorObject : GeometryCollectionActors)
 		{
-			if (bIsCollisionEventListeningEnabled)
-			{
-				if (Chaos::FPBDRigidsSolver* Solver = GeometryCollectionPhysicsObject->GetSolver())
-				{
-					if (Solver->GetGenerateCollisionData() && Solver->GetSolverTime() > 0.f)
-					{
-						Chaos::FPBDRigidsSolver::FScopedGetEventsData SopedAccess = Solver->ScopedGetEventsData();
-
-						// ----------------------------------------------------------------------------------------------------------
-						// GETTING DATA FROM PBDRIGIDSOLVER
-						// ----------------------------------------------------------------------------------------------------------
-						const Chaos::FPBDRigidsSolver::FAllCollisionDataMaps& AllCollisionData_Maps = SopedAccess.GetAllCollisions_Maps();
-						Chaos::FPBDRigidsSolver::FAllCollisionData AllCollisions;
-						Chaos::FPBDRigidsSolver::FSolverObjectReverseMapping SolverObjectReverseMapping;
-						Chaos::FPBDRigidsSolver::FParticleIndexReverseMapping ParticleIndexReverseMapping;
-						Chaos::FPBDRigidsSolver::FAllCollisionsIndicesBySolverObject AllCollisionsIndicesBySolverObject;
-						if (AllCollisionData_Maps.IsValid())
-						{
-							if (AllCollisionData_Maps.AllCollisionData)
-							{
-								AllCollisions = *AllCollisionData_Maps.AllCollisionData;
-							}
-							else
-							{
-								return;
-							}
-							if (AllCollisionData_Maps.SolverObjectReverseMapping)
-							{
-								SolverObjectReverseMapping = *AllCollisionData_Maps.SolverObjectReverseMapping;
-							}
-							else
-							{
-								return;
-							}
-							if (AllCollisionData_Maps.ParticleIndexReverseMapping)
-							{
-								ParticleIndexReverseMapping = *AllCollisionData_Maps.ParticleIndexReverseMapping;
-							}
-							else
-							{
-								return;
-							}
-							if (AllCollisionData_Maps.AllCollisionsIndicesBySolverObject)
-							{
-								AllCollisionsIndicesBySolverObject = *AllCollisionData_Maps.AllCollisionsIndicesBySolverObject;
-							}
-							else
-							{
-								return;
-							}
-
-							// ----------------------------------------------------------------------------------------------------------
-							// END OF GETTING DATA FROM PBDRIGIDSOLVER
-							// ----------------------------------------------------------------------------------------------------------
-
-							int32 NumCollisions = AllCollisions.AllCollisionsArray.Num();
-							if (NumCollisions > 0)
-							{
-								// Get collisions for this GeometryCollectionPhysicsObject from AllCollisions.AllCollisions[]
-								if (AllCollisionsIndicesBySolverObject.AllCollisionsIndicesBySolverObjectMap.Contains(GeometryCollectionPhysicsObject))
-								{
-									TArray<int32> CollisionIndices = AllCollisionsIndicesBySolverObject.AllCollisionsIndicesBySolverObjectMap[GeometryCollectionPhysicsObject];
-									for (int32 Idx = 0; Idx < CollisionIndices.Num(); ++Idx)
-									{
-										RawCollisionDataArray.Add(AllCollisions.AllCollisionsArray[CollisionIndices[Idx]]);
-									}
-								}
-							}
-						}
-					}
-				}
-			}
-
-			if (bIsBreakingEventListeningEnabled)
-			{
-				if (Chaos::FPBDRigidsSolver* Solver = GeometryCollectionPhysicsObject->GetSolver())
-				{
-					Chaos::FPBDRigidsSolver::FScopedGetEventsData SopedAccess = Solver->ScopedGetEventsData();
-
-					if (Solver->GetGenerateBreakingData() && Solver->GetSolverTime() > 0.f)
-					{
-						// ----------------------------------------------------------------------------------------------------------
-						// GETTING DATA FROM PBDRIGIDSOLVER
-						// ----------------------------------------------------------------------------------------------------------
-						const Chaos::FPBDRigidsSolver::FAllBreakingDataMaps& AllBreakingData_Maps = SopedAccess.GetAllBreakings_Maps();
-						Chaos::FPBDRigidsSolver::FAllBreakingData AllBreakings;
-						Chaos::FPBDRigidsSolver::FSolverObjectReverseMapping SolverObjectReverseMapping;
-						Chaos::FPBDRigidsSolver::FParticleIndexReverseMapping ParticleIndexReverseMapping;
-						Chaos::FPBDRigidsSolver::FAllBreakingsIndicesBySolverObject AllBreakingsIndicesBySolverObject;
-						if (AllBreakingData_Maps.IsValid())
-						{
-							if (AllBreakingData_Maps.AllBreakingData)
-							{
-								AllBreakings = *AllBreakingData_Maps.AllBreakingData;
-							}
-							else
-							{
-								return;
-							}
-							if (AllBreakingData_Maps.SolverObjectReverseMapping)
-							{
-								SolverObjectReverseMapping = *AllBreakingData_Maps.SolverObjectReverseMapping;
-							}
-							else
-							{
-								return;
-							}
-							if (AllBreakingData_Maps.ParticleIndexReverseMapping)
-							{
-								ParticleIndexReverseMapping = *AllBreakingData_Maps.ParticleIndexReverseMapping;
-							}
-							else
-							{
-								return;
-							}
-							if (AllBreakingData_Maps.AllBreakingsIndicesBySolverObject)
-							{
-								AllBreakingsIndicesBySolverObject = *AllBreakingData_Maps.AllBreakingsIndicesBySolverObject;
-							}
-							else
-							{
-								return;
-							}
-
-							// ----------------------------------------------------------------------------------------------------------
-							// END OF GETTING DATA FROM PBDRIGIDSOLVER
-							// ----------------------------------------------------------------------------------------------------------
-
-							int32 NumBreakings = AllBreakings.AllBreakingsArray.Num();
-							if (NumBreakings > 0)
-							{
-								// Get breakings for this GeometryCollectionPhysicsObject from AllBreakings.AllBreakings[]
-								if (AllBreakingsIndicesBySolverObject.AllBreakingsIndicesBySolverObjectMap.Contains(GeometryCollectionPhysicsObject))
-								{
-									TArray<int32> BreakingIndices = AllBreakingsIndicesBySolverObject.AllBreakingsIndicesBySolverObjectMap[GeometryCollectionPhysicsObject];
-									for (int32 Idx = 0; Idx < BreakingIndices.Num(); ++Idx)
-									{
-										RawBreakingDataArray.Add(AllBreakings.AllBreakingsArray[BreakingIndices[Idx]]);
-									}
-								}
-							}
-						}
-					}
-				}
-			}
-
-			if (bIsTrailingEventListeningEnabled)
+			if (GeometryCollectionActorObject)
 			{
-				if (Chaos::FPBDRigidsSolver* Solver = GeometryCollectionPhysicsObject->GetSolver())
+				// Get GeometryCollectionComponent
+				if (const UGeometryCollectionComponent* GeometryCollectionComponent = GeometryCollectionActorObject->GetGeometryCollectionComponent())
 				{
-					Chaos::FPBDRigidsSolver::FScopedGetEventsData ScopedAccess = Solver->ScopedGetEventsData();
-
-					if (Solver->GetGenerateTrailingData() && Solver->GetSolverTime() > 0.f)
+					// Get GeometryCollectionPhysicsProxies
+					if (const FGeometryCollectionPhysicsProxy* GeometryCollectionPhysicsProxy = GeometryCollectionComponent->GetPhysicsProxy())
 					{
-						// ----------------------------------------------------------------------------------------------------------
-						// GETTING DATA FROM PBDRIGIDSOLVER
-						// ----------------------------------------------------------------------------------------------------------
-						const Chaos::FPBDRigidsSolver::FAllTrailingDataMaps& AllTrailingData_Maps = ScopedAccess.GetAllTrailings_Maps();
-						Chaos::FPBDRigidsSolver::FAllTrailingData AllTrailings;
-						Chaos::FPBDRigidsSolver::FSolverObjectReverseMapping SolverObjectReverseMapping;
-						Chaos::FPBDRigidsSolver::FParticleIndexReverseMapping ParticleIndexReverseMapping;
-						Chaos::FPBDRigidsSolver::FAllTrailingsIndicesBySolverObject AllTrailingsIndicesBySolverObject;
-						if (AllTrailingData_Maps.IsValid())
+						if (Chaos::FPhysicsSolver* Solver = GeometryCollectionPhysicsProxy->GetSolver())
 						{
-							if (AllTrailingData_Maps.AllTrailingData)
-							{
-								AllTrailings = *AllTrailingData_Maps.AllTrailingData;
-							}
-							else
-							{
-								return;
-							}
-							if (AllTrailingData_Maps.SolverObjectReverseMapping)
-							{
-								SolverObjectReverseMapping = *AllTrailingData_Maps.SolverObjectReverseMapping;
-							}
-							else
-							{
-								return;
-							}
-							if (AllTrailingData_Maps.ParticleIndexReverseMapping)
-							{
-								ParticleIndexReverseMapping = *AllTrailingData_Maps.ParticleIndexReverseMapping;
-							}
-							else
-							{
-								return;
-							}
-							if (AllTrailingData_Maps.AllTrailingsIndicesBySolverObject)
-							{
-								AllTrailingsIndicesBySolverObject = *AllTrailingData_Maps.AllTrailingsIndicesBySolverObject;
-							}
-							else
-							{
-								return;
-							}
-
-							// ----------------------------------------------------------------------------------------------------------
-							// END OF GETTING DATA FROM PBDRIGIDSOLVER
-							// ----------------------------------------------------------------------------------------------------------
-
-							int32 NumTrailings = AllTrailings.AllTrailingsArray.Num();
-							if (NumTrailings > 0)
+							if (!Solvers.Contains(Solver))
 							{
-								// Get trailings for this GeometryCollectionPhysicsObject from AllTrailings.AllTrailings[]
-								if (AllTrailingsIndicesBySolverObject.AllTrailingsIndicesBySolverObjectMap.Contains(GeometryCollectionPhysicsObject))
-								{
-									TArray<int32> TrailingIndices = AllTrailingsIndicesBySolverObject.AllTrailingsIndicesBySolverObjectMap[GeometryCollectionPhysicsObject];
-									for (int32 Idx = 0; Idx < TrailingIndices.Num(); ++Idx)
-									{
-										RawTrailingDataArray.Add(AllTrailings.AllTrailingsArray[TrailingIndices[Idx]]);
-									}
-								}
+								GeometryCollectionPhysicsProxies.Add(GeometryCollectionPhysicsProxy);
 							}
 						}
 					}
@@ -527,6 +119,7 @@ void UChaosDestructionListener::GetDataFromGeometryCollectionPhysicsObjects()
 	}
 #endif
 }
+#endif
 
 #if WITH_EDITOR
 void UChaosDestructionListener::PostEditChangeProperty(struct FPropertyChangedEvent& PropertyChangedEvent)
@@ -559,8 +152,21 @@ void UChaosDestructionListener::UpdateTransformSettings()
 	bChanged = true;
 }
 
+void UChaosDestructionListener::BeginPlay()
+{
+	Super::BeginPlay();
+	UpdateEvents();
+}
+
+void UChaosDestructionListener::EndPlay(const EEndPlayReason::Type EndPlayReason)
+{
+	ClearEvents();
+	Super::EndPlay(EndPlayReason);
+}
+
 void UChaosDestructionListener::TickComponent(float DeltaTime, enum ELevelTick TickType, FActorComponentTickFunction *ThisTickFunction)
 {
+#if 0
 	bool bIsListening = IsEventListening();
 
 	// if owning actor is disabled, don't listen
@@ -631,9 +237,9 @@ void UChaosDestructionListener::TickComponent(float DeltaTime, enum ELevelTick T
 		UpdateSolvers();
 	}
 
-	if (!GeometryCollectionPhysicsObjects.Num())
+	if (!GeometryCollectionPhysicsProxies.Num())
 	{
-		UpdateGeometryCollectionPhysicsObjects();
+		UpdateGeometryCollectionPhysicsProxies();
 	}
 
 	// Reset our cached data arrays for various destruction types
@@ -645,8 +251,8 @@ void UChaosDestructionListener::TickComponent(float DeltaTime, enum ELevelTick T
 	// Retrieve the raw data arrays from the solvers
 	GetDataFromSolvers();
 
-	// Retrieve the raw data arrays from the GeometryCollectionPhysicsObjects
-	GetDataFromGeometryCollectionPhysicsObjects();
+	// Retrieve the raw data arrays from the GeometryCollectionPhysicsProxy
+	GetDataFromGeometryCollectionPhysicsProxies();
 
 	TaskState.Set((int32)ETaskState::Processing);
 
@@ -684,41 +290,56 @@ void UChaosDestructionListener::TickComponent(float DeltaTime, enum ELevelTick T
 #endif
 			TaskState.Set((int32)ETaskState::Finished);
 		});
+#endif // if 0
 }
 
 void UChaosDestructionListener::AddChaosSolverActor(AChaosSolverActor* ChaosSolverActor)
 {
-	if (ChaosSolverActor)
+#if 0 // solver actors no longer functional, using GetWorld()->GetPhysicsScene() instead
+	if (ChaosSolverActor && !ChaosSolverActors.Contains(ChaosSolverActor))
 	{
 		ChaosSolverActors.Add(ChaosSolverActor);
-		UpdateSolvers();
+		RegisterChaosEvents(ChaosSolverActor->GetPhysicsScene());
 	}
+#endif
 }
 
 void UChaosDestructionListener::RemoveChaosSolverActor(AChaosSolverActor* ChaosSolverActor)
 {
+#if 0 // solver actors no longer functional, using GetWorld()->GetPhysicsScene() instead
 	if (ChaosSolverActor)
 	{
+		ClearEvents();
 		ChaosSolverActors.Remove(ChaosSolverActor);
-		UpdateSolvers();
+		UpdateEvents();
 	}
+#endif
 }
 
 void UChaosDestructionListener::AddGeometryCollectionActor(AGeometryCollectionActor* GeometryCollectionActor)
 {
-	if (GeometryCollectionActor)
+#if INCLUDE_CHAOS
+	if (GeometryCollectionActor && !GeometryCollectionActors.Contains(GeometryCollectionActor))
 	{
 		GeometryCollectionActors.Add(GeometryCollectionActor);
-		UpdateGeometryCollectionPhysicsObjects();
+		if (const UGeometryCollectionComponent* GeometryCollectionComponent = GeometryCollectionActor->GetGeometryCollectionComponent())
+		{
+			if (const FGeometryCollectionPhysicsProxy* GeometryCollectionPhysicsObject = GeometryCollectionComponent->GetPhysicsProxy())
+			{
+				RegisterChaosEvents(GeometryCollectionComponent->GetPhysicsScene());
+			}
+		}
 	}
+#endif
 }
 
 void UChaosDestructionListener::RemoveGeometryCollectionActor(AGeometryCollectionActor* GeometryCollectionActor)
 {
 	if (GeometryCollectionActor)
 	{
+		ClearEvents();
 		GeometryCollectionActors.Remove(GeometryCollectionActor);
-		UpdateGeometryCollectionPhysicsObjects();
+		UpdateEvents();
 	}
 }
 
@@ -787,3 +408,113 @@ void UChaosDestructionListener::SortTrailingEvents(TArray<FChaosTrailingEventDat
 	}
 #endif
 }
+
+#if INCLUDE_CHAOS
+
+void UChaosDestructionListener::RegisterChaosEvents(FPhysScene* Scene)
+{
+#if WITH_CHAOS
+	Scene->RegisterEventHandler<Chaos::FCollisionEventData>(Chaos::EEventType::Collision, this, &UChaosDestructionListener::HandleCollisionEvents);
+	Scene->RegisterEventHandler<Chaos::FBreakingEventData>(Chaos::EEventType::Breaking, this, &UChaosDestructionListener::HandleBreakingEvents);
+	Scene->RegisterEventHandler<Chaos::FTrailingEventData>(Chaos::EEventType::Trailing, this, &UChaosDestructionListener::HandleTrailingEvents);
+#endif
+}
+
+void UChaosDestructionListener::UnregisterChaosEvents(FPhysScene* Scene)
+{
+#if WITH_CHAOS
+	Scene->UnregisterEventHandler(Chaos::EEventType::Collision, this);
+	Scene->UnregisterEventHandler(Chaos::EEventType::Breaking, this);
+	Scene->UnregisterEventHandler(Chaos::EEventType::Trailing, this);
+#endif
+}
+
+void UChaosDestructionListener::RegisterChaosEvents(TSharedPtr<FPhysScene_Chaos> Scene)
+{
+#if WITH_CHAOS
+	Scene->RegisterEventHandler<Chaos::FCollisionEventData>(Chaos::EEventType::Collision, this, &UChaosDestructionListener::HandleCollisionEvents);
+	Scene->RegisterEventHandler<Chaos::FBreakingEventData>(Chaos::EEventType::Breaking, this, &UChaosDestructionListener::HandleBreakingEvents);
+	Scene->RegisterEventHandler<Chaos::FTrailingEventData>(Chaos::EEventType::Trailing, this, &UChaosDestructionListener::HandleTrailingEvents);
+#endif
+}
+
+void UChaosDestructionListener::UnregisterChaosEvents(TSharedPtr<FPhysScene_Chaos> Scene)
+{
+#if WITH_CHAOS
+	Scene->UnregisterEventHandler(Chaos::EEventType::Collision, this);
+	Scene->UnregisterEventHandler(Chaos::EEventType::Breaking, this);
+	Scene->UnregisterEventHandler(Chaos::EEventType::Trailing, this);
+#endif
+}
+
+
+void UChaosDestructionListener::HandleCollisionEvents(const Chaos::FCollisionEventData& Event)
+{
+	if (bIsCollisionEventListeningEnabled)
+	{
+		int NumCollisions = Event.CollisionData.AllCollisionsArray.Num();
+
+		RawCollisionDataArray.Append(Event.CollisionData.AllCollisionsArray.GetData(), NumCollisions);
+
+#if DISPATCH_BLUEPRINTS_IMMEDIATE
+		if (ChaosCollisionFilter.IsValid())
+		{
+			ChaosCollisionFilter->FilterEvents(ChaosComponentTransform, RawCollisionDataArray);
+
+			if (ChaosCollisionFilter->GetNumEvents() > 0 && OnCollisionEvents.IsBound())
+			{
+				OnCollisionEvents.Broadcast(ChaosCollisionFilter->GetFilteredResults());
+			}
+		}
+		RawCollisionDataArray.Reset();
+#endif
+	}
+}
+
+void UChaosDestructionListener::HandleBreakingEvents(const Chaos::FBreakingEventData& Event)
+{
+	if (bIsBreakingEventListeningEnabled)
+	{
+		int NumBreakings = Event.BreakingData.AllBreakingsArray.Num();
+		RawBreakingDataArray.Append(Event.BreakingData.AllBreakingsArray.GetData(), NumBreakings);
+
+#if DISPATCH_BLUEPRINTS_IMMEDIATE
+		if (ChaosBreakingFilter.IsValid())
+		{
+			ChaosBreakingFilter->FilterEvents(ChaosComponentTransform, RawBreakingDataArray);
+
+			if (ChaosBreakingFilter->GetNumEvents() > 0 && OnBreakingEvents.IsBound())
+			{
+				OnBreakingEvents.Broadcast(ChaosBreakingFilter->GetFilteredResults());
+			}
+		}
+		RawBreakingDataArray.Reset();
+#endif
+	}
+}
+
+void UChaosDestructionListener::HandleTrailingEvents(const Chaos::FTrailingEventData& Event)
+{
+	if (bIsTrailingEventListeningEnabled)
+	{
+		int NumTrailings = Event.TrailingData.AllTrailingsArray.Num();
+		RawTrailingDataArray.Append(Event.TrailingData.AllTrailingsArray.GetData(), NumTrailings);
+
+#if DISPATCH_BLUEPRINTS_IMMEDIATE
+		if (ChaosTrailingFilter.IsValid())
+		{
+			ChaosTrailingFilter->FilterEvents(ChaosComponentTransform, RawTrailingDataArray);
+
+			if (ChaosTrailingFilter->GetNumEvents() > 0 && OnTrailingEvents.IsBound())
+			{
+				OnTrailingEvents.Broadcast(ChaosTrailingFilter->GetFilteredResults());
+			}
+		}
+		RawTrailingDataArray.Reset();
+#endif
+	}
+}
+
+#endif // INCLUDE_CHAOS
+
+
diff --git a/Engine/Source/Runtime/Experimental/GeometryCollectionEngine/Private/ChaosBreakingEventFilter.cpp b/Engine/Source/Runtime/Experimental/GeometryCollectionEngine/Private/ChaosBreakingEventFilter.cpp
--- a/Engine/Source/Runtime/Experimental/GeometryCollectionEngine/Private/ChaosBreakingEventFilter.cpp
+++ b/Engine/Source/Runtime/Experimental/GeometryCollectionEngine/Private/ChaosBreakingEventFilter.cpp
@@ -3,7 +3,7 @@
 
 #if INCLUDE_CHAOS
 
-void FChaosBreakingEventFilter::FilterEvents(const FTransform& ChaosComponentTransform, const FPBDRigidsSolver::FBreakingDataArray& RawBreakingDataArray)
+void FChaosBreakingEventFilter::FilterEvents(const FTransform& ChaosComponentTransform, const Chaos::FBreakingDataArray& RawBreakingDataArray)
 {
 	FilteredDataArray.Reset();
 
diff --git a/Engine/Source/Runtime/Experimental/GeometryCollectionEngine/Private/ChaosCollisionEventFilter.cpp b/Engine/Source/Runtime/Experimental/GeometryCollectionEngine/Private/ChaosCollisionEventFilter.cpp
--- a/Engine/Source/Runtime/Experimental/GeometryCollectionEngine/Private/ChaosCollisionEventFilter.cpp
+++ b/Engine/Source/Runtime/Experimental/GeometryCollectionEngine/Private/ChaosCollisionEventFilter.cpp
@@ -3,7 +3,7 @@
 
 #if INCLUDE_CHAOS
 
-void FChaosCollisionEventFilter::FilterEvents(const FTransform& ChaosComponentTransform, const FPBDRigidsSolver::FCollisionDataArray& RawCollisionDataArray)
+void FChaosCollisionEventFilter::FilterEvents(const FTransform& ChaosComponentTransform, const Chaos::FCollisionDataArray& RawCollisionDataArray)
 {
 	FilteredDataArray.Reset();
 
@@ -61,8 +61,8 @@ void FChaosCollisionEventFilter::FilterEvents(const FTransform& ChaosComponentTr
 			NewData.Mass1 = CollisionData.Mass1;
 			NewData.Mass2 = CollisionData.Mass2;
 			NewData.Impulse = CollisionData.AccumulatedImpulse;
-			NewData.ParticleIndex = CollisionData.ParticleIndex;
-			NewData.LevelsetIndex = CollisionData.LevelsetIndex;
+			NewData.Particle = CollisionData.Particle;
+			NewData.Levelset = CollisionData.Levelset;
 
 			FilteredDataArray.Add(NewData);
 
diff --git a/Engine/Source/Runtime/Experimental/GeometryCollectionEngine/Private/ChaosTrailingEventFilter.cpp b/Engine/Source/Runtime/Experimental/GeometryCollectionEngine/Private/ChaosTrailingEventFilter.cpp
--- a/Engine/Source/Runtime/Experimental/GeometryCollectionEngine/Private/ChaosTrailingEventFilter.cpp
+++ b/Engine/Source/Runtime/Experimental/GeometryCollectionEngine/Private/ChaosTrailingEventFilter.cpp
@@ -1,8 +1,11 @@
 // Copyright 1998-2019 Epic Games, Inc. All Rights Reserved.
+
+#if INCLUDE_CHAOS
+
 #include "ChaosTrailingEventFilter.h"
 
 
-void FChaosTrailingEventFilter::FilterEvents(const FTransform& ChaosComponentTransform, const TArray<Chaos::TTrailingData<float, 3>>& RawTrailingDataArray)
+void FChaosTrailingEventFilter::FilterEvents(const FTransform& ChaosComponentTransform, const Chaos::FTrailingDataArray& RawTrailingDataArray)
 {
 	FilteredDataArray.Reset();
 
@@ -45,7 +48,9 @@ void FChaosTrailingEventFilter::FilterEvents(const FTransform& ChaosComponentTra
 			NewData.Velocity = TrailingData.Velocity;
 			NewData.AngularVelocity = TrailingData.AngularVelocity;
 			NewData.Mass = TrailingData.Mass;
-			NewData.ParticleIndex = TrailingData.ParticleIndex;
+#if TODO_MAKE_A_BLUEPRINT_WAY_TO_ACCESS_PARTICLES
+			NewData.ParticleIndex = TrailingData.Particle;
+#endif
 
 			FilteredDataArray.Add(NewData);
 
@@ -115,4 +120,4 @@ void FChaosTrailingEventFilter::SortEvents(TArray<FChaosTrailingEventData>& InOu
 	}
 }
 
-
+#endif
diff --git a/Engine/Source/Runtime/Experimental/GeometryCollectionEngine/Private/GeometryCollection/DerivedDataGeometryCollectionCooker.cpp b/Engine/Source/Runtime/Experimental/GeometryCollectionEngine/Private/GeometryCollection/DerivedDataGeometryCollectionCooker.cpp
--- a/Engine/Source/Runtime/Experimental/GeometryCollectionEngine/Private/GeometryCollection/DerivedDataGeometryCollectionCooker.cpp
+++ b/Engine/Source/Runtime/Experimental/GeometryCollectionEngine/Private/GeometryCollection/DerivedDataGeometryCollectionCooker.cpp
@@ -3,7 +3,7 @@
 #include "GeometryCollection/GeometryCollectionObject.h"
 #include "GeometryCollection/GeometryCollection.h"
 #include "Serialization/MemoryWriter.h"
-#include "SolverObjects/GeometryCollectionPhysicsObject.h"
+#include "PhysicsProxy/GeometryCollectionPhysicsProxy.h"
 #include "Chaos/ChaosArchive.h"
 #include "UObject/DestructionObjectVersion.h"
 #include "Chaos/ErrorReporter.h"
diff --git a/Engine/Source/Runtime/Experimental/GeometryCollectionEngine/Private/GeometryCollection/GeometryCollectionActor.cpp b/Engine/Source/Runtime/Experimental/GeometryCollectionEngine/Private/GeometryCollection/GeometryCollectionActor.cpp
--- a/Engine/Source/Runtime/Experimental/GeometryCollectionEngine/Private/GeometryCollection/GeometryCollectionActor.cpp
+++ b/Engine/Source/Runtime/Experimental/GeometryCollectionEngine/Private/GeometryCollection/GeometryCollectionActor.cpp
@@ -4,10 +4,10 @@
 GeometryCollectionActor.cpp: AGeometryCollectionActor methods.
 =============================================================================*/
 
-#if INCLUDE_CHAOS
-
 #include "GeometryCollection/GeometryCollectionActor.h"
 
+#if INCLUDE_CHAOS
+
 #include "Chaos/Utilities.h"
 #include "Chaos/Plane.h"
 #include "Chaos/Box.h"
@@ -20,7 +20,7 @@ GeometryCollectionActor.cpp: AGeometryCollectionActor methods.
 #include "GeometryCollection/GeometryCollectionUtility.h"
 #include "Math/Box.h"
 #include "Physics/PhysicsInterfaceCore.h"
-#include "PBDRigidsSolver.h"
+#include "PhysicsSolver.h"
 
 
 DEFINE_LOG_CATEGORY_STATIC(AGeometryCollectionActorLogging, Log, All);
@@ -55,7 +55,7 @@ void AGeometryCollectionActor::Tick(float DeltaTime)
 }
 
 
-const Chaos::FPBDRigidsSolver* GetSolver(const AGeometryCollectionActor& GeomCollectionActor)
+const Chaos::FPhysicsSolver* GetSolver(const AGeometryCollectionActor& GeomCollectionActor)
 {
 	return GeomCollectionActor.GetGeometryCollectionComponent()->ChaosSolverActor != nullptr ? GeomCollectionActor.GetGeometryCollectionComponent()->ChaosSolverActor->GetSolver() : GeomCollectionActor.GetWorld()->PhysicsScene_Chaos->GetSolver();
 }
@@ -69,9 +69,10 @@ bool LowLevelRaycastImp(const Chaos::TVector<float, 3>& Start, const Chaos::TVec
 	const TSharedPtr<FPhysScene_Chaos> Scene = GeomCollectionActor.GetGeometryCollectionComponent()->GetPhysicsScene();
 	ensure(Scene);
 
-	const Chaos::FPBDRigidsSolver* Solver = GetSolver(GeomCollectionActor);
+	const Chaos::FPhysicsSolver* Solver = GetSolver(GeomCollectionActor);
 	if(ensure(Solver))
 	{
+#if TODO_REIMPLEMENT_GET_RIGID_PARTICLES
 		const TPBDRigidParticles<float, 3>& Particles = Solver->GetRigidParticles();	//todo(ocohen): should these just get passed in instead of hopping through scene?
 		TBox<float, 3> RayBox(Start, Start);
 		RayBox.Thicken(Dir * DeltaMag);
@@ -105,6 +106,7 @@ bool LowLevelRaycastImp(const Chaos::TVector<float, 3>& Start, const Chaos::TVec
 				return true;
 			}
 		}
+#endif
 	}
 
 	return false;
@@ -128,9 +130,12 @@ bool AGeometryCollectionActor::RaycastSingle(FVector Start, FVector End, FHitRes
 	return false;
 }
 
+#endif
+
 #if WITH_EDITOR
 bool AGeometryCollectionActor::GetReferencedContentObjects(TArray<UObject*>& Objects) const
 {
+#if INCLUDE_CHAOS
 	Super::GetReferencedContentObjects(Objects);
 
 	if (GeometryCollectionComponent)
@@ -141,8 +146,7 @@ bool AGeometryCollectionActor::GetReferencedContentObjects(TArray<UObject*>& Obj
 			Objects.Add(GeometryCollection);
 		}
 	}
+#endif
 	return true;
 }
 #endif
-
-#endif
diff --git a/Engine/Source/Runtime/Experimental/GeometryCollectionEngine/Private/GeometryCollection/GeometryCollectionActor_ImmediateImpl.cpp b/Engine/Source/Runtime/Experimental/GeometryCollectionEngine/Private/GeometryCollection/GeometryCollectionActor_ImmediateImpl.cpp
--- a/Engine/Source/Runtime/Experimental/GeometryCollectionEngine/Private/GeometryCollection/GeometryCollectionActor_ImmediateImpl.cpp
+++ b/Engine/Source/Runtime/Experimental/GeometryCollectionEngine/Private/GeometryCollection/GeometryCollectionActor_ImmediateImpl.cpp
@@ -6,7 +6,7 @@ GeometryCollectionActor.cpp: AGeometryCollectionActor methods.
 
 #include "GeometryCollection/GeometryCollectionActor.h"
 
-#if !INCLUDE_CHAOS
+#if !INCLUDE_CHAOS && 0
 
 #include "GeometryCollection/GeometryCollectionAlgo.h"
 #include "GeometryCollection/GeometryCollectionComponent.h"
@@ -102,104 +102,102 @@ void AGeometryCollectionActor::Tick(float DeltaTime)
 #endif
 }
 
-void AGeometryCollectionActor::StartFrameCallback(float EndFrame)
-{
-#if 0	//we don't access dynamic collection directly anymore
-	UE_LOG(AGeometryCollectionActorLogging, Verbose, TEXT("AGeometryCollectionActor::StartFrameCallback()"));
-	UGeometryCollection* Collection = GeometryCollectionComponent->GetDynamicCollection();
-	if (!Scene.GetSimulation()->NumActors() && Collection->GetGeometryCollection()->HasAttribute("RigidBodyID", FGeometryCollection::TransformGroup))
-	{
-		TManagedArray<int32> & RigidBodyId = *RigidBodyIdArray;
-		TManagedArray<FVector> & CenterOfMass = *CenterOfMassArray;
-
-		TManagedArray<FTransform> & Transform = *Collection->GetGeometryCollection()->Transform;
-		TManagedArray<int32> & BoneMap = *Collection->GetGeometryCollection()->BoneMap;
-		TManagedArray<FVector> & Vertex = Collection->GetGeometryCollection()->Vertex;
-
-		PxMaterial* NewMaterial = GPhysXSDK->createMaterial(0, 0, 0);
-
-		// floor
-		FTransform FloorTransform;
-		PxRigidStatic* FloorActor = GPhysXSDK->createRigidStatic(U2PTransform(FTransform::Identity));
-		PxShape* FloorShape = PxRigidActorExt::createExclusiveShape(*FloorActor, PxBoxGeometry(U2PVector(FVector(10000.f, 10000.f, 10.f))), *NewMaterial);
-		// This breaks threading correctness in a general sense but is needed until we can call this in create rigid bodies
-		const_cast<ImmediatePhysics::FSimulation*>(Scene.GetSimulation())->CreateStaticActor(FloorActor, FloorTransform);
-
-		FVector Scale = GeometryCollectionComponent->GetComponentTransform().GetScale3D();
-
-		TArray<FBox> Bounds;
-		Bounds.AddZeroed(Collection->GetGeometryCollection()->NumElements(FGeometryCollection::TransformGroup));
-
-		TArray<int32> SurfaceParticlesCount;
-		SurfaceParticlesCount.AddZeroed(Collection->GetGeometryCollection()->NumElements(FGeometryCollection::TransformGroup));
-
-		TArray<FVector> SumOfMass;
-		SumOfMass.AddZeroed(Collection->GetGeometryCollection()->NumElements(FGeometryCollection::TransformGroup));
-
-		for (int i = 0; i < Vertex.Num(); i++)
-		{
-			FVector ScaledVertex = Scale * Vertex[i];
-			int32 ParticleIndex = BoneMap[i];
-			Bounds[ParticleIndex] += ScaledVertex;
-			SurfaceParticlesCount[ParticleIndex]++;
-			SumOfMass[ParticleIndex] += ScaledVertex;
-		}
-
-
-		for (int32 i = 0; i < Collection->GetGeometryCollection()->TransformNum(); ++i)
-		{
-			if (SurfaceParticlesCount[i] && 0.f < Bounds[i].GetSize().SizeSquared())
-			{
-				CenterOfMass[i] = SumOfMass[i] / SurfaceParticlesCount[i];
-				Bounds[i] = Bounds[i].InverseTransformBy(FTransform(CenterOfMass[i]));
-
-				RigidBodyId[i] = i;
-				int32 RigidBodyIndex = RigidBodyId[i];
-
-				FTransform NewTransform = TransformMatrix(GeometryCollectionComponent->GetComponentTransform(), Transform[i]);
-				float SideSquared = Bounds[i].GetSize()[0] * Bounds[i].GetSize()[0] / 6.f;
-
-				PxRigidDynamic* NewActor = GPhysXSDK->createRigidDynamic(U2PTransform(FTransform::Identity));
-				NewActor->setLinearVelocity(U2PVector(FVector(0.f, 0.f, 0.f)));
-				NewActor->setAngularVelocity(U2PVector(FVector(0.f, 0.f, 0.f)));
-				NewActor->setMass(1.f);
-				NewActor->setMassSpaceInertiaTensor(U2PVector(FVector(SideSquared, SideSquared, SideSquared)));
-				PxShape* NewShape = PxRigidActorExt::createExclusiveShape(*NewActor, PxBoxGeometry(U2PVector((Bounds[i].Max - Bounds[i].Min) / 2.f)), *NewMaterial);
-				const_cast<ImmediatePhysics::FSimulation*>(Scene.GetSimulation())->CreateDynamicActor(NewActor, NewTransform);
-			}
-		}
-	}
-#endif
-}
-
-void AGeometryCollectionActor::CreateRigidBodyCallback(FSolverCallbacks::FParticlesType& Particles)
-{
-}
-
-void AGeometryCollectionActor::EndFrameCallback(float EndFrame)
-{
-#if 0	//we don't use dynamic collection directly anymore
-	UE_LOG(AGeometryCollectionActorLogging, Log, TEXT("AGeometryCollectionActor::EndFrameFunction()"));
-	UGeometryCollection* Collection = GeometryCollectionComponent->GetDynamicCollection();
-	if (Collection->GetGeometryCollection()->HasAttribute("RigidBodyID", FGeometryCollection::TransformGroup))
-	{
-		TManagedArray<int32> & RigidBodyId = *RigidBodyIdArray;
-		TManagedArray<FVector> & CenterOfMass = *CenterOfMassArray;
-		TManagedArray<FTransform> & Transform = *Collection->GetGeometryCollection()->Transform;
-
-		const TArray<ImmediatePhysics::FActorHandle*>& Actors = Scene.GetSimulation()->GetActorHandles();
-
-		FTransform InverseComponentTransform = GeometryCollectionComponent->GetComponentTransform().Inverse();
-		for (int i = 0; i < Collection->GetGeometryCollection()->NumElements(FGeometryCollection::TransformGroup); i++)
-		{
-			int32 RigidBodyIndex = RigidBodyId[i];
-			Transform[i] = TransformMatrix(InverseComponentTransform, Actors[RigidBodyIndex]->GetWorldTransform());
-		}
-
-		GeometryCollectionComponent->SetRenderStateDirty();
-	}
-#endif
-}
+// void AGeometryCollectionActor::StartFrameCallback(float EndFrame)
+// {
+// #if 0	//we don't access dynamic collection directly anymore
+// 	UE_LOG(AGeometryCollectionActorLogging, Verbose, TEXT("AGeometryCollectionActor::StartFrameCallback()"));
+// 	UGeometryCollection* Collection = GeometryCollectionComponent->GetDynamicCollection();
+// 	if (!Scene.GetSimulation()->NumActors() && Collection->GetGeometryCollection()->HasAttribute("RigidBodyID", FGeometryCollection::TransformGroup))
+// 	{
+// 		TManagedArray<int32> & RigidBodyId = *RigidBodyIdArray;
+// 		TManagedArray<FVector> & CenterOfMass = *CenterOfMassArray;
+// 
+// 		TManagedArray<FTransform> & Transform = *Collection->GetGeometryCollection()->Transform;
+// 		TManagedArray<int32> & BoneMap = *Collection->GetGeometryCollection()->BoneMap;
+// 		TManagedArray<FVector> & Vertex = Collection->GetGeometryCollection()->Vertex;
+// 
+// 		PxMaterial* NewMaterial = GPhysXSDK->createMaterial(0, 0, 0);
+// 
+// 		// floor
+// 		FTransform FloorTransform;
+// 		PxRigidStatic* FloorActor = GPhysXSDK->createRigidStatic(U2PTransform(FTransform::Identity));
+// 		PxShape* FloorShape = PxRigidActorExt::createExclusiveShape(*FloorActor, PxBoxGeometry(U2PVector(FVector(10000.f, 10000.f, 10.f))), *NewMaterial);
+// 		// This breaks threading correctness in a general sense but is needed until we can call this in create rigid bodies
+// 		const_cast<ImmediatePhysics::FSimulation*>(Scene.GetSimulation())->CreateStaticActor(FloorActor, FloorTransform);
+// 
+// 		FVector Scale = GeometryCollectionComponent->GetComponentTransform().GetScale3D();
+// 
+// 		TArray<FBox> Bounds;
+// 		Bounds.AddZeroed(Collection->GetGeometryCollection()->NumElements(FGeometryCollection::TransformGroup));
+// 
+// 		TArray<int32> SurfaceParticlesCount;
+// 		SurfaceParticlesCount.AddZeroed(Collection->GetGeometryCollection()->NumElements(FGeometryCollection::TransformGroup));
+// 
+// 		TArray<FVector> SumOfMass;
+// 		SumOfMass.AddZeroed(Collection->GetGeometryCollection()->NumElements(FGeometryCollection::TransformGroup));
+// 
+// 		for (int i = 0; i < Vertex.Num(); i++)
+// 		{
+// 			FVector ScaledVertex = Scale * Vertex[i];
+// 			int32 ParticleIndex = BoneMap[i];
+// 			Bounds[ParticleIndex] += ScaledVertex;
+// 			SurfaceParticlesCount[ParticleIndex]++;
+// 			SumOfMass[ParticleIndex] += ScaledVertex;
+// 		}
+// 
+// 
+// 		for (int32 i = 0; i < Collection->GetGeometryCollection()->TransformNum(); ++i)
+// 		{
+// 			if (SurfaceParticlesCount[i] && 0.f < Bounds[i].GetSize().SizeSquared())
+// 			{
+// 				CenterOfMass[i] = SumOfMass[i] / SurfaceParticlesCount[i];
+// 				Bounds[i] = Bounds[i].InverseTransformBy(FTransform(CenterOfMass[i]));
+// 
+// 				RigidBodyId[i] = i;
+// 				int32 RigidBodyIndex = RigidBodyId[i];
+// 
+// 				FTransform NewTransform = TransformMatrix(GeometryCollectionComponent->GetComponentTransform(), Transform[i]);
+// 				float SideSquared = Bounds[i].GetSize()[0] * Bounds[i].GetSize()[0] / 6.f;
+// 
+// 				PxRigidDynamic* NewActor = GPhysXSDK->createRigidDynamic(U2PTransform(FTransform::Identity));
+// 				NewActor->setLinearVelocity(U2PVector(FVector(0.f, 0.f, 0.f)));
+// 				NewActor->setAngularVelocity(U2PVector(FVector(0.f, 0.f, 0.f)));
+// 				NewActor->setMass(1.f);
+// 				NewActor->setMassSpaceInertiaTensor(U2PVector(FVector(SideSquared, SideSquared, SideSquared)));
+// 				PxShape* NewShape = PxRigidActorExt::createExclusiveShape(*NewActor, PxBoxGeometry(U2PVector((Bounds[i].Max - Bounds[i].Min) / 2.f)), *NewMaterial);
+// 				const_cast<ImmediatePhysics::FSimulation*>(Scene.GetSimulation())->CreateDynamicActor(NewActor, NewTransform);
+// 			}
+// 		}
+// 	}
+// #endif
+// }
+// 
+// void AGeometryCollectionActor::CreateRigidBodyCallback(FSolverCallbacks::FParticlesType& Particles)
+// {
+// }
+// 
+// void AGeometryCollectionActor::EndFrameCallback(float EndFrame)
+// {
+// #if 0	//we don't use dynamic collection directly anymore
+// 	UE_LOG(AGeometryCollectionActorLogging, Log, TEXT("AGeometryCollectionActor::EndFrameFunction()"));
+// 	UGeometryCollection* Collection = GeometryCollectionComponent->GetDynamicCollection();
+// 	if (Collection->GetGeometryCollection()->HasAttribute("RigidBodyID", FGeometryCollection::TransformGroup))
+// 	{
+// 		TManagedArray<int32> & RigidBodyId = *RigidBodyIdArray;
+// 		TManagedArray<FVector> & CenterOfMass = *CenterOfMassArray;
+// 		TManagedArray<FTransform> & Transform = *Collection->GetGeometryCollection()->Transform;
+// 
+// 		FTransform InverseComponentTransform = GeometryCollectionComponent->GetComponentTransform().Inverse();
+// 		for (int i = 0; i < Collection->GetGeometryCollection()->NumElements(FGeometryCollection::TransformGroup); i++)
+// 		{
+// 			int32 RigidBodyIndex = RigidBodyId[i];
+// 			Transform[i] = TransformMatrix(InverseComponentTransform, Scene.GetSimulation()->GetActorHandle(RigidBodyIndex)->GetWorldTransform());
+// 		}
+// 
+// 		GeometryCollectionComponent->SetRenderStateDirty();
+// 	}
+// #endif
+// }
 
 bool AGeometryCollectionActor::RaycastSingle(FVector Start, FVector End, FHitResult& OutHit) const
 {
@@ -209,13 +207,13 @@ bool AGeometryCollectionActor::RaycastSingle(FVector Start, FVector End, FHitRes
 	return false;
 }
 
-void AGeometryCollectionActor::UpdateKinematicBodiesCallback(FSolverCallbacks::FParticlesType& Particles, const float Dt, const float Time, const int32 Index) {}
-
-void AGeometryCollectionActor::ParameterUpdateCallback(FSolverCallbacks::FParticlesType& Particles, const float Time) {}
-
-void AGeometryCollectionActor::DisableCollisionsCallback(TSet<TTuple<int32, int32>>& CollisionPairs) {}
-
-void AGeometryCollectionActor::AddForceCallback(FSolverCallbacks::FParticlesType& Particles, const float Dt, const int32 Index) {}
+//void AGeometryCollectionActor::UpdateKinematicBodiesCallback(FSolverCallbacks::FParticlesType& Particles, const float Dt, const float Time, const int32 Index) {}
+//
+//void AGeometryCollectionActor::ParameterUpdateCallback(FSolverCallbacks::FParticlesType& Particles, const float Time) {}
+//
+//void AGeometryCollectionActor::DisableCollisionsCallback(TSet<TTuple<int32, int32>>& CollisionPairs) {}
+//
+//void AGeometryCollectionActor::AddForceCallback(FSolverCallbacks::FParticlesType& Particles, const float Dt, const int32 Index) {}
 
 #if WITH_EDITOR
 bool AGeometryCollectionActor::GetReferencedContentObjects(TArray<UObject*>& Objects) const
diff --git a/Engine/Source/Runtime/Experimental/GeometryCollectionEngine/Private/GeometryCollection/GeometryCollectionComponent.cpp b/Engine/Source/Runtime/Experimental/GeometryCollectionEngine/Private/GeometryCollection/GeometryCollectionComponent.cpp
--- a/Engine/Source/Runtime/Experimental/GeometryCollectionEngine/Private/GeometryCollection/GeometryCollectionComponent.cpp
+++ b/Engine/Source/Runtime/Experimental/GeometryCollectionEngine/Private/GeometryCollection/GeometryCollectionComponent.cpp
@@ -20,8 +20,8 @@
 #include "Modules/ModuleManager.h"
 #include "ChaosSolversModule.h"
 #include "ChaosStats.h"
-#include "SolverObjects/GeometryCollectionPhysicsObject.h"
-#include "PBDRigidsSolver.h"
+#include "PhysicsProxy/GeometryCollectionPhysicsProxy.h"
+#include "PhysicsSolver.h"
 
 #if WITH_EDITOR
 #include "AssetToolsModule.h"
@@ -48,6 +48,7 @@ FGeometryCollectionSQAccelerator GlobalGeomCollectionAccelerator;	//todo(ocohen)
 
 void HackRegisterGeomAccelerator(UGeometryCollectionComponent& Component)
 {
+#if TODO_REIMPLEMENT_SCENEQUERY_CROSSENGINE
 	if (UWorld* World = Component.GetWorld())
 	{
 		if (FPhysScene* PhysScene = World->GetPhysicsScene())
@@ -58,6 +59,7 @@ void HackRegisterGeomAccelerator(UGeometryCollectionComponent& Component)
 			}
 		}
 	}
+#endif
 }
 #endif
 
@@ -113,7 +115,7 @@ UGeometryCollectionComponent::UGeometryCollectionComponent(const FObjectInitiali
 	, NavmeshInvalidationTimeSliceIndex(0)
 	, IsObjectDynamic(false)
 	, IsObjectLoading(true)
-	, PhysicsObject(nullptr)
+	, PhysicsProxy(nullptr)
 #if WITH_EDITOR && WITH_EDITORONLY_DATA
 	, EditorActor(nullptr)
 #endif
@@ -146,7 +148,7 @@ UGeometryCollectionComponent::UGeometryCollectionComponent(const FObjectInitiali
 
 
 #if INCLUDE_CHAOS
-Chaos::FPBDRigidsSolver* GetSolver(const UGeometryCollectionComponent& GeometryCollectionComponent)
+Chaos::FPhysicsSolver* GetSolver(const UGeometryCollectionComponent& GeometryCollectionComponent)
 {
 	return	GeometryCollectionComponent.ChaosSolverActor != nullptr ? GeometryCollectionComponent.ChaosSolverActor->GetSolver() : GeometryCollectionComponent.GetOwner()->GetWorld()->PhysicsScene_Chaos->GetSolver();
 }
@@ -206,16 +208,16 @@ void UGeometryCollectionComponent::BeginPlay()
 	FChaosSolversModule* ChaosModule = FModuleManager::Get().GetModulePtr<FChaosSolversModule>("ChaosSolvers");
 	if (ChaosModule != nullptr)
 	{
-		Chaos::FPBDRigidsSolver* Solver = GetSolver(*this);
+		Chaos::FPhysicsSolver* Solver = GetSolver(*this);
 		if (Solver != nullptr)
 		{
-			if (PhysicsObject != nullptr)
+			if (PhysicsProxy != nullptr)
 			{
-				ChaosModule->GetDispatcher()->EnqueueCommand(Solver, [&InPhysicsObject = PhysicsObject](Chaos::FPBDRigidsSolver* InSolver)
+				ChaosModule->GetDispatcher()->EnqueueCommandImmediate(Solver, [&InPhysicsProxy = PhysicsProxy](Chaos::FPhysicsSolver* InSolver)
 				{
-					if (InPhysicsObject)
+					if (InPhysicsProxy)
 					{
-						InPhysicsObject->ActivateBodies();
+						InPhysicsProxy->ActivateBodies();
 					}
 				});
 			}
@@ -375,7 +377,7 @@ bool UGeometryCollectionComponent::ShouldCreatePhysicsState() const
 
 bool UGeometryCollectionComponent::HasValidPhysicsState() const
 {
-	return PhysicsObject != nullptr;
+	return PhysicsProxy != nullptr;
 }
 
 void UGeometryCollectionComponent::SetNotifyBreaks(bool bNewNotifyBreaks)
@@ -854,12 +856,15 @@ void UGeometryCollectionComponent::InitDynamicData(FGeometryCollectionDynamicDat
 				{
 #if INCLUDE_CHAOS
 					// clear events on the solver
-					Chaos::FPBDRigidsSolver *Solver = GetSolver(*this);
+					Chaos::FPhysicsSolver *Solver = GetSolver(*this);
 					
 					if (Solver)
 					{
+#if TODO_REPLACE_SOLVER_LOCK
 						Chaos::FSolverWriteLock ScopedWriteLock(Solver);
+#endif
 
+#if TODO_REIMPLEMENT_EVENTS_DATA_ARRAYS
 						//////////////////////////////////////////////////////////////////////////
 						// Temporary workaround for writing on multiple threads.
 						// The following is called wide from the render thread to populate
@@ -870,9 +875,9 @@ void UGeometryCollectionComponent::InitDynamicData(FGeometryCollectionDynamicDat
 						// TO BE REFACTORED
 						// #TODO BG
 						//////////////////////////////////////////////////////////////////////////
-						Chaos::FPBDRigidsSolver::FAllCollisionData *CollisionDataToWriteTo = const_cast<Chaos::FPBDRigidsSolver::FAllCollisionData*>(Solver->GetAllCollisions_FromSequencerCache_NEEDSLOCK());
-						Chaos::FPBDRigidsSolver::FAllBreakingData *BreakingDataToWriteTo = const_cast<Chaos::FPBDRigidsSolver::FAllBreakingData*>(Solver->GetAllBreakings_FromSequencerCache_NEEDSLOCK());
-						Chaos::FPBDRigidsSolver::FAllTrailingData *TrailingDataToWriteTo = const_cast<Chaos::FPBDRigidsSolver::FAllTrailingData*>(Solver->GetAllTrailings_FromSequencerCache_NEEDSLOCK());
+						Chaos::FPhysicsSolver::FAllCollisionData *CollisionDataToWriteTo = const_cast<Chaos::FPhysicsSolver::FAllCollisionData*>(Solver->GetAllCollisions_FromSequencerCache_NEEDSLOCK());
+						Chaos::FPhysicsSolver::FAllBreakingData *BreakingDataToWriteTo = const_cast<Chaos::FPhysicsSolver::FAllBreakingData*>(Solver->GetAllBreakings_FromSequencerCache_NEEDSLOCK());
+						Chaos::FPhysicsSolver::FAllTrailingData *TrailingDataToWriteTo = const_cast<Chaos::FPhysicsSolver::FAllTrailingData*>(Solver->GetAllTrailings_FromSequencerCache_NEEDSLOCK());
 
 						if (!FMath::IsNearlyEqual(CollisionDataToWriteTo->TimeCreated, DesiredCacheTime))
 						{
@@ -883,7 +888,7 @@ void UGeometryCollectionComponent::InitDynamicData(FGeometryCollectionDynamicDat
 						int32 Index = CacheParameters.TargetCache->GetData()->FindLastKeyBefore(CurrentCacheTime);
 						const FRecordedFrame *RecordedFrame = &CacheParameters.TargetCache->GetData()->Records[Index];				
 
-						if (RecordedFrame && PhysicsObject && !EventsPlayed[Index])
+						if (RecordedFrame && PhysicsProxy && !EventsPlayed[Index])
 						{
 							EventsPlayed[Index] = true;
 
@@ -905,7 +910,9 @@ void UGeometryCollectionComponent::InitDynamicData(FGeometryCollectionDynamicDat
 									AllCollisionsDataArrayItem.AngularVelocity2 = RecordedFrame->Collisions[Idx].AngularVelocity2;
 									AllCollisionsDataArrayItem.Mass1 = RecordedFrame->Collisions[Idx].Mass1;
 									AllCollisionsDataArrayItem.Mass2 = RecordedFrame->Collisions[Idx].Mass2;
+#if TODO_CONVERT_GEOMETRY_COLLECTION_PARTICLE_INDICES_TO_PARTICLE_POINTERS
 									AllCollisionsDataArrayItem.ParticleIndex = RecordedFrame->Collisions[Idx].ParticleIndex;
+#endif
 									AllCollisionsDataArrayItem.LevelsetIndex = RecordedFrame->Collisions[Idx].LevelsetIndex;
 									AllCollisionsDataArrayItem.ParticleIndexMesh = RecordedFrame->Collisions[Idx].ParticleIndexMesh;
 									AllCollisionsDataArrayItem.LevelsetIndexMesh = RecordedFrame->Collisions[Idx].LevelsetIndexMesh;
@@ -925,7 +932,9 @@ void UGeometryCollectionComponent::InitDynamicData(FGeometryCollectionDynamicDat
 									AllBreakingsDataArrayItem.Velocity = RecordedFrame->Breakings[Idx].Velocity;
 									AllBreakingsDataArrayItem.AngularVelocity = RecordedFrame->Breakings[Idx].AngularVelocity;
 									AllBreakingsDataArrayItem.Mass = RecordedFrame->Breakings[Idx].Mass;
+#if TODO_CONVERT_GEOMETRY_COLLECTION_PARTICLE_INDICES_TO_PARTICLE_POINTERS
 									AllBreakingsDataArrayItem.ParticleIndex = RecordedFrame->Breakings[Idx].ParticleIndex;
+#endif
 									AllBreakingsDataArrayItem.ParticleIndexMesh = RecordedFrame->Breakings[Idx].ParticleIndexMesh;
 								}
 							}
@@ -943,11 +952,14 @@ void UGeometryCollectionComponent::InitDynamicData(FGeometryCollectionDynamicDat
 									AllTrailingsDataArrayItem.Velocity = Trailing.Velocity;
 									AllTrailingsDataArrayItem.AngularVelocity = Trailing.AngularVelocity;
 									AllTrailingsDataArrayItem.Mass = Trailing.Mass;
+#if TODO_CONVERT_GEOMETRY_COLLECTION_PARTICLE_INDICES_TO_PARTICLE_POINTERS
 									AllTrailingsDataArrayItem.ParticleIndex = Trailing.ParticleIndex;
+#endif
 									AllTrailingsDataArrayItem.ParticleIndexMesh = Trailing.ParticleIndexMesh;
 								}
 							}
 						}
+#endif
 					}
 #endif
 				}								
@@ -1159,7 +1171,7 @@ void UGeometryCollectionComponent::OnCreatePhysicsState()
 	DummyBodyInstance.SetResponseToAllChannels(ECR_Block);
 #endif
 
-	if (!PhysicsObject)
+	if (!PhysicsProxy)
 	{
 #if WITH_EDITOR && WITH_EDITORONLY_DATA
 		EditorActor = nullptr;
@@ -1391,9 +1403,9 @@ void UGeometryCollectionComponent::OnCreatePhysicsState()
 			}
 			// end temporary 
 
-			PhysicsObject = new FGeometryCollectionPhysicsObject(this, DynamicCollection.Get(), InitFunc, CacheSyncFunc, FinalSyncFunc);
+			PhysicsProxy = new FGeometryCollectionPhysicsProxy(this, DynamicCollection.Get(), InitFunc, CacheSyncFunc, FinalSyncFunc);
 			TSharedPtr<FPhysScene_Chaos> Scene = GetPhysicsScene();
-			Scene->AddObject(this, PhysicsObject);
+			Scene->AddObject(this, PhysicsProxy);
 
 			RegisterForEvents();
 		}
@@ -1402,7 +1414,7 @@ void UGeometryCollectionComponent::OnCreatePhysicsState()
 
 #if INCLUDE_CHAOS
 #if WITH_PHYSX && !WITH_CHAOS_NEEDS_TO_BE_FIXED
-	if (PhysicsObject)
+	if (PhysicsProxy)
 	{
 		GlobalGeomCollectionAccelerator.AddComponent(this);
 	}
@@ -1428,14 +1440,14 @@ void UGeometryCollectionComponent::OnDestroyPhysicsState()
 #endif
 
 #if INCLUDE_CHAOS
-	if(PhysicsObject)
+	if(PhysicsProxy)
 	{
 		TSharedPtr<FPhysScene_Chaos> Scene = GetPhysicsScene();
-		Scene->RemoveObject(PhysicsObject);
+		Scene->RemoveObject(PhysicsProxy);
 		InitializationState = ESimulationInitializationState::Unintialized;
 
 		// Discard the pointer (cleanup happens through the scene or dedicated thread)
-		PhysicsObject = nullptr;
+		PhysicsProxy = nullptr;
 	}
 #endif
 }
@@ -1952,7 +1964,7 @@ void UGeometryCollectionComponent::ApplyPhysicsField(bool Enabled, EGeometryColl
 void UGeometryCollectionComponent::DispatchCommand(const FFieldSystemCommand& InCommand)
 {
 #if INCLUDE_CHAOS
-	if (PhysicsObject)
+	if (PhysicsProxy)
 	{
 		FChaosSolversModule* ChaosModule = FModuleManager::Get().GetModulePtr<FChaosSolversModule>("ChaosSolvers");
 		checkSlow(ChaosModule);
@@ -1960,10 +1972,10 @@ void UGeometryCollectionComponent::DispatchCommand(const FFieldSystemCommand& In
 		Chaos::IDispatcher* PhysicsDispatcher = ChaosModule->GetDispatcher();
 		checkSlow(PhysicsDispatcher); // Should always have one of these
 
-		PhysicsDispatcher->EnqueueCommand([PhysicsObject = this->PhysicsObject, NewCommand = InCommand]()
+		PhysicsDispatcher->EnqueueCommandImmediate([PhysicsProxy = this->PhysicsProxy, NewCommand = InCommand]()
 		{
 			// Pass through nullptr here as geom component commands can never affect other solvers
-			PhysicsObject->BufferCommand(nullptr, NewCommand);
+			PhysicsProxy->BufferCommand(nullptr, NewCommand);
 		});
 	}
 #endif
@@ -2055,7 +2067,7 @@ void UGeometryCollectionComponent::CalculateGlobalMatrices()
 	FChaosSolversModule* Module = FChaosSolversModule::GetModule();
 	Module->LockResultsRead();
 
-	const FGeometryCollectionResults* Results = PhysicsObject ? &PhysicsObject->GetPhysicsResults().GetGameDataForRead() : nullptr;
+	const FGeometryCollectionResults* Results = PhysicsProxy ? &PhysicsProxy->GetPhysicsResults().GetGameDataForRead() : nullptr;
 
 	const int32 NumTransforms = Results ? Results->GlobalTransforms.Num() : 0;
 	if(NumTransforms > 0)
diff --git a/Engine/Source/Runtime/Experimental/GeometryCollectionEngine/Private/GeometryCollection/GeometryCollectionDebugDrawActor.cpp b/Engine/Source/Runtime/Experimental/GeometryCollectionEngine/Private/GeometryCollection/GeometryCollectionDebugDrawActor.cpp
--- a/Engine/Source/Runtime/Experimental/GeometryCollectionEngine/Private/GeometryCollection/GeometryCollectionDebugDrawActor.cpp
+++ b/Engine/Source/Runtime/Experimental/GeometryCollectionEngine/Private/GeometryCollection/GeometryCollectionDebugDrawActor.cpp
@@ -19,9 +19,9 @@
 #include "HAL/IConsoleManager.h"
 #if INCLUDE_CHAOS
 #include "PBDRigidsSolver.h"
+#include "PhysicsSolver.h"  // #if TODO_REIMPLEMENT_GET_RIGID_PARTICLES
 #endif  // #if INCLUDE_CHAOS
 
-
 DEFINE_LOG_CATEGORY_STATIC(LogGeometryCollectionDebugDrawActor, Log, All);
 
 namespace GeometryCollectionDebugDrawActorConstants
@@ -286,6 +286,7 @@ void AGeometryCollectionDebugDrawActor::Tick(float DeltaSeconds)
 	if (World && SelectedRigidBody.Id != INDEX_NONE && !SelectedRigidBody.GeometryCollection)
 	{
 #if INCLUDE_CHAOS
+#if TODO_REIMPLEMENT_GET_RIGID_PARTICLES
 		// Check the id is within the selected solver range
 		const Chaos::FPBDRigidsSolver* const Solver = 
 			SelectedRigidBody.Solver ? SelectedRigidBody.Solver->GetSolver() :  // Selected solver
@@ -298,6 +299,7 @@ void AGeometryCollectionDebugDrawActor::Tick(float DeltaSeconds)
 			UE_LOG(LogGeometryCollectionDebugDrawActor, VeryVerbose, TEXT("The selection id is out of range."));
 		}
 		else  // Statement continues below...
+#endif  // #if TODO_REIMPLEMENT_GET_RIGID_PARTICLES
 #endif  // #if INCLUDE_CHAOS
 		{
 			UE_LOG(LogGeometryCollectionDebugDrawActor, VeryVerbose, TEXT("The selection couldn't be found. The property update will run on all components still containing any invalid rigid body ids."));
diff --git a/Engine/Source/Runtime/Experimental/GeometryCollectionEngine/Private/GeometryCollection/GeometryCollectionDebugDrawComponent.cpp b/Engine/Source/Runtime/Experimental/GeometryCollectionEngine/Private/GeometryCollection/GeometryCollectionDebugDrawComponent.cpp
--- a/Engine/Source/Runtime/Experimental/GeometryCollectionEngine/Private/GeometryCollection/GeometryCollectionDebugDrawComponent.cpp
+++ b/Engine/Source/Runtime/Experimental/GeometryCollectionEngine/Private/GeometryCollection/GeometryCollectionDebugDrawComponent.cpp
@@ -9,7 +9,7 @@
 #include "GeometryCollection/GeometryCollectionDebugDrawActor.h"
 #include "GeometryCollection/GeometryCollection.h"
 #if INCLUDE_CHAOS
-#include "SolverObjects/GeometryCollectionPhysicsObject.h"
+#include "PhysicsProxy/GeometryCollectionPhysicsProxy.h"
 #endif  // #if INCLUDE_CHAOS
 #endif  // #if GEOMETRYCOLLECTION_DEBUG_DRAW
 #include "HAL/IConsoleManager.h"
@@ -877,10 +877,10 @@ void UGeometryCollectionDebugDrawComponent::DebugDrawChaosTick()
 
 	// Retrieve synced particle and clustering data
 	const TManagedArray<int32>& RigidBodyIds = GeometryCollectionComponent->RigidBodyIds;
-	const FGeometryCollectionPhysicsObject* const PhysicsObject = GeometryCollectionComponent->GetPhysicsObject();
-	if (PhysicsObject)
+	const FGeometryCollectionPhysicsProxy* const PhysicsProxy = GeometryCollectionComponent->GetPhysicsProxy();
+	if (PhysicsProxy)
 	{
-		ParticlesData.Sync(PhysicsObject->GetSolver(), RigidBodyIds);
+		ParticlesData.Sync(PhysicsProxy->GetSolver(), RigidBodyIds);
 	}
 
 	// Visualize single rigid body
diff --git a/Engine/Source/Runtime/Experimental/GeometryCollectionEngine/Private/GeometryCollection/GeometryCollectionParticlesData.cpp b/Engine/Source/Runtime/Experimental/GeometryCollectionEngine/Private/GeometryCollection/GeometryCollectionParticlesData.cpp
--- a/Engine/Source/Runtime/Experimental/GeometryCollectionEngine/Private/GeometryCollection/GeometryCollectionParticlesData.cpp
+++ b/Engine/Source/Runtime/Experimental/GeometryCollectionEngine/Private/GeometryCollection/GeometryCollectionParticlesData.cpp
@@ -4,7 +4,7 @@
 
 #if INCLUDE_CHAOS && !(UE_BUILD_SHIPPING || UE_BUILD_TEST)
 
-#include "PBDRigidsSolver.h"
+#include "PhysicsSolver.h"
 #include "Chaos/Sphere.h"
 #include "Chaos/Box.h"
 #include "Chaos/ImplicitObjectTransformed.h"
@@ -26,7 +26,7 @@ TGeometryCollectionParticlesData<T, d>::TGeometryCollectionParticlesData()
 }
 
 template<class T, int d>
-void TGeometryCollectionParticlesData<T, d>::Sync(const Chaos::FPBDRigidsSolver* Solver, const TManagedArray<int32>& RigidBodyIds)
+void TGeometryCollectionParticlesData<T, d>::Sync(const Chaos::FPhysicsSolver* Solver, const TManagedArray<int32>& RigidBodyIds)
 {
 	// No point in calling twice the sync function within the same frame
 	if (!ensureMsgf(SyncFrame != GFrameCounter, TEXT("Sync should not happen twice during the same tick.")))
@@ -53,7 +53,7 @@ void TGeometryCollectionParticlesData<T, d>::Sync(const Chaos::FPBDRigidsSolver*
 			Chaos::IDispatcher* const PhysicsDispatcher = ChaosModule->GetDispatcher();
 			check(PhysicsDispatcher);
 
-			PhysicsDispatcher->EnqueueCommand([this, Solver, &RigidBodyIds]()
+			PhysicsDispatcher->EnqueueCommandImmediate([this, Solver, &RigidBodyIds]()
 			{
 				// Iterate through all data
 				FData& Data = BufferedData.GetPhysicsDataForWrite();
@@ -135,9 +135,10 @@ void TGeometryCollectionParticlesData<T, d>::FData::Reset(EGeometryCollectionPar
 }
 
 template<class T, int d>
-void TGeometryCollectionParticlesData<T, d>::FData::Copy(EGeometryCollectionParticlesData Data, const Chaos::FPBDRigidsSolver* Solver, const TManagedArray<int32>& RigidBodyIds)
+void TGeometryCollectionParticlesData<T, d>::FData::Copy(EGeometryCollectionParticlesData Data, const Chaos::FPhysicsSolver* Solver, const TManagedArray<int32>& RigidBodyIds)
 {
 	check(Solver);
+#if TODO_REIMPLEMENT_GET_RIGID_PARTICLES
 	const Chaos::TPBDRigidParticles<T, d>& Particles = Solver->GetRigidParticles();
 	const Chaos::TPBDRigidClustering<FPBDRigidsEvolution, FPBDCollisionConstraints, T, d>& Clustering = Solver->GetRigidClustering();
 
@@ -162,8 +163,8 @@ void TGeometryCollectionParticlesData<T, d>::FData::Copy(EGeometryCollectionPart
 	auto HasBoundingBox = [](Chaos::TSerializablePtr<Chaos::TImplicitObject<T, d>> ImplicitObject) { return ImplicitObject ? ImplicitObject->HasBoundingBox(): false; };
 	auto BoxMin         = [](Chaos::TSerializablePtr<Chaos::TImplicitObject<T, d>> ImplicitObject) { const Chaos::TBox     <T, d>* Box     ; return ImplicitObject && (Box      = ImplicitObject->template GetObject<Chaos::TBox     <T, d>>()) != nullptr ? Box     ->Min    (): Chaos::TVector<T, d>(T(0)); };
 	auto BoxMax         = [](Chaos::TSerializablePtr<Chaos::TImplicitObject<T, d>> ImplicitObject) { const Chaos::TBox     <T, d>* Box     ; return ImplicitObject && (Box      = ImplicitObject->template GetObject<Chaos::TBox     <T, d>>()) != nullptr ? Box     ->Max    (): Chaos::TVector<T, d>(T(0)); };
-	auto SphereCenter   = [](Chaos::TSerializablePtr<Chaos::TImplicitObject<T, d>> ImplicitObject) { const Chaos::TSphere  <T, d>* Sphere  ; return ImplicitObject && (Sphere   = ImplicitObject->template GetObject<Chaos::TSphere  <T, d>>()) != nullptr ? Sphere  ->Center (): Chaos::TVector<T, d>(T(0)); };
-	auto SphereRadius   = [](Chaos::TSerializablePtr<Chaos::TImplicitObject<T, d>> ImplicitObject) { const Chaos::TSphere  <T, d>* Sphere  ; return ImplicitObject && (Sphere   = ImplicitObject->template GetObject<Chaos::TSphere  <T, d>>()) != nullptr ? Sphere  ->Radius (): T(0); };
+	auto SphereCenter   = [](Chaos::TSerializablePtr<Chaos::TImplicitObject<T, d>> ImplicitObject) { const Chaos::TSphere  <T, d>* Sphere  ; return ImplicitObject && (Sphere   = ImplicitObject->template GetObject<Chaos::TSphere  <T, d>>()) != nullptr ? Sphere  ->GetCenter (): Chaos::TVector<T, d>(T(0)); };
+	auto SphereRadius   = [](Chaos::TSerializablePtr<Chaos::TImplicitObject<T, d>> ImplicitObject) { const Chaos::TSphere  <T, d>* Sphere  ; return ImplicitObject && (Sphere   = ImplicitObject->template GetObject<Chaos::TSphere  <T, d>>()) != nullptr ? Sphere  ->GetRadius (): T(0); };
 	auto LevelSetGrid   = [](Chaos::TSerializablePtr<Chaos::TImplicitObject<T, d>> ImplicitObject) { const Chaos::TLevelSet<T, d>* LevelSet; return ImplicitObject && (LevelSet = ImplicitObject->template GetObject<Chaos::TLevelSet<T, d>>()) != nullptr ? LevelSet->GetGrid(): Chaos::TUniformGrid<T, d>(); };
 
 	// Data type copy for all particles
@@ -201,6 +202,7 @@ void TGeometryCollectionParticlesData<T, d>::FData::Copy(EGeometryCollectionPart
 	case EGeometryCollectionParticlesData::ConnectivityEdges     : ConnectivityEdges     .Resize(Count); for (int32 i = 0; i < Count; ++i) { if (RigidBodyIds[i] != INDEX_NONE) { ConnectivityEdges     [i] = Clustering.GetConnectivityEdges()[RigidBodyIds[i]] ; } } break;
 	case EGeometryCollectionParticlesData::ChildToParentMap      : ChildToParentMap      .Resize(Count); for (int32 i = 0; i < Count; ++i) { if (RigidBodyIds[i] != INDEX_NONE) { ChildToParentMap      [i] = Clustering.GetChildToParentMap ()[RigidBodyIds[i]] ; } } break;
 	}
+#endif
 }
 
 template<class T, int d>
diff --git a/Engine/Source/Runtime/Experimental/GeometryCollectionEngine/Private/GeometryCollection/GeometryCollectionSQAccelerator.cpp b/Engine/Source/Runtime/Experimental/GeometryCollectionEngine/Private/GeometryCollection/GeometryCollectionSQAccelerator.cpp
--- a/Engine/Source/Runtime/Experimental/GeometryCollectionEngine/Private/GeometryCollection/GeometryCollectionSQAccelerator.cpp
+++ b/Engine/Source/Runtime/Experimental/GeometryCollectionEngine/Private/GeometryCollection/GeometryCollectionSQAccelerator.cpp
@@ -13,8 +13,8 @@
 #include "Components/BoxComponent.h"
 #include "ChaosSolversModule.h"
 #include "ChaosStats.h"
-#include "PBDRigidsSolver.h"
-#include "SolverObjects/GeometryCollectionPhysicsObject.h"
+#include "PhysicsSolver.h"
+#include "PhysicsProxy/GeometryCollectionPhysicsProxy.h"
 
 #if INCLUDE_CHAOS && !WITH_CHAOS_NEEDS_TO_BE_FIXED
 
@@ -23,6 +23,8 @@ DECLARE_CYCLE_STAT(TEXT("LowLevelSweep"), STAT_LowLevelSweep, STATGROUP_Chaos);
 DECLARE_CYCLE_STAT(TEXT("LowLevelRaycast"), STAT_LowLevelRaycast, STATGROUP_Chaos);
 DECLARE_CYCLE_STAT(TEXT("LowLevelOverlap"), STAT_LowLevelOverlap, STATGROUP_Chaos);
 
+#if TODO_REIMPLEMENT_SCENEQUERY_CROSSENGINE
+
 bool IsValidIndexAndTransform(const FGeometryCollectionResults& PhysResult, const Chaos::TPBDRigidParticles<float, 3>& Particles, const TManagedArray<FTransform>& TransformArray, const TArray<bool>& DisabledFlags, const int32 RigidBodyIdx, const bool bCanBeDisabled)
 {
 	if(RigidBodyIdx == -1)
@@ -69,7 +71,7 @@ bool IsValidIndexAndTransform(const FGeometryCollectionResults& PhysResult, cons
 	return true;
 }
 
-bool LowLevelRaycastSingleElement(int32 InParticleIndex, const Chaos::FPBDRigidsSolver* InSolver, const Chaos::TClusterBuffer<float, 3>& ClusterBuffer, const FGeometryCollectionPhysicsObject* InObject, const FVector& Start, const FVector& Dir, float DeltaMag, bool bCanBeDisabled, EHitFlags OutputFlags, FHitRaycast& OutHit)
+bool LowLevelRaycastSingleElement(int32 InParticleIndex, const Chaos::FPhysicsSolver* InSolver, const Chaos::TClusterBuffer<float, 3>& ClusterBuffer, const FGeometryCollectionPhysicsProxy* InObject, const FVector& Start, const FVector& Dir, float DeltaMag, bool bCanBeDisabled, EHitFlags OutputFlags, FHitRaycast& OutHit)
 {
 	using namespace Chaos;
 
@@ -137,7 +139,7 @@ static TAutoConsoleVariable<int32> CVarMaxSweepSteps(
 	TEXT("Number of steps during a sweep"),
 	ECVF_Default);
 
-bool LowLevelSweepSingleElement(int32 InParticleIndex, const Chaos::FPBDRigidsSolver* InSolver, const Chaos::TClusterBuffer<float, 3>& ClusterBuffer, const FGeometryCollectionPhysicsObject* InObject, const Chaos::TImplicitObject<float, 3>& QueryGeom, const Chaos::TParticles<float, 3>& CollisionParticles, const FTransform& StartPose, const FVector& Dir, float DeltaMag, const bool bCanBeDisabled, FHitSweep& OutHit)
+bool LowLevelSweepSingleElement(int32 InParticleIndex, const Chaos::FPhysicsSolver* InSolver, const Chaos::TClusterBuffer<float, 3>& ClusterBuffer, const FGeometryCollectionPhysicsProxy* InObject, const Chaos::TImplicitObject<float, 3>& QueryGeom, const Chaos::TParticles<float, 3>& CollisionParticles, const FTransform& StartPose, const FVector& Dir, float DeltaMag, const bool bCanBeDisabled, FHitSweep& OutHit)
 {
 	using namespace Chaos;
 
@@ -202,7 +204,7 @@ bool LowLevelOverlap(const UGeometryCollectionComponent& GeomCollectionComponent
 
 	SCOPE_CYCLE_COUNTER(STAT_LowLevelOverlap);
 
-	const FGeometryCollectionPhysicsObject* PhysObject = GeomCollectionComponent.GetPhysicsObject();
+	const FGeometryCollectionPhysicsProxy* PhysObject = GeomCollectionComponent.GetPhysicsProxy();
 	if (!ensure(PhysObject))
 	{
 		return false;
@@ -215,7 +217,7 @@ bool LowLevelOverlap(const UGeometryCollectionComponent& GeomCollectionComponent
 	
 	bool bFound = false;
 
-	if (Chaos::FPBDRigidsSolver* Solver = GeomCollectionComponent.ChaosSolverActor != nullptr ? GeomCollectionComponent.ChaosSolverActor->GetSolver() : GeomCollectionComponent.GetOwner()->GetWorld()->PhysicsScene_Chaos->GetSolver())
+	if (Chaos::FPhysicsSolver* Solver = GeomCollectionComponent.ChaosSolverActor != nullptr ? GeomCollectionComponent.ChaosSolverActor->GetSolver() : GeomCollectionComponent.GetOwner()->GetWorld()->PhysicsScene_Chaos->GetSolver())
 	{
 		const TPBDRigidParticles<float, 3>& Particles = Solver->GetRigidParticles();	//todo(ocohen): should these just get passed in instead of hopping through scene?
 
@@ -253,7 +255,7 @@ bool LowLevelOverlap(const UGeometryCollectionComponent& GeomCollectionComponent
 int32 UseSlowSQ = 0;
 FAutoConsoleVariableRef CVarUseSlowSQ(TEXT("p.UseSlowSQ"), UseSlowSQ, TEXT(""));
 
-void FGeometryCollectionSQAccelerator::Raycast(const FVector& Start, const FVector& Dir, const float DeltaMagnitude, FPhysicsHitCallback<FHitRaycast>& HitBuffer, EHitFlags OutputFlags, FQueryFlags QueryFlags, const FCollisionFilterData& QueryFilter, const FQueryFilterData& QueryFilterData, ICollisionQueryFilterCallbackBase& QueryCallback) const
+void FGeometryCollectionSQAccelerator::Raycast(const FVector& Start, const FVector& Dir, const float DeltaMagnitude, ChaosInterface::FSQHitBuffer<ChaosInterface::FRaycastHit>& HitBuffer, EHitFlags OutputFlags, const FQueryFilterData& QueryFilterData, ICollisionQueryFilterCallbackBase& QueryCallback) const
 {
 	SCOPE_CYCLE_COUNTER(STAT_GCRaycast);
 	FChaosScopeSolverLock SolverScopeLock;
@@ -262,15 +264,15 @@ void FGeometryCollectionSQAccelerator::Raycast(const FVector& Start, const FVect
 
 	FChaosSolversModule* Module = FChaosSolversModule::GetModule();
 
-	TMap<const Chaos::FPBDRigidsSolver*, TArray<int32>> SolverIntersectionSets;
+	TMap<const Chaos::FPhysicsSolver*, TArray<int32>> SolverIntersectionSets;
 
 	Chaos::TSpatialRay<float, 3> Ray(Start, Start + Dir * DeltaMagnitude);
 
 #if WITH_PHYSX
 
-	const TArray<Chaos::FPBDRigidsSolver*>& Solvers = Module->GetSolvers();
+	const TArray<Chaos::FPhysicsSolver*>& Solvers = Module->GetSolvers();
 
-	for(const Chaos::FPBDRigidsSolver* Solver : Solvers)
+	for(const Chaos::FPhysicsSolver* Solver : Solvers)
 	{
 		if (!Solver)
 		{
@@ -283,15 +285,15 @@ void FGeometryCollectionSQAccelerator::Raycast(const FVector& Start, const FVect
 		const Chaos::TClusterBuffer<float, 3>& Buffer = Solver->GetRigidClustering().GetBufferedData();
 
 
-		const Chaos::FPBDRigidsSolver::FSolverObjectReverseMapping& ObjectMap = Solver->GetSolverObjectReverseMapping_GameThread();
+		const Chaos::FPhysicsSolver::FPhysicsProxyReverseMapping& ObjectMap = Solver->GetPhysicsProxyReverseMapping_GameThread();
 
 		FHitRaycast Hit;
 		int32 IntersectionSetSize = IntersectionSet.Num();
 		for(int32 i = 0; i < IntersectionSet.Num(); ++i)
 		{
 			const int32 IntersectParticleIndex = IntersectionSet[i];
-			const SolverObjectWrapper& ObjectWrapper = ObjectMap.SolverObjectReverseMappingArray[IntersectParticleIndex];
-			if (!ObjectWrapper.SolverObject)
+			const PhysicsProxyWrapper& ObjectWrapper = ObjectMap.PhysicsProxyReverseMappingArray[IntersectParticleIndex];
+			if (!ObjectWrapper.PhysicsProxy)
 			{
 				const TImplicitObject<float, 3>* Object = Buffer.GeometryPtrs[IntersectParticleIndex].Get();
 				// Ignore the ground plane
@@ -341,16 +343,16 @@ void FGeometryCollectionSQAccelerator::Raycast(const FVector& Start, const FVect
 				continue;
 			}
 
-			if(ObjectWrapper.Type == ESolverObjectType::GeometryCollectionType && ensure(ObjectWrapper.SolverObject))
+			if(ObjectWrapper.Type == EPhysicsProxyType::GeometryCollectionType && ensure(ObjectWrapper.PhysicsProxy))
 			{
-				LowLevelRaycastSingleElement(IntersectParticleIndex, Solver, Buffer, static_cast<FGeometryCollectionPhysicsObject*>(ObjectWrapper.SolverObject), Start, Dir, DeltaMagnitude, i >= IntersectionSetSize, OutputFlags, Hit);
+				LowLevelRaycastSingleElement(IntersectParticleIndex, Solver, Buffer, static_cast<FGeometryCollectionPhysicsProxy*>(ObjectWrapper.PhysicsProxy), Start, Dir, DeltaMagnitude, i >= IntersectionSetSize, OutputFlags, Hit);
 
 				// If we registered a hit
-				if(Hit.distance != PX_MAX_REAL && ensure(ObjectWrapper.SolverObject))
+				if(Hit.distance != PX_MAX_REAL && ensure(ObjectWrapper.PhysicsProxy))
 				{
 #if !WITH_IMMEDIATE_PHYSX && PHYSICS_INTERFACE_PHYSX
 					//todo(ocohen):hack placeholder while we convert over to non physx API
-					UGeometryCollectionComponent* Component = Cast<UGeometryCollectionComponent>(ObjectWrapper.SolverObject->GetOwner());
+					UGeometryCollectionComponent* Component = Cast<UGeometryCollectionComponent>(ObjectWrapper.PhysicsProxy->GetOwner());
 					check(Component);
 					if (Component->IsRegistered())
 					{
@@ -374,7 +376,7 @@ void FGeometryCollectionSQAccelerator::Raycast(const FVector& Start, const FVect
 		}
 
 		Solver->GetRigidClustering().ReleaseBufferedData();
-		Solver->ReleaseSolverObjectReverseMapping();
+		Solver->ReleasePhysicsProxyReverseMapping();
 	}
 
 #endif
@@ -384,7 +386,7 @@ DECLARE_CYCLE_STAT(TEXT("Sweep Broadphase"), STAT_SQSweepBroadPhase, STATGROUP_C
 DECLARE_CYCLE_STAT(TEXT("Sweep Narrowphase"), STAT_SQSweepNarrowPhase, STATGROUP_Chaos);
 
 //@todo(mlentine): Avoid duplicated code between this and overlap
-void FGeometryCollectionSQAccelerator::Sweep(const FPhysicsGeometry& QueryGeom, const FTransform& StartTM, const FVector& Dir, const float DeltaMag, FPhysicsHitCallback<FHitSweep>& HitBuffer, EHitFlags OutputFlags, FQueryFlags QueryFlags, const FCollisionFilterData& QueryFilter, const FQueryFilterData& QueryFilterData, ICollisionQueryFilterCallbackBase& QueryCallback) const
+void FGeometryCollectionSQAccelerator::Sweep(const Chaos::TImplicitObject<float, 3>& QueryGeom, const FTransform& StartTM, const FVector& Dir, const float DeltaMagnitude, ChaosInterface::FSQHitBuffer<ChaosInterface::FSweepHit>& HitBuffer, EHitFlags OutputFlags, const FQueryFilterData& QueryFilterData, ICollisionQueryFilterCallbackBase& QueryCallback) const
 {
 	SCOPE_CYCLE_COUNTER(STAT_GCSweep);
 	FChaosScopeSolverLock SolverScopeLock;
@@ -392,10 +394,10 @@ void FGeometryCollectionSQAccelerator::Sweep(const FPhysicsGeometry& QueryGeom,
 	using namespace Chaos;
 #if WITH_PHYSX
 
-	TMap<const Chaos::FPBDRigidsSolver*, TArray<int32>> SolverIntersectionSets;
+	TMap<const Chaos::FPhysicsSolver*, TArray<int32>> SolverIntersectionSets;
 
 	// Getter for intersections from the mapping above
-	auto GetIntersectionsFunc = [&](const Chaos::FPBDRigidsSolver* InSolver, const Chaos::TParticles<float, 3>& InCollisionParticles, float InDeltaMag, const FTransform& InPose) -> TArray<int32>
+	auto GetIntersectionsFunc = [&](const Chaos::FPhysicsSolver* InSolver, const Chaos::TParticles<float, 3>& InCollisionParticles, float InDeltaMag, const FTransform& InPose) -> TArray<int32>
 	{
 		SCOPE_CYCLE_COUNTER(STAT_SQSweepBroadPhase)
 
@@ -532,9 +534,9 @@ void FGeometryCollectionSQAccelerator::Sweep(const FPhysicsGeometry& QueryGeom,
 		return;
 	}
 
-	const TArray<Chaos::FPBDRigidsSolver*>& Solvers = FChaosSolversModule::GetModule()->GetSolvers();
+	const TArray<Chaos::FPhysicsSolver*>& Solvers = FChaosSolversModule::GetModule()->GetSolvers();
 
-	for(const Chaos::FPBDRigidsSolver* Solver : Solvers)
+	for(const Chaos::FPhysicsSolver* Solver : Solvers)
 	{		
 		if (!Solver)
 		{
@@ -543,7 +545,7 @@ void FGeometryCollectionSQAccelerator::Sweep(const FPhysicsGeometry& QueryGeom,
 
 		TArray<int32> IntersectionSet = GetIntersectionsFunc(Solver, CollisionParticles, DeltaMag, StartTM);
 
-		const Chaos::FPBDRigidsSolver::FSolverObjectReverseMapping& ObjectMap = Solver->GetSolverObjectReverseMapping_GameThread();
+		const Chaos::FPhysicsSolver::FPhysicsProxyReverseMapping& ObjectMap = Solver->GetPhysicsProxyReverseMapping_GameThread();
 		const Chaos::TClusterBuffer<float, 3>& Buffer = Solver->GetRigidClustering().GetBufferedData();
 
 
@@ -553,9 +555,9 @@ void FGeometryCollectionSQAccelerator::Sweep(const FPhysicsGeometry& QueryGeom,
 			SCOPE_CYCLE_COUNTER(STAT_SQSweepNarrowPhase)
 
 			const int32 ParticleIndex = IntersectionSet[i];
-			const SolverObjectWrapper& ObjectWrapper = ObjectMap.SolverObjectReverseMappingArray[ParticleIndex];
+			const PhysicsProxyWrapper& ObjectWrapper = ObjectMap.PhysicsProxyReverseMappingArray[ParticleIndex];
 
-			if (!ObjectWrapper.SolverObject)
+			if (!ObjectWrapper.PhysicsProxy)
 			{
 				const TImplicitObject<float, 3>* Object = Buffer.GeometryPtrs[ParticleIndex].Get();
 				
@@ -606,14 +608,14 @@ void FGeometryCollectionSQAccelerator::Sweep(const FPhysicsGeometry& QueryGeom,
 				continue;
 			}
 
-			if(ObjectWrapper.Type == ESolverObjectType::GeometryCollectionType && ensure(ObjectWrapper.SolverObject))
+			if(ObjectWrapper.Type == EPhysicsProxyType::GeometryCollectionType && ensure(ObjectWrapper.PhysicsProxy))
 			{
-				if(LowLevelSweepSingleElement(ParticleIndex, Solver, Buffer, static_cast<FGeometryCollectionPhysicsObject*>(ObjectWrapper.SolverObject), *Implicit, CollisionParticles, StartTM, Dir, DeltaMag, i >= IntersectionSetSize, Hit))
+				if(LowLevelSweepSingleElement(ParticleIndex, Solver, Buffer, static_cast<FGeometryCollectionPhysicsProxy*>(ObjectWrapper.PhysicsProxy), *Implicit, CollisionParticles, StartTM, Dir, DeltaMag, i >= IntersectionSetSize, Hit))
 				{
 #if !WITH_IMMEDIATE_PHYSX && PHYSICS_INTERFACE_PHYSX
 					//todo(mlentine): This is duplicated from above and should be merged
 					//todo(ocohen):hack placeholder while we convert over to non physx API
-					UGeometryCollectionComponent* Component = Cast<UGeometryCollectionComponent>(ObjectWrapper.SolverObject->GetOwner());
+					UGeometryCollectionComponent* Component = Cast<UGeometryCollectionComponent>(ObjectWrapper.PhysicsProxy->GetOwner());
 					check(Component);
 
 					if (Component->IsRegistered())
@@ -639,12 +641,12 @@ void FGeometryCollectionSQAccelerator::Sweep(const FPhysicsGeometry& QueryGeom,
 
 		Solver->GetRigidClustering().ReleaseBufferedData();
 
-		Solver->ReleaseSolverObjectReverseMapping();
+		Solver->ReleasePhysicsProxyReverseMapping();
 	}
 #endif
 }
 
-bool LowLevelOverlapSingleElement(int32 InParticleIndex, const Chaos::FPBDRigidsSolver* InSolver, const Chaos::TClusterBuffer<float, 3>& ClusterBuffer, const FGeometryCollectionPhysicsObject* InObject, const Chaos::TImplicitObject<float, 3>& QueryGeom, const FTransform& InPose, FHitOverlap& OutHit)
+bool LowLevelOverlapSingleElement(int32 InParticleIndex, const Chaos::FPhysicsSolver* InSolver, const Chaos::TClusterBuffer<float, 3>& ClusterBuffer, const FGeometryCollectionPhysicsProxy* InObject, const Chaos::TImplicitObject<float, 3>& QueryGeom, const FTransform& InPose, FHitOverlap& OutHit)
 {
 	using namespace Chaos;
 
@@ -678,7 +680,8 @@ bool LowLevelOverlapSingleElement(int32 InParticleIndex, const Chaos::FPBDRigids
 	return Result.Second;
 }
 
-void FGeometryCollectionSQAccelerator::Overlap(const FPhysicsGeometry& QueryGeom, const FTransform& GeomPose, FPhysicsHitCallback<FHitOverlap>& HitBuffer, FQueryFlags QueryFlags, const FCollisionFilterData& QueryFilter, const FQueryFilterData& QueryFilterData, ICollisionQueryFilterCallbackBase& QueryCallback) const
+
+void FGeometryCollectionSQAccelerator::Overlap(const Chaos::TImplicitObject<float, 3>& QueryGeom, const FTransform& GeomPose, ChaosInterface::FSQHitBuffer<ChaosInterface::FOverlapHit>& HitBuffer, const FQueryFilterData& QueryFilterData, ICollisionQueryFilterCallbackBase& QueryCallback) const
 {
 	return;	//todo: This is currently broken because it doesn'thandle cluster unions. Need to fix function - disabling entirely for now
 	SCOPE_CYCLE_COUNTER(STAT_GCOverlap);
@@ -686,10 +689,10 @@ void FGeometryCollectionSQAccelerator::Overlap(const FPhysicsGeometry& QueryGeom
 
 #if WITH_PHYSX
 
-	TMap<const Chaos::FPBDRigidsSolver*, TArray<int32>> SolverIntersectionSets;
+	TMap<const Chaos::FPhysicsSolver*, TArray<int32>> SolverIntersectionSets;
 
 	// Getter for intersections from the mapping above
-	auto GetIntersectionsFunc = [&](const Chaos::FPBDRigidsSolver* InSolver, Chaos::TImplicitObject<float, 3>* InGeometry, const FTransform& InPose) -> TArray<int32>*
+	auto GetIntersectionsFunc = [&](const Chaos::FPhysicsSolver* InSolver, Chaos::TImplicitObject<float, 3>* InGeometry, const FTransform& InPose) -> TArray<int32>*
 	{
 		TArray<int32>* CurrentIntersections = SolverIntersectionSets.Find(InSolver);
 
@@ -758,31 +761,31 @@ void FGeometryCollectionSQAccelerator::Overlap(const FPhysicsGeometry& QueryGeom
 	}
 
 	FChaosSolversModule* Module = FChaosSolversModule::GetModule();
-	const TArray<Chaos::FPBDRigidsSolver*> Solvers = Module->GetSolvers();
+	const TArray<Chaos::FPhysicsSolver*> Solvers = Module->GetSolvers();
 
-	for(const Chaos::FPBDRigidsSolver* Solver : Solvers)
+	for(const Chaos::FPhysicsSolver* Solver : Solvers)
 	{
 		// Collect all intersections for this solver
 		TArray<int32> IntersectionSet = Solver->GetSpatialAcceleration()->FindAllIntersections(Implicit->BoundingBox().TransformedBox(GeomPose));
 		Solver->ReleaseSpatialAcceleration();
 
-		const Chaos::FPBDRigidsSolver::FSolverObjectReverseMapping& ObjectMap = Solver->GetSolverObjectReverseMapping_GameThread();
+		const Chaos::FPhysicsSolver::FPhysicsProxyReverseMapping& ObjectMap = Solver->GetPhysicsProxyReverseMapping_GameThread();
 
 		const Chaos::TClusterBuffer<float, 3>& ClusterBuffer = Solver->GetRigidClustering().GetBufferedData();
 
 		for(int32 i = 0; i < IntersectionSet.Num(); ++i)
 		{
 			const int32 IntersectParticleIndex = IntersectionSet[i];
-			const SolverObjectWrapper& ObjectWrapper = ObjectMap.SolverObjectReverseMappingArray[IntersectParticleIndex];
+			const PhysicsProxyWrapper& ObjectWrapper = ObjectMap.PhysicsProxyReverseMappingArray[IntersectParticleIndex];
 
-			if(ObjectWrapper.Type == ESolverObjectType::GeometryCollectionType && ensure(ObjectWrapper.SolverObject))
+			if(ObjectWrapper.Type == EPhysicsProxyType::GeometryCollectionType && ensure(ObjectWrapper.PhysicsProxy))
 			{
-				if(LowLevelOverlapSingleElement(IntersectParticleIndex, Solver, ClusterBuffer, static_cast<FGeometryCollectionPhysicsObject*>(ObjectWrapper.SolverObject), *Implicit, GeomPose, Hit))
+				if(LowLevelOverlapSingleElement(IntersectParticleIndex, Solver, ClusterBuffer, static_cast<FGeometryCollectionPhysicsProxy*>(ObjectWrapper.PhysicsProxy), *Implicit, GeomPose, Hit))
 				{
 #if !WITH_IMMEDIATE_PHYSX && PHYSICS_INTERFACE_PHYSX
 					//todo(mlentine): This is duplicated from above and should be merged
 					//todo(ocohen):hack placeholder while we convert over to non physx API
-					UGeometryCollectionComponent* Component = Cast<UGeometryCollectionComponent>(ObjectWrapper.SolverObject->GetOwner());
+					UGeometryCollectionComponent* Component = Cast<UGeometryCollectionComponent>(ObjectWrapper.PhysicsProxy->GetOwner());
 					check(Component);
 
 					if (Component->IsRegistered())
@@ -807,12 +810,14 @@ void FGeometryCollectionSQAccelerator::Overlap(const FPhysicsGeometry& QueryGeom
 		}
 
 		Solver->GetRigidClustering().ReleaseBufferedData();
-		Solver->ReleaseSolverObjectReverseMapping();
+		Solver->ReleasePhysicsProxyReverseMapping();
 	}
 
 #endif
 }
 
+#endif // TODO_REIMPLEMENT_SCENEQUERY_CROSSENGINE
+
 void FGeometryCollectionSQAccelerator::AddComponent(UGeometryCollectionComponent* Component)
 {
 	if (ensure(Component))
diff --git a/Engine/Source/Runtime/Experimental/GeometryCollectionEngine/Private/GeometryCollection/GeometryCollectionSQAccelerator.h b/Engine/Source/Runtime/Experimental/GeometryCollectionEngine/Private/GeometryCollection/GeometryCollectionSQAccelerator.h
--- a/Engine/Source/Runtime/Experimental/GeometryCollectionEngine/Private/GeometryCollection/GeometryCollectionSQAccelerator.h
+++ b/Engine/Source/Runtime/Experimental/GeometryCollectionEngine/Private/GeometryCollection/GeometryCollectionSQAccelerator.h
@@ -7,12 +7,12 @@
 
 class UGeometryCollectionComponent;
 
-class FGeometryCollectionSQAccelerator : public ISQAccelerator
+class FGeometryCollectionSQAccelerator
 {
 public:
-	virtual void Raycast(const FVector& Start, const FVector& Dir, const float DeltaMagnitude, FPhysicsHitCallback<FHitRaycast>& HitBuffer, EHitFlags OutputFlags, FQueryFlags QueryFlags, const FCollisionFilterData& QueryFilter, const FQueryFilterData& QueryFilterData, ICollisionQueryFilterCallbackBase& QueryCallback) const override;
-	virtual void Sweep(const FPhysicsGeometry& QueryGeom, const FTransform& StartTM, const FVector& Dir, const float DeltaMagnitude, FPhysicsHitCallback<FHitSweep>& HitBuffer, EHitFlags OutputFlags, FQueryFlags QueryFlags, const FCollisionFilterData& QueryFilter, const FQueryFilterData& QueryFilterData, ICollisionQueryFilterCallbackBase& QueryCallback) const override;
-	virtual void Overlap(const FPhysicsGeometry& QueryGeom, const FTransform& GeomPose, FPhysicsHitCallback<FHitOverlap>& HitBuffer, FQueryFlags QueryFlags, const FCollisionFilterData& QueryFilter, const FQueryFilterData& QueryFilterData, ICollisionQueryFilterCallbackBase& QueryCallback) const override;
+	void Raycast(const FVector& Start, const FVector& Dir, const float DeltaMagnitude, ChaosInterface::FSQHitBuffer<ChaosInterface::FRaycastHit>& HitBuffer, EHitFlags OutputFlags, const FQueryFilterData& QueryFilterData, ICollisionQueryFilterCallbackBase& QueryCallback) const;
+	void Sweep(const Chaos::TImplicitObject<float, 3>& QueryGeom, const FTransform& StartTM, const FVector& Dir, const float DeltaMagnitude, ChaosInterface::FSQHitBuffer<ChaosInterface::FSweepHit>& HitBuffer, EHitFlags OutputFlags, const FQueryFilterData& QueryFilterData, ICollisionQueryFilterCallbackBase& QueryCallback) const;
+	void Overlap(const Chaos::TImplicitObject<float, 3>& QueryGeom, const FTransform& GeomPose, ChaosInterface::FSQHitBuffer<ChaosInterface::FOverlapHit>& HitBuffer, const FQueryFilterData& QueryFilterData, ICollisionQueryFilterCallbackBase& QueryCallback) const;
 	virtual ~FGeometryCollectionSQAccelerator() {}
 
 	void AddComponent(UGeometryCollectionComponent* Component);
