Fixed HashCombine usage in implicit physics intersections

diff --git a/Engine/Source/Runtime/Experimental/Chaos/Public/Chaos/ImplicitObjectIntersection.h b/Engine/Source/Runtime/Experimental/Chaos/Public/Chaos/ImplicitObjectIntersection.h
--- a/Engine/Source/Runtime/Experimental/Chaos/Public/Chaos/ImplicitObjectIntersection.h
+++ b/Engine/Source/Runtime/Experimental/Chaos/Public/Chaos/ImplicitObjectIntersection.h
@@ -64,7 +64,7 @@ class TImplicitObjectIntersection : public TImplicitObject<T, d>
 
 		for(const TUniquePtr<TImplicitObject<T, d>>& Ptr : MObjects)
 		{
-			HashCombine(Ptr->GetTypeHash(), OutHash);
+			OutHash = HashCombine(Ptr->GetTypeHash(), OutHash);
 		}
 
 		return OutHash;
