Moved the geometry collection debug draw component properties to the debug draw actor to solve the slow down issue, and resolve a few other related problems.

diff --git a/Engine/Source/Runtime/Experimental/GeometryCollectionEngine/Private/GeometryCollection/GeometryCollectionDebugDrawActor.cpp b/Engine/Source/Runtime/Experimental/GeometryCollectionEngine/Private/GeometryCollection/GeometryCollectionDebugDrawActor.cpp
--- a/Engine/Source/Runtime/Experimental/GeometryCollectionEngine/Private/GeometryCollection/GeometryCollectionDebugDrawActor.cpp
+++ b/Engine/Source/Runtime/Experimental/GeometryCollectionEngine/Private/GeometryCollection/GeometryCollectionDebugDrawActor.cpp
@@ -1,6 +1,7 @@
 // Copyright 1998-2019 Epic Games, Inc. All Rights Reserved.
 
 #include "GeometryCollection/GeometryCollectionDebugDrawActor.h"
+
 #include "GeometryCollection/GeometryCollection.h"
 #include "GeometryCollection/GeometryCollectionParticlesData.h"
 #include "GeometryCollection/GeometryCollectionComponent.h"
@@ -24,9 +25,10 @@
 
 DEFINE_LOG_CATEGORY_STATIC(LogGeometryCollectionDebugDrawActor, Log, All);
 
+// Constants
 namespace GeometryCollectionDebugDrawActorConstants
 {
-	// Constants
+	// Invariables
 	static const bool bPersistent = true;  // Debug draw needs persistency to work well within the editor.
 	static const float LifeTime = -1.0f;  // Lifetime is infinite.
 	static const uint8 DepthPriority = 0;
@@ -39,6 +41,15 @@ namespace GeometryCollectionDebugDrawActorConstants
 		FIntPoint(0, 4), FIntPoint(1, 5), FIntPoint(2, 6), FIntPoint(3, 7)
 	});
 
+	// Base colors
+	static const FLinearColor DarkerTintFactor(1.0f, 1.0f, 0.7f);  // Darker HSV multiplier
+	static const FLinearColor LighterTintFactor(1.0f, 1.0f, 2.0f);  // Lighter HSV multiplier
+	static const FLinearColor RigidBodyTint(0.8f, 0.1f, 0.1f);  // Red
+	static const FLinearColor ClusteringTint(0.6, 0.4f, 0.2f);  // Orange
+	static const FLinearColor GeometryTint(0.4f, 0.2f, 0.6f);  // Purple
+	static const FLinearColor SingleFaceTint(0.6, 0.2f, 0.4f);  // Pink
+	static const FLinearColor VertexTint(0.2f, 0.4f, 0.6f);  // Blue
+
 	// Defaults
 	static const FString SelectedRigidBodySolverDefault = FName(NAME_None).ToString();
 	static const int32 SelectedRigidBodyIdDefault = INDEX_NONE;
@@ -76,8 +87,27 @@ namespace GeometryCollectionDebugDrawActorConstants
 	static const int32 TextShadowDefault = 1;
 	static const float TextScaleDefault = 1.0f;
 	static const float NormalScaleDefault = 10.0f;
-	static const float TransformScaleDefault = 20.0f;
+	static const float AxisScaleDefault = 20.0f;
 	static const float ArrowScaleDefault = 2.5f;
+	static const float TransformScaleDefault = 1.f;
+	static const FColor RigidBodyIdsColorDefault = (RigidBodyTint.LinearRGBToHSV() * LighterTintFactor).HSVToLinearRGB().ToFColor(true);
+	static const FColor RigidBodyCollisionColorDefault = RigidBodyTint.ToFColor(true);
+	static const FColor RigidBodyInertiaColorDefault = (RigidBodyTint.LinearRGBToHSV() * LighterTintFactor).HSVToLinearRGB().ToFColor(true);
+	static const FColor RigidBodyVelocityColorDefault = (RigidBodyTint.LinearRGBToHSV() * DarkerTintFactor).HSVToLinearRGB().ToFColor(true);
+	static const FColor RigidBodyForceColorDefault = (RigidBodyTint.LinearRGBToHSV() * DarkerTintFactor).HSVToLinearRGB().ToFColor(true);
+	static const FColor RigidBodyInfoColorDefault = (RigidBodyTint.LinearRGBToHSV() * LighterTintFactor).HSVToLinearRGB().ToFColor(true);
+	static const FColor TransformIndexColorDefault = (ClusteringTint.LinearRGBToHSV() * LighterTintFactor).HSVToLinearRGB().ToFColor(true);
+	static const FColor LevelColorDefault = (ClusteringTint.LinearRGBToHSV() * LighterTintFactor).HSVToLinearRGB().ToFColor(true);
+	static const FColor ParentColorDefault = ClusteringTint.ToFColor(true);
+	static const FColor GeometryIndexColorDefault = (GeometryTint.LinearRGBToHSV() * LighterTintFactor).HSVToLinearRGB().ToFColor(true);
+	static const FColor BoundingBoxColorDefault = (GeometryTint.LinearRGBToHSV() * DarkerTintFactor).HSVToLinearRGB().ToFColor(true);
+	static const FColor FaceColorDefault = GeometryTint.ToFColor(true);
+	static const FColor FaceIndexColorDefault = (GeometryTint.LinearRGBToHSV() * LighterTintFactor).HSVToLinearRGB().ToFColor(true);
+	static const FColor FaceNormalColorDefault = (GeometryTint.LinearRGBToHSV() * DarkerTintFactor).HSVToLinearRGB().ToFColor(true);
+	static const FColor SingleFaceColorDefault = (SingleFaceTint.LinearRGBToHSV() * LighterTintFactor).HSVToLinearRGB().ToFColor(true);
+	static const FColor VertexColorDefault = VertexTint.ToFColor(true);
+	static const FColor VertexIndexColorDefault = (VertexTint.LinearRGBToHSV() * LighterTintFactor).HSVToLinearRGB().ToFColor(true);
+	static const FColor VertexNormalColorDefault = (VertexTint.LinearRGBToHSV() * DarkerTintFactor).HSVToLinearRGB().ToFColor(true);
 }
 
 namespace GeometryCollectionDebugDrawActorCVars
@@ -119,7 +149,7 @@ namespace GeometryCollectionDebugDrawActorCVars
 	static TAutoConsoleVariable<int32  > TextShadow              (TEXT("p.gc.TextShadow"              ), GeometryCollectionDebugDrawActorConstants::TextShadowDefault              , TEXT("Geometry Collection debug draw, text shadow under indices for better readability.\nDefault = 1."), ECVF_Cheat);
 	static TAutoConsoleVariable<float  > TextScale               (TEXT("p.gc.TextScale"               ), GeometryCollectionDebugDrawActorConstants::TextScaleDefault               , TEXT("Geometry Collection debug draw, text scale.\nDefault = 1."), ECVF_Cheat);
 	static TAutoConsoleVariable<float  > NormalScale             (TEXT("p.gc.NormalScale"             ), GeometryCollectionDebugDrawActorConstants::NormalScaleDefault             , TEXT("Geometry Collection debug draw, normal size.\nDefault = 10."), ECVF_Cheat);
-	static TAutoConsoleVariable<float  > TransformScale          (TEXT("p.gc.TransformScale"          ), GeometryCollectionDebugDrawActorConstants::TransformScaleDefault          , TEXT("Geometry Collection debug draw, transform size.\nDefault = 20."), ECVF_Cheat);
+	static TAutoConsoleVariable<float  > AxisScale               (TEXT("p.gc.AxisScale"               ), GeometryCollectionDebugDrawActorConstants::AxisScaleDefault               , TEXT("Geometry Collection debug draw, size of the axis used for visualizing all transforms.\nDefault = 20."), ECVF_Cheat);
 	static TAutoConsoleVariable<float  > ArrowScale              (TEXT("p.gc.ArrowScale"              ), GeometryCollectionDebugDrawActorConstants::ArrowScaleDefault              , TEXT("Geometry Collection debug draw, arrow size for normals.\nDefault = 2.5."), ECVF_Cheat);
 }
 
@@ -192,8 +222,30 @@ AGeometryCollectionDebugDrawActor::AGeometryCollectionDebugDrawActor(const FObje
 	, bTextShadow              (!!GeometryCollectionDebugDrawActorConstants::TextShadowDefault              )
 	, TextScale                (  GeometryCollectionDebugDrawActorConstants::TextScaleDefault               )
 	, NormalScale              (  GeometryCollectionDebugDrawActorConstants::NormalScaleDefault             )
-	, TransformScale           (  GeometryCollectionDebugDrawActorConstants::TransformScaleDefault          )
+	, AxisScale                (  GeometryCollectionDebugDrawActorConstants::AxisScaleDefault               )
 	, ArrowScale               (  GeometryCollectionDebugDrawActorConstants::ArrowScaleDefault              )
+	, RigidBodyIdColor         (  GeometryCollectionDebugDrawActorConstants::RigidBodyIdsColorDefault       )
+	, RigidBodyTransformScale  (  GeometryCollectionDebugDrawActorConstants::TransformScaleDefault          )
+	, RigidBodyCollisionColor  (  GeometryCollectionDebugDrawActorConstants::RigidBodyCollisionColorDefault )
+	, RigidBodyInertiaColor    (  GeometryCollectionDebugDrawActorConstants::RigidBodyInertiaColorDefault   )
+	, RigidBodyVelocityColor   (  GeometryCollectionDebugDrawActorConstants::RigidBodyVelocityColorDefault  )
+	, RigidBodyForceColor      (  GeometryCollectionDebugDrawActorConstants::RigidBodyForceColorDefault     )
+	, RigidBodyInfoColor       (  GeometryCollectionDebugDrawActorConstants::RigidBodyInfoColorDefault      )
+	, TransformIndexColor      (  GeometryCollectionDebugDrawActorConstants::TransformIndexColorDefault     )
+	, TransformScale           (  GeometryCollectionDebugDrawActorConstants::TransformScaleDefault          )
+	, LevelColor               (  GeometryCollectionDebugDrawActorConstants::LevelColorDefault              )
+	, ParentColor              (  GeometryCollectionDebugDrawActorConstants::ParentColorDefault             )
+	, ConnectivityEdgeThickness(  GeometryCollectionDebugDrawActorConstants::LineThicknessDefault           )
+	, GeometryIndexColor       (  GeometryCollectionDebugDrawActorConstants::GeometryIndexColorDefault      )
+	, GeometryTransformScale   (  GeometryCollectionDebugDrawActorConstants::TransformScaleDefault          )
+	, BoundingBoxColor         (  GeometryCollectionDebugDrawActorConstants::BoundingBoxColorDefault        )
+	, FaceColor                (  GeometryCollectionDebugDrawActorConstants::FaceColorDefault               )
+	, FaceIndexColor           (  GeometryCollectionDebugDrawActorConstants::FaceIndexColorDefault          )
+	, FaceNormalColor          (  GeometryCollectionDebugDrawActorConstants::FaceNormalColorDefault         )
+	, SingleFaceColor          (  GeometryCollectionDebugDrawActorConstants::SingleFaceColorDefault         )
+	, VertexColor              (  GeometryCollectionDebugDrawActorConstants::VertexColorDefault             )
+	, VertexIndexColor         (  GeometryCollectionDebugDrawActorConstants::VertexIndexColorDefault        )
+	, VertexNormalColor        (  GeometryCollectionDebugDrawActorConstants::VertexNormalColorDefault       )
 	, SpriteComponent()
 	, ConsoleVariableSinkHandle()
 	, DebugDrawTextDelegateHandle()
@@ -383,7 +435,7 @@ void AGeometryCollectionDebugDrawActor::PostLoad()
 	GeometryCollectionDebugDrawActorCVars::TextShadow              ->Set(int32(bTextShadow                ), SetBy);
 	GeometryCollectionDebugDrawActorCVars::TextScale               ->Set(      TextScale                   , SetBy);
 	GeometryCollectionDebugDrawActorCVars::NormalScale             ->Set(      NormalScale                 , SetBy);
-	GeometryCollectionDebugDrawActorCVars::TransformScale          ->Set(      TransformScale              , SetBy);
+	GeometryCollectionDebugDrawActorCVars::AxisScale               ->Set(      AxisScale                   , SetBy);
 	GeometryCollectionDebugDrawActorCVars::ArrowScale              ->Set(      ArrowScale                  , SetBy);
 
 	Super::PostLoad();
@@ -515,7 +567,7 @@ void AGeometryCollectionDebugDrawActor::OnCVarsChanged()
 	UpdatePropertyValue(bTextShadow              , GeometryCollectionDebugDrawActorCVars::TextShadow              , bHavePropertiesChanged);
 	UpdatePropertyValue(TextScale                , GeometryCollectionDebugDrawActorCVars::TextScale               , bHavePropertiesChanged);
 	UpdatePropertyValue(NormalScale              , GeometryCollectionDebugDrawActorCVars::NormalScale             , bHavePropertiesChanged);
-	UpdatePropertyValue(TransformScale           , GeometryCollectionDebugDrawActorCVars::TransformScale          , bHavePropertiesChanged);
+	UpdatePropertyValue(AxisScale                , GeometryCollectionDebugDrawActorCVars::AxisScale               , bHavePropertiesChanged);
 	UpdatePropertyValue(ArrowScale               , GeometryCollectionDebugDrawActorCVars::ArrowScale              , bHavePropertiesChanged);
 
 	// Update geometry collection component, but only if this actor has begun play
@@ -566,16 +618,18 @@ void AGeometryCollectionDebugDrawActor::PostEditChangeProperty(FPropertyChangedE
 	else if (PropertyName == GET_MEMBER_NAME_CHECKED(AGeometryCollectionDebugDrawActor, bShowVertices            )) { GeometryCollectionDebugDrawActorCVars::ShowVertices            ->Set(int32(bShowVertices              ), SetBy); }
 	else if (PropertyName == GET_MEMBER_NAME_CHECKED(AGeometryCollectionDebugDrawActor, bShowVertexIndices       )) { GeometryCollectionDebugDrawActorCVars::ShowVertexIndices       ->Set(int32(bShowVertexIndices         ), SetBy); }
 	else if (PropertyName == GET_MEMBER_NAME_CHECKED(AGeometryCollectionDebugDrawActor, bShowVertexNormals       )) { GeometryCollectionDebugDrawActorCVars::ShowVertexNormals       ->Set(int32(bShowVertexNormals         ), SetBy); }
-	else if (PropertyName == GET_MEMBER_NAME_CHECKED(AGeometryCollectionDebugDrawActor, bUseActiveVisualization  )) { GeometryCollectionDebugDrawActorCVars::UseActiveVisualization  ->Set(int32(bUseActiveVisualization    ), SetBy); }
-	else if (PropertyName == GET_MEMBER_NAME_CHECKED(AGeometryCollectionDebugDrawActor, PointThickness           )) { GeometryCollectionDebugDrawActorCVars::PointThickness          ->Set(      PointThickness              , SetBy); }
-	else if (PropertyName == GET_MEMBER_NAME_CHECKED(AGeometryCollectionDebugDrawActor, LineThickness            )) { GeometryCollectionDebugDrawActorCVars::LineThickness           ->Set(      LineThickness               , SetBy); }
-	else if (PropertyName == GET_MEMBER_NAME_CHECKED(AGeometryCollectionDebugDrawActor, bTextShadow              )) { GeometryCollectionDebugDrawActorCVars::TextShadow              ->Set(int32(bTextShadow                ), SetBy); }
-	else if (PropertyName == GET_MEMBER_NAME_CHECKED(AGeometryCollectionDebugDrawActor, TextScale                )) { GeometryCollectionDebugDrawActorCVars::TextScale               ->Set(      TextScale                   , SetBy); }
-	else if (PropertyName == GET_MEMBER_NAME_CHECKED(AGeometryCollectionDebugDrawActor, NormalScale              )) { GeometryCollectionDebugDrawActorCVars::NormalScale             ->Set(      NormalScale                 , SetBy); }
-	else if (PropertyName == GET_MEMBER_NAME_CHECKED(AGeometryCollectionDebugDrawActor, TransformScale           )) { GeometryCollectionDebugDrawActorCVars::TransformScale          ->Set(      TransformScale              , SetBy); }
-	else if (PropertyName == GET_MEMBER_NAME_CHECKED(AGeometryCollectionDebugDrawActor, ArrowScale               )) { GeometryCollectionDebugDrawActorCVars::ArrowScale              ->Set(      ArrowScale                  , SetBy); }
 	else
 	{
+		// These properties are cosmetic changes and don't require visibility updates or enabling the component tick
+		if      (PropertyName == GET_MEMBER_NAME_CHECKED(AGeometryCollectionDebugDrawActor, bUseActiveVisualization  )) { GeometryCollectionDebugDrawActorCVars::UseActiveVisualization  ->Set(int32(bUseActiveVisualization    ), SetBy); }
+		else if (PropertyName == GET_MEMBER_NAME_CHECKED(AGeometryCollectionDebugDrawActor, PointThickness           )) { GeometryCollectionDebugDrawActorCVars::PointThickness          ->Set(      PointThickness              , SetBy); }
+		else if (PropertyName == GET_MEMBER_NAME_CHECKED(AGeometryCollectionDebugDrawActor, LineThickness            )) { GeometryCollectionDebugDrawActorCVars::LineThickness           ->Set(      LineThickness               , SetBy); }
+		else if (PropertyName == GET_MEMBER_NAME_CHECKED(AGeometryCollectionDebugDrawActor, bTextShadow              )) { GeometryCollectionDebugDrawActorCVars::TextShadow              ->Set(int32(bTextShadow                ), SetBy); }
+		else if (PropertyName == GET_MEMBER_NAME_CHECKED(AGeometryCollectionDebugDrawActor, TextScale                )) { GeometryCollectionDebugDrawActorCVars::TextScale               ->Set(      TextScale                   , SetBy); }
+		else if (PropertyName == GET_MEMBER_NAME_CHECKED(AGeometryCollectionDebugDrawActor, NormalScale              )) { GeometryCollectionDebugDrawActorCVars::NormalScale             ->Set(      NormalScale                 , SetBy); }
+		else if (PropertyName == GET_MEMBER_NAME_CHECKED(AGeometryCollectionDebugDrawActor, AxisScale                )) { GeometryCollectionDebugDrawActorCVars::AxisScale               ->Set(      AxisScale                   , SetBy); }
+		else if (PropertyName == GET_MEMBER_NAME_CHECKED(AGeometryCollectionDebugDrawActor, ArrowScale               )) { GeometryCollectionDebugDrawActorCVars::ArrowScale              ->Set(      ArrowScale                  , SetBy); }
+
 		Super::PostEditChangeProperty(PropertyChangedEvent);
 		return;  // Don't call OnPropertyChanged()
 	}
@@ -704,7 +758,7 @@ void AGeometryCollectionDebugDrawActor::DrawVertices(const TArray<FTransform>& G
 #if ENABLE_DRAW_DEBUG
 	check(GeometryCollectionComponent);
 	check(GeometryCollectionComponent->RestCollection);
-	check(TransformIndex >= 0);
+	check(TransformIndex >= INDEX_NONE);
 
 	const UWorld* const World = GetWorld();
 	check(World);
@@ -1433,7 +1487,7 @@ void AGeometryCollectionDebugDrawActor::DrawGeometryIndex(const TArray<FTransfor
 void AGeometryCollectionDebugDrawActor::DrawTransforms(const TArray<FTransform>& GlobalTransforms, const UGeometryCollectionComponent* GeometryCollectionComponent, float Scale)
 {
 #if ENABLE_DRAW_DEBUG
-	Scale *= TransformScale;
+	Scale *= AxisScale;
 
 	check(GeometryCollectionComponent);
 	check(GeometryCollectionComponent->RestCollection);
@@ -1500,7 +1554,7 @@ void AGeometryCollectionDebugDrawActor::DrawTransform(const TArray<FTransform>&
 		const FVector Position = Transform.GetLocation();
 		const FRotator Rotator = Transform.Rotator();
 
-		DrawDebugCoordinateSystem(World, Position, Rotator, Scale * TransformScale, GeometryCollectionDebugDrawActorConstants::bPersistent, GeometryCollectionDebugDrawActorConstants::LifeTime, GeometryCollectionDebugDrawActorConstants::DepthPriority, LineThickness);
+		DrawDebugCoordinateSystem(World, Position, Rotator, Scale * AxisScale, GeometryCollectionDebugDrawActorConstants::bPersistent, GeometryCollectionDebugDrawActorConstants::LifeTime, GeometryCollectionDebugDrawActorConstants::DepthPriority, LineThickness);
 		bNeedsDebugLinesFlush = true;
 	}
 
@@ -2109,7 +2163,7 @@ void AGeometryCollectionDebugDrawActor::DrawRigidBodyTransform(const UGeometryCo
 void AGeometryCollectionDebugDrawActor::DrawRigidBodyTransformNoChecks(const UGeometryCollectionComponent* GeometryCollectionComponent, int32 TransformIndex, const FGeometryCollectionParticlesData& ParticlesData, float Scale)
 {
 #if ENABLE_DRAW_DEBUG
-	Scale *= TransformScale;
+	Scale *= AxisScale;
 
 	// Retrieve particle transform
 	const FTransform Transform = GetParticleTransformNoChecks(GeometryCollectionComponent, TransformIndex, ParticlesData);
@@ -2507,7 +2561,7 @@ void AGeometryCollectionDebugDrawActor::DrawRigidBodyInfoNoChecks(const UGeometr
 #endif  // #if ENABLE_DRAW_DEBUG
 }
 
-void AGeometryCollectionDebugDrawActor::DrawConnectivityEdges(const UGeometryCollectionComponent* GeometryCollectionComponent, const FGeometryCollectionParticlesData& ParticlesData, const TManagedArray<int32>& RigidBodyIdArray, float Thickness)
+void AGeometryCollectionDebugDrawActor::DrawConnectivityEdges(const UGeometryCollectionComponent* GeometryCollectionComponent, const FGeometryCollectionParticlesData& ParticlesData, const TManagedArray<int32>& RigidBodyIdArray)
 {
 #if ENABLE_DRAW_DEBUG
 	check(GeometryCollectionComponent);
@@ -2535,14 +2589,14 @@ void AGeometryCollectionDebugDrawActor::DrawConnectivityEdges(const UGeometryCol
 				Hue = (Hue + 157) % 256;  // 157 is a prime number that gives a good spread of colors without getting too similar as a rand might do.
 				const FColor RandomColor = FLinearColor::MakeFromHSV8(Hue, 160, 128).ToFColor(true);
 				const FColor ActiveColor = bUseActiveVisualization ? MakeDarker(RandomColor, GetLevel(TransformIndex, ParentArray)) : RandomColor;
-				DrawConnectivityEdgesNoChecks(GeometryCollectionComponent, TransformIndex, ParticlesData, RigidBodyIdArray, Thickness, ActiveColor);
+				DrawConnectivityEdgesNoChecks(GeometryCollectionComponent, TransformIndex, ParticlesData, RigidBodyIdArray, ActiveColor);
 			}
 		}
 	}
 #endif  // #if ENABLE_DRAW_DEBUG
 }
 
-void AGeometryCollectionDebugDrawActor::DrawConnectivityEdges(const UGeometryCollectionComponent* GeometryCollectionComponent, int32 TransformIndex, const FGeometryCollectionParticlesData& ParticlesData, const TManagedArray<int32>& RigidBodyIdArray, float Thickness, FColor HSVColor)
+void AGeometryCollectionDebugDrawActor::DrawConnectivityEdges(const UGeometryCollectionComponent* GeometryCollectionComponent, int32 TransformIndex, const FGeometryCollectionParticlesData& ParticlesData, const TManagedArray<int32>& RigidBodyIdArray, FColor HSVColor)
 {
 #if ENABLE_DRAW_DEBUG
 	check(GeometryCollectionComponent);
@@ -2564,7 +2618,7 @@ void AGeometryCollectionDebugDrawActor::DrawConnectivityEdges(const UGeometryCol
 		if (bHasParent)
 		{
 			const FLinearColor LinearColor = FLinearColor::MakeFromHSV8(HSVColor.R, HSVColor.G, HSVColor.B);  // HSV stored as RGB values
-			DrawConnectivityEdgesNoChecks(GeometryCollectionComponent, TransformIndex, ParticlesData, RigidBodyIdArray, Thickness, LinearColor.ToFColor(true));
+			DrawConnectivityEdgesNoChecks(GeometryCollectionComponent, TransformIndex, ParticlesData, RigidBodyIdArray, LinearColor.ToFColor(true));
 		}
 
 		// Debug draw children if the cluster mode is on
@@ -2575,14 +2629,14 @@ void AGeometryCollectionDebugDrawActor::DrawConnectivityEdges(const UGeometryCol
 			for (int32 ChildTransformIndex : ChildrenArrayRest[TransformIndex])
 			{
 				HSVColor.R = (HSVColor.R + 157) % 256;   // HSV stored as RGB values, this moves to the next "random" hue
-				DrawConnectivityEdges(GeometryCollectionComponent, ChildTransformIndex, ParticlesData, RigidBodyIdArray, Thickness, HSVColor);
+				DrawConnectivityEdges(GeometryCollectionComponent, ChildTransformIndex, ParticlesData, RigidBodyIdArray, HSVColor);
 			}
 		}
 	}
 #endif  // #if ENABLE_DRAW_DEBUG
 }
 
-void AGeometryCollectionDebugDrawActor::DrawConnectivityEdgesNoChecks(const UGeometryCollectionComponent* GeometryCollectionComponent, int32 TransformIndex, const FGeometryCollectionParticlesData& ParticlesData, const TManagedArray<int32>& RigidBodyIdArray, float Thickness, const FColor& Color)
+void AGeometryCollectionDebugDrawActor::DrawConnectivityEdgesNoChecks(const UGeometryCollectionComponent* GeometryCollectionComponent, int32 TransformIndex, const FGeometryCollectionParticlesData& ParticlesData, const TManagedArray<int32>& RigidBodyIdArray, const FColor& Color)
 {
 #if ENABLE_DRAW_DEBUG
 	const UWorld* const World = GetWorld();
@@ -2609,8 +2663,8 @@ void AGeometryCollectionDebugDrawActor::DrawConnectivityEdgesNoChecks(const UGeo
 	// Retrieve connectivity edges information
 	const TArray<Chaos::TConnectivityEdge<float>>& ConnectivityEdges = ParticlesData.GetConnectivityEdges(TransformIndex);
 
-	// Line thickness
-	Thickness *= LineThickness;
+	// Edge thickness
+	const float Thickness = ConnectivityEdgeThickness * LineThickness;
 
 	// Draw connectivity information
 	for (const auto& ConnectivityEdge: ConnectivityEdges)
diff --git a/Engine/Source/Runtime/Experimental/GeometryCollectionEngine/Private/GeometryCollection/GeometryCollectionDebugDrawComponent.cpp b/Engine/Source/Runtime/Experimental/GeometryCollectionEngine/Private/GeometryCollection/GeometryCollectionDebugDrawComponent.cpp
--- a/Engine/Source/Runtime/Experimental/GeometryCollectionEngine/Private/GeometryCollection/GeometryCollectionDebugDrawComponent.cpp
+++ b/Engine/Source/Runtime/Experimental/GeometryCollectionEngine/Private/GeometryCollection/GeometryCollectionDebugDrawComponent.cpp
@@ -16,89 +16,12 @@
 
 DEFINE_LOG_CATEGORY_STATIC(LogGeometryCollectionDebugDraw, Log, All);
 
-// Constants
-namespace GeometryCollectionDebugDrawComponentConstants
-{
-	static const FLinearColor DarkerTintFactor(1.0f, 1.0f, 0.7f);  // Darker HSV multiplier
-	static const FLinearColor LighterTintFactor(1.0f, 1.0f, 2.0f);  // Lighter HSV multiplier
-
-	static const FLinearColor RigidBodyTint(0.8f, 0.1f, 0.1f);  // Red
-	static const FLinearColor ClusteringTint(0.6, 0.4f, 0.2f);  // Orange
-	static const FLinearColor GeometryTint(0.4f, 0.2f, 0.6f);  // Purple
-	static const FLinearColor SingleFaceTint(0.6, 0.2f, 0.4f);  // Pink
-	static const FLinearColor VertexTint(0.2f, 0.4f, 0.6f);  // Blue
-
-	static const FColor RigidBodyIdsColorDefault = (RigidBodyTint.LinearRGBToHSV() * LighterTintFactor).HSVToLinearRGB().ToFColor(true);
-	static const FColor RigidBodyCollisionColorDefault = RigidBodyTint.ToFColor(true);
-	static const FColor RigidBodyInertiaColorDefault = (RigidBodyTint.LinearRGBToHSV() * LighterTintFactor).HSVToLinearRGB().ToFColor(true);
-	static const FColor RigidBodyVelocityColorDefault = (RigidBodyTint.LinearRGBToHSV() * DarkerTintFactor).HSVToLinearRGB().ToFColor(true);
-	static const FColor RigidBodyForceColorDefault = (RigidBodyTint.LinearRGBToHSV() * DarkerTintFactor).HSVToLinearRGB().ToFColor(true);
-	static const FColor RigidBodyInfoColorDefault = (RigidBodyTint.LinearRGBToHSV() * LighterTintFactor).HSVToLinearRGB().ToFColor(true);
-	static const FColor TransformIndexColorDefault = (ClusteringTint.LinearRGBToHSV() * LighterTintFactor).HSVToLinearRGB().ToFColor(true);
-	static const FColor LevelColorDefault = (ClusteringTint.LinearRGBToHSV() * LighterTintFactor).HSVToLinearRGB().ToFColor(true);
-	static const FColor ParentColorDefault = ClusteringTint.ToFColor(true);
-	static const FColor GeometryIndexColorDefault = (GeometryTint.LinearRGBToHSV() * LighterTintFactor).HSVToLinearRGB().ToFColor(true);
-	static const FColor BoundingBoxColorDefault = (GeometryTint.LinearRGBToHSV() * DarkerTintFactor).HSVToLinearRGB().ToFColor(true);
-	static const FColor FaceColorDefault = GeometryTint.ToFColor(true);
-	static const FColor FaceIndexColorDefault = (GeometryTint.LinearRGBToHSV() * LighterTintFactor).HSVToLinearRGB().ToFColor(true);
-	static const FColor FaceNormalColorDefault = (GeometryTint.LinearRGBToHSV() * DarkerTintFactor).HSVToLinearRGB().ToFColor(true);
-	static const FColor SingleFaceColorDefault = (SingleFaceTint.LinearRGBToHSV() * LighterTintFactor).HSVToLinearRGB().ToFColor(true);
-	static const FColor VertexColorDefault = VertexTint.ToFColor(true);
-	static const FColor VertexIndexColorDefault = (VertexTint.LinearRGBToHSV() * LighterTintFactor).HSVToLinearRGB().ToFColor(true);
-	static const FColor VertexNormalColorDefault = (VertexTint.LinearRGBToHSV() * DarkerTintFactor).HSVToLinearRGB().ToFColor(true);
-}
-
 // Static member variables
 UGeometryCollectionDebugDrawComponent* UGeometryCollectionDebugDrawComponent::RenderLevelSetOwner = nullptr;
 int32 UGeometryCollectionDebugDrawComponent::LastRenderedId = INDEX_NONE;
 
 UGeometryCollectionDebugDrawComponent::UGeometryCollectionDebugDrawComponent(const FObjectInitializer& ObjectInitializer)
 	: Super(ObjectInitializer)
-	, bShowRigidBodyIds(false)
-	, bShowRigidBodyTransforms(false)
-	, bShowRigidBodyCollisions(false)
-	, bShowRigidBodyInertias(false)
-	, bShowRigidBodyVelocities(false)
-	, bShowRigidBodyForces(false)
-	, bShowRigidBodyInfos(false)
-	, RigidBodyIdColor(GeometryCollectionDebugDrawComponentConstants::RigidBodyIdsColorDefault)
-	, RigidBodyTransformScale(1.0f)
-	, RigidBodyCollisionColor(GeometryCollectionDebugDrawComponentConstants::RigidBodyCollisionColorDefault)
-	, RigidBodyInertiaColor(GeometryCollectionDebugDrawComponentConstants::RigidBodyInertiaColorDefault)
-	, RigidBodyVelocityColor(GeometryCollectionDebugDrawComponentConstants::RigidBodyVelocityColorDefault)
-	, RigidBodyForceColor(GeometryCollectionDebugDrawComponentConstants::RigidBodyForceColorDefault)
-	, RigidBodyInfoColor(GeometryCollectionDebugDrawComponentConstants::RigidBodyInfoColorDefault)
-	, bShowTransformIndices(false)
-	, bShowTransforms(false)
-	, bShowLevels(false)
-	, bShowParents(false)
-	, bShowConnectivityEdges(false)
-	, TransformIndexColor(GeometryCollectionDebugDrawComponentConstants::TransformIndexColorDefault)
-	, TransformScale(1.0f)
-	, LevelColor(GeometryCollectionDebugDrawComponentConstants::LevelColorDefault)
-	, ParentColor(GeometryCollectionDebugDrawComponentConstants::ParentColorDefault)
-	, ConnectivityEdgeThickness(1.0f)
-	, bShowGeometryIndices(false)
-	, bShowGeometryTransforms(false)
-	, bShowBoundingBoxes(false)
-	, bShowFaces(false)
-	, bShowFaceIndices(false)
-	, bShowFaceNormals(false)
-	, bShowSingleFace(false)
-	, SingleFaceIndex(0)
-	, bShowVertices(false)
-	, bShowVertexIndices(false)
-	, bShowVertexNormals(false)
-	, GeometryIndexColor(GeometryCollectionDebugDrawComponentConstants::GeometryIndexColorDefault)
-	, GeometryTransformScale(1.0f)
-	, BoundingBoxColor(GeometryCollectionDebugDrawComponentConstants::BoundingBoxColorDefault)
-	, FaceColor(GeometryCollectionDebugDrawComponentConstants::FaceColorDefault)
-	, FaceIndexColor(GeometryCollectionDebugDrawComponentConstants::FaceIndexColorDefault)
-	, FaceNormalColor(GeometryCollectionDebugDrawComponentConstants::FaceNormalColorDefault)
-	, SingleFaceColor(GeometryCollectionDebugDrawComponentConstants::SingleFaceColorDefault)
-	, VertexColor(GeometryCollectionDebugDrawComponentConstants::VertexColorDefault)
-	, VertexIndexColor(GeometryCollectionDebugDrawComponentConstants::VertexIndexColorDefault)
-	, VertexNormalColor(GeometryCollectionDebugDrawComponentConstants::VertexNormalColorDefault)
 	, GeometryCollectionDebugDrawActor(nullptr)
 	, GeometryCollectionRenderLevelSetActor(nullptr)
 	, GeometryCollectionComponent(nullptr)
@@ -234,37 +157,6 @@ void UGeometryCollectionDebugDrawComponent::TickComponent(float DeltaTime, enum
 #endif  // #if GEOMETRYCOLLECTION_DEBUG_DRAW
 }
 
-#if WITH_EDITOR
-void UGeometryCollectionDebugDrawComponent::PostEditChangeProperty(FPropertyChangedEvent& PropertyChangedEvent)
-{
-#if GEOMETRYCOLLECTION_DEBUG_DRAW
-	if (GeometryCollectionComponent)
-	{
-		const FName PropertyName = PropertyChangedEvent.Property ? PropertyChangedEvent.Property->GetFName(): NAME_None;
-		if (PropertyName == GET_MEMBER_NAME_CHECKED(UGeometryCollectionDebugDrawComponent, SingleFaceIndex))
-		{
-			const int32 NumFaces = GeometryCollectionComponent->GetNumElements(FGeometryCollection::FacesGroup);
-			SingleFaceIndex = FMath::Clamp(SingleFaceIndex, 0, NumFaces - 1);
-		}
-	}
-
-	// Update selection and visibility
-	OnDebugDrawPropertiesChanged(false);
-#endif  // #if GEOMETRYCOLLECTION_DEBUG_DRAW
-
-	Super::PostEditChangeProperty(PropertyChangedEvent);
-}
-
-bool UGeometryCollectionDebugDrawComponent::CanEditChange(const UProperty* InProperty) const
-{
-	if (HasBegunPlay())
-	{
-		return false;
-	}
-	return Super::CanEditChange(InProperty);
-}
-#endif  // #if WITH_EDITOR
-
 #if GEOMETRYCOLLECTION_DEBUG_DRAW
 bool UGeometryCollectionDebugDrawComponent::OnDebugDrawPropertiesChanged(bool bForceVisibilityUpdate)
 {
@@ -398,146 +290,131 @@ void UGeometryCollectionDebugDrawComponent::DebugDrawTick()
 	check(GeometryCollectionComponent->RestCollection);
 	check(GeometryCollectionDebugDrawActor);
 
-	AActor* const Actor = GetOwner();
-	check(Actor);
+	const bool bIsSelected = (SelectedTransformIndex != INDEX_NONE);
+	const bool bIsOneSelected = bIsSelected && !GeometryCollectionDebugDrawActor->bDebugDrawWholeCollection;
+	const bool bAreAllSelected = (bIsSelected && GeometryCollectionDebugDrawActor->bDebugDrawWholeCollection) ||
+		(GeometryCollectionDebugDrawActor->SelectedRigidBody.Id == INDEX_NONE &&
+			GeometryCollectionDebugDrawActor->SelectedRigidBody.Solver == GeometryCollectionComponent->ChaosSolverActor);
 
 	// Compute world space geometry and cluster transforms
 	TArray<FTransform> ClusterTransforms;
 	TArray<FTransform> GeometryTransforms;
 	ComputeTransforms(ClusterTransforms, GeometryTransforms);
 
-	const bool bIsSelected = (SelectedTransformIndex != INDEX_NONE);
-	const bool bIsOneSelected = bIsSelected && !GeometryCollectionDebugDrawActor->bDebugDrawWholeCollection;
-	const bool bAreAllSelected = (bIsSelected && GeometryCollectionDebugDrawActor->bDebugDrawWholeCollection) ||
-		(GeometryCollectionDebugDrawActor->SelectedRigidBody.Id == INDEX_NONE && 
-		 GeometryCollectionDebugDrawActor->SelectedRigidBody.Solver == GeometryCollectionComponent->ChaosSolverActor);
-
-	// Clustering
-	if (!bShowTransformIndices && GeometryCollectionDebugDrawActor->bShowTransformIndex && bIsOneSelected)
-	{
-		GeometryCollectionDebugDrawActor->DrawTransformIndex(ClusterTransforms, GeometryCollectionComponent, SelectedTransformIndex, TransformIndexColor);
-	}
-	else if (bShowTransformIndices || (GeometryCollectionDebugDrawActor->bShowTransformIndex && bAreAllSelected))
-	{
-		GeometryCollectionDebugDrawActor->DrawTransformIndices(ClusterTransforms, GeometryCollectionComponent, TransformIndexColor);
-	}
-
-	if (!bShowTransforms && GeometryCollectionDebugDrawActor->bShowTransform && bIsOneSelected)
-	{
-		GeometryCollectionDebugDrawActor->DrawTransform(ClusterTransforms, GeometryCollectionComponent, SelectedTransformIndex, TransformScale);
-	}
-	else if (bShowTransforms || (GeometryCollectionDebugDrawActor->bShowTransform && bAreAllSelected))
-	{
-		GeometryCollectionDebugDrawActor->DrawTransforms(ClusterTransforms, GeometryCollectionComponent, TransformScale);
-	}
-
-	if (!bShowParents && GeometryCollectionDebugDrawActor->bShowParent && bIsOneSelected)
-	{
-		GeometryCollectionDebugDrawActor->DrawParent(ClusterTransforms, GeometryCollectionComponent, SelectedTransformIndex, ParentColor);
-	}
-	else if (bShowParents || (GeometryCollectionDebugDrawActor->bShowParent && bAreAllSelected))
-	{
-		GeometryCollectionDebugDrawActor->DrawParents(ClusterTransforms, GeometryCollectionComponent, ParentColor);
-	}
-
-	if (!bShowLevels && GeometryCollectionDebugDrawActor->bShowLevel && bIsOneSelected)
-	{
-		GeometryCollectionDebugDrawActor->DrawLevel(ClusterTransforms, GeometryCollectionComponent, SelectedTransformIndex, LevelColor);
-	}
-	else if (bShowLevels || (GeometryCollectionDebugDrawActor->bShowLevel && bAreAllSelected))
-	{
-		GeometryCollectionDebugDrawActor->DrawLevels(ClusterTransforms, GeometryCollectionComponent, LevelColor);
-	}
-
-	// Geometry
-	if (!bShowVertices && GeometryCollectionDebugDrawActor->bShowVertices && bIsOneSelected)
-	{
-		GeometryCollectionDebugDrawActor->DrawVertices(GeometryTransforms, GeometryCollectionComponent, SelectedTransformIndex, VertexColor);
-	}
-	else if (bShowVertices || (GeometryCollectionDebugDrawActor->bShowVertices && bAreAllSelected))
-	{
-		GeometryCollectionDebugDrawActor->DrawVertices(GeometryTransforms, GeometryCollectionComponent, VertexColor);
-	}
-
-	if (!bShowVertexIndices && GeometryCollectionDebugDrawActor->bShowVertexIndices && bIsOneSelected)
-	{ 
-		GeometryCollectionDebugDrawActor->DrawVertexIndices(GeometryTransforms, GeometryCollectionComponent, SelectedTransformIndex, VertexIndexColor);
-	}
-	else if (bShowVertexIndices || (GeometryCollectionDebugDrawActor->bShowVertexIndices && bAreAllSelected))
-	{
-		GeometryCollectionDebugDrawActor->DrawVertexIndices(GeometryTransforms, GeometryCollectionComponent, VertexIndexColor);
-	}
-
-	if (!bShowVertexNormals && GeometryCollectionDebugDrawActor->bShowVertexNormals && bIsOneSelected)
-	{
-		GeometryCollectionDebugDrawActor->DrawVertexNormals(GeometryTransforms, GeometryCollectionComponent, SelectedTransformIndex, VertexNormalColor);
-	}
-	else if (bShowVertexNormals || (GeometryCollectionDebugDrawActor->bShowVertexNormals && bAreAllSelected))
-	{
-		GeometryCollectionDebugDrawActor->DrawVertexNormals(GeometryTransforms, GeometryCollectionComponent, VertexNormalColor);
-	}
-
-	if (!bShowFaces && GeometryCollectionDebugDrawActor->bShowFaces && bIsOneSelected)
+	if (bIsOneSelected)
 	{
-		GeometryCollectionDebugDrawActor->DrawFaces(GeometryTransforms, GeometryCollectionComponent, SelectedTransformIndex, FaceColor);
-	}
-	else if (bShowFaces || (GeometryCollectionDebugDrawActor->bShowFaces && bAreAllSelected))
-	{
-		GeometryCollectionDebugDrawActor->DrawFaces(GeometryTransforms, GeometryCollectionComponent, FaceColor);
-	}
-
-	if (!bShowFaceIndices && GeometryCollectionDebugDrawActor->bShowFaceIndices && bIsOneSelected)
-	{
-		GeometryCollectionDebugDrawActor->DrawFaceIndices(GeometryTransforms, GeometryCollectionComponent, SelectedTransformIndex, FaceIndexColor);
-	}
-	else if (bShowFaceIndices || (GeometryCollectionDebugDrawActor->bShowFaceIndices && bAreAllSelected))
-	{
-		GeometryCollectionDebugDrawActor->DrawFaceIndices(GeometryTransforms, GeometryCollectionComponent, FaceIndexColor);
-	}
-
-	if (bShowSingleFace)
-	{
-		GeometryCollectionDebugDrawActor->DrawSingleFace(GeometryTransforms, GeometryCollectionComponent, SingleFaceIndex, SingleFaceColor);
-	}
-	if (GeometryCollectionDebugDrawActor->bShowSingleFace && (bIsOneSelected || bAreAllSelected))  // No else required here, it should be able to draw the two faces at the same time
-	{
-		GeometryCollectionDebugDrawActor->DrawSingleFace(GeometryTransforms, GeometryCollectionComponent, GeometryCollectionDebugDrawActor->SingleFaceIndex, SingleFaceColor);
-	}
-
-	if (!bShowFaceNormals && GeometryCollectionDebugDrawActor->bShowFaceNormals && bIsOneSelected)
-	{
-		GeometryCollectionDebugDrawActor->DrawFaceNormals(GeometryTransforms, GeometryCollectionComponent, SelectedTransformIndex, FaceNormalColor);
-	}
-	else if (bShowFaceNormals || (GeometryCollectionDebugDrawActor->bShowFaceNormals && bAreAllSelected))
-	{
-		GeometryCollectionDebugDrawActor->DrawFaceNormals(GeometryTransforms, GeometryCollectionComponent, FaceNormalColor);
-	}
-
-	if (!bShowGeometryIndices && GeometryCollectionDebugDrawActor->bShowGeometryIndex && bIsOneSelected)
-	{
-		GeometryCollectionDebugDrawActor->DrawGeometryIndex(GeometryTransforms, GeometryCollectionComponent, SelectedTransformIndex, GeometryIndexColor);
-	}
-	else if (bShowGeometryIndices || (GeometryCollectionDebugDrawActor->bShowGeometryIndex && bAreAllSelected))
-	{
-		GeometryCollectionDebugDrawActor->DrawGeometryIndices(GeometryTransforms, GeometryCollectionComponent, GeometryIndexColor);
-	}
-
-	if (!bShowGeometryTransforms && GeometryCollectionDebugDrawActor->bShowGeometryTransform && bIsOneSelected)
-	{
-		GeometryCollectionDebugDrawActor->DrawTransform(GeometryTransforms, GeometryCollectionComponent, SelectedTransformIndex, GeometryTransformScale);
-	}
-	else if (bShowGeometryTransforms || (GeometryCollectionDebugDrawActor->bShowGeometryTransform && bAreAllSelected))
-	{
-		GeometryCollectionDebugDrawActor->DrawTransforms(GeometryTransforms, GeometryCollectionComponent, GeometryTransformScale);
+		if (GeometryCollectionDebugDrawActor->bShowTransformIndex)
+		{
+			GeometryCollectionDebugDrawActor->DrawTransformIndex(ClusterTransforms, GeometryCollectionComponent, SelectedTransformIndex, GeometryCollectionDebugDrawActor->TransformIndexColor);
+		}
+		if (GeometryCollectionDebugDrawActor->bShowTransform)
+		{
+			GeometryCollectionDebugDrawActor->DrawTransform(ClusterTransforms, GeometryCollectionComponent, SelectedTransformIndex, GeometryCollectionDebugDrawActor->TransformScale);
+		}
+		if (GeometryCollectionDebugDrawActor->bShowParent)
+		{
+			GeometryCollectionDebugDrawActor->DrawParent(ClusterTransforms, GeometryCollectionComponent, SelectedTransformIndex, GeometryCollectionDebugDrawActor->ParentColor);
+		}
+		if (GeometryCollectionDebugDrawActor->bShowLevel)
+		{
+			GeometryCollectionDebugDrawActor->DrawLevel(ClusterTransforms, GeometryCollectionComponent, SelectedTransformIndex, GeometryCollectionDebugDrawActor->LevelColor);
+		}
+		if (GeometryCollectionDebugDrawActor->bShowVertices)
+		{
+			GeometryCollectionDebugDrawActor->DrawVertices(GeometryTransforms, GeometryCollectionComponent, SelectedTransformIndex, GeometryCollectionDebugDrawActor->VertexColor);
+		}
+		if (GeometryCollectionDebugDrawActor->bShowVertexIndices)
+		{
+			GeometryCollectionDebugDrawActor->DrawVertexIndices(GeometryTransforms, GeometryCollectionComponent, SelectedTransformIndex, GeometryCollectionDebugDrawActor->VertexIndexColor);
+		}
+		if (GeometryCollectionDebugDrawActor->bShowVertexNormals)
+		{
+			GeometryCollectionDebugDrawActor->DrawVertexNormals(GeometryTransforms, GeometryCollectionComponent, SelectedTransformIndex, GeometryCollectionDebugDrawActor->VertexNormalColor);
+		}
+		if (GeometryCollectionDebugDrawActor->bShowFaces)
+		{
+			GeometryCollectionDebugDrawActor->DrawFaces(GeometryTransforms, GeometryCollectionComponent, SelectedTransformIndex, GeometryCollectionDebugDrawActor->FaceColor);
+		}
+		if (GeometryCollectionDebugDrawActor->bShowFaceIndices)
+		{
+			GeometryCollectionDebugDrawActor->DrawFaceIndices(GeometryTransforms, GeometryCollectionComponent, SelectedTransformIndex, GeometryCollectionDebugDrawActor->FaceIndexColor);
+		}
+		if (GeometryCollectionDebugDrawActor->bShowFaceNormals)
+		{
+			GeometryCollectionDebugDrawActor->DrawFaceNormals(GeometryTransforms, GeometryCollectionComponent, SelectedTransformIndex, GeometryCollectionDebugDrawActor->FaceNormalColor);
+		}
+		if (GeometryCollectionDebugDrawActor->bShowGeometryIndex)
+		{
+			GeometryCollectionDebugDrawActor->DrawGeometryIndex(GeometryTransforms, GeometryCollectionComponent, SelectedTransformIndex, GeometryCollectionDebugDrawActor->GeometryIndexColor);
+		}
+		if (GeometryCollectionDebugDrawActor->bShowGeometryTransform)
+		{
+			GeometryCollectionDebugDrawActor->DrawTransform(GeometryTransforms, GeometryCollectionComponent, SelectedTransformIndex, GeometryCollectionDebugDrawActor->GeometryTransformScale);
+		}
+		if (GeometryCollectionDebugDrawActor->bShowBoundingBox)
+		{
+			GeometryCollectionDebugDrawActor->DrawBoundingBox(GeometryTransforms, GeometryCollectionComponent, SelectedTransformIndex, GeometryCollectionDebugDrawActor->BoundingBoxColor);
+		}
 	}
-
-	if (!bShowBoundingBoxes && GeometryCollectionDebugDrawActor->bShowBoundingBox && bIsOneSelected)
+	else if (bAreAllSelected)
 	{
-		GeometryCollectionDebugDrawActor->DrawBoundingBox(GeometryTransforms, GeometryCollectionComponent, SelectedTransformIndex, BoundingBoxColor);
+		if (GeometryCollectionDebugDrawActor->bShowTransformIndex)
+		{
+			GeometryCollectionDebugDrawActor->DrawTransformIndices(ClusterTransforms, GeometryCollectionComponent, GeometryCollectionDebugDrawActor->TransformIndexColor);
+		}
+		if (GeometryCollectionDebugDrawActor->bShowTransform)
+		{
+			GeometryCollectionDebugDrawActor->DrawTransforms(ClusterTransforms, GeometryCollectionComponent, GeometryCollectionDebugDrawActor->TransformScale);
+		}
+		if (GeometryCollectionDebugDrawActor->bShowParent)
+		{
+			GeometryCollectionDebugDrawActor->DrawParents(ClusterTransforms, GeometryCollectionComponent, GeometryCollectionDebugDrawActor->ParentColor);
+		}
+		if (GeometryCollectionDebugDrawActor->bShowLevel)
+		{
+			GeometryCollectionDebugDrawActor->DrawLevels(ClusterTransforms, GeometryCollectionComponent, GeometryCollectionDebugDrawActor->LevelColor);
+		}
+		if (GeometryCollectionDebugDrawActor->bShowVertices)
+		{
+			GeometryCollectionDebugDrawActor->DrawVertices(GeometryTransforms, GeometryCollectionComponent, GeometryCollectionDebugDrawActor->VertexColor);
+		}
+		if (GeometryCollectionDebugDrawActor->bShowVertexIndices)
+		{
+			GeometryCollectionDebugDrawActor->DrawVertexIndices(GeometryTransforms, GeometryCollectionComponent, GeometryCollectionDebugDrawActor->VertexIndexColor);
+		}
+		if (GeometryCollectionDebugDrawActor->bShowVertexNormals)
+		{
+			GeometryCollectionDebugDrawActor->DrawVertexNormals(GeometryTransforms, GeometryCollectionComponent, GeometryCollectionDebugDrawActor->VertexNormalColor);
+		}
+		if (GeometryCollectionDebugDrawActor->bShowFaces)
+		{
+			GeometryCollectionDebugDrawActor->DrawFaces(GeometryTransforms, GeometryCollectionComponent, GeometryCollectionDebugDrawActor->FaceColor);
+		}
+		if (GeometryCollectionDebugDrawActor->bShowFaceIndices)
+		{
+			GeometryCollectionDebugDrawActor->DrawFaceIndices(GeometryTransforms, GeometryCollectionComponent, GeometryCollectionDebugDrawActor->FaceIndexColor);
+		}
+		if (GeometryCollectionDebugDrawActor->bShowFaceNormals)
+		{
+			GeometryCollectionDebugDrawActor->DrawFaceNormals(GeometryTransforms, GeometryCollectionComponent, GeometryCollectionDebugDrawActor->FaceNormalColor);
+		}
+		if (GeometryCollectionDebugDrawActor->bShowGeometryIndex)
+		{
+			GeometryCollectionDebugDrawActor->DrawGeometryIndices(GeometryTransforms, GeometryCollectionComponent, GeometryCollectionDebugDrawActor->GeometryIndexColor);
+		}
+		if (GeometryCollectionDebugDrawActor->bShowGeometryTransform)
+		{
+			GeometryCollectionDebugDrawActor->DrawTransforms(GeometryTransforms, GeometryCollectionComponent, GeometryCollectionDebugDrawActor->GeometryTransformScale);
+		}
+		if (GeometryCollectionDebugDrawActor->bShowBoundingBox)
+		{
+			GeometryCollectionDebugDrawActor->DrawBoundingBoxes(GeometryTransforms, GeometryCollectionComponent, GeometryCollectionDebugDrawActor->BoundingBoxColor);
+		}
 	}
-	else if (bShowBoundingBoxes || (GeometryCollectionDebugDrawActor->bShowBoundingBox && bAreAllSelected))
+	// Draw single face
+	if (bIsSelected && GeometryCollectionDebugDrawActor->bShowSingleFace)
 	{
-		GeometryCollectionDebugDrawActor->DrawBoundingBoxes(GeometryTransforms, GeometryCollectionComponent, BoundingBoxColor);
+		GeometryCollectionDebugDrawActor->DrawSingleFace(GeometryTransforms, GeometryCollectionComponent, GeometryCollectionDebugDrawActor->SingleFaceIndex, GeometryCollectionDebugDrawActor->SingleFaceColor);
 	}
 }
 
@@ -713,7 +590,7 @@ void UGeometryCollectionDebugDrawComponent::UpdateGeometryVisibility(bool bForce
 			HiddenTransformIndex = INDEX_NONE;
 			break;
 		case EGeometryCollectionDebugDrawActorHideGeometry::HideWithCollision:
-			if (bShowRigidBodyCollisions || (bAreAnySelected && GeometryCollectionDebugDrawActor->bShowRigidBodyCollision))
+			if (bAreAnySelected && GeometryCollectionDebugDrawActor->bShowRigidBodyCollision)
 			{
 				bIsVisible = !bAreAllSelected;
 				HiddenTransformIndex = bIsVisible ? SelectedTransformIndex: INDEX_NONE;
@@ -813,8 +690,8 @@ void UGeometryCollectionDebugDrawComponent::UpdateTickStatus()
 			 GeometryCollectionDebugDrawActor->SelectedRigidBody.Solver == GeometryCollectionComponent->ChaosSolverActor));
 
 		bIsEnabled =
-			(bAreAnySelected && (
-				   GeometryCollectionDebugDrawActor->bShowRigidBodyId  // Check the debug draw's actor debug options
+			bAreAnySelected && (
+				   GeometryCollectionDebugDrawActor->bShowRigidBodyId
 				|| GeometryCollectionDebugDrawActor->bShowRigidBodyCollision
 				|| GeometryCollectionDebugDrawActor->bShowRigidBodyTransform
 				|| GeometryCollectionDebugDrawActor->bShowRigidBodyInertia
@@ -836,29 +713,7 @@ void UGeometryCollectionDebugDrawComponent::UpdateTickStatus()
 				|| GeometryCollectionDebugDrawActor->SingleFaceIndex
 				|| GeometryCollectionDebugDrawActor->bShowVertices
 				|| GeometryCollectionDebugDrawActor->bShowVertexIndices
-				|| GeometryCollectionDebugDrawActor->bShowVertexNormals))
-			|| bShowRigidBodyIds  // Check this component's debug options
-			|| bShowRigidBodyTransforms
-			|| bShowRigidBodyCollisions
-			|| bShowRigidBodyInertias
-			|| bShowRigidBodyVelocities
-			|| bShowRigidBodyForces
-			|| bShowRigidBodyInfos
-			|| bShowTransformIndices
-			|| bShowTransforms
-			|| bShowLevels
-			|| bShowParents
-			|| bShowConnectivityEdges
-			|| bShowGeometryIndices
-			|| bShowGeometryTransforms
-			|| bShowBoundingBoxes
-			|| bShowFaces
-			|| bShowFaceIndices
-			|| bShowFaceNormals
-			|| bShowSingleFace
-			|| bShowVertices
-			|| bShowVertexIndices
-			|| bShowVertexNormals;
+				|| GeometryCollectionDebugDrawActor->bShowVertexNormals);
 	}
 	// Update component's ability to tick 
 	SetComponentTickEnabled(bIsEnabled);
@@ -945,88 +800,86 @@ void UGeometryCollectionDebugDrawComponent::DebugDrawChaosTick()
 	// Visualize other rigid body debug draw informations
 	const bool bIsOneSelected = bIsSelected && !GeometryCollectionDebugDrawActor->bDebugDrawWholeCollection;
 	const bool bAreAllSelected = (bIsSelected && GeometryCollectionDebugDrawActor->bDebugDrawWholeCollection) ||
-		(GeometryCollectionDebugDrawActor->SelectedRigidBody.Id == INDEX_NONE && 
+		(GeometryCollectionDebugDrawActor->SelectedRigidBody.Id == INDEX_NONE &&
 		 GeometryCollectionDebugDrawActor->SelectedRigidBody.Solver == GeometryCollectionComponent->ChaosSolverActor);
 
-	if (!bShowRigidBodyIds && GeometryCollectionDebugDrawActor->bShowRigidBodyId && bIsOneSelected)
-	{
-		GeometryCollectionDebugDrawActor->DrawRigidBodyId(GeometryCollectionComponent, SelectedTransformIndex, ParticlesData, RigidBodyIds, RigidBodyIdColor);
-	}
-	else if (bShowRigidBodyIds || (GeometryCollectionDebugDrawActor->bShowRigidBodyId && bAreAllSelected))
-	{
-		GeometryCollectionDebugDrawActor->DrawRigidBodiesId(GeometryCollectionComponent, ParticlesData, RigidBodyIds, RigidBodyIdColor);
-	}
-
-	if (!bShowRigidBodyTransforms && GeometryCollectionDebugDrawActor->bShowRigidBodyTransform && bIsOneSelected)
-	{
-		GeometryCollectionDebugDrawActor->DrawRigidBodyTransform(GeometryCollectionComponent, SelectedTransformIndex, ParticlesData, RigidBodyTransformScale);
-	}
-	else if (bShowRigidBodyTransforms || (GeometryCollectionDebugDrawActor->bShowRigidBodyTransform && bAreAllSelected))
-	{
-		GeometryCollectionDebugDrawActor->DrawRigidBodiesTransform(GeometryCollectionComponent, ParticlesData, RigidBodyTransformScale);
-	}
-
-	const bool bIsShowingLevelSet = (RenderLevelSetOwner == this && LastRenderedId == SelectedRigidBodyId);  // Only draw single collision whenever there isn't a level set being already rendered
-	if (!bIsShowingLevelSet && !bShowRigidBodyCollisions && GeometryCollectionDebugDrawActor->bShowRigidBodyCollision && bIsOneSelected)
-	{
-		GeometryCollectionDebugDrawActor->DrawRigidBodyCollision(GeometryCollectionComponent, SelectedTransformIndex, ParticlesData, RigidBodyCollisionColor);
-	}
-	else if (bShowRigidBodyCollisions || (GeometryCollectionDebugDrawActor->bShowRigidBodyCollision && bAreAllSelected))
-	{
-		GeometryCollectionDebugDrawActor->DrawRigidBodiesCollision(GeometryCollectionComponent, ParticlesData, RigidBodyCollisionColor);
-	}
-
-	if (!bShowRigidBodyInertias && GeometryCollectionDebugDrawActor->bShowRigidBodyInertia && bIsOneSelected)
+	if (bIsOneSelected)
 	{
-		GeometryCollectionDebugDrawActor->DrawRigidBodyInertia(GeometryCollectionComponent, SelectedTransformIndex, ParticlesData, RigidBodyInertiaColor);
-	}
-	else if (bShowRigidBodyInertias || (GeometryCollectionDebugDrawActor->bShowRigidBodyInertia && bAreAllSelected))
-	{
-		GeometryCollectionDebugDrawActor->DrawRigidBodiesInertia(GeometryCollectionComponent, ParticlesData, RigidBodyInertiaColor);
-	}
-
-	if (!bShowRigidBodyVelocities && GeometryCollectionDebugDrawActor->bShowRigidBodyVelocity && bIsOneSelected)
-	{
-		GeometryCollectionDebugDrawActor->DrawRigidBodyVelocity(GeometryCollectionComponent, SelectedTransformIndex, ParticlesData, RigidBodyVelocityColor);
-	}
-	else if (bShowRigidBodyVelocities || (GeometryCollectionDebugDrawActor->bShowRigidBodyVelocity && bAreAllSelected))
-	{
-		GeometryCollectionDebugDrawActor->DrawRigidBodiesVelocity(GeometryCollectionComponent, ParticlesData, RigidBodyVelocityColor);
-	}
-
-	if (!bShowRigidBodyForces && GeometryCollectionDebugDrawActor->bShowRigidBodyForce && bIsOneSelected)
-	{
-		GeometryCollectionDebugDrawActor->DrawRigidBodyForce(GeometryCollectionComponent, SelectedTransformIndex, ParticlesData, RigidBodyForceColor);
-	}
-	else if (bShowRigidBodyForces || (GeometryCollectionDebugDrawActor->bShowRigidBodyForce && bAreAllSelected))
-	{
-		GeometryCollectionDebugDrawActor->DrawRigidBodiesForce(GeometryCollectionComponent, ParticlesData, RigidBodyForceColor);
-	}
-
-	if (!bShowRigidBodyInfos && GeometryCollectionDebugDrawActor->bShowRigidBodyInfos && bIsOneSelected)
-	{
-		GeometryCollectionDebugDrawActor->DrawRigidBodyInfo(GeometryCollectionComponent, SelectedTransformIndex, ParticlesData, RigidBodyInfoColor);
-	}
-	else if (bShowRigidBodyInfos || (GeometryCollectionDebugDrawActor->bShowRigidBodyInfos && bAreAllSelected))
-	{
-		GeometryCollectionDebugDrawActor->DrawRigidBodiesInfo(GeometryCollectionComponent, ParticlesData, RigidBodyInfoColor);
-	}
-
-	if (!bShowConnectivityEdges && GeometryCollectionDebugDrawActor->bShowConnectivityEdges && bIsOneSelected)
-	{
-		GeometryCollectionDebugDrawActor->DrawConnectivityEdges(GeometryCollectionComponent, SelectedTransformIndex, ParticlesData, RigidBodyIds, ConnectivityEdgeThickness);
+		if (GeometryCollectionDebugDrawActor->bShowRigidBodyId)
+		{
+			GeometryCollectionDebugDrawActor->DrawRigidBodyId(GeometryCollectionComponent, SelectedTransformIndex, ParticlesData, RigidBodyIds, GeometryCollectionDebugDrawActor->RigidBodyIdColor);
+		}
+		if (GeometryCollectionDebugDrawActor->bShowRigidBodyTransform)
+		{
+			GeometryCollectionDebugDrawActor->DrawRigidBodyTransform(GeometryCollectionComponent, SelectedTransformIndex, ParticlesData, GeometryCollectionDebugDrawActor->RigidBodyTransformScale);
+		}
+		const bool bIsShowingLevelSet = (RenderLevelSetOwner == this && LastRenderedId == SelectedRigidBodyId);  // Only draw single collision whenever there isn't a level set being already rendered
+		if (!bIsShowingLevelSet && GeometryCollectionDebugDrawActor->bShowRigidBodyCollision)
+		{
+			GeometryCollectionDebugDrawActor->DrawRigidBodyCollision(GeometryCollectionComponent, SelectedTransformIndex, ParticlesData, GeometryCollectionDebugDrawActor->RigidBodyCollisionColor);
+		}
+		if (GeometryCollectionDebugDrawActor->bShowRigidBodyInertia)
+		{
+			GeometryCollectionDebugDrawActor->DrawRigidBodyInertia(GeometryCollectionComponent, SelectedTransformIndex, ParticlesData, GeometryCollectionDebugDrawActor->RigidBodyInertiaColor);
+		}
+		if (GeometryCollectionDebugDrawActor->bShowRigidBodyVelocity)
+		{
+			GeometryCollectionDebugDrawActor->DrawRigidBodyVelocity(GeometryCollectionComponent, SelectedTransformIndex, ParticlesData, GeometryCollectionDebugDrawActor->RigidBodyVelocityColor);
+		}
+		if (GeometryCollectionDebugDrawActor->bShowRigidBodyForce)
+		{
+			GeometryCollectionDebugDrawActor->DrawRigidBodyForce(GeometryCollectionComponent, SelectedTransformIndex, ParticlesData, GeometryCollectionDebugDrawActor->RigidBodyForceColor);
+		}
+		if (GeometryCollectionDebugDrawActor->bShowRigidBodyInfos)
+		{
+			GeometryCollectionDebugDrawActor->DrawRigidBodyInfo(GeometryCollectionComponent, SelectedTransformIndex, ParticlesData, GeometryCollectionDebugDrawActor->RigidBodyInfoColor);
+		}
+		if (GeometryCollectionDebugDrawActor->bShowConnectivityEdges)
+		{
+			GeometryCollectionDebugDrawActor->DrawConnectivityEdges(GeometryCollectionComponent, SelectedTransformIndex, ParticlesData, RigidBodyIds);
+		}
 	}
-	else if (bShowConnectivityEdges || (GeometryCollectionDebugDrawActor->bShowConnectivityEdges && bAreAllSelected))
+	else if (bAreAllSelected)
 	{
-		GeometryCollectionDebugDrawActor->DrawConnectivityEdges(GeometryCollectionComponent, ParticlesData, RigidBodyIds, ConnectivityEdgeThickness);
+		if (GeometryCollectionDebugDrawActor->bShowRigidBodyCollision)
+		{
+			GeometryCollectionDebugDrawActor->DrawRigidBodiesCollision(GeometryCollectionComponent, ParticlesData, GeometryCollectionDebugDrawActor->RigidBodyCollisionColor);
+		}
+		if (GeometryCollectionDebugDrawActor->bShowRigidBodyInertia)
+		{
+			GeometryCollectionDebugDrawActor->DrawRigidBodiesInertia(GeometryCollectionComponent, ParticlesData, GeometryCollectionDebugDrawActor->RigidBodyInertiaColor);
+		}
+		if (GeometryCollectionDebugDrawActor->bShowRigidBodyVelocity)
+		{
+			GeometryCollectionDebugDrawActor->DrawRigidBodiesVelocity(GeometryCollectionComponent, ParticlesData, GeometryCollectionDebugDrawActor->RigidBodyVelocityColor);
+		}
+		if (GeometryCollectionDebugDrawActor->bShowRigidBodyForce)
+		{
+			GeometryCollectionDebugDrawActor->DrawRigidBodiesForce(GeometryCollectionComponent, ParticlesData, GeometryCollectionDebugDrawActor->RigidBodyForceColor);
+		}
+		if (GeometryCollectionDebugDrawActor->bShowRigidBodyInfos)
+		{
+			GeometryCollectionDebugDrawActor->DrawRigidBodiesInfo(GeometryCollectionComponent, ParticlesData, GeometryCollectionDebugDrawActor->RigidBodyInfoColor);
+		}
+		if (GeometryCollectionDebugDrawActor->bShowConnectivityEdges)
+		{
+			GeometryCollectionDebugDrawActor->DrawConnectivityEdges(GeometryCollectionComponent, ParticlesData, RigidBodyIds);
+		}
+		if (GeometryCollectionDebugDrawActor->bShowRigidBodyTransform)
+		{
+			GeometryCollectionDebugDrawActor->DrawRigidBodiesTransform(GeometryCollectionComponent, ParticlesData, GeometryCollectionDebugDrawActor->RigidBodyTransformScale);
+		}
+		if (GeometryCollectionDebugDrawActor->bShowRigidBodyId)
+		{
+			GeometryCollectionDebugDrawActor->DrawRigidBodiesId(GeometryCollectionComponent, ParticlesData, RigidBodyIds, GeometryCollectionDebugDrawActor->RigidBodyIdColor);
+		}
 	}
 }
 
 void UGeometryCollectionDebugDrawComponent::UpdateLevelSetVisibility()
 {
 	const bool bIsSelected = (SelectedTransformIndex != INDEX_NONE);
-	const bool bShowCollision = bShowRigidBodyCollisions || 
-		(bIsSelected && GeometryCollectionDebugDrawActor && GeometryCollectionDebugDrawActor->bShowRigidBodyCollision);
+	const bool bShowCollision = bIsSelected && GeometryCollectionDebugDrawActor && GeometryCollectionDebugDrawActor->bShowRigidBodyCollision;
 	
 	if (RenderLevelSetOwner == this && (LastRenderedId == INDEX_NONE || !bShowCollision))
 	{
diff --git a/Engine/Source/Runtime/Experimental/GeometryCollectionEngine/Public/GeometryCollection/GeometryCollectionDebugDrawActor.h b/Engine/Source/Runtime/Experimental/GeometryCollectionEngine/Public/GeometryCollection/GeometryCollectionDebugDrawActor.h
--- a/Engine/Source/Runtime/Experimental/GeometryCollectionEngine/Public/GeometryCollection/GeometryCollectionDebugDrawActor.h
+++ b/Engine/Source/Runtime/Experimental/GeometryCollectionEngine/Public/GeometryCollection/GeometryCollectionDebugDrawActor.h
@@ -4,7 +4,6 @@
 
 #include "GameFramework/Actor.h"
 #include "Components/BillboardComponent.h"
-#include "GeometryCollection/GeometryCollectionDebugDrawComponent.h"
 
 #include "GeometryCollectionDebugDrawActor.generated.h"
 
@@ -35,6 +34,16 @@ enum class EGeometryCollectionDebugDrawActorHideGeometry : uint8
 	HideAll
 };
 
+/**
+* FGeometryCollectionDebugDrawWarningMessage
+*   Empty structure used to embed a warning message in the UI through a detail customization.
+*/
+USTRUCT()
+struct FGeometryCollectionDebugDrawWarningMessage
+{
+	GENERATED_USTRUCT_BODY()
+};
+
 /**
 * FGeometryCollectionDebugDrawActorSelectedRigidBody
 *   Structure used to select a rigid body id with a picking tool through a detail customization.
@@ -220,14 +229,102 @@ public:
 	UPROPERTY(EditAnywhere, Category = "Debug Draw|Settings", meta = (ClampMin="0.0001"))
 	float NormalScale;
 
-	/** Scale factor used for visualizing transforms. */
+	/** Scale of the axis used for visualizing all transforms. */
 	UPROPERTY(EditAnywhere, Category = "Debug Draw|Settings", meta = (ClampMin="0.0001"))
-	float TransformScale;
+	float AxisScale;
 
 	/** Size of arrows used for visualizing normals, breaking information, ...etc. */
 	UPROPERTY(EditAnywhere, Category = "Debug Draw|Settings", meta = (ClampMin="0.0001"))
 	float ArrowScale;
 
+	/** Color used for the visualization of the rigid body ids. */
+	UPROPERTY(EditAnywhere, Category = "Debug Draw|Settings", meta = (HideAlphaChannel))
+	FColor RigidBodyIdColor;
+
+	/** Scale for rigid body transform visualization. */
+	UPROPERTY(EditAnywhere, Category = "Debug Draw|Settings", meta = (ClampMin="0.0001"))
+	float RigidBodyTransformScale;
+
+	/** Color used for collision primitives visualization. */
+	UPROPERTY(EditAnywhere, Category = "Debug Draw|Settings", meta = (HideAlphaChannel))
+	FColor RigidBodyCollisionColor;
+
+	/** Color used for the visualization of the rigid body inertia tensor box. */
+	UPROPERTY(EditAnywhere, Category = "Debug Draw|Settings", meta = (HideAlphaChannel))
+	FColor RigidBodyInertiaColor;
+
+	/** Color used for rigid body velocities visualization. */
+	UPROPERTY(EditAnywhere, Category = "Debug Draw|Settings", meta = (HideAlphaChannel))
+	FColor RigidBodyVelocityColor;
+
+	/** Color used for rigid body applied force and torque visualization. */
+	UPROPERTY(EditAnywhere, Category = "Debug Draw|Settings", meta = (HideAlphaChannel))
+	FColor RigidBodyForceColor;
+
+	/** Color used for the visualization of the rigid body infos. */
+	UPROPERTY(EditAnywhere, Category = "Debug Draw|Settings", meta = (HideAlphaChannel))
+	FColor RigidBodyInfoColor;
+
+	/** Color used for the visualization of the transform indices. */
+	UPROPERTY(EditAnywhere, Category = "Debug Draw|Settings", meta = (HideAlphaChannel))
+	FColor TransformIndexColor;
+
+	/** Scale for cluster transform visualization. */
+	UPROPERTY(EditAnywhere, Category = "Debug Draw|Settings", meta = (ClampMin="0.0001"))
+	float TransformScale;
+
+	/** Color used for the visualization of the levels. */
+	UPROPERTY(EditAnywhere, Category = "Debug Draw|Settings", meta = (HideAlphaChannel))
+	FColor LevelColor;
+
+	/** Color used for the visualization of the link from the parents. */
+	UPROPERTY(EditAnywhere, Category = "Debug Draw|Settings", meta = (HideAlphaChannel))
+	FColor ParentColor;
+
+	/** Line thickness used for the visualization of the rigid clustering connectivity edges. */
+	UPROPERTY(EditAnywhere, Category = "Debug Draw|Settings", meta = (ClampMin="0.0001"))
+	float ConnectivityEdgeThickness;
+
+	/** Color used for the visualization of the geometry indices. */
+	UPROPERTY(EditAnywhere, Category = "Debug Draw|Settings", meta = (HideAlphaChannel))
+	FColor GeometryIndexColor;
+
+	/** Scale for geometry transform visualization. */
+	UPROPERTY(EditAnywhere, Category = "Debug Draw|Settings", meta = (ClampMin="0.0001"))
+	float GeometryTransformScale;
+
+	/** Color used for the visualization of the bounding boxes. */
+	UPROPERTY(EditAnywhere, Category = "Debug Draw|Settings", meta = (HideAlphaChannel))
+	FColor BoundingBoxColor;
+
+	/** Color used for the visualization of the faces. */
+	UPROPERTY(EditAnywhere, Category = "Debug Draw|Settings", meta = (HideAlphaChannel))
+	FColor FaceColor;
+
+	/** Color used for the visualization of the face indices. */
+	UPROPERTY(EditAnywhere, Category = "Debug Draw|Settings", meta = (HideAlphaChannel))
+	FColor FaceIndexColor;
+
+	/** Color used for the visualization of the face normals. */
+	UPROPERTY(EditAnywhere, Category = "Debug Draw|Settings", meta = (HideAlphaChannel))
+	FColor FaceNormalColor;
+
+	/** Color used for the visualization of the single face. */
+	UPROPERTY(EditAnywhere, Category = "Debug Draw|Settings", meta = (HideAlphaChannel))
+	FColor SingleFaceColor;
+
+	/** Color used for the visualization of the vertices. */
+	UPROPERTY(EditAnywhere, Category = "Debug Draw|Settings", meta = (HideAlphaChannel))
+	FColor VertexColor;
+
+	/** Color used for the visualization of the vertex indices. */
+	UPROPERTY(EditAnywhere, Category = "Debug Draw|Settings", meta = (HideAlphaChannel))
+	FColor VertexIndexColor;
+
+	/** Color used for the visualization of the vertex normals. */
+	UPROPERTY(EditAnywhere, Category = "Debug Draw|Settings", meta = (HideAlphaChannel))
+	FColor VertexNormalColor;
+
 	/** Display icon in the editor. */
 	UPROPERTY()
 	UBillboardComponent* SpriteComponent;
@@ -268,7 +365,7 @@ public:
 	void DrawVertexIndices(const TArray<FTransform>& GlobalTransforms, const UGeometryCollectionComponent* GeometryCollectionComponent, const FColor& Color);
 
 	/** Draw vertex indices for the part of the geometry attached to the specified transform index. */
-	void DrawVertexIndices(const TArray<FTransform>& GlobalTransforms, const UGeometryCollectionComponent* GeometryCollectionComponent, int32 TransformIndex, const FColor & Color);
+	void DrawVertexIndices(const TArray<FTransform>& GlobalTransforms, const UGeometryCollectionComponent* GeometryCollectionComponent, int32 TransformIndex, const FColor& Color);
 
 	/** Draw vertex normals. */
 	void DrawVertexNormals(const TArray<FTransform>& GlobalTransforms, const UGeometryCollectionComponent* GeometryCollectionComponent, const FColor& Color);
@@ -368,10 +465,10 @@ public:
 	void DrawRigidBodyInfo(const UGeometryCollectionComponent* GeometryCollectionComponent, int32 TransformIndex, const FGeometryCollectionParticlesData& ParticlesData, const FColor& Color);
 
 	/** Draw Chaos' rigid clustering's connectivities edges */
-	void DrawConnectivityEdges(const UGeometryCollectionComponent* GeometryCollectionComponent, const FGeometryCollectionParticlesData& ParticlesData, const TManagedArray<int32>& RigidBodyIdArray, float Thickness);
+	void DrawConnectivityEdges(const UGeometryCollectionComponent* GeometryCollectionComponent, const FGeometryCollectionParticlesData& ParticlesData, const TManagedArray<int32>& RigidBodyIdArray);
 
 	/** Draw Chaos' single rigid clustering's connectivity edges. */
-	void DrawConnectivityEdges(const UGeometryCollectionComponent* GeometryCollectionComponent, int32 TransformIndex, const FGeometryCollectionParticlesData& ParticlesData, const TManagedArray<int32>& RigidBodyIdArray, float Thickness, FColor HSVColor = FColor(157, 160, 128));
+	void DrawConnectivityEdges(const UGeometryCollectionComponent* GeometryCollectionComponent, int32 TransformIndex, const FGeometryCollectionParticlesData& ParticlesData, const TManagedArray<int32>& RigidBodyIdArray, FColor HSVColor = FColor(157, 160, 128));
 
 	/** Draw Chaos' rigid body velocity. */
 	void DrawRigidBodiesVelocity(const UGeometryCollectionComponent* GeometryCollectionComponent, const FGeometryCollectionParticlesData& ParticlesData, const FColor& Color);
@@ -435,7 +532,7 @@ private:
 	void DrawRigidBodyInfoNoChecks(const UGeometryCollectionComponent* GeometryCollectionComponent, int32 TransformIndex, const FGeometryCollectionParticlesData& ParticlesData, const FColor& Color);
 
 	/** Draw Chaos' single rigid clustering's connectivity edges without synchronization checks. */
-	void DrawConnectivityEdgesNoChecks(const UGeometryCollectionComponent* GeometryCollectionComponent, int32 TransformIndex, const FGeometryCollectionParticlesData& ParticlesData, const TManagedArray<int32>& RigidBodyIdArray, float Thickness, const FColor& Color);
+	void DrawConnectivityEdgesNoChecks(const UGeometryCollectionComponent* GeometryCollectionComponent, int32 TransformIndex, const FGeometryCollectionParticlesData& ParticlesData, const TManagedArray<int32>& RigidBodyIdArray, const FColor& Color);
 
 	/** Draw Chaos' single rigid body velocity without synchronization checks. */
 	void DrawRigidBodyVelocityNoChecks(const UGeometryCollectionComponent* GeometryCollectionComponent, int32 TransformIndex, const FGeometryCollectionParticlesData& ParticlesData, const FColor& Color);
diff --git a/Engine/Source/Runtime/Experimental/GeometryCollectionEngine/Public/GeometryCollection/GeometryCollectionDebugDrawComponent.h b/Engine/Source/Runtime/Experimental/GeometryCollectionEngine/Public/GeometryCollection/GeometryCollectionDebugDrawComponent.h
--- a/Engine/Source/Runtime/Experimental/GeometryCollectionEngine/Public/GeometryCollection/GeometryCollectionDebugDrawComponent.h
+++ b/Engine/Source/Runtime/Experimental/GeometryCollectionEngine/Public/GeometryCollection/GeometryCollectionDebugDrawComponent.h
@@ -24,16 +24,6 @@ namespace Chaos { template<class T, int d> class TImplicitObject; }
 namespace Chaos { template<class T, int d> class TPBDRigidParticles; }
 #endif  // #if INCLUDE_CHAOS && GEOMETRYCOLLECTION_DEBUG_DRAW
 
-/**
-* FGeometryCollectionDebugDrawWarningMessage
-*   Empty structure used to embed a warning message in the UI through a detail customization.
-*/
-USTRUCT()
-struct FGeometryCollectionDebugDrawWarningMessage
-{
-	GENERATED_USTRUCT_BODY()
-};
-
 /**
 * UGeometryCollectionDebugDrawComponent
 *   Component adding debug drawing functionality to a GeometryCollectionActor.
@@ -46,195 +36,11 @@ class GEOMETRYCOLLECTIONENGINE_API UGeometryCollectionDebugDrawComponent : publi
 
 public:
 
-	/** Warning message to explain that the debug draw properties have no effect until starting playing/simulating. */
-	UPROPERTY(EditAnywhere, Category = "Debug Draw")
-	FGeometryCollectionDebugDrawWarningMessage WarningMessage;
-
-	/** Enable the visualization of the rigid body ids. */
-	UPROPERTY(EditAnywhere, Category = "Debug Draw|Rigid Body")
-	bool bShowRigidBodyIds;
-
-	/** Enable rigid body transform visualization. */
-	UPROPERTY(EditAnywhere, Category = "Debug Draw|Rigid Body")
-	bool bShowRigidBodyTransforms;
-
-	/** Enable the visualization of the implicit collision primitives (excluding level sets). */
-	UPROPERTY(EditAnywhere, Category = "Debug Draw|Rigid Body")
-	bool bShowRigidBodyCollisions;
-
-	/** Enable rigid body inertia tensor box visualization. */
-	UPROPERTY(EditAnywhere, Category = "Debug Draw|Rigid Body")
-	bool bShowRigidBodyInertias;
-
-	/** Enable rigid body velocities visualization. */
-	UPROPERTY(EditAnywhere, Category = "Debug Draw|Rigid Body")
-	bool bShowRigidBodyVelocities;
-
-	/** Enable rigid body applied force and torque visualization. */
-	UPROPERTY(EditAnywhere, Category = "Debug Draw|Rigid Body")
-	bool bShowRigidBodyForces;
-
-	/** Enable the visualization of the rigid body infos. */
-	UPROPERTY(EditAnywhere, Category = "Debug Draw|Rigid Body", AdvancedDisplay)
-	bool bShowRigidBodyInfos;
-
-	/** Color used for the visualization of the rigid body ids. */
-	UPROPERTY(EditAnywhere, Category = "Debug Draw|Rigid Body Settings", meta = (HideAlphaChannel))
-	FColor RigidBodyIdColor;
-
-	/** Scale for rigid body transform visualization. */
-	UPROPERTY(EditAnywhere, Category = "Debug Draw|Rigid Body Settings", meta = (ClampMin="0.0001"))
-	float RigidBodyTransformScale;
-
-	/** Color used for collision primitives visualization. */
-	UPROPERTY(EditAnywhere, Category = "Debug Draw|Rigid Body Settings", meta = (HideAlphaChannel))
-	FColor RigidBodyCollisionColor;
-
-	/** Color used for the visualization of the rigid body inertia tensor box. */
-	UPROPERTY(EditAnywhere, Category = "Debug Draw|Rigid Body Settings", meta = (HideAlphaChannel))
-	FColor RigidBodyInertiaColor;
-
-	/** Color used for rigid body velocities visualization. */
-	UPROPERTY(EditAnywhere, Category = "Debug Draw|Rigid Body Settings", meta = (HideAlphaChannel))
-	FColor RigidBodyVelocityColor;
-
-	/** Color used for rigid body applied force and torque visualization. */
-	UPROPERTY(EditAnywhere, Category = "Debug Draw|Rigid Body Settings", meta = (HideAlphaChannel))
-	FColor RigidBodyForceColor;
-
-	/** Color used for the visualization of the rigid body infos. */
-	UPROPERTY(EditAnywhere, Category = "Debug Draw|Rigid Body Settings", meta = (HideAlphaChannel))
-	FColor RigidBodyInfoColor;
-
-	/** Enable the visualization of the transform indices. */
-	UPROPERTY(EditAnywhere, Category = "Debug Draw|Clustering")
-	bool bShowTransformIndices;
-
-	/** Enable cluster transform visualization. */
-	UPROPERTY(EditAnywhere, Category = "Debug Draw|Clustering")
-	bool bShowTransforms;
-
-	/** Enable the visualization of the cluster levels. */
-	UPROPERTY(EditAnywhere, Category = "Debug Draw|Clustering")
-	bool bShowLevels;
-
-	/** Enable the visualization of the link from the parents. */
-	UPROPERTY(EditAnywhere, Category = "Debug Draw|Clustering")
-	bool bShowParents;
-
-	/** Enable the visualization of the rigid clustering connectivity edges. */
-	UPROPERTY(EditAnywhere, Category = "Debug Draw|Clustering")
-	bool bShowConnectivityEdges;
-
-	/** Color used for the visualization of the transform indices. */
-	UPROPERTY(EditAnywhere, Category = "Debug Draw|Clustering Settings", meta = (HideAlphaChannel))
-	FColor TransformIndexColor;
-
-	/** Scale for cluster transform visualization. */
-	UPROPERTY(EditAnywhere, Category = "Debug Draw|Clustering Settings", meta = (ClampMin="0.0001"))
-	float TransformScale;
-
-	/** Color used for the visualization of the levels. */
-	UPROPERTY(EditAnywhere, Category = "Debug Draw|Clustering Settings", meta = (HideAlphaChannel))
-	FColor LevelColor;
-
-	/** Color used for the visualization of the link from the parents. */
-	UPROPERTY(EditAnywhere, Category = "Debug Draw|Clustering Settings", meta = (HideAlphaChannel))
-	FColor ParentColor;
-
-	/** Line thickness used for the visualization of the rigid clustering connectivity edges. */
-	UPROPERTY(EditAnywhere, Category = "Debug Draw|Clustering Settings", meta = (ClampMin="0.0001"))
-	float ConnectivityEdgeThickness;
-
-	/** Enable the visualization of the geometry indices. */
-	UPROPERTY(EditAnywhere, Category = "Debug Draw|Geometry")
-	bool bShowGeometryIndices;
-
-	/** Enable geometry transform visualization. */
-	UPROPERTY(EditAnywhere, Category = "Debug Draw|Geometry")
-	bool bShowGeometryTransforms;
-
-	/** Enable the visualization of the bounding boxes. */
-	UPROPERTY(EditAnywhere, Category = "Debug Draw|Geometry")
-	bool bShowBoundingBoxes;
-
-	/** Enable the visualization of faces. */
-	UPROPERTY(EditAnywhere, Category = "Debug Draw|Geometry")
-	bool bShowFaces;
-
-	/** Enable the visualization of the face indices. */
-	UPROPERTY(EditAnywhere, Category = "Debug Draw|Geometry")
-	bool bShowFaceIndices;
-
-	/** Enable the visualization of the face normals. */
-	UPROPERTY(EditAnywhere, Category = "Debug Draw|Geometry")
-	bool bShowFaceNormals;
-
-	/** Enable single face visualization. */
-	UPROPERTY(EditAnywhere, Category = "Debug Draw|Geometry")
-	bool bShowSingleFace;
-
-	/** Index of the single face to visualize. */
-	UPROPERTY(EditAnywhere, Category = "Debug Draw|Geometry", meta = (ClampMin="0"))
-	int32 SingleFaceIndex;
-
-	/** Enable the visualization of the vertices. */
-	UPROPERTY(EditAnywhere, Category = "Debug Draw|Geometry")
-	bool bShowVertices;
-
-	/** Enable the visualization of the vertex indices. */
-	UPROPERTY(EditAnywhere, Category = "Debug Draw|Geometry")
-	bool bShowVertexIndices;
-
-	/** Enable the visualization of the vertex normals. */
-	UPROPERTY(EditAnywhere, Category = "Debug Draw|Geometry")
-	bool bShowVertexNormals;
-
-	/** Color used for the visualization of the geometry indices. */
-	UPROPERTY(EditAnywhere, Category = "Debug Draw|Geometry Settings", meta = (HideAlphaChannel))
-	FColor GeometryIndexColor;
-
-	/** Scale for geometry transform visualization. */
-	UPROPERTY(EditAnywhere, Category = "Debug Draw|Geometry Settings", meta = (ClampMin="0.0001"))
-	float GeometryTransformScale;
-
-	/** Color used for the visualization of the bounding boxes. */
-	UPROPERTY(EditAnywhere, Category = "Debug Draw|Geometry Settings", meta = (HideAlphaChannel))
-	FColor BoundingBoxColor;
-
-	/** Color used for the visualization of the faces. */
-	UPROPERTY(EditAnywhere, Category = "Debug Draw|Geometry Settings", meta = (HideAlphaChannel))
-	FColor FaceColor;
-
-	/** Color used for the visualization of the face indices. */
-	UPROPERTY(EditAnywhere, Category = "Debug Draw|Geometry Settings", meta = (HideAlphaChannel))
-	FColor FaceIndexColor;
-
-	/** Color used for the visualization of the face normals. */
-	UPROPERTY(EditAnywhere, Category = "Debug Draw|Geometry Settings", meta = (HideAlphaChannel))
-	FColor FaceNormalColor;
-
-	/** Color used for the visualization of the single face. */
-	UPROPERTY(EditAnywhere, Category = "Debug Draw|Geometry Settings", meta = (HideAlphaChannel))
-	FColor SingleFaceColor;
-
-	/** Color used for the visualization of the vertices. */
-	UPROPERTY(EditAnywhere, Category = "Debug Draw|Geometry Settings", meta = (HideAlphaChannel))
-	FColor VertexColor;
-
-	/** Color used for the visualization of the vertex indices. */
-	UPROPERTY(EditAnywhere, Category = "Debug Draw|Geometry Settings", meta = (HideAlphaChannel))
-	FColor VertexIndexColor;
-
-	/** Color used for the visualization of the vertex normals. */
-	UPROPERTY(EditAnywhere, Category = "Debug Draw|Geometry Settings", meta = (HideAlphaChannel))
-	FColor VertexNormalColor;
-
-	/** Singleton actor, containing the debug draw properties. Automatically populated at play time unless explicitely set. */
+	/** Singleton actor, containing the debug draw properties. Automatically populated at play time unless explicitly set. */
 	UPROPERTY(EditAnywhere, Category = "Debug Draw", AdvancedDisplay)
 	AGeometryCollectionDebugDrawActor* GeometryCollectionDebugDrawActor;
 
-	/** Level Set singleton actor, containing the Render properties. Automatically populated at play time unless explicitely set. */
+	/** Level Set singleton actor, containing the Render properties. Automatically populated at play time unless explicitly set. */
 	UPROPERTY(EditAnywhere, Category = "Debug Draw", AdvancedDisplay)
 	AGeometryCollectionRenderLevelSetActor* GeometryCollectionRenderLevelSetActor;
 
@@ -244,14 +50,6 @@ public:
 	virtual void EndPlay(EEndPlayReason::Type ReasonEnd) override;
 	virtual void TickComponent(float DeltaTime, enum ELevelTick TickType, FActorComponentTickFunction *ThisTickFunction) override;
 	
-#if WITH_EDITOR
-	/** Property changed callback. Used to clamp the level set and single face index properties. */
-	virtual void PostEditChangeProperty(FPropertyChangedEvent& PropertyChangedEvent) override;
-
-	/** Some properties cab be locked depending on the value of the debug draw actor. */
-	virtual bool CanEditChange(const UProperty* InProperty) const override;
-#endif
-
 #if GEOMETRYCOLLECTION_DEBUG_DRAW
 	/** Update selection and visibility after any change in properties. Also enable/disable this component tick update. Return true if this geometry collection is selected. */
 	bool OnDebugDrawPropertiesChanged(bool bForceVisibilityUpdate);
