Fix compile error with VS2019.

#rnx

diff --git a/Engine/Source/Runtime/Experimental/Chaos/Public/Chaos/ParticleIterator.h b/Engine/Source/Runtime/Experimental/Chaos/Public/Chaos/ParticleIterator.h
--- a/Engine/Source/Runtime/Experimental/Chaos/Public/Chaos/ParticleIterator.h
+++ b/Engine/Source/Runtime/Experimental/Chaos/Public/Chaos/ParticleIterator.h
@@ -251,7 +251,7 @@ public:
 
 	TConstParticleIterator(const TConstParticleIterator& Rhs) = default;
 
-	operator bool() const { return TransientHandle.GeometryParticles; }
+	operator bool() const { return TransientHandle.GeometryParticles != nullptr; }
 	TConstParticleIterator<TSOA>& operator++()
 	{
 		RangedForValidation();
