Remove debug/test code submitted along with fix for UE-77399.

diff --git a/Engine/Source/Runtime/Experimental/GeometryCollectionEngine/Private/GeometryCollection/GeometryCollectionDebugDrawActor.cpp b/Engine/Source/Runtime/Experimental/GeometryCollectionEngine/Private/GeometryCollection/GeometryCollectionDebugDrawActor.cpp
--- a/Engine/Source/Runtime/Experimental/GeometryCollectionEngine/Private/GeometryCollection/GeometryCollectionDebugDrawActor.cpp
+++ b/Engine/Source/Runtime/Experimental/GeometryCollectionEngine/Private/GeometryCollection/GeometryCollectionDebugDrawActor.cpp
@@ -758,7 +758,7 @@ void AGeometryCollectionDebugDrawActor::DrawVertices(const TArray<FTransform>& G
 #if ENABLE_DRAW_DEBUG
 	check(GeometryCollectionComponent);
 	check(GeometryCollectionComponent->RestCollection);
-	check(TransformIndex >= INDEX_NONE);
+	check(TransformIndex >= 0);
 
 	const UWorld* const World = GetWorld();
 	check(World);
