Fix Box/Sphere Collisions.

diff --git a/Engine/Source/Runtime/Experimental/ChaosSolvers/Private/PhysicsProxy/GeometryCollectionPhysicsProxy.cpp b/Engine/Source/Runtime/Experimental/ChaosSolvers/Private/PhysicsProxy/GeometryCollectionPhysicsProxy.cpp
--- a/Engine/Source/Runtime/Experimental/ChaosSolvers/Private/PhysicsProxy/GeometryCollectionPhysicsProxy.cpp
+++ b/Engine/Source/Runtime/Experimental/ChaosSolvers/Private/PhysicsProxy/GeometryCollectionPhysicsProxy.cpp
@@ -2482,6 +2482,33 @@ void FGeometryCollectionPhysicsProxy::InitializeSharedCollisionStructures(Chaos:
 					CollectionSimplicials[ClusterTransformIdx] = TUniquePtr< FSimplicial >(
 						FCollisionStructureManager::NewSimplicial(MassSpaceParticles, *UnionMesh, CollectionImplicits[ClusterTransformIdx].Get()));
 				}
+				else if (SizeSpecificData.ImplicitType == EImplicitTypeEnum::Chaos_Implicit_Box)
+				{
+					ErrorReporter.SetPrefix(BaseErrorPrefix + " | Cluster Transform Index: " + FString::FromInt(ClusterTransformIdx));
+					CollectionImplicits[ClusterTransformIdx] = TUniquePtr<TImplicitObject<float, 3>>(
+						FCollisionStructureManager::NewImplicit(ErrorReporter, MassSpaceParticles, *UnionMesh,
+							InstanceBoundingBox, 0, 0, 0,
+							SizeSpecificData.CollisionObjectReductionPercentage, SizeSpecificData.CollisionType,
+							SizeSpecificData.ImplicitType)
+						);
+
+					CollectionSimplicials[ClusterTransformIdx] = TUniquePtr< FSimplicial >(
+						FCollisionStructureManager::NewSimplicial(MassSpaceParticles, *UnionMesh, CollectionImplicits[ClusterTransformIdx].Get()));
+
+				}
+				else if (SizeSpecificData.ImplicitType == EImplicitTypeEnum::Chaos_Implicit_Sphere)
+				{
+					ErrorReporter.SetPrefix(BaseErrorPrefix + " | Cluster Transform Index: " + FString::FromInt(ClusterTransformIdx));
+					CollectionImplicits[ClusterTransformIdx] = TUniquePtr<TImplicitObject<float, 3>>(
+						FCollisionStructureManager::NewImplicit(ErrorReporter, MassSpaceParticles, *UnionMesh,
+							InstanceBoundingBox, InstanceBoundingBox.GetExtent().GetAbsMin() / 2, 0, 0,
+							SizeSpecificData.CollisionObjectReductionPercentage, SizeSpecificData.CollisionType,
+							SizeSpecificData.ImplicitType)
+						);
+
+					CollectionSimplicials[ClusterTransformIdx] = TUniquePtr< FSimplicial >(
+						FCollisionStructureManager::NewSimplicial(MassSpaceParticles, *UnionMesh, CollectionImplicits[ClusterTransformIdx].Get()));
+				}
 				else
 				{
 					CollectionImplicits[ClusterTransformIdx].Reset();	//union so just set as null
diff --git a/Engine/Source/Runtime/Experimental/GeometryCollectionEngine/Private/GeometryCollection/DerivedDataGeometryCollectionCooker.cpp b/Engine/Source/Runtime/Experimental/GeometryCollectionEngine/Private/GeometryCollection/DerivedDataGeometryCollectionCooker.cpp
--- a/Engine/Source/Runtime/Experimental/GeometryCollectionEngine/Private/GeometryCollection/DerivedDataGeometryCollectionCooker.cpp
+++ b/Engine/Source/Runtime/Experimental/GeometryCollectionEngine/Private/GeometryCollection/DerivedDataGeometryCollectionCooker.cpp
@@ -82,7 +82,8 @@ const TCHAR* FDerivedDataGeometryCollectionCooker::GetVersionString() const
 	//return TEXT("0182FCF14245402B8A902298417ECAA3");	//Increase resolution for cluster parents
 	//return TEXT("2F4F3A038A7B41F28CA86799DD8D4300");	//No tiny faces or non-manifold edges in the hole fills + fix materials
 	//return TEXT("B5F1831D60884F9DAC2C2CBCF2ECFA12");	//simplicials
-	return TEXT("851938C6E7EA44C4B2DA0C8E600693A4");	//simplicials and inertia
+	//return TEXT("851938C6E7EA44C4B2DA0C8E600693A4");	//simplicials and inertia
+	return TEXT("E237FFBE23A54668A07C3D14AB68A7B4");	//box and sphere
 
 }
 
