Add back kinematic activation.

diff --git a/Engine/Source/Runtime/Experimental/ChaosSolvers/Private/PhysicsProxy/FieldSystemPhysicsProxy.cpp b/Engine/Source/Runtime/Experimental/ChaosSolvers/Private/PhysicsProxy/FieldSystemPhysicsProxy.cpp
--- a/Engine/Source/Runtime/Experimental/ChaosSolvers/Private/PhysicsProxy/FieldSystemPhysicsProxy.cpp
+++ b/Engine/Source/Runtime/Experimental/ChaosSolvers/Private/PhysicsProxy/FieldSystemPhysicsProxy.cpp
@@ -149,6 +149,13 @@ void FFieldSystemPhysicsProxy::FieldParameterUpdateCallback(Chaos::FPhysicsSolve
 							{
 								Particles.SetObjectState(RigidBodyIndex, Chaos::EObjectStateType::Dynamic);
 							}
+							// todo: This breaks shooting kinematic characters as dynamics takes over but enables cache interaction
+							else if (DynamicStateView[RigidBodyIndex] == (int32)EObjectStateTypeEnum::Chaos_Object_Dynamic
+								&& Particles.ObjectState(RigidBodyIndex) == Chaos::EObjectStateType::Kinematic
+								&& FLT_EPSILON < Particles.M(RigidBodyIndex))
+							{
+								Particles.SetObjectState(RigidBodyIndex, Chaos::EObjectStateType::Dynamic);
+							}
 							else if ((DynamicStateView[RigidBodyIndex] == (int32)EObjectStateTypeEnum::Chaos_Object_Kinematic)
 								&& Particles.ObjectState(RigidBodyIndex) == Chaos::EObjectStateType::Dynamic)
 							{
