Change Chaos Defaults to be a bit more stable.

diff --git a/Engine/Source/Runtime/Experimental/ChaosSolverEngine/Private/Chaos/ChaosSolverActor.cpp b/Engine/Source/Runtime/Experimental/ChaosSolverEngine/Private/Chaos/ChaosSolverActor.cpp
--- a/Engine/Source/Runtime/Experimental/ChaosSolverEngine/Private/Chaos/ChaosSolverActor.cpp
+++ b/Engine/Source/Runtime/Experimental/ChaosSolverEngine/Private/Chaos/ChaosSolverActor.cpp
@@ -194,9 +194,9 @@ static TUniquePtr<FChaosSolverActorConsoleObjects> ChaosSolverActorConsoleObject
 AChaosSolverActor::AChaosSolverActor(const FObjectInitializer& ObjectInitializer)
 	: Super(ObjectInitializer)
 	, TimeStepMultiplier(1.f)
-	, CollisionIterations(5)
-	, PushOutIterations(1)
-	, PushOutPairIterations(1)
+	, CollisionIterations(1)
+	, PushOutIterations(3)
+	, PushOutPairIterations(2)
 	, ClusterConnectionFactor(1.0)
 	, ClusterUnionConnectionType(EClusterConnectionTypeEnum::Chaos_DelaunayTriangulation)
 	, DoGenerateCollisionData(true)
