Fix mac build.

diff --git a/Engine/Source/Runtime/Experimental/Chaos/Private/Chaos/PBDRigidClustering.cpp b/Engine/Source/Runtime/Experimental/Chaos/Private/Chaos/PBDRigidClustering.cpp
--- a/Engine/Source/Runtime/Experimental/Chaos/Private/Chaos/PBDRigidClustering.cpp
+++ b/Engine/Source/Runtime/Experimental/Chaos/Private/Chaos/PBDRigidClustering.cpp
@@ -9,6 +9,8 @@
 #include "Chaos/MassProperties.h"
 #include "Chaos/PBDRigidsEvolution.h"
 #include "Chaos/PBDRigidsEvolutionPGS.h"
+#include "Chaos/PBDCollisionConstraint.h"
+#include "Chaos/PBDCollisionConstraintPGS.h"
 #include "Chaos/Sphere.h"
 #include "Chaos/UniformGrid.h"
 #include "ChaosStats.h"
@@ -1910,9 +1912,9 @@ namespace Chaos
 }
 
 using namespace Chaos;
-template class Chaos::TPBDRigidClustering<TPBDRigidsEvolutionGBF<float, 3>, TPBDCollisionConstraint<float, 3>, float, 3>;
+template class CHAOS_API Chaos::TPBDRigidClustering<TPBDRigidsEvolutionGBF<float, 3>, TPBDCollisionConstraint<float, 3>, float, 3>;
 #if CHAOS_PARTICLEHANDLE_TODO
-template class Chaos::TPBDRigidClustering<TPBDRigidsEvolutionPGS<float, 3>, TPBDCollisionConstraintPGS<float, 3>, float, 3>;
+template class CHAOS_API Chaos::TPBDRigidClustering<TPBDRigidsEvolutionPGS<float, 3>, TPBDCollisionConstraintPGS<float, 3>, float, 3>;
 #endif
 template CHAOS_API void Chaos::UpdateClusterMassProperties<float, 3>(TPBDRigidClusteredParticles<float, 3>& Particles, const TArray<uint32>& Children, const uint32 ClusterIndex, const TRigidTransform<float, 3>* ForceMassOrientation,
 	const TArrayCollectionArray<TUniquePtr<TMultiChildProxyData<float, 3>>>* MMultiChildProxyData,
diff --git a/Engine/Source/Runtime/Experimental/Chaos/Private/Chaos/PBDRigidsEvolution.cpp b/Engine/Source/Runtime/Experimental/Chaos/Private/Chaos/PBDRigidsEvolution.cpp
--- a/Engine/Source/Runtime/Experimental/Chaos/Private/Chaos/PBDRigidsEvolution.cpp
+++ b/Engine/Source/Runtime/Experimental/Chaos/Private/Chaos/PBDRigidsEvolution.cpp
@@ -37,4 +37,4 @@ namespace Chaos
 	}
 }
 
-template class Chaos::TPBDRigidsEvolutionBase<Chaos::TPBDRigidsEvolutionGBF<float, 3>, Chaos::TPBDCollisionConstraint<float,3>, float, 3>;
\ No newline at end of file
+template class CHAOS_API Chaos::TPBDRigidsEvolutionBase<Chaos::TPBDRigidsEvolutionGBF<float, 3>, Chaos::TPBDCollisionConstraint<float,3>, float, 3>;
diff --git a/Engine/Source/Runtime/Experimental/Chaos/Public/Chaos/BoundingVolumeHierarchy.h b/Engine/Source/Runtime/Experimental/Chaos/Public/Chaos/BoundingVolumeHierarchy.h
--- a/Engine/Source/Runtime/Experimental/Chaos/Public/Chaos/BoundingVolumeHierarchy.h
+++ b/Engine/Source/Runtime/Experimental/Chaos/Public/Chaos/BoundingVolumeHierarchy.h
@@ -166,6 +166,11 @@ class CHAOS_API TBoundingVolumeHierarchy final : public ISpatialAcceleration<int
 	FCriticalSection CriticalSection;
 };
 
+#if PLATFORM_MAC
+extern template class CHAOS_API Chaos::TBoundingVolumeHierarchy<Chaos::TParticles<float, 3>, TArray<int32>, float, 3>;
+extern template class CHAOS_API Chaos::TBoundingVolumeHierarchy<Chaos::TGeometryParticles<float, 3>, TArray<int32>, float, 3>;
+#endif
+
 template<class OBJECT_ARRAY, class LEAF_TYPE, class T, int d>
 FArchive& operator<<(FArchive& Ar, TBoundingVolumeHierarchy<OBJECT_ARRAY, LEAF_TYPE, T, d>& BVH)
 {
