Fix windows build after mac fix.

diff --git a/Engine/Source/Runtime/Experimental/Chaos/Public/Chaos/PBDRigidsEvolution.h b/Engine/Source/Runtime/Experimental/Chaos/Public/Chaos/PBDRigidsEvolution.h
--- a/Engine/Source/Runtime/Experimental/Chaos/Public/Chaos/PBDRigidsEvolution.h
+++ b/Engine/Source/Runtime/Experimental/Chaos/Public/Chaos/PBDRigidsEvolution.h
@@ -48,7 +48,7 @@ struct CHAOS_API FEvolutionStats
 };
 
 template<class FPBDRigidsEvolution, class FPBDCollisionConstraint, class T, int d>
-class TPBDRigidsEvolutionBase
+class CHAOS_API TPBDRigidsEvolutionBase
 {
   public:
 	typedef TFunction<void(TTransientPBDRigidParticleHandle<T,d>& Particle, const T)> FForceRule;
@@ -56,8 +56,8 @@ class TPBDRigidsEvolutionBase
 	typedef TFunction<void(const TParticleView<TPBDRigidParticles<T, d>>&, const T)> FUpdatePositionRule;
 	typedef TFunction<void(TPBDRigidParticles<T, d>&, const T, const T, const int32)> FKinematicUpdateRule;
 
-	CHAOS_API TPBDRigidsEvolutionBase(TPBDRigidsSOAs<T, d>& InParticles, int32 InNumIterations = 1);
-	CHAOS_API virtual ~TPBDRigidsEvolutionBase()
+	TPBDRigidsEvolutionBase(TPBDRigidsSOAs<T, d>& InParticles, int32 InNumIterations = 1);
+	virtual ~TPBDRigidsEvolutionBase()
 	{
 		Particles.GetParticleHandles().RemoveArray(&PhysicsMaterials);
 		Particles.GetParticleHandles().RemoveArray(&PerParticlePhysicsMaterials);
@@ -65,38 +65,38 @@ class TPBDRigidsEvolutionBase
 		Particles.GetParticleHandles().RemoveArray(&Collided);
 	}
 
-	CHAOS_API TArray<TGeometryParticleHandle<T, d>*> CreateStaticParticles(int32 NumParticles, const TGeometryParticleParameters<T, d>& Params = TGeometryParticleParameters<T, d>()) { return Particles.CreateStaticParticles(NumParticles, Params); }
-	CHAOS_API TArray<TKinematicGeometryParticleHandle<T, d>*> CreateKinematicParticles(int32 NumParticles, const TKinematicGeometryParticleParameters<T, d>& Params = TKinematicGeometryParticleParameters<T, d>()) { return Particles.CreateKinematicParticles(NumParticles, Params); }
-	CHAOS_API TArray<TPBDRigidParticleHandle<T, d>*> CreateDynamicParticles(int32 NumParticles, const TPBDRigidParticleParameters<T, d>& Params = TPBDRigidParticleParameters<T, d>()) { return Particles.CreateDynamicParticles(NumParticles, Params); }
-	CHAOS_API TArray<TPBDRigidClusteredParticleHandle<T, d>*> CreateClusteredParticles(int32 NumParticles, const TPBDRigidParticleParameters<T, d>& Params = TPBDRigidParticleParameters<T, d>()) { return Particles.CreateClusteredParticles(NumParticles, Params); }
+	TArray<TGeometryParticleHandle<T, d>*> CreateStaticParticles(int32 NumParticles, const TGeometryParticleParameters<T, d>& Params = TGeometryParticleParameters<T, d>()) { return Particles.CreateStaticParticles(NumParticles, Params); }
+	TArray<TKinematicGeometryParticleHandle<T, d>*> CreateKinematicParticles(int32 NumParticles, const TKinematicGeometryParticleParameters<T, d>& Params = TKinematicGeometryParticleParameters<T, d>()) { return Particles.CreateKinematicParticles(NumParticles, Params); }
+	TArray<TPBDRigidParticleHandle<T, d>*> CreateDynamicParticles(int32 NumParticles, const TPBDRigidParticleParameters<T, d>& Params = TPBDRigidParticleParameters<T, d>()) { return Particles.CreateDynamicParticles(NumParticles, Params); }
+	TArray<TPBDRigidClusteredParticleHandle<T, d>*> CreateClusteredParticles(int32 NumParticles, const TPBDRigidParticleParameters<T, d>& Params = TPBDRigidParticleParameters<T, d>()) { return Particles.CreateClusteredParticles(NumParticles, Params); }
 
-	CHAOS_API void AddForceFunction(FForceRule ForceFunction) { ForceRules.Add(ForceFunction); }
-	CHAOS_API void SetParticleUpdateVelocityFunction(FUpdateVelocityRule ParticleUpdate) { ParticleUpdateVelocity = ParticleUpdate; }
-	CHAOS_API void SetParticleUpdatePositionFunction(FUpdatePositionRule ParticleUpdate) { ParticleUpdatePosition = ParticleUpdate; }
+	void AddForceFunction(FForceRule ForceFunction) { ForceRules.Add(ForceFunction); }
+	void SetParticleUpdateVelocityFunction(FUpdateVelocityRule ParticleUpdate) { ParticleUpdateVelocity = ParticleUpdate; }
+	void SetParticleUpdatePositionFunction(FUpdatePositionRule ParticleUpdate) { ParticleUpdatePosition = ParticleUpdate; }
 
-	CHAOS_API TGeometryParticleHandles<T, d>& GetParticleHandles() { return Particles.GetParticleHandles(); }
-	CHAOS_API const TGeometryParticleHandles<T, d>& GetParticleHandles() const { return Particles.GetParticleHandles(); }
+	TGeometryParticleHandles<T, d>& GetParticleHandles() { return Particles.GetParticleHandles(); }
+	const TGeometryParticleHandles<T, d>& GetParticleHandles() const { return Particles.GetParticleHandles(); }
 
-	CHAOS_API TPBDRigidsSOAs<T,d>& GetParticles() { return Particles; }
-	CHAOS_API const TPBDRigidsSOAs<T, d>& GetParticles() const { return Particles; }
+	TPBDRigidsSOAs<T,d>& GetParticles() { return Particles; }
+	const TPBDRigidsSOAs<T, d>& GetParticles() const { return Particles; }
 
 	typedef TPBDConstraintGraph<T, d> FConstraintGraph;
 	typedef TPBDConstraintGraphRule<T, d> FConstraintRule;
 
-	CHAOS_API void AddConstraintRule(FConstraintRule* ConstraintRule)
+	void AddConstraintRule(FConstraintRule* ConstraintRule)
 	{
 		uint32 ContainerId = (uint32)ConstraintRules.Num();
 		ConstraintRules.Add(ConstraintRule);
 		ConstraintRule->BindToGraph(ConstraintGraph, ContainerId);
 	}
 
-	CHAOS_API void EnableParticle(TGeometryParticleHandle<T,d>* Particle, const TGeometryParticleHandle<T, d>* ParentParticle)
+	void EnableParticle(TGeometryParticleHandle<T,d>* Particle, const TGeometryParticleHandle<T, d>* ParentParticle)
 	{
 		Particles.EnableParticle(Particle);
 		ConstraintGraph.EnableParticle(Particle, ParentParticle);
 	}
 
-	CHAOS_API void DisableParticle(TGeometryParticleHandle<T,d>* Particle)
+	void DisableParticle(TGeometryParticleHandle<T,d>* Particle)
 	{
 		Particles.DisableParticle(Particle);
 		ConstraintGraph.DisableParticle(Particle);
@@ -104,14 +104,14 @@ class TPBDRigidsEvolutionBase
 		RemoveConstraints(TSet<TGeometryParticleHandle<T,d>*>({ Particle }));
 	}
 
-	CHAOS_API void DestroyParticle(TGeometryParticleHandle<T, d>* Particle)
+	void DestroyParticle(TGeometryParticleHandle<T, d>* Particle)
 	{
 		ConstraintGraph.RemoveParticle(Particle);
 		RemoveConstraints(TSet<TGeometryParticleHandle<T, d>*>({ Particle }));
 		Particles.DestroyParticle(Particle);
 	}
 
-	CHAOS_API void DisableParticles(const TSet<TGeometryParticleHandle<T,d>*>& InParticles)
+	void DisableParticles(const TSet<TGeometryParticleHandle<T,d>*>& InParticles)
 	{
 		for (TGeometryParticleHandle<T, d>* Particle : InParticles)
 		{
@@ -123,7 +123,7 @@ class TPBDRigidsEvolutionBase
 		RemoveConstraints(InParticles);
 	}
 
-	CHAOS_API void WakeIsland(const int32 Island)
+	void WakeIsland(const int32 Island)
 	{
 		ConstraintGraph.WakeIsland(Island);
 		//Update Particles SOAs
@@ -134,7 +134,7 @@ class TPBDRigidsEvolutionBase
 	}
 
 	// @todo(ccaulfield): Remove the uint version
-	CHAOS_API void RemoveConstraints(const TSet<TGeometryParticleHandle<T, d>*>& RemovedParticles)
+	void RemoveConstraints(const TSet<TGeometryParticleHandle<T, d>*>& RemovedParticles)
 	{
 		for (FConstraintRule* ConstraintRule : ConstraintRules)
 		{
@@ -146,15 +146,15 @@ class TPBDRigidsEvolutionBase
 	const auto& GetActiveClusteredArray() const { return Particles.GetActiveClusteredArray(); }
 	const auto& GetNonDisabledClusteredArray() const { return Particles.GetNonDisabledClusteredArray(); }
 
-	CHAOS_API TSerializablePtr<TChaosPhysicsMaterial<T>> GetPhysicsMaterial(const TGeometryParticleHandle<T, d>* Particle) const { return Particle->AuxilaryValue(PhysicsMaterials); }
-	CHAOS_API void SetPhysicsMaterial(TGeometryParticleHandle<T,d>* Particle, TSerializablePtr<TChaosPhysicsMaterial<T>> InMaterial)
+	TSerializablePtr<TChaosPhysicsMaterial<T>> GetPhysicsMaterial(const TGeometryParticleHandle<T, d>* Particle) const { return Particle->AuxilaryValue(PhysicsMaterials); }
+	void SetPhysicsMaterial(TGeometryParticleHandle<T,d>* Particle, TSerializablePtr<TChaosPhysicsMaterial<T>> InMaterial)
 	{
 		check(!Particle->AuxilaryValue(PerParticlePhysicsMaterials)); //shouldn't be setting non unique material if a unique one already exists
 		Particle->AuxilaryValue(PhysicsMaterials) = InMaterial;
 	}
 
-	CHAOS_API const TArray<TGeometryParticleHandle<T,d>*>& GetIslandParticles(const int32 Island) const { return ConstraintGraph.GetIslandParticles(Island); }
-	CHAOS_API int32 NumIslands() const { return ConstraintGraph.NumIslands(); }
+	const TArray<TGeometryParticleHandle<T,d>*>& GetIslandParticles(const int32 Island) const { return ConstraintGraph.GetIslandParticles(Island); }
+	int32 NumIslands() const { return ConstraintGraph.NumIslands(); }
 
 	void InitializeAccelerationStructures()
 	{
