[NO NON-PII COMMIT MESSAGE]

diff --git a/Engine/Source/Runtime/Experimental/ChaosSolverEngine/Private/Chaos/ChaosSolverActor.cpp b/Engine/Source/Runtime/Experimental/ChaosSolverEngine/Private/Chaos/ChaosSolverActor.cpp
--- a/Engine/Source/Runtime/Experimental/ChaosSolverEngine/Private/Chaos/ChaosSolverActor.cpp
+++ b/Engine/Source/Runtime/Experimental/ChaosSolverEngine/Private/Chaos/ChaosSolverActor.cpp
@@ -226,20 +226,20 @@ AChaosSolverActor::AChaosSolverActor(const FObjectInitializer& ObjectInitializer
 	struct FConstructorStatics
 	{
 		// A helper class object we use to find target UTexture2D object in resource package
-		ConstructorHelpers::FObjectFinderOptional<UTexture2D> NoteTextureObject;
+		ConstructorHelpers::FObjectFinderOptional<UTexture2D> SolverTextureObject;
 
 		// Icon sprite category name
-		FName ID_Notes;
+		FName ID_Solver;
 
 		// Icon sprite display name
-		FText NAME_Notes;
+		FText NAME_Solver;
 
 		FConstructorStatics()
 			// Use helper class object to find the texture
-			// "/Engine/EditorResources/S_Note" is resource path
-			: NoteTextureObject(TEXT("/Engine/EditorResources/S_Note"))
-			, ID_Notes(TEXT("Notes"))
-			, NAME_Notes(NSLOCTEXT("SpriteCategory", "Notes", "Notes"))
+			// "/Engine/EditorResources/S_Solver" is resource path
+			: SolverTextureObject(TEXT("/Engine/EditorResources/S_Solver.S_Solver"))
+			, ID_Solver(TEXT("Solver"))
+			, NAME_Solver(NSLOCTEXT("SpriteCategory", "Solver", "Solver"))
 		{
 		}
 	};
@@ -254,9 +254,9 @@ AChaosSolverActor::AChaosSolverActor(const FObjectInitializer& ObjectInitializer
 	SpriteComponent = ObjectInitializer.CreateEditorOnlyDefaultSubobject<UBillboardComponent>(this, TEXT("Sprite"));
 	if (SpriteComponent)
 	{
-		SpriteComponent->Sprite = ConstructorStatics.NoteTextureObject.Get();		// Get the sprite texture from helper class object
-		SpriteComponent->SpriteInfo.Category = ConstructorStatics.ID_Notes;		// Assign sprite category name
-		SpriteComponent->SpriteInfo.DisplayName = ConstructorStatics.NAME_Notes;	// Assign sprite display name
+		SpriteComponent->Sprite = ConstructorStatics.SolverTextureObject.Get();		// Get the sprite texture from helper class object
+		SpriteComponent->SpriteInfo.Category = ConstructorStatics.ID_Solver;		// Assign sprite category name
+		SpriteComponent->SpriteInfo.DisplayName = ConstructorStatics.NAME_Solver;	// Assign sprite display name
 		SpriteComponent->AttachToComponent(RootComponent, FAttachmentTransformRules::KeepRelativeTransform);
 		SpriteComponent->Mobility = EComponentMobility::Static;
 	}
