Add a parameter to chance max particles.

diff --git a/Engine/Source/Runtime/Experimental/ChaosSolvers/Private/PhysicsProxy/GeometryCollectionPhysicsProxy.cpp b/Engine/Source/Runtime/Experimental/ChaosSolvers/Private/PhysicsProxy/GeometryCollectionPhysicsProxy.cpp
--- a/Engine/Source/Runtime/Experimental/ChaosSolvers/Private/PhysicsProxy/GeometryCollectionPhysicsProxy.cpp
+++ b/Engine/Source/Runtime/Experimental/ChaosSolvers/Private/PhysicsProxy/GeometryCollectionPhysicsProxy.cpp
@@ -2480,7 +2480,8 @@ void FGeometryCollectionPhysicsProxy::InitializeSharedCollisionStructures(Chaos:
 						);
 
 					CollectionSimplicials[ClusterTransformIdx] = TUniquePtr< FSimplicial >(
-						FCollisionStructureManager::NewSimplicial(MassSpaceParticles, *UnionMesh, CollectionImplicits[ClusterTransformIdx].Get()));
+						FCollisionStructureManager::NewSimplicial(MassSpaceParticles, *UnionMesh, CollectionImplicits[ClusterTransformIdx].Get(),
+						SharedParams.MaximumCollisionParticleCount));
 				}
 				else if (SizeSpecificData.ImplicitType == EImplicitTypeEnum::Chaos_Implicit_Box)
 				{
@@ -2493,7 +2494,8 @@ void FGeometryCollectionPhysicsProxy::InitializeSharedCollisionStructures(Chaos:
 						);
 
 					CollectionSimplicials[ClusterTransformIdx] = TUniquePtr< FSimplicial >(
-						FCollisionStructureManager::NewSimplicial(MassSpaceParticles, *UnionMesh, CollectionImplicits[ClusterTransformIdx].Get()));
+						FCollisionStructureManager::NewSimplicial(MassSpaceParticles, *UnionMesh, CollectionImplicits[ClusterTransformIdx].Get(),
+						SharedParams.MaximumCollisionParticleCount));
 
 				}
 				else if (SizeSpecificData.ImplicitType == EImplicitTypeEnum::Chaos_Implicit_Sphere)
@@ -2507,7 +2509,8 @@ void FGeometryCollectionPhysicsProxy::InitializeSharedCollisionStructures(Chaos:
 						);
 
 					CollectionSimplicials[ClusterTransformIdx] = TUniquePtr< FSimplicial >(
-						FCollisionStructureManager::NewSimplicial(MassSpaceParticles, *UnionMesh, CollectionImplicits[ClusterTransformIdx].Get()));
+						FCollisionStructureManager::NewSimplicial(MassSpaceParticles, *UnionMesh, CollectionImplicits[ClusterTransformIdx].Get(),
+						SharedParams.MaximumCollisionParticleCount));
 				}
 				else
 				{
diff --git a/Engine/Source/Runtime/Experimental/GeometryCollectionEngine/Private/GeometryCollection/GeometryCollectionObject.cpp b/Engine/Source/Runtime/Experimental/GeometryCollectionEngine/Private/GeometryCollection/GeometryCollectionObject.cpp
--- a/Engine/Source/Runtime/Experimental/GeometryCollectionEngine/Private/GeometryCollection/GeometryCollectionObject.cpp
+++ b/Engine/Source/Runtime/Experimental/GeometryCollectionEngine/Private/GeometryCollection/GeometryCollectionObject.cpp
@@ -37,6 +37,7 @@ UGeometryCollection::UGeometryCollection(const FObjectInitializer& ObjectInitial
 	, Mass(1.0f)
 	, MinimumMassClamp(0.1f)
 	, CollisionParticlesFraction(1.0f)
+	, MaximumCollisionParticles(60)
 	, EnableRemovePiecesOnFracture(false)
 	, GeometryCollection(new FGeometryCollection())
 {
@@ -54,6 +55,7 @@ FGeometryCollectionSizeSpecificData::FGeometryCollectionSizeSpecificData()
 	, MaxClusterLevelSetResolution(50)
 	, CollisionObjectReductionPercentage(0.0f)
 	, CollisionParticlesFraction(1.0f)
+	, MaximumCollisionParticles(60)
 {
 }
 
@@ -68,6 +70,7 @@ void FillSharedSimulationSizeSpecificData(FSharedSimulationSizeSpecificData& ToD
 	ToData.MaxClusterLevelSetResolution = FromData.MaxClusterLevelSetResolution;
 	ToData.CollisionObjectReductionPercentage = FromData.CollisionObjectReductionPercentage;
 	ToData.CollisionParticlesFraction = FromData.CollisionParticlesFraction;
+	ToData.MaximumCollisionParticles = FromData.MaximumCollisionParticles;
 }
 
 
@@ -86,6 +89,7 @@ void UGeometryCollection::GetSharedSimulationParams(FSharedSimulationParameters&
 	OutParams.bMassAsDensity = bMassAsDensity;
 	OutParams.Mass = bMassAsDensity ? KgM3ToKgCm3(Mass) : Mass;	//todo(ocohen): we still have the solver working in old units. This is mainly to fix ui issues. Long term need to normalize units for best precision
 	OutParams.MinimumMassClamp = MinimumMassClamp;
+	OutParams.MaximumCollisionParticleCount = MaximumCollisionParticles;
 
 	FGeometryCollectionSizeSpecificData InfSize;
 	InfSize.MaxSize = FLT_MAX;
@@ -97,6 +101,7 @@ void UGeometryCollection::GetSharedSimulationParams(FSharedSimulationParameters&
 	InfSize.MaxClusterLevelSetResolution = MaxClusterLevelSetResolution;
 	InfSize.CollisionObjectReductionPercentage = CollisionObjectReductionPercentage;
 	InfSize.CollisionParticlesFraction = CollisionParticlesFraction;
+	InfSize.MaximumCollisionParticles = MaximumCollisionParticles;
 
 	OutParams.SizeSpecificData.SetNum(SizeSpecificData.Num() + 1);
 	FillSharedSimulationSizeSpecificData(OutParams.SizeSpecificData[0], InfSize);
diff --git a/Engine/Source/Runtime/Experimental/GeometryCollectionEngine/Public/GeometryCollection/GeometryCollectionObject.h b/Engine/Source/Runtime/Experimental/GeometryCollectionEngine/Public/GeometryCollection/GeometryCollectionObject.h
--- a/Engine/Source/Runtime/Experimental/GeometryCollectionEngine/Public/GeometryCollection/GeometryCollectionObject.h
+++ b/Engine/Source/Runtime/Experimental/GeometryCollectionEngine/Public/GeometryCollection/GeometryCollectionObject.h
@@ -73,6 +73,13 @@ struct GEOMETRYCOLLECTIONENGINE_API FGeometryCollectionSizeSpecificData
 	 */
 	UPROPERTY(EditAnywhere, Category = "ChaosPhysics|Collisions")
 	float CollisionParticlesFraction;
+
+	/**
+	 * Max number of particles.
+	 */
+	UPROPERTY(EditAnywhere, Category = "ChaosPhysics|Collisions")
+	int32 MaximumCollisionParticles;
+
 };
 
 /**
@@ -210,6 +217,12 @@ public:
 	UPROPERTY(EditAnywhere, BlueprintReadOnly, Category = "ChaosPhysics|Collisions")
 	float CollisionParticlesFraction;
 
+	/**
+	 * Max number of particles.
+	 */
+	UPROPERTY(EditAnywhere, BlueprintReadOnly, Category = "ChaosPhysics|Collisions")
+	int32 MaximumCollisionParticles;
+
 	UPROPERTY(EditAnywhere, Category = "ChaosPhysics|Collisions")
 	TArray<FGeometryCollectionSizeSpecificData> SizeSpecificData;
 	
diff --git a/Engine/Source/Runtime/Experimental/GeometryCollectionSimulationCore/Private/GeometryCollection/GeometryCollectionCollisionStructureManager.cpp b/Engine/Source/Runtime/Experimental/GeometryCollectionSimulationCore/Private/GeometryCollection/GeometryCollectionCollisionStructureManager.cpp
--- a/Engine/Source/Runtime/Experimental/GeometryCollectionSimulationCore/Private/GeometryCollection/GeometryCollectionCollisionStructureManager.cpp
+++ b/Engine/Source/Runtime/Experimental/GeometryCollectionSimulationCore/Private/GeometryCollection/GeometryCollectionCollisionStructureManager.cpp
@@ -28,14 +28,15 @@ FAutoConsoleVariableRef CVarCollisionParticlesSpatialDivision(TEXT("p.CollisionP
 int32 CollisionParticlesMin = 10;
 FAutoConsoleVariableRef CVarCollisionParticlesMin(TEXT("p.CollisionParticlesMin"), CollisionParticlesMin, TEXT("Minimum number of particles after simplicial pruning (assuming it started with more)"));
 
-int32 CollisionParticlesMax = 60;
+int32 CollisionParticlesMax = 2000;
 FAutoConsoleVariableRef CVarCollisionParticlesMax(TEXT("p.CollisionParticlesMax"), CollisionParticlesMax, TEXT("Maximum number of particles after simplicial pruning"));
 
 FCollisionStructureManager::FSimplicial*
 FCollisionStructureManager::NewSimplicial(
 	const Chaos::TParticles<float, 3>& Vertices,
 	Chaos::TTriangleMesh<float>& TriMesh,
-	const Chaos::TImplicitObject<float, 3>* Implicit)
+	const Chaos::TImplicitObject<float, 3>* Implicit,
+	int32 CollisionParticlesMaxInput)
 {
 	FCollisionStructureManager::FSimplicial * Simplicial = new FCollisionStructureManager::FSimplicial();
 	if (Implicit || Vertices.Size())
@@ -47,7 +48,8 @@ FCollisionStructureManager::NewSimplicial(
 		TArray<Chaos::TVector<float, 3>> OutsideVertices;
 
 		bool bFullCopy = true;
-		if (bCollisionParticlesUseImplicitCulling!=0 && Implicit && Indices.Num()>CollisionParticlesMax)
+		int32 LocalCollisionParticlesMax = CollisionParticlesMaxInput > 0 ? FMath::Min(CollisionParticlesMaxInput, CollisionParticlesMax) : CollisionParticlesMax;
+		if (bCollisionParticlesUseImplicitCulling!=0 && Implicit && Indices.Num()>LocalCollisionParticlesMax)
 		{
 			Extent = Implicit->HasBoundingBox() ? Implicit->BoundingBox().Extents().Size() : 1.f;
 
@@ -68,7 +70,7 @@ FCollisionStructureManager::NewSimplicial(
 			}
 			OutsideVertices.SetNum(LSVCounter);
 
-			if (OutsideVertices.Num() > CollisionParticlesMax)
+			if (OutsideVertices.Num() > LocalCollisionParticlesMax)
 				bFullCopy = false;
 		}
 		
@@ -91,7 +93,7 @@ FCollisionStructureManager::NewSimplicial(
 
 		float SnapThreshold = Extent / float(CollisionParticlesSpatialDivision);
 		OutsideVertices = Chaos::CleanCollisionParticles(OutsideVertices, SnapThreshold);
-		int32 NumParticles = (OutsideVertices.Num() > CollisionParticlesMax) ? CollisionParticlesMax : OutsideVertices.Num();
+		int32 NumParticles = (OutsideVertices.Num() > LocalCollisionParticlesMax) ? LocalCollisionParticlesMax : OutsideVertices.Num();
 
 		if (NumParticles)
 		{
diff --git a/Engine/Source/Runtime/Experimental/GeometryCollectionSimulationCore/Public/GeometryCollection/GeometryCollectionCollisionStructureManager.h b/Engine/Source/Runtime/Experimental/GeometryCollectionSimulationCore/Public/GeometryCollection/GeometryCollectionCollisionStructureManager.h
--- a/Engine/Source/Runtime/Experimental/GeometryCollectionSimulationCore/Public/GeometryCollection/GeometryCollectionCollisionStructureManager.h
+++ b/Engine/Source/Runtime/Experimental/GeometryCollectionSimulationCore/Public/GeometryCollection/GeometryCollectionCollisionStructureManager.h
@@ -35,7 +35,8 @@ public:
 	static FSimplicial* NewSimplicial(
 		const Chaos::TParticles<float, 3>& Vertices,
 		Chaos::TTriangleMesh<float>& TriMesh,
-		const Chaos::TImplicitObject<float, 3>* Implicit
+		const Chaos::TImplicitObject<float, 3>* Implicit,
+		int32 CollisionParticlesMaxInput
 	);
 	
 	
diff --git a/Engine/Source/Runtime/Experimental/GeometryCollectionSimulationCore/Public/GeometryCollection/GeometryCollectionSimulationCoreTypes.h b/Engine/Source/Runtime/Experimental/GeometryCollectionSimulationCore/Public/GeometryCollection/GeometryCollectionSimulationCoreTypes.h
--- a/Engine/Source/Runtime/Experimental/GeometryCollectionSimulationCore/Public/GeometryCollection/GeometryCollectionSimulationCoreTypes.h
+++ b/Engine/Source/Runtime/Experimental/GeometryCollectionSimulationCore/Public/GeometryCollection/GeometryCollectionSimulationCoreTypes.h
@@ -23,6 +23,7 @@ struct FSharedSimulationSizeSpecificData
 		, MaxClusterLevelSetResolution(50)
 		, CollisionObjectReductionPercentage(0.f)
 		, CollisionParticlesFraction(1.f)
+		, MaximumCollisionParticles(60)
 	{
 	}
 
@@ -35,6 +36,7 @@ struct FSharedSimulationSizeSpecificData
 	int32 MaxClusterLevelSetResolution;
 	float CollisionObjectReductionPercentage;
 	float CollisionParticlesFraction;
+	int32 MaximumCollisionParticles;
 
 	bool operator<(const FSharedSimulationSizeSpecificData& Rhs) const { return MaxSize < Rhs.MaxSize; }
 };
@@ -55,6 +57,7 @@ struct FSharedSimulationParameters
 	: bMassAsDensity(false)
 	, Mass(1.0)
 	, MinimumMassClamp(0.1)
+	, MaximumCollisionParticleCount(60)
 	{
 		SizeSpecificData.AddDefaulted();
 	}
@@ -68,10 +71,12 @@ struct FSharedSimulationParameters
 		,bool InMassAsDensity
 		,float InMass
 		,float InMinimumMassClamp
-		,float InCollisionParticlesFraction)
+		,float InCollisionParticlesFraction
+		,int32 InMaximumCollisionParticleCount)
 	: bMassAsDensity(InMassAsDensity)
 	, Mass(InMass)
 	, MinimumMassClamp(InMinimumMassClamp)
+	, MaximumCollisionParticleCount(InMaximumCollisionParticleCount)
 	{
 		SizeSpecificData.AddDefaulted();
 		SizeSpecificData[0].CollisionType = InCollisionType;
@@ -81,6 +86,7 @@ struct FSharedSimulationParameters
 		SizeSpecificData[0].MinClusterLevelSetResolution = InMinClusterLevelSetResolution;
 		SizeSpecificData[0].MaxClusterLevelSetResolution = InMaxClusterLevelSetResolution;
 		SizeSpecificData[0].CollisionParticlesFraction = InCollisionParticlesFraction;
+		SizeSpecificData[0].MaximumCollisionParticles = InMaximumCollisionParticleCount;
 	}
 
 	bool bMassAsDensity;
@@ -88,6 +94,7 @@ struct FSharedSimulationParameters
 	float MinimumMassClamp;
 	TArray<FSharedSimulationSizeSpecificData> SizeSpecificData;
 	TArray<int32> RemoveOnFractureIndices;
+	int32 MaximumCollisionParticleCount;
 };
 
 struct FCollisionDataSimulationParameters
