Fix crash with record and play.

diff --git a/Engine/Source/Runtime/Experimental/ChaosSolvers/Private/PhysicsProxy/GeometryCollectionPhysicsProxy.cpp b/Engine/Source/Runtime/Experimental/ChaosSolvers/Private/PhysicsProxy/GeometryCollectionPhysicsProxy.cpp
--- a/Engine/Source/Runtime/Experimental/ChaosSolvers/Private/PhysicsProxy/GeometryCollectionPhysicsProxy.cpp
+++ b/Engine/Source/Runtime/Experimental/ChaosSolvers/Private/PhysicsProxy/GeometryCollectionPhysicsProxy.cpp
@@ -571,7 +571,7 @@ void FGeometryCollectionPhysicsProxy::StartFrameCallback(const float Dt, const f
 							{
 								// Check if the particle is still kinematic
 								if (RecordedFrame->Collisions[Idx].ParticleIndex < 0 ||
-									(RecordedFrame->Collisions[Idx].ParticleIndex >= 0 && Particles.ObjectState(RecordedFrame->Collisions[Idx].ParticleIndex) == Chaos::EObjectStateType::Kinematic))
+									(RecordedFrame->Collisions[Idx].ParticleIndex >= 0 && RecordedFrame->Collisions[Idx].ParticleIndex < static_cast<int32>(Particles.Size()) && Particles.ObjectState(RecordedFrame->Collisions[Idx].ParticleIndex) == Chaos::EObjectStateType::Kinematic))
 								{
 									const int32 NewIdx = AllCollisionsDataArray.Add(Chaos::TCollisionData<float, 3>());
 									Chaos::TCollisionData<float, 3>& AllCollisionsDataArrayItem = AllCollisionsDataArray[NewIdx];
@@ -615,7 +615,7 @@ void FGeometryCollectionPhysicsProxy::StartFrameCallback(const float Dt, const f
 							{
 								// Check if the particle is still kinematic
 								if (RecordedFrame->Breakings[Idx].ParticleIndex < 0 ||
-									(RecordedFrame->Breakings[Idx].ParticleIndex >= 0 && Particles.ObjectState(RecordedFrame->Breakings[Idx].ParticleIndex) == Chaos::EObjectStateType::Kinematic))
+									(RecordedFrame->Breakings[Idx].ParticleIndex >= 0 && RecordedFrame->Breakings[Idx].ParticleIndex < static_cast<int32>(Particles.Size()) && Particles.ObjectState(RecordedFrame->Breakings[Idx].ParticleIndex) == Chaos::EObjectStateType::Kinematic))
 								{
 									const int32 NewIdx = AllBreakingsDataArray.Add(Chaos::TBreakingData<float, 3>());
 									Chaos::TBreakingData<float, 3>& AllBreakingsDataArrayItem = AllBreakingsDataArray[NewIdx];
@@ -652,7 +652,7 @@ void FGeometryCollectionPhysicsProxy::StartFrameCallback(const float Dt, const f
 							{
 								// Check if the particle is still kinematic
 								if (Trailing.ParticleIndex < 0 ||
-									(Trailing.ParticleIndex >= 0 && Particles.ObjectState(Trailing.ParticleIndex) == Chaos::EObjectStateType::Kinematic))
+									(Trailing.ParticleIndex >= 0 && Trailing.ParticleIndex < static_cast<int32>(Particles.Size()) && Particles.ObjectState(Trailing.ParticleIndex) == Chaos::EObjectStateType::Kinematic))
 								{
 									const int32 NewIdx = AllTrailingsDataArray.Add(Chaos::TTrailingData<float, 3>());
 									Chaos::TTrailingData<float, 3>& AllTrailingsDataArrayItem = AllTrailingsDataArray[NewIdx];
