Fix CIS.

diff --git a/Engine/Source/Runtime/Experimental/Chaos/Public/Chaos/ParticleIterator.h b/Engine/Source/Runtime/Experimental/Chaos/Public/Chaos/ParticleIterator.h
--- a/Engine/Source/Runtime/Experimental/Chaos/Public/Chaos/ParticleIterator.h
+++ b/Engine/Source/Runtime/Experimental/Chaos/Public/Chaos/ParticleIterator.h
@@ -251,7 +251,7 @@ public:
 
 	TConstParticleIterator(const TConstParticleIterator& Rhs) = default;
 
-	operator bool() const { return TransientHandle.GeometryParticles; }
+	operator bool() const { return TransientHandle.GeometryParticles != nullptr; }
 	TConstParticleIterator<TSOA>& operator++()
 	{
 		RangedForValidation();
diff --git a/Engine/Source/Runtime/Experimental/ChaosSolvers/Public/PBDRigidsSolver.h b/Engine/Source/Runtime/Experimental/ChaosSolvers/Public/PBDRigidsSolver.h
--- a/Engine/Source/Runtime/Experimental/ChaosSolvers/Public/PBDRigidsSolver.h
+++ b/Engine/Source/Runtime/Experimental/ChaosSolvers/Public/PBDRigidsSolver.h
@@ -179,7 +179,7 @@ namespace Chaos
 			});
 		}
 
-		bool GetNumPhysicsProxies() const {
+		int32 GetNumPhysicsProxies() const {
 			return GeometryParticlePhysicsProxies.Num() + KinematicGeometryParticlePhysicsProxies.Num() + RigidParticlePhysicsProxies.Num()
 				+ SkeletalMeshPhysicsProxies.Num() + StaticMeshPhysicsProxies.Num()
 				+ GeometryCollectionPhysicsProxies.Num() + FieldSystemPhysicsProxies.Num() ;
