Fix CIS.

diff --git a/Engine/Source/Runtime/Experimental/ChaosSolvers/Public/PBDRigidsSolver.h b/Engine/Source/Runtime/Experimental/ChaosSolvers/Public/PBDRigidsSolver.h
--- a/Engine/Source/Runtime/Experimental/ChaosSolvers/Public/PBDRigidsSolver.h
+++ b/Engine/Source/Runtime/Experimental/ChaosSolvers/Public/PBDRigidsSolver.h
@@ -179,7 +179,7 @@ namespace Chaos
 			});
 		}
 
-		bool GetNumPhysicsProxies() const {
+		int32 GetNumPhysicsProxies() const {
 			return GeometryParticlePhysicsProxies.Num() + KinematicGeometryParticlePhysicsProxies.Num() + RigidParticlePhysicsProxies.Num()
 				+ SkeletalMeshPhysicsProxies.Num() + StaticMeshPhysicsProxies.Num()
 				+ GeometryCollectionPhysicsProxies.Num() + FieldSystemPhysicsProxies.Num() ;
