Fix CIS.

diff --git a/Engine/Source/Runtime/Experimental/Chaos/Private/Chaos/PBDCollisionConstraint.cpp b/Engine/Source/Runtime/Experimental/Chaos/Private/Chaos/PBDCollisionConstraint.cpp
--- a/Engine/Source/Runtime/Experimental/Chaos/Private/Chaos/PBDCollisionConstraint.cpp
+++ b/Engine/Source/Runtime/Experimental/Chaos/Private/Chaos/PBDCollisionConstraint.cpp
@@ -19,8 +19,19 @@
 #include "Chaos/PBDRigidsSOAs.h"
 
 #if INTEL_ISPC
+
+#if USING_CODE_ANALYSIS
+MSVC_PRAGMA(warning(push))
+MSVC_PRAGMA(warning(disable : ALL_CODE_ANALYSIS_WARNINGS))
+#endif    // USING_CODE_ANALYSIS
+
 #include "PBDCollisionConstraint.ispc.generated.h"
-#endif
+
+#if USING_CODE_ANALYSIS
+MSVC_PRAGMA(warning(pop))
+#endif    // USING_CODE_ANALYSIS
+
+#endif    // INTEL_ISPC
 
 int32 CollisionParticlesBVHDepth = 4;
 FAutoConsoleVariableRef CVarCollisionParticlesBVHDepth(TEXT("p.CollisionParticlesBVHDepth"), CollisionParticlesBVHDepth, TEXT("The maximum depth for collision particles bvh"));
diff --git a/Engine/Source/Runtime/Experimental/GeometryCollectionEngine/Private/GeometryCollection/GeometryCollectionComponent.cpp b/Engine/Source/Runtime/Experimental/GeometryCollectionEngine/Private/GeometryCollection/GeometryCollectionComponent.cpp
--- a/Engine/Source/Runtime/Experimental/GeometryCollectionEngine/Private/GeometryCollection/GeometryCollectionComponent.cpp
+++ b/Engine/Source/Runtime/Experimental/GeometryCollectionEngine/Private/GeometryCollection/GeometryCollectionComponent.cpp
@@ -38,7 +38,18 @@
 #endif //!(UE_BUILD_SHIPPING || UE_BUILD_TEST)
 
 #if INTEL_ISPC
+
+#if USING_CODE_ANALYSIS
+MSVC_PRAGMA(warning(push))
+MSVC_PRAGMA(warning(disable : ALL_CODE_ANALYSIS_WARNINGS))
+#endif    // USING_CODE_ANALYSIS
+
 #include "GeometryCollectionComponent.ispc.generated.h"
+
+#if USING_CODE_ANALYSIS
+MSVC_PRAGMA(warning(pop))
+#endif    // USING_CODE_ANALYSIS
+
 #endif
 
 DEFINE_LOG_CATEGORY_STATIC(UGCC_LOG, Error, All);
