Add a zero initalizer to EObjectsStateTypeEnum in chaos

#rnx

diff --git a/Engine/Source/Runtime/Experimental/GeometryCollectionSimulationCore/Public/GeometryCollection/GeometryCollectionSimulationTypes.h b/Engine/Source/Runtime/Experimental/GeometryCollectionSimulationCore/Public/GeometryCollection/GeometryCollectionSimulationTypes.h
--- a/Engine/Source/Runtime/Experimental/GeometryCollectionSimulationCore/Public/GeometryCollection/GeometryCollectionSimulationTypes.h
+++ b/Engine/Source/Runtime/Experimental/GeometryCollectionSimulationCore/Public/GeometryCollection/GeometryCollectionSimulationTypes.h
@@ -30,6 +30,7 @@ enum class EImplicitTypeEnum : uint8
 UENUM(BlueprintType)
 enum class EObjectStateTypeEnum : uint8
 {
+	Chaos_NONE = 0 UMETA(Hidden),
 	Chaos_Object_Sleeping  = 1 /*Chaos::EObjectStateType::Sleeping*/   UMETA(DisplayName = "Sleeping"),
 	Chaos_Object_Kinematic = 2 /*Chaos::EObjectStateType::Kinematic*/  UMETA(DisplayName = "Kinematic"),
 	Chaos_Object_Static = 3    /*Chaos::EObjectStateType::Static*/     UMETA(DisplayName = "Static"),
