Fix for shadow warnging.

diff --git a/Engine/Source/Runtime/Experimental/Chaos/Public/Chaos/ChaosPerfTest.h b/Engine/Source/Runtime/Experimental/Chaos/Public/Chaos/ChaosPerfTest.h
--- a/Engine/Source/Runtime/Experimental/Chaos/Public/Chaos/ChaosPerfTest.h
+++ b/Engine/Source/Runtime/Experimental/Chaos/Public/Chaos/ChaosPerfTest.h
@@ -8,7 +8,7 @@
 #endif
 
 #if CHAOS_PERF_TEST_ENABLED
-enum EChaosPerfUnits
+enum class EChaosPerfUnits
 {
 	S,
 	Ms,
@@ -30,13 +30,13 @@ public:
 
 	~FChaosScopedDurationTimeLogger()
 	{
-		static const float Multipliers[EChaosPerfUnits::Num] = { 1.f, 1000.f, 1000000.f };
-		static const TCHAR* Units[EChaosPerfUnits::Num] = { TEXT("s"), TEXT("ms"), TEXT("us") };
+		static const float Multipliers[static_cast<int>(EChaosPerfUnits::Num)] = { 1.f, 1000.f, 1000000.f };
+		static const TCHAR* Units[static_cast<int>(EChaosPerfUnits::Num)] = { TEXT("s"), TEXT("ms"), TEXT("us") };
 		Timer.Stop();
 		if (GlobalLabel)
 		{
 			//Device->Logf(TEXT("%s - %s: %4.fus"), GlobalLabel, Label, Accumulator * 1000000.f);
-			Device->Logf(TEXT("%s - %s: %f%s"), GlobalLabel, Label, Accumulator * Multipliers[GlobalUnits], Units[GlobalUnits]);
+			Device->Logf(TEXT("%s - %s: %f%s"), GlobalLabel, Label, Accumulator * Multipliers[static_cast<int>(GlobalUnits)], Units[static_cast<int>(GlobalUnits)]);
 		}
 	}
 
