Fix overload const correctness.

#jire ue-80234

diff --git a/Engine/Source/Runtime/Experimental/ChaosSolverEngine/Private/Chaos/ChaosSolverActor.cpp b/Engine/Source/Runtime/Experimental/ChaosSolverEngine/Private/Chaos/ChaosSolverActor.cpp
--- a/Engine/Source/Runtime/Experimental/ChaosSolverEngine/Private/Chaos/ChaosSolverActor.cpp
+++ b/Engine/Source/Runtime/Experimental/ChaosSolverEngine/Private/Chaos/ChaosSolverActor.cpp
@@ -337,7 +337,7 @@ void AChaosSolverActor::BeginPlay()
 #endif  // #if TODO_REIMPLEMENT_DEBUG_SUBSTEP
 }
 
-void AChaosSolverActor::EndPlay(EEndPlayReason::Type ReasonEnd)
+void AChaosSolverActor::EndPlay(const EEndPlayReason::Type ReasonEnd)
 {
 #if INCLUDE_CHAOS
 	Chaos::IDispatcher* PhysDispatcher = PhysScene->GetDispatcher();
diff --git a/Engine/Source/Runtime/Experimental/ChaosSolverEngine/Public/Chaos/ChaosSolverActor.h b/Engine/Source/Runtime/Experimental/ChaosSolverEngine/Public/Chaos/ChaosSolverActor.h
--- a/Engine/Source/Runtime/Experimental/ChaosSolverEngine/Public/Chaos/ChaosSolverActor.h
+++ b/Engine/Source/Runtime/Experimental/ChaosSolverEngine/Public/Chaos/ChaosSolverActor.h
@@ -192,7 +192,7 @@ public:
 	virtual void PostRegisterAllComponents() override;
 	
 	virtual void BeginPlay() override;
-	virtual void EndPlay(EEndPlayReason::Type ReasonEnd) override;
+	virtual void EndPlay(const EEndPlayReason::Type ReasonEnd) override;
 
 private:
 #if INCLUDE_CHAOS
