Fix endplay signature

#jire ue-80235

diff --git a/Engine/Source/Runtime/Experimental/GeometryCollectionEngine/Private/GeometryCollection/GeometryCollectionComponent.cpp b/Engine/Source/Runtime/Experimental/GeometryCollectionEngine/Private/GeometryCollection/GeometryCollectionComponent.cpp
--- a/Engine/Source/Runtime/Experimental/GeometryCollectionEngine/Private/GeometryCollection/GeometryCollectionComponent.cpp
+++ b/Engine/Source/Runtime/Experimental/GeometryCollectionEngine/Private/GeometryCollection/GeometryCollectionComponent.cpp
@@ -243,7 +243,7 @@ void UGeometryCollectionComponent::BeginPlay()
 }
 
 
-void UGeometryCollectionComponent::EndPlay(EEndPlayReason::Type ReasonEnd)
+void UGeometryCollectionComponent::EndPlay(const EEndPlayReason::Type ReasonEnd)
 {
 #if WITH_EDITOR && WITH_EDITORONLY_DATA
 	// Track our editor component if needed for syncing simulations back from PIE on shutdown
diff --git a/Engine/Source/Runtime/Experimental/GeometryCollectionEngine/Public/GeometryCollection/GeometryCollectionComponent.h b/Engine/Source/Runtime/Experimental/GeometryCollectionEngine/Public/GeometryCollection/GeometryCollectionComponent.h
--- a/Engine/Source/Runtime/Experimental/GeometryCollectionEngine/Public/GeometryCollection/GeometryCollectionComponent.h
+++ b/Engine/Source/Runtime/Experimental/GeometryCollectionEngine/Public/GeometryCollection/GeometryCollectionComponent.h
@@ -275,7 +275,7 @@ public:
 	virtual void SendRenderDynamicData_Concurrent() override;
 	FORCEINLINE void SetRenderStateDirty() { bRenderStateDirty = true; }
 	virtual void BeginPlay() override;
-	virtual void EndPlay(EEndPlayReason::Type ReasonEnd) override;
+	virtual void EndPlay(const EEndPlayReason::Type ReasonEnd) override;
 
 	//~ Begin UActorComponent Interface. 
 
