Speculative fix for Incremental UE4Editor Mac compile issue
#rnx

diff --git a/Engine/Source/Runtime/Experimental/Chaos/Private/Chaos/KinematicGeometryParticles.cpp b/Engine/Source/Runtime/Experimental/Chaos/Private/Chaos/KinematicGeometryParticles.cpp
--- a/Engine/Source/Runtime/Experimental/Chaos/Private/Chaos/KinematicGeometryParticles.cpp
+++ b/Engine/Source/Runtime/Experimental/Chaos/Private/Chaos/KinematicGeometryParticles.cpp
@@ -21,8 +21,8 @@ namespace Chaos
 	}
 }
 
-template class Chaos::TKinematicGeometryParticlesImp<float, 3, Chaos::EGeometryParticlesSimType::RigidBodySim>;
-template class Chaos::TKinematicGeometryParticlesImp<float, 3, Chaos::EGeometryParticlesSimType::Other>;
+template class CHAOS_API Chaos::TKinematicGeometryParticlesImp<float, 3, Chaos::EGeometryParticlesSimType::RigidBodySim>;
+template class CHAOS_API Chaos::TKinematicGeometryParticlesImp<float, 3, Chaos::EGeometryParticlesSimType::Other>;
 
 // Re-enable when double precision is able to compile
 //template Chaos::TKinematicGeometryParticlesImp<double, 3>;
\ No newline at end of file
