Copying //UE4/Dev-Physics to Dev-Main (//UE4/Dev-Main) [at] 9098278

diff --git a/Engine/Source/Runtime/Engine/Public/MaterialCompiler.h b/Engine/Source/Runtime/Engine/Public/MaterialCompiler.h
--- a/Engine/Source/Runtime/Engine/Public/MaterialCompiler.h
+++ b/Engine/Source/Runtime/Engine/Public/MaterialCompiler.h
@@ -352,6 +352,11 @@ public:
 	virtual int32 SkyAtmosphereViewLuminance() = 0;
 	virtual int32 SkyAtmosphereAerialPerspective(int32 WorldPosition) = 0;
 	virtual int32 SkyAtmosphereDistantLightScatteredLuminance() = 0;
+	virtual int32 GetHairUV() = 0;
+	virtual int32 GetHairDimensions() = 0;
+	virtual int32 GetHairSeed() = 0;
+	virtual int32 GetHairTangent() = 0;
+	virtual int32 GetHairRootUV() = 0;
 	virtual int32 CustomPrimitiveData(int32 OutputIndex, EMaterialValueType Type) = 0;
 	virtual int32 ShadingModel(EMaterialShadingModel InSelectedShadingModel) = 0;
 
@@ -605,6 +610,12 @@ public:
 		return Compiler->DepthOfFieldFunction(Depth, FunctionValueIndex);
 	}
 
+	virtual int32 GetHairUV() override { return Compiler->GetHairUV(); }
+	virtual int32 GetHairDimensions() override { return Compiler->GetHairDimensions(); }
+	virtual int32 GetHairSeed() override { return Compiler->GetHairSeed(); }
+	virtual int32 GetHairTangent() override { return Compiler->GetHairTangent(); }
+	virtual int32 GetHairRootUV() override { return Compiler->GetHairRootUV(); }
+
 	virtual int32 RotateScaleOffsetTexCoords(int32 TexCoordCodeIndex, int32 RotationScale, int32 Offset) override
 	{
 		return Compiler->RotateScaleOffsetTexCoords(TexCoordCodeIndex, RotationScale, Offset);
diff --git a/Engine/Source/Runtime/Engine/Public/MaterialShared.h b/Engine/Source/Runtime/Engine/Public/MaterialShared.h
--- a/Engine/Source/Runtime/Engine/Public/MaterialShared.h
+++ b/Engine/Source/Runtime/Engine/Public/MaterialShared.h
@@ -1464,6 +1464,7 @@ public:
 	virtual bool IsUsedWithUI() const { return false; }
 	virtual bool IsUsedWithGeometryCache() const { return false; }
 	virtual bool IsUsedWithWater() const { return false; }
+	virtual bool IsUsedWithHairStrands() const { return false; }
 	ENGINE_API virtual enum EMaterialTessellationMode GetTessellationMode() const;
 	virtual bool IsCrackFreeDisplacementEnabled() const { return false; }
 	virtual bool IsAdaptiveTessellationEnabled() const { return false; }
@@ -2221,6 +2222,7 @@ public:
 	ENGINE_API virtual bool IsUsedWithAPEXCloth() const override;
 	ENGINE_API virtual bool IsUsedWithGeometryCache() const override;
 	ENGINE_API virtual bool IsUsedWithWater() const override;
+	ENGINE_API virtual bool IsUsedWithHairStrands() const override;
 	ENGINE_API virtual enum EMaterialTessellationMode GetTessellationMode() const override;
 	ENGINE_API virtual bool IsCrackFreeDisplacementEnabled() const override;
 	ENGINE_API virtual bool IsAdaptiveTessellationEnabled() const override;
diff --git a/Engine/Source/Runtime/Engine/Public/MeshBatch.h b/Engine/Source/Runtime/Engine/Public/MeshBatch.h
--- a/Engine/Source/Runtime/Engine/Public/MeshBatch.h
+++ b/Engine/Source/Runtime/Engine/Public/MeshBatch.h
@@ -145,6 +145,22 @@ struct FMeshBatchElement
 	}
 };
 
+FORCEINLINE bool IsCompatibleWithHairStrands(const FMaterial* Material, const ERHIFeatureLevel::Type FeatureLevel)
+{
+	return
+		ERHIFeatureLevel::SM5 == FeatureLevel &&
+		Material && Material->IsUsedWithHairStrands() && Material->GetShadingModels().HasShadingModel(MSM_Hair) &&
+		(Material->GetBlendMode() == BLEND_Opaque || Material->GetBlendMode() == BLEND_Masked);
+}
+
+FORCEINLINE bool IsCompatibleWithHairStrands(const FMaterial* Material, const EShaderPlatform Platform)
+{
+	return
+		IsPCPlatform(Platform) && GetMaxSupportedFeatureLevel(Platform) == ERHIFeatureLevel::SM5 &&
+		Material && Material->IsUsedWithHairStrands() && Material->GetShadingModels().HasShadingModel(MSM_Hair) &&
+		(Material->GetBlendMode() == BLEND_Opaque || Material->GetBlendMode() == BLEND_Masked);
+}
+
 /**
  * A batch of mesh elements, all with the same material and vertex buffer
  */
@@ -247,10 +263,13 @@ struct FMeshBatch
 		return Mat->IsDeferredDecal();
 	}
 
-	FORCEINLINE bool CastsDeepShadow(/*ERHIFeatureLevel::Type InFeatureLevel*/) const
+	FORCEINLINE bool UseForHairStrands(ERHIFeatureLevel::Type InFeatureLevel) const
 	{
-		const FMaterial* Mat = MaterialRenderProxy->GetMaterial(ERHIFeatureLevel::SM5);
-		return Mat->GetShadingModels().HasOnlyShadingModel(EMaterialShadingModel::MSM_Hair);
+		if (ERHIFeatureLevel::SM5 != InFeatureLevel)
+			return false;
+
+		const FMaterial* Mat = MaterialRenderProxy->GetMaterial(InFeatureLevel);
+		return IsCompatibleWithHairStrands(Mat, InFeatureLevel);
 	}
 
 	FORCEINLINE bool IsMasked(ERHIFeatureLevel::Type InFeatureLevel) const
diff --git a/Engine/Source/Runtime/Engine/Public/Physics/Experimental/PhysInterface_Chaos.h b/Engine/Source/Runtime/Engine/Public/Physics/Experimental/PhysInterface_Chaos.h
--- a/Engine/Source/Runtime/Engine/Public/Physics/Experimental/PhysInterface_Chaos.h
+++ b/Engine/Source/Runtime/Engine/Public/Physics/Experimental/PhysInterface_Chaos.h
@@ -334,8 +334,12 @@ public:
     // Interface needed for cmd
     static bool ExecuteRead(const FPhysicsActorHandle& InActorReference, TFunctionRef<void(const FPhysicsActorHandle& Actor)> InCallable)
     {
-		InCallable(InActorReference);
-		return true;
+		if (InActorReference)
+		{
+			InCallable(InActorReference);
+			return true;
+		}
+		return false;
     }
 
     static bool ExecuteRead(USkeletalMeshComponent* InMeshComponent, TFunctionRef<void()> InCallable)
@@ -374,10 +378,26 @@ public:
 
     static bool ExecuteWrite(const FPhysicsActorHandle& InActorReference, TFunctionRef<void(const FPhysicsActorHandle& Actor)> InCallable)
     {
-		InCallable(InActorReference);
-		return true;
+		//why do we have a write that takes in a const handle?
+		if (InActorReference)
+		{
+			InCallable(InActorReference);
+			return true;
+		}
+		return false;
     }
 
+	static bool ExecuteWrite(FPhysicsActorHandle& InActorReference, TFunctionRef<void(FPhysicsActorHandle& Actor)> InCallable)
+	{
+		if (InActorReference)
+		{
+			InCallable(InActorReference);
+			return true;
+		}
+
+		return false;
+	}
+
     static bool ExecuteWrite(USkeletalMeshComponent* InMeshComponent, TFunctionRef<void()> InCallable)
     {
         InCallable();
@@ -499,11 +519,6 @@ FORCEINLINE float GetHalfHeight(const Chaos::TCapsule<float>& Capsule)
 }
 */
 
-FVector FindBoxOpposingNormal(const FLocationHit& PHit, const FVector& TraceDirectionDenorm, const FVector& InNormal);
-FVector FindHeightFieldOpposingNormal(const FLocationHit& PHit, const FVector& TraceDirectionDenorm, const FVector& InNormal);
-FVector FindConvexMeshOpposingNormal(const FLocationHit& PHit, const FVector& TraceDirectionDenorm, const FVector& InNormal);
-FVector FindTriMeshOpposingNormal(const FLocationHit& PHit, const FVector& TraceDirectionDenorm, const FVector& InNormal);
-
 FORCEINLINE void DrawOverlappingTris(const UWorld* World, const FLocationHit& Hit, const Chaos::TImplicitObject<float, 3>& Geom, const FTransform& QueryTM)
 {
 	//TODO_SQ_IMPLEMENTATION
diff --git a/Engine/Source/Runtime/Engine/Public/Physics/Experimental/PhysScene_Chaos.h b/Engine/Source/Runtime/Engine/Public/Physics/Experimental/PhysScene_Chaos.h
--- a/Engine/Source/Runtime/Engine/Public/Physics/Experimental/PhysScene_Chaos.h
+++ b/Engine/Source/Runtime/Engine/Public/Physics/Experimental/PhysScene_Chaos.h
@@ -14,7 +14,7 @@
 #include "PhysInterface_Chaos.h"
 #include "Chaos/Transform.h"
 #include "Chaos/Declares.h"
-#include "Chaos/Framework/SingleParticlePhysicsProxyFwd.h"
+#include "PhysicsProxy/SingleParticlePhysicsProxyFwd.h"
 
 #ifndef CHAOS_WITH_PAUSABLE_SOLVER
 #define CHAOS_WITH_PAUSABLE_SOLVER 1
@@ -44,6 +44,12 @@ namespace Chaos
 
 	template<typename PayloadType, typename HandlerType>
 	class TRawEventHandler;
+
+	template <typename T, int d>
+	class TAccelerationStructureHandle;
+
+	template <typename TPayload, typename T, int d>
+	class ISpatialAcceleration;
 }
 
 /**
@@ -103,6 +109,8 @@ public:
 	void RemoveObject(FGeometryCollectionPhysicsProxy* InObject);
 	void RemoveObject(FFieldSystemPhysicsProxy* InObject);	
 
+	void RemoveActorFromAccelerationStructure(FPhysicsActorHandle& Actor);
+
 	template<typename PayloadType>
 	void RegisterEvent(const Chaos::EEventType& EventID, TFunction<void(const Chaos::FPBDRigidsSolver* Solver, PayloadType& EventData)> InLambda);
 	void UnregisterEvent(const Chaos::EEventType& EventID);
@@ -136,6 +144,8 @@ public:
 		return PhysicsProxyPtr ? *PhysicsProxyPtr : nullptr;
 	}
 
+	TUniquePtr<Chaos::ISpatialAcceleration<Chaos::TAccelerationStructureHandle<float, 3>, float, 3>> SolverAccelerationStructure;
+
 private:
 #if CHAOS_WITH_PAUSABLE_SOLVER
 	/** Callback that checks the status of the world settings for this scene before pausing/unpausing its solver. */
@@ -249,6 +259,9 @@ public:
 	Chaos::FPhysicsSolver* GetSolver();
 	const Chaos::FPhysicsSolver* GetSolver() const;
 
+	FPhysScene_Chaos& GetScene() { return Scene; }
+	const FPhysScene_Chaos& GetScene() const { return Scene; }
+
 	FPhysicsReplication* GetPhysicsReplication();
 	void RemoveBodyInstanceFromPendingLists_AssumesLocked(FBodyInstance* BodyInstance, int32 SceneType);
 	void AddCustomPhysics_AssumesLocked(FBodyInstance* BodyInstance, FCalculateCustomPhysics& CalculateCustomPhysics);
diff --git a/Engine/Source/Runtime/Engine/Public/Physics/ImmediatePhysics/ImmediatePhysicsChaos/ImmediatePhysicsActorHandle_Chaos.h b/Engine/Source/Runtime/Engine/Public/Physics/ImmediatePhysics/ImmediatePhysicsChaos/ImmediatePhysicsActorHandle_Chaos.h
--- a/Engine/Source/Runtime/Engine/Public/Physics/ImmediatePhysics/ImmediatePhysicsChaos/ImmediatePhysicsActorHandle_Chaos.h
+++ b/Engine/Source/Runtime/Engine/Public/Physics/ImmediatePhysics/ImmediatePhysicsChaos/ImmediatePhysicsActorHandle_Chaos.h
@@ -109,6 +109,10 @@ namespace ImmediatePhysics_Chaos
 		/** Get the max contact impulse*/
 		float GetMaxContactImpulse() const;
 
+		/** Get the actor-space centre of mass offset */
+		// @todo(ccaulfield): should be FTransform
+		const FVector& GetCoMTranslation() const;
+
 	private:
 		friend struct FSimulation;
 		friend struct FJointHandle;
diff --git a/Engine/Source/Runtime/Engine/Public/Physics/ImmediatePhysics/ImmediatePhysicsChaos/ImmediatePhysicsCore_Chaos.h b/Engine/Source/Runtime/Engine/Public/Physics/ImmediatePhysics/ImmediatePhysicsChaos/ImmediatePhysicsCore_Chaos.h
--- a/Engine/Source/Runtime/Engine/Public/Physics/ImmediatePhysics/ImmediatePhysicsChaos/ImmediatePhysicsCore_Chaos.h
+++ b/Engine/Source/Runtime/Engine/Public/Physics/ImmediatePhysics/ImmediatePhysicsChaos/ImmediatePhysicsCore_Chaos.h
@@ -11,13 +11,16 @@
 
 namespace Chaos
 {
-	template<typename T, int D> class TPBDRigidsEvolutionGBF;
-	template<typename T, int D> class TPBDRigidsSOAs;
-	template<typename T, int D> struct TKinematicGeometryParticleParameters;
-	template<typename T, int D> struct TPBDRigidParticleParameters;
 	template<typename T, int D> class TImplicitObject;
-	template<typename T, int D> class TPBDJointConstraints;
+	template<typename T, int D> struct TKinematicGeometryParticleParameters;
 	template<typename P, typename T, int D> class TPBDConstraintIslandRule;
+	template<typename T, int D> class TPBDJointConstraintHandle;
+	template<typename T, int D> class TPBDJointConstraints;
+	template<typename T, int D> class TPBD6DJointConstraintHandle;
+	template<typename T, int D> class TPBD6DJointConstraints;
+	template<typename T, int D> struct TPBDRigidParticleParameters;
+	template<typename T, int D> class TPBDRigidsEvolutionGBF;
+	template<typename T, int D> class TPBDRigidsSOAs;
 	template<typename T, int D> class TPerShapeData;
 }
 
diff --git a/Engine/Source/Runtime/Engine/Public/Physics/ImmediatePhysics/ImmediatePhysicsChaos/ImmediatePhysicsJointHandle_Chaos.h b/Engine/Source/Runtime/Engine/Public/Physics/ImmediatePhysics/ImmediatePhysicsChaos/ImmediatePhysicsJointHandle_Chaos.h
--- a/Engine/Source/Runtime/Engine/Public/Physics/ImmediatePhysics/ImmediatePhysicsChaos/ImmediatePhysicsJointHandle_Chaos.h
+++ b/Engine/Source/Runtime/Engine/Public/Physics/ImmediatePhysics/ImmediatePhysicsChaos/ImmediatePhysicsJointHandle_Chaos.h
@@ -14,12 +14,15 @@ namespace ImmediatePhysics_Chaos
 	struct ENGINE_API FJointHandle
 	{
 	public:
-		FJointHandle(Chaos::TPBDJointConstraints<FReal, Dimensions>* InConstraints, FConstraintInstance* ConstraintInstance, FActorHandle* Actor1, FActorHandle* Actor2);
+		using FChaosConstraintContainer = Chaos::TPBD6DJointConstraints<FReal, Dimensions>;
+		using FChaosConstraintHandle = typename Chaos::TPBD6DJointConstraintHandle<FReal, Dimensions>;
+
+		FJointHandle(FChaosConstraintContainer* InConstraints, FConstraintInstance* ConstraintInstance, FActorHandle* Actor1, FActorHandle* Actor2);
 		~FJointHandle();
 
 	private:
-		Chaos::TPBDJointConstraints<FReal, Dimensions>* Constraints;
-		int32 ConstraintIndex;
+		FChaosConstraintContainer* Constraints;
+		FChaosConstraintHandle* ConstraintHandle;
 	};
 }
 #endif
diff --git a/Engine/Source/Runtime/Engine/Public/Physics/ImmediatePhysics/ImmediatePhysicsChaos/ImmediatePhysicsSimulation_Chaos.h b/Engine/Source/Runtime/Engine/Public/Physics/ImmediatePhysics/ImmediatePhysicsChaos/ImmediatePhysicsSimulation_Chaos.h
--- a/Engine/Source/Runtime/Engine/Public/Physics/ImmediatePhysics/ImmediatePhysicsChaos/ImmediatePhysicsSimulation_Chaos.h
+++ b/Engine/Source/Runtime/Engine/Public/Physics/ImmediatePhysics/ImmediatePhysicsChaos/ImmediatePhysicsSimulation_Chaos.h
@@ -80,8 +80,8 @@ namespace ImmediatePhysics_Chaos
 	private:
 		TUniquePtr<Chaos::TPBDRigidsEvolutionGBF<FReal, Dimensions>> Evolution;
 		TUniquePtr<Chaos::TPBDRigidsSOAs<FReal, Dimensions>> Particles;
-		TUniquePtr<Chaos::TPBDJointConstraints<FReal, Dimensions>> Joints;
-		TUniquePtr<Chaos::TPBDConstraintIslandRule<Chaos::TPBDJointConstraints<FReal, Dimensions>, FReal, Dimensions>> JointsRule;		// @todo(ccaulfield): Don't need islands for anim node physics...
+		TUniquePtr<Chaos::TPBD6DJointConstraints<FReal, Dimensions>> Joints;
+		TUniquePtr<Chaos::TPBDConstraintIslandRule<Chaos::TPBD6DJointConstraints<FReal, Dimensions>, FReal, Dimensions>> JointsRule;		// @todo(ccaulfield): Don't need islands for anim node physics...
 
 		/** Mapping from entity index to handle */
 		// @todo(ccaulfield): we now have handles pointing to handles which is inefficient - we can do better than this, but don't want to change API yet
diff --git a/Engine/Source/Runtime/Engine/Public/PrimitiveViewRelevance.h b/Engine/Source/Runtime/Engine/Public/PrimitiveViewRelevance.h
--- a/Engine/Source/Runtime/Engine/Public/PrimitiveViewRelevance.h
+++ b/Engine/Source/Runtime/Engine/Public/PrimitiveViewRelevance.h
@@ -49,6 +49,8 @@ struct FPrimitiveViewRelevance
 	uint32 bDrawRelevance : 1;
 	/** The primitive is casting a shadow. */
 	uint32 bShadowRelevance : 1;
+	/** The primitive is hair strands geometry. */
+	uint32 bHairStrandsRelevance : 1;
 	/** The primitive should render velocity. */
 	uint32 bVelocityRelevance : 1;
 	/** The primitive should render to the custom depth pass. */
diff --git a/Engine/Source/Runtime/Engine/Public/RayTracingInstanceUtils.h b/Engine/Source/Runtime/Engine/Public/RayTracingInstanceUtils.h
new file mode 100644
--- /dev/null
+++ b/Engine/Source/Runtime/Engine/Public/RayTracingInstanceUtils.h
@@ -0,0 +1,19 @@
+// Copyright 1998-2019 Epic Games, Inc. All Rights Reserved.
+
+#pragma once
+
+#include "RenderResource.h"
+
+#if RHI_RAYTRACING
+
+enum ERaytracingInstanceMask
+{
+	RaytracingInstanceMask_Opaque			= 0x01,
+	RaytracingInstanceMask_Translucent		= 0x02,
+	RaytracingInstanceMask_ThinShadow		= 0x04,
+	RaytracingInstanceMask_Shadow			= 0x08,
+	RaytracingInstanceMask_All				= 0xFF
+};
+
+ENGINE_API void AddOpaqueRaytracingInstance(const FMatrix& InstanceTransform, const FRayTracingGeometry* RayTracingGeometry, const uint32 Mask, TArray<struct FRayTracingInstance>& OutRayTracingInstances);
+#endif
diff --git a/Engine/Source/Runtime/Engine/Public/SceneManagement.h b/Engine/Source/Runtime/Engine/Public/SceneManagement.h
--- a/Engine/Source/Runtime/Engine/Public/SceneManagement.h
+++ b/Engine/Source/Runtime/Engine/Public/SceneManagement.h
@@ -1361,6 +1361,7 @@ public:
 	inline bool CastsStaticShadow() const { return bCastStaticShadow; }
 	inline bool CastsTranslucentShadows() const { return bCastTranslucentShadows; }
 	inline bool CastsVolumetricShadow() const { return bCastVolumetricShadow; }
+	inline bool CastsHairStrandsDeepShadow() const { return bCastHairStrandsDeepShadow; }
 	inline bool CastsRaytracedShadow() const { return bCastRaytracedShadow; }
 	inline bool AffectReflection() const { return bAffectReflection; }
 	inline bool AffectGlobalIllumination() const { return bAffectGlobalIllumination; }
@@ -1527,7 +1528,7 @@ protected:
 	const uint8 bTransmission : 1;
 
 	const uint8 bCastVolumetricShadow : 1;
-
+	const uint8 bCastHairStrandsDeepShadow : 1;
 	const uint8 bCastShadowsFromCinematicObjectsOnly : 1;
 
 	const uint8 bForceCachedShadowsForMovablePrimitives : 1;
diff --git a/Engine/Source/Runtime/Experimental/Chaos/Private/Chaos/HeightField.cpp b/Engine/Source/Runtime/Experimental/Chaos/Private/Chaos/HeightField.cpp
--- a/Engine/Source/Runtime/Experimental/Chaos/Private/Chaos/HeightField.cpp
+++ b/Engine/Source/Runtime/Experimental/Chaos/Private/Chaos/HeightField.cpp
@@ -66,9 +66,15 @@ bool THeightField<T>::SweepGeom(const TImplicitObject<T, 3>& QueryGeom, const TR
 }
 
 template <typename T>
-int32 THeightField<T>::FindMostOpposingFace(const TVector<T, 3>& Position, const TVector<T, 3>& UnitDir, int32 HintFaceIndex) const
+int32 THeightField<T>::FindMostOpposingFace(const TVector<T, 3>& Position, const TVector<T, 3>& UnitDir, int32 HintFaceIndex, T SearchDist) const
 {
-	return MTriangleMeshImplicitObject->FindMostOpposingFace(Position, UnitDir, HintFaceIndex);
+	return MTriangleMeshImplicitObject->FindMostOpposingFace(Position, UnitDir, HintFaceIndex, SearchDist);
+}
+
+template <typename T>
+TVector<T, 3> THeightField<T>::FindGeometryOpposingNormal(const TVector<T, 3>& DenormDir, int32 FaceIndex, const TVector<T, 3>& OriginalNormal) const
+{
+	return MTriangleMeshImplicitObject->FindGeometryOpposingNormal(DenormDir, FaceIndex, OriginalNormal);
 }
 
 }
diff --git a/Engine/Source/Runtime/Experimental/Chaos/Private/Chaos/PBD6DJointConstraints.cpp b/Engine/Source/Runtime/Experimental/Chaos/Private/Chaos/PBD6DJointConstraints.cpp
new file mode 100644
--- /dev/null
+++ b/Engine/Source/Runtime/Experimental/Chaos/Private/Chaos/PBD6DJointConstraints.cpp
@@ -0,0 +1,221 @@
+// Copyright 1998-2019 Epic Games, Inc. All Rights Reserved.
+#include "Chaos/PBD6DJointConstraints.h"
+#include "Chaos/ParticleHandle.h"
+#include "Chaos/Utilities.h"
+
+using namespace Chaos;
+
+
+template<class T, int d>
+PMatrix<T, 6, 6> ComputeD6JointFactorMatrix()
+{
+	// Ra, Rb = world-space relative position of connector on body A, B
+	// Twista, Twistb = world-space twist axis for body A, B
+	// Swing0a, Swing0b = first world-space swing axis for body A, B
+	// Swing1a, Swing1b = second world-space swing axis for body A, B
+	// RaX, RbX = Cross-product matrix of Ra, Rb
+	//
+	// Jacobians:
+	// Ja =		|	1(3x3)			-RaX					|
+	//			|	0(1x3)			(Twista x Twistb)T		|
+	//			|	0(1x3)			(Swinga0 x Swingb0)T	|
+	//			|	0(1x3)			(Swinga1 x Swingb1)T	|
+	//
+	// Jb =		|	1(3x3)			RBX						|
+	//			|	0(1x3)			-(Twista x Twistb)T		|
+	//			|	0(1x3)			-(Swinga0 x Swingb0)T	|
+	//			|	0(1x3)			-(Swinga1 x Swingb1)T	|
+	//
+	// MaInv =	|	1(3x3)/Ma		0(3x3)			|
+	//			|	0(3x3)			InertaaInv		|
+	//
+	// MbInv =	|	1(3x3)/Mb		0(3x3)			|
+	//			|	0(3x3)			InertabInv		|
+	//
+	// FactorMatrix = J.MInv.(J)T
+	//
+	// FMa =	Ja.MaInv.(Ja)T
+	//			|	1/Ma(3x3) - RaX.InertiaaInv.RaX				-RaX.InertaaInv.(Twista x Twistb)						-RaX.InertaaInv.(Swinga0 x Swingb0)						-RaX.InertaaInv.(Swinga1 x Swingb1)						|
+	//			|	(-RaX.InertaaInv.(Twista x Twistb))T		(Twista x Twistb)T.InertiaaInv.(Twista x Twistb)		(Twista x Twistb)T.InertiaaInv.(Swinga0 x Swingb0)		(Twista x Twistb)T.InertiaaInv.(Swinga1 x Swingb1)		|
+	//			|	(-RaX.InertaaInv.(Swinga0 x Swingb0)T		(Swinga0 x Swingb0)T.InertiaaInv.(Twista x Twistb)		(Swinga0 x Swingb0)T.InertiaaInv.(Swinga0 x Swingb0)	(Swinga0 x Swingb0)T.InertiaaInv.(Swinga1 x Swingb1)	|
+	//			|	(-RaX.InertaaInv.(Swinga1 x Swingb1))T		(Swinga1 x Swingb1)T.InertiaaInv.(Twista x Twistb)		(Swinga1 x Swingb1)T.InertiaaInv.(Swinga0 x Swingb0)	(Swinga1 x Swingb1)T.InertiaaInv.(Swinga1 x Swingb1)	|
+	//
+	// MFb =	Jb.MbInv.(Jb)T
+}
+
+//
+// Constraint Handle
+//
+
+//
+// Constraint Settings
+//
+
+template<class T, int d>
+TPBD6DJointConstraintSettings<T, d>::TPBD6DJointConstraintSettings()
+	: ConstraintFrames({ FTransform::Identity, FTransform::Identity })
+	, LinearMotionTypes({ E6DJointMotionType::Locked, E6DJointMotionType::Locked, E6DJointMotionType::Locked })
+	, AngularMotionTypes({ E6DJointMotionType::Free, E6DJointMotionType::Free, E6DJointMotionType::Free })
+	, LinearLimits(TVector<T, d>(FLT_MAX, FLT_MAX, FLT_MAX))
+	, AngularLimits(TVector<T, d>(FLT_MAX, FLT_MAX, FLT_MAX))
+{
+}
+
+
+//
+// Constraint Container
+//
+
+
+template<class T, int d>
+void TPBD6DJointConstraints<T, d>::InitConstraint(int ConstraintIndex, const FTransformPair& ConstraintFrames)
+{
+	ConstraintSettings[ConstraintIndex].ConstraintFrames = ConstraintFrames;
+}
+
+template<class T, int d>
+TVector<T, d> TPBD6DJointConstraints<T, d>::GetDeltaDynamicDynamic(const TVector<T, d>& P0, const TVector<T, d>& P1, const TVector<T, d>& C0, const TVector<T, d>& C1, const PMatrix<T, d, d>& InvI0, const PMatrix<T, d, d>& InvI1, const T InvM0, const T InvM1)
+{
+	PMatrix<T, d, d> Factor = Utilities::ComputeJointFactorMatrix(C0 - P0, InvI0, InvM0) + Utilities::ComputeJointFactorMatrix(C1 - P1, InvI1, InvM1);
+	PMatrix<T, d, d> FactorInv = Factor.Inverse();
+	FactorInv.M[3][3] = 1;
+	TVector<T, d> Delta = C1 - C0;
+	return FactorInv * Delta;
+}
+
+template<class T, int d>
+void TPBD6DJointConstraints<T, d>::ApplySingle(const T Dt, const int32 ConstraintIndex)
+{
+	// @todo(ccaulfield): bApplyProjection should be an option. Either per-constraint or per-container...
+	const bool bApplyProjection = true;
+
+	const TVector<TGeometryParticleHandle<T, d>*, 2>& Constraint = ConstraintParticles[ConstraintIndex];
+	if (Constraint[0]->AsDynamic() && Constraint[1]->AsDynamic())
+	{
+		ApplyDynamicDynamic(Dt, ConstraintIndex, 0, 1, bApplyProjection);
+	}
+	else if (Constraint[0]->AsDynamic())
+	{
+		ApplyDynamicStatic(Dt, ConstraintIndex, 0, 1, bApplyProjection);
+	}
+	else
+	{
+		ApplyDynamicStatic(Dt, ConstraintIndex, 1, 0, bApplyProjection);
+	}
+}
+
+template<class T, int d>
+void TPBD6DJointConstraints<T, d>::ApplyDynamicDynamic(const T Dt, const int32 ConstraintIndex, const int32 PBDRigid0Index, const int32 PBDRigid1Index, const bool bApplyProjection)
+{
+	check((PBDRigid0Index == 0) || (PBDRigid0Index == 1));
+	check((PBDRigid1Index == 0) || (PBDRigid1Index == 1));
+	check(PBDRigid0Index != PBDRigid1Index);
+
+	TPBDRigidParticleHandle<T, d>* PBDRigid0 = ConstraintParticles[ConstraintIndex][PBDRigid0Index]->AsDynamic();
+	TPBDRigidParticleHandle<T, d>* PBDRigid1 = ConstraintParticles[ConstraintIndex][PBDRigid1Index]->AsDynamic();
+	check(PBDRigid0 && PBDRigid1 && (PBDRigid0->Island() == PBDRigid1->Island()));
+
+	TRotation<T, d>& Q0 = PBDRigid0->Q();
+	TVector<T, d>& P0 = PBDRigid0->P();
+	TRotation<T, d>& Q1 = PBDRigid1->Q();
+	TVector<T, d>& P1 = PBDRigid1->P();
+
+	// Calculate world-space constraint positions
+	const TVector<T, 3>& Distance0 = ConstraintSettings[ConstraintIndex].ConstraintFrames[PBDRigid0Index].GetTranslation();
+	const TVector<T, 3>& Distance1 = ConstraintSettings[ConstraintIndex].ConstraintFrames[PBDRigid1Index].GetTranslation();
+	const TVector<T, d> C0 = Q0.RotateVector(Distance0) + P0;
+	const TVector<T, d> C1 = Q1.RotateVector(Distance1) + P1;
+
+	// Calculate world-space mass
+	const PMatrix<T, d, d> InvI0 = (Q0 * FMatrix::Identity) * PBDRigid0->InvI() * (Q0 * FMatrix::Identity).GetTransposed();
+	const PMatrix<T, d, d> InvI1 = (Q1 * FMatrix::Identity) * PBDRigid1->InvI() * (Q1 * FMatrix::Identity).GetTransposed();
+	const float InvM0 = PBDRigid0->InvM();
+	const float InvM1 = PBDRigid1->InvM();
+
+	// Calculate mass-weighted correction
+	const TVector<T, d> Delta = GetDeltaDynamicDynamic(P0, P1, C0, C1, InvI0, InvI1, InvM0, InvM1);
+
+	// Apply corrections
+	P0 += InvM0 * Delta;
+	Q0 += TRotation<T, d>(InvI0 * TVector<T, d>::CrossProduct(C0 - P0, Delta), 0.f) * Q0 * T(0.5);
+	Q0.Normalize();
+
+	P1 -= InvM1 * Delta;
+	Q1 += TRotation<T, d>(InvI1 * TVector<T, d>::CrossProduct(C1 - P1, -Delta), 0.f) * Q1 * T(0.5);
+	Q1.Normalize();
+
+	// Correct any remaining error by translating
+	if (bApplyProjection)
+	{
+		const TVector<T, d> C0Proj = Q0.RotateVector(Distance0) + P0;
+		const TVector<T, d> C1Proj = Q1.RotateVector(Distance1) + P1;
+		const TVector<T, d> DeltaProj = (C1Proj - C0Proj) / (InvM0 + InvM1);
+
+		P0 += InvM0 * DeltaProj;
+		P1 -= InvM1 * DeltaProj;
+	}
+}
+
+template<class T, int d>
+TVector<T, d> TPBD6DJointConstraints<T, d>::GetDeltaDynamicKinematic(const TVector<T, d>& P0, const TVector<T, d>& C0, const TVector<T, d>& C1, const PMatrix<T, d, d>& InvI0, const T InvM0)
+{
+	PMatrix<T, d, d> Factor = Utilities::ComputeJointFactorMatrix(C0 - P0, InvI0, InvM0);
+	PMatrix<T, d, d> FactorInv = Factor.Inverse();
+	FactorInv.M[3][3] = 1;
+	TVector<T, d> Delta = C1 - C0;
+	return FactorInv * Delta;
+}
+
+template<class T, int d>
+void TPBD6DJointConstraints<T, d>::ApplyDynamicStatic(const T Dt, const int32 ConstraintIndex, const int32 PBDRigid0Index, const int32 Static1Index, const bool bApplyProjection)
+{
+	check((PBDRigid0Index == 0) || (PBDRigid0Index == 1));
+	check((Static1Index == 0) || (Static1Index == 1));
+	check(PBDRigid0Index != Static1Index);
+
+	TPBDRigidParticleHandle<T, d>* PBDRigid0 = ConstraintParticles[ConstraintIndex][PBDRigid0Index]->AsDynamic();
+	TGeometryParticleHandle<T, d>* Static1 = ConstraintParticles[ConstraintIndex][Static1Index];
+	check(PBDRigid0 && Static1 && !Static1->AsDynamic());
+
+	TRotation<T, d>& Q0 = PBDRigid0->Q();
+	TVector<T, d>& P0 = PBDRigid0->P();
+	const TRotation<T, d>& Q1 = Static1->R();
+	const TVector<T, d>& P1 = Static1->X();
+
+	// Calculate world-space constraint positions
+	const TVector<T, 3>& Distance0 = ConstraintSettings[ConstraintIndex].ConstraintFrames[PBDRigid0Index].GetTranslation();
+	const TVector<T, 3>& Distance1 = ConstraintSettings[ConstraintIndex].ConstraintFrames[Static1Index].GetTranslation();
+	const TVector<T, d> C0 = Q0.RotateVector(Distance0) + P0;
+	const TVector<T, d> C1 = Q1.RotateVector(Distance1) + P1;
+
+	// Calculate world-space mass
+	const PMatrix<T, d, d> InvI0 = (Q0 * FMatrix::Identity) * PBDRigid0->InvI() * (Q0 * FMatrix::Identity).GetTransposed();
+	const float InvM0 = PBDRigid0->InvM();
+
+	// Calculate mass-weighted correction
+	const TVector<T, d> Delta = GetDeltaDynamicKinematic(P0, C0, C1, InvI0, InvM0);
+
+	// Apply correction
+	P0 += InvM0 * Delta;
+	Q0 += TRotation<T, d>(InvI0 * TVector<T, d>::CrossProduct(C0 - P0, Delta), 0.f) * Q0 * T(0.5);
+	Q0.Normalize();
+
+	// Correct any remaining error by translating
+	if (bApplyProjection)
+	{
+		const TVector<T, d> C0Proj = Q0.RotateVector(Distance0) + P0;
+		const TVector<T, d> C1Proj = Q1.RotateVector(Distance1) + P1;
+		const TVector<T, d> DeltaProj = (C1Proj - C0Proj);
+
+		P0 += DeltaProj;
+	}
+}
+
+
+namespace Chaos
+{
+	template class TPBD6DJointConstraintSettings<float, 3>;
+	template class TPBD6DJointConstraintHandle<float, 3>;
+	template class TContainerConstraintHandle<TPBD6DJointConstraints<float, 3>>;
+	template class TPBD6DJointConstraints<float, 3>;
+}
diff --git a/Engine/Source/Runtime/Experimental/Chaos/Private/Chaos/PBDCollisionConstraint.cpp b/Engine/Source/Runtime/Experimental/Chaos/Private/Chaos/PBDCollisionConstraint.cpp
--- a/Engine/Source/Runtime/Experimental/Chaos/Private/Chaos/PBDCollisionConstraint.cpp
+++ b/Engine/Source/Runtime/Experimental/Chaos/Private/Chaos/PBDCollisionConstraint.cpp
@@ -56,6 +56,7 @@ namespace Chaos
 
 namespace Chaos
 {
+
 template<ECollisionUpdateType, typename T, int d>
 void UpdateConstraintImp2(const TRigidTransform<T, d>& ParticleTM, const TImplicitObject<T, d>& LevelsetObject, const TRigidTransform<T, d>& LevelsetTM, const T Thickness, TRigidBodyContactConstraint<T, d>& Constraint);
 
@@ -80,6 +81,7 @@ template<typename T, int d>
 void TPBDCollisionConstraint<T, d>::Reset(/*const TPBDRigidParticles<T, d>& InParticles, const TArray<int32>& InIndices*/)
 {
 	Constraints.Empty();
+	Handles.Empty();
 	MAngularFriction = 0;
 	bUseCCD = false;
 }
@@ -101,6 +103,15 @@ void TPBDCollisionConstraint<T, d>::UpdatePositionBasedState(/*const TPBDRigidPa
 	}
 }
 
+template <typename T, int d>
+TUniquePtr<ISpatialAcceleration<TAccelerationStructureHandle<T, d>, T, d>> TPBDCollisionConstraint<T, d>::CreateExternalAccelerationStructure()
+{
+	auto Copy = SpatialAccelerationResource.GetWritable().Copy();
+	Copy->Reinitialize(Particles.GetNonDisabledView(), true, 0); //todo(ocohen): bounds are wrong - also no need to recreate, just remove deleted items
+	return Copy;
+}
+
+
 
 #if CHAOS_PARTICLEHANDLE_TODO
 template <typename T, int d>
@@ -279,7 +290,7 @@ void TPBDCollisionConstraint<T, d>::ComputeConstraintsHelper(T Dt, const SPATIAL
 
 			//SCOPE_CYCLE_COUNTER(STAT_ComputeConstraintsNP2);
 
-			TArray<TGeometryParticleHandle<T,d>*> PotentialIntersections;
+			TArray<TAccelerationStructureHandle<T,d>> PotentialIntersections;
 
 			const bool bBody1Bounded = HasBoundingBox(Particle1);
 
@@ -316,7 +327,7 @@ void TPBDCollisionConstraint<T, d>::ComputeConstraintsHelper(T Dt, const SPATIAL
 			const int32 NumPotentials = PotentialIntersections.Num();
 			for (int32 i = 0; i < NumPotentials; ++i)
 			{
-				auto& Particle2 = *PotentialIntersections[i];
+				auto& Particle2 = *PotentialIntersections[i].GetGeometryParticleHandle_PhysicsThread();
 				const TGenericParticleHandle<T, d> Particle2Generic(&Particle2);
 				// Collision group culling...
 				// CollisionGroup == 0 : Collide_With_Everything
@@ -410,6 +421,7 @@ void TPBDCollisionConstraint<T, d>::ComputeConstraintsHelper(T Dt, const SPATIAL
 		{				
 			int32 Idx = Constraints.AddUninitialized(1);
 			Queue.Dequeue(Constraints[Idx]);
+			Handles.Add(HandleAllocator.AllocHandle(this, Idx));
 		}
 	}
 
@@ -496,6 +508,11 @@ void TPBDCollisionConstraint<T, d>::RemoveConstraints(const TSet<TGeometryPartic
 		if (RemovedParticles.Contains(Constraint.ParticleIndex) || RemovedParticles.Contains(Constraint.LevelsetIndex))
 		{
 			Constraints.RemoveAtSwap(i);
+			Handles.RemoveAtSwap(i);
+			if (i < Handles.Num())
+			{
+				SetConstraintIndex(Handles[i], i);
+			}
 			i--;
 		}
 	}
@@ -639,6 +656,7 @@ void TPBDCollisionConstraint<T, d>::UpdateConstraintsHelper(/*const TPBDRigidPar
 	{				
 		int32 Idx = Constraints.AddUninitialized(1);
 		Queue.Dequeue(Constraints[Idx]);
+		Handles.Add(HandleAllocator.AllocHandle(this, Idx));
 	}
 
 	Timer.Stop();
@@ -662,6 +680,11 @@ void TPBDCollisionConstraint<T, d>::UpdateConstraints(/*const TPBDRigidParticles
 			if(InParticles.Disabled(Constraint.ParticleIndex) || InParticles.Disabled(Constraint.LevelsetIndex))
 			{
 				Constraints.RemoveAtSwap(i);
+				Handles.RemoveAt(i);
+				if (i < Handles.Num())
+				{
+					SetConstraintIndex(Handles[i], i);
+				}
 				++NumRemovedConstraints;
 				i--;
 			}
@@ -939,11 +962,12 @@ void TPBDCollisionConstraint<T, d>::Apply(const T Dt, FRigidBodyContactConstrain
 
 DECLARE_CYCLE_STAT(TEXT("Apply"), STAT_Apply2, STATGROUP_ChaosWide);
 template<typename T, int d>
-void TPBDCollisionConstraint<T, d>::Apply(const T Dt, const TArray<int32>& InConstraintIndices)
+void TPBDCollisionConstraint<T, d>::Apply(const T Dt, const TArray<FConstraintHandle*>& InConstraintHandles)
 {
-	PhysicsParallelFor(InConstraintIndices.Num(), [&](int32 ConstraintIndex) {
-		FRigidBodyContactConstraint& Constraint = Constraints[InConstraintIndices[ConstraintIndex]];
-		Apply(Dt, Constraint);
+	PhysicsParallelFor(InConstraintHandles.Num(), [&](int32 ConstraintHandleIndex) {
+		FConstraintHandle* ConstraintHandle = InConstraintHandles[ConstraintHandleIndex];
+		check(ConstraintHandle != nullptr);
+		Apply(Dt, Constraints[ConstraintHandle->GetConstraintIndex()]);
 	});
 }
 
@@ -1066,15 +1090,16 @@ void TPBDCollisionConstraint<T, d>::ApplyPushOut(const T Dt, FRigidBodyContactCo
 
 //DECLARE_CYCLE_STAT(TEXT("ApplyPushOut"), STAT_ApplyPushOut, STATGROUP_ChaosWide);
 template<typename T, int d>
-bool TPBDCollisionConstraint<T, d>::ApplyPushOut(const T Dt, const TArray<int32>& InConstraintIndices, const TSet<TGeometryParticleHandle<T,d>*>& IsTemporarilyStatic, int32 Iteration, int32 NumIterations)
+bool TPBDCollisionConstraint<T, d>::ApplyPushOut(const T Dt, const TArray<FConstraintHandle*>& InConstraintHandles, const TSet<TGeometryParticleHandle<T,d>*>& IsTemporarilyStatic, int32 Iteration, int32 NumIterations)
 {
 	SCOPE_CYCLE_COUNTER(STAT_ApplyPushOut2);
 
 	bool NeedsAnotherIteration = false;
 
-	PhysicsParallelFor(InConstraintIndices.Num(), [&](int32 ConstraintIndex) {
-		FRigidBodyContactConstraint& Constraint = Constraints[InConstraintIndices[ConstraintIndex]];
-		ApplyPushOut(Dt, Constraint, IsTemporarilyStatic, Iteration, NumIterations, NeedsAnotherIteration);
+	PhysicsParallelFor(InConstraintHandles.Num(), [&](int32 ConstraintHandleIndex) {
+		FConstraintHandle* ConstraintHandle = InConstraintHandles[ConstraintHandleIndex];
+		check(ConstraintHandle != nullptr);
+		ApplyPushOut(Dt, Constraints[ConstraintHandle->GetConstraintIndex()], IsTemporarilyStatic, Iteration, NumIterations, NeedsAnotherIteration);
 		});
 
 	return NeedsAnotherIteration;
@@ -2494,6 +2519,7 @@ void TPBDCollisionConstraint<T, d>::UpdateConstraint(const T Thickness, FRigidBo
 	}
 }
 
+template class TAccelerationStructureHandle<float, 3>;
 template class TPBDCollisionConstraint<float, 3>;
 template void TPBDCollisionConstraint<float, 3>::UpdateConstraint<ECollisionUpdateType::Any>(const float Thickness, FRigidBodyContactConstraint& Constraint);
 template void TPBDCollisionConstraint<float, 3>::UpdateConstraint<ECollisionUpdateType::Deepest>(const float Thickness, FRigidBodyContactConstraint& Constraint);
diff --git a/Engine/Source/Runtime/Experimental/Chaos/Private/Chaos/PBDConstraintColor.cpp b/Engine/Source/Runtime/Experimental/Chaos/Private/Chaos/PBDConstraintColor.cpp
--- a/Engine/Source/Runtime/Experimental/Chaos/Private/Chaos/PBDConstraintColor.cpp
+++ b/Engine/Source/Runtime/Experimental/Chaos/Private/Chaos/PBDConstraintColor.cpp
@@ -1,12 +1,14 @@
 // Copyright 1998-2019 Epic Games, Inc. All Rights Reserved.
 
 #include "Chaos/PBDConstraintColor.h"
+
+#include "Chaos/ConstraintHandle.h"
+#include "Chaos/Framework/Parallel.h"
 #include "Chaos/ParticleHandle.h"
+#include "Chaos/PBDRigidParticles.h"
 
 #include "ProfilingDebugging/ScopedTimers.h"
 #include "ChaosStats.h"
-#include "Chaos/Framework/Parallel.h"
-#include "Chaos/PBDRigidParticles.h"
 #include "Containers/Queue.h"
 
 #include <memory>
@@ -136,7 +138,7 @@ void TPBDConstraintColor<T, d>::ComputeIslandColoring(const int32 Island, const
 					LevelToColorToConstraintListMap[Level].Add(ColorEdge.Color, {});
 				}
 
-				LevelToColorToConstraintListMap[Level][ColorEdge.Color].Add(GraphEdge.Data.ConstraintIndex);
+				LevelToColorToConstraintListMap[Level][ColorEdge.Color].Add(GraphEdge.Data.ConstraintHandle);
 
 				if (OtherNodeIndex != INDEX_NONE)
 				{
diff --git a/Engine/Source/Runtime/Experimental/Chaos/Private/Chaos/PBDConstraintGraph.cpp b/Engine/Source/Runtime/Experimental/Chaos/Private/Chaos/PBDConstraintGraph.cpp
--- a/Engine/Source/Runtime/Experimental/Chaos/Private/Chaos/PBDConstraintGraph.cpp
+++ b/Engine/Source/Runtime/Experimental/Chaos/Private/Chaos/PBDConstraintGraph.cpp
@@ -84,14 +84,14 @@ void TPBDConstraintGraph<T, d>::ReserveConstraints(const int32 NumConstraints)
 }
 
 template<typename T, int d>
-void TPBDConstraintGraph<T, d>::AddConstraint(const uint32 InContainerId, const int32 InConstraintIndex, const TVector<TGeometryParticleHandle<T,d>*, 2>& ConstrainedParticles)
+void TPBDConstraintGraph<T, d>::AddConstraint(const uint32 InContainerId, TConstraintHandle<T, d>* InConstraintHandle, const TVector<TGeometryParticleHandle<T,d>*, 2>& ConstrainedParticles)
 {
 	// Must have at least one constrained particle
 	check((ConstrainedParticles[0]) || (ConstrainedParticles[1]));
 
 	const int32 NewEdgeIndex = Edges.Num();
 	FGraphEdge NewEdge;
-	NewEdge.Data = { InContainerId, InConstraintIndex };
+	NewEdge.Data = { InContainerId, InConstraintHandle };
 
 	if (ConstrainedParticles[0])
 	{
diff --git a/Engine/Source/Runtime/Experimental/Chaos/Private/Chaos/PBDConstraintRule.cpp b/Engine/Source/Runtime/Experimental/Chaos/Private/Chaos/PBDConstraintRule.cpp
--- a/Engine/Source/Runtime/Experimental/Chaos/Private/Chaos/PBDConstraintRule.cpp
+++ b/Engine/Source/Runtime/Experimental/Chaos/Private/Chaos/PBDConstraintRule.cpp
@@ -1,6 +1,8 @@
 // Copyright 1998-2019 Epic Games, Inc. All Rights Reserved.
 
 #include "Chaos/PBDConstraintRule.h"
+
+#include "Chaos/PBD6DJointConstraints.h"
 #include "Chaos/PBDCollisionConstraint.h"
 #include "Chaos/PBDJointConstraints.h"
 #include "Chaos/PBDPositionConstraints.h"
@@ -10,8 +12,9 @@
 namespace Chaos
 {
 	template class TPBDConstraintRule<float, 3>;
+	template class TPBDConstraintColorRule<TPBDCollisionConstraint<float, 3>, float, 3>;
+	template class TPBDConstraintIslandRule<TPBD6DJointConstraints<float, 3>, float, 3>;
 	template class TPBDConstraintIslandRule<TPBDJointConstraints<float, 3>, float, 3>;
 	template class TPBDConstraintIslandRule<TPBDPositionConstraints<float, 3>, float, 3>;
 	template class TPBDConstraintIslandRule<TPBDRigidDynamicSpringConstraints<float, 3>, float, 3>;
-	template class TPBDConstraintColorRule<TPBDCollisionConstraint<float, 3>, float, 3>;
 }
\ No newline at end of file
diff --git a/Engine/Source/Runtime/Experimental/Chaos/Private/Chaos/PBDJointConstraints.cpp b/Engine/Source/Runtime/Experimental/Chaos/Private/Chaos/PBDJointConstraints.cpp
--- a/Engine/Source/Runtime/Experimental/Chaos/Private/Chaos/PBDJointConstraints.cpp
+++ b/Engine/Source/Runtime/Experimental/Chaos/Private/Chaos/PBDJointConstraints.cpp
@@ -1,4 +1,200 @@
 // Copyright 1998-2019 Epic Games, Inc. All Rights Reserved.
 #include "Chaos/PBDJointConstraints.h"
+#include "Chaos/ParticleHandle.h"
+#include "Chaos/Utilities.h"
 
-template class Chaos::TPBDJointConstraints<float, 3>;
+using namespace Chaos;
+
+//
+// Constraint Handle
+//
+
+template<class T, int d>
+const TVector<TVector<T, 3>, 2>& TPBDJointConstraintHandle<T, d>::GetConstraintPositions() const
+{
+	return ConstraintContainer->GetConstraintPositions(ConstraintIndex);
+}
+
+template<class T, int d>
+void TPBDJointConstraintHandle<T, d>::SetConstraintPositions(const TVector<TVector<T, 3>, 2>& ConstraintPositions)
+{
+	ConstraintContainer->SetConstraintPositions(ConstraintIndex, ConstraintPositions);
+}
+
+
+
+//
+// Constraint Container
+//
+
+template<class T, int d>
+void TPBDJointConstraints<T, d>::UpdateDistanceInternal(const TVector<T, d>& InLocation, const int32 InConstraintIndex)
+{
+	const TGeometryParticleHandle<T, d>* Particle0 = Constraints[InConstraintIndex][0];
+	const TGeometryParticleHandle<T, d>* Particle1 = Constraints[InConstraintIndex][1];
+	Distances[InConstraintIndex][0] = Particle0->R().Inverse().RotateVector(InLocation - Particle0->X());
+	Distances[InConstraintIndex][1] = Particle1->R().Inverse().RotateVector(InLocation - Particle1->X());
+}
+
+template<class T, int d>
+void TPBDJointConstraints<T, d>::UpdateDistance(const TVector<T, d>& InLocation, const int32 InConstraintIndex)
+{
+	Distances.SetNum(Constraints.Num());
+	UpdateDistanceInternal(InLocation, InConstraintIndex);
+}
+
+template<class T, int d>
+void TPBDJointConstraints<T, d>::UpdateDistances(const TArray<TVector<T, d>>& InLocations)
+{
+	const int32 NumConstraints = Constraints.Num();
+	Distances.SetNum(NumConstraints);
+	for (int32 ConstraintIndex = 0; ConstraintIndex < NumConstraints; ++ConstraintIndex)
+	{
+		UpdateDistanceInternal(InLocations[ConstraintIndex], ConstraintIndex);
+	}
+}
+
+template<class T, int d>
+TVector<T, d> TPBDJointConstraints<T, d>::GetDeltaDynamicDynamic(const TVector<T, d>& P0, const TVector<T, d>& P1, const TVector<T, d>& C0, const TVector<T, d>& C1, const PMatrix<T, d, d>& InvI0, const PMatrix<T, d, d>& InvI1, const T InvM0, const T InvM1)
+{
+	PMatrix<T, d, d> Factor = Utilities::ComputeJointFactorMatrix(C0 - P0, InvI0, InvM0) + Utilities::ComputeJointFactorMatrix(C1 - P1, InvI1, InvM1);
+	PMatrix<T, d, d> FactorInv = Factor.Inverse();
+	FactorInv.M[3][3] = 1;
+	TVector<T, d> Delta = C1 - C0;
+	return FactorInv * Delta;
+}
+
+template<class T, int d>
+void TPBDJointConstraints<T, d>::ApplySingle(const T Dt, const int32 ConstraintIndex)
+{
+	// @todo(ccaulfield): bApplyProjection should be an option. Either per-constraint or per-container...
+	const bool bApplyProjection = true;
+
+	const TVector<TGeometryParticleHandle<T, d>*, 2>& Constraint = Constraints[ConstraintIndex];
+	if (Constraint[0]->AsDynamic() && Constraint[1]->AsDynamic())
+	{
+		ApplyDynamicDynamic(Dt, ConstraintIndex, 0, 1, bApplyProjection);
+	}
+	else if (Constraint[0]->AsDynamic())
+	{
+		ApplyDynamicStatic(Dt, ConstraintIndex, 0, 1, bApplyProjection);
+	}
+	else
+	{
+		ApplyDynamicStatic(Dt, ConstraintIndex, 1, 0, bApplyProjection);
+	}
+}
+
+template<class T, int d>
+void TPBDJointConstraints<T, d>::ApplyDynamicDynamic(const T Dt, const int32 ConstraintIndex, const int32 PBDRigid0Index, const int32 PBDRigid1Index, const bool bApplyProjection)
+{
+	check((PBDRigid0Index == 0) || (PBDRigid0Index == 1));
+	check((PBDRigid1Index == 0) || (PBDRigid1Index == 1));
+	check(PBDRigid0Index != PBDRigid1Index);
+
+	TPBDRigidParticleHandle<T, d>* PBDRigid0 = Constraints[ConstraintIndex][PBDRigid0Index]->AsDynamic();
+	TPBDRigidParticleHandle<T, d>* PBDRigid1 = Constraints[ConstraintIndex][PBDRigid1Index]->AsDynamic();
+	check(PBDRigid0 && PBDRigid1 && (PBDRigid0->Island() == PBDRigid1->Island()));
+
+	TRotation<T, d>& Q0 = PBDRigid0->Q();
+	TVector<T, d>& P0 = PBDRigid0->P();
+	TRotation<T, d>& Q1 = PBDRigid1->Q();
+	TVector<T, d>& P1 = PBDRigid1->P();
+
+	// Calculate world-space constraint positions
+	const TVector<T, 3>& Distance0 = Distances[ConstraintIndex][PBDRigid0Index];
+	const TVector<T, 3>& Distance1 = Distances[ConstraintIndex][PBDRigid1Index];
+	const TVector<T, d> C0 = Q0.RotateVector(Distance0) + P0;
+	const TVector<T, d> C1 = Q1.RotateVector(Distance1) + P1;
+
+	// Calculate world-space mass
+	const PMatrix<T, d, d> InvI0 = (Q0 * FMatrix::Identity) * PBDRigid0->InvI() * (Q0 * FMatrix::Identity).GetTransposed();
+	const PMatrix<T, d, d> InvI1 = (Q1 * FMatrix::Identity) * PBDRigid1->InvI() * (Q1 * FMatrix::Identity).GetTransposed();
+	const float InvM0 = PBDRigid0->InvM();
+	const float InvM1 = PBDRigid1->InvM();
+
+	// Calculate mass-weighted correction
+	const TVector<T, d> Delta = GetDeltaDynamicDynamic(P0, P1, C0, C1, InvI0, InvI1, InvM0, InvM1);
+
+	// Apply corrections
+	P0 += InvM0 * Delta;
+	Q0 += TRotation<T, d>(InvI0 * TVector<T, d>::CrossProduct(C0 - P0, Delta), 0.f) * Q0 * T(0.5);
+	Q0.Normalize();
+
+	P1 -= InvM1 * Delta;
+	Q1 += TRotation<T, d>(InvI1 * TVector<T, d>::CrossProduct(C1 - P1, -Delta), 0.f) * Q1 * T(0.5);
+	Q1.Normalize();
+
+	// Correct any remaining error by translating
+	if (bApplyProjection)
+	{
+		const TVector<T, d> C0Proj = Q0.RotateVector(Distance0) + P0;
+		const TVector<T, d> C1Proj = Q1.RotateVector(Distance1) + P1;
+		const TVector<T, d> DeltaProj = (C1Proj - C0Proj) / (InvM0 + InvM1);
+
+		P0 += InvM0 * DeltaProj;
+		P1 -= InvM1 * DeltaProj;
+	}
+}
+
+template<class T, int d>
+TVector<T, d> TPBDJointConstraints<T, d>::GetDeltaDynamicKinematic(const TVector<T, d>& P0, const TVector<T, d>& C0, const TVector<T, d>& C1, const PMatrix<T, d, d>& InvI0, const T InvM0)
+{
+	PMatrix<T, d, d> Factor = Utilities::ComputeJointFactorMatrix(C0 - P0, InvI0, InvM0);
+	PMatrix<T, d, d> FactorInv = Factor.Inverse();
+	FactorInv.M[3][3] = 1;
+	TVector<T, d> Delta = C1 - C0;
+	return FactorInv * Delta;
+}
+
+template<class T, int d>
+void TPBDJointConstraints<T, d>::ApplyDynamicStatic(const T Dt, const int32 ConstraintIndex, const int32 PBDRigid0Index, const int32 Static1Index, const bool bApplyProjection)
+{
+	check((PBDRigid0Index == 0) || (PBDRigid0Index == 1));
+	check((Static1Index == 0) || (Static1Index == 1));
+	check(PBDRigid0Index != Static1Index);
+
+	TPBDRigidParticleHandle<T, d>* PBDRigid0 = Constraints[ConstraintIndex][PBDRigid0Index]->AsDynamic();
+	TGeometryParticleHandle<T, d>* Static1 = Constraints[ConstraintIndex][Static1Index];
+	check(PBDRigid0 && Static1 && !Static1->AsDynamic());
+
+	TRotation<T, d>& Q0 = PBDRigid0->Q();
+	TVector<T, d>& P0 = PBDRigid0->P();
+	const TRotation<T, d>& Q1 = Static1->R();
+	const TVector<T, d>& P1 = Static1->X();
+
+	// Calculate world-space constraint positions
+	const TVector<T, 3>& Distance0 = Distances[ConstraintIndex][PBDRigid0Index];
+	const TVector<T, 3>& Distance1 = Distances[ConstraintIndex][Static1Index];
+	const TVector<T, d> C0 = Q0.RotateVector(Distance0) + P0;
+	const TVector<T, d> C1 = Q1.RotateVector(Distance1) + P1;
+
+	// Calculate world-space mass
+	const PMatrix<T, d, d> InvI0 = (Q0 * FMatrix::Identity) * PBDRigid0->InvI() * (Q0 * FMatrix::Identity).GetTransposed();
+	const float InvM0 = PBDRigid0->InvM();
+
+	// Calculate mass-weighted correction
+	const TVector<T, d> Delta = GetDeltaDynamicKinematic(P0, C0, C1, InvI0, InvM0);
+
+	// Apply correction
+	P0 += InvM0 * Delta;
+	Q0 += TRotation<T, d>(InvI0 * TVector<T, d>::CrossProduct(C0 - P0, Delta), 0.f) * Q0 * T(0.5);
+	Q0.Normalize();
+
+	// Correct any remaining error by translating
+	if (bApplyProjection)
+	{
+		const TVector<T, d> C0Proj = Q0.RotateVector(Distance0) + P0;
+		const TVector<T, d> C1Proj = Q1.RotateVector(Distance1) + P1;
+		const TVector<T, d> DeltaProj = (C1Proj - C0Proj);
+
+		P0 += DeltaProj;
+	}
+}
+
+
+namespace Chaos
+{
+	template class Chaos::TPBDJointConstraintHandle<float, 3>;
+	template class Chaos::TPBDJointConstraints<float, 3>;
+}
diff --git a/Engine/Source/Runtime/Experimental/Chaos/Private/Chaos/PBDRigidDynamicSpringConstraints.cpp b/Engine/Source/Runtime/Experimental/Chaos/Private/Chaos/PBDRigidDynamicSpringConstraints.cpp
--- a/Engine/Source/Runtime/Experimental/Chaos/Private/Chaos/PBDRigidDynamicSpringConstraints.cpp
+++ b/Engine/Source/Runtime/Experimental/Chaos/Private/Chaos/PBDRigidDynamicSpringConstraints.cpp
@@ -5,49 +5,149 @@
 using namespace Chaos;
 
 template<class T, int d>
-void TPBDRigidDynamicSpringConstraints<T, d>::ApplyHelper(const T Dt, const TArray<int32>& InConstraintIndices) const
+void TPBDRigidDynamicSpringConstraints<T, d>::UpdatePositionBasedState(const T Dt)
 {
- 	for (int32 ConstraintIndex : InConstraintIndices)
+	const int32 NumConstraints = Constraints.Num();
+	for (int32 ConstraintIndex = 0; ConstraintIndex < NumConstraints; ++ConstraintIndex)
 	{
 		TGeometryParticleHandle<T, d>* Static0 = Constraints[ConstraintIndex][0];
 		TGeometryParticleHandle<T, d>* Static1 = Constraints[ConstraintIndex][1];
 		TPBDRigidParticleHandle<T, d>* PBDRigid0 = Static0->AsDynamic();
 		TPBDRigidParticleHandle<T, d>* PBDRigid1 = Static1->AsDynamic();
-		check(PBDRigid0 || PBDRigid1);
-		check(!PBDRigid0 || !PBDRigid1 || (PBDRigid0->Island() == PBDRigid1->Island()));
 
-		TRotation<T, d>& Q0 = PBDRigid0 ? PBDRigid0->Q() : Static0->R();
-		TRotation<T, d>& Q1 = PBDRigid1 ? PBDRigid1->Q() : Static1->R();
-		TVector<T, d>& P0 = PBDRigid0 ? PBDRigid0->P() : Static0->X();
-		TVector<T, d>& P1 = PBDRigid1 ? PBDRigid1->P() : Static1->X();
+		// Do not create springs between objects with no geometry
+		if (!Static0->Geometry() || !Static1->Geometry())
+		{
+			continue;
+		}
+
+		const TRotation<T, d>& Q0 = PBDRigid0 ? PBDRigid0->Q() : Static0->R();
+		const TRotation<T, d>& Q1 = PBDRigid1 ? PBDRigid1->Q() : Static1->R();
+		const TVector<T, d>& P0 = PBDRigid0 ? PBDRigid0->P() : Static0->X();
+		const TVector<T, d>& P1 = PBDRigid1 ? PBDRigid1->P() : Static1->X();
 
+		// Delete constraints
 		const int32 NumSprings = SpringDistances[ConstraintIndex].Num();
-		const PMatrix<T, d, d> WorldSpaceInvI1 = PBDRigid0? (Q0 * FMatrix::Identity) * PBDRigid0->InvI() * (Q0 * FMatrix::Identity).GetTransposed() : PMatrix<T, d, d>(0);
-		const PMatrix<T, d, d> WorldSpaceInvI2 = PBDRigid1 ? (Q1 * FMatrix::Identity) * PBDRigid1->InvI() * (Q1 * FMatrix::Identity).GetTransposed() : PMatrix<T, d, d>(0);;
-		for (int32 SpringIndex = 0; SpringIndex < NumSprings; ++SpringIndex)
+		for (int32 SpringIndex = NumSprings - 1; SpringIndex >= 0; --SpringIndex)
 		{
 			const TVector<T, d>& Distance0 = Distances[ConstraintIndex][SpringIndex][0];
 			const TVector<T, d>& Distance1 = Distances[ConstraintIndex][SpringIndex][1];
 			const TVector<T, d> WorldSpaceX1 = Q0.RotateVector(Distance0) + P0;
 			const TVector<T, d> WorldSpaceX2 = Q1.RotateVector(Distance1) + P1;
-			const TVector<T, d> Delta = Base::GetDelta(WorldSpaceX1, WorldSpaceX2, ConstraintIndex, SpringIndex);
-
-			if (PBDRigid0)
+			const TVector<T, d> Difference = WorldSpaceX2 - WorldSpaceX1;
+			float Distance = Difference.Size();
+			if (Distance > CreationThreshold * 2)
 			{
-				const TVector<T, d> Radius = WorldSpaceX1 - P0;
-				P0 += PBDRigid0->InvM() * Delta;
-				Q0 += TRotation<T, d>(WorldSpaceInvI1 * TVector<T, d>::CrossProduct(Radius, Delta), 0.f) * Q0 * T(0.5);
-				Q0.Normalize();
+				Distances[ConstraintIndex].RemoveAtSwap(SpringIndex);
+				SpringDistances[ConstraintIndex].RemoveAtSwap(SpringIndex);
 			}
+		}
+
+		if (SpringDistances[ConstraintIndex].Num() == MaxSprings)
+		{
+			continue;
+		}
+
+		TRigidTransform<T, d> Transform1(P0, Q0);
+		TRigidTransform<T, d> Transform2(P1, Q1);
 
-			if (PBDRigid1)
+		// Create constraints
+		if (Static0->Geometry()->HasBoundingBox() && Static1->Geometry()->HasBoundingBox())
+		{
+			// Matrix multiplication is reversed intentionally to be compatible with unreal
+			TBox<T, d> Box1 = Static0->Geometry()->BoundingBox().TransformedBox(Transform1 * Transform2.Inverse());
+			Box1.Thicken(CreationThreshold);
+			TBox<T, d> Box2 = Static1->Geometry()->BoundingBox();
+			Box2.Thicken(CreationThreshold);
+			if (!Box1.Intersects(Box2))
 			{
-				const TVector<T, d> Radius = WorldSpaceX2 - P1;
-				P1 -= PBDRigid1->InvM() * Delta;
-				Q1 += TRotation<T, d>(WorldSpaceInvI2 * TVector<T, d>::CrossProduct(Radius, -Delta), 0.f) * Q1 * T(0.5);
-				Q1.Normalize();
+				continue;
 			}
 		}
+		const TVector<T, d> Midpoint = (P0 + P1) / (T)2;
+		TVector<T, d> Normal1;
+		const T Phi1 = Static0->Geometry()->PhiWithNormal(Transform1.InverseTransformPosition(Midpoint), Normal1);
+		Normal1 = Transform2.TransformVector(Normal1);
+		TVector<T, d> Normal2;
+		const T Phi2 = Static1->Geometry()->PhiWithNormal(Transform2.InverseTransformPosition(Midpoint), Normal2);
+		Normal2 = Transform2.TransformVector(Normal2);
+		if ((Phi1 + Phi2) > CreationThreshold)
+		{
+			continue;
+		}
+		TVector<T, d> Location0 = Midpoint - Phi1 * Normal1;
+		TVector<T, d> Location1 = Midpoint - Phi2 * Normal2;
+		TVector<TVector<T, 3>, 2> Distance;
+		Distance[0] = Q0.Inverse().RotateVector(Location0 - P0);
+		Distance[1] = Q0.Inverse().RotateVector(Location1 - P1);
+		Distances[ConstraintIndex].Add(MoveTemp(Distance));
+		SpringDistances[ConstraintIndex].Add((Location0 - Location1).Size());
+	}
+}
+
+template<class T, int d>
+TVector<T, d> TPBDRigidDynamicSpringConstraints<T, d>::GetDelta(const TVector<T, d>& WorldSpaceX1, const TVector<T, d>& WorldSpaceX2, const int32 ConstraintIndex, const int32 SpringIndex) const
+{
+	TPBDRigidParticleHandle<T, d>* PBDRigid0 = Constraints[ConstraintIndex][0]->AsDynamic();
+	TPBDRigidParticleHandle<T, d>* PBDRigid1 = Constraints[ConstraintIndex][1]->AsDynamic();
+
+	if (!PBDRigid0 && !PBDRigid1)
+	{
+		return TVector<T, d>(0);
+	}
+
+	const TVector<T, d> Difference = WorldSpaceX2 - WorldSpaceX1;
+	float Distance = Difference.Size();
+	check(Distance > 1e-7);
+
+	const T InvM0 = PBDRigid0 ? PBDRigid0->InvM() : (T)0;
+	const T InvM1 = PBDRigid1 ? PBDRigid1->InvM() : (T)0;
+	const TVector<T, d> Direction = Difference / Distance;
+	const TVector<T, d> Delta = (Distance - SpringDistances[ConstraintIndex][SpringIndex]) * Direction;
+	return Stiffness * Delta / (InvM0 + InvM1);
+}
+
+template<class T, int d>
+void TPBDRigidDynamicSpringConstraints<T, d>::ApplySingle(const T Dt, int32 ConstraintIndex) const
+{
+	TGeometryParticleHandle<T, d>* Static0 = Constraints[ConstraintIndex][0];
+	TGeometryParticleHandle<T, d>* Static1 = Constraints[ConstraintIndex][1];
+	TPBDRigidParticleHandle<T, d>* PBDRigid0 = Static0->AsDynamic();
+	TPBDRigidParticleHandle<T, d>* PBDRigid1 = Static1->AsDynamic();
+	check(PBDRigid0 || PBDRigid1);
+	check(!PBDRigid0 || !PBDRigid1 || (PBDRigid0->Island() == PBDRigid1->Island()));
+
+	TRotation<T, d>& Q0 = PBDRigid0 ? PBDRigid0->Q() : Static0->R();
+	TRotation<T, d>& Q1 = PBDRigid1 ? PBDRigid1->Q() : Static1->R();
+	TVector<T, d>& P0 = PBDRigid0 ? PBDRigid0->P() : Static0->X();
+	TVector<T, d>& P1 = PBDRigid1 ? PBDRigid1->P() : Static1->X();
+
+	const int32 NumSprings = SpringDistances[ConstraintIndex].Num();
+	const PMatrix<T, d, d> WorldSpaceInvI1 = PBDRigid0? (Q0 * FMatrix::Identity) * PBDRigid0->InvI() * (Q0 * FMatrix::Identity).GetTransposed() : PMatrix<T, d, d>(0);
+	const PMatrix<T, d, d> WorldSpaceInvI2 = PBDRigid1 ? (Q1 * FMatrix::Identity) * PBDRigid1->InvI() * (Q1 * FMatrix::Identity).GetTransposed() : PMatrix<T, d, d>(0);;
+	for (int32 SpringIndex = 0; SpringIndex < NumSprings; ++SpringIndex)
+	{
+		const TVector<T, d>& Distance0 = Distances[ConstraintIndex][SpringIndex][0];
+		const TVector<T, d>& Distance1 = Distances[ConstraintIndex][SpringIndex][1];
+		const TVector<T, d> WorldSpaceX1 = Q0.RotateVector(Distance0) + P0;
+		const TVector<T, d> WorldSpaceX2 = Q1.RotateVector(Distance1) + P1;
+		const TVector<T, d> Delta = GetDelta(WorldSpaceX1, WorldSpaceX2, ConstraintIndex, SpringIndex);
+
+		if (PBDRigid0)
+		{
+			const TVector<T, d> Radius = WorldSpaceX1 - P0;
+			P0 += PBDRigid0->InvM() * Delta;
+			Q0 += TRotation<T, d>(WorldSpaceInvI1 * TVector<T, d>::CrossProduct(Radius, Delta), 0.f) * Q0 * T(0.5);
+			Q0.Normalize();
+		}
+
+		if (PBDRigid1)
+		{
+			const TVector<T, d> Radius = WorldSpaceX2 - P1;
+			P1 -= PBDRigid1->InvM() * Delta;
+			Q1 += TRotation<T, d>(WorldSpaceInvI2 * TVector<T, d>::CrossProduct(Radius, -Delta), 0.f) * Q1 * T(0.5);
+			Q1.Normalize();
+		}
 	}
 }
 
diff --git a/Engine/Source/Runtime/Experimental/Chaos/Private/Chaos/PBDRigidsEvolutionGBF.cpp b/Engine/Source/Runtime/Experimental/Chaos/Private/Chaos/PBDRigidsEvolutionGBF.cpp
--- a/Engine/Source/Runtime/Experimental/Chaos/Private/Chaos/PBDRigidsEvolutionGBF.cpp
+++ b/Engine/Source/Runtime/Experimental/Chaos/Private/Chaos/PBDRigidsEvolutionGBF.cpp
@@ -166,6 +166,12 @@ TPBDRigidsEvolutionGBF<T, d>::TPBDRigidsEvolutionGBF(TPBDRigidsSOAs<T, d>& InPar
 		});
 	});
 
+	AddForceFunction([this](TTransientPBDRigidParticleHandle<T, d>& HandleIn, const T Dt) 
+	{
+		ExternalForces.Apply(HandleIn, Dt);
+	});
+
+
 	AddConstraintRule(&CollisionRule);
 }
 
diff --git a/Engine/Source/Runtime/Experimental/Chaos/Public/Chaos/BoundingVolume.h b/Engine/Source/Runtime/Experimental/Chaos/Public/Chaos/BoundingVolume.h
--- a/Engine/Source/Runtime/Experimental/Chaos/Public/Chaos/BoundingVolume.h
+++ b/Engine/Source/Runtime/Experimental/Chaos/Public/Chaos/BoundingVolume.h
@@ -28,7 +28,7 @@ struct TPayloadHelper
 	template <typename TParticle>
 	static TPayloadType GetPayload(TParticle& Particle, int32 Idx)
 	{
-		return Particle.Handle();
+		return TPayloadType(Particle.Handle());
 	}
 };
 
@@ -112,9 +112,9 @@ public:
 		return *this;
 	}
 
-	TBoundingVolume<TSOA, TPayloadType, T, d> Copy() const
+	TUniquePtr<TBoundingVolume<TSOA, TPayloadType, T, d>> Copy() const
 	{
-		return TBoundingVolume<TSOA, TPayloadType, T, d>(*this);
+		return TUniquePtr<TBoundingVolume<TSOA, TPayloadType, T, d>>(new TBoundingVolume<TSOA, TPayloadType, T, d>(*this));
 	}
 
 	template <typename ParticleView>
@@ -155,6 +155,35 @@ public:
 		return IntersectionList;
 	}
 
+	virtual void RemoveElement(const TPayloadType& Payload) override
+	{
+		//todo: remove with a faster search
+		if(MGlobalPayloads.RemoveSwap(Payload) == 0)
+		{
+			const TVector<int32, d> NumCells = MGrid.Counts();
+			for (int32 X = 0; X < NumCells[0]; ++X)
+			{
+				for (int32 Y = 0; Y < NumCells[1]; ++Y)
+				{
+					for (int32 Z = 0; Z < NumCells[2]; ++Z)
+					{
+						TArray<FCellElement>& Elems = MElements(X, Y, Z);
+						int32 ElemIdx = 0;
+						for (FCellElement& Elem : Elems)
+						{
+							if (Elem.Payload == Payload)
+							{
+								Elems.RemoveAtSwap(ElemIdx);
+								break;
+							}
+							++ElemIdx;
+						}
+					}
+				}
+			}
+		}
+	}
+
 	// Begin ISpatialAcceleration interface
 	virtual TArray<TPayloadType> FindAllIntersections(const TBox<T, d>& Box) const override { return FindAllIntersectionsImp(Box); }
 	virtual TArray<TPayloadType> FindAllIntersections(const TVector<T, d>& Point) const override
diff --git a/Engine/Source/Runtime/Experimental/Chaos/Public/Chaos/Box.h b/Engine/Source/Runtime/Experimental/Chaos/Public/Chaos/Box.h
--- a/Engine/Source/Runtime/Experimental/Chaos/Public/Chaos/Box.h
+++ b/Engine/Source/Runtime/Experimental/Chaos/Public/Chaos/Box.h
@@ -339,6 +339,40 @@ namespace Chaos
 			return MakePair(TVector<T, d>(0), false);
 		}
 
+		virtual TVector<T, d> FindGeometryOpposingNormal(const TVector<T, d>& DenormDir, int32 FaceIndex, const TVector<T, d>& OriginalNormal) const override
+		{
+			// Find which faces were included in the contact normal, and for multiple faces, use the one most opposing the sweep direction.
+			TVector<T,d> BestNormal(OriginalNormal);
+			T BestOpposingDot = TNumericLimits<T>::Max();
+
+			for (int32 Axis = 0; Axis < d; Axis++)
+			{
+				// Select axis of face to compare to, based on normal.
+				if (OriginalNormal[Axis] > KINDA_SMALL_NUMBER)
+				{
+					const float TraceDotFaceNormal = DenormDir[Axis]; // TraceDirDenormLocal.dot(BoxFaceNormal)
+					if (TraceDotFaceNormal < BestOpposingDot)
+					{
+						BestOpposingDot = TraceDotFaceNormal;
+						BestNormal = TVector<T,d>(0);
+						BestNormal[Axis] = 1;
+					}
+				}
+				else if (OriginalNormal[Axis] < -KINDA_SMALL_NUMBER)
+				{
+					const float TraceDotFaceNormal = -DenormDir[Axis]; // TraceDirDenormLocal.dot(BoxFaceNormal)
+					if (TraceDotFaceNormal < BestOpposingDot)
+					{
+						BestOpposingDot = TraceDotFaceNormal;
+						BestNormal = FVector(0.f);
+						BestNormal[Axis] = -1.f;
+					}
+				}
+			}
+
+			return BestNormal;
+		}
+
 		virtual TVector<T, d> Support(const TVector<T, d>& Direction, const T Thickness) const override
 		{
 			TVector<T, d> ChosenPt;
diff --git a/Engine/Source/Runtime/Experimental/Chaos/Public/Chaos/ConstraintHandle.h b/Engine/Source/Runtime/Experimental/Chaos/Public/Chaos/ConstraintHandle.h
new file mode 100644
--- /dev/null
+++ b/Engine/Source/Runtime/Experimental/Chaos/Public/Chaos/ConstraintHandle.h
@@ -0,0 +1,85 @@
+// Copyright 1998-2019 Epic Games, Inc. All Rights Reserved.
+#pragma once
+
+#include "Chaos/Declares.h"
+#include "Chaos/Vector.h"
+
+namespace Chaos
+{
+	template<class T, int d>
+	class TPBDConstraintContainer;
+
+	/**
+	 * Base class for constraint handles.
+	 */
+	template<class T, int d>
+	class CHAOS_API TConstraintHandle
+	{
+	public:
+		using FReal = T;
+		static const int Dimensions = d;
+		using FGeometryParticleHandle = TGeometryParticleHandle<FReal, Dimensions>;
+
+		TConstraintHandle() : ConstraintIndex(INDEX_NONE) {}
+		TConstraintHandle(int32 InConstraintIndex) : ConstraintIndex(InConstraintIndex) {}
+		~TConstraintHandle() {}
+
+		bool IsValid() const
+		{
+			return (ConstraintIndex != INDEX_NONE);
+		}
+
+		int32 GetConstraintIndex() const
+		{
+			return ConstraintIndex;
+		}
+
+	protected:
+		friend class TPBDConstraintContainer<T, d>;
+
+		int32 ConstraintIndex;
+	};
+
+
+	/**
+	 * Utility base class for ConstraintHandles. Provides basic functionality common to most constraint containers.
+	 */
+	template<typename T_CONTAINER>
+	class CHAOS_API TContainerConstraintHandle : public TConstraintHandle<typename T_CONTAINER::FReal, T_CONTAINER::Dimensions>
+	{
+	public:
+		using Base = TConstraintHandle<typename T_CONTAINER::FReal, T_CONTAINER::Dimensions>;
+		using FReal = typename Base::FReal;
+		static const int Dimensions = Base::Dimensions;
+		using FGeometryParticleHandle = typename Base::FGeometryParticleHandle;
+		using FConstraintContainer = T_CONTAINER;
+
+		TContainerConstraintHandle() : ConstraintContainer(nullptr) {}
+		TContainerConstraintHandle(FConstraintContainer* InConstraintContainer, int32 InConstraintIndex) : TConstraintHandle<FReal, Dimensions>(InConstraintIndex), ConstraintContainer(InConstraintContainer) {}
+
+		void RemoveConstraint() { ConstraintContainer->RemoveConstraint(ConstraintIndex); }
+
+		TVector<FGeometryParticleHandle*, 2> GetConstrainedParticles() const { return ConstraintContainer->GetConstrainedParticles(ConstraintIndex); }
+
+	protected:
+		using Base::ConstraintIndex;
+		FConstraintContainer* ConstraintContainer;
+	};
+
+
+	/**
+	 * An allocator for constraint handles.
+	 *
+	 * @todo(ccaulfield): block allocator for handles, or support custom allocators in constraint containers.
+	 */
+	template<class T_CONTAINER>
+	class CHAOS_API TConstraintHandleAllocator
+	{
+	public:
+		using FConstraintContainer = T_CONTAINER;
+		using FConstraintHandle = typename FConstraintContainer::FConstraintHandle;
+
+		FConstraintHandle* AllocHandle(FConstraintContainer* ConstraintContainer, int32 ConstraintIndex) { return new FConstraintHandle(ConstraintContainer, ConstraintIndex); }
+		void FreeHandle(FConstraintHandle* Handle) { delete Handle; }
+	};
+}
diff --git a/Engine/Source/Runtime/Experimental/Chaos/Public/Chaos/Convex.h b/Engine/Source/Runtime/Experimental/Chaos/Public/Chaos/Convex.h
--- a/Engine/Source/Runtime/Experimental/Chaos/Public/Chaos/Convex.h
+++ b/Engine/Source/Runtime/Experimental/Chaos/Public/Chaos/Convex.h
@@ -94,6 +94,40 @@ namespace Chaos
 			return MakePair(TVector<T, 3>(0), false);
 		}
 
+		virtual int32 FindMostOpposingFace(const TVector<T, 3>& Position, const TVector<T, 3>& UnitDir, int32 HintFaceIndex, T SearchDist) const override
+		{
+			//todo: use hill climbing
+			int32 MostOpposingIdx = INDEX_NONE;
+			T MostOpposingDot = TNumericLimits<T>::Max();
+			int32 Idx = 0;
+			for (const TPlane<T, d>& Plane : Planes)
+			{
+				const T Distance = Plane.SignedDistance(Position);
+				if (FMath::Abs(Distance) < SearchDist)
+				{
+					const T Dot = TVector<T, d>::DotProduct(Plane.Normal(), UnitDir);
+					if (Dot < MostOpposingDot)
+					{
+						MostOpposingDot = Dot;
+						MostOpposingIdx = Idx;
+					}
+				}
+				++Idx;
+			}
+			ensure(MostOpposingIdx != INDEX_NONE);
+			return MostOpposingIdx;
+		}
+
+		TVector<T, d> FindGeometryOpposingNormal(const TVector<T, d>& DenormDir, int32 FaceIndex, const TVector<T, d>& OriginalNormal) const
+		{
+			if (ensure(FaceIndex != INDEX_NONE))
+			{
+				const TPlane<float, 3>& OpposingFace = GetFaces()[FaceIndex];
+				return OpposingFace.Normal();
+			}
+			return FVector(0, 0, 1);
+		}
+
 		virtual TVector<T, d> Support(const TVector<T, d>& Direction, const T Thickness) const override
 		{
 			T MaxDot = TNumericLimits<T>::Lowest();
@@ -124,6 +158,11 @@ namespace Chaos
 			return SurfaceParticles;
 		}
 
+		const TArray<TPlane<T, d>>& GetFaces() const
+		{
+			return Planes;
+		}
+
 		virtual uint32 GetTypeHash() const override
 		{
 			uint32 Result = LocalBoundingBox.GetTypeHash();
diff --git a/Engine/Source/Runtime/Experimental/Chaos/Public/Chaos/Declares.h b/Engine/Source/Runtime/Experimental/Chaos/Public/Chaos/Declares.h
--- a/Engine/Source/Runtime/Experimental/Chaos/Public/Chaos/Declares.h
+++ b/Engine/Source/Runtime/Experimental/Chaos/Public/Chaos/Declares.h
@@ -3,6 +3,7 @@
 #pragma once
 
 #include "Chaos/ParticleHandleFwd.h"
+
 namespace Chaos
 {
 	class FPBDRigidsSolver;
@@ -22,6 +23,7 @@ namespace Chaos
 #define TODO_REIMPLMEENT_EVOLUTION_ACCESSORS (CHAOS_PARTICLEHANDLE_TODO && 1)
 #define TODO_REIMPLEMENT_DYNAMIC_CONSTRAINT_ACCESSORS (CHAOS_PARTICLEHANDLE_TODO && 1)
 #define TODO_REIMPLEMENT_FIELDS_TO_USE_PARTICLEHANDLES (CHAOS_PARTICLEHANDLE_TODO && 1)
+#define TODO_REIMPLEMENT_GEOMETRY_COLLECTION_PHYSICS (CHAOS_PARTICLEHANDLE_TODO && 1)
 #define TODO_REIMPLEMENT_RIGID_CLUSTERING (CHAOS_PARTICLEHANDLE_TODO && 1)
 #define TODO_REIMPLEMENT_GETFLOORINDEX (CHAOS_PARTICLEHANDLE_TODO && 1)
 #define TODO_REIMPLEMENT_PHYSICS_PROXY_REVERSE_MAPPING (CHAOS_PARTICLEHANDLE_TODO && 1)
@@ -38,3 +40,5 @@ namespace Chaos
 #define TODO_REIMPLEMENT_INITIALIZE_EVOLUTION_FROM_PARTICLE_DATA (CHAOS_PARTICLEHANDLE_TODO && 1)
 #define TODO_REIMPLEMENT_SKELETAL_TESTS (CHAOS_PARTICLEHANDLE_TODO && 1)
 #define TODO_REIMPLEMENT_SCENEQUERY_CROSSENGINE (0)		// ISQAccelerator must be re-implemented for Chaos/PhysX SQ interop. Without this, GDC 2019 demos will not work as designed.
+
+#define CHAOS_CONSTRAINTHANDLE_TODO 0
diff --git a/Engine/Source/Runtime/Experimental/Chaos/Public/Chaos/Framework/PhysicsProxy.h b/Engine/Source/Runtime/Experimental/Chaos/Public/Chaos/Framework/PhysicsProxy.h
--- a/Engine/Source/Runtime/Experimental/Chaos/Public/Chaos/Framework/PhysicsProxy.h
+++ b/Engine/Source/Runtime/Experimental/Chaos/Public/Chaos/Framework/PhysicsProxy.h
@@ -77,8 +77,8 @@ class TPhysicsProxy : public IPhysicsProxyBase
 {
 
 public:
-	using FParticleType = typename Concrete;
-	using FParticleData = typename ConcreteData;
+	using FParticleType = Concrete;
+	using FParticleData = ConcreteData;
 
 	using FParticlesType = Chaos::TPBDRigidParticles<float, 3>;
 	using FCollisionConstraintsType = Chaos::TPBDCollisionConstraint<float, 3>;
diff --git a/Engine/Source/Runtime/Experimental/Chaos/Public/Chaos/HeightField.h b/Engine/Source/Runtime/Experimental/Chaos/Public/Chaos/HeightField.h
--- a/Engine/Source/Runtime/Experimental/Chaos/Public/Chaos/HeightField.h
+++ b/Engine/Source/Runtime/Experimental/Chaos/Public/Chaos/HeightField.h
@@ -31,7 +31,8 @@ namespace Chaos
 		virtual bool Overlap(const TVector<T, 3>& Point, const T Thickness) const override;
 		bool OverlapGeom(const TImplicitObject<T, 3>& QueryGeom, const TRigidTransform<T, 3>& QueryTM, const T Thickness) const;
 		bool SweepGeom(const TImplicitObject<T, 3>& QueryGeom, const TRigidTransform<T, 3>& StartTM, const TVector<T, 3>& Dir, const T Length, T& OutTime, TVector<T, 3>& OutPosition, TVector<T, 3>& OutNormal, int32& OutFaceIndex, const T Thickness = 0) const;
-		virtual int32 FindMostOpposingFace(const TVector<T, 3>& Position, const TVector<T, 3>& UnitDir, int32 HintFaceIndex) const override;
+		virtual int32 FindMostOpposingFace(const TVector<T, 3>& Position, const TVector<T, 3>& UnitDir, int32 HintFaceIndex, T SearchDist) const override;
+		virtual TVector<T, 3> FindGeometryOpposingNormal(const TVector<T, 3>& DenormDir, int32 FaceIndex, const TVector<T, 3>& OriginalNormal) const override;
 
 		virtual const TBox<T, 3>& BoundingBox() const
 		{
diff --git a/Engine/Source/Runtime/Experimental/Chaos/Public/Chaos/ISpatialAcceleration.h b/Engine/Source/Runtime/Experimental/Chaos/Public/Chaos/ISpatialAcceleration.h
--- a/Engine/Source/Runtime/Experimental/Chaos/Public/Chaos/ISpatialAcceleration.h
+++ b/Engine/Source/Runtime/Experimental/Chaos/Public/Chaos/ISpatialAcceleration.h
@@ -97,6 +97,11 @@ public:
 	virtual void Sweep(const TVector<T, d>& Start, const TVector<T, d>& Dir, T OriginalLength, const TVector<T, d> QueryHalfExtents, ISpatialVisitor<TPayloadType, T>& Visitor, const TVector<T, d>& Scale = TVector<T, d>(1)) const {}
 	virtual void Overlap(const TBox<T, d>& QueryBounds, ISpatialVisitor<TPayloadType, T>& Visitor, const TVector<T, d>& Scale = TVector<T, d>(1)) const {}
 
+	virtual void RemoveElement(const TPayloadType& Payload)
+	{
+		check(false);	//not implemented
+	}
+
 #if !UE_BUILD_SHIPPING
 	virtual void DebugDraw(ISpacialDebugDrawInterface<T>* InInterface) const {}
 	virtual void DumpStats() const {}
diff --git a/Engine/Source/Runtime/Experimental/Chaos/Public/Chaos/ImplicitObject.h b/Engine/Source/Runtime/Experimental/Chaos/Public/Chaos/ImplicitObject.h
--- a/Engine/Source/Runtime/Experimental/Chaos/Public/Chaos/ImplicitObject.h
+++ b/Engine/Source/Runtime/Experimental/Chaos/Public/Chaos/ImplicitObject.h
@@ -161,13 +161,27 @@ public:
 		@param Position - local position to search around (for example an edge of a convex hull)
 		@param UnitDir - the direction we want to oppose (for example a ray moving into the edge of a convex hull would get the face with the most negative dot(FaceNormal, UnitDir)
 		@param HintFaceIndex - for certain geometry we can use this to accelerate the search.
+		@return Index of the most opposing face
 	*/
-	virtual int32 FindMostOpposingFace(const TVector<T, d>& Position, const TVector<T, d>& UnitDir, int32 HintFaceIndex) const
+	virtual int32 FindMostOpposingFace(const TVector<T, d>& Position, const TVector<T, d>& UnitDir, int32 HintFaceIndex, T SearchDist) const
 	{
 		//Many objects have no concept of a face
 		return INDEX_NONE;
 	}
 
+	/** Given a normal and a face index, compute the most opposing normal associated with the underlying geometry features.
+		For example a sphere swept against a box may not give a normal associated with one of the box faces. This function will return a normal associated with one of the faces.
+		@param DenormDir - the direction we want to oppose
+		@param FaceIndex - the face index associated with the geometry (for example if we hit a specific face of a convex hull)
+		@param OriginalNormal - the original normal given by something like a sphere sweep
+		@return The most opposing normal associated with the underlying geometry's feature (like a face)
+	*/
+	virtual TVector<T,d> FindGeometryOpposingNormal(const TVector<T, d>& DenormDir, int32 FaceIndex, const TVector<T,d>& OriginalNormal) const
+	{
+		//Many objects have no concept of a face
+		return OriginalNormal;
+	}
+
 	//This gives derived types a way to do an overlap check without calling PhiWithNormal todo: this api is confusing
 	virtual bool Overlap(const TVector<T, d>& Point, const T Thickness) const
 	{
diff --git a/Engine/Source/Runtime/Experimental/Chaos/Public/Chaos/ImplicitObjectScaled.h b/Engine/Source/Runtime/Experimental/Chaos/Public/Chaos/ImplicitObjectScaled.h
--- a/Engine/Source/Runtime/Experimental/Chaos/Public/Chaos/ImplicitObjectScaled.h
+++ b/Engine/Source/Runtime/Experimental/Chaos/Public/Chaos/ImplicitObjectScaled.h
@@ -97,6 +97,38 @@ public:
 		return false;
 	}
 
+	virtual int32 FindMostOpposingFace(const TVector<T, d>& Position, const TVector<T, d>& UnitDir, int32 HintFaceIndex, T SearchDist) const override
+	{
+		ensure(FMath::IsNearlyEqual(UnitDir.SizeSquared(), 1, KINDA_SMALL_NUMBER));
+
+		const TVector<T, d> UnscaledPosition = MInvScale * Position;
+		TVector<T, d> UnscaledDir = MInvScale * UnitDir;
+		const T UnscaledLength = UnscaledDir.SafeNormalize();
+		const T UnscaledSearchDist = SearchDist * MInvScale.Max();	//this is not quite right since it's no longer a sphere, but the whole thing is fuzzy anyway
+
+		return MObject->FindMostOpposingFace(UnscaledPosition, UnscaledDir, HintFaceIndex, UnscaledSearchDist);
+	}
+
+	virtual TVector<T, 3> FindGeometryOpposingNormal(const TVector<T, d>& DenormDir, int32 FaceIndex, const TVector<T, d>& OriginalNormal) const override
+	{
+		TVector<T, 3> LocalDenormDir = DenormDir * MInvScale;
+		TVector<T, 3> LocalOriginalNormal = OriginalNormal * MInvScale;
+		
+		if (ensure(LocalOriginalNormal.SafeNormalize()) == 0)
+		{
+			LocalOriginalNormal = TVector<T,3>(0,0,1);
+		}
+
+		const TVector<T, d> LocalNormal = MObject->FindGeometryOpposingNormal(LocalDenormDir, FaceIndex, LocalOriginalNormal);
+		TVector<T, d> Normal = LocalNormal * MScale;
+		if (ensure(Normal.SafeNormalize()) == 0)
+		{
+			Normal = TVector<T,3>(0,0,1);
+		}
+
+		return Normal;
+	}
+
 	virtual bool Overlap(const TVector<T, d>& Point, const T Thickness) const override
 	{
 		const TVector<T, d> UnscaledPoint = MInvScale * Point;
diff --git a/Engine/Source/Runtime/Experimental/Chaos/Public/Chaos/ImplicitObjectTransformed.h b/Engine/Source/Runtime/Experimental/Chaos/Public/Chaos/ImplicitObjectTransformed.h
--- a/Engine/Source/Runtime/Experimental/Chaos/Public/Chaos/ImplicitObjectTransformed.h
+++ b/Engine/Source/Runtime/Experimental/Chaos/Public/Chaos/ImplicitObjectTransformed.h
@@ -105,6 +105,21 @@ public:
 		return false;
 	}
 
+	virtual int32 FindMostOpposingFace(const TVector<T, 3>& Position, const TVector<T, 3>& UnitDir, int32 HintFaceIndex, T SearchDistance) const override
+	{
+		const TVector<T, d> LocalPosition = MTransform.InverseTransformPositionNoScale(Position);
+		const TVector<T, d> LocalDir = MTransform.InverseTransformVectorNoScale(UnitDir);
+		return MObject->FindMostOpposingFace(LocalPosition, LocalDir, HintFaceIndex, SearchDistance);
+	}
+
+	virtual TVector<T, 3> FindGeometryOpposingNormal(const TVector<T, d>& DenormDir, int32 FaceIndex, const TVector<T, d>& OriginalNormal) const override
+	{
+		const TVector<T, d> LocalDenormDir = MTransform.InverseTransformVectorNoScale(DenormDir);
+		const TVector<T, d> LocalOriginalNormal = MTransform.InverseTransformVectorNoScale(OriginalNormal);
+		const TVector<T, d> LocalNormal = MObject->FindGeometryOpposingNormal(LocalDenormDir, FaceIndex, LocalOriginalNormal);
+		return MTransform.TransformVectorNoScale(LocalNormal);
+	}
+
 	virtual bool Overlap(const TVector<T, d>& Point, const T Thickness) const override
 	{
 		const TVector<T, d> LocalPoint = MTransform.InverseTransformPosition(Point);
diff --git a/Engine/Source/Runtime/Experimental/Chaos/Public/Chaos/PBD6DJointConstraints.h b/Engine/Source/Runtime/Experimental/Chaos/Public/Chaos/PBD6DJointConstraints.h
new file mode 100644
--- /dev/null
+++ b/Engine/Source/Runtime/Experimental/Chaos/Public/Chaos/PBD6DJointConstraints.h
@@ -0,0 +1,215 @@
+// Copyright 1998-2019 Epic Games, Inc. All Rights Reserved.
+#pragma once
+
+#include "Chaos/Array.h"
+#include "Chaos/Transform.h"
+#include "Chaos/Vector.h"
+
+#include "Chaos/ConstraintHandle.h"
+#include "Chaos/ParticleHandleFwd.h"
+#include "Chaos/PBDConstraintContainer.h"
+
+namespace Chaos
+{
+	template<class T, int d>
+	class TPBD6DJointConstraints;
+
+	template<class T, int d>
+	class CHAOS_API TPBD6DJointConstraintHandle : public TContainerConstraintHandle<TPBD6DJointConstraints<T, d>>
+	{
+	public:
+		using Base = TContainerConstraintHandle<TPBD6DJointConstraints<T, d>>;
+		using FConstraintContainer = TPBD6DJointConstraints<T, d>;
+
+		TPBD6DJointConstraintHandle() {}
+		TPBD6DJointConstraintHandle(FConstraintContainer* InConstraintContainer, int32 InConstraintIndex) : TContainerConstraintHandle<TPBD6DJointConstraints<T, d>>(InConstraintContainer, InConstraintIndex) {}
+
+	protected:
+		using Base::ConstraintIndex;
+		using Base::ConstraintContainer;
+	};
+
+
+	enum class E6DJointMotionType
+	{
+		Free,
+		Limited,
+		Locked,
+	};
+
+
+	template<class T, int d>
+	class CHAOS_API TPBD6DJointConstraintSettings
+	{
+	public:
+		using FTransformPair = TVector<TRigidTransform<T, d>, 2>;
+
+		TPBD6DJointConstraintSettings();
+
+	private:
+		friend class TPBD6DJointConstraints<T, d>;
+
+		// Particle-relative joint axes and positions
+		// X Axis: Twist
+		// Y Axis: Swing 1
+		// Z Axis: Swing 2
+		FTransformPair ConstraintFrames;
+
+		// Motion Types
+		TVector<E6DJointMotionType, d> LinearMotionTypes;
+		TVector<E6DJointMotionType, d> AngularMotionTypes;
+
+		// Motion Limits (if MotionType == Limited)
+		TVector<T, d> LinearLimits;
+		TVector<T, d> AngularLimits;
+	};
+
+
+	/**
+	 * A joint restricting up to 6 degrees of freedom, with linear and angular limits.
+	 */
+	template<class T, int d>
+	class CHAOS_API TPBD6DJointConstraints : public TPBDConstraintContainer<T, d>
+	{
+	public:
+		using Base = TPBDConstraintContainer<T, d>;
+		using FReal = T;
+		static const int Dimensions = d;
+		using FConstraintHandle = TPBD6DJointConstraintHandle<FReal, Dimensions>;
+		using FConstraintHandleAllocator = TConstraintHandleAllocator<TPBD6DJointConstraints<FReal, Dimensions>>;
+		using FParticlePair = TVector<TGeometryParticleHandle<T, d>*, 2>;
+		using FVectorPair = TVector<TVector<T, d>, 2>;
+		using FTransformPair = TVector<TRigidTransform<T, d>, 2>;
+
+		TPBD6DJointConstraints(const T InStiffness = (T)1)
+			: Stiffness(InStiffness)
+		{
+		}
+
+		virtual ~TPBD6DJointConstraints() {}
+
+		//
+		// Constraint Container API
+		//
+
+		/**
+		 * Get the number of constraints.
+		 */
+		int32 NumConstraints() const
+		{
+			return ConstraintParticles.Num();
+		}
+
+		/** 
+		 * Add a constraint with particle-space constraint offsets. 
+		 */
+		FConstraintHandle* AddConstraint(const FParticlePair& InConstrainedParticles, const FTransformPair& ConstraintFrames)
+		{
+			int ConstraintIndex = Handles.Num();
+			Handles.Add(HandleAllocator.AllocHandle(this, ConstraintIndex));
+			ConstraintParticles.Add(InConstrainedParticles);
+			ConstraintSettings.Add(TPBD6DJointConstraintSettings<T, d>());
+			InitConstraint(ConstraintIndex, ConstraintFrames);
+			return Handles.Last();
+		}
+
+		/**
+		 * Remove the specified constraint.
+		 */
+		void RemoveConstraint(int ConstraintIndex)
+		{
+			FConstraintHandle* ConstraintHandle = Handles[ConstraintIndex];
+			if (ConstraintHandle != nullptr)
+			{
+				// Release the handle for the freed constraint
+				HandleAllocator.FreeHandle(ConstraintHandle);
+				Handles[ConstraintIndex] = nullptr;
+			}
+
+			// Swap the last constraint into the gap to keep the array packed
+			ConstraintParticles.RemoveAtSwap(ConstraintIndex);
+			ConstraintSettings.RemoveAtSwap(ConstraintIndex);
+			Handles.RemoveAtSwap(ConstraintIndex);
+
+			// Update the handle for the constraint that was moved
+			if (ConstraintIndex < Handles.Num())
+			{
+				SetConstraintIndex(Handles[ConstraintIndex], ConstraintIndex);
+			}
+		}
+
+		// @todo(ccaulfield): rename/remove  this
+		void RemoveConstraints(const TSet<TGeometryParticleHandle<T, d>*>& RemovedParticles)
+		{
+		}
+
+		//
+		// Constraint API
+		//
+
+		const FConstraintHandle* GetConstraintHandle(int32 ConstraintIndex) const
+		{
+			return Handles[ConstraintIndex];
+		}
+
+		FConstraintHandle* GetConstraintHandle(int32 ConstraintIndex)
+		{
+			return Handles[ConstraintIndex];
+		}
+
+		/**
+		 * Get the particles that are affected by the specified constraint.
+		 */
+		const FParticlePair& GetConstrainedParticles(int32 ConstraintIndex) const
+		{
+			return ConstraintParticles[ConstraintIndex];
+		}
+
+
+		//
+		// Island Rule API
+		//
+
+		void UpdatePositionBasedState(const T Dt)
+		{
+		}
+
+		void Apply(const T Dt, const TArray<FConstraintHandle*>& InConstraintHandles)
+		{
+			for (FConstraintHandle* ConstraintHandle : InConstraintHandles)
+			{
+				ApplySingle(Dt, ConstraintHandle->GetConstraintIndex());
+			}
+		}
+
+		// @todo(ccaulfield): remove  this
+		void ApplyPushOut(const T Dt, const TArray<FConstraintHandle*>& InConstraintHandles)
+		{
+		}
+
+
+	protected:
+		using Base::GetConstraintIndex;
+		using Base::SetConstraintIndex;
+
+		void ApplySingle(const T Dt, const int32 ConstraintIndex);
+
+	private:
+		void InitConstraint(int ConstraintIndex, const FTransformPair& InConstraintTransforms);
+
+		// Double dynamic body solve
+		void ApplyDynamicDynamic(const T Dt, const int32 ConstraintIndex, const int32 PBDRigid0Index, const int32 PBDRigid1Index, const bool bApplyProjection);
+		TVector<T, d> GetDeltaDynamicDynamic(const TVector<T, d>& P0, const TVector<T, d>& P1, const TVector<T, d>& WorldSpaceX0, const TVector<T, d>& WorldSpaceX1, const PMatrix<T, d, d>& WorldSpaceInvI0, const PMatrix<T, d, d>& WorldSpaceInvI1, const T InvM0, const T InvM1);
+
+		// Single dynamic body solve
+		void ApplyDynamicStatic(const T Dt, const int32 ConstraintIndex, const int32 PBDRigid0Index, const int32 Static1Index, const bool bApplyProjection);
+		TVector<T, d> GetDeltaDynamicKinematic(const TVector<T, d>& P0, const TVector<T, d>& WorldSpaceX0, const TVector<T, d>& WorldSpaceX1, const PMatrix<T, d, d>& WorldSpaceInvI0, const T InvM0);
+
+		TArray<TPBD6DJointConstraintSettings<T, d>> ConstraintSettings;
+		TArray<FParticlePair> ConstraintParticles;
+		T Stiffness;
+
+		TArray<FConstraintHandle*> Handles;
+		FConstraintHandleAllocator HandleAllocator;
+	};
+}
diff --git a/Engine/Source/Runtime/Experimental/Chaos/Public/Chaos/PBDChainConstraints.h b/Engine/Source/Runtime/Experimental/Chaos/Public/Chaos/PBDChainConstraints.h
--- a/Engine/Source/Runtime/Experimental/Chaos/Public/Chaos/PBDChainConstraints.h
+++ b/Engine/Source/Runtime/Experimental/Chaos/Public/Chaos/PBDChainConstraints.h
@@ -10,7 +10,7 @@ namespace Chaos
 template<class T, int d>
 class TPBDChainConstraints : public TPBDConstraintContainer<T, d>
 {
-  public:
+public:
 	TPBDChainConstraints(const TDynamicParticles<T, d>& InParticles, TArray<TArray<int32>>&& Constraints, const T Coefficient = (T)1)
 	    : MConstraints(Constraints), MCoefficient(Coefficient)
 	{
@@ -70,7 +70,10 @@ class TPBDChainConstraints : public TPBDConstraintContainer<T, d>
 		});
 	}
 
-  private:
+private:
+	using Base::GetConstraintIndex;
+	using Base::SetConstraintIndex;
+	  
 	TArray<TArray<int32>> MConstraints;
 	TArray<TArray<T>> MDists;
 	T MCoefficient;
diff --git a/Engine/Source/Runtime/Experimental/Chaos/Public/Chaos/PBDCollisionConstraint.h b/Engine/Source/Runtime/Experimental/Chaos/Public/Chaos/PBDCollisionConstraint.h
--- a/Engine/Source/Runtime/Experimental/Chaos/Public/Chaos/PBDCollisionConstraint.h
+++ b/Engine/Source/Runtime/Experimental/Chaos/Public/Chaos/PBDCollisionConstraint.h
@@ -2,7 +2,9 @@
 #pragma once
 
 #include "Chaos/BoundingVolume.h"
+#include "Chaos/ConstraintHandle.h"
 #include "Chaos/PBDCollisionTypes.h"
+#include "Chaos/PBDConstraintContainer.h"
 #include "Framework/BufferedData.h"
 
 #include <memory>
@@ -15,6 +17,9 @@ namespace Chaos
 template<typename T, int d>
 class TPBDCollisionConstraintAccessor;
 
+template<typename T, int d>
+class TPBDCollisionConstraint;
+
 template <typename T, int d>
 class TRigidTransform;
 
@@ -41,48 +46,162 @@ enum class ECollisionUpdateType
 };
 
 
+template<class T, int d>
+class CHAOS_API TPBDCollisionConstraintHandle : public TContainerConstraintHandle<TPBDCollisionConstraint<T, d>>
+{
+public:
+	using Base = TContainerConstraintHandle<TPBDCollisionConstraint<T, d>>;
+	using FConstraintContainer = TPBDCollisionConstraint<T, d>;
+
+	TPBDCollisionConstraintHandle() {}
+	TPBDCollisionConstraintHandle(FConstraintContainer* InConstraintContainer, int32 InConstraintIndex) : TContainerConstraintHandle<TPBDCollisionConstraint<T, d>>(InConstraintContainer, InConstraintIndex) {}
+
+protected:
+	using Base::ConstraintIndex;
+	using Base::ConstraintContainer;
+};
+
+/** Wrapper that holds both physics thread data and GT data. It's possible that the physics handle is null if we're doing operations entirely on external threads*/
+template <typename T, int d>
+class CHAOS_API TAccelerationStructureHandle
+{
+public:
+	TAccelerationStructureHandle(TGeometryParticleHandle<T, d>* InHandle)
+		: ExternalGeometryParticle(InHandle->GTGeometryParticle())
+		, GeometryParticleHandle(InHandle)
+	{
+	}
+
+	TAccelerationStructureHandle(TGeometryParticle<T, d>* InGeometryParticle)
+		: ExternalGeometryParticle(InGeometryParticle)
+		, GeometryParticleHandle(nullptr)
+	{
+	}
+
+	//Should only be used by GT and scene query threads where an appropriate lock has been acquired
+	TGeometryParticle<T, d>* GetExternalGeometryParticle_ExternalThread() const { return ExternalGeometryParticle; }
+
+	//Should only be used by PT
+	TGeometryParticleHandle<T, d>* GetGeometryParticleHandle_PhysicsThread() const { return GeometryParticleHandle; }
+
+	bool operator==(const TAccelerationStructureHandle<T, d>& Rhs) const
+	{
+		return ExternalGeometryParticle == Rhs.ExternalGeometryParticle;
+	}
+
+	bool operator<(const TAccelerationStructureHandle<T, d>& Rhs) const
+	{
+		return  ExternalGeometryParticle < Rhs.ExternalGeometryParticle;
+	}
+
+	uint32 GetTypeHash() const
+	{
+		return ::GetTypeHash((void*)ExternalGeometryParticle);
+	}
+
+private:
+	TGeometryParticle<T, d>* ExternalGeometryParticle;
+	TGeometryParticleHandle<T, d>* GeometryParticleHandle;
+};
+
+template <typename T, int d>
+uint32 GetTypeHash(const TAccelerationStructureHandle<T,d>& A)
+{
+	return A.GetTypeHash();
+}
+
 /**
  * Manages a set of contact constraints:
  *	- Performs collision detection to generate constraints.
  *	- Responsible for applying corrections to particles affected by the constraints.
  * @todo(ccaulfield): rename to TPBDCollisionConstraints
+ * @todo(ccaulfield): remove TPBDCollisionConstraintAccessor
+ * @todo(ccaulfield): separate collision detection from constraint container
  */
 template<typename T, int d>
-class CHAOS_API TPBDCollisionConstraint
+class CHAOS_API TPBDCollisionConstraint : public TPBDConstraintContainer<T, d>
 {
-  public:
+public:
+	using Base = TPBDConstraintContainer<T, d>;
 	friend class TPBDCollisionConstraintAccessor<T, d>;
-
-	typedef TRigidBodyContactConstraint<T, d> FRigidBodyContactConstraint;
+	using FReal = T;
+	static const int Dimensions = d;
+	using FConstraintHandle = TPBDCollisionConstraintHandle<T, d>;
+	using FConstraintHandleAllocator = TConstraintHandleAllocator<TPBDCollisionConstraint<T, d>>;
+	using FRigidBodyContactConstraint = TRigidBodyContactConstraint<T, d>;
 
 	TPBDCollisionConstraint(const TPBDRigidsSOAs<T,d>& InParticles, TArrayCollectionArray<bool>& Collided, const TArrayCollectionArray<TSerializablePtr<TChaosPhysicsMaterial<T>>>& PerParticleMaterials, const int32 PairIterations = 1, const T Thickness = (T)0);
 	virtual ~TPBDCollisionConstraint() {}
 
+	//
+	// Constraint Container API
+	//
+
 	/**
-	 * Generate all contact constraints.
+	 * Get the number of constraints.
 	 */
-	void UpdatePositionBasedState(/*const TPBDRigidParticles<T, d>& InParticles, const TArray<int32>& InIndices,*/ const T Dt);
+	int32 NumConstraints() const
+	{
+		return Constraints.Num();
+	}
+
+	const FConstraintHandle* GetConstraintHandle(int32 ConstraintIndex) const
+	{
+		return Handles[ConstraintIndex];
+	}
+
+	FConstraintHandle* GetConstraintHandle(int32 ConstraintIndex)
+	{
+		return Handles[ConstraintIndex];
+	}
+
+	// @todo(ccaulfield): rename/remove
+	void RemoveConstraints(const TSet<TGeometryParticleHandle<T, d>*>& RemovedParticles);
+
+	//
+	// Constraint API
+	//
+
+	/**
+	 * Get the particles that are affected by the specified constraint.
+	 */
+	TVector<TGeometryParticleHandle<T, d>*, 2> GetConstrainedParticles(int32 ConstraintIndex) const
+	{
+		return { Constraints[ConstraintIndex].Particle, Constraints[ConstraintIndex].Levelset };
+	}
+
+
+	//
+	// Island Rule API
+	//
 
 	/**
 	 * Apply corrections for the specified list of constraints.
 	 */
-	// @todo(ccaulfield): this runs wide. The serial/parallel decision should be in the ConstraintRule
-	void Apply(const T Dt, const TArray<int32>& InConstraintIndices);
+	 // @todo(ccaulfield): this runs wide. The serial/parallel decision should be in the ConstraintRule
+	void Apply(const T Dt, const TArray<FConstraintHandle*>& InConstraintHandles);
+
+	/**
+	 * Generate all contact constraints.
+	 */
+	void UpdatePositionBasedState(/*const TPBDRigidParticles<T, d>& InParticles, const TArray<int32>& InIndices,*/ const T Dt);
 
 	/** 
 	 * Apply push out for the specified list of constraints.
 	 * Return true if we need another iteration 
 	 */
 	 // @todo(ccaulfield): this runs wide. The serial/parallel decision should be in the ConstraintRule
-	bool ApplyPushOut(const T Dt, const TArray<int32>& InConstraintIndices, const TSet<TGeometryParticleHandle<T,d>*>& IsTemporarilyStatic, int32 Iteration, int32 NumIterations);
+	bool ApplyPushOut(const T Dt, const TArray<FConstraintHandle*>& InConstraintHandles, const TSet<TGeometryParticleHandle<T,d>*>& IsTemporarilyStatic, int32 Iteration, int32 NumIterations);
+
+
+
 
-	void RemoveConstraints(const TSet<TGeometryParticleHandle<T, d>*>& RemovedParticles);
 	void UpdateConstraints(/*const TPBDRigidParticles<T, d>& InParticles, const TArray<int32>& InIndices,*/ T Dt, const TSet<TGeometryParticleHandle<T, d>*>& AddedParticles);
 
 	static bool NearestPoint(TArray<Pair<TVector<T, d>, TVector<T, d>>>& Points, TVector<T, d>& Direction, TVector<T, d>& ClosestPoint);
 
 	const TArray<FRigidBodyContactConstraint>& GetAllConstraints() const { return Constraints; }
-	const ISpatialAcceleration<TGeometryParticleHandle<T, d>*, T, d>& GetSpatialAcceleration() const { return SpatialAccelerationResource.GetRead(); }
+	const ISpatialAcceleration<TAccelerationStructureHandle<T, d>, T, d>& GetSpatialAcceleration() const { return SpatialAccelerationResource.GetRead(); }
 	void ReleaseSpatialAcceleration() const;
 	void SwapSpatialAcceleration();
 
@@ -93,9 +212,6 @@ class CHAOS_API TPBDCollisionConstraint
 	template<ECollisionUpdateType>
 	static void UpdateLevelsetConstraintGJK(const T Thickness, FRigidBodyContactConstraint& Constraint);
 
-
-	int32 NumConstraints() const { return Constraints.Num(); }
-
 	/**
 	 * Get the particles that are affected by the specified constraint.
 	 */
@@ -103,6 +219,13 @@ class CHAOS_API TPBDCollisionConstraint
 
 	void SetPushOutPairIterations(int32 InPushOutPairIterations) { MPairIterations = InPushOutPairIterations; }
 
+	/** Make a copy of the acceleration structure to allow for external modification. This is needed for supporting sync operations on SQ structure from game thread */
+	TUniquePtr<ISpatialAcceleration<TAccelerationStructureHandle<T, d>, T, d>> CreateExternalAccelerationStructure();
+	
+protected:
+	using Base::GetConstraintIndex;
+	using Base::SetConstraintIndex;
+
 private:
 	void Reset(/*const TPBDRigidParticles<T, d>& InParticles, const TArray<int32>& InIndices*/);
 
@@ -130,7 +253,7 @@ private:
 #endif
 
 	const TPBDRigidsSOAs<T,d>& Particles;
-	TChaosReadWriteResource<TBoundingVolume<TGeometryParticles<T, d>, TGeometryParticleHandle<T,d>*, T, d>> SpatialAccelerationResource;
+	TChaosReadWriteResource<TBoundingVolume<TGeometryParticles<T, d>, TAccelerationStructureHandle<T,d>, T, d>> SpatialAccelerationResource;
 
 	TArray<FRigidBodyContactConstraint> Constraints;
 	TArrayCollectionArray<bool>& MCollided;
@@ -139,6 +262,9 @@ private:
 	T MThickness;
 	T MAngularFriction;
 	bool bUseCCD;
+
+	TArray<FConstraintHandle*> Handles;
+	FConstraintHandleAllocator HandleAllocator;
 };
 
 extern template void TPBDCollisionConstraint<float, 3>::UpdateConstraint<ECollisionUpdateType::Any>(const float Thickness, FRigidBodyContactConstraint& Constraint);
diff --git a/Engine/Source/Runtime/Experimental/Chaos/Public/Chaos/PBDConstraintColor.h b/Engine/Source/Runtime/Experimental/Chaos/Public/Chaos/PBDConstraintColor.h
--- a/Engine/Source/Runtime/Experimental/Chaos/Public/Chaos/PBDConstraintColor.h
+++ b/Engine/Source/Runtime/Experimental/Chaos/Public/Chaos/PBDConstraintColor.h
@@ -20,7 +20,7 @@ namespace Chaos
 	public:
 		typedef TPBDConstraintGraph<T, d> FConstraintGraph;
 		typedef TSet<int32> FColorSet;
-		typedef TArray<int32> FConstraintList;
+		typedef TArray<TConstraintHandle<T, d>*> FConstraintList;
 		typedef TMap<int32, FConstraintList> FColorToConstraintListMap;
 		typedef TArray<FColorToConstraintListMap> FLevelToColorToConstraintListMap;
 
diff --git a/Engine/Source/Runtime/Experimental/Chaos/Public/Chaos/PBDConstraintContainer.h b/Engine/Source/Runtime/Experimental/Chaos/Public/Chaos/PBDConstraintContainer.h
--- a/Engine/Source/Runtime/Experimental/Chaos/Public/Chaos/PBDConstraintContainer.h
+++ b/Engine/Source/Runtime/Experimental/Chaos/Public/Chaos/PBDConstraintContainer.h
@@ -2,6 +2,7 @@
 #pragma once
 
 #include "Chaos/Defines.h"
+#include "Chaos/ConstraintHandle.h"
 
 namespace Chaos
 {
@@ -14,5 +15,11 @@ namespace Chaos
 	class CHAOS_API TPBDConstraintContainer
 	{
 	public:
+		using FConstraintHandle = TConstraintHandle<T, d>;
+
+	protected:
+		// friend access to the Constraint Handle's container API
+		int32 GetConstraintIndex(const FConstraintHandle* ConstraintHandle) const { return ConstraintHandle->ConstraintIndex; }
+		void SetConstraintIndex(FConstraintHandle* ConstraintHandle, int32 ConstraintIndex) const { ConstraintHandle->ConstraintIndex = ConstraintIndex; }
 	};
 }
\ No newline at end of file
diff --git a/Engine/Source/Runtime/Experimental/Chaos/Public/Chaos/PBDConstraintGraph.h b/Engine/Source/Runtime/Experimental/Chaos/Public/Chaos/PBDConstraintGraph.h
--- a/Engine/Source/Runtime/Experimental/Chaos/Public/Chaos/PBDConstraintGraph.h
+++ b/Engine/Source/Runtime/Experimental/Chaos/Public/Chaos/PBDConstraintGraph.h
@@ -16,6 +16,9 @@ namespace Chaos
 	template<typename T, int d>
 	class TPBDConstraintColor;
 
+	template<typename T, int d>
+	class TConstraintHandle;
+
 	template<class T>
 	class TChaosPhysicsMaterial;
 
@@ -50,7 +53,7 @@ namespace Chaos
 		struct FConstraintData
 		{
 			uint32 ContainerId;
-			int32 ConstraintIndex;
+			TConstraintHandle<T, d>* ConstraintHandle;
 		};
 
 		TPBDConstraintGraph();
@@ -71,7 +74,7 @@ namespace Chaos
 		/**
 		 * Add a constraint to the graph for each constraint in the container.
 		 */
-		void AddConstraint(const uint32 InContainerId, const int32 InConstraintIndex, const TVector<TGeometryParticleHandle<T,d>*, 2>& InConstrainedParticles);
+		void AddConstraint(const uint32 InContainerId, TConstraintHandle<T, d>* InConstraintHandle, const TVector<TGeometryParticleHandle<T,d>*, 2>& InConstrainedParticles);
 
 		/**
 		 * Add particles/constraints to their particle's already-assigned islands (if applicable).
@@ -101,8 +104,8 @@ namespace Chaos
 		void ReconcileIslands();
 
 		/**
-		 * Get the list of ConstraintData indices associated with the specified island. NOTE: ConstraintDataIndex != ConstraintIndex.
-		 * Indices returned are into the ConstraintData array in the ConstraintGraph to get to the Constraint Index and Container Id
+		 * Get the list of ConstraintsData indices associated with the specified island. NOTE: ConstraintDataIndex is an internal index and not related to 
+		 * a constraint's index in its owning container. Indices are into the ConstraintData array in the ConstraintGraph to get to the Constraint Index and Container Id.
 		 * @see GetConstraintsData().
 		 */
 		const TArray<int32>& GetIslandConstraintData(int32 Island) const
