Copying //UE4/Dev-Physics to Dev-Main (//UE4/Dev-Main) [at] 9098278

diff --git a/Engine/Source/Runtime/Experimental/Chaos/Public/Chaos/PBDConstraintRule.h b/Engine/Source/Runtime/Experimental/Chaos/Public/Chaos/PBDConstraintRule.h
--- a/Engine/Source/Runtime/Experimental/Chaos/Public/Chaos/PBDConstraintRule.h
+++ b/Engine/Source/Runtime/Experimental/Chaos/Public/Chaos/PBDConstraintRule.h
@@ -1,6 +1,7 @@
 // Copyright 1998-2019 Epic Games, Inc. All Rights Reserved.
 #pragma once
 
+#include "Chaos/Declares.h"
 #include "Chaos/PBDCollisionTypes.h"
 #include "Chaos/PBDConstraintGraph.h"
 #include "Chaos/PBDConstraintColor.h"
@@ -9,6 +10,9 @@
 
 namespace Chaos
 {
+	template<typename T, int d>
+	class TConstraintHandle;
+
 	/**
 	 * Constraint Rules bind constraint collections to the evolution and provide their update algorithm.
 	 */
@@ -98,7 +102,7 @@ namespace Chaos
 			ConstraintGraph->ReserveConstraints(Constraints.NumConstraints());
 			for (int32 ConstraintIndex = 0; ConstraintIndex < Constraints.NumConstraints(); ++ConstraintIndex)
 			{
-				ConstraintGraph->AddConstraint(ContainerId, ConstraintIndex, Constraints.ConstraintParticles(ConstraintIndex));
+				ConstraintGraph->AddConstraint(ContainerId, Constraints.GetConstraintHandle(ConstraintIndex), Constraints.GetConstrainedParticles(ConstraintIndex));
 			}
 		}
 
@@ -126,8 +130,10 @@ namespace Chaos
 		typedef TPBDConstraintGraphRuleImpl<T_CONSTRAINTS, T, d> Base;
 
 	public:
-		typedef typename Base::FConstraints FConstraints;
-		typedef typename Base::FConstraintGraph FConstraintGraph;
+		using FConstraints = typename Base::FConstraints;
+		using FConstraintHandle = typename FConstraints::FConstraintHandle;
+		using FConstraintList = TArray<TConstraintHandle<T, d>*>;
+		using FConstraintGraph = typename Base::FConstraintGraph;
 
 		TPBDConstraintIslandRule(FConstraints& InConstraints)
 			: TPBDConstraintGraphRuleImpl<T_CONSTRAINTS, T, d>(InConstraints)
@@ -138,7 +144,7 @@ namespace Chaos
 		{
 			if (IslandConstraintLists[Island].Num())
 			{
-				Constraints.Apply(Dt, IslandConstraintLists[Island]);
+				Constraints.Apply(Dt, GetIslandConstraints(Island));
 			}
 		}
 
@@ -146,14 +152,14 @@ namespace Chaos
 		{
 			if (IslandConstraintLists[Island].Num())
 			{
-				Constraints.ApplyPushOut(Dt, IslandConstraintLists[Island]);
+				Constraints.ApplyPushOut(Dt, GetIslandConstraints(Island));
 			}
 		}
 
 		virtual void InitializeAccelerationStructures() override
 		{
 			IslandConstraintLists.SetNum(ConstraintGraph->NumIslands());
-			for (TArray<int32>& IslandConstraintList : IslandConstraintLists)
+			for (FConstraintList& IslandConstraintList : IslandConstraintLists)
 			{
 				IslandConstraintList.Reset();
 			}
@@ -162,7 +168,7 @@ namespace Chaos
 		virtual void UpdateAccelerationStructures(const int32 Island) override
 		{
 			const TArray<int32>& ConstraintDataIndices = ConstraintGraph->GetIslandConstraintData(Island);
-			TArray<int32>& IslandConstraintList = IslandConstraintLists[Island];
+			FConstraintList& IslandConstraintList = IslandConstraintLists[Island];
 			IslandConstraintList.Reset();
 			IslandConstraintList.Reserve(ConstraintDataIndices.Num());
 			for (int32 ConstraintDataIndex : ConstraintDataIndices)
@@ -170,7 +176,7 @@ namespace Chaos
 				const typename FConstraintGraph::FConstraintData& ConstraintData = ConstraintGraph->GetConstraintData(ConstraintDataIndex);
 				if (ConstraintData.ContainerId == ContainerId)
 				{
-					IslandConstraintList.Add(ConstraintData.ConstraintIndex);
+					IslandConstraintList.Add(ConstraintData.ConstraintHandle);
 				}
 			}
 		}
@@ -180,9 +186,16 @@ namespace Chaos
 		using Base::ConstraintGraph;
 		using Base::ContainerId;
 
+		const TArray<FConstraintHandle*>& GetIslandConstraints(int32 Island)
+		{
+			// Constraint rules are bound to a single type, but the FConstraintGraph works with many types. We have
+			// already pre-filtered the constraint lists based on type, so this case is safe.
+			return reinterpret_cast<const TArray<FConstraintHandle*>&>(IslandConstraintLists[Island]);
+		}
+
 		// @todo(ccaulfield): optimize: eliminate the need for this index list - it is a subset of EdgeData. 
 		// If EdgeData were sorted by ContainerId we could use TArray<TArrayView<int32>>...
-		TArray<TArray<int32>> IslandConstraintLists;
+		TArray<FConstraintList> IslandConstraintLists;
 	};
 
 	/**
@@ -224,7 +237,8 @@ namespace Chaos
 				{
 					if (LevelToColorToConstraintListMap[Level].Contains(Color) && LevelToColorToConstraintListMap[Level][Color].Num())
 					{
-						Constraints.Apply(Dt, LevelToColorToConstraintListMap[Level][Color]);
+						const TArray<typename FConstraints::FConstraintHandle*>& ConstraintHandles = GetLevelColorConstraints(LevelToColorToConstraintListMap, Level, Color);
+						Constraints.Apply(Dt, ConstraintHandles);
 					}
 				}
 			}
@@ -247,7 +261,8 @@ namespace Chaos
 					{
 						if (LevelToColorToConstraintListMap[Level].Contains(Color) && LevelToColorToConstraintListMap[Level][Color].Num())
 						{
-							NeedsAnotherIteration = Constraints.ApplyPushOut(Dt, LevelToColorToConstraintListMap[Level][Color], IsTemporarilyStatic, Iteration, PushOutIterations);
+							const TArray<typename FConstraints::FConstraintHandle*>& ConstraintHandles = GetLevelColorConstraints(LevelToColorToConstraintListMap, Level, Color);
+							NeedsAnotherIteration = Constraints.ApplyPushOut(Dt, ConstraintHandles, IsTemporarilyStatic, Iteration, PushOutIterations);
 						}
 					}
 #if USE_SHOCK_PROPOGATION
@@ -257,7 +272,7 @@ namespace Chaos
 						{
 							for (int32 Edge = 0; Edge < LevelToColorToConstraintListMap[Level][Color].Num(); ++Edge)
 							{
-								const int32 ConstraintIndex = LevelToColorToConstraintListMap[Level][Color][Edge];
+								const int32 ConstraintIndex = LevelToColorToConstraintListMap[Level][Color][Edge]->GetConstraintIndex();
 								const TVector<TGeometryParticleHandle<T,d>*, 2> Particles = Constraints.ConstraintParticles(ConstraintIndex);
 								if (Iteration == PushOutIterations - 1)
 								{
@@ -298,6 +313,12 @@ namespace Chaos
 		using Base::ConstraintGraph;
 		using Base::ContainerId;
 
+		const TArray<typename FConstraints::FConstraintHandle*>& GetLevelColorConstraints(const typename FConstraintColor::FLevelToColorToConstraintListMap& LevelToColorToConstraintListMap, int32 Level, int32 Color) const
+		{
+			// FConstraintColor works with any constraint type (in principle - currently only used with Collisions), but the rule is bound to a single type and so this cast is ok
+			return reinterpret_cast<const TArray<typename FConstraints::FConstraintHandle*>&>(LevelToColorToConstraintListMap[Level][Color]);
+		}
+
 		FConstraintColor GraphColor;
 		int32 PushOutIterations;
 	};
diff --git a/Engine/Source/Runtime/Experimental/Chaos/Public/Chaos/PBDJointConstraints.h b/Engine/Source/Runtime/Experimental/Chaos/Public/Chaos/PBDJointConstraints.h
--- a/Engine/Source/Runtime/Experimental/Chaos/Public/Chaos/PBDJointConstraints.h
+++ b/Engine/Source/Runtime/Experimental/Chaos/Public/Chaos/PBDJointConstraints.h
@@ -3,49 +3,223 @@
 
 #include "Chaos/Array.h"
 #include "Chaos/Map.h"
-#include "Chaos/PBDJointConstraintsBase.h"
+#include "Chaos/Matrix.h"
+#include "Chaos/Vector.h"
+
+#include "Chaos/ConstraintHandle.h"
+#include "Chaos/ParticleHandleFwd.h"
 #include "Chaos/PBDConstraintContainer.h"
 
 namespace Chaos
 {
-template<class T, int d>
-class CHAOS_API TPBDJointConstraints : public TPBDJointConstraintsBase<T, d>, public TPBDConstraintContainer<T, d>
-{
-	typedef TPBDJointConstraintsBase<T, d> Base;
+	template<class T, int d>
+	class TPBDJointConstraints;
 
-  public:
-	using Base::AddConstraint;
-	using Base::ApplySingle;
-	using Base::ConstraintParticles;
-	using Base::NumConstraints;
-	using Base::RemoveConstraints;
+	template<class T, int d>
+	class CHAOS_API TPBDJointConstraintHandle : public TContainerConstraintHandle<TPBDJointConstraints<T, d>>
+	{
+	public:
+		using Base = TContainerConstraintHandle<TPBDJointConstraints<T, d>>;
+		using FConstraintContainer = TPBDJointConstraints<T, d>;
 
-	TPBDJointConstraints(const T InStiffness = (T)1)
-	    : TPBDJointConstraintsBase<T, d>(InStiffness) 
-	{}
+		TPBDJointConstraintHandle() {}
+		TPBDJointConstraintHandle(FConstraintContainer* InConstraintContainer, int32 InConstraintIndex) : TContainerConstraintHandle<TPBDJointConstraints<T, d>>(InConstraintContainer, InConstraintIndex) {}
 
-	TPBDJointConstraints(const TArray<TVector<T, 3>>& Locations, TArray<TVector<TGeometryParticleHandle<T,d>*, 2>>&& InConstraints, const T InStiffness = (T)1)
-	    : TPBDJointConstraintsBase<T, d>(Locations, MoveTemp(InConstraints), InStiffness) 
-	{}
+		const TVector<TVector<T, 3>, 2>& GetConstraintPositions() const;
+		void SetConstraintPositions(const TVector<TVector<T, 3>, 2>& ConstraintPositions);
 
-	virtual ~TPBDJointConstraints()
-	{}
+	protected:
+		using Base::ConstraintIndex;
+		using Base::ConstraintContainer;
+	};
 
 
-	void UpdatePositionBasedState(const T Dt)
+	template<class T, int d>
+	class CHAOS_API TPBDJointConstraints : public TPBDConstraintContainer<T, d>
 	{
-	}
+	public:
+		using Base = TPBDConstraintContainer<T, d>;
+		using FReal = T;
+		static const int Dimensions = d;
+		using FConstraintHandle = TPBDJointConstraintHandle<FReal, Dimensions>;
+		using FConstraintHandleAllocator = TConstraintHandleAllocator<TPBDJointConstraints<FReal, Dimensions>>;
+		using FConstrainedParticlePair = TVector<TGeometryParticleHandle<T, d>*, 2>;
 
-	void Apply(const T Dt, const TArray<int32>& InConstraintIndices) const
-	{
-		for (int32 ConstraintIndex : InConstraintIndices)
+
+		TPBDJointConstraints(const T InStiffness = (T)1)
+			: Stiffness(InStiffness)
+		{
+		}
+
+		TPBDJointConstraints(const TArray<TVector<T, d>>& Locations, TArray<FConstrainedParticlePair>&& InConstraints, const T InStiffness = (T)1)
+			: Constraints(MoveTemp(InConstraints)), Stiffness(InStiffness)
 		{
-			ApplySingle(Dt, ConstraintIndex);
+			if (Constraints.Num() > 0)
+			{
+				Handles.Reserve(Constraints.Num());
+				for (int32 ConstraintIndex = 0; ConstraintIndex < Constraints.Num(); ++ConstraintIndex)
+				{
+					Handles.Add(HandleAllocator.AllocHandle(this, ConstraintIndex));
+				}
 			}
-	}
+			UpdateDistances(Locations);
+		}
 
-	void ApplyPushOut(const T Dt, const TArray<int32>& InConstraintIndices) const
-	{
-	}
-};
+		virtual ~TPBDJointConstraints() {}
+
+		//
+		// Constraint Container API
+		//
+
+		/**
+		 * Get the number of constraints.
+		 */
+		int32 NumConstraints() const
+		{
+			return Constraints.Num();
+		}
+
+		/** 
+		 * Add a constraint initialized from current world-space particle positions. 
+		 * You would use this method when your objects are already positioned in the world. 
+		 */
+		FConstraintHandle* AddConstraint(const FConstrainedParticlePair& InConstrainedParticles, const TVector<T, d>& InLocation)
+		{
+			Handles.Add(HandleAllocator.AllocHandle(this, Handles.Num()));
+			int32 ConstraintIndex = Constraints.Add(InConstrainedParticles);
+			UpdateDistance(InLocation, ConstraintIndex);
+			return Handles.Last();
+		}
+
+		/** 
+		 * Add a constraint and explicitly provide constrained particle offsets. 
+		 * You would use this method to initialize a constraint when the particles are not already in the correct locations. This might be because the
+		 * particles have not been positioned yet, or if the particles are not positioned correctly (it will be the constraint's job to correct the positions).
+		 */
+		FConstraintHandle* AddConstraintLocal(const FConstrainedParticlePair& InConstrainedParticles, const TVector<TVector<T, 3>, 2>& InLocations)
+		{
+			Handles.Add(HandleAllocator.AllocHandle(this, Handles.Num()));
+			Constraints.Add(InConstrainedParticles);
+			Distances.Add(InLocations);
+			return Handles.Last();
+		}
+
+		/**
+		 * Remove the specified constraint.
+		 */
+		void RemoveConstraint(int ConstraintIndex)
+		{
+			FConstraintHandle* ConstraintHandle = Handles[ConstraintIndex];
+			if (ConstraintHandle != nullptr)
+			{
+				// Release the handle for the freed constraint
+				HandleAllocator.FreeHandle(ConstraintHandle);
+				Handles[ConstraintIndex] = nullptr;
+			}
+
+			// Swap the last constraint into the gap to keep the array packed
+			Constraints.RemoveAtSwap(ConstraintIndex);
+			Distances.RemoveAtSwap(ConstraintIndex);
+			Handles.RemoveAtSwap(ConstraintIndex);
+
+			// Update the handle for the constraint that was moved
+			if (ConstraintIndex < Handles.Num())
+			{
+				SetConstraintIndex(Handles[ConstraintIndex], ConstraintIndex);
+			}
+		}
+
+		// @todo(ccaulfield): rename/remove  this
+		void RemoveConstraints(const TSet<TGeometryParticleHandle<T, d>*>& RemovedParticles)
+		{
+		}
+
+		//
+		// Constraint API
+		//
+
+		const FConstraintHandle* GetConstraintHandle(int32 ConstraintIndex) const
+		{
+			return Handles[ConstraintIndex];
+		}
+
+		FConstraintHandle* GetConstraintHandle(int32 ConstraintIndex)
+		{
+			return Handles[ConstraintIndex];
+		}
+
+		/**
+		 * Get the particles that are affected by the specified constraint.
+		 */
+		const FConstrainedParticlePair& GetConstrainedParticles(int32 ConstraintIndex) const
+		{
+			return Constraints[ConstraintIndex];
+		}
+
+		/**
+		 * Get the local-space constraint positions for each body.
+		 */
+		const TVector<TVector<T, 3>, 2>& GetConstraintPositions(int ConstraintIndex) const
+		{
+			return Distances[ConstraintIndex];
+		}
+
+		/**
+		 * Set the local-space constraint positions for each body.
+		 */
+		void SetConstraintPositions(int ConstraintIndex, const TVector<TVector<T, 3>, 2>& ConstraintPositions)
+		{
+			Distances[ConstraintIndex] = ConstraintPositions;
+		}
+
+
+		//
+		// Island Rule API
+		//
+
+		void UpdatePositionBasedState(const T Dt)
+		{
+		}
+
+		void Apply(const T Dt, const TArray<FConstraintHandle*>& InConstraintHandles)
+		{
+			for (FConstraintHandle* ConstraintHandle : InConstraintHandles)
+			{
+				ApplySingle(Dt, ConstraintHandle->GetConstraintIndex());
+			}
+		}
+
+		// @todo(ccaulfield): remove  this
+		void ApplyPushOut(const T Dt, const TArray<FConstraintHandle*>& InConstraintHandles)
+		{
+		}
+
+
+	protected:
+		using Base::GetConstraintIndex;
+		using Base::SetConstraintIndex;
+
+		void ApplySingle(const T Dt, const int32 ConstraintIndex);
+
+	private:
+		void UpdateDistanceInternal(const TVector<T, d>& InLocation, const int32 InConstraintIndex);
+		void UpdateDistance(const TVector<T, d>& InLocation, const int32 InConstraintIndex);
+		void UpdateDistances(const TArray<TVector<T, d>>& InLocations);
+
+		// Double dynamic body solve
+		void ApplyDynamicDynamic(const T Dt, const int32 ConstraintIndex, const int32 PBDRigid0Index, const int32 PBDRigid1Index, const bool bApplyProjection);
+		TVector<T, d> GetDeltaDynamicDynamic(const TVector<T, d>& P0, const TVector<T, d>& P1, const TVector<T, d>& WorldSpaceX0, const TVector<T, d>& WorldSpaceX1, const PMatrix<T, d, d>& WorldSpaceInvI0, const PMatrix<T, d, d>& WorldSpaceInvI1, const T InvM0, const T InvM1);
+
+		// Single dynamic body solve
+		void ApplyDynamicStatic(const T Dt, const int32 ConstraintIndex, const int32 PBDRigid0Index, const int32 Static1Index, const bool bApplyProjection);
+		TVector<T, d> GetDeltaDynamicKinematic(const TVector<T, d>& P0, const TVector<T, d>& WorldSpaceX0, const TVector<T, d>& WorldSpaceX1, const PMatrix<T, d, d>& WorldSpaceInvI0, const T InvM0);
+
+		// @todo(ccaulfield): we never iterate over these separately. Do we still want SoA?
+		TArray<FConstrainedParticlePair> Constraints;
+		TArray<TVector<TVector<T, 3>, 2>> Distances;
+		T Stiffness;
+
+		TArray<FConstraintHandle*> Handles;
+		FConstraintHandleAllocator HandleAllocator;
+	};
 }
diff --git a/Engine/Source/Runtime/Experimental/Chaos/Public/Chaos/PBDPositionConstraints.h b/Engine/Source/Runtime/Experimental/Chaos/Public/Chaos/PBDPositionConstraints.h
--- a/Engine/Source/Runtime/Experimental/Chaos/Public/Chaos/PBDPositionConstraints.h
+++ b/Engine/Source/Runtime/Experimental/Chaos/Public/Chaos/PBDPositionConstraints.h
@@ -3,42 +3,178 @@
 
 #include "Chaos/Array.h"
 #include "Chaos/Map.h"
+#include "Chaos/ConstraintHandle.h"
+#include "Chaos/ParticleHandle.h"
 #include "Chaos/PBDConstraintContainer.h"
-#include "ParticleHandle.h"
 
 namespace Chaos
 {
-template<class T, int d>
-class TPBDPositionConstraints : public TPBDConstraintContainer<T, d>
-{
-  public:
-	TPBDPositionConstraints(const T InStiffness = (T)1)
-		: Stiffness(InStiffness)
-	{}
+	template<class T, int d>
+	class TPBDPositionConstraints;
+
+	template<class T, int d>
+	class TPBDPositionConstraintHandle : public TContainerConstraintHandle<TPBDPositionConstraints<T, d>>
+	{
+	public:
+		using Base = TContainerConstraintHandle<TPBDPositionConstraints<T, d>>;
+		using FConstraintContainer = TPBDPositionConstraints<T, d>;
 
-	TPBDPositionConstraints(TArray<TVector<T, d>>&& Locations, TArray<TPBDRigidParticleHandle<T,d>*>&& InConstrainedParticles, const T InStiffness = (T)1)
-		: Targets(MoveTemp(Locations)), ConstrainedParticles(MoveTemp(InConstrainedParticles)), Stiffness(InStiffness)
-	{}
+		TPBDPositionConstraintHandle() {}
+		TPBDPositionConstraintHandle(FConstraintContainer* InConstraintContainer, int32 InConstraintIndex) : TContainerConstraintHandle<TPBDPositionConstraints<T, d>>(InConstraintContainer, InConstraintIndex) {}
 
-	virtual ~TPBDPositionConstraints() {}
+	protected:
+		using Base::ConstraintIndex;
+		using Base::ConstraintContainer;
+	};
 
-	int32 NumConstraints() const
+	template<class T, int d>
+	class TPBDPositionConstraints : public TPBDConstraintContainer<T, d>
 	{
-		return ConstrainedParticles.Num();
-	}
+	public:
+		using Base = TPBDConstraintContainer<T, d>;
+		using FReal = T;
+		static const int Dimensions = d;
+		using FConstraintHandle = TPBDPositionConstraintHandle<FReal, Dimensions>;
+		using FConstraintHandleAllocator = TConstraintHandleAllocator<TPBDPositionConstraints<FReal, Dimensions>>;
 
-	TVector<TGeometryParticleHandle<T,d>*, 2> ConstraintParticles(int32 ConstraintIndex) const
-	{ 
-		return { ConstrainedParticles[ConstraintIndex], nullptr };
-	}
+		TPBDPositionConstraints(const T InStiffness = (T)1)
+			: Stiffness(InStiffness)
+		{}
 
-	void UpdatePositionBasedState(const T Dt)
-	{
-	}
+		TPBDPositionConstraints(TArray<TVector<T, d>>&& Locations, TArray<TPBDRigidParticleHandle<T,d>*>&& InConstrainedParticles, const T InStiffness = (T)1)
+			: Targets(MoveTemp(Locations)), ConstrainedParticles(MoveTemp(InConstrainedParticles)), Stiffness(InStiffness)
+		{
+			if (ConstrainedParticles.Num() > 0)
+			{
+				Handles.Reserve(ConstrainedParticles.Num());
+				for (int32 ConstraintIndex = 0; ConstraintIndex < ConstrainedParticles.Num(); ++ConstraintIndex)
+				{
+					Handles.Add(HandleAllocator.AllocHandle(this, ConstraintIndex));
+				}
+			}
+		}
 
-	void ApplyHelper(const T Dt, const TArray<int32>& InConstraintIndices) const
-	{
-		for (int32 ConstraintIndex : InConstraintIndices)
+		virtual ~TPBDPositionConstraints() {}
+
+
+		//
+		// Constraint Container API
+		//
+
+		/**
+		 * Get the number of constraints.
+		 */
+		int32 NumConstraints() const
+		{
+			return ConstrainedParticles.Num();
+		}
+
+		/**
+		 * Add a constraint.
+		 */
+		FConstraintHandle* AddConstraint(TPBDRigidParticleHandle<T, d>* Particle, const TVector<T, d>& Position)
+		{
+			int32 NewIndex = Targets.Num();
+			Targets.Add(Position);
+			ConstrainedParticles.Add(Particle);
+			Handles.Add(HandleAllocator.AllocHandle(this, NewIndex));
+			return Handles[NewIndex];
+		}
+
+		/**
+		 * Remove a constraint.
+		 */
+		void RemoveConstraint(int ConstraintIndex)
+		{
+			FConstraintHandle* ConstraintHandle = Handles[ConstraintIndex];
+			if (ConstraintHandle != nullptr)
+			{
+				// Release the handle for the freed constraint
+				HandleAllocator.FreeHandle(ConstraintHandle);
+				Handles[ConstraintIndex] = nullptr;
+			}
+
+			// Swap the last constraint into the gap to keep the array packed
+			Targets.RemoveAtSwap(ConstraintIndex);
+			ConstrainedParticles.RemoveAtSwap(ConstraintIndex);
+			Handles.RemoveAtSwap(ConstraintIndex);
+
+			// Update the handle for the constraint that was moved
+			if (ConstraintIndex < Handles.Num())
+			{
+				SetConstraintIndex(Handles[ConstraintIndex], ConstraintIndex);
+			}
+		}
+
+		// @todo(ccaulfield): remove/rename/implement
+		void RemoveConstraints(const TSet<TGeometryParticleHandle<T, d>*>& RemovedParticles)
+		{
+		}
+
+
+		//
+		// Constraint API
+		//
+
+		const FConstraintHandle* GetConstraintHandle(int32 ConstraintIndex) const
+		{
+			return Handles[ConstraintIndex];
+		}
+
+		FConstraintHandle* GetConstraintHandle(int32 ConstraintIndex)
+		{
+			return Handles[ConstraintIndex];
+		}
+
+		/**
+		 * Get the particles that are affected by the specified constraint.
+		 */
+		TVector<TGeometryParticleHandle<T, d>*, 2> GetConstrainedParticles(int32 ConstraintIndex) const
+		{
+			return { ConstrainedParticles[ConstraintIndex], nullptr };
+		}
+
+		/**
+		 * Get the world-space constraint positions for each body.
+		 */
+		const TVector<T, d>& GetConstraintPosition(int ConstraintIndex) const
+		{
+			return Targets[ConstraintIndex];
+		}
+
+		// @todo(ccaulfield): remove/rename
+		void Replace(const int32 ConstraintIndex, const TVector<T, d>& Position)
+		{
+			Targets[ConstraintIndex] = Position;
+		}
+
+
+		//
+		// Island Rule API
+		//
+
+		void UpdatePositionBasedState(const T Dt)
+		{
+		}
+
+		void Apply(const T Dt, const TArray<FConstraintHandle*>& ConstraintHandles) const
+		{
+			for (FConstraintHandle* ConstraintHandle : ConstraintHandles)
+			{
+				ApplySingle(Dt, ConstraintHandle->GetConstraintIndex());
+			}
+		}
+
+		void ApplyPushOut(const T Dt, const TArray<FConstraintHandle*>& InConstraintIndices) const
+		{
+		}
+
+	protected:
+		using Base::GetConstraintIndex;
+		using Base::SetConstraintIndex;
+
+	private:
+		void ApplySingle(const T Dt, int32 ConstraintIndex) const
 		{
 			if (TPBDRigidParticleHandle<T, d>* PBDRigid = ConstrainedParticles[ConstraintIndex])
 			{
@@ -48,38 +184,12 @@ class TPBDPositionConstraints : public TPBDConstraintContainer<T, d>
 				PBDRigid->P() -= Stiffness * Difference;
 			}
 		}
-	}
 
-	void Apply(const T Dt, const TArray<int32>& InConstraintIndices) const
-	{
-		ApplyHelper(Dt, InConstraintIndices);
-	}
-
-	void ApplyPushOut(const T Dt, const TArray<int32>& InConstraintIndices) const
-	{
-	}
-
-	int32 Add(TGeometryParticleHandle<T,d>* Particle, const TVector<T, d>& Position)
-	{
-		int32 NewIndex = Targets.Num();
-		Targets.Add(Position);
-		ConstrainedParticles.Add(Particle);
-		return NewIndex;
-	}
-
-	void Replace(const int32 Index, const TVector<T, d>& Position)
-	{
-		Targets[Index] = Position;
-	}
+		TArray<TVector<T, d>> Targets;
+		TArray<TPBDRigidParticleHandle<T,d>*> ConstrainedParticles;
+		T Stiffness;
 
-	void RemoveConstraints(const TSet<TGeometryParticleHandle<T,d>*>& RemovedParticles)
-	{
-		// @todo(ccaulfield): constraint management
-	}
-
-private:
-	TArray<TVector<T, d>> Targets;
-	TArray<TPBDRigidParticleHandle<T,d>*> ConstrainedParticles;
-	const T Stiffness;
-};
+		TArray<FConstraintHandle*> Handles;
+		FConstraintHandleAllocator HandleAllocator;
+	};
 }
diff --git a/Engine/Source/Runtime/Experimental/Chaos/Public/Chaos/PBDRigidDynamicSpringConstraints.h b/Engine/Source/Runtime/Experimental/Chaos/Public/Chaos/PBDRigidDynamicSpringConstraints.h
--- a/Engine/Source/Runtime/Experimental/Chaos/Public/Chaos/PBDRigidDynamicSpringConstraints.h
+++ b/Engine/Source/Runtime/Experimental/Chaos/Public/Chaos/PBDRigidDynamicSpringConstraints.h
@@ -3,49 +3,196 @@
 
 #include "Chaos/Array.h"
 #include "Chaos/Map.h"
-#include "Chaos/PBDRigidDynamicSpringConstraintsBase.h"
+#include "Chaos/ConstraintHandle.h"
+#include "Chaos/ParticleHandle.h"
 #include "Chaos/PBDConstraintContainer.h"
 
 namespace Chaos
 {
-template<class T, int d>
-class TPBDRigidDynamicSpringConstraints : public TPBDRigidDynamicSpringConstraintsBase2<T, d>, public TPBDConstraintContainer<T, d>
-{
-	typedef TPBDRigidDynamicSpringConstraintsBase2<T, d> Base;
+	template<class T, int d>
+	class TPBDRigidDynamicSpringConstraints;
 
-	using Base::Constraints;
-	using Base::Distances;
-	using Base::SpringDistances;
+	template<class T, int d>
+	class CHAOS_API TPBDRigidDynamicSpringConstraintHandle : public TContainerConstraintHandle<TPBDRigidDynamicSpringConstraints<T, d>>
+	{
+	public:
+		using Base = TContainerConstraintHandle<TPBDRigidDynamicSpringConstraints<T, d>>;
+		using FConstraintContainer = TPBDRigidDynamicSpringConstraints<T, d>;
 
-public:
-	using Base::UpdatePositionBasedState;
+		TPBDRigidDynamicSpringConstraintHandle() {}
+		TPBDRigidDynamicSpringConstraintHandle(FConstraintContainer* InConstraintContainer, int32 InConstraintIndex) : TContainerConstraintHandle<TPBDRigidDynamicSpringConstraints<T, d>>(InConstraintContainer, InConstraintIndex) {}
 
-	TPBDRigidDynamicSpringConstraints(const T InStiffness = (T)1)
-	    : TPBDRigidDynamicSpringConstraintsBase2<T, d>(InStiffness) 
-	{}
+	protected:
+		using Base::ConstraintIndex;
+		using Base::ConstraintContainer;
+	};
 
-	TPBDRigidDynamicSpringConstraints(TArray<TVector<TGeometryParticleHandle<T, d>*, 2>>&& InConstraints, const T InCreationThreshold = (T)1, const int32 InMaxSprings = 1, const T InStiffness = (T)1)
-	    : TPBDRigidDynamicSpringConstraintsBase2<T, d>(MoveTemp(InConstraints), InCreationThreshold, InMaxSprings, InStiffness)
-	{}
+	template<class T, int d>
+	class CHAOS_API TPBDRigidDynamicSpringConstraints : public TPBDConstraintContainer<T, d>
+	{
+	public:
+		using Base = TPBDConstraintContainer<T, d>;
+		using FConstrainedParticlePair = TVector<TGeometryParticleHandle<T, d>*, 2>;
+		using FReal = T;
+		static const int Dimensions = d;
+		using FConstraintHandle = TPBDRigidDynamicSpringConstraintHandle<FReal, Dimensions>;
+		using FConstraintHandleAllocator = TConstraintHandleAllocator<TPBDRigidDynamicSpringConstraints<FReal, Dimensions>>;
 
-	virtual ~TPBDRigidDynamicSpringConstraints() {}
+		TPBDRigidDynamicSpringConstraints(const T InStiffness = (T)1)
+			: CreationThreshold(1), MaxSprings(1), Stiffness(InStiffness) 
+		{}
 
-	CHAOS_API void UpdatePositionBasedState(const T Dt)
-	{
-		Base::UpdatePositionBasedState();
-	}
+		TPBDRigidDynamicSpringConstraints(TArray<FConstrainedParticlePair>&& InConstraints, const T InCreationThreshold = (T)1, const int32 InMaxSprings = 1, const T InStiffness = (T)1)
+			: Constraints(MoveTemp(InConstraints)), CreationThreshold(InCreationThreshold), MaxSprings(InMaxSprings), Stiffness(InStiffness)
+		{
+			if (Constraints.Num() > 0)
+			{
+				Handles.Reserve(Constraints.Num());
+				Distances.Reserve(Constraints.Num());
+				SpringDistances.Reserve(Constraints.Num());
+				for (int32 ConstraintIndex = 0; ConstraintIndex < Constraints.Num(); ++ConstraintIndex)
+				{
+					Handles.Add(HandleAllocator.AllocHandle(this, ConstraintIndex));
+					Distances.Add({});
+					SpringDistances.Add({});
+				}
+			}
+		}
 
-	// @todo(mlentine): Why is this needed?
-	CHAOS_API void ApplyHelper(const T Dt, const TArray<int32>& InConstraintIndices) const;
+		virtual ~TPBDRigidDynamicSpringConstraints() {}
 
-	CHAOS_API void Apply(const T Dt, const TArray<int32>& InConstraintIndices) const
-	{
-		ApplyHelper(Dt, InConstraintIndices);
-	}
+		//
+		// Constraint Container API
+		//
 
-	CHAOS_API void ApplyPushOut(const T Dt, const TArray<int32>& InConstraintIndices)
-	{
+		/**
+		 * Get the number of constraints.
+		 */
+		int32 NumConstraints() const
+		{
+			return Constraints.Num();
+		}
+
+		/**
+		 * Add a constraint initialized from current world-space particle positions.
+		 * You would use this method when your objects are already positioned in the world.
+		 */
+		FConstraintHandle* AddConstraint(const FConstrainedParticlePair& InConstrainedParticles)
+		{
+			Handles.Add(HandleAllocator.AllocHandle(this, Handles.Num()));
+			Constraints.Add(InConstrainedParticles);
+			Distances.Add({});
+			SpringDistances.Add({});
+			return Handles.Last();
+		}
+
+		/**
+		 * Remove the specified constraint.
+		 */
+		void RemoveConstraint(int ConstraintIndex)
+		{
+			FConstraintHandle* ConstraintHandle = Handles[ConstraintIndex];
+			if (ConstraintHandle != nullptr)
+			{
+				// Release the handle for the freed constraint
+				HandleAllocator.FreeHandle(ConstraintHandle);
+				Handles[ConstraintIndex] = nullptr;
+			}
+
+			// Swap the last constraint into the gap to keep the array packed
+			Constraints.RemoveAtSwap(ConstraintIndex);
+			Distances.RemoveAtSwap(ConstraintIndex);
+			SpringDistances.RemoveAtSwap(ConstraintIndex);
+			Handles.RemoveAtSwap(ConstraintIndex);
+
+			// Update the handle for the constraint that was moved
+			if (ConstraintIndex < Handles.Num())
+			{
+				SetConstraintIndex(Handles[ConstraintIndex], ConstraintIndex);
+			}
+		}
+
+		// @todo(ccaulfield): rename/remove  this
+		void RemoveConstraints(const TSet<TGeometryParticleHandle<T, d>*>& RemovedParticles)
+		{
+		}
+
+		/**
+		 * Set the distance threshold below which springs get created between particles.
+		 */
+		void SetCreationThreshold(const T InCreationThreshold)
+		{
+			CreationThreshold = InCreationThreshold;
+		}
+
+		/**
+		 * Set the maximum number of springs
+		 */
+		void SetMaxSprings(const T InMaxSprings)
+		{
+			MaxSprings = InMaxSprings;
+		}
+
+
+		//
+		// Constraint API
+		//
+
+		const FConstraintHandle* GetConstraintHandle(int32 ConstraintIndex) const
+		{
+			return Handles[ConstraintIndex];
+		}
+
+		FConstraintHandle* GetConstraintHandle(int32 ConstraintIndex)
+		{
+			return Handles[ConstraintIndex];
+		}
+
+		/**
+		 * Get the particles that are affected by the specified constraint.
+		 */
+		const FConstrainedParticlePair& GetConstrainedParticles(int32 ConstraintIndex) const
+		{
+			return Constraints[ConstraintIndex];
+		}
+
+
+		//
+		// Island Rule API
+		//
+
+		void UpdatePositionBasedState(const T Dt);
+
+		void Apply(const T Dt, const TArray<FConstraintHandle*>& InConstraintHandles)
+		{
+			for (FConstraintHandle* ConstraintHandle : InConstraintHandles)
+			{
+				ApplySingle(Dt, ConstraintHandle->GetConstraintIndex());
+			}
+		}
+
+		// @todo(ccaulfield): remove  this
+		void ApplyPushOut(const T Dt, const TArray<FConstraintHandle*>& InConstraintHandles)
+		{
+		}
+
+	protected:
+		using Base::GetConstraintIndex;
+		using Base::SetConstraintIndex;
+
+	private:
+		void ApplySingle(const T Dt, int32 ConstraintIndex) const;
+
+		TVector<T, d> GetDelta(const TVector<T, d>& WorldSpaceX1, const TVector<T, d>& WorldSpaceX2, const int32 ConstraintIndex, const int32 SpringIndex) const;
+
+		TArray<FConstrainedParticlePair> Constraints;
+		TArray<TArray<TVector<TVector<T, 3>, 2>>> Distances;
+		TArray<TArray<T>> SpringDistances;
+		T CreationThreshold;
+		int32 MaxSprings;
+		T Stiffness;
 
-	}
-};
+		TArray<FConstraintHandle*> Handles;
+		FConstraintHandleAllocator HandleAllocator;
+	};
 }
diff --git a/Engine/Source/Runtime/Experimental/Chaos/Public/Chaos/PBDRigidSpringConstraints.h b/Engine/Source/Runtime/Experimental/Chaos/Public/Chaos/PBDRigidSpringConstraints.h
--- a/Engine/Source/Runtime/Experimental/Chaos/Public/Chaos/PBDRigidSpringConstraints.h
+++ b/Engine/Source/Runtime/Experimental/Chaos/Public/Chaos/PBDRigidSpringConstraints.h
@@ -1,84 +1,203 @@
 // Copyright 1998-2019 Epic Games, Inc. All Rights Reserved.
 #pragma once
 
-#include "Chaos/PBDRigidSpringConstraintsBase.h"
+#include "Chaos/Array.h"
+#include "Chaos/ConstraintHandle.h"
+#include "Chaos/ParticleHandle.h"
 #include "Chaos/PBDConstraintContainer.h"
 
 namespace Chaos
 {
-template<class T, int d>
-class TPBDRigidSpringConstraints2 : public TPBDRigidSpringConstraintsBase<T, d>, public TPBDConstraintContainer<T, d>
-{
-	typedef TPBDRigidSpringConstraintsBase<T, d> Base;
-
-	using Base::Constraints;
-	using Base::Distances;
+	template<class T, int d>
+	class TPBDRigidSpringConstraints;
 
-public:
-	using FConstrainedParticlePair = TVector<TGeometryParticleHandle<T, d>*, 2>;
+	template<class T, int d>
+	class CHAOS_API TPBDRigidSpringConstraintHandle : public TContainerConstraintHandle<TPBDRigidSpringConstraints<T, d>>
+	{
+	public:
+		using Base = TContainerConstraintHandle<TPBDRigidSpringConstraints<T, d>>;
+		using FConstraintContainer = TPBDRigidSpringConstraints<T, d>;
 
+		TPBDRigidSpringConstraintHandle() {}
+		TPBDRigidSpringConstraintHandle(FConstraintContainer* InConstraintContainer, int32 InConstraintIndex) : TContainerConstraintHandle<TPBDRigidSpringConstraints<T, d>>(InConstraintContainer, InConstraintIndex) {}
 
-	TPBDRigidSpringConstraints2(const T InStiffness = (T)1)
-	    : TPBDRigidSpringConstraintsBase<T, d>(InStiffness) 
-	{}
+		const TVector<TVector<T, 3>, 2>& GetConstraintPositions() const;
+		void SetConstraintPositions(const TVector<TVector<T, 3>, 2>& ConstraintPositions);
 
-	TPBDRigidSpringConstraints2(const TArray<TVector<T, 3>>& Locations0, const TArray<TVector<T, 3>>& Locations1, TArray<FConstrainedParticlePair>&& InConstraints, const T InStiffness = (T)1)
-	    : TPBDRigidSpringConstraintsBase<T, d>(Locations0, Locations1, MoveTemp(InConstraints), InStiffness)
-	{}
+	protected:
+		using Base::ConstraintIndex;
+		using Base::ConstraintContainer;
+	};
 
-	virtual ~TPBDRigidSpringConstraints2() {}
 
-	void UpdatePositionBasedState(const T Dt)
+	template<class T, int d>
+	class TPBDRigidSpringConstraints : public TPBDConstraintContainer<T, d>
 	{
-	}
+	public:
+		using Base = TPBDConstraintContainer<T, d>;
+		using FReal = T;
+		static const int Dimensions = d;
+		using FConstraintHandle = TPBDRigidSpringConstraintHandle<FReal, Dimensions>;
+		using FConstraintHandleAllocator = TConstraintHandleAllocator<TPBDRigidSpringConstraints<FReal, Dimensions>>;
+		using FConstrainedParticlePair = TVector<TGeometryParticleHandle<T, d>*, 2>;
+
+		TPBDRigidSpringConstraints(const T InStiffness = (T)1)
+			: Stiffness(InStiffness) 
+		{}
+
+		TPBDRigidSpringConstraints(const TArray<TVector<T, 3>>& Locations0, const TArray<TVector<T, 3>>& Locations1, TArray<FConstrainedParticlePair>&& InConstraints, const T InStiffness = (T)1)
+			: Constraints(MoveTemp(InConstraints)), Stiffness(InStiffness)
+		{
+			if (Constraints.Num() > 0)
+			{
+				Handles.Reserve(Constraints.Num());
+				Distances.Reserve(Constraints.Num());
+				SpringDistances.Reserve(Constraints.Num());
+				for (int32 ConstraintIndex = 0; ConstraintIndex < Constraints.Num(); ++ConstraintIndex)
+				{
+					Handles.Add(HandleAllocator.AllocHandle(this, ConstraintIndex));
+					Distances.Add({});
+					SpringDistances.Add({});
+					UpdateDistance(ConstraintIndex, Locations0[ConstraintIndex], Locations1[ConstraintIndex]);
+				}
+			}
+		}
 
-	void ApplyHelper(const T Dt, const TArray<int32>& InConstraintIndices) const
-	{
-		for (int32 ConstraintIndex : InConstraintIndices)
+		virtual ~TPBDRigidSpringConstraints() {}
+
+		//
+		// Constraint Container API
+		//
+
+		/**
+		 * Get the number of constraints.
+		 */
+		int32 NumConstraints() const
 		{
-			const FConstrainedParticlePair& Constraint = Constraints[ConstraintIndex];
+			return Constraints.Num();
+		}
 
-			TPBDRigidParticleHandle<T, d>* PBDRigid0 = Constraint[0]->AsDynamic();
-			TPBDRigidParticleHandle<T, d>* PBDRigid1 = Constraint[1]->AsDynamic();
-			check((PBDRigid0 && PBDRigid1 && PBDRigid0->Island() == PBDRigid1->Island()) || (!PBDRigid0 && PBDRigid1) || (PBDRigid0 && PBDRigid1));
+		/**
+		 * Add a constraint initialized from current world-space particle positions.
+		 * You would use this method when your objects are already positioned in the world.
+		 */
+		FConstraintHandle* AddConstraint(const FConstrainedParticlePair& InConstrainedParticles, const  TVector<TVector<T, 3>, 2>& InLocations)
+		{
+			Handles.Add(HandleAllocator.AllocHandle(this, Handles.Num()));
+			int32 ConstraintIndex = Constraints.Add(InConstrainedParticles);
+			Distances.Add({});
+			SpringDistances.Add({});
+			UpdateDistance(ConstraintIndex, InLocations[0], InLocations[1]);
+			return Handles.Last();
+		}
 
-			// @todo(ccaulfield): see if we can eliminate the need for all these ifs
-			const TRotation<T, d> & Q0 = PBDRigid0 ? PBDRigid0->Q() : Constraint[0]->R();
-			const TRotation<T, d> & Q1 = PBDRigid1 ? PBDRigid1->Q() : Constraint[1]->R();
-			const TVector<T, d> & P0 = PBDRigid0 ? PBDRigid0->P() : Constraint[0]->X();
-			const TVector<T, d> & P1 = PBDRigid1 ? PBDRigid1->P() : Constraint[1]->X();
+		/**
+		 * Remove the specified constraint.
+		 */
+		void RemoveConstraint(int ConstraintIndex)
+		{
+			FConstraintHandle* ConstraintHandle = Handles[ConstraintIndex];
+			if (ConstraintHandle != nullptr)
+			{
+				// Release the handle for the freed constraint
+				HandleAllocator.FreeHandle(ConstraintHandle);
+				Handles[ConstraintIndex] = nullptr;
+			}
 
-			const TVector<T, d> WorldSpaceX1 = Q0.RotateVector(Distances[ConstraintIndex][0]) + P0;
-			const TVector<T, d> WorldSpaceX2 = Q1.RotateVector(Distances[ConstraintIndex][1]) + P1;
-			const PMatrix<T, d, d> WorldSpaceInvI1 = PBDRigid0 ? (PBDRigid0->Q() * FMatrix::Identity) * PBDRigid0->InvI() * (PBDRigid0->Q() * FMatrix::Identity).GetTransposed() : PMatrix<T, d, d>(0);
-			const PMatrix<T, d, d> WorldSpaceInvI2 = PBDRigid1 ? (PBDRigid1->Q() * FMatrix::Identity) * PBDRigid1->InvI() * (PBDRigid1->Q() * FMatrix::Identity).GetTransposed() : PMatrix<T, d, d>(0);
-			const TVector<T, d> Delta = Base::GetDelta(WorldSpaceX1, WorldSpaceX2, ConstraintIndex);
+			// Swap the last constraint into the gap to keep the array packed
+			Constraints.RemoveAtSwap(ConstraintIndex);
+			Distances.RemoveAtSwap(ConstraintIndex);
+			Handles.RemoveAtSwap(ConstraintIndex);
 
-			if (PBDRigid0)
+			// Update the handle for the constraint that was moved
+			if (ConstraintIndex < Handles.Num())
 			{
-				const TVector<T, d> Radius = WorldSpaceX1 - PBDRigid0->P();
-				PBDRigid0->P() += PBDRigid0->InvM() * Delta;
-				PBDRigid0->Q() += TRotation<T, d>(WorldSpaceInvI1 * TVector<T, d>::CrossProduct(Radius, Delta), 0.f) * PBDRigid0->Q() * T(0.5);
-				PBDRigid0->Q().Normalize();
+				SetConstraintIndex(Handles[ConstraintIndex], ConstraintIndex);
 			}
+		}
+
+		// @todo(ccaulfield): rename/remove  this
+		void RemoveConstraints(const TSet<TGeometryParticleHandle<T, d>*>& RemovedParticles)
+		{
+		}
 
-			if (PBDRigid1)
+		//
+		// Constraint API
+		//
+
+		const FConstraintHandle* GetConstraintHandle(int32 ConstraintIndex) const
+		{
+			return Handles[ConstraintIndex];
+		}
+
+		FConstraintHandle* GetConstraintHandle(int32 ConstraintIndex)
+		{
+			return Handles[ConstraintIndex];
+		}
+
+		/**
+		 * Get the particles that are affected by the specified constraint.
+		 */
+		const FConstrainedParticlePair& GetConstrainedParticles(int32 ConstraintIndex) const
+		{
+			return Constraints[ConstraintIndex];
+		}
+
+		/**
+		 * Get the local-space constraint positions for each body.
+		 */
+		const TVector<TVector<T, 3>, 2>& GetConstraintPositions(int ConstraintIndex) const
+		{
+			return Distances[ConstraintIndex];
+		}
+
+		/**
+		 * Set the local-space constraint positions for each body.
+		 */
+		void SetConstraintPositions(int ConstraintIndex, const TVector<TVector<T, 3>, 2>& ConstraintPositions)
+		{
+			Distances[ConstraintIndex] = ConstraintPositions;
+		}
+
+
+		//
+		// Island Rule API
+		//
+
+		void UpdatePositionBasedState(const T Dt)
+		{
+		}
+
+		void Apply(const T Dt, const TArray<FConstraintHandle*>& InConstraintHandles)
+		{
+			for (FConstraintHandle* ConstraintHandle : InConstraintHandles)
 			{
-				const TVector<T, d> Radius = WorldSpaceX2 - PBDRigid1->P();
-				PBDRigid1->P() -= PBDRigid1->InvM() * Delta;
-				PBDRigid1->Q() += TRotation<T, d>(WorldSpaceInvI2 * TVector<T, d>::CrossProduct(Radius, -Delta), 0.f) * PBDRigid1->Q() * T(0.5);
-				PBDRigid1->Q().Normalize();
+				ApplySingle(Dt, ConstraintHandle->GetConstraintIndex());
 			}
 		}
-	}
 
-	void Apply(const T Dt, const TArray<int32>& InConstraintIndices) const
-	{
-		ApplyHelper(Dt, InConstraintIndices);
-	}
+		// @todo(ccaulfield): remove  this
+		void ApplyPushOut(const T Dt, const TArray<FConstraintHandle*>& InConstraintHandles)
+		{
+		}
 
-	CHAOS_API void ApplyPushOut(const T Dt, const TArray<int32>& InConstraintIndices)
-	{
-	}
-};
+	protected:
+		using Base::GetConstraintIndex;
+		using Base::SetConstraintIndex;
+
+	private:
+		void ApplySingle(const T Dt, int32 ConstraintIndex) const;
+
+		void UpdateDistance(int32 ConstraintIndex, const TVector<T, d>& Location0, const TVector<T, d>& Location1);
+
+		TVector<T, d> GetDelta(int32 ConstraintIndex, const TVector<T, d>& WorldSpaceX1, const TVector<T, d>& WorldSpaceX2) const;
+
+		TArray<FConstrainedParticlePair> Constraints;
+		TArray<TVector<TVector<T, 3>, 2>> Distances;
+		TArray<T> SpringDistances;
+		T Stiffness;
+
+		TArray<FConstraintHandle*> Handles;
+		FConstraintHandleAllocator HandleAllocator;
+	};
 }
diff --git a/Engine/Source/Runtime/Experimental/Chaos/Public/Chaos/PBDRigidsEvolutionGBF.h b/Engine/Source/Runtime/Experimental/Chaos/Public/Chaos/PBDRigidsEvolutionGBF.h
--- a/Engine/Source/Runtime/Experimental/Chaos/Public/Chaos/PBDRigidsEvolutionGBF.h
+++ b/Engine/Source/Runtime/Experimental/Chaos/Public/Chaos/PBDRigidsEvolutionGBF.h
@@ -8,6 +8,7 @@
 #include "Chaos/PerParticleEulerStepVelocity.h"
 #include "Chaos/PerParticleEtherDrag.h"
 #include "Chaos/PerParticlePBDEulerStep.h"
+#include "Chaos/PerParticleExternalForces.h"
 
 namespace Chaos
 {
@@ -30,9 +31,11 @@ public:
 	using Base::SetParticleUpdateVelocityFunction;
 	using Base::SetParticleUpdatePositionFunction;
 	using Base::AddConstraintRule;
+    using Base::AddForceFunction;
 	using Base::Clustering;
 	using typename Base::FForceRule;
 	using FCollisionConstraints = TPBDCollisionConstraint<T, d>;
+	using FExternalForces = TPerParticleExternalForces<T, d>;
 	using FCollisionConstraintRule = TPBDConstraintColorRule<FCollisionConstraints, T, d>;
 
 	static constexpr int32 DefaultNumIterations = 1;
@@ -50,6 +53,20 @@ public:
 	FCollisionConstraints& GetCollisionConstraints() { return CollisionConstraints; }
 	const FCollisionConstraints& GetCollisionConstraints() const { return CollisionConstraints; }
 
+	FExternalForces& GetExternalForces() { return ExternalForces; }
+	const FExternalForces& GetExternalForces() const { return ExternalForces; }
+
+
+	CHAOS_API inline void EndFrame(T Dt)
+	{
+		Particles.GetNonDisabledDynamicView().ParallelFor([&](auto& Particle, int32 Index)
+		{
+			Particle.ExternalForce() = TVector<T, 3>(0);
+			Particle.ExternalTorque() = TVector<T, 3>(0);
+		});
+	}
+
+
 	template <typename TParticleView>
 	void Integrate(const TParticleView& InParticles, T Dt)
 	{
@@ -115,6 +132,7 @@ protected:
 	using Base::ParticleDisableCount;
 	using Base::Collided;
 
+	FExternalForces ExternalForces;
 	FCollisionConstraints CollisionConstraints;
 	FCollisionConstraintRule CollisionRule;
 };
diff --git a/Engine/Source/Runtime/Experimental/Chaos/Public/Chaos/ParticleDirtyFlags.h b/Engine/Source/Runtime/Experimental/Chaos/Public/Chaos/ParticleDirtyFlags.h
--- a/Engine/Source/Runtime/Experimental/Chaos/Public/Chaos/ParticleDirtyFlags.h
+++ b/Engine/Source/Runtime/Experimental/Chaos/Public/Chaos/ParticleDirtyFlags.h
@@ -26,7 +26,9 @@ namespace Chaos
 		M =				1 << 14,
 		InvM =			1 << 15,
 		ObjectState =	1 << 16,
-		Geometry =		1 << 17
+		Geometry =		1 << 17,
+		ExternalForce = 1 << 18,
+		ExternalTorque =1 << 19
 	};
 
 	class FParticleDirtyFlags
diff --git a/Engine/Source/Runtime/Experimental/Chaos/Public/Chaos/ParticleHandle.h b/Engine/Source/Runtime/Experimental/Chaos/Public/Chaos/ParticleHandle.h
--- a/Engine/Source/Runtime/Experimental/Chaos/Public/Chaos/ParticleHandle.h
+++ b/Engine/Source/Runtime/Experimental/Chaos/Public/Chaos/ParticleHandle.h
@@ -64,6 +64,8 @@ void PBDRigidParticleDefaultConstruct(FConcrete& Concrete, const TPBDRigidPartic
 	Concrete.SetQ(Concrete.R());
 	Concrete.SetF(TVector<T, d>(0));
 	Concrete.SetTorque(TVector<T, d>(0));
+	Concrete.SetExternalForce(TVector<T, d>(0));
+	Concrete.SetExternalTorque(TVector<T, d>(0));
 	Concrete.SetM(1);
 	Concrete.SetInvM(1);
 	Concrete.SetI(PMatrix<T, d, d>(1, 1, 1));
@@ -379,6 +381,14 @@ public:
 	TVector<T, d>& Torque() { return PBDRigidParticles->Torque(ParticleIdx); }
 	void SetTorque(const TVector<T, d>& InTorque) { PBDRigidParticles->Torque(ParticleIdx) = InTorque; }
 
+	const TVector<T, d>& ExternalForce() const { return PBDRigidParticles->ExternalForce(ParticleIdx); }
+	TVector<T, d>& ExternalForce() { return PBDRigidParticles->ExternalForce(ParticleIdx); }
+	void SetExternalForce(const TVector<T, d>& InF) { PBDRigidParticles->ExternalForce(ParticleIdx) = InF; }
+
+	const TVector<T, d>& ExternalTorque() const { return PBDRigidParticles->ExternalTorque(ParticleIdx); }
+	TVector<T, d>& ExternalTorque() { return PBDRigidParticles->ExternalTorque(ParticleIdx); }
+	void SetExternalTorque(const TVector<T, d>& InTorque) { PBDRigidParticles->ExternalTorque(ParticleIdx) = InTorque; }
+
 	const PMatrix<T, d, d>& I() const { return PBDRigidParticles->I(ParticleIdx); }
 	PMatrix<T, d, d>& I() { return PBDRigidParticles->I(ParticleIdx); }
 	void SetI(const PMatrix<T, d, d>& InI) { PBDRigidParticles->I(ParticleIdx) = InI; }
@@ -988,6 +998,21 @@ public:
 		this->MTorque = InTorque;
 	}
 
+	const TVector<T, d>& ExternalForce() const { return MExternalForce; }
+	void SetExternalForce(const TVector<T, d>& InExternalForce)
+	{
+		this->MarkDirty(EParticleFlags::ExternalForce);
+		this->MExternalForce = InExternalForce;
+	}
+
+	const TVector<T, d>& ExternalTorque() const { return MExternalTorque; }
+	void SetExternalTorque(const TVector<T, d>& InExternalTorque)
+	{
+		this->MarkDirty(EParticleFlags::ExternalTorque);
+		this->MExternalTorque = InExternalTorque;
+	}
+
+
 	const PMatrix<T, d, d>& I() const { return MI; }
 	void SetI(const PMatrix<T, d, d>& InI)
 	{
@@ -1051,6 +1076,8 @@ private:
 	TVector<T, d> MP;
 	TVector<T, d> MF;
 	TVector<T, d> MTorque;
+	TVector<T, d> MExternalForce;
+	TVector<T, d> MExternalTorque;
 	PMatrix<T, d, d> MI;
 	PMatrix<T, d, d> MInvI;
 	TUniquePtr<TBVHParticles<T, d>> MCollisionParticles;
@@ -1078,8 +1105,8 @@ public:
 		, MPreV(TVector<T, d>(0))
 		, MPreW(TVector<T, d>(0))
 		, MP(TVector<T, d>(0))
-		, MF(TVector<T, d>(0))
-		, MTorque(TVector<T, d>(0))
+		, MExternalForce(TVector<T, d>(0))
+		, MExternalTorque(TVector<T, d>(0))
 		, MI(PMatrix<T, d, d>(0))
 		, MInvI(PMatrix<T, d, d>(0))
 		, MCollisionParticles(nullptr)
@@ -1098,8 +1125,8 @@ public:
 		, MPreV(InParticle.PreV())
 		, MPreW(InParticle.PreW())
 		, MP(InParticle.P())
-		, MF(InParticle.F())
-		, MTorque(InParticle.Torque())
+		, MExternalForce(InParticle.ExternalForce())
+		, MExternalTorque(InParticle.ExternalTorque())
 		, MI(InParticle.I())
 		, MInvI(InParticle.InvI())
 		, MCollisionParticles(nullptr)
@@ -1119,8 +1146,8 @@ public:
 	TVector<T, d> MPreV;
 	TVector<T, d> MPreW;
 	TVector<T, d> MP;
-	TVector<T, d> MF;
-	TVector<T, d> MTorque;
+	TVector<T, d> MExternalForce;
+	TVector<T, d> MExternalTorque;
 	PMatrix<T, d, d> MI;
 	PMatrix<T, d, d> MInvI;
 	const TBVHParticles<T, d> * MCollisionParticles;
@@ -1139,8 +1166,8 @@ public:
 		MPreV = TVector<T, d>(0);
 		MPreW = TVector<T, d>(0);
 		MP = TVector<T, d>(0);
-		MF = TVector<T, d>(0);
-		MTorque = TVector<T, d>(0);
+		MExternalForce = TVector<T, d>(0);
+		MExternalTorque = TVector<T, d>(0);
 		MI = PMatrix<T, d, d>(0);
 		MInvI = PMatrix<T, d, d>(0);
 		MCollisionParticles = nullptr;
diff --git a/Engine/Source/Runtime/Experimental/Chaos/Public/Chaos/PerParticleExternalForces.h b/Engine/Source/Runtime/Experimental/Chaos/Public/Chaos/PerParticleExternalForces.h
new file mode 100644
--- /dev/null
+++ b/Engine/Source/Runtime/Experimental/Chaos/Public/Chaos/PerParticleExternalForces.h
@@ -0,0 +1,29 @@
+// Copyright 1998-2019 Epic Games, Inc. All Rights Reserved.
+#pragma once
+
+#include "Chaos/Map.h"
+#include "Chaos/ParticleHandle.h"
+#include "Chaos/PerParticleRule.h"
+
+namespace Chaos
+{
+	template<class T, int d>
+	class TPerParticleExternalForces : public TPerParticleRule<T, d>
+	{
+
+	public:
+		TPerParticleExternalForces() {}
+
+		virtual ~TPerParticleExternalForces() {}
+
+		inline void Apply(TTransientPBDRigidParticleHandle<T, d>& HandleIn, const T Dt) const override //-V762
+		{
+			if (TPBDRigidParticleHandleImp<T, d, true>* Handle = HandleIn.Handle())
+			{
+				Handle->F() += Handle->ExternalForce();
+				Handle->Torque() += Handle->ExternalForce();
+			}
+		} 
+
+	};
+}
diff --git a/Engine/Source/Runtime/Experimental/Chaos/Public/Chaos/RigidParticles.h b/Engine/Source/Runtime/Experimental/Chaos/Public/Chaos/RigidParticles.h
--- a/Engine/Source/Runtime/Experimental/Chaos/Public/Chaos/RigidParticles.h
+++ b/Engine/Source/Runtime/Experimental/Chaos/Public/Chaos/RigidParticles.h
@@ -52,6 +52,8 @@ class TRigidParticles : public TKinematicGeometryParticles<T, d>
 	{
 		TArrayCollection::AddArray(&MF);
 		TArrayCollection::AddArray(&MT);
+		TArrayCollection::AddArray(&MExternalForce);
+		TArrayCollection::AddArray(&MExternalTorque);
 		TArrayCollection::AddArray(&MI);
 		TArrayCollection::AddArray(&MInvI);
 		TArrayCollection::AddArray(&MM);
@@ -65,10 +67,12 @@ class TRigidParticles : public TKinematicGeometryParticles<T, d>
 	}
 	TRigidParticles(const TRigidParticles<T, d>& Other) = delete;
 	TRigidParticles(TRigidParticles<T, d>&& Other)
-	    : TKinematicGeometryParticles<T, d>(MoveTemp(Other)), MF(MoveTemp(Other.MF)), MT(MoveTemp(Other.MT)), MI(MoveTemp(Other.MI)), MInvI(MoveTemp(Other.MInvI)), MM(MoveTemp(Other.MM)), MInvM(MoveTemp(Other.MInvM)), MCollisionParticles(MoveTemp(Other.MCollisionParticles)), MCollisionGroup(MoveTemp(Other.MCollisionGroup)), MObjectState(MoveTemp(Other.MObjectState))
+	    : TKinematicGeometryParticles<T, d>(MoveTemp(Other)), MF(MoveTemp(Other.MF)), MT(MoveTemp(Other.MT)), MExternalForce(MoveTemp(Other.MExternalForce)), MExternalTorque(MoveTemp(Other.MExternalTorque)), MI(MoveTemp(Other.MI)), MInvI(MoveTemp(Other.MInvI)), MM(MoveTemp(Other.MM)), MInvM(MoveTemp(Other.MInvM)), MCollisionParticles(MoveTemp(Other.MCollisionParticles)), MCollisionGroup(MoveTemp(Other.MCollisionGroup)), MObjectState(MoveTemp(Other.MObjectState))
 	{
 		TArrayCollection::AddArray(&MF);
 		TArrayCollection::AddArray(&MT);
+		TArrayCollection::AddArray(&MExternalForce);
+		TArrayCollection::AddArray(&MExternalTorque);
 		TArrayCollection::AddArray(&MI);
 		TArrayCollection::AddArray(&MInvI);
 		TArrayCollection::AddArray(&MM);
@@ -87,6 +91,12 @@ class TRigidParticles : public TKinematicGeometryParticles<T, d>
 	const TVector<T, d>& F(const int32 Index) const { return MF[Index]; }
 	TVector<T, d>& F(const int32 Index) { return MF[Index]; }
 
+	const TVector<T, d>& ExternalTorque(const int32 Index) const { return MExternalTorque[Index]; }
+	TVector<T, d>& ExternalTorque(const int32 Index) { return MExternalTorque[Index]; }
+
+	const TVector<T, d>& ExternalForce(const int32 Index) const { return MExternalForce[Index]; }
+	TVector<T, d>& ExternalForce(const int32 Index) { return MExternalForce[Index]; }
+
 	const PMatrix<T, d, d>& I(const int32 Index) const { return MI[Index]; }
 	PMatrix<T, d, d>& I(const int32 Index) { return MI[Index]; }
 
@@ -135,19 +145,21 @@ class TRigidParticles : public TKinematicGeometryParticles<T, d>
 	FString ToString(int32 index) const
 	{
 		FString BaseString = TKinematicGeometryParticles<T, d>::ToString(index);
-		return FString::Printf(TEXT("%s, MF:%s, MT:%s, MI:%s, MInvI:%s, MM:%f, MInvM:%f, MCollisionParticles(num):%d, MCollisionGroup:%d, MDisabled:%d, MSleepring:%d, MIsland:%d"), *BaseString, *F(index).ToString(), *Torque(index).ToString(), *I(index).ToString(), *InvI(index).ToString(), M(index), InvM(index), CollisionParticlesSize(index), CollisionGroup(index), Disabled(index), Sleeping(index), Island(index));
+		return FString::Printf(TEXT("%s, MF:%s, MT:%s, MExternalForce:%s, MExternalTorque:%s, MI:%s, MInvI:%s, MM:%f, MInvM:%f, MCollisionParticles(num):%d, MCollisionGroup:%d, MDisabled:%d, MSleepring:%d, MIsland:%d"), *BaseString, *F(index).ToString(), *Torque(index).ToString(), *ExternalForce(index).ToString(), *ExternalTorque(index).ToString(), *I(index).ToString(), *InvI(index).ToString(), M(index), InvM(index), CollisionParticlesSize(index), CollisionGroup(index), Disabled(index), Sleeping(index), Island(index));
 	}
 
 	void Serialize(FChaosArchive& Ar)
 	{
 		TKinematicGeometryParticles<T,d>::Serialize(Ar);
-		Ar << MF << MT << MI << MInvI << MM << MInvM;
+		Ar << MF << MT << MExternalForce << MExternalTorque << MI << MInvI << MM << MInvM;
 		Ar << MCollisionParticles << MCollisionGroup << MIsland << MDisabled << MObjectState;
 	}
 
   private:
 	TArrayCollectionArray<TVector<T, d>> MF;
 	TArrayCollectionArray<TVector<T, d>> MT;
+	TArrayCollectionArray<TVector<T, d>> MExternalForce;
+	TArrayCollectionArray<TVector<T, d>> MExternalTorque;
 	TArrayCollectionArray<PMatrix<T, d, d>> MI;
 	TArrayCollectionArray<PMatrix<T, d, d>> MInvI;
 	TArrayCollectionArray<T> MM;
diff --git a/Engine/Source/Runtime/Experimental/Chaos/Public/Chaos/TaperedCylinder.h b/Engine/Source/Runtime/Experimental/Chaos/Public/Chaos/TaperedCylinder.h
--- a/Engine/Source/Runtime/Experimental/Chaos/Public/Chaos/TaperedCylinder.h
+++ b/Engine/Source/Runtime/Experimental/Chaos/Public/Chaos/TaperedCylinder.h
@@ -208,7 +208,7 @@ namespace Chaos
 			if (Determinant == 0)
 			{
 				T Root = -b / (2 * a);
-				auto RootPoint = Root * Direction + StartPoint;
+				TVector<T, 3> RootPoint = Root * Direction + StartPoint;
 				if (Root >= 0 && Root <= Length && TVector<T, 3>::DotProduct(RootPoint - Top, -BaseNormal) >= 0)
 				{
 					Intersections.Add(MakePair(Root, RootPoint));
@@ -218,8 +218,8 @@ namespace Chaos
 			{
 				T Root1 = (-b - sqrt(Determinant)) / (2 * a);
 				T Root2 = (-b + sqrt(Determinant)) / (2 * a);
-				auto Root1Point = Root1 * Direction + StartPoint;
-				auto Root2Point = Root2 * Direction + StartPoint;
+				TVector<T, 3> Root1Point = Root1 * Direction + StartPoint;
+				TVector<T, 3> Root2Point = Root2 * Direction + StartPoint;
 				if (Root1 < 0 || Root1 > Length || TVector<T, 3>::DotProduct(Root1Point - Top, -BaseNormal) < 0)
 				{
 					if (Root2 >= 0 && Root2 <= Length && TVector<T, 3>::DotProduct(Root2Point - Top, -BaseNormal) >= 0)
diff --git a/Engine/Source/Runtime/Experimental/Chaos/Public/Chaos/TriangleMeshImplicitObject.h b/Engine/Source/Runtime/Experimental/Chaos/Public/Chaos/TriangleMeshImplicitObject.h
--- a/Engine/Source/Runtime/Experimental/Chaos/Public/Chaos/TriangleMeshImplicitObject.h
+++ b/Engine/Source/Runtime/Experimental/Chaos/Public/Chaos/TriangleMeshImplicitObject.h
@@ -33,7 +33,8 @@ namespace Chaos
 
 		bool OverlapGeom(const TImplicitObject<T,3>& QueryGeom, const TRigidTransform<T, 3>& QueryTM, const T Thickness, const TVector<T,3> Scale = TVector<T,3>(1)) const;
 		bool SweepGeom(const TImplicitObject<T, 3>& QueryGeom, const TRigidTransform<T, 3>& StartTM, const TVector<T,3>& Dir, const T Length, T& OutTime, TVector<T,3>& OutPosition, TVector<T,3>& OutNormal, int32& OutFaceIndex, const T Thickness = 0, const TVector<T,3> Scale = TVector<T,3>(1)) const;
-		virtual int32 FindMostOpposingFace(const TVector<T, 3>& Position, const TVector<T, 3>& UnitDir, int32 HintFaceIndex) const override;
+		virtual int32 FindMostOpposingFace(const TVector<T, 3>& Position, const TVector<T, 3>& UnitDir, int32 HintFaceIndex, T SearchDistance) const override;
+		virtual TVector<T, 3> FindGeometryOpposingNormal(const TVector<T, 3>& DenormDir, int32 FaceIndex, const TVector<T, 3>& OriginalNormal) const override;
 
 		virtual const TBox<T, 3>& BoundingBox() const
 		{
@@ -92,11 +93,14 @@ namespace Chaos
 			return Result;
 		}
 
+		TVector<T, 3> GetFaceNormal(const int32 FaceIdx) const;
+
 	private:
 
 		void RebuildConvexHulls();
 		void RebuildBV();
 
+
 		TParticles<T, 3> MParticles;
 		TArray<TVector<int32, 3>> MElements;
 		TBox<T, 3> MLocalBoundingBox;
diff --git a/Engine/Source/Runtime/Experimental/ChaosSolvers/Private/Framework/PersistentTask.cpp b/Engine/Source/Runtime/Experimental/ChaosSolvers/Private/Framework/PersistentTask.cpp
--- a/Engine/Source/Runtime/Experimental/ChaosSolvers/Private/Framework/PersistentTask.cpp
+++ b/Engine/Source/Runtime/Experimental/ChaosSolvers/Private/Framework/PersistentTask.cpp
@@ -17,7 +17,7 @@
 
 #include "PhysicsProxy/SkeletalMeshPhysicsProxy.h"
 #include "PhysicsProxy/StaticMeshPhysicsProxy.h"
-#include "Chaos/Framework/SingleParticlePhysicsProxy.h"
+#include "PhysicsProxy/SingleParticlePhysicsProxy.h"
 #include "PhysicsProxy/GeometryCollectionPhysicsProxy.h"
 #include "PhysicsProxy/FieldSystemPhysicsProxy.h"
 
diff --git a/Engine/Source/Runtime/Experimental/ChaosSolvers/Private/PBDRigidsSolver.cpp b/Engine/Source/Runtime/Experimental/ChaosSolvers/Private/PBDRigidsSolver.cpp
--- a/Engine/Source/Runtime/Experimental/ChaosSolvers/Private/PBDRigidsSolver.cpp
+++ b/Engine/Source/Runtime/Experimental/ChaosSolvers/Private/PBDRigidsSolver.cpp
@@ -7,13 +7,13 @@
 
 #include "Async/AsyncWork.h"
 #include "Chaos/ChaosArchive.h"
-#include "Chaos/Framework/SingleParticlePhysicsProxy.h"
 #include "Chaos/PBDCollisionConstraintUtil.h"
 #include "Chaos/Utilities.h"
 #include "ChaosStats.h"
 #include "ChaosSolversModule.h"
 #include "HAL/FileManager.h"
 #include "Misc/ScopeLock.h"
+#include "PhysicsProxy/SingleParticlePhysicsProxy.h"
 #include "PhysicsProxy/SkeletalMeshPhysicsProxy.h"
 #include "PhysicsProxy/StaticMeshPhysicsProxy.h"
 #include "PhysicsProxy/GeometryCollectionPhysicsProxy.h"
@@ -75,8 +75,8 @@ namespace Chaos
 			}
 
 			{
-				//SCOPE_CYCLE_COUNTER(STAT_EndFrame);
-				//MSolver->EndFrameCallback(MDeltaTime);
+				SCOPE_CYCLE_COUNTER(STAT_EndFrame);
+				MSolver->GetEvolution()->EndFrame(MDeltaTime);
 			}
 
 			MSolver->GetSolverTime() += MDeltaTime;
@@ -266,7 +266,7 @@ namespace Chaos
 			}
 
 			// Use the handle to destroy the particle data
-			Solver->Particles.DestroyParticle(Handle);
+			Solver->GetEvolution()->DestroyParticle(Handle);
 		});
 	}
 
diff --git a/Engine/Source/Runtime/Experimental/ChaosSolvers/Public/Framework/CommandBuffer.h b/Engine/Source/Runtime/Experimental/ChaosSolvers/Public/Framework/CommandBuffer.h
--- a/Engine/Source/Runtime/Experimental/ChaosSolvers/Public/Framework/CommandBuffer.h
+++ b/Engine/Source/Runtime/Experimental/ChaosSolvers/Public/Framework/CommandBuffer.h
@@ -22,7 +22,7 @@ namespace Chaos
 
 		friend class FCommandList;
 
-		template<enum Chaos::EThreadingMode>
+		template<Chaos::EThreadingMode>
 		friend class FDispatcher;
 		
 		// Element type for solver command storage
diff --git a/Engine/Source/Runtime/Experimental/ChaosSolvers/Public/PBDRigidsSolver.h b/Engine/Source/Runtime/Experimental/ChaosSolvers/Public/PBDRigidsSolver.h
--- a/Engine/Source/Runtime/Experimental/ChaosSolvers/Public/PBDRigidsSolver.h
+++ b/Engine/Source/Runtime/Experimental/ChaosSolvers/Public/PBDRigidsSolver.h
@@ -8,7 +8,6 @@
 #include "Chaos/Defines.h"
 #include "Chaos/Framework/PhysicsProxy.h"
 #include "Chaos/Framework/MultiBufferResource.h"
-#include "Chaos/Framework/SingleParticlePhysicsProxy.h"
 #include "Chaos/PBDRigidParticles.h"
 #include "Chaos/PBDRigidsEvolutionGBF.h"
 #include "Chaos/PBDCollisionConstraint.h"
@@ -24,6 +23,7 @@
 #include "Framework/Dispatcher.h"
 #include "Field/FieldSystem.h"
 #include "EventManager.h"
+#include "PhysicsProxy/SingleParticlePhysicsProxyFwd.h"
 #include "SolverEventFilters.h"
 
 #ifndef CHAOS_WITH_PAUSABLE_SOLVER
diff --git a/Engine/Source/Runtime/Experimental/ChaosSolvers/Public/PhysicsProxy/PhysicsProxies.h b/Engine/Source/Runtime/Experimental/ChaosSolvers/Public/PhysicsProxy/PhysicsProxies.h
--- a/Engine/Source/Runtime/Experimental/ChaosSolvers/Public/PhysicsProxy/PhysicsProxies.h
+++ b/Engine/Source/Runtime/Experimental/ChaosSolvers/Public/PhysicsProxy/PhysicsProxies.h
@@ -10,7 +10,7 @@
  * the currently supported set of objects.
  */
 
-#include "Chaos/Framework/SingleParticlePhysicsProxy.h"
+#include "SingleParticlePhysicsProxy.h"
 #include "FieldSystemPhysicsProxy.h"
 #include "GeometryCollectionPhysicsProxy.h"
 #include "SkeletalMeshPhysicsProxy.h"
diff --git a/Engine/Source/Runtime/Landscape/Private/LandscapeCollision.cpp b/Engine/Source/Runtime/Landscape/Private/LandscapeCollision.cpp
--- a/Engine/Source/Runtime/Landscape/Private/LandscapeCollision.cpp
+++ b/Engine/Source/Runtime/Landscape/Private/LandscapeCollision.cpp
@@ -317,7 +317,7 @@ void ULandscapeHeightfieldCollisionComponent::OnCreatePhysicsState()
 					if (LandscapeComponentGeomEd.isValid())
 					{
 #if WITH_CHAOS || WITH_IMMEDIATE_PHYSX
-                        ensure(false);
+						UE_LOG(LogLandscape, Warning, TEXT("Failed to create editor shapes, currently unimplemented for Chaos"));
 #else
 						FPhysicsMaterialHandle_PhysX MaterialHandle = GEngine->DefaultPhysMaterial->GetPhysicsMaterial();
 						PxMaterial* PDefaultMat = MaterialHandle.Material;
@@ -350,6 +350,10 @@ void ULandscapeHeightfieldCollisionComponent::OnCreatePhysicsState()
 				BodyInstance.PhysicsUserData = FPhysicsUserData(&BodyInstance);
 				BodyInstance.OwnerComponent = this;
 
+#if WITH_CHAOS
+				PhysHandle->SetUserData(&BodyInstance.PhysicsUserData);
+#endif
+
 #if WITH_CHAOS || WITH_IMMEDIATE_PHYSX
 				TArray<FPhysicsActorHandle> Actors;
 				Actors.Add(PhysHandle);
@@ -440,7 +444,10 @@ void ULandscapeHeightfieldCollisionComponent::CreateCollisionObject()
 				for (UPhysicalMaterial* PhysicalMaterial : CookedPhysicalMaterials)
 				{
 #if WITH_CHAOS || WITH_IMMEDIATE_PHYSX
-                    ensure(false);
+#if WITH_PHYSX
+					//todo: total hack until we get landscape fully converted to chaos
+					HeightfieldRef->UsedPhysicalMaterialArray.Add(GPhysXSDK->createMaterial(1,1,1));
+#endif
 #else
 					const FPhysicsMaterialHandle_PhysX& MaterialHandle = PhysicalMaterial->GetPhysicsMaterial();
 					HeightfieldRef->UsedPhysicalMaterialArray.Add(MaterialHandle.Material);
diff --git a/Engine/Source/Runtime/Launch/Private/LaunchEngineLoop.cpp b/Engine/Source/Runtime/Launch/Private/LaunchEngineLoop.cpp
--- a/Engine/Source/Runtime/Launch/Private/LaunchEngineLoop.cpp
+++ b/Engine/Source/Runtime/Launch/Private/LaunchEngineLoop.cpp
@@ -3374,7 +3374,7 @@ bool FEngineLoop::LoadStartupCoreModules()
 	FModuleManager::Get().LoadModule(TEXT("MediaAssets"));
 #endif
 
-	FModuleManager::Get().LoadModule(TEXT("ClothingSystemRuntime"));
+	FModuleManager::Get().LoadModule(TEXT("ClothingSystemRuntimeNv"));
 #if WITH_EDITOR
 	FModuleManager::Get().LoadModule(TEXT("ClothingSystemEditor"));
 #endif
diff --git a/Engine/Source/Runtime/MeshDescription/Public/MeshAttributeArray.h b/Engine/Source/Runtime/MeshDescription/Public/MeshAttributeArray.h
--- a/Engine/Source/Runtime/MeshDescription/Public/MeshAttributeArray.h
+++ b/Engine/Source/Runtime/MeshDescription/Public/MeshAttributeArray.h
@@ -857,7 +857,7 @@ protected:
 	}
 
 	/** Serialization */
-	friend FArchive& operator<<( FArchive& Ar, FAttributesSetBase& AttributesSet );
+	friend MESHDESCRIPTION_API FArchive& operator<<( FArchive& Ar, FAttributesSetBase& AttributesSet );
 
 	template <typename T>
 	friend void SerializeLegacy( FArchive& Ar, FAttributesSetBase& AttributesSet );
diff --git a/Engine/Source/Runtime/PhysicsCore/Private/PhysicsInterfaceUtilsCore.cpp b/Engine/Source/Runtime/PhysicsCore/Private/PhysicsInterfaceUtilsCore.cpp
--- a/Engine/Source/Runtime/PhysicsCore/Private/PhysicsInterfaceUtilsCore.cpp
+++ b/Engine/Source/Runtime/PhysicsCore/Private/PhysicsInterfaceUtilsCore.cpp
@@ -4,6 +4,11 @@
 #include "PhysXPublicCore.h"
 #include "PhysicsInterfaceTypesCore.h"
 
+#if INCLUDE_CHAOS
+#include "Chaos/GeometryParticles.h"
+#include "Chaos/ParticleHandle.h"
+#endif
+
 #if WITH_PHYSX
 
 PxShapeFlags BuildPhysXShapeFlags(FBodyCollisionFlags BodyCollisionFlags, bool bPhysicsStatic, bool bIsTriangleMesh)
@@ -264,7 +269,9 @@ uint32 FindFaceIndex(const FHitLocation& PHit, const FVector& UnitDir)
 	}
 
 	return PHit.faceIndex;	//If no custom logic just return whatever face index they initially had
+#else
+	const FTransform WorldTM(PHit.Actor->R(), PHit.Actor->X());
+	const FVector LocalPosition = WorldTM.InverseTransformPositionNoScale(PHit.WorldPosition);
+	return PHit.Shape->Geometry->FindMostOpposingFace(LocalPosition, UnitDir, PHit.FaceIndex, 1);	//todo:this number matches the one above, but is it right?
 #endif
-	ensure(false);
-	return 0;
 }
diff --git a/Engine/Source/Runtime/PhysicsCore/Public/ChaosInterfaceWrapperCore.h b/Engine/Source/Runtime/PhysicsCore/Public/ChaosInterfaceWrapperCore.h
--- a/Engine/Source/Runtime/PhysicsCore/Public/ChaosInterfaceWrapperCore.h
+++ b/Engine/Source/Runtime/PhysicsCore/Public/ChaosInterfaceWrapperCore.h
@@ -79,6 +79,7 @@ inline ECollisionShapeType GetImplicitType(const Chaos::TImplicitObject<float, 3
 	case Chaos::ImplicitObjectType::Convex: return ECollisionShapeType::Convex;
 	case Chaos::ImplicitObjectType::TriangleMesh: return ECollisionShapeType::Trimesh;
 	case Chaos::ImplicitObjectType::HeightField: return ECollisionShapeType::Heightfield;
+	case Chaos::ImplicitObjectType::Scaled: return ECollisionShapeType::Scaled;
 	default: break;
 	}
 
diff --git a/Engine/Source/Runtime/PhysicsCore/Public/PhysXIncludes.h b/Engine/Source/Runtime/PhysicsCore/Public/PhysXIncludes.h
--- a/Engine/Source/Runtime/PhysicsCore/Public/PhysXIncludes.h
+++ b/Engine/Source/Runtime/PhysicsCore/Public/PhysXIncludes.h
@@ -76,10 +76,10 @@ THIRD_PARTY_INCLUDES_END
 THIRD_PARTY_INCLUDES_START
 
 #if WITH_APEX_CLOTHING
-#include "ModuleClothing.h"
+#include "clothing/ModuleClothing.h"
 #include "clothing/ClothingAsset.h"
-#include "ClothingActor.h"
-#include "ClothingCollision.h"
+#include "clothing/ClothingActor.h"
+#include "clothing/ClothingCollision.h"
 #endif
 
 #if WITH_APEX_LEGACY
diff --git a/Engine/Source/Runtime/PhysicsCore/Public/PhysicsInterfaceTypesCore.h b/Engine/Source/Runtime/PhysicsCore/Public/PhysicsInterfaceTypesCore.h
--- a/Engine/Source/Runtime/PhysicsCore/Public/PhysicsInterfaceTypesCore.h
+++ b/Engine/Source/Runtime/PhysicsCore/Public/PhysicsInterfaceTypesCore.h
@@ -66,6 +66,7 @@ enum class ECollisionShapeType : uint8
 	Convex,
 	Trimesh,
 	Heightfield,
+	Scaled,
 	None
 };
 
diff --git a/Engine/Source/Runtime/PhysicsSQ/Private/SQAccelerator.cpp b/Engine/Source/Runtime/PhysicsSQ/Private/SQAccelerator.cpp
--- a/Engine/Source/Runtime/PhysicsSQ/Private/SQAccelerator.cpp
+++ b/Engine/Source/Runtime/PhysicsSQ/Private/SQAccelerator.cpp
@@ -58,8 +58,8 @@ void FSQAcceleratorUnion::RemoveSQAccelerator(ISQAccelerator* AcceleratorToRemov
 }
 
 #if INCLUDE_CHAOS
-FChaosSQAccelerator::FChaosSQAccelerator(const Chaos::TPBDRigidsEvolutionGBF<float, 3>& InEvolution)
-	: Evolution(InEvolution)
+FChaosSQAccelerator::FChaosSQAccelerator(const Chaos::ISpatialAcceleration<Chaos::TAccelerationStructureHandle<float, 3>, float, 3>& InSpatialAcceleration)
+	: SpatialAcceleration(InSpatialAcceleration)
 {}
 
 struct FPreFilterInfo
@@ -68,15 +68,16 @@ struct FPreFilterInfo
 	int32 ActorIdx;
 };
 
-void FillHitHelper(ChaosInterface::FLocationHit& Hit, const float Distance, const FVector& WorldPosition, const FVector& WorldNormal)
+void FillHitHelper(ChaosInterface::FLocationHit& Hit, const float Distance, const FVector& WorldPosition, const FVector& WorldNormal, int32 FaceIdx)
 {
 	Hit.Distance = Distance;
 	Hit.WorldPosition = WorldPosition;
 	Hit.WorldNormal = WorldNormal;
-	Hit.Flags = Distance > 0 ? EHitFlags::Distance | EHitFlags::Normal | EHitFlags::Position : EHitFlags::Distance;
+	Hit.Flags = Distance > 0 ? EHitFlags::Distance | EHitFlags::Normal | EHitFlags::Position : EHitFlags::Distance | EHitFlags::FaceIndex;
+	Hit.FaceIndex = FaceIdx;
 }
 
-void FillHitHelper(ChaosInterface::FOverlapHit& Hit, const float Distance, const FVector& WorldPosition, const FVector& WorldNormal)
+void FillHitHelper(ChaosInterface::FOverlapHit& Hit, const float Distance, const FVector& WorldPosition, const FVector& WorldNormal, int32 FaceIdx)
 {
 }
 
@@ -163,7 +164,8 @@ private:
 	{
 		//todo: add a check to ensure hitbuffer matches SQ type
 		using namespace Chaos;
-		const TShapesArray<float,3>& Shapes = Payload->GTGeometryParticle()->ShapesArray();	//TODO: use GT particles directly
+		TGeometryParticle<float, 3>* GeometryParticle = Payload.GetExternalGeometryParticle_ExternalThread();
+		const TShapesArray<float,3>& Shapes = GeometryParticle->ShapesArray();
 
 		for (const auto& Shape : Shapes)
 		{
@@ -171,7 +173,7 @@ private:
 			//TODO: use gt particles directly
 			//#TODO alternative to px flags
 #if WITH_PHYSX
-			ECollisionQueryHitType HitType = QueryFilterData.flags & PxQueryFlag::ePREFILTER ? QueryCallback.PreFilter(P2UFilterData(QueryFilterData.data), *Shape, *Payload->GTGeometryParticle()) : ECollisionQueryHitType::Block;
+			ECollisionQueryHitType HitType = QueryFilterData.flags & PxQueryFlag::ePREFILTER ? QueryCallback.PreFilter(P2UFilterData(QueryFilterData.data), *Shape, *GeometryParticle) : ECollisionQueryHitType::Block;
 #else
 			//#TODO Chaos flag alternative
 			ensure(false);
@@ -179,17 +181,17 @@ private:
 #endif
 			if (HitType != ECollisionQueryHitType::None)
 			{
-				const TRigidTransform<float, 3> ActorTM(Payload->X(), Payload->R());
+				const TRigidTransform<float, 3> ActorTM(GeometryParticle->X(), GeometryParticle->R());
 
 				THitType Hit;
-				Hit.Actor = Payload->GTGeometryParticle();	//todo: use GT particles directly
+				Hit.Actor = GeometryParticle;
 				Hit.Shape = Shape.Get();
 
 				bool bHit = false;
 
 				TVector<float, 3> WorldPosition, WorldNormal;
 				float Distance = 0;	//not needed but fixes compiler warning for overlap
-				int32 FaceIdx;	//todo: pass back to unreal
+				int32 FaceIdx = INDEX_NONE;	//not needed but fixes compiler warning for overlap
 
 				if (SQ == ESQType::Raycast)
 				{
@@ -216,7 +218,7 @@ private:
 
 				if(bHit)
 				{
-					FillHitHelper(Hit, Distance, WorldPosition, WorldNormal);
+					FillHitHelper(Hit, Distance, WorldPosition, WorldNormal, FaceIdx);
 #if WITH_PHYSX
 					HitType = QueryFilterData.flags & PxQueryFlag::ePOSTFILTER ? QueryCallback.PostFilter(P2UFilterData(QueryFilterData.data), Hit) : HitType;
 #else
@@ -267,12 +269,10 @@ void FChaosSQAccelerator::Raycast(const FVector& Start, const FVector& Dir, cons
 	using namespace Chaos;
 	using namespace ChaosInterface;
 
-	TSQVisitor<TGeometryParticleHandle<float, 3>*, FRaycastHit> RaycastVisitor(Start, Dir, HitBuffer, OutputFlags, QueryFilterData, QueryCallback);
-	const auto& SpatialAcceleration = Evolution.GetCollisionConstraints().GetSpatialAcceleration();
+	TSQVisitor<TAccelerationStructureHandle<float, 3>, FRaycastHit> RaycastVisitor(Start, Dir, HitBuffer, OutputFlags, QueryFilterData, QueryCallback);
 	HitBuffer.IncFlushCount();
 	SpatialAcceleration.Raycast(Start, Dir, DeltaMagnitude, RaycastVisitor);
 	HitBuffer.DecFlushCount();
-	Evolution.GetCollisionConstraints().ReleaseSpatialAcceleration();
 }
 
 void FChaosSQAccelerator::Sweep(const Chaos::TImplicitObject<float, 3>& QueryGeom, const FTransform& StartTM, const FVector& Dir, const float DeltaMagnitude, ChaosInterface::FSQHitBuffer<ChaosInterface::FSweepHit>& HitBuffer, EHitFlags OutputFlags, const FQueryFilterData& QueryFilterData, ICollisionQueryFilterCallbackBase& QueryCallback) const
@@ -282,12 +282,10 @@ void FChaosSQAccelerator::Sweep(const Chaos::TImplicitObject<float, 3>& QueryGeo
 
 	const TBox<float, 3> Bounds = QueryGeom.BoundingBox().TransformedBox(StartTM);
 	const FVector HalfExtents = Bounds.Extents() * 0.5f;
-	TSQVisitor<TGeometryParticleHandle<float, 3>*, FSweepHit> SweepVisitor(StartTM, Dir, HitBuffer, OutputFlags, QueryFilterData, QueryCallback, QueryGeom);
-	const auto& SpatialAcceleration = Evolution.GetCollisionConstraints().GetSpatialAcceleration();
+	TSQVisitor<TAccelerationStructureHandle<float, 3>, FSweepHit> SweepVisitor(StartTM, Dir, HitBuffer, OutputFlags, QueryFilterData, QueryCallback, QueryGeom);
 	HitBuffer.IncFlushCount();
 	SpatialAcceleration.Sweep(Bounds.GetCenter(), Dir, DeltaMagnitude, HalfExtents, SweepVisitor);
 	HitBuffer.DecFlushCount();
-	Evolution.GetCollisionConstraints().ReleaseSpatialAcceleration();
 }
 
 void FChaosSQAccelerator::Overlap(const Chaos::TImplicitObject<float, 3>& QueryGeom, const FTransform& GeomPose, ChaosInterface::FSQHitBuffer<ChaosInterface::FOverlapHit>& HitBuffer, const FQueryFilterData& QueryFilterData, ICollisionQueryFilterCallbackBase& QueryCallback) const
@@ -296,19 +294,17 @@ void FChaosSQAccelerator::Overlap(const Chaos::TImplicitObject<float, 3>& QueryG
 	using namespace ChaosInterface;
 
 	const TBox<float, 3> Bounds = QueryGeom.BoundingBox().TransformedBox(GeomPose);
-	TSQVisitor<TGeometryParticleHandle<float, 3>*, FOverlapHit> OverlapVisitor(GeomPose, HitBuffer, QueryFilterData, QueryCallback, QueryGeom);
-	const auto& SpatialAcceleration = Evolution.GetCollisionConstraints().GetSpatialAcceleration();
+	TSQVisitor<TAccelerationStructureHandle<float, 3>, FOverlapHit> OverlapVisitor(GeomPose, HitBuffer, QueryFilterData, QueryCallback, QueryGeom);
 	HitBuffer.IncFlushCount();
 	SpatialAcceleration.Overlap(Bounds, OverlapVisitor);
 	HitBuffer.DecFlushCount();
-	Evolution.GetCollisionConstraints().ReleaseSpatialAcceleration();
 }
 
 #endif
 
 #if WITH_PHYSX && INCLUDE_CHAOS
-FChaosSQAcceleratorAdapter::FChaosSQAcceleratorAdapter(const Chaos::TPBDRigidsEvolutionGBF<float, 3>& InEvolution)
-	: ChaosSQAccelerator(InEvolution)
+FChaosSQAcceleratorAdapter::FChaosSQAcceleratorAdapter(const Chaos::ISpatialAcceleration<Chaos::TAccelerationStructureHandle<float, 3>, float, 3>& InSpatialAcceleration)
+	: ChaosSQAccelerator(InSpatialAcceleration)
 {
 }
 
diff --git a/Engine/Source/Runtime/PhysicsSQ/Public/SQAccelerator.h b/Engine/Source/Runtime/PhysicsSQ/Public/SQAccelerator.h
--- a/Engine/Source/Runtime/PhysicsSQ/Public/SQAccelerator.h
+++ b/Engine/Source/Runtime/PhysicsSQ/Public/SQAccelerator.h
@@ -23,11 +23,14 @@ namespace physx
 
 namespace Chaos
 {
-	template <typename, int>
-	class TPBDRigidsEvolutionGBF;
-
 	template <typename, int>
 	class TImplicitObject;
+
+	template <typename T, int d>
+	class TAccelerationStructureHandle;
+
+	template <typename TPayload, typename T, int d>
+	class ISpatialAcceleration;
 }
 
 class FSQAccelerator;
@@ -41,7 +44,7 @@ class PHYSICSSQ_API FChaosSQAccelerator
 {
 public:
 
-	FChaosSQAccelerator(const Chaos::TPBDRigidsEvolutionGBF<float, 3>& InEvolution);
+	FChaosSQAccelerator(const Chaos::ISpatialAcceleration<Chaos::TAccelerationStructureHandle<float, 3>,float, 3>& InSpatialAcceleration);
 	virtual ~FChaosSQAccelerator() {};
 
 	void Raycast(const FVector& Start, const FVector& Dir, const float DeltaMagnitude, ChaosInterface::FSQHitBuffer<ChaosInterface::FRaycastHit>& HitBuffer, EHitFlags OutputFlags, const FQueryFilterData& QueryFilterData, ICollisionQueryFilterCallbackBase& QueryCallback) const;
@@ -49,7 +52,7 @@ public:
 	void Overlap(const Chaos::TImplicitObject<float, 3>& QueryGeom, const FTransform& GeomPose, ChaosInterface::FSQHitBuffer<ChaosInterface::FOverlapHit>& HitBuffer, const FQueryFilterData& QueryFilterData, ICollisionQueryFilterCallbackBase& QueryCallback) const;
 
 private:
-	const Chaos::TPBDRigidsEvolutionGBF<float, 3>& Evolution;
+	const Chaos::ISpatialAcceleration<Chaos::TAccelerationStructureHandle<float, 3>, float, 3>& SpatialAcceleration;
 
 };
 #endif
@@ -90,7 +93,7 @@ class PHYSICSSQ_API FChaosSQAcceleratorAdapter : public ISQAccelerator
 {
 public:
 
-	FChaosSQAcceleratorAdapter(const Chaos::TPBDRigidsEvolutionGBF<float, 3>& InEvolution);
+	FChaosSQAcceleratorAdapter(const Chaos::ISpatialAcceleration<Chaos::TAccelerationStructureHandle<float, 3>, float, 3>& InSpatialAcceleration);
 	virtual ~FChaosSQAcceleratorAdapter() {};
 
 	virtual void Raycast(const FVector& Start, const FVector& Dir, const float DeltaMagnitude, FPhysicsHitCallback<FHitRaycast>& HitBuffer, EHitFlags OutputFlags, const FQueryFilterData& QueryFilterData, ICollisionQueryFilterCallbackBase& QueryCallback) const override;
diff --git a/Engine/Source/Runtime/PhysicsSQ/Public/SQVerifier.h b/Engine/Source/Runtime/PhysicsSQ/Public/SQVerifier.h
--- a/Engine/Source/Runtime/PhysicsSQ/Public/SQVerifier.h
+++ b/Engine/Source/Runtime/PhysicsSQ/Public/SQVerifier.h
@@ -6,6 +6,8 @@
 
 #if WITH_PHYSX && INCLUDE_CHAOS
 
+#include "Chaos/PBDRigidsEvolutionGBF.h"
+
 #ifndef SQ_REPLAY_TEST
 #define SQ_REPLAY_TEST(cond) bEnsureOnMismatch ? ensure(cond) : (cond)
 #endif
@@ -43,8 +45,10 @@ bool SQComparisonHelper(FPhysTestSerializer& Serializer, bool bEnsureOnMismatch
 		}
 
 		auto ChaosHitBuffer = MakeUnique<ChaosInterface::FSQHitBuffer<ChaosInterface::FRaycastHit>>();
-		FChaosSQAccelerator SQAccelerator(*Serializer.GetChaosData());
+		auto& Accelerator = Serializer.GetChaosData()->GetCollisionConstraints().GetSpatialAcceleration();
+		FChaosSQAccelerator SQAccelerator(Accelerator);
 		SQAccelerator.Raycast(CapturedSQ.StartPoint, CapturedSQ.Dir, CapturedSQ.DeltaMag, *ChaosHitBuffer, CapturedSQ.OutputFlags.HitFlags, CapturedSQ.QueryFilterData, *CapturedSQ.FilterCallback);
+		Serializer.GetChaosData()->GetCollisionConstraints().ReleaseSpatialAcceleration();
 
 		bTestPassed &= SQ_REPLAY_TEST(ChaosHitBuffer->HasBlockingHit() == CapturedSQ.PhysXRaycastBuffer.hasBlock);
 		bTestPassed &= SQ_REPLAY_TEST(ChaosHitBuffer->GetNumHits() == CapturedSQ.PhysXRaycastBuffer.GetNumHits());
@@ -84,8 +88,10 @@ bool SQComparisonHelper(FPhysTestSerializer& Serializer, bool bEnsureOnMismatch
 		}
 
 		auto ChaosHitBuffer = MakeUnique<ChaosInterface::FSQHitBuffer<ChaosInterface::FSweepHit>>();
-		FChaosSQAccelerator SQAccelerator(*Serializer.GetChaosData());
+		auto& Accelerator = Serializer.GetChaosData()->GetCollisionConstraints().GetSpatialAcceleration();
+		FChaosSQAccelerator SQAccelerator(Accelerator);
 		SQAccelerator.Sweep(*CapturedSQ.ChaosGeometry, CapturedSQ.StartTM, CapturedSQ.Dir, CapturedSQ.DeltaMag, *ChaosHitBuffer, CapturedSQ.OutputFlags.HitFlags, CapturedSQ.QueryFilterData, *CapturedSQ.FilterCallback);
+		Serializer.GetChaosData()->GetCollisionConstraints().ReleaseSpatialAcceleration();
 
 		bTestPassed &= SQ_REPLAY_TEST(ChaosHitBuffer->HasBlockingHit() == CapturedSQ.PhysXSweepBuffer.hasBlock);
 		bTestPassed &= SQ_REPLAY_TEST(ChaosHitBuffer->GetNumHits() == CapturedSQ.PhysXSweepBuffer.GetNumHits());
@@ -111,8 +117,10 @@ bool SQComparisonHelper(FPhysTestSerializer& Serializer, bool bEnsureOnMismatch
 		bTestPassed &= SQ_REPLAY_TEST(PxHitBuffer->GetNumHits() == CapturedSQ.PhysXOverlapBuffer.GetNumHits());
 
 		auto ChaosHitBuffer = MakeUnique<ChaosInterface::FSQHitBuffer<ChaosInterface::FOverlapHit>>();
-		FChaosSQAccelerator SQAccelerator(*Serializer.GetChaosData());
+		auto& Accelerator = Serializer.GetChaosData()->GetCollisionConstraints().GetSpatialAcceleration();
+		FChaosSQAccelerator SQAccelerator(Accelerator);
 		SQAccelerator.Overlap(*CapturedSQ.ChaosGeometry, CapturedSQ.StartTM, *ChaosHitBuffer, CapturedSQ.QueryFilterData, *CapturedSQ.FilterCallback);
+		Serializer.GetChaosData()->GetCollisionConstraints().ReleaseSpatialAcceleration();
 
 		bTestPassed &= SQ_REPLAY_TEST(ChaosHitBuffer->GetNumHits() == CapturedSQ.PhysXOverlapBuffer.GetNumHits());
 		break;
@@ -143,7 +151,8 @@ void SQPerfComparisonHelper(const FString& TestName, FPhysTestSerializer& Serial
 			Timer.Stop();
 		}
 
-		FChaosSQAccelerator SQAccelerator(*Serializer.GetChaosData());
+		auto& Accelerator = Serializer.GetChaosData()->GetCollisionConstraints().GetSpatialAcceleration();
+		FChaosSQAccelerator SQAccelerator(Accelerator);
 		for (double i = 0; i < NumIterations; ++i)
 		{
 			auto ChaosHitBuffer = MakeUnique<ChaosInterface::FSQHitBuffer<ChaosInterface::FRaycastHit>>();
@@ -152,7 +161,7 @@ void SQPerfComparisonHelper(const FString& TestName, FPhysTestSerializer& Serial
 			SQAccelerator.Raycast(CapturedSQ.StartPoint, CapturedSQ.Dir, CapturedSQ.DeltaMag, *ChaosHitBuffer, CapturedSQ.OutputFlags.HitFlags, CapturedSQ.QueryFilterData, *CapturedSQ.FilterCallback);
 			Timer.Stop();
 		}
-
+		Serializer.GetChaosData()->GetCollisionConstraints().ReleaseSpatialAcceleration();
 		break;
 	}
 	case FSQCapture::ESQType::Sweep:
@@ -166,7 +175,8 @@ void SQPerfComparisonHelper(const FString& TestName, FPhysTestSerializer& Serial
 			Timer.Stop();
 		}
 
-		FChaosSQAccelerator SQAccelerator(*Serializer.GetChaosData());
+		auto& Accelerator = Serializer.GetChaosData()->GetCollisionConstraints().GetSpatialAcceleration();
+		FChaosSQAccelerator SQAccelerator(Accelerator);
 		for (double i = 0; i < NumIterations; ++i)
 		{
 			auto ChaosHitBuffer = MakeUnique<ChaosInterface::FSQHitBuffer<ChaosInterface::FSweepHit>>();
@@ -175,7 +185,7 @@ void SQPerfComparisonHelper(const FString& TestName, FPhysTestSerializer& Serial
 			SQAccelerator.Sweep(*CapturedSQ.ChaosGeometry, CapturedSQ.StartTM, CapturedSQ.Dir, CapturedSQ.DeltaMag, *ChaosHitBuffer, CapturedSQ.OutputFlags.HitFlags, CapturedSQ.QueryFilterData, *CapturedSQ.FilterCallback);
 			Timer.Stop();
 		}
-
+		Serializer.GetChaosData()->GetCollisionConstraints().ReleaseSpatialAcceleration();
 		break;
 	}
 	case FSQCapture::ESQType::Overlap:
@@ -189,7 +199,8 @@ void SQPerfComparisonHelper(const FString& TestName, FPhysTestSerializer& Serial
 			Timer.Stop();
 		}
 
-		FChaosSQAccelerator SQAccelerator(*Serializer.GetChaosData());
+		auto& Accelerator = Serializer.GetChaosData()->GetCollisionConstraints().GetSpatialAcceleration();
+		FChaosSQAccelerator SQAccelerator(Accelerator);
 		for (double i = 0; i < NumIterations; ++i)
 		{
 			auto ChaosHitBuffer = MakeUnique<ChaosInterface::FSQHitBuffer<ChaosInterface::FOverlapHit>>();
@@ -198,6 +209,7 @@ void SQPerfComparisonHelper(const FString& TestName, FPhysTestSerializer& Serial
 			SQAccelerator.Overlap(*CapturedSQ.ChaosGeometry, CapturedSQ.StartTM, *ChaosHitBuffer, CapturedSQ.QueryFilterData, *CapturedSQ.FilterCallback);
 			Timer.Stop();
 		}
+		Serializer.GetChaosData()->GetCollisionConstraints().ReleaseSpatialAcceleration();
 		break;
 	}
 	}
diff --git a/Engine/Source/Runtime/RHI/Public/RHIUtilities.h b/Engine/Source/Runtime/RHI/Public/RHIUtilities.h
--- a/Engine/Source/Runtime/RHI/Public/RHIUtilities.h
+++ b/Engine/Source/Runtime/RHI/Public/RHIUtilities.h
@@ -55,40 +55,37 @@ static inline bool IsStencilFormat(EPixelFormat Format)
 	return false;
 }
 
-/** Encapsulates a GPU read/write texture with its UAV and SRV. */
-struct FTextureRWBuffer
+/** Encapsulates a GPU read/write texture 2D with its UAV and SRV. */
+struct FTextureRWBuffer2D
 {
 	FTexture2DRHIRef Buffer;
 	FUnorderedAccessViewRHIRef UAV;
 	FShaderResourceViewRHIRef SRV;
 	uint32 NumBytes;
 
-	FTextureRWBuffer()
+	FTextureRWBuffer2D()
 		: NumBytes(0)
 	{}
 
-	~FTextureRWBuffer()
+	~FTextureRWBuffer2D()
 	{
 		Release();
 	}
 
 	// @param AdditionalUsage passed down to RHICreateVertexBuffer(), get combined with "BUF_UnorderedAccess | BUF_ShaderResource" e.g. BUF_Static
-	void Initialize(uint32 SizeX, uint32 SizeY)
+	void Initialize(const uint32 BytesPerElement, const uint32 SizeX, const uint32 SizeY, const EPixelFormat Format)
 	{
 		check(GMaxRHIFeatureLevel == ERHIFeatureLevel::SM5
 			|| IsVulkanPlatform(GMaxRHIShaderPlatform)
 			|| IsMetalPlatform(GMaxRHIShaderPlatform)
 			|| (GMaxRHIFeatureLevel == ERHIFeatureLevel::ES3_1 && GSupportsResourceView)
 		);		
-				
-		// #todo(dmp): hardcoded for now...
-		uint32 BlockBytes = 4;
 
-		NumBytes = SizeX * SizeY * BlockBytes;
+		NumBytes = SizeX * SizeY * BytesPerElement;
 
 		FRHIResourceCreateInfo CreateInfo;
 		Buffer = RHICreateTexture2D(
-			SizeX, SizeY, PF_R32_FLOAT,
+			SizeX, SizeY, Format, //PF_R32_FLOAT,
 			/*NumMips=*/ 1,
 			1,
 			/*Flags=*/ TexCreate_ShaderResource | TexCreate_UAV,
@@ -124,6 +121,69 @@ struct FTextureRWBuffer
 	}
 };
 
+/** Encapsulates a GPU read/write texture 3d with its UAV and SRV. */
+struct FTextureRWBuffer3D
+{
+	FTexture3DRHIRef Buffer;
+	FUnorderedAccessViewRHIRef UAV;
+	FShaderResourceViewRHIRef SRV;
+	uint32 NumBytes;
+
+	FTextureRWBuffer3D()
+		: NumBytes(0)
+	{}
+
+	~FTextureRWBuffer3D()
+	{
+		Release();
+	}
+
+	// @param AdditionalUsage passed down to RHICreateVertexBuffer(), get combined with "BUF_UnorderedAccess | BUF_ShaderResource" e.g. BUF_Static
+	void Initialize(uint32 BytesPerElement, uint32 SizeX, uint32 SizeY, uint32 SizeZ, EPixelFormat Format)
+	{
+		check(GMaxRHIFeatureLevel == ERHIFeatureLevel::SM5
+			|| IsVulkanPlatform(GMaxRHIShaderPlatform)
+			|| IsMetalPlatform(GMaxRHIShaderPlatform)
+			|| (GMaxRHIFeatureLevel == ERHIFeatureLevel::ES3_1 && GSupportsResourceView)
+		);
+
+		NumBytes = SizeX * SizeY * SizeZ * BytesPerElement;
+
+		FRHIResourceCreateInfo CreateInfo;
+		Buffer = RHICreateTexture3D(
+			SizeX, SizeY, SizeZ, Format,
+			/*NumMips=*/ 1,
+			/*Flags=*/ TexCreate_ShaderResource | TexCreate_UAV,
+			/*BulkData=*/ CreateInfo);
+
+		UAV = RHICreateUnorderedAccessView(Buffer, 0);
+		SRV = RHICreateShaderResourceView(Buffer, 0);
+	}
+
+	void AcquireTransientResource()
+	{
+		RHIAcquireTransientResource(Buffer);
+	}
+	void DiscardTransientResource()
+	{
+		RHIDiscardTransientResource(Buffer);
+	}
+
+	void Release()
+	{
+		int32 BufferRefCount = Buffer ? Buffer->GetRefCount() : -1;
+
+		if (BufferRefCount == 1)
+		{
+			DiscardTransientResource();
+		}
+
+		NumBytes = 0;
+		Buffer.SafeRelease();
+		UAV.SafeRelease();
+		SRV.SafeRelease();
+	}
+};
 
 /** Encapsulates a GPU read/write buffer with its UAV and SRV. */
 struct FRWBuffer
diff --git a/Engine/Source/Runtime/RenderCore/Private/ClearQuad.cpp b/Engine/Source/Runtime/RenderCore/Private/ClearQuad.cpp
--- a/Engine/Source/Runtime/RenderCore/Private/ClearQuad.cpp
+++ b/Engine/Source/Runtime/RenderCore/Private/ClearQuad.cpp
@@ -155,12 +155,6 @@ void ClearUAV(FRHICommandList& RHICmdList, const FRWBufferStructured& Structured
 	}
 }
 
-
-void ClearUAV(FRHICommandList& RHICmdList, const FTextureRWBuffer& Buffer, FLinearColor Value)
-{
-	ClearTexture2DUAV(RHICmdList, Buffer.UAV, Buffer.Buffer->GetSizeX(), Buffer.Buffer->GetSizeY(), Value);	
-}
-
 void ClearUAV(FRHICommandList& RHICmdList, const FRWBuffer& Buffer, uint32 Value, bool bBarriers)
 {
 	if (Buffer.NumBytes <= uint32(CVarFastClearUAVMaxSize.GetValueOnRenderThread()))
@@ -246,7 +240,7 @@ inline void ClearUAV_T(FRHICommandList& RHICmdList, FRHITexture* Texture, FRHIUn
 	}
 }
 
-void ClearTexture2DUAV(FRHICommandList& RHICmdList, FRHIUnorderedAccessView* UAV, int32 Width, int32 Height, const FLinearColor& ClearColor)
+void ClearUAV(FRHICommandList& RHICmdList, FRHIUnorderedAccessView* UAV, int32 Width, int32 Height, const FLinearColor& ClearColor)
 {
 	TShaderMapRef< FClearTexture2DReplacementCS<float> > ComputeShader(GetGlobalShaderMap(GMaxRHIFeatureLevel));
 	FRHIComputeShader* ShaderRHI = ComputeShader->GetComputeShader();
@@ -258,6 +252,16 @@ void ClearTexture2DUAV(FRHICommandList& RHICmdList, FRHIUnorderedAccessView* UAV
 	ComputeShader->FinalizeParameters(RHICmdList, UAV);
 }
 
+void ClearUAV(FRHICommandList& RHICmdList, const FTextureRWBuffer2D& Buffer, FLinearColor Value)
+{
+	ClearUAV_T(RHICmdList, Buffer.Buffer, Buffer.UAV, reinterpret_cast<const float(&)[4]>(Value));
+}
+
+void ClearUAV(FRHICommandList& RHICmdList, const FTextureRWBuffer3D& Buffer, FLinearColor Value)
+{
+	ClearUAV_T(RHICmdList, Buffer.Buffer, Buffer.UAV, reinterpret_cast<const float(&)[4]>(Value));
+}
+
 void ClearUAV(FRHICommandList& RHICmdList, const FSceneRenderTargetItem& RenderTargetItem, const float(&ClearValues)[4])
 {
 	ClearUAV_T(RHICmdList, RenderTargetItem.TargetableTexture, RenderTargetItem.UAV, ClearValues);
diff --git a/Engine/Source/Runtime/RenderCore/Private/RenderGraphValidation.cpp b/Engine/Source/Runtime/RenderCore/Private/RenderGraphValidation.cpp
--- a/Engine/Source/Runtime/RenderCore/Private/RenderGraphValidation.cpp
+++ b/Engine/Source/Runtime/RenderCore/Private/RenderGraphValidation.cpp
@@ -471,9 +471,9 @@ void FRDGUserValidation::ValidateAddPass(const FRDGPass* Pass, bool bSkipPassAcc
 
 			if (FRDGTextureRef Texture = DepthStencil.GetTexture())
 			{
-				ensureMsgf(Texture->Desc.NumSamples == 1,
-					TEXT("Pass '%s' uses an MSAA depth-stencil render target. This is not yet supported."),
-					Pass->GetName());
+				//ensureMsgf(Texture->Desc.NumSamples == 1,
+				//	TEXT("Pass '%s' uses an MSAA depth-stencil render target. This is not yet supported."),
+				//	Pass->GetName());
 
 				ensureMsgf(!bIsGeneratingMips,
 					TEXT("Pass '%s' is marked to generate mips but has a depth stencil texture. This is not supported."),
diff --git a/Engine/Source/Runtime/RenderCore/Public/ClearQuad.h b/Engine/Source/Runtime/RenderCore/Public/ClearQuad.h
--- a/Engine/Source/Runtime/RenderCore/Public/ClearQuad.h
+++ b/Engine/Source/Runtime/RenderCore/Public/ClearQuad.h
@@ -44,20 +44,20 @@ struct FClearQuadCallbacks
 };
 
 extern RENDERCORE_API void ClearUAV(FRHICommandList& RHICmdList, const FRWBufferStructured& StructuredBuffer, uint32 Value);
-extern RENDERCORE_API void ClearUAV(FRHICommandList& RHICmdList, const FTextureRWBuffer& Buffer, FLinearColor Value);
+extern RENDERCORE_API void ClearUAV(FRHICommandList& RHICmdList, const FTextureRWBuffer2D& Buffer, FLinearColor Value);
+extern RENDERCORE_API void ClearUAV(FRHICommandList& RHICmdList, const FTextureRWBuffer3D& Buffer, FLinearColor Value);
 extern RENDERCORE_API void ClearUAV(FRHICommandList& RHICmdList, const FRWBuffer& Buffer, uint32 Value, bool bBarriers = true);
 extern RENDERCORE_API void ClearUAV(FRHICommandList& RHICmdList, FRHIUnorderedAccessView* Buffer, uint32 NumBytes, uint32 Value);
 
 extern RENDERCORE_API void ClearUAV(FRHICommandList& RHICmdList, FRHITexture* Texture, FRHIUnorderedAccessView* TextureUAV, const float(&ClearValues)[4]);
 extern RENDERCORE_API void ClearUAV(FRHICommandList& RHICmdList, FRHITexture* Texture, FRHIUnorderedAccessView* TextureUAV, const uint32(&ClearValues)[4]);
 extern RENDERCORE_API void ClearUAV(FRHICommandList& RHICmdList, FRHITexture* Texture, FRHIUnorderedAccessView* TextureUAV, const FLinearColor& ClearColor);
+extern RENDERCORE_API void ClearUAV(FRHICommandList& RHICmdList, FRHIUnorderedAccessView* UAV, int32 Width, int32 Height, const FLinearColor& ClearColor);
 
 extern RENDERCORE_API void ClearUAV(FRHICommandList& RHICmdList, const FSceneRenderTargetItem& RenderTargetItem, const float(&ClearValues)[4]);
 extern RENDERCORE_API void ClearUAV(FRHICommandList& RHICmdList, const FSceneRenderTargetItem& RenderTargetItem, const uint32(&ClearValues)[4]);
 extern RENDERCORE_API void ClearUAV(FRHICommandList& RHICmdList, const FSceneRenderTargetItem& RenderTargetItem, const FLinearColor& ClearColor);
 
-extern RENDERCORE_API void ClearTexture2DUAV(FRHICommandList& RHICmdList, FRHIUnorderedAccessView* UAV, int32 Width, int32 Height, const FLinearColor& ClearColor);
-
 extern RENDERCORE_API void DrawClearQuadMRT( FRHICommandList& RHICmdList, bool bClearColor, int32 NumClearColors, const FLinearColor* ClearColorArray, bool bClearDepth, float Depth, bool bClearStencil, uint32 Stencil );
 extern RENDERCORE_API void DrawClearQuadMRT(FRHICommandList& RHICmdList, bool bClearColor, int32 NumClearColors, const FLinearColor* ClearColorArray, bool bClearDepth, float Depth, bool bClearStencil, uint32 Stencil, FClearQuadCallbacks ClearQuadCallbacks);
 extern RENDERCORE_API void DrawClearQuadMRT( FRHICommandList& RHICmdList, bool bClearColor, int32 NumClearColors, const FLinearColor* ClearColorArray, bool bClearDepth, float Depth, bool bClearStencil, uint32 Stencil, FIntPoint ViewSize, FIntRect ExcludeRect );
@@ -80,4 +80,4 @@ inline void DrawClearQuad(FRHICommandList& RHICmdList, const FLinearColor& Color
 inline void DrawClearQuad(FRHICommandList& RHICmdList, const FLinearColor& Color, FClearQuadCallbacks ClearQuadCallbacks)
 {
 	DrawClearQuadMRT(RHICmdList, true, 1, &Color, false, 0, false, 0, ClearQuadCallbacks);
-}
\ No newline at end of file
+}
diff --git a/Engine/Source/Runtime/RenderCore/Public/VertexFactory.h b/Engine/Source/Runtime/RenderCore/Public/VertexFactory.h
--- a/Engine/Source/Runtime/RenderCore/Public/VertexFactory.h
+++ b/Engine/Source/Runtime/RenderCore/Public/VertexFactory.h
@@ -556,8 +556,14 @@ public:
 		return PrimitiveIdStreamIndex[static_cast<uint8>(InputStreamType)];
 	}
 
+
 protected:
 
+	inline void SetPrimitiveIdStreamIndex(EVertexInputStreamType InputStreamType, int32 StreamIndex)
+	{
+		PrimitiveIdStreamIndex[static_cast<uint8>(InputStreamType)] = StreamIndex;
+	}
+
 	/**
 	 * Creates a vertex element for a vertex stream components.  Adds a unique stream index for the vertex buffer used by the component.
 	 * @param Component - The vertex stream component.
diff --git a/Engine/Source/Runtime/Renderer/Private/DeferredShadingRenderer.cpp b/Engine/Source/Runtime/Renderer/Private/DeferredShadingRenderer.cpp
--- a/Engine/Source/Runtime/Renderer/Private/DeferredShadingRenderer.cpp
+++ b/Engine/Source/Runtime/Renderer/Private/DeferredShadingRenderer.cpp
@@ -39,6 +39,7 @@
 #include "RayTracingDefinitions.h"
 #include "RayTracingInstance.h"
 #include "ShaderPrint.h"
+#include "HairStrands/HairStrandsRendering.h"
 
 static TAutoConsoleVariable<int32> CVarStencilForLODDither(
 	TEXT("r.StencilForLODDither"),
@@ -1140,6 +1141,11 @@ void FDeferredShadingSceneRenderer::Render(FRHICommandListImmediate& RHICmdList)
 		Scene->FXSystem->PreRender(RHICmdList, &Views[0].GlobalDistanceFieldInfo.ParameterData, Views[0].AllowGPUParticleUpdate());
 	}
 
+	if (IsHairStrandsEnable(Scene->GetShaderPlatform()))
+	{
+		RunHairStrandsInterpolation(RHICmdList);
+	}
+
 	bool bDidAfterTaskWork = false;
 	auto AfterTasksAreStarted = [&bDidAfterTaskWork, bDoInitViewAftersPrepass, this, &RHICmdList, &ILCTaskData, &UpdateViewCustomDataEvents]()
 	{
@@ -1696,6 +1702,19 @@ void FDeferredShadingSceneRenderer::Render(FRHICommandListImmediate& RHICmdList)
 
 	checkSlow(RHICmdList.IsOutsideRenderPass());
 
+	FHairStrandsDatas* HairDatas = nullptr;
+	FHairStrandsDatas HairDatasStorage;
+	if (IsHairStrandsEnable(Scene->GetShaderPlatform()))
+	{
+		HairDatasStorage.HairClusterPerViews = CreateHairStrandsClusters(RHICmdList, Scene, Views);
+		VoxelizeHairStrands(RHICmdList, Scene, Views, HairDatasStorage.HairClusterPerViews);
+		HairDatasStorage.DeepShadowViews = RenderHairStrandsDeepShadows(RHICmdList, Scene, Views, HairDatasStorage.HairClusterPerViews);
+		HairDatasStorage.HairVisibilityViews = RenderHairStrandsVisibilityBuffer(RHICmdList, Scene, Views, SceneContext.GBufferB, SceneContext.GetSceneColor(), SceneContext.SceneDepthZ, SceneContext.SceneVelocity, HairDatasStorage.HairClusterPerViews);
+		ServiceLocalQueue();
+
+		HairDatas = &HairDatasStorage;
+	}
+
 	// Render lighting.
 	if (bRenderDeferredLighting)
 	{
@@ -1728,7 +1747,9 @@ void FDeferredShadingSceneRenderer::Render(FRHICommandListImmediate& RHICmdList)
 		}
 
 		RHICmdList.SetCurrentStat(GET_STATID(STAT_CLM_Lighting));
-		RenderLights(RHICmdList, SortedLightSet);
+		{
+			RenderLights(RHICmdList, SortedLightSet, HairDatas);
+		}
 		RHICmdList.SetCurrentStat(GET_STATID(STAT_CLM_AfterLighting));
 		ServiceLocalQueue();
 
@@ -1769,7 +1790,7 @@ void FDeferredShadingSceneRenderer::Render(FRHICommandListImmediate& RHICmdList)
 		checkSlow(RHICmdList.IsOutsideRenderPass());
 
 		// Render diffuse sky lighting and reflections that only operate on opaque pixels
-		RenderDeferredReflectionsAndSkyLighting(RHICmdList, DynamicBentNormalAO, SceneContext.SceneVelocity);
+		RenderDeferredReflectionsAndSkyLighting(RHICmdList, DynamicBentNormalAO, SceneContext.SceneVelocity, HairDatas);
 
 		DynamicBentNormalAO = NULL;
 
@@ -2040,6 +2061,12 @@ void FDeferredShadingSceneRenderer::Render(FRHICommandListImmediate& RHICmdList)
 
 	}
 
+	if (HairDatas)
+	{
+		RenderHairComposeSubPixel(RHICmdList, Views, HairDatas);
+		RenderHairStrandsDebugInfo(RHICmdList, Views, HairDatas);
+	}
+
 	checkSlow(RHICmdList.IsOutsideRenderPass());
 
 	if (bCanOverlayRayTracingOutput && ViewFamily.EngineShowFlags.LightShafts)
@@ -2132,7 +2159,7 @@ void FDeferredShadingSceneRenderer::Render(FRHICommandListImmediate& RHICmdList)
 			}
 			else
 			{
-				GPostProcessing.Process(RHICmdList, Views[ ViewIndex ], SceneContext.SceneVelocity);
+				GPostProcessing.Process(RHICmdList, Views[ViewIndex], SceneContext.SceneVelocity);
 			}
 		}
 
