Try Andrew Rodham's suggestion on fixing the incremental Mac link errors
#rnx

diff --git a/Engine/Source/Runtime/Experimental/Chaos/Public/Chaos/KinematicGeometryParticles.h b/Engine/Source/Runtime/Experimental/Chaos/Public/Chaos/KinematicGeometryParticles.h
--- a/Engine/Source/Runtime/Experimental/Chaos/Public/Chaos/KinematicGeometryParticles.h
+++ b/Engine/Source/Runtime/Experimental/Chaos/Public/Chaos/KinematicGeometryParticles.h
@@ -63,8 +63,8 @@ FChaosArchive& operator<<(FChaosArchive& Ar, TKinematicGeometryParticlesImp<T, d
 	return Ar;
 }
 
-extern template class Chaos::TKinematicGeometryParticlesImp<float, 3, Chaos::EGeometryParticlesSimType::RigidBodySim>;
-extern template class Chaos::TKinematicGeometryParticlesImp<float, 3, Chaos::EGeometryParticlesSimType::Other>;
+extern template class CHAOS_API Chaos::TKinematicGeometryParticlesImp<float, 3, Chaos::EGeometryParticlesSimType::RigidBodySim>;
+extern template class CHAOS_API Chaos::TKinematicGeometryParticlesImp<float, 3, Chaos::EGeometryParticlesSimType::Other>;
 
 template <typename T, int d>
 using TKinematicGeometryParticles = TKinematicGeometryParticlesImp<T, d, EGeometryParticlesSimType::RigidBodySim>;
