Only export the externs for clang since MSVC does not approve
#rnx

diff --git a/Engine/Source/Runtime/Experimental/Chaos/Public/Chaos/KinematicGeometryParticles.h b/Engine/Source/Runtime/Experimental/Chaos/Public/Chaos/KinematicGeometryParticles.h
--- a/Engine/Source/Runtime/Experimental/Chaos/Public/Chaos/KinematicGeometryParticles.h
+++ b/Engine/Source/Runtime/Experimental/Chaos/Public/Chaos/KinematicGeometryParticles.h
@@ -63,8 +63,13 @@ FChaosArchive& operator<<(FChaosArchive& Ar, TKinematicGeometryParticlesImp<T, d
 	return Ar;
 }
 
+#ifdef __clang__
 extern template class CHAOS_API Chaos::TKinematicGeometryParticlesImp<float, 3, Chaos::EGeometryParticlesSimType::RigidBodySim>;
 extern template class CHAOS_API Chaos::TKinematicGeometryParticlesImp<float, 3, Chaos::EGeometryParticlesSimType::Other>;
+#else
+extern template class Chaos::TKinematicGeometryParticlesImp<float, 3, Chaos::EGeometryParticlesSimType::RigidBodySim>;
+extern template class Chaos::TKinematicGeometryParticlesImp<float, 3, Chaos::EGeometryParticlesSimType::Other>;
+#endif
 
 template <typename T, int d>
 using TKinematicGeometryParticles = TKinematicGeometryParticlesImp<T, d, EGeometryParticlesSimType::RigidBodySim>;
