Export TRigidParticles in a way that makes incremental mac happy
#rnx

diff --git a/Engine/Source/Runtime/Experimental/Chaos/Private/Chaos/RigidParticles.cpp b/Engine/Source/Runtime/Experimental/Chaos/Private/Chaos/RigidParticles.cpp
--- a/Engine/Source/Runtime/Experimental/Chaos/Private/Chaos/RigidParticles.cpp
+++ b/Engine/Source/Runtime/Experimental/Chaos/Private/Chaos/RigidParticles.cpp
@@ -20,4 +20,4 @@ void TRigidParticles<T, d>::SetCollisionParticles(const int32 Index, TParticles<
 	MCollisionParticles[Index] = MakeUnique<TBVHParticles<T, d>>(MoveTemp(Points));
 }
 
-template class Chaos::TRigidParticles<float, 3>;
\ No newline at end of file
+template class CHAOS_API Chaos::TRigidParticles<float, 3>;
\ No newline at end of file
diff --git a/Engine/Source/Runtime/Experimental/Chaos/Public/Chaos/RigidParticles.h b/Engine/Source/Runtime/Experimental/Chaos/Public/Chaos/RigidParticles.h
--- a/Engine/Source/Runtime/Experimental/Chaos/Public/Chaos/RigidParticles.h
+++ b/Engine/Source/Runtime/Experimental/Chaos/Public/Chaos/RigidParticles.h
@@ -188,6 +188,9 @@ FChaosArchive& operator<<(FChaosArchive& Ar, TRigidParticles<T, d>& Particles)
 	return Ar;
 }
 
+#ifdef __clang__
+extern template class CHAOS_API TRigidParticles<float, 3>;
+#else
 extern template class TRigidParticles<float, 3>;
-
+#endif
 }
