Don't export symbol for MSVC
#rnx

diff --git a/Engine/Source/Runtime/Experimental/Chaos/Private/Chaos/RigidParticles.cpp b/Engine/Source/Runtime/Experimental/Chaos/Private/Chaos/RigidParticles.cpp
--- a/Engine/Source/Runtime/Experimental/Chaos/Private/Chaos/RigidParticles.cpp
+++ b/Engine/Source/Runtime/Experimental/Chaos/Private/Chaos/RigidParticles.cpp
@@ -20,4 +20,8 @@ void TRigidParticles<T, d>::SetCollisionParticles(const int32 Index, TParticles<
 	MCollisionParticles[Index] = MakeUnique<TBVHParticles<T, d>>(MoveTemp(Points));
 }
 
-template class CHAOS_API Chaos::TRigidParticles<float, 3>;
\ No newline at end of file
+#ifdef __clang__
+template class CHAOS_API Chaos::TRigidParticles<float, 3>;
+#else
+template class Chaos::TRigidParticles<float, 3>;
+#endif
\ No newline at end of file
