Implement TAABB (POD Box type), and make TBox use TAABB.

diff --git a/Engine/Source/Runtime/Experimental/Chaos/Private/Chaos/AABB.cpp b/Engine/Source/Runtime/Experimental/Chaos/Private/Chaos/AABB.cpp
new file mode 100644
--- /dev/null
+++ b/Engine/Source/Runtime/Experimental/Chaos/Private/Chaos/AABB.cpp
@@ -0,0 +1,179 @@
+// Copyright 1998-2019 Epic Games, Inc. All Rights Reserved.
+#include "Chaos/AABB.h"
+#include "Chaos/Sphere.h"
+#include "Chaos/Capsule.h"
+
+namespace Chaos
+{
+template <typename T, int d>
+bool TAABB<T, d>::Raycast(const TVector<T, d>& StartPoint, const TVector<T, d>& Dir, const T Length, const T Thickness, T& OutTime, TVector<T, d>& OutPosition, TVector<T, d>& OutNormal, int32& OutFaceIndex) const
+{
+	ensure(Length > 0);
+	ensure(FMath::IsNearlyEqual(Dir.SizeSquared(), 1, KINDA_SMALL_NUMBER));
+
+	OutFaceIndex = INDEX_NONE;
+	const TVector<T, d> MinInflated = MMin - Thickness;
+	const TVector<T,d> StartToMin = MinInflated - StartPoint;
+	
+	const TVector<T, d> MaxInflated = MMax + Thickness;
+	const TVector<T, d> StartToMax = MaxInflated - StartPoint;
+
+	//For each axis record the start and end time when ray is in the box. If the intervals overlap the ray is inside the box
+	T LatestStartTime = 0;
+	T EarliestEndTime = FLT_MAX;
+	TVector<T, d> Normal(0);	//not needed but fixes compiler warning
+
+	for (int Axis = 0; Axis < d; ++Axis)
+	{
+		const bool bParallel = FMath::IsNearlyZero(Dir[Axis]);
+		T Time1, Time2;
+		if (bParallel)
+		{
+			if (StartToMin[Axis] > 0 || StartToMax[Axis] < 0)
+			{
+				return false;	//parallel and outside
+			}
+			else
+			{
+				Time1 = 0;
+				Time2 = FLT_MAX;
+			}
+		}
+		else
+		{
+			const T InvDir = (T)1 / Dir[Axis];
+			Time1 = StartToMin[Axis] * InvDir;
+			Time2 = StartToMax[Axis] * InvDir;
+		}
+
+		TVector<T, d> CurNormal = TVector<T, d>::AxisVector(Axis);
+
+		if (Time1 > Time2)
+		{
+			//going from max to min direction
+			std::swap(Time1, Time2);
+		}
+		else
+		{
+			//hit negative plane first
+			CurNormal[Axis] = -1;
+		}
+
+		if (Time1 > LatestStartTime)
+		{
+			//last plane to enter so save its normal
+			Normal = CurNormal;
+		}
+		LatestStartTime = FMath::Max(LatestStartTime, Time1);
+		EarliestEndTime = FMath::Min(EarliestEndTime, Time2);
+
+		if (LatestStartTime > EarliestEndTime)
+		{
+			return false;	//Outside of slab before entering another
+		}
+	}
+
+	//infinite ray intersects with inflated box
+	if (LatestStartTime > Length || EarliestEndTime < 0)
+	{
+		//outside of line segment given
+		return false;
+	}
+	
+	const TVector<T, d> BoxIntersection = StartPoint + LatestStartTime * Dir;
+
+	//If the box is rounded we have to consider corners and edges.
+	//Break the box into voronoi regions based on features (corner, edge, face) and see which region the raycast hit
+
+	if (Thickness)
+	{
+		check(d == 3);
+		TVector<T, d> GeomStart;
+		TVector<T, d> GeomEnd;
+		int32 NumAxes = 0;
+
+		for (int Axis = 0; Axis < d; ++Axis)
+		{
+			if (BoxIntersection[Axis] < MMin[Axis])
+			{
+				GeomStart[Axis] = MMin[Axis];
+				GeomEnd[Axis] = MMin[Axis];
+				++NumAxes;
+			}
+			else if (BoxIntersection[Axis] > MMax[Axis])
+			{
+				GeomStart[Axis] = MMax[Axis];
+				GeomEnd[Axis] = MMax[Axis];
+				++NumAxes;
+			}
+			else
+			{
+				GeomStart[Axis] = MMin[Axis];
+				GeomEnd[Axis] = MMax[Axis];
+			}
+		}
+
+		if (NumAxes >= 2)
+		{
+			bool bHit = false;
+			if (NumAxes == 3)
+			{
+				//hit a corner. For now just use 3 capsules, there's likely a better way to determine which capsule is needed
+				T CornerTimes[3];
+				TVector<T, d> CornerPositions[3];
+				TVector<T, d> CornerNormals[3];
+				int32 HitIdx = INDEX_NONE;
+				T MinTime = 0;	//initialization just here for compiler warning
+				for (int CurIdx = 0; CurIdx < 3; ++CurIdx)
+				{
+					TVector<T, d> End = GeomStart;
+					End[CurIdx] = End[CurIdx] == MMin[CurIdx] ? MMax[CurIdx] : MMin[CurIdx];
+					TCapsule<T> Capsule(GeomStart, End, Thickness);
+					if (Capsule.Raycast(StartPoint, Dir, Length, 0, CornerTimes[CurIdx], CornerPositions[CurIdx], CornerNormals[CurIdx], OutFaceIndex))
+					{
+						if (HitIdx == INDEX_NONE || CornerTimes[CurIdx] < MinTime)
+						{
+							MinTime = CornerTimes[CurIdx];
+							HitIdx = CurIdx;
+
+							if (MinTime == 0)
+							{
+								OutTime = 0;	//initial overlap so just exit
+								return true;
+							}
+						}
+					}
+				}
+
+				if (HitIdx != INDEX_NONE)
+				{
+					OutPosition = CornerPositions[HitIdx];
+					OutTime = MinTime;
+					OutNormal = CornerNormals[HitIdx];
+					bHit = true;
+				}
+			}
+			else
+			{
+				//capsule: todo(use a cylinder which is cheaper. Our current cylinder raycast implementation doesn't quite work for this setup)
+				TCapsule<T> CapsuleBorder(GeomStart, GeomEnd, Thickness);
+				bHit = CapsuleBorder.Raycast(StartPoint, Dir, Length, 0, OutTime, OutPosition, OutNormal, OutFaceIndex);
+			}
+
+			if (bHit && OutTime > 0)
+			{
+				OutPosition -= OutNormal * Thickness;
+			}
+			return bHit;
+		}
+	}
+	
+	// didn't hit any rounded parts so just use the box intersection
+	OutTime = LatestStartTime;
+	OutNormal = Normal;
+	OutPosition = BoxIntersection - Thickness * Normal;
+	return true;
+}
+}
+
+template class Chaos::TAABB<float, 3>;
diff --git a/Engine/Source/Runtime/Experimental/Chaos/Private/Chaos/Box.cpp b/Engine/Source/Runtime/Experimental/Chaos/Private/Chaos/Box.cpp
--- a/Engine/Source/Runtime/Experimental/Chaos/Private/Chaos/Box.cpp
+++ b/Engine/Source/Runtime/Experimental/Chaos/Private/Chaos/Box.cpp
@@ -1,179 +0,0 @@
-// Copyright 1998-2019 Epic Games, Inc. All Rights Reserved.
-#include "Chaos/Box.h"
-#include "Chaos/Sphere.h"
-#include "Chaos/Capsule.h"
-
-namespace Chaos
-{
-template <typename T, int d>
-bool TBox<T, d>::Raycast(const TVector<T, d>& StartPoint, const TVector<T, d>& Dir, const T Length, const T Thickness, T& OutTime, TVector<T, d>& OutPosition, TVector<T, d>& OutNormal, int32& OutFaceIndex) const
-{
-	ensure(Length > 0);
-	ensure(FMath::IsNearlyEqual(Dir.SizeSquared(), 1, KINDA_SMALL_NUMBER));
-
-	OutFaceIndex = INDEX_NONE;
-	const TVector<T, d> MinInflated = MMin - Thickness;
-	const TVector<T,d> StartToMin = MinInflated - StartPoint;
-	
-	const TVector<T, d> MaxInflated = MMax + Thickness;
-	const TVector<T, d> StartToMax = MaxInflated - StartPoint;
-
-	//For each axis record the start and end time when ray is in the box. If the intervals overlap the ray is inside the box
-	T LatestStartTime = 0;
-	T EarliestEndTime = FLT_MAX;
-	TVector<T, d> Normal(0);	//not needed but fixes compiler warning
-
-	for (int Axis = 0; Axis < d; ++Axis)
-	{
-		const bool bParallel = FMath::IsNearlyZero(Dir[Axis]);
-		T Time1, Time2;
-		if (bParallel)
-		{
-			if (StartToMin[Axis] > 0 || StartToMax[Axis] < 0)
-			{
-				return false;	//parallel and outside
-			}
-			else
-			{
-				Time1 = 0;
-				Time2 = FLT_MAX;
-			}
-		}
-		else
-		{
-			const T InvDir = (T)1 / Dir[Axis];
-			Time1 = StartToMin[Axis] * InvDir;
-			Time2 = StartToMax[Axis] * InvDir;
-		}
-
-		TVector<T, d> CurNormal = TVector<T, d>::AxisVector(Axis);
-
-		if (Time1 > Time2)
-		{
-			//going from max to min direction
-			std::swap(Time1, Time2);
-		}
-		else
-		{
-			//hit negative plane first
-			CurNormal[Axis] = -1;
-		}
-
-		if (Time1 > LatestStartTime)
-		{
-			//last plane to enter so save its normal
-			Normal = CurNormal;
-		}
-		LatestStartTime = FMath::Max(LatestStartTime, Time1);
-		EarliestEndTime = FMath::Min(EarliestEndTime, Time2);
-
-		if (LatestStartTime > EarliestEndTime)
-		{
-			return false;	//Outside of slab before entering another
-		}
-	}
-
-	//infinite ray intersects with inflated box
-	if (LatestStartTime > Length || EarliestEndTime < 0)
-	{
-		//outside of line segment given
-		return false;
-	}
-	
-	const TVector<T, d> BoxIntersection = StartPoint + LatestStartTime * Dir;
-
-	//If the box is rounded we have to consider corners and edges.
-	//Break the box into voronoi regions based on features (corner, edge, face) and see which region the raycast hit
-
-	if (Thickness)
-	{
-		check(d == 3);
-		TVector<T, d> GeomStart;
-		TVector<T, d> GeomEnd;
-		int32 NumAxes = 0;
-
-		for (int Axis = 0; Axis < d; ++Axis)
-		{
-			if (BoxIntersection[Axis] < MMin[Axis])
-			{
-				GeomStart[Axis] = MMin[Axis];
-				GeomEnd[Axis] = MMin[Axis];
-				++NumAxes;
-			}
-			else if (BoxIntersection[Axis] > MMax[Axis])
-			{
-				GeomStart[Axis] = MMax[Axis];
-				GeomEnd[Axis] = MMax[Axis];
-				++NumAxes;
-			}
-			else
-			{
-				GeomStart[Axis] = MMin[Axis];
-				GeomEnd[Axis] = MMax[Axis];
-			}
-		}
-
-		if (NumAxes >= 2)
-		{
-			bool bHit = false;
-			if (NumAxes == 3)
-			{
-				//hit a corner. For now just use 3 capsules, there's likely a better way to determine which capsule is needed
-				T CornerTimes[3];
-				TVector<T, d> CornerPositions[3];
-				TVector<T, d> CornerNormals[3];
-				int32 HitIdx = INDEX_NONE;
-				T MinTime = 0;	//initialization just here for compiler warning
-				for (int CurIdx = 0; CurIdx < 3; ++CurIdx)
-				{
-					TVector<T, d> End = GeomStart;
-					End[CurIdx] = End[CurIdx] == MMin[CurIdx] ? MMax[CurIdx] : MMin[CurIdx];
-					TCapsule<T> Capsule(GeomStart, End, Thickness);
-					if (Capsule.Raycast(StartPoint, Dir, Length, 0, CornerTimes[CurIdx], CornerPositions[CurIdx], CornerNormals[CurIdx], OutFaceIndex))
-					{
-						if (HitIdx == INDEX_NONE || CornerTimes[CurIdx] < MinTime)
-						{
-							MinTime = CornerTimes[CurIdx];
-							HitIdx = CurIdx;
-
-							if (MinTime == 0)
-							{
-								OutTime = 0;	//initial overlap so just exit
-								return true;
-							}
-						}
-					}
-				}
-
-				if (HitIdx != INDEX_NONE)
-				{
-					OutPosition = CornerPositions[HitIdx];
-					OutTime = MinTime;
-					OutNormal = CornerNormals[HitIdx];
-					bHit = true;
-				}
-			}
-			else
-			{
-				//capsule: todo(use a cylinder which is cheaper. Our current cylinder raycast implementation doesn't quite work for this setup)
-				TCapsule<T> CapsuleBorder(GeomStart, GeomEnd, Thickness);
-				bHit = CapsuleBorder.Raycast(StartPoint, Dir, Length, 0, OutTime, OutPosition, OutNormal, OutFaceIndex);
-			}
-
-			if (bHit && OutTime > 0)
-			{
-				OutPosition -= OutNormal * Thickness;
-			}
-			return bHit;
-		}
-	}
-	
-	// didn't hit any rounded parts so just use the box intersection
-	OutTime = LatestStartTime;
-	OutNormal = Normal;
-	OutPosition = BoxIntersection - Thickness * Normal;
-	return true;
-}
-}
-
-template class Chaos::TBox<float,3>;
\ No newline at end of file
diff --git a/Engine/Source/Runtime/Experimental/Chaos/Public/Chaos/AABB.h b/Engine/Source/Runtime/Experimental/Chaos/Public/Chaos/AABB.h
new file mode 100644
--- /dev/null
+++ b/Engine/Source/Runtime/Experimental/Chaos/Public/Chaos/AABB.h
@@ -0,0 +1,560 @@
+// Copyright 1998-2019 Epic Games, Inc. All Rights Reserved.
+#pragma once
+
+#include "Chaos/Vector.h"
+#include "Chaos/Matrix.h"
+#include "Chaos/Plane.h"
+#include "Chaos/Rotation.h"
+#include "ChaosArchive.h"
+
+namespace Chaos
+{ 
+	template<class T, int d>
+	class TAABB;
+
+	template<typename T, int d>
+	struct TAABBSpecializeSamplingHelper
+	{
+		static FORCEINLINE TArray<TVector<T, d>> ComputeLocalSamplePoints(const class TAABB<T, d>& AABB)
+		{
+			check(false);
+			return TArray<TVector<T, d>>();
+		}
+	};
+
+
+	template<class T, int d>
+	class TAABB
+	{
+	public:
+		FORCEINLINE TAABB()
+			: MMin()
+			, MMax()
+		{
+		}
+
+		FORCEINLINE TAABB(const TVector<T, d>& Min, const TVector<T, d>&Max)
+			: MMin(Min)
+			, MMax(Max)
+		{
+		}
+
+		FORCEINLINE TAABB(const TAABB<T, d>& Other)
+			: MMin(Other.MMin)
+			, MMax(Other.MMax)
+		{
+		}
+
+		FORCEINLINE TAABB(TAABB<T, d>&& Other)
+			: MMin(MoveTemp(Other.MMin))
+			, MMax(MoveTemp(Other.MMax))
+		{
+		}
+
+		FORCEINLINE TAABB<T, d>& operator=(const TAABB<T, d>& Other)
+		{
+			MMin = Other.MMin;
+			MMax = Other.MMax;
+			return *this;
+		}
+
+		FORCEINLINE TAABB<T, d>& operator=(TAABB<T, d>&& Other)
+		{
+			MMin = MoveTemp(Other.MMin);
+			MMax = MoveTemp(Other.MMax);
+			return *this;
+		}
+
+		/**
+		 * Returns sample points centered about the origin.
+		 */
+		FORCEINLINE TArray<TVector<T, d>> ComputeLocalSamplePoints() const
+		{
+			const TVector<T, d> Mid = Center();
+			return TAABBSpecializeSamplingHelper<T, d>::ComputeSamplePoints(TAABB<T, d>(Min() - Mid, Max() - Mid));
+		}
+
+		/**
+		 * Returns sample points at the current location of the box.
+		 */
+		FORCEINLINE TArray<TVector<T, d>> ComputeSamplePoints() const
+		{
+			return TAABBSpecializeSamplingHelper<T, d>::ComputeSamplePoints(*this);
+		}
+
+		template<class TTRANSFORM>
+		FORCEINLINE TAABB<T, d> TransformedAABB(const TTRANSFORM& SpaceTransform) const
+		{
+			TVector<T, d> CurrentExtents = Extents();
+			int32 Idx = 0;
+			const TVector<T, d> MinToNewSpace = SpaceTransform.TransformPosition(MMin);
+			TAABB<T, d> NewAABB(MinToNewSpace, MinToNewSpace);
+			NewAABB.GrowToInclude(SpaceTransform.TransformPosition(MMax));
+
+			for (int32 j = 0; j < d; ++j)
+			{
+				NewAABB.GrowToInclude(SpaceTransform.TransformPosition(MMin + TVector<T, d>::AxisVector(j) * CurrentExtents));
+				NewAABB.GrowToInclude(SpaceTransform.TransformPosition(MMax - TVector<T, d>::AxisVector(j) * CurrentExtents));
+			}
+
+			return NewAABB;
+		}
+
+		FORCEINLINE bool Intersects(const TAABB<T, d>& Other) const
+		{
+			for (int32 i = 0; i < d; ++i)
+			{
+				if (Other.MMax[i] < MMin[i] || Other.MMin[i] > MMax[i])
+					return false;
+			}
+			return true;
+		}
+
+		FORCEINLINE TAABB<T, d> GetIntersection(const TAABB<T, d>& Other) const
+		{
+			TVector<float, 3> Tmp;
+			return TAABB<T, d>(MMin.ComponentwiseMax(Other.MMin), MMax.ComponentwiseMin(Other.MMax));
+		}
+
+		FORCEINLINE bool Contains(const TVector<T, d>& Point) const
+		{
+			for (int i = 0; i < d; i++)
+			{
+				if (Point[i] < MMin[i] || Point[i] > MMax[i])
+				{
+					return false;
+				}
+			}
+			return true;
+		}
+
+		FORCEINLINE bool Contains(const TVector<T, d>& Point, const T Tolerance) const
+		{
+			for (int i = 0; i < d; i++)
+			{
+				if (Point[i] < MMin[i] - Tolerance || Point[i] > MMax[i] + Tolerance)
+				{
+					return false;
+				}
+			}
+			return true;
+		}
+
+		FORCEINLINE T SignedDistance(const TVector<T, d>& x) const
+		{
+			TVector<T, d> Normal;
+			return PhiWithNormal(x, Normal);
+		}
+
+
+		FORCEINLINE T PhiWithNormal(const TVector<T, d>& x, TVector<T, d>& Normal) const 
+		{
+			const TVector<T, d> MaxDists = x - MMax;
+			const TVector<T, d> MinDists = MMin - x;
+			if (x <= MMax && x >= MMin)
+			{
+				const Pair<T, int32> MaxAndAxis = TVector<T, d>::MaxAndAxis(MinDists, MaxDists);
+				Normal = MaxDists[MaxAndAxis.Second] > MinDists[MaxAndAxis.Second] ? TVector<T, d>::AxisVector(MaxAndAxis.Second) : -TVector<T, d>::AxisVector(MaxAndAxis.Second);
+				return MaxAndAxis.First;
+			}
+			else
+			{
+				for (int i = 0; i < d; ++i)
+				{
+					if (MaxDists[i] > 0)
+					{
+						Normal[i] = MaxDists[i];
+					}
+					else if (MinDists[i] > 0)
+					{
+						Normal[i] = -MinDists[i];
+					}
+					else
+					{
+						Normal[i] = 0;
+					}
+				}
+				T Phi = Normal.SafeNormalize();
+				if (Phi < KINDA_SMALL_NUMBER)
+				{
+					for (int i = 0; i < d; ++i)
+					{
+						if (Normal[i] > 0)
+						{
+							Normal[i] = 1;
+						}
+						else if (Normal[i] < 0)
+						{
+							Normal[i] = -1;
+						}
+					}
+					Normal.Normalize();
+				}
+				return Phi;
+			}
+		}
+
+		bool CHAOS_API Raycast(const TVector<T, d>& StartPoint, const TVector<T, d>& Dir, const T Length, const T Thickness, T& OutTime, TVector<T, d>& OutPosition, TVector<T, d>& OutNormal, int32& OutFaceIndex) const;
+
+
+		FORCEINLINE bool RaycastFast(const TVector<T, d>& StartPoint, const TVector<T, d>& Dir, const TVector<T, d>& InvDir, const bool* bParallel, const T Length, const T InvLength, T& OutTime, TVector<T, d>& OutPosition)
+		{
+			const TVector<T, d> StartToMin = MMin - StartPoint;
+			const TVector<T, d> StartToMax = MMax - StartPoint;
+
+			//For each axis record the start and end time when ray is in the box. If the intervals overlap the ray is inside the box
+			T LatestStartTime = 0;
+			T EarliestEndTime = FLT_MAX;
+
+			for (int Axis = 0; Axis < d; ++Axis)
+			{
+				T Time1, Time2;
+				if (bParallel[Axis])
+				{
+					if (StartToMin[Axis] > 0 || StartToMax[Axis] < 0)
+					{
+						return false;	//parallel and outside
+					}
+					else
+					{
+						Time1 = 0;
+						Time2 = FLT_MAX;
+					}
+				}
+				else
+				{
+					Time1 = StartToMin[Axis] * InvDir[Axis];
+					Time2 = StartToMax[Axis] * InvDir[Axis];
+				}
+
+				if (Time1 > Time2)
+				{
+					//going from max to min direction
+					Swap(Time1, Time2);
+				}
+
+				LatestStartTime = FMath::Max(LatestStartTime, Time1);
+				EarliestEndTime = FMath::Min(EarliestEndTime, Time2);
+
+				if (LatestStartTime > EarliestEndTime)
+				{
+					return false;	//Outside of slab before entering another
+				}
+			}
+
+			//infinite ray intersects with inflated box
+			if (LatestStartTime > Length || EarliestEndTime < 0)
+			{
+				//outside of line segment given
+				return false;
+			}
+
+			OutTime = LatestStartTime;
+			OutPosition = StartPoint + LatestStartTime * Dir;
+			return true;
+		}
+
+
+		FORCEINLINE TVector<T, d> FindClosestPoint(const TVector<T, d>& StartPoint, const T Thickness = (T)0) const
+		{
+			TVector<T, d> Result(0);
+
+			// clamp exterior to surface
+			bool bIsExterior = false;
+			for (int i = 0; i < 3; i++)
+			{
+				float v = StartPoint[i];
+				if (v < MMin[i])
+				{
+					v = MMin[i];
+					bIsExterior = true;
+				}
+				if (v > MMax[i])
+				{
+					v = MMax[i];
+					bIsExterior = true;
+				}
+				Result[i] = v;
+			}
+
+			if (!bIsExterior)
+			{
+				TArray<Pair<T, TVector<T, d>>> Intersections;
+
+				// sum interior direction to surface
+				for (int32 i = 0; i < d; ++i)
+				{
+					auto PlaneIntersection = TPlane<T, d>(MMin - Thickness, -TVector<T, d>::AxisVector(i)).FindClosestPoint(Result, 0);
+					Intersections.Add(MakePair((PlaneIntersection - Result).Size(), -TVector<T, d>::AxisVector(i)));
+					PlaneIntersection = TPlane<T, d>(MMax + Thickness, TVector<T, d>::AxisVector(i)).FindClosestPoint(Result, 0);
+					Intersections.Add(MakePair((PlaneIntersection - Result).Size(), TVector<T, d>::AxisVector(i)));
+				}
+				Intersections.Sort([](const Pair<T, TVector<T, d>>& Elem1, const Pair<T, TVector<T, d>>& Elem2) { return Elem1.First < Elem2.First; });
+
+				if (!FMath::IsNearlyEqual(Intersections[0].First, 0.f))
+				{
+					T SmallestDistance = Intersections[0].First;
+					Result += Intersections[0].Second * Intersections[0].First;
+					for (int32 i = 1; i < 3 && FMath::IsNearlyEqual(SmallestDistance, Intersections[i].First); ++i)
+					{
+						Result += Intersections[i].Second * Intersections[i].First;
+					}
+				}
+			}
+			return Result;
+		}
+
+		FORCEINLINE Pair<TVector<T, d>, bool> FindClosestIntersectionImp(const TVector<T, d>& StartPoint, const TVector<T, d>& EndPoint, const T Thickness) const 
+		{
+			TArray<Pair<T, TVector<T, d>>> Intersections;
+			for (int32 i = 0; i < d; ++i)
+			{
+				auto PlaneIntersection = TPlane<T, d>(MMin - Thickness, -TVector<T, d>::AxisVector(i)).FindClosestIntersection(StartPoint, EndPoint, 0);
+				if (PlaneIntersection.Second)
+					Intersections.Add(MakePair((PlaneIntersection.First - StartPoint).Size(), PlaneIntersection.First));
+				PlaneIntersection = TPlane<T, d>(MMax + Thickness, TVector<T, d>::AxisVector(i)).FindClosestIntersection(StartPoint, EndPoint, 0);
+				if (PlaneIntersection.Second)
+					Intersections.Add(MakePair((PlaneIntersection.First - StartPoint).Size(), PlaneIntersection.First));
+			}
+			Intersections.Sort([](const Pair<T, TVector<T, d>>& Elem1, const Pair<T, TVector<T, d>>& Elem2) { return Elem1.First < Elem2.First; });
+			for (const auto& Elem : Intersections)
+			{
+				if (SignedDistance(Elem.Second) < (Thickness + 1e-4))
+				{
+					return MakePair(Elem.Second, true);
+				}
+			}
+			return MakePair(TVector<T, d>(0), false);
+		}
+
+		FORCEINLINE TVector<T, d> FindGeometryOpposingNormal(const TVector<T, d>& DenormDir, int32 FaceIndex, const TVector<T, d>& OriginalNormal) const 
+		{
+			// Find which faces were included in the contact normal, and for multiple faces, use the one most opposing the sweep direction.
+			TVector<T, d> BestNormal(OriginalNormal);
+			T BestOpposingDot = TNumericLimits<T>::Max();
+
+			for (int32 Axis = 0; Axis < d; Axis++)
+			{
+				// Select axis of face to compare to, based on normal.
+				if (OriginalNormal[Axis] > KINDA_SMALL_NUMBER)
+				{
+					const float TraceDotFaceNormal = DenormDir[Axis]; // TraceDirDenormLocal.dot(BoxFaceNormal)
+					if (TraceDotFaceNormal < BestOpposingDot)
+					{
+						BestOpposingDot = TraceDotFaceNormal;
+						BestNormal = TVector<T, d>(0);
+						BestNormal[Axis] = 1;
+					}
+				}
+				else if (OriginalNormal[Axis] < -KINDA_SMALL_NUMBER)
+				{
+					const float TraceDotFaceNormal = -DenormDir[Axis]; // TraceDirDenormLocal.dot(BoxFaceNormal)
+					if (TraceDotFaceNormal < BestOpposingDot)
+					{
+						BestOpposingDot = TraceDotFaceNormal;
+						BestNormal = FVector(0.f);
+						BestNormal[Axis] = -1.f;
+					}
+				}
+			}
+
+			return BestNormal;
+		}
+
+		FORCEINLINE TVector<T, d> Support(const TVector<T, d>& Direction, const T Thickness) const 
+		{
+			TVector<T, d> ChosenPt;
+			for (int Axis = 0; Axis < d; ++Axis)
+			{
+				ChosenPt[Axis] = Direction[Axis] < 0 ? MMin[Axis] : MMax[Axis];
+			}
+
+			if (Thickness)
+			{
+				//We want N / ||N|| and to avoid inf
+				//So we want N / ||N|| < 1 / eps => N eps < ||N||, but this is clearly true for all eps < 1 and N > 0
+				T SizeSqr = Direction.SizeSquared();
+				if (SizeSqr <= TNumericLimits<T>::Min())
+				{
+					return ChosenPt;
+				}
+				const TVector<T, d> Normalized = Direction / sqrt(SizeSqr);
+
+				const TVector<T, d> InflatedPt = ChosenPt + Normalized.GetSafeNormal() * Thickness;
+				return InflatedPt;
+			}
+
+			return ChosenPt;
+		}
+
+		FORCEINLINE void GrowToInclude(const TVector<T, d>& V)
+		{
+			MMin = TVector<T, d>(FGenericPlatformMath::Min(MMin[0], V[0]), FGenericPlatformMath::Min(MMin[1], V[1]), FGenericPlatformMath::Min(MMin[2], V[2]));
+			MMax = TVector<T, d>(FGenericPlatformMath::Max(MMax[0], V[0]), FGenericPlatformMath::Max(MMax[1], V[1]), FGenericPlatformMath::Max(MMax[2], V[2]));
+		}
+
+		FORCEINLINE void GrowToInclude(const TAABB<T, d>& Other)
+		{
+			MMin = TVector<T, d>(FGenericPlatformMath::Min(MMin[0], Other.MMin[0]), FGenericPlatformMath::Min(MMin[1], Other.MMin[1]), FGenericPlatformMath::Min(MMin[2], Other.MMin[2]));
+			MMax = TVector<T, d>(FGenericPlatformMath::Max(MMax[0], Other.MMax[0]), FGenericPlatformMath::Max(MMax[1], Other.MMax[1]), FGenericPlatformMath::Max(MMax[2], Other.MMax[2]));
+		}
+
+		FORCEINLINE void ShrinkToInclude(const TAABB<T, d>& Other)
+		{
+			MMin = TVector<T, d>(FGenericPlatformMath::Max(MMin[0], Other.MMin[0]), FGenericPlatformMath::Max(MMin[1], Other.MMin[1]), FGenericPlatformMath::Max(MMin[2], Other.MMin[2]));
+			MMax = TVector<T, d>(FGenericPlatformMath::Min(MMax[0], Other.MMax[0]), FGenericPlatformMath::Min(MMax[1], Other.MMax[1]), FGenericPlatformMath::Min(MMax[2], Other.MMax[2]));
+		}
+
+		FORCEINLINE void Thicken(const float Thickness)
+		{
+			MMin -= TVector<T, d>(Thickness);
+			MMax += TVector<T, d>(Thickness);
+		}
+
+		//Grows (or shrinks) the box by this vector symmetrically - Changed name because previous Thicken had different semantics which caused several bugs
+		FORCEINLINE void ThickenSymmetrically(const TVector<T, d>& Thickness)
+		{
+			const TVector<T, d> AbsThickness = TVector<T, d>(FGenericPlatformMath::Abs(Thickness.X), FGenericPlatformMath::Abs(Thickness.Y), FGenericPlatformMath::Abs(Thickness.Z));
+			MMin -= AbsThickness;
+			MMax += AbsThickness;
+		}
+
+		FORCEINLINE TVector<T, d> Center() const { return (MMax - MMin) / (T)2 + MMin; }
+		FORCEINLINE TVector<T, d> GetCenter() const { return Center(); }
+		FORCEINLINE TVector<T, d> GetCenterOfMass() const { return GetCenter(); }
+		FORCEINLINE TVector<T, d> Extents() const { return MMax - MMin; }
+
+		FORCEINLINE int LargestAxis() const
+		{
+			const auto Extents = this->Extents();
+			if (Extents[0] > Extents[1] && Extents[0] > Extents[2])
+			{
+				return 0;
+			}
+			else if (Extents[1] > Extents[2])
+			{
+				return 1;
+			}
+			else
+			{
+				return 2;
+			}
+		}
+
+		FORCEINLINE const TVector<T, d>& Min() const { return MMin; }
+		FORCEINLINE const TVector<T, d>& Max() const { return MMax; }
+
+		FORCEINLINE T GetArea() const { return GetArea(Extents()); }
+		FORCEINLINE static T GetArea(const TVector<T, d>& Dim) { return d == 2 ? Dim.Product() : 2. * (Dim[0] * Dim[1] + Dim[0] * Dim[2] + Dim[1] * Dim[2]); }
+
+		FORCEINLINE T GetVolume() const { return GetVolume(Extents()); }
+		FORCEINLINE static T GetVolume(const TVector<T, 3>& Dim) { return Dim.Product(); }
+
+		FORCEINLINE static TAABB<T, d> EmptyAABB() { return TAABB<T, d>(TVector<T, d>(TNumericLimits<T>::Max()), TVector<T, d>(-TNumericLimits<T>::Max())); }
+		FORCEINLINE static TAABB<T, d> ZeroAABB() { return TAABB<T, d>(TVector<T, d>((T)0), TVector<T, d>((T)0)); }
+
+		FORCEINLINE void Serialize(FArchive &Ar) 
+		{
+			Ar << MMin << MMax;
+		}
+
+		FORCEINLINE uint32 GetTypeHash() const
+		{
+			return HashCombine(::GetTypeHash(MMin), ::GetTypeHash(MMax));
+		}
+
+	private:
+		TVector<T, d> MMin, MMax;
+	};
+
+	template<typename T>
+	struct TAABBSpecializeSamplingHelper<T, 2>
+	{
+		static FORCEINLINE TArray<TVector<T, 2>> ComputeSamplePoints(const TAABB<T, 2>& AABB)
+		{
+			const TVector<T, 2>& Min = AABB.Min();
+			const TVector<T, 2>& Max = AABB.Max();
+			const TVector<T, 2> Mid = AABB.Center();
+
+			TArray<TVector<T, 2>> SamplePoints;
+			SamplePoints.SetNum(8);
+			//top line (min y)
+			SamplePoints[0] = TVector<T, 2>{Min.X, Min.Y};
+			SamplePoints[1] = TVector<T, 2>{Mid.X, Min.Y};
+			SamplePoints[2] = TVector<T, 2>{Max.X, Min.Y};
+
+			//mid line (y=0) (mid point removed because internal)
+			SamplePoints[3] = TVector<T, 2>{Min.X, Mid.Y};
+			SamplePoints[4] = TVector<T, 2>{Max.X, Mid.Y};
+
+			//bottom line (max y)
+			SamplePoints[5] = TVector<T, 2>{Min.X, Max.Y};
+			SamplePoints[6] = TVector<T, 2>{Mid.X, Max.Y};
+			SamplePoints[7] = TVector<T, 2>{Max.X, Max.Y};
+
+			return SamplePoints;
+		}
+	};
+
+	template<typename T>
+	struct TAABBSpecializeSamplingHelper<T, 3>
+	{
+		static FORCEINLINE TArray<TVector<T, 3>> ComputeSamplePoints(const TAABB<T, 3>& AABB)
+		{
+			const TVector<T, 3>& Min = AABB.Min();
+			const TVector<T, 3>& Max = AABB.Max();
+			const TVector<T, 3> Mid = AABB.Center();
+
+			//todo(ocohen): should order these for best levelset cache traversal
+			TArray<TVector<T, 3>> SamplePoints;
+			SamplePoints.SetNum(26);
+			{
+				//xy plane for Min Z
+				SamplePoints[0] = TVector<T, 3>{Min.X, Min.Y, Min.Z};
+				SamplePoints[1] = TVector<T, 3>{Mid.X, Min.Y, Min.Z};
+				SamplePoints[2] = TVector<T, 3>{Max.X, Min.Y, Min.Z};
+
+				SamplePoints[3] = TVector<T, 3>{Min.X, Mid.Y, Min.Z};
+				SamplePoints[4] = TVector<T, 3>{Mid.X, Mid.Y, Min.Z};
+				SamplePoints[5] = TVector<T, 3>{Max.X, Mid.Y, Min.Z};
+
+				SamplePoints[6] = TVector<T, 3>{Min.X, Max.Y, Min.Z};
+				SamplePoints[7] = TVector<T, 3>{Mid.X, Max.Y, Min.Z};
+				SamplePoints[8] = TVector<T, 3>{Max.X, Max.Y, Min.Z};
+			}
+
+			{
+				//xy plane for z = 0 (skip mid point since inside)
+				SamplePoints[9] = TVector<T, 3>{Min.X, Min.Y, Mid.Z};
+				SamplePoints[10] = TVector<T, 3>{Mid.X, Min.Y, Mid.Z};
+				SamplePoints[11] = TVector<T, 3>{Max.X, Min.Y, Mid.Z};
+
+				SamplePoints[12] = TVector<T, 3>{Min.X, Mid.Y, Mid.Z};
+				SamplePoints[13] = TVector<T, 3>{Max.X, Mid.Y, Mid.Z};
+
+				SamplePoints[14] = TVector<T, 3>{Min.X, Max.Y, Mid.Z};
+				SamplePoints[15] = TVector<T, 3>{Mid.X, Max.Y, Mid.Z};
+				SamplePoints[16] = TVector<T, 3>{Max.X, Max.Y, Mid.Z};
+			}
+
+			{
+				//xy plane for Max Z
+				SamplePoints[17] = TVector<T, 3>{Min.X, Min.Y, Max.Z};
+				SamplePoints[18] = TVector<T, 3>{Mid.X, Min.Y, Max.Z};
+				SamplePoints[19] = TVector<T, 3>{Max.X, Min.Y, Max.Z};
+
+				SamplePoints[20] = TVector<T, 3>{Min.X, Mid.Y, Max.Z};
+				SamplePoints[21] = TVector<T, 3>{Mid.X, Mid.Y, Max.Z};
+				SamplePoints[22] = TVector<T, 3>{Max.X, Mid.Y, Max.Z};
+
+				SamplePoints[23] = TVector<T, 3>{Min.X, Max.Y, Max.Z};
+				SamplePoints[24] = TVector<T, 3>{Mid.X, Max.Y, Max.Z};
+				SamplePoints[25] = TVector<T, 3>{Max.X, Max.Y, Max.Z};
+			}
+
+			return SamplePoints;
+		}
+	};
+
+
+}
\ No newline at end of file
diff --git a/Engine/Source/Runtime/Experimental/Chaos/Public/Chaos/Box.h b/Engine/Source/Runtime/Experimental/Chaos/Public/Chaos/Box.h
--- a/Engine/Source/Runtime/Experimental/Chaos/Public/Chaos/Box.h
+++ b/Engine/Source/Runtime/Experimental/Chaos/Public/Chaos/Box.h
@@ -4,22 +4,12 @@
 #include "Templates/UnrealTemplate.h"
 
 #include "Chaos/ImplicitObject.h"
-#include "Chaos/Plane.h"
+#include "Chaos/AABB.h"
 #include "Chaos/Transform.h"
 #include "ChaosArchive.h"
 
 namespace Chaos
 {
-	template<typename T, int d>
-	struct TBoxSpecializeSamplingHelper
-	{
-		static FORCEINLINE TArray<TVector<T, d>> ComputeLocalSamplePoints(const class TBox<T, d>& Box)
-		{
-			check(false);
-			return TArray<TVector<T, d>>();
-		}
-	};
-
 	template<class T, int d>
 	class TBox final : public TImplicitObject<T, d>
 	{
@@ -32,8 +22,7 @@ namespace Chaos
 		    : TImplicitObject<T, d>(EImplicitObject::FiniteConvex, ImplicitObjectType::Box){};
 		FORCEINLINE TBox(const TVector<T, d>& Min, const TVector<T, d>& Max)
 		    : TImplicitObject<T, d>(EImplicitObject::FiniteConvex, ImplicitObjectType::Box)
-		    , MMin(Min)
-		    , MMax(Max)
+			, AABB(Min, Max)
 		{
 			//todo: turn back on
 			/*for (int Axis = 0; Axis < d; ++Axis)
@@ -44,15 +33,19 @@ namespace Chaos
 
 		FORCEINLINE TBox(const TBox<T, d>& Other)
 		    : TImplicitObject<T, d>(EImplicitObject::FiniteConvex, ImplicitObjectType::Box)
-		    , MMin(Other.MMin)
-		    , MMax(Other.MMax)
+			, AABB(Other.AABB)
+		{
+		}
+
+		FORCEINLINE TBox(const TAABB<T, d>& AABB)
+		    : TImplicitObject<T, d>(EImplicitObject::FiniteConvex, ImplicitObjectType::Box)
+			, AABB(AABB)
 		{
 		}
 
 		FORCEINLINE TBox(TBox<T, d>&& Other)
 		    : TImplicitObject<T, d>(EImplicitObject::FiniteConvex, ImplicitObjectType::Box)
-		    , MMin(MoveTemp(Other.MMin))
-		    , MMax(MoveTemp(Other.MMax))
+		    , AABB(MoveTemp(Other.AABB))
 		{
 		}
 
@@ -63,8 +56,7 @@ namespace Chaos
 			this->bIgnoreAnalyticCollisions = Other.bIgnoreAnalyticCollisions;
 			this->bHasBoundingBox = Other.bHasBoundingBox;
 
-			MMin = Other.MMin;
-			MMax = Other.MMax;
+			AABB = Other.AABB;
 			return *this;
 		}
 
@@ -75,10 +67,10 @@ namespace Chaos
 			this->bIgnoreAnalyticCollisions = Other.bIgnoreAnalyticCollisions;
 			this->bHasBoundingBox = Other.bHasBoundingBox;
 
-			MMin = MoveTemp(Other.MMin);
-			MMax = MoveTemp(Other.MMax);
+			AABB = MoveTemp(Other.AABB);
 			return *this;
 		}
+
 		virtual ~TBox() {}
 
 		virtual TUniquePtr<TImplicitObject<T, d>> Copy() const override
@@ -91,8 +83,7 @@ namespace Chaos
 		 */
 		TArray<TVector<T, d>> ComputeLocalSamplePoints() const
 		{
-			const TVector<T, d> Mid = Center();
-			return TBoxSpecializeSamplingHelper<T, d>::ComputeSamplePoints(TBox<T, d>(Min() - Mid, Max() - Mid));
+			return AABB.ComputeLocalSamplePoints();
 		}
 
 		/**
@@ -100,65 +91,33 @@ namespace Chaos
 		 */
 		TArray<TVector<T, d>> ComputeSamplePoints() const
 		{
-			return TBoxSpecializeSamplingHelper<T, d>::ComputeSamplePoints(*this);
+			return AABB.ComputeSamplePoints();
 		}
 
 		template<class TTRANSFORM>
 		TBox<T, d> TransformedBox(const TTRANSFORM& SpaceTransform) const
 		{
-			TVector<T, d> CurrentExtents = Extents();
-			int32 Idx = 0;
-			const TVector<T, d> MinToNewSpace = SpaceTransform.TransformPosition(MMin);
-			TBox<T, d> NewBox(MinToNewSpace, MinToNewSpace);
-			NewBox.GrowToInclude(SpaceTransform.TransformPosition(MMax));
-
-			for (int32 j = 0; j < d; ++j)
-			{
-				NewBox.GrowToInclude(SpaceTransform.TransformPosition(MMin + TVector<T, d>::AxisVector(j) * CurrentExtents));
-				NewBox.GrowToInclude(SpaceTransform.TransformPosition(MMax - TVector<T, d>::AxisVector(j) * CurrentExtents));
-			}
-
-			return NewBox;
+			return TBox<T, d>(AABB.TransformedAABB(SpaceTransform));
 		}
 
 		FORCEINLINE bool Intersects(const TBox<T, d>& Other) const
 		{
-			for (int32 i = 0; i < d; ++i)
-			{
-				if (Other.MMax[i] < MMin[i] || Other.MMin[i] > MMax[i])
-					return false;
-			}
-			return true;
+			return AABB.Intersects(Other.AABB);
 		}
 
 		TBox<T, d> GetIntersection(const TBox<T, d>& Other) const
 		{
-			TVector<float, 3> Tmp;
-			return TBox<T, d>(MMin.ComponentwiseMax(Other.MMin), MMax.ComponentwiseMin(Other.MMax));
+			return TBox<T, d>(AABB.GetIntersection(Other.AABB));
 		}
 
 		FORCEINLINE bool Contains(const TVector<T, d>& Point) const
 		{
-			for (int i = 0; i < d; i++)
-			{
-				if (Point[i] < MMin[i] || Point[i] > MMax[i])
-				{
-					return false;
-				}
-			}
-			return true;
+			return AABB.Contains(Point);
 		}
 
 		FORCEINLINE bool Contains(const TVector<T, d>& Point, const T Tolerance) const
 		{
-			for (int i = 0; i < d; i++)
-			{
-				if (Point[i] < MMin[i] - Tolerance || Point[i] > MMax[i] + Tolerance)
-				{
-					return false;
-				}
-			}
-			return true;
+			return AABB.Contains(Point, Tolerance);
 		}
 
 		FORCEINLINE static ImplicitObjectType GetType() { return ImplicitObjectType::Box; }
@@ -167,308 +126,86 @@ namespace Chaos
 
 		virtual T PhiWithNormal(const TVector<T, d>& x, TVector<T, d>& Normal) const override
 		{
-			const TVector<T, d> MaxDists = x - MMax;
-			const TVector<T, d> MinDists = MMin - x;
-			if (x <= MMax && x >= MMin)
-			{
-				const Pair<T, int32> MaxAndAxis = TVector<T, d>::MaxAndAxis(MinDists, MaxDists);
-				Normal = MaxDists[MaxAndAxis.Second] > MinDists[MaxAndAxis.Second] ? TVector<T, d>::AxisVector(MaxAndAxis.Second) : -TVector<T, d>::AxisVector(MaxAndAxis.Second);
-				return MaxAndAxis.First;
-			}
-			else
-			{
-				for (int i = 0; i < d; ++i)
-				{
-					if (MaxDists[i] > 0)
-					{
-						Normal[i] = MaxDists[i];
-					}
-					else if (MinDists[i] > 0)
-					{
-						Normal[i] = -MinDists[i];
-					}
-					else
-					{
-						Normal[i] = 0;
-					}
-				}
-				T Phi = Normal.SafeNormalize();
-				if (Phi < KINDA_SMALL_NUMBER)
-				{
-					for (int i = 0; i < d; ++i)
-					{
-						if (Normal[i] > 0)
-						{
-							Normal[i] = 1;
-						}
-						else if (Normal[i] < 0)
-						{
-							Normal[i] = -1;
-						}
-					}
-					Normal.Normalize();
-				}
-				return Phi;
-			}
-		}
-
-		virtual bool CHAOS_API Raycast(const TVector<T, d>& StartPoint, const TVector<T, d>& Dir, const T Length, const T Thickness, T& OutTime, TVector<T, d>& OutPosition, TVector<T, d>& OutNormal, int32& OutFaceIndex) const override;
-
-		static bool RaycastFast(const TVector<T,d>& MMin, const TVector<T,d>& MMax, const TVector<T, d>& StartPoint, const TVector<T, d>& Dir, const TVector<T, d>& InvDir, const bool* bParallel, const T Length, const T InvLength, T& OutTime, TVector<T, d>& OutPosition)
-		{
-			const TVector<T, d> StartToMin = MMin - StartPoint;
-			const TVector<T, d> StartToMax = MMax - StartPoint;
-
-			//For each axis record the start and end time when ray is in the box. If the intervals overlap the ray is inside the box
-			T LatestStartTime = 0;
-			T EarliestEndTime = FLT_MAX;
-
-			for (int Axis = 0; Axis < d; ++Axis)
-			{
-				T Time1, Time2;
-				if (bParallel[Axis])
-				{
-					if (StartToMin[Axis] > 0 || StartToMax[Axis] < 0)
-					{
-						return false;	//parallel and outside
-					}
-					else
-					{
-						Time1 = 0;
-						Time2 = FLT_MAX;
-					}
-				}
-				else
-				{
-					Time1 = StartToMin[Axis] * InvDir[Axis];
-					Time2 = StartToMax[Axis] * InvDir[Axis];
-				}
-
-				if (Time1 > Time2)
-				{
-					//going from max to min direction
-					Swap(Time1, Time2);
-				}
-
-				LatestStartTime = FMath::Max(LatestStartTime, Time1);
-				EarliestEndTime = FMath::Min(EarliestEndTime, Time2);
-
-				if (LatestStartTime > EarliestEndTime)
-				{
-					return false;	//Outside of slab before entering another
-				}
-			}
-
-			//infinite ray intersects with inflated box
-			if (LatestStartTime > Length || EarliestEndTime < 0)
-			{
-				//outside of line segment given
-				return false;
-			}
+			return AABB.PhiWithNormal(x, Normal);
+		}
+
 
-			OutTime = LatestStartTime;
-			OutPosition = StartPoint + LatestStartTime * Dir;
-			return true;
+		static FORCEINLINE bool RaycastFast(const TVector<T,d>& Min, const TVector<T,d>& Max, const TVector<T, d>& StartPoint, const TVector<T, d>& Dir, const TVector<T, d>& InvDir, const bool* bParallel, const T Length, const T InvLength, T& OutTime, TVector<T, d>& OutPosition)
+		{
+			TAABB<T, d> AABB(Min, Max);
+			return AABB.RaycastFast(StartPoint, Dir, InvDir, bParallel, Length, InvLength, OutTime, OutPosition);
 		}
-		
-		
-		TVector<T, d> FindClosestPoint(const TVector<T, d>& StartPoint, const T Thickness = (T)0) const
+
+		virtual bool CHAOS_API Raycast(const TVector<T, d>& StartPoint, const TVector<T, d>& Dir, const T Length, const T Thickness, T& OutTime, TVector<T, d>& OutPosition, TVector<T, d>& OutNormal, int32& OutFaceIndex) const override
 		{
-			TVector<T, d> Result(0);
+			return AABB.Raycast(StartPoint, Dir, Length, Thickness, OutTime, OutPosition, OutNormal, OutFaceIndex);
+		}
 
-			// clamp exterior to surface
-			bool bIsExterior = false;
-			for (int i = 0; i < 3; i++)
-			{
-				float v = StartPoint[i];
-				if (v < MMin[i])
-				{
-					v = MMin[i];
-					bIsExterior = true;
-				}
-				if (v > MMax[i])
-				{
-					v = MMax[i];
-					bIsExterior = true;
-				}
-				Result[i] = v;
-			}
-
-			if (!bIsExterior)
-			{
-				TArray<Pair<T, TVector<T, d>>> Intersections;
-
-				// sum interior direction to surface
-				for (int32 i = 0; i < d; ++i)
-				{
-					auto PlaneIntersection = TPlane<T, d>(MMin - Thickness, -TVector<T, d>::AxisVector(i)).FindClosestPoint(Result, 0);
-					Intersections.Add(MakePair((PlaneIntersection - Result).Size(), -TVector<T, d>::AxisVector(i)));
-					PlaneIntersection = TPlane<T, d>(MMax + Thickness, TVector<T, d>::AxisVector(i)).FindClosestPoint(Result, 0);
-					Intersections.Add(MakePair((PlaneIntersection - Result).Size(), TVector<T, d>::AxisVector(i)));
-				}
-				Intersections.Sort([](const Pair<T, TVector<T, d>>& Elem1, const Pair<T, TVector<T, d>>& Elem2) { return Elem1.First < Elem2.First; });
-
-				if (!FMath::IsNearlyEqual(Intersections[0].First, 0.f))
-				{
-					T SmallestDistance = Intersections[0].First;
-					Result += Intersections[0].Second * Intersections[0].First;
-					for (int32 i = 1; i < 3 && FMath::IsNearlyEqual(SmallestDistance, Intersections[i].First); ++i)
-					{
-						Result += Intersections[i].Second * Intersections[i].First;
-					}
-				}
-			}
-			return Result;
+		TVector<T, d> FindClosestPoint(const TVector<T, d>& StartPoint, const T Thickness = (T)0) const
+		{
+			return AABB.FindClosestPoint(StartPoint, Thickness);
 		}
 
 		virtual Pair<TVector<T, d>, bool> FindClosestIntersectionImp(const TVector<T, d>& StartPoint, const TVector<T, d>& EndPoint, const T Thickness) const override
 		{
-			TArray<Pair<T, TVector<T, d>>> Intersections;
-			for (int32 i = 0; i < d; ++i)
-			{
-				auto PlaneIntersection = TPlane<T, d>(MMin - Thickness, -TVector<T, d>::AxisVector(i)).FindClosestIntersection(StartPoint, EndPoint, 0);
-				if (PlaneIntersection.Second)
-					Intersections.Add(MakePair((PlaneIntersection.First - StartPoint).Size(), PlaneIntersection.First));
-				PlaneIntersection = TPlane<T, d>(MMax + Thickness, TVector<T, d>::AxisVector(i)).FindClosestIntersection(StartPoint, EndPoint, 0);
-				if (PlaneIntersection.Second)
-					Intersections.Add(MakePair((PlaneIntersection.First - StartPoint).Size(), PlaneIntersection.First));
-			}
-			Intersections.Sort([](const Pair<T, TVector<T, d>>& Elem1, const Pair<T, TVector<T, d>>& Elem2) { return Elem1.First < Elem2.First; });
-			for (const auto& Elem : Intersections)
-			{
-				if (SignedDistance(Elem.Second) < (Thickness + 1e-4))
-				{
-					return MakePair(Elem.Second, true);
-				}
-			}
-			return MakePair(TVector<T, d>(0), false);
+			return AABB.FindClosestIntersectionImp(StartPoint, EndPoint, Thickness);
 		}
 
 		virtual TVector<T, d> FindGeometryOpposingNormal(const TVector<T, d>& DenormDir, int32 FaceIndex, const TVector<T, d>& OriginalNormal) const override
 		{
-			// Find which faces were included in the contact normal, and for multiple faces, use the one most opposing the sweep direction.
-			TVector<T,d> BestNormal(OriginalNormal);
-			T BestOpposingDot = TNumericLimits<T>::Max();
-
-			for (int32 Axis = 0; Axis < d; Axis++)
-			{
-				// Select axis of face to compare to, based on normal.
-				if (OriginalNormal[Axis] > KINDA_SMALL_NUMBER)
-				{
-					const float TraceDotFaceNormal = DenormDir[Axis]; // TraceDirDenormLocal.dot(BoxFaceNormal)
-					if (TraceDotFaceNormal < BestOpposingDot)
-					{
-						BestOpposingDot = TraceDotFaceNormal;
-						BestNormal = TVector<T,d>(0);
-						BestNormal[Axis] = 1;
-					}
-				}
-				else if (OriginalNormal[Axis] < -KINDA_SMALL_NUMBER)
-				{
-					const float TraceDotFaceNormal = -DenormDir[Axis]; // TraceDirDenormLocal.dot(BoxFaceNormal)
-					if (TraceDotFaceNormal < BestOpposingDot)
-					{
-						BestOpposingDot = TraceDotFaceNormal;
-						BestNormal = FVector(0.f);
-						BestNormal[Axis] = -1.f;
-					}
-				}
-			}
-
-			return BestNormal;
+			return AABB.FindGeometryOpposingNormal(DenormDir, FaceIndex, OriginalNormal);
 		}
 
 		virtual TVector<T, d> Support(const TVector<T, d>& Direction, const T Thickness) const override
 		{
-			TVector<T, d> ChosenPt;
-			for (int Axis = 0; Axis < d; ++Axis)
-			{
-				ChosenPt[Axis] = Direction[Axis] < 0 ? MMin[Axis] : MMax[Axis];
-			}
-
-			if (Thickness)
-			{
-				//We want N / ||N|| and to avoid inf
-				//So we want N / ||N|| < 1 / eps => N eps < ||N||, but this is clearly true for all eps < 1 and N > 0
-				T SizeSqr = Direction.SizeSquared();
-				if (SizeSqr <= TNumericLimits<T>::Min())
-				{
-					return ChosenPt;
-				}
-				const TVector<T, d> Normalized = Direction / sqrt(SizeSqr);
-
-				const TVector<T, d> InflatedPt = ChosenPt + Normalized.GetSafeNormal() * Thickness;
-				return InflatedPt;
-			}
-
-			return ChosenPt;
+			return AABB.Support(Direction, Thickness);
 		}
 
 		FORCEINLINE void GrowToInclude(const TVector<T, d>& V)
 		{
-			MMin = TVector<T, d>(FGenericPlatformMath::Min(MMin[0], V[0]), FGenericPlatformMath::Min(MMin[1], V[1]), FGenericPlatformMath::Min(MMin[2], V[2]));
-			MMax = TVector<T, d>(FGenericPlatformMath::Max(MMax[0], V[0]), FGenericPlatformMath::Max(MMax[1], V[1]), FGenericPlatformMath::Max(MMax[2], V[2]));
+			AABB.GrowToInclude(V);
 		}
 
 		FORCEINLINE void GrowToInclude(const TBox<T, d>& Other)
 		{
-			MMin = TVector<T, d>(FGenericPlatformMath::Min(MMin[0], Other.MMin[0]), FGenericPlatformMath::Min(MMin[1], Other.MMin[1]), FGenericPlatformMath::Min(MMin[2], Other.MMin[2]));
-			MMax = TVector<T, d>(FGenericPlatformMath::Max(MMax[0], Other.MMax[0]), FGenericPlatformMath::Max(MMax[1], Other.MMax[1]), FGenericPlatformMath::Max(MMax[2], Other.MMax[2]));
+			AABB.GrowToInclude(Other.AABB);
 		}
 
 		FORCEINLINE void ShrinkToInclude(const TBox<T, d>& Other)
 		{
-			MMin = TVector<T, d>(FGenericPlatformMath::Max(MMin[0], Other.MMin[0]), FGenericPlatformMath::Max(MMin[1], Other.MMin[1]), FGenericPlatformMath::Max(MMin[2], Other.MMin[2]));
-			MMax = TVector<T, d>(FGenericPlatformMath::Min(MMax[0], Other.MMax[0]), FGenericPlatformMath::Min(MMax[1], Other.MMax[1]), FGenericPlatformMath::Min(MMax[2], Other.MMax[2]));
+			AABB.ShrinkToInclude(Other.AABB);
 		}
 
 		FORCEINLINE void Thicken(const float Thickness)
 		{
-			MMin -= TVector<T, d>(Thickness);
-			MMax += TVector<T, d>(Thickness);
+			AABB.Thicken(Thickness);
 		}
 
 		//Grows (or shrinks) the box by this vector symmetrically - Changed name because previous Thicken had different semantics which caused several bugs
 		FORCEINLINE void ThickenSymmetrically(const TVector<T,d>& Thickness)
 		{
-			const TVector<T,d> AbsThickness = TVector<T,d>(FGenericPlatformMath::Abs(Thickness.X), FGenericPlatformMath::Abs(Thickness.Y), FGenericPlatformMath::Abs(Thickness.Z));
-			MMin -= AbsThickness;
-			MMax += AbsThickness;
+			AABB.ThickenSymmetrically(Thickness);
 		}
 
-		FORCEINLINE TVector<T, d> Center() const { return (MMax - MMin) / (T)2 + MMin; }
-		FORCEINLINE TVector<T, d> GetCenter() const { return Center(); }
-		FORCEINLINE TVector<T, d> GetCenterOfMass() const { return GetCenter(); }
-		FORCEINLINE TVector<T, d> Extents() const { return MMax - MMin; }
+		FORCEINLINE TVector<T, d> Center() const { return AABB.Center(); }
+		FORCEINLINE TVector<T, d> GetCenter() const { return AABB.GetCenter(); }
+		FORCEINLINE TVector<T, d> GetCenterOfMass() const { return AABB.GetCenterOfMass(); }
+		FORCEINLINE TVector<T, d> Extents() const { return AABB.Extents(); }
 
 		int LargestAxis() const
 		{
-			const auto Extents = this->Extents();
-			if (Extents[0] > Extents[1] && Extents[0] > Extents[2])
-			{
-				return 0;
-			}
-			else if (Extents[1] > Extents[2])
-			{
-				return 1;
-			}
-			else
-			{
-				return 2;
-			}
+			return AABB.LargestAxis();
 		}
 
-		FORCEINLINE const TVector<T, d>& Min() const { return MMin; }
-		FORCEINLINE const TVector<T, d>& Max() const { return MMax; }
+		FORCEINLINE const TVector<T, d>& Min() const { return AABB.Min(); }
+		FORCEINLINE const TVector<T, d>& Max() const { return AABB.Max(); }
 
-		T GetArea() const { return GetArea(Extents()); }
-		static T GetArea(const TVector<T, d>& Dim) { return d == 2 ? Dim.Product() : 2. * (Dim[0] * Dim[1] + Dim[0] * Dim[2] + Dim[1] * Dim[2]); }
+		T GetArea() const { return AABB.GetArea(); }
 
-		T GetVolume() const { return GetVolume(Extents()); }
-		static T GetVolume(const TVector<T, 3>& Dim) { return Dim.Product(); }
+		T GetVolume() const { return AABB.GetVolume(); }
 
 		PMatrix<T, d, d> GetInertiaTensor(const T Mass) const { return GetInertiaTensor(Mass, Extents()); }
-		static PMatrix<T, 3, 3> GetInertiaTensor(const T Mass, const TVector<T, 3>& Dim)
+		FORCEINLINE static PMatrix<T, 3, 3> GetInertiaTensor(const T Mass, const TVector<T, 3>& Dim)
 		{
 			// https://www.wolframalpha.com/input/?i=cuboid
 			const T M = Mass / 12;
@@ -477,7 +214,9 @@ namespace Chaos
 			const T DD = Dim[2] * Dim[2];
 			return PMatrix<T, 3, 3>(M * (HH + DD), M * (WW + DD), M * (WW + HH));
 		}
-		static TRotation<T, d> GetRotationOfMass()
+
+
+		FORCEINLINE static TRotation<T, d> GetRotationOfMass()
 		{
 			return TRotation<T, d>::FromElements(TVector<T, d>(0), 1);
 		}
@@ -487,13 +226,15 @@ namespace Chaos
 		FORCEINLINE void SerializeImp(FArchive& Ar)
 		{
 			TImplicitObject<T, d>::SerializeImp(Ar);
-			Ar << MMin << MMax;
+			AABB.Serialize(Ar);
 		}
+
 		virtual void Serialize(FChaosArchive& Ar) override
 		{
 			FChaosArchiveScopedMemory ScopedMemory(Ar, GetTypeName());
 			SerializeImp(Ar);
 		}
+
 		virtual void Serialize(FArchive& Ar) override { SerializeImp(Ar); }
 
 		static TBox<T, d> EmptyBox() { return TBox<T, d>(TVector<T, d>(TNumericLimits<T>::Max()), TVector<T, d>(-TNumericLimits<T>::Max())); }
@@ -501,99 +242,10 @@ namespace Chaos
 
 		virtual uint32 GetTypeHash() const override 
 		{
-			return HashCombine(::GetTypeHash(MMin), ::GetTypeHash(MMax));
+			return AABB.GetTypeHash();
 		}
 
 	private:
-		TVector<T, d> MMin, MMax;
-	};
-
-	template<typename T>
-	struct TBoxSpecializeSamplingHelper<T, 2>
-	{
-		static FORCEINLINE TArray<TVector<T, 2>> ComputeSamplePoints(const TBox<T, 2>& Box)
-		{
-			const TVector<T, 2>& Min = Box.Min();
-			const TVector<T, 2>& Max = Box.Max();
-			const TVector<T, 2> Mid = Box.Center();
-
-			TArray<TVector<T, 2>> SamplePoints;
-			SamplePoints.SetNum(8);
-			//top line (min y)
-			SamplePoints[0] = TVector<T, 2>{Min.X, Min.Y};
-			SamplePoints[1] = TVector<T, 2>{Mid.X, Min.Y};
-			SamplePoints[2] = TVector<T, 2>{Max.X, Min.Y};
-
-			//mid line (y=0) (mid point removed because internal)
-			SamplePoints[3] = TVector<T, 2>{Min.X, Mid.Y};
-			SamplePoints[4] = TVector<T, 2>{Max.X, Mid.Y};
-
-			//bottom line (max y)
-			SamplePoints[5] = TVector<T, 2>{Min.X, Max.Y};
-			SamplePoints[6] = TVector<T, 2>{Mid.X, Max.Y};
-			SamplePoints[7] = TVector<T, 2>{Max.X, Max.Y};
-
-			return SamplePoints;
-		}
-	};
-
-	template<typename T>
-	struct TBoxSpecializeSamplingHelper<T, 3>
-	{
-		static FORCEINLINE TArray<TVector<T, 3>> ComputeSamplePoints(const TBox<T, 3>& Box)
-		{
-			const TVector<T, 3>& Min = Box.Min();
-			const TVector<T, 3>& Max = Box.Max();
-			const TVector<T, 3> Mid = Box.Center();
-
-			//todo(ocohen): should order these for best levelset cache traversal
-			TArray<TVector<T, 3>> SamplePoints;
-			SamplePoints.SetNum(26);
-			{
-				//xy plane for Min Z
-				SamplePoints[0] = TVector<T, 3>{Min.X, Min.Y, Min.Z};
-				SamplePoints[1] = TVector<T, 3>{Mid.X, Min.Y, Min.Z};
-				SamplePoints[2] = TVector<T, 3>{Max.X, Min.Y, Min.Z};
-
-				SamplePoints[3] = TVector<T, 3>{Min.X, Mid.Y, Min.Z};
-				SamplePoints[4] = TVector<T, 3>{Mid.X, Mid.Y, Min.Z};
-				SamplePoints[5] = TVector<T, 3>{Max.X, Mid.Y, Min.Z};
-
-				SamplePoints[6] = TVector<T, 3>{Min.X, Max.Y, Min.Z};
-				SamplePoints[7] = TVector<T, 3>{Mid.X, Max.Y, Min.Z};
-				SamplePoints[8] = TVector<T, 3>{Max.X, Max.Y, Min.Z};
-			}
-
-			{
-				//xy plane for z = 0 (skip mid point since inside)
-				SamplePoints[9] = TVector<T, 3>{Min.X, Min.Y, Mid.Z};
-				SamplePoints[10] = TVector<T, 3>{Mid.X, Min.Y, Mid.Z};
-				SamplePoints[11] = TVector<T, 3>{Max.X, Min.Y, Mid.Z};
-
-				SamplePoints[12] = TVector<T, 3>{Min.X, Mid.Y, Mid.Z};
-				SamplePoints[13] = TVector<T, 3>{Max.X, Mid.Y, Mid.Z};
-
-				SamplePoints[14] = TVector<T, 3>{Min.X, Max.Y, Mid.Z};
-				SamplePoints[15] = TVector<T, 3>{Mid.X, Max.Y, Mid.Z};
-				SamplePoints[16] = TVector<T, 3>{Max.X, Max.Y, Mid.Z};
-			}
-
-			{
-				//xy plane for Max Z
-				SamplePoints[17] = TVector<T, 3>{Min.X, Min.Y, Max.Z};
-				SamplePoints[18] = TVector<T, 3>{Mid.X, Min.Y, Max.Z};
-				SamplePoints[19] = TVector<T, 3>{Max.X, Min.Y, Max.Z};
-
-				SamplePoints[20] = TVector<T, 3>{Min.X, Mid.Y, Max.Z};
-				SamplePoints[21] = TVector<T, 3>{Mid.X, Mid.Y, Max.Z};
-				SamplePoints[22] = TVector<T, 3>{Max.X, Mid.Y, Max.Z};
-
-				SamplePoints[23] = TVector<T, 3>{Min.X, Max.Y, Max.Z};
-				SamplePoints[24] = TVector<T, 3>{Mid.X, Max.Y, Max.Z};
-				SamplePoints[25] = TVector<T, 3>{Max.X, Max.Y, Max.Z};
-			}
-
-			return SamplePoints;
-		}
+		TAABB<T, d> AABB;
 	};
 }
