Chaos uses spatial collection by default. For now it's a collection of 1

diff --git a/Engine/Source/Runtime/Engine/Private/PhysicsEngine/Experimental/PhysScene_Chaos.cpp b/Engine/Source/Runtime/Engine/Private/PhysicsEngine/Experimental/PhysScene_Chaos.cpp
--- a/Engine/Source/Runtime/Engine/Private/PhysicsEngine/Experimental/PhysScene_Chaos.cpp
+++ b/Engine/Source/Runtime/Engine/Private/PhysicsEngine/Experimental/PhysScene_Chaos.cpp
@@ -33,6 +33,8 @@
 #include "Chaos/PBDSpringConstraints.h"
 #include "Chaos/PerParticleGravity.h"
 #include "PBDRigidActiveParticlesBuffer.h"
+#include "Chaos/GeometryParticlesfwd.h"
+
 
 #if !UE_BUILD_SHIPPING
 #include "Engine/World.h"
@@ -578,7 +580,8 @@ void FPhysScene_Chaos::RemoveActorFromAccelerationStructure(FPhysicsActorHandle&
 	{
 		ExternalDataLock.WriteLock();
 		Chaos::TAccelerationStructureHandle<float, 3> AccelerationHandle(Actor);
-		GetSpacialAcceleration()->RemoveElement(AccelerationHandle);
+		Chaos::FSpatialAccelerationIdx Idx = { 0,0 };	//todo: use real idx;
+		GetSpacialAcceleration()->RemoveElementFrom(AccelerationHandle, Idx);
 		ExternalDataLock.WriteUnlock();
 	}
 #endif
diff --git a/Engine/Source/Runtime/Engine/Public/Physics/Experimental/PhysScene_Chaos.h b/Engine/Source/Runtime/Engine/Public/Physics/Experimental/PhysScene_Chaos.h
--- a/Engine/Source/Runtime/Engine/Public/Physics/Experimental/PhysScene_Chaos.h
+++ b/Engine/Source/Runtime/Engine/Public/Physics/Experimental/PhysScene_Chaos.h
@@ -54,6 +54,9 @@ namespace Chaos
 	template <typename TPayload, typename T, int d>
 	class ISpatialAcceleration;
 
+	template <typename TPayload, typename T, int d>
+	class ISpatialAccelerationCollection;
+
 	class IDispatcher;
 }
 
@@ -161,7 +164,7 @@ public:
 	void CopySolverAccelerationStructure();
 
 private:
-	TUniquePtr<Chaos::ISpatialAcceleration<Chaos::TAccelerationStructureHandle<float, 3>, float, 3>> SolverAccelerationStructure;
+	TUniquePtr<Chaos::ISpatialAccelerationCollection<Chaos::TAccelerationStructureHandle<float, 3>, float, 3>> SolverAccelerationStructure;
 
 #if CHAOS_WITH_PAUSABLE_SOLVER
 	/** Callback that checks the status of the world settings for this scene before pausing/unpausing its solver. */
diff --git a/Engine/Source/Runtime/Experimental/Chaos/Private/Chaos/PBDCollisionConstraint.cpp b/Engine/Source/Runtime/Experimental/Chaos/Private/Chaos/PBDCollisionConstraint.cpp
--- a/Engine/Source/Runtime/Experimental/Chaos/Private/Chaos/PBDCollisionConstraint.cpp
+++ b/Engine/Source/Runtime/Experimental/Chaos/Private/Chaos/PBDCollisionConstraint.cpp
@@ -1,7 +1,7 @@
 // Copyright 1998-2019 Epic Games, Inc. All Rights Reserved.
 
 #include "Chaos/PBDCollisionConstraint.h"
-
+#include "Chaos/PBDCollisionConstraintImp.h"
 #include "Chaos/ChaosPerfTest.h"
 #include "Chaos/Defines.h"
 #include "Chaos/GeometryQueries.h"
@@ -18,6 +18,7 @@
 #include "ProfilingDebugging/ScopedTimers.h"
 #include "Chaos/PBDRigidsSOAs.h"
 #include "Chaos/ChaosDebugDraw.h"
+#include "Chaos/ISpatialAccelerationCollection.h"
 
 #if INTEL_ISPC
 #include "PBDCollisionConstraint.ispc.generated.h"
@@ -26,8 +27,6 @@
 int32 CollisionParticlesBVHDepth = 4;
 FAutoConsoleVariableRef CVarCollisionParticlesBVHDepth(TEXT("p.CollisionParticlesBVHDepth"), CollisionParticlesBVHDepth, TEXT("The maximum depth for collision particles bvh"));
 
-int32 EnableCollisions = 1;
-FAutoConsoleVariableRef CVarEnableCollisions(TEXT("p.EnableCollisions"), EnableCollisions, TEXT("Enable/Disable collisions on the Chaos solver."));
 
 int32 ConstraintBPBVHDepth = 2;
 FAutoConsoleVariableRef CVarConstraintBPBVHDepth(TEXT("p.ConstraintBPBVHDepth"), ConstraintBPBVHDepth, TEXT("The maximum depth for constraint bvh"));
@@ -47,6 +46,8 @@ namespace Chaos
 
 namespace Chaos
 {
+CHAOS_API int32 EnableCollisions = 1;
+FAutoConsoleVariableRef CVarEnableCollisions(TEXT("p.EnableCollisions"), EnableCollisions, TEXT("Enable/Disable collisions on the Chaos solver."));
 
 template<ECollisionUpdateType, typename T, int d>
 void UpdateConstraintImp2(const TRigidTransform<T, d>& ParticleTM, const TImplicitObject<T, d>& LevelsetObject, const TRigidTransform<T, d>& LevelsetTM, const T Thickness, TRigidBodyContactConstraint<T, d>& Constraint);
@@ -162,373 +163,43 @@ void TPBDCollisionConstraint<T, d>::UpdatePositionBasedState(/*const TPBDRigidPa
 	}
 }
 
-#if !UE_BUILD_SHIPPING
-
-int32 BucketSizes2[] = { 0,1,4,8,16,32,64,128,512,MAX_int32 };
-
-
-template <bool bGatherStats>
-struct FStatHelper2
-{
-	int32 MaxCount;
-
-	FStatHelper2() {}
-	void Record(int32 Count) {}
-	FString ToString() const { return TEXT(""); }
-};
-
-template <>
-struct FStatHelper2<true>
-{
-	int32 BucketCount[UE_ARRAY_COUNT(BucketSizes2)];
-	int32 MaxCount;
-
-	FStatHelper2()
-	{
-		FMemory::Memset(BucketCount, 0, sizeof(BucketCount));
-		MaxCount = 0;
-	}
-
-	void Record(int32 Count)
-	{
-		for (int32 BucketIdx = 1; BucketIdx < UE_ARRAY_COUNT(BucketSizes2); ++BucketIdx)
-		{
-			if (Count >= BucketSizes2[BucketIdx - 1] && Count < BucketSizes2[BucketIdx])
-			{
-				++BucketCount[BucketIdx];
-			}
-		}
-
-		if (Count > MaxCount)
-		{
-			MaxCount = Count;
-		}
-	}
-
-	FString ToString() const
-	{
-		FString OutLog;
-		int32 MaxBucketCount = 0;
-		for (int32 Count : BucketCount)
-		{
-			if (Count > MaxBucketCount)
-			{
-				MaxBucketCount = Count;
-			}
-		}
-
-		const float CountPerChar = MaxBucketCount / 20.f;
-		for (int32 Idx = 1; Idx < UE_ARRAY_COUNT(BucketSizes2); ++Idx)
-		{
-			int32 NumChars = BucketCount[Idx] / CountPerChar;
-			if (Idx < UE_ARRAY_COUNT(BucketSizes2) - 1)
-			{
-				OutLog += FString::Printf(TEXT("\t[%4d - %4d) (%4d) |"), BucketSizes2[Idx - 1], BucketSizes2[Idx], BucketCount[Idx]);
-			}
-			else
-			{
-				OutLog += FString::Printf(TEXT("\t[%4d -  inf) (%4d) |"), BucketSizes2[Idx - 1], BucketCount[Idx]);
-			}
-			for (int32 Count = 0; Count < NumChars; ++Count)
-			{
-				OutLog += TEXT("-");
-			}
-			OutLog += TEXT("\n");
-		}
-
-		return OutLog;
-	}
-};
+DEFINE_STAT(STAT_ComputeConstraints);
+DEFINE_STAT(STAT_ComputeConstraintsNP);
+DEFINE_STAT(STAT_ComputeConstraintsBP);
 
-#endif
-
-DECLARE_CYCLE_STAT(TEXT("ComputeConstraints"), STAT_ComputeConstraints, STATGROUP_Chaos);
-DECLARE_CYCLE_STAT(TEXT("ComputeConstraintsNP"), STAT_ComputeConstraintsNP, STATGROUP_Chaos);
-
-
-int32 ComputeConstraintsUseAny2 = 1;
-FAutoConsoleVariableRef CVarComputeConstraintsUseAny2(TEXT("p.ComputeConstraintsUseAny2"), ComputeConstraintsUseAny2, TEXT(""));
-
-int32 CollisionConstraintsForceSingleThreaded = 0;
+CHAOS_API int32 CollisionConstraintsForceSingleThreaded = 0;
 FAutoConsoleVariableRef CVarCollisionConstraintsForceSingleThreaded(TEXT("p.Chaos.Collision.ForceSingleThreaded"), CollisionConstraintsForceSingleThreaded, TEXT("CollisionConstraintsForceSingleThreaded"));
 
-template <typename T, int d>
-struct TSimOverlapVisitor
-{
-	TArray<TAccelerationStructureHandle<T, d>>& Intersections;
-	TSimOverlapVisitor(TArray<TAccelerationStructureHandle<T, d>>& InIntersections)
-		: Intersections(InIntersections)
-	{
-	}
-
-	bool VisitOverlap(const TSpatialVisitorData<TAccelerationStructureHandle<T, d>>& Instance)
-	{
-		Intersections.Add(Instance.Payload);
-		return true;
-	}
-
-	bool VisitSweep(TSpatialVisitorData<TAccelerationStructureHandle<T,d>>, T Length)
-	{
-		check(false);
-		return false;
-	}
-
-	bool VisitRaycast(TSpatialVisitorData<TAccelerationStructureHandle<T, d>>, T Length)
-	{
-		check(false);
-		return false;
-	}
-};
-
-template <typename T, int d>
-template <typename SPATIAL_ACCELERATION, bool bGatherStats>
-void TPBDCollisionConstraint<T, d>::ComputeConstraintsHelper(T Dt, const SPATIAL_ACCELERATION& InSpatialAcceleration)
-{
-	SCOPE_CYCLE_COUNTER(STAT_ComputeConstraints);
-	CHAOS_SCOPED_TIMER(ComputeConstraints);
-	if (!EnableCollisions) return;
-	// Broad phase
-
-#if !UE_BUILD_SHIPPING
-	FStatHelper2<bGatherStats> BroadphasePotentials;
-	FStatHelper2<bGatherStats> NarrowPhaseSkipped;
-	FStatHelper2<bGatherStats> NarrowPhasePerformed;
-	FStatHelper2<bGatherStats> NarrowPhaseRejected;
-	FStatHelper2<bGatherStats> BoundsDistribution;
-	int32 SimulatedParticles = 0;
-#endif
-	{
-		SCOPE_CYCLE_COUNTER(STAT_ComputeConstraintsNP);
-
-
-		// Narrow phase
-		CHAOS_SCOPED_TIMER(ComputeConstraints_NP);
-
-		TQueue<TRigidBodyContactConstraint<T, d>, EQueueMode::Mpsc> Queue;	//todo(ocohen): use per thread buffer instead, need better support than ParallelFor for this
-		Particles.GetNonDisabledDynamicView().ParallelFor([&](auto& Particle1, int32 ActiveIdxIdx)
-		{
-#if !UE_BUILD_SHIPPING
-			if (bGatherStats)
-			{
-				++SimulatedParticles;
-			}
-#endif
-
-			TArray<TAccelerationStructureHandle<T,d>> PotentialIntersections;
-
-			const bool bBody1Bounded = HasBoundingBox(Particle1);
-
-			const T Box1Thickness = ComputeThickness(Particle1, Dt).Size();
-			if (bBody1Bounded)
-			{
-#if CHAOS_PARTICLEHANDLE_TODO
-				const TBox <T, d> Box1 = InSpatialAcceleration.GetWorldSpaceBoundingBox(Particle1);
-#else
-				const TBox <T, d> Box1 = ComputeWorldSpaceBoundingBox(Particle1);	//NOTE: this ignores the velocity expansion which is wrong
-#endif
-
-#if !UE_BUILD_SHIPPING
-				if (bGatherStats)
-				{
-					//const TBox<T, d> Box1 = SpatialAcceleration.GetWorldSpaceBoundingBox(InParticles, Body1Index);
-
-					BoundsDistribution.Record(Box1.Extents().GetMax());
-				}
-#endif
-				TSimOverlapVisitor<T,d> OverlapVisitor(PotentialIntersections);
-				InSpatialAcceleration.Overlap(Box1, OverlapVisitor);
-			}
-			else
-			{
-				const auto& GlobalElems = InSpatialAcceleration.GlobalObjects();
-				PotentialIntersections.Reserve(GlobalElems.Num());
-
-				for (auto& Elem : GlobalElems)
-				{
-					PotentialIntersections.Add(Elem.Payload);
-				}
-			}
-
-#if !UE_BUILD_SHIPPING
-			BroadphasePotentials.Record(PotentialIntersections.Num());
-
-			int32 CountNP = 0;
-			int32 RejectedNP = 0;
-#endif
-			const int32 NumPotentials = PotentialIntersections.Num();
-			for (int32 i = 0; i < NumPotentials; ++i)
-			{
-				auto& Particle2 = *PotentialIntersections[i].GetGeometryParticleHandle_PhysicsThread();
-				const TGenericParticleHandle<T, d> Particle2Generic(&Particle2);
-				// Collision group culling...
-				// CollisionGroup == 0 : Collide_With_Everything
-				// CollisionGroup == INDEX_NONE : Disabled collisions
-				// CollisionGroup_A != CollisionGroup_B : Skip Check
-				
-				if (Particle1.CollisionGroup() == INDEX_NONE || Particle2Generic->CollisionGroup() == INDEX_NONE)
-				{
-					continue;
-				}
-				if (Particle1.CollisionGroup() && Particle2Generic->CollisionGroup() && Particle1.CollisionGroup() != Particle2Generic->CollisionGroup())
-				{
-					continue;
-				}
-
-				if (!Particle1.Geometry() && !Particle2.Geometry())
-				{
-					continue;
-				}
-
-				const bool bBody2Bounded = HasBoundingBox(Particle2);
-
-				if (Particle1.Handle() == Particle2.Handle())
-				{
-					continue;
-				}
-				
-				if(bBody1Bounded == bBody2Bounded && Particle2.AsDynamic())
-				{
-					//if both are dynamic, assume index order matters
-#if CHAOS_DETERMINISTIC
-					if(Particle2.ParticleID() > Particle1.ParticleID())
-#else
-					//not deterministic just use memory address to avoid pair duplication
-					if (Particle2.Handle() > Particle1.Handle())
-#endif
-					{
-						continue;
-					}
-				}
-				//				++TotalTested;
-
-				const TVector<T, d> Box2Thickness = Particle2.AsDynamic() ? ComputeThickness(*Particle2.AsDynamic(), Dt) : TVector<T, d>(0);
-				const T UseThickness = FMath::Max(Box1Thickness, Box2Thickness.Size());// + MThickness
-
-				auto Constraint = ComputeConstraint(Particle1.Handle(), Particle2.Handle(), UseThickness);
-
-				//if (true || !InParticles.Geometry(Body1Index)->HasBoundingBox() || !InParticles.Geometry(Body2Index)->HasBoundingBox())
-				{
-					//SCOPE_CYCLE_COUNTER(STAT_ComputeConstraintsNP3);
-					//use narrow phase to determine if constraint is needed. Without this we can't do shock propagation
-
-#if !UE_BUILD_SHIPPING
-					if (bGatherStats)
-					{
-						++CountNP;
-					}
-#endif
-					if (ComputeConstraintsUseAny2)
-					{
-						UpdateConstraint<ECollisionUpdateType::Any>(UseThickness, Constraint);
-					}
-					else
-					{
-						UpdateConstraint<ECollisionUpdateType::Deepest>(UseThickness, Constraint);
-					}
-
-					if (Constraint.Phi < UseThickness)
-					{
-						Queue.Enqueue(Constraint);
-					}
-					else
-					{
-#if !UE_BUILD_SHIPPING
-						++RejectedNP;
-#endif
-					}
-				}
-
-			}
-
-#if !UE_BUILD_SHIPPING
-			NarrowPhasePerformed.Record(CountNP);
-			const int32 NPSkipped = NumPotentials - CountNP;
-			NarrowPhaseSkipped.Record(NPSkipped);
-			NarrowPhaseRejected.Record(RejectedNP);
-#endif
-		}, bGatherStats || CollisionConstraintsForceSingleThreaded);
-		
-		while(!Queue.IsEmpty())
-		{				
-			int32 Idx = Constraints.AddUninitialized(1);
-			Queue.Dequeue(Constraints[Idx]);
-			Handles.Add(HandleAllocator.AllocHandle(this, Idx));
-		}
-	}
-
-#if CHAOS_DETERMINISTIC
-	//todo: sort constraints
-#endif
-
-#if !UE_BUILD_SHIPPING
-	FString OutLog;
-	if (bGatherStats)
-	{
-#if CHAOS_PARTICLEHANDLE_TODO
-		const float NumParticles = InParticles.Size();
-		OutLog = FString::Printf(TEXT("ComputeConstraints stats:\n"
-			"Total Particles:%d\nSimulated Particles:%d (%.2f%%)\n"
-			"Max candidates per instance:%d (%.2f%%)\n"
-			"Max candidates skipped per instance (NP skipped):%d (%.2f%%)\n"
-			"Max narrow phase tests per instance:%d (%.2f%%)\n"
-			"Max narrow phase rejected per instance (NP rejected):%d (%.2f%%)\n"
-			"Constraints generated:%d\n"
-		),
-			InParticles.Size(),
-			SimulatedParticles, SimulatedParticles / NumParticles * 100.f,
-			BroadphasePotentials.MaxCount, BroadphasePotentials.MaxCount / NumParticles * 100.f,
-			NarrowPhaseSkipped.MaxCount, NarrowPhaseSkipped.MaxCount / NumParticles * 100.f,
-			NarrowPhasePerformed.MaxCount, NarrowPhasePerformed.MaxCount / NumParticles * 100.f,
-			NarrowPhaseRejected.MaxCount, NarrowPhaseRejected.MaxCount / NumParticles * 100.f,
-			Constraints.Num()
-		);
-
-		OutLog += FString::Printf(TEXT("Potentials per instance distribution:\n"));
-		OutLog += BroadphasePotentials.ToString();
-
-		
-		OutLog += FString::Printf(TEXT("\nCandidates skipped per instance (NP skipped) distribution:\n"));
-		OutLog += NarrowPhaseSkipped.ToString();
-		
-		OutLog += FString::Printf(TEXT("\nNarrow phase performed per instance distribution:\n"));
-		OutLog += NarrowPhasePerformed.ToString();
-
-		OutLog += FString::Printf(TEXT("\nNarrow phase candidates rejected per instance distribution:\n"));
-		OutLog += NarrowPhaseRejected.ToString();
-
-		OutLog += FString::Printf(TEXT("\nBounds distribution:\n"));
-		OutLog += BoundsDistribution.ToString();
-
-		UE_LOG(LogChaos, Warning, TEXT("%s"), *OutLog);
-		PendingHierarchyDump = 0;
-#endif
-	}
-#endif
-	//	if (TotalTested > 0)
-	//	{
-			//UE_LOG(LogChaos, Warning, TEXT("ComputeConstraints: rejected:%f out of rejected:%f = %.3f"), TotalRejected, TotalTested, TotalRejected / TotalTested);
-	//	}
-}
-
 template<typename T, int d>
 template <bool bGatherStats>
 void TPBDCollisionConstraint<T, d>::ComputeConstraints(const FAccelerationStructure& AccelerationStructure, T Dt)
 {
 	if (const auto AABBTree = AccelerationStructure.template As<TAABBTree<TAccelerationStructureHandle<T, d>, TAABBTreeLeafArray<TAccelerationStructureHandle<T, d>, T>, T>>())
 	{
-		ComputeConstraintsHelper<decltype(*AABBTree), bGatherStats>(Dt, *AABBTree);
+		ComputeConstraintsHelperLowLevel<bGatherStats>(*AABBTree, Dt);
 	}
 	else if(const auto BV = AccelerationStructure.template As<TBoundingVolume<TAccelerationStructureHandle<T, d>, T, d>>())
 	{
-		ComputeConstraintsHelper<decltype(*BV), bGatherStats>(Dt, *BV);
+		ComputeConstraintsHelperLowLevel<bGatherStats>(*BV, Dt);
 	}
 	else if (const auto AABBTreeBV = AccelerationStructure.template As<TAABBTree<TAccelerationStructureHandle<T, d>, TBoundingVolume<TAccelerationStructureHandle<T, d>, T, d>, T>>())
 	{
-		ComputeConstraintsHelper<decltype(*AABBTreeBV), bGatherStats>(Dt, *AABBTreeBV);
+		ComputeConstraintsHelperLowLevel<bGatherStats>(*AABBTreeBV, Dt);
+	}
+	else if(const auto Collection = AccelerationStructure.template As<ISpatialAccelerationCollection<TAccelerationStructureHandle<T, d>, T, d>>())
+	{
+		if (bGatherStats)
+		{
+			Collection->PBDComputeConstraintsLowLevel_GatherStats(*this, Dt);
+		}
+		else
+		{
+			Collection->PBDComputeConstraintsLowLevel(*this, Dt);
+		}
 	}
 	else
 	{
-		check(false); //question: do we want to support a dynamic dispatch version?
+		check(false);  //question: do we want to support a dynamic dispatch version?
 	}
 	
 	if (PostComputeCallback != nullptr)
diff --git a/Engine/Source/Runtime/Experimental/Chaos/Private/Chaos/PBDRigidsEvolution.cpp b/Engine/Source/Runtime/Experimental/Chaos/Private/Chaos/PBDRigidsEvolution.cpp
--- a/Engine/Source/Runtime/Experimental/Chaos/Private/Chaos/PBDRigidsEvolution.cpp
+++ b/Engine/Source/Runtime/Experimental/Chaos/Private/Chaos/PBDRigidsEvolution.cpp
@@ -5,6 +5,7 @@
 #include "ChaosStats.h"
 #include "Chaos/PBDRigidsEvolutionGBF.h"
 #include "Chaos/ParticleHandle.h"
+#include "Chaos/ISpatialAccelerationCollection.h"
 
 namespace Chaos
 {
@@ -28,7 +29,17 @@ namespace Chaos
 		ComputeIntermediateSpatialAcceleration();
 	}
 
-	DECLARE_CYCLE_STAT(TEXT("ComputeConstraintsBP"), STAT_ComputeConstraintsBP, STATGROUP_Chaos);
+	template<class FPBDRigidsEvolution, class FPBDCollisionConstraint, class T, int d>
+	TPBDRigidsEvolutionBase<FPBDRigidsEvolution, FPBDCollisionConstraint, T, d>::~TPBDRigidsEvolutionBase()
+	{
+		Particles.GetParticleHandles().RemoveArray(&PhysicsMaterials);
+		Particles.GetParticleHandles().RemoveArray(&PerParticlePhysicsMaterials);
+		Particles.GetParticleHandles().RemoveArray(&ParticleDisableCount);
+		Particles.GetParticleHandles().RemoveArray(&Collided);
+		WaitOnAccelerationStructure();
+	}
+
+	
 	DECLARE_CYCLE_STAT(TEXT("CacheAccelerationBounds"), STAT_CacheAccelerationBounds, STATGROUP_Chaos);
 	DECLARE_CYCLE_STAT(TEXT("ComputeIntermediateSpatialAcceleration"), STAT_ComputeIntermediateSpatialAcceleration, STATGROUP_Chaos);
 
@@ -82,19 +93,58 @@ namespace Chaos
 	FAutoConsoleVariableRef CVarMaxPayloadSize(TEXT("p.MaxPayloadSize"), MaxPayloadSize, TEXT(""));
 
 	template<class FPBDRigidsEvolution, class FPBDCollisionConstraint, class T, int d>
-	ISpatialAcceleration<TAccelerationStructureHandle<T, d>, T, d>*  TPBDRigidsEvolutionBase<FPBDRigidsEvolution, FPBDCollisionConstraint, T, d>::CreateNewSpatialStructure(const TArray<FAccelerationStructureBuilder>& Bounds)
+	ISpatialAccelerationCollection<TAccelerationStructureHandle<T, d>, T, d>*  TPBDRigidsEvolutionBase<FPBDRigidsEvolution, FPBDCollisionConstraint, T, d>::CreateNewSpatialStructure(const TArray<FAccelerationStructureBuilder>& Bounds)
 	{
 		if (BroadphaseType == 0)
 		{
-			return new TBoundingVolume<TAccelerationStructureHandle<T, d>, T, d>(Bounds, false, 0, BoundingVolumeNumCells, MaxPayloadSize);
+			using AccelType = TBoundingVolume<TAccelerationStructureHandle<T, d>, T, d>;
+			auto Structure = MakeUnique<AccelType>(Bounds, false, 0, BoundingVolumeNumCells, MaxPayloadSize);
+			auto Collection = new TSpatialAccelerationCollection<AccelType>();
+			Collection->AddSubstructure(MoveTemp(Structure), 0);
+			return Collection;
 		}
 		else if(BroadphaseType == 1)
 		{
-			return new TAABBTree<TAccelerationStructureHandle<T, d>, TAABBTreeLeafArray<TAccelerationStructureHandle<T, d>, T>, T>(Bounds, MaxChildrenInLeaf, MaxTreeDepth, MaxPayloadSize);
+			using AccelType = TAABBTree<TAccelerationStructureHandle<T, d>, TAABBTreeLeafArray<TAccelerationStructureHandle<T, d>, T>, T>;
+			auto Structure = MakeUnique<AccelType>(Bounds, MaxChildrenInLeaf, MaxTreeDepth, MaxPayloadSize);
+			auto Collection = new TSpatialAccelerationCollection<AccelType>();
+			Collection->AddSubstructure(MoveTemp(Structure), 0);
+			return Collection;
 		}
 		else
 		{
-			return new TAABBTree<TAccelerationStructureHandle<T, d>, TBoundingVolume<TAccelerationStructureHandle<T, d>, T, d>, T>(Bounds, AABBMaxChildrenInLeaf, AABBMaxTreeDepth, MaxPayloadSize);
+			using AccelType = TAABBTree<TAccelerationStructureHandle<T, d>, TBoundingVolume<TAccelerationStructureHandle<T, d>, T, d>, T>;
+			auto Structure = MakeUnique<AccelType>(Bounds, AABBMaxChildrenInLeaf, AABBMaxTreeDepth, MaxPayloadSize);
+			auto Collection = new TSpatialAccelerationCollection<AccelType>();
+			Collection->AddSubstructure(MoveTemp(Structure), 0);
+			return Collection;
+		}
+	}
+
+	template<typename T, int d>
+	TUniquePtr<ISpatialAccelerationCollection<TAccelerationStructureHandle<T, d>, T, d>> CreateNewSpatialStructureFromSubStructure(TUniquePtr<ISpatialAcceleration<TAccelerationStructureHandle<T,d>, T,d>>&& Substructure)
+	{
+		using BVType = TBoundingVolume<TAccelerationStructureHandle<T, d>, T, d>;
+		using AABBType = TAABBTree<TAccelerationStructureHandle<T, d>, TAABBTreeLeafArray<TAccelerationStructureHandle<T, d>, T>, T>;
+
+		if (Substructure->template As<BVType>())
+		{
+			auto Collection = MakeUnique<TSpatialAccelerationCollection<BVType>>();
+			Collection->AddSubstructure(MoveTemp(Substructure), 0);
+			return Collection;
+		}
+		else if (Substructure->template As<AABBType>())
+		{
+			auto Collection = MakeUnique<TSpatialAccelerationCollection<AABBType>>();
+			Collection->AddSubstructure(MoveTemp(Substructure), 0);
+			return Collection;
+		}
+		else
+		{
+			using AccelType = TAABBTree<TAccelerationStructureHandle<T, d>, TBoundingVolume<TAccelerationStructureHandle<T, d>, T, d>, T>;
+			auto Collection = MakeUnique<TSpatialAccelerationCollection<AccelType>>();
+			Collection->AddSubstructure(MoveTemp(Substructure), 0);
+			return Collection;
 		}
 	}
 
@@ -102,7 +152,7 @@ namespace Chaos
 	void TPBDRigidsEvolutionBase<FPBDRigidsEvolution, FPBDCollisionConstraint, T, d>::FChaosAccelerationStructureTask::DoTask(ENamedThreads::Type CurrentThread, const FGraphEventRef& MyCompletionGraphEvent)
 	{
 		AccelerationStructure.Reset(CreateNewSpatialStructure(CachedSpatialBuilderData));
-		AccelerationStructureCopy = AccelerationStructure->Copy();
+		AccelerationStructureCopy = AsUniqueSpatialAccelerationChecked<FAccelerationStructure>(AccelerationStructure->Copy());
 	}
 
 	template<class FPBDRigidsEvolution, class FPBDCollisionConstraint, class T, int d>
@@ -116,7 +166,7 @@ namespace Chaos
 		//As long as we delete first and update second this will be respected
 		if (SpatialData.bDelete)
 		{
-			AccelerationStructure.RemoveElement(SpatialData.AccelerationHandle);
+			AccelerationStructure.RemoveElementFrom(SpatialData.AccelerationHandle, SpatialData.SpatialIdx);
 
 			if (bAsync)
 			{
@@ -138,7 +188,7 @@ namespace Chaos
 
 		if(SpatialData.bUpdate)
 		{
-			AccelerationStructure.UpdateElement(Particle, Particle->WorldSpaceInflatedBounds(), Particle->HasBounds());
+			AccelerationStructure.UpdateElementIn(Particle, Particle->WorldSpaceInflatedBounds(), Particle->HasBounds(), SpatialData.SpatialIdx);	//todo: can the spatial idx change when we recreate and re-use memory? (yes, potential bug)
 			
 			if (bAsync)
 			{
@@ -215,9 +265,12 @@ namespace Chaos
 		CHAOS_SCOPED_TIMER(ComputeIntermediateSpatialAcceleration);
 		if (!AccelerationStructureTaskComplete)
 		{
-			//initial frame so make empty structure
+			//initial frame so make empty structures
+
 			InternalAcceleration = TUniquePtr<FAccelerationStructure>(CreateNewSpatialStructure(TArray<FAccelerationStructureBuilder>()));
 			ScratchExternalAcceleration = TUniquePtr<FAccelerationStructure>(CreateNewSpatialStructure(TArray<FAccelerationStructureBuilder>()));
+			AsyncInternalAcceleration = TUniquePtr<FAccelerationStructure>(CreateNewSpatialStructure(TArray<FAccelerationStructureBuilder>()));
+			AsyncExternalAcceleration = TUniquePtr<FAccelerationStructure>(CreateNewSpatialStructure(TArray<FAccelerationStructureBuilder>()));
 			FlushInternalAccelerationQueue();
 			FlushExternalAccelerationQueue(*ScratchExternalAcceleration);
 			bExternalReady = true;
@@ -249,7 +302,7 @@ namespace Chaos
 	}
 
 	template<class FPBDRigidsEvolution, class FPBDCollisionConstraint, class T, int d>
-	void TPBDRigidsEvolutionBase<FPBDRigidsEvolution, FPBDCollisionConstraint, T, d>::UpdateExternalAccelerationStructure(TUniquePtr<ISpatialAcceleration<TAccelerationStructureHandle<T, d>, T, d>>& StructToUpdate)
+	void TPBDRigidsEvolutionBase<FPBDRigidsEvolution, FPBDCollisionConstraint, T, d>::UpdateExternalAccelerationStructure(TUniquePtr<ISpatialAccelerationCollection<TAccelerationStructureHandle<T, d>, T, d>>& StructToUpdate)
 	{
 		DECLARE_SCOPE_CYCLE_COUNTER(TEXT("CreateExternalAccelerationStructure"), STAT_CreateExternalAccelerationStructure, STATGROUP_Physics);
 		if (bExternalReady)
@@ -263,6 +316,49 @@ namespace Chaos
 			FlushExternalAccelerationQueue(*StructToUpdate);
 		}
 	}
+
+	template<class FPBDRigidsEvolution, class FPBDCollisionConstraint, class T, int d>
+	void TPBDRigidsEvolutionBase<FPBDRigidsEvolution, FPBDCollisionConstraint, T, d>::Serialize(FChaosArchive& Ar)
+	{
+		Particles.Serialize(Ar);
+
+		Ar.UsingCustomVersion(FExternalPhysicsCustomObjectVersion::GUID);
+		if (Ar.CustomVer(FExternalPhysicsCustomObjectVersion::GUID) >= FExternalPhysicsCustomObjectVersion::SerializeEvolutionBV)
+		{
+			//for now just assume a single structure
+			TUniquePtr<ISpatialAcceleration<TAccelerationStructureHandle<T, d>, T, d>> SubStructure;
+			if (!Ar.IsLoading())
+			{
+				SubStructure = InternalAcceleration->RemoveSubstructure(FSpatialAccelerationIdx{ 0,0 });
+				Ar << SubStructure;
+				InternalAcceleration->AddSubstructure(MoveTemp(SubStructure), 0);
+			}
+			else
+			{
+				Ar << SubStructure;
+				InternalAcceleration = CreateNewSpatialStructureFromSubStructure(MoveTemp(SubStructure));
+			}
+
+			SerializePendingMap(Ar, InternalAccelerationQueue);
+			SerializePendingMap(Ar, AsyncAccelerationQueue);
+			SerializePendingMap(Ar, ExternalAccelerationQueue);
+
+			ScratchExternalAcceleration = AsUniqueSpatialAccelerationChecked<FAccelerationStructure>(InternalAcceleration->Copy());
+		}
+		else if (Ar.IsLoading())
+		{
+			AccelerationStructureTaskComplete = nullptr;
+			for (auto& Particle : Particles.GetNonDisabledView())
+			{
+				DirtyParticle(Particle);
+			}
+
+			//force build acceleration structure with latest data
+			ComputeIntermediateSpatialAcceleration(true);
+			ComputeIntermediateSpatialAcceleration(true);	//having to do it multiple times because of the various caching involved over multiple frames.
+			ComputeIntermediateSpatialAcceleration(true);
+		}
+	}
 }
 
 #ifdef __clang__
diff --git a/Engine/Source/Runtime/Experimental/Chaos/Public/Chaos/ISpatialAcceleration.h b/Engine/Source/Runtime/Experimental/Chaos/Public/Chaos/ISpatialAcceleration.h
--- a/Engine/Source/Runtime/Experimental/Chaos/Public/Chaos/ISpatialAcceleration.h
+++ b/Engine/Source/Runtime/Experimental/Chaos/Public/Chaos/ISpatialAcceleration.h
@@ -3,7 +3,6 @@
 #include "Chaos/Vector.h"
 #include "Chaos/Box.h"
 #include "GeometryParticlesfwd.h"
-#include <tuple>
 
 namespace Chaos
 {
@@ -113,6 +112,7 @@ enum ESpatialAcceleration
 	BoundingVolume,
 	AABBTree,
 	AABBTreeBV,
+	Collection,
 	Unknown
 };
 
@@ -187,12 +187,6 @@ public:
 		return nullptr;
 	}
 
-	/** Begin Root-level API. This is kind of a hack to avoid virtuals. We simply route calls into templated functions */
-	virtual void PBDComputeConstraintsHelper() const
-	{
-		check(false);	//not implemented
-	}
-
 #if !UE_BUILD_SHIPPING
 	virtual void DebugDraw(ISpacialDebugDrawInterface<T>* InInterface) const {}
 	virtual void DumpStats() const {}
@@ -230,29 +224,28 @@ public:
 		return static_cast<const TConcrete&>(*this);
 	}
 
+private:
+	ESpatialAcceleration Type;
+};
 
-	template <typename TDerived>
-	static TUniquePtr<TDerived> AsUnique(TUniquePtr<ISpatialAcceleration<TPayloadType, T, d>>&& Base)
-	{
-		if (TDerived* Derived = Base->template As<TDerived>())
-		{
-			Base.Release();
-			return TUniquePtr<TDerived>(Derived);
-		}
-		return nullptr;
-	}
-
-	template <typename TDerived>
-	static TUniquePtr<TDerived> AsUniqueChecked(TUniquePtr<ISpatialAcceleration<TPayloadType, T, d>>&& Base)
+template <typename TBase, typename TDerived>
+static TUniquePtr<TDerived> AsUniqueSpatialAcceleration(TUniquePtr<TBase>&& Base)
+{
+	if (TDerived* Derived = Base->template As<TDerived>())
 	{
-		TDerived& Derived = Base->template AsChecked<TDerived>();
 		Base.Release();
-		return TUniquePtr<TDerived>(&Derived);
+		return TUniquePtr<TDerived>(Derived);
 	}
+	return nullptr;
+}
 
-private:
-	ESpatialAcceleration Type;
-};
+template <typename TDerived, typename TBase>
+static TUniquePtr<TDerived> AsUniqueSpatialAccelerationChecked(TUniquePtr<TBase>&& Base)
+{
+	TDerived& Derived = Base->template AsChecked<TDerived>();
+	Base.Release();
+	return TUniquePtr<TDerived>(&Derived);
+}
 
 /** Helper class used to bridge virtual to template implementation of acceleration structures */
 template <typename TPayloadType, typename T>
@@ -280,348 +273,4 @@ private:
 	ISpatialVisitor<TPayloadType, T>& Visitor;
 };
 
-template <typename TPayloadType, typename T, int d>
-class CHAOS_API ISpatialAccelerationCollection : public ISpatialAcceleration<TPayloadType, T, d>
-{
-public:
-	virtual FSpatialAccelerationIdx AddSubstructure(TUniquePtr<ISpatialAcceleration<TPayloadType, T, d>>&& Substructure, int32 Bucket) = 0;
-	virtual void RemoveSubstructure(FSpatialAccelerationIdx Idx) = 0;
-	virtual ISpatialAcceleration<TPayloadType, T, d>* GetSubstructure(FSpatialAccelerationIdx Idx) = 0;
-};
-
-template <typename TSpatialAcceleration>
-class CHAOS_API TSpatialAccelerationCollectionBucket
-{
-public:
-	using AccelerationType = TSpatialAcceleration;
-	using TPayloadType = typename AccelerationType::PayloadType;
-	using T = typename AccelerationType::TType;
-	static constexpr int d = AccelerationType::D;
-
-	TSpatialAccelerationCollectionBucket() = default;
-
-	//question: should this be private and make higher up code more explicit? kind of annoying because of tuple
-	TSpatialAccelerationCollectionBucket(const TSpatialAccelerationCollectionBucket<TSpatialAcceleration>& Other)
-		: FreeIndices(Other.FreeIndices)
-	{
-		Accelerations.Reserve(Other.Accelerations.Num());
-		for (const auto& OtherAccel : Other.Accelerations)
-		{
-			if (OtherAccel)
-			{
-				Accelerations.Add(TSpatialAcceleration::template AsUniqueChecked<TSpatialAcceleration>(OtherAccel->Copy()));
-			}
-			else
-			{
-				Accelerations.Add(nullptr);
-			}
-		}
-	}
-
-	uint16 AddSubstructure(TUniquePtr<TSpatialAcceleration>&& Substructure)
-	{
-		uint16 NewIdx = NewSlot();
-		Accelerations[NewIdx] = MoveTemp(Substructure);
-		
-		return NewIdx;
-	}
-	
-	void RemoveSubstructure(uint16 Idx)
-	{
-		ReleaseSlot(Idx);
-	}
-
-	TSpatialAcceleration* GetSubstructure(int32 Idx)
-	{
-		return Accelerations[Idx].Get();
-	}
-
-	void Raycast(const TVector<T, d>& Start, const TVector<T, d>& Dir, const T OriginalLength, ISpatialVisitor<TPayloadType, T>& Visitor) const
-	{
-		for (const auto& Acceleration : Accelerations)
-		{
-			if (Acceleration)
-			{
-				Acceleration->Raycast(Start, Dir, OriginalLength, Visitor);
-			}
-		}
-	}
-
-	void Sweep(const TVector<T, d>& Start, const TVector<T, d>& Dir, const T OriginalLength, const TVector<T,3>& QueryHalfExtents, ISpatialVisitor<TPayloadType, T>& Visitor) const
-	{
-		for (const auto& Acceleration : Accelerations)
-		{
-			if (Acceleration)
-			{
-				Acceleration->Sweep(Start, Dir, OriginalLength, QueryHalfExtents, Visitor);
-			}
-		}
-	}
-
-	void Overlap(const TBox<T,d>& QueryBounds, ISpatialVisitor<TPayloadType, T>& Visitor) const
-	{
-		for (const auto& Acceleration : Accelerations)
-		{
-			if (Acceleration)
-			{
-				Acceleration->Overlap(QueryBounds, Visitor);
-			}
-		}
-	}
-
-	void RemoveElement(const TPayloadType& Payload, int32 Idx)
-	{
-		Accelerations[Idx]->RemoveElement(Payload);
-	}
-
-	void UpdateElement(const TPayloadType& Payload, const TBox<T, d>& NewBounds, bool bHasBounds, int32 Idx)
-	{
-		Accelerations[Idx]->UpdateElement(Payload, NewBounds, bHasBounds);
-	}
-
-private:
-	uint16 NewSlot()
-	{
-		if (FreeIndices.Num())
-		{
-			return FreeIndices.Pop(false);
-		}
-		else if (Accelerations.Num() < FSpatialAccelerationIdx::MaxBucketEntries)
-		{
-			return Accelerations.Add(nullptr);
-		}
-		check(false);	//too many substructures in a single bucket
-		return 0;
-	}
-
-	void ReleaseSlot(uint16 Index)
-	{
-		Accelerations[Index] = nullptr;
-		FreeIndices.Add(Index);
-		if (Index + 1 == Accelerations.Num())
-		{
-			Accelerations.Pop(false);
-		}
-	}
-
-	TArray<TUniquePtr<TSpatialAcceleration>> Accelerations;
-	TArray<uint16> FreeIndices;
-};
-
-template <typename... TRemaining>
-struct CHAOS_API TSpatialBucketTuple
-{
-};
-
-template <typename TAcceleration, typename... TRemaining>
-struct CHAOS_API TSpatialBucketTuple<TAcceleration, TRemaining...>
-{
-	using FirstType = TAcceleration;
-	TSpatialAccelerationCollectionBucket<TAcceleration> First;
-	TSpatialBucketTuple<TRemaining...> Remaining;
-};
-
-template<int Idx, typename ... Rest>
-struct CHAOS_API TSpatialBucketTupleGetter
-{
-
-};
-
-template<int Idx, typename First, typename... Rest>
-struct CHAOS_API TSpatialBucketTupleGetter<Idx, First, Rest...>
-{
-	static auto& Get(TSpatialBucketTuple<First, Rest...>& Buckets) { return TSpatialBucketTupleGetter<Idx - 1, Rest...>::Get(Buckets.Remaining); }
-	static const auto& Get(const TSpatialBucketTuple<First, Rest...>& Buckets) { return TSpatialBucketTupleGetter<Idx - 1, Rest...>::Get(Buckets.Remaining); }
-};
-
-template<typename First, typename... Rest>
-struct CHAOS_API TSpatialBucketTupleGetter<0, First, Rest...>
-{
-	static auto& Get(TSpatialBucketTuple<First, Rest...>& Buckets) { return Buckets.First; }
-	static const auto& Get(const TSpatialBucketTuple<First, Rest...>& Buckets) { return Buckets.First; }
-};
-
-template <int Idx, typename First, typename... Rest>
-auto& GetBucket(TSpatialBucketTuple<First, Rest...>& Buckets)
-{
-	return TSpatialBucketTupleGetter<Idx, First, Rest...>::Get(Buckets);
-}
-
-template <int Idx, typename First, typename... Rest>
-const auto& GetBucket(const TSpatialBucketTuple<First, Rest...>& Buckets)
-{
-	return TSpatialBucketTupleGetter<Idx, First, Rest...>::Get(Buckets);
-}
-
-template <int BucketIdx, int NumBuckets, typename Tuple, typename TPayloadType, typename T, int d>
-struct TSpatialAccelerationCollectionHelper
-{
-	static void Raycast(const Tuple& Buckets, const TVector<T, d>& Start, const TVector<T, d>& Dir, const T OriginalLength, ISpatialVisitor<TPayloadType, T>& Visitor)
-	{
-		GetBucket<BucketIdx>(Buckets).Raycast(Start, Dir, OriginalLength, Visitor);
-		constexpr int NextBucket = BucketIdx + 1;
-		if (NextBucket < NumBuckets)
-		{
-			TSpatialAccelerationCollectionHelper<NextBucket < NumBuckets ? NextBucket : 0, NumBuckets, Tuple, TPayloadType, T, d>::Raycast(Buckets, Start, Dir, OriginalLength, Visitor);
-		}
-	}
-
-	static void Sweep(const Tuple& Buckets, const TVector<T, d>& Start, const TVector<T, d>& Dir, const T OriginalLength, const TVector<T, d> QueryHalfExtents, ISpatialVisitor<TPayloadType, T>& Visitor)
-	{
-		GetBucket<BucketIdx>(Buckets).Sweep(Start, Dir, OriginalLength, QueryHalfExtents, Visitor);
-		constexpr int NextBucket = BucketIdx + 1;
-		if (NextBucket < NumBuckets)
-		{
-			TSpatialAccelerationCollectionHelper < NextBucket < NumBuckets ? NextBucket : 0, NumBuckets, Tuple, TPayloadType, T, d>::Sweep(Buckets, Start, Dir, OriginalLength, QueryHalfExtents, Visitor);
-		}
-	}
-
-	static void Overlap(const Tuple& Buckets, const TBox<T, d> QueryBounds, ISpatialVisitor<TPayloadType, T>& Visitor)
-	{
-		GetBucket<BucketIdx>(Buckets).Overlap(QueryBounds, Visitor);
-		constexpr int NextBucket = BucketIdx + 1;
-		if (NextBucket < NumBuckets)
-		{
-			TSpatialAccelerationCollectionHelper < NextBucket < NumBuckets ? NextBucket : 0, NumBuckets, Tuple, TPayloadType, T, d>::Overlap(Buckets, QueryBounds, Visitor);
-		}
-	}
-};
-
-template <typename ... TSpatialAccelerationBuckets>
-class CHAOS_API TSpatialAccelerationCollection : public
-	ISpatialAccelerationCollection<typename std::tuple_element<0, std::tuple< TSpatialAccelerationBuckets...>>::type::PayloadType,
-	typename std::tuple_element<0, std::tuple< TSpatialAccelerationBuckets...>>::type::TType,
-	std::tuple_element<0, std::tuple< TSpatialAccelerationBuckets...>>::type::D>
-{
-public:
-	using FirstAccelerationType = typename std::tuple_element<0, std::tuple< TSpatialAccelerationBuckets...>>::type;
-	using TPayloadType = typename FirstAccelerationType::PayloadType;
-	using T = typename FirstAccelerationType::TType;
-	static constexpr int d = FirstAccelerationType::D;
-	TSpatialAccelerationCollection() = default;
-
-	virtual FSpatialAccelerationIdx AddSubstructure(TUniquePtr<ISpatialAcceleration<TPayloadType, T, d>>&& Substructure, int32 BucketIdx) override
-	{
-		check(BucketIdx < NumBuckets);
-		FSpatialAccelerationIdx Result;
-		Result.Bucket = BucketIdx;
-		
-		switch (BucketIdx)
-		{
-		case 0: Result.InnerIdx = GetBucket<0>(Buckets).AddSubstructure(ISpatialAcceleration<TPayloadType, T, d>::template AsUniqueChecked<typename std::remove_reference<decltype(GetBucket<0>(Buckets))>::type::AccelerationType>(MoveTemp(Substructure))); break;
-		case 1: Result.InnerIdx = GetBucket<ClampedIdx(1)>(Buckets).AddSubstructure(ISpatialAcceleration<TPayloadType, T, d>::template AsUniqueChecked<typename std::remove_reference<decltype(GetBucket<ClampedIdx(1)>(Buckets))>::type::AccelerationType>(MoveTemp(Substructure))); break;
-		case 2: Result.InnerIdx = GetBucket<ClampedIdx(2)>(Buckets).AddSubstructure(ISpatialAcceleration<TPayloadType, T, d>::template AsUniqueChecked<typename std::remove_reference<decltype(GetBucket<ClampedIdx(2)>(Buckets))>::type::AccelerationType>(MoveTemp(Substructure))); break;
-		case 3: Result.InnerIdx = GetBucket<ClampedIdx(3)>(Buckets).AddSubstructure(ISpatialAcceleration<TPayloadType, T, d>::template AsUniqueChecked<typename std::remove_reference<decltype(GetBucket<ClampedIdx(3)>(Buckets))>::type::AccelerationType>(MoveTemp(Substructure))); break;
-		case 4: Result.InnerIdx = GetBucket<ClampedIdx(4)>(Buckets).AddSubstructure(ISpatialAcceleration<TPayloadType, T, d>::template AsUniqueChecked<typename std::remove_reference<decltype(GetBucket<ClampedIdx(4)>(Buckets))>::type::AccelerationType>(MoveTemp(Substructure))); break;
-		case 5: Result.InnerIdx = GetBucket<ClampedIdx(5)>(Buckets).AddSubstructure(ISpatialAcceleration<TPayloadType, T, d>::template AsUniqueChecked<typename std::remove_reference<decltype(GetBucket<ClampedIdx(5)>(Buckets))>::type::AccelerationType>(MoveTemp(Substructure))); break;
-		case 6: Result.InnerIdx = GetBucket<ClampedIdx(6)>(Buckets).AddSubstructure(ISpatialAcceleration<TPayloadType, T, d>::template AsUniqueChecked<typename std::remove_reference<decltype(GetBucket<ClampedIdx(6)>(Buckets))>::type::AccelerationType>(MoveTemp(Substructure))); break;
-		case 7: Result.InnerIdx = GetBucket<ClampedIdx(7)>(Buckets).AddSubstructure(ISpatialAcceleration<TPayloadType, T, d>::template AsUniqueChecked<typename std::remove_reference<decltype(GetBucket<ClampedIdx(7)>(Buckets))>::type::AccelerationType>(MoveTemp(Substructure))); break;
-		default: check(false);
-		}
-		return Result;
-	}
-
-	virtual void RemoveSubstructure(FSpatialAccelerationIdx Idx) override
-	{
-		check(Idx.Bucket < NumBuckets);
-		switch (Idx.Bucket)
-		{
-		case 0: return GetBucket<0>(Buckets).RemoveSubstructure(Idx.InnerIdx);
-		case 1: return GetBucket<ClampedIdx(1)>(Buckets).RemoveSubstructure(Idx.InnerIdx);
-		case 2: return GetBucket<ClampedIdx(2)>(Buckets).RemoveSubstructure(Idx.InnerIdx);
-		case 3: return GetBucket<ClampedIdx(3)>(Buckets).RemoveSubstructure(Idx.InnerIdx);
-		case 4: return GetBucket<ClampedIdx(4)>(Buckets).RemoveSubstructure(Idx.InnerIdx);
-		case 5: return GetBucket<ClampedIdx(5)>(Buckets).RemoveSubstructure(Idx.InnerIdx);
-		case 6: return GetBucket<ClampedIdx(6)>(Buckets).RemoveSubstructure(Idx.InnerIdx);
-		case 7: return GetBucket<ClampedIdx(7)>(Buckets).RemoveSubstructure(Idx.InnerIdx);
-		default: check(false);
-		}
-	}
-
-	virtual ISpatialAcceleration<TPayloadType, T, d>* GetSubstructure(FSpatialAccelerationIdx Idx) override
-	{
-		check (Idx.Bucket < NumBuckets);
-		switch (Idx.Bucket)
-		{
-		case 0: return GetBucket<0>(Buckets).GetSubstructure(Idx.InnerIdx);
-		case 1: return GetBucket<ClampedIdx(1)>(Buckets).GetSubstructure(Idx.InnerIdx);
-		case 2: return GetBucket<ClampedIdx(2)>(Buckets).GetSubstructure(Idx.InnerIdx);
-		case 3: return GetBucket<ClampedIdx(3)>(Buckets).GetSubstructure(Idx.InnerIdx);
-		case 4: return GetBucket<ClampedIdx(4)>(Buckets).GetSubstructure(Idx.InnerIdx);
-		case 5: return GetBucket<ClampedIdx(5)>(Buckets).GetSubstructure(Idx.InnerIdx);
-		case 6: return GetBucket<ClampedIdx(6)>(Buckets).GetSubstructure(Idx.InnerIdx);
-		case 7: return GetBucket<ClampedIdx(7)>(Buckets).GetSubstructure(Idx.InnerIdx);
-		default: check(false);
-		}
-		return nullptr;
-	}
-
-	virtual void Raycast(const TVector<T, d>& Start, const TVector<T, d>& Dir, const T OriginalLength, ISpatialVisitor<TPayloadType, T>& Visitor) const
-	{
-		TSpatialAccelerationCollectionHelper<0, NumBuckets, decltype(Buckets), TPayloadType, T, d>::Raycast(Buckets, Start, Dir, OriginalLength, Visitor);
-	}
-
-	virtual void Sweep(const TVector<T, d>& Start, const TVector<T, d>& Dir, T OriginalLength, const TVector<T, d> QueryHalfExtents, ISpatialVisitor<TPayloadType, T>& Visitor) const
-	{
-		TSpatialAccelerationCollectionHelper<0, NumBuckets, decltype(Buckets), TPayloadType, T, d>::Sweep(Buckets, Start, Dir, OriginalLength, QueryHalfExtents, Visitor);
-	}
-	virtual void Overlap(const TBox<T, d>& QueryBounds, ISpatialVisitor<TPayloadType, T>& Visitor) const
-	{
-		TSpatialAccelerationCollectionHelper<0, NumBuckets, decltype(Buckets), TPayloadType, T, d>::Overlap(Buckets, QueryBounds, Visitor);
-	}
-
-	virtual void RemoveElementFrom(const TPayloadType& Payload, FSpatialAccelerationIdx SpatialIdx) override
-	{
-		check(SpatialIdx.Bucket < NumBuckets);
-		switch (SpatialIdx.Bucket)
-		{
-		case 0: return GetBucket<0>(Buckets).RemoveElement(Payload, SpatialIdx.InnerIdx);
-		case 1: return GetBucket<ClampedIdx(1)>(Buckets).RemoveElement(Payload, SpatialIdx.InnerIdx);
-		case 2: return GetBucket<ClampedIdx(2)>(Buckets).RemoveElement(Payload, SpatialIdx.InnerIdx);
-		case 3: return GetBucket<ClampedIdx(3)>(Buckets).RemoveElement(Payload, SpatialIdx.InnerIdx);
-		case 4: return GetBucket<ClampedIdx(4)>(Buckets).RemoveElement(Payload, SpatialIdx.InnerIdx);
-		case 5: return GetBucket<ClampedIdx(5)>(Buckets).RemoveElement(Payload, SpatialIdx.InnerIdx);
-		case 6: return GetBucket<ClampedIdx(6)>(Buckets).RemoveElement(Payload, SpatialIdx.InnerIdx);
-		case 7: return GetBucket<ClampedIdx(7)>(Buckets).RemoveElement(Payload, SpatialIdx.InnerIdx);
-		default: check(false);
-		}
-	}
-
-	virtual void UpdateElementIn(const TPayloadType& Payload, const TBox<T, d>& NewBounds, bool bHasBounds, FSpatialAccelerationIdx SpatialIdx)
-	{
-		check(SpatialIdx.Bucket < NumBuckets);
-		switch (SpatialIdx.Bucket)
-		{
-		case 0: return GetBucket<0>(Buckets).UpdateElement(Payload, NewBounds, bHasBounds, SpatialIdx.InnerIdx);
-		case 1: return GetBucket<ClampedIdx(1)>(Buckets).UpdateElement(Payload, NewBounds, bHasBounds, SpatialIdx.InnerIdx);
-		case 2: return GetBucket<ClampedIdx(2)>(Buckets).UpdateElement(Payload, NewBounds, bHasBounds, SpatialIdx.InnerIdx);
-		case 3: return GetBucket<ClampedIdx(3)>(Buckets).UpdateElement(Payload, NewBounds, bHasBounds, SpatialIdx.InnerIdx);
-		case 4: return GetBucket<ClampedIdx(4)>(Buckets).UpdateElement(Payload, NewBounds, bHasBounds, SpatialIdx.InnerIdx);
-		case 5: return GetBucket<ClampedIdx(5)>(Buckets).UpdateElement(Payload, NewBounds, bHasBounds, SpatialIdx.InnerIdx);
-		case 6: return GetBucket<ClampedIdx(6)>(Buckets).UpdateElement(Payload, NewBounds, bHasBounds, SpatialIdx.InnerIdx);
-		case 7: return GetBucket<ClampedIdx(7)>(Buckets).UpdateElement(Payload, NewBounds, bHasBounds, SpatialIdx.InnerIdx);
-		default: check(false);
-		}
-	}
-
-	virtual TUniquePtr<ISpatialAcceleration<TPayloadType, T, d>> Copy() const
-	{
-		return TUniquePtr<ISpatialAcceleration<TPayloadType, T, d>>(new TSpatialAccelerationCollection<TSpatialAccelerationBuckets...>(*this));
-	}
-
-private:
-
-	TSpatialAccelerationCollection(const TSpatialAccelerationCollection<TSpatialAccelerationBuckets...>& Other) = default;
-
-	static constexpr uint32 ClampedIdx(uint32 Idx)
-	{
-		return Idx < NumBuckets ? Idx : 0;
-	}
-
-	TSpatialBucketTuple< TSpatialAccelerationBuckets...> Buckets;
-	static constexpr uint32 NumBuckets = sizeof...(TSpatialAccelerationBuckets);
-	static_assert(NumBuckets < 8, "8 max buckets supported");
-};
-
-
 }
diff --git a/Engine/Source/Runtime/Experimental/Chaos/Public/Chaos/ISpatialAccelerationCollection.h b/Engine/Source/Runtime/Experimental/Chaos/Public/Chaos/ISpatialAccelerationCollection.h
new file mode 100644
--- /dev/null
+++ b/Engine/Source/Runtime/Experimental/Chaos/Public/Chaos/ISpatialAccelerationCollection.h
@@ -0,0 +1,480 @@
+// Copyright 1998-2019 Epic Games, Inc. All Rights Reserved.
+#pragma once
+#include "Chaos/ISpatialAcceleration.h"
+#include "Chaos/Box.h"
+#include "GeometryParticlesfwd.h"
+#include <tuple>
+#include "PBDCollisionConstraintImp.h"
+
+namespace Chaos
+{
+
+template <typename TPayloadType, typename T, int d>
+class CHAOS_API ISpatialAccelerationCollection : public ISpatialAcceleration<TPayloadType, T, d>
+{
+public:
+	ISpatialAccelerationCollection() : ISpatialAcceleration<TPayloadType, T, d>(StaticType) {}
+	static constexpr ESpatialAcceleration StaticType = ESpatialAcceleration::Collection;
+	virtual FSpatialAccelerationIdx AddSubstructure(TUniquePtr<ISpatialAcceleration<TPayloadType, T, d>>&& Substructure, int32 Bucket) = 0;
+	virtual TUniquePtr<ISpatialAcceleration<TPayloadType, T, d>> RemoveSubstructure(FSpatialAccelerationIdx Idx) = 0;
+	virtual ISpatialAcceleration<TPayloadType, T, d>* GetSubstructure(FSpatialAccelerationIdx Idx) = 0;
+
+	/** This is kind of a hack to avoid virtuals. We simply route calls into templated functions */
+	virtual void PBDComputeConstraintsLowLevel_GatherStats(TPBDCollisionConstraint<T, d>& CollisionConstraint, T Dt) const = 0;
+	virtual void PBDComputeConstraintsLowLevel(TPBDCollisionConstraint<T, d>& CollisionConstraint, T Dt) const = 0;
+};
+
+template <typename TSpatialAcceleration>
+class CHAOS_API TSpatialAccelerationCollectionBucket
+{
+public:
+	using AccelerationType = TSpatialAcceleration;
+	using TPayloadType = typename AccelerationType::PayloadType;
+	using T = typename AccelerationType::TType;
+	static constexpr int d = AccelerationType::D;
+
+	TSpatialAccelerationCollectionBucket() = default;
+
+	TSpatialAccelerationCollectionBucket<TSpatialAcceleration> Copy() const
+	{
+		return TSpatialAccelerationCollectionBucket<TSpatialAcceleration>(*this);
+	}
+
+	uint16 AddSubstructure(TUniquePtr<TSpatialAcceleration>&& Substructure)
+	{
+		uint16 NewIdx = NewSlot();
+		Accelerations[NewIdx] = MoveTemp(Substructure);
+		
+		return NewIdx;
+	}
+	
+	TUniquePtr <ISpatialAcceleration<TPayloadType, T, d>> RemoveSubstructure(uint16 Idx)
+	{
+		return ReleaseSlot(Idx);
+	}
+
+	TSpatialAcceleration* GetSubstructure(int32 Idx)
+	{
+		return Accelerations[Idx].Get();
+	}
+
+	template <typename SQVisitor>
+	void Raycast(const TVector<T, d>& Start, const TVector<T, d>& Dir, const T OriginalLength, SQVisitor& Visitor) const
+	{
+		for (const auto& Acceleration : Accelerations)
+		{
+			if (Acceleration)
+			{
+				Acceleration->Raycast(Start, Dir, OriginalLength, Visitor);
+			}
+		}
+	}
+
+	template <typename SQVisitor>
+	void Sweep(const TVector<T, d>& Start, const TVector<T, d>& Dir, const T OriginalLength, const TVector<T,3>& QueryHalfExtents, SQVisitor& Visitor) const
+	{
+		for (const auto& Acceleration : Accelerations)
+		{
+			if (Acceleration)
+			{
+				Acceleration->Sweep(Start, Dir, OriginalLength, QueryHalfExtents, Visitor);
+			}
+		}
+	}
+
+	template <typename SQVisitor>
+	void Overlap(const TBox<T,d>& QueryBounds, SQVisitor& Visitor) const
+	{
+		for (const auto& Acceleration : Accelerations)
+		{
+			if (Acceleration)
+			{
+				Acceleration->Overlap(QueryBounds, Visitor);
+			}
+		}
+	}
+
+	void GlobalObjects(TArray<TPayloadBoundsElement<TPayloadType, T>>& ObjList) const
+	{
+		for (const auto& Acceleration : Accelerations)
+		{
+			if (Acceleration)
+			{
+				ObjList.Append(Acceleration->GlobalObjects());
+			}
+		}
+	}
+
+	void RemoveElement(const TPayloadType& Payload, int32 Idx)
+	{
+		Accelerations[Idx]->RemoveElement(Payload);
+	}
+
+	void UpdateElement(const TPayloadType& Payload, const TBox<T, d>& NewBounds, bool bHasBounds, int32 Idx)
+	{
+		Accelerations[Idx]->UpdateElement(Payload, NewBounds, bHasBounds);
+	}
+
+	explicit TSpatialAccelerationCollectionBucket(const TSpatialAccelerationCollectionBucket<TSpatialAcceleration>& Other)
+		: FreeIndices(Other.FreeIndices)
+	{
+		Accelerations.Reserve(Other.Accelerations.Num());
+		for (const auto& OtherAccel : Other.Accelerations)
+		{
+			if (OtherAccel)
+			{
+				Accelerations.Add(AsUniqueSpatialAccelerationChecked<TSpatialAcceleration>(OtherAccel->Copy()));
+			}
+			else
+			{
+				Accelerations.Add(nullptr);
+			}
+		}
+	}
+
+private:
+	uint16 NewSlot()
+	{
+		if (FreeIndices.Num())
+		{
+			return FreeIndices.Pop(false);
+		}
+		else if (Accelerations.Num() < FSpatialAccelerationIdx::MaxBucketEntries)
+		{
+			return Accelerations.Add(nullptr);
+		}
+		check(false);	//too many substructures in a single bucket
+		return 0;
+	}
+
+	TUniquePtr <ISpatialAcceleration<TPayloadType, T, d>> ReleaseSlot(uint16 Index)
+	{
+		TUniquePtr <ISpatialAcceleration<TPayloadType, T, d>> Result = MoveTemp(Accelerations[Index]);
+		Accelerations[Index] = nullptr;
+		if (Index + 1 == Accelerations.Num())
+		{
+			Accelerations.Pop(false);
+		}
+		else
+		{
+			FreeIndices.Add(Index);
+		}
+		return Result;
+	}
+
+	TSpatialAccelerationCollectionBucket& operator=(const TSpatialAccelerationCollectionBucket<TSpatialAcceleration>& Other)
+	{
+		if (&Other != this)
+		{
+			FreeIndices = Other.FreeIndices;
+			Accelerations.Reset(Other.Accelerations.Num());
+
+			for (const auto& OtherAccel : Other.Accelerations)
+			{
+				if (OtherAccel)
+				{
+					Accelerations.Add(AsUniqueSpatialAccelerationChecked<TSpatialAcceleration>(OtherAccel->Copy()));
+				}
+				else
+				{
+					Accelerations.Add(nullptr);
+				}
+			}
+		}
+
+		return *this;
+	}
+
+	TArray<TUniquePtr<TSpatialAcceleration>> Accelerations;
+	TArray<uint16> FreeIndices;
+};
+
+template <typename... TRemaining>
+struct CHAOS_API TSpatialBucketTuple
+{
+};
+
+template <typename TAcceleration, typename... TRemaining>
+struct CHAOS_API TSpatialBucketTuple<TAcceleration, TRemaining...>
+{
+	using FirstType = TAcceleration;
+	TSpatialAccelerationCollectionBucket<TAcceleration> First;
+	TSpatialBucketTuple<TRemaining...> Remaining;
+
+	TSpatialBucketTuple() = default;
+
+	explicit TSpatialBucketTuple(const TSpatialBucketTuple<TAcceleration, TRemaining...>& Other)
+		: First(Other.First)
+		, Remaining(Other.Remaining)
+	{
+	}
+};
+
+template<int Idx, typename ... Rest>
+struct CHAOS_API TSpatialBucketTupleGetter
+{
+
+};
+
+template<int Idx, typename First, typename... Rest>
+struct CHAOS_API TSpatialBucketTupleGetter<Idx, First, Rest...>
+{
+	static auto& Get(TSpatialBucketTuple<First, Rest...>& Buckets) { return TSpatialBucketTupleGetter<Idx - 1, Rest...>::Get(Buckets.Remaining); }
+	static const auto& Get(const TSpatialBucketTuple<First, Rest...>& Buckets) { return TSpatialBucketTupleGetter<Idx - 1, Rest...>::Get(Buckets.Remaining); }
+};
+
+template<typename First, typename... Rest>
+struct CHAOS_API TSpatialBucketTupleGetter<0, First, Rest...>
+{
+	static auto& Get(TSpatialBucketTuple<First, Rest...>& Buckets) { return Buckets.First; }
+	static const auto& Get(const TSpatialBucketTuple<First, Rest...>& Buckets) { return Buckets.First; }
+};
+
+template <int Idx, typename First, typename... Rest>
+auto& GetBucket(TSpatialBucketTuple<First, Rest...>& Buckets)
+{
+	return TSpatialBucketTupleGetter<Idx, First, Rest...>::Get(Buckets);
+}
+
+template <int Idx, typename First, typename... Rest>
+const auto& GetBucket(const TSpatialBucketTuple<First, Rest...>& Buckets)
+{
+	return TSpatialBucketTupleGetter<Idx, First, Rest...>::Get(Buckets);
+}
+
+template <int BucketIdx, int NumBuckets, typename Tuple, typename TPayloadType, typename T, int d>
+struct TSpatialAccelerationCollectionHelper
+{
+	template <typename SQVisitor>
+	static void Raycast(const Tuple& Buckets, const TVector<T, d>& Start, const TVector<T, d>& Dir, const T OriginalLength, SQVisitor& Visitor)
+	{
+		GetBucket<BucketIdx>(Buckets).Raycast(Start, Dir, OriginalLength, Visitor);
+		constexpr int NextBucket = BucketIdx + 1;
+		if (NextBucket < NumBuckets)
+		{
+			TSpatialAccelerationCollectionHelper<NextBucket < NumBuckets ? NextBucket : 0, NumBuckets, Tuple, TPayloadType, T, d>::Raycast(Buckets, Start, Dir, OriginalLength, Visitor);
+		}
+	}
+
+	template <typename SQVisitor>
+	static void Sweep(const Tuple& Buckets, const TVector<T, d>& Start, const TVector<T, d>& Dir, const T OriginalLength, const TVector<T, d> QueryHalfExtents, SQVisitor& Visitor)
+	{
+		GetBucket<BucketIdx>(Buckets).Sweep(Start, Dir, OriginalLength, QueryHalfExtents, Visitor);
+		constexpr int NextBucket = BucketIdx + 1;
+		if (NextBucket < NumBuckets)
+		{
+			TSpatialAccelerationCollectionHelper < NextBucket < NumBuckets ? NextBucket : 0, NumBuckets, Tuple, TPayloadType, T, d>::Sweep(Buckets, Start, Dir, OriginalLength, QueryHalfExtents, Visitor);
+		}
+	}
+
+	template <typename SQVisitor>
+	static void Overlap(const Tuple& Buckets, const TBox<T, d> QueryBounds, SQVisitor& Visitor)
+	{
+		GetBucket<BucketIdx>(Buckets).Overlap(QueryBounds, Visitor);
+		constexpr int NextBucket = BucketIdx + 1;
+		if (NextBucket < NumBuckets)
+		{
+			TSpatialAccelerationCollectionHelper < NextBucket < NumBuckets ? NextBucket : 0, NumBuckets, Tuple, TPayloadType, T, d>::Overlap(Buckets, QueryBounds, Visitor);
+		}
+	}
+
+	static void GlobalObjects(const Tuple& Buckets, TArray<TPayloadBoundsElement<TPayloadType, T>>& ObjList)
+	{
+		GetBucket<BucketIdx>(Buckets).GlobalObjects(ObjList);
+		constexpr int NextBucket = BucketIdx + 1;
+		if (NextBucket < NumBuckets)
+		{
+			TSpatialAccelerationCollectionHelper < NextBucket < NumBuckets ? NextBucket : 0, NumBuckets, Tuple, TPayloadType, T, d>::GlobalObjects(Buckets, ObjList);
+		}
+	}
+};
+
+template <bool bGatherStats, typename SpatialAccelerationCollection, typename T, int d>
+typename TEnableIf<TIsSame<typename SpatialAccelerationCollection::TPayloadType, TAccelerationStructureHandle<T,d>>::Value, void>::Type PBDComputeConstraintsLowLevel_Helper(TPBDCollisionConstraint<T, d>& CollisionConstraint, const SpatialAccelerationCollection& Accel, T Dt)
+{
+	CollisionConstraint.template ComputeConstraintsHelperLowLevel<bGatherStats>(Accel, Dt);
+}
+
+template <bool bGatherStats, typename SpatialAccelerationCollection, typename T, int d>
+typename TEnableIf<!TIsSame<typename SpatialAccelerationCollection::TPayloadType, TAccelerationStructureHandle<T, d>>::Value, void>::Type PBDComputeConstraintsLowLevel_Helper(TPBDCollisionConstraint<T, d>& CollisionConstraint, const SpatialAccelerationCollection& Accel, T Dt)
+{
+}
+
+template <typename ... TSpatialAccelerationBuckets>
+class CHAOS_API TSpatialAccelerationCollection : public
+	ISpatialAccelerationCollection<typename std::tuple_element<0, std::tuple< TSpatialAccelerationBuckets...>>::type::PayloadType,
+	typename std::tuple_element<0, std::tuple< TSpatialAccelerationBuckets...>>::type::TType,
+	std::tuple_element<0, std::tuple< TSpatialAccelerationBuckets...>>::type::D>
+{
+public:
+	using FirstAccelerationType = typename std::tuple_element<0, std::tuple< TSpatialAccelerationBuckets...>>::type;
+	using TPayloadType = typename FirstAccelerationType::PayloadType;
+	using T = typename FirstAccelerationType::TType;
+	static constexpr int d = FirstAccelerationType::D;
+	TSpatialAccelerationCollection() = default;
+
+	virtual FSpatialAccelerationIdx AddSubstructure(TUniquePtr<ISpatialAcceleration<TPayloadType, T, d>>&& Substructure, int32 BucketIdx) override
+	{
+		check(BucketIdx < NumBuckets);
+		FSpatialAccelerationIdx Result;
+		Result.Bucket = BucketIdx;
+		
+		switch (BucketIdx)
+		{
+		case 0: Result.InnerIdx = GetBucket<0>(Buckets).AddSubstructure(AsUniqueSpatialAccelerationChecked<typename std::remove_reference<decltype(GetBucket<0>(Buckets))>::type::AccelerationType>(MoveTemp(Substructure))); break;
+		case 1: Result.InnerIdx = GetBucket<ClampedIdx(1)>(Buckets).AddSubstructure(AsUniqueSpatialAccelerationChecked<typename std::remove_reference<decltype(GetBucket<ClampedIdx(1)>(Buckets))>::type::AccelerationType>(MoveTemp(Substructure))); break;
+		case 2: Result.InnerIdx = GetBucket<ClampedIdx(2)>(Buckets).AddSubstructure(AsUniqueSpatialAccelerationChecked<typename std::remove_reference<decltype(GetBucket<ClampedIdx(2)>(Buckets))>::type::AccelerationType>(MoveTemp(Substructure))); break;
+		case 3: Result.InnerIdx = GetBucket<ClampedIdx(3)>(Buckets).AddSubstructure(AsUniqueSpatialAccelerationChecked<typename std::remove_reference<decltype(GetBucket<ClampedIdx(3)>(Buckets))>::type::AccelerationType>(MoveTemp(Substructure))); break;
+		case 4: Result.InnerIdx = GetBucket<ClampedIdx(4)>(Buckets).AddSubstructure(AsUniqueSpatialAccelerationChecked<typename std::remove_reference<decltype(GetBucket<ClampedIdx(4)>(Buckets))>::type::AccelerationType>(MoveTemp(Substructure))); break;
+		case 5: Result.InnerIdx = GetBucket<ClampedIdx(5)>(Buckets).AddSubstructure(AsUniqueSpatialAccelerationChecked<typename std::remove_reference<decltype(GetBucket<ClampedIdx(5)>(Buckets))>::type::AccelerationType>(MoveTemp(Substructure))); break;
+		case 6: Result.InnerIdx = GetBucket<ClampedIdx(6)>(Buckets).AddSubstructure(AsUniqueSpatialAccelerationChecked<typename std::remove_reference<decltype(GetBucket<ClampedIdx(6)>(Buckets))>::type::AccelerationType>(MoveTemp(Substructure))); break;
+		case 7: Result.InnerIdx = GetBucket<ClampedIdx(7)>(Buckets).AddSubstructure(AsUniqueSpatialAccelerationChecked<typename std::remove_reference<decltype(GetBucket<ClampedIdx(7)>(Buckets))>::type::AccelerationType>(MoveTemp(Substructure))); break;
+		default: check(false);
+		}
+		return Result;
+	}
+
+	virtual TUniquePtr <ISpatialAcceleration<TPayloadType, T, d>> RemoveSubstructure(FSpatialAccelerationIdx Idx) override
+	{
+		check(Idx.Bucket < NumBuckets);
+		switch (Idx.Bucket)
+		{
+		case 0: return GetBucket<0>(Buckets).RemoveSubstructure(Idx.InnerIdx);
+		case 1: return GetBucket<ClampedIdx(1)>(Buckets).RemoveSubstructure(Idx.InnerIdx);
+		case 2: return GetBucket<ClampedIdx(2)>(Buckets).RemoveSubstructure(Idx.InnerIdx);
+		case 3: return GetBucket<ClampedIdx(3)>(Buckets).RemoveSubstructure(Idx.InnerIdx);
+		case 4: return GetBucket<ClampedIdx(4)>(Buckets).RemoveSubstructure(Idx.InnerIdx);
+		case 5: return GetBucket<ClampedIdx(5)>(Buckets).RemoveSubstructure(Idx.InnerIdx);
+		case 6: return GetBucket<ClampedIdx(6)>(Buckets).RemoveSubstructure(Idx.InnerIdx);
+		case 7: return GetBucket<ClampedIdx(7)>(Buckets).RemoveSubstructure(Idx.InnerIdx);
+		default: check(false);
+		}
+		return nullptr;
+	}
+
+	virtual ISpatialAcceleration<TPayloadType, T, d>* GetSubstructure(FSpatialAccelerationIdx Idx) override
+	{
+		check (Idx.Bucket < NumBuckets);
+		switch (Idx.Bucket)
+		{
+		case 0: return GetBucket<0>(Buckets).GetSubstructure(Idx.InnerIdx);
+		case 1: return GetBucket<ClampedIdx(1)>(Buckets).GetSubstructure(Idx.InnerIdx);
+		case 2: return GetBucket<ClampedIdx(2)>(Buckets).GetSubstructure(Idx.InnerIdx);
+		case 3: return GetBucket<ClampedIdx(3)>(Buckets).GetSubstructure(Idx.InnerIdx);
+		case 4: return GetBucket<ClampedIdx(4)>(Buckets).GetSubstructure(Idx.InnerIdx);
+		case 5: return GetBucket<ClampedIdx(5)>(Buckets).GetSubstructure(Idx.InnerIdx);
+		case 6: return GetBucket<ClampedIdx(6)>(Buckets).GetSubstructure(Idx.InnerIdx);
+		case 7: return GetBucket<ClampedIdx(7)>(Buckets).GetSubstructure(Idx.InnerIdx);
+		default: check(false);
+		}
+		return nullptr;
+	}
+
+	virtual void Raycast(const TVector<T, d>& Start, const TVector<T, d>& Dir, const T OriginalLength, ISpatialVisitor<TPayloadType, T>& Visitor) const
+	{
+		TSpatialVisitor<TPayloadType, T> ProxyVisitor(Visitor);
+		TSpatialAccelerationCollectionHelper<0, NumBuckets, decltype(Buckets), TPayloadType, T, d>::Raycast(Buckets, Start, Dir, OriginalLength, ProxyVisitor);
+	}
+
+	template <typename SQVisitor>
+	void Raycast(const TVector<T, d>& Start, const TVector<T, d>& Dir, const T OriginalLength, SQVisitor& Visitor) const
+	{
+		TSpatialAccelerationCollectionHelper<0, NumBuckets, decltype(Buckets), TPayloadType, T, d>::Raycast(Buckets, Start, Dir, OriginalLength, Visitor);
+	}
+
+	virtual void Sweep(const TVector<T, d>& Start, const TVector<T, d>& Dir, T OriginalLength, const TVector<T, d> QueryHalfExtents, ISpatialVisitor<TPayloadType, T>& Visitor) const
+	{
+		TSpatialVisitor<TPayloadType, T> ProxyVisitor(Visitor);
+		TSpatialAccelerationCollectionHelper<0, NumBuckets, decltype(Buckets), TPayloadType, T, d>::Sweep(Buckets, Start, Dir, OriginalLength, QueryHalfExtents, ProxyVisitor);
+	}
+
+	template <typename SQVisitor>
+	void Sweep(const TVector<T, d>& Start, const TVector<T, d>& Dir, T OriginalLength, const TVector<T, d> QueryHalfExtents, SQVisitor& Visitor) const
+	{
+		TSpatialAccelerationCollectionHelper<0, NumBuckets, decltype(Buckets), TPayloadType, T, d>::Sweep(Buckets, Start, Dir, OriginalLength, QueryHalfExtents, Visitor);
+	}
+
+	virtual void Overlap(const TBox<T, d>& QueryBounds, ISpatialVisitor<TPayloadType, T>& Visitor) const
+	{
+		TSpatialVisitor<TPayloadType, T> ProxyVisitor(Visitor);
+		Overlap(QueryBounds, ProxyVisitor);
+	}
+
+	template <typename SQVisitor>
+	void Overlap(const TBox<T, 3>& QueryBounds, SQVisitor& Visitor) const
+	{
+		TSpatialAccelerationCollectionHelper<0, NumBuckets, decltype(Buckets), TPayloadType, T, d>::Overlap(Buckets, QueryBounds, Visitor);
+	}
+
+	TArray<TPayloadBoundsElement<TPayloadType, T>> GlobalObjects() const
+	{
+		TArray<TPayloadBoundsElement<TPayloadType, T>> ObjList;
+		TSpatialAccelerationCollectionHelper<0, NumBuckets, decltype(Buckets), TPayloadType, T, d>::GlobalObjects(Buckets, ObjList);
+		return ObjList;
+	}
+
+	virtual void RemoveElementFrom(const TPayloadType& Payload, FSpatialAccelerationIdx SpatialIdx) override
+	{
+		check(SpatialIdx.Bucket < NumBuckets);
+		switch (SpatialIdx.Bucket)
+		{
+		case 0: return GetBucket<0>(Buckets).RemoveElement(Payload, SpatialIdx.InnerIdx);
+		case 1: return GetBucket<ClampedIdx(1)>(Buckets).RemoveElement(Payload, SpatialIdx.InnerIdx);
+		case 2: return GetBucket<ClampedIdx(2)>(Buckets).RemoveElement(Payload, SpatialIdx.InnerIdx);
+		case 3: return GetBucket<ClampedIdx(3)>(Buckets).RemoveElement(Payload, SpatialIdx.InnerIdx);
+		case 4: return GetBucket<ClampedIdx(4)>(Buckets).RemoveElement(Payload, SpatialIdx.InnerIdx);
+		case 5: return GetBucket<ClampedIdx(5)>(Buckets).RemoveElement(Payload, SpatialIdx.InnerIdx);
+		case 6: return GetBucket<ClampedIdx(6)>(Buckets).RemoveElement(Payload, SpatialIdx.InnerIdx);
+		case 7: return GetBucket<ClampedIdx(7)>(Buckets).RemoveElement(Payload, SpatialIdx.InnerIdx);
+		default: check(false);
+		}
+	}
+
+	virtual void UpdateElementIn(const TPayloadType& Payload, const TBox<T, d>& NewBounds, bool bHasBounds, FSpatialAccelerationIdx SpatialIdx)
+	{
+		check(SpatialIdx.Bucket < NumBuckets);
+		switch (SpatialIdx.Bucket)
+		{
+		case 0: return GetBucket<0>(Buckets).UpdateElement(Payload, NewBounds, bHasBounds, SpatialIdx.InnerIdx);
+		case 1: return GetBucket<ClampedIdx(1)>(Buckets).UpdateElement(Payload, NewBounds, bHasBounds, SpatialIdx.InnerIdx);
+		case 2: return GetBucket<ClampedIdx(2)>(Buckets).UpdateElement(Payload, NewBounds, bHasBounds, SpatialIdx.InnerIdx);
+		case 3: return GetBucket<ClampedIdx(3)>(Buckets).UpdateElement(Payload, NewBounds, bHasBounds, SpatialIdx.InnerIdx);
+		case 4: return GetBucket<ClampedIdx(4)>(Buckets).UpdateElement(Payload, NewBounds, bHasBounds, SpatialIdx.InnerIdx);
+		case 5: return GetBucket<ClampedIdx(5)>(Buckets).UpdateElement(Payload, NewBounds, bHasBounds, SpatialIdx.InnerIdx);
+		case 6: return GetBucket<ClampedIdx(6)>(Buckets).UpdateElement(Payload, NewBounds, bHasBounds, SpatialIdx.InnerIdx);
+		case 7: return GetBucket<ClampedIdx(7)>(Buckets).UpdateElement(Payload, NewBounds, bHasBounds, SpatialIdx.InnerIdx);
+		default: check(false);
+		}
+	}
+
+	virtual TUniquePtr<ISpatialAcceleration<TPayloadType, T, d>> Copy() const
+	{
+		return TUniquePtr<ISpatialAcceleration<TPayloadType, T, d>>(new TSpatialAccelerationCollection<TSpatialAccelerationBuckets...>(*this));
+	}
+
+	virtual void PBDComputeConstraintsLowLevel_GatherStats(TPBDCollisionConstraint<T, d>& CollisionConstraint, T Dt) const override
+	{
+		PBDComputeConstraintsLowLevel_Helper<true>(CollisionConstraint, *this, Dt);
+	}
+
+	virtual void PBDComputeConstraintsLowLevel(TPBDCollisionConstraint<T,d>& CollisionConstraint, T Dt) const override
+	{
+		PBDComputeConstraintsLowLevel_Helper<false>(CollisionConstraint, *this, Dt);
+	}
+
+private:
+
+	TSpatialAccelerationCollection(const TSpatialAccelerationCollection<TSpatialAccelerationBuckets...>& Other) = default;
+	
+	static constexpr uint32 ClampedIdx(uint32 Idx)
+	{
+		return Idx < NumBuckets ? Idx : 0;
+	}
+
+	TSpatialBucketTuple< TSpatialAccelerationBuckets...> Buckets;
+	static constexpr uint32 NumBuckets = sizeof...(TSpatialAccelerationBuckets);
+	static_assert(NumBuckets < 8, "8 max buckets supported");
+};
+
+
+}
diff --git a/Engine/Source/Runtime/Experimental/Chaos/Public/Chaos/PBDCollisionConstraint.h b/Engine/Source/Runtime/Experimental/Chaos/Public/Chaos/PBDCollisionConstraint.h
--- a/Engine/Source/Runtime/Experimental/Chaos/Public/Chaos/PBDCollisionConstraint.h
+++ b/Engine/Source/Runtime/Experimental/Chaos/Public/Chaos/PBDCollisionConstraint.h
@@ -206,6 +206,10 @@ public:
 	 * Get the particles that are affected by the specified constraint.
 	 */
 	TVector<TGeometryParticleHandle<T, d>*, 2> ConstraintParticles(int32 ContactIndex) const { return { Constraints[ContactIndex].Particle, Constraints[ContactIndex].Levelset }; }
+
+	//NOTE: this should not be called by anyone other than ISpatialAccelerationCollection and CollisionConstraints - todo: make private with friends?
+	template <bool bGatherStats, typename SPATIAL_ACCELERATION>
+	void ComputeConstraintsHelperLowLevel(const SPATIAL_ACCELERATION& SpatialAcceleration, T Dt);
 	
 protected:
 	using Base::GetConstraintIndex;
@@ -225,8 +229,6 @@ private:
 
 	FRigidBodyContactConstraint ComputeConstraint(TGeometryParticleHandle<T, d>* Particle0, TGeometryParticleHandle<T, d>* Particle1, const T Thickness);
 
-	template <typename SPATIAL_ACCELERATION, bool bGatherStats>
-	void ComputeConstraintsHelper(T Dt, const SPATIAL_ACCELERATION& SpatialAcceleration);
 
 	template<typename SPATIAL_ACCELERATION>
 	void UpdateConstraintsHelper(/*const TPBDRigidParticles<T, d>& InParticles, const TArray<int32>& InIndices,*/ T Dt, const TSet<TGeometryParticleHandle<T, d>*>& AddedParticles, SPATIAL_ACCELERATION& SpatialAcceleration);
diff --git a/Engine/Source/Runtime/Experimental/Chaos/Public/Chaos/PBDCollisionConstraintImp.h b/Engine/Source/Runtime/Experimental/Chaos/Public/Chaos/PBDCollisionConstraintImp.h
new file mode 100644
--- /dev/null
+++ b/Engine/Source/Runtime/Experimental/Chaos/Public/Chaos/PBDCollisionConstraintImp.h
@@ -0,0 +1,346 @@
+// Copyright 1998-2019 Epic Games, Inc. All Rights Reserved.
+#pragma once
+
+#include "PBDCollisionConstraint.h"
+#include "Chaos/ChaosPerfTest.h"
+
+namespace Chaos
+{
+#if !UE_BUILD_SHIPPING
+
+	static constexpr int32 BucketSizes2[] = { 0,1,4,8,16,32,64,128,512,MAX_int32 };
+
+
+	template <bool bGatherStats>
+	struct FStatHelper2
+	{
+		int32 MaxCount;
+
+		FStatHelper2() {}
+		void Record(int32 Count) {}
+		FString ToString() const { return TEXT(""); }
+	};
+
+	template <>
+	struct FStatHelper2<true>
+	{
+		int32 BucketCount[UE_ARRAY_COUNT(BucketSizes2)];
+		int32 MaxCount;
+
+		FStatHelper2()
+		{
+			FMemory::Memset(BucketCount, 0, sizeof(BucketCount));
+			MaxCount = 0;
+		}
+
+		void Record(int32 Count)
+		{
+			for (int32 BucketIdx = 1; BucketIdx < UE_ARRAY_COUNT(BucketSizes2); ++BucketIdx)
+			{
+				if (Count >= BucketSizes2[BucketIdx - 1] && Count < BucketSizes2[BucketIdx])
+				{
+					++BucketCount[BucketIdx];
+				}
+			}
+
+			if (Count > MaxCount)
+			{
+				MaxCount = Count;
+			}
+		}
+
+		FString ToString() const
+		{
+			FString OutLog;
+			int32 MaxBucketCount = 0;
+			for (int32 Count : BucketCount)
+			{
+				if (Count > MaxBucketCount)
+				{
+					MaxBucketCount = Count;
+				}
+			}
+
+			const float CountPerChar = MaxBucketCount / 20.f;
+			for (int32 Idx = 1; Idx < UE_ARRAY_COUNT(BucketSizes2); ++Idx)
+			{
+				int32 NumChars = BucketCount[Idx] / CountPerChar;
+				if (Idx < UE_ARRAY_COUNT(BucketSizes2) - 1)
+				{
+					OutLog += FString::Printf(TEXT("\t[%4d - %4d) (%4d) |"), BucketSizes2[Idx - 1], BucketSizes2[Idx], BucketCount[Idx]);
+				}
+				else
+				{
+					OutLog += FString::Printf(TEXT("\t[%4d -  inf) (%4d) |"), BucketSizes2[Idx - 1], BucketCount[Idx]);
+				}
+				for (int32 Count = 0; Count < NumChars; ++Count)
+				{
+					OutLog += TEXT("-");
+				}
+				OutLog += TEXT("\n");
+			}
+
+			return OutLog;
+		}
+	};
+
+#endif
+
+	template <typename T, int d>
+	struct TSimOverlapVisitor
+	{
+		TArray<TAccelerationStructureHandle<T, d>>& Intersections;
+		TSimOverlapVisitor(TArray<TAccelerationStructureHandle<T, d>>& InIntersections)
+			: Intersections(InIntersections)
+		{
+		}
+
+		bool VisitOverlap(const TSpatialVisitorData<TAccelerationStructureHandle<T, d>>& Instance)
+		{
+			Intersections.Add(Instance.Payload);
+			return true;
+		}
+
+		bool VisitSweep(TSpatialVisitorData<TAccelerationStructureHandle<T, d>>, T Length)
+		{
+			check(false);
+			return false;
+		}
+
+		bool VisitRaycast(TSpatialVisitorData<TAccelerationStructureHandle<T, d>>, T Length)
+		{
+			check(false);
+			return false;
+		}
+	};
+
+	extern CHAOS_API int32 EnableCollisions;
+	extern CHAOS_API int32 CollisionConstraintsForceSingleThreaded;
+
+	DECLARE_CYCLE_STAT_EXTERN(TEXT("ComputeConstraints"), STAT_ComputeConstraints, STATGROUP_Chaos, CHAOS_API);
+	DECLARE_CYCLE_STAT_EXTERN(TEXT("ComputeConstraintsNP"), STAT_ComputeConstraintsNP, STATGROUP_Chaos, CHAOS_API);
+	DECLARE_CYCLE_STAT_EXTERN(TEXT("ComputeConstraintsBP"), STAT_ComputeConstraintsBP, STATGROUP_Chaos, CHAOS_API);
+
+	template <typename T, int d>
+	template <bool bGatherStats, typename SPATIAL_ACCELERATION>
+	void TPBDCollisionConstraint<T, d>::ComputeConstraintsHelperLowLevel(const SPATIAL_ACCELERATION& InSpatialAcceleration, T Dt)
+	{
+		SCOPE_CYCLE_COUNTER(STAT_ComputeConstraints);
+		CHAOS_SCOPED_TIMER(ComputeConstraints);
+		if (!EnableCollisions) return;
+		// Broad phase
+
+#if !UE_BUILD_SHIPPING
+		FStatHelper2<bGatherStats> BroadphasePotentials;
+		FStatHelper2<bGatherStats> NarrowPhaseSkipped;
+		FStatHelper2<bGatherStats> NarrowPhasePerformed;
+		FStatHelper2<bGatherStats> NarrowPhaseRejected;
+		FStatHelper2<bGatherStats> BoundsDistribution;
+		int32 SimulatedParticles = 0;
+#endif
+		{
+			SCOPE_CYCLE_COUNTER(STAT_ComputeConstraintsNP);
+
+
+			// Narrow phase
+			CHAOS_SCOPED_TIMER(ComputeConstraints_NP);
+
+			TQueue<TRigidBodyContactConstraint<T, d>, EQueueMode::Mpsc> Queue;	//todo(ocohen): use per thread buffer instead, need better support than ParallelFor for this
+			Particles.GetNonDisabledDynamicView().ParallelFor([&](auto& Particle1, int32 ActiveIdxIdx)
+			{
+#if !UE_BUILD_SHIPPING
+				if (bGatherStats)
+				{
+					++SimulatedParticles;
+				}
+#endif
+
+				TArray<TAccelerationStructureHandle<T, d>> PotentialIntersections;
+
+				const bool bBody1Bounded = HasBoundingBox(Particle1);
+
+				const T Box1Thickness = ComputeThickness(Particle1, Dt).Size();
+				if (bBody1Bounded)
+				{
+#if CHAOS_PARTICLEHANDLE_TODO
+					const TBox <T, d> Box1 = InSpatialAcceleration.GetWorldSpaceBoundingBox(Particle1);
+#else
+					const TBox <T, d> Box1 = ComputeWorldSpaceBoundingBox(Particle1);	//NOTE: this ignores the velocity expansion which is wrong
+#endif
+
+#if !UE_BUILD_SHIPPING
+					if (bGatherStats)
+					{
+						//const TBox<T, d> Box1 = SpatialAcceleration.GetWorldSpaceBoundingBox(InParticles, Body1Index);
+
+						BoundsDistribution.Record(Box1.Extents().GetMax());
+					}
+#endif
+					TSimOverlapVisitor<T, d> OverlapVisitor(PotentialIntersections);
+					InSpatialAcceleration.Overlap(Box1, OverlapVisitor);
+				}
+				else
+				{
+					const auto& GlobalElems = InSpatialAcceleration.GlobalObjects();
+					PotentialIntersections.Reserve(GlobalElems.Num());
+
+					for (auto& Elem : GlobalElems)
+					{
+						PotentialIntersections.Add(Elem.Payload);
+					}
+				}
+
+#if !UE_BUILD_SHIPPING
+				BroadphasePotentials.Record(PotentialIntersections.Num());
+
+				int32 CountNP = 0;
+				int32 RejectedNP = 0;
+#endif
+				const int32 NumPotentials = PotentialIntersections.Num();
+				for (int32 i = 0; i < NumPotentials; ++i)
+				{
+					auto& Particle2 = *PotentialIntersections[i].GetGeometryParticleHandle_PhysicsThread();
+					const TGenericParticleHandle<T, d> Particle2Generic(&Particle2);
+					// Collision group culling...
+					// CollisionGroup == 0 : Collide_With_Everything
+					// CollisionGroup == INDEX_NONE : Disabled collisions
+					// CollisionGroup_A != CollisionGroup_B : Skip Check
+
+					if (Particle1.CollisionGroup() == INDEX_NONE || Particle2Generic->CollisionGroup() == INDEX_NONE)
+					{
+						continue;
+					}
+					if (Particle1.CollisionGroup() && Particle2Generic->CollisionGroup() && Particle1.CollisionGroup() != Particle2Generic->CollisionGroup())
+					{
+						continue;
+					}
+
+					if (!Particle1.Geometry() && !Particle2.Geometry())
+					{
+						continue;
+					}
+
+					const bool bBody2Bounded = HasBoundingBox(Particle2);
+
+					if (Particle1.Handle() == Particle2.Handle())
+					{
+						continue;
+					}
+
+					if (bBody1Bounded == bBody2Bounded && Particle2.AsDynamic())
+					{
+						//if both are dynamic, assume index order matters
+#if CHAOS_DETERMINISTIC
+						if (Particle2.ParticleID() > Particle1.ParticleID())
+#else
+					//not deterministic just use memory address to avoid pair duplication
+						if (Particle2.Handle() > Particle1.Handle())
+#endif
+						{
+							continue;
+						}
+					}
+					//				++TotalTested;
+
+					const TVector<T, d> Box2Thickness = Particle2.AsDynamic() ? ComputeThickness(*Particle2.AsDynamic(), Dt) : TVector<T, d>(0);
+					const T UseThickness = FMath::Max(Box1Thickness, Box2Thickness.Size());// + MThickness
+
+					auto Constraint = ComputeConstraint(Particle1.Handle(), Particle2.Handle(), UseThickness);
+
+					//if (true || !InParticles.Geometry(Body1Index)->HasBoundingBox() || !InParticles.Geometry(Body2Index)->HasBoundingBox())
+					{
+						//SCOPE_CYCLE_COUNTER(STAT_ComputeConstraintsNP3);
+						//use narrow phase to determine if constraint is needed. Without this we can't do shock propagation
+
+#if !UE_BUILD_SHIPPING
+						if (bGatherStats)
+						{
+							++CountNP;
+						}
+#endif
+						UpdateConstraint<ECollisionUpdateType::Any>(UseThickness, Constraint);
+
+						if (Constraint.Phi < UseThickness)
+						{
+							Queue.Enqueue(Constraint);
+						}
+						else
+						{
+#if !UE_BUILD_SHIPPING
+							++RejectedNP;
+#endif
+						}
+					}
+
+				}
+
+#if !UE_BUILD_SHIPPING
+				NarrowPhasePerformed.Record(CountNP);
+				const int32 NPSkipped = NumPotentials - CountNP;
+				NarrowPhaseSkipped.Record(NPSkipped);
+				NarrowPhaseRejected.Record(RejectedNP);
+#endif
+			}, bGatherStats || CollisionConstraintsForceSingleThreaded);
+
+			while (!Queue.IsEmpty())
+			{
+				int32 Idx = Constraints.AddUninitialized(1);
+				Queue.Dequeue(Constraints[Idx]);
+				Handles.Add(HandleAllocator.AllocHandle(this, Idx));
+			}
+		}
+
+#if CHAOS_DETERMINISTIC
+		//todo: sort constraints
+#endif
+
+#if !UE_BUILD_SHIPPING
+		FString OutLog;
+		if (bGatherStats)
+		{
+#if CHAOS_PARTICLEHANDLE_TODO
+			const float NumParticles = InParticles.Size();
+			OutLog = FString::Printf(TEXT("ComputeConstraints stats:\n"
+				"Total Particles:%d\nSimulated Particles:%d (%.2f%%)\n"
+				"Max candidates per instance:%d (%.2f%%)\n"
+				"Max candidates skipped per instance (NP skipped):%d (%.2f%%)\n"
+				"Max narrow phase tests per instance:%d (%.2f%%)\n"
+				"Max narrow phase rejected per instance (NP rejected):%d (%.2f%%)\n"
+				"Constraints generated:%d\n"
+			),
+				InParticles.Size(),
+				SimulatedParticles, SimulatedParticles / NumParticles * 100.f,
+				BroadphasePotentials.MaxCount, BroadphasePotentials.MaxCount / NumParticles * 100.f,
+				NarrowPhaseSkipped.MaxCount, NarrowPhaseSkipped.MaxCount / NumParticles * 100.f,
+				NarrowPhasePerformed.MaxCount, NarrowPhasePerformed.MaxCount / NumParticles * 100.f,
+				NarrowPhaseRejected.MaxCount, NarrowPhaseRejected.MaxCount / NumParticles * 100.f,
+				Constraints.Num()
+			);
+
+			OutLog += FString::Printf(TEXT("Potentials per instance distribution:\n"));
+			OutLog += BroadphasePotentials.ToString();
+
+
+			OutLog += FString::Printf(TEXT("\nCandidates skipped per instance (NP skipped) distribution:\n"));
+			OutLog += NarrowPhaseSkipped.ToString();
+
+			OutLog += FString::Printf(TEXT("\nNarrow phase performed per instance distribution:\n"));
+			OutLog += NarrowPhasePerformed.ToString();
+
+			OutLog += FString::Printf(TEXT("\nNarrow phase candidates rejected per instance distribution:\n"));
+			OutLog += NarrowPhaseRejected.ToString();
+
+			OutLog += FString::Printf(TEXT("\nBounds distribution:\n"));
+			OutLog += BoundsDistribution.ToString();
+
+			UE_LOG(LogChaos, Warning, TEXT("%s"), *OutLog);
+			PendingHierarchyDump = 0;
+#endif
+		}
+#endif
+		//	if (TotalTested > 0)
+		//	{
+				//UE_LOG(LogChaos, Warning, TEXT("ComputeConstraints: rejected:%f out of rejected:%f = %.3f"), TotalRejected, TotalTested, TotalRejected / TotalTested);
+		//	}
+	}
+}
\ No newline at end of file
diff --git a/Engine/Source/Runtime/Experimental/Chaos/Public/Chaos/PBDRigidsEvolution.h b/Engine/Source/Runtime/Experimental/Chaos/Public/Chaos/PBDRigidsEvolution.h
--- a/Engine/Source/Runtime/Experimental/Chaos/Public/Chaos/PBDRigidsEvolution.h
+++ b/Engine/Source/Runtime/Experimental/Chaos/Public/Chaos/PBDRigidsEvolution.h
@@ -12,12 +12,16 @@
 #include "Chaos/Framework/DebugSubstep.h"
 #include "HAL/Event.h"
 #include "Chaos/PBDRigidsSOAs.h"
+#include "Chaos/ISpatialAccelerationCollection.h"
 
 namespace Chaos
 {
 template<class T, int d> class TPBDRigidsEvolutionGBF;
 class FChaosArchive;
 
+template <typename TPayload, typename T, int d>
+class ISpatialAccelerationCollection;
+
 struct CHAOS_API FEvolutionStats
 {
 	int32 ActiveCollisionPoints;
@@ -58,14 +62,7 @@ class TPBDRigidsEvolutionBase
 	typedef TFunction<void(TPBDRigidParticles<T, d>&, const T, const T, const int32)> FKinematicUpdateRule;
 
 	CHAOS_API TPBDRigidsEvolutionBase(TPBDRigidsSOAs<T, d>& InParticles, int32 InNumIterations = 1);
-	CHAOS_API virtual ~TPBDRigidsEvolutionBase()
-	{
-		Particles.GetParticleHandles().RemoveArray(&PhysicsMaterials);
-		Particles.GetParticleHandles().RemoveArray(&PerParticlePhysicsMaterials);
-		Particles.GetParticleHandles().RemoveArray(&ParticleDisableCount);
-		Particles.GetParticleHandles().RemoveArray(&Collided);
-		WaitOnAccelerationStructure();
-	}
+	CHAOS_API virtual ~TPBDRigidsEvolutionBase();
 
 	CHAOS_API TArray<TGeometryParticleHandle<T, d>*> CreateStaticParticles(int32 NumParticles, const TGeometryParticleParameters<T, d>& Params = TGeometryParticleParameters<T, d>())
 	{
@@ -261,41 +258,13 @@ class TPBDRigidsEvolutionBase
 	}
 
 	/** Make a copy of the acceleration structure to allow for external modification. This is needed for supporting sync operations on SQ structure from game thread */
-	CHAOS_API void UpdateExternalAccelerationStructure(TUniquePtr<ISpatialAcceleration<TAccelerationStructureHandle<T, d>, T, d>>& ExternalStructure);
-	ISpatialAcceleration<TAccelerationStructureHandle<T, d>, T, d>* GetSpatialAcceleration() { return InternalAcceleration.Get(); }
+	CHAOS_API void UpdateExternalAccelerationStructure(TUniquePtr<ISpatialAccelerationCollection<TAccelerationStructureHandle<T, d>, T, d>>& ExternalStructure);
+	ISpatialAccelerationCollection<TAccelerationStructureHandle<T, d>, T, d>* GetSpatialAcceleration() { return InternalAcceleration.Get(); }
 
 	const auto& GetRigidClustering() const { return Clustering; }
 	auto& GetRigidClustering() { return Clustering; }
 
-	void Serialize(FChaosArchive& Ar)
-	{
-		Particles.Serialize(Ar);
-
-		Ar.UsingCustomVersion(FExternalPhysicsCustomObjectVersion::GUID);
-		if (Ar.CustomVer(FExternalPhysicsCustomObjectVersion::GUID) >= FExternalPhysicsCustomObjectVersion::SerializeEvolutionBV)
-		{
-			Ar << InternalAcceleration;
-
-			SerializePendingMap(Ar, InternalAccelerationQueue);
-			SerializePendingMap(Ar, AsyncAccelerationQueue);
-			SerializePendingMap(Ar, ExternalAccelerationQueue);
-
-			ScratchExternalAcceleration = InternalAcceleration->Copy();
-		}
-		else if (Ar.IsLoading())
-		{
-			AccelerationStructureTaskComplete = nullptr;
-			for (auto& Particle : Particles.GetNonDisabledView())
-			{
-				DirtyParticle(Particle);
-			}
-
-			//force build acceleration structure with latest data
-			ComputeIntermediateSpatialAcceleration(true);
-			ComputeIntermediateSpatialAcceleration(true);	//having to do it multiple times because of the various caching involved over multiple frames.
-			ComputeIntermediateSpatialAcceleration(true);
-		}
-	}
+	void Serialize(FChaosArchive& Ar);
 
 protected:
 	int32 NumConstraints() const
@@ -320,7 +289,7 @@ protected:
 
 		//remove particle immediately for intermediate structure
 		InternalAccelerationQueue.Remove(Particle);
-		InternalAcceleration->RemoveElement(AsyncSpatialData.AccelerationHandle);
+		InternalAcceleration->RemoveElementFrom(AsyncSpatialData.AccelerationHandle, AsyncSpatialData.SpatialIdx);	//todo: use proper idx
 	}
 
 	void UpdateConstraintPositionBasedState(T Dt)
@@ -364,7 +333,7 @@ protected:
 		}
 	}
 
-	using FAccelerationStructure = typename FPBDCollisionConstraint::FAccelerationStructure;
+	using FAccelerationStructure = ISpatialAccelerationCollection<TAccelerationStructureHandle<T, d>, T, d>;
 
 	void ComputeIntermediateSpatialAcceleration(bool bBlock = false);
 	void FlushInternalAccelerationQueue();
@@ -396,11 +365,13 @@ protected:
 	struct FPendingSpatialData
 	{
 		TAccelerationStructureHandle<T, d> AccelerationHandle;
+		FSpatialAccelerationIdx SpatialIdx;
 		bool bUpdate;
 		bool bDelete;
 
 		FPendingSpatialData()
-			: bUpdate(false)
+			: SpatialIdx{ 0,0 }	//todo: set properly
+			, bUpdate(false)
 			, bDelete(false)
 		{}
 
@@ -486,7 +457,7 @@ protected:
 
 	int32 NumIterations;
 
-	static ISpatialAcceleration<TAccelerationStructureHandle<T, d>, T, d>* CreateNewSpatialStructure(const TArray<FAccelerationStructureBuilder>& Bounds);
+	static ISpatialAccelerationCollection<TAccelerationStructureHandle<T, d>, T, d>* CreateNewSpatialStructure(const TArray<FAccelerationStructureBuilder>& Bounds);
 };
 
 }
diff --git a/Engine/Source/Runtime/PhysicsSQ/Public/SQVerifier.h b/Engine/Source/Runtime/PhysicsSQ/Public/SQVerifier.h
--- a/Engine/Source/Runtime/PhysicsSQ/Public/SQVerifier.h
+++ b/Engine/Source/Runtime/PhysicsSQ/Public/SQVerifier.h
@@ -45,7 +45,7 @@ bool SQComparisonHelper(FPhysTestSerializer& Serializer, bool bEnsureOnMismatch
 		}
 
 		auto ChaosHitBuffer = MakeUnique<ChaosInterface::FSQHitBuffer<ChaosInterface::FRaycastHit>>();
-		TUniquePtr<ISpatialAcceleration<TAccelerationStructureHandle<float, 3>, float, 3>> Accelerator;
+		TUniquePtr<ISpatialAccelerationCollection<TAccelerationStructureHandle<float, 3>, float, 3>> Accelerator;
 		Serializer.GetChaosData()->UpdateExternalAccelerationStructure(Accelerator);
 		FChaosSQAccelerator SQAccelerator(*Accelerator); SQAccelerator.Raycast(CapturedSQ.StartPoint, CapturedSQ.Dir, CapturedSQ.DeltaMag, *ChaosHitBuffer, CapturedSQ.OutputFlags.HitFlags, CapturedSQ.QueryFilterData, *CapturedSQ.FilterCallback);
 		
@@ -87,7 +87,7 @@ bool SQComparisonHelper(FPhysTestSerializer& Serializer, bool bEnsureOnMismatch
 		}
 
 		auto ChaosHitBuffer = MakeUnique<ChaosInterface::FSQHitBuffer<ChaosInterface::FSweepHit>>();
-		TUniquePtr<ISpatialAcceleration<TAccelerationStructureHandle<float, 3>, float, 3>> Accelerator;
+		TUniquePtr<ISpatialAccelerationCollection<TAccelerationStructureHandle<float, 3>, float, 3>> Accelerator;
 		Serializer.GetChaosData()->UpdateExternalAccelerationStructure(Accelerator);
 		FChaosSQAccelerator SQAccelerator(*Accelerator); SQAccelerator.Sweep(*CapturedSQ.ChaosGeometry, CapturedSQ.StartTM, CapturedSQ.Dir, CapturedSQ.DeltaMag, *ChaosHitBuffer, CapturedSQ.OutputFlags.HitFlags, CapturedSQ.QueryFilterData, *CapturedSQ.FilterCallback);
 		
@@ -115,7 +115,7 @@ bool SQComparisonHelper(FPhysTestSerializer& Serializer, bool bEnsureOnMismatch
 		bTestPassed &= SQ_REPLAY_TEST(PxHitBuffer->GetNumHits() == CapturedSQ.PhysXOverlapBuffer.GetNumHits());
 
 		auto ChaosHitBuffer = MakeUnique<ChaosInterface::FSQHitBuffer<ChaosInterface::FOverlapHit>>();
-		TUniquePtr<ISpatialAcceleration<TAccelerationStructureHandle<float, 3>, float, 3>> Accelerator;
+		TUniquePtr<ISpatialAccelerationCollection<TAccelerationStructureHandle<float, 3>, float, 3>> Accelerator;
 		Serializer.GetChaosData()->UpdateExternalAccelerationStructure(Accelerator);
 		FChaosSQAccelerator SQAccelerator(*Accelerator);
 		SQAccelerator.Overlap(*CapturedSQ.ChaosGeometry, CapturedSQ.StartTM, *ChaosHitBuffer, CapturedSQ.QueryFilterData, *CapturedSQ.FilterCallback);
@@ -154,7 +154,7 @@ void SQPerfComparisonHelper(const FString& TestName, FPhysTestSerializer& Serial
 			}
 		}
 
-		TUniquePtr<ISpatialAcceleration<TAccelerationStructureHandle<float, 3>, float, 3>> Accelerator;
+		TUniquePtr<ISpatialAccelerationCollection<TAccelerationStructureHandle<float, 3>, float, 3>> Accelerator;
 		Serializer.GetChaosData()->UpdateExternalAccelerationStructure(Accelerator);
 		FChaosSQAccelerator SQAccelerator(*Accelerator);
 		for (double i = 0; i < NumIterations; ++i)
@@ -181,7 +181,7 @@ void SQPerfComparisonHelper(const FString& TestName, FPhysTestSerializer& Serial
 			}
 		}
 
-		TUniquePtr<ISpatialAcceleration<TAccelerationStructureHandle<float, 3>, float, 3>> Accelerator;
+		TUniquePtr<ISpatialAccelerationCollection<TAccelerationStructureHandle<float, 3>, float, 3>> Accelerator;
 		Serializer.GetChaosData()->UpdateExternalAccelerationStructure(Accelerator);
 		FChaosSQAccelerator SQAccelerator(*Accelerator);
 		for (double i = 0; i < NumIterations; ++i)
@@ -208,7 +208,7 @@ void SQPerfComparisonHelper(const FString& TestName, FPhysTestSerializer& Serial
 			}
 		}
 
-		TUniquePtr<ISpatialAcceleration<TAccelerationStructureHandle<float, 3>, float, 3>> Accelerator;
+		TUniquePtr<ISpatialAccelerationCollection<TAccelerationStructureHandle<float, 3>, float, 3>> Accelerator;
 		Serializer.GetChaosData()->UpdateExternalAccelerationStructure(Accelerator);
 		FChaosSQAccelerator SQAccelerator(*Accelerator);
 		for (double i = 0; i < NumIterations; ++i)
@@ -254,7 +254,7 @@ bool SQValidityHelper(FPhysTestSerializer& Serializer)
 		case FSQCapture::ESQType::Raycast:
 		{
 			ChaosInterface::FSQHitBuffer<ChaosInterface::FRaycastHit> ChaosHitBuffer;
-			TUniquePtr<ISpatialAcceleration<TAccelerationStructureHandle<float, 3>, float, 3>> Accelerator;
+			TUniquePtr<ISpatialAccelerationCollection<TAccelerationStructureHandle<float, 3>, float, 3>> Accelerator;
 			Serializer.GetChaosData()->UpdateExternalAccelerationStructure(Accelerator);
 			FChaosSQAccelerator SQAccelerator(*Accelerator);
 			SQAccelerator.Raycast(CapturedSQ.StartPoint, CapturedSQ.Dir, CapturedSQ.DeltaMag, ChaosHitBuffer, CapturedSQ.OutputFlags.HitFlags, CapturedSQ.QueryFilterData, *CapturedSQ.FilterCallback);
@@ -270,7 +270,7 @@ bool SQValidityHelper(FPhysTestSerializer& Serializer)
 		case FSQCapture::ESQType::Sweep:
 		{
 			ChaosInterface::FSQHitBuffer<ChaosInterface::FSweepHit> ChaosHitBuffer;
-			TUniquePtr<ISpatialAcceleration<TAccelerationStructureHandle<float, 3>, float, 3>> Accelerator;
+			TUniquePtr<ISpatialAccelerationCollection<TAccelerationStructureHandle<float, 3>, float, 3>> Accelerator;
 			Serializer.GetChaosData()->UpdateExternalAccelerationStructure(Accelerator);
 			FChaosSQAccelerator SQAccelerator(*Accelerator);
 			SQAccelerator.Sweep(*CapturedSQ.ChaosGeometry, CapturedSQ.StartTM, CapturedSQ.Dir, CapturedSQ.DeltaMag, ChaosHitBuffer, CapturedSQ.OutputFlags.HitFlags, CapturedSQ.QueryFilterData, *CapturedSQ.FilterCallback);
@@ -293,7 +293,7 @@ bool SQValidityHelper(FPhysTestSerializer& Serializer)
 		case FSQCapture::ESQType::Overlap:
 		{
 			ChaosInterface::FSQHitBuffer<ChaosInterface::FOverlapHit> ChaosHitBuffer;
-			TUniquePtr<ISpatialAcceleration<TAccelerationStructureHandle<float, 3>, float, 3>> Accelerator;
+			TUniquePtr<ISpatialAccelerationCollection<TAccelerationStructureHandle<float, 3>, float, 3>> Accelerator;
 			Serializer.GetChaosData()->UpdateExternalAccelerationStructure(Accelerator);
 			FChaosSQAccelerator SQAccelerator(*Accelerator);
 			SQAccelerator.Overlap(*CapturedSQ.ChaosGeometry, CapturedSQ.StartTM, ChaosHitBuffer, CapturedSQ.QueryFilterData, *CapturedSQ.FilterCallback);
