Fix scaled Time output for non-uniform scaled raycast and sweep. Pass HeadlessChaos tests again

diff --git a/Engine/Source/Runtime/Experimental/Chaos/Public/Chaos/ImplicitObjectScaled.h b/Engine/Source/Runtime/Experimental/Chaos/Public/Chaos/ImplicitObjectScaled.h
--- a/Engine/Source/Runtime/Experimental/Chaos/Public/Chaos/ImplicitObjectScaled.h
+++ b/Engine/Source/Runtime/Experimental/Chaos/Public/Chaos/ImplicitObjectScaled.h
@@ -90,13 +90,14 @@ public:
 		const TVector<T, d> UnscaledStart = MInvScale * StartPoint;
 		TVector<T, d> UnscaledDir = MScale * Dir * Length;
 		const T UnscaledLength = UnscaledDir.SafeNormalize();
+		const T ScaledLengthRatio = UnscaledLength / Length;
 		
 		TVector<T, d> UnscaledPosition;
 		TVector<T, d> UnscaledNormal;
 
 		if (MObject->Raycast(UnscaledStart, UnscaledDir, UnscaledLength, MInternalThickness + Thickness * MInvScale[0], OutTime, UnscaledPosition, UnscaledNormal, OutFaceIndex))
 		{
-			OutTime = UnscaledLength ? OutTime * (Length / UnscaledLength) : 0;	//is this needed? SafeNormalize above seems bad
+			OutTime = ScaledLengthRatio * OutTime;	//is this needed? SafeNormalize above seems bad
 			OutPosition = MScale * UnscaledPosition;
 			OutNormal = (MInvScale * UnscaledNormal).GetSafeNormal();
 			return true;
@@ -130,6 +131,7 @@ public:
 			UnscaledLength = sqrt(SizeSqr);
 			UnscaledDir /= UnscaledLength;
 		}
+		const T ScaledLengthRatio = UnscaledLength / Length;
 
 		TVector<T, d> UnscaledPosition;
 		TVector<T, d> UnscaledNormal;
@@ -142,7 +144,7 @@ public:
 		
 		if (InternalGeom.SweepGeom(ScaledB, BToATMNoScale, UnscaledDir, UnscaledLength, OutTime, UnscaledPosition, UnscaledNormal, OutFaceIndex, OwningScaled.MInternalThickness + Thickness))
 		{
-			OutTime = UnscaledLength ? OutTime * (Length / UnscaledLength) : 0;	//is this needed? SafeNormalize above seems bad
+			OutTime = ScaledLengthRatio * OutTime;	//is this needed? SafeNormalize above seems bad
 			LocalPosition = OwningScaled.MScale * UnscaledPosition;
 			LocalNormal = (OwningScaled.MInvScale * UnscaledNormal).GetSafeNormal();
 			return true;
