Export additional functions in TGeometryParticlesImp as needed for Mac build
#rnx

diff --git a/Engine/Source/Runtime/Experimental/Chaos/Public/Chaos/GeometryParticles.h b/Engine/Source/Runtime/Experimental/Chaos/Public/Chaos/GeometryParticles.h
--- a/Engine/Source/Runtime/Experimental/Chaos/Public/Chaos/GeometryParticles.h
+++ b/Engine/Source/Runtime/Experimental/Chaos/Public/Chaos/GeometryParticles.h
@@ -173,10 +173,10 @@ namespace Chaos
 		CHAOS_API virtual ~TGeometryParticlesImp()
 		{}
 
-		const TRotation<T, d>& R(const int32 Index) const { return MR[Index]; }
-		TRotation<T, d>& R(const int32 Index) { return MR[Index]; }
+		CHAOS_API const TRotation<T, d>& R(const int32 Index) const { return MR[Index]; }
+		CHAOS_API TRotation<T, d>& R(const int32 Index) { return MR[Index]; }
 
-		TSerializablePtr<TImplicitObject<T, d>> Geometry(const int32 Index) const { return MGeometry[Index]; }
+		CHAOS_API TSerializablePtr<TImplicitObject<T, d>> Geometry(const int32 Index) const { return MGeometry[Index]; }
 
 		const TUniquePtr<TImplicitObject<T, d>>& DynamicGeometry(const int32 Index) const { return MDynamicGeometry[Index]; }
 
@@ -185,8 +185,8 @@ namespace Chaos
 		const TShapesArray<T, d>& ShapesArray(const int32 Index) const { return MShapesArray[Index]; }
 
 #if CHAOS_DETERMINISTIC
-		FParticleID ParticleID(const int32 Idx) const { return MParticleIDs[Idx]; }
-		FParticleID& ParticleID(const int32 Idx) { return MParticleIDs[Idx]; }
+		CHAOS_API FParticleID ParticleID(const int32 Idx) const { return MParticleIDs[Idx]; }
+		CHAOS_API FParticleID& ParticleID(const int32 Idx) { return MParticleIDs[Idx]; }
 #endif
 
 		CHAOS_API void SetDynamicGeometry(const int32 Index, TUniquePtr<TImplicitObject<T, d>>&& InUnique)
@@ -223,12 +223,12 @@ namespace Chaos
 			return MLocalBounds[Index];
 		}
 
-		bool HasBounds(const int32 Index) const
+		CHAOS_API bool HasBounds(const int32 Index) const
 		{
 			return MHasBounds[Index];
 		}
 
-		bool& HasBounds(const int32 Index)
+		CHAOS_API bool& HasBounds(const int32 Index)
 		{
 			return MHasBounds[Index];
 		}
@@ -260,8 +260,8 @@ namespace Chaos
 
 		CHAOS_API void SetHandle(int32 Index, TGeometryParticleHandle<T, d>* Handle);
 		
-		TGeometryParticle<T, d>* GTGeometryParticle(const int32 Index) const { return MGeometryParticle[Index]; }
-		TGeometryParticle<T, d>*& GTGeometryParticle(const int32 Index) { return MGeometryParticle[Index]; }
+		CHAOS_API TGeometryParticle<T, d>* GTGeometryParticle(const int32 Index) const { return MGeometryParticle[Index]; }
+		CHAOS_API TGeometryParticle<T, d>*& GTGeometryParticle(const int32 Index) { return MGeometryParticle[Index]; }
 
 		FString ToString(int32 index) const
 		{
