Fix non-unity
#rnx

diff --git a/Engine/Source/Runtime/Experimental/Chaos/Public/Chaos/GeometryParticlesfwd.h b/Engine/Source/Runtime/Experimental/Chaos/Public/Chaos/GeometryParticlesfwd.h
--- a/Engine/Source/Runtime/Experimental/Chaos/Public/Chaos/GeometryParticlesfwd.h
+++ b/Engine/Source/Runtime/Experimental/Chaos/Public/Chaos/GeometryParticlesfwd.h
@@ -1,6 +1,8 @@
 // Copyright 1998-2019 Epic Games, Inc. All Rights Reserved.
 #pragma once
 
+#include "CoreMinimal.h"
+
 namespace Chaos
 {
 	enum class EGeometryParticlesSimType
