One more CHAOS_API add
#rnx

diff --git a/Engine/Source/Runtime/Experimental/Chaos/Public/Chaos/GeometryParticles.h b/Engine/Source/Runtime/Experimental/Chaos/Public/Chaos/GeometryParticles.h
--- a/Engine/Source/Runtime/Experimental/Chaos/Public/Chaos/GeometryParticles.h
+++ b/Engine/Source/Runtime/Experimental/Chaos/Public/Chaos/GeometryParticles.h
@@ -305,7 +305,7 @@ namespace Chaos
 			}
 		}
 
-		EParticleType ParticleType() const { return MParticleType; }
+		CHAOS_API EParticleType ParticleType() const { return MParticleType; }
 
 	protected:
 		EParticleType MParticleType;
