More CHAOS_API
#rnx

diff --git a/Engine/Source/Runtime/Experimental/Chaos/Public/Chaos/ParticleHandle.h b/Engine/Source/Runtime/Experimental/Chaos/Public/Chaos/ParticleHandle.h
--- a/Engine/Source/Runtime/Experimental/Chaos/Public/Chaos/ParticleHandle.h
+++ b/Engine/Source/Runtime/Experimental/Chaos/Public/Chaos/ParticleHandle.h
@@ -357,10 +357,10 @@ public:
 	const TShapesArray<T,d>& ShapesArray() const { return GeometryParticles->ShapesArray(ParticleIdx); }
 
 	const TBox<T, d>& LocalBounds() const { return GeometryParticles->LocalBounds(ParticleIdx); }
-	void SetLocalBounds(const TBox<T, d>& NewBounds) { GeometryParticles->LocalBounds(ParticleIdx) = NewBounds; }
+	CHAOS_API void SetLocalBounds(const TBox<T, d>& NewBounds) { GeometryParticles->LocalBounds(ParticleIdx) = NewBounds; }
 
 	const TBox<T, d>& WorldSpaceInflatedBounds() const { return GeometryParticles->WorldSpaceInflatedBounds(ParticleIdx); }
-	void SetWorldSpaceInflatedBounds(const TBox<T,d>& WorldSpaceInflatedBounds) { GeometryParticles->WorldSpaceInflatedBounds(ParticleIdx) = WorldSpaceInflatedBounds; }
+	CHAOS_API void SetWorldSpaceInflatedBounds(const TBox<T,d>& WorldSpaceInflatedBounds) { GeometryParticles->WorldSpaceInflatedBounds(ParticleIdx) = WorldSpaceInflatedBounds; }
 
 	bool HasBounds() const { return GeometryParticles->HasBounds(ParticleIdx); }
 	void SetHasBounds(bool bHasBounds) { GeometryParticles->HasBounds(ParticleIdx) = bHasBounds; }
