Reduce threshold for terminating gjkraycast.

diff --git a/Engine/Source/Runtime/Experimental/Chaos/Public/Chaos/GJK.h b/Engine/Source/Runtime/Experimental/Chaos/Public/Chaos/GJK.h
--- a/Engine/Source/Runtime/Experimental/Chaos/Public/Chaos/GJK.h
+++ b/Engine/Source/Runtime/Experimental/Chaos/Public/Chaos/GJK.h
@@ -528,7 +528,7 @@ namespace Chaos
 			V = SimplexFindClosestToOrigin(Simplex, SimplexIDs, Barycentric, As, Bs);
 
 			T NewDist2 = V.SizeSquared();	//todo: relative error
-			bNearZero = NewDist2 < 1e-8;
+			bNearZero = NewDist2 < 1e-6;
 			bDegenerate = NewDist2 >= InGJKPreDist2;
 			InGJKPreDist2 = NewDist2;
 			bTerminate = bNearZero || bDegenerate;
