Fix static analysis warnings about Dereferencing NULL pointer
#rnx

diff --git a/Engine/Source/Runtime/Experimental/Chaos/Private/Chaos/PBDConstraintGraph.cpp b/Engine/Source/Runtime/Experimental/Chaos/Private/Chaos/PBDConstraintGraph.cpp
--- a/Engine/Source/Runtime/Experimental/Chaos/Private/Chaos/PBDConstraintGraph.cpp
+++ b/Engine/Source/Runtime/Experimental/Chaos/Private/Chaos/PBDConstraintGraph.cpp
@@ -189,14 +189,14 @@ void TPBDConstraintGraph<T, d>::AddConstraint(const uint32 InContainerId, TConst
 	int32* PNodeIndex1 = (ConstrainedParticles[1])? ParticleToNodeIndex.Find(ConstrainedParticles[1]) : nullptr;
 	if (ensure(PNodeIndex0 || PNodeIndex1))
 	{
-		if (ConstrainedParticles[0])
+		if (PNodeIndex0)
 		{
 			NewEdge.FirstNode = *PNodeIndex0;
 			Nodes[NewEdge.FirstNode].Particle = ConstrainedParticles[0];
 			Nodes[NewEdge.FirstNode].Edges.Add(NewEdgeIndex);
 			UpdatedNodes.Add(NewEdge.FirstNode);
 		}
-		if (ConstrainedParticles[1])
+		if (PNodeIndex1)
 		{
 			NewEdge.SecondNode = *PNodeIndex1;
 			Nodes[NewEdge.SecondNode].Particle = ConstrainedParticles[1];
