Try additional exports
#rnx

diff --git a/Engine/Source/Runtime/Experimental/Chaos/Public/Chaos/GeometryParticles.h b/Engine/Source/Runtime/Experimental/Chaos/Public/Chaos/GeometryParticles.h
--- a/Engine/Source/Runtime/Experimental/Chaos/Public/Chaos/GeometryParticles.h
+++ b/Engine/Source/Runtime/Experimental/Chaos/Public/Chaos/GeometryParticles.h
@@ -213,12 +213,12 @@ namespace Chaos
 			UpdateShapesArray(Index);
 		}
 
-		const TBox<T,d>& LocalBounds(const int32 Index) const
+		CHAOS_API const TBox<T,d>& LocalBounds(const int32 Index) const
 		{
 			return MLocalBounds[Index];
 		}
 
-		TBox<T, d>& LocalBounds(const int32 Index)
+		CHAOS_API TBox<T, d>& LocalBounds(const int32 Index)
 		{
 			return MLocalBounds[Index];
 		}
