Plumbing to use spatial index from particles

diff --git a/Engine/Source/Runtime/Engine/Private/PhysicsEngine/Experimental/PhysScene_Chaos.cpp b/Engine/Source/Runtime/Engine/Private/PhysicsEngine/Experimental/PhysScene_Chaos.cpp
--- a/Engine/Source/Runtime/Engine/Private/PhysicsEngine/Experimental/PhysScene_Chaos.cpp
+++ b/Engine/Source/Runtime/Engine/Private/PhysicsEngine/Experimental/PhysScene_Chaos.cpp
@@ -580,8 +580,7 @@ void FPhysScene_Chaos::RemoveActorFromAccelerationStructure(FPhysicsActorHandle&
 	{
 		ExternalDataLock.WriteLock();
 		Chaos::TAccelerationStructureHandle<float, 3> AccelerationHandle(Actor);
-		Chaos::FSpatialAccelerationIdx Idx = { 0,0 };	//todo: use real idx;
-		GetSpacialAcceleration()->RemoveElementFrom(AccelerationHandle, Idx);
+		GetSpacialAcceleration()->RemoveElementFrom(AccelerationHandle, Actor->SpatialIdx());
 		ExternalDataLock.WriteUnlock();
 	}
 #endif
diff --git a/Engine/Source/Runtime/Experimental/Chaos/Public/Chaos/GeometryParticlesfwd.h b/Engine/Source/Runtime/Experimental/Chaos/Public/Chaos/GeometryParticlesfwd.h
--- a/Engine/Source/Runtime/Experimental/Chaos/Public/Chaos/GeometryParticlesfwd.h
+++ b/Engine/Source/Runtime/Experimental/Chaos/Public/Chaos/GeometryParticlesfwd.h
@@ -27,5 +27,10 @@ namespace Chaos
 		uint16 InnerIdx : 13;
 
 		static constexpr uint16 MaxBucketEntries = 1 << 13;
+
+		bool operator==(const FSpatialAccelerationIdx& Rhs) const
+		{
+			return ((const uint16&)*this) == ((const uint16&)Rhs);
+		}
 	};
 }
\ No newline at end of file
diff --git a/Engine/Source/Runtime/Experimental/Chaos/Public/Chaos/PBDRigidsEvolution.h b/Engine/Source/Runtime/Experimental/Chaos/Public/Chaos/PBDRigidsEvolution.h
--- a/Engine/Source/Runtime/Experimental/Chaos/Public/Chaos/PBDRigidsEvolution.h
+++ b/Engine/Source/Runtime/Experimental/Chaos/Public/Chaos/PBDRigidsEvolution.h
@@ -151,6 +151,8 @@ class TPBDRigidsEvolutionBase
 		FPendingSpatialData& SpatialData = InternalAccelerationQueue.FindOrAdd(Particle.Handle());
 		SpatialData.AccelerationHandle = TAccelerationStructureHandle<T,d>(Particle);
 		SpatialData.bUpdate = true;
+		ensure(!SpatialData.bDelete || SpatialData.SpatialIdx == Particle.SpatialIdx());	//delete and add got same memory, but in a different spatial acceleration (need to handle case)
+		SpatialData.SpatialIdx = Particle.SpatialIdx();	//todo: handle case where we delete 
 
 		AsyncAccelerationQueue.FindOrAdd(Particle.Handle()) = SpatialData;
 		ExternalAccelerationQueue.FindOrAdd(Particle.Handle()) = SpatialData;
@@ -285,6 +287,7 @@ protected:
 		AsyncSpatialData.AccelerationHandle = TAccelerationStructureHandle<T, d>(ParticleHandle);
 		AsyncSpatialData.bUpdate = false;	//don't bother updating since deleting anyway
 		AsyncSpatialData.bDelete = true;
+		AsyncSpatialData.SpatialIdx = ParticleHandle.SpatialIdx();
 		ExternalAccelerationQueue.FindOrAdd(Particle) = AsyncSpatialData;
 
 		//remove particle immediately for intermediate structure
@@ -370,8 +373,7 @@ protected:
 		bool bDelete;
 
 		FPendingSpatialData()
-			: SpatialIdx{ 0,0 }	//todo: set properly
-			, bUpdate(false)
+			: bUpdate(false)
 			, bDelete(false)
 		{}
 
diff --git a/Engine/Source/Runtime/Experimental/Chaos/Public/Chaos/ParticleDirtyFlags.h b/Engine/Source/Runtime/Experimental/Chaos/Public/Chaos/ParticleDirtyFlags.h
--- a/Engine/Source/Runtime/Experimental/Chaos/Public/Chaos/ParticleDirtyFlags.h
+++ b/Engine/Source/Runtime/Experimental/Chaos/Public/Chaos/ParticleDirtyFlags.h
@@ -30,6 +30,7 @@ namespace Chaos
 		ExternalForce =  1 << 18,
 		ExternalTorque = 1 << 19,
 		GravityEnabled = 1 << 20,
+		SpatialIdx     = 1 << 21
 	};
 
 	class FParticleDirtyFlags
diff --git a/Engine/Source/Runtime/Experimental/Chaos/Public/Chaos/ParticleHandle.h b/Engine/Source/Runtime/Experimental/Chaos/Public/Chaos/ParticleHandle.h
--- a/Engine/Source/Runtime/Experimental/Chaos/Public/Chaos/ParticleHandle.h
+++ b/Engine/Source/Runtime/Experimental/Chaos/Public/Chaos/ParticleHandle.h
@@ -47,6 +47,7 @@ void GeometryParticleDefaultConstruct(FConcrete& Concrete, const TGeometryPartic
 {
 	Concrete.SetX(TVector<T, d>(0));
 	Concrete.SetR(TRotation<T, d>::Identity);
+	Concrete.SetSpatialIdx(FSpatialAccelerationIdx{ 0,0 });
 }
 
 template <typename T, int d, typename FConcrete>
@@ -364,6 +365,9 @@ public:
 
 	bool HasBounds() const { return GeometryParticles->HasBounds(ParticleIdx); }
 	void SetHasBounds(bool bHasBounds) { GeometryParticles->HasBounds(ParticleIdx) = bHasBounds; }
+
+	FSpatialAccelerationIdx SpatialIdx() const { return GeometryParticles->SpatialIdx(ParticleIdx); }
+	void SetSpatialIdx(FSpatialAccelerationIdx Idx) { GeometryParticles->SpatialIdx(ParticleIdx) = Idx; }
 	
 	EObjectStateType ObjectState() const;
 
@@ -1036,6 +1040,13 @@ public:
 	const TPBDRigidParticle<T, d>* AsDynamic() const;
 	TPBDRigidParticle<T, d>* AsDynamic();
 
+	FSpatialAccelerationIdx SpatialIdx() const { return MSpatialIdx; }
+	void SetSpatialIdx(FSpatialAccelerationIdx Idx)
+	{
+		MarkDirty(EParticleFlags::SpatialIdx);
+		MSpatialIdx = Idx;
+	}
+
 	FParticleData* NewData() const { return new TGeometryParticleData<T, d>( *this ); }
 
 	bool IsDirty() const
@@ -1073,6 +1084,7 @@ private:
 	TRotation<T, d> MR;
 	TSharedPtr<TImplicitObject<T, d>, ESPMode::ThreadSafe> MGeometry;	//TODO: geometry should live in bodysetup
 	TShapesArray<T,d> MShapesArray;
+	FSpatialAccelerationIdx MSpatialIdx;
 
 	// Pointer to some arbitrary data associated with the particle, but not used by Chaos. External systems may use this for whatever.
 	void* MUserData;
@@ -1113,13 +1125,15 @@ public:
 		: FParticleData(InType)
 		, X(TVector<T, d>(0))
 		, R(TRotation<T, d>())
+		, SpatialIdx(FSpatialAccelerationIdx{ 0,0 })
 	{}
 
-	TGeometryParticleData(const TGeometryParticle<T,d>& InParticle)
+	TGeometryParticleData(const TGeometryParticle<T, d>& InParticle)
 		: FParticleData(EParticleType::Static)
 		, X(InParticle.X())
 		, R(InParticle.R())
 		, Geometry(InParticle.GeometrySharedLowLevel())
+		, SpatialIdx(InParticle.SpatialIdx())
 	{}
 
 	void Reset() 
@@ -1128,11 +1142,13 @@ public:
 		X = TVector<T, d>(0); 
 		R = TRotation<T, d>(); 
 		Geometry = TSharedPtr<TImplicitObjectUnion<T, d>, ESPMode::ThreadSafe>();
+		SpatialIdx = FSpatialAccelerationIdx{ 0,0 };
 	}
 
 	TVector<T, d> X;
 	TRotation<T, d> R;
 	TSharedPtr<TImplicitObject<T, d>, ESPMode::ThreadSafe> Geometry;
+	FSpatialAccelerationIdx SpatialIdx;
 };
 
 
diff --git a/Engine/Source/Runtime/Experimental/ChaosSolvers/Private/PhysicsProxy/SingleParticlePhysicsProxy.cpp b/Engine/Source/Runtime/Experimental/ChaosSolvers/Private/PhysicsProxy/SingleParticlePhysicsProxy.cpp
--- a/Engine/Source/Runtime/Experimental/ChaosSolvers/Private/PhysicsProxy/SingleParticlePhysicsProxy.cpp
+++ b/Engine/Source/Runtime/Experimental/ChaosSolvers/Private/PhysicsProxy/SingleParticlePhysicsProxy.cpp
@@ -59,6 +59,7 @@ void FSingleParticlePhysicsProxy<Chaos::TGeometryParticle<float, 3>>::PushToPhys
 			RigidHandle->SetLocalBounds(Data->Geometry->BoundingBox());
 			RigidHandle->SetWorldSpaceInflatedBounds(Data->Geometry->BoundingBox().TransformedBox(Chaos::TRigidTransform<float, 3>(Data->X, Data->R)));
 		}
+		RigidHandle->SetSpatialIdx(Data->SpatialIdx);	//todo: this needs to only happen once during initialization
 	}
 }
 
