More CHAOS_API
#rnx

diff --git a/Engine/Source/Runtime/Experimental/Chaos/Public/Chaos/GeometryParticles.h b/Engine/Source/Runtime/Experimental/Chaos/Public/Chaos/GeometryParticles.h
--- a/Engine/Source/Runtime/Experimental/Chaos/Public/Chaos/GeometryParticles.h
+++ b/Engine/Source/Runtime/Experimental/Chaos/Public/Chaos/GeometryParticles.h
@@ -233,12 +233,12 @@ namespace Chaos
 			return MHasBounds[Index];
 		}
 
-		FSpatialAccelerationIdx SpatialIdx(const int32 Index) const
+		CHAOS_API FSpatialAccelerationIdx SpatialIdx(const int32 Index) const
 		{
 			return MSpatialIdx[Index];
 		}
 
-		FSpatialAccelerationIdx& SpatialIdx(const int32 Index)
+		CHAOS_API FSpatialAccelerationIdx& SpatialIdx(const int32 Index)
 		{
 			return MSpatialIdx[Index];
 		}
