CHAOS_API applications is getting silly
#rnx

diff --git a/Engine/Source/Runtime/Experimental/Chaos/Public/Chaos/ParticleHandle.h b/Engine/Source/Runtime/Experimental/Chaos/Public/Chaos/ParticleHandle.h
--- a/Engine/Source/Runtime/Experimental/Chaos/Public/Chaos/ParticleHandle.h
+++ b/Engine/Source/Runtime/Experimental/Chaos/Public/Chaos/ParticleHandle.h
@@ -360,7 +360,7 @@ public:
 	const TBox<T, d>& LocalBounds() const { return GeometryParticles->LocalBounds(ParticleIdx); }
 	CHAOS_API void SetLocalBounds(const TBox<T, d>& NewBounds) { GeometryParticles->LocalBounds(ParticleIdx) = NewBounds; }
 
-	const TBox<T, d>& WorldSpaceInflatedBounds() const { return GeometryParticles->WorldSpaceInflatedBounds(ParticleIdx); }
+	CHAOS_API const TBox<T, d>& WorldSpaceInflatedBounds() const { return GeometryParticles->WorldSpaceInflatedBounds(ParticleIdx); }
 	CHAOS_API void SetWorldSpaceInflatedBounds(const TBox<T,d>& WorldSpaceInflatedBounds) { GeometryParticles->WorldSpaceInflatedBounds(ParticleIdx) = WorldSpaceInflatedBounds; }
 
 	bool HasBounds() const { return GeometryParticles->HasBounds(ParticleIdx); }
