Back out CHAOS_API changes to get main running again.

diff --git a/Engine/Source/Runtime/Experimental/Chaos/Public/Chaos/ParticleHandle.h b/Engine/Source/Runtime/Experimental/Chaos/Public/Chaos/ParticleHandle.h
--- a/Engine/Source/Runtime/Experimental/Chaos/Public/Chaos/ParticleHandle.h
+++ b/Engine/Source/Runtime/Experimental/Chaos/Public/Chaos/ParticleHandle.h
@@ -358,10 +358,10 @@ public:
 	const TShapesArray<T,d>& ShapesArray() const { return GeometryParticles->ShapesArray(ParticleIdx); }
 
 	const TBox<T, d>& LocalBounds() const { return GeometryParticles->LocalBounds(ParticleIdx); }
-	CHAOS_API void SetLocalBounds(const TBox<T, d>& NewBounds) { GeometryParticles->LocalBounds(ParticleIdx) = NewBounds; }
+	void SetLocalBounds(const TBox<T, d>& NewBounds) { GeometryParticles->LocalBounds(ParticleIdx) = NewBounds; }
 
-	CHAOS_API const TBox<T, d>& WorldSpaceInflatedBounds() const { return GeometryParticles->WorldSpaceInflatedBounds(ParticleIdx); }
-	CHAOS_API void SetWorldSpaceInflatedBounds(const TBox<T,d>& WorldSpaceInflatedBounds) { GeometryParticles->WorldSpaceInflatedBounds(ParticleIdx) = WorldSpaceInflatedBounds; }
+	const TBox<T, d>& WorldSpaceInflatedBounds() const { return GeometryParticles->WorldSpaceInflatedBounds(ParticleIdx); }
+	void SetWorldSpaceInflatedBounds(const TBox<T,d>& WorldSpaceInflatedBounds) { GeometryParticles->WorldSpaceInflatedBounds(ParticleIdx) = WorldSpaceInflatedBounds; }
 
 	bool HasBounds() const { return GeometryParticles->HasBounds(ParticleIdx); }
 	void SetHasBounds(bool bHasBounds) { GeometryParticles->HasBounds(ParticleIdx) = bHasBounds; }
