Remove no longer used #ifndef/#define
#rnx

diff --git a/Engine/Source/Runtime/Experimental/ChaosSolvers/Public/PBDRigidsSolver.h b/Engine/Source/Runtime/Experimental/ChaosSolvers/Public/PBDRigidsSolver.h
--- a/Engine/Source/Runtime/Experimental/ChaosSolvers/Public/PBDRigidsSolver.h
+++ b/Engine/Source/Runtime/Experimental/ChaosSolvers/Public/PBDRigidsSolver.h
@@ -24,12 +24,6 @@
 #include "PhysicsProxy/SingleParticlePhysicsProxyFwd.h"
 #include "SolverEventFilters.h"
 
-
-
-#ifndef CHAOS_WITH_PAUSABLE_SOLVER
-#define CHAOS_WITH_PAUSABLE_SOLVER 1
-#endif
-
 class FPhysicsSolverAdvanceTask;
 class FPhysInterface_Chaos;
 
