More plumbing to allow for a spatial acceleration collection

diff --git a/Engine/Source/Runtime/Experimental/Chaos/Private/Chaos/PBDRigidsEvolution.cpp b/Engine/Source/Runtime/Experimental/Chaos/Private/Chaos/PBDRigidsEvolution.cpp
--- a/Engine/Source/Runtime/Experimental/Chaos/Private/Chaos/PBDRigidsEvolution.cpp
+++ b/Engine/Source/Runtime/Experimental/Chaos/Private/Chaos/PBDRigidsEvolution.cpp
@@ -44,10 +44,10 @@ namespace Chaos
 	DECLARE_CYCLE_STAT(TEXT("ComputeIntermediateSpatialAcceleration"), STAT_ComputeIntermediateSpatialAcceleration, STATGROUP_Chaos);
 
 	template<class FPBDRigidsEvolution, class FPBDCollisionConstraint, class T, int d>
-	TPBDRigidsEvolutionBase<FPBDRigidsEvolution, FPBDCollisionConstraint, T, d>::FChaosAccelerationStructureTask::FChaosAccelerationStructureTask(const TArray<FAccelerationStructureBuilder>& InBounds
+	TPBDRigidsEvolutionBase<FPBDRigidsEvolution, FPBDCollisionConstraint, T, d>::FChaosAccelerationStructureTask::FChaosAccelerationStructureTask(const TArray<FAccelerationStructBuilderCache>& CacheMap
 		, TUniquePtr<FAccelerationStructure>& InAccelerationStructure
 		, TUniquePtr<FAccelerationStructure>& InAccelerationStructureCopy)
-		: CachedSpatialBuilderData(InBounds)
+		: BuilderCacheMap(CacheMap)
 		, AccelerationStructure(InAccelerationStructure)
 		, AccelerationStructureCopy(InAccelerationStructureCopy)
 	{
@@ -151,7 +151,30 @@ namespace Chaos
 	template<class FPBDRigidsEvolution, class FPBDCollisionConstraint, class T, int d>
 	void TPBDRigidsEvolutionBase<FPBDRigidsEvolution, FPBDCollisionConstraint, T, d>::FChaosAccelerationStructureTask::DoTask(ENamedThreads::Type CurrentThread, const FGraphEventRef& MyCompletionGraphEvent)
 	{
-		AccelerationStructure.Reset(CreateNewSpatialStructure(CachedSpatialBuilderData));
+		//todo: no reason to do this in serial
+		for (const FAccelerationStructBuilderCache& Cache : BuilderCacheMap)
+		{
+			auto OldStruct = AccelerationStructure->RemoveSubstructure(Cache.SpatialIdx);
+			//This is a hack, need to have a Reinitialize API instead of downcasting
+			using BVType = TBoundingVolume<TAccelerationStructureHandle<T, d>, T, d>;
+			using AABBType = TAABBTree<TAccelerationStructureHandle<T, d>, TAABBTreeLeafArray<TAccelerationStructureHandle<T, d>, T>, T>;
+			ISpatialAcceleration<TAccelerationStructureHandle<T,d>,T,d>* NewStruct;
+			if (OldStruct->template As<BVType>())
+			{
+				NewStruct = new BVType(*Cache.CachedSpatialBuilderData);
+			}
+			else if (OldStruct->template As<AABBType>())
+			{
+				NewStruct = new AABBType(*Cache.CachedSpatialBuilderData);
+			}
+			else
+			{
+				using AccelType = TAABBTree<TAccelerationStructureHandle<T, d>, TBoundingVolume<TAccelerationStructureHandle<T, d>, T, d>, T>;
+				NewStruct = new AccelType(*Cache.CachedSpatialBuilderData);
+			}
+
+			AccelerationStructure->AddSubstructure(TUniquePtr<ISpatialAcceleration<TAccelerationStructureHandle<T, d>, T,d>>(NewStruct), Cache.SpatialIdx.Bucket); //this assumes bucket size is 1, which will not be true in future. Need to fix all this
+		}
 		AccelerationStructureCopy = AsUniqueSpatialAccelerationChecked<FAccelerationStructure>(AccelerationStructure->Copy());
 	}
 
@@ -166,12 +189,14 @@ namespace Chaos
 		//As long as we delete first and update second this will be respected
 		if (SpatialData.bDelete)
 		{
-			AccelerationStructure.RemoveElementFrom(SpatialData.AccelerationHandle, SpatialData.SpatialIdx);
+			AccelerationStructure.RemoveElementFrom(SpatialData.AccelerationHandle, SpatialData.DeletedSpatialIdx);
 
 			if (bAsync)
 			{
 				if (int32* CacheIdxPtr = ParticleToCacheIdx.Find(Particle))
 				{
+					const auto SpatialIdx = SpatialData.DeletedSpatialIdx;
+					TArray<FAccelerationStructureBuilder>& CachedSpatialBuilderData = *CachedSpatialBuilderDataMap.FindByPredicate([SpatialIdx](const auto& Cache){ return Cache.SpatialIdx == SpatialIdx;})->CachedSpatialBuilderData;
 					const int32 CacheIdx = *CacheIdxPtr;
 					if (CacheIdx + 1 < CachedSpatialBuilderData.Num())	//will get swapped with last element, so update it
 					{
@@ -188,10 +213,13 @@ namespace Chaos
 
 		if(SpatialData.bUpdate)
 		{
-			AccelerationStructure.UpdateElementIn(Particle, Particle->WorldSpaceInflatedBounds(), Particle->HasBounds(), SpatialData.SpatialIdx);	//todo: can the spatial idx change when we recreate and re-use memory? (yes, potential bug)
+			AccelerationStructure.UpdateElementIn(Particle, Particle->WorldSpaceInflatedBounds(), Particle->HasBounds(), SpatialData.UpdatedSpatialIdx);
 			
 			if (bAsync)
 			{
+				const auto SpatialIdx = SpatialData.UpdatedSpatialIdx;
+				TArray<FAccelerationStructureBuilder>& CachedSpatialBuilderData = *CachedSpatialBuilderDataMap.FindByPredicate([SpatialIdx](const auto& Cache) { return Cache.SpatialIdx == SpatialIdx; })->CachedSpatialBuilderData;
+
 				//make sure in mapping
 				int32 CacheIdx;
 				if (int32* CacheIdxPtr = ParticleToCacheIdx.Find(Particle))
@@ -274,6 +302,7 @@ namespace Chaos
 			FlushInternalAccelerationQueue();
 			FlushExternalAccelerationQueue(*ScratchExternalAcceleration);
 			bExternalReady = true;
+			InitializeAccelerationCache();
 		}
 
 		if (bBlock)
@@ -293,7 +322,7 @@ namespace Chaos
 				bExternalReady = true;
 			}
 
-			AccelerationStructureTaskComplete = TGraphTask<FChaosAccelerationStructureTask>::CreateTask().ConstructAndDispatchWhenReady(CachedSpatialBuilderData, AsyncInternalAcceleration, AsyncExternalAcceleration);
+			AccelerationStructureTaskComplete = TGraphTask<FChaosAccelerationStructureTask>::CreateTask().ConstructAndDispatchWhenReady(CachedSpatialBuilderDataMap, AsyncInternalAcceleration, AsyncExternalAcceleration);
 		}
 		else
 		{
@@ -301,6 +330,17 @@ namespace Chaos
 		}
 	}
 
+	template<class FPBDRigidsEvolution, class FPBDCollisionConstraint, class T, int d>
+	void TPBDRigidsEvolutionBase<FPBDRigidsEvolution, FPBDCollisionConstraint, T, d>::InitializeAccelerationCache()
+	{
+		CachedSpatialBuilderDataMap.Empty();
+		const auto SpatialIndices = InternalAcceleration->GetAllSpatialIndices();
+		for (const auto& Idx : SpatialIndices)
+		{
+			CachedSpatialBuilderDataMap.Add(Idx);
+		}
+	}
+
 	template<class FPBDRigidsEvolution, class FPBDCollisionConstraint, class T, int d>
 	void TPBDRigidsEvolutionBase<FPBDRigidsEvolution, FPBDCollisionConstraint, T, d>::UpdateExternalAccelerationStructure(TUniquePtr<ISpatialAccelerationCollection<TAccelerationStructureHandle<T, d>, T, d>>& StructToUpdate)
 	{
@@ -337,6 +377,7 @@ namespace Chaos
 			{
 				Ar << SubStructure;
 				InternalAcceleration = CreateNewSpatialStructureFromSubStructure(MoveTemp(SubStructure));
+				InitializeAccelerationCache();
 			}
 
 			SerializePendingMap(Ar, InternalAccelerationQueue);
diff --git a/Engine/Source/Runtime/Experimental/Chaos/Public/Chaos/ISpatialAccelerationCollection.h b/Engine/Source/Runtime/Experimental/Chaos/Public/Chaos/ISpatialAccelerationCollection.h
--- a/Engine/Source/Runtime/Experimental/Chaos/Public/Chaos/ISpatialAccelerationCollection.h
+++ b/Engine/Source/Runtime/Experimental/Chaos/Public/Chaos/ISpatialAccelerationCollection.h
@@ -22,6 +22,7 @@ public:
 	/** This is kind of a hack to avoid virtuals. We simply route calls into templated functions */
 	virtual void PBDComputeConstraintsLowLevel_GatherStats(TPBDCollisionConstraint<T, d>& CollisionConstraint, T Dt) const = 0;
 	virtual void PBDComputeConstraintsLowLevel(TPBDCollisionConstraint<T, d>& CollisionConstraint, T Dt) const = 0;
+	virtual TArray<FSpatialAccelerationIdx> GetAllSpatialIndices() const = 0;
 };
 
 template <typename TSpatialAcceleration>
@@ -105,6 +106,16 @@ public:
 		}
 	}
 
+	void GetAllSpatialIndices(TArray<FSpatialAccelerationIdx>& Indices, uint16 BucketIdx) const
+	{
+		Indices.Reserve(Indices.Num() + Accelerations.Num());
+		uint16 InnerIdx = 0;
+		for (const auto& Acceleration : Accelerations)
+		{
+			Indices.Add(FSpatialAccelerationIdx{ BucketIdx, InnerIdx++ });
+		}
+	}
+
 	void RemoveElement(const TPayloadType& Payload, int32 Idx)
 	{
 		Accelerations[Idx]->RemoveElement(Payload);
@@ -287,6 +298,16 @@ struct TSpatialAccelerationCollectionHelper
 			TSpatialAccelerationCollectionHelper < NextBucket < NumBuckets ? NextBucket : 0, NumBuckets, Tuple, TPayloadType, T, d>::GlobalObjects(Buckets, ObjList);
 		}
 	}
+
+	static void GetAllSpatialIndices(const Tuple& Buckets, TArray<FSpatialAccelerationIdx>& Indices)
+	{
+		GetBucket<BucketIdx>(Buckets).GetAllSpatialIndices(Indices, BucketIdx);
+		constexpr int NextBucket = BucketIdx + 1;
+		if (NextBucket < NumBuckets)
+		{
+			TSpatialAccelerationCollectionHelper < NextBucket < NumBuckets ? NextBucket : 0, NumBuckets, Tuple, TPayloadType, T, d>::GetAllSpatialIndices(Buckets, Indices);
+		}
+	}
 };
 
 template <bool bGatherStats, typename SpatialAccelerationCollection, typename T, int d>
@@ -413,6 +434,13 @@ public:
 		return ObjList;
 	}
 
+	virtual TArray<FSpatialAccelerationIdx> GetAllSpatialIndices() const override
+	{
+		TArray<FSpatialAccelerationIdx> Indices;
+		TSpatialAccelerationCollectionHelper<0, NumBuckets, decltype(Buckets), TPayloadType, T, d>::GetAllSpatialIndices(Buckets, Indices);
+		return Indices;
+	}
+
 	virtual void RemoveElementFrom(const TPayloadType& Payload, FSpatialAccelerationIdx SpatialIdx) override
 	{
 		check(SpatialIdx.Bucket < NumBuckets);
diff --git a/Engine/Source/Runtime/Experimental/Chaos/Public/Chaos/PBDRigidsEvolution.h b/Engine/Source/Runtime/Experimental/Chaos/Public/Chaos/PBDRigidsEvolution.h
--- a/Engine/Source/Runtime/Experimental/Chaos/Public/Chaos/PBDRigidsEvolution.h
+++ b/Engine/Source/Runtime/Experimental/Chaos/Public/Chaos/PBDRigidsEvolution.h
@@ -151,8 +151,7 @@ class TPBDRigidsEvolutionBase
 		FPendingSpatialData& SpatialData = InternalAccelerationQueue.FindOrAdd(Particle.Handle());
 		SpatialData.AccelerationHandle = TAccelerationStructureHandle<T,d>(Particle);
 		SpatialData.bUpdate = true;
-		ensure(!SpatialData.bDelete || SpatialData.SpatialIdx == Particle.SpatialIdx());	//delete and add got same memory, but in a different spatial acceleration (need to handle case)
-		SpatialData.SpatialIdx = Particle.SpatialIdx();	//todo: handle case where we delete 
+		SpatialData.UpdatedSpatialIdx = Particle.SpatialIdx();
 
 		AsyncAccelerationQueue.FindOrAdd(Particle.Handle()) = SpatialData;
 		ExternalAccelerationQueue.FindOrAdd(Particle.Handle()) = SpatialData;
@@ -285,14 +284,21 @@ protected:
 		auto Particle = ParticleHandle.Handle();
 		FPendingSpatialData& AsyncSpatialData = AsyncAccelerationQueue.FindOrAdd(Particle);
 		AsyncSpatialData.AccelerationHandle = TAccelerationStructureHandle<T, d>(ParticleHandle);
+		if (!AsyncSpatialData.bDelete)
+		{
+			//There are three cases to consider:
+			//Simple single delete happens, in that case just use the index you see (the first and only index)
+			//Delete followed by any number deletes and updates and finally an update, in that case we must delete the first particle and add the final (so use first index for delete)
+			//Delete followed by multiple updates and or deletes and a final delete. In that case we still only delete the first particle since the final delete is not really needed (add will be cancelled)
+			AsyncSpatialData.DeletedSpatialIdx = ParticleHandle.SpatialIdx();
+		}
 		AsyncSpatialData.bUpdate = false;	//don't bother updating since deleting anyway
 		AsyncSpatialData.bDelete = true;
-		AsyncSpatialData.SpatialIdx = ParticleHandle.SpatialIdx();
 		ExternalAccelerationQueue.FindOrAdd(Particle) = AsyncSpatialData;
 
 		//remove particle immediately for intermediate structure
 		InternalAccelerationQueue.Remove(Particle);
-		InternalAcceleration->RemoveElementFrom(AsyncSpatialData.AccelerationHandle, AsyncSpatialData.SpatialIdx);	//todo: use proper idx
+		InternalAcceleration->RemoveElementFrom(AsyncSpatialData.AccelerationHandle, AsyncSpatialData.DeletedSpatialIdx);	//even though we remove immediately, future adds are still pending
 	}
 
 	void UpdateConstraintPositionBasedState(T Dt)
@@ -368,7 +374,8 @@ protected:
 	struct FPendingSpatialData
 	{
 		TAccelerationStructureHandle<T, d> AccelerationHandle;
-		FSpatialAccelerationIdx SpatialIdx;
+		FSpatialAccelerationIdx UpdatedSpatialIdx;
+		FSpatialAccelerationIdx DeletedSpatialIdx;	//need both updated and deleted in case memory is reused but a different idx is neede
 		bool bUpdate;
 		bool bDelete;
 
@@ -436,22 +443,33 @@ protected:
 	};
 
 	/** Used for async acceleration rebuild */
-	TArray<FAccelerationStructureBuilder> CachedSpatialBuilderData;
 	TMap<TGeometryParticleHandle<T, d>*, int32> ParticleToCacheIdx;
 
+	struct FAccelerationStructBuilderCache
+	{
+		FSpatialAccelerationIdx SpatialIdx;
+		TUniquePtr <TArray<FAccelerationStructureBuilder>> CachedSpatialBuilderData;
+
+		FAccelerationStructBuilderCache(FSpatialAccelerationIdx Idx)
+			: SpatialIdx(Idx)
+			, CachedSpatialBuilderData(MakeUnique<TArray<FAccelerationStructureBuilder>>())
+		{}
+	};
+	TArray<FAccelerationStructBuilderCache> CachedSpatialBuilderDataMap;
+
 	FORCEINLINE_DEBUGGABLE void ApplyParticlePendingData(TGeometryParticleHandle<T, d>* Particle, const FPendingSpatialData& PendingData, FAccelerationStructure& SpatialAcceleration, bool bAsync);
 
 	class FChaosAccelerationStructureTask
 	{
 	public:
-		FChaosAccelerationStructureTask(const TArray<FAccelerationStructureBuilder>& InBounds, TUniquePtr<FAccelerationStructure>& InAccelerationStructure,
+		FChaosAccelerationStructureTask(const TArray<FAccelerationStructBuilderCache>& BuilderCacheMap, TUniquePtr<FAccelerationStructure>& InAccelerationStructure,
 			TUniquePtr<FAccelerationStructure>& InAccelerationStructureCopy);
 		static FORCEINLINE TStatId GetStatId();
 		static FORCEINLINE ENamedThreads::Type GetDesiredThread();
 		static FORCEINLINE ESubsequentsMode::Type GetSubsequentsMode();
 		void DoTask(ENamedThreads::Type CurrentThread, const FGraphEventRef& MyCompletionGraphEvent);
 
-		const TArray<FAccelerationStructureBuilder>& CachedSpatialBuilderData;
+		const TArray<FAccelerationStructBuilderCache>& BuilderCacheMap;
 		TUniquePtr<FAccelerationStructure>& AccelerationStructure;
 		TUniquePtr<FAccelerationStructure>& AccelerationStructureCopy;
 	};
@@ -460,6 +478,7 @@ protected:
 	int32 NumIterations;
 
 	static ISpatialAccelerationCollection<TAccelerationStructureHandle<T, d>, T, d>* CreateNewSpatialStructure(const TArray<FAccelerationStructureBuilder>& Bounds);
+	void InitializeAccelerationCache();
 };
 
 }
