Fix CIS (Missing file from last submit)

diff --git a/Engine/Source/Runtime/Experimental/Chaos/Public/Chaos/ImplicitObjectTransformed.h b/Engine/Source/Runtime/Experimental/Chaos/Public/Chaos/ImplicitObjectTransformed.h
--- a/Engine/Source/Runtime/Experimental/Chaos/Public/Chaos/ImplicitObjectTransformed.h
+++ b/Engine/Source/Runtime/Experimental/Chaos/Public/Chaos/ImplicitObjectTransformed.h
@@ -34,14 +34,23 @@ void TImplicitObjectTransformAccumulateSerializableHelper(TArray<Pair<TSerializa
 	check(false);
 }
 
+/**
+ * Transform the contained shape. If you pass a TUniquePtr to the constructor, ownership is transferred to the TransformedImplicit. If you pass a
+ * SerializablePtr, the lifetime of the object must be handled externally (do not delete it before deleting the TransformedImplicit).
+ * @template bSerializable Whether the shape can be serialized (usually true). Set to false for transient/stack-allocated objects. 
+ */
 template<class T, int d, bool bSerializable = true>
 class TImplicitObjectTransformed final : public TImplicitObject<T, d>
 {
-using ObjectType = typename TChooseClass<bSerializable, TSerializablePtr<TImplicitObject<T, d>>, const TImplicitObject<T, d>*>::Result;
+	using FStorage = TImplicitObjectPtrStorage<T, d, bSerializable>;
+	using ObjectType = typename FStorage::PtrType;
 
 public:
 	using TImplicitObject<T, d>::GetTypeName;
 
+	/**
+	 * Create a transform around an ImplicitObject. Lifetime of the wrapped object is managed externally.
+	 */
 	TImplicitObjectTransformed(ObjectType Object, const TRigidTransform<T, d>& InTransform)
 	    : TImplicitObject<T, d>(EImplicitObject::HasBoundingBox, ImplicitObjectType::Transformed)
 	    , MObject(Object)
@@ -50,20 +59,26 @@ public:
 	{
 		this->bIsConvex = Object->IsConvex();
 	}
-	TImplicitObjectTransformed(ObjectType Object, TUniquePtr<Chaos::TImplicitObject<T,d>> &&ObjectOwner, const TRigidTransform<T, d>& InTransform)
+
+	/**
+	 * Create a transform around an ImplicitObject and take control of its lifetime.
+	 */
+	TImplicitObjectTransformed(TUniquePtr<Chaos::TImplicitObject<T,d>> &&ObjectOwner, const TRigidTransform<T, d>& InTransform)
 	    : TImplicitObject<T, d>(EImplicitObject::HasBoundingBox, ImplicitObjectType::Transformed)
-	    , MObject(Object)
 		, MObjectOwner(MoveTemp(ObjectOwner))
 	    , MTransform(InTransform)
-	    , MLocalBoundingBox(Object->BoundingBox().TransformedBox(InTransform))
 	{
-		this->bIsConvex = Object->IsConvex();
+		static_assert(bSerializable, "Non-serializable TImplicitObjectTransformed created with a UniquePtr");
+		this->MObject = FStorage::Convert(MObjectOwner);
+		this->MLocalBoundingBox = MObject->BoundingBox().TransformedBox(InTransform);
+		this->bIsConvex = MObject->IsConvex();
 	}
 
 	TImplicitObjectTransformed(const TImplicitObjectTransformed<T, d, bSerializable>& Other) = delete;
 	TImplicitObjectTransformed(TImplicitObjectTransformed<T, d, bSerializable>&& Other)
 	    : TImplicitObject<T, d>(EImplicitObject::HasBoundingBox, ImplicitObjectType::Transformed)
 	    , MObject(Other.MObject)
+		, MObjectOwner(MoveTemp(Other.MObjectOwner))
 	    , MTransform(Other.MTransform)
 	    , MLocalBoundingBox(MoveTemp(Other.MLocalBoundingBox))
 	{
