Use enable_if over static_assert to resolve static analysis warnings
#rnx

diff --git a/Engine/Source/Runtime/Experimental/ChaosCore/Public/Chaos/Vector.h b/Engine/Source/Runtime/Experimental/ChaosCore/Public/Chaos/Vector.h
--- a/Engine/Source/Runtime/Experimental/ChaosCore/Public/Chaos/Vector.h
+++ b/Engine/Source/Runtime/Experimental/ChaosCore/Public/Chaos/Vector.h
@@ -54,24 +54,24 @@ namespace Chaos
 			}
 		}
 
+		template <int N=d, typename std::enable_if<N==2, int>::type = 0>
 		TVector(const T& V0, const T& V1)
 		{
-			static_assert(NumElements == 2, "Two-element constructor is only for two-element vectors");
 			V[0] = V0;
 			V[1] = V1;
 		}
 
+		template <int N=d, typename std::enable_if<N==3, int>::type = 0>
 		TVector(const T& V0, const T& V1, const T& V2)
 		{
-			static_assert(NumElements == 3, "Three-element constructor is only for three-element vectors");
 			V[0] = V0;
 			V[1] = V1;
 			V[2] = V2;
 		}
 
+		template <int N=d, typename std::enable_if<N==4, int>::type = 0>
 		TVector(const T& V0, const T& V1, const T& V2, const T& V3)
 		{
-			static_assert(NumElements == 4, "Four-element constructor is only for four-element vectors");
 			V[0] = V0;
 			V[1] = V1;
 			V[2] = V2;
@@ -79,9 +79,9 @@ namespace Chaos
 		}
 
 #if !COMPILE_WITHOUT_UNREAL_SUPPORT
+		template <int N=d, typename std::enable_if<N==3, int>::type = 0>
 		TVector(const FVector& Other)
 		{
-			static_assert(NumElements == 3, "Three-element constructor is only for three-element vectors");
 			V[0] = static_cast<FElement>(Other.X);
 			V[1] = static_cast<FElement>(Other.Y);
 			V[2] = static_cast<FElement>(Other.Z);
