Fix incorrect logic

diff --git a/Engine/Source/Runtime/Experimental/Chaos/Public/Chaos/AABBTree.h b/Engine/Source/Runtime/Experimental/Chaos/Public/Chaos/AABBTree.h
--- a/Engine/Source/Runtime/Experimental/Chaos/Public/Chaos/AABBTree.h
+++ b/Engine/Source/Runtime/Experimental/Chaos/Public/Chaos/AABBTree.h
@@ -351,7 +351,7 @@ public:
 			}
 			PayloadInfo->GlobalPayloadIdx = INDEX_NONE;
 		}
-		else if(PayloadInfo->GlobalPayloadIdx == INDEX_NONE)
+		else
 		{
 			TBox<T, 3> GlobalBounds = bTooBig ? NewBounds : TBox<T, 3>(TVector<T, 3>(TNumericLimits<T>::Lowest()), TVector<T, 3>(TNumericLimits<T>::Max()));
 			if (PayloadInfo->GlobalPayloadIdx == INDEX_NONE)
