Trivial: Added Scale func to Box/AABB to scale the extents

diff --git a/Engine/Source/Runtime/Experimental/Chaos/Public/Chaos/AABB.h b/Engine/Source/Runtime/Experimental/Chaos/Public/Chaos/AABB.h
--- a/Engine/Source/Runtime/Experimental/Chaos/Public/Chaos/AABB.h
+++ b/Engine/Source/Runtime/Experimental/Chaos/Public/Chaos/AABB.h
@@ -441,6 +441,12 @@ namespace Chaos
 			}
 		}
 
+		FORCEINLINE void Scale(const TVector<T, d>& InScale)
+		{
+			MMin *= InScale;
+			MMax *= InScale;
+		}
+
 		FORCEINLINE const TVector<T, d>& Min() const { return MMin; }
 		FORCEINLINE const TVector<T, d>& Max() const { return MMax; }
 
diff --git a/Engine/Source/Runtime/Experimental/Chaos/Public/Chaos/Box.h b/Engine/Source/Runtime/Experimental/Chaos/Public/Chaos/Box.h
--- a/Engine/Source/Runtime/Experimental/Chaos/Public/Chaos/Box.h
+++ b/Engine/Source/Runtime/Experimental/Chaos/Public/Chaos/Box.h
@@ -187,6 +187,11 @@ namespace Chaos
 			AABB.ThickenSymmetrically(Thickness);
 		}
 
+		FORCEINLINE void Scale(const TVector<T, d>& InScale)
+		{
+			AABB.Scale(InScale);
+		}
+
 		FORCEINLINE TVector<T, d> Center() const { return AABB.Center(); }
 		FORCEINLINE TVector<T, d> GetCenter() const { return AABB.GetCenter(); }
 		FORCEINLINE TVector<T, d> GetCenterOfMass() const { return AABB.GetCenterOfMass(); }
