Only use _API for mac.

diff --git a/Engine/Source/Runtime/Experimental/Chaos/Public/Chaos/KinematicGeometryParticles.h b/Engine/Source/Runtime/Experimental/Chaos/Public/Chaos/KinematicGeometryParticles.h
--- a/Engine/Source/Runtime/Experimental/Chaos/Public/Chaos/KinematicGeometryParticles.h
+++ b/Engine/Source/Runtime/Experimental/Chaos/Public/Chaos/KinematicGeometryParticles.h
@@ -63,7 +63,7 @@ FChaosArchive& operator<<(FChaosArchive& Ar, TKinematicGeometryParticlesImp<T, d
 	return Ar;
 }
 
-#ifdef __clang__
+#if PLATFORM_MAC
 extern template class CHAOS_API Chaos::TKinematicGeometryParticlesImp<float, 3, Chaos::EGeometryParticlesSimType::RigidBodySim>;
 extern template class CHAOS_API Chaos::TKinematicGeometryParticlesImp<float, 3, Chaos::EGeometryParticlesSimType::Other>;
 #else
