Fix static analysis warnings in gjk.h

diff --git a/Engine/Source/Runtime/Experimental/Chaos/Public/Chaos/GJK.h b/Engine/Source/Runtime/Experimental/Chaos/Public/Chaos/GJK.h
--- a/Engine/Source/Runtime/Experimental/Chaos/Public/Chaos/GJK.h
+++ b/Engine/Source/Runtime/Experimental/Chaos/Public/Chaos/GJK.h
@@ -606,7 +606,7 @@ namespace Chaos
 
 		FSimplex SimplexIDs;
 		TVector<T, 3> Simplex[4], SimplexA[4], SimplexB[4];
-		T Barycentric[4];
+		T Barycentric[4] = { -1, -1, -1, -1 };
 
 		const TRotation<T, 3> AToBRotation = BToATM.GetRotation().Inverse();
 		T Mu = 0;
@@ -653,6 +653,7 @@ namespace Chaos
 					for (int32 VertIndex = 0; VertIndex < SimplexIDs.NumVerts; ++VertIndex)
 					{
 						int32 WIndex = SimplexIDs[VertIndex];
+						check(Barycentric[WIndex] >= (T)0);
 						OutNearestA += Barycentric[WIndex] * SimplexA[WIndex];
 						OutNearestB += Barycentric[WIndex] * SimplexB[WIndex];
 					}
