Expose more functions for mac build.

diff --git a/Engine/Source/Runtime/Experimental/Chaos/Public/Chaos/GeometryParticles.h b/Engine/Source/Runtime/Experimental/Chaos/Public/Chaos/GeometryParticles.h
--- a/Engine/Source/Runtime/Experimental/Chaos/Public/Chaos/GeometryParticles.h
+++ b/Engine/Source/Runtime/Experimental/Chaos/Public/Chaos/GeometryParticles.h
@@ -178,7 +178,7 @@ namespace Chaos
 
 		CHAOS_API TSerializablePtr<TImplicitObject<T, d>> Geometry(const int32 Index) const { return MGeometry[Index]; }
 
-		const TUniquePtr<TImplicitObject<T, d>>& DynamicGeometry(const int32 Index) const { return MDynamicGeometry[Index]; }
+		CHAOS_API const TUniquePtr<TImplicitObject<T, d>>& DynamicGeometry(const int32 Index) const { return MDynamicGeometry[Index]; }
 
 		const TSharedPtr<TImplicitObject<T, d>, ESPMode::ThreadSafe>& SharedGeometry(const int32 Index) const { return MSharedGeometry[Index]; }
 
@@ -243,12 +243,12 @@ namespace Chaos
 			return MSpatialIdx[Index];
 		}
 
-		const TBox<T, d>& WorldSpaceInflatedBounds(const int32 Index) const
+		CHAOS_API const TBox<T, d>& WorldSpaceInflatedBounds(const int32 Index) const
 		{
 			return MWorldSpaceInflatedBounds[Index];
 		}
 
-		TBox<T, d>& WorldSpaceInflatedBounds(const int32 Index)
+		CHAOS_API TBox<T, d>& WorldSpaceInflatedBounds(const int32 Index)
 		{
 			return MWorldSpaceInflatedBounds[Index];
 		}
@@ -365,4 +365,4 @@ namespace Chaos
 
 	extern template class TGeometryParticlesImp<float, 3, EGeometryParticlesSimType::RigidBodySim>;
 	extern template class TGeometryParticlesImp<float, 3, EGeometryParticlesSimType::Other>;
-}
\ No newline at end of file
+}
