Chaos : Collision Caching [Part 1/3]

- Added TMap to allow for persistent collisions between frames.

diff --git a/Engine/Source/Runtime/Experimental/Chaos/Private/Chaos/PBDCollisionConstraint.cpp b/Engine/Source/Runtime/Experimental/Chaos/Private/Chaos/PBDCollisionConstraint.cpp
--- a/Engine/Source/Runtime/Experimental/Chaos/Private/Chaos/PBDCollisionConstraint.cpp
+++ b/Engine/Source/Runtime/Experimental/Chaos/Private/Chaos/PBDCollisionConstraint.cpp
@@ -89,22 +89,35 @@ TPBDCollisionConstraint<T, d>::TPBDCollisionConstraint(const TPBDRigidsSOAs<T,d>
 	, MThickness(Thickness)
 	, MAngularFriction(0)
 	, bUseCCD(false)
+	, LifespanCounter(0)
 	, PostComputeCallback(nullptr)
 	, PostApplyCallback(nullptr)
 	, PostApplyPushOutCallback(nullptr)
 {
 }
 
+DECLARE_CYCLE_STAT(TEXT("CollisionConstraint::Reset"), STAT_CollisionConstraintsReset, STATGROUP_Chaos);
+
 template<typename T, int d>
 void TPBDCollisionConstraint<T, d>::Reset(/*const TPBDRigidParticles<T, d>& InParticles, const TArray<int32>& InIndices*/)
 {
-	for (FConstraintHandle* Handle : Handles)
+	SCOPE_CYCLE_COUNTER(STAT_CollisionConstraintsReset);
+
+	int32 LowestRemovedConstraintIndex = INT_MAX;
+	int32 Threshold = LifespanCounter - 1; // Maybe this should be solver time?
+	for (int32 Idx = Constraints.Num() - 1; Idx >= 0; Idx--)
 	{
-		HandleAllocator.FreeHandle(Handle);
+		if (Constraints[Idx].Lifespan < Threshold)
+		{
+			HandleAllocator.FreeHandle(Handles.FindAndRemoveChecked(GetConstraintHandleID(Idx)));
+			Constraints.RemoveAtSwap(Idx);
+			LowestRemovedConstraintIndex = Idx;
+		}
+	}
+	for (int32 Idx = LowestRemovedConstraintIndex; Idx < Constraints.Num(); Idx++)
+	{
+		Handles[GetConstraintHandleID(Idx)]->SetConstraintIndex(Idx);
 	}
-	Handles.Empty();
-
-	Constraints.Empty();
 
 	MAngularFriction = 0;
 	bUseCCD = false;
@@ -166,6 +179,7 @@ void TPBDCollisionConstraint<T, d>::UpdatePositionBasedState(/*const TPBDRigidPa
 DEFINE_STAT(STAT_ComputeConstraints);
 DEFINE_STAT(STAT_ComputeConstraintsNP);
 DEFINE_STAT(STAT_ComputeConstraintsBP);
+DEFINE_STAT(STAT_ComputeConstraintsSU);
 
 CHAOS_API int32 CollisionConstraintsForceSingleThreaded = 0;
 FAutoConsoleVariableRef CVarCollisionConstraintsForceSingleThreaded(TEXT("p.Chaos.Collision.ForceSingleThreaded"), CollisionConstraintsForceSingleThreaded, TEXT("CollisionConstraintsForceSingleThreaded"));
@@ -364,12 +378,26 @@ void TPBDCollisionConstraint<T, d>::UpdateConstraintsHelper(/*const TPBDRigidPar
 			}*/
 		}
 	});
-	while(!Queue.IsEmpty())
-	{				
-		int32 Idx = Constraints.AddUninitialized(1);
-		Queue.Dequeue(Constraints[Idx]);
-		Handles.Add(HandleAllocator.AllocHandle(this, Idx));
+	while (!Queue.IsEmpty())
+	{
+		const TRigidBodyContactConstraint<T, d> * Constraint = Queue.Peek();
+		FConstraintHandleID HandleID = GetConstraintHandleID(*Constraint);
+		if (Handles.Contains(HandleID))
+		{
+			FConstraintHandle* Handle = Handles[HandleID];
+			int32 Idx = Handle->GetConstraintIndex();
+			Queue.Dequeue(Constraints[Idx]);
+			Constraints[Idx].Lifespan = LifespanCounter;
+		}
+		else
+		{
+			int32 Idx = Constraints.AddUninitialized(1);
+			Queue.Dequeue(Constraints[Idx]);
+			Handles.Add(GetConstraintHandleID(Idx), HandleAllocator.AllocHandle(this, Idx));
+			Constraints[Idx].Lifespan = LifespanCounter;
+		}
 	}
+	LifespanCounter++;
 
 	Timer.Stop();
 	UE_LOG(LogChaos, Verbose, TEXT("\tPBDCollisionConstraint Update %d Constraints with Potential Collisions %f"), Constraints.Num(), Time);
diff --git a/Engine/Source/Runtime/Experimental/Chaos/Public/Chaos/PBDCollisionConstraint.h b/Engine/Source/Runtime/Experimental/Chaos/Public/Chaos/PBDCollisionConstraint.h
--- a/Engine/Source/Runtime/Experimental/Chaos/Public/Chaos/PBDCollisionConstraint.h
+++ b/Engine/Source/Runtime/Experimental/Chaos/Public/Chaos/PBDCollisionConstraint.h
@@ -57,6 +57,7 @@ public:
 	TPBDCollisionConstraintHandle(FConstraintContainer* InConstraintContainer, int32 InConstraintIndex);
 
 	const TRigidBodyContactConstraint<T, d>& GetContact() const;
+	void SetConstraintIndex(int32 IndexIn) { ConstraintIndex = IndexIn; }
 
 protected:
 	using Base::ConstraintIndex;
@@ -92,6 +93,7 @@ public:
 	using FConstraintHandle = TPBDCollisionConstraintHandle<T, d>;
 	using FConstraintHandleAllocator = TConstraintHandleAllocator<TPBDCollisionConstraint<T, d>>;
 	using FRigidBodyContactConstraint = TRigidBodyContactConstraint<T, d>;
+	typedef TPair<TGeometryParticleHandle<T, d>*, const TGeometryParticleHandle<T, d>*> FConstraintHandleID;
 
 	TPBDCollisionConstraint(const TPBDRigidsSOAs<T,d>& InParticles, TArrayCollectionArray<bool>& Collided, const TArrayCollectionArray<TSerializablePtr<TChaosPhysicsMaterial<T>>>& PerParticleMaterials, const int32 PairIterations = 1, const T Thickness = (T)0);
 	virtual ~TPBDCollisionConstraint() {}
@@ -123,14 +125,24 @@ public:
 		return Constraints.Num();
 	}
 
+	FConstraintHandleID GetConstraintHandleID(const FRigidBodyContactConstraint & Constraint) const
+	{
+		return  FConstraintHandleID( Constraint.Particle, Constraint.Levelset );
+	}
+
+	FConstraintHandleID GetConstraintHandleID(int32 ConstraintIndex) const
+	{
+		return  FConstraintHandleID( Constraints[ConstraintIndex].Particle, Constraints[ConstraintIndex].Levelset );
+	}
+		
 	const FConstraintHandle* GetConstraintHandle(int32 ConstraintIndex) const
 	{
-		return Handles[ConstraintIndex];
+		return Handles[GetConstraintHandleID(ConstraintIndex)];
 	}
 
 	FConstraintHandle* GetConstraintHandle(int32 ConstraintIndex)
 	{
-		return Handles[ConstraintIndex];
+		return Handles[ GetConstraintHandleID(ConstraintIndex) ];
 	}
 
 	// @todo(ccaulfield): rename/remove
@@ -246,12 +258,13 @@ private:
 	T MThickness;
 	T MAngularFriction;
 	bool bUseCCD;
+	int32 LifespanCounter;
 
 	TRigidBodyContactConstraintsPostComputeCallback<T, d> PostComputeCallback;
 	TRigidBodyContactConstraintsPostApplyCallback<T, d> PostApplyCallback;
 	TRigidBodyContactConstraintsPostApplyPushOutCallback<T, d> PostApplyPushOutCallback;
 
-	TArray<FConstraintHandle*> Handles;
+	TMap<FConstraintHandleID, FConstraintHandle*> Handles;
 	FConstraintHandleAllocator HandleAllocator;
 };
 
diff --git a/Engine/Source/Runtime/Experimental/Chaos/Public/Chaos/PBDCollisionConstraintImp.h b/Engine/Source/Runtime/Experimental/Chaos/Public/Chaos/PBDCollisionConstraintImp.h
--- a/Engine/Source/Runtime/Experimental/Chaos/Public/Chaos/PBDCollisionConstraintImp.h
+++ b/Engine/Source/Runtime/Experimental/Chaos/Public/Chaos/PBDCollisionConstraintImp.h
@@ -120,6 +120,7 @@ namespace Chaos
 	DECLARE_CYCLE_STAT_EXTERN(TEXT("ComputeConstraints"), STAT_ComputeConstraints, STATGROUP_Chaos, CHAOS_API);
 	DECLARE_CYCLE_STAT_EXTERN(TEXT("ComputeConstraintsNP"), STAT_ComputeConstraintsNP, STATGROUP_Chaos, CHAOS_API);
 	DECLARE_CYCLE_STAT_EXTERN(TEXT("ComputeConstraintsBP"), STAT_ComputeConstraintsBP, STATGROUP_Chaos, CHAOS_API);
+	DECLARE_CYCLE_STAT_EXTERN(TEXT("ComputeConstraintsSU"), STAT_ComputeConstraintsSU, STATGROUP_Chaos, CHAOS_API);
 
 	template <typename T, int d>
 	template <bool bGatherStats, typename SPATIAL_ACCELERATION>
@@ -281,12 +282,30 @@ namespace Chaos
 				NarrowPhaseRejected.Record(RejectedNP);
 #endif
 			}, bGatherStats || CollisionConstraintsForceSingleThreaded);
-
-			while (!Queue.IsEmpty())
+			
 			{
-				int32 Idx = Constraints.AddUninitialized(1);
-				Queue.Dequeue(Constraints[Idx]);
-				Handles.Add(HandleAllocator.AllocHandle(this, Idx));
+				SCOPE_CYCLE_COUNTER(STAT_ComputeConstraintsSU);
+
+				while (!Queue.IsEmpty())
+				{
+					const TRigidBodyContactConstraint<T, d> * Constraint = Queue.Peek();
+					FConstraintHandleID HandleID = GetConstraintHandleID(*Constraint);
+					if (Handles.Contains(HandleID))
+					{
+						FConstraintHandle* Handle = Handles[HandleID];
+						int32 Idx = Handle->GetConstraintIndex();
+						Queue.Dequeue(Constraints[Idx]);
+						Constraints[Idx].Lifespan = LifespanCounter;
+					}
+					else
+					{
+						int32 Idx = Constraints.AddUninitialized(1);
+						Queue.Dequeue(Constraints[Idx]);
+						Handles.Add(GetConstraintHandleID(Idx), HandleAllocator.AllocHandle(this, Idx));
+						Constraints[Idx].Lifespan = LifespanCounter;
+					}
+				}
+				LifespanCounter++;
 			}
 		}
 
diff --git a/Engine/Source/Runtime/Experimental/Chaos/Public/Chaos/PBDCollisionTypes.h b/Engine/Source/Runtime/Experimental/Chaos/Public/Chaos/PBDCollisionTypes.h
--- a/Engine/Source/Runtime/Experimental/Chaos/Public/Chaos/PBDCollisionTypes.h
+++ b/Engine/Source/Runtime/Experimental/Chaos/Public/Chaos/PBDCollisionTypes.h
@@ -13,13 +13,14 @@ namespace Chaos
 template<class T, int d>
 struct TRigidBodyContactConstraint
 {
-	TRigidBodyContactConstraint() : AccumulatedImpulse(0.f) {}
+	TRigidBodyContactConstraint() : AccumulatedImpulse(0.f), Lifespan(0) {}
 	TGeometryParticleHandle<T, d>* Particle;
 	TGeometryParticleHandle<T, d>* Levelset;
 	TVector<T, d> Normal;
 	TVector<T, d> Location;
 	T Phi;
 	TVector<T, d> AccumulatedImpulse;
+	int32 Lifespan;
 
 	FString ToString() const
 	{
