Silence incorrect PVS warning
#rnx

diff --git a/Engine/Source/Runtime/Experimental/ChaosCore/Public/Chaos/Vector.h b/Engine/Source/Runtime/Experimental/ChaosCore/Public/Chaos/Vector.h
--- a/Engine/Source/Runtime/Experimental/ChaosCore/Public/Chaos/Vector.h
+++ b/Engine/Source/Runtime/Experimental/ChaosCore/Public/Chaos/Vector.h
@@ -66,7 +66,7 @@ namespace Chaos
 		{
 			V[0] = V0;
 			V[1] = V1;
-			V[2] = V2;
+			V[2] = V2; //-V557
 		}
 
 		template <int N=d, typename std::enable_if<N==4, int>::type = 0>
@@ -74,8 +74,8 @@ namespace Chaos
 		{
 			V[0] = V0;
 			V[1] = V1;
-			V[2] = V2;
-			V[3] = V3;
+			V[2] = V2; //-V557
+			V[3] = V3; //-V557
 		}
 
 #if !COMPILE_WITHOUT_UNREAL_SUPPORT
@@ -84,7 +84,7 @@ namespace Chaos
 		{
 			V[0] = static_cast<FElement>(Other.X);
 			V[1] = static_cast<FElement>(Other.Y);
-			V[2] = static_cast<FElement>(Other.Z);
+			V[2] = static_cast<FElement>(Other.Z); //-V557
 		}
 #endif
 
