Chaos 6dof re-enable elliptical constraint limits

diff --git a/Engine/Source/Runtime/Experimental/Chaos/Private/Chaos/PBDJointConstraintUtilities.cpp b/Engine/Source/Runtime/Experimental/Chaos/Private/Chaos/PBDJointConstraintUtilities.cpp
--- a/Engine/Source/Runtime/Experimental/Chaos/Private/Chaos/PBDJointConstraintUtilities.cpp
+++ b/Engine/Source/Runtime/Experimental/Chaos/Private/Chaos/PBDJointConstraintUtilities.cpp
@@ -528,12 +528,10 @@ namespace Chaos
 		// Elliptical swing limit
 		if (!FMath::IsNearlyEqual(Swing1Limit, Swing2Limit, KINDA_SMALL_NUMBER))
 		{
-			SwingAngleMax = (T)0.5 * (Swing1Limit + Swing2Limit);
-
 			// Map swing axis to ellipse and calculate limit for this swing axis
-			//T DotSwing1 = FMath::Abs(TVector<T, d>::DotProduct(SwingAxis01, TJointConstants<T, d>::Swing1Axis()));
-			//T DotSwing2 = FMath::Abs(TVector<T, d>::DotProduct(SwingAxis01, TJointConstants<T, d>::Swing2Axis()));
-			//SwingAngleMax = FMath::Sqrt(Swing1Limit * DotSwing1 * Swing1Limit * DotSwing1 + Swing2Limit * DotSwing2 * Swing2Limit * DotSwing2);
+			T DotSwing1 = FMath::Abs(TVector<T, d>::DotProduct(SwingAxis01, TJointConstants<T, d>::Swing1Axis()));
+			T DotSwing2 = FMath::Abs(TVector<T, d>::DotProduct(SwingAxis01, TJointConstants<T, d>::Swing2Axis()));
+			SwingAngleMax = FMath::Sqrt(Swing1Limit * DotSwing1 * Swing1Limit * DotSwing1 + Swing2Limit * DotSwing2 * Swing2Limit * DotSwing2);
 		}
 
 		// Calculate swing error we need to correct
