Chaos
- add ability to disable collision detection
- disable RBAN collisions - needs fixes to mesh with persistent collisions
- fix TransformedImplicit transform order.
- re-enable Transformed boxes.

diff --git a/Engine/Source/Runtime/Engine/Private/PhysicsEngine/Experimental/ChaosInterfaceUtils.cpp b/Engine/Source/Runtime/Engine/Private/PhysicsEngine/Experimental/ChaosInterfaceUtils.cpp
--- a/Engine/Source/Runtime/Engine/Private/PhysicsEngine/Experimental/ChaosInterfaceUtils.cpp
+++ b/Engine/Source/Runtime/Engine/Private/PhysicsEngine/Experimental/ChaosInterfaceUtils.cpp
@@ -173,12 +173,12 @@ namespace ChaosInterface
 
 			// TBox can handle translations internally but if we have a rotation we need to wrap it in a transform
 			TUniquePtr<Chaos::TImplicitObject<float, 3>> Implicit;
-			//if (!BoxTransform.GetRotation().IsIdentity())
-			//{
-			//	auto ImplicitBox = MakeUnique<Chaos::TBox<float, 3>>(-HalfExtents, HalfExtents);
-			//	Implicit = TUniquePtr<Chaos::TImplicitObject<float, 3>>(new Chaos::TImplicitObjectTransformed<float, 3>(MoveTemp(ImplicitBox), BoxTransform));
-			//}
-			//else
+			if (!BoxTransform.GetRotation().IsIdentity())
+			{
+				auto ImplicitBox = MakeUnique<Chaos::TBox<float, 3>>(-HalfExtents, HalfExtents);
+				Implicit = TUniquePtr<Chaos::TImplicitObject<float, 3>>(new Chaos::TImplicitObjectTransformed<float, 3>(MoveTemp(ImplicitBox), BoxTransform));
+			}
+			else
 			{
 				Implicit = MakeUnique<Chaos::TBox<float, 3>>(BoxTransform.GetTranslation() - HalfExtents, BoxTransform.GetTranslation() + HalfExtents);
 			}
diff --git a/Engine/Source/Runtime/Engine/Private/PhysicsEngine/ImmediatePhysics/ImmediatePhysicsChaos/ImmediatePhysicsSimulation_Chaos.cpp b/Engine/Source/Runtime/Engine/Private/PhysicsEngine/ImmediatePhysics/ImmediatePhysicsChaos/ImmediatePhysicsSimulation_Chaos.cpp
--- a/Engine/Source/Runtime/Engine/Private/PhysicsEngine/ImmediatePhysics/ImmediatePhysicsChaos/ImmediatePhysicsSimulation_Chaos.cpp
+++ b/Engine/Source/Runtime/Engine/Private/PhysicsEngine/ImmediatePhysics/ImmediatePhysicsChaos/ImmediatePhysicsSimulation_Chaos.cpp
@@ -20,7 +20,7 @@
 // @todo(ccaulfield): remove when finished
 float ChaosImmediate_Evolution_DeltaTime = 0.03f;
 int32 ChaosImmediate_Evolution_Iterations = 10;
-int32 ChaosImmediate_Collision_Enabled = 1;
+int32 ChaosImmediate_Collision_Enabled = 0;
 int32 ChaosImmediate_Collision_ApplyEnabled = 0;
 int32 ChaosImmediate_Collision_PushOutIterations = 5;
 int32 ChaosImmediate_Collision_PushOutPairIterations = 2;
@@ -409,6 +409,9 @@ namespace ImmediatePhysics_Chaos
 			Evolution->GetCollisionConstraints().SetVelocitySolveEnabled(ChaosImmediate_Collision_ApplyEnabled != 0);
 			Evolution->GetCollisionConstraints().SetPushOutPairIterations(ChaosImmediate_Collision_PushOutPairIterations);
 			Evolution->GetCollisionConstraintsRule().SetPushOutIterations(ChaosImmediate_Collision_PushOutIterations);
+		
+			// TEMP until we can remove constraints again, or I add broad-phase filtering. (FilterCollisionConstraints will crash since the persistent collision changes)
+			Evolution->GetCollisionConstraints().SetCollisionsEnabled(ChaosImmediate_Collision_Enabled != 0);
 		}
 
 		DebugDrawElementsAll(2, 2, 0.7f);
diff --git a/Engine/Source/Runtime/Experimental/Chaos/Private/Chaos/PBDCollisionConstraint.cpp b/Engine/Source/Runtime/Experimental/Chaos/Private/Chaos/PBDCollisionConstraint.cpp
--- a/Engine/Source/Runtime/Experimental/Chaos/Private/Chaos/PBDCollisionConstraint.cpp
+++ b/Engine/Source/Runtime/Experimental/Chaos/Private/Chaos/PBDCollisionConstraint.cpp
@@ -89,6 +89,7 @@ TPBDCollisionConstraint<T, d>::TPBDCollisionConstraint(const TPBDRigidsSOAs<T,d>
 	, MThickness(Thickness)
 	, MAngularFriction(0)
 	, bUseCCD(false)
+	, bEnableCollisions(true)
 	, LifespanCounter(0)
 	, PostComputeCallback(nullptr)
 	, PostApplyCallback(nullptr)
@@ -164,15 +165,18 @@ void TPBDCollisionConstraint<T, d>::UpdatePositionBasedState(/*const TPBDRigidPa
 {
 	Reset();
 
-#if !UE_BUILD_SHIPPING
-	if (PendingHierarchyDump)
+	if (bEnableCollisions)
 	{
-		ComputeConstraints<true>(*SpatialAcceleration, Dt);
-	}
-	else
+#if !UE_BUILD_SHIPPING
+		if (PendingHierarchyDump)
+		{
+			ComputeConstraints<true>(*SpatialAcceleration, Dt);
+		}
+		else
 #endif
-	{
-		ComputeConstraints(*SpatialAcceleration, Dt);
+		{
+			ComputeConstraints(*SpatialAcceleration, Dt);
+		}
 	}
 }
 
@@ -1619,7 +1623,7 @@ void UpdateSingleUnionConstraint(const T Thickness, TRigidBodyContactConstraint<
 	for (const Pair<const TImplicitObject<T, d>*, TRigidTransform<T, d>>& LevelsetObjPair : LevelsetShapes)
 	{
 		const TImplicitObject<T, d>& LevelsetInnerObj = *LevelsetObjPair.First;
-		const TRigidTransform<T, d> LevelsetInnerObjTM = LevelsetTM * LevelsetObjPair.Second;
+		const TRigidTransform<T, d> LevelsetInnerObjTM = LevelsetObjPair.Second * LevelsetTM;
 		UpdateConstraintImp2<UpdateType>(*ParticleObj, ParticlesTM, LevelsetInnerObj, LevelsetInnerObjTM, Thickness, Constraint);
 	}
 }
diff --git a/Engine/Source/Runtime/Experimental/Chaos/Public/Chaos/ImplicitObjectTransformed.h b/Engine/Source/Runtime/Experimental/Chaos/Public/Chaos/ImplicitObjectTransformed.h
--- a/Engine/Source/Runtime/Experimental/Chaos/Public/Chaos/ImplicitObjectTransformed.h
+++ b/Engine/Source/Runtime/Experimental/Chaos/Public/Chaos/ImplicitObjectTransformed.h
@@ -186,7 +186,7 @@ public:
 		MObject->FindAllIntersectingObjects(Out, SubobjectBounds);
 		if (Out.Num() > NumOut)
 		{
-			Out[NumOut].Second = MTransform * Out[NumOut].Second;
+			Out[NumOut].Second = Out[NumOut].Second * MTransform;
 		}
 	}
 
diff --git a/Engine/Source/Runtime/Experimental/Chaos/Public/Chaos/PBDCollisionConstraint.h b/Engine/Source/Runtime/Experimental/Chaos/Public/Chaos/PBDCollisionConstraint.h
--- a/Engine/Source/Runtime/Experimental/Chaos/Public/Chaos/PBDCollisionConstraint.h
+++ b/Engine/Source/Runtime/Experimental/Chaos/Public/Chaos/PBDCollisionConstraint.h
@@ -117,6 +117,11 @@ public:
 		MPairIterations = InPairIterations;
 	}
 
+	void SetCollisionsEnabled(bool bInEnableCollisions)
+	{
+		bEnableCollisions = bInEnableCollisions;
+	}
+
 	/**
 	 * Get the number of constraints.
 	 */
@@ -258,6 +263,7 @@ private:
 	T MThickness;
 	T MAngularFriction;
 	bool bUseCCD;
+	bool bEnableCollisions;
 	int32 LifespanCounter;
 
 	TRigidBodyContactConstraintsPostComputeCallback<T, d> PostComputeCallback;
