Chaos : Collision Constraints

- Added RemoveConstraint(idx) to help external clients manage handle updates when they remove constraints.
- Removed RemoveConstraints(...), needed as indicated by the comment.

diff --git a/Engine/Source/Runtime/Experimental/Chaos/Private/Chaos/PBDCollisionConstraint.cpp b/Engine/Source/Runtime/Experimental/Chaos/Private/Chaos/PBDCollisionConstraint.cpp
--- a/Engine/Source/Runtime/Experimental/Chaos/Private/Chaos/PBDCollisionConstraint.cpp
+++ b/Engine/Source/Runtime/Experimental/Chaos/Private/Chaos/PBDCollisionConstraint.cpp
@@ -104,26 +104,30 @@ void TPBDCollisionConstraint<T, d>::Reset(/*const TPBDRigidParticles<T, d>& InPa
 {
 	SCOPE_CYCLE_COUNTER(STAT_CollisionConstraintsReset);
 
-	int32 LowestRemovedConstraintIndex = INT_MAX;
 	int32 Threshold = LifespanCounter - 1; // Maybe this should be solver time?
 	for (int32 Idx = Constraints.Num() - 1; Idx >= 0; Idx--)
 	{
 		if (Constraints[Idx].Lifespan < Threshold)
 		{
-			HandleAllocator.FreeHandle(Handles.FindAndRemoveChecked(GetConstraintHandleID(Idx)));
-			Constraints.RemoveAtSwap(Idx);
-			LowestRemovedConstraintIndex = Idx;
+			RemoveConstraint(Idx);
 		}
 	}
-	for (int32 Idx = LowestRemovedConstraintIndex; Idx < Constraints.Num(); Idx++)
-	{
-		Handles[GetConstraintHandleID(Idx)]->SetConstraintIndex(Idx);
-	}
 
 	MAngularFriction = 0;
 	bUseCCD = false;
 }
 
+template<typename T, int d>
+void TPBDCollisionConstraint<T, d>::RemoveConstraint(int32 Idx)
+{
+	HandleAllocator.FreeHandle(Handles.FindAndRemoveChecked(GetConstraintHandleID(Idx)));
+	Constraints.RemoveAtSwap(Idx);
+	if (Idx < Constraints.Num())
+	{
+		Handles[GetConstraintHandleID(Idx)]->SetConstraintIndex(Idx);
+	}
+}
+
 template<typename T, int d>
 void TPBDCollisionConstraint<T, d>::SetPostComputeCallback(const TRigidBodyContactConstraintsPostComputeCallback<T, d>& Callback)
 {
@@ -226,28 +230,6 @@ void TPBDCollisionConstraint<T, d>::ComputeConstraints(const FAccelerationStruct
 	}
 }
 
-template<typename T, int d>
-void TPBDCollisionConstraint<T, d>::RemoveConstraints(const TSet<TGeometryParticleHandle<T, d>*>& RemovedParticles)
-{
-#if CHAOS_PARTICLEHANDLE_TODO
-	SpatialAccelerationResource.GetWritable().RemoveElements(RemovedParticles.Array());
-	for (int32 i = 0; i < Constraints.Num(); ++i)
-	{
-		const auto& Constraint = Constraints[i];
-		if (RemovedParticles.Contains(Constraint.ParticleIndex) || RemovedParticles.Contains(Constraint.LevelsetIndex))
-		{
-			Constraints.RemoveAtSwap(i);
-			HandleAllocator.FreeHandle(Handles[i]);
-			Handles.RemoveAtSwap(i);
-			if (i < Handles.Num())
-			{
-				SetConstraintIndex(Handles[i], i);
-			}
-			i--;
-		}
-	}
-#endif
-}
 
 DECLARE_CYCLE_STAT(TEXT("UpdateConstraints"), STAT_UpdateConstraints, STATGROUP_Chaos);
 
diff --git a/Engine/Source/Runtime/Experimental/Chaos/Public/Chaos/PBDCollisionConstraint.h b/Engine/Source/Runtime/Experimental/Chaos/Public/Chaos/PBDCollisionConstraint.h
--- a/Engine/Source/Runtime/Experimental/Chaos/Public/Chaos/PBDCollisionConstraint.h
+++ b/Engine/Source/Runtime/Experimental/Chaos/Public/Chaos/PBDCollisionConstraint.h
@@ -150,8 +150,7 @@ public:
 		return Handles[ GetConstraintHandleID(ConstraintIndex) ];
 	}
 
-	// @todo(ccaulfield): rename/remove
-	void RemoveConstraints(const TSet<TGeometryParticleHandle<T, d>*>& RemovedParticles);
+	void RemoveConstraint(int32 Idx);
 
 	/**
 	 * Set the callback used just after contacts are generated at the start of a frame tick.
diff --git a/Engine/Source/Runtime/Experimental/Chaos/Public/Chaos/PBDConstraintRule.h b/Engine/Source/Runtime/Experimental/Chaos/Public/Chaos/PBDConstraintRule.h
--- a/Engine/Source/Runtime/Experimental/Chaos/Public/Chaos/PBDConstraintRule.h
+++ b/Engine/Source/Runtime/Experimental/Chaos/Public/Chaos/PBDConstraintRule.h
@@ -106,13 +106,6 @@ namespace Chaos
 			}
 		}
 
-		virtual void RemoveConstraints(const TSet<TGeometryParticleHandle<T,d>*>& RemovedParticles) override
-		{
-			// @todo(ccaulfield): should also remove graph edges here too. Currently the edges are rebuilt every tick so not a problem...
-			//ConstraintGraph->RemoveEdges(...)
-			Constraints.RemoveConstraints(RemovedParticles);
-		}
-
 		virtual int32 NumConstraints() const override { return Constraints.NumConstraints(); }
 
 	protected:
