Chaos : Collision Caching

- Save the previous location on the collision constraint.

diff --git a/Engine/Source/Runtime/Experimental/Chaos/Public/Chaos/PBDCollisionConstraintImp.h b/Engine/Source/Runtime/Experimental/Chaos/Public/Chaos/PBDCollisionConstraintImp.h
--- a/Engine/Source/Runtime/Experimental/Chaos/Public/Chaos/PBDCollisionConstraintImp.h
+++ b/Engine/Source/Runtime/Experimental/Chaos/Public/Chaos/PBDCollisionConstraintImp.h
@@ -294,7 +294,10 @@ namespace Chaos
 					{
 						FConstraintHandle* Handle = Handles[HandleID];
 						int32 Idx = Handle->GetConstraintIndex();
+						FVector Position = Constraint[Idx].Location;
+
 						Queue.Dequeue(Constraints[Idx]);
+						Constraints[Idx].PreviousLocation = Position;
 						Constraints[Idx].Lifespan = LifespanCounter;
 					}
 					else
diff --git a/Engine/Source/Runtime/Experimental/Chaos/Public/Chaos/PBDCollisionTypes.h b/Engine/Source/Runtime/Experimental/Chaos/Public/Chaos/PBDCollisionTypes.h
--- a/Engine/Source/Runtime/Experimental/Chaos/Public/Chaos/PBDCollisionTypes.h
+++ b/Engine/Source/Runtime/Experimental/Chaos/Public/Chaos/PBDCollisionTypes.h
@@ -13,13 +13,15 @@ namespace Chaos
 template<class T, int d>
 struct TRigidBodyContactConstraint
 {
-	TRigidBodyContactConstraint() : AccumulatedImpulse(0.f), Lifespan(0) {}
+	TRigidBodyContactConstraint() : AccumulatedImpulse(0.f), PreviousLocation( FLT_MAX ), Lifespan(0) {}
 	TGeometryParticleHandle<T, d>* Particle;
 	TGeometryParticleHandle<T, d>* Levelset;
 	TVector<T, d> Normal;
 	TVector<T, d> Location;
 	T Phi;
 	TVector<T, d> AccumulatedImpulse;
+
+	TVector<T, d> PreviousLocation;
 	int32 Lifespan;
 
 	FString ToString() const
