Fix Windows Clang linker errors
#rnx

diff --git a/Engine/Source/Runtime/Experimental/Chaos/Private/Chaos/RigidParticles.cpp b/Engine/Source/Runtime/Experimental/Chaos/Private/Chaos/RigidParticles.cpp
--- a/Engine/Source/Runtime/Experimental/Chaos/Private/Chaos/RigidParticles.cpp
+++ b/Engine/Source/Runtime/Experimental/Chaos/Private/Chaos/RigidParticles.cpp
@@ -21,7 +21,11 @@ void TRigidParticles<T, d>::SetCollisionParticles(const int32 Index, TParticles<
 }
 
 #ifdef __clang__
+#if PLATFORM_WINDOWS
+template class Chaos::TRigidParticles<float, 3>;
+#else
 template class CHAOS_API Chaos::TRigidParticles<float, 3>;
+#endif
 #else
 template class Chaos::TRigidParticles<float, 3>;
 #endif
\ No newline at end of file
diff --git a/Engine/Source/Runtime/Experimental/Chaos/Public/Chaos/RigidParticles.h b/Engine/Source/Runtime/Experimental/Chaos/Public/Chaos/RigidParticles.h
--- a/Engine/Source/Runtime/Experimental/Chaos/Public/Chaos/RigidParticles.h
+++ b/Engine/Source/Runtime/Experimental/Chaos/Public/Chaos/RigidParticles.h
@@ -189,7 +189,11 @@ FChaosArchive& operator<<(FChaosArchive& Ar, TRigidParticles<T, d>& Particles)
 }
 
 #ifdef __clang__
+#if PLATFORM_WINDOWS
+extern template class TRigidParticles<float, 3>;
+#else
 extern template class CHAOS_API TRigidParticles<float, 3>;
+#endif
 #else
 extern template class TRigidParticles<float, 3>;
 #endif
