Option to warn about or simplify convex geometry that is too complex and would cause a performance issues.

diff --git a/Engine/Source/Runtime/Engine/Private/PhysicsEngine/BodySetup.cpp b/Engine/Source/Runtime/Engine/Private/PhysicsEngine/BodySetup.cpp
--- a/Engine/Source/Runtime/Engine/Private/PhysicsEngine/BodySetup.cpp
+++ b/Engine/Source/Runtime/Engine/Private/PhysicsEngine/BodySetup.cpp
@@ -696,6 +696,13 @@ void UBodySetup::FinishCreatingPhysicsMeshes_Chaos(FChaosDerivedDataReader<float
 		{
 			FKConvexElem& ConvexElem = AggGeom.ConvexElems[ElementIndex];
 			ConvexElem.SetChaosConvexMesh(MoveTemp(InReader.ConvexImplicitObjects[ElementIndex]));
+
+			if (ConvexElem.GetChaosConvexMesh()->IsPerformanceWarning())
+			{
+				const FString& PerformanceString = ConvexElem.GetChaosConvexMesh()->PerformanceWarningAndSimplifaction();
+				UE_LOG(LogPhysics, Warning, TEXT("TConvex Name:%s, Element [%d], %s"), FullName.GetCharArray().GetData(), ElementIndex, PerformanceString.GetCharArray().GetData());
+			}
+
 		}
 		InReader.ConvexImplicitObjects.Reset();
 	}
diff --git a/Engine/Source/Runtime/Experimental/Chaos/Private/Chaos/CollisionConvexMesh.cpp b/Engine/Source/Runtime/Experimental/Chaos/Private/Chaos/CollisionConvexMesh.cpp
new file mode 100644
--- /dev/null
+++ b/Engine/Source/Runtime/Experimental/Chaos/Private/Chaos/CollisionConvexMesh.cpp
@@ -0,0 +1,17 @@
+/// Copyright 1998-2019 Epic Games, Inc. All Rights Reserved.
+
+#include "Chaos/CollisionConvexMesh.h"
+
+namespace Chaos
+{
+	// CVars variables for controlling geometry complexity checking and simplification
+	int32 TConvexBuilder<float>::PerformGeometryCheck = 0;
+	int32 TConvexBuilder<float>::PerformGeometryReduction = 0;
+	int32 TConvexBuilder<float>::ParticlesThreshold = 50;
+
+	FAutoConsoleVariableRef CVarConvexGeometryCheckEnable(TEXT("p.Chaos.ConvexGeometryCheckEnable"), TConvexBuilder<float>::PerformGeometryCheck, TEXT("Perform convex geometry complexity check for Chaos physics."));
+
+	FAutoConsoleVariableRef CVarConvexGeometrySimplifyEnable(TEXT("p.Chaos.PerformGeometryReduction"), TConvexBuilder<float>::PerformGeometryReduction, TEXT("Perform convex geometry simplification to increase performance in Chaos physics."));
+
+	FAutoConsoleVariableRef CVarConvexParticlesWarningThreshold(TEXT("p.Chaos.ConvexParticlesWarningThreshold"), TConvexBuilder<float>::ParticlesThreshold, TEXT("Threshold beyond which we warn about collision geometry complexity."));
+}
\ No newline at end of file
diff --git a/Engine/Source/Runtime/Experimental/Chaos/Public/Chaos/CollisionConvexMesh.h b/Engine/Source/Runtime/Experimental/Chaos/Public/Chaos/CollisionConvexMesh.h
--- a/Engine/Source/Runtime/Experimental/Chaos/Public/Chaos/CollisionConvexMesh.h
+++ b/Engine/Source/Runtime/Experimental/Chaos/Public/Chaos/CollisionConvexMesh.h
@@ -122,6 +122,106 @@ namespace Chaos
 			return TTriangleMesh<T>(MoveTemp(Indices));
 		}
 
+		static bool IsPerformanceWarning(int32 NumPlanes, int32 NumParticles)
+		{
+			if (!PerformGeometryCheck)
+				return false;
+
+			return (NumParticles > ParticlesThreshold);
+		}
+
+		static bool IsGeometryReductionEnabled()
+		{
+			return PerformGeometryReduction;
+		}
+
+		static FString PerformanceWarningString(int32 NumPlanes, int32 NumParticles)
+		{
+			return FString::Printf(TEXT("Planes %d, SurfaceParticles %d"), NumPlanes, NumParticles);
+		}
+
+		static void Simplify(TArray <TPlane<T, 3>>& InOutPlanes, TParticles<T, 3>& InOutParticles, TBox<T, 3>& InOutLocalBounds)
+		{
+			struct TPair
+			{
+				TPair() : A(-1), B(-1) {}
+				uint32 A;
+				uint32 B;
+			};
+
+			uint32 NumberOfParticlesRequired = ParticlesThreshold;
+			uint32 NumberOfParticlesWeHave = InOutParticles.Size();
+			int32 NumToDelete = NumberOfParticlesWeHave - NumberOfParticlesRequired;
+
+			uint32 Size = InOutParticles.Size();
+			TArray<TVector<T, 3>> Particles;
+			Particles.AddUninitialized(Size);
+			for (uint32 A = 0; A < Size; A++)
+			{
+				Particles[A] = InOutParticles.X(A);
+			}
+
+			TArray<bool> IsDeleted;
+			IsDeleted.Reset();
+			IsDeleted.Init(false, Size);
+
+			if (NumToDelete > 0)
+			{
+				for (uint32 Iteration = 0; Iteration < (uint32)NumToDelete; Iteration++)
+				{
+					TPair ClosestPair;
+					float ClosestDistSqr = FLT_MAX;
+
+					for (uint32 A = 0; A < (Size - 1); A++)
+					{
+						if (!IsDeleted[A])
+						{
+							for (uint32 B = A + 1; B < Size; B++)
+							{
+								if (!IsDeleted[B])
+								{
+									TVector<T, 3> Vec = Particles[A] - Particles[B];
+									float LengthSqr = Vec.SizeSquared();
+									if (LengthSqr < ClosestDistSqr)
+									{
+										ClosestDistSqr = LengthSqr;
+										ClosestPair.A = A;
+										ClosestPair.B = B;
+									}
+								}
+							}
+						}
+					}
+
+					if (ClosestPair.A != -1)
+					{
+						// merge to mid point
+						Particles[ClosestPair.A] = Particles[ClosestPair.A] + (Particles[ClosestPair.B] - Particles[ClosestPair.A]) * 0.5f;
+						IsDeleted[ClosestPair.B] = true;
+					}
+				}
+			}
+
+			TParticles<T, 3> TmpParticles;
+			for (int Idx = 0; Idx < Particles.Num(); Idx++)
+			{
+				// Only add particles that have not been merged away
+				if (!IsDeleted[Idx])
+				{
+					TmpParticles.AddParticles(1);
+					TmpParticles.X(TmpParticles.Size() - 1) = Particles[Idx];
+				}
+			}
+
+			Build(TmpParticles, InOutPlanes, InOutParticles, InOutLocalBounds);
+			check(InOutParticles.Size() > 3);
+		}
+
+		// CVars variables for controlling geometry complexity checking and simplification
+		static CHAOS_API int32 PerformGeometryCheck;
+		static CHAOS_API int32 PerformGeometryReduction;
+		static CHAOS_API int32 ParticlesThreshold;
+
 	private:
 
 		struct FHalfEdge;
@@ -575,5 +675,6 @@ namespace Chaos
 			//todo(ocohen): need to explicitly test for merge failures. Coplaner, nonconvex, etc...
 			//getting this in as is for now to unblock other systems
 		}
+
 	};
-}
\ No newline at end of file
+}
diff --git a/Engine/Source/Runtime/Experimental/Chaos/Public/Chaos/Convex.h b/Engine/Source/Runtime/Experimental/Chaos/Public/Chaos/Convex.h
--- a/Engine/Source/Runtime/Experimental/Chaos/Public/Chaos/Convex.h
+++ b/Engine/Source/Runtime/Experimental/Chaos/Public/Chaos/Convex.h
@@ -213,6 +213,29 @@ namespace Chaos
 			SerializeImp(Ar);
 		}
 
+		virtual bool IsPerformanceWarning() const override
+		{
+			return TConvexBuilder<T>::IsPerformanceWarning(Planes.Num(), SurfaceParticles.Size());
+		}
+
+		virtual FString PerformanceWarningAndSimplifaction() override
+		{
+
+			FString PerformanceWarningString = TConvexBuilder<T>::PerformanceWarningString(Planes.Num(), SurfaceParticles.Size());
+			if (TConvexBuilder<T>::IsGeometryReductionEnabled())
+			{
+				PerformanceWarningString += ", [Simplifying]";
+				SimplifyGeometry();
+			}
+
+			return PerformanceWarningString;
+		}
+
+		void SimplifyGeometry()
+		{
+			TConvexBuilder<T>::Simplify(Planes, SurfaceParticles, LocalBoundingBox);
+		}
+
 	private:
 		TArray<TPlane<T, d>> Planes;
 		TParticles<T, d> SurfaceParticles;	//copy of the vertices that are just on the convex hull boundary
diff --git a/Engine/Source/Runtime/Experimental/Chaos/Public/Chaos/ImplicitObject.h b/Engine/Source/Runtime/Experimental/Chaos/Public/Chaos/ImplicitObject.h
--- a/Engine/Source/Runtime/Experimental/Chaos/Public/Chaos/ImplicitObject.h
+++ b/Engine/Source/Runtime/Experimental/Chaos/Public/Chaos/ImplicitObject.h
@@ -156,6 +156,11 @@ public:
 	void IgnoreAnalyticCollisions(const bool Ignore = true) { bIgnoreAnalyticCollisions = Ignore; }
 	bool GetIgnoreAnalyticCollisions() const { return bIgnoreAnalyticCollisions; }
 	void SetConvex(const bool Convex = true) { bIsConvex = Convex; }
+	virtual bool IsPerformanceWarning() const { return false; }
+	virtual FString PerformanceWarningAndSimplifaction() 
+	{
+		return FString::Printf(TEXT("ImplicitObject - No Performance String"));
+	};
 
 	Pair<TVector<T, d>, bool> FindDeepestIntersection(const TImplicitObject<T, d>* Other, const TBVHParticles<float, d>* Particles, const PMatrix<T, d, d>& OtherToLocalTransform, const T Thickness) const;
 	Pair<TVector<T, d>, bool> FindDeepestIntersection(const TImplicitObject<T, d>* Other, const TParticles<float, d>* Particles, const PMatrix<T, d, d>& OtherToLocalTransform, const T Thickness) const;
