Fix CIS on non-windows platforms

diff --git a/Engine/Source/Runtime/Experimental/Chaos/Private/Chaos/CollisionConvexMesh.cpp b/Engine/Source/Runtime/Experimental/Chaos/Private/Chaos/CollisionConvexMesh.cpp
--- a/Engine/Source/Runtime/Experimental/Chaos/Private/Chaos/CollisionConvexMesh.cpp
+++ b/Engine/Source/Runtime/Experimental/Chaos/Private/Chaos/CollisionConvexMesh.cpp
@@ -5,8 +5,13 @@
 namespace Chaos
 {
 	// CVars variables for controlling geometry complexity checking and simplification
+	template <>
 	int32 TConvexBuilder<float>::PerformGeometryCheck = 0;
+
+	template <>
 	int32 TConvexBuilder<float>::PerformGeometryReduction = 0;
+
+	template <>
 	int32 TConvexBuilder<float>::ParticlesThreshold = 50;
 
 	FAutoConsoleVariableRef CVarConvexGeometryCheckEnable(TEXT("p.Chaos.ConvexGeometryCheckEnable"), TConvexBuilder<float>::PerformGeometryCheck, TEXT("Perform convex geometry complexity check for Chaos physics."));
