Fix CIS

diff --git a/Engine/Source/Runtime/Experimental/Chaos/Private/Chaos/CollisionConvexMesh.cpp b/Engine/Source/Runtime/Experimental/Chaos/Private/Chaos/CollisionConvexMesh.cpp
--- a/Engine/Source/Runtime/Experimental/Chaos/Private/Chaos/CollisionConvexMesh.cpp
+++ b/Engine/Source/Runtime/Experimental/Chaos/Private/Chaos/CollisionConvexMesh.cpp
@@ -1,4 +1,4 @@
-/// Copyright 1998-2019 Epic Games, Inc. All Rights Reserved.
+// Copyright 1998-2019 Epic Games, Inc. All Rights Reserved.
 
 #include "Chaos/CollisionConvexMesh.h"
 
diff --git a/Engine/Source/Runtime/Experimental/Chaos/Public/Chaos/CollisionConvexMesh.h b/Engine/Source/Runtime/Experimental/Chaos/Public/Chaos/CollisionConvexMesh.h
--- a/Engine/Source/Runtime/Experimental/Chaos/Public/Chaos/CollisionConvexMesh.h
+++ b/Engine/Source/Runtime/Experimental/Chaos/Public/Chaos/CollisionConvexMesh.h
@@ -125,14 +125,16 @@ namespace Chaos
 		static bool IsPerformanceWarning(int32 NumPlanes, int32 NumParticles)
 		{
 			if (!PerformGeometryCheck)
+			{
 				return false;
+			}
 
 			return (NumParticles > ParticlesThreshold);
 		}
 
 		static bool IsGeometryReductionEnabled()
 		{
-			return PerformGeometryReduction;
+			return (PerformGeometryReduction>0)?true:false;
 		}
 
 		static FString PerformanceWarningString(int32 NumPlanes, int32 NumParticles)
