CIS fixes - dependent types and bracing

diff --git a/Engine/Source/Runtime/Experimental/Chaos/Private/Chaos/HeightField.cpp b/Engine/Source/Runtime/Experimental/Chaos/Private/Chaos/HeightField.cpp
--- a/Engine/Source/Runtime/Experimental/Chaos/Private/Chaos/HeightField.cpp
+++ b/Engine/Source/Runtime/Experimental/Chaos/Private/Chaos/HeightField.cpp
@@ -276,8 +276,8 @@ namespace Chaos
 	template<typename HeightfieldT, typename BufferType>
 	void BuildGeomData(TArrayView<BufferType> BufferView, int32 NumRows, int32 NumCols, const TVector<HeightfieldT, 3>& InScale, TUniqueFunction<HeightfieldT(const BufferType)> ToRealFunc, typename THeightField<HeightfieldT>::FDataType& OutData, TBox<HeightfieldT, 3>& OutBounds)
 	{
-		using FDataType = THeightField<HeightfieldT>::FDataType;
-		using RealType = FDataType::RealType;
+		using FDataType = typename THeightField<HeightfieldT>::FDataType;
+		using RealType = typename FDataType::RealType;
 
 		const int32 NumCells = NumRows * NumCols;
 		ensure(BufferView.Num() == NumCells);
@@ -296,7 +296,7 @@ namespace Chaos
 
 		for(int32 HeightIndex = 1; HeightIndex < NumHeights; ++HeightIndex)
 		{
-			const typename RealType CurrHeight = ToRealFunc(BufferView[HeightIndex]);
+			const RealType CurrHeight = ToRealFunc(BufferView[HeightIndex]);
 
 			if(CurrHeight > OutData.MaxValue)
 			{
@@ -317,7 +317,7 @@ namespace Chaos
 
 			int32 X = HeightIndex % (NumCols);
 			int32 Y = HeightIndex / (NumCols);
-			TVector<HeightfieldT, 3> Position(typename RealType(X), typename RealType(Y), OutData.MinValue + OutData.Heights[HeightIndex] * OutData.HeightPerUnit);
+			TVector<HeightfieldT, 3> Position(RealType(X), RealType(Y), OutData.MinValue + OutData.Heights[HeightIndex] * OutData.HeightPerUnit);
 			if(HeightIndex == 0)
 			{
 				OutBounds = TBox<HeightfieldT, 3>(Position * InScale, Position * InScale);
diff --git a/Engine/Source/Runtime/Landscape/Private/LandscapeCollision.cpp b/Engine/Source/Runtime/Landscape/Private/LandscapeCollision.cpp
--- a/Engine/Source/Runtime/Landscape/Private/LandscapeCollision.cpp
+++ b/Engine/Source/Runtime/Landscape/Private/LandscapeCollision.cpp
@@ -353,7 +353,7 @@ void ULandscapeHeightfieldCollisionComponent::OnCreatePhysicsState()
 				SCOPED_SCENE_WRITE_LOCK(SyncScene);
 				SyncScene->addActor(*HeightFieldActorSync);
 			}
-			else
+
 #elif WITH_CHAOS
 			{
 				FActorCreationParams Params;
@@ -431,8 +431,8 @@ void ULandscapeHeightfieldCollisionComponent::OnCreatePhysicsState()
 				TArray<FPhysicsActorHandle> Actors;
 				Actors.Add(PhysHandle);
 				PhysScene->AddActorsToScene_AssumesLocked(Actors);
-#endif
 			}
+#endif
 		}
 #endif// WITH_PHYSX
 	}
