Geometry Particles assignment in conditinoal error.
-Merging //UE4/Dev-VR to Dev-Main (//UE4/Dev-Main) : 9448937
-Not sure why this is only being reported as an error on Hololens, I don't see a relavant compiler setting differing from windows.  Anyway it was not very human readable so lets make it more explicit.

diff --git a/Engine/Source/Runtime/Experimental/Chaos/Public/Chaos/GeometryParticles.h b/Engine/Source/Runtime/Experimental/Chaos/Public/Chaos/GeometryParticles.h
--- a/Engine/Source/Runtime/Experimental/Chaos/Public/Chaos/GeometryParticles.h
+++ b/Engine/Source/Runtime/Experimental/Chaos/Public/Chaos/GeometryParticles.h
@@ -296,7 +296,8 @@ namespace Chaos
 				//just assume all bounds come from geometry (technically wrong for pbd rigids with only sample points, but backwards compat is not that important right now)
 				for (int32 Idx = 0; Idx < MGeometry.Num(); ++Idx)
 				{
-					if ((MHasBounds[Idx] = MGeometry[Idx] && MGeometry[Idx]->HasBoundingBox()))
+					MHasBounds[Idx] = MGeometry[Idx] && MGeometry[Idx]->HasBoundingBox();
+					if (MHasBounds[Idx])
 					{
 						MLocalBounds[Idx] = MGeometry[Idx]->BoundingBox();
 						MWorldSpaceInflatedBounds[Idx] = MLocalBounds[Idx].TransformedBox(TRigidTransform<T,d>(X(Idx), R(Idx)));	//ignore velocity too, really just trying to get something reasonable
