Chaos - fix elliptical limits (swing1/2 limits were flipped)

diff --git a/Engine/Source/Runtime/Experimental/Chaos/Private/Chaos/PBDJointConstraintUtilities.cpp b/Engine/Source/Runtime/Experimental/Chaos/Private/Chaos/PBDJointConstraintUtilities.cpp
--- a/Engine/Source/Runtime/Experimental/Chaos/Private/Chaos/PBDJointConstraintUtilities.cpp
+++ b/Engine/Source/Runtime/Experimental/Chaos/Private/Chaos/PBDJointConstraintUtilities.cpp
@@ -176,7 +176,7 @@ namespace Chaos
 		}
 		else
 		{
-			// Line Limits
+			// Line/Square/Cube Limits (no way to author square or cube limits, but would work if we wanted it)
 			TVector<T, d> CX = InCX;
 			if (Motion[0] != EJointMotionType::Locked)
 			{
@@ -531,7 +531,7 @@ namespace Chaos
 			// Map swing axis to ellipse and calculate limit for this swing axis
 			T DotSwing1 = FMath::Abs(TVector<T, d>::DotProduct(SwingAxis01, TJointConstants<T, d>::Swing1Axis()));
 			T DotSwing2 = FMath::Abs(TVector<T, d>::DotProduct(SwingAxis01, TJointConstants<T, d>::Swing2Axis()));
-			SwingAngleMax = FMath::Sqrt(Swing1Limit * DotSwing1 * Swing1Limit * DotSwing1 + Swing2Limit * DotSwing2 * Swing2Limit * DotSwing2);
+			SwingAngleMax = FMath::Sqrt(Swing1Limit * DotSwing2 * Swing1Limit * DotSwing2 + Swing2Limit * DotSwing1 * Swing2Limit * DotSwing1);
 		}
 
 		// Calculate swing error we need to correct
