Don't export symbols that are not needed.

diff --git a/Engine/Source/Runtime/Experimental/Chaos/Public/Chaos/CollisionConvexMesh.h b/Engine/Source/Runtime/Experimental/Chaos/Public/Chaos/CollisionConvexMesh.h
--- a/Engine/Source/Runtime/Experimental/Chaos/Public/Chaos/CollisionConvexMesh.h
+++ b/Engine/Source/Runtime/Experimental/Chaos/Public/Chaos/CollisionConvexMesh.h
@@ -220,9 +220,9 @@ namespace Chaos
 		}
 
 		// CVars variables for controlling geometry complexity checking and simplification
-		static CHAOS_API int32 PerformGeometryCheck;
-		static CHAOS_API int32 PerformGeometryReduction;
-		static CHAOS_API int32 ParticlesThreshold;
+		static int32 PerformGeometryCheck;
+		static int32 PerformGeometryReduction;
+		static int32 ParticlesThreshold;
 
 	private:
 
