- Create new PBD spherical constraints
- Cloth: Implement Max distance and backstop constraints in terms of spherical constraints

diff --git a/Engine/Plugins/Experimental/ChaosCloth/Source/Chaos/Private/ChaosCloth/ChaosClothingSimulation.cpp b/Engine/Plugins/Experimental/ChaosCloth/Source/Chaos/Private/ChaosCloth/ChaosClothingSimulation.cpp
--- a/Engine/Plugins/Experimental/ChaosCloth/Source/Chaos/Private/ChaosCloth/ChaosClothingSimulation.cpp
+++ b/Engine/Plugins/Experimental/ChaosCloth/Source/Chaos/Private/ChaosCloth/ChaosClothingSimulation.cpp
@@ -23,6 +23,7 @@
 #include "Chaos/PerParticleGravity.h"
 #include "Chaos/PerParticlePBDLongRangeConstraints.h"
 #include "Chaos/PerParticlePBDShapeConstraints.h"
+#include "Chaos/PBDSphericalConstraint.h"
 #include "Chaos/Plane.h"
 #include "Chaos/Sphere.h"
 #include "Chaos/TaperedCylinder.h"
@@ -396,6 +397,36 @@ void ClothingSimulation::CreateActor(USkeletalMeshComponent* InOwnerComponent, U
 				std::placeholders::_1,
 				std::placeholders::_2));
     }
+
+	// Maximum Distance Constraints
+	const UEnum* const MeshTargets = PhysMesh->GetFloatArrayTargets();	
+	const uint32 PhysMeshMaxDistanceIndex = MeshTargets->GetValueByName(TEXT("MaxDistance"));;
+	if (PhysMesh->GetFloatArray(PhysMeshMaxDistanceIndex)->Num() > 0)
+	{
+		check(Mesh->GetNumElements() > 0);
+		Chaos::PBDSphericalConstraint<float, 3> SphericalContraint(Offset, PhysMesh->GetFloatArray(PhysMeshMaxDistanceIndex)->Num(), true, &AnimationPositions, PhysMesh->GetFloatArray(PhysMeshMaxDistanceIndex));
+		Evolution->AddPBDConstraintFunction([=](TPBDParticles<float, 3>& InParticles, const float Dt)
+		{
+			SphericalContraint.Apply(InParticles, Dt);
+		});
+	}
+
+	// Backstop Constraints
+	const uint32 PhysMeshBackstopDistanceIndex = MeshTargets->GetValueByName(TEXT("BackstopDistance"));
+	const uint32 PhysMeshBackstopRadiusIndex = MeshTargets->GetValueByName(TEXT("BackstopRadius"));
+	if (PhysMesh->GetFloatArray(PhysMeshBackstopRadiusIndex)->Num() > 0 && PhysMesh->GetFloatArray(PhysMeshBackstopDistanceIndex)->Num() > 0)
+	{
+		check(Mesh->GetNumElements() > 0);
+		check(PhysMesh->GetFloatArray(PhysMeshBackstopRadiusIndex)->Num() == PhysMesh->GetFloatArray(PhysMeshBackstopDistanceIndex)->Num());
+
+		Chaos::PBDSphericalConstraint<float, 3> SphericalContraint(Offset, PhysMesh->GetFloatArray(PhysMeshBackstopRadiusIndex)->Num(), false, &AnimationPositions, 
+			PhysMesh->GetFloatArray(PhysMeshBackstopRadiusIndex), PhysMesh->GetFloatArray(PhysMeshBackstopDistanceIndex), &AnimationNormals);
+		Evolution->AddPBDConstraintFunction([=](TPBDParticles<float, 3>& InParticles, const float Dt)
+		{
+			SphericalContraint.Apply(InParticles, Dt);
+		});		
+	}
+
     // Add Self Collisions
     if (bUseSelfCollisions)
     {
diff --git a/Engine/Source/Runtime/ClothingSystemRuntimeNv/Public/ClothPhysicalMeshDataNv.h b/Engine/Source/Runtime/ClothingSystemRuntimeNv/Public/ClothPhysicalMeshDataNv.h
--- a/Engine/Source/Runtime/ClothingSystemRuntimeNv/Public/ClothPhysicalMeshDataNv.h
+++ b/Engine/Source/Runtime/ClothingSystemRuntimeNv/Public/ClothPhysicalMeshDataNv.h
@@ -31,9 +31,9 @@ public:
 	ClearParticleParameters() override;
 
 	virtual void
-	BuildSelfCollisionData(const UClothConfigBase* ClothConfig);
+	BuildSelfCollisionData(const UClothConfigBase* ClothConfig) override;
 
-	virtual UEnum* GetFloatArrayTargets() const
+	virtual UEnum* GetFloatArrayTargets() const override
 	{ return StaticEnum<MaskTarget_PhysMesh>(); }
 
 	virtual bool IsFullyKinematic() const override
diff --git a/Engine/Source/Runtime/Experimental/Chaos/Public/Chaos/PBDSphericalConstraint.h b/Engine/Source/Runtime/Experimental/Chaos/Public/Chaos/PBDSphericalConstraint.h
new file mode 100644
--- /dev/null
+++ b/Engine/Source/Runtime/Experimental/Chaos/Public/Chaos/PBDSphericalConstraint.h
@@ -0,0 +1,71 @@
+// Copyright 1998-2019 Epic Games, Inc. All Rights Reserved.
+#pragma once
+
+#include "ParticleRule.h"
+
+namespace Chaos
+{
+
+template<typename T, int d>
+class PBDSphericalConstraint : public TParticleRule<T, d>
+{
+public:
+	PBDSphericalConstraint(
+		const uint32 InFirstParticleIndex,
+		const uint32 InParticleCount,
+		bool bInside,
+		const TArray<TVector<T, d>>* const InSpherePositions,
+		const TArray<T>* const InSphereRadii,
+		const TArray <T>* const InSphereOffsetDistances = nullptr,
+		const TArray<TVector<T, d>>* const InSphereOffsetDirections = nullptr)
+		: FirstParticleIndex(InFirstParticleIndex), ParticleCount(InParticleCount), bConstraintInsideOfSphere(bInside), SpherePositions(InSpherePositions), SphereRadii(InSphereRadii), SphereOffsetDistances(InSphereOffsetDistances), SphereOffsetDirections(InSphereOffsetDirections)
+	{
+		check(SphereRadii);
+		check(SpherePositions);
+	}
+	virtual ~PBDSphericalConstraint() {}
+
+	inline void Apply(TPBDParticles<T, d>& InParticles, const T Dt) const override
+	{
+		const T CenterDeadZoneRadius = KINDA_SMALL_NUMBER; // We will not push the particle away in the dead zone
+
+		PhysicsParallelFor(ParticleCount, [&](int32 Index)
+		{
+			if (InParticles.InvM(FirstParticleIndex + Index) == 0)
+			{
+				return;
+			}
+			
+			const T Radius = (*SphereRadii)[Index];
+			TVector<T, d> Centre = (*SpherePositions)[Index + FirstParticleIndex];
+			if (SphereOffsetDistances && SphereOffsetDirections)
+			{
+				Centre -= (*SphereOffsetDistances)[Index] * (*SphereOffsetDirections)[Index + FirstParticleIndex];
+			}
+
+			const TVector<T, d> CenterToParticle = InParticles.P(FirstParticleIndex + Index) - Centre;
+			const T DistanceSquared = CenterToParticle.SizeSquared();
+			const bool bPushParticleToShpericalBoundery = ((DistanceSquared < Radius * Radius) ^ bConstraintInsideOfSphere) && DistanceSquared > CenterDeadZoneRadius * CenterDeadZoneRadius;
+
+			if (bPushParticleToShpericalBoundery)
+			{
+				const T Distance = sqrt(DistanceSquared);
+				const TVector<T, d> PositionOnSphere = (Radius / Distance) * (CenterToParticle);
+				InParticles.P(FirstParticleIndex + Index) = Centre + PositionOnSphere;
+			}
+		});
+	}
+
+private:	
+	const uint32 FirstParticleIndex; // TODO: Get rid of this member variable (this is specific to a temporary cloth instancing implementation)
+	const uint32 ParticleCount;
+	const bool bConstraintInsideOfSphere; // Should the points be constrained to the inside of their respective spheres or the outsides
+
+	const TArray<TVector<T, d>>* const SpherePositions;  // Positions of spheres without adding an offset vector // Start at index FirstParticleIndex
+	const TArray<T>* const SphereRadii; // Start at index 0
+
+	const TArray<T>* const SphereOffsetDistances;  // Optional Sphere position offsets // start at index 0
+	const TArray<TVector<T, d>>* const SphereOffsetDirections; //	Optional Sphere offset directions (TBD: This is really implementation specific to backstop constraints, consider creating a specialized constraint instead).
+};
+
+}
