Export functions over vars.

diff --git a/Engine/Source/Runtime/Experimental/Chaos/Public/Chaos/CollisionConvexMesh.h b/Engine/Source/Runtime/Experimental/Chaos/Public/Chaos/CollisionConvexMesh.h
--- a/Engine/Source/Runtime/Experimental/Chaos/Public/Chaos/CollisionConvexMesh.h
+++ b/Engine/Source/Runtime/Experimental/Chaos/Public/Chaos/CollisionConvexMesh.h
@@ -122,7 +122,7 @@ namespace Chaos
 			return TTriangleMesh<T>(MoveTemp(Indices));
 		}
 
-		static bool IsPerformanceWarning(int32 NumPlanes, int32 NumParticles)
+		static CHAOS_API bool IsPerformanceWarning(int32 NumPlanes, int32 NumParticles)
 		{
 			if (!PerformGeometryCheck)
 			{
@@ -132,7 +132,7 @@ namespace Chaos
 			return (NumParticles > ParticlesThreshold);
 		}
 
-		static bool IsGeometryReductionEnabled()
+		static CHAOS_API bool IsGeometryReductionEnabled()
 		{
 			return (PerformGeometryReduction>0)?true:false;
 		}
@@ -142,7 +142,7 @@ namespace Chaos
 			return FString::Printf(TEXT("Planes %d, SurfaceParticles %d"), NumPlanes, NumParticles);
 		}
 
-		static void Simplify(TArray <TPlane<T, 3>>& InOutPlanes, TParticles<T, 3>& InOutParticles, TBox<T, 3>& InOutLocalBounds)
+		static CHAOS_API void Simplify(TArray <TPlane<T, 3>>& InOutPlanes, TParticles<T, 3>& InOutParticles, TBox<T, 3>& InOutLocalBounds)
 		{
 			struct TPair
 			{
@@ -220,9 +220,9 @@ namespace Chaos
 		}
 
 		// CVars variables for controlling geometry complexity checking and simplification
-		static CHAOS_API int32 PerformGeometryCheck;
-		static CHAOS_API int32 PerformGeometryReduction;
-		static CHAOS_API int32 ParticlesThreshold;
+		static int32 PerformGeometryCheck;
+		static int32 PerformGeometryReduction;
+		static int32 ParticlesThreshold;
 
 	private:
 
