Chaos - fix disable linear limit logic

diff --git a/Engine/Source/Runtime/Experimental/Chaos/Private/Chaos/PBDJointConstraints.cpp b/Engine/Source/Runtime/Experimental/Chaos/Private/Chaos/PBDJointConstraints.cpp
--- a/Engine/Source/Runtime/Experimental/Chaos/Private/Chaos/PBDJointConstraints.cpp
+++ b/Engine/Source/Runtime/Experimental/Chaos/Private/Chaos/PBDJointConstraints.cpp
@@ -400,7 +400,7 @@ namespace Chaos
 		EJointMotionType Swing2Motion = JointSettings.Motion.AngularMotionTypes[(int32)EJointAngularConstraintIndex::Swing2];
 
 		// Disable a constraint if it has any linear limits?
-		if (!Settings.bEnableLinearLimits && ((LinearMotion[0] == EJointMotionType::Limited) || (LinearMotion[1] == EJointMotionType::Limited) || (LinearMotion[1] == EJointMotionType::Limited)))
+		if (!Settings.bEnableLinearLimits && ((LinearMotion[0] == EJointMotionType::Limited) || (LinearMotion[1] == EJointMotionType::Limited) || (LinearMotion[2] == EJointMotionType::Limited)))
 		{
 			return;
 		}
