Add _API for incorrect variables on mac only.

diff --git a/Engine/Source/Runtime/Experimental/Chaos/Public/Chaos/CollisionConvexMesh.h b/Engine/Source/Runtime/Experimental/Chaos/Public/Chaos/CollisionConvexMesh.h
--- a/Engine/Source/Runtime/Experimental/Chaos/Public/Chaos/CollisionConvexMesh.h
+++ b/Engine/Source/Runtime/Experimental/Chaos/Public/Chaos/CollisionConvexMesh.h
@@ -220,9 +220,15 @@ namespace Chaos
 		}
 
 		// CVars variables for controlling geometry complexity checking and simplification
+#if PLATFORM_MAC
+		static CHAOS_API int32 PerformGeometryCheck;
+		static CHAOS_API int32 PerformGeometryReduction;
+		static CHAOS_API int32 ParticlesThreshold;
+#else
 		static int32 PerformGeometryCheck;
 		static int32 PerformGeometryReduction;
 		static int32 ParticlesThreshold;
+#endif
 
 	private:
 
