Chaos - hook debug draw into VisLog (cvar controlled). Fix disable collisions logic.

diff --git a/Engine/Source/Runtime/Experimental/Chaos/Private/Chaos/ChaosDebugDraw.cpp b/Engine/Source/Runtime/Experimental/Chaos/Private/Chaos/ChaosDebugDraw.cpp
--- a/Engine/Source/Runtime/Experimental/Chaos/Private/Chaos/ChaosDebugDraw.cpp
+++ b/Engine/Source/Runtime/Experimental/Chaos/Private/Chaos/ChaosDebugDraw.cpp
@@ -54,13 +54,14 @@ namespace Chaos
 			case ImplicitObjectType::Sphere:
 			{
 				const TSphere<T, d>* Sphere = Shape->template GetObject<TSphere<T, d>>();
-				FDebugDrawQueue::GetInstance().DrawDebugSphere(ShapeTransform.GetLocation(), Sphere->GetRadius(), 20, Color, false, KINDA_SMALL_NUMBER, DrawPriority, LineThickness);
+				const TVector<T, d> P = ShapeTransform.TransformPosition(Sphere->GetCenter());
+				FDebugDrawQueue::GetInstance().DrawDebugSphere(P, Sphere->GetRadius(), 20, Color, false, KINDA_SMALL_NUMBER, DrawPriority, LineThickness);
 				break;
 			}
 			case ImplicitObjectType::Box:
 			{
 				const TBox<T, d>* Box = Shape->template GetObject<TBox<T, d>>();
-				const TVector<T, d> P = ShapeTransform.TransformPosition(Box->Center());
+				const TVector<T, d> P = ShapeTransform.TransformPosition(Box->GetCenter());
 				FDebugDrawQueue::GetInstance().DrawDebugBox(P, (T)0.5 * Box->Extents(), ShapeTransform.GetRotation(), Color, false, KINDA_SMALL_NUMBER, DrawPriority, LineThickness);
 				break;
 			}
@@ -115,7 +116,7 @@ namespace Chaos
 		void DrawParticleShapesImpl(const TRigidTransform<float, 3>& SpaceTransform, const TGeometryParticleHandle<T, d>* Particle, T ColorScale)
 		{
 			FColor ShapeColor = Particle->AsDynamic() ? FColor::Yellow : FColor::Orange;
-			FColor Color = (ColorScale * ShapeColor).ToFColor(false);
+			FColor Color = (((T)0.5 * ColorScale) * ShapeColor).ToFColor(false);
 			TVector<T, d> P = SpaceTransform.TransformPosition(Particle->AsDynamic() ? Particle->AsDynamic()->P() : Particle->X());
 			TRotation<T, d> Q = SpaceTransform.GetRotation() * (Particle->AsDynamic() ? Particle->AsDynamic()->Q() : Particle->R());
 
diff --git a/Engine/Source/Runtime/Experimental/Chaos/Private/Chaos/PBDCollisionConstraint.cpp b/Engine/Source/Runtime/Experimental/Chaos/Private/Chaos/PBDCollisionConstraint.cpp
--- a/Engine/Source/Runtime/Experimental/Chaos/Private/Chaos/PBDCollisionConstraint.cpp
+++ b/Engine/Source/Runtime/Experimental/Chaos/Private/Chaos/PBDCollisionConstraint.cpp
@@ -107,7 +107,7 @@ void TPBDCollisionConstraint<T, d>::Reset(/*const TPBDRigidParticles<T, d>& InPa
 	int32 Threshold = LifespanCounter - 1; // Maybe this should be solver time?
 	for (int32 Idx = Constraints.Num() - 1; Idx >= 0; Idx--)
 	{
-		if (Constraints[Idx].Lifespan < Threshold)
+		if ((Constraints[Idx].Lifespan < Threshold) || !bEnableCollisions)
 		{
 			RemoveConstraint(Idx);
 		}
diff --git a/Engine/Source/Runtime/Experimental/Chaos/Public/Chaos/DebugDrawQueue.h b/Engine/Source/Runtime/Experimental/Chaos/Public/Chaos/DebugDrawQueue.h
--- a/Engine/Source/Runtime/Experimental/Chaos/Public/Chaos/DebugDrawQueue.h
+++ b/Engine/Source/Runtime/Experimental/Chaos/Public/Chaos/DebugDrawQueue.h
@@ -119,6 +119,11 @@ struct CHAOS_API FLatentDrawCommand
 		Capsule,
 	} Type;
 
+	FLatentDrawCommand()
+		: TestBaseActor(nullptr)
+	{
+	}
+
 	static FLatentDrawCommand DrawPoint(const FVector& Position, const FColor& Color, bool bPersistentLines, float LifeTime, uint8 DepthPriority, float Thickness)
 	{
 		FLatentDrawCommand Command;
diff --git a/Engine/Source/Runtime/Experimental/ChaosSolverEngine/Private/Chaos/ChaosDebugDrawComponent.cpp b/Engine/Source/Runtime/Experimental/ChaosSolverEngine/Private/Chaos/ChaosDebugDrawComponent.cpp
--- a/Engine/Source/Runtime/Experimental/ChaosSolverEngine/Private/Chaos/ChaosDebugDrawComponent.cpp
+++ b/Engine/Source/Runtime/Experimental/ChaosSolverEngine/Private/Chaos/ChaosDebugDrawComponent.cpp
@@ -2,12 +2,23 @@
 
 #include "Chaos/ChaosDebugDrawComponent.h"
 #include "Chaos/DebugDrawQueue.h"
+#include "ChaosLog.h"
 #include "DrawDebugHelpers.h"
 #include "Engine/World.h"
+#include "VisualLogger/VisualLogger.h"
+
+int bChaosDebugDraw_DrawMode = 0;
+FAutoConsoleVariableRef CVarArrowSize(TEXT("p.Chaos.DebugDrawMode"), bChaosDebugDraw_DrawMode, TEXT("Where to send debug draw commands. 0 = UE Debug Draw; 1 = VisLog; 2 = Both"));
 
 #if CHAOS_DEBUG_DRAW
-void DebugDrawChaos(UWorld* World)
+void DebugDrawChaos(AActor* DebugDrawActor)
 {
+	if (!DebugDrawActor)
+	{
+		return;
+	}
+
+	UWorld* World = DebugDrawActor->GetWorld();
 	if (!World || !World->IsGameWorld())
 	{
 		return;
@@ -16,51 +27,116 @@ void DebugDrawChaos(UWorld* World)
 	using namespace Chaos;
 	TArray<FLatentDrawCommand> LatenetDrawCommands;
 
+	bool bDrawUe = bChaosDebugDraw_DrawMode != 1;
+	bool bDrawVisLog = bChaosDebugDraw_DrawMode != 0;
+
 	FDebugDrawQueue::GetInstance().ExtractAllElements(LatenetDrawCommands, !World->IsPaused());
 	for (const FLatentDrawCommand& Command : LatenetDrawCommands)
 	{
+		AActor* Actor = (Command.TestBaseActor) ? Command.TestBaseActor : DebugDrawActor;
+
 		switch (Command.Type)
 		{
 		case FLatentDrawCommand::EDrawType::Point:
 		{
-			DrawDebugPoint(World, Command.LineStart, Command.Thickness, Command.Color, Command.bPersistentLines, Command.LifeTime, Command.DepthPriority);
+			if (bDrawUe)
+			{
+				DrawDebugPoint(World, Command.LineStart, Command.Thickness, Command.Color, Command.bPersistentLines, Command.LifeTime, Command.DepthPriority);
+			}
+			if (bDrawVisLog)
+			{
+				UE_VLOG_SEGMENT_THICK(Actor, LogChaos, Log, Command.LineStart, Command.LineStart, Command.Color, Command.Thickness, TEXT_EMPTY);
+			}
 			break;
 		}
 		case FLatentDrawCommand::EDrawType::Line:
 		{
-			DrawDebugLine(World, Command.LineStart, Command.LineEnd, Command.Color, Command.bPersistentLines, Command.LifeTime, Command.DepthPriority, Command.Thickness);
+			if (bDrawUe)
+			{
+				DrawDebugLine(World, Command.LineStart, Command.LineEnd, Command.Color, Command.bPersistentLines, Command.LifeTime, Command.DepthPriority, Command.Thickness);
+			}
+			if (bDrawVisLog)
+			{
+				UE_VLOG_SEGMENT(Actor, LogChaos, Log, Command.LineStart, Command.LineEnd, Command.Color, TEXT_EMPTY);
+			}
 			break;
 		}
 		case FLatentDrawCommand::EDrawType::DirectionalArrow:
 		{
-			DrawDebugDirectionalArrow(World, Command.LineStart, Command.LineEnd, Command.ArrowSize, Command.Color, Command.bPersistentLines, Command.LifeTime, Command.DepthPriority, Command.Thickness);
+			if (bDrawUe)
+			{
+				DrawDebugDirectionalArrow(World, Command.LineStart, Command.LineEnd, Command.ArrowSize, Command.Color, Command.bPersistentLines, Command.LifeTime, Command.DepthPriority, Command.Thickness);
+			}
+			if (bDrawVisLog)
+			{
+				UE_VLOG_SEGMENT(Actor, LogChaos, Log, Command.LineStart, Command.LineEnd, Command.Color, TEXT_EMPTY);
+			}
 			break;
 		}
 		case FLatentDrawCommand::EDrawType::Sphere:
 		{
-			DrawDebugSphere(World, Command.LineStart, Command.Radius, Command.Segments, Command.Color, Command.bPersistentLines, Command.LifeTime, Command.DepthPriority, Command.Thickness);
+			if (bDrawUe)
+			{
+				DrawDebugSphere(World, Command.LineStart, Command.Radius, Command.Segments, Command.Color, Command.bPersistentLines, Command.LifeTime, Command.DepthPriority, Command.Thickness);
+			}
+			if (bDrawVisLog)
+			{
+				// VLOG Capsule uses the bottom end as the origin (though the variable is named Center)
+				FVector Base = Command.LineStart - Command.Radius * FVector::UpVector;
+				UE_VLOG_CAPSULE(Actor, LogChaos, Log, Base, Command.Radius + KINDA_SMALL_NUMBER, Command.Radius, FQuat::Identity, Command.Color, TEXT_EMPTY);
+			}
 			break;
 		}
 		case FLatentDrawCommand::EDrawType::Box:
 		{
-			DrawDebugBox(World, Command.Center, Command.Extent, Command.Rotation, Command.Color, Command.bPersistentLines, Command.LifeTime, Command.DepthPriority, Command.Thickness);
+			if (bDrawUe)
+			{
+				DrawDebugBox(World, Command.Center, Command.Extent, Command.Rotation, Command.Color, Command.bPersistentLines, Command.LifeTime, Command.DepthPriority, Command.Thickness);
+			}
+			if (bDrawVisLog)
+			{
+				UE_VLOG_OBOX(Actor, LogChaos, Log, FBox(-Command.Extent, Command.Extent), FQuatRotationTranslationMatrix::Make(Command.Rotation, Command.Center), Command.Color, TEXT_EMPTY);
+			}
 			break;
 		}
 		case FLatentDrawCommand::EDrawType::String:
 		{
-			DrawDebugString(World, Command.TextLocation, Command.Text, Command.TestBaseActor, Command.Color, Command.LifeTime, Command.bDrawShadow, Command.FontScale);
+			if (bDrawUe)
+			{
+				DrawDebugString(World, Command.TextLocation, Command.Text, Command.TestBaseActor, Command.Color, Command.LifeTime, Command.bDrawShadow, Command.FontScale);
+			}
+			if (bDrawVisLog)
+			{
+				UE_VLOG(Command.TestBaseActor, LogChaos, Log, TEXT("%s"), *Command.Text);
+			}
 			break;
 		}
 		case FLatentDrawCommand::EDrawType::Circle:
 		{
-			FMatrix M = FRotationMatrix::MakeFromYZ(Command.YAxis, Command.ZAxis);
-			M.SetOrigin(Command.Center);
-			DrawDebugCircle(World, M, Command.Radius, Command.Segments, Command.Color, Command.bPersistentLines, Command.LifeTime, Command.DepthPriority, Command.Thickness, Command.bDrawAxis);
+			if (bDrawUe)
+			{
+				FMatrix M = FRotationMatrix::MakeFromYZ(Command.YAxis, Command.ZAxis);
+				M.SetOrigin(Command.Center);
+				DrawDebugCircle(World, M, Command.Radius, Command.Segments, Command.Color, Command.bPersistentLines, Command.LifeTime, Command.DepthPriority, Command.Thickness, Command.bDrawAxis);
+			}
+			//if (bDrawVisLog)
+			//{
+			//	UE_VLOG_CONE(Actor, LogChaos, Log, Command.Center, Command.ZAxis, 0.1f, PI, Command.Color, TEXT_EMPTY);
+			//}
 			break;
 		}
 		case FLatentDrawCommand::EDrawType::Capsule:
 		{
-			DrawDebugCapsule(World, Command.Center, Command.HalfHeight, Command.Radius, Command.Rotation, Command.Color, Command.bPersistentLines, Command.LifeTime, Command.DepthPriority, Command.Thickness);
+			if (bDrawUe)
+			{
+				DrawDebugCapsule(World, Command.Center, Command.HalfHeight, Command.Radius, Command.Rotation, Command.Color, Command.bPersistentLines, Command.LifeTime, Command.DepthPriority, Command.Thickness);
+			}
+			if (bDrawVisLog)
+			{
+				// VLOG Capsule uses the bottom end as the origin (though the variable is named Center)
+				FVector Base = Command.Center - Command.HalfHeight * (Command.Rotation * FVector::UpVector);
+				UE_VLOG_CAPSULE(Actor, LogChaos, Log, Base, Command.HalfHeight, Command.Radius, Command.Rotation, Command.Color, TEXT_EMPTY);
+			}
 			break;
 		}
 		default:
@@ -122,7 +198,7 @@ void UChaosDebugDrawComponent::TickComponent(float DeltaTime, enum ELevelTick Ti
 	Super::TickComponent(DeltaTime, TickType, ThisTickFunction);
 
 #if CHAOS_DEBUG_DRAW
-	DebugDrawChaos(GetWorld());
+	DebugDrawChaos(GetOwner());
 #endif
 }
 
