Correct ray and sweep length scaling in ImplicitObjectScaled

diff --git a/Engine/Source/Runtime/Experimental/Chaos/Public/Chaos/ImplicitObjectScaled.h b/Engine/Source/Runtime/Experimental/Chaos/Public/Chaos/ImplicitObjectScaled.h
--- a/Engine/Source/Runtime/Experimental/Chaos/Public/Chaos/ImplicitObjectScaled.h
+++ b/Engine/Source/Runtime/Experimental/Chaos/Public/Chaos/ImplicitObjectScaled.h
@@ -88,19 +88,26 @@ public:
 		ensure(Thickness == 0 || (FMath::IsNearlyEqual(MScale[0], MScale[1]) && FMath::IsNearlyEqual(MScale[0], MScale[2])));	//non uniform turns sphere into an ellipsoid so no longer a raycast and requires a more expensive sweep
 
 		const TVector<T, d> UnscaledStart = MInvScale * StartPoint;
-		TVector<T, d> UnscaledDir = MScale * Dir * Length;
-		const T UnscaledLength = UnscaledDir.SafeNormalize();
-		const T ScaledLengthRatio = UnscaledLength / Length;
-		
-		TVector<T, d> UnscaledPosition;
-		TVector<T, d> UnscaledNormal;
-
-		if (MObject->Raycast(UnscaledStart, UnscaledDir, UnscaledLength, MInternalThickness + Thickness * MInvScale[0], OutTime, UnscaledPosition, UnscaledNormal, OutFaceIndex))
+		const TVector<T, d> UnscaledDirDenorm = MScale * Dir;
+		const T LengthScale = UnscaledDirDenorm.Size();
+		if (ensure(LengthScale > TNumericLimits::Min()))
 		{
-			OutTime = ScaledLengthRatio * OutTime;	//is this needed? SafeNormalize above seems bad
-			OutPosition = MScale * UnscaledPosition;
-			OutNormal = (MInvScale * UnscaledNormal).GetSafeNormal();
-			return true;
+			const T LengthScaleInv = 1.f / LengthScale;
+			const T UnscaledLength = Length * LengthScaleInv;
+			const TVector<T, d> UnscaledDir = UnscaledDirDenorm * LengthScaleInv;
+			
+			TVector<T, d> UnscaledPosition;
+			TVector<T, d> UnscaledNormal;
+			float UnscaledTime;
+
+			if (MObject->Raycast(UnscaledStart, UnscaledDir, UnscaledLength, MInternalThickness + Thickness * MInvScale[0], UnscaledTime, UnscaledPosition, UnscaledNormal, OutFaceIndex))
+			{
+				OutTime = LengthScale * UnscaledTime;
+				OutPosition = MScale * UnscaledPosition;
+				OutNormal = (MInvScale * UnscaledNormal).GetSafeNormal();
+				ensure(OutTime <= Length);
+				return true;
+			}
 		}
 		
 		return false;
@@ -114,40 +121,32 @@ public:
 		ensure(FMath::IsNearlyEqual(LocalDir.SizeSquared(), 1, KINDA_SMALL_NUMBER));
 		ensure(Thickness == 0 || (FMath::IsNearlyEqual(OwningScaled.MScale[0], OwningScaled.MScale[1]) && FMath::IsNearlyEqual(OwningScaled.MScale[0], OwningScaled.MScale[2])));
 
-		TVector<T, d> UnscaledDir = OwningScaled.MScale * LocalDir * Length;
-
-		// TODO: Fix and use TVector::SafeNormalize().
-		//We want N / ||N|| and to avoid inf
-		//So we want N / ||N|| < 1 / eps => N eps < ||N||, but this is clearly true for all eps < 1 and N > 0
-		T SizeSqr = UnscaledDir.SizeSquared();
-		T UnscaledLength = 0;
-		if (SizeSqr <= TNumericLimits<T>::Min())
-		{
-			UnscaledDir = TVector<T, d>::AxisVector(0.0f);
-			ensure(false);
-		}
-		else
+		const TVector<T, d> UnscaledDirDenorm = OwningScaled.MScale * LocalDir;
+		const T LengthScale = UnscaledDirDenorm.Size();
+		if (ensure(LengthScale > TNumericLimits::Min()))
 		{
-			UnscaledLength = sqrt(SizeSqr);
-			UnscaledDir /= UnscaledLength;
-		}
-		const T ScaledLengthRatio = UnscaledLength / Length;
+			const T LengthScaleInv = 1.f / LengthScale;
+			const T UnscaledLength = Length * LengthScaleInv;
+			const TVector<T, d> UnscaledDir = UnscaledDirDenorm * LengthScaleInv;
 
-		TVector<T, d> UnscaledPosition;
-		TVector<T, d> UnscaledNormal;
+			TVector<T, d> UnscaledPosition;
+			TVector<T, d> UnscaledNormal;
+			float UnscaledTime;
 
-		TUniquePtr<TImplicitObject<T, d>> HackBPtr(const_cast<TImplicitObject<T,d>*>(&B));	//todo: hack, need scaled object to except raw ptr similar to transformed implicit
-		TImplicitObjectScaled<T, d> ScaledB(MakeSerializable(HackBPtr), OwningScaled.MInvScale);
-		HackBPtr.Release();
+			TUniquePtr<TImplicitObject<T, d>> HackBPtr(const_cast<TImplicitObject<T,d>*>(&B));	//todo: hack, need scaled object to except raw ptr similar to transformed implicit
+			TImplicitObjectScaled<T, d> ScaledB(MakeSerializable(HackBPtr), OwningScaled.MInvScale);
+			HackBPtr.Release();
 
-		TRigidTransform<T, d> BToATMNoScale(BToATM.GetLocation() * OwningScaled.MInvScale, BToATM.GetRotation());
-		
-		if (InternalGeom.SweepGeom(ScaledB, BToATMNoScale, UnscaledDir, UnscaledLength, OutTime, UnscaledPosition, UnscaledNormal, OutFaceIndex, OwningScaled.MInternalThickness + Thickness))
-		{
-			OutTime = ScaledLengthRatio * OutTime;	//is this needed? SafeNormalize above seems bad
-			LocalPosition = OwningScaled.MScale * UnscaledPosition;
-			LocalNormal = (OwningScaled.MInvScale * UnscaledNormal).GetSafeNormal();
-			return true;
+			TRigidTransform<T, d> BToATMNoScale(BToATM.GetLocation() * OwningScaled.MInvScale, BToATM.GetRotation());
+			
+			if (InternalGeom.SweepGeom(ScaledB, BToATMNoScale, UnscaledDir, UnscaledLength, UnscaledTime, UnscaledPosition, UnscaledNormal, OutFaceIndex, OwningScaled.MInternalThickness + Thickness))
+			{
+				OutTime = LengthScale * UnscaledTime;
+				LocalPosition = OwningScaled.MScale * UnscaledPosition;
+				LocalNormal = (OwningScaled.MInvScale * UnscaledNormal).GetSafeNormal();
+				ensure(OutTime <= Length);
+				return true;
+			}
 		}
 
 		return false;
@@ -325,4 +324,4 @@ private:
 template <typename T, int d>
 using TImplicitObjectScaledNonSerializable = TImplicitObjectScaled<T, d, false>;
 
-}
\ No newline at end of file
+}
