Fixed landscape query crash when casting directly at a boundary (Chaos)

diff --git a/Engine/Source/Runtime/Experimental/Chaos/Private/Chaos/HeightField.cpp b/Engine/Source/Runtime/Experimental/Chaos/Private/Chaos/HeightField.cpp
--- a/Engine/Source/Runtime/Experimental/Chaos/Private/Chaos/HeightField.cpp
+++ b/Engine/Source/Runtime/Experimental/Chaos/Private/Chaos/HeightField.cpp
@@ -461,7 +461,9 @@ namespace Chaos
 				// This is a cast down the Z axis, handle in a simpler way as this is the common case
 				const int32 QueryX = (int32)(ClippedFlatRayStart[0] / GeomData.Scale[0]);
 				const int32 QueryY = (int32)(ClippedFlatRayStart[1] / GeomData.Scale[1]);
-				Visitor.VisitRaycast(QueryY * (GeomData.NumCols - 1) + QueryX, CurrentLength);
+				const TVector<int32, 2> Cell = FlatGrid.ClampIndex({QueryX, QueryY});
+
+				Visitor.VisitRaycast(Cell[1] * (GeomData.NumCols - 1) + Cell[0], CurrentLength);
 			}
 			else
 			{
