Fix heightfield serialization not being backwards compat

diff --git a/Engine/Source/Runtime/Core/Public/UObject/ExternalPhysicsCustomObjectVersion.h b/Engine/Source/Runtime/Core/Public/UObject/ExternalPhysicsCustomObjectVersion.h
--- a/Engine/Source/Runtime/Core/Public/UObject/ExternalPhysicsCustomObjectVersion.h
+++ b/Engine/Source/Runtime/Core/Public/UObject/ExternalPhysicsCustomObjectVersion.h
@@ -30,6 +30,9 @@ struct CORE_API FExternalPhysicsCustomObjectVersion
 		//SpatialIdx serialized
 		SpatialIdxSerialized,
 
+		//Save out heightfield data
+		HeightfieldData,
+
 		// -----<new versions can be added above this line>-------------------------------------------------
 		VersionPlusOne,
 		LatestVersion = VersionPlusOne - 1
diff --git a/Engine/Source/Runtime/Experimental/Chaos/Public/Chaos/HeightField.h b/Engine/Source/Runtime/Experimental/Chaos/Public/Chaos/HeightField.h
--- a/Engine/Source/Runtime/Experimental/Chaos/Public/Chaos/HeightField.h
+++ b/Engine/Source/Runtime/Experimental/Chaos/Public/Chaos/HeightField.h
@@ -11,6 +11,7 @@
 #include "Templates/UnrealTypeTraits.h"
 #include "UniformGrid.h"
 #include "Utilities.h"
+#include "UObject/ExternalPhysicsCustomObjectVersion.h"
 
 namespace Chaos
 {
@@ -77,10 +78,15 @@ namespace Chaos
 			
 			GeomData.Serialize(Ar);
 
-			Ar << FlatGrid;
-			Ar << FlattenedBounds.Min;
-			Ar << FlattenedBounds.Max;
-			Ar << LocalBounds;
+			Ar.UsingCustomVersion(FExternalPhysicsCustomObjectVersion::GUID);
+			if (Ar.CustomVer(FExternalPhysicsCustomObjectVersion::GUID) >= FExternalPhysicsCustomObjectVersion::HeightfieldData)
+			{
+				Ar << FlatGrid;
+				Ar << FlattenedBounds.Min;
+				Ar << FlattenedBounds.Max;
+				Ar << LocalBounds;
+			}
+			
 
 			if(Ar.IsLoading())
 			{
@@ -185,9 +191,14 @@ namespace Chaos
 				Ar << MaxValue;
 				Ar << NumRows;
 				Ar << NumCols;
-				Ar << Range;
-				Ar << HeightPerUnit;
-				Ar << CellBounds;
+
+				Ar.UsingCustomVersion(FExternalPhysicsCustomObjectVersion::GUID);
+				if (Ar.CustomVer(FExternalPhysicsCustomObjectVersion::GUID) >= FExternalPhysicsCustomObjectVersion::HeightfieldData)
+				{
+					Ar << Range;
+					Ar << HeightPerUnit;
+					Ar << CellBounds;
+				}
 			}
 		};
 
diff --git a/Engine/Source/Runtime/Landscape/Private/LandscapeCollision.cpp b/Engine/Source/Runtime/Landscape/Private/LandscapeCollision.cpp
--- a/Engine/Source/Runtime/Landscape/Private/LandscapeCollision.cpp
+++ b/Engine/Source/Runtime/Landscape/Private/LandscapeCollision.cpp
@@ -129,7 +129,7 @@ ULandscapeMeshCollisionComponent::FPhysXMeshRef::~FPhysXMeshRef()
 }
 
 // Generate a new guid to force a recache of landscape collison derived data
-#define LANDSCAPE_COLLISION_DERIVEDDATA_VER	TEXT("582D8D523EAF429B92760385A69C8D63")
+#define LANDSCAPE_COLLISION_DERIVEDDATA_VER	TEXT("9CAF1A48A5224D4BB307CEBE2776DE27")
 
 static FString GetHFDDCKeyString(const FName& Format, bool bDefMaterial, const FGuid& StateId, const TArray<UPhysicalMaterial*>& PhysicalMaterials)
 {
