Change acceleration structures to have a fast path for collection case

diff --git a/Engine/Source/Runtime/Experimental/Chaos/Public/Chaos/AABBTree.h b/Engine/Source/Runtime/Experimental/Chaos/Public/Chaos/AABBTree.h
--- a/Engine/Source/Runtime/Experimental/Chaos/Public/Chaos/AABBTree.h
+++ b/Engine/Source/Runtime/Experimental/Chaos/Public/Chaos/AABBTree.h
@@ -244,40 +244,87 @@ public:
 		return TUniquePtr<ISpatialAcceleration<TPayloadType, T, 3>>(new TAABBTree<TPayloadType, TLeafType, T>(*this));
 	}
 
-	virtual void Raycast(const TVector<T, 3>& Start, const TVector<T, 3>& Dir, const T OriginalLength, ISpatialVisitor<TPayloadType, T>& Visitor) const override
+	virtual void Raycast(const TVector<T, 3>& Start, const TVector<T, 3>& Dir, const T Length, ISpatialVisitor<TPayloadType, T>& Visitor) const override
 	{
 		TSpatialVisitor<TPayloadType, T> ProxyVisitor(Visitor);
-		Raycast(Start, Dir, OriginalLength, ProxyVisitor);
+		Raycast(Start, Dir, Length, ProxyVisitor);
 	}
 
 	template <typename SQVisitor>
-	void Raycast(const TVector<T, 3>& Start, const TVector<T, 3>& Dir, const T OriginalLength, SQVisitor& Visitor) const
+	void Raycast(const TVector<T, 3>& Start, const TVector<T, 3>& Dir, const T Length, SQVisitor& Visitor) const
 	{
-		QueryImp<EAABBQueryType::Raycast>(Start, Dir, OriginalLength, TVector<T,3>(), TBox<T,3>(), Visitor);
+		ensure(Length);
+		T CurLength = Length;
+		bool bParallel[3];
+		TVector<T, 3> InvDir;
+
+		T InvLength = 1 / Length;
+		for (int Axis = 0; Axis < 3; ++Axis)
+		{
+			bParallel[Axis] = Dir[Axis] == 0;
+			InvDir[Axis] = bParallel[Axis] ? 0 : 1 / Dir[Axis];
+		}
+
+		QueryImp<EAABBQueryType::Raycast>(Start, Dir, InvDir, bParallel, CurLength, InvLength, TVector<T,3>(), TBox<T,3>(), Visitor);
 	}
 
-	void Sweep(const TVector<T, 3>& Start, const TVector<T, 3>& Dir, T OriginalLength, const TVector<T, 3> QueryHalfExtents, ISpatialVisitor<TPayloadType, T>& Visitor) const override
+	template <typename SQVisitor>
+	bool RaycastFast(const TVector<T, 3>& Start, const TVector<T, 3>& Dir, const TVector<T,3>& InvDir, const bool* bParallel, T& CurLength, T& InvLength, SQVisitor& Visitor) const
+	{
+		return QueryImp<EAABBQueryType::Raycast>(Start, Dir, InvDir, bParallel, CurLength, InvLength, TVector<T, 3>(), TBox<T, 3>(), Visitor);
+	}
+
+	void Sweep(const TVector<T, 3>& Start, const TVector<T, 3>& Dir, const T Length, const TVector<T, 3> QueryHalfExtents, ISpatialVisitor<TPayloadType, T>& Visitor) const override
 	{
 		TSpatialVisitor<TPayloadType, T> ProxyVisitor(Visitor);
-		Sweep(Start, Dir, OriginalLength, QueryHalfExtents, ProxyVisitor);
+		Sweep(Start, Dir, Length, QueryHalfExtents, ProxyVisitor);
 	}
 
 	template <typename SQVisitor>
-	void Sweep(const TVector<T, 3>& Start, const TVector<T, 3>& Dir, const T OriginalLength, const TVector<T, 3> QueryHalfExtents, SQVisitor& Visitor) const
+	void Sweep(const TVector<T, 3>& Start, const TVector<T, 3>& Dir, const T Length, const TVector<T, 3> QueryHalfExtents, SQVisitor& Visitor) const
 	{
-		QueryImp<EAABBQueryType::Sweep>(Start, Dir, OriginalLength, QueryHalfExtents, TBox<T, 3>(), Visitor);
+		bool bParallel[3];
+		TVector<T, 3> InvDir;
+
+		ensure(Length);
+		T CurLength = Length;
+		T InvLength = 1 / Length;
+		for (int Axis = 0; Axis < 3; ++Axis)
+		{
+			bParallel[Axis] = Dir[Axis] == 0;
+			InvDir[Axis] = bParallel[Axis] ? 0 : 1 / Dir[Axis];
+		}
+
+		QueryImp<EAABBQueryType::Sweep>(Start, Dir, InvDir, bParallel, CurLength, InvLength, QueryHalfExtents, TBox<T, 3>(), Visitor);
+	}
+
+	template <typename SQVisitor>
+	bool SweepFast(const TVector<T, 3>& Start, const TVector<T, 3>& Dir, const TVector<T,3>& InvDir, const bool* bParallel, T& CurLength, T& InvLength, const TVector<T, 3> QueryHalfExtents, SQVisitor& Visitor) const
+	{
+		return QueryImp<EAABBQueryType::Sweep>(Start, Dir, InvDir, bParallel, CurLength, InvLength, QueryHalfExtents, TBox<T, 3>(), Visitor);
 	}
 
 	void Overlap(const TBox<T, 3>& QueryBounds, ISpatialVisitor<TPayloadType, T>& Visitor) const override
 	{
 		TSpatialVisitor<TPayloadType, T> ProxyVisitor(Visitor);
-		return Overlap(QueryBounds, ProxyVisitor);
+		Overlap(QueryBounds, ProxyVisitor);
 	}
 
 	template <typename SQVisitor>
 	void Overlap(const TBox<T,3>& QueryBounds, SQVisitor& Visitor) const
 	{
-		QueryImp<EAABBQueryType::Overlap>(TVector<T,3>(), TVector<T, 3>(), 0, TVector<T,3>(), QueryBounds, Visitor);
+		OverlapFast(QueryBounds, Visitor);
+	}
+
+	template <typename SQVisitor>
+	bool OverlapFast(const TBox<T, 3>& QueryBounds, SQVisitor& Visitor) const
+	{
+		//dummy variables to reuse templated path
+		T Length;
+		T InvLength;
+		TVector<T, 3> InvDir;
+		bool bParallel;
+		return QueryImp<EAABBQueryType::Overlap>(TVector<T, 3>(), TVector<T, 3>(), InvDir, &bParallel, Length, InvLength, TVector<T, 3>(), QueryBounds, Visitor);
 	}
 
 	virtual void RemoveElement(const TPayloadType& Payload)
@@ -403,25 +450,8 @@ private:
 	using FNode = TAABBTreeNode<T>;
 
 	template <EAABBQueryType Query, typename SQVisitor>
-	void QueryImp(const TVector<T, 3>& Start, const TVector<T, 3>& Dir, const T OriginalLength, const TVector<T, 3> QueryHalfExtents, const TBox<T,3>& QueryBounds, SQVisitor& Visitor) const
+	bool QueryImp(const TVector<T, 3>& Start, const TVector<T, 3>& Dir, const TVector<T,3>& InvDir, const bool* bParallel, T& CurrentLength, T& InvCurrentLength, const TVector<T, 3> QueryHalfExtents, const TBox<T,3>& QueryBounds, SQVisitor& Visitor) const
 	{
-		T CurrentLength = OriginalLength;
-
-		bool bParallel[3];
-		TVector<T, 3> InvDir;
-
-		T InvCurrentLength;
-
-		if (Query != EAABBQueryType::Overlap)
-		{
-			InvCurrentLength = 1 / CurrentLength;
-			for (int Axis = 0; Axis < 3; ++Axis)
-			{
-				bParallel[Axis] = Dir[Axis] == 0;
-				InvDir[Axis] = bParallel[Axis] ? 0 : 1 / Dir[Axis];
-			}
-		}
-
 		TVector<T, 3> TmpPosition;
 		T TOI = 0;
 
@@ -444,7 +474,7 @@ private:
 
 				if (!bContinue)
 				{
-					return;
+					return false;
 				}
 
 				if (Query != EAABBQueryType::Overlap)
@@ -474,7 +504,7 @@ private:
 				
 				if (!bContinue)
 				{
-					return;
+					return false;
 				}
 
 				if (Query != EAABBQueryType::Overlap)
@@ -495,9 +525,12 @@ private:
 		while (NodeStack.Num())
 		{
 			const FNodeQueueEntry NodeEntry = NodeStack.Pop(false);
-			if (NodeEntry.TOI > CurrentLength)
+			if (Query != EAABBQueryType::Overlap)
 			{
-				continue;
+				if (NodeEntry.TOI > CurrentLength)
+				{
+					continue;
+				}
 			}
 
 			const FNode& Node = Nodes[NodeEntry.NodeIdx];
@@ -508,19 +541,19 @@ private:
 				{
 					if (Leaf.OverlapFast(QueryBounds, Visitor) == false)
 					{
-						return;
+						return false;
 					}
 				}
 				else if (Query == EAABBQueryType::Sweep)
 				{
 					if (Leaf.SweepFast(Start, Dir, InvDir, bParallel, CurrentLength, InvCurrentLength, QueryHalfExtents, Visitor) == false)
 					{
-						return;
+						return false;
 					}
 				}
 				else if (Leaf.RaycastFast(Start, Dir, InvDir, bParallel, CurrentLength, InvCurrentLength, Visitor) == false)
 				{
-					return;
+					return false;
 				}
 			}
 			else
@@ -537,6 +570,8 @@ private:
 				}
 			}
 		}
+
+		return true;
 	}
 
 	template <typename TParticles>
diff --git a/Engine/Source/Runtime/Experimental/Chaos/Public/Chaos/BoundingVolume.h b/Engine/Source/Runtime/Experimental/Chaos/Public/Chaos/BoundingVolume.h
--- a/Engine/Source/Runtime/Experimental/Chaos/Public/Chaos/BoundingVolume.h
+++ b/Engine/Source/Runtime/Experimental/Chaos/Public/Chaos/BoundingVolume.h
@@ -255,38 +255,132 @@
 		return MGlobalPayloads;
 	}
 
-	virtual void Raycast(const TVector<T, d>& Start, const TVector<T, d>& Dir, const T OriginalLength, ISpatialVisitor<TPayloadType, T>& Visitor) const override
+	virtual void Raycast(const TVector<T, d>& Start, const TVector<T, d>& Dir, const T Length, ISpatialVisitor<TPayloadType, T>& Visitor) const override
 	{
 		TSpatialVisitor<TPayloadType, T> ProxyVisitor(Visitor);
-		Raycast(Start, Dir, OriginalLength, ProxyVisitor);
+		return Raycast(Start, Dir, Length, ProxyVisitor);
 	}
 
 	template <typename SQVisitor>
 	bool RaycastFast(const TVector<T,d>& Start, const TVector<T,d>& Dir, const TVector<T,d>& InvDir, const bool* bParallel, T& CurrentLength, T& InvCurrentLength, SQVisitor& Visitor) const
 	{
-		if (Raycast(Start, Dir, CurrentLength, Visitor))
-		{
-			InvCurrentLength = 1 / CurrentLength;
-			return true;
-		}
-		return false;
+		return RaycastImp(Start, Dir, InvDir, bParallel, CurrentLength, InvCurrentLength, Visitor);
 	}
 
 	template <typename SQVisitor, bool bPruneDuplicates = true>
-	bool Raycast(const TVector<T, d>& Start, const TVector<T, d>& Dir, const T OriginalLength, SQVisitor& Visitor) const
+	void Raycast(const TVector<T, d>& Start, const TVector<T, d>& Dir, const T Length, SQVisitor& Visitor) const
 	{
-		T CurrentLength = OriginalLength;
-
+		ensure(Length);
+		T CurLength = Length;
 		bool bParallel[d];
 		TVector<T, d> InvDir;
 
-		T InvCurrentLength = 1 / CurrentLength;
+		T InvLength = 1 / Length;
 		for (int Axis = 0; Axis < d; ++Axis)
 		{
 			bParallel[Axis] = Dir[Axis] == 0;
 			InvDir[Axis] = bParallel[Axis] ? 0 : 1 / Dir[Axis];
 		}
 
+		RaycastImp<SQVisitor, bPruneDuplicates>(Start, Dir, InvDir, bParallel, CurLength, InvLength, Visitor);
+	}
+
+	void Sweep(const TVector<T, d>& Start, const TVector<T, d>& Dir, const T Length, const TVector<T, d> QueryHalfExtents, ISpatialVisitor<TPayloadType, T>& Visitor) const override
+	{
+		TSpatialVisitor<TPayloadType, T> ProxyVisitor(Visitor);
+		Sweep(Start, Dir, Length, QueryHalfExtents, ProxyVisitor);
+	}
+
+	template <typename SQVisitor>
+	bool SweepFast(const TVector<T,d>& Start, const TVector<T,d>& Dir, const TVector<T,d>& InvDir, const bool* bParallel, T& CurrentLength, T& InvCurrentLength, const TVector<T,d>& QueryHalfExtents, SQVisitor& Visitor) const
+	{
+		return SweepImp(Start, Dir, InvDir, bParallel, CurrentLength, InvCurrentLength, QueryHalfExtents, Visitor);
+	}
+
+	template <typename SQVisitor, bool bPruneDuplicates = true>
+	void Sweep(const TVector<T, d>& Start, const TVector<T, d>& Dir, const T Length, const TVector<T, d> QueryHalfExtents, SQVisitor& Visitor) const
+	{
+		ensure(Length);
+		T CurLength = Length;
+		bool bParallel[d];
+		TVector<T, d> InvDir;
+
+		T InvLength = 1 / Length;
+		for (int Axis = 0; Axis < d; ++Axis)
+		{
+			bParallel[Axis] = Dir[Axis] == 0;
+			InvDir[Axis] = bParallel[Axis] ? 0 : 1 / Dir[Axis];
+		}
+
+		SweepImp<SQVisitor, bPruneDuplicates>(Start, Dir, InvDir, bParallel, CurLength, InvLength, QueryHalfExtents, Visitor);
+	}
+
+	virtual void Overlap(const TBox<T, d>& QueryBounds, ISpatialVisitor<TPayloadType, T>& Visitor) const override
+	{
+		TSpatialVisitor<TPayloadType, T> ProxyVisitor(Visitor);
+		Overlap(QueryBounds, ProxyVisitor);
+	}
+
+	template <typename SQVisitor>
+	bool OverlapFast(const TBox<T, d>& QueryBounds, SQVisitor& Visitor) const
+	{
+		return OverlapImp(QueryBounds, Visitor);
+	}
+
+	template <typename SQVisitor, bool bPruneDuplicates = true>
+	void Overlap(const TBox<T, d>& QueryBounds, SQVisitor& Visitor) const
+	{
+		OverlapImp<SQVisitor, bPruneDuplicates>(QueryBounds, Visitor);
+	}
+
+	virtual void Serialize(FChaosArchive& Ar) override
+	{
+		Ar.UsingCustomVersion(FExternalPhysicsCustomObjectVersion::GUID);
+		if (Ar.CustomVer(FExternalPhysicsCustomObjectVersion::GUID) < FExternalPhysicsCustomObjectVersion::GlobalElementsHaveBounds)
+		{
+			TArray<TPayloadType> TmpPayloads;
+			Ar << TmpPayloads;
+			MGlobalPayloads.Reserve(TmpPayloads.Num());
+			for (auto& Payload : TmpPayloads)
+			{
+				MGlobalPayloads.Add({ Payload, TBox<T,d>(TVector<T,d>(TNumericLimits<T>::Lowest()), TVector<T,d>(TNumericLimits<T>::Max())) });
+			}
+			MaxPayloadBounds = DefaultMaxPayloadBounds;
+		}
+		else
+		{
+			Ar << MGlobalPayloads;
+			Ar << MaxPayloadBounds;
+		}
+
+		Ar << MGrid;
+		Ar << MElements;
+		Ar << MDirtyElements;
+		Ar << bIsEmpty;
+
+		TArray<TPayloadType> Payloads;
+		if (!Ar.IsLoading())
+		{
+			MPayloadInfo.GenerateKeyArray(Payloads);
+		}
+		Ar << Payloads;
+
+		for (auto Payload : Payloads)
+		{
+			auto& Info = MPayloadInfo.FindOrAdd(Payload);
+			Ar << Info;
+		}
+
+	}
+
+private:
+
+	using FCellElement = TBVCellElement<TPayloadType, T, d>;
+	using FPayloadInfo = TBVPayloadInfo<T, d>;
+
+	template <typename SQVisitor, bool bPruneDuplicates = true>
+	bool RaycastImp(const TVector<T, d>& Start, const TVector<T, d>& Dir, const TVector<T, d>& InvDir, const bool* bParallel, T& CurrentLength, T& InvCurrentLength, SQVisitor& Visitor) const
+	{
 		TVector<T, d> TmpPosition;
 		T TOI;
 
@@ -297,11 +391,12 @@
 			{
 				TSpatialVisitorData<TPayloadType> VisitData(Elem.Payload, true, InstanceBounds);
 				const bool bContinue = Visitor.VisitRaycast(VisitData, CurrentLength);
+				InvCurrentLength = 1 / CurrentLength;
+
 				if (!bContinue)
 				{
 					return false;
 				}
-				InvCurrentLength = 1 / CurrentLength;
 			}
 		}
 
@@ -323,7 +418,7 @@
 		TBox<T, d> GlobalBounds(MGrid.MinCorner(), MGrid.MaxCorner());
 		TVector<T, d> NextStart;
 		TVector<int32, d> CellIdx;
-		bool bCellsLeft = MElements.Num() && TBox<T,d>::RaycastFast(GlobalBounds.Min(), GlobalBounds.Max(), Start, Dir, InvDir, bParallel, CurrentLength, InvCurrentLength, TOI, NextStart);
+		bool bCellsLeft = MElements.Num() && TBox<T, d>::RaycastFast(GlobalBounds.Min(), GlobalBounds.Max(), Start, Dir, InvDir, bParallel, CurrentLength, InvCurrentLength, TOI, NextStart);
 		if (bCellsLeft)
 		{
 			CellIdx = MGrid.Cell(NextStart);
@@ -340,7 +435,6 @@
 				{
 					if (bPruneDuplicates)
 					{
-
 						bool bSkip = false;
 						if (Elem.StartIdx[0] != Elem.EndIdx[0] || Elem.StartIdx[1] != Elem.EndIdx[1] || Elem.StartIdx[2] != Elem.EndIdx[2])
 						{
@@ -366,7 +460,7 @@
 						}
 					}
 					const auto& InstanceBounds = Elem.Bounds;
-					if (TBox<T,d>::RaycastFast(InstanceBounds.Min(), InstanceBounds.Max(), Start, Dir, InvDir, bParallel, CurrentLength, InvCurrentLength, TOI, TmpPosition))
+					if (TBox<T, d>::RaycastFast(InstanceBounds.Min(), InstanceBounds.Max(), Start, Dir, InvDir, bParallel, CurrentLength, InvCurrentLength, TOI, TmpPosition))
 					{
 						TSpatialVisitorData<TPayloadType> VisitData(Elem.Payload, true, InstanceBounds);
 						const bool bContinue = Visitor.VisitRaycast(VisitData, CurrentLength);
@@ -374,7 +468,7 @@
 						{
 							return false;
 						}
-						InvCurrentLength = 1 / CurrentLength;
+						InvCurrentLength = 1 / CurrentLength;	//todo: pass into visitor
 					}
 				}
 
@@ -428,43 +522,14 @@
 				NextStart = NextStart + Dir * BestTime;
 			} while (true);
 		}
-		
+
 		return true;
 	}
 
-	void Sweep(const TVector<T, d>& Start, const TVector<T, d>& Dir, T OriginalLength, const TVector<T, d> QueryHalfExtents, ISpatialVisitor<TPayloadType, T>& Visitor) const override
-	{
-		TSpatialVisitor<TPayloadType, T> ProxyVisitor(Visitor);
-		Sweep(Start, Dir, OriginalLength, QueryHalfExtents, ProxyVisitor);
-	}
-
-	template <typename SQVisitor>
-	bool SweepFast(const TVector<T,d>& Start, const TVector<T,d>& Dir, const TVector<T,d>& InvDir, const bool* bParallel, T& CurrentLength, T& InvCurrentLength, const TVector<T,d>& QueryHalfExtents, SQVisitor& Visitor) const
-	{
-		if (Sweep(Start, Dir, CurrentLength, QueryHalfExtents, Visitor))
-		{
-			InvCurrentLength = 1 / CurrentLength;
-			return true;
-		}
-		return false;
-	}
-
 	template <typename SQVisitor, bool bPruneDuplicates = true>
-	bool Sweep(const TVector<T, d>& Start, const TVector<T, d>& Dir, T OriginalLength, const TVector<T, d> QueryHalfExtents, SQVisitor& Visitor) const
+	bool SweepImp(const TVector<T, d>& Start, const TVector<T, d>& Dir, const TVector<T,d>& InvDir, const bool* bParallel, T& CurrentLength, T& InvCurrentLength, const TVector<T, d> QueryHalfExtents, SQVisitor& Visitor) const
 	{
-		T CurrentLength = OriginalLength;
-
-		T TOI = 0;	//not needed, but fixes compiler warning
-		bool bParallel[d];
-		TVector<T, d> InvDir;
-
-		T InvCurrentLength = 1 / CurrentLength;
-		for (int Axis = 0; Axis < d; ++Axis)
-		{
-			bParallel[Axis] = Dir[Axis] == 0;
-			InvDir[Axis] = bParallel[Axis] ? 0 : 1 / Dir[Axis];
-		}
-
+		T TOI = 0;
 		for (const auto& Elem : MGlobalPayloads)
 		{
 			const TBox<T, d>& InstanceBounds = Elem.Bounds;
@@ -507,7 +572,7 @@
 		}
 
 		TBox<T, d> GlobalBounds(MGrid.MinCorner() - QueryHalfExtents, MGrid.MaxCorner() + QueryHalfExtents);
-		
 
+
 		struct FCellIntersection
 		{
@@ -515,11 +580,11 @@
 			T TOI;
 		};
 
-		
+
 		TVector<T, d> HitPoint;
 		FGridSet IdxsSeen(MGrid.Counts());
 		FGridSet CellsVisited(MGrid.Counts());
-		const bool bInitialHit = TBox<T,d>::RaycastFast(GlobalBounds.Min(), GlobalBounds.Max(), Start, Dir, InvDir, bParallel, CurrentLength, InvCurrentLength, TOI, HitPoint);
+		const bool bInitialHit = TBox<T, d>::RaycastFast(GlobalBounds.Min(), GlobalBounds.Max(), Start, Dir, InvDir, bParallel, CurrentLength, InvCurrentLength, TOI, HitPoint);
 		if (bInitialHit)
 		{
 			//Flood fill from inflated cell so that we get all cells along the ray
@@ -579,7 +644,7 @@
 						const TVector<T, d> NeighborCenter = MGrid.Location(NeighborIdx);
 						const TVector<T, d> Min = NeighborCenter - QueryHalfExtents - HalfDx;
 						const TVector<T, d> Max = NeighborCenter + QueryHalfExtents + HalfDx;
-						if (TBox<T,d>::RaycastFast(Min, Max, Start, Dir, InvDir, bParallel, CurrentLength, InvCurrentLength, TOI, HitPoint))
+						if (TBox<T, d>::RaycastFast(Min, Max, Start, Dir, InvDir, bParallel, CurrentLength, InvCurrentLength, TOI, HitPoint))
 						{
 							IdxsQueue.Add({ NeighborIdx, TOI });	//should we sort by TOI?
 						}
@@ -601,7 +666,7 @@
 								{
 									for (int32 Z = Elem.StartIdx[2]; Z <= Elem.EndIdx[2]; ++Z)
 									{
-										if (CellsVisited.Contains(TVector<int32,3>(X,Y,Z)))
+										if (CellsVisited.Contains(TVector<int32, 3>(X, Y, Z)))
 										{
 											bSkip = true;
 											break;
@@ -620,7 +685,7 @@
 					const TBox<T, d>& InstanceBounds = Elem.Bounds;
 					const TVector<T, d> Min = InstanceBounds.Min() - QueryHalfExtents;
 					const TVector<T, d> Max = InstanceBounds.Max() + QueryHalfExtents;
-					if (TBox<T,d>::RaycastFast(Min, Max, Start, Dir, InvDir, bParallel, CurrentLength, InvCurrentLength, TOI, HitPoint))
+					if (TBox<T, d>::RaycastFast(Min, Max, Start, Dir, InvDir, bParallel, CurrentLength, InvCurrentLength, TOI, HitPoint))
 					{
 						TSpatialVisitorData<TPayloadType> VisitData(Elem.Payload, true, InstanceBounds);
 						const bool bContinue = Visitor.VisitSweep(VisitData, CurrentLength);
@@ -638,20 +703,8 @@
 		return true;
 	}
 
-	void Overlap(const TBox<T, d>& QueryBounds, ISpatialVisitor<TPayloadType, T>& Visitor) const override
-	{
-		TSpatialVisitor<TPayloadType, T> ProxyVisitor(Visitor);
-		Overlap(QueryBounds, ProxyVisitor);
-	}
-
-	template <typename SQVisitor>
-	bool OverlapFast(const TBox<T, d>& QueryBounds, SQVisitor& Visitor) const
-	{
-		return Overlap(QueryBounds, Visitor);
-	}
-
 	template <typename SQVisitor, bool bPruneDuplicates = true>
-	bool Overlap(const TBox<T, d>& QueryBounds, SQVisitor& Visitor) const
+	bool OverlapImp(const TBox<T, d>& QueryBounds, SQVisitor& Visitor) const
 	{
 		for (const auto& Elem : MGlobalPayloads)
 		{
@@ -719,51 +772,6 @@
 		return true;
 	}
 
-	virtual void Serialize(FChaosArchive& Ar) override
-	{
-		Ar.UsingCustomVersion(FExternalPhysicsCustomObjectVersion::GUID);
-		if (Ar.CustomVer(FExternalPhysicsCustomObjectVersion::GUID) < FExternalPhysicsCustomObjectVersion::GlobalElementsHaveBounds)
-		{
-			TArray<TPayloadType> TmpPayloads;
-			Ar << TmpPayloads;
-			MGlobalPayloads.Reserve(TmpPayloads.Num());
-			for (auto& Payload : TmpPayloads)
-			{
-				MGlobalPayloads.Add({ Payload, TBox<T,d>(TVector<T,d>(TNumericLimits<T>::Lowest()), TVector<T,d>(TNumericLimits<T>::Max())) });
-			}
-			MaxPayloadBounds = DefaultMaxPayloadBounds;
-		}
-		else
-		{
-			Ar << MGlobalPayloads;
-			Ar << MaxPayloadBounds;
-		}
-
-		Ar << MGrid;
-		Ar << MElements;
-		Ar << MDirtyElements;
-		Ar << bIsEmpty;
-
-		TArray<TPayloadType> Payloads;
-		if (!Ar.IsLoading())
-		{
-			MPayloadInfo.GenerateKeyArray(Payloads);
-		}
-		Ar << Payloads;
-
-		for (auto Payload : Payloads)
-		{
-			auto& Info = MPayloadInfo.FindOrAdd(Payload);
-			Ar << Info;
-		}
-
-	}
-
-private:
-
-	using FCellElement = TBVCellElement<TPayloadType, T, d>;
-	using FPayloadInfo = TBVPayloadInfo<T, d>;
-
 	template <typename ParticleView>
 	void GenerateTree(const ParticleView& Particles, const bool bUseVelocity, const T Dt, const int32 MaxCells)
 	{
diff --git a/Engine/Source/Runtime/Experimental/Chaos/Public/Chaos/ISpatialAcceleration.h b/Engine/Source/Runtime/Experimental/Chaos/Public/Chaos/ISpatialAcceleration.h
--- a/Engine/Source/Runtime/Experimental/Chaos/Public/Chaos/ISpatialAcceleration.h
+++ b/Engine/Source/Runtime/Experimental/Chaos/Public/Chaos/ISpatialAcceleration.h
@@ -82,12 +82,6 @@ public:
 	virtual bool Sweep(const TSpatialVisitorData<TPayloadType>& Instance, T& CurLength) = 0;
 };
 
-enum class ESpatialAccelerationType
-{
-	Grid,
-	BVH
-};
-
 /**
  * Can be implemented by external, non-chaos systems to collect / render
  * debug information from spacial structures. When passed to the debug
@@ -113,9 +107,12 @@ enum ESpatialAcceleration
 	AABBTree,
 	AABBTreeBV,
 	Collection,
-	Unknown
+	Unknown,
+	//For custom types continue the enum after ESpatialAcceleration::Unknown
 };
 
+using SpatialAccelerationType = uint8;	//see ESpatialAcceleration. Projects can add their own custom types by using enum values higher than ESpatialAcceleration::Unknown
+
 
 template <typename TPayloadType, typename T>
 struct TPayloadBoundsElement
@@ -149,7 +146,7 @@ class CHAOS_API ISpatialAcceleration
 {
 public:
 
-	ISpatialAcceleration(ESpatialAcceleration InType = ESpatialAcceleration::Unknown)
+	ISpatialAcceleration(SpatialAccelerationType InType = ESpatialAcceleration::Unknown)
 		: Type(InType)
 	{
 	}
@@ -157,8 +154,8 @@ public:
 
 	virtual TArray<TPayloadType> FindAllIntersections(const TBox<T, d>& Box) const { check(false); return TArray<TPayloadType>(); }
 
-	virtual void Raycast(const TVector<T, d>& Start, const TVector<T, d>& Dir, const T OriginalLength, ISpatialVisitor<TPayloadType, T>& Visitor) const { check(false); }
-	virtual void Sweep(const TVector<T, d>& Start, const TVector<T, d>& Dir, T OriginalLength, const TVector<T, d> QueryHalfExtents, ISpatialVisitor<TPayloadType, T>& Visitor) const { check(false); }
+	virtual void Raycast(const TVector<T, d>& Start, const TVector<T, d>& Dir, const T Length, ISpatialVisitor<TPayloadType, T>& Visitor) const { check(false); }
+	virtual void Sweep(const TVector<T, d>& Start, const TVector<T, d>& Dir, const T Length, const TVector<T, d> QueryHalfExtents, ISpatialVisitor<TPayloadType, T>& Visitor) const { check(false);}
 	virtual void Overlap(const TBox<T, d>& QueryBounds, ISpatialVisitor<TPayloadType, T>& Visitor) const { check(false); }
 
 	virtual void RemoveElement(const TPayloadType& Payload)
@@ -225,7 +222,7 @@ public:
 	}
 
 private:
-	ESpatialAcceleration Type;
+	SpatialAccelerationType Type;
 };
 
 template <typename TBase, typename TDerived>
