Fixed bad scale handling in heightfield construction causing SQ unit tests to fail (quick fix - will remove scale from geom build at later time as it's no longer required)

diff --git a/Engine/Source/Runtime/Experimental/Chaos/Private/Chaos/HeightField.cpp b/Engine/Source/Runtime/Experimental/Chaos/Private/Chaos/HeightField.cpp
--- a/Engine/Source/Runtime/Experimental/Chaos/Private/Chaos/HeightField.cpp
+++ b/Engine/Source/Runtime/Experimental/Chaos/Private/Chaos/HeightField.cpp
@@ -334,8 +334,9 @@ namespace Chaos
 	THeightField<T>::THeightField(TArray<T>&& Height, int32 NumRows, int32 NumCols, const TVector<T, 3>& InScale)
 		: TImplicitObject<T, 3>(EImplicitObject::HasBoundingBox, ImplicitObjectType::HeightField)
 	{
-		BuildGeomData<T, T>(MakeArrayView(Height), NumRows, NumCols, InScale, [](const T InVal) {return InVal; }, GeomData, LocalBounds);
+		BuildGeomData<T, T>(MakeArrayView(Height), NumRows, NumCols, TVector<T, 3>(1), [](const T InVal) {return InVal; }, GeomData, LocalBounds);
 		CalcBounds();
+		SetScale(InScale);
 	}
 
 	template<typename T>
@@ -347,8 +348,9 @@ namespace Chaos
 			return (float)((int32)InVal - 32768);
 		};
 
-		BuildGeomData<T, const uint16>(InHeights, InNumRows, InNumCols, InScale, MoveTemp(ConversionFunc), GeomData, LocalBounds);
+		BuildGeomData<T, const uint16>(InHeights, InNumRows, InNumCols, TVector<T, 3>(1), MoveTemp(ConversionFunc), GeomData, LocalBounds);
 		CalcBounds();
+		SetScale(InScale);
 	}
 
 	template<typename T>
diff --git a/Engine/Source/Runtime/Experimental/Chaos/Public/Chaos/HeightField.h b/Engine/Source/Runtime/Experimental/Chaos/Public/Chaos/HeightField.h
--- a/Engine/Source/Runtime/Experimental/Chaos/Public/Chaos/HeightField.h
+++ b/Engine/Source/Runtime/Experimental/Chaos/Public/Chaos/HeightField.h
@@ -86,6 +86,10 @@ namespace Chaos
 				Ar << FlattenedBounds.Max;
 				Ar << LocalBounds;
 			}
+			else
+			{
+				CalcBounds();
+			}
 			
 
 			if(Ar.IsLoading())
