Use defaults with default constructor. Fixes trimesh using random max bounds leading to every tri being in global array

diff --git a/Engine/Source/Runtime/Experimental/Chaos/Public/Chaos/AABBTree.h b/Engine/Source/Runtime/Experimental/Chaos/Public/Chaos/AABBTree.h
--- a/Engine/Source/Runtime/Experimental/Chaos/Public/Chaos/AABBTree.h
+++ b/Engine/Source/Runtime/Experimental/Chaos/Public/Chaos/AABBTree.h
@@ -220,10 +220,15 @@ public:
 	static constexpr int D = 3;
 	using TType = T;
 	static constexpr T DefaultMaxPayloadBounds = 100000;
+	static constexpr int32 DefaultMaxChildrenInLeaf = 12;
+	static constexpr int32 DefaultMaxTreeDepth = 16;
 	static constexpr ESpatialAcceleration StaticType = TIsSame<TAABBTreeLeafArray<TPayloadType, T>, TLeafType>::Value ? ESpatialAcceleration::AABBTree : 
 		(TIsSame<TBoundingVolume<TPayloadType, T, 3>, TLeafType>::Value ? ESpatialAcceleration::AABBTreeBV : ESpatialAcceleration::Unknown);
 	TAABBTree()
 		: ISpatialAcceleration<TPayloadType, T, 3>(StaticType)
+		, MaxChildrenInLeaf(DefaultMaxChildrenInLeaf)
+		, MaxTreeDepth(DefaultMaxTreeDepth)
+		, MaxPayloadBounds(DefaultMaxPayloadBounds)
 	{
 	}
 
diff --git a/Engine/Source/Runtime/Experimental/Chaos/Public/Chaos/BoundingVolume.h b/Engine/Source/Runtime/Experimental/Chaos/Public/Chaos/BoundingVolume.h
--- a/Engine/Source/Runtime/Experimental/Chaos/Public/Chaos/BoundingVolume.h
+++ b/Engine/Source/Runtime/Experimental/Chaos/Public/Chaos/BoundingVolume.h
@@ -123,6 +123,7 @@ class TBoundingVolume final : public ISpatialAcceleration<InPayloadType, T,d>
 	static constexpr ESpatialAcceleration StaticType = ESpatialAcceleration::BoundingVolume;
 	TBoundingVolume()
 		: ISpatialAcceleration<InPayloadType, T, d>(StaticType)
+		, MaxPayloadBounds(DefaultMaxPayloadBounds)
 	{
 	}
 
@@ -415,6 +416,11 @@ private:
 			}
 		}
 
+		if (bIsEmpty)
+		{
+			return true;
+		}
+
 		TBox<T, d> GlobalBounds(MGrid.MinCorner(), MGrid.MaxCorner());
 		TVector<T, d> NextStart;
 		TVector<int32, d> CellIdx;
@@ -566,7 +572,7 @@ private:
 			}
 		}
 
-		if (MElements.Num() == 0)
+		if (bIsEmpty)
 		{
 			return true;
 		}
@@ -732,6 +738,11 @@ private:
 			}
 		}
 
+		if (bIsEmpty)
+		{
+			return true;
+		}
+
 		TBox<T, d> GlobalBounds(MGrid.MinCorner(), MGrid.MaxCorner());
 
 		const TVector<int32, d> StartIndex = MGrid.ClampIndex(MGrid.Cell(QueryBounds.Min()));
