Refactor spatial collection to make it easier to ab test stuff. Explicitly placing large objects into second spatial acceleration

diff --git a/Engine/Source/Runtime/Experimental/Chaos/Private/Chaos/PBDRigidsEvolution.cpp b/Engine/Source/Runtime/Experimental/Chaos/Private/Chaos/PBDRigidsEvolution.cpp
--- a/Engine/Source/Runtime/Experimental/Chaos/Private/Chaos/PBDRigidsEvolution.cpp
+++ b/Engine/Source/Runtime/Experimental/Chaos/Private/Chaos/PBDRigidsEvolution.cpp
@@ -71,7 +71,7 @@ namespace Chaos
 		return ESubsequentsMode::TrackSubsequents;
 	}
 
-	int32 BroadphaseType = 1;
+	int32 BroadphaseType = 3;
 	FAutoConsoleVariableRef CVarBroadphaseIsTree(TEXT("p.BroadphaseType"), BroadphaseType, TEXT(""));
 
 	int32 BoundingVolumeNumCells = 35;
@@ -93,30 +93,49 @@ namespace Chaos
 	FAutoConsoleVariableRef CVarMaxPayloadSize(TEXT("p.MaxPayloadSize"), MaxPayloadSize, TEXT(""));
 
 	template<class FPBDRigidsEvolution, class FPBDCollisionConstraint, class T, int d>
-	ISpatialAccelerationCollection<TAccelerationStructureHandle<T, d>, T, d>*  TPBDRigidsEvolutionBase<FPBDRigidsEvolution, FPBDCollisionConstraint, T, d>::CreateNewSpatialStructure(const TArray<FAccelerationStructureBuilder>& Bounds)
+	ISpatialAccelerationCollection<TAccelerationStructureHandle<T, d>, T, d>*  TPBDRigidsEvolutionBase<FPBDRigidsEvolution, FPBDCollisionConstraint, T, d>::CreateNewSpatialStructure()
 	{
+		const TArray<FAccelerationStructureBuilder> Empty;
+		TArray<TSpatialAccelerationParams<TAccelerationStructureHandle<T, d>, T, d>> EmptyParams;
 		if (BroadphaseType == 0)
 		{
 			using AccelType = TBoundingVolume<TAccelerationStructureHandle<T, d>, T, d>;
-			auto Structure = MakeUnique<AccelType>(Bounds, false, 0, BoundingVolumeNumCells, MaxPayloadSize);
-			auto Collection = new TSpatialAccelerationCollection<AccelType>();
+			auto Collection = new TSpatialAccelerationCollection<AccelType>(EmptyParams);
+
+			auto Structure = MakeUnique<AccelType>(Empty, false, 0, BoundingVolumeNumCells, MaxPayloadSize);
 			Collection->AddSubstructure(MoveTemp(Structure), 0);
+			
 			return Collection;
 		}
 		else if(BroadphaseType == 1)
 		{
 			using AccelType = TAABBTree<TAccelerationStructureHandle<T, d>, TAABBTreeLeafArray<TAccelerationStructureHandle<T, d>, T>, T>;
-			auto Structure = MakeUnique<AccelType>(Bounds, MaxChildrenInLeaf, MaxTreeDepth, MaxPayloadSize);
-			auto Collection = new TSpatialAccelerationCollection<AccelType>();
+			auto Collection = new TSpatialAccelerationCollection<AccelType>(EmptyParams);
+
+			auto Structure = MakeUnique<AccelType>(Empty, MaxChildrenInLeaf, MaxTreeDepth, MaxPayloadSize);
 			Collection->AddSubstructure(MoveTemp(Structure), 0);
 			return Collection;
 		}
-		else
+		else if(BroadphaseType == 2)
 		{
 			using AccelType = TAABBTree<TAccelerationStructureHandle<T, d>, TBoundingVolume<TAccelerationStructureHandle<T, d>, T, d>, T>;
-			auto Structure = MakeUnique<AccelType>(Bounds, AABBMaxChildrenInLeaf, AABBMaxTreeDepth, MaxPayloadSize);
-			auto Collection = new TSpatialAccelerationCollection<AccelType>();
+			auto Collection = new TSpatialAccelerationCollection<AccelType>(EmptyParams);
+
+			auto Structure = MakeUnique<AccelType>(Empty, AABBMaxChildrenInLeaf, AABBMaxTreeDepth, MaxPayloadSize);
+			Collection->AddSubstructure(MoveTemp(Structure), 0);
+			return Collection;
+		}
+		else
+		{
+			using AccelType = TAABBTree<TAccelerationStructureHandle<T, d>, TAABBTreeLeafArray<TAccelerationStructureHandle<T, d>, T>, T>;
+			using GridType = TBoundingVolume<TAccelerationStructureHandle<T, d>, T, d>;
+			auto Collection = new TSpatialAccelerationCollection<AccelType, GridType>(EmptyParams);
+
+			auto Structure = MakeUnique<AccelType>(Empty, MaxChildrenInLeaf, MaxTreeDepth, MaxPayloadSize);
 			Collection->AddSubstructure(MoveTemp(Structure), 0);
+
+			auto Structure2 = MakeUnique<GridType>(Empty, false, 0, BoundingVolumeNumCells, MaxPayloadSize);
+			Collection->AddSubstructure(MoveTemp(Structure2), 1);
 			return Collection;
 		}
 	}
@@ -126,23 +145,24 @@ namespace Chaos
 	{
 		using BVType = TBoundingVolume<TAccelerationStructureHandle<T, d>, T, d>;
 		using AABBType = TAABBTree<TAccelerationStructureHandle<T, d>, TAABBTreeLeafArray<TAccelerationStructureHandle<T, d>, T>, T>;
+		TArray<TSpatialAccelerationParams<TAccelerationStructureHandle<T, d>, T, d>> EmptyParams;
 
 		if (Substructure->template As<BVType>())
 		{
-			auto Collection = MakeUnique<TSpatialAccelerationCollection<BVType>>();
+			auto Collection = MakeUnique<TSpatialAccelerationCollection<BVType>>(EmptyParams);
 			Collection->AddSubstructure(MoveTemp(Substructure), 0);
 			return Collection;
 		}
 		else if (Substructure->template As<AABBType>())
 		{
-			auto Collection = MakeUnique<TSpatialAccelerationCollection<AABBType>>();
+			auto Collection = MakeUnique<TSpatialAccelerationCollection<AABBType>>(EmptyParams);
 			Collection->AddSubstructure(MoveTemp(Substructure), 0);
 			return Collection;
 		}
 		else
 		{
 			using AccelType = TAABBTree<TAccelerationStructureHandle<T, d>, TBoundingVolume<TAccelerationStructureHandle<T, d>, T, d>, T>;
-			auto Collection = MakeUnique<TSpatialAccelerationCollection<AccelType>>();
+			auto Collection = MakeUnique<TSpatialAccelerationCollection<AccelType>>(EmptyParams);
 			Collection->AddSubstructure(MoveTemp(Substructure), 0);
 			return Collection;
 		}
@@ -295,10 +315,10 @@ namespace Chaos
 		{
 			//initial frame so make empty structures
 
-			InternalAcceleration = TUniquePtr<FAccelerationStructure>(CreateNewSpatialStructure(TArray<FAccelerationStructureBuilder>()));
-			ScratchExternalAcceleration = TUniquePtr<FAccelerationStructure>(CreateNewSpatialStructure(TArray<FAccelerationStructureBuilder>()));
-			AsyncInternalAcceleration = TUniquePtr<FAccelerationStructure>(CreateNewSpatialStructure(TArray<FAccelerationStructureBuilder>()));
-			AsyncExternalAcceleration = TUniquePtr<FAccelerationStructure>(CreateNewSpatialStructure(TArray<FAccelerationStructureBuilder>()));
+			InternalAcceleration = TUniquePtr<FAccelerationStructure>(CreateNewSpatialStructure());
+			ScratchExternalAcceleration = TUniquePtr<FAccelerationStructure>(CreateNewSpatialStructure());
+			AsyncInternalAcceleration = TUniquePtr<FAccelerationStructure>(CreateNewSpatialStructure());
+			AsyncExternalAcceleration = TUniquePtr<FAccelerationStructure>(CreateNewSpatialStructure());
 			FlushInternalAccelerationQueue();
 			FlushExternalAccelerationQueue(*ScratchExternalAcceleration);
 			bExternalReady = true;
@@ -333,6 +353,7 @@ namespace Chaos
 	template<class FPBDRigidsEvolution, class FPBDCollisionConstraint, class T, int d>
 	void TPBDRigidsEvolutionBase<FPBDRigidsEvolution, FPBDCollisionConstraint, T, d>::InitializeAccelerationCache()
 	{
+		WaitOnAccelerationStructure();
 		CachedSpatialBuilderDataMap.Empty();
 		const auto SpatialIndices = InternalAcceleration->GetAllSpatialIndices();
 		for (const auto& Idx : SpatialIndices)
diff --git a/Engine/Source/Runtime/Experimental/Chaos/Public/Chaos/ISpatialAccelerationCollection.h b/Engine/Source/Runtime/Experimental/Chaos/Public/Chaos/ISpatialAccelerationCollection.h
--- a/Engine/Source/Runtime/Experimental/Chaos/Public/Chaos/ISpatialAccelerationCollection.h
+++ b/Engine/Source/Runtime/Experimental/Chaos/Public/Chaos/ISpatialAccelerationCollection.h
@@ -9,6 +9,23 @@
 namespace Chaos
 {
 
+/** Used for the creation of different acceleration structures */
+template <typename TPayloadType, typename T, int d>
+struct CHAOS_API TSpatialAccelerationParams
+{
+	ESpatialAcceleration Type;	//Specifies the acceleration type. Unknown can be used for custom types, see TSpatialAccelerationParams
+
+	//Union of all params, some are only applicable to subset of structure types
+	float MaxPayloadBounds;	//max bounds size before inserting into a linear list of "large" objects. This list will always be checked
+
+	int32 MaxCells;	//(grid only) max number of cells in grid structure
+	int32 MaxChildrenInLeaf; //(tree only) if the tree sees MaxChildrenInLeaf or fewer objects, no more splitting occurs
+	int32 MaxTreeDepth;	//(tree only) if the tree has reached MaxTreeDepth, no more splitting occurs
+
+	typedef ISpatialAcceleration<TPayloadType, T, d> (*ConstructFn)(TSpatialAccelerationParams<TPayloadType, T, d>& Params, int32 ParamIdx);
+	ConstructFn ConstructionFactory;	//if unknown type is specified, use function pointer set by user
+};
+
 template <typename TPayloadType, typename T, int d>
 class CHAOS_API ISpatialAccelerationCollection : public ISpatialAcceleration<TPayloadType, T, d>
 {
@@ -25,287 +42,251 @@ public:
 	virtual TArray<FSpatialAccelerationIdx> GetAllSpatialIndices() const = 0;
 };
 
-template <typename TSpatialAcceleration>
-class CHAOS_API TSpatialAccelerationCollectionBucket
+template <typename T>
+void FreeObjHelper(T*& RawPtr)
 {
-public:
-	using AccelerationType = TSpatialAcceleration;
-	using TPayloadType = typename AccelerationType::PayloadType;
-	using T = typename AccelerationType::TType;
-	static constexpr int d = AccelerationType::D;
-
-	TSpatialAccelerationCollectionBucket() = default;
-
-	TSpatialAccelerationCollectionBucket<TSpatialAcceleration> Copy() const
-	{
-		return TSpatialAccelerationCollectionBucket<TSpatialAcceleration>(*this);
-	}
+	RawPtr = nullptr;
+}
 
-	uint16 AddSubstructure(TUniquePtr<TSpatialAcceleration>&& Substructure)
-	{
-		uint16 NewIdx = NewSlot();
-		Accelerations[NewIdx] = MoveTemp(Substructure);
-		
-		return NewIdx;
-	}
-	
-	TUniquePtr <ISpatialAcceleration<TPayloadType, T, d>> RemoveSubstructure(uint16 Idx)
-	{
-		return ReleaseSlot(Idx);
-	}
+template <typename TPayloadType, typename T, int d>
+struct CHAOS_API TSpatialAccelerationBucketEntry
+{
+	TUniquePtr<ISpatialAcceleration<TPayloadType, T, d>> Acceleration;
+	uint16 TypeInnerIdx;
 
-	TSpatialAcceleration* GetSubstructure(int32 Idx)
+	void CopyFrom(TSpatialAccelerationBucketEntry<TPayloadType, T, d>& Src)
 	{
-		return Accelerations[Idx].Get();
+		Acceleration = Src.Acceleration->Copy();
+		TypeInnerIdx = Src.TypeInnerIdx;
 	}
 
-	template <typename SQVisitor>
-	void Raycast(const TVector<T, d>& Start, const TVector<T, d>& Dir, const T OriginalLength, SQVisitor& Visitor) const
-	{
-		for (const auto& Acceleration : Accelerations)
-		{
-			if (Acceleration)
-			{
-				Acceleration->Raycast(Start, Dir, OriginalLength, Visitor);
-			}
-		}
-	}
+	TSpatialAccelerationBucketEntry() = default;
+	TSpatialAccelerationBucketEntry(const TSpatialAccelerationBucketEntry<TPayloadType, T, d>& Other) = delete;
+	TSpatialAccelerationBucketEntry(TSpatialAccelerationBucketEntry<TPayloadType, T, d>&& Other) = default;
+	TSpatialAccelerationBucketEntry<TPayloadType, T, d>& operator=(TSpatialAccelerationBucketEntry<TPayloadType, T, d>&& Other) = default;
+	TSpatialAccelerationBucketEntry<TPayloadType, T, d>& operator=(TSpatialAccelerationBucketEntry<TPayloadType, T, d>& Other) = delete;
+};
 
-	template <typename SQVisitor>
-	void Sweep(const TVector<T, d>& Start, const TVector<T, d>& Dir, const T OriginalLength, const TVector<T,3>& QueryHalfExtents, SQVisitor& Visitor) const
-	{
-		for (const auto& Acceleration : Accelerations)
-		{
-			if (Acceleration)
-			{
-				Acceleration->Sweep(Start, Dir, OriginalLength, QueryHalfExtents, Visitor);
-			}
-		}
-	}
 
-	template <typename SQVisitor>
-	void Overlap(const TBox<T,d>& QueryBounds, SQVisitor& Visitor) const
-	{
-		for (const auto& Acceleration : Accelerations)
-		{
-			if (Acceleration)
-			{
-				Acceleration->Overlap(QueryBounds, Visitor);
-			}
-		}
-	}
+template <typename TPayloadType, typename T, int d>
+void FreeObjHelper(TSpatialAccelerationBucketEntry<TPayloadType, T, d >& BucketEntry)
+{
+	BucketEntry.Acceleration = nullptr;
+}
 
-	void GlobalObjects(TArray<TPayloadBoundsElement<TPayloadType, T>>& ObjList) const
-	{
-		for (const auto& Acceleration : Accelerations)
-		{
-			if (Acceleration)
-			{
-				ObjList.Append(Acceleration->GlobalObjects());
-			}
-		}
-	}
+template <typename T>
+T CopyFromHelper(const T& Src)
+{
+	return Src;
+}
 
-	void GetAllSpatialIndices(TArray<FSpatialAccelerationIdx>& Indices, uint16 BucketIdx) const
-	{
-		Indices.Reserve(Indices.Num() + Accelerations.Num());
-		uint16 InnerIdx = 0;
-		for (const auto& Acceleration : Accelerations)
-		{
-			Indices.Add(FSpatialAccelerationIdx{ BucketIdx, InnerIdx++ });
-		}
-	}
+template <typename TPayloadType, typename T, int d>
+TSpatialAccelerationBucketEntry<TPayloadType, T, d> CopyFromHelper(const TSpatialAccelerationBucketEntry<TPayloadType, T, d >& Src)
+{
+	TSpatialAccelerationBucketEntry<TPayloadType, T, d> Copy;
+	Copy.Acceleration = Src.Acceleration ? Src.Acceleration->Copy() : nullptr;
+	Copy.TypeInnerIdx = Src.TypeInnerIdx;
+	return Copy;
+}
 
-	void RemoveElement(const TPayloadType& Payload, int32 Idx)
-	{
-		Accelerations[Idx]->RemoveElement(Payload);
-	}
 
-	void UpdateElement(const TPayloadType& Payload, const TBox<T, d>& NewBounds, bool bHasBounds, int32 Idx)
-	{
-		Accelerations[Idx]->UpdateElement(Payload, NewBounds, bHasBounds);
-	}
+template <typename TObj>
+struct CHAOS_API TSpatialCollectionBucket
+{
+	TSpatialCollectionBucket() = default;
+	TSpatialCollectionBucket(const TSpatialCollectionBucket<TObj>& Other) = delete;
+	TSpatialCollectionBucket(TSpatialCollectionBucket<TObj>&& Other) = default;
+	TSpatialCollectionBucket<TObj>& operator=(const TSpatialCollectionBucket<TObj>& Other) = delete;
+	TSpatialCollectionBucket<TObj>& operator=(TSpatialCollectionBucket<TObj>&& Other) = default;
 
-	explicit TSpatialAccelerationCollectionBucket(const TSpatialAccelerationCollectionBucket<TSpatialAcceleration>& Other)
-		: FreeIndices(Other.FreeIndices)
-	{
-		Accelerations.Reserve(Other.Accelerations.Num());
-		for (const auto& OtherAccel : Other.Accelerations)
-		{
-			if (OtherAccel)
-			{
-				Accelerations.Add(AsUniqueSpatialAccelerationChecked<TSpatialAcceleration>(OtherAccel->Copy()));
-			}
-			else
-			{
-				Accelerations.Add(nullptr);
-			}
-		}
-	}
+	TArray<TObj> Objects;
+	TArray<uint16> FreeIndices;
 
-private:
-	uint16 NewSlot()
+	int32 Add(TObj&& Obj)
 	{
+		uint16 Idx;
 		if (FreeIndices.Num())
 		{
-			return FreeIndices.Pop(false);
+			Idx = FreeIndices.Pop();
+			Objects[Idx] = MoveTemp(Obj);
 		}
-		else if (Accelerations.Num() < FSpatialAccelerationIdx::MaxBucketEntries)
+		else
 		{
-			return Accelerations.Add(nullptr);
+			Idx = Objects.Add(MoveTemp(Obj));
 		}
-		check(false);	//too many substructures in a single bucket
-		return 0;
+
+		return Idx;
 	}
 
-	TUniquePtr <ISpatialAcceleration<TPayloadType, T, d>> ReleaseSlot(uint16 Index)
+	void Remove(uint16 Idx)
 	{
-		TUniquePtr <ISpatialAcceleration<TPayloadType, T, d>> Result = MoveTemp(Accelerations[Index]);
-		Accelerations[Index] = nullptr;
-		if (Index + 1 == Accelerations.Num())
+		if (Objects.Num() == Idx + 1)
 		{
-			Accelerations.Pop(false);
+			Objects.Pop(false);
 		}
 		else
 		{
-			FreeIndices.Add(Index);
+			FreeObjHelper(Objects[Idx]);
+			FreeIndices.Add(Idx);
 		}
-		return Result;
 	}
 
-	TSpatialAccelerationCollectionBucket& operator=(const TSpatialAccelerationCollectionBucket<TSpatialAcceleration>& Other)
+	void CopyFrom(const TSpatialCollectionBucket& Src)
 	{
-		if (&Other != this)
+		Objects.SetNum(Src.Objects.Num());
+		for (int32 ObjIdx = 0; ObjIdx < Objects.Num(); ++ObjIdx)
 		{
-			FreeIndices = Other.FreeIndices;
-			Accelerations.Reset(Other.Accelerations.Num());
-
-			for (const auto& OtherAccel : Other.Accelerations)
-			{
-				if (OtherAccel)
-				{
-					Accelerations.Add(AsUniqueSpatialAccelerationChecked<TSpatialAcceleration>(OtherAccel->Copy()));
-				}
-				else
-				{
-					Accelerations.Add(nullptr);
-				}
-			}
+			Objects[ObjIdx] = CopyFromHelper(Src.Objects[ObjIdx]);
 		}
-
-		return *this;
+		FreeIndices = Src.FreeIndices;
 	}
 
-	TArray<TUniquePtr<TSpatialAcceleration>> Accelerations;
-	TArray<uint16> FreeIndices;
 };
 
 template <typename... TRemaining>
-struct CHAOS_API TSpatialBucketTuple
+struct CHAOS_API TSpatialTypeTuple
 {
 };
 
 template <typename TAcceleration, typename... TRemaining>
-struct CHAOS_API TSpatialBucketTuple<TAcceleration, TRemaining...>
+struct CHAOS_API TSpatialTypeTuple<TAcceleration, TRemaining...>
 {
 	using FirstType = TAcceleration;
-	TSpatialAccelerationCollectionBucket<TAcceleration> First;
-	TSpatialBucketTuple<TRemaining...> Remaining;
+	TSpatialCollectionBucket<TAcceleration*> First;
+	TSpatialTypeTuple<TRemaining...> Remaining;
 
-	TSpatialBucketTuple() = default;
+	TSpatialTypeTuple() = default;
 
-	explicit TSpatialBucketTuple(const TSpatialBucketTuple<TAcceleration, TRemaining...>& Other)
-		: First(Other.First)
-		, Remaining(Other.Remaining)
-	{
-	}
+	TSpatialTypeTuple(const TSpatialTypeTuple<TAcceleration, TRemaining...>& Other) = delete;	//raw pointers need updating because buckets use a deep copy
 };
 
 template<int Idx, typename ... Rest>
-struct CHAOS_API TSpatialBucketTupleGetter
+struct CHAOS_API TSpatialTypeTupleGetter
 {
 
 };
 
 template<int Idx, typename First, typename... Rest>
-struct CHAOS_API TSpatialBucketTupleGetter<Idx, First, Rest...>
+struct CHAOS_API TSpatialTypeTupleGetter<Idx, First, Rest...>
 {
-	static auto& Get(TSpatialBucketTuple<First, Rest...>& Buckets) { return TSpatialBucketTupleGetter<Idx - 1, Rest...>::Get(Buckets.Remaining); }
-	static const auto& Get(const TSpatialBucketTuple<First, Rest...>& Buckets) { return TSpatialBucketTupleGetter<Idx - 1, Rest...>::Get(Buckets.Remaining); }
+	static auto& Get(TSpatialTypeTuple<First, Rest...>& Types) { return TSpatialTypeTupleGetter<Idx - 1, Rest...>::Get(Types.Remaining); }
+	static const auto& Get(const TSpatialTypeTuple<First, Rest...>& Types) { return TSpatialTypeTupleGetter<Idx - 1, Rest...>::Get(Types.Remaining); }
 };
 
 template<typename First, typename... Rest>
-struct CHAOS_API TSpatialBucketTupleGetter<0, First, Rest...>
+struct CHAOS_API TSpatialTypeTupleGetter<0, First, Rest...>
 {
-	static auto& Get(TSpatialBucketTuple<First, Rest...>& Buckets) { return Buckets.First; }
-	static const auto& Get(const TSpatialBucketTuple<First, Rest...>& Buckets) { return Buckets.First; }
+	static auto& Get(TSpatialTypeTuple<First, Rest...>& Types) { return Types.First; }
+	static const auto& Get(const TSpatialTypeTuple<First, Rest...>& Types) { return Types.First; }
 };
 
 template <int Idx, typename First, typename... Rest>
-auto& GetBucket(TSpatialBucketTuple<First, Rest...>& Buckets)
+auto& GetAccelerationsPerType(TSpatialTypeTuple<First, Rest...>& Types)
 {
-	return TSpatialBucketTupleGetter<Idx, First, Rest...>::Get(Buckets);
+	return TSpatialTypeTupleGetter<Idx, First, Rest...>::Get(Types);
 }
 
 template <int Idx, typename First, typename... Rest>
-const auto& GetBucket(const TSpatialBucketTuple<First, Rest...>& Buckets)
+const auto& GetAccelerationsPerType(const TSpatialTypeTuple<First, Rest...>& Types)
 {
-	return TSpatialBucketTupleGetter<Idx, First, Rest...>::Get(Buckets);
+	return TSpatialTypeTupleGetter<Idx, First, Rest...>::Get(Types);
 }
 
-template <int BucketIdx, int NumBuckets, typename Tuple, typename TPayloadType, typename T, int d>
+template <int TypeIdx, int NumTypes, typename Tuple, typename TPayloadType, typename T, int d>
 struct TSpatialAccelerationCollectionHelper
 {
 	template <typename SQVisitor>
-	static void Raycast(const Tuple& Buckets, const TVector<T, d>& Start, const TVector<T, d>& Dir, const T OriginalLength, SQVisitor& Visitor)
+	static bool RaycastFast(const Tuple& Types, const TVector<T, d>& Start, const TVector<T, d>& Dir, const TVector<T,d>& InvDir, const bool* bParallel, T& CurLength, T& InvCurLength, SQVisitor& Visitor)
 	{
-		GetBucket<BucketIdx>(Buckets).Raycast(Start, Dir, OriginalLength, Visitor);
-		constexpr int NextBucket = BucketIdx + 1;
-		if (NextBucket < NumBuckets)
+		const auto& Accelerations = GetAccelerationsPerType<TypeIdx>(Types).Objects;
+		for (const auto& Accelerator : Accelerations)
+		{
+			if (Accelerator && !Accelerator->RaycastFast(Start, Dir, InvDir, bParallel, CurLength, InvCurLength, Visitor))
+			{
+				return false;
+			}
+		}
+		
+		constexpr int NextType = TypeIdx + 1;
+		if (NextType < NumTypes)
 		{
-			TSpatialAccelerationCollectionHelper<NextBucket < NumBuckets ? NextBucket : 0, NumBuckets, Tuple, TPayloadType, T, d>::Raycast(Buckets, Start, Dir, OriginalLength, Visitor);
+			return TSpatialAccelerationCollectionHelper<NextType < NumTypes ? NextType : 0, NumTypes, Tuple, TPayloadType, T, d>::RaycastFast(Types, Start, Dir, InvDir, bParallel, CurLength, InvCurLength, Visitor);
 		}
+
+		return true;
 	}
 
 	template <typename SQVisitor>
-	static void Sweep(const Tuple& Buckets, const TVector<T, d>& Start, const TVector<T, d>& Dir, const T OriginalLength, const TVector<T, d> QueryHalfExtents, SQVisitor& Visitor)
+	static bool SweepFast(const Tuple& Types, const TVector<T, d>& Start, const TVector<T, d>& Dir, const TVector<T,d>& InvDir, const bool* bParallel, T& CurLength, T& InvCurLength, const TVector<T, d> QueryHalfExtents, SQVisitor& Visitor)
 	{
-		GetBucket<BucketIdx>(Buckets).Sweep(Start, Dir, OriginalLength, QueryHalfExtents, Visitor);
-		constexpr int NextBucket = BucketIdx + 1;
-		if (NextBucket < NumBuckets)
+		const auto& Accelerations = GetAccelerationsPerType<TypeIdx>(Types).Objects;
+		for (const auto& Accelerator : Accelerations)
+		{
+			if (Accelerator && !Accelerator->SweepFast(Start, Dir, InvDir, bParallel, CurLength, InvCurLength, QueryHalfExtents, Visitor))
+			{
+				return false;
+			}
+		}
+
+		constexpr int NextType = TypeIdx + 1;
+		if (NextType < NumTypes)
 		{
-			TSpatialAccelerationCollectionHelper < NextBucket < NumBuckets ? NextBucket : 0, NumBuckets, Tuple, TPayloadType, T, d>::Sweep(Buckets, Start, Dir, OriginalLength, QueryHalfExtents, Visitor);
+			return TSpatialAccelerationCollectionHelper < NextType < NumTypes ? NextType : 0, NumTypes, Tuple, TPayloadType, T, d>::SweepFast(Types, Start, Dir, InvDir, bParallel, CurLength, InvCurLength, QueryHalfExtents, Visitor);
 		}
+
+		return true;
 	}
 
 	template <typename SQVisitor>
-	static void Overlap(const Tuple& Buckets, const TBox<T, d> QueryBounds, SQVisitor& Visitor)
+	static bool OverlapFast(const Tuple& Types, const TBox<T, d> QueryBounds, SQVisitor& Visitor)
 	{
-		GetBucket<BucketIdx>(Buckets).Overlap(QueryBounds, Visitor);
-		constexpr int NextBucket = BucketIdx + 1;
-		if (NextBucket < NumBuckets)
+		const auto& Accelerations = GetAccelerationsPerType<TypeIdx>(Types).Objects;
+		for (const auto& Accelerator : Accelerations)
+		{
+			if (Accelerator && !Accelerator->OverlapFast(QueryBounds, Visitor))
+			{
+				return false;
+			}
+		}
+
+		constexpr int NextType = TypeIdx + 1;
+		if (NextType < NumTypes)
 		{
-			TSpatialAccelerationCollectionHelper < NextBucket < NumBuckets ? NextBucket : 0, NumBuckets, Tuple, TPayloadType, T, d>::Overlap(Buckets, QueryBounds, Visitor);
+			return TSpatialAccelerationCollectionHelper < NextType < NumTypes ? NextType : 0, NumTypes, Tuple, TPayloadType, T, d>::OverlapFast(Types, QueryBounds, Visitor);
 		}
+
+		return true;
 	}
 
-	static void GlobalObjects(const Tuple& Buckets, TArray<TPayloadBoundsElement<TPayloadType, T>>& ObjList)
+	static void GlobalObjects(const Tuple& Types, TArray<TPayloadBoundsElement<TPayloadType, T>>& ObjList)
 	{
-		GetBucket<BucketIdx>(Buckets).GlobalObjects(ObjList);
-		constexpr int NextBucket = BucketIdx + 1;
-		if (NextBucket < NumBuckets)
+		const auto& Accelerations = GetAccelerationsPerType<TypeIdx>(Types).Objects;
+		for (const auto& Accelerator : Accelerations)
 		{
-			TSpatialAccelerationCollectionHelper < NextBucket < NumBuckets ? NextBucket : 0, NumBuckets, Tuple, TPayloadType, T, d>::GlobalObjects(Buckets, ObjList);
+			if (Accelerator)
+			{
+				//todo: handle early out
+				ObjList.Append(Accelerator->GlobalObjects());
+			}
+		}
+
+		constexpr int NextType = TypeIdx + 1;
+		if (NextType < NumTypes)
+		{
+			TSpatialAccelerationCollectionHelper < NextType < NumTypes ? NextType : 0, NumTypes, Tuple, TPayloadType, T, d>::GlobalObjects(Types, ObjList);
 		}
 	}
 
-	static void GetAllSpatialIndices(const Tuple& Buckets, TArray<FSpatialAccelerationIdx>& Indices)
+	static void SetNumFrom(const Tuple& TypesSrc, Tuple& TypesDest)
 	{
-		GetBucket<BucketIdx>(Buckets).GetAllSpatialIndices(Indices, BucketIdx);
-		constexpr int NextBucket = BucketIdx + 1;
-		if (NextBucket < NumBuckets)
+		const auto& SrcAccelerations = GetAccelerationsPerType<TypeIdx>(TypesSrc).Objects;
+		auto& DestAccelerations = GetAccelerationsPerType<TypeIdx>(TypesDest).Objects;
+		DestAccelerations.SetNum(SrcAccelerations.Num());
+
+		constexpr int NextType = TypeIdx + 1;
+		if (NextType < NumTypes)
 		{
-			TSpatialAccelerationCollectionHelper < NextBucket < NumBuckets ? NextBucket : 0, NumBuckets, Tuple, TPayloadType, T, d>::GetAllSpatialIndices(Buckets, Indices);
+			TSpatialAccelerationCollectionHelper < NextType < NumTypes ? NextType : 0, NumTypes, Tuple, TPayloadType, T, d>::SetNumFrom(TypesSrc, TypesDest);
 		}
 	}
 };
@@ -321,18 +302,21 @@ typename TEnableIf<!TIsSame<typename SpatialAccelerationCollection::TPayloadType
 {
 }
 
-template <typename ... TSpatialAccelerationBuckets>
+template <typename ... TSpatialAccelerationTypes>
 class CHAOS_API TSpatialAccelerationCollection : public
-	ISpatialAccelerationCollection<typename std::tuple_element<0, std::tuple< TSpatialAccelerationBuckets...>>::type::PayloadType,
-	typename std::tuple_element<0, std::tuple< TSpatialAccelerationBuckets...>>::type::TType,
-	std::tuple_element<0, std::tuple< TSpatialAccelerationBuckets...>>::type::D>
+	ISpatialAccelerationCollection<typename std::tuple_element<0, std::tuple< TSpatialAccelerationTypes...>>::type::PayloadType,
+	typename std::tuple_element<0, std::tuple< TSpatialAccelerationTypes...>>::type::TType,
+	std::tuple_element<0, std::tuple< TSpatialAccelerationTypes...>>::type::D>
 {
 public:
-	using FirstAccelerationType = typename std::tuple_element<0, std::tuple< TSpatialAccelerationBuckets...>>::type;
+	using FirstAccelerationType = typename std::tuple_element<0, std::tuple< TSpatialAccelerationTypes...>>::type;
 	using TPayloadType = typename FirstAccelerationType::PayloadType;
 	using T = typename FirstAccelerationType::TType;
 	static constexpr int d = FirstAccelerationType::D;
-	TSpatialAccelerationCollection() = default;
+	TSpatialAccelerationCollection(const TArray<TSpatialAccelerationParams<TPayloadType, T, d>>& BucketsParams)
+	{
+		//todo: verify params match template args
+	}
 
 	virtual FSpatialAccelerationIdx AddSubstructure(TUniquePtr<ISpatialAcceleration<TPayloadType, T, d>>&& Substructure, int32 BucketIdx) override
 	{
@@ -340,82 +324,105 @@ public:
 		FSpatialAccelerationIdx Result;
 		Result.Bucket = BucketIdx;
 		
-		switch (BucketIdx)
+		ISpatialAcceleration<TPayloadType, T, d>* AccelPtr = Substructure.Get();
+		TSpatialAccelerationBucketEntry<TPayloadType, T, d> BucketEntry;
+		BucketEntry.Acceleration = MoveTemp(Substructure);
+		
+		const int32 TypeIdx = GetTypeIdx(AccelPtr);
+		switch (TypeIdx)
 		{
-		case 0: Result.InnerIdx = GetBucket<0>(Buckets).AddSubstructure(AsUniqueSpatialAccelerationChecked<typename std::remove_reference<decltype(GetBucket<0>(Buckets))>::type::AccelerationType>(MoveTemp(Substructure))); break;
-		case 1: Result.InnerIdx = GetBucket<ClampedIdx(1)>(Buckets).AddSubstructure(AsUniqueSpatialAccelerationChecked<typename std::remove_reference<decltype(GetBucket<ClampedIdx(1)>(Buckets))>::type::AccelerationType>(MoveTemp(Substructure))); break;
-		case 2: Result.InnerIdx = GetBucket<ClampedIdx(2)>(Buckets).AddSubstructure(AsUniqueSpatialAccelerationChecked<typename std::remove_reference<decltype(GetBucket<ClampedIdx(2)>(Buckets))>::type::AccelerationType>(MoveTemp(Substructure))); break;
-		case 3: Result.InnerIdx = GetBucket<ClampedIdx(3)>(Buckets).AddSubstructure(AsUniqueSpatialAccelerationChecked<typename std::remove_reference<decltype(GetBucket<ClampedIdx(3)>(Buckets))>::type::AccelerationType>(MoveTemp(Substructure))); break;
-		case 4: Result.InnerIdx = GetBucket<ClampedIdx(4)>(Buckets).AddSubstructure(AsUniqueSpatialAccelerationChecked<typename std::remove_reference<decltype(GetBucket<ClampedIdx(4)>(Buckets))>::type::AccelerationType>(MoveTemp(Substructure))); break;
-		case 5: Result.InnerIdx = GetBucket<ClampedIdx(5)>(Buckets).AddSubstructure(AsUniqueSpatialAccelerationChecked<typename std::remove_reference<decltype(GetBucket<ClampedIdx(5)>(Buckets))>::type::AccelerationType>(MoveTemp(Substructure))); break;
-		case 6: Result.InnerIdx = GetBucket<ClampedIdx(6)>(Buckets).AddSubstructure(AsUniqueSpatialAccelerationChecked<typename std::remove_reference<decltype(GetBucket<ClampedIdx(6)>(Buckets))>::type::AccelerationType>(MoveTemp(Substructure))); break;
-		case 7: Result.InnerIdx = GetBucket<ClampedIdx(7)>(Buckets).AddSubstructure(AsUniqueSpatialAccelerationChecked<typename std::remove_reference<decltype(GetBucket<ClampedIdx(7)>(Buckets))>::type::AccelerationType>(MoveTemp(Substructure))); break;
-		default: check(false);
+		case 0: BucketEntry.TypeInnerIdx = GetAccelerationsPerType<0>(Types).Add(&AccelPtr->template AsChecked<typename std::tuple_element<0, std::tuple<TSpatialAccelerationTypes...>>::type>()); break;
+		case 1: BucketEntry.TypeInnerIdx = GetAccelerationsPerType<ClampedIdx(1)>(Types).Add(&AccelPtr->template AsChecked<typename std::tuple_element<ClampedIdx(1), std::tuple<TSpatialAccelerationTypes...>>::type>()); break;
+		case 2: BucketEntry.TypeInnerIdx = GetAccelerationsPerType<ClampedIdx(2)>(Types).Add(&AccelPtr->template AsChecked<typename std::tuple_element<ClampedIdx(2), std::tuple<TSpatialAccelerationTypes...>>::type>()); break;
+		case 3: BucketEntry.TypeInnerIdx = GetAccelerationsPerType<ClampedIdx(3)>(Types).Add(&AccelPtr->template AsChecked<typename std::tuple_element<ClampedIdx(3), std::tuple<TSpatialAccelerationTypes...>>::type>()); break;
+		case 4: BucketEntry.TypeInnerIdx = GetAccelerationsPerType<ClampedIdx(4)>(Types).Add(&AccelPtr->template AsChecked<typename std::tuple_element<ClampedIdx(4), std::tuple<TSpatialAccelerationTypes...>>::type>()); break;
+		case 5: BucketEntry.TypeInnerIdx = GetAccelerationsPerType<ClampedIdx(5)>(Types).Add(&AccelPtr->template AsChecked<typename std::tuple_element<ClampedIdx(5), std::tuple<TSpatialAccelerationTypes...>>::type>()); break;
+		case 6: BucketEntry.TypeInnerIdx = GetAccelerationsPerType<ClampedIdx(6)>(Types).Add(&AccelPtr->template AsChecked<typename std::tuple_element<ClampedIdx(6), std::tuple<TSpatialAccelerationTypes...>>::type>()); break;
+		case 7: BucketEntry.TypeInnerIdx = GetAccelerationsPerType<ClampedIdx(7)>(Types).Add(&AccelPtr->template AsChecked<typename std::tuple_element<ClampedIdx(7), std::tuple<TSpatialAccelerationTypes...>>::type>()); break;
 		}
+
+		Result.InnerIdx = Buckets[BucketIdx].Add(MoveTemp(BucketEntry));
 		return Result;
 	}
 
 	virtual TUniquePtr <ISpatialAcceleration<TPayloadType, T, d>> RemoveSubstructure(FSpatialAccelerationIdx Idx) override
 	{
+		TUniquePtr <ISpatialAcceleration<TPayloadType, T, d>> Removed;
 		check(Idx.Bucket < NumBuckets);
-		switch (Idx.Bucket)
 		{
-		case 0: return GetBucket<0>(Buckets).RemoveSubstructure(Idx.InnerIdx);
-		case 1: return GetBucket<ClampedIdx(1)>(Buckets).RemoveSubstructure(Idx.InnerIdx);
-		case 2: return GetBucket<ClampedIdx(2)>(Buckets).RemoveSubstructure(Idx.InnerIdx);
-		case 3: return GetBucket<ClampedIdx(3)>(Buckets).RemoveSubstructure(Idx.InnerIdx);
-		case 4: return GetBucket<ClampedIdx(4)>(Buckets).RemoveSubstructure(Idx.InnerIdx);
-		case 5: return GetBucket<ClampedIdx(5)>(Buckets).RemoveSubstructure(Idx.InnerIdx);
-		case 6: return GetBucket<ClampedIdx(6)>(Buckets).RemoveSubstructure(Idx.InnerIdx);
-		case 7: return GetBucket<ClampedIdx(7)>(Buckets).RemoveSubstructure(Idx.InnerIdx);
-		default: check(false);
+			auto& BucketEntry = Buckets[Idx.Bucket].Objects[Idx.InnerIdx];
+			const int32 TypeIdx = GetTypeIdx(BucketEntry.Acceleration.Get());
+			switch (TypeIdx)
+			{
+			case 0: GetAccelerationsPerType<0>(Types).Remove(BucketEntry.TypeInnerIdx); break;
+			case 1: GetAccelerationsPerType<ClampedIdx(1)>(Types).Remove(BucketEntry.TypeInnerIdx); break;
+			case 2: GetAccelerationsPerType<ClampedIdx(2)>(Types).Remove(BucketEntry.TypeInnerIdx); break;
+			case 3: GetAccelerationsPerType<ClampedIdx(3)>(Types).Remove(BucketEntry.TypeInnerIdx); break;
+			case 4: GetAccelerationsPerType<ClampedIdx(4)>(Types).Remove(BucketEntry.TypeInnerIdx); break;
+			case 5: GetAccelerationsPerType<ClampedIdx(5)>(Types).Remove(BucketEntry.TypeInnerIdx); break;
+			case 6: GetAccelerationsPerType<ClampedIdx(6)>(Types).Remove(BucketEntry.TypeInnerIdx); break;
+			case 7: GetAccelerationsPerType<ClampedIdx(7)>(Types).Remove(BucketEntry.TypeInnerIdx); break;
+			}
+			Removed = MoveTemp(BucketEntry.Acceleration);
 		}
-		return nullptr;
+		Buckets[Idx.Bucket].Remove(Idx.InnerIdx);
+		return Removed;
 	}
 
 	virtual ISpatialAcceleration<TPayloadType, T, d>* GetSubstructure(FSpatialAccelerationIdx Idx) override
 	{
-		check (Idx.Bucket < NumBuckets);
-		switch (Idx.Bucket)
-		{
-		case 0: return GetBucket<0>(Buckets).GetSubstructure(Idx.InnerIdx);
-		case 1: return GetBucket<ClampedIdx(1)>(Buckets).GetSubstructure(Idx.InnerIdx);
-		case 2: return GetBucket<ClampedIdx(2)>(Buckets).GetSubstructure(Idx.InnerIdx);
-		case 3: return GetBucket<ClampedIdx(3)>(Buckets).GetSubstructure(Idx.InnerIdx);
-		case 4: return GetBucket<ClampedIdx(4)>(Buckets).GetSubstructure(Idx.InnerIdx);
-		case 5: return GetBucket<ClampedIdx(5)>(Buckets).GetSubstructure(Idx.InnerIdx);
-		case 6: return GetBucket<ClampedIdx(6)>(Buckets).GetSubstructure(Idx.InnerIdx);
-		case 7: return GetBucket<ClampedIdx(7)>(Buckets).GetSubstructure(Idx.InnerIdx);
-		default: check(false);
-		}
-		return nullptr;
+		check(Idx.Bucket < NumBuckets);
+		return Buckets[Idx.Bucket].Objects[Idx.InnerIdx].Acceleration.Get();
 	}
 
-	virtual void Raycast(const TVector<T, d>& Start, const TVector<T, d>& Dir, const T OriginalLength, ISpatialVisitor<TPayloadType, T>& Visitor) const
+	virtual void Raycast(const TVector<T, d>& Start, const TVector<T, d>& Dir, const T Length, ISpatialVisitor<TPayloadType, T>& Visitor) const override
 	{
 		TSpatialVisitor<TPayloadType, T> ProxyVisitor(Visitor);
-		TSpatialAccelerationCollectionHelper<0, NumBuckets, decltype(Buckets), TPayloadType, T, d>::Raycast(Buckets, Start, Dir, OriginalLength, ProxyVisitor);
+		Raycast(Start, Dir, Length, ProxyVisitor);
 	}
 
 	template <typename SQVisitor>
-	void Raycast(const TVector<T, d>& Start, const TVector<T, d>& Dir, const T OriginalLength, SQVisitor& Visitor) const
+	void Raycast(const TVector<T, d>& Start, const TVector<T, d>& Dir, const T Length, SQVisitor& Visitor) const
 	{
-		TSpatialAccelerationCollectionHelper<0, NumBuckets, decltype(Buckets), TPayloadType, T, d>::Raycast(Buckets, Start, Dir, OriginalLength, Visitor);
+		ensure(Length);
+		T CurLength = Length;
+		bool bParallel[3];
+		TVector<T, 3> InvDir;
+
+		T InvLength = 1 / Length;
+		for (int Axis = 0; Axis < 3; ++Axis)
+		{
+			bParallel[Axis] = Dir[Axis] == 0;
+			InvDir[Axis] = bParallel[Axis] ? 0 : 1 / Dir[Axis];
+		}
+
+		TSpatialAccelerationCollectionHelper<0, NumTypes, decltype(Types), TPayloadType, T, d>::RaycastFast(Types, Start, Dir, InvDir, bParallel, CurLength, InvLength, Visitor);
 	}
 
-	virtual void Sweep(const TVector<T, d>& Start, const TVector<T, d>& Dir, T OriginalLength, const TVector<T, d> QueryHalfExtents, ISpatialVisitor<TPayloadType, T>& Visitor) const
+	void Sweep(const TVector<T, d>& Start, const TVector<T, d>& Dir, const T Length, const TVector<T, d> QueryHalfExtents, ISpatialVisitor<TPayloadType, T>& Visitor) const override
 	{
 		TSpatialVisitor<TPayloadType, T> ProxyVisitor(Visitor);
-		TSpatialAccelerationCollectionHelper<0, NumBuckets, decltype(Buckets), TPayloadType, T, d>::Sweep(Buckets, Start, Dir, OriginalLength, QueryHalfExtents, ProxyVisitor);
+		Sweep(Start, Dir, Length, QueryHalfExtents, ProxyVisitor);
 	}
 
 	template <typename SQVisitor>
-	void Sweep(const TVector<T, d>& Start, const TVector<T, d>& Dir, T OriginalLength, const TVector<T, d> QueryHalfExtents, SQVisitor& Visitor) const
+	void Sweep(const TVector<T, d>& Start, const TVector<T, d>& Dir, const T Length, const TVector<T, d> QueryHalfExtents, SQVisitor& Visitor) const
 	{
-		TSpatialAccelerationCollectionHelper<0, NumBuckets, decltype(Buckets), TPayloadType, T, d>::Sweep(Buckets, Start, Dir, OriginalLength, QueryHalfExtents, Visitor);
+		ensure(Length);
+		T CurLength = Length;
+		bool bParallel[3];
+		TVector<T, 3> InvDir;
+
+		T InvLength = 1 / Length;
+		for (int Axis = 0; Axis < 3; ++Axis)
+		{
+			bParallel[Axis] = Dir[Axis] == 0;
+			InvDir[Axis] = bParallel[Axis] ? 0 : 1 / Dir[Axis];
+		}
+		TSpatialAccelerationCollectionHelper<0, NumTypes, decltype(Types), TPayloadType, T, d>::SweepFast(Types, Start, Dir, InvDir, bParallel, CurLength, InvLength, QueryHalfExtents, Visitor);
 	}
 
-	virtual void Overlap(const TBox<T, d>& QueryBounds, ISpatialVisitor<TPayloadType, T>& Visitor) const
+	virtual void Overlap(const TBox<T, d>& QueryBounds, ISpatialVisitor<TPayloadType, T>& Visitor) const override
 	{
 		TSpatialVisitor<TPayloadType, T> ProxyVisitor(Visitor);
 		Overlap(QueryBounds, ProxyVisitor);
@@ -424,60 +431,51 @@ public:
 	template <typename SQVisitor>
 	void Overlap(const TBox<T, 3>& QueryBounds, SQVisitor& Visitor) const
 	{
-		TSpatialAccelerationCollectionHelper<0, NumBuckets, decltype(Buckets), TPayloadType, T, d>::Overlap(Buckets, QueryBounds, Visitor);
+		TSpatialAccelerationCollectionHelper<0, NumTypes, decltype(Types), TPayloadType, T, d>::OverlapFast(Types, QueryBounds, Visitor);
 	}
 
 	TArray<TPayloadBoundsElement<TPayloadType, T>> GlobalObjects() const
 	{
 		TArray<TPayloadBoundsElement<TPayloadType, T>> ObjList;
-		TSpatialAccelerationCollectionHelper<0, NumBuckets, decltype(Buckets), TPayloadType, T, d>::GlobalObjects(Buckets, ObjList);
+		TSpatialAccelerationCollectionHelper<0, NumTypes, decltype(Types), TPayloadType, T, d>::GlobalObjects(Types, ObjList);
 		return ObjList;
 	}
 
 	virtual TArray<FSpatialAccelerationIdx> GetAllSpatialIndices() const override
 	{
+		uint16 BucketIdx = 0;
 		TArray<FSpatialAccelerationIdx> Indices;
-		TSpatialAccelerationCollectionHelper<0, NumBuckets, decltype(Buckets), TPayloadType, T, d>::GetAllSpatialIndices(Buckets, Indices);
+		for (const auto& Bucket : Buckets)
+		{
+			uint16 InnerIdx = 0;
+			for (const auto& Entry : Bucket.Objects)
+			{
+				if (Entry.Acceleration)
+				{
+					Indices.Add(FSpatialAccelerationIdx{ BucketIdx, InnerIdx });
+				}
+				++InnerIdx;
+			}
+			++BucketIdx;
+		}
 		return Indices;
 	}
 
 	virtual void RemoveElementFrom(const TPayloadType& Payload, FSpatialAccelerationIdx SpatialIdx) override
 	{
 		check(SpatialIdx.Bucket < NumBuckets);
-		switch (SpatialIdx.Bucket)
-		{
-		case 0: return GetBucket<0>(Buckets).RemoveElement(Payload, SpatialIdx.InnerIdx);
-		case 1: return GetBucket<ClampedIdx(1)>(Buckets).RemoveElement(Payload, SpatialIdx.InnerIdx);
-		case 2: return GetBucket<ClampedIdx(2)>(Buckets).RemoveElement(Payload, SpatialIdx.InnerIdx);
-		case 3: return GetBucket<ClampedIdx(3)>(Buckets).RemoveElement(Payload, SpatialIdx.InnerIdx);
-		case 4: return GetBucket<ClampedIdx(4)>(Buckets).RemoveElement(Payload, SpatialIdx.InnerIdx);
-		case 5: return GetBucket<ClampedIdx(5)>(Buckets).RemoveElement(Payload, SpatialIdx.InnerIdx);
-		case 6: return GetBucket<ClampedIdx(6)>(Buckets).RemoveElement(Payload, SpatialIdx.InnerIdx);
-		case 7: return GetBucket<ClampedIdx(7)>(Buckets).RemoveElement(Payload, SpatialIdx.InnerIdx);
-		default: check(false);
-		}
+		Buckets[SpatialIdx.Bucket].Objects[SpatialIdx.InnerIdx].Acceleration->RemoveElement(Payload);
 	}
 
 	virtual void UpdateElementIn(const TPayloadType& Payload, const TBox<T, d>& NewBounds, bool bHasBounds, FSpatialAccelerationIdx SpatialIdx)
 	{
 		check(SpatialIdx.Bucket < NumBuckets);
-		switch (SpatialIdx.Bucket)
-		{
-		case 0: return GetBucket<0>(Buckets).UpdateElement(Payload, NewBounds, bHasBounds, SpatialIdx.InnerIdx);
-		case 1: return GetBucket<ClampedIdx(1)>(Buckets).UpdateElement(Payload, NewBounds, bHasBounds, SpatialIdx.InnerIdx);
-		case 2: return GetBucket<ClampedIdx(2)>(Buckets).UpdateElement(Payload, NewBounds, bHasBounds, SpatialIdx.InnerIdx);
-		case 3: return GetBucket<ClampedIdx(3)>(Buckets).UpdateElement(Payload, NewBounds, bHasBounds, SpatialIdx.InnerIdx);
-		case 4: return GetBucket<ClampedIdx(4)>(Buckets).UpdateElement(Payload, NewBounds, bHasBounds, SpatialIdx.InnerIdx);
-		case 5: return GetBucket<ClampedIdx(5)>(Buckets).UpdateElement(Payload, NewBounds, bHasBounds, SpatialIdx.InnerIdx);
-		case 6: return GetBucket<ClampedIdx(6)>(Buckets).UpdateElement(Payload, NewBounds, bHasBounds, SpatialIdx.InnerIdx);
-		case 7: return GetBucket<ClampedIdx(7)>(Buckets).UpdateElement(Payload, NewBounds, bHasBounds, SpatialIdx.InnerIdx);
-		default: check(false);
-		}
+		Buckets[SpatialIdx.Bucket].Objects[SpatialIdx.InnerIdx].Acceleration->UpdateElement(Payload, NewBounds, bHasBounds);
 	}
 
 	virtual TUniquePtr<ISpatialAcceleration<TPayloadType, T, d>> Copy() const
 	{
-		return TUniquePtr<ISpatialAcceleration<TPayloadType, T, d>>(new TSpatialAccelerationCollection<TSpatialAccelerationBuckets...>(*this));
+		return TUniquePtr<ISpatialAcceleration<TPayloadType, T, d>>(new TSpatialAccelerationCollection<TSpatialAccelerationTypes...>(*this));
 	}
 
 	virtual void PBDComputeConstraintsLowLevel_GatherStats(TPBDCollisionConstraint<T, d>& CollisionConstraint, T Dt) const override
@@ -492,16 +490,90 @@ public:
 
 private:
 
-	TSpatialAccelerationCollection(const TSpatialAccelerationCollection<TSpatialAccelerationBuckets...>& Other) = default;
+	TSpatialAccelerationCollection(const TSpatialAccelerationCollection<TSpatialAccelerationTypes...>& Other)
+	{
+		//need to make a deep copy of the unique ptrs and then update raw ptrs
+		TSpatialAccelerationCollectionHelper<0, NumTypes, decltype(Types), TPayloadType, T, d>::SetNumFrom(Other.Types, Types);
+		
+
+		for (int BucketIdx = 0; BucketIdx < NumBuckets; ++BucketIdx)
+		{
+			Buckets[BucketIdx].CopyFrom(Other.Buckets[BucketIdx]);
+			bool bFirst = true;
+			for (const TSpatialAccelerationBucketEntry<TPayloadType, T, d>& Entry : Buckets[BucketIdx].Objects)
+			{
+				ISpatialAcceleration<TPayloadType, T, d>* RawPtr = Entry.Acceleration.Get();
+				const int32 TypeIdx = GetTypeIdx(RawPtr);
+				switch (TypeIdx)
+				{
+				case 0: GetAccelerationsPerType<0>(Types).Objects[Entry.TypeInnerIdx] = &RawPtr->template AsChecked<typename std::tuple_element<0, std::tuple<TSpatialAccelerationTypes...>>::type>(); break;
+				case 1: GetAccelerationsPerType<ClampedIdx(1)>(Types).Objects[Entry.TypeInnerIdx] = &RawPtr->template AsChecked<typename std::tuple_element<ClampedIdx(1), std::tuple<TSpatialAccelerationTypes...>>::type>(); break;
+				case 2: GetAccelerationsPerType<ClampedIdx(2)>(Types).Objects[Entry.TypeInnerIdx] = &RawPtr->template AsChecked<typename std::tuple_element<ClampedIdx(2), std::tuple<TSpatialAccelerationTypes...>>::type>(); break;
+				case 3: GetAccelerationsPerType<ClampedIdx(3)>(Types).Objects[Entry.TypeInnerIdx] = &RawPtr->template AsChecked<typename std::tuple_element<ClampedIdx(3), std::tuple<TSpatialAccelerationTypes...>>::type>(); break;
+				case 4: GetAccelerationsPerType<ClampedIdx(4)>(Types).Objects[Entry.TypeInnerIdx] = &RawPtr->template AsChecked<typename std::tuple_element<ClampedIdx(4), std::tuple<TSpatialAccelerationTypes...>>::type>(); break;
+				case 5: GetAccelerationsPerType<ClampedIdx(5)>(Types).Objects[Entry.TypeInnerIdx] = &RawPtr->template AsChecked<typename std::tuple_element<ClampedIdx(5), std::tuple<TSpatialAccelerationTypes...>>::type>(); break;
+				case 6: GetAccelerationsPerType<ClampedIdx(6)>(Types).Objects[Entry.TypeInnerIdx] = &RawPtr->template AsChecked<typename std::tuple_element<ClampedIdx(6), std::tuple<TSpatialAccelerationTypes...>>::type>(); break;
+				case 7: GetAccelerationsPerType<ClampedIdx(7)>(Types).Objects[Entry.TypeInnerIdx] = &RawPtr->template AsChecked<typename std::tuple_element<ClampedIdx(7), std::tuple<TSpatialAccelerationTypes...>>::type>(); break;
+				}
+			}
+		}
+	}
 	
 	static constexpr uint32 ClampedIdx(uint32 Idx)
 	{
-		return Idx < NumBuckets ? Idx : 0;
+		return Idx < NumTypes ? Idx : 0;
+	}
+
+	static uint8 GetTypeIdx(const ISpatialAcceleration<TPayloadType, T, d>* Accel)
+	{
+		if (Accel->template As<typename std::tuple_element<0, std::tuple< TSpatialAccelerationTypes...>>::type> ())
+		{
+			return 0;
+		}
+		
+		if (Accel->template As<typename std::tuple_element<ClampedIdx(1), std::tuple< TSpatialAccelerationTypes...>>::type> ())
+		{
+			return 1;
+		}
+
+		if (Accel->template As<typename std::tuple_element<ClampedIdx(2), std::tuple< TSpatialAccelerationTypes...>>::type> ())
+		{
+			return 2;
+		}
+
+		if (Accel->template As<typename std::tuple_element<ClampedIdx(3), std::tuple< TSpatialAccelerationTypes...>>::type> ())
+		{
+			return 3;
+		}
+
+		if (Accel->template As<typename std::tuple_element<ClampedIdx(4), std::tuple< TSpatialAccelerationTypes...>>::type> ())
+		{
+			return 4;
+		}
+
+		if (Accel->template As<typename std::tuple_element<ClampedIdx(5), std::tuple< TSpatialAccelerationTypes...>>::type> ())
+		{
+			return 5;
+		}
+
+		if (Accel->template As<typename std::tuple_element<ClampedIdx(6), std::tuple< TSpatialAccelerationTypes...>>::type> ())
+		{
+			return 6;
+		}
+
+		if (Accel->template As<typename std::tuple_element<ClampedIdx(7), std::tuple< TSpatialAccelerationTypes...>>::type> ())
+		{
+			return 7;
+		}
+
+		check(false);	//Did you instantiate the collection with the right template arguments?
+		return 0;
 	}
 
-	TSpatialBucketTuple< TSpatialAccelerationBuckets...> Buckets;
-	static constexpr uint32 NumBuckets = sizeof...(TSpatialAccelerationBuckets);
-	static_assert(NumBuckets < 8, "8 max buckets supported");
+	static constexpr uint16 NumBuckets = 8;
+	TSpatialCollectionBucket<TSpatialAccelerationBucketEntry<TPayloadType, T, d>> Buckets[NumBuckets];
+	TSpatialTypeTuple< TSpatialAccelerationTypes...> Types;
+	static constexpr uint32 NumTypes = sizeof...(TSpatialAccelerationTypes);
 };
 
 
diff --git a/Engine/Source/Runtime/Experimental/Chaos/Public/Chaos/PBDRigidsEvolution.h b/Engine/Source/Runtime/Experimental/Chaos/Public/Chaos/PBDRigidsEvolution.h
--- a/Engine/Source/Runtime/Experimental/Chaos/Public/Chaos/PBDRigidsEvolution.h
+++ b/Engine/Source/Runtime/Experimental/Chaos/Public/Chaos/PBDRigidsEvolution.h
@@ -480,7 +480,7 @@ protected:
 
 	int32 NumIterations;
 
-	static ISpatialAccelerationCollection<TAccelerationStructureHandle<T, d>, T, d>* CreateNewSpatialStructure(const TArray<FAccelerationStructureBuilder>& Bounds);
+	static ISpatialAccelerationCollection<TAccelerationStructureHandle<T, d>, T, d>* CreateNewSpatialStructure();
 	void InitializeAccelerationCache();
 };
 
diff --git a/Engine/Source/Runtime/Experimental/ChaosSolvers/Private/PBDRigidsSolver.cpp b/Engine/Source/Runtime/Experimental/ChaosSolvers/Private/PBDRigidsSolver.cpp
--- a/Engine/Source/Runtime/Experimental/ChaosSolvers/Private/PBDRigidsSolver.cpp
+++ b/Engine/Source/Runtime/Experimental/ChaosSolvers/Private/PBDRigidsSolver.cpp
@@ -140,6 +140,11 @@ namespace Chaos
 		// Make sure this particle doesn't already have a proxy
 		checkSlow(GTParticle->Proxy == nullptr);
 
+		if (GTParticle->Geometry() && GTParticle->Geometry()->HasBoundingBox() && GTParticle->Geometry()->BoundingBox().Extents().Max() >= 100000)
+		{
+			GTParticle->SetSpatialIdx(FSpatialAccelerationIdx{ 1,0 });
+		}
+
 		// Grab the particle's type
 		const EParticleType InParticleType = GTParticle->ObjectType();
 
