Fixed grid OOB when sweeping right on the edge of a heightfield

diff --git a/Engine/Source/Runtime/Experimental/Chaos/Private/Chaos/HeightField.cpp b/Engine/Source/Runtime/Experimental/Chaos/Private/Chaos/HeightField.cpp
--- a/Engine/Source/Runtime/Experimental/Chaos/Private/Chaos/HeightField.cpp
+++ b/Engine/Source/Runtime/Experimental/Chaos/Private/Chaos/HeightField.cpp
@@ -659,7 +659,7 @@ namespace Chaos
 						TVector<int32, 2> NeighCoord = CellCoord.Index + Neighbor;
 
 						FBounds2D CellBounds;
-						if(!Seen.Contains(NeighCoord) && GetCellBounds3DScaled(NeighCoord, Min, Max, HalfExtents3D))
+						if(GetCellBounds3DScaled(NeighCoord, Min, Max, HalfExtents3D) && !Seen.Contains(NeighCoord))
 						{
 							if(TBox<T, 3>::RaycastFast(Min, Max, StartPoint, Dir, InvDir, bParallel, CurrentLength, InvCurrentLength, ToI, HitPoint))
 							{
