Change the behaviour of the TTriangleMesh::GetVertexSet function to first reset the output set rather than simply append the vertices to it, in ordert to better match its documentation, expected and current usage.
Plus a couple of minor TriangleMesh improvements.

diff --git a/Engine/Source/Runtime/Experimental/Chaos/Private/Chaos/TriangleMesh.cpp b/Engine/Source/Runtime/Experimental/Chaos/Private/Chaos/TriangleMesh.cpp
--- a/Engine/Source/Runtime/Experimental/Chaos/Private/Chaos/TriangleMesh.cpp
+++ b/Engine/Source/Runtime/Experimental/Chaos/Private/Chaos/TriangleMesh.cpp
@@ -71,6 +71,7 @@ void TTriangleMesh<T>::InitHelper(const int32 StartIdx, const int32 EndIdx)
 				MaxIdx = FMath::Max(MaxIdx, MElements[i][Axis]);
 			}
 			check(MElements[i][0] != MElements[i][1]);
+			check(MElements[i][0] != MElements[i][2]);
 			check(MElements[i][1] != MElements[i][2]);
 		}
 		// This assumes vertices are contiguous in the vertex buffer. Assumption is held throughout TTriangleMesh
@@ -98,7 +99,7 @@ TSet<int32> TTriangleMesh<T>::GetVertices() const
 template<class T>
 void TTriangleMesh<T>::GetVertexSet(TSet<int32>& VertexSet) const
 {
-	VertexSet.Reserve(MNumIndices);
+	VertexSet.Reset(MNumIndices);
 	for (const TVector<int32, 3>& Element : MElements)
 	{
 		VertexSet.Append({Element[0], Element[1], Element[2]});
@@ -273,31 +274,31 @@ void TTriangleMesh<T>::GetFaceNormals(TArray<TVector<T, 3>>& Normals, const TArr
 }
 
 template<class T>
-TArray<TVector<T, 3>> TTriangleMesh<T>::GetPointNormals(const TArrayView<const TVector<T, 3>>& Points, const bool bReturnEmptyOnError, const bool bFillAtStartIndex)
+TArray<TVector<T, 3>> TTriangleMesh<T>::GetPointNormals(const TArrayView<const TVector<T, 3>>& Points, const bool bReturnEmptyOnError, const bool bUseGlobalArray)
 {
 	TArray<TVector<T, 3>> FaceNormals = GetFaceNormals(Points, bReturnEmptyOnError);
 	TArray<TVector<T, 3>> PointNormals;
-	GetPointNormals(PointNormals, FaceNormals, bReturnEmptyOnError, bFillAtStartIndex);
+	GetPointNormals(PointNormals, FaceNormals, bReturnEmptyOnError, bUseGlobalArray);
 	return PointNormals;
 }
 
 template<class T>
-void TTriangleMesh<T>::GetPointNormals(TArray<TVector<T, 3>>& PointNormals, const TArray<TVector<T, 3>>& FaceNormals, const bool bReturnEmptyOnError, const bool bFillAtStartIndex)
+void TTriangleMesh<T>::GetPointNormals(TArray<TVector<T, 3>>& PointNormals, const TArray<TVector<T, 3>>& FaceNormals, const bool bReturnEmptyOnError, const bool bUseGlobalArray)
 {
 	GetPointToTriangleMap(); // build MPointToTriangleMap
 	const TTriangleMesh<T>* ConstThis = this;
-	ConstThis->GetPointNormals(PointNormals, FaceNormals, bReturnEmptyOnError, bFillAtStartIndex);
+	ConstThis->GetPointNormals(PointNormals, FaceNormals, bReturnEmptyOnError, bUseGlobalArray);
 }
 
 template<class T>
-void TTriangleMesh<T>::GetPointNormals(TArray<TVector<T, 3>>& PointNormals, const TArray<TVector<T, 3>>& FaceNormals, const bool bReturnEmptyOnError, const bool bFillAtStartIndex) const
+void TTriangleMesh<T>::GetPointNormals(TArray<TVector<T, 3>>& PointNormals, const TArray<TVector<T, 3>>& FaceNormals, const bool bReturnEmptyOnError, const bool bUseGlobalArray) const
 {
 	check(MPointToTriangleMap.Num() != 0);
 	PointNormals.SetNum(MNumIndices);
 	for (auto Element : MPointToTriangleMap)
 	{
 		checkSlow(Element.Key >= MStartIdx);
-		const int32 NormalIndex = bFillAtStartIndex ? Element.Key : Element.Key - MStartIdx;  // Select whether the normal indices match the points indices or start at 0
+		const int32 NormalIndex = bUseGlobalArray ? Element.Key : GlobalToLocal(Element.Key);  // Select whether the points normal indices match the points indices or start at 0
 		if (PointNormals.Num() <= NormalIndex)
 		{
 			PointNormals.SetNum(NormalIndex);
diff --git a/Engine/Source/Runtime/Experimental/Chaos/Public/Chaos/TriangleMesh.h b/Engine/Source/Runtime/Experimental/Chaos/Public/Chaos/TriangleMesh.h
--- a/Engine/Source/Runtime/Experimental/Chaos/Public/Chaos/TriangleMesh.h
+++ b/Engine/Source/Runtime/Experimental/Chaos/Public/Chaos/TriangleMesh.h
@@ -132,17 +132,17 @@ namespace Chaos
 		TArray<TVector<T, 3>> GetFaceNormals(const TParticles<T, 3>& InParticles, const bool ReturnEmptyOnError = true) const
 		{ return GetFaceNormals(InParticles.X(), ReturnEmptyOnError); }
 
-		TArray<TVector<T, 3>> GetPointNormals(const TArrayView<const TVector<T, 3>>& points, const bool bReturnEmptyOnError = false, const bool bFillAtStartIndex = true);
-		void GetPointNormals(TArray<TVector<T, 3>>& PointNormals, const TArray<TVector<T, 3>>& FaceNormals, const bool ReturnEmptyOnError = false, const bool bFillAtStartIndex = true);
+		TArray<TVector<T, 3>> GetPointNormals(const TArrayView<const TVector<T, 3>>& points, const bool bReturnEmptyOnError = false, const bool bUseGlobalArray = true);
+		void GetPointNormals(TArray<TVector<T, 3>>& PointNormals, const TArray<TVector<T, 3>>& FaceNormals, const bool ReturnEmptyOnError = false, const bool bUseGlobalArray = true);
 		/** \brief Get per-point normals. 
 		 * This const version of this function requires \c GetPointToTriangleMap() 
 		 * to be called prior to invoking this function. 
-		 * @param bFillAtStartIndex When true, fill the array from the StartIdx to StartIdx + NumIndices - 1 positions, otherwise fill the array from the 0 to NumIndices - 1 positions.
+		 * @param bUseGlobalArray When true, fill the array from the StartIdx to StartIdx + NumIndices - 1 positions, otherwise fill the array from the 0 to NumIndices - 1 positions.
 		 */
-		void GetPointNormals(TArray<TVector<T, 3>>& PointNormals, const TArray<TVector<T, 3>>& FaceNormals, const bool bReturnEmptyOnError, const bool bFillAtStartIndex) const;
+		void GetPointNormals(TArray<TVector<T, 3>>& PointNormals, const TArray<TVector<T, 3>>& FaceNormals, const bool bReturnEmptyOnError, const bool bUseGlobalArray) const;
 		/** Deprecated. Use TArrayView version. */
-		TArray<TVector<T, 3>> GetPointNormals(const TParticles<T, 3>& InParticles, const bool bReturnEmptyOnError = false, const bool bFillAtStartIndex = true)
-		{ return GetPointNormals(InParticles.X(), bReturnEmptyOnError, bFillAtStartIndex); }
+		TArray<TVector<T, 3>> GetPointNormals(const TParticles<T, 3>& InParticles, const bool bReturnEmptyOnError = false, const bool bUseGlobalArray = true)
+		{ return GetPointNormals(InParticles.X(), bReturnEmptyOnError, bUseGlobalArray); }
 
 		static TTriangleMesh<T> GetConvexHullFromParticles(const TArrayView<const TVector<T, 3>>& points);
 		/** Deprecated. Use TArrayView version. */
