Chaos : Persistent Collisions

- Remove collision constraints from the collision list when particles are removed or disabled.

diff --git a/Engine/Source/Runtime/Experimental/Chaos/Private/Chaos/PBDCollisionConstraint.cpp b/Engine/Source/Runtime/Experimental/Chaos/Private/Chaos/PBDCollisionConstraint.cpp
--- a/Engine/Source/Runtime/Experimental/Chaos/Private/Chaos/PBDCollisionConstraint.cpp
+++ b/Engine/Source/Runtime/Experimental/Chaos/Private/Chaos/PBDCollisionConstraint.cpp
@@ -128,6 +128,26 @@ void TPBDCollisionConstraint<T, d>::RemoveConstraint(int32 Idx)
 	}
 }
 
+
+template<typename T, int d>
+void TPBDCollisionConstraint<T, d>::RemoveConstraints(const TSet<TGeometryParticleHandle<T, d>*>&  InHandleSet)
+{
+	const TArray<TGeometryParticleHandle<T, d>*> HandleArray = InHandleSet.Array();
+	for (auto ParticleHandle : HandleArray)
+	{
+		TArray<FConstraintHandleID> Keys;
+		Handles.GetKeys(Keys);
+		for (FConstraintHandleID Key : Keys)
+		{
+			if (Key.Key == ParticleHandle || Key.Value == ParticleHandle)
+			{
+				RemoveConstraint(Handles[Key]->GetConstraintIndex());
+			}
+		}
+	}
+	Handles.Compact();
+}
+
 template<typename T, int d>
 void TPBDCollisionConstraint<T, d>::ApplyCollisionModifier(const TFunction<ECollisionModifierResult(FRigidBodyContactConstraint& Constraint)>& CollisionModifier)
 {
diff --git a/Engine/Source/Runtime/Experimental/Chaos/Public/Chaos/PBDCollisionConstraint.h b/Engine/Source/Runtime/Experimental/Chaos/Public/Chaos/PBDCollisionConstraint.h
--- a/Engine/Source/Runtime/Experimental/Chaos/Public/Chaos/PBDCollisionConstraint.h
+++ b/Engine/Source/Runtime/Experimental/Chaos/Public/Chaos/PBDCollisionConstraint.h
@@ -159,6 +159,11 @@ public:
 
 	void RemoveConstraint(int32 Idx);
 
+	/**
+	* Remove the constraints associated with the ParticleHandle.
+	*/
+	void RemoveConstraints(const TSet<TGeometryParticleHandle<T, d>*>&  ParticleHandle);
+
 	/**
 	 * Apply a modifier to the constraints and specify which constraints should be disabled.
 	 * You would probably call this in the PostComputeCallback. Prefer this to calling RemoveConstraints in a loop, 
diff --git a/Engine/Source/Runtime/Experimental/Chaos/Public/Chaos/PBDConstraintRule.h b/Engine/Source/Runtime/Experimental/Chaos/Public/Chaos/PBDConstraintRule.h
--- a/Engine/Source/Runtime/Experimental/Chaos/Public/Chaos/PBDConstraintRule.h
+++ b/Engine/Source/Runtime/Experimental/Chaos/Public/Chaos/PBDConstraintRule.h
@@ -243,6 +243,11 @@ namespace Chaos
 			}
 		}
 
+		virtual void RemoveConstraints(const TSet<TGeometryParticleHandle<T, d>*>& InConstraints)
+		{
+			Constraints.RemoveConstraints(InConstraints);
+		}
+
 		virtual void ApplyPushOut(const T Dt, int32 Island) override
 		{
 			const typename FConstraintColor::FLevelToColorToConstraintListMap& LevelToColorToConstraintListMap = GraphColor.GetIslandLevelToColorToConstraintListMap(Island);
