FindMostOpposingFace and FindGeometryOpposingNormal was doing an inverted transform on normals and directions - corrected these. Convex's FindGeometryOpposingNormal function was not implemented.

diff --git a/Engine/Source/Runtime/Experimental/Chaos/Public/Chaos/Convex.h b/Engine/Source/Runtime/Experimental/Chaos/Public/Chaos/Convex.h
--- a/Engine/Source/Runtime/Experimental/Chaos/Public/Chaos/Convex.h
+++ b/Engine/Source/Runtime/Experimental/Chaos/Public/Chaos/Convex.h
@@ -133,14 +133,28 @@ namespace Chaos
 			return MostOpposingIdx;
 		}
 
-		TVector<T, d> FindGeometryOpposingNormal(const TVector<T, d>& DenormDir, int32 FaceIndex, const TVector<T, d>& OriginalNormal) const
+		TVector<T, d> FindGeometryOpposingNormal(const TVector<T, d>& DenormDir, int32 HintFaceIndex, const TVector<T, d>& OriginalNormal) const
 		{
-			if (ensure(FaceIndex != INDEX_NONE))
+			if (HintFaceIndex != INDEX_NONE)
 			{
-				const TPlane<float, 3>& OpposingFace = GetFaces()[FaceIndex];
+				const TPlane<float, 3>& OpposingFace = GetFaces()[HintFaceIndex];
 				return OpposingFace.Normal();
 			}
-			return FVector(0, 0, 1);
+		
+			// todo: make a way to call FindMostOpposingFace without a search dist
+			int32 MostOpposingIdx = INDEX_NONE;
+			T MostOpposingDot = TNumericLimits<T>::Max();
+			for (int32 Idx = 0; Idx < Planes.Num(); ++Idx)
+			{
+				const T Dot = TVector<T, d>::DotProduct(Planes[Idx].Normal(), DenormDir);
+				if (Dot < MostOpposingDot)
+				{
+					MostOpposingDot = Dot;
+					MostOpposingIdx = Idx;
+				}
+			}
+			ensure(MostOpposingIdx != INDEX_NONE);
+			return Planes[MostOpposingIdx].Normal();
 		}
 
 		virtual TVector<T, d> Support(const TVector<T, d>& Direction, const T Thickness) const override
diff --git a/Engine/Source/Runtime/Experimental/Chaos/Public/Chaos/ImplicitObjectScaled.h b/Engine/Source/Runtime/Experimental/Chaos/Public/Chaos/ImplicitObjectScaled.h
--- a/Engine/Source/Runtime/Experimental/Chaos/Public/Chaos/ImplicitObjectScaled.h
+++ b/Engine/Source/Runtime/Experimental/Chaos/Public/Chaos/ImplicitObjectScaled.h
@@ -174,36 +174,33 @@ public:
 		ensure(FMath::IsNearlyEqual(UnitDir.SizeSquared(), 1, KINDA_SMALL_NUMBER));
 
 		const TVector<T, d> UnscaledPosition = MInvScale * Position;
-		TVector<T, d> UnscaledDir = MInvScale * UnitDir;
-		const T UnscaledLength = UnscaledDir.SafeNormalize();
+		const TVector<T, d> UnscaledDirDenorm = MScale * UnitDir;
+		const float LengthScale = UnscaledDirDenorm.Size();
+		const TVector<T, d> UnscaledDir
+			= ensure(LengthScale > TNumericLimits<T>::Min())
+			? UnscaledDirDenorm / LengthScale
+			: TVector<T, d>(0.f, 0.f, 1.f);
 		const T UnscaledSearchDist = SearchDist * MScale.Max();	//this is not quite right since it's no longer a sphere, but the whole thing is fuzzy anyway
-
 		return MObject->FindMostOpposingFace(UnscaledPosition, UnscaledDir, HintFaceIndex, UnscaledSearchDist);
 	}
 
-	virtual TVector<T, 3> FindGeometryOpposingNormal(const TVector<T, d>& DenormDir, int32 FaceIndex, const TVector<T, d>& OriginalNormal) const override
+	virtual TVector<T, 3> FindGeometryOpposingNormal(const TVector<T, d>& DenormDir, int32 HintFaceIndex, const TVector<T, d>& OriginalNormal) const override
 	{
 		ensure(MInternalThickness == 0);	//not supported: do we care?
-		TVector<T, 3> LocalDenormDir = DenormDir * MInvScale;
-		TVector<T, 3> LocalOriginalNormal = OriginalNormal * MInvScale;
-		
-
-		// TODO: Fix and use TVector::SafeNormalize().
-		//We want N / ||N|| and to avoid inf
-		//So we want N / ||N|| < 1 / eps => N eps < ||N||, but this is clearly true for all eps < 1 and N > 0
-		T SizeSqr = LocalOriginalNormal.SizeSquared();
-		if (SizeSqr <= TNumericLimits<T>::Min())
-		{
-			LocalOriginalNormal = TVector<T, 3>(0, 0, 1);
-			ensure(false);
-		}
-		else
-		{
-			LocalOriginalNormal /= sqrt(SizeSqr);
-		}
-
-		const TVector<T, d> LocalNormal = MObject->FindGeometryOpposingNormal(LocalDenormDir, FaceIndex, LocalOriginalNormal);
-		TVector<T, d> Normal = LocalNormal * MScale;
+		ensure(FMath::IsNearlyEqual(OriginalNormal.SizeSquared(), 1, KINDA_SMALL_NUMBER));
+
+		// Get unscaled dir and normal
+		const TVector<T, 3> LocalDenormDir = DenormDir * MScale;
+		const TVector<T, 3> LocalOriginalNormalDenorm = OriginalNormal * MScale;
+		const float NormalLengthScale = LocalOriginalNormalDenorm.Size();
+		const TVector<T, 3> LocalOriginalNormal
+			= ensure(NormalLengthScale > SMALL_NUMBER)
+			? LocalOriginalNormalDenorm / NormalLengthScale
+			: TVector<T, d>(0, 0, 1);
+
+		// Compute final normal
+		const TVector<T, d> LocalNormal = MObject->FindGeometryOpposingNormal(LocalDenormDir, HintFaceIndex, LocalOriginalNormal);
+		TVector<T, d> Normal = LocalNormal * MInvScale;
 		if (ensure(Normal.SafeNormalize()) == 0)
 		{
 			Normal = TVector<T,3>(0,0,1);
