Stopping point for bucket work.

diff --git a/Engine/Source/Runtime/Experimental/Chaos/Private/Chaos/PBDRigidsEvolution.cpp b/Engine/Source/Runtime/Experimental/Chaos/Private/Chaos/PBDRigidsEvolution.cpp
--- a/Engine/Source/Runtime/Experimental/Chaos/Private/Chaos/PBDRigidsEvolution.cpp
+++ b/Engine/Source/Runtime/Experimental/Chaos/Private/Chaos/PBDRigidsEvolution.cpp
@@ -89,20 +89,21 @@ namespace Chaos
 	int32 AABBMaxTreeDepth = 200;
 	FAutoConsoleVariableRef CVarAABBMaxTreeDepth(TEXT("p.AABBMaxTreeDepth"), AABBMaxTreeDepth, TEXT(""));
 
-	int32 MaxPayloadSize = 20000;
+	float MaxPayloadSize = 20000;
 	FAutoConsoleVariableRef CVarMaxPayloadSize(TEXT("p.MaxPayloadSize"), MaxPayloadSize, TEXT(""));
 
 	template<class FPBDRigidsEvolution, class FPBDCollisionConstraint, class T, int d>
 	ISpatialAccelerationCollection<TAccelerationStructureHandle<T, d>, T, d>*  TPBDRigidsEvolutionBase<FPBDRigidsEvolution, FPBDCollisionConstraint, T, d>::CreateNewSpatialStructure()
 	{
 		const TArray<FAccelerationStructureBuilder> Empty;
-		TArray<TSpatialAccelerationParams<TAccelerationStructureHandle<T, d>, T, d>> EmptyParams;
 		if (BroadphaseType == 0)
 		{
 			using AccelType = TBoundingVolume<TAccelerationStructureHandle<T, d>, T, d>;
-			auto Collection = new TSpatialAccelerationCollection<AccelType>(EmptyParams);
+			TArray<TSpatialAccelerationParams<TAccelerationStructureHandle<T, d>, T, d>> Buckets = { {AccelType::StaticType, MaxPayloadSize, BoundingVolumeNumCells } };
 
-			auto Structure = MakeUnique<AccelType>(Empty, false, 0, BoundingVolumeNumCells, MaxPayloadSize);
+			auto Collection = new TSpatialAccelerationCollection<AccelType>(Buckets);
+
+			auto Structure = MakeUnique<AccelType>(Empty, false, 0, Buckets[0].MaxCells, Buckets[0].MaxPayloadBounds);
 			Collection->AddSubstructure(MoveTemp(Structure), 0);
 			
 			return Collection;
@@ -110,18 +111,20 @@ namespace Chaos
 		else if(BroadphaseType == 1)
 		{
 			using AccelType = TAABBTree<TAccelerationStructureHandle<T, d>, TAABBTreeLeafArray<TAccelerationStructureHandle<T, d>, T>, T>;
-			auto Collection = new TSpatialAccelerationCollection<AccelType>(EmptyParams);
+			TArray<TSpatialAccelerationParams<TAccelerationStructureHandle<T, d>, T, d>> Buckets = { {AccelType::StaticType, MaxPayloadSize, 0, MaxChildrenInLeaf, MaxTreeDepth} };
+			auto Collection = new TSpatialAccelerationCollection<AccelType>(Buckets);
 
-			auto Structure = MakeUnique<AccelType>(Empty, MaxChildrenInLeaf, MaxTreeDepth, MaxPayloadSize);
+			auto Structure = MakeUnique<AccelType>(Empty, Buckets[0].MaxChildrenInLeaf, Buckets[0].MaxTreeDepth, Buckets[0].MaxPayloadBounds);
 			Collection->AddSubstructure(MoveTemp(Structure), 0);
 			return Collection;
 		}
 		else if(BroadphaseType == 2)
 		{
 			using AccelType = TAABBTree<TAccelerationStructureHandle<T, d>, TBoundingVolume<TAccelerationStructureHandle<T, d>, T, d>, T>;
-			auto Collection = new TSpatialAccelerationCollection<AccelType>(EmptyParams);
+			TArray<TSpatialAccelerationParams<TAccelerationStructureHandle<T, d>, T, d>> Buckets = { {AccelType::StaticType, MaxPayloadSize, 0, AABBMaxChildrenInLeaf, AABBMaxTreeDepth} };
+			auto Collection = new TSpatialAccelerationCollection<AccelType>(Buckets);
 
-			auto Structure = MakeUnique<AccelType>(Empty, AABBMaxChildrenInLeaf, AABBMaxTreeDepth, MaxPayloadSize);
+			auto Structure = MakeUnique<AccelType>(Empty, Buckets[0].MaxChildrenInLeaf, Buckets[0].MaxTreeDepth, Buckets[0].MaxPayloadBounds);
 			Collection->AddSubstructure(MoveTemp(Structure), 0);
 			return Collection;
 		}
@@ -129,13 +132,19 @@ namespace Chaos
 		{
 			using AccelType = TAABBTree<TAccelerationStructureHandle<T, d>, TAABBTreeLeafArray<TAccelerationStructureHandle<T, d>, T>, T>;
 			using GridType = TBoundingVolume<TAccelerationStructureHandle<T, d>, T, d>;
-			auto Collection = new TSpatialAccelerationCollection<AccelType, GridType>(EmptyParams);
 
-			auto Structure = MakeUnique<AccelType>(Empty, MaxChildrenInLeaf, MaxTreeDepth, MaxPayloadSize);
-			Collection->AddSubstructure(MoveTemp(Structure), 0);
+			TArray<TSpatialAccelerationParams<TAccelerationStructureHandle<T, d>, T, d>> Buckets = {
+				{AccelType::StaticType, MaxPayloadSize, 0, AABBMaxChildrenInLeaf, AABBMaxTreeDepth},
+				{GridType::StaticType, MaxPayloadSize*100, BoundingVolumeNumCells, AABBMaxChildrenInLeaf, AABBMaxTreeDepth},
+			};
+
+			auto Collection = new TSpatialAccelerationCollection<AccelType, GridType>(Buckets);
+
+			auto Structure0 = MakeUnique<AccelType>(Empty, Buckets[0].MaxChildrenInLeaf, Buckets[0].MaxTreeDepth, Buckets[0].MaxPayloadBounds);
+			Collection->AddSubstructure(MoveTemp(Structure0), 0);
 
-			auto Structure2 = MakeUnique<GridType>(Empty, false, 0, BoundingVolumeNumCells, MaxPayloadSize);
-			Collection->AddSubstructure(MoveTemp(Structure2), 1);
+			auto Structure1 = MakeUnique<GridType>(Empty, false, 0, Buckets[1].MaxCells, Buckets[1].MaxPayloadBounds);
+			Collection->AddSubstructure(MoveTemp(Structure1), 1);
 			return Collection;
 		}
 	}
@@ -145,24 +154,26 @@ namespace Chaos
 	{
 		using BVType = TBoundingVolume<TAccelerationStructureHandle<T, d>, T, d>;
 		using AABBType = TAABBTree<TAccelerationStructureHandle<T, d>, TAABBTreeLeafArray<TAccelerationStructureHandle<T, d>, T>, T>;
-		TArray<TSpatialAccelerationParams<TAccelerationStructureHandle<T, d>, T, d>> EmptyParams;
-
+		
 		if (Substructure->template As<BVType>())
 		{
-			auto Collection = MakeUnique<TSpatialAccelerationCollection<BVType>>(EmptyParams);
+			TArray<TSpatialAccelerationParams<TAccelerationStructureHandle<T, d>, T, d>> Buckets = { {BVType::StaticType } };
+			auto Collection = MakeUnique<TSpatialAccelerationCollection<BVType>>(Buckets);
 			Collection->AddSubstructure(MoveTemp(Substructure), 0);
 			return Collection;
 		}
 		else if (Substructure->template As<AABBType>())
 		{
-			auto Collection = MakeUnique<TSpatialAccelerationCollection<AABBType>>(EmptyParams);
+			TArray<TSpatialAccelerationParams<TAccelerationStructureHandle<T, d>, T, d>> Buckets = { {AABBType::StaticType } };
+			auto Collection = MakeUnique<TSpatialAccelerationCollection<AABBType>>(Buckets);
 			Collection->AddSubstructure(MoveTemp(Substructure), 0);
 			return Collection;
 		}
 		else
 		{
 			using AccelType = TAABBTree<TAccelerationStructureHandle<T, d>, TBoundingVolume<TAccelerationStructureHandle<T, d>, T, d>, T>;
-			auto Collection = MakeUnique<TSpatialAccelerationCollection<AccelType>>(EmptyParams);
+			TArray<TSpatialAccelerationParams<TAccelerationStructureHandle<T, d>, T, d>> Buckets = { {AccelType::StaticType } };
+			auto Collection = MakeUnique<TSpatialAccelerationCollection<AccelType>>(Buckets);
 			Collection->AddSubstructure(MoveTemp(Substructure), 0);
 			return Collection;
 		}
@@ -207,6 +218,7 @@ namespace Chaos
 		//While pending we may have updated an object, we may have also created and destroyed the object, but the final event is a delete, so just remove from acceleration structure
 		//While pending we destroyed, recreated using the same memory address, and then did an update. In this case we should remove first and then update as global bounds may have changed
 		//As long as we delete first and update second this will be respected
+
 		if (SpatialData.bDelete)
 		{
 			AccelerationStructure.RemoveElementFrom(SpatialData.AccelerationHandle, SpatialData.DeletedSpatialIdx);
@@ -216,7 +228,8 @@ namespace Chaos
 				if (int32* CacheIdxPtr = ParticleToCacheIdx.Find(Particle))
 				{
 					const auto SpatialIdx = SpatialData.DeletedSpatialIdx;
-					TArray<FAccelerationStructureBuilder>& CachedSpatialBuilderData = *CachedSpatialBuilderDataMap.FindByPredicate([SpatialIdx](const auto& Cache){ return Cache.SpatialIdx == SpatialIdx;})->CachedSpatialBuilderData;
+					auto MapFound = CachedSpatialBuilderDataMap.FindByPredicate([SpatialIdx](const auto& Cache) { return Cache.SpatialIdx == SpatialIdx; });
+					TArray<FAccelerationStructureBuilder>& CachedSpatialBuilderData = *(MapFound ? MapFound->CachedSpatialBuilderData : CachedSpatialBuilderDataMap.Last().CachedSpatialBuilderData);
 					const int32 CacheIdx = *CacheIdxPtr;
 					if (CacheIdx + 1 < CachedSpatialBuilderData.Num())	//will get swapped with last element, so update it
 					{
@@ -238,7 +251,8 @@ namespace Chaos
 			if (bAsync)
 			{
 				const auto SpatialIdx = SpatialData.UpdatedSpatialIdx;
-				TArray<FAccelerationStructureBuilder>& CachedSpatialBuilderData = *CachedSpatialBuilderDataMap.FindByPredicate([SpatialIdx](const auto& Cache) { return Cache.SpatialIdx == SpatialIdx; })->CachedSpatialBuilderData;
+				auto MapFound = CachedSpatialBuilderDataMap.FindByPredicate([SpatialIdx](const auto& Cache) { return Cache.SpatialIdx == SpatialIdx; });
+				TArray<FAccelerationStructureBuilder>& CachedSpatialBuilderData = *(MapFound ? MapFound->CachedSpatialBuilderData : CachedSpatialBuilderDataMap.Last().CachedSpatialBuilderData);
 
 				//make sure in mapping
 				int32 CacheIdx;
diff --git a/Engine/Source/Runtime/Experimental/Chaos/Public/Chaos/ISpatialAcceleration.h b/Engine/Source/Runtime/Experimental/Chaos/Public/Chaos/ISpatialAcceleration.h
--- a/Engine/Source/Runtime/Experimental/Chaos/Public/Chaos/ISpatialAcceleration.h
+++ b/Engine/Source/Runtime/Experimental/Chaos/Public/Chaos/ISpatialAcceleration.h
@@ -195,6 +195,8 @@ public:
 		check(false);
 	}
 
+	SpatialAccelerationType GetType() const { return Type; }
+
 	template <typename TConcrete>
 	TConcrete* As()
 	{
diff --git a/Engine/Source/Runtime/Experimental/Chaos/Public/Chaos/ISpatialAccelerationCollection.h b/Engine/Source/Runtime/Experimental/Chaos/Public/Chaos/ISpatialAccelerationCollection.h
--- a/Engine/Source/Runtime/Experimental/Chaos/Public/Chaos/ISpatialAccelerationCollection.h
+++ b/Engine/Source/Runtime/Experimental/Chaos/Public/Chaos/ISpatialAccelerationCollection.h
@@ -13,7 +13,7 @@ namespace Chaos
 template <typename TPayloadType, typename T, int d>
 struct CHAOS_API TSpatialAccelerationParams
 {
-	ESpatialAcceleration Type;	//Specifies the acceleration type. Unknown can be used for custom types, see TSpatialAccelerationParams
+	SpatialAccelerationType Type;	//Specifies the acceleration type. Unknown can be used for custom types, see TSpatialAccelerationParams
 
 	//Union of all params, some are only applicable to subset of structure types
 	float MaxPayloadBounds;	//max bounds size before inserting into a linear list of "large" objects. This list will always be checked
@@ -289,6 +289,25 @@ struct TSpatialAccelerationCollectionHelper
 			TSpatialAccelerationCollectionHelper < NextType < NumTypes ? NextType : 0, NumTypes, Tuple, TPayloadType, T, d>::SetNumFrom(TypesSrc, TypesDest);
 		}
 	}
+
+	static uint16 FindTypeIdx(const Tuple& Types, SpatialAccelerationType Type)
+	{
+		using AccelType = typename std::remove_pointer<decltype(GetAccelerationsPerType<TypeIdx>(Types).Objects)::ElementType>::type;
+		if (AccelType::StaticType == Type)
+		{
+			return TypeIdx;
+		}
+
+		constexpr int NextType = TypeIdx + 1;
+		if (NextType < NumTypes)
+		{
+			return TSpatialAccelerationCollectionHelper < NextType < NumTypes ? NextType : 0, NumTypes, Tuple, TPayloadType, T, d>::FindTypeIdx(Types, Type);
+		}
+		else
+		{
+			return NumTypes;
+		}
+	}
 };
 
 template <bool bGatherStats, typename SpatialAccelerationCollection, typename T, int d>
@@ -314,13 +333,20 @@ public:
 	using T = typename FirstAccelerationType::TType;
 	static constexpr int d = FirstAccelerationType::D;
 	TSpatialAccelerationCollection(const TArray<TSpatialAccelerationParams<TPayloadType, T, d>>& BucketsParams)
+		: NumBuckets(BucketsParams.Num())
 	{
-		//todo: verify params match template args
+		check(NumBuckets < MaxBuckets);
+		int32 ParamIdx = 0;
+		for (const auto Params : BucketsParams)
+		{
+			check((TSpatialAccelerationCollectionHelper<0, NumTypes, decltype(Types), TPayloadType, T, d>::FindTypeIdx(Types, Params.Type) < NumTypes));
+		}
 	}
 
 	virtual FSpatialAccelerationIdx AddSubstructure(TUniquePtr<ISpatialAcceleration<TPayloadType, T, d>>&& Substructure, int32 BucketIdx) override
 	{
 		check(BucketIdx < NumBuckets);
+		//todo: verify the bucket and type match
 		FSpatialAccelerationIdx Result;
 		Result.Bucket = BucketIdx;
 		
@@ -371,7 +397,7 @@ public:
 
 	virtual ISpatialAcceleration<TPayloadType, T, d>* GetSubstructure(FSpatialAccelerationIdx Idx) override
 	{
-		check(Idx.Bucket < NumBuckets);
+		check(Idx.Bucket < MaxBuckets);
 		return Buckets[Idx.Bucket].Objects[Idx.InnerIdx].Acceleration.Get();
 	}
 
@@ -463,14 +489,16 @@ public:
 
 	virtual void RemoveElementFrom(const TPayloadType& Payload, FSpatialAccelerationIdx SpatialIdx) override
 	{
-		check(SpatialIdx.Bucket < NumBuckets);
-		Buckets[SpatialIdx.Bucket].Objects[SpatialIdx.InnerIdx].Acceleration->RemoveElement(Payload);
+		const uint16 LastBucket = NumBuckets - 1;
+		const uint16 UseBucket = FMath::Min(LastBucket, SpatialIdx.Bucket);
+		Buckets[UseBucket].Objects[SpatialIdx.InnerIdx].Acceleration->RemoveElement(Payload);
 	}
 
 	virtual void UpdateElementIn(const TPayloadType& Payload, const TBox<T, d>& NewBounds, bool bHasBounds, FSpatialAccelerationIdx SpatialIdx)
 	{
-		check(SpatialIdx.Bucket < NumBuckets);
-		Buckets[SpatialIdx.Bucket].Objects[SpatialIdx.InnerIdx].Acceleration->UpdateElement(Payload, NewBounds, bHasBounds);
+		const uint16 LastBucket = NumBuckets - 1;
+		const uint16 UseBucket = FMath::Min(LastBucket, SpatialIdx.Bucket);
+		Buckets[UseBucket].Objects[SpatialIdx.InnerIdx].Acceleration->UpdateElement(Payload, NewBounds, bHasBounds);
 	}
 
 	virtual TUniquePtr<ISpatialAcceleration<TPayloadType, T, d>> Copy() const
@@ -491,12 +519,13 @@ public:
 private:
 
 	TSpatialAccelerationCollection(const TSpatialAccelerationCollection<TSpatialAccelerationTypes...>& Other)
+		: NumBuckets(Other.NumBuckets)
 	{
 		//need to make a deep copy of the unique ptrs and then update raw ptrs
 		TSpatialAccelerationCollectionHelper<0, NumTypes, decltype(Types), TPayloadType, T, d>::SetNumFrom(Other.Types, Types);
 		
 
-		for (int BucketIdx = 0; BucketIdx < NumBuckets; ++BucketIdx)
+		for (int BucketIdx = 0; BucketIdx < MaxBuckets; ++BucketIdx)
 		{
 			Buckets[BucketIdx].CopyFrom(Other.Buckets[BucketIdx]);
 			bool bFirst = true;
@@ -570,10 +599,11 @@ private:
 		return 0;
 	}
 
-	static constexpr uint16 NumBuckets = 8;
-	TSpatialCollectionBucket<TSpatialAccelerationBucketEntry<TPayloadType, T, d>> Buckets[NumBuckets];
+	static constexpr uint16 MaxBuckets = 8;
+	TSpatialCollectionBucket<TSpatialAccelerationBucketEntry<TPayloadType, T, d>> Buckets[MaxBuckets];
 	TSpatialTypeTuple< TSpatialAccelerationTypes...> Types;
 	static constexpr uint32 NumTypes = sizeof...(TSpatialAccelerationTypes);
+	uint16 NumBuckets;
 };
 
 
