Fix CIS

diff --git a/Engine/Source/Runtime/Experimental/Chaos/Public/Chaos/ISpatialAccelerationCollection.h b/Engine/Source/Runtime/Experimental/Chaos/Public/Chaos/ISpatialAccelerationCollection.h
--- a/Engine/Source/Runtime/Experimental/Chaos/Public/Chaos/ISpatialAccelerationCollection.h
+++ b/Engine/Source/Runtime/Experimental/Chaos/Public/Chaos/ISpatialAccelerationCollection.h
@@ -292,7 +292,7 @@ struct TSpatialAccelerationCollectionHelper
 
 	static uint16 FindTypeIdx(const Tuple& Types, SpatialAccelerationType Type)
 	{
-		using AccelType = typename std::remove_pointer<decltype(GetAccelerationsPerType<TypeIdx>(Types).Objects)::ElementType>::type;
+		using AccelType = typename std::remove_pointer<typename decltype(GetAccelerationsPerType<TypeIdx>(Types).Objects)::ElementType>::type;
 		if (AccelType::StaticType == Type)
 		{
 			return TypeIdx;
