Fix backwards dir scaling in non-uniform scaled object cast and sweep

diff --git a/Engine/Source/Runtime/Experimental/Chaos/Public/Chaos/ImplicitObjectScaled.h b/Engine/Source/Runtime/Experimental/Chaos/Public/Chaos/ImplicitObjectScaled.h
--- a/Engine/Source/Runtime/Experimental/Chaos/Public/Chaos/ImplicitObjectScaled.h
+++ b/Engine/Source/Runtime/Experimental/Chaos/Public/Chaos/ImplicitObjectScaled.h
@@ -89,12 +89,12 @@ public:
 		ensure(Thickness == 0 || (FMath::IsNearlyEqual(MScale[0], MScale[1]) && FMath::IsNearlyEqual(MScale[0], MScale[2])));	//non uniform turns sphere into an ellipsoid so no longer a raycast and requires a more expensive sweep
 
 		const TVector<T, d> UnscaledStart = MInvScale * StartPoint;
-		const TVector<T, d> UnscaledDirDenorm = MScale * Dir;
+		const TVector<T, d> UnscaledDirDenorm = MInvScale * Dir;
 		const T LengthScale = UnscaledDirDenorm.Size();
 		if (ensure(LengthScale > TNumericLimits<T>::Min()))
 		{
 			const T LengthScaleInv = 1.f / LengthScale;
-			const T UnscaledLength = Length * LengthScaleInv;
+			const T UnscaledLength = Length * LengthScale;
 			const TVector<T, d> UnscaledDir = UnscaledDirDenorm * LengthScaleInv;
 			
 			TVector<T, d> UnscaledPosition;
@@ -103,7 +103,7 @@ public:
 
 			if (MObject->Raycast(UnscaledStart, UnscaledDir, UnscaledLength, MInternalThickness + Thickness * MInvScale[0], UnscaledTime, UnscaledPosition, UnscaledNormal, OutFaceIndex))
 			{
-				OutTime = LengthScale * UnscaledTime;
+				OutTime = LengthScaleInv * UnscaledTime;
 				OutPosition = MScale * UnscaledPosition;
 				OutNormal = (MInvScale * UnscaledNormal).GetSafeNormal();
 				ensure(OutTime <= Length);
@@ -122,12 +122,12 @@ public:
 		ensure(FMath::IsNearlyEqual(LocalDir.SizeSquared(), 1, KINDA_SMALL_NUMBER));
 		ensure(Thickness == 0 || (FMath::IsNearlyEqual(OwningScaled.MScale[0], OwningScaled.MScale[1]) && FMath::IsNearlyEqual(OwningScaled.MScale[0], OwningScaled.MScale[2])));
 
-		const TVector<T, d> UnscaledDirDenorm = OwningScaled.MScale * LocalDir;
+		const TVector<T, d> UnscaledDirDenorm = OwningScaled.MInvScale * LocalDir;
 		const T LengthScale = UnscaledDirDenorm.Size();
 		if (ensure(LengthScale > TNumericLimits<T>::Min()))
 		{
 			const T LengthScaleInv = 1.f / LengthScale;
-			const T UnscaledLength = Length * LengthScaleInv;
+			const T UnscaledLength = Length * LengthScale;
 			const TVector<T, d> UnscaledDir = UnscaledDirDenorm * LengthScaleInv;
 
 			TVector<T, d> UnscaledPosition;
@@ -142,7 +142,7 @@ public:
 			
 			if (InternalGeom.SweepGeom(ScaledB, BToATMNoScale, UnscaledDir, UnscaledLength, UnscaledTime, UnscaledPosition, UnscaledNormal, OutFaceIndex, OwningScaled.MInternalThickness + Thickness))
 			{
-				OutTime = LengthScale * UnscaledTime;
+				OutTime = LengthScaleInv * UnscaledTime;
 				LocalPosition = OwningScaled.MScale * UnscaledPosition;
 				LocalNormal = (OwningScaled.MInvScale * UnscaledNormal).GetSafeNormal();
 				ensure(OutTime <= Length);
