Fix Linux linker error.

diff --git a/Engine/Source/Runtime/Experimental/Chaos/Public/Chaos/CollisionConvexMesh.h b/Engine/Source/Runtime/Experimental/Chaos/Public/Chaos/CollisionConvexMesh.h
--- a/Engine/Source/Runtime/Experimental/Chaos/Public/Chaos/CollisionConvexMesh.h
+++ b/Engine/Source/Runtime/Experimental/Chaos/Public/Chaos/CollisionConvexMesh.h
@@ -220,7 +220,7 @@ namespace Chaos
 		}
 
 		// CVars variables for controlling geometry complexity checking and simplification
-#if PLATFORM_MAC
+#if PLATFORM_MAC || PLATFORM_LINUX
 		static CHAOS_API int32 PerformGeometryCheck;
 		static CHAOS_API int32 PerformGeometryReduction;
 		static CHAOS_API int32 ParticlesThreshold;
