Make it easier to serialize spatial collection

diff --git a/Engine/Source/Runtime/Core/Public/UObject/ExternalPhysicsCustomObjectVersion.h b/Engine/Source/Runtime/Core/Public/UObject/ExternalPhysicsCustomObjectVersion.h
--- a/Engine/Source/Runtime/Core/Public/UObject/ExternalPhysicsCustomObjectVersion.h
+++ b/Engine/Source/Runtime/Core/Public/UObject/ExternalPhysicsCustomObjectVersion.h
@@ -33,6 +33,9 @@ struct CORE_API FExternalPhysicsCustomObjectVersion
 		//Save out heightfield data
 		HeightfieldData,
 
+		//Save out multiple acceleration structures
+		SerializeMultiStructures,
+
 		// -----<new versions can be added above this line>-------------------------------------------------
 		VersionPlusOne,
 		LatestVersion = VersionPlusOne - 1
diff --git a/Engine/Source/Runtime/Experimental/Chaos/Private/Chaos/BoundingVolume.cpp b/Engine/Source/Runtime/Experimental/Chaos/Private/Chaos/BoundingVolume.cpp
--- a/Engine/Source/Runtime/Experimental/Chaos/Private/Chaos/BoundingVolume.cpp
+++ b/Engine/Source/Runtime/Experimental/Chaos/Private/Chaos/BoundingVolume.cpp
@@ -30,6 +30,7 @@ namespace Chaos
 		case ESpatialAcceleration::BoundingVolume: return Ar.IsLoading() ? new TBoundingVolume<TPayloadType, T, d>() : nullptr;
 		case ESpatialAcceleration::AABBTree: return Ar.IsLoading() ? new TAABBTree<TPayloadType, TAABBTreeLeafArray<TPayloadType, T>, T>() : nullptr;
 		case ESpatialAcceleration::AABBTreeBV: return Ar.IsLoading() ? new TAABBTree<TPayloadType, TBoundingVolume<TPayloadType, T, 3>, T>() : nullptr;
+		case ESpatialAcceleration::Collection: check(false);	//Collections must be serialized directly since they are variadic
 		default: check(false); return nullptr;
 		}
 	}
diff --git a/Engine/Source/Runtime/Experimental/Chaos/Private/Chaos/PBDRigidsEvolution.cpp b/Engine/Source/Runtime/Experimental/Chaos/Private/Chaos/PBDRigidsEvolution.cpp
--- a/Engine/Source/Runtime/Experimental/Chaos/Private/Chaos/PBDRigidsEvolution.cpp
+++ b/Engine/Source/Runtime/Experimental/Chaos/Private/Chaos/PBDRigidsEvolution.cpp
@@ -9,12 +9,109 @@
 
 namespace Chaos
 {
+	int32 BroadphaseType = 3;
+	FAutoConsoleVariableRef CVarBroadphaseIsTree(TEXT("p.BroadphaseType"), BroadphaseType, TEXT(""));
+
+	int32 BoundingVolumeNumCells = 35;
+	FAutoConsoleVariableRef CVarBoundingVolumeNumCells(TEXT("p.BoundingVolumeNumCells"), BoundingVolumeNumCells, TEXT(""));
+
+	int32 MaxChildrenInLeaf = 5;
+	FAutoConsoleVariableRef CVarMaxChildrenInLeaf(TEXT("p.MaxChildrenInLeaf"), MaxChildrenInLeaf, TEXT(""));
+
+	int32 MaxTreeDepth = 200;
+	FAutoConsoleVariableRef CVarMaxTreeDepth(TEXT("p.MaxTreeDepth"), MaxTreeDepth, TEXT(""));
+
+	int32 AABBMaxChildrenInLeaf = 500;
+	FAutoConsoleVariableRef CVarAABBMaxChildrenInLeaf(TEXT("p.AABBMaxChildrenInLeaf"), AABBMaxChildrenInLeaf, TEXT(""));
+
+	int32 AABBMaxTreeDepth = 200;
+	FAutoConsoleVariableRef CVarAABBMaxTreeDepth(TEXT("p.AABBMaxTreeDepth"), AABBMaxTreeDepth, TEXT(""));
+
+	float MaxPayloadSize = 20000;
+	FAutoConsoleVariableRef CVarMaxPayloadSize(TEXT("p.MaxPayloadSize"), MaxPayloadSize, TEXT(""));
+
+	template<typename T, int d>
+	struct TDefaultCollectionFactory : public ISpatialAccelerationCollectionFactory<T, d>
+	{
+		virtual TUniquePtr<ISpatialAccelerationCollection<TAccelerationStructureHandle<T, d>, T, d>> CreateEmptyCollection() override
+		{
+			const TArray<TAccelerationStructureBuilder<T, d>> Empty;
+			if (BroadphaseType == 0)
+			{
+				using AccelType = TBoundingVolume<TAccelerationStructureHandle<T, d>, T, d>;
+				TArray<TSpatialAccelerationParams<TArray<TAccelerationStructureBuilder<T, d>>, TAccelerationStructureHandle<T, d>, T, d>> Buckets = { {AccelType::StaticType, MaxPayloadSize, BoundingVolumeNumCells } };
+
+				auto Structure = MakeUnique<AccelType>(Empty, false, 0, Buckets[0].MaxCells, Buckets[0].MaxPayloadBounds);
+				auto Collection = new TSpatialAccelerationCollection<TArray<TAccelerationStructureBuilder<T, d>>, AccelType>(MoveTemp(Buckets));
+
+				Collection->AddSubstructure(MoveTemp(Structure), 0);
+
+				return TUniquePtr<ISpatialAccelerationCollection<TAccelerationStructureHandle<T, d>, T, d>>(Collection);
+			}
+			else if (BroadphaseType == 1)
+			{
+				using AccelType = TAABBTree<TAccelerationStructureHandle<T, d>, TAABBTreeLeafArray<TAccelerationStructureHandle<T, d>, T>, T>;
+				TArray<TSpatialAccelerationParams<TArray<TAccelerationStructureBuilder<T, d>>, TAccelerationStructureHandle<T, d>, T, d>> Buckets = { {AccelType::StaticType, MaxPayloadSize, 0, MaxChildrenInLeaf, MaxTreeDepth} };
+
+				auto Structure = MakeUnique<AccelType>(Empty, Buckets[0].MaxChildrenInLeaf, Buckets[0].MaxTreeDepth, Buckets[0].MaxPayloadBounds);
+				auto Collection = new TSpatialAccelerationCollection<TArray<TAccelerationStructureBuilder<T, d>>, AccelType>(MoveTemp(Buckets));
+				Collection->AddSubstructure(MoveTemp(Structure), 0);
+				return TUniquePtr<ISpatialAccelerationCollection<TAccelerationStructureHandle<T, d>, T, d>>(Collection);
+			}
+			else if (BroadphaseType == 2)
+			{
+				using AccelType = TAABBTree<TAccelerationStructureHandle<T, d>, TBoundingVolume<TAccelerationStructureHandle<T, d>, T, d>, T>;
+				TArray<TSpatialAccelerationParams<TArray<TAccelerationStructureBuilder<T, d>>, TAccelerationStructureHandle<T, d>, T, d>> Buckets = { {AccelType::StaticType, MaxPayloadSize, 0, AABBMaxChildrenInLeaf, AABBMaxTreeDepth} };
+
+				auto Structure = MakeUnique<AccelType>(Empty, Buckets[0].MaxChildrenInLeaf, Buckets[0].MaxTreeDepth, Buckets[0].MaxPayloadBounds);
+				auto Collection = new TSpatialAccelerationCollection<TArray<TAccelerationStructureBuilder<T, d>>, AccelType>(MoveTemp(Buckets));
+				Collection->AddSubstructure(MoveTemp(Structure), 0);
+				return TUniquePtr<ISpatialAccelerationCollection<TAccelerationStructureHandle<T, d>, T, d>>(Collection);
+			}
+			else
+			{
+				using AccelType = TAABBTree<TAccelerationStructureHandle<T, d>, TAABBTreeLeafArray<TAccelerationStructureHandle<T, d>, T>, T>;
+				using GridType = TBoundingVolume<TAccelerationStructureHandle<T, d>, T, d>;
+
+				TArray<TSpatialAccelerationParams<TArray<TAccelerationStructureBuilder<T, d>>, TAccelerationStructureHandle<T, d>, T, d>> Buckets = {
+					{AccelType::StaticType, MaxPayloadSize, 0, AABBMaxChildrenInLeaf, AABBMaxTreeDepth},
+					{GridType::StaticType, MaxPayloadSize * 100, BoundingVolumeNumCells, AABBMaxChildrenInLeaf, AABBMaxTreeDepth},
+				};
+
+				auto Structure0 = MakeUnique<AccelType>(Empty, Buckets[0].MaxChildrenInLeaf, Buckets[0].MaxTreeDepth, Buckets[0].MaxPayloadBounds);
+
+				auto Structure1 = MakeUnique<GridType>(Empty, false, 0, Buckets[1].MaxCells, Buckets[1].MaxPayloadBounds);
+				auto Collection = new TSpatialAccelerationCollection<TArray<TAccelerationStructureBuilder<T, d>>, AccelType, GridType>(MoveTemp(Buckets));
+				Collection->AddSubstructure(MoveTemp(Structure0), 0);
+				Collection->AddSubstructure(MoveTemp(Structure1), 1);
+
+				return TUniquePtr<ISpatialAccelerationCollection<TAccelerationStructureHandle<T, d>, T, d>>(Collection);
+			}
+		}
+
+		virtual void Serialize(TUniquePtr<ISpatialAccelerationCollection<TAccelerationStructureHandle<T, d>, T, d>>& Ptr, FChaosArchive& Ar) override
+		{
+			if (Ar.IsLoading())
+			{
+				//todo: actually read in parameters
+				Ptr = CreateEmptyCollection();
+				Ptr->Serialize(Ar);
+			}
+			else
+			{
+				//todo: actually save out parameters
+				Ptr->Serialize(Ar);
+			}
+		}
+	};
+
 	template<class FPBDRigidsEvolution, class FPBDCollisionConstraint, class T, int d>
 	TPBDRigidsEvolutionBase<FPBDRigidsEvolution, FPBDCollisionConstraint, T, d>::TPBDRigidsEvolutionBase(TPBDRigidsSOAs<T, d>& InParticles, int32 InNumIterations)
 		: Particles(InParticles)
 		, bExternalReady(false)
 		, Clustering(static_cast<FPBDRigidsEvolution&>(*this), Particles.GetClusteredParticles())
 		, NumIterations(InNumIterations)
+		, SpatialCollectionFactory(new TDefaultCollectionFactory<T,d>())
 	{
 		Particles.GetParticleHandles().AddArray(&PhysicsMaterials);
 		Particles.GetParticleHandles().AddArray(&PerParticlePhysicsMaterials);
@@ -71,84 +168,6 @@ namespace Chaos
 		return ESubsequentsMode::TrackSubsequents;
 	}
 
-	int32 BroadphaseType = 3;
-	FAutoConsoleVariableRef CVarBroadphaseIsTree(TEXT("p.BroadphaseType"), BroadphaseType, TEXT(""));
-
-	int32 BoundingVolumeNumCells = 35;
-	FAutoConsoleVariableRef CVarBoundingVolumeNumCells(TEXT("p.BoundingVolumeNumCells"), BoundingVolumeNumCells, TEXT(""));
-
-	int32 MaxChildrenInLeaf = 5;
-	FAutoConsoleVariableRef CVarMaxChildrenInLeaf(TEXT("p.MaxChildrenInLeaf"), MaxChildrenInLeaf, TEXT(""));
-
-	int32 MaxTreeDepth = 200;
-	FAutoConsoleVariableRef CVarMaxTreeDepth(TEXT("p.MaxTreeDepth"), MaxTreeDepth, TEXT(""));
-
-	int32 AABBMaxChildrenInLeaf = 500;
-	FAutoConsoleVariableRef CVarAABBMaxChildrenInLeaf(TEXT("p.AABBMaxChildrenInLeaf"), AABBMaxChildrenInLeaf, TEXT(""));
-
-	int32 AABBMaxTreeDepth = 200;
-	FAutoConsoleVariableRef CVarAABBMaxTreeDepth(TEXT("p.AABBMaxTreeDepth"), AABBMaxTreeDepth, TEXT(""));
-
-	float MaxPayloadSize = 20000;
-	FAutoConsoleVariableRef CVarMaxPayloadSize(TEXT("p.MaxPayloadSize"), MaxPayloadSize, TEXT(""));
-
-	template<class FPBDRigidsEvolution, class FPBDCollisionConstraint, class T, int d>
-	ISpatialAccelerationCollection<TAccelerationStructureHandle<T, d>, T, d>*  TPBDRigidsEvolutionBase<FPBDRigidsEvolution, FPBDCollisionConstraint, T, d>::CreateNewSpatialStructure()
-	{
-		const TArray<FAccelerationStructureBuilder> Empty;
-		if (BroadphaseType == 0)
-		{
-			using AccelType = TBoundingVolume<TAccelerationStructureHandle<T, d>, T, d>;
-			TArray<TSpatialAccelerationParams<TAccelerationStructureHandle<T, d>, T, d>> Buckets = { {AccelType::StaticType, MaxPayloadSize, BoundingVolumeNumCells } };
-
-			auto Collection = new TSpatialAccelerationCollection<AccelType>(Buckets);
-
-			auto Structure = MakeUnique<AccelType>(Empty, false, 0, Buckets[0].MaxCells, Buckets[0].MaxPayloadBounds);
-			Collection->AddSubstructure(MoveTemp(Structure), 0);
-			
-			return Collection;
-		}
-		else if(BroadphaseType == 1)
-		{
-			using AccelType = TAABBTree<TAccelerationStructureHandle<T, d>, TAABBTreeLeafArray<TAccelerationStructureHandle<T, d>, T>, T>;
-			TArray<TSpatialAccelerationParams<TAccelerationStructureHandle<T, d>, T, d>> Buckets = { {AccelType::StaticType, MaxPayloadSize, 0, MaxChildrenInLeaf, MaxTreeDepth} };
-			auto Collection = new TSpatialAccelerationCollection<AccelType>(Buckets);
-
-			auto Structure = MakeUnique<AccelType>(Empty, Buckets[0].MaxChildrenInLeaf, Buckets[0].MaxTreeDepth, Buckets[0].MaxPayloadBounds);
-			Collection->AddSubstructure(MoveTemp(Structure), 0);
-			return Collection;
-		}
-		else if(BroadphaseType == 2)
-		{
-			using AccelType = TAABBTree<TAccelerationStructureHandle<T, d>, TBoundingVolume<TAccelerationStructureHandle<T, d>, T, d>, T>;
-			TArray<TSpatialAccelerationParams<TAccelerationStructureHandle<T, d>, T, d>> Buckets = { {AccelType::StaticType, MaxPayloadSize, 0, AABBMaxChildrenInLeaf, AABBMaxTreeDepth} };
-			auto Collection = new TSpatialAccelerationCollection<AccelType>(Buckets);
-
-			auto Structure = MakeUnique<AccelType>(Empty, Buckets[0].MaxChildrenInLeaf, Buckets[0].MaxTreeDepth, Buckets[0].MaxPayloadBounds);
-			Collection->AddSubstructure(MoveTemp(Structure), 0);
-			return Collection;
-		}
-		else
-		{
-			using AccelType = TAABBTree<TAccelerationStructureHandle<T, d>, TAABBTreeLeafArray<TAccelerationStructureHandle<T, d>, T>, T>;
-			using GridType = TBoundingVolume<TAccelerationStructureHandle<T, d>, T, d>;
-
-			TArray<TSpatialAccelerationParams<TAccelerationStructureHandle<T, d>, T, d>> Buckets = {
-				{AccelType::StaticType, MaxPayloadSize, 0, AABBMaxChildrenInLeaf, AABBMaxTreeDepth},
-				{GridType::StaticType, MaxPayloadSize*100, BoundingVolumeNumCells, AABBMaxChildrenInLeaf, AABBMaxTreeDepth},
-			};
-
-			auto Collection = new TSpatialAccelerationCollection<AccelType, GridType>(Buckets);
-
-			auto Structure0 = MakeUnique<AccelType>(Empty, Buckets[0].MaxChildrenInLeaf, Buckets[0].MaxTreeDepth, Buckets[0].MaxPayloadBounds);
-			Collection->AddSubstructure(MoveTemp(Structure0), 0);
-
-			auto Structure1 = MakeUnique<GridType>(Empty, false, 0, Buckets[1].MaxCells, Buckets[1].MaxPayloadBounds);
-			Collection->AddSubstructure(MoveTemp(Structure1), 1);
-			return Collection;
-		}
-	}
-
 	template<typename T, int d>
 	TUniquePtr<ISpatialAccelerationCollection<TAccelerationStructureHandle<T, d>, T, d>> CreateNewSpatialStructureFromSubStructure(TUniquePtr<ISpatialAcceleration<TAccelerationStructureHandle<T,d>, T,d>>&& Substructure)
 	{
@@ -157,23 +176,23 @@ namespace Chaos
 		
 		if (Substructure->template As<BVType>())
 		{
-			TArray<TSpatialAccelerationParams<TAccelerationStructureHandle<T, d>, T, d>> Buckets = { {BVType::StaticType } };
-			auto Collection = MakeUnique<TSpatialAccelerationCollection<BVType>>(Buckets);
+			TArray<TSpatialAccelerationParams<TArray<TAccelerationStructureBuilder<T,d>>, TAccelerationStructureHandle<T, d>, T, d>> Buckets = { {BVType::StaticType } };
+			auto Collection = MakeUnique<TSpatialAccelerationCollection<TArray<TAccelerationStructureBuilder<T, d>>, BVType>>(MoveTemp(Buckets));
 			Collection->AddSubstructure(MoveTemp(Substructure), 0);
 			return Collection;
 		}
 		else if (Substructure->template As<AABBType>())
 		{
-			TArray<TSpatialAccelerationParams<TAccelerationStructureHandle<T, d>, T, d>> Buckets = { {AABBType::StaticType } };
-			auto Collection = MakeUnique<TSpatialAccelerationCollection<AABBType>>(Buckets);
+			TArray<TSpatialAccelerationParams<TArray<TAccelerationStructureBuilder<T,d>>, TAccelerationStructureHandle<T, d>, T, d>> Buckets = { {AABBType::StaticType } };
+			auto Collection = MakeUnique<TSpatialAccelerationCollection<TArray<TAccelerationStructureBuilder<T, d>>, AABBType>>(MoveTemp(Buckets));
 			Collection->AddSubstructure(MoveTemp(Substructure), 0);
 			return Collection;
 		}
 		else
 		{
 			using AccelType = TAABBTree<TAccelerationStructureHandle<T, d>, TBoundingVolume<TAccelerationStructureHandle<T, d>, T, d>, T>;
-			TArray<TSpatialAccelerationParams<TAccelerationStructureHandle<T, d>, T, d>> Buckets = { {AccelType::StaticType } };
-			auto Collection = MakeUnique<TSpatialAccelerationCollection<AccelType>>(Buckets);
+			TArray<TSpatialAccelerationParams<TArray<TAccelerationStructureBuilder<T,d>>, TAccelerationStructureHandle<T, d>, T, d>> Buckets = { {AccelType::StaticType } };
+			auto Collection = MakeUnique<TSpatialAccelerationCollection<TArray<TAccelerationStructureBuilder<T, d>>, AccelType>>(MoveTemp(Buckets));
 			Collection->AddSubstructure(MoveTemp(Substructure), 0);
 			return Collection;
 		}
@@ -229,7 +248,7 @@ namespace Chaos
 				{
 					const auto SpatialIdx = SpatialData.DeletedSpatialIdx;
 					auto MapFound = CachedSpatialBuilderDataMap.FindByPredicate([SpatialIdx](const auto& Cache) { return Cache.SpatialIdx == SpatialIdx; });
-					TArray<FAccelerationStructureBuilder>& CachedSpatialBuilderData = *(MapFound ? MapFound->CachedSpatialBuilderData : CachedSpatialBuilderDataMap.Last().CachedSpatialBuilderData);
+					TArray<TAccelerationStructureBuilder<T,d>>& CachedSpatialBuilderData = *(MapFound ? MapFound->CachedSpatialBuilderData : CachedSpatialBuilderDataMap.Last().CachedSpatialBuilderData);
 					const int32 CacheIdx = *CacheIdxPtr;
 					if (CacheIdx + 1 < CachedSpatialBuilderData.Num())	//will get swapped with last element, so update it
 					{
@@ -252,7 +271,7 @@ namespace Chaos
 			{
 				const auto SpatialIdx = SpatialData.UpdatedSpatialIdx;
 				auto MapFound = CachedSpatialBuilderDataMap.FindByPredicate([SpatialIdx](const auto& Cache) { return Cache.SpatialIdx == SpatialIdx; });
-				TArray<FAccelerationStructureBuilder>& CachedSpatialBuilderData = *(MapFound ? MapFound->CachedSpatialBuilderData : CachedSpatialBuilderDataMap.Last().CachedSpatialBuilderData);
+				TArray<TAccelerationStructureBuilder<T,d>>& CachedSpatialBuilderData = *(MapFound ? MapFound->CachedSpatialBuilderData : CachedSpatialBuilderDataMap.Last().CachedSpatialBuilderData);
 
 				//make sure in mapping
 				int32 CacheIdx;
@@ -268,7 +287,7 @@ namespace Chaos
 				}
 
 				//update cache itself
-				CachedSpatialBuilderData[CacheIdx] = FAccelerationStructureBuilder{ Particle->HasBounds(), Particle->WorldSpaceInflatedBounds(), SpatialData.AccelerationHandle };
+				CachedSpatialBuilderData[CacheIdx] = TAccelerationStructureBuilder<T,d>{ Particle->HasBounds(), Particle->WorldSpaceInflatedBounds(), SpatialData.AccelerationHandle };
 			}
 		}
 	}
@@ -329,10 +348,10 @@ namespace Chaos
 		{
 			//initial frame so make empty structures
 
-			InternalAcceleration = TUniquePtr<FAccelerationStructure>(CreateNewSpatialStructure());
-			ScratchExternalAcceleration = TUniquePtr<FAccelerationStructure>(CreateNewSpatialStructure());
-			AsyncInternalAcceleration = TUniquePtr<FAccelerationStructure>(CreateNewSpatialStructure());
-			AsyncExternalAcceleration = TUniquePtr<FAccelerationStructure>(CreateNewSpatialStructure());
+			InternalAcceleration = TUniquePtr<FAccelerationStructure>(SpatialCollectionFactory->CreateEmptyCollection());
+			ScratchExternalAcceleration = TUniquePtr<FAccelerationStructure>(SpatialCollectionFactory->CreateEmptyCollection());
+			AsyncInternalAcceleration = TUniquePtr<FAccelerationStructure>(SpatialCollectionFactory->CreateEmptyCollection());
+			AsyncExternalAcceleration = TUniquePtr<FAccelerationStructure>(SpatialCollectionFactory->CreateEmptyCollection());
 			FlushInternalAccelerationQueue();
 			FlushExternalAccelerationQueue(*ScratchExternalAcceleration);
 			bExternalReady = true;
@@ -409,19 +428,26 @@ namespace Chaos
 		Ar.UsingCustomVersion(FExternalPhysicsCustomObjectVersion::GUID);
 		if (Ar.CustomVer(FExternalPhysicsCustomObjectVersion::GUID) >= FExternalPhysicsCustomObjectVersion::SerializeEvolutionBV)
 		{
-			//for now just assume a single structure
-			TUniquePtr<ISpatialAcceleration<TAccelerationStructureHandle<T, d>, T, d>> SubStructure;
-			if (!Ar.IsLoading())
+			if (Ar.CustomVer(FExternalPhysicsCustomObjectVersion::GUID) < FExternalPhysicsCustomObjectVersion::SerializeMultiStructures)
 			{
-				SubStructure = InternalAcceleration->RemoveSubstructure(FSpatialAccelerationIdx{ 0,0 });
-				Ar << SubStructure;
-				InternalAcceleration->AddSubstructure(MoveTemp(SubStructure), 0);
+				//old path assumes single sub-structure
+				TUniquePtr<ISpatialAcceleration<TAccelerationStructureHandle<T, d>, T, d>> SubStructure;
+				if (!Ar.IsLoading())
+				{
+					SubStructure = InternalAcceleration->RemoveSubstructure(FSpatialAccelerationIdx{ 0,0 });
+					Ar << SubStructure;
+					InternalAcceleration->AddSubstructure(MoveTemp(SubStructure), 0);
+				}
+				else
+				{
+					Ar << SubStructure;
+					InternalAcceleration = CreateNewSpatialStructureFromSubStructure(MoveTemp(SubStructure));
+					InitializeAccelerationCache();
+				}
 			}
 			else
 			{
-				Ar << SubStructure;
-				InternalAcceleration = CreateNewSpatialStructureFromSubStructure(MoveTemp(SubStructure));
-				InitializeAccelerationCache();
+				SpatialCollectionFactory->Serialize(InternalAcceleration, Ar);
 			}
 
 			SerializePendingMap(Ar, InternalAccelerationQueue);
diff --git a/Engine/Source/Runtime/Experimental/Chaos/Public/Chaos/ISpatialAccelerationCollection.h b/Engine/Source/Runtime/Experimental/Chaos/Public/Chaos/ISpatialAccelerationCollection.h
--- a/Engine/Source/Runtime/Experimental/Chaos/Public/Chaos/ISpatialAccelerationCollection.h
+++ b/Engine/Source/Runtime/Experimental/Chaos/Public/Chaos/ISpatialAccelerationCollection.h
@@ -10,7 +10,7 @@ namespace Chaos
 {
 
 /** Used for the creation of different acceleration structures */
-template <typename TPayloadType, typename T, int d>
+template <typename TParticleView, typename TPayloadType, typename T, int d>
 struct CHAOS_API TSpatialAccelerationParams
 {
 	SpatialAccelerationType Type;	//Specifies the acceleration type. Unknown can be used for custom types, see TSpatialAccelerationParams
@@ -22,8 +22,8 @@ struct CHAOS_API TSpatialAccelerationParams
 	int32 MaxChildrenInLeaf; //(tree only) if the tree sees MaxChildrenInLeaf or fewer objects, no more splitting occurs
 	int32 MaxTreeDepth;	//(tree only) if the tree has reached MaxTreeDepth, no more splitting occurs
 
-	typedef ISpatialAcceleration<TPayloadType, T, d> (*ConstructFn)(TSpatialAccelerationParams<TPayloadType, T, d>& Params, int32 ParamIdx);
-	ConstructFn ConstructionFactory;	//if unknown type is specified, use function pointer set by user
+	//typedef ISpatialAcceleration<TPayloadType, T, d> (*ConstructFn)(TSpatialAccelerationParams<TPayloadType, T, d>& Params, int32 ParamIdx);
+	//ConstructFn ConstructionFactory;	//if unknown type is specified, use function pointer set by user
 };
 
 template <typename TPayloadType, typename T, int d>
@@ -67,6 +67,14 @@ struct CHAOS_API TSpatialAccelerationBucketEntry
 	TSpatialAccelerationBucketEntry<TPayloadType, T, d>& operator=(TSpatialAccelerationBucketEntry<TPayloadType, T, d>& Other) = delete;
 };
 
+template <typename TPayloadType, typename T, int d>
+FChaosArchive& operator<<(FChaosArchive& Ar, TSpatialAccelerationBucketEntry<TPayloadType, T, d>& BucketEntry)
+{
+	Ar << BucketEntry.Acceleration;
+	Ar << BucketEntry.TypeInnerIdx;
+	return Ar;
+}
+
 
 template <typename TPayloadType, typename T, int d>
 void FreeObjHelper(TSpatialAccelerationBucketEntry<TPayloadType, T, d >& BucketEntry)
@@ -143,6 +151,14 @@ struct CHAOS_API TSpatialCollectionBucket
 
 };
 
+template <typename TObj>
+FChaosArchive& operator<<(FChaosArchive& Ar, TSpatialCollectionBucket<TObj>& Bucket)
+{
+	Ar << Bucket.Objects;
+	Ar << Bucket.FreeIndices;
+	return Ar;
+}
+
 template <typename... TRemaining>
 struct CHAOS_API TSpatialTypeTuple
 {
@@ -321,7 +337,7 @@ typename TEnableIf<!TIsSame<typename SpatialAccelerationCollection::TPayloadType
 {
 }
 
-template <typename ... TSpatialAccelerationTypes>
+template <typename TParticleView, typename ... TSpatialAccelerationTypes>
 class CHAOS_API TSpatialAccelerationCollection : public
 	ISpatialAccelerationCollection<typename std::tuple_element<0, std::tuple< TSpatialAccelerationTypes...>>::type::PayloadType,
 	typename std::tuple_element<0, std::tuple< TSpatialAccelerationTypes...>>::type::TType,
@@ -332,8 +348,10 @@ public:
 	using TPayloadType = typename FirstAccelerationType::PayloadType;
 	using T = typename FirstAccelerationType::TType;
 	static constexpr int d = FirstAccelerationType::D;
-	TSpatialAccelerationCollection(const TArray<TSpatialAccelerationParams<TPayloadType, T, d>>& BucketsParams)
-		: NumBuckets(BucketsParams.Num())
+
+	TSpatialAccelerationCollection(TArray<TSpatialAccelerationParams<TParticleView, TPayloadType, T, d>>&& InBucketsParams)
+		: NumBuckets(InBucketsParams.Num())
+		, BucketsParams(MoveTemp(InBucketsParams))
 	{
 		check(NumBuckets < MaxBuckets);
 		int32 ParamIdx = 0;
@@ -503,7 +521,7 @@ public:
 
 	virtual TUniquePtr<ISpatialAcceleration<TPayloadType, T, d>> Copy() const
 	{
-		return TUniquePtr<ISpatialAcceleration<TPayloadType, T, d>>(new TSpatialAccelerationCollection<TSpatialAccelerationTypes...>(*this));
+		return TUniquePtr<ISpatialAcceleration<TPayloadType, T, d>>(new TSpatialAccelerationCollection<TParticleView, TSpatialAccelerationTypes...>(*this));
 	}
 
 	virtual void PBDComputeConstraintsLowLevel_GatherStats(TPBDCollisionConstraint<T, d>& CollisionConstraint, T Dt) const override
@@ -516,9 +534,40 @@ public:
 		PBDComputeConstraintsLowLevel_Helper<false>(CollisionConstraint, *this, Dt);
 	}
 
+	virtual void Serialize(FChaosArchive& Ar)
+	{
+		//todo: let user serialize out bucket params
+
+		//serialize out sub structures
+		for (int BucketIdx = 0; BucketIdx < MaxBuckets; ++BucketIdx)
+		{
+			Ar << Buckets[BucketIdx];
+
+			if (Ar.IsLoading())
+			{
+				for (const TSpatialAccelerationBucketEntry<TPayloadType, T, d>& Entry : Buckets[BucketIdx].Objects)
+				{
+					ISpatialAcceleration<TPayloadType, T, d>* RawPtr = Entry.Acceleration.Get();
+					const int32 TypeIdx = GetTypeIdx(RawPtr);
+					switch (TypeIdx)
+					{
+					case 0: GetAccelerationsPerType<0>(Types).Objects[Entry.TypeInnerIdx] = &RawPtr->template AsChecked<typename std::tuple_element<0, std::tuple<TSpatialAccelerationTypes...>>::type>(); break;
+					case 1: GetAccelerationsPerType<ClampedIdx(1)>(Types).Objects[Entry.TypeInnerIdx] = &RawPtr->template AsChecked<typename std::tuple_element<ClampedIdx(1), std::tuple<TSpatialAccelerationTypes...>>::type>(); break;
+					case 2: GetAccelerationsPerType<ClampedIdx(2)>(Types).Objects[Entry.TypeInnerIdx] = &RawPtr->template AsChecked<typename std::tuple_element<ClampedIdx(2), std::tuple<TSpatialAccelerationTypes...>>::type>(); break;
+					case 3: GetAccelerationsPerType<ClampedIdx(3)>(Types).Objects[Entry.TypeInnerIdx] = &RawPtr->template AsChecked<typename std::tuple_element<ClampedIdx(3), std::tuple<TSpatialAccelerationTypes...>>::type>(); break;
+					case 4: GetAccelerationsPerType<ClampedIdx(4)>(Types).Objects[Entry.TypeInnerIdx] = &RawPtr->template AsChecked<typename std::tuple_element<ClampedIdx(4), std::tuple<TSpatialAccelerationTypes...>>::type>(); break;
+					case 5: GetAccelerationsPerType<ClampedIdx(5)>(Types).Objects[Entry.TypeInnerIdx] = &RawPtr->template AsChecked<typename std::tuple_element<ClampedIdx(5), std::tuple<TSpatialAccelerationTypes...>>::type>(); break;
+					case 6: GetAccelerationsPerType<ClampedIdx(6)>(Types).Objects[Entry.TypeInnerIdx] = &RawPtr->template AsChecked<typename std::tuple_element<ClampedIdx(6), std::tuple<TSpatialAccelerationTypes...>>::type>(); break;
+					case 7: GetAccelerationsPerType<ClampedIdx(7)>(Types).Objects[Entry.TypeInnerIdx] = &RawPtr->template AsChecked<typename std::tuple_element<ClampedIdx(7), std::tuple<TSpatialAccelerationTypes...>>::type>(); break;
+					}
+				}
+			}
+		}
+	}
+
 private:
 
-	TSpatialAccelerationCollection(const TSpatialAccelerationCollection<TSpatialAccelerationTypes...>& Other)
+	TSpatialAccelerationCollection(const TSpatialAccelerationCollection<TParticleView, TSpatialAccelerationTypes...>& Other)
 		: NumBuckets(Other.NumBuckets)
 	{
 		//need to make a deep copy of the unique ptrs and then update raw ptrs
@@ -604,6 +653,7 @@ private:
 	TSpatialTypeTuple< TSpatialAccelerationTypes...> Types;
 	static constexpr uint32 NumTypes = sizeof...(TSpatialAccelerationTypes);
 	uint16 NumBuckets;
+	const TArray<TSpatialAccelerationParams<TParticleView, TPayloadType, T, d>> BucketsParams;
 };
 
 
diff --git a/Engine/Source/Runtime/Experimental/Chaos/Public/Chaos/PBDRigidsEvolution.h b/Engine/Source/Runtime/Experimental/Chaos/Public/Chaos/PBDRigidsEvolution.h
--- a/Engine/Source/Runtime/Experimental/Chaos/Public/Chaos/PBDRigidsEvolution.h
+++ b/Engine/Source/Runtime/Experimental/Chaos/Public/Chaos/PBDRigidsEvolution.h
@@ -22,6 +22,18 @@ class FChaosArchive;
 template <typename TPayload, typename T, int d>
 class ISpatialAccelerationCollection;
 
+template <typename T, int d>
+struct CHAOS_API ISpatialAccelerationCollectionFactory
+{
+	//Create an empty acceleration collection with the desired buckets. Chaos uses the bucket params inside the collection to add / remove sub-structures as needed
+	virtual TUniquePtr<ISpatialAccelerationCollection<TAccelerationStructureHandle<T, d>, T, d>> CreateEmptyCollection() = 0;
+
+	//Serialize the collection in and out
+	virtual void Serialize(TUniquePtr<ISpatialAccelerationCollection<TAccelerationStructureHandle<T, d>, T, d>>& Ptr, FChaosArchive& Ar) = 0;
+
+	virtual ~ISpatialAccelerationCollectionFactory() = default;
+};
+
 struct CHAOS_API FEvolutionStats
 {
 	int32 ActiveCollisionPoints;
@@ -52,6 +64,32 @@ struct CHAOS_API FEvolutionStats
 	}
 };
 
+/** Used for building an acceleration structure out of cached bounds and payloads */
+template <typename T, int d>
+struct TAccelerationStructureBuilder
+{
+	//todo: should these be arrays instead? might make it easier in some cases
+	bool bHasBoundingBox;
+	TBox<T, d> CachedSpatialBounds;
+	TAccelerationStructureHandle<T, d> CachedSpatialPayload;
+
+	const TBox<T, d>& BoundingBox() const
+	{
+		return CachedSpatialBounds;
+	}
+
+	bool HasBoundingBox() const
+	{
+		return bHasBoundingBox;
+	}
+
+	template <typename TPayloadType>
+	TAccelerationStructureHandle<T, d> GetPayload(int32 Idx) const
+	{
+		return CachedSpatialPayload;
+	}
+};
+
 template<class FPBDRigidsEvolution, class FPBDCollisionConstraint, class T, int d>
 class TPBDRigidsEvolutionBase
 {
@@ -422,43 +460,17 @@ protected:
 		}
 	}
 
-
-	/** Used for building an acceleration structure out of cached bounds and payloads */
-	struct FAccelerationStructureBuilder
-	{
-		//todo: should these be arrays instead? might make it easier in some cases
-		bool bHasBoundingBox;
-		TBox<T, d> CachedSpatialBounds;
-		TAccelerationStructureHandle<T,d> CachedSpatialPayload;
-
-		const TBox<T, d>& BoundingBox() const
-		{
-			return CachedSpatialBounds;
-		}
-
-		bool HasBoundingBox() const
-		{
-			return bHasBoundingBox;
-		}
-
-		template <typename TPayloadType>
-		TAccelerationStructureHandle<T, d> GetPayload(int32 Idx) const
-		{
-			return CachedSpatialPayload;
-		}
-	};
-
 	/** Used for async acceleration rebuild */
 	TMap<TGeometryParticleHandle<T, d>*, int32> ParticleToCacheIdx;
 
 	struct FAccelerationStructBuilderCache
 	{
 		FSpatialAccelerationIdx SpatialIdx;
-		TUniquePtr <TArray<FAccelerationStructureBuilder>> CachedSpatialBuilderData;
+		TUniquePtr <TArray<TAccelerationStructureBuilder<T, d>>> CachedSpatialBuilderData;
 
 		FAccelerationStructBuilderCache(FSpatialAccelerationIdx Idx)
 			: SpatialIdx(Idx)
-			, CachedSpatialBuilderData(MakeUnique<TArray<FAccelerationStructureBuilder>>())
+			, CachedSpatialBuilderData(MakeUnique<TArray<TAccelerationStructureBuilder<T, d>>>())
 		{}
 	};
 	TArray<FAccelerationStructBuilderCache> CachedSpatialBuilderDataMap;
@@ -482,8 +494,8 @@ protected:
 	FGraphEventRef AccelerationStructureTaskComplete;
 
 	int32 NumIterations;
+	TUniquePtr<ISpatialAccelerationCollectionFactory<T, d>> SpatialCollectionFactory;
 
-	static ISpatialAccelerationCollection<TAccelerationStructureHandle<T, d>, T, d>* CreateNewSpatialStructure();
 	void InitializeAccelerationCache();
 };
 
