Chaos - Kinematic Targets

diff --git a/Engine/Source/Runtime/Experimental/Chaos/Private/Chaos/PBDCollisionConstraint.cpp b/Engine/Source/Runtime/Experimental/Chaos/Private/Chaos/PBDCollisionConstraint.cpp
--- a/Engine/Source/Runtime/Experimental/Chaos/Private/Chaos/PBDCollisionConstraint.cpp
+++ b/Engine/Source/Runtime/Experimental/Chaos/Private/Chaos/PBDCollisionConstraint.cpp
@@ -521,8 +521,8 @@ TVector<T, d> GetEnergyClampedImpulse(const TRigidBodyContactConstraint<T, d>& C
 template<typename T, int d>
 void TPBDCollisionConstraint<T, d>::Apply(const T Dt, FRigidBodyContactConstraint& Constraint)
 {
-	TGeometryParticleHandle<T, d>* Particle0 = Constraint.Particle;
-	TGeometryParticleHandle<T, d>* Particle1 = Constraint.Levelset;
+	TGenericParticleHandle<T, d> Particle0 = TGenericParticleHandle<T, d>(Constraint.Particle);
+	TGenericParticleHandle<T, d> Particle1 = TGenericParticleHandle<T, d>(Constraint.Levelset);
 	TPBDRigidParticleHandle<T, d>* PBDRigid0 = Particle0->AsDynamic();
 	TPBDRigidParticleHandle<T, d>* PBDRigid1 = Particle1->AsDynamic();
 
@@ -549,14 +549,14 @@ void TPBDCollisionConstraint<T, d>::Apply(const T Dt, FRigidBodyContactConstrain
 
 	// @todo(ccaulfield): CHAOS_PARTICLEHANDLE_TODO split function to avoid ifs
 	const TVector<T, d> ZeroVector = TVector<T, d>(0);
-	const TRotation<T, d>& Q0 = PBDRigid0? PBDRigid0->Q() : Particle0->R();
-	const TRotation<T, d>& Q1 = PBDRigid1? PBDRigid1->Q() : Particle1->R();
-	const TVector<T, d>& P0 = PBDRigid0 ? PBDRigid0->P() : Particle0->X();
-	const TVector<T, d>& P1 = PBDRigid1 ? PBDRigid1->P() : Particle1->X();
-	const TVector<T, d>& V0 = PBDRigid0 ? PBDRigid0->V() : ZeroVector;
-	const TVector<T, d>& V1 = PBDRigid1 ? PBDRigid1->V() : ZeroVector;
-	const TVector<T, d>& W0 = PBDRigid0 ? PBDRigid0->W() : ZeroVector;
-	const TVector<T, d>& W1 = PBDRigid1 ? PBDRigid1->W() : ZeroVector;
+	const TRotation<T, d>& Q0 = Particle0->Q();
+	const TRotation<T, d>& Q1 = Particle1->Q();
+	const TVector<T, d>& P0 = Particle0->P();
+	const TVector<T, d>& P1 = Particle1->P();
+	const TVector<T, d>& V0 = Particle0->V();
+	const TVector<T, d>& V1 = Particle1->V();
+	const TVector<T, d>& W0 = Particle0->W();
+	const TVector<T, d>& W1 = Particle1->W();
 	TSerializablePtr<TChaosPhysicsMaterial<T>> PhysicsMaterial0 = Particle0->AuxilaryValue(MPhysicsMaterials);
 	TSerializablePtr<TChaosPhysicsMaterial<T>> PhysicsMaterial1 = Particle1->AuxilaryValue(MPhysicsMaterials);
 
diff --git a/Engine/Source/Runtime/Experimental/Chaos/Private/Chaos/PBDRigidsEvolutionGBF.cpp b/Engine/Source/Runtime/Experimental/Chaos/Private/Chaos/PBDRigidsEvolutionGBF.cpp
--- a/Engine/Source/Runtime/Experimental/Chaos/Private/Chaos/PBDRigidsEvolutionGBF.cpp
+++ b/Engine/Source/Runtime/Experimental/Chaos/Private/Chaos/PBDRigidsEvolutionGBF.cpp
@@ -42,6 +42,7 @@ FAutoConsoleVariableRef CVarDisableThreshold(TEXT("p.DisableThreshold2"), Disabl
 
 DECLARE_CYCLE_STAT(TEXT("TPBDRigidsEvolutionGBF::AdvanceOneTimeStep"), STAT_AdvanceOneTimeStep, STATGROUP_Chaos);
 DECLARE_CYCLE_STAT(TEXT("TPBDRigidsEvolutionGBF::Integrate"), STAT_Integrate, STATGROUP_Chaos);
+DECLARE_CYCLE_STAT(TEXT("TPBDRigidsEvolutionGBF::KinematicTargets"), STAT_KinematicTargets, STATGROUP_Chaos);
 DECLARE_CYCLE_STAT(TEXT("TPBDRigidsEvolutionGBF::UpdateConstraintPositionBasedState"), STAT_UpdateConstraintPositionBasedState, STATGROUP_Chaos);
 DECLARE_CYCLE_STAT(TEXT("TPBDRigidsEvolutionGBF::CreateConstraintGraph"), STAT_CreateConstraintGraph, STATGROUP_Chaos);
 DECLARE_CYCLE_STAT(TEXT("TPBDRigidsEvolutionGBF::CreateIslands"), STAT_CreateIslands, STATGROUP_Chaos);
@@ -100,6 +101,11 @@ void TPBDRigidsEvolutionGBF<T, d>::AdvanceOneTimeStep(T Dt)
 		Integrate(Particles.GetNonDisabledDynamicView(), Dt);	//Question: should we use an awake view?
 	}
 
+	{
+		SCOPE_CYCLE_COUNTER(STAT_KinematicTargets);
+		ApplyKinematicTargets(Dt);
+	}
+
 	if (PostIntegrateCallback != nullptr)
 	{
 		PostIntegrateCallback();
diff --git a/Engine/Source/Runtime/Experimental/Chaos/Public/Chaos/KinematicGeometryParticles.h b/Engine/Source/Runtime/Experimental/Chaos/Public/Chaos/KinematicGeometryParticles.h
--- a/Engine/Source/Runtime/Experimental/Chaos/Public/Chaos/KinematicGeometryParticles.h
+++ b/Engine/Source/Runtime/Experimental/Chaos/Public/Chaos/KinematicGeometryParticles.h
@@ -6,6 +6,59 @@
 
 namespace Chaos
 {
+	/**
+	 * Controls how a kinematic body is integrated each Evolution Advance
+	 */
+	enum class EKinematicTargetMode
+	{
+		None,			/** Particle does not move and no data is changed */
+		Zero,			/** Particle does not move, velocity and angular velocity are zeroed, then mode is set to "None". */
+		Position,		/** Particle is moved to Kinematic Target transform, velocity and angular velocity updated to reflect the change, then mode is reset to "Zero". */
+		Velocity,		/** Particle is moved based on velocity and angular velocity, mode remains as "Velocity". */
+	};
+
+	/**
+	 * Data used to integrate kinematic bodies
+	 */
+	template<class T, int d>
+	class CHAOS_API TKinematicTarget
+	{
+	public:
+		TKinematicTarget()
+			: Mode(EKinematicTargetMode::None)
+		{
+		}
+
+		/** Whether this kinematic target has been set (either velocity or position mode) */
+		bool IsSet() const { return (Mode == EKinematicTargetMode::Position) || (Mode == EKinematicTargetMode::Velocity); }
+
+		/** Get the kinematic target mode */
+		EKinematicTargetMode GetMode() const { return Mode; }
+
+		/** Get the target transform (asserts if not in Position mode) */
+		const TRigidTransform<T, d>& GetTarget() const { check(Mode == EKinematicTargetMode::Position); return Target; }
+
+		/** Use transform target mode and set the transform target */
+		void SetTargetMode(const TRigidTransform<T, d>& InTarget) { Target = InTarget;  Mode = EKinematicTargetMode::Position; }
+
+		/** Use velocity target mode */
+		void SetVelocityMode() { Mode = EKinematicTargetMode::Velocity; }
+
+		// For internal use only
+		void SetMode(EKinematicTargetMode InMode) { Mode = InMode; }
+
+		friend FChaosArchive& operator<<(FChaosArchive& Ar, TKinematicTarget<T, d>& KinematicTarget)
+		{
+			Ar << KinematicTarget.Target << KinematicTarget.Mode;
+			return Ar;
+		}
+
+	private:
+		TRigidTransform<T, d> Target;
+		EKinematicTargetMode Mode;
+	};
+
+
 template<class T, int d, EGeometryParticlesSimType SimType>
 class TKinematicGeometryParticlesImp : public TGeometryParticlesImp<T, d, SimType>
 {
@@ -16,6 +69,7 @@ class TKinematicGeometryParticlesImp : public TGeometryParticlesImp<T, d, SimTyp
 		this->MParticleType = EParticleType::Kinematic;
 		TArrayCollection::AddArray(&MV);
 		TArrayCollection::AddArray(&MW);
+		TArrayCollection::AddArray(&KinematicTargets);
 	}
 	TKinematicGeometryParticlesImp(const TKinematicGeometryParticlesImp<T, d, SimType>& Other) = delete;
 	CHAOS_API TKinematicGeometryParticlesImp(TKinematicGeometryParticlesImp<T, d, SimType>&& Other)
@@ -24,6 +78,7 @@ class TKinematicGeometryParticlesImp : public TGeometryParticlesImp<T, d, SimTyp
 		this->MParticleType = EParticleType::Kinematic;
 		TArrayCollection::AddArray(&MV);
 		TArrayCollection::AddArray(&MW);
+		TArrayCollection::AddArray(&KinematicTargets);
 	}
 	CHAOS_API virtual ~TKinematicGeometryParticlesImp();
 
@@ -33,6 +88,9 @@ class TKinematicGeometryParticlesImp : public TGeometryParticlesImp<T, d, SimTyp
 	const TVector<T, d>& W(const int32 Index) const { return MW[Index]; }
 	TVector<T, d>& W(const int32 Index) { return MW[Index]; }
 
+	const TKinematicTarget<T, d>& KinematicTarget(const int32 Index) const { return KinematicTargets[Index]; }
+	TKinematicTarget<T, d>& KinematicTarget(const int32 Index) { return KinematicTargets[Index]; }
+
 	FString ToString(int32 index) const
 	{
 		FString BaseString = TGeometryParticlesImp<T, d, SimType>::ToString(index);
@@ -48,12 +106,13 @@ class TKinematicGeometryParticlesImp : public TGeometryParticlesImp<T, d, SimTyp
 	CHAOS_API virtual void Serialize(FChaosArchive& Ar) override
 	{
 		TGeometryParticlesImp<T, d, SimType>::Serialize(Ar);
-		Ar << MV << MW;
+		Ar << MV << MW << KinematicTargets;
 	}
 
   private:
 	TArrayCollectionArray<TVector<T, d>> MV;
 	TArrayCollectionArray<TVector<T, d>> MW;
+	TArrayCollectionArray<TKinematicTarget<T, d>> KinematicTargets;
 };
 
 template <typename T, int d, EGeometryParticlesSimType SimType>
diff --git a/Engine/Source/Runtime/Experimental/Chaos/Public/Chaos/PBDRigidsEvolution.h b/Engine/Source/Runtime/Experimental/Chaos/Public/Chaos/PBDRigidsEvolution.h
--- a/Engine/Source/Runtime/Experimental/Chaos/Public/Chaos/PBDRigidsEvolution.h
+++ b/Engine/Source/Runtime/Experimental/Chaos/Public/Chaos/PBDRigidsEvolution.h
@@ -296,6 +296,58 @@ class TPBDRigidsEvolutionBase
 		}
 	}
 
+	void ApplyKinematicTargets(T Dt)
+	{
+		// @todo(ccaulfield): optimize. Depending on the number of kinematics relative to the number that have 
+		// targets set, it may be faster to process a command list rather than iterate over them all each frame. 
+		const T MinDt = 1e-6f;
+		for (auto& Particle : Particles.GetActiveKinematicParticlesView())
+		{
+			TKinematicTarget<T, d>& KinematicTarget = Particle.KinematicTarget();
+			switch (KinematicTarget.GetMode())
+			{
+			case EKinematicTargetMode::None:
+				// Nothing to do
+				break;
+
+			case EKinematicTargetMode::Zero:
+			{
+				// Reset velocity and then switch to do-nothing mode
+				Particle.V() = TVector<T, d>(0, 0, 0);
+				Particle.W() = TVector<T, d>(0, 0, 0);
+				KinematicTarget.SetMode(EKinematicTargetMode::None);
+				break;
+			}
+
+			case EKinematicTargetMode::Position:
+			{
+				// Move to kinematic target and update velocities to match
+				// Target positions only need to be processed once, and we reset the velocity next frame (if no new target is set)
+				if (Dt > MinDt)
+				{
+					TVector<float, 3> V = TVector<float, 3>::CalculateVelocity(Particle.X(), KinematicTarget.GetTarget().GetLocation(), Dt);
+					Particle.V() = V;
+
+					TVector<float, 3> W = TRotation<float, 3>::CalculateAngularVelocity(Particle.R(), KinematicTarget.GetTarget().GetRotation(), Dt);
+					Particle.W() = W;
+				}
+				Particle.X() = KinematicTarget.GetTarget().GetTranslation();
+				Particle.R() = KinematicTarget.GetTarget().GetRotation();
+				KinematicTarget.SetMode(EKinematicTargetMode::Zero);
+				break;
+			}
+
+			case EKinematicTargetMode::Velocity:
+			{
+				// Move based on velocity
+				Particle.X() = Particle.X() + Particle.V() * Dt;
+				Particle.R() = TRotation<T, d>::IntegrateRotationWithAngularVelocity(Particle.R(), Particle.W(), Dt);
+				break;
+			}
+			}
+		}
+	}
+
 	/** Make a copy of the acceleration structure to allow for external modification. This is needed for supporting sync operations on SQ structure from game thread */
 	CHAOS_API void UpdateExternalAccelerationStructure(TUniquePtr<ISpatialAccelerationCollection<TAccelerationStructureHandle<T, d>, T, d>>& ExternalStructure);
 	ISpatialAccelerationCollection<TAccelerationStructureHandle<T, d>, T, d>* GetSpatialAcceleration() { return InternalAcceleration.Get(); }
diff --git a/Engine/Source/Runtime/Experimental/Chaos/Public/Chaos/PBDRigidsEvolutionGBF.h b/Engine/Source/Runtime/Experimental/Chaos/Public/Chaos/PBDRigidsEvolutionGBF.h
--- a/Engine/Source/Runtime/Experimental/Chaos/Public/Chaos/PBDRigidsEvolutionGBF.h
+++ b/Engine/Source/Runtime/Experimental/Chaos/Public/Chaos/PBDRigidsEvolutionGBF.h
@@ -39,6 +39,7 @@ public:
 	using Base::Particles;
 	using Base::ForceRules;
 	using Base::ParticleUpdatePosition;
+	using Base::ApplyKinematicTargets;
 	using Base::SetParticleUpdateVelocityFunction;
 	using Base::SetParticleUpdatePositionFunction;
 	using Base::AddConstraintRule;
diff --git a/Engine/Source/Runtime/Experimental/Chaos/Public/Chaos/PBDRigidsSOAs.h b/Engine/Source/Runtime/Experimental/Chaos/Public/Chaos/PBDRigidsSOAs.h
--- a/Engine/Source/Runtime/Experimental/Chaos/Public/Chaos/PBDRigidsSOAs.h
+++ b/Engine/Source/Runtime/Experimental/Chaos/Public/Chaos/PBDRigidsSOAs.h
@@ -236,6 +236,9 @@ public:
 
 	const TParticleView<TGeometryParticles<T, d>>& GetAllParticlesView() const { return AllParticlesView; }
 
+	const TParticleView<TKinematicGeometryParticles<T, d>>& GetActiveKinematicParticlesView() const { return ActiveKinematicParticlesView; }
+	TParticleView<TKinematicGeometryParticles<T, d>>& GetActiveKinematicParticlesView() { return ActiveKinematicParticlesView; }
+
 	const TGeometryParticleHandles<T, d>& GetParticleHandles() const { return ParticleHandles; }
 	TGeometryParticleHandles<T, d>& GetParticleHandles() { return ParticleHandles; }
 
@@ -346,12 +349,15 @@ private:
 			TArray<TSOAView<TPBDRigidParticles<T, d>>> TmpArray = { {&ActiveParticlesArray} };
 			ActiveParticlesView = MakeParticleView(MoveTemp(TmpArray));
 		}
-
 		{
 			TArray<TSOAView<TGeometryParticles<T, d>>> TmpArray = { StaticParticles.Get(), StaticDisabledParticles.Get(), KinematicParticles.Get(), KinematicDisabledParticles.Get(),
 				DynamicParticles.Get(), DynamicDisabledParticles.Get(), ClusteredParticles.Get() };
 			AllParticlesView = MakeParticleView(MoveTemp(TmpArray));
 		}
+		{
+			TArray<TSOAView<TKinematicGeometryParticles<T, d>>> TmpArray = { KinematicParticles.Get() };
+			ActiveKinematicParticlesView = MakeParticleView(MoveTemp(TmpArray));
+		}
 	}
 
 	//Organized by SOA type
@@ -377,10 +383,11 @@ private:
 	TArray<TPBDRigidClusteredParticleHandle<T, d>*> NonDisabledClusteredArray;
 
 	//Particle Views
-	TParticleView<TGeometryParticles<T, d>> NonDisabledView;	//all particles that are not disabled
-	TParticleView<TPBDRigidParticles<T, d>> NonDisabledDynamicView;	//all dynamic particles that are not disabled
-	TParticleView<TPBDRigidParticles<T, d>> ActiveParticlesView;	//all particles that are active
-	TParticleView<TGeometryParticles<T, d>> AllParticlesView;	//all particles
+	TParticleView<TGeometryParticles<T, d>> NonDisabledView;							//all particles that are not disabled
+	TParticleView<TPBDRigidParticles<T, d>> NonDisabledDynamicView;						//all dynamic particles that are not disabled
+	TParticleView<TPBDRigidParticles<T, d>> ActiveParticlesView;						//all particles that are active
+	TParticleView<TGeometryParticles<T, d>> AllParticlesView;							//all particles
+	TParticleView<TKinematicGeometryParticles<T, d>> ActiveKinematicParticlesView;		//all kinematic particles that are not disabled
 
 	//Auxiliary data synced with particle handles
 	TGeometryParticleHandles<T, d> ParticleHandles;
diff --git a/Engine/Source/Runtime/Experimental/Chaos/Public/Chaos/ParticleHandle.h b/Engine/Source/Runtime/Experimental/Chaos/Public/Chaos/ParticleHandle.h
--- a/Engine/Source/Runtime/Experimental/Chaos/Public/Chaos/ParticleHandle.h
+++ b/Engine/Source/Runtime/Experimental/Chaos/Public/Chaos/ParticleHandle.h
@@ -508,6 +508,9 @@ public:
 	TVector<T, d>& W() { return KinematicGeometryParticles->W(ParticleIdx); }
 	void SetW(const TVector<T, d>& InW) { KinematicGeometryParticles->W(ParticleIdx) = InW; }
 
+	const TKinematicTarget<T, d>& KinematicTarget() const { return KinematicGeometryParticles->KinematicTarget(ParticleIdx); }
+	TKinematicTarget<T, d>& KinematicTarget() { return KinematicGeometryParticles->KinematicTarget(ParticleIdx); }
+
 	//Really only useful when using a transient handle
 	const TKinematicGeometryParticleHandleImp<T, d, true>* Handle() const { return KinematicGeometryParticles->Handle(ParticleIdx); }
 	TKinematicGeometryParticleHandleImp<T, d, true>* Handle() { return KinematicGeometryParticles->Handle(ParticleIdx); }
@@ -794,11 +797,23 @@ public:
 		bool IsKinematic() const { return (Handle->ObjectState() == EObjectStateType::Kinematic); }
 		bool IsDynamic() const { return (Handle->ObjectState() == EObjectStateType::Dynamic) || (Handle->ObjectState() == EObjectStateType::Sleeping); }
 
+		const TKinematicGeometryParticleHandle<T, d>* AsKinematic() const { return Handle->AsKinematic(); }
+		TKinematicGeometryParticleHandle<T, d>* AsKinematic() { return Handle->AsKinematic(); }
+		const TPBDRigidParticleHandle<T, d>* AsDynamic() const { return Handle->AsDynamic(); }
+		TPBDRigidParticleHandle<T, d>* AsDynamic() { return Handle->AsDynamic(); }
+
 		// Static Particles
 		const TVector<T, d>& X() const { return Handle->X(); }
 		const TRotation<T, d>& R() const { return Handle->R(); }
 		TSerializablePtr<TImplicitObject<T, d>> Geometry() const { return Handle->Geometry(); }
 		const TUniquePtr<TImplicitObject<T, d>>& DynamicGeometry() const { return Handle->DynamicGeometry(); }
+		bool Sleeping() const { return Handle->Sleeping(); }
+		FString ToString() const { return Handle->ToString(); }
+
+		template <typename Container>
+		const auto& AuxilaryValue(const Container& AuxContainer) const { return Handle->AuxilaryValue(AuxContainer); }
+		template <typename Container>
+		auto& AuxilaryValue(Container& AuxContainer) { return Handle->AuxilaryValue(AuxContainer); }
 
 		// Kinematic Particles
 		const TVector<T, d>& V() const { return (Handle->AsKinematic()) ? Handle->AsKinematic()->V() : ZeroVector; }
diff --git a/Engine/Source/Runtime/Experimental/ChaosCore/Public/Chaos/Vector.h b/Engine/Source/Runtime/Experimental/ChaosCore/Public/Chaos/Vector.h
--- a/Engine/Source/Runtime/Experimental/ChaosCore/Public/Chaos/Vector.h
+++ b/Engine/Source/Runtime/Experimental/ChaosCore/Public/Chaos/Vector.h
@@ -525,6 +525,12 @@ namespace Chaos
 			float c = DotProduct(V1, V2);
 			return atan2(s, c);
 		}
+		/** Calculate the velocity to move from P0 to P1 in time Dt. Exists just for symmetry with TRotation::CalculateAngularVelocity! */
+		static TVector<float, 3> CalculateVelocity(const TVector<float, 3>& P0, const TVector<float, 3>& P1, const float Dt)
+		{
+			return (P1 - P0) / Dt;
+		}
+
 	};
 
 	template<>
