Fix CIS - explicit instance of TKinematicTarget<float, 3>

diff --git a/Engine/Source/Runtime/Experimental/Chaos/Private/Chaos/KinematicGeometryParticles.cpp b/Engine/Source/Runtime/Experimental/Chaos/Private/Chaos/KinematicGeometryParticles.cpp
--- a/Engine/Source/Runtime/Experimental/Chaos/Private/Chaos/KinematicGeometryParticles.cpp
+++ b/Engine/Source/Runtime/Experimental/Chaos/Private/Chaos/KinematicGeometryParticles.cpp
@@ -32,9 +32,11 @@ namespace Chaos
 #ifdef __clang__
 template class CHAOS_API Chaos::TKinematicGeometryParticlesImp<float, 3, Chaos::EGeometryParticlesSimType::RigidBodySim>;
 template class CHAOS_API Chaos::TKinematicGeometryParticlesImp<float, 3, Chaos::EGeometryParticlesSimType::Other>;
+template class CHAOS_API Chaos::TKinematicTarget<float, 3>;
 #else
 template class Chaos::TKinematicGeometryParticlesImp<float, 3, Chaos::EGeometryParticlesSimType::RigidBodySim>;
 template class Chaos::TKinematicGeometryParticlesImp<float, 3, Chaos::EGeometryParticlesSimType::Other>;
+template class Chaos::TKinematicTarget<float, 3>;
 #endif
 
 // Re-enable when double precision is able to compile
