Chaos - add kinematic target serialization version

diff --git a/Engine/Source/Runtime/Core/Public/UObject/ExternalPhysicsCustomObjectVersion.h b/Engine/Source/Runtime/Core/Public/UObject/ExternalPhysicsCustomObjectVersion.h
--- a/Engine/Source/Runtime/Core/Public/UObject/ExternalPhysicsCustomObjectVersion.h
+++ b/Engine/Source/Runtime/Core/Public/UObject/ExternalPhysicsCustomObjectVersion.h
@@ -36,6 +36,9 @@ struct CORE_API FExternalPhysicsCustomObjectVersion
 		//Save out multiple acceleration structures
 		SerializeMultiStructures,
 
+		// Add kinematic targets to TKinematicGeometryParticles
+		KinematicTargets,
+
 		// -----<new versions can be added above this line>-------------------------------------------------
 		VersionPlusOne,
 		LatestVersion = VersionPlusOne - 1
diff --git a/Engine/Source/Runtime/Experimental/Chaos/Public/Chaos/KinematicGeometryParticles.h b/Engine/Source/Runtime/Experimental/Chaos/Public/Chaos/KinematicGeometryParticles.h
--- a/Engine/Source/Runtime/Experimental/Chaos/Public/Chaos/KinematicGeometryParticles.h
+++ b/Engine/Source/Runtime/Experimental/Chaos/Public/Chaos/KinematicGeometryParticles.h
@@ -106,7 +106,13 @@ class TKinematicGeometryParticlesImp : public TGeometryParticlesImp<T, d, SimTyp
 	CHAOS_API virtual void Serialize(FChaosArchive& Ar) override
 	{
 		TGeometryParticlesImp<T, d, SimType>::Serialize(Ar);
-		Ar << MV << MW << KinematicTargets;
+		Ar << MV << MW;
+		
+		Ar.UsingCustomVersion(FExternalPhysicsCustomObjectVersion::GUID);
+		if (Ar.CustomVer(FExternalPhysicsCustomObjectVersion::GUID) >= FExternalPhysicsCustomObjectVersion::KinematicTargets)
+		{
+			Ar << KinematicTargets;
+		}
 	}
 
   private:
