Merging //UE4/Main@9388041 to Dev-RenderPlat-Staging

diff --git a/Engine/Source/Runtime/Engine/Public/Physics/ImmediatePhysics/ImmediatePhysicsChaos/ImmediatePhysicsJointHandle_Chaos.h b/Engine/Source/Runtime/Engine/Public/Physics/ImmediatePhysics/ImmediatePhysicsChaos/ImmediatePhysicsJointHandle_Chaos.h
--- a/Engine/Source/Runtime/Engine/Public/Physics/ImmediatePhysics/ImmediatePhysicsChaos/ImmediatePhysicsJointHandle_Chaos.h
+++ b/Engine/Source/Runtime/Engine/Public/Physics/ImmediatePhysics/ImmediatePhysicsChaos/ImmediatePhysicsJointHandle_Chaos.h
@@ -2,9 +2,10 @@
 
 #pragma once
 
-#if INCLUDE_CHAOS
-
 #include "Physics/ImmediatePhysics/ImmediatePhysicsChaos/ImmediatePhysicsCore_Chaos.h"
+#include "Physics/ImmediatePhysics/ImmediatePhysicsChaos/ImmediatePhysicsActorHandle_Chaos.h"
+
+#include "Chaos/Vector.h"
 
 #include "Engine/EngineTypes.h"
 
@@ -14,15 +15,23 @@ namespace ImmediatePhysics_Chaos
 	struct ENGINE_API FJointHandle
 	{
 	public:
-		using FChaosConstraintContainer = Chaos::TPBD6DJointConstraints<FReal, Dimensions>;
-		using FChaosConstraintHandle = typename Chaos::TPBD6DJointConstraintHandle<FReal, Dimensions>;
+		using FChaosConstraintContainer = Chaos::TPBDJointConstraints<FReal, Dimensions>;
+		using FChaosConstraintHandle = typename Chaos::TPBDJointConstraintHandle<FReal, Dimensions>;
 
-		FJointHandle(FChaosConstraintContainer* InConstraints, FConstraintInstance* ConstraintInstance, FActorHandle* Actor1, FActorHandle* Actor2);
+		FJointHandle(FChaosConstraintContainer* InConstraints, FConstraintInstance* ConstraintInstance, FActorHandle* InActor1, FActorHandle* InActor2);
 		~FJointHandle();
 
+		FChaosConstraintHandle* GetConstraint();
+		const FChaosConstraintHandle* GetConstraint() const;
+
+		const Chaos::TVector<FActorHandle*, 2>& GetActorHandles();
+		const Chaos::TVector<const FActorHandle*, 2>& GetActorHandles() const;
+
+		void UpdateLevels();
+
 	private:
+		Chaos::TVector<FActorHandle*, 2> ActorHandles;
 		FChaosConstraintContainer* Constraints;
 		FChaosConstraintHandle* ConstraintHandle;
 	};
 }
-#endif
diff --git a/Engine/Source/Runtime/Engine/Public/Physics/ImmediatePhysics/ImmediatePhysicsChaos/ImmediatePhysicsSimulation_Chaos.h b/Engine/Source/Runtime/Engine/Public/Physics/ImmediatePhysics/ImmediatePhysicsChaos/ImmediatePhysicsSimulation_Chaos.h
--- a/Engine/Source/Runtime/Engine/Public/Physics/ImmediatePhysics/ImmediatePhysicsChaos/ImmediatePhysicsSimulation_Chaos.h
+++ b/Engine/Source/Runtime/Engine/Public/Physics/ImmediatePhysics/ImmediatePhysicsChaos/ImmediatePhysicsSimulation_Chaos.h
@@ -2,10 +2,9 @@
 
 #pragma once
 
-#if INCLUDE_CHAOS
-
 #include "Physics/ImmediatePhysics/ImmediatePhysicsChaos/ImmediatePhysicsCore_Chaos.h"
 
+#include "Chaos/ChaosDebugDrawDeclares.h"
 #include "Engine/EngineTypes.h"
 #include "Templates/UniquePtr.h"
 
@@ -77,11 +76,19 @@ namespace ImmediatePhysics_Chaos
 		//void SetPositionIterationCount(int32 InIterationCount);
 		//void SetVelocityIterationCount(int32 InIterationCount);
 
+		void SetSimulationSpaceTransform(const FTransform& Transform) { SimulationSpaceTransform = Transform; }
+
 	private:
+		void ConditionConstraints();
+
+		void DebugDrawElementsAll(const int32 MinDebugLevel, const int32 MaxDebugLevel, const float ColorScale);
+		void DebugDrawElementsIsland(const int32 Island, const int32 MinDebugLevel, const int32 MaxDebugLevel, const float ColorScale);
+
+		using FParticleHandle = Chaos::TGeometryParticleHandle<FReal, Dimensions>;
 		TUniquePtr<Chaos::TPBDRigidsEvolutionGBF<FReal, Dimensions>> Evolution;
 		TUniquePtr<Chaos::TPBDRigidsSOAs<FReal, Dimensions>> Particles;
-		TUniquePtr<Chaos::TPBD6DJointConstraints<FReal, Dimensions>> Joints;
-		TUniquePtr<Chaos::TPBDConstraintIslandRule<Chaos::TPBD6DJointConstraints<FReal, Dimensions>, FReal, Dimensions>> JointsRule;		// @todo(ccaulfield): Don't need islands for anim node physics...
+		TUniquePtr<Chaos::TPBDJointConstraints<FReal, Dimensions>> Joints;
+		TUniquePtr<Chaos::TPBDConstraintIslandRule<Chaos::TPBDJointConstraints<FReal, Dimensions>, FReal, Dimensions>> JointsRule;		// @todo(ccaulfield): Don't need islands for anim node physics...
 
 		/** Mapping from entity index to handle */
 		// @todo(ccaulfield): we now have handles pointing to handles which is inefficient - we can do better than this, but don't want to change API yet
@@ -91,13 +98,11 @@ namespace ImmediatePhysics_Chaos
 		/** Mapping from constraint index to handle */
 		TArray<FJointHandle*> JointHandles;
 
-		/** Both of these are slow to access. Make sure to use iteration cache when possible */
-		TMap<FActorHandle*, TSet<FActorHandle*>> IgnoreCollisionPairTable;
-		TSet<FActorHandle*> IgnoreCollisionActors;
+		/** Slow to access. */
+		// @todo(ccaulfield): Optimize
+		TMap<const FParticleHandle*, TSet<const FParticleHandle*>> IgnoreCollisionParticlePairTable;
 
-		FVector Gravity;
+		FTransform SimulationSpaceTransform;
 	};
 
 }
-
-#endif // INCLUDE_CHAOS
\ No newline at end of file
diff --git a/Engine/Source/Runtime/Engine/Public/Physics/PhysicsInterfacePhysX.h b/Engine/Source/Runtime/Engine/Public/Physics/PhysicsInterfacePhysX.h
--- a/Engine/Source/Runtime/Engine/Public/Physics/PhysicsInterfacePhysX.h
+++ b/Engine/Source/Runtime/Engine/Public/Physics/PhysicsInterfacePhysX.h
@@ -172,6 +172,7 @@ struct ENGINE_API FPhysicsCommand_PhysX
 
 	// Executes with appropriate write locking, return true if execution took place (actor was valid)
 	static bool ExecuteWrite(const FPhysicsActorHandle_PhysX& InHandle, TFunctionRef<void(const FPhysicsActorHandle_PhysX& Actor)> InCallable);
+	static bool ExecuteWrite(FPhysicsActorHandle_PhysX& InHandle, TFunctionRef<void(FPhysicsActorHandle_PhysX& Actor)> InCallable);
 	static bool ExecuteWrite(USkeletalMeshComponent* InMeshComponent, TFunctionRef<void()> InCallable);
 	static bool ExecuteWrite(const FPhysicsActorHandle_PhysX& InHandleA, const FPhysicsActorHandle_PhysX& InHandleB, TFunctionRef<void(const FPhysicsActorHandle_PhysX& ActorA, const FPhysicsActorHandle_PhysX& ActorB)> InCallable);
 	static bool ExecuteWrite(const FPhysicsConstraintHandle_PhysX& InHandle, TFunctionRef<void(const FPhysicsConstraintHandle_PhysX& Constraint)> InCallable);
diff --git a/Engine/Source/Runtime/Experimental/Chaos/Chaos.Build.cs b/Engine/Source/Runtime/Experimental/Chaos/Chaos.Build.cs
--- a/Engine/Source/Runtime/Experimental/Chaos/Chaos.Build.cs
+++ b/Engine/Source/Runtime/Experimental/Chaos/Chaos.Build.cs
@@ -28,6 +28,15 @@ namespace UnrealBuildTool.Rules
             {
                 PublicDefinitions.Add("INCLUDE_CHAOS=0");
             }
+
+			if (Target.bUseChaosMemoryTracking == true)
+			{
+				PublicDefinitions.Add("CHAOS_MEMORY_TRACKING=1");
+			}
+			else
+			{
+				PublicDefinitions.Add("CHAOS_MEMORY_TRACKING=0");
+			}
         }
     }
 }
diff --git a/Engine/Source/Runtime/Experimental/Chaos/Private/Chaos/BoundingVolume.cpp b/Engine/Source/Runtime/Experimental/Chaos/Private/Chaos/BoundingVolume.cpp
--- a/Engine/Source/Runtime/Experimental/Chaos/Private/Chaos/BoundingVolume.cpp
+++ b/Engine/Source/Runtime/Experimental/Chaos/Private/Chaos/BoundingVolume.cpp
@@ -1,12 +1,39 @@
 // Copyright 1998-2019 Epic Games, Inc. All Rights Reserved.
 
 #include "Chaos/BoundingVolume.h"
+#include "Chaos/AABBTree.h"
+#include "UObject/ExternalPhysicsCustomObjectVersion.h"
 
-int FBoundingVolumeCVars::FilterFarBodies = 1;
+int FBoundingVolumeCVars::FilterFarBodies = 0;
 
 FAutoConsoleVariableRef FBoundingVolumeCVars::CVarFilterFarBodies(
 	TEXT("p.RemoveFarBodiesFromBVH"),
 	FBoundingVolumeCVars::FilterFarBodies,
 	TEXT("Removes bodies far from the scene from the bvh\n")
 	TEXT("0: Kept, 1: Removed"),
-	ECVF_Default);
\ No newline at end of file
+	ECVF_Default);
+
+namespace Chaos
+{
+	template <typename TPayloadType, typename T, int d>
+	ISpatialAcceleration<TPayloadType, T, d>* ISpatialAcceleration<TPayloadType, T, d>::SerializationFactory(FChaosArchive& Ar, ISpatialAcceleration<TPayloadType, T, d>* Accel)
+	{
+		if (Ar.CustomVer(FExternalPhysicsCustomObjectVersion::GUID) < FExternalPhysicsCustomObjectVersion::SerializeEvolutionGenericAcceleration)
+		{
+			return new TBoundingVolume<TPayloadType, T, d>();
+		}
+
+		int8 AccelType = Ar.IsLoading() ? 0 : (int8)Accel->Type;
+		Ar << AccelType;
+		switch ((ESpatialAcceleration)AccelType)
+		{
+		case ESpatialAcceleration::BoundingVolume: return Ar.IsLoading() ? new TBoundingVolume<TPayloadType, T, d>() : nullptr;
+		case ESpatialAcceleration::AABBTree: return Ar.IsLoading() ? new TAABBTree<TPayloadType, TAABBTreeLeafArray<TPayloadType, T>, T>() : nullptr;
+		case ESpatialAcceleration::AABBTreeBV: return Ar.IsLoading() ? new TAABBTree<TPayloadType, TBoundingVolume<TPayloadType, T, 3>, T>() : nullptr;
+		default: check(false); return nullptr;
+		}
+	}
+
+	template class CHAOS_API Chaos::ISpatialAcceleration<int32, float, 3>;
+	template class CHAOS_API Chaos::ISpatialAcceleration<TAccelerationStructureHandle<float,3>, float, 3>;
+}
\ No newline at end of file
diff --git a/Engine/Source/Runtime/Experimental/Chaos/Private/Chaos/BoundingVolumeHierarchy.cpp b/Engine/Source/Runtime/Experimental/Chaos/Private/Chaos/BoundingVolumeHierarchy.cpp
--- a/Engine/Source/Runtime/Experimental/Chaos/Private/Chaos/BoundingVolumeHierarchy.cpp
+++ b/Engine/Source/Runtime/Experimental/Chaos/Private/Chaos/BoundingVolumeHierarchy.cpp
@@ -73,7 +73,7 @@ TArray<int32> MakeNewLeaf(const OBJECT_ARRAY& Objects, const TArray<int32>& AllO
 }
 
 template<class OBJECT_ARRAY, class T, int d, typename TPayloadType>
-TBoundingVolume<OBJECT_ARRAY, TPayloadType, T, d> MakeNewLeaf(const OBJECT_ARRAY& Objects, const TArray<int32>& AllObjects, const TBoundingVolume<OBJECT_ARRAY, TPayloadType, T, d>* /*Type*/)
+TBoundingVolume<TPayloadType, T, d> MakeNewLeaf(const OBJECT_ARRAY& Objects, const TArray<int32>& AllObjects, const TBoundingVolume<TPayloadType, T, d>* /*Type*/)
 {
 #if CHAOS_PARTICLEHANDLE_TODO
 	// @todo(mlentine): This is pretty dirty but should work for now.
@@ -82,7 +82,7 @@ TBoundingVolume<OBJECT_ARRAY, TPayloadType, T, d> MakeNewLeaf(const OBJECT_ARRAY
 	TArray<TSOAView<OBJECT_ARRAY>> SOAs;
 	SOAs.Add(TSOAView<OBJECT_ARRAY>(const_cast<OBJECT_ARRAY*>(&Objects)));	//todo: this const_cast is here because bvh is holding on to things as const but giving it out as non const. Not sure if we should change API or not
 	//return TBoundingVolume<OBJECT_ARRAY, T, d>(MakeParticleView<OBJECT_ARRAY>({ {&Objects} }));
-	return TBoundingVolume<OBJECT_ARRAY, TPayloadType, T, d>(MakeParticleView<OBJECT_ARRAY>(MoveTemp(SOAs)));
+	return TBoundingVolume<TPayloadType, T, d>(MakeParticleView<OBJECT_ARRAY>(MoveTemp(SOAs)));
 #endif
 }
 
@@ -188,7 +188,7 @@ TArray<int32> FindAllIntersectionsLeafHelper(const TArray<int32>& Leaf, const TV
 }
 
 template<class OBJECT_ARRAY,typename TPayloadType, class T, int d>
-TArray<int32> FindAllIntersectionsLeafHelper(const TBoundingVolume<OBJECT_ARRAY, TPayloadType, T, d>& Leaf, const TVector<T, d>& Point)
+TArray<int32> FindAllIntersectionsLeafHelper(const TBoundingVolume<TPayloadType, T, d>& Leaf, const TVector<T, d>& Point)
 {
 #if CHAOS_PARTICLEHANDLE_TODO
 	return Leaf.FindAllIntersections(Point);
@@ -297,7 +297,7 @@ TArray<int32> FindAllIntersectionsLeafHelper(const TArray<int32>& Leaf, const QU
 }
 
 template<class OBJECT_ARRAY, typename TPayloadType, class T, int d, typename QUERY_OBJECT>
-TArray<int32> FindAllIntersectionsLeafHelper(const TBoundingVolume<OBJECT_ARRAY, TPayloadType, T, d>& Leaf, const QUERY_OBJECT& QueryObject)
+TArray<int32> FindAllIntersectionsLeafHelper(const TBoundingVolume<TPayloadType, T, d>& Leaf, const QUERY_OBJECT& QueryObject)
 {
 #if CHAOS_PARTICLEHANDLE_TODO
 	return Leaf.FindAllIntersectionsImp(QueryObject);
@@ -654,7 +654,7 @@ int32 TBoundingVolumeHierarchy<OBJECT_ARRAY, LEAF_TYPE, T, d>::GenerateNextLevel
 namespace Chaos
 {
 template <typename OBJECT_ARRAY, typename TPayloadType, typename T, int d>
-void FixupLeafObj(const OBJECT_ARRAY& Objects, TArray<TBoundingVolume<OBJECT_ARRAY, TPayloadType, T, d>>& Leafs)
+void FixupLeafObj(const OBJECT_ARRAY& Objects, TArray<TBoundingVolume<TPayloadType, T, d>>& Leafs)
 {
 #if CHAOS_PARTICLEHANDLE_TODO
 	for (TBoundingVolume<OBJECT_ARRAY, TPayloadType, T, d>& Leaf : Leafs)
@@ -680,10 +680,10 @@ void TBoundingVolumeHierarchy<OBJECT_ARRAY, LEAF_TYPE, T, d>::Serialize(FArchive
 }
 }
 
-template class CHAOS_API Chaos::TBoundingVolume<TArray<Chaos::TSphere<float, 3>*>, int32, float, 3>;
+template class CHAOS_API Chaos::TBoundingVolume<int32, float, 3>;
 template class CHAOS_API Chaos::TBoundingVolumeHierarchy<TArray<Chaos::TSphere<float, 3>*>, TArray<int32>, float, 3>;
 template class CHAOS_API Chaos::TBoundingVolumeHierarchy<Chaos::TPBDRigidParticles<float, 3>, TArray<int32>, float, 3>;
 template class CHAOS_API Chaos::TBoundingVolumeHierarchy<Chaos::TParticles<float, 3>, TArray<int32>, float, 3>;
 template class CHAOS_API Chaos::TBoundingVolumeHierarchy<Chaos::TGeometryParticles<float, 3>, TArray<int32>, float, 3>;
-template class CHAOS_API Chaos::TBoundingVolumeHierarchy<Chaos::TPBDRigidParticles<float, 3>, TBoundingVolume<Chaos::TPBDRigidParticles<float, 3>, TPBDRigidParticleHandle<float,3>*, float, 3>, float, 3>;
-template class CHAOS_API Chaos::TBoundingVolumeHierarchy<Chaos::TGeometryParticles<float, 3>, TBoundingVolume<Chaos::TGeometryParticles<float, 3>, TGeometryParticleHandle<float,3>*, float, 3>, float, 3>;
+template class CHAOS_API Chaos::TBoundingVolumeHierarchy<Chaos::TPBDRigidParticles<float, 3>, TBoundingVolume<TPBDRigidParticleHandle<float,3>*, float, 3>, float, 3>;
+template class CHAOS_API Chaos::TBoundingVolumeHierarchy<Chaos::TGeometryParticles<float, 3>, TBoundingVolume<TGeometryParticleHandle<float,3>*, float, 3>, float, 3>;
diff --git a/Engine/Source/Runtime/Experimental/Chaos/Private/Chaos/ChaosArchive.cpp b/Engine/Source/Runtime/Experimental/Chaos/Private/Chaos/ChaosArchive.cpp
--- a/Engine/Source/Runtime/Experimental/Chaos/Private/Chaos/ChaosArchive.cpp
+++ b/Engine/Source/Runtime/Experimental/Chaos/Private/Chaos/ChaosArchive.cpp
@@ -10,4 +10,62 @@ void FChaosArchive::SerializeLegacy(TUniquePtr<TImplicitObject<float, 3>>& Obj)
 {
 	TImplicitObject<float, 3>::SerializeLegacyHelper(InnerArchive, Obj);
 }
+
+#if CHAOS_MEMORY_TRACKING
+const FName FChaosArchiveMemoryTrackingContext::UntrackedName = TEXT("Untracked");
+
+void FChaosArchiveMemoryTrackingContext::PushSection(const FName& SectionName, const int64 MemoryLocation, const bool bAbsorbChildren)
+{
+	check(SectionStack.Num() == 0 || SectionName != SectionStack.Top().Name);
+	SectionStack.Push({ SectionName, MemoryLocation, 0, bAbsorbChildren });
+	ChildAbsorbers += (int32)bAbsorbChildren;
+}
+
+void FChaosArchiveMemoryTrackingContext::PopSection(const int64 MemoryLocation)
+{
+	// Remove & get the top section
+	const FChaosArchiveSection Section = SectionStack.Pop(false);
+
+	// Decrement the number of absorbers if the top section was absorbing subsections
+	ChildAbsorbers -= (int32)Section.bAbsorbChildren;
+
+	// Measure the size of this archive section
+	const int64 SizeInclusive = MemoryLocation - Section.MemoryLocation;
+	const int64 SizeExclusive = SizeInclusive - Section.ChildMemory;
+
+	// Only count this section if it was at the bottom of the stack,
+	// to avoid double counting.
+	if (ChildAbsorbers == 0)
+	{
+		// Update this section's data
+		FChaosArchiveSectionData& SectionData = SectionMap.FindOrAdd(Section.Name);
+		SectionData.Count += 1;
+		SectionData.SizeInclusive += SizeInclusive;
+		SectionData.SizeExclusive += SizeExclusive;
+
+		// Remove the size of this child from all parent sections
+		for (FChaosArchiveSection& Parent : SectionStack)
+		{
+			Parent.ChildMemory += SizeExclusive;
+		}
+	}
+}
+
+void FChaosArchiveMemoryTrackingContext::BeginSerialize(const int64 MemoryLocation)
+{
+	if (SectionStack.Num() == 0)
+	{
+		PushSection(UntrackedName, MemoryLocation, false);
+	}
+}
+
+void FChaosArchiveMemoryTrackingContext::EndSerialize(const int64 MemoryLocation)
+{
+	if (SectionStack.Top().Name == UntrackedName)
+	{
+		PopSection(MemoryLocation);
+	}
+}
+#endif
+
 }
\ No newline at end of file
diff --git a/Engine/Source/Runtime/Experimental/Chaos/Private/Chaos/ChaosDebugDraw.cpp b/Engine/Source/Runtime/Experimental/Chaos/Private/Chaos/ChaosDebugDraw.cpp
new file mode 100644
--- /dev/null
+++ b/Engine/Source/Runtime/Experimental/Chaos/Private/Chaos/ChaosDebugDraw.cpp
@@ -0,0 +1,370 @@
+// Copyright 1998-2019 Epic Games, Inc. All Rights Reserved.
+#include "Chaos/ChaosDebugDraw.h"
+#include "Chaos/Box.h"
+#include "Chaos/Capsule.h"
+#include "Chaos/DebugDrawQueue.h"
+#include "Chaos/ImplicitObjectScaled.h"
+#include "Chaos/ImplicitObjectTransformed.h"
+#include "Chaos/ImplicitObjectUnion.h"
+#include "Chaos/ParticleHandle.h"
+#include "Chaos/PBD6DJointConstraints.h"
+#include "Chaos/PBDCollisionConstraint.h"
+#include "Chaos/PBDJointConstraints.h"
+#include "Chaos/PBDRigidParticles.h"
+#include "Chaos/Sphere.h"
+#include "Chaos/Utilities.h"
+
+namespace Chaos
+{
+	namespace DebugDraw
+	{
+		float ArrowSize = 1.5f;
+		float BodyAxisLen = 12.0f;
+		float ContactLen = 6.0f;
+		float ContactWidth = 2.0f;
+		float ContactPhiWidth = 1.5f;
+		float ConstraintAxisLen = 5.0f;
+		float LineThickness = 0.25f;
+		int DrawPriority = 10.0f;
+		float DrawScale = 1.0f;
+		float FontHeight = 10.0f;
+		float FontScale = 1.0f;
+
+		FAutoConsoleVariableRef CVarArrowSize(TEXT("p.Chaos.DebugDrawArrowSize"), ArrowSize, TEXT("ArrowSize."));
+		FAutoConsoleVariableRef CVarBodyAxisLen(TEXT("p.Chaos.DebugDrawBodyAxisLen"), BodyAxisLen, TEXT("BodyAxisLen."));
+		FAutoConsoleVariableRef CVarContactLen(TEXT("p.Chaos.DebugDrawContactLen"), ContactLen, TEXT("ContactLen."));
+		FAutoConsoleVariableRef CVarContactWidth(TEXT("p.Chaos.DebugDrawContactWidth"), ContactWidth, TEXT("ContactWidth."));
+		FAutoConsoleVariableRef CVarContactPhiWidth(TEXT("p.Chaos.DebugDrawContactPhiWidth"), ContactPhiWidth, TEXT("ContactPhiWidth."));
+		FAutoConsoleVariableRef CVarConstraintAxisLen(TEXT("p.Chaos.DebugDrawConstraintAxisLen"), ConstraintAxisLen, TEXT("ConstraintAxisLen."));
+		FAutoConsoleVariableRef CVarLineThickness(TEXT("p.Chaos.DebugDrawLineThickness"), LineThickness, TEXT("LineThickness."));
+		FAutoConsoleVariableRef CVarScale(TEXT("p.Chaos.DebugDrawScale"), DrawScale, TEXT("Scale applied to all Chaos Debug Draw line lengths etc."));
+
+		//
+		//
+		//
+
+#if CHAOS_DEBUG_DRAW
+
+		template <typename T, int d>
+		void DrawShapesImpl(const TRigidTransform<float, 3>& ShapeTransform, const TImplicitObject<T, d>* Shape, FColor Color)
+		{
+			// @todo(ccaulfield): handle scale throughout
+			switch (Shape->GetType(false))
+			{
+			case ImplicitObjectType::Sphere:
+			{
+				const TSphere<T, d>* Sphere = Shape->template GetObject<TSphere<T, d>>();
+				FDebugDrawQueue::GetInstance().DrawDebugSphere(ShapeTransform.GetLocation(), Sphere->GetRadius(), 20, Color, false, KINDA_SMALL_NUMBER, DrawPriority, LineThickness);
+				break;
+			}
+			case ImplicitObjectType::Box:
+			{
+				const TBox<T, d>* Box = Shape->template GetObject<TBox<T, d>>();
+				FDebugDrawQueue::GetInstance().DrawDebugBox(ShapeTransform.GetLocation(), (T)0.5 * Box->Extents(), ShapeTransform.GetRotation(), Color, false, KINDA_SMALL_NUMBER, DrawPriority, LineThickness);
+				break;
+			}
+			case ImplicitObjectType::Plane:
+				break;
+			case ImplicitObjectType::Capsule:
+			{
+				const TCapsule<T>* Capsule = Shape->template GetObject<TCapsule<T>>();
+				const TVector<T, d> P = ShapeTransform.TransformPosition(Capsule->GetCenter());
+				const TRotation<T, d> Q = ShapeTransform.GetRotation() * FRotationMatrix::MakeFromZ(Capsule->GetAxis());
+				FDebugDrawQueue::GetInstance().DrawDebugCapsule(P, (T)0.5 * Capsule->GetHeight() + Capsule->GetRadius(), Capsule->GetRadius(), Q, Color, false, KINDA_SMALL_NUMBER, DrawPriority, LineThickness);
+				break;
+			}
+			case ImplicitObjectType::Transformed:
+			{
+				const TImplicitObjectTransformed<T, d>* Transformed = Shape->template GetObject<TImplicitObjectTransformed<T, d>>();
+				TRigidTransform<float, 3> TransformedTransform = TRigidTransform<float, 3>(ShapeTransform.TransformPosition(Transformed->GetTransform().GetLocation()), ShapeTransform.GetRotation() * Transformed->GetTransform().GetRotation());
+				DrawShapesImpl<T, d>(TransformedTransform, Transformed->GetTransformedObject(), Color);
+				break;
+			}
+			case ImplicitObjectType::Union:
+			{
+				const TImplicitObjectUnion<T, d>* Union = Shape->template GetObject<TImplicitObjectUnion<T, d>>();
+				for (auto& UnionShape : Union->GetObjects())
+				{
+					DrawShapesImpl<T, d>(ShapeTransform, UnionShape.Get(), Color);
+				}
+				break;
+			}
+			case ImplicitObjectType::LevelSet:
+				break;
+			case ImplicitObjectType::Unknown:
+				break;
+			case ImplicitObjectType::Convex:
+				break;
+			case ImplicitObjectType::TaperedCylinder:
+				break;
+			case ImplicitObjectType::Cylinder:
+				break;
+			case ImplicitObjectType::TriangleMesh:
+				break;
+			case ImplicitObjectType::HeightField:
+				break;
+			case ImplicitObjectType::Scaled:
+				break;
+			default:
+				break;
+			}
+		}
+
+		template <typename T, int d>
+		void DrawParticleShapesImpl(const TRigidTransform<float, 3>& SpaceTransform, const TGeometryParticleHandle<T, d>* Particle, T ColorScale)
+		{
+			FColor Color = (ColorScale * FColor::White).ToFColor(false);
+			TVector<T, d> P = SpaceTransform.TransformPosition(Particle->AsDynamic() ? Particle->AsDynamic()->P() : Particle->X());
+			TRotation<T, d> Q = SpaceTransform.GetRotation() * (Particle->AsDynamic() ? Particle->AsDynamic()->Q() : Particle->R());
+
+			DrawShapesImpl<T, d>(TRigidTransform<T, d>(P, Q), Particle->Geometry().Get(), Color);
+		}
+
+		template <typename T, int d>
+		void DrawParticleTransformImpl(const TRigidTransform<float, 3>& SpaceTransform, const TGeometryParticleHandle<T, d>* Particle, T ColorScale)
+		{
+			FColor R = (ColorScale * FColor::Red).ToFColor(false);
+			FColor G = (ColorScale * FColor::Green).ToFColor(false);
+			FColor B = (ColorScale * FColor::Blue).ToFColor(false);
+
+			TVector<T, d> P = SpaceTransform.TransformPosition((Particle->AsDynamic()) ? Particle->AsDynamic()->P() : Particle->X());
+			TRotation<T, d> Q = SpaceTransform.GetRotation() * ((Particle->AsDynamic()) ? Particle->AsDynamic()->Q() : Particle->R());
+			PMatrix<T, d, d> Qm = Q.ToMatrix();
+			FDebugDrawQueue::GetInstance().DrawDebugDirectionalArrow(P, P + DrawScale * BodyAxisLen * Qm.GetAxis(0), DrawScale * ArrowSize, R, false, KINDA_SMALL_NUMBER, DrawPriority, LineThickness);
+			FDebugDrawQueue::GetInstance().DrawDebugDirectionalArrow(P, P + DrawScale * BodyAxisLen * Qm.GetAxis(1), DrawScale * ArrowSize, G, false, KINDA_SMALL_NUMBER, DrawPriority, LineThickness);
+			FDebugDrawQueue::GetInstance().DrawDebugDirectionalArrow(P, P + DrawScale * BodyAxisLen * Qm.GetAxis(2), DrawScale * ArrowSize, B, false, KINDA_SMALL_NUMBER, DrawPriority, LineThickness);
+		}
+
+		template <typename T, int d>
+		void DrawCollisionImpl(const TRigidTransform<float, 3>& SpaceTransform, const TPBDCollisionConstraintHandle<float, 3>* ConstraintHandle, float ColorScale)
+		{
+			const TRigidBodyContactConstraint<T, d>& Contact = ConstraintHandle->GetContact();
+			if (Contact.Phi > 0)
+			{
+				ColorScale = ColorScale * (T)0.1;
+			}
+
+			TVector<T, d> Location = SpaceTransform.TransformPosition(Contact.Location);
+			TVector<T, d> Normal = SpaceTransform.TransformVector(Contact.Normal);
+
+			if (ContactWidth > 0)
+			{
+				FColor C0 = (ColorScale * FColor(128, 0, 0)).ToFColor(false);
+				FMatrix Axes = FRotationMatrix::MakeFromX(Normal);
+				FDebugDrawQueue::GetInstance().DrawDebugCircle(Location, DrawScale * ContactWidth, 12, C0, false, KINDA_SMALL_NUMBER, DrawPriority, LineThickness, Axes.GetUnitAxis(EAxis::Y), Axes.GetUnitAxis(EAxis::Z), false);
+			}
+			if (ContactLen > 0)
+			{
+				FColor C1 = (ColorScale * FColor(255, 0, 0)).ToFColor(false);
+				FDebugDrawQueue::GetInstance().DrawDebugLine(Location, Location + DrawScale * ContactLen * Normal, C1, false, KINDA_SMALL_NUMBER, DrawPriority, LineThickness);
+			}
+			if (ContactPhiWidth > 0 && Contact.Phi < FLT_MAX)
+			{
+				FColor C2 = (ColorScale * FColor(128, 128, 0)).ToFColor(false);
+				FMatrix Axes = FRotationMatrix::MakeFromX(Normal);
+				FDebugDrawQueue::GetInstance().DrawDebugCircle(Location - Contact.Phi * Normal, DrawScale * ContactPhiWidth, 12, C2, false, KINDA_SMALL_NUMBER, DrawPriority, LineThickness, Axes.GetUnitAxis(EAxis::Y), Axes.GetUnitAxis(EAxis::Z), false);
+			}
+		}
+
+		template<typename T, int d>
+		void DrawJointConstraintImpl(const TRigidTransform<float, 3>& SpaceTransform, const TVector<T, d>& InXa, const PMatrix<T, d, d>& Ra, const TVector<T, d>& InXb, const PMatrix<T, d, d>& Rb, const TVector<T, d>& CR, T ColorScale)
+		{
+			using namespace Chaos::DebugDraw;
+			FColor R = (ColorScale * FColor::Red).ToFColor(false);
+			FColor G = (ColorScale * FColor::Green).ToFColor(false);
+			FColor B = (ColorScale * FColor::Blue).ToFColor(false);
+			FColor C = (ColorScale * FColor::Cyan).ToFColor(false);
+			FColor M = (ColorScale * FColor::Magenta).ToFColor(false);
+			FColor Y = (ColorScale * FColor::Yellow).ToFColor(false);
+			TVector<T, d> Xa = SpaceTransform.TransformPosition(InXa);
+			TVector<T, d> Xb = SpaceTransform.TransformPosition(InXb);
+			FDebugDrawQueue::GetInstance().DrawDebugDirectionalArrow(Xa, Xa + DrawScale * ConstraintAxisLen * SpaceTransform.TransformVector(Ra.GetAxis(0)), DrawScale * ArrowSize, R, false, KINDA_SMALL_NUMBER, DrawPriority, LineThickness);
+			FDebugDrawQueue::GetInstance().DrawDebugDirectionalArrow(Xa, Xa + DrawScale * ConstraintAxisLen * SpaceTransform.TransformVector(Ra.GetAxis(1)), DrawScale * ArrowSize, G, false, KINDA_SMALL_NUMBER, DrawPriority, LineThickness);
+			FDebugDrawQueue::GetInstance().DrawDebugDirectionalArrow(Xa, Xa + DrawScale * ConstraintAxisLen * SpaceTransform.TransformVector(Ra.GetAxis(2)), DrawScale * ArrowSize, B, false, KINDA_SMALL_NUMBER, DrawPriority, LineThickness);
+			FDebugDrawQueue::GetInstance().DrawDebugDirectionalArrow(Xb, Xb + DrawScale * ConstraintAxisLen * SpaceTransform.TransformVector(Rb.GetAxis(0)), DrawScale * ArrowSize, C, false, KINDA_SMALL_NUMBER, DrawPriority, LineThickness);
+			FDebugDrawQueue::GetInstance().DrawDebugDirectionalArrow(Xb, Xb + DrawScale * ConstraintAxisLen * SpaceTransform.TransformVector(Rb.GetAxis(1)), DrawScale * ArrowSize, M, false, KINDA_SMALL_NUMBER, DrawPriority, LineThickness);
+			FDebugDrawQueue::GetInstance().DrawDebugDirectionalArrow(Xb, Xb + DrawScale * ConstraintAxisLen * SpaceTransform.TransformVector(Rb.GetAxis(2)), DrawScale * ArrowSize, Y, false, KINDA_SMALL_NUMBER, DrawPriority, LineThickness);
+			//FDebugDrawQueue::GetInstance().DrawDebugString(Xb + 3 * FontHeight * TVector<T, d>(0, 0, 1), FString::Format(TEXT("T=  {0}"), { FMath::RadiansToDegrees(CR[(int32)E6DJointAngularConstraintIndex::Twist]) }), nullptr, FColor::Red, KINDA_SMALL_NUMBER, false, FontScale);
+			//FDebugDrawQueue::GetInstance().DrawDebugString(Xb + 2 * FontHeight * TVector<T, d>(0, 0, 1), FString::Format(TEXT("S1= {0}"), { FMath::RadiansToDegrees(CR[(int32)E6DJointAngularConstraintIndex::Swing1]) }), nullptr, FColor::Red, KINDA_SMALL_NUMBER, false, FontScale);
+			//FDebugDrawQueue::GetInstance().DrawDebugString(Xb + 1 * FontHeight * TVector<T, d>(0, 0, 1), FString::Format(TEXT("S2= {0}"), { FMath::RadiansToDegrees(CR[(int32)E6DJointAngularConstraintIndex::Swing2]) }), nullptr, FColor::Red, KINDA_SMALL_NUMBER, false, FontScale);
+		}
+
+		template<typename T, int d>
+		void DrawJointConstraintImpl(const TRigidTransform<float, 3>& SpaceTransform, const TPBDJointConstraintHandle<float, 3>* ConstraintHandle, T ColorScale)
+		{
+			TVector<TGeometryParticleHandle<T, d>*, 2> ConstrainedParticles = ConstraintHandle->GetConstrainedParticles();
+			if (ConstrainedParticles[0]->AsDynamic() || ConstrainedParticles[1]->AsDynamic())
+			{
+				TVector<T, d> Xa, Xb, CR;
+				PMatrix<T, d, d> Ra, Rb;
+				ConstraintHandle->CalculateConstraintSpace(Xa, Ra, Xb, Rb, CR);
+				DrawJointConstraintImpl<T, d>(SpaceTransform, Xa, Ra, Xb, Rb, CR, ColorScale);
+			}
+		}
+
+		template<typename T, int d>
+		void Draw6DofConstraintImpl(const TRigidTransform<float, 3>& SpaceTransform, const TPBD6DJointConstraintHandle<float, 3>* ConstraintHandle, T ColorScale)
+		{
+			TVector<T, d> Xa, Xb, CR;
+			PMatrix<T, d, d> Ra, Rb;
+			ConstraintHandle->CalculateConstraintSpace(Xa, Ra, Xb, Rb, CR);
+			DrawJointConstraintImpl<T, d>(SpaceTransform, Xa, Ra, Xb, Rb, CR, ColorScale);
+		}
+
+#endif
+
+		//
+		//
+		//
+
+		void DrawParticleShapes(const TRigidTransform<float, 3>& SpaceTransform, const TParticleView<TGeometryParticles<float, 3>>& ParticlesView, float ColorScale)
+		{
+#if CHAOS_DEBUG_DRAW
+			if (FDebugDrawQueue::IsDebugDrawingEnabled())
+			{
+				for (auto& Particle : ParticlesView)
+				{
+					DrawParticleShapesImpl<float, 3>(SpaceTransform, GetHandleHelper<float, 3>(&Particle), ColorScale);
+				}
+			}
+#endif
+		}
+
+		void DrawParticleShapes(const TRigidTransform<float, 3>& SpaceTransform, const TArray<TGeometryParticleHandle<float, 3>*>& Particles, float ColorScale)
+		{
+#if CHAOS_DEBUG_DRAW
+			if (FDebugDrawQueue::IsDebugDrawingEnabled())
+			{
+				for (auto& Particle : Particles)
+				{
+					DrawParticleShapesImpl<float, 3>(SpaceTransform, Particle, ColorScale);
+				}
+			}
+#endif
+		}
+
+		void DrawParticleTransforms(const TRigidTransform<float, 3>& SpaceTransform, const TParticleView<TGeometryParticles<float, 3>>& ParticlesView, float ColorScale)
+		{
+#if CHAOS_DEBUG_DRAW
+			if (FDebugDrawQueue::IsDebugDrawingEnabled())
+			{
+				for (auto& Particle : ParticlesView)
+				{
+					DrawParticleTransformImpl<float, 3>(SpaceTransform, GetHandleHelper<float, 3>(&Particle), ColorScale);
+				}
+			}
+#endif
+		}
+
+		void DrawParticleTransforms(const TRigidTransform<float, 3>& SpaceTransform, const TArray<TGeometryParticleHandle<float, 3>*>& Particles, float ColorScale)
+		{
+#if CHAOS_DEBUG_DRAW
+			if (FDebugDrawQueue::IsDebugDrawingEnabled())
+			{
+				for (auto& Particle : Particles)
+				{
+					DrawParticleTransformImpl<float, 3>(SpaceTransform, Particle, ColorScale);
+				}
+			}
+#endif
+		}
+
+		void DrawParticleCollisions(const TRigidTransform<float, 3>& SpaceTransform, const TGeometryParticleHandle<float, 3>* Particle, const TPBDCollisionConstraint<float, 3>& Collisions)
+		{
+#if CHAOS_DEBUG_DRAW
+			if (FDebugDrawQueue::IsDebugDrawingEnabled())
+			{
+				for (int32 ConstraintIndex = 0; ConstraintIndex < Collisions.NumConstraints(); ++ConstraintIndex)
+				{
+					TVector<TGeometryParticleHandle<float, 3>*, 2> ConstrainedParticles = Collisions.GetConstrainedParticles(ConstraintIndex);
+					if ((ConstrainedParticles[0] == Particle) || (ConstrainedParticles[1] == Particle))
+					{
+						const typename TPBDCollisionConstraint<float, 3>::FConstraintHandle* ConstraintHandle = Collisions.GetConstraintHandle(ConstraintIndex);
+						DrawCollisionImpl<float, 3>(SpaceTransform, ConstraintHandle, 1.0f);
+					}
+				}
+			}
+#endif
+		}
+
+		void DrawCollisions(const TRigidTransform<float, 3>& SpaceTransform, const TPBDCollisionConstraint<float, 3>& Collisions, float ColorScale)
+		{
+#if CHAOS_DEBUG_DRAW
+			if (FDebugDrawQueue::IsDebugDrawingEnabled())
+			{
+				for (int32 ConstraintIndex = 0; ConstraintIndex < Collisions.NumConstraints(); ++ConstraintIndex)
+				{
+					DrawCollisionImpl<float, 3>(SpaceTransform, Collisions.GetConstraintHandle(ConstraintIndex), ColorScale);
+				}
+			}
+#endif
+		}
+
+		void DrawCollisions(const TRigidTransform<float, 3>& SpaceTransform, const TArray<TPBDCollisionConstraintHandle<float, 3>*>& ConstraintHandles, float ColorScale)
+		{
+#if CHAOS_DEBUG_DRAW
+			if (FDebugDrawQueue::IsDebugDrawingEnabled())
+			{
+				for (const TPBDCollisionConstraintHandle<float, 3>* ConstraintHandle : ConstraintHandles)
+				{
+					DrawCollisionImpl<float, 3>(SpaceTransform, ConstraintHandle, ColorScale);
+				}
+			}
+#endif
+		}
+
+
+		void DrawJointConstraints(const TRigidTransform<float, 3>& SpaceTransform, const TArray<TPBDJointConstraintHandle<float, 3>*>& ConstraintHandles, float ColorScale)
+		{
+#if CHAOS_DEBUG_DRAW
+			if (FDebugDrawQueue::IsDebugDrawingEnabled())
+			{
+				for (const TPBDJointConstraintHandle<float, 3>* ConstraintHandle : ConstraintHandles)
+				{
+					DrawJointConstraintImpl<float, 3>(SpaceTransform, ConstraintHandle, ColorScale);
+				}
+			}
+#endif
+		}
+
+		void DrawJointConstraints(const TRigidTransform<float, 3>& SpaceTransform, const TPBDJointConstraints<float, 3>& Constraints, float ColorScale)
+		{
+#if CHAOS_DEBUG_DRAW
+			if (FDebugDrawQueue::IsDebugDrawingEnabled())
+			{
+				for (int32 ConstraintIndex = 0; ConstraintIndex < Constraints.NumConstraints(); ++ConstraintIndex)
+				{
+					DrawJointConstraintImpl<float, 3>(SpaceTransform, Constraints.GetConstraintHandle(ConstraintIndex), ColorScale);
+				}
+			}
+#endif
+		}
+
+		void Draw6DofConstraints(const TRigidTransform<float, 3>& SpaceTransform, const TArray<TPBD6DJointConstraintHandle<float, 3>*>& ConstraintHandles, float ColorScale)
+		{
+#if CHAOS_DEBUG_DRAW
+			if (FDebugDrawQueue::IsDebugDrawingEnabled())
+			{
+				for (const TPBD6DJointConstraintHandle<float, 3>* ConstraintHandle : ConstraintHandles)
+				{
+					Draw6DofConstraintImpl<float, 3>(SpaceTransform, ConstraintHandle, ColorScale);
+				}
+			}
+#endif
+		}
+
+		void Draw6DofConstraints(const TRigidTransform<float, 3>& SpaceTransform, const TPBD6DJointConstraints<float, 3>& Constraints, float ColorScale)
+		{
+#if CHAOS_DEBUG_DRAW
+			if (FDebugDrawQueue::IsDebugDrawingEnabled())
+			{
+				for (int32 ConstraintIndex = 0; ConstraintIndex < Constraints.NumConstraints(); ++ConstraintIndex)
+				{
+					Draw6DofConstraintImpl<float, 3>(SpaceTransform, Constraints.GetConstraintHandle(ConstraintIndex), ColorScale);
+				}
+			}
+#endif
+		}
+
+	}
+}
diff --git a/Engine/Source/Runtime/Experimental/Chaos/Private/Chaos/DebugDrawQueue.cpp b/Engine/Source/Runtime/Experimental/Chaos/Private/Chaos/DebugDrawQueue.cpp
--- a/Engine/Source/Runtime/Experimental/Chaos/Private/Chaos/DebugDrawQueue.cpp
+++ b/Engine/Source/Runtime/Experimental/Chaos/Private/Chaos/DebugDrawQueue.cpp
@@ -1,10 +1,37 @@
 // Copyright 1998-2019 Epic Games, Inc. All Rights Reserved.
 #include "Chaos/DebugDrawQueue.h"
 #include "HAL/IConsoleManager.h"
+#include "Misc/ScopeLock.h"
 
 #if CHAOS_DEBUG_DRAW
 using namespace Chaos;
 
 int32 FDebugDrawQueue::EnableDebugDrawing = 0;
-FAutoConsoleVariableRef CVarEnableDebugDrawingChaos(TEXT("p.ChaosDebugDrawing"), FDebugDrawQueue::EnableDebugDrawing, TEXT("Whether to debug draw low level physics solver information"));
+
+int32 FDebugDrawQueue::NumConsumers = 0;
+
+FAutoConsoleVariableRef CVarEnableDebugDrawingChaos = FDebugDrawQueue::MakeCVarRef();
+
+void FDebugDrawQueue::SetConsumerActive(void* Consumer, bool bConsumerActive)
+{
+	FScopeLock Lock(&ConsumersCS);
+
+	if (bConsumerActive)
+	{
+		Consumers.AddUnique(Consumer);
+	}
+	else
+	{
+		Consumers.Remove(Consumer);
+	}
+
+	NumConsumers = Consumers.Num();
+}
+
+
+FAutoConsoleVariableRef FDebugDrawQueue::MakeCVarRef()
+{
+	return FAutoConsoleVariableRef(TEXT("p.Chaos.DebugDrawing"), EnableDebugDrawing, TEXT("Whether to debug draw low level physics solver information"));
+}
+
 #endif
diff --git a/Engine/Source/Runtime/Experimental/Chaos/Private/Chaos/GeometryParticles.cpp b/Engine/Source/Runtime/Experimental/Chaos/Private/Chaos/GeometryParticles.cpp
--- a/Engine/Source/Runtime/Experimental/Chaos/Private/Chaos/GeometryParticles.cpp
+++ b/Engine/Source/Runtime/Experimental/Chaos/Private/Chaos/GeometryParticles.cpp
@@ -3,11 +3,12 @@
 #include "Chaos/GeometryParticles.h"
 #include "Chaos/ImplicitObject.h"
 #include "Chaos/ImplicitObjectUnion.h"
+#include "Chaos/ParticleHandle.h"
 
 namespace Chaos
 {
 	template <typename T, int d>
-	void UpdateShapesArrayFromGeometry(TShapesArray<T, d>& ShapesArray, const TImplicitObject<T, d>* Geometry)
+	void UpdateShapesArrayFromGeometry(TShapesArray<T, d>& ShapesArray, TSerializablePtr<TImplicitObject<T, d>> Geometry)
 	{
 		if(Geometry)
 		{
@@ -17,15 +18,15 @@ namespace Chaos
 				int32 Inner = 0;
 				for(const auto& Geom : Union->GetObjects())
 				{
-					ShapesArray[Inner] = MakeUnique<TPerShapeData<T, d>>();
-					ShapesArray[Inner]->Geometry = Geom.Get();
+					ShapesArray[Inner] = TPerShapeData<T, d>::CreatePerShapeData();
+					ShapesArray[Inner]->Geometry = MakeSerializable(Geom);
 					++Inner;
 				}
 			}
 			else
 			{
 				ShapesArray.SetNum(1);
-				ShapesArray[0] = MakeUnique<TPerShapeData<T, d>>();
+				ShapesArray[0] = TPerShapeData<T, d>::CreatePerShapeData();
 				ShapesArray[0]->Geometry = Geometry;
 			}
 		}
@@ -35,5 +36,83 @@ namespace Chaos
 		}
 	}
 
-	template void UpdateShapesArrayFromGeometry(TShapesArray<float, 3>& ShapesArray, const TImplicitObject<float, 3>* Geometry);
+	template <typename T, int d>
+	TPerShapeData<T, d>::TPerShapeData()
+	: UserData(nullptr)
+	{
+	}
+
+	template <typename T, int d>
+	TPerShapeData<T, d>::~TPerShapeData()
+	{
+	}
+
+	template <typename T, int d>
+	TUniquePtr<TPerShapeData<T, d>> TPerShapeData<T, d>::CreatePerShapeData()
+	{
+		return TUniquePtr<TPerShapeData<T, d>>(new TPerShapeData<T, d>());
+	}
+
+	template <typename T, int d>
+	TPerShapeData<T, d>* TPerShapeData<T, d>::SerializationFactory(FChaosArchive& Ar, TPerShapeData<T, d>*)
+	{
+		return Ar.IsLoading() ? new TPerShapeData<T, d>() : nullptr;
+	}
+
+	template <typename T, int d>
+	void TPerShapeData<T, d>::Serialize(FChaosArchive& Ar)
+	{
+		Ar << Geometry;
+		Ar << QueryData;
+		Ar << SimData;
+	}
+
+
+	template <typename T, int d, EGeometryParticlesSimType SimType>
+	void TGeometryParticlesImp<T, d, SimType>::SetHandle(int32 Index, TGeometryParticleHandle<T, d>* Handle)
+	{
+		Handle->SetSOALowLevel(this);
+		MGeometryParticleHandle[Index] = AsAlwaysSerializable(Handle);
+	}
+
+	template <>
+	void TGeometryParticlesImp<float, 3, EGeometryParticlesSimType::Other>::SetHandle(int32 Index, TGeometryParticleHandle<float, 3>* Handle)
+	{
+		check(false);  // TODO: Implement EGeometryParticlesSimType::Other (cloth) particle serialization
+	}
+
+	template<class T, int d, EGeometryParticlesSimType SimType>
+	CHAOS_API TGeometryParticlesImp<T, d, SimType>* TGeometryParticlesImp<T, d, SimType>::SerializationFactory(FChaosArchive& Ar, TGeometryParticlesImp<T,d,SimType>* Particles)
+	{
+		int8 ParticleType = Ar.IsLoading() ? 0 : (int8)Particles->ParticleType();
+		Ar << ParticleType;
+		switch ((EParticleType)ParticleType)
+		{
+		case EParticleType::Static: return Ar.IsLoading() ? new TGeometryParticlesImp<T, d, SimType>() : nullptr;
+		case EParticleType::Kinematic: return Ar.IsLoading() ? new TKinematicGeometryParticlesImp<T, d, SimType>() : nullptr;
+		case EParticleType::Dynamic: return Ar.IsLoading() ? new TPBDRigidParticles<T, d>() : nullptr;
+		case EParticleType::Clustered: return Ar.IsLoading() ? new TPBDRigidClusteredParticles<T, d>() : nullptr;
+		default:
+			check(false); return nullptr;
+		}
+	}
+	
+	template<>
+	TGeometryParticlesImp<float, 3, EGeometryParticlesSimType::Other>* TGeometryParticlesImp<float, 3, EGeometryParticlesSimType::Other>::SerializationFactory(FChaosArchive& Ar, TGeometryParticlesImp<float, 3, EGeometryParticlesSimType::Other>* Particles)
+	{
+		check(false);  // TODO: Implement EGeometryParticlesSimType::Other (cloth) particle serialization
+		return nullptr;
+	}
+
+	template <typename T, int d, EGeometryParticlesSimType SimType>
+	void TGeometryParticlesImp<T, d, SimType>::SerializeGeometryParticleHelper(FChaosArchive& Ar, TGeometryParticlesImp<T, d, EGeometryParticlesSimType::RigidBodySim>* GeometryParticles)
+	{
+		auto& SerializableGeometryParticles = AsAlwaysSerializableArray(GeometryParticles->MGeometryParticle);
+		Ar << SerializableGeometryParticles;
+	}
+	
+	template class TGeometryParticlesImp<float, 3, EGeometryParticlesSimType::RigidBodySim>;
+	template class TGeometryParticlesImp<float, 3, EGeometryParticlesSimType::Other>;
+	template class TPerShapeData<float, 3>;
+	template void UpdateShapesArrayFromGeometry(TShapesArray<float, 3>& ShapesArray, TSerializablePtr<TImplicitObject<float, 3>> Geometry);
 }
diff --git a/Engine/Source/Runtime/Experimental/Chaos/Private/Chaos/HeightField.cpp b/Engine/Source/Runtime/Experimental/Chaos/Private/Chaos/HeightField.cpp
--- a/Engine/Source/Runtime/Experimental/Chaos/Private/Chaos/HeightField.cpp
+++ b/Engine/Source/Runtime/Experimental/Chaos/Private/Chaos/HeightField.cpp
@@ -1,82 +1,1034 @@
 // Copyright 1998-2019 Epic Games, Inc. All Rights Reserved.
 #include "Chaos/HeightField.h"
+#include "Chaos/Convex.h"
+#include "Chaos/Box.h"
+#include "Chaos/GJK.h"
+#include "Chaos/ImplicitObjectScaled.h"
+#include "Chaos/Capsule.h"
+#include "HAL/IConsoleManager.h"
+#include "HAL/PlatformMisc.h"
 
 namespace Chaos
 {
-template <typename T>
-THeightField<T>::THeightField(TArray<T>&& Height, int32 NumRows, int32 NumCols, const TVector<T, 3>& Scale)
-	: TImplicitObject<T, 3>(EImplicitObject::HasBoundingBox, ImplicitObjectType::HeightField)
-{
-	//For now just generate a triangle mesh
-	const int32 NumCells = NumRows * NumCols;
-	ensure(Height.Num() == NumCells);
-	ensure(NumRows > 1);
-	ensure(NumCols > 1);
 
-	TArray<TVector<int32, 3>> Triangles;
-	Triangles.Reserve(2 * (NumCols - 1)*(NumRows - 1));
+	template<typename T>
+	class THeightfieldRaycastVisitor
+	{
+	public:
+
+		THeightfieldRaycastVisitor(const typename THeightField<T>::FDataType* InData, const TVector<T, 3>& InStart, const TVector<T, 3>& InDir, const T InThickness)
+			: OutTime(TNumericLimits<T>::Max())
+			, OutFaceIndex(INDEX_NONE)
+			, GeomData(InData)
+			, Start(InStart)
+			, Dir(InDir)
+			, Thickness(InThickness)
+		{}
+
+		enum class ERaycastType
+		{
+			Raycast,
+			Sweep
+		};
+
+		template<ERaycastType SQType>
+		bool Visit(int32 Payload, T& CurrentLength)
+		{
+			const int32 SubX = Payload % (GeomData->NumCols - 1);
+			const int32 SubY = Payload / (GeomData->NumCols - 1);
+
+			const int32 FullIndex = Payload + SubY;
+
+			const T Radius = Thickness + SMALL_NUMBER;
+			const T Radius2 = Radius * Radius;
+			bool bIntersection = false;
+
+			auto TestTriangle = [&](int32 FaceIndex, const TVector<T, 3>& A, const TVector<T, 3>& B, const TVector<T, 3>& C) -> bool
+			{
+				const TVector<T, 3> AB = B - A;
+				const TVector<T, 3> AC = C - A;
+
+				TVector<T, 3> Normal = TVector<T, 3>::CrossProduct(AB, AC);
+				const T Len2 = Normal.SafeNormalize();
+
+				if(!ensure(Len2 > SMALL_NUMBER))
+				{
+					// Bad triangle, co-linear points or very thin
+					return true;
+				}
+
+				const TPlane<T, 3> TrianglePlane(A, Normal);
+
+				TVector<T, 3> ResultPosition(0);
+				TVector<T, 3> ResultNormal(0);
+				T Time = TNumericLimits<T>::Max();
+				int32 DummyFaceIndex = INDEX_NONE;
+
+				if(TrianglePlane.Raycast(Start, Dir, CurrentLength, Thickness, Time, ResultPosition, ResultNormal, DummyFaceIndex))
+				{
+					if(Time == 0)
+					{
+						// Initial overlap
+						const TVector<T, 3> ClosestPtOnTri = FindClosestPointOnTriangle(TrianglePlane, A, B, C, Start);
+						const T DistToTriangle2 = (Start - ClosestPtOnTri).SizeSquared();
+						if(DistToTriangle2 <= Radius2)
+						{
+							OutTime = 0;
+							OutPosition = ClosestPtOnTri;
+							OutNormal = Normal;
+							OutFaceIndex = FaceIndex;
+							return false;
+						}
+					}
+					else
+					{
+						const TVector<T, 3> ClosestPtOnTri = FindClosestPointOnTriangle(ResultPosition, A, B, C, ResultPosition);
+						const T DistToTriangle2 = (ResultPosition - ClosestPtOnTri).SizeSquared();
+						bIntersection = DistToTriangle2 <= SMALL_NUMBER;
+					}
+				}
+
+				if(SQType == ERaycastType::Sweep && !bIntersection)
+				{
+					//sphere is not immediately touching the triangle, but it could start intersecting the perimeter as it sweeps by
+					TVector<T, 3> BorderPositions[3];
+					TVector<T, 3> BorderNormals[3];
+					T BorderTimes[3];
+					bool bBorderIntersections[3];
+
+					const TCapsule<T> ABCapsule(A, B, Thickness);
+					bBorderIntersections[0] = ABCapsule.Raycast(Start, Dir, CurrentLength, 0, BorderTimes[0], BorderPositions[0], BorderNormals[0], DummyFaceIndex);
+
+					const TCapsule<T> BCCapsule(B, C, Thickness);
+					bBorderIntersections[1] = BCCapsule.Raycast(Start, Dir, CurrentLength, 0, BorderTimes[1], BorderPositions[1], BorderNormals[1], DummyFaceIndex);
+
+					const TCapsule<T> ACCapsule(A, C, Thickness);
+					bBorderIntersections[2] = ACCapsule.Raycast(Start, Dir, CurrentLength, 0, BorderTimes[2], BorderPositions[2], BorderNormals[2], DummyFaceIndex);
+
+					int32 MinBorderIdx = INDEX_NONE;
+					T MinBorderTime = 0;
+
+					for(int32 BorderIdx = 0; BorderIdx < 3; ++BorderIdx)
+					{
+						if(bBorderIntersections[BorderIdx])
+						{
+							if(!bIntersection || BorderTimes[BorderIdx] < MinBorderTime)
+							{
+								MinBorderTime = BorderTimes[BorderIdx];
+								MinBorderIdx = BorderIdx;
+								bIntersection = true;
+							}
+						}
+					}
+
+					if(MinBorderIdx != INDEX_NONE)
+					{
+						ResultNormal = BorderNormals[MinBorderIdx];
+						ResultPosition = BorderPositions[MinBorderIdx] - ResultNormal * Thickness;
+
+						if(Time == 0)
+						{
+							//we were initially overlapping with triangle plane so no normal was given. Compute it now
+							TVector<T, 3> TmpNormal;
+							const T SignedDistance = TrianglePlane.PhiWithNormal(Start, TmpNormal);
+							ResultNormal = SignedDistance >= 0 ? TmpNormal : -TmpNormal;
+						}
+
+						Time = MinBorderTime;
+					}
+				}
+
+				if(bIntersection)
+				{
+					if(Time < OutTime)
+					{
+						OutPosition = ResultPosition;
+						OutNormal = ResultNormal;
+						OutTime = Time;
+						OutFaceIndex = FaceIndex;
+						CurrentLength = Time;
+					}
+				}
+
+				return true;
+			};
+
+			TVector<T, 3> Points[4];
+			GeomData->GetPoints(FullIndex, Points);
+
+			// Test both triangles that are in this cell, as we could hit both in any order
+			TestTriangle(FullIndex * 2, Points[0], Points[1], Points[2]);
+			TestTriangle(FullIndex * 2 + 1, Points[2], Points[1], Points[3]);
+
+			return OutTime > 0;
+		}
+
+		bool VisitRaycast(int32 Payload, T& CurLength)
+		{
+			return Visit<ERaycastType::Raycast>(Payload, CurLength);
+		}
+
+		bool VisitSweep(int32 Payload, T& CurLength)
+		{
+			return Visit<ERaycastType::Sweep>(Payload, CurLength);
+		}
+
+		T OutTime;
+		TVector<T, 3> OutPosition;
+		TVector<T, 3> OutNormal;
+		int32 OutFaceIndex;
+
+	private:
+
+		const typename THeightField<T>::FDataType* GeomData;
+
+		TVector<T, 3> Start;
+		TVector<T, 3> Dir;
+		T Thickness;
+	};
+
+	template<typename T>
+	class THeightfieldSweepVisitor
+	{
+	public:
+
+		THeightfieldSweepVisitor(const typename THeightField<T>::FDataType* InData, const TImplicitObject<T, 3>& InQueryGeom, const TRigidTransform<T, 3>& InStartTM, const TVector<T, 3>& InDir, const T InThickness)
+			: OutTime(TNumericLimits<T>::Max())
+			, OutFaceIndex(INDEX_NONE)
+			, HfData(InData)
+			, StartTM(InStartTM)
+			, OtherGeom(InQueryGeom)
+			, Dir(InDir)
+			, Thickness(InThickness)
+		{}
+
+		bool VisitSweep(int32 Payload, T& CurrentLength)
+		{
+			const int32 SubX = Payload % (HfData->NumCols - 1);
+			const int32 SubY = Payload / (HfData->NumCols - 1);
+
+			const int32 FullIndex = Payload + SubY;
+
+			auto TestTriangle = [&](int32 FaceIndex, const TVector<T, 3>& A, const TVector<T, 3>& B, const TVector<T, 3>& C) -> bool
+			{
+				if(OutTime == 0)
+				{
+					return false;
+				}
+
+				// This isn't great as we build a convex on-the-fly but we have no other
+				// trimesh support. Update later as better trimesh collisions come online
+				TParticles<T, 3> TriParticles;
+				TriParticles.AddParticles(3);
+				TriParticles.X(0) = A;
+				TriParticles.X(1) = B;
+				TriParticles.X(2) = C;
+
+				TConvex<T, 3> Convex(TriParticles);
+
+				T Time;
+				TVector<T, 3> HitPosition;
+				TVector<T, 3> HitNormal;
+				if(GJKRaycast<T>(Convex, OtherGeom, StartTM, Dir, CurrentLength, Time, HitPosition, HitNormal, Thickness))
+				{
+					if(Time < OutTime)
+					{
+						OutNormal = HitNormal;
+						OutPosition = HitPosition;
+						OutTime = Time;
+						CurrentLength = Time;
+						OutFaceIndex = FaceIndex;
+
+						if(Time == 0)
+						{
+							return false;
+						}
+					}
+				}
+
+				return true;
+			};
+
+			TVector<T, 3> Points[4];
+			HfData->GetPoints(FullIndex, Points);
+
+			TestTriangle(FullIndex * 2, Points[0], Points[1], Points[2]);
+			TestTriangle(FullIndex * 2 + 1, Points[2], Points[1], Points[3]);
+
+			return OutTime != 0;
+		}
+
+		T OutTime;
+		TVector<T, 3> OutPosition;
+		TVector<T, 3> OutNormal;
+		int32 OutFaceIndex;
+
+	private:
+
+		const typename THeightField<T>::FDataType* HfData;
+		const TRigidTransform<T, 3> StartTM;
+		const TImplicitObject<T, 3>& OtherGeom;
+		const TVector<T, 3>& Dir;
+		const T Thickness;
+
+	};
+
+	template <typename T>
+	THeightField<T>::THeightField(TArray<T>&& Height, int32 NumRows, int32 NumCols, const TVector<T, 3>& InScale)
+		: TImplicitObject<T, 3>(EImplicitObject::HasBoundingBox, ImplicitObjectType::HeightField)
+	{
+		//For now just generate a triangle mesh
+		const int32 NumCells = NumRows * NumCols;
+		ensure(Height.Num() == NumCells);
+		ensure(NumRows > 1);
+		ensure(NumCols > 1);
+
+		// Populate data.
+		const int32 NumHeights = Height.Num();
+		GeomData.Heights.SetNum(NumHeights);
+
+		GeomData.NumRows = NumRows;
+		GeomData.NumCols = NumCols;
+		GeomData.MinValue = Height[0];
+		GeomData.MaxValue = Height[0];
+		GeomData.Scale = InScale;
+
+		for(int32 HeightIndex = 1; HeightIndex < NumHeights; ++HeightIndex)
+		{
+			const typename FDataType::RealType CurrHeight = (typename FDataType::RealType)Height[HeightIndex];
+
+			if(CurrHeight > GeomData.MaxValue)
+			{
+				GeomData.MaxValue = CurrHeight;
+			}
+			else if(CurrHeight < GeomData.MinValue)
+			{
+				GeomData.MinValue = CurrHeight;
+			}
+		}
+
+		GeomData.Range = GeomData.MaxValue - GeomData.MinValue;
+		GeomData.HeightPerUnit = GeomData.Range / FDataType::StorageRange;
+
+		for(int32 HeightIndex = 0; HeightIndex < NumHeights; ++HeightIndex)
+		{
+			GeomData.Heights[HeightIndex] = static_cast<typename FDataType::StorageType>((Height[HeightIndex] - GeomData.MinValue) / GeomData.HeightPerUnit);
+
+			int32 X = HeightIndex % (NumCols);
+			int32 Y = HeightIndex / (NumCols);
+			TVector<T, 3> Position(typename FDataType::RealType(X), typename FDataType::RealType(Y), GeomData.MinValue + GeomData.Heights[HeightIndex] * GeomData.HeightPerUnit);
+			if(HeightIndex == 0)
+			{
+				LocalBounds = TBox<T, 3>(Position * InScale, Position * InScale);
+			}
+			else
+			{
+				LocalBounds.GrowToInclude(Position * InScale);
+			}
+		}
+
+		// Flatten out the Z axis
+		FlattenedBounds.Min = TVector<T, 2>(LocalBounds.Min()[0], LocalBounds.Min()[1]);
+		FlattenedBounds.Max = TVector<T, 2>(LocalBounds.Max()[0], LocalBounds.Max()[1]);
+
+		BuildQueryData();
+
+		// Cache per-cell bounds
+		const int32 NumX = GeomData.NumCols - 1;
+		const int32 NumY = GeomData.NumRows - 1;
+		GeomData.CellBounds.SetNum(NumX * NumY);
+		for(int32 XIndex = 0; XIndex < NumX; ++XIndex)
+		{
+			for(int32 YIndex = 0; YIndex < NumY; ++YIndex)
+			{
+				const TVector<int32, 2> Cell(XIndex, YIndex);
+				TVector<T, 3> Min, Max;
+				CalcCellBounds3D(Cell, Min, Max);
+				GeomData.CellBounds[YIndex * NumX + XIndex] = TBox<T, 3>(Min, Max);
+			}
+		}
+	}
 
-	TParticles<T, 3> Particles;
-	Particles.AddParticles(NumCells);
-	int32 Idx = 0;
-	for (int32 Row = 0; Row < NumRows; ++Row)
+	template<typename T>
+	bool Chaos::THeightField<T>::GetCellBounds2D(const TVector<int32, 2> InCoord, FBounds2D& OutBounds, const TVector<T, 2>& InInflate /*= {0}*/) const
 	{
-		for (int32 Col = 0; Col < NumCols; ++Col)
+		if(FlatGrid.IsValid(InCoord))
 		{
-			if (Row < NumRows - 1 && Col < NumCols - 1)
+			const TBox<T, 3>& Bound = GeomData.CellBounds[InCoord[1] * (GeomData.NumCols - 1) + InCoord[0]];
+			OutBounds.Min = {Bound.Min()[0], Bound.Min()[1]};
+			OutBounds.Max = {Bound.Max()[0], Bound.Max()[1]};
+			OutBounds.Min -= InInflate;
+			OutBounds.Max += InInflate;
+
+			return true;
+		}
+
+		return false;
+	}
+
+	template<typename T>
+	bool Chaos::THeightField<T>::GetCellBounds3D(const TVector<int32, 2> InCoord, TVector<T, 3>& OutMin, TVector<T, 3>& OutMax, const TVector<T, 3>& InInflate /*= {0}*/) const
+	{
+		if(FlatGrid.IsValid(InCoord))
+		{
+			const TBox<T, 3>& Bound = GeomData.CellBounds[InCoord[1] * (GeomData.NumCols - 1) + InCoord[0]];
+			OutMin = Bound.Min() - InInflate;
+			OutMax = Bound.Max() + InInflate;
+
+			return true;
+		}
+
+		return false;
+	}
+
+
+	template<typename T>
+	bool Chaos::THeightField<T>::CalcCellBounds3D(const TVector<int32, 2> InCoord, TVector<T, 3>& OutMin, TVector<T, 3>& OutMax, const TVector<T, 3>& InInflate /*= {0}*/) const
+	{
+		if(FlatGrid.IsValid(InCoord))
+		{
+			int32 Index = InCoord[1] * (GeomData.NumCols - 1) + InCoord[0] + InCoord[1];
+			static TVector<T, 3> Points[4];
+			GeomData.GetPoints(Index, Points);
+
+			OutMin = OutMax = Points[0];
+
+			for(int32 PointIndex = 1; PointIndex < 4; ++PointIndex)
 			{
-				//add two triangles question: do we care about tesselation control?
-				Triangles.Emplace(Idx, Idx + 1, Idx + 1 + NumCols);
-				Triangles.Emplace(Idx, Idx + 1 + NumCols, Idx + NumCols);
-				//Triangles.Emplace(Idx, Idx + 1, Idx + NumCols);
-				//Triangles.Emplace(Idx + 1, Idx + NumCols + 1, Idx + NumCols);
+				const TVector<T, 3>& Point = Points[PointIndex];
+				OutMin = TVector<T, 3>(FGenericPlatformMath::Min(OutMin[0], Point[0]), FGenericPlatformMath::Min(OutMin[1], Point[1]), FGenericPlatformMath::Min(OutMin[2], Point[2]));
+				OutMax = TVector<T, 3>(FGenericPlatformMath::Max(OutMax[0], Point[0]), FGenericPlatformMath::Max(OutMax[1], Point[1]), FGenericPlatformMath::Max(OutMax[2], Point[2]));
 			}
 
-			Particles.X(Idx) = TVector<T, 3>(Col, Row, Height[Idx]) * Scale;
-			++Idx;
+			OutMin -= InInflate;
+			OutMax += InInflate;
+
+			return true;
 		}
+
+		return false;
 	}
 
-	MTriangleMeshImplicitObject = MakeUnique<TTriangleMeshImplicitObject<T>>(MoveTemp(Particles), MoveTemp(Triangles));
-}
 
-template <typename T>
-bool THeightField<T>::Raycast(const TVector<T, 3>& StartPoint, const TVector<T, 3>& Dir, const T Length, const T Thickness, T& OutTime, TVector<T,3>& OutPosition, TVector<T,3>& OutNormal, int32& OutFaceIndex) const
-{
-	return MTriangleMeshImplicitObject->Raycast(StartPoint, Dir, Length, Thickness, OutTime, OutPosition, OutNormal, OutFaceIndex);
-}
+	template <typename T>
+	bool THeightField<T>::GridCast(const TVector<T, 3>& StartPoint, const TVector<T, 3>& Dir, const T Length, THeightfieldRaycastVisitor<T>& Visitor) const
+	{
+		T CurrentLength = Length;
+		TVector<T, 2> ClippedFlatRayStart;
+		TVector<T, 2> ClippedFlatRayEnd;
 
-template <typename T>
-bool THeightField<T>::Overlap(const TVector<T, 3>& Point, const T Thickness) const
-{
-	return MTriangleMeshImplicitObject->Overlap(Point, Thickness);
-}
+		if(FlattenedBounds.ClipLine(StartPoint, StartPoint + Dir * Length, ClippedFlatRayStart, ClippedFlatRayEnd))
+		{
+			// The line is now valid and is entirely enclosed by the bounds (and thus, the grid)
+			if(FMath::Abs((ClippedFlatRayEnd - ClippedFlatRayStart).SizeSquared()) < SMALL_NUMBER)
+			{
+				// This is a cast down the Z axis, handle in a simpler way as this is the common case
+				const int32 QueryX = (int32)(ClippedFlatRayStart[0] / GeomData.Scale[0]);
+				const int32 QueryY = (int32)(ClippedFlatRayStart[1] / GeomData.Scale[1]);
+				Visitor.VisitRaycast(QueryY * (GeomData.NumCols - 1) + QueryX, CurrentLength);
+			}
+			else
+			{
+				// Rasterize the line over the grid
+				TVector<int32, 2> StartCell = FlatGrid.Cell(ClippedFlatRayStart);
+				TVector<int32, 2> EndCell = FlatGrid.Cell(ClippedFlatRayEnd);
 
+				// Boundaries might push us one cell over
+				StartCell = FlatGrid.ClampIndex(StartCell);
+				EndCell = FlatGrid.ClampIndex(EndCell);
 
-template <typename T>
-bool THeightField<T>::OverlapGeom(const TImplicitObject<T, 3>& QueryGeom, const TRigidTransform<T, 3>& QueryTM, const T Thickness) const
-{
-	return MTriangleMeshImplicitObject->OverlapGeom(QueryGeom, QueryTM, Thickness);
-}
+				int32 DeltaX = FMath::Abs(EndCell[0] - StartCell[0]);
+				int32 DeltaY = -FMath::Abs(EndCell[1] - StartCell[1]);
+				int32 DirX = StartCell[0] < EndCell[0] ? 1 : -1;
+				int32 DirY = StartCell[1] < EndCell[1] ? 1 : -1;
+				int32 Error = DeltaX + DeltaY;
 
-template <typename T>
-bool THeightField<T>::SweepGeom(const TImplicitObject<T, 3>& QueryGeom, const TRigidTransform<T, 3>& StartTM, const TVector<T, 3>& Dir, const T Length, T& OutTime, TVector<T, 3>& OutPosition, TVector<T, 3>& OutNormal, int32& OutFaceIndex, const T Thickness /* = 0 */) const
-{
-	return MTriangleMeshImplicitObject->SweepGeom(QueryGeom, StartTM, Dir, Length, OutTime, OutPosition, OutNormal, OutFaceIndex, Thickness);
-}
+				if(StartCell == EndCell)
+				{
+					// Visit the selected cell, in this case no need to rasterize because we never
+					// leave the initial cell
+					Visitor.VisitRaycast(StartCell[1] * (GeomData.NumCols - 1) + StartCell[0], CurrentLength);
+				}
+				else
+				{
+					Visitor.VisitRaycast(StartCell[1] * (GeomData.NumCols - 1) + StartCell[0], CurrentLength);
 
-template <typename T>
-int32 THeightField<T>::FindMostOpposingFace(const TVector<T, 3>& Position, const TVector<T, 3>& UnitDir, int32 HintFaceIndex, T SearchDist) const
-{
-	return MTriangleMeshImplicitObject->FindMostOpposingFace(Position, UnitDir, HintFaceIndex, SearchDist);
-}
+					do 
+					{
+						const int32 DoubleError = Error * 2;
 
-template <typename T>
-TVector<T, 3> THeightField<T>::FindGeometryOpposingNormal(const TVector<T, 3>& DenormDir, int32 FaceIndex, const TVector<T, 3>& OriginalNormal) const
-{
-	return MTriangleMeshImplicitObject->FindGeometryOpposingNormal(DenormDir, FaceIndex, OriginalNormal);
-}
+						if(DoubleError >= DeltaY)
+						{
+							Error += DeltaY;
+							StartCell[0] += DirX;
+						}
+
+						if(DoubleError <= DeltaX)
+						{
+							Error += DeltaX;
+							StartCell[1] += DirY;
+						}
+
+						// Visit the selected cell
+						bool bContinue = Visitor.VisitRaycast(StartCell[1] * (GeomData.NumCols - 1) + StartCell[0], CurrentLength);
+						if(!bContinue)
+						{
+							return true;
+						}
+					} while(StartCell != EndCell);
+				}
+			}
+		}
+
+		return false;
+	}
+
+	struct F2DGridSet
+	{
+		F2DGridSet(TVector<int32, 2> Size)
+			: NumX(Size[0])
+			, NumY(Size[1])
+		{
+			int32 BitsNeeded = NumX * NumY;
+			int32 BytesNeeded = 1 + (BitsNeeded) / 8;
+			Data = MakeUnique<uint8[]>(BytesNeeded);
+			FMemory::Memzero(Data.Get(), BytesNeeded);
+		}
+
+		bool Contains(const TVector<int32, 2>& Coordinate)
+		{
+			int32 Idx = Coordinate[1] * NumX + Coordinate[0];
+			int32 ByteIdx = Idx / 8;
+			int32 BitIdx = Idx % 8;
+			bool bContains = (Data[ByteIdx] >> BitIdx) & 0x1;
+			return bContains;
+		}
+
+		void Add(const TVector<int32, 2>& Coordinate)
+		{
+			int32 Idx = Coordinate[1] * NumX + Coordinate[0];
+			int32 ByteIdx = Idx / 8;
+			int32 BitIdx = Idx % 8;
+			uint8 Mask = 1 << BitIdx;
+			Data[ByteIdx] |= Mask;
+		}
+
+	private:
+		int32 NumX;
+		int32 NumY;
+		TUniquePtr<uint8[]> Data;
+	};
+
+	template<typename T>
+	template<typename SQVisitor>
+	bool THeightField<T>::GridSweep(const TVector<T, 3>& StartPoint, const TVector<T, 3>& Dir, const T Length, const TVector<T, 2> InHalfExtents, SQVisitor& Visitor) const
+	{
+		T CurrentLength = Length;
+
+		FBounds2D InflatedBounds = FlattenedBounds;
+		InflatedBounds.Min -= InHalfExtents;
+		InflatedBounds.Max += InHalfExtents;
+
+		TVector<T, 3> HalfExtents3D(InHalfExtents[0], InHalfExtents[1], InHalfExtents[1]);
+
+		const TVector<T, 3> EndPoint = StartPoint + Dir * Length;
+		const TVector<T, 2> Start2D(StartPoint[0], StartPoint[1]);
+		const TVector<T, 2> End2D(EndPoint[0], EndPoint[1]);
+
+		TVector<T, 2> ClippedStart;
+		TVector<T, 2> ClippedEnd;
+
+		if(InflatedBounds.ClipLine(StartPoint, StartPoint + Dir * Length, ClippedStart, ClippedEnd))
+		{
+			// Rasterize the line over the grid
+			TVector<int32, 2> StartCell = FlatGrid.Cell(ClippedStart);
+			TVector<int32, 2> EndCell = FlatGrid.Cell(ClippedEnd);
+
+			// Boundaries might push us one cell over
+			StartCell = FlatGrid.ClampIndex(StartCell);
+			EndCell = FlatGrid.ClampIndex(EndCell);
+
+			const int32 DeltaX = FMath::Abs(EndCell[0] - StartCell[0]);
+			const int32 DeltaY = -FMath::Abs(EndCell[1] - StartCell[1]);
+			const bool bSameCell = DeltaX == 0 && DeltaY == 0;
+
+			const int32 DirX = StartCell[0] < EndCell[0] ? 1 : -1;
+			const int32 DirY = StartCell[1] < EndCell[1] ? 1 : -1;
+			int32 Error = DeltaX + DeltaY;
+			const TVector<int32, 2> ThickenDir = FMath::Abs(DeltaX) > FMath::Abs(DeltaY) ? TVector<int32, 2>(0, 1) : TVector<int32, 2>(1, 0);
+
+			struct FQueueEntry
+			{
+				TVector<int32, 2> Index;
+				T ToI;
+			};
+
+			// Tracking data for cells to query (similar to bounding volume approach)
+			F2DGridSet Seen(FlatGrid.Counts());
+			TArray<FQueueEntry> Queue;
+			Queue.Add({StartCell, -1});
+			Seen.Add(StartCell);
+
+			// Data for fast box cast
+			TVector<T, 3> Min, Max, HitPoint;
+			float ToI;
+			bool bParallel[3];
+			TVector<T, 3> InvDir;
+
+			T InvCurrentLength = 1 / CurrentLength;
+			for(int Axis = 0; Axis < 3; ++Axis)
+			{
+				bParallel[Axis] = Dir[Axis] == 0;
+				InvDir[Axis] = bParallel[Axis] ? 0 : 1 / Dir[Axis];
+			}
+
+			int32 QueueIndex = 0;
+			while(QueueIndex < Queue.Num())
+			{
+				// Copy so we don't lost the entry through reallocs
+				FQueueEntry CellCoord = Queue[QueueIndex++];
+
+				if(CellCoord.ToI > CurrentLength)
+				{
+					continue;
+				}
+
+				if(bSameCell)
+				{
+					// Test the current cell
+					bool bContinue = Visitor.VisitSweep(CellCoord.Index[1] * (GeomData.NumCols - 1) + CellCoord.Index[0], CurrentLength);
+					
+					if(!bContinue)
+					{
+						return true;
+					}
+
+					// Flatten out a double loop and skip the centre cell
+					// to search cells immediately adjacent to the current cell
+					static const TVector<int32, 2> Neighbors[] =
+					{
+						{-1, -1}, {0, -1}, {1, -1},
+						{-1, 0}, {1, 0},
+						{-1, 1}, {0, 1}, {1, 1}
+					};
+
+					for(const TVector<int32, 2>& Neighbor : Neighbors)
+					{
+						TVector<int32, 2> NeighCoord = CellCoord.Index + Neighbor;
+
+						FBounds2D CellBounds;
+						if(!Seen.Contains(NeighCoord) && GetCellBounds3D(NeighCoord, Min, Max, HalfExtents3D))
+						{
+							if(TBox<T, 3>::RaycastFast(Min, Max, StartPoint, Dir, InvDir, bParallel, CurrentLength, InvCurrentLength, ToI, HitPoint))
+							{
+								Seen.Add(NeighCoord);
+								Queue.Add({NeighCoord, ToI});
+							}
+						}
+					}
+				}
+				else
+				{
+					// Expand each cell along the thicken direction
+					// Although the line should minimally thicken around the perpendicular to the line direction
+					// it's cheaper to just expand in the cardinal opposite the current major direction. We end up
+					// doing a broad test on more cells but avoid having to run many rasterize/walk steps for each
+					// perpendicular step.
+					auto Expand = [&](const TVector<int32, 2>& Begin, const TVector<int32, 2>& Direction, const int32 NumSteps)
+					{
+						TVector<int32, 2> CurrentCell = Begin;
+
+						for(int32 CurrStep = 0; CurrStep < NumSteps; ++CurrStep)
+						{
+							CurrentCell += Direction;
+
+							// Fail if we leave the grid
+							if(Begin[0] < 0 || Begin[1] < 0 || Begin[0] > FlatGrid.Counts()[0] - 1 || Begin[1] > FlatGrid.Counts()[1] - 1)
+							{
+								break;
+							}
+
+							// No intersections here. We set the ToI to zero to cause an intersection check to happen
+							// without any expansion when we reach this cell in the queue.
+							if(!Seen.Contains(CurrentCell))
+							{
+								Seen.Add(CurrentCell);
+								Queue.Add({CurrentCell, 0});
+							}
+						}
+					};
+
+					// Check the current cell, if we hit its 3D bound we can move on to narrow phase
+					const TVector<int32, 2> Coord = CellCoord.Index;
+					if(GetCellBounds3D(Coord, Min, Max, HalfExtents3D) &&
+						TBox<T, 3>::RaycastFast(Min, Max, StartPoint, Dir, InvDir, bParallel, CurrentLength, InvCurrentLength, ToI, HitPoint))
+					{
+						bool bContinue = Visitor.VisitSweep(CellCoord.Index[1] * (GeomData.NumCols - 1) + CellCoord.Index[0], CurrentLength);
+
+						if(!bContinue)
+						{
+							return true;
+						}
+					}
+
+					// This time isn't used to reject things for this method but to flag cells that should be expanded
+					if(CellCoord.ToI < 0)
+					{
+						// Perform expansion for thickness
+						int32 ExpandAxis;
+						if(ThickenDir[0] == 0)
+						{
+							ExpandAxis = 1;
+						}
+						else
+						{
+							ExpandAxis = 0;
+						}
+						T ExpandSize = HalfExtents3D[ExpandAxis];
+						int32 Steps = FMath::RoundFromZero(ExpandSize / GeomData.Scale[ExpandAxis]);
+
+						Expand(Coord, ThickenDir, Steps);
+						Expand(Coord, -ThickenDir, Steps);
+
+						// Walk the line and add to the queue
+						if(StartCell != EndCell)
+						{
+							const int32 DoubleError = Error * 2;
+
+							if(DoubleError >= DeltaY)
+							{
+								Error += DeltaY;
+								StartCell[0] += DirX;
+							}
+
+							if(DoubleError <= DeltaX)
+							{
+								Error += DeltaX;
+								StartCell[1] += DirY;
+							}
+
+							if(!Seen.Contains(StartCell))
+							{
+								Seen.Add(StartCell);
+								Queue.Add({StartCell, -1});
+							}
+						}
+					}
+				}
+			}
+		}
+
+		return false;
+	}
+
+	template <typename T>
+	bool THeightField<T>::Raycast(const TVector<T, 3>& StartPoint, const TVector<T, 3>& Dir, const T Length, const T Thickness, T& OutTime, TVector<T, 3>& OutPosition, TVector<T, 3>& OutNormal, int32& OutFaceIndex) const
+	{
+		OutFaceIndex = INDEX_NONE;
+
+		THeightfieldRaycastVisitor<T> Visitor(&GeomData, StartPoint, Dir, Thickness);
+
+		if(Thickness > 0)
+		{
+			GridSweep(StartPoint, Dir, Length, TVector<T, 2>(Thickness), Visitor);
+		}
+		else
+		{
+			GridCast(StartPoint, Dir, Length, Visitor);
+		}
+
+		if(Visitor.OutTime <= Length)
+		{
+			OutTime = Visitor.OutTime;
+			OutPosition = Visitor.OutPosition;
+			OutNormal = Visitor.OutNormal;
+			OutFaceIndex = Visitor.OutFaceIndex;
+			return true;
+		}
+
+		return false;
+	}
+
+	template<typename T>
+	bool Chaos::THeightField<T>::GetGridIntersections(FBounds2D InFlatBounds, TArray<TVector<int32, 2>>& OutInterssctions) const
+	{
+		OutInterssctions.Reset();
+
+		InFlatBounds.Min = FlattenedBounds.Clamp(InFlatBounds.Min);
+		InFlatBounds.Max = FlattenedBounds.Clamp(InFlatBounds.Max);
+		TVector<int32, 2> MinCell = FlatGrid.Cell(InFlatBounds.Min);
+		TVector<int32, 2> MaxCell = FlatGrid.Cell(InFlatBounds.Max);
+
+		const int32 DeltaX = MaxCell[0] - MinCell[0];
+		const int32 DeltaY = MaxCell[1] - MinCell[1];
+
+		for(int32 CurrX = 0; CurrX < DeltaX; ++CurrX)
+		{
+			for(int32 CurrY = 0; CurrY < DeltaY; ++CurrY)
+			{
+				OutInterssctions.Add(TVector<int32, 2>(MinCell[0] + CurrX, MinCell[1] + CurrY));
+			}
+		}
+
+		return OutInterssctions.Num() > 0;
+	}
+
+	template <typename T>
+	bool THeightField<T>::Overlap(const TVector<T, 3>& Point, const T Thickness) const
+	{
+		auto OverlapTriangle = [&](const TVector<T, 3>& A, const TVector<T, 3>& B, const TVector<T, 3>& C) -> bool
+		{
+			const TVector<T, 3> AB = B - A;
+			const TVector<T, 3> AC = C - A;
+			TVector<T, 3> Normal = TVector<T, 3>::CrossProduct(AB, AC);
+			const T NormalLength = Normal.SafeNormalize();
+
+			if(!ensure(NormalLength > KINDA_SMALL_NUMBER))
+			{
+				return false;
+			}
+
+			const TPlane<T, 3> TriPlane{A, Normal};
+			const TVector<T, 3> ClosestPointOnTri = FindClosestPointOnTriangle(TriPlane, A, B, C, Point);
+			const T Distance2 = (ClosestPointOnTri - Point).SizeSquared();
+
+			if(Distance2 <= Thickness * Thickness)	//This really only has a hope in working if thickness is > 0
+			{
+				return true;
+			}
+
+			return false;
+		};
+
+		TBox<T, 3> QueryBounds(Point, Point);
+		QueryBounds.Thicken(Thickness);
+
+		FBounds2D FlatQueryBounds;
+		FlatQueryBounds.Min = TVector<T, 2>(QueryBounds.Min()[0], QueryBounds.Min()[1]);
+		FlatQueryBounds.Max = TVector<T, 2>(QueryBounds.Max()[0], QueryBounds.Max()[1]);
+
+		TArray<TVector<int32, 2>> Intersections;
+		TVector<T, 3> Points[4];
+
+		GetGridIntersections(FlatQueryBounds, Intersections);
+
+		for(const TVector<int32, 2>& Cell : Intersections)
+		{
+			const int32 SingleIndex = Cell[1] * (GeomData.NumCols - 1) + Cell[0];
+			GeomData.GetPoints(SingleIndex, Points);
+
+			if(OverlapTriangle(Points[0], Points[1], Points[2]))
+			{
+				return true;
+			}
+
+			if(OverlapTriangle(Points[2], Points[1], Points[3]))
+			{
+				return true;
+			}
+		}
+		
+		return false;
+	}
+
+
+	template <typename T>
+	bool THeightField<T>::OverlapGeom(const TImplicitObject<T, 3>& QueryGeom, const TRigidTransform<T, 3>& QueryTM, const T Thickness) const
+	{
+		auto OverlapTriangle = [&](const TVector<T, 3>& A, const TVector<T, 3>& B, const TVector<T, 3>& C) -> bool
+		{
+			const TVector<T, 3> AB = B - A;
+			const TVector<T, 3> AC = C - A;
+
+			//It's most likely that the query object is in front of the triangle since queries tend to be on the outside.
+			//However, maybe we should check if it's behind the triangle plane. Also, we should enforce this winding in some way
+			const TVector<T, 3> Offset = TVector<T, 3>::CrossProduct(AB, AC);
+
+			// Ugly but required for now until we have an easier way to do tri collisons 
+			TParticles<T, 3> Particles;
+			Particles.AddParticles(3);
+
+			Particles.X(0) = A;
+			Particles.X(1) = B;
+			Particles.X(2) = C;
+
+			TConvex<T, 3> TriangleConvex(Particles);
+
+			return GJKIntersection(TriangleConvex, QueryGeom, QueryTM, Thickness, Offset);
+		};
+
+		bool bResult = false;
+		TBox<T, 3> QueryBounds = QueryGeom.BoundingBox();
+		QueryBounds.Thicken(Thickness);
+		QueryBounds = QueryBounds.TransformedBox(QueryTM);
+
+		FBounds2D FlatQueryBounds;
+		FlatQueryBounds.Min = TVector<T, 2>(QueryBounds.Min()[0], QueryBounds.Min()[1]);
+		FlatQueryBounds.Max = TVector<T, 2>(QueryBounds.Max()[0], QueryBounds.Max()[1]);
+
+		TArray<TVector<int32, 2>> Intersections;
+		TVector<T, 3> Points[4];
+
+		GetGridIntersections(FlatQueryBounds, Intersections);
+
+		for(const TVector<int32, 2>& Cell : Intersections)
+		{
+			const int32 SingleIndex = Cell[1] * (GeomData.NumCols - 1) + Cell[0];
+			GeomData.GetPoints(SingleIndex, Points);
+
+			if(OverlapTriangle(Points[0], Points[1], Points[2]))
+			{
+				return true;
+			}
+
+			if(OverlapTriangle(Points[2], Points[1], Points[3]))
+			{
+				return true;
+			}
+		}
+
+		return false;
+	}
+
+	template <typename T>
+	bool THeightField<T>::SweepGeom(const TImplicitObject<T, 3>& QueryGeom, const TRigidTransform<T, 3>& StartTM, const TVector<T, 3>& Dir, const T Length, T& OutTime, TVector<T, 3>& OutPosition, TVector<T, 3>& OutNormal, int32& OutFaceIndex, const T Thickness) const
+	{
+		bool bHit = false;
+		THeightfieldSweepVisitor<T> SQVisitor(&GeomData, QueryGeom, StartTM, Dir, Thickness);
+		const TBox<T, 3> QueryBounds = QueryGeom.BoundingBox();
+		const TVector<T, 3> StartPoint = StartTM.TransformPositionNoScale(QueryBounds.Center());
+
+		const TVector<T, 3> Inflation3D = QueryBounds.Extents() * 0.5 + TVector<T, 3>(Thickness);
+		GridSweep(StartPoint, Dir, Length, TVector<T, 2>(Inflation3D[0], Inflation3D[1]), SQVisitor);
+
+		if(SQVisitor.OutTime <= Length)
+		{
+			OutTime = SQVisitor.OutTime;
+			OutPosition = SQVisitor.OutPosition;
+			OutNormal = SQVisitor.OutNormal;
+			OutFaceIndex = SQVisitor.OutFaceIndex;
+			bHit = true;
+		}
+
+		return bHit;
+	}
+
+	template <typename T>
+	int32 THeightField<T>::FindMostOpposingFace(const TVector<T, 3>& Position, const TVector<T, 3>& UnitDir, int32 HintFaceIndex, T SearchDist) const
+	{
+		const T SearchDist2 = SearchDist * SearchDist;
+
+		TBox<T, 3> QueryBounds(Position - TVector<T, 3>(SearchDist), Position + TVector<T, 3>(SearchDist));
+		const FBounds2D FlatBounds(QueryBounds);
+		TArray<TVector<int32, 2>> PotentialIntersections;
+		GetGridIntersections(FlatBounds, PotentialIntersections);
+
+		T MostOpposingDot = TNumericLimits<T>::Max();
+		int32 MostOpposingFace = HintFaceIndex;
+
+		auto CheckTriangle = [&](int32 FaceIndex, const TVector<T, 3>& A, const TVector<T, 3>& B, const TVector<T, 3>& C)
+		{
+			const TVector<T, 3> AB = B - A;
+			const TVector<T, 3> AC = C - A;
+			TVector<T, 3> Normal = TVector<T, 3>::CrossProduct(AB, AC);
+			const T NormalLength = Normal.SafeNormalize();
+			if(!ensure(NormalLength > KINDA_SMALL_NUMBER))
+			{
+				//hitting degenerate triangle - should be fixed before we get to this stage
+				return;
+			}
+
+			const TPlane<T, 3> TriPlane{A, Normal};
+			const TVector<T, 3> ClosestPointOnTri = FindClosestPointOnTriangle(TriPlane, A, B, C, Position);
+			const T Distance2 = (ClosestPointOnTri - Position).SizeSquared();
+			if(Distance2 < SearchDist2)
+			{
+				const T Dot = TVector<T, 3>::DotProduct(Normal, UnitDir);
+				if(Dot < MostOpposingDot)
+				{
+					MostOpposingDot = Dot;
+					MostOpposingFace = FaceIndex;
+				}
+			}
+		};
+
+		for(const TVector<int32, 2>& CellCoord : PotentialIntersections)
+		{
+			const int32 CellIndex = CellCoord[1] * (GeomData.NumCols - 1) + CellCoord[0];
+			const int32 SubX = CellIndex % (GeomData.NumCols - 1);
+			const int32 SubY = CellIndex / (GeomData.NumCols - 1);
+
+			const int32 FullIndex = CellIndex + SubY;
+
+			TVector<T, 3> Points[4];
+
+			GeomData.GetPoints(FullIndex, Points);
+
+			CheckTriangle(FullIndex * 2, Points[0], Points[1], Points[2]);
+			CheckTriangle(FullIndex * 2 + 1, Points[2], Points[1], Points[3]);
+		}
+
+		return MostOpposingFace;
+	}
+
+	template <typename T>
+	TVector<T, 3> THeightField<T>::FindGeometryOpposingNormal(const TVector<T, 3>& DenormDir, int32 FaceIndex, const TVector<T, 3>& OriginalNormal) const
+	{
+		if(ensure(FaceIndex != INDEX_NONE))
+		{
+			bool bSecondFace = FaceIndex % 2 != 0;
+
+			if(bSecondFace)
+			{
+				FaceIndex -= 1;
+			}
+			FaceIndex /= 2;
+
+			TVector<T, 3> Points[4];
+
+			GeomData.GetPoints(FaceIndex, Points);
+
+			TVector<T, 3> A;
+			TVector<T, 3> B;
+			TVector<T, 3> C;
+
+			if(bSecondFace)
+			{
+				A = Points[2];
+				B = Points[1];
+				C = Points[3];
+			}
+			else
+			{
+				A = Points[0];
+				B = Points[1];
+				C = Points[2];
+			}
+
+			const TVector<T, 3> AB = B - A;
+			const TVector<T, 3> AC = C - A;
+			TVector<T, 3> Normal = TVector<T, 3>::CrossProduct(AB, AC);
+			const T Length = Normal.SafeNormalize();
+			ensure(Length);
+			return Normal;
+		}
+
+		return TVector<T, 3>(0, 0, 1);
+	}
+
+	template<typename T>
+	void THeightField<T>::BuildQueryData()
+	{
+		// NumCols and NumRows are the actual heights, there are n-1 cells between those heights
+		TVector<int32, 2> Cells(GeomData.NumCols - 1, GeomData.NumRows - 1);
+		
+		TVector<T, 2> MinCorner(0, 0);
+		TVector<T, 2> MaxCorner(GeomData.NumCols - 1, GeomData.NumRows - 1);
+		MaxCorner *= {GeomData.Scale[0], GeomData.Scale[1]};
+
+		FlatGrid = TUniformGrid<T, 2>(MinCorner, MaxCorner, Cells);
+	}
 
 }
 
-template class Chaos::THeightField<float>;
\ No newline at end of file
+template class Chaos::THeightField<float>;
diff --git a/Engine/Source/Runtime/Experimental/Chaos/Private/Chaos/ImplicitObject.cpp b/Engine/Source/Runtime/Experimental/Chaos/Private/Chaos/ImplicitObject.cpp
--- a/Engine/Source/Runtime/Experimental/Chaos/Private/Chaos/ImplicitObject.cpp
+++ b/Engine/Source/Runtime/Experimental/Chaos/Private/Chaos/ImplicitObject.cpp
@@ -4,6 +4,8 @@
 #include "Chaos/Box.h"
 #include "Chaos/Capsule.h"
 #include "Chaos/Convex.h"
+#include "Chaos/HeightField.h"
+#include "Chaos/ImplicitObjectScaled.h"
 #include "Chaos/ImplicitObjectTransformed.h"
 #include "Chaos/ImplicitObjectUnion.h"
 #include "Chaos/Levelset.h"
@@ -33,7 +35,7 @@ TImplicitObject<T, d>::~TImplicitObject()
 template<class T, int d>
 TVector<T, d> TImplicitObject<T, d>::Support(const TVector<T, d>& Direction, const T Thickness) const
 {
-	check(false);	//not a good implementation, don't use this
+	check(false); //not a good implementation, don't use this
 	return TVector<T, d>(0);
 #if 0
 	check(bHasBoundingBox);
@@ -120,41 +122,42 @@ Pair<TVector<T, d>, bool> TImplicitObject<T, d>::FindDeepestIntersection(const T
 template<class T, int d>
 Pair<TVector<T, d>, bool> TImplicitObject<T, d>::FindClosestIntersection(const TVector<T, d>& StartPoint, const TVector<T, d>& EndPoint, const T Thickness) const
 {
-	T Epsilon = (T)1e-4;
+	constexpr T Epsilon = (T)1e-4;
+	constexpr T EpsilonSquared = Epsilon * Epsilon;
+
 	//Consider 0 thickness with Start sitting on abs(Phi) < Epsilon. This is a common case; for example a particle sitting perfectly on a floor. In this case intersection could return false.
 	//If start is in this fuzzy region we simply return that spot snapped onto the surface. This is valid because low precision means we don't really know where we are, so let's take the cheapest option
 	//If end is in this fuzzy region it is also a valid hit. However, there could be multiple hits between start and end and since we want the first one, we can't simply return this point.
 	//As such we move end away from start (and out of the fuzzy region) so that we always get a valid intersection if no earlier ones exist
 	//When Thickness > 0 the same idea applies, but we must consider Phi = (Thickness - Epsilon, Thickness + Epsilon)
 	TVector<T, d> Normal;
-	T Phi = PhiWithNormal(StartPoint, Normal);
-
+	const T Phi = PhiWithNormal(StartPoint, Normal);
 	if (FMath::IsNearlyEqual(Phi, Thickness, Epsilon))
 	{
-		return MakePair(TVector<T,d>(StartPoint - Normal * Phi), true); //snap to surface
+		return MakePair(TVector<T, d>(StartPoint - Normal * Phi), true); //snap to surface
 	}
 
 	TVector<T, d> ModifiedEnd = EndPoint;
 	{
 		const TVector<T, d> OriginalStartToEnd = (EndPoint - StartPoint);
-		const T OriginalLength = OriginalStartToEnd.Size();
-		if (OriginalLength < Epsilon)
+		const T OriginalLength2 = OriginalStartToEnd.SizeSquared();
+		if (OriginalLength2 < EpsilonSquared)
 		{
 			return MakePair(TVector<T, d>(0), false); //start was not close to surface, and end is very close to start so no hit
 		}
-		const TVector<T, d> OriginalDir = OriginalStartToEnd / OriginalLength;
 
 		TVector<T, d> EndNormal;
-		T EndPhi = PhiWithNormal(EndPoint, EndNormal);
+		const T EndPhi = PhiWithNormal(EndPoint, EndNormal);
 		if (FMath::IsNearlyEqual(EndPhi, Thickness, Epsilon))
 		{
 			//We want to push End out of the fuzzy region. Moving along the normal direction is best since direction could be nearly parallel with fuzzy band
 			//To ensure an intersection, we must go along the normal, but in the same general direction as the ray.
+			const TVector<T, d> OriginalDir = OriginalStartToEnd / FMath::Sqrt(OriginalLength2);
 			const T Dot = TVector<T, d>::DotProduct(OriginalDir, EndNormal);
 			if (FMath::IsNearlyZero(Dot, Epsilon))
 			{
 				//End is in the fuzzy region, and the direction from start to end is nearly parallel with this fuzzy band, so we should just return End since no other hits will occur
-				return MakePair(TVector<T,d>(EndPoint - Normal * Phi), true); //snap to surface
+				return MakePair(TVector<T, d>(EndPoint - Normal * Phi), true); //snap to surface
 			}
 			else
 			{
@@ -234,7 +237,7 @@ void TImplicitObject<T, d>::FindAllIntersectingObjects(TArray<Pair<const TImplic
 {
 	if (!HasBoundingBox() || LocalBounds.Intersects(BoundingBox()))
 	{
-		Out.Add(MakePair(this, TRigidTransform<T, d>(TVector<T, d>(0), TRotation<T, d>(TVector<T, d>(0), (T)1))));
+		Out.Add(MakePair(this, TRigidTransform<T, d>(TVector<T, d>(0), TRotation<T, d>::FromElements(TVector<T, d>(0), (T)1))));
 	}
 }
 
@@ -291,53 +294,68 @@ void TImplicitObject<T, d>::Serialize(FChaosArchive& Ar)
 }
 
 template<typename T, int d>
-void TImplicitObject<T, d>::StaticSerialize(FChaosArchive& Ar, TSerializablePtr<TImplicitObject<T, d>>& Serializable)
+const FName TImplicitObject<T, d>::GetTypeName(const ImplicitObjectType InType)
 {
-	TImplicitObject<T, d>* ImplicitObject = const_cast<TImplicitObject<T, d>*>(Serializable.Get());
-	int8 ObjectType = Ar.IsLoading() ? 0 : (int8)ImplicitObject->Type;
-	Ar << ObjectType;
+	static const FName SphereName = TEXT("Sphere");
+	static const FName BoxName = TEXT("Box");
+	static const FName PlaneName = TEXT("Plane");
+	static const FName CapsuleName = TEXT("Capsule");
+	static const FName TransformedName = TEXT("Transformed");
+	static const FName UnionName = TEXT("Union");
+	static const FName LevelSetName = TEXT("LevelSet");
+	static const FName UnknownName = TEXT("Unknown");
+	static const FName ConvexName = TEXT("Convex");
+	static const FName TaperedCylinderName = TEXT("TaperedCylinder");
+	static const FName CylinderName = TEXT("Cylinder");
+	static const FName TriangleMeshName = TEXT("TriangleMesh");
+	static const FName HeightFieldName = TEXT("HeightField");
+	static const FName ScaledName = TEXT("Scaled");
 
-	if (Ar.IsLoading())
+	switch (InType)
 	{
-		switch ((ImplicitObjectType)ObjectType)
-		{
-		case ImplicitObjectType::Sphere: { ImplicitObject = new TSphere<T, d>(); break; }
-		case ImplicitObjectType::Box: { ImplicitObject = new TBox<T, d>(); break; }
-		case ImplicitObjectType::Plane: { ImplicitObject = new TPlane<T, d>(); break; }
-		case ImplicitObjectType::Capsule: { ImplicitObject = new TCapsule<T>(); break; }
-		case ImplicitObjectType::Transformed: { ImplicitObject = new TImplicitObjectTransformed<T, d>(); break; }
-		case ImplicitObjectType::Union: { ImplicitObject = new TImplicitObjectUnion<T, d>(); break; }
-		case ImplicitObjectType::LevelSet: { ImplicitObject = new TLevelSet<T, d>(); break; }
-		case ImplicitObjectType::Convex: { ImplicitObject = new TConvex<T, d>(); break; }
-		case ImplicitObjectType::TaperedCylinder: { ImplicitObject = new TTaperedCylinder<T>(); break; }
-		case ImplicitObjectType::TriangleMesh: { ImplicitObject = new TTriangleMeshImplicitObject<T>(); break; }
-		default:
-			check(false);
-		}
-
-		Serializable.SetFromRawLowLevel(ImplicitObject);
+		case ImplicitObjectType::Sphere: return SphereName;
+		case ImplicitObjectType::Box: return BoxName;
+		case ImplicitObjectType::Plane: return PlaneName;
+		case ImplicitObjectType::Capsule: return CapsuleName;
+		case ImplicitObjectType::Transformed: return TransformedName;
+		case ImplicitObjectType::Union: return UnionName;
+		case ImplicitObjectType::LevelSet: return LevelSetName;
+		case ImplicitObjectType::Unknown: return UnknownName;
+		case ImplicitObjectType::Convex: return ConvexName;
+		case ImplicitObjectType::TaperedCylinder: return TaperedCylinderName;
+		case ImplicitObjectType::Cylinder: return CylinderName;
+		case ImplicitObjectType::TriangleMesh: return TriangleMeshName;
+		case ImplicitObjectType::HeightField: return HeightFieldName;
+		case ImplicitObjectType::Scaled: return ScaledName;
 	}
-	else
+	return NAME_None;
+}
+
+template<typename T, int d>
+TImplicitObject<T, d>* TImplicitObject<T, d>::SerializationFactory(FChaosArchive& Ar, TImplicitObject<T, d>* Obj)
+{
+	int8 ObjectType = Ar.IsLoading() ? 0 : (int8)Obj->Type;
+	Ar << ObjectType;
+	switch ((ImplicitObjectType)ObjectType)
 	{
-		switch ((ImplicitObjectType)ObjectType)
-		{
-		case ImplicitObjectType::Sphere:
-		case ImplicitObjectType::Box:
-		case ImplicitObjectType::Plane:
-		case ImplicitObjectType::Capsule:
-		case ImplicitObjectType::Transformed:
-		case ImplicitObjectType::Union:
-		case ImplicitObjectType::LevelSet:
-		case ImplicitObjectType::Convex:
-		case ImplicitObjectType::TaperedCylinder:
-		case ImplicitObjectType::TriangleMesh:
-			break;
-		default:
-			check(false);
-		}
+	case ImplicitObjectType::Sphere: if (Ar.IsLoading()) { return new TSphere<T, d>(); } break;
+	case ImplicitObjectType::Box: if (Ar.IsLoading()) { return new TBox<T, d>(); } break;
+	case ImplicitObjectType::Plane: if (Ar.IsLoading()) { return new TPlane<T, d>(); } break;
+	case ImplicitObjectType::Capsule: if (Ar.IsLoading()) { return new TCapsule<T>(); } break;
+	case ImplicitObjectType::Transformed: if (Ar.IsLoading()) { return new TImplicitObjectTransformed<T, d>(); } break;
+	case ImplicitObjectType::Union: if (Ar.IsLoading()) { return new TImplicitObjectUnion<T, d>(); } break;
+	case ImplicitObjectType::LevelSet: if (Ar.IsLoading()) { return new TLevelSet<T, d>(); } break;
+	case ImplicitObjectType::Convex: if (Ar.IsLoading()) { return new TConvex<T, d>(); } break;
+	case ImplicitObjectType::TaperedCylinder: if (Ar.IsLoading()) { return new TTaperedCylinder<T>(); } break;
+	case ImplicitObjectType::TriangleMesh: if (Ar.IsLoading()) { return new TTriangleMeshImplicitObject<T>(); } break;
+	case ImplicitObjectType::Scaled: if (Ar.IsLoading()) { return new TImplicitObjectScaled<T,d>(); } break;
+	case ImplicitObjectType::HeightField: if (Ar.IsLoading()) { return new THeightField<T>(); } break;
+	case ImplicitObjectType::Cylinder: if (Ar.IsLoading()) { return new TCylinder<T>(); } break;
+	default:
+		check(false);
 	}
-	ImplicitObject->Serialize(Ar);
+	return nullptr;
 }
 
 //template class Chaos::TImplicitObject<float, 2>; // Missing 2D Rotation class
-template class Chaos::TImplicitObject<float, 3>;
\ No newline at end of file
+template class Chaos::TImplicitObject<float, 3>;
diff --git a/Engine/Source/Runtime/Experimental/Chaos/Private/Chaos/KinematicGeometryParticles.cpp b/Engine/Source/Runtime/Experimental/Chaos/Private/Chaos/KinematicGeometryParticles.cpp
--- a/Engine/Source/Runtime/Experimental/Chaos/Private/Chaos/KinematicGeometryParticles.cpp
+++ b/Engine/Source/Runtime/Experimental/Chaos/Private/Chaos/KinematicGeometryParticles.cpp
@@ -19,10 +19,23 @@ namespace Chaos
 		//static_assert(SimType == EGeometryParticlesSimType::RigidBodySim, "Handles require a rigid body sim type");
 		return static_cast<TKinematicGeometryParticlesImp<T, d, SimType>::THandleType*>(TGeometryParticlesImp<T, d, SimType>::Handle(Index));
 	}
+
+
+	template<class T, int d, EGeometryParticlesSimType SimType>
+	CHAOS_API TKinematicGeometryParticlesImp<T, d, SimType>::~TKinematicGeometryParticlesImp<T, d, SimType>()
+	{
+
+	}
+
 }
 
+#ifdef __clang__
 template class CHAOS_API Chaos::TKinematicGeometryParticlesImp<float, 3, Chaos::EGeometryParticlesSimType::RigidBodySim>;
 template class CHAOS_API Chaos::TKinematicGeometryParticlesImp<float, 3, Chaos::EGeometryParticlesSimType::Other>;
+#else
+template class Chaos::TKinematicGeometryParticlesImp<float, 3, Chaos::EGeometryParticlesSimType::RigidBodySim>;
+template class Chaos::TKinematicGeometryParticlesImp<float, 3, Chaos::EGeometryParticlesSimType::Other>;
+#endif
 
 // Re-enable when double precision is able to compile
-//template Chaos::TKinematicGeometryParticlesImp<double, 3>;
\ No newline at end of file
+//template Chaos::TKinematicGeometryParticlesImp<double, 3>;
diff --git a/Engine/Source/Runtime/Experimental/Chaos/Private/Chaos/Levelset.cpp b/Engine/Source/Runtime/Experimental/Chaos/Private/Chaos/Levelset.cpp
--- a/Engine/Source/Runtime/Experimental/Chaos/Private/Chaos/Levelset.cpp
+++ b/Engine/Source/Runtime/Experimental/Chaos/Private/Chaos/Levelset.cpp
@@ -9,6 +9,7 @@
 #include <vector>
 #include "Chaos/ErrorReporter.h"
 #include "Chaos/MassProperties.h"
+#include "Chaos/TriangleMesh.h"
 
 int32 OutputFailedLevelSetDebugData = 0;
 FAutoConsoleVariableRef CVarOutputFailedLevelSetDebugData(TEXT("p.LevelSetOutputFailedDebugData"), OutputFailedLevelSetDebugData, TEXT("Output debug obj files for level set and mesh when error tolerances are too high"));
diff --git a/Engine/Source/Runtime/Experimental/Chaos/Private/Chaos/MassProperties.cpp b/Engine/Source/Runtime/Experimental/Chaos/Private/Chaos/MassProperties.cpp
--- a/Engine/Source/Runtime/Experimental/Chaos/Private/Chaos/MassProperties.cpp
+++ b/Engine/Source/Runtime/Experimental/Chaos/Private/Chaos/MassProperties.cpp
@@ -2,6 +2,8 @@
 #include "Chaos/MassProperties.h"
 #include "Chaos/Rotation.h"
 #include "Chaos/Matrix.h"
+#include "Chaos/Particles.h"
+#include "Chaos/TriangleMesh.h"
 
 namespace Chaos
 {
@@ -17,13 +19,13 @@ namespace Chaos
 		if (!ensure(Trace > SMALL_NUMBER))
 		{
 			// Tiny inertia - numerical instability would follow. We should not get this unless we have bad input.
-			return TRotation<T, d>(TVector<T, d>(0), 1);
+			return TRotation<T, d>::FromElements(TVector<T, d>(0), 1);
 		}
 
 		if ((OffDiagSize / Trace) < SMALL_NUMBER)
 		{
 			// Almost diagonal matrix - we are already in local space.
-			return TRotation<T, d>(TVector<T, d>(0), 1);
+			return TRotation<T, d>::FromElements(TVector<T, d>(0), 1);
 		}
 
 		T Size = FMath::Sqrt((FMath::Square(Inertia.M[0][0] - Trace) + FMath::Square(Inertia.M[1][1] - Trace) + FMath::Square(Inertia.M[2][2] - Trace) + 2 * OffDiagSize) / 6);
@@ -52,7 +54,7 @@ namespace Chaos
 		if (!ensure((SqrtIM1Scale0 > KINDA_SMALL_NUMBER) || (SqrtIM1Scale1 > KINDA_SMALL_NUMBER)))
 		{
 			// We hit numerical accuracy, despite the early off-diagonal check. We should not see this any more.
-			return TRotation<T, d>(TVector<T, d>(0), 1);
+			return TRotation<T, d>::FromElements(TVector<T, d>(0), 1);
 		}
 
 		TVector<T, d - 1> SmallEigenvector2 = IM1Scale0 > IM1Scale1 ? (TVector<T, d - 1>(IM1.M[3], -IM1.M[1]) / SqrtIM1Scale0) : (IM1Scale1 > 0 ? (TVector<T, d - 1>(-IM1.M[1], IM1.M[0]) / SqrtIM1Scale1) : TVector<T, d - 1>(1, 0));
