Made tri meshes serialize their BV instead of building on demand

diff --git a/Engine/Source/Runtime/Core/Public/UObject/ExternalPhysicsCustomObjectVersion.h b/Engine/Source/Runtime/Core/Public/UObject/ExternalPhysicsCustomObjectVersion.h
--- a/Engine/Source/Runtime/Core/Public/UObject/ExternalPhysicsCustomObjectVersion.h
+++ b/Engine/Source/Runtime/Core/Public/UObject/ExternalPhysicsCustomObjectVersion.h
@@ -38,6 +38,9 @@ struct CORE_API FExternalPhysicsCustomObjectVersion
 
 		// Add kinematic targets to TKinematicGeometryParticles
 		KinematicTargets,
+		
+		// Allow trimeshes to serialize their acceleration structure
+		TrimeshSerializesBV,
 
 		// -----<new versions can be added above this line>-------------------------------------------------
 		VersionPlusOne,
diff --git a/Engine/Source/Runtime/Engine/Private/PhysicsEngine/Experimental/ChaosDerivedData.cpp b/Engine/Source/Runtime/Engine/Private/PhysicsEngine/Experimental/ChaosDerivedData.cpp
--- a/Engine/Source/Runtime/Engine/Private/PhysicsEngine/Experimental/ChaosDerivedData.cpp
+++ b/Engine/Source/Runtime/Engine/Private/PhysicsEngine/Experimental/ChaosDerivedData.cpp
@@ -21,7 +21,7 @@ const TCHAR* FChaosDerivedDataCooker::GetPluginName() const
 
 const TCHAR* FChaosDerivedDataCooker::GetVersionString() const
 {
-	return TEXT("UN4SF7NU5G9O0X2LHNHLBG0R0CIS3TON");
+	return TEXT("DFE79D57413E4527AD10CEEF771A4654");
 }
 
 FString FChaosDerivedDataCooker::GetPluginSpecificCacheKeySuffix() const
diff --git a/Engine/Source/Runtime/Experimental/Chaos/Public/Chaos/TriangleMeshImplicitObject.h b/Engine/Source/Runtime/Experimental/Chaos/Public/Chaos/TriangleMeshImplicitObject.h
--- a/Engine/Source/Runtime/Experimental/Chaos/Public/Chaos/TriangleMeshImplicitObject.h
+++ b/Engine/Source/Runtime/Experimental/Chaos/Public/Chaos/TriangleMeshImplicitObject.h
@@ -64,28 +64,16 @@ namespace Chaos
 				Ar << ConvexHulls;
 			}
 
-			if(Ar.IsLoading())
+			if(Ar.CustomVer(FExternalPhysicsCustomObjectVersion::GUID) < FExternalPhysicsCustomObjectVersion::TrimeshSerializesBV)
 			{
+				// Should now only hit when loading older trimeshes
 				RebuildBV();
 			}
-#if 0
-#if 0
-			// Disabled during 2-1 replacement. Replace when BV is actually working again
-			Ar << ConvexHulls;
-			Ar << BVH;
-
-			if(Ar.IsLoading())
-			{
-				// Re-link the object array when we load this back in
-				BVH.SetObjects(ConvexHulls);
-			}
-#else
-			if(Ar.IsLoading())
+			else
 			{
-				RebuildBV();
+				// Serialize acceleration
+				Ar << BVH;
 			}
-#endif
-#endif
 		}
 
 		virtual void Serialize(FChaosArchive& Ar) override
