Collection factory generates acceleration per bucket

diff --git a/Engine/Source/Runtime/Experimental/Chaos/Private/Chaos/PBDRigidsEvolution.cpp b/Engine/Source/Runtime/Experimental/Chaos/Private/Chaos/PBDRigidsEvolution.cpp
--- a/Engine/Source/Runtime/Experimental/Chaos/Private/Chaos/PBDRigidsEvolution.cpp
+++ b/Engine/Source/Runtime/Experimental/Chaos/Private/Chaos/PBDRigidsEvolution.cpp
@@ -33,59 +33,53 @@ namespace Chaos
 	template<typename T, int d>
 	struct TDefaultCollectionFactory : public ISpatialAccelerationCollectionFactory<T, d>
 	{
+		using BVType = TBoundingVolume<TAccelerationStructureHandle<T, d>, T, d>;
+		using AABBTreeType = TAABBTree<TAccelerationStructureHandle<T, d>, TAABBTreeLeafArray<TAccelerationStructureHandle<T, d>, T>, T>;
+		using AABBTreeOfGridsType = TAABBTree<TAccelerationStructureHandle<T, d>, TBoundingVolume<TAccelerationStructureHandle<T, d>, T, d>, T>;
+
 		virtual TUniquePtr<ISpatialAccelerationCollection<TAccelerationStructureHandle<T, d>, T, d>> CreateEmptyCollection() override
 		{
+			const uint16 NumBuckets = BroadphaseType == 3 ? 2 : 1;
+			auto Collection = new TSpatialAccelerationCollection<AABBTreeType, BVType, AABBTreeOfGridsType>();
 			const TArray<TAccelerationStructureBuilder<T, d>> Empty;
-			if (BroadphaseType == 0)
-			{
-				using AccelType = TBoundingVolume<TAccelerationStructureHandle<T, d>, T, d>;
-				TArray<TSpatialAccelerationParams<TArray<TAccelerationStructureBuilder<T, d>>, TAccelerationStructureHandle<T, d>, T, d>> Buckets = { {AccelType::StaticType, MaxPayloadSize, BoundingVolumeNumCells } };
 
-				auto Structure = MakeUnique<AccelType>(Empty, false, 0, Buckets[0].MaxCells, Buckets[0].MaxPayloadBounds);
-				auto Collection = new TSpatialAccelerationCollection<TArray<TAccelerationStructureBuilder<T, d>>, AccelType>(MoveTemp(Buckets));
+			for (uint16 BucketIdx = 0; BucketIdx < NumBuckets; ++BucketIdx)
+			{
+				Collection->AddSubstructure(CreateAccelerationPerBucket(Empty, BucketIdx), BucketIdx);
+			}
 
-				Collection->AddSubstructure(MoveTemp(Structure), 0);
+			return TUniquePtr<ISpatialAccelerationCollection<TAccelerationStructureHandle<T, d>, T, d>>(Collection);
+		}
 
-				return TUniquePtr<ISpatialAccelerationCollection<TAccelerationStructureHandle<T, d>, T, d>>(Collection);
-			}
-			else if (BroadphaseType == 1)
+		virtual TUniquePtr<ISpatialAcceleration<TAccelerationStructureHandle<T, d>, T, d>> CreateAccelerationPerBucket(const TArray<TAccelerationStructureBuilder<T, d>>& Particles, uint16 BucketIdx) override
+		{
+			switch (BucketIdx)
 			{
-				using AccelType = TAABBTree<TAccelerationStructureHandle<T, d>, TAABBTreeLeafArray<TAccelerationStructureHandle<T, d>, T>, T>;
-				TArray<TSpatialAccelerationParams<TArray<TAccelerationStructureBuilder<T, d>>, TAccelerationStructureHandle<T, d>, T, d>> Buckets = { {AccelType::StaticType, MaxPayloadSize, 0, MaxChildrenInLeaf, MaxTreeDepth} };
-
-				auto Structure = MakeUnique<AccelType>(Empty, Buckets[0].MaxChildrenInLeaf, Buckets[0].MaxTreeDepth, Buckets[0].MaxPayloadBounds);
-				auto Collection = new TSpatialAccelerationCollection<TArray<TAccelerationStructureBuilder<T, d>>, AccelType>(MoveTemp(Buckets));
-				Collection->AddSubstructure(MoveTemp(Structure), 0);
-				return TUniquePtr<ISpatialAccelerationCollection<TAccelerationStructureHandle<T, d>, T, d>>(Collection);
+			case 0:
+			{
+				if (BroadphaseType == 0)
+				{
+					return MakeUnique<BVType>(Particles, false, 0, BoundingVolumeNumCells, MaxPayloadSize);
+				}
+				else if(BroadphaseType == 1 || BroadphaseType == 3)
+				{
+					return MakeUnique<AABBTreeType>(Particles, MaxChildrenInLeaf, MaxTreeDepth, MaxPayloadSize);
+				}
+				else
+				{
+					return MakeUnique<AABBTreeOfGridsType>(Particles, AABBMaxChildrenInLeaf, AABBMaxTreeDepth, MaxPayloadSize);
+				}
 			}
-			else if (BroadphaseType == 2)
+			case 1:
 			{
-				using AccelType = TAABBTree<TAccelerationStructureHandle<T, d>, TBoundingVolume<TAccelerationStructureHandle<T, d>, T, d>, T>;
-				TArray<TSpatialAccelerationParams<TArray<TAccelerationStructureBuilder<T, d>>, TAccelerationStructureHandle<T, d>, T, d>> Buckets = { {AccelType::StaticType, MaxPayloadSize, 0, AABBMaxChildrenInLeaf, AABBMaxTreeDepth} };
-
-				auto Structure = MakeUnique<AccelType>(Empty, Buckets[0].MaxChildrenInLeaf, Buckets[0].MaxTreeDepth, Buckets[0].MaxPayloadBounds);
-				auto Collection = new TSpatialAccelerationCollection<TArray<TAccelerationStructureBuilder<T, d>>, AccelType>(MoveTemp(Buckets));
-				Collection->AddSubstructure(MoveTemp(Structure), 0);
-				return TUniquePtr<ISpatialAccelerationCollection<TAccelerationStructureHandle<T, d>, T, d>>(Collection);
+				ensure(BroadphaseType == 3);
+				return MakeUnique<BVType>(Particles, false, 0, BoundingVolumeNumCells, MaxPayloadSize);
 			}
-			else
+			default:
 			{
-				using AccelType = TAABBTree<TAccelerationStructureHandle<T, d>, TAABBTreeLeafArray<TAccelerationStructureHandle<T, d>, T>, T>;
-				using GridType = TBoundingVolume<TAccelerationStructureHandle<T, d>, T, d>;
-
-				TArray<TSpatialAccelerationParams<TArray<TAccelerationStructureBuilder<T, d>>, TAccelerationStructureHandle<T, d>, T, d>> Buckets = {
-					{AccelType::StaticType, MaxPayloadSize, 0, AABBMaxChildrenInLeaf, AABBMaxTreeDepth},
-					{GridType::StaticType, MaxPayloadSize * 100, BoundingVolumeNumCells, AABBMaxChildrenInLeaf, AABBMaxTreeDepth},
-				};
-
-				auto Structure0 = MakeUnique<AccelType>(Empty, Buckets[0].MaxChildrenInLeaf, Buckets[0].MaxTreeDepth, Buckets[0].MaxPayloadBounds);
-
-				auto Structure1 = MakeUnique<GridType>(Empty, false, 0, Buckets[1].MaxCells, Buckets[1].MaxPayloadBounds);
-				auto Collection = new TSpatialAccelerationCollection<TArray<TAccelerationStructureBuilder<T, d>>, AccelType, GridType>(MoveTemp(Buckets));
-				Collection->AddSubstructure(MoveTemp(Structure0), 0);
-				Collection->AddSubstructure(MoveTemp(Structure1), 1);
-
-				return TUniquePtr<ISpatialAccelerationCollection<TAccelerationStructureHandle<T, d>, T, d>>(Collection);
+				check(false);
+				return nullptr;
+			}
 			}
 		}
 
@@ -141,10 +135,13 @@ namespace Chaos
 	DECLARE_CYCLE_STAT(TEXT("ComputeIntermediateSpatialAcceleration"), STAT_ComputeIntermediateSpatialAcceleration, STATGROUP_Chaos);
 
 	template<class FPBDRigidsEvolution, class FPBDCollisionConstraint, class T, int d>
-	TPBDRigidsEvolutionBase<FPBDRigidsEvolution, FPBDCollisionConstraint, T, d>::FChaosAccelerationStructureTask::FChaosAccelerationStructureTask(const TArray<FAccelerationStructBuilderCache>& CacheMap
+	TPBDRigidsEvolutionBase<FPBDRigidsEvolution, FPBDCollisionConstraint, T, d>::FChaosAccelerationStructureTask::FChaosAccelerationStructureTask(
+		ISpatialAccelerationCollectionFactory<T, d>& InSpatialCollectionFactory
+		, const TArray<FAccelerationStructBuilderCache>& CacheMap
 		, TUniquePtr<FAccelerationStructure>& InAccelerationStructure
 		, TUniquePtr<FAccelerationStructure>& InAccelerationStructureCopy)
-		: BuilderCacheMap(CacheMap)
+		: SpatialCollectionFactory(InSpatialCollectionFactory)
+		, BuilderCacheMap(CacheMap)
 		, AccelerationStructure(InAccelerationStructure)
 		, AccelerationStructureCopy(InAccelerationStructureCopy)
 	{
@@ -176,23 +173,20 @@ namespace Chaos
 		
 		if (Substructure->template As<BVType>())
 		{
-			TArray<TSpatialAccelerationParams<TArray<TAccelerationStructureBuilder<T,d>>, TAccelerationStructureHandle<T, d>, T, d>> Buckets = { {BVType::StaticType } };
-			auto Collection = MakeUnique<TSpatialAccelerationCollection<TArray<TAccelerationStructureBuilder<T, d>>, BVType>>(MoveTemp(Buckets));
+			auto Collection = MakeUnique<TSpatialAccelerationCollection<BVType>>();
 			Collection->AddSubstructure(MoveTemp(Substructure), 0);
 			return Collection;
 		}
 		else if (Substructure->template As<AABBType>())
 		{
-			TArray<TSpatialAccelerationParams<TArray<TAccelerationStructureBuilder<T,d>>, TAccelerationStructureHandle<T, d>, T, d>> Buckets = { {AABBType::StaticType } };
-			auto Collection = MakeUnique<TSpatialAccelerationCollection<TArray<TAccelerationStructureBuilder<T, d>>, AABBType>>(MoveTemp(Buckets));
+			auto Collection = MakeUnique<TSpatialAccelerationCollection<AABBType>>();
 			Collection->AddSubstructure(MoveTemp(Substructure), 0);
 			return Collection;
 		}
 		else
 		{
 			using AccelType = TAABBTree<TAccelerationStructureHandle<T, d>, TBoundingVolume<TAccelerationStructureHandle<T, d>, T, d>, T>;
-			TArray<TSpatialAccelerationParams<TArray<TAccelerationStructureBuilder<T,d>>, TAccelerationStructureHandle<T, d>, T, d>> Buckets = { {AccelType::StaticType } };
-			auto Collection = MakeUnique<TSpatialAccelerationCollection<TArray<TAccelerationStructureBuilder<T, d>>, AccelType>>(MoveTemp(Buckets));
+			auto Collection = MakeUnique<TSpatialAccelerationCollection<AccelType>>();
 			Collection->AddSubstructure(MoveTemp(Substructure), 0);
 			return Collection;
 		}
@@ -205,25 +199,8 @@ namespace Chaos
 		for (const FAccelerationStructBuilderCache& Cache : BuilderCacheMap)
 		{
 			auto OldStruct = AccelerationStructure->RemoveSubstructure(Cache.SpatialIdx);
-			//This is a hack, need to have a Reinitialize API instead of downcasting
-			using BVType = TBoundingVolume<TAccelerationStructureHandle<T, d>, T, d>;
-			using AABBType = TAABBTree<TAccelerationStructureHandle<T, d>, TAABBTreeLeafArray<TAccelerationStructureHandle<T, d>, T>, T>;
-			ISpatialAcceleration<TAccelerationStructureHandle<T,d>,T,d>* NewStruct;
-			if (OldStruct->template As<BVType>())
-			{
-				NewStruct = new BVType(*Cache.CachedSpatialBuilderData);
-			}
-			else if (OldStruct->template As<AABBType>())
-			{
-				NewStruct = new AABBType(*Cache.CachedSpatialBuilderData);
-			}
-			else
-			{
-				using AccelType = TAABBTree<TAccelerationStructureHandle<T, d>, TBoundingVolume<TAccelerationStructureHandle<T, d>, T, d>, T>;
-				NewStruct = new AccelType(*Cache.CachedSpatialBuilderData);
-			}
-
-			AccelerationStructure->AddSubstructure(TUniquePtr<ISpatialAcceleration<TAccelerationStructureHandle<T, d>, T,d>>(NewStruct), Cache.SpatialIdx.Bucket); //this assumes bucket size is 1, which will not be true in future. Need to fix all this
+			auto NewStruct = SpatialCollectionFactory.CreateAccelerationPerBucket(*Cache.CachedSpatialBuilderData, Cache.SpatialIdx.Bucket);
+			AccelerationStructure->AddSubstructure(MoveTemp(NewStruct), Cache.SpatialIdx.Bucket);
 		}
 		AccelerationStructureCopy = AsUniqueSpatialAccelerationChecked<FAccelerationStructure>(AccelerationStructure->Copy());
 	}
@@ -375,7 +352,7 @@ namespace Chaos
 				bExternalReady = true;
 			}
 
-			AccelerationStructureTaskComplete = TGraphTask<FChaosAccelerationStructureTask>::CreateTask().ConstructAndDispatchWhenReady(CachedSpatialBuilderDataMap, AsyncInternalAcceleration, AsyncExternalAcceleration);
+			AccelerationStructureTaskComplete = TGraphTask<FChaosAccelerationStructureTask>::CreateTask().ConstructAndDispatchWhenReady(*SpatialCollectionFactory, CachedSpatialBuilderDataMap, AsyncInternalAcceleration, AsyncExternalAcceleration);
 		}
 		else
 		{
diff --git a/Engine/Source/Runtime/Experimental/Chaos/Public/Chaos/ISpatialAccelerationCollection.h b/Engine/Source/Runtime/Experimental/Chaos/Public/Chaos/ISpatialAccelerationCollection.h
--- a/Engine/Source/Runtime/Experimental/Chaos/Public/Chaos/ISpatialAccelerationCollection.h
+++ b/Engine/Source/Runtime/Experimental/Chaos/Public/Chaos/ISpatialAccelerationCollection.h
@@ -9,30 +9,13 @@
 namespace Chaos
 {
 
-/** Used for the creation of different acceleration structures */
-template <typename TParticleView, typename TPayloadType, typename T, int d>
-struct CHAOS_API TSpatialAccelerationParams
-{
-	SpatialAccelerationType Type;	//Specifies the acceleration type. Unknown can be used for custom types, see TSpatialAccelerationParams
-
-	//Union of all params, some are only applicable to subset of structure types
-	float MaxPayloadBounds;	//max bounds size before inserting into a linear list of "large" objects. This list will always be checked
-
-	int32 MaxCells;	//(grid only) max number of cells in grid structure
-	int32 MaxChildrenInLeaf; //(tree only) if the tree sees MaxChildrenInLeaf or fewer objects, no more splitting occurs
-	int32 MaxTreeDepth;	//(tree only) if the tree has reached MaxTreeDepth, no more splitting occurs
-
-	//typedef ISpatialAcceleration<TPayloadType, T, d> (*ConstructFn)(TSpatialAccelerationParams<TPayloadType, T, d>& Params, int32 ParamIdx);
-	//ConstructFn ConstructionFactory;	//if unknown type is specified, use function pointer set by user
-};
-
 template <typename TPayloadType, typename T, int d>
 class CHAOS_API ISpatialAccelerationCollection : public ISpatialAcceleration<TPayloadType, T, d>
 {
 public:
 	ISpatialAccelerationCollection() : ISpatialAcceleration<TPayloadType, T, d>(StaticType) {}
 	static constexpr ESpatialAcceleration StaticType = ESpatialAcceleration::Collection;
-	virtual FSpatialAccelerationIdx AddSubstructure(TUniquePtr<ISpatialAcceleration<TPayloadType, T, d>>&& Substructure, int32 Bucket) = 0;
+	virtual FSpatialAccelerationIdx AddSubstructure(TUniquePtr<ISpatialAcceleration<TPayloadType, T, d>>&& Substructure, uint16 Bucket) = 0;
 	virtual TUniquePtr<ISpatialAcceleration<TPayloadType, T, d>> RemoveSubstructure(FSpatialAccelerationIdx Idx) = 0;
 	virtual ISpatialAcceleration<TPayloadType, T, d>* GetSubstructure(FSpatialAccelerationIdx Idx) = 0;
 
@@ -337,7 +320,7 @@ typename TEnableIf<!TIsSame<typename SpatialAccelerationCollection::TPayloadType
 {
 }
 
-template <typename TParticleView, typename ... TSpatialAccelerationTypes>
+template <typename ... TSpatialAccelerationTypes>
 class CHAOS_API TSpatialAccelerationCollection : public
 	ISpatialAccelerationCollection<typename std::tuple_element<0, std::tuple< TSpatialAccelerationTypes...>>::type::PayloadType,
 	typename std::tuple_element<0, std::tuple< TSpatialAccelerationTypes...>>::type::TType,
@@ -349,21 +332,15 @@ public:
 	using T = typename FirstAccelerationType::TType;
 	static constexpr int d = FirstAccelerationType::D;
 
-	TSpatialAccelerationCollection(TArray<TSpatialAccelerationParams<TParticleView, TPayloadType, T, d>>&& InBucketsParams)
-		: NumBuckets(InBucketsParams.Num())
-		, BucketsParams(MoveTemp(InBucketsParams))
+	TSpatialAccelerationCollection()
+		: LastBucket(0)
 	{
-		check(NumBuckets < MaxBuckets);
-		int32 ParamIdx = 0;
-		for (const auto Params : BucketsParams)
-		{
-			check((TSpatialAccelerationCollectionHelper<0, NumTypes, decltype(Types), TPayloadType, T, d>::FindTypeIdx(Types, Params.Type) < NumTypes));
-		}
 	}
 
-	virtual FSpatialAccelerationIdx AddSubstructure(TUniquePtr<ISpatialAcceleration<TPayloadType, T, d>>&& Substructure, int32 BucketIdx) override
+	virtual FSpatialAccelerationIdx AddSubstructure(TUniquePtr<ISpatialAcceleration<TPayloadType, T, d>>&& Substructure, uint16 BucketIdx) override
 	{
-		check(BucketIdx < NumBuckets);
+		check(BucketIdx < MaxBuckets);
+		LastBucket = FMath::Max(LastBucket, BucketIdx);
 		//todo: verify the bucket and type match
 		FSpatialAccelerationIdx Result;
 		Result.Bucket = BucketIdx;
@@ -392,7 +369,6 @@ public:
 	virtual TUniquePtr <ISpatialAcceleration<TPayloadType, T, d>> RemoveSubstructure(FSpatialAccelerationIdx Idx) override
 	{
 		TUniquePtr <ISpatialAcceleration<TPayloadType, T, d>> Removed;
-		check(Idx.Bucket < NumBuckets);
 		{
 			auto& BucketEntry = Buckets[Idx.Bucket].Objects[Idx.InnerIdx];
 			const int32 TypeIdx = GetTypeIdx(BucketEntry.Acceleration.Get());
@@ -410,6 +386,16 @@ public:
 			Removed = MoveTemp(BucketEntry.Acceleration);
 		}
 		Buckets[Idx.Bucket].Remove(Idx.InnerIdx);
+
+		LastBucket = 0;
+		for (uint16 BucketIdx = Idx.Bucket; BucketIdx > 0; --BucketIdx)
+		{
+			if (Buckets[BucketIdx].Objects.Num())
+			{
+				LastBucket = BucketIdx;
+				break;
+			}
+		}
 		return Removed;
 	}
 
@@ -507,21 +493,19 @@ public:
 
 	virtual void RemoveElementFrom(const TPayloadType& Payload, FSpatialAccelerationIdx SpatialIdx) override
 	{
-		const uint16 LastBucket = NumBuckets - 1;
 		const uint16 UseBucket = FMath::Min(LastBucket, SpatialIdx.Bucket);
 		Buckets[UseBucket].Objects[SpatialIdx.InnerIdx].Acceleration->RemoveElement(Payload);
 	}
 
 	virtual void UpdateElementIn(const TPayloadType& Payload, const TBox<T, d>& NewBounds, bool bHasBounds, FSpatialAccelerationIdx SpatialIdx)
 	{
-		const uint16 LastBucket = NumBuckets - 1;
 		const uint16 UseBucket = FMath::Min(LastBucket, SpatialIdx.Bucket);
 		Buckets[UseBucket].Objects[SpatialIdx.InnerIdx].Acceleration->UpdateElement(Payload, NewBounds, bHasBounds);
 	}
 
 	virtual TUniquePtr<ISpatialAcceleration<TPayloadType, T, d>> Copy() const
 	{
-		return TUniquePtr<ISpatialAcceleration<TPayloadType, T, d>>(new TSpatialAccelerationCollection<TParticleView, TSpatialAccelerationTypes...>(*this));
+		return TUniquePtr<ISpatialAcceleration<TPayloadType, T, d>>(new TSpatialAccelerationCollection<TSpatialAccelerationTypes...>(*this));
 	}
 
 	virtual void PBDComputeConstraintsLowLevel_GatherStats(TPBDCollisionConstraint<T, d>& CollisionConstraint, T Dt) const override
@@ -567,8 +551,8 @@ public:
 
 private:
 
-	TSpatialAccelerationCollection(const TSpatialAccelerationCollection<TParticleView, TSpatialAccelerationTypes...>& Other)
-		: NumBuckets(Other.NumBuckets)
+	TSpatialAccelerationCollection(const TSpatialAccelerationCollection<TSpatialAccelerationTypes...>& Other)
+		: LastBucket(Other.LastBucket)
 	{
 		//need to make a deep copy of the unique ptrs and then update raw ptrs
 		TSpatialAccelerationCollectionHelper<0, NumTypes, decltype(Types), TPayloadType, T, d>::SetNumFrom(Other.Types, Types);
@@ -652,8 +636,7 @@ private:
 	TSpatialCollectionBucket<TSpatialAccelerationBucketEntry<TPayloadType, T, d>> Buckets[MaxBuckets];
 	TSpatialTypeTuple< TSpatialAccelerationTypes...> Types;
 	static constexpr uint32 NumTypes = sizeof...(TSpatialAccelerationTypes);
-	uint16 NumBuckets;
-	const TArray<TSpatialAccelerationParams<TParticleView, TPayloadType, T, d>> BucketsParams;
+	uint16 LastBucket;	//the last bucket that has an acceleration structure
 };
 
 
diff --git a/Engine/Source/Runtime/Experimental/Chaos/Public/Chaos/PBDRigidsEvolution.h b/Engine/Source/Runtime/Experimental/Chaos/Public/Chaos/PBDRigidsEvolution.h
--- a/Engine/Source/Runtime/Experimental/Chaos/Public/Chaos/PBDRigidsEvolution.h
+++ b/Engine/Source/Runtime/Experimental/Chaos/Public/Chaos/PBDRigidsEvolution.h
@@ -22,18 +22,6 @@ class FChaosArchive;
 template <typename TPayload, typename T, int d>
 class ISpatialAccelerationCollection;
 
-template <typename T, int d>
-struct CHAOS_API ISpatialAccelerationCollectionFactory
-{
-	//Create an empty acceleration collection with the desired buckets. Chaos uses the bucket params inside the collection to add / remove sub-structures as needed
-	virtual TUniquePtr<ISpatialAccelerationCollection<TAccelerationStructureHandle<T, d>, T, d>> CreateEmptyCollection() = 0;
-
-	//Serialize the collection in and out
-	virtual void Serialize(TUniquePtr<ISpatialAccelerationCollection<TAccelerationStructureHandle<T, d>, T, d>>& Ptr, FChaosArchive& Ar) = 0;
-
-	virtual ~ISpatialAccelerationCollectionFactory() = default;
-};
-
 struct CHAOS_API FEvolutionStats
 {
 	int32 ActiveCollisionPoints;
@@ -90,6 +78,21 @@ struct TAccelerationStructureBuilder
 	}
 };
 
+template <typename T, int d>
+struct CHAOS_API ISpatialAccelerationCollectionFactory
+{
+	//Create an empty acceleration collection with the desired buckets. Chaos enqueues acceleration structure operations per bucket
+	virtual TUniquePtr<ISpatialAccelerationCollection<TAccelerationStructureHandle<T, d>, T, d>> CreateEmptyCollection() = 0;
+
+	//Chaos creates new acceleration structures per bucket. Factory can change underlying type, but changing the number of buckets requires re-registering factory
+	virtual TUniquePtr<ISpatialAcceleration<TAccelerationStructureHandle<T, d>, T, d>> CreateAccelerationPerBucket(const TArray<TAccelerationStructureBuilder<T, d>>& Particles, uint16 BucketIdx) = 0;
+
+	//Serialize the collection in and out
+	virtual void Serialize(TUniquePtr<ISpatialAccelerationCollection<TAccelerationStructureHandle<T, d>, T, d>>& Ptr, FChaosArchive& Ar) = 0;
+
+	virtual ~ISpatialAccelerationCollectionFactory() = default;
+};
+
 template<class FPBDRigidsEvolution, class FPBDCollisionConstraint, class T, int d>
 class TPBDRigidsEvolutionBase
 {
@@ -532,13 +535,14 @@ protected:
 	class FChaosAccelerationStructureTask
 	{
 	public:
-		FChaosAccelerationStructureTask(const TArray<FAccelerationStructBuilderCache>& BuilderCacheMap, TUniquePtr<FAccelerationStructure>& InAccelerationStructure,
+		FChaosAccelerationStructureTask(ISpatialAccelerationCollectionFactory<T,d>& InSpatialCollectionFactory, const TArray<FAccelerationStructBuilderCache>& BuilderCacheMap, TUniquePtr<FAccelerationStructure>& InAccelerationStructure,
 			TUniquePtr<FAccelerationStructure>& InAccelerationStructureCopy);
 		static FORCEINLINE TStatId GetStatId();
 		static FORCEINLINE ENamedThreads::Type GetDesiredThread();
 		static FORCEINLINE ESubsequentsMode::Type GetSubsequentsMode();
 		void DoTask(ENamedThreads::Type CurrentThread, const FGraphEventRef& MyCompletionGraphEvent);
 
+		ISpatialAccelerationCollectionFactory<T, d>& SpatialCollectionFactory;
 		const TArray<FAccelerationStructBuilderCache>& BuilderCacheMap;
 		TUniquePtr<FAccelerationStructure>& AccelerationStructure;
 		TUniquePtr<FAccelerationStructure>& AccelerationStructureCopy;
