Clean up cvars

diff --git a/Engine/Source/Runtime/Experimental/Chaos/Private/Chaos/PBDRigidsEvolution.cpp b/Engine/Source/Runtime/Experimental/Chaos/Private/Chaos/PBDRigidsEvolution.cpp
--- a/Engine/Source/Runtime/Experimental/Chaos/Private/Chaos/PBDRigidsEvolution.cpp
+++ b/Engine/Source/Runtime/Experimental/Chaos/Private/Chaos/PBDRigidsEvolution.cpp
@@ -9,37 +9,48 @@
 
 namespace Chaos
 {
-	int32 BroadphaseType = 3;
-	FAutoConsoleVariableRef CVarBroadphaseIsTree(TEXT("p.BroadphaseType"), BroadphaseType, TEXT(""));
-
-	int32 BoundingVolumeNumCells = 35;
-	FAutoConsoleVariableRef CVarBoundingVolumeNumCells(TEXT("p.BoundingVolumeNumCells"), BoundingVolumeNumCells, TEXT(""));
-
-	int32 MaxChildrenInLeaf = 5;
-	FAutoConsoleVariableRef CVarMaxChildrenInLeaf(TEXT("p.MaxChildrenInLeaf"), MaxChildrenInLeaf, TEXT(""));
-
-	int32 MaxTreeDepth = 200;
-	FAutoConsoleVariableRef CVarMaxTreeDepth(TEXT("p.MaxTreeDepth"), MaxTreeDepth, TEXT(""));
-
-	int32 AABBMaxChildrenInLeaf = 500;
-	FAutoConsoleVariableRef CVarAABBMaxChildrenInLeaf(TEXT("p.AABBMaxChildrenInLeaf"), AABBMaxChildrenInLeaf, TEXT(""));
-
-	int32 AABBMaxTreeDepth = 200;
-	FAutoConsoleVariableRef CVarAABBMaxTreeDepth(TEXT("p.AABBMaxTreeDepth"), AABBMaxTreeDepth, TEXT(""));
+	struct FAccelerationConfig
+	{
+		int32 BroadphaseType;
+		int32 BVNumCells;
+		int32 MaxChildrenInLeaf;
+		int32 MaxTreeDepth;
+		int32 AABBMaxChildrenInLeaf;
+		int32 AABBMaxTreeDepth;
+		float MaxPayloadSize;
+
+		FAccelerationConfig()
+		{
+			BroadphaseType = 3;
+			BVNumCells = 35;
+			MaxChildrenInLeaf = 5;
+			MaxTreeDepth = 200;
+			AABBMaxChildrenInLeaf = 500;
+			AABBMaxTreeDepth = 200;
+			MaxPayloadSize = 20000;
+		}
+	} ConfigSettings;
 
-	float MaxPayloadSize = 20000;
-	FAutoConsoleVariableRef CVarMaxPayloadSize(TEXT("p.MaxPayloadSize"), MaxPayloadSize, TEXT(""));
+	FAutoConsoleVariableRef CVarBroadphaseIsTree(TEXT("p.BroadphaseType"), ConfigSettings.BroadphaseType, TEXT(""));
+	FAutoConsoleVariableRef CVarBoundingVolumeNumCells(TEXT("p.BoundingVolumeNumCells"), ConfigSettings.BVNumCells, TEXT(""));
+	FAutoConsoleVariableRef CVarMaxChildrenInLeaf(TEXT("p.MaxChildrenInLeaf"), ConfigSettings.MaxChildrenInLeaf, TEXT(""));
+	FAutoConsoleVariableRef CVarMaxTreeDepth(TEXT("p.MaxTreeDepth"), ConfigSettings.MaxTreeDepth, TEXT(""));
+	FAutoConsoleVariableRef CVarAABBMaxChildrenInLeaf(TEXT("p.AABBMaxChildrenInLeaf"), ConfigSettings.AABBMaxChildrenInLeaf, TEXT(""));
+	FAutoConsoleVariableRef CVarAABBMaxTreeDepth(TEXT("p.AABBMaxTreeDepth"), ConfigSettings.AABBMaxTreeDepth, TEXT(""));
+	FAutoConsoleVariableRef CVarMaxPayloadSize(TEXT("p.MaxPayloadSize"), ConfigSettings.MaxPayloadSize, TEXT(""));
 
 	template<typename T, int d>
 	struct TDefaultCollectionFactory : public ISpatialAccelerationCollectionFactory<T, d>
 	{
+		FAccelerationConfig Config;
+
 		using BVType = TBoundingVolume<TAccelerationStructureHandle<T, d>, T, d>;
 		using AABBTreeType = TAABBTree<TAccelerationStructureHandle<T, d>, TAABBTreeLeafArray<TAccelerationStructureHandle<T, d>, T>, T>;
 		using AABBTreeOfGridsType = TAABBTree<TAccelerationStructureHandle<T, d>, TBoundingVolume<TAccelerationStructureHandle<T, d>, T, d>, T>;
 
-		virtual TUniquePtr<ISpatialAccelerationCollection<TAccelerationStructureHandle<T, d>, T, d>> CreateEmptyCollection() override
+		TUniquePtr<ISpatialAccelerationCollection<TAccelerationStructureHandle<T, d>, T, d>> CreateEmptyCollection() override
 		{
-			const uint16 NumBuckets = BroadphaseType == 3 ? 2 : 1;
+			const uint16 NumBuckets = ConfigSettings.BroadphaseType == 3 ? 2 : 1;
 			auto Collection = new TSpatialAccelerationCollection<AABBTreeType, BVType, AABBTreeOfGridsType>();
 			const TArray<TAccelerationStructureBuilder<T, d>> Empty;
 
@@ -57,23 +68,23 @@ namespace Chaos
 			{
 			case 0:
 			{
-				if (BroadphaseType == 0)
+				if (ConfigSettings.BroadphaseType == 0)
 				{
-					return MakeUnique<BVType>(Particles, false, 0, BoundingVolumeNumCells, MaxPayloadSize);
+					return MakeUnique<BVType>(Particles, false, 0, ConfigSettings.BVNumCells, ConfigSettings.MaxPayloadSize);
 				}
-				else if(BroadphaseType == 1 || BroadphaseType == 3)
+				else if(ConfigSettings.BroadphaseType == 1 || ConfigSettings.BroadphaseType == 3)
 				{
-					return MakeUnique<AABBTreeType>(Particles, MaxChildrenInLeaf, MaxTreeDepth, MaxPayloadSize);
+					return MakeUnique<AABBTreeType>(Particles, ConfigSettings.MaxChildrenInLeaf, ConfigSettings.MaxTreeDepth, ConfigSettings.MaxPayloadSize);
 				}
 				else
 				{
-					return MakeUnique<AABBTreeOfGridsType>(Particles, AABBMaxChildrenInLeaf, AABBMaxTreeDepth, MaxPayloadSize);
+					return MakeUnique<AABBTreeOfGridsType>(Particles, ConfigSettings.AABBMaxChildrenInLeaf, ConfigSettings.AABBMaxTreeDepth, ConfigSettings.MaxPayloadSize);
 				}
 			}
 			case 1:
 			{
-				ensure(BroadphaseType == 3);
-				return MakeUnique<BVType>(Particles, false, 0, BoundingVolumeNumCells, MaxPayloadSize);
+				ensure(ConfigSettings.BroadphaseType == 3);
+				return MakeUnique<BVType>(Particles, false, 0, ConfigSettings.BVNumCells, ConfigSettings.MaxPayloadSize);
 			}
 			default:
 			{
@@ -87,13 +98,11 @@ namespace Chaos
 		{
 			if (Ar.IsLoading())
 			{
-				//todo: actually read in parameters
 				Ptr = CreateEmptyCollection();
 				Ptr->Serialize(Ar);
 			}
 			else
 			{
-				//todo: actually save out parameters
 				Ptr->Serialize(Ar);
 			}
 		}
