Clean up spatial acceleration cache. Made it so you can swap between any structure type at runtime for AB testing

diff --git a/Engine/Source/Runtime/Experimental/Chaos/Private/Chaos/PBDRigidsEvolution.cpp b/Engine/Source/Runtime/Experimental/Chaos/Private/Chaos/PBDRigidsEvolution.cpp
--- a/Engine/Source/Runtime/Experimental/Chaos/Private/Chaos/PBDRigidsEvolution.cpp
+++ b/Engine/Source/Runtime/Experimental/Chaos/Private/Chaos/PBDRigidsEvolution.cpp
@@ -50,19 +50,25 @@ namespace Chaos
 
 		TUniquePtr<ISpatialAccelerationCollection<TAccelerationStructureHandle<T, d>, T, d>> CreateEmptyCollection() override
 		{
+			TConstParticleView<TSpatialAccelerationCache<T, d>> Empty;
+
 			const uint16 NumBuckets = ConfigSettings.BroadphaseType == 3 ? 2 : 1;
 			auto Collection = new TSpatialAccelerationCollection<AABBTreeType, BVType, AABBTreeOfGridsType>();
-			const TArray<TAccelerationStructureBuilder<T, d>> Empty;
 
 			for (uint16 BucketIdx = 0; BucketIdx < NumBuckets; ++BucketIdx)
 			{
-				Collection->AddSubstructure(CreateAccelerationPerBucket(Empty, BucketIdx), BucketIdx);
+				Collection->AddSubstructure(CreateAccelerationPerBucket_Threaded(Empty, BucketIdx), BucketIdx);
 			}
 
 			return TUniquePtr<ISpatialAccelerationCollection<TAccelerationStructureHandle<T, d>, T, d>>(Collection);
 		}
 
-		virtual TUniquePtr<ISpatialAcceleration<TAccelerationStructureHandle<T, d>, T, d>> CreateAccelerationPerBucket(const TArray<TAccelerationStructureBuilder<T, d>>& Particles, uint16 BucketIdx) override
+		virtual uint8 GetActiveBucketsMask() const
+		{
+			return ConfigSettings.BroadphaseType == 3 ? 3 : 1;
+		}
+
+		virtual TUniquePtr<ISpatialAcceleration<TAccelerationStructureHandle<T, d>, T, d>> CreateAccelerationPerBucket_Threaded(const TConstParticleView<TSpatialAccelerationCache<T, d>>& Particles, uint16 BucketIdx) override
 		{
 			switch (BucketIdx)
 			{
@@ -146,11 +152,11 @@ namespace Chaos
 	template<class FPBDRigidsEvolution, class FPBDCollisionConstraint, class T, int d>
 	TPBDRigidsEvolutionBase<FPBDRigidsEvolution, FPBDCollisionConstraint, T, d>::FChaosAccelerationStructureTask::FChaosAccelerationStructureTask(
 		ISpatialAccelerationCollectionFactory<T, d>& InSpatialCollectionFactory
-		, const TArray<FAccelerationStructBuilderCache>& CacheMap
+		, const TMap<FSpatialAccelerationIdx, TSpatialAccelerationCache<T,d>>& InSpatialAccelerationCache
 		, TUniquePtr<FAccelerationStructure>& InAccelerationStructure
 		, TUniquePtr<FAccelerationStructure>& InAccelerationStructureCopy)
 		: SpatialCollectionFactory(InSpatialCollectionFactory)
-		, BuilderCacheMap(CacheMap)
+		, SpatialAccelerationCache(InSpatialAccelerationCache)
 		, AccelerationStructure(InAccelerationStructure)
 		, AccelerationStructureCopy(InAccelerationStructureCopy)
 	{
@@ -204,13 +210,34 @@ namespace Chaos
 	template<class FPBDRigidsEvolution, class FPBDCollisionConstraint, class T, int d>
 	void TPBDRigidsEvolutionBase<FPBDRigidsEvolution, FPBDCollisionConstraint, T, d>::FChaosAccelerationStructureTask::DoTask(ENamedThreads::Type CurrentThread, const FGraphEventRef& MyCompletionGraphEvent)
 	{
-		//todo: no reason to do this in serial
-		for (const FAccelerationStructBuilderCache& Cache : BuilderCacheMap)
+		uint8 ActiveBucketsMask = SpatialCollectionFactory.GetActiveBucketsMask();
+		TArray<TSOAView<TSpatialAccelerationCache<T, d>>> ViewsPerBucket[8];
+		
+		//merge buckets. todo: support multiple entries per bucket (i.e. dynamic vs static)
+		for (const auto& Itr : SpatialAccelerationCache)
 		{
-			auto OldStruct = AccelerationStructure->RemoveSubstructure(Cache.SpatialIdx);
-			auto NewStruct = SpatialCollectionFactory.CreateAccelerationPerBucket(*Cache.CachedSpatialBuilderData, Cache.SpatialIdx.Bucket);
-			AccelerationStructure->AddSubstructure(MoveTemp(NewStruct), Cache.SpatialIdx.Bucket);
+			const FSpatialAccelerationIdx SpatialIdx = Itr.Key;
+			const TSpatialAccelerationCache<T, d>& Cache = Itr.Value;
+			const uint8 BucketIdx = (1 << SpatialIdx.Bucket) & ActiveBucketsMask ? SpatialIdx.Bucket : 0;
+			ViewsPerBucket[BucketIdx].Add(const_cast<TSpatialAccelerationCache<T, d>*>(&Cache));
+
+			if (AccelerationStructure->IsBucketActive(SpatialIdx.Bucket))
+			{
+				AccelerationStructure->RemoveSubstructure(SpatialIdx);
+			}
 		}
+
+		//todo: creation can go wide, insertion to collection cannot
+		for (uint8 BucketIdx = 0; BucketIdx < 8; ++BucketIdx)
+		{
+			if (ViewsPerBucket[BucketIdx].Num())
+			{
+				auto ParticleView = MakeConstParticleView(MoveTemp(ViewsPerBucket[BucketIdx]));
+				auto NewStruct = SpatialCollectionFactory.CreateAccelerationPerBucket_Threaded(ParticleView, BucketIdx);
+				AccelerationStructure->AddSubstructure(MoveTemp(NewStruct), BucketIdx);
+			}
+		}
+		
 		AccelerationStructureCopy = AsUniqueSpatialAccelerationChecked<FAccelerationStructure>(AccelerationStructure->Copy());
 	}
 
@@ -230,21 +257,19 @@ namespace Chaos
 
 			if (bAsync)
 			{
-				if (int32* CacheIdxPtr = ParticleToCacheIdx.Find(Particle))
+				if (uint32* InnerIdxPtr = ParticleToCacheInnerIdx.Find(Particle))
 				{
 					const auto SpatialIdx = SpatialData.DeletedSpatialIdx;
-					auto MapFound = CachedSpatialBuilderDataMap.FindByPredicate([SpatialIdx](const auto& Cache) { return Cache.SpatialIdx == SpatialIdx; });
-					TArray<TAccelerationStructureBuilder<T,d>>& CachedSpatialBuilderData = *(MapFound ? MapFound->CachedSpatialBuilderData : CachedSpatialBuilderDataMap.Last().CachedSpatialBuilderData);
-					const int32 CacheIdx = *CacheIdxPtr;
-					if (CacheIdx + 1 < CachedSpatialBuilderData.Num())	//will get swapped with last element, so update it
+					TSpatialAccelerationCache<T, d>& Cache = SpatialAccelerationCache.FindChecked(SpatialIdx);	//can't delete from cache that doesn't exist
+					const uint32 CacheInnerIdx = *InnerIdxPtr;
+					if (CacheInnerIdx + 1 < Cache.Size())	//will get swapped with last element, so update it
 					{
-						//in cached bounds so must be in mapping, update mapping to new position
-						int32& PrevIdx = ParticleToCacheIdx.FindChecked(CachedSpatialBuilderData.Last().CachedSpatialPayload.GetGeometryParticleHandle_PhysicsThread());
-						PrevIdx = CacheIdx;
+						TGeometryParticleHandle<T, d>* LastParticleInCache = Cache.Payload(Cache.Size() - 1).GetGeometryParticleHandle_PhysicsThread();
+						ParticleToCacheInnerIdx.FindChecked(LastParticleInCache) = CacheInnerIdx;
 					}
 
-					CachedSpatialBuilderData.RemoveAtSwap(CacheIdx);
-					ParticleToCacheIdx.Remove(Particle);
+					Cache.DestroyElement(CacheInnerIdx);
+					ParticleToCacheInnerIdx.Remove(Particle);
 				}
 			}
 		}
@@ -255,25 +280,25 @@ namespace Chaos
 			
 			if (bAsync)
 			{
-				const auto SpatialIdx = SpatialData.UpdatedSpatialIdx;
-				auto MapFound = CachedSpatialBuilderDataMap.FindByPredicate([SpatialIdx](const auto& Cache) { return Cache.SpatialIdx == SpatialIdx; });
-				TArray<TAccelerationStructureBuilder<T,d>>& CachedSpatialBuilderData = *(MapFound ? MapFound->CachedSpatialBuilderData : CachedSpatialBuilderDataMap.Last().CachedSpatialBuilderData);
+				TSpatialAccelerationCache<T, d>& Cache = SpatialAccelerationCache.FindOrAdd(SpatialData.UpdatedSpatialIdx);
 
 				//make sure in mapping
-				int32 CacheIdx;
-				if (int32* CacheIdxPtr = ParticleToCacheIdx.Find(Particle))
+				uint32 CacheInnerIdx;
+				if (uint32* CacheInnerIdxPtr = ParticleToCacheInnerIdx.Find(Particle))
 				{
-					CacheIdx = *CacheIdxPtr;
+					CacheInnerIdx = *CacheInnerIdxPtr;
 				}
 				else
 				{
-					CacheIdx = CachedSpatialBuilderData.Num();
-					ParticleToCacheIdx.Add(Particle, CacheIdx);
-					CachedSpatialBuilderData.AddUninitialized();
+					CacheInnerIdx = Cache.Size();
+					Cache.AddElements(1);
+					ParticleToCacheInnerIdx.Add(Particle, CacheInnerIdx);
 				}
 
-				//update cache itself
-				CachedSpatialBuilderData[CacheIdx] = TAccelerationStructureBuilder<T,d>{ Particle->HasBounds(), Particle->WorldSpaceInflatedBounds(), SpatialData.AccelerationHandle };
+				//update cache entry
+				Cache.HasBounds(CacheInnerIdx) = Particle->HasBounds();
+				Cache.Bounds(CacheInnerIdx) = Particle->WorldSpaceInflatedBounds();
+				Cache.Payload(CacheInnerIdx) = SpatialData.AccelerationHandle;
 			}
 		}
 	}
@@ -341,7 +366,6 @@ namespace Chaos
 			FlushInternalAccelerationQueue();
 			FlushExternalAccelerationQueue(*ScratchExternalAcceleration);
 			bExternalReady = true;
-			InitializeAccelerationCache();
 		}
 
 		if (bBlock)
@@ -361,7 +385,7 @@ namespace Chaos
 				bExternalReady = true;
 			}
 
-			AccelerationStructureTaskComplete = TGraphTask<FChaosAccelerationStructureTask>::CreateTask().ConstructAndDispatchWhenReady(*SpatialCollectionFactory, CachedSpatialBuilderDataMap, AsyncInternalAcceleration, AsyncExternalAcceleration);
+			AccelerationStructureTaskComplete = TGraphTask<FChaosAccelerationStructureTask>::CreateTask().ConstructAndDispatchWhenReady(*SpatialCollectionFactory, SpatialAccelerationCache, AsyncInternalAcceleration, AsyncExternalAcceleration);
 		}
 		else
 		{
@@ -369,18 +393,6 @@ namespace Chaos
 		}
 	}
 
-	template<class FPBDRigidsEvolution, class FPBDCollisionConstraint, class T, int d>
-	void TPBDRigidsEvolutionBase<FPBDRigidsEvolution, FPBDCollisionConstraint, T, d>::InitializeAccelerationCache()
-	{
-		WaitOnAccelerationStructure();
-		CachedSpatialBuilderDataMap.Empty();
-		const auto SpatialIndices = InternalAcceleration->GetAllSpatialIndices();
-		for (const auto& Idx : SpatialIndices)
-		{
-			CachedSpatialBuilderDataMap.Add(Idx);
-		}
-	}
-
 	template<class FPBDRigidsEvolution, class FPBDCollisionConstraint, class T, int d>
 	void TPBDRigidsEvolutionBase<FPBDRigidsEvolution, FPBDCollisionConstraint, T, d>::UpdateExternalAccelerationStructure(TUniquePtr<ISpatialAccelerationCollection<TAccelerationStructureHandle<T, d>, T, d>>& StructToUpdate)
 	{
@@ -428,7 +440,6 @@ namespace Chaos
 				{
 					Ar << SubStructure;
 					InternalAcceleration = CreateNewSpatialStructureFromSubStructure(MoveTemp(SubStructure));
-					InitializeAccelerationCache();
 				}
 			}
 			else
diff --git a/Engine/Source/Runtime/Experimental/Chaos/Public/Chaos/GeometryParticlesfwd.h b/Engine/Source/Runtime/Experimental/Chaos/Public/Chaos/GeometryParticlesfwd.h
--- a/Engine/Source/Runtime/Experimental/Chaos/Public/Chaos/GeometryParticlesfwd.h
+++ b/Engine/Source/Runtime/Experimental/Chaos/Public/Chaos/GeometryParticlesfwd.h
@@ -34,6 +34,11 @@ namespace Chaos
 		}
 	};
 
+	inline uint32 GetTypeHash(const FSpatialAccelerationIdx& Idx)
+	{
+		return ::GetTypeHash((const uint16&) Idx);
+	}
+
 	inline FArchive& operator<<(FArchive& Ar, FSpatialAccelerationIdx& Idx)
 	{
 		return Ar << (uint16&)Idx;
diff --git a/Engine/Source/Runtime/Experimental/Chaos/Public/Chaos/ISpatialAccelerationCollection.h b/Engine/Source/Runtime/Experimental/Chaos/Public/Chaos/ISpatialAccelerationCollection.h
--- a/Engine/Source/Runtime/Experimental/Chaos/Public/Chaos/ISpatialAccelerationCollection.h
+++ b/Engine/Source/Runtime/Experimental/Chaos/Public/Chaos/ISpatialAccelerationCollection.h
@@ -13,7 +13,10 @@ template <typename TPayloadType, typename T, int d>
 class CHAOS_API ISpatialAccelerationCollection : public ISpatialAcceleration<TPayloadType, T, d>
 {
 public:
-	ISpatialAccelerationCollection() : ISpatialAcceleration<TPayloadType, T, d>(StaticType) {}
+	ISpatialAccelerationCollection()
+	: ISpatialAcceleration<TPayloadType, T, d>(StaticType)
+	, ActiveBucketsMask(0)
+	{}
 	static constexpr ESpatialAcceleration StaticType = ESpatialAcceleration::Collection;
 	virtual FSpatialAccelerationIdx AddSubstructure(TUniquePtr<ISpatialAcceleration<TPayloadType, T, d>>&& Substructure, uint16 Bucket) = 0;
 	virtual TUniquePtr<ISpatialAcceleration<TPayloadType, T, d>> RemoveSubstructure(FSpatialAccelerationIdx Idx) = 0;
@@ -23,6 +26,13 @@ public:
 	virtual void PBDComputeConstraintsLowLevel_GatherStats(TPBDCollisionConstraint<T, d>& CollisionConstraint, T Dt) const = 0;
 	virtual void PBDComputeConstraintsLowLevel(TPBDCollisionConstraint<T, d>& CollisionConstraint, T Dt) const = 0;
 	virtual TArray<FSpatialAccelerationIdx> GetAllSpatialIndices() const = 0;
+
+	bool IsBucketActive(uint8 BucketIdx) const
+	{
+		return (1 << BucketIdx) & ActiveBucketsMask;
+	}
+protected:
+	uint8 ActiveBucketsMask;
 };
 
 template <typename T>
@@ -333,15 +343,12 @@ public:
 	static constexpr int d = FirstAccelerationType::D;
 
 	TSpatialAccelerationCollection()
-		: LastBucket(0)
 	{
 	}
 
 	virtual FSpatialAccelerationIdx AddSubstructure(TUniquePtr<ISpatialAcceleration<TPayloadType, T, d>>&& Substructure, uint16 BucketIdx) override
 	{
 		check(BucketIdx < MaxBuckets);
-		LastBucket = FMath::Max(LastBucket, BucketIdx);
-		//todo: verify the bucket and type match
 		FSpatialAccelerationIdx Result;
 		Result.Bucket = BucketIdx;
 		
@@ -362,6 +369,8 @@ public:
 		case 7: BucketEntry.TypeInnerIdx = GetAccelerationsPerType<ClampedIdx(7)>(Types).Add(&AccelPtr->template AsChecked<typename std::tuple_element<ClampedIdx(7), std::tuple<TSpatialAccelerationTypes...>>::type>()); break;
 		}
 
+		this->ActiveBucketsMask |= (1 << BucketIdx);
+
 		Result.InnerIdx = Buckets[BucketIdx].Add(MoveTemp(BucketEntry));
 		return Result;
 	}
@@ -386,15 +395,9 @@ public:
 			Removed = MoveTemp(BucketEntry.Acceleration);
 		}
 		Buckets[Idx.Bucket].Remove(Idx.InnerIdx);
-
-		LastBucket = 0;
-		for (uint16 BucketIdx = Idx.Bucket; BucketIdx > 0; --BucketIdx)
+		if (Buckets[Idx.Bucket].Objects.Num() == 0)
 		{
-			if (Buckets[BucketIdx].Objects.Num())
-			{
-				LastBucket = BucketIdx;
-				break;
-			}
+			this->ActiveBucketsMask &= ~(1 << Idx.Bucket);
 		}
 		return Removed;
 	}
@@ -493,13 +496,13 @@ public:
 
 	virtual void RemoveElementFrom(const TPayloadType& Payload, FSpatialAccelerationIdx SpatialIdx) override
 	{
-		const uint16 UseBucket = FMath::Min(LastBucket, SpatialIdx.Bucket);
+		const uint16 UseBucket = ((1 << SpatialIdx.Bucket) & this->ActiveBucketsMask) ? SpatialIdx.Bucket : 0;
 		Buckets[UseBucket].Objects[SpatialIdx.InnerIdx].Acceleration->RemoveElement(Payload);
 	}
 
 	virtual void UpdateElementIn(const TPayloadType& Payload, const TBox<T, d>& NewBounds, bool bHasBounds, FSpatialAccelerationIdx SpatialIdx)
 	{
-		const uint16 UseBucket = FMath::Min(LastBucket, SpatialIdx.Bucket);
+		const uint16 UseBucket = ((1 << SpatialIdx.Bucket) & this->ActiveBucketsMask) ? SpatialIdx.Bucket : 0;
 		Buckets[UseBucket].Objects[SpatialIdx.InnerIdx].Acceleration->UpdateElement(Payload, NewBounds, bHasBounds);
 	}
 
@@ -552,11 +555,10 @@ public:
 private:
 
 	TSpatialAccelerationCollection(const TSpatialAccelerationCollection<TSpatialAccelerationTypes...>& Other)
-		: LastBucket(Other.LastBucket)
 	{
 		//need to make a deep copy of the unique ptrs and then update raw ptrs
 		TSpatialAccelerationCollectionHelper<0, NumTypes, decltype(Types), TPayloadType, T, d>::SetNumFrom(Other.Types, Types);
-		
+		this->ActiveBucketsMask = Other.ActiveBucketsMask;
 
 		for (int BucketIdx = 0; BucketIdx < MaxBuckets; ++BucketIdx)
 		{
@@ -636,7 +638,6 @@ private:
 	TSpatialCollectionBucket<TSpatialAccelerationBucketEntry<TPayloadType, T, d>> Buckets[MaxBuckets];
 	TSpatialTypeTuple< TSpatialAccelerationTypes...> Types;
 	static constexpr uint32 NumTypes = sizeof...(TSpatialAccelerationTypes);
-	uint16 LastBucket;	//the last bucket that has an acceleration structure
 };
 
 
diff --git a/Engine/Source/Runtime/Experimental/Chaos/Public/Chaos/PBDRigidsEvolution.h b/Engine/Source/Runtime/Experimental/Chaos/Public/Chaos/PBDRigidsEvolution.h
--- a/Engine/Source/Runtime/Experimental/Chaos/Public/Chaos/PBDRigidsEvolution.h
+++ b/Engine/Source/Runtime/Experimental/Chaos/Public/Chaos/PBDRigidsEvolution.h
@@ -52,30 +52,142 @@ struct CHAOS_API FEvolutionStats
 	}
 };
 
-/** Used for building an acceleration structure out of cached bounds and payloads */
 template <typename T, int d>
-struct TAccelerationStructureBuilder
+struct TSpatialAccelerationCacheHandle;
+
+/** The SOA cache used for a single acceleration structure */
+template <typename T, int d>
+class TSpatialAccelerationCache : public TArrayCollection
 {
-	//todo: should these be arrays instead? might make it easier in some cases
-	bool bHasBoundingBox;
-	TBox<T, d> CachedSpatialBounds;
-	TAccelerationStructureHandle<T, d> CachedSpatialPayload;
+public:
+	using THandleType = TSpatialAccelerationCacheHandle<T, d>;
 
-	const TBox<T, d>& BoundingBox() const
+	TSpatialAccelerationCache()
 	{
-		return CachedSpatialBounds;
+		AddArray(&MHasBoundingBoxes);
+		AddArray(&MBounds);
+		AddArray(&MPayloads);
+
+#if PARTICLE_ITERATOR_RANGED_FOR_CHECK
+		MDirtyValidationCount = 0;
+#endif
 	}
 
-	bool HasBoundingBox() const
+	TSpatialAccelerationCache(const TSpatialAccelerationCache<T, d>&) = delete;
+	TSpatialAccelerationCache(TSpatialAccelerationCache<T, d>&& Other)
+		: TArrayCollection()
+		, MHasBoundingBoxes(MoveTemp(Other.MHasBoundingBoxes))
+		, MBounds(MoveTemp(Other.MBounds))
+		, MPayloads(MoveTemp(Other.MPayloads))
+	{
+		Other.MSize = 0;
+
+		AddArray(&MHasBoundingBoxes);
+		AddArray(&MBounds);
+		AddArray(&MPayloads);
+#if PARTICLE_ITERATOR_RANGED_FOR_CHECK
+		MDirtyValidationCount = 0;
+#endif
+	}
+
+	TSpatialAccelerationCache& operator=(TSpatialAccelerationCache<T, d>&& Other)
+	{
+		if (&Other != this)
+		{
+			MHasBoundingBoxes = MoveTemp(Other.MHasBoundingBoxes);
+			MBounds = MoveTemp(Other.MBounds);
+			MPayloads = MoveTemp(Other.MPayloads);
+			Other.MSize = 0;
+#if PARTICLE_ITERATOR_RANGED_FOR_CHECK
+			MDirtyValidationCount = 0;
+			++Other.MDirtyValidationCount;
+#endif
+		}
+
+		return *this;
+	}
+
+#if PARTICLE_ITERATOR_RANGED_FOR_CHECK
+	int32 DirtyValidationCount() const { return MDirtyValidationCount; }
+#endif
+
+	void AddElements(const int32 Num)
 	{
-		return bHasBoundingBox;
+		AddElementsHelper(Num);
+		IncrementDirtyValidation();
 	}
 
+	void DestroyElement(const int32 Idx)
+	{
+		RemoveAtSwapHelper(Idx);
+		IncrementDirtyValidation();
+	}
+
+	bool HasBounds(const int32 Idx) const { return MHasBoundingBoxes[Idx]; }
+	bool& HasBounds(const int32 Idx) { return MHasBoundingBoxes[Idx]; }
+
+	const TBox<T,d>& Bounds(const int32 Idx) const { return MBounds[Idx]; }
+	TBox<T, d>& Bounds(const int32 Idx) { return MBounds[Idx]; }
+
+	const TAccelerationStructureHandle<T, d>& Payload(const int32 Idx) const { return MPayloads[Idx]; }
+	TAccelerationStructureHandle<T, d>& Payload(const int32 Idx) { return MPayloads[Idx]; }
+
+private:
+	void IncrementDirtyValidation()
+	{
+#if PARTICLE_ITERATOR_RANGED_FOR_CHECK
+		++MDirtyValidationCount;
+#endif
+	}
+
+	TArrayCollectionArray<bool> MHasBoundingBoxes;
+	TArrayCollectionArray<TBox<T, d>> MBounds;
+	TArrayCollectionArray<TAccelerationStructureHandle<T, d>> MPayloads;
+
+#if PARTICLE_ITERATOR_RANGED_FOR_CHECK
+	int32 MDirtyValidationCount;
+#endif
+};
+
+/** The handle the acceleration structure uses to access the data (similar to particle handle) */
+template <typename T, int d>
+struct TSpatialAccelerationCacheHandle
+{
+	using THandleBase = TSpatialAccelerationCacheHandle<T, d>;
+	using TTransientHandle = TSpatialAccelerationCacheHandle<T, d>;
+
+	TSpatialAccelerationCacheHandle(TSpatialAccelerationCache<T, d>* InCache = nullptr, int32 InEntryIdx = INDEX_NONE)
+		: Cache(InCache)
+		, EntryIdx(InEntryIdx)
+	{}
+
 	template <typename TPayloadType>
-	TAccelerationStructureHandle<T, d> GetPayload(int32 Idx) const
+	TPayloadType GetPayload(int32 Idx) const
 	{
-		return CachedSpatialPayload;
+		return Cache->Payload(EntryIdx);
 	}
+
+	bool HasBoundingBox() const
+	{
+		return Cache->HasBounds(EntryIdx);
+	}
+
+	const TBox<T, d>& BoundingBox() const
+	{
+		return Cache->Bounds(EntryIdx);
+	}
+
+	union
+	{
+		TSpatialAccelerationCache<T, d>* GeometryParticles;	//using same name as particles SOA for template reuse, should probably rethink this
+		TSpatialAccelerationCache<T, d>* Cache;
+	};
+
+	union
+	{
+		int32 ParticleIdx;	//same name for template reasons. Not really a particle idx
+		int32 EntryIdx;
+	};
 };
 
 template <typename T, int d>
@@ -84,8 +196,11 @@ struct CHAOS_API ISpatialAccelerationCollectionFactory
 	//Create an empty acceleration collection with the desired buckets. Chaos enqueues acceleration structure operations per bucket
 	virtual TUniquePtr<ISpatialAccelerationCollection<TAccelerationStructureHandle<T, d>, T, d>> CreateEmptyCollection() = 0;
 
-	//Chaos creates new acceleration structures per bucket. Factory can change underlying type, but changing the number of buckets requires re-registering factory
-	virtual TUniquePtr<ISpatialAcceleration<TAccelerationStructureHandle<T, d>, T, d>> CreateAccelerationPerBucket(const TArray<TAccelerationStructureBuilder<T, d>>& Particles, uint16 BucketIdx) = 0;
+	//Chaos creates new acceleration structures per bucket. Factory can change underlying type at runtime as well as number of buckets to AB test
+	virtual TUniquePtr<ISpatialAcceleration<TAccelerationStructureHandle<T, d>, T, d>> CreateAccelerationPerBucket_Threaded(const TConstParticleView<TSpatialAccelerationCache<T, d>>& Particles, uint16 BucketIdx) = 0;
+
+	//Mask indicating which bucket is active. Spatial indices in inactive buckets fallback to bucket 0. Bit 0 indicates bucket 0 is active, Bit 1 indicates bucket 1 is active, etc...
+	virtual uint8 GetActiveBucketsMask() const = 0;
 
 	//Serialize the collection in and out
 	virtual void Serialize(TUniquePtr<ISpatialAccelerationCollection<TAccelerationStructureHandle<T, d>, T, d>>& Ptr, FChaosArchive& Ar) = 0;
@@ -516,34 +631,26 @@ protected:
 	}
 
 	/** Used for async acceleration rebuild */
-	TMap<TGeometryParticleHandle<T, d>*, int32> ParticleToCacheIdx;
+	TMap<TGeometryParticleHandle<T, d>*, uint32> ParticleToCacheInnerIdx;
 
-	struct FAccelerationStructBuilderCache
-	{
-		FSpatialAccelerationIdx SpatialIdx;
-		TUniquePtr <TArray<TAccelerationStructureBuilder<T, d>>> CachedSpatialBuilderData;
-
-		FAccelerationStructBuilderCache(FSpatialAccelerationIdx Idx)
-			: SpatialIdx(Idx)
-			, CachedSpatialBuilderData(MakeUnique<TArray<TAccelerationStructureBuilder<T, d>>>())
-		{}
-	};
-	TArray<FAccelerationStructBuilderCache> CachedSpatialBuilderDataMap;
+	TMap<FSpatialAccelerationIdx, TSpatialAccelerationCache<T, d>> SpatialAccelerationCache;
 
 	FORCEINLINE_DEBUGGABLE void ApplyParticlePendingData(TGeometryParticleHandle<T, d>* Particle, const FPendingSpatialData& PendingData, FAccelerationStructure& SpatialAcceleration, bool bAsync);
 
 	class FChaosAccelerationStructureTask
 	{
 	public:
-		FChaosAccelerationStructureTask(ISpatialAccelerationCollectionFactory<T,d>& InSpatialCollectionFactory, const TArray<FAccelerationStructBuilderCache>& BuilderCacheMap, TUniquePtr<FAccelerationStructure>& InAccelerationStructure,
-			TUniquePtr<FAccelerationStructure>& InAccelerationStructureCopy);
+		FChaosAccelerationStructureTask(ISpatialAccelerationCollectionFactory<T,d>& InSpatialCollectionFactory
+			, const TMap<FSpatialAccelerationIdx, TSpatialAccelerationCache<T,d>>& InSpatialAccelerationCache
+			, TUniquePtr<FAccelerationStructure>& InAccelerationStructure
+			, TUniquePtr<FAccelerationStructure>& InAccelerationStructureCopy);
 		static FORCEINLINE TStatId GetStatId();
 		static FORCEINLINE ENamedThreads::Type GetDesiredThread();
 		static FORCEINLINE ESubsequentsMode::Type GetSubsequentsMode();
 		void DoTask(ENamedThreads::Type CurrentThread, const FGraphEventRef& MyCompletionGraphEvent);
 
 		ISpatialAccelerationCollectionFactory<T, d>& SpatialCollectionFactory;
-		const TArray<FAccelerationStructBuilderCache>& BuilderCacheMap;
+		const TMap<FSpatialAccelerationIdx, TSpatialAccelerationCache<T, d>>& SpatialAccelerationCache;
 		TUniquePtr<FAccelerationStructure>& AccelerationStructure;
 		TUniquePtr<FAccelerationStructure>& AccelerationStructureCopy;
 	};
@@ -551,8 +658,6 @@ protected:
 
 	int32 NumIterations;
 	TUniquePtr<ISpatialAccelerationCollectionFactory<T, d>> SpatialCollectionFactory;
-
-	void InitializeAccelerationCache();
 };
 
 }
