Chaos - profiling markers, reconcile cvars to control parallel updates

diff --git a/Engine/Source/Runtime/Experimental/Chaos/Private/Chaos/Framework/Parallel.cpp b/Engine/Source/Runtime/Experimental/Chaos/Private/Chaos/Framework/Parallel.cpp
--- a/Engine/Source/Runtime/Experimental/Chaos/Private/Chaos/Framework/Parallel.cpp
+++ b/Engine/Source/Runtime/Experimental/Chaos/Private/Chaos/Framework/Parallel.cpp
@@ -5,12 +5,25 @@
 
 using namespace Chaos;
 
-static bool GNoParallelFor = false;
+namespace Chaos
+{
+#if !UE_BUILD_SHIPPING
+	bool bDisablePhysicsParallelFor = false;
+	CHAOS_API bool bDisableParticleParallelFor = false;
+	CHAOS_API bool bDisableCollisionParallelFor = false;
+
+	FAutoConsoleVariableRef CVarDisablePhysicsParallelFor(TEXT("p.Chaos.DisablePhysicsParallelFor"), bDisablePhysicsParallelFor, TEXT("Disable parallel execution in Chaos Evolution"));
+	FAutoConsoleVariableRef CVarDisableParticleParallelFor(TEXT("p.Chaos.DisableParticleParallelFor"), bDisableParticleParallelFor, TEXT("Disable parallel execution for Chaos Particles (Collisions, "));
+	FAutoConsoleVariableRef CVarDisableCollisionParallelFor(TEXT("p.Chaos.DisableCollisionParallelFor"), bDisableCollisionParallelFor, TEXT("Disable parallel execution for Chaos Collisions (also disabled by DisableParticleParallelFor)"));
+#else
+	const bool bDisablePhysicsParallelFor = false;
+#endif
+}
 
 void Chaos::PhysicsParallelFor(int32 InNum, TFunctionRef<void(int32)> InCallable, bool bForceSingleThreaded)
 {
 	// Passthrough for now, except with global flag to disable parallel
-	::ParallelFor(InNum, InCallable, GNoParallelFor || bForceSingleThreaded);
+	::ParallelFor(InNum, InCallable, bDisablePhysicsParallelFor || bForceSingleThreaded);
 }
 
 //class FRecursiveDivideTask
diff --git a/Engine/Source/Runtime/Experimental/Chaos/Private/Chaos/PBDCollisionConstraint.cpp b/Engine/Source/Runtime/Experimental/Chaos/Private/Chaos/PBDCollisionConstraint.cpp
--- a/Engine/Source/Runtime/Experimental/Chaos/Private/Chaos/PBDCollisionConstraint.cpp
+++ b/Engine/Source/Runtime/Experimental/Chaos/Private/Chaos/PBDCollisionConstraint.cpp
@@ -97,7 +97,7 @@ TPBDCollisionConstraint<T, d>::TPBDCollisionConstraint(const TPBDRigidsSOAs<T,d>
 {
 }
 
-DECLARE_CYCLE_STAT(TEXT("CollisionConstraint::Reset"), STAT_CollisionConstraintsReset, STATGROUP_Chaos);
+DECLARE_CYCLE_STAT(TEXT("TPBDCollisionConstraint::Reset"), STAT_CollisionConstraintsReset, STATGROUP_Chaos);
 
 template<typename T, int d>
 void TPBDCollisionConstraint<T, d>::Reset(/*const TPBDRigidParticles<T, d>& InParticles, const TArray<int32>& InIndices*/)
@@ -223,9 +223,6 @@ DEFINE_STAT(STAT_ComputeConstraintsNP);
 DEFINE_STAT(STAT_ComputeConstraintsBP);
 DEFINE_STAT(STAT_ComputeConstraintsSU);
 
-CHAOS_API int32 CollisionConstraintsForceSingleThreaded = 0;
-FAutoConsoleVariableRef CVarCollisionConstraintsForceSingleThreaded(TEXT("p.Chaos.Collision.ForceSingleThreaded"), CollisionConstraintsForceSingleThreaded, TEXT("CollisionConstraintsForceSingleThreaded"));
-
 template<typename T, int d>
 template <bool bGatherStats>
 void TPBDCollisionConstraint<T, d>::ComputeConstraints(const FAccelerationStructure& AccelerationStructure, T Dt)
@@ -424,13 +421,13 @@ void TPBDCollisionConstraint<T, d>::UpdateConstraintsHelper(/*const TPBDRigidPar
 #endif
 }
 
-DECLARE_CYCLE_STAT(TEXT("Reconcile Updated Constraints"), STAT_ReconcileConstraints2, STATGROUP_Chaos);
+DECLARE_CYCLE_STAT(TEXT("TPBDCollisionConstraint::ReconcileUpdatedConstraints"), STAT_ReconcileConstraints, STATGROUP_ChaosWide);
 template<typename T, int d>
 void TPBDCollisionConstraint<T, d>::UpdateConstraints(/*const TPBDRigidParticles<T, d>& InParticles, const TArray<int32>& InIndices,*/ T Dt, const TSet<TGeometryParticleHandle<T, d>*>& AddedParticles)
 {
 #if CHAOS_PARTICLEHANDLE_TODO
 	{
-		SCOPE_CYCLE_COUNTER(STAT_ReconcileConstraints2);
+		SCOPE_CYCLE_COUNTER(STAT_ReconcileConstraints);
 
 		// Updating post-clustering, we will have invalid constraints
 		int32 NumRemovedConstraints = 0;
@@ -721,17 +718,18 @@ void TPBDCollisionConstraint<T, d>::Apply(const T Dt, FRigidBodyContactConstrain
 	}
 }
 
-DECLARE_CYCLE_STAT(TEXT("Apply"), STAT_Apply2, STATGROUP_ChaosWide);
+DECLARE_CYCLE_STAT(TEXT("TPBDCollisionConstraint::Apply"), STAT_Apply, STATGROUP_Chaos);
 template<typename T, int d>
 void TPBDCollisionConstraint<T, d>::Apply(const T Dt, const TArray<FConstraintHandle*>& InConstraintHandles, const int32 It, const int32 NumIts)
 {
+	SCOPE_CYCLE_COUNTER(STAT_Apply);
 	if (bEnableVelocitySolve)
 	{
 		PhysicsParallelFor(InConstraintHandles.Num(), [&](int32 ConstraintHandleIndex) {
 			FConstraintHandle* ConstraintHandle = InConstraintHandles[ConstraintHandleIndex];
 			check(ConstraintHandle != nullptr);
 			Apply(Dt, Constraints[ConstraintHandle->GetConstraintIndex()]);
-			});
+			}, bDisableCollisionParallelFor);
 	}
 
 	if (PostApplyCallback != nullptr)
@@ -740,11 +738,9 @@ void TPBDCollisionConstraint<T, d>::Apply(const T Dt, const TArray<FConstraintHa
 	}
 }
 
-DECLARE_CYCLE_STAT(TEXT("ApplyPushOut"), STAT_ApplyPushOut2, STATGROUP_ChaosWide);
 template<typename T, int d>
 void TPBDCollisionConstraint<T, d>::ApplyPushOut(const T Dt, FRigidBodyContactConstraint& Constraint, const TSet<TGeometryParticleHandle<T,d>*>& IsTemporarilyStatic, int32 Iteration, int32 NumIterations, bool &NeedsAnotherIteration)
 {
-	SCOPE_CYCLE_COUNTER(STAT_ApplyPushOut2);
 	TGeometryParticleHandle<T, d>* Particle0 = Constraint.Particle;
 	TGeometryParticleHandle<T, d>* Particle1 = Constraint.Levelset;
 	TPBDRigidParticleHandle<T, d>* PBDRigid0 = Particle0->AsDynamic();
@@ -857,11 +853,11 @@ void TPBDCollisionConstraint<T, d>::ApplyPushOut(const T Dt, FRigidBodyContactCo
 	}
 }
 
-//DECLARE_CYCLE_STAT(TEXT("ApplyPushOut"), STAT_ApplyPushOut, STATGROUP_ChaosWide);
+DECLARE_CYCLE_STAT(TEXT("TPBDCollisionConstraint::ApplyPushOut"), STAT_ApplyPushOut, STATGROUP_Chaos);
 template<typename T, int d>
 bool TPBDCollisionConstraint<T, d>::ApplyPushOut(const T Dt, const TArray<FConstraintHandle*>& InConstraintHandles, const TSet<TGeometryParticleHandle<T,d>*>& IsTemporarilyStatic, int32 Iteration, int32 NumIterations)
 {
-	SCOPE_CYCLE_COUNTER(STAT_ApplyPushOut2);
+	SCOPE_CYCLE_COUNTER(STAT_ApplyPushOut);
 
 	bool NeedsAnotherIteration = false;
 
@@ -871,7 +867,7 @@ bool TPBDCollisionConstraint<T, d>::ApplyPushOut(const T Dt, const TArray<FConst
 			FConstraintHandle* ConstraintHandle = InConstraintHandles[ConstraintHandleIndex];
 			check(ConstraintHandle != nullptr);
 			ApplyPushOut(Dt, Constraints[ConstraintHandle->GetConstraintIndex()], IsTemporarilyStatic, Iteration, NumIterations, NeedsAnotherIteration);
-			});
+			}, bDisableCollisionParallelFor);
 	}
 
 	if (PostApplyPushOutCallback != nullptr)
@@ -962,12 +958,12 @@ bool SampleObjectNormalAverageHelper2(const TImplicitObject<T, d>& Object, const
 	return false;
 }
 
-DECLARE_CYCLE_STAT(TEXT("UpdateLevelsetPartial"), STAT_UpdateLevelsetPartial2, STATGROUP_ChaosWide);
-DECLARE_CYCLE_STAT(TEXT("UpdateLevelsetFindParticles"), STAT_UpdateLevelsetFindParticles2, STATGROUP_ChaosWide);
-DECLARE_CYCLE_STAT(TEXT("UpdateLevelsetBVHTraversal"), STAT_UpdateLevelsetBVHTraversal2, STATGROUP_ChaosWide);
-DECLARE_CYCLE_STAT(TEXT("UpdateLevelsetSignedDistance"), STAT_UpdateLevelsetSignedDistance2, STATGROUP_ChaosWide);
-DECLARE_CYCLE_STAT(TEXT("UpdateLevelsetAll"), STAT_UpdateLevelsetAll2, STATGROUP_ChaosWide);
-DECLARE_CYCLE_STAT(TEXT("SampleObject"), STAT_SampleObject2, STATGROUP_ChaosWide);
+DECLARE_CYCLE_STAT(TEXT("TPBDCollisionConstraint::UpdateLevelsetPartial"), STAT_UpdateLevelsetPartial, STATGROUP_ChaosWide);
+DECLARE_CYCLE_STAT(TEXT("TPBDCollisionConstraint::UpdateLevelsetFindParticles"), STAT_UpdateLevelsetFindParticles, STATGROUP_ChaosWide);
+DECLARE_CYCLE_STAT(TEXT("TPBDCollisionConstraint::UpdateLevelsetBVHTraversal"), STAT_UpdateLevelsetBVHTraversal, STATGROUP_ChaosWide);
+DECLARE_CYCLE_STAT(TEXT("TPBDCollisionConstraint::UpdateLevelsetSignedDistance"), STAT_UpdateLevelsetSignedDistance, STATGROUP_ChaosWide);
+DECLARE_CYCLE_STAT(TEXT("TPBDCollisionConstraint::UpdateLevelsetAll"), STAT_UpdateLevelsetAll, STATGROUP_ChaosWide);
+DECLARE_CYCLE_STAT(TEXT("TPBDCollisionConstraint::SampleObject"), STAT_SampleObject, STATGROUP_ChaosWide);
 
 int32 NormalAveraging2 = 1;
 FAutoConsoleVariableRef CVarNormalAveraging2(TEXT("p.NormalAveraging2"), NormalAveraging2, TEXT(""));
@@ -978,7 +974,7 @@ FAutoConsoleVariableRef CVarSampleMinParticlesForAcceleration2(TEXT("p.SampleMin
 template <ECollisionUpdateType UpdateType, typename T, int d>
 void SampleObject2(const TImplicitObject<T, d>& Object, const TRigidTransform<T, d>& ObjectTransform, const TBVHParticles<T, d>& SampleParticles, const TRigidTransform<T, d>& SampleParticlesTransform, T Thickness, TRigidBodyContactConstraint<T, d>& Constraint)
 {
-	SCOPE_CYCLE_COUNTER(STAT_SampleObject2);
+	SCOPE_CYCLE_COUNTER(STAT_SampleObject);
 	TRigidBodyContactConstraint<T, d> AvgConstraint;
 	AvgConstraint.Particle = Constraint.Particle;
 	AvgConstraint.Levelset = Constraint.Levelset;
@@ -993,16 +989,16 @@ void SampleObject2(const TImplicitObject<T, d>& Object, const TRigidTransform<T,
 	const TRigidTransform<T, d> & SampleToObjectTM = SampleParticlesTransform.GetRelativeTransform(ObjectTransform);
 	if (NumParticles > SampleMinParticlesForAcceleration2 && Object.HasBoundingBox())
 	{
-		SCOPE_CYCLE_COUNTER(STAT_UpdateLevelsetPartial2);
+		SCOPE_CYCLE_COUNTER(STAT_UpdateLevelsetPartial);
 		TBox<T, d> ImplicitBox = Object.BoundingBox().TransformedBox(ObjectTransform.GetRelativeTransform(SampleParticlesTransform));
 		ImplicitBox.Thicken(Thickness);
 		TArray<int32> PotentialParticles;
 		{
-			SCOPE_CYCLE_COUNTER(STAT_UpdateLevelsetFindParticles2);
+			SCOPE_CYCLE_COUNTER(STAT_UpdateLevelsetFindParticles);
 			PotentialParticles = SampleParticles.FindAllIntersections(ImplicitBox);
 		}
 		{
-			SCOPE_CYCLE_COUNTER(STAT_UpdateLevelsetSignedDistance2);
+			SCOPE_CYCLE_COUNTER(STAT_UpdateLevelsetSignedDistance);
 			for (int32 i : PotentialParticles)
 			{
 				if (NormalAveraging2 && UpdateType != ECollisionUpdateType::Any)	//if we just want one don't bother with normal
@@ -1026,7 +1022,7 @@ void SampleObject2(const TImplicitObject<T, d>& Object, const TRigidTransform<T,
 	}
 	else
 	{
-		SCOPE_CYCLE_COUNTER(STAT_UpdateLevelsetAll2);
+		SCOPE_CYCLE_COUNTER(STAT_UpdateLevelsetAll);
 		for (int32 i = 0; i < NumParticles; ++i)
 		{
 			if (NormalAveraging2 && UpdateType != ECollisionUpdateType::Any)	//if we just want one don't bother with normal
@@ -1083,7 +1079,7 @@ void SampleObject2(const TImplicitObject<T, d>& Object, const TRigidTransform<T,
 template<ECollisionUpdateType UpdateType>
 void SampleObject2(const TImplicitObject<float, 3>& Object, const TRigidTransform<float, 3>& ObjectTransform, const TBVHParticles<float, 3>& SampleParticles, const TRigidTransform<float, 3>& SampleParticlesTransform, float Thickness, TRigidBodyContactConstraint<float, 3>& Constraint)
 {
-	SCOPE_CYCLE_COUNTER(STAT_SampleObject2);
+	SCOPE_CYCLE_COUNTER(STAT_SampleObject);
 	TRigidBodyContactConstraint<float, 3> AvgConstraint;
 	AvgConstraint.Particle = Constraint.Particle;
 	AvgConstraint.Levelset = Constraint.Levelset;
@@ -1099,16 +1095,16 @@ void SampleObject2(const TImplicitObject<float, 3>& Object, const TRigidTransfor
 
 	if (NumParticles > SampleMinParticlesForAcceleration2 && Object.HasBoundingBox())
 	{
-		SCOPE_CYCLE_COUNTER(STAT_UpdateLevelsetPartial2);
+		SCOPE_CYCLE_COUNTER(STAT_UpdateLevelsetPartial);
 		TBox<float, 3> ImplicitBox = Object.BoundingBox().TransformedBox(ObjectTransform.GetRelativeTransform(SampleParticlesTransform));
 		ImplicitBox.Thicken(Thickness);
 		TArray<int32> PotentialParticles;
 		{
-			SCOPE_CYCLE_COUNTER(STAT_UpdateLevelsetFindParticles2);
+			SCOPE_CYCLE_COUNTER(STAT_UpdateLevelsetFindParticles);
 			PotentialParticles = SampleParticles.FindAllIntersections(ImplicitBox);
 		}
 		{
-			SCOPE_CYCLE_COUNTER(STAT_UpdateLevelsetSignedDistance2);
+			SCOPE_CYCLE_COUNTER(STAT_UpdateLevelsetSignedDistance);
 
 			if (Object.GetType(true) == ImplicitObjectType::LevelSet && PotentialParticles.Num() > 0)
 			{
@@ -1218,7 +1214,7 @@ void SampleObject2(const TImplicitObject<float, 3>& Object, const TRigidTransfor
 	}
 	else
 	{
-		SCOPE_CYCLE_COUNTER(STAT_UpdateLevelsetAll2);
+		SCOPE_CYCLE_COUNTER(STAT_UpdateLevelsetAll);
 		if (Object.GetType(true) == ImplicitObjectType::LevelSet && NumParticles > 0)
 		{
 			const TLevelSet<float, 3>* LevelSet = Object.GetObject<Chaos::TLevelSet<float, 3>>();
@@ -1554,11 +1550,11 @@ void UpdateCapsuleBoxConstraint(const TCapsule<T>& A, const TRigidTransform<T, d
 
 }
 
-DECLARE_CYCLE_STAT(TEXT("FindRelevantShapes2"), STAT_FindRelevantShapes2, STATGROUP_ChaosWide);
+DECLARE_CYCLE_STAT(TEXT("TPBDCollisionConstraint::FindRelevantShapes"), STAT_FindRelevantShapes, STATGROUP_ChaosWide);
 template <typename T, int d>
 TArray<Pair<const TImplicitObject<T, d>*, TRigidTransform<T, d>>> FindRelevantShapes2(const TImplicitObject<T,d>* ParticleObj, const TRigidTransform<T,d>& ParticlesTM, const TImplicitObject<T,d>& LevelsetObj, const TRigidTransform<T,d>& LevelsetTM, const T Thickness)
 {
-	SCOPE_CYCLE_COUNTER(STAT_FindRelevantShapes2);
+	SCOPE_CYCLE_COUNTER(STAT_FindRelevantShapes);
 	TArray<Pair<const TImplicitObject<T, d>*, TRigidTransform<T, d>>> RelevantShapes;
 	//find all levelset inner objects
 	if (ParticleObj)
@@ -1586,11 +1582,11 @@ TArray<Pair<const TImplicitObject<T, d>*, TRigidTransform<T, d>>> FindRelevantSh
 	return RelevantShapes;
 }
 
-DECLARE_CYCLE_STAT(TEXT("UpdateUnionUnionConstraint"), STAT_UpdateUnionUnionConstraint2, STATGROUP_ChaosWide);
+DECLARE_CYCLE_STAT(TEXT("TPBDCollisionConstraint::UpdateUnionUnionConstraint"), STAT_UpdateUnionUnionConstraint, STATGROUP_ChaosWide);
 template<ECollisionUpdateType UpdateType, typename T, int d>
 void UpdateUnionUnionConstraint(const T Thickness, TRigidBodyContactConstraint<T, d>& Constraint)
 {
-	SCOPE_CYCLE_COUNTER(STAT_UpdateUnionUnionConstraint2);
+	SCOPE_CYCLE_COUNTER(STAT_UpdateUnionUnionConstraint);
 
 	TGenericParticleHandle<T, d> Particle0 = Constraint.Particle;
 	TGenericParticleHandle<T, d> Particle1 = Constraint.Levelset;
@@ -1620,11 +1616,11 @@ void UpdateUnionUnionConstraint(const T Thickness, TRigidBodyContactConstraint<T
 	}
 }
 
-DECLARE_CYCLE_STAT(TEXT("UpdateSingleUnionConstraint"), STAT_UpdateSingleUnionConstraint2, STATGROUP_ChaosWide);
+DECLARE_CYCLE_STAT(TEXT("TPBDCollisionConstraint::UpdateSingleUnionConstraint"), STAT_UpdateSingleUnionConstraint, STATGROUP_ChaosWide);
 template<ECollisionUpdateType UpdateType, typename T, int d>
 void UpdateSingleUnionConstraint(const T Thickness, TRigidBodyContactConstraint<T, d>& Constraint)
 {
-	SCOPE_CYCLE_COUNTER(STAT_UpdateSingleUnionConstraint2);
+	SCOPE_CYCLE_COUNTER(STAT_UpdateSingleUnionConstraint);
 
 	TGenericParticleHandle<T, d> Particle0 = Constraint.Particle;
 	TGenericParticleHandle<T, d> Particle1 = Constraint.Levelset;
@@ -1644,12 +1640,12 @@ void UpdateSingleUnionConstraint(const T Thickness, TRigidBodyContactConstraint<
 	}
 }
 
-DECLARE_CYCLE_STAT(TEXT("UpdateLevelsetConstraint"), STAT_UpdateLevelsetConstraint2, STATGROUP_ChaosWide);
+DECLARE_CYCLE_STAT(TEXT("TPBDCollisionConstraint::UpdateLevelsetConstraint"), STAT_UpdateLevelsetConstraint, STATGROUP_ChaosWide);
 template<typename T, int d>
 template<ECollisionUpdateType UpdateType>
 void TPBDCollisionConstraint<T, d>::UpdateLevelsetConstraint(const T Thickness, FRigidBodyContactConstraint& Constraint)
 {
-	SCOPE_CYCLE_COUNTER(STAT_UpdateLevelsetConstraint2);
+	SCOPE_CYCLE_COUNTER(STAT_UpdateLevelsetConstraint);
 	
 	TGenericParticleHandle<T, d> Particle0 = Constraint.Particle;
 	TRigidTransform<T, d> ParticlesTM = TRigidTransform<T, d>(Particle0->P(), Particle0->Q());
@@ -1683,11 +1679,11 @@ void UpdateConvexConstraintsUsingCoreShapes(const TImplicitObject<T, d> & AObj,
 	GJKCoreShapeIntersection<T, 3>(AObj, ATM, BObj, BTM, Constraint.Location, Constraint.Phi, Constraint.Normal, Thickness);
 }
 
-DECLARE_CYCLE_STAT(TEXT("UpdateUnionLevelsetConstraint"), STAT_UpdateUnionLevelsetConstraint2, STATGROUP_ChaosWide);
+DECLARE_CYCLE_STAT(TEXT("TPBDCollisionConstraint::UpdateUnionLevelsetConstraint"), STAT_UpdateUnionLevelsetConstraint, STATGROUP_ChaosWide);
 template<ECollisionUpdateType UpdateType, typename T, int d>
 void UpdateUnionLevelsetConstraint(const T Thickness, TRigidBodyContactConstraint<T, d>& Constraint)
 {
-	SCOPE_CYCLE_COUNTER(STAT_UpdateUnionLevelsetConstraint2);
+	SCOPE_CYCLE_COUNTER(STAT_UpdateUnionLevelsetConstraint);
 
 	TGenericParticleHandle<T, d> Particle0 = Constraint.Particle;
 	TGenericParticleHandle<T, d> Particle1 = Constraint.Levelset;
@@ -1759,11 +1755,11 @@ void UpdateUnionLevelsetConstraint(const T Thickness, TRigidBodyContactConstrain
 }
 
 
-DECLARE_CYCLE_STAT(TEXT("UpdateLevelsetUnionConstraint"), STAT_UpdateLevelsetUnionConstraint2, STATGROUP_ChaosWide);
+DECLARE_CYCLE_STAT(TEXT("TPBDCollisionConstraint::UpdateLevelsetUnionConstraint"), STAT_UpdateLevelsetUnionConstraint, STATGROUP_ChaosWide);
 template<ECollisionUpdateType UpdateType, typename T, int d>
 void UpdateLevelsetUnionConstraint(const T Thickness, TRigidBodyContactConstraint<T, d>& Constraint)
 {
-	SCOPE_CYCLE_COUNTER(STAT_UpdateLevelsetUnionConstraint2);
+	SCOPE_CYCLE_COUNTER(STAT_UpdateLevelsetUnionConstraint);
 
 	TGenericParticleHandle<T, d> Particle0 = Constraint.Particle;
 	TGenericParticleHandle<T, d> Particle1 = Constraint.Levelset;
@@ -2156,13 +2152,13 @@ void UpdateConstraintImp2(const TImplicitObject<T, d>& ParticleObject, const TRi
 	}
 }
 
-DECLARE_CYCLE_STAT(TEXT("UpdateConstraint"), STAT_UpdateConstraint2, STATGROUP_ChaosWide);
+DECLARE_CYCLE_STAT(TEXT("TPBDCollisionConstraint::UpdateConstraint"), STAT_UpdateConstraint, STATGROUP_ChaosWide);
 
 template<typename T, int d>
 template<ECollisionUpdateType UpdateType>
 void TPBDCollisionConstraint<T, d>::UpdateConstraint(const T Thickness, FRigidBodyContactConstraint& Constraint)
 {
-	SCOPE_CYCLE_COUNTER(STAT_UpdateConstraint2);
+	SCOPE_CYCLE_COUNTER(STAT_UpdateConstraint);
 
 	Constraint.Phi = Thickness;
 	
diff --git a/Engine/Source/Runtime/Experimental/Chaos/Private/Chaos/PBDJointConstraints.cpp b/Engine/Source/Runtime/Experimental/Chaos/Private/Chaos/PBDJointConstraints.cpp
--- a/Engine/Source/Runtime/Experimental/Chaos/Private/Chaos/PBDJointConstraints.cpp
+++ b/Engine/Source/Runtime/Experimental/Chaos/Private/Chaos/PBDJointConstraints.cpp
@@ -6,12 +6,16 @@
 #include "Chaos/PBDJointConstraintUtilities.h"
 #include "Chaos/Utilities.h"
 #include "ChaosLog.h"
+#include "ChaosStats.h"
 
 #include "HAL/IConsoleManager.h"
 
 
 namespace Chaos
 {
+	DECLARE_CYCLE_STAT(TEXT("TPBDJointConstraints::Apply"), STAT_ApplyJointConstraints, STATGROUP_Chaos);
+
+
 	//
 	// Constraint Handle
 	//
@@ -317,6 +321,8 @@ namespace Chaos
 	template<class T, int d>
 	void TPBDJointConstraints<T, d>::Apply(const T Dt, const TArray<FConstraintHandle*>& InConstraintHandles, const int32 It, const int32 NumIts)
 	{
+		SCOPE_CYCLE_COUNTER(STAT_ApplyJointConstraints);
+
 		// @todo(ccaulfield): make sorting optional
 		// @todo(ccaulfield): handles should be sorted by level by the constraint rule/graph
 		// @todo(ccaulfield): the best sort order depends on whether we are freezing.
diff --git a/Engine/Source/Runtime/Experimental/Chaos/Private/Chaos/PBDRigidsEvolutionGBF.cpp b/Engine/Source/Runtime/Experimental/Chaos/Private/Chaos/PBDRigidsEvolutionGBF.cpp
--- a/Engine/Source/Runtime/Experimental/Chaos/Private/Chaos/PBDRigidsEvolutionGBF.cpp
+++ b/Engine/Source/Runtime/Experimental/Chaos/Private/Chaos/PBDRigidsEvolutionGBF.cpp
@@ -43,6 +43,9 @@ FAutoConsoleVariableRef CVarDisableThreshold(TEXT("p.DisableThreshold2"), Disabl
 DECLARE_CYCLE_STAT(TEXT("TPBDRigidsEvolutionGBF::AdvanceOneTimeStep"), STAT_AdvanceOneTimeStep, STATGROUP_Chaos);
 DECLARE_CYCLE_STAT(TEXT("TPBDRigidsEvolutionGBF::Integrate"), STAT_Integrate, STATGROUP_Chaos);
 DECLARE_CYCLE_STAT(TEXT("TPBDRigidsEvolutionGBF::KinematicTargets"), STAT_KinematicTargets, STATGROUP_Chaos);
+DECLARE_CYCLE_STAT(TEXT("TPBDRigidsEvolutionGBF::ApplyConstraints"), STAT_ApplyConstraints, STATGROUP_Chaos);
+DECLARE_CYCLE_STAT(TEXT("TPBDRigidsEvolutionGBF::UpdateVelocities"), STAT_UpdateVelocites, STATGROUP_Chaos);
+DECLARE_CYCLE_STAT(TEXT("TPBDRigidsEvolutionGBF::ApplyPushOut"), STAT_ApplyPushOut, STATGROUP_Chaos);
 DECLARE_CYCLE_STAT(TEXT("TPBDRigidsEvolutionGBF::UpdateConstraintPositionBasedState"), STAT_UpdateConstraintPositionBasedState, STATGROUP_Chaos);
 DECLARE_CYCLE_STAT(TEXT("TPBDRigidsEvolutionGBF::CreateConstraintGraph"), STAT_CreateConstraintGraph, STATGROUP_Chaos);
 DECLARE_CYCLE_STAT(TEXT("TPBDRigidsEvolutionGBF::CreateIslands"), STAT_CreateIslands, STATGROUP_Chaos);
@@ -145,16 +148,25 @@ void TPBDRigidsEvolutionGBF<T, d>::AdvanceOneTimeStep(T Dt)
 		PhysicsParallelFor(ConstraintGraph.NumIslands(), [&](int32 Island) {
 			const TArray<TGeometryParticleHandle<T, d>*>& IslandParticles = ConstraintGraph.GetIslandParticles(Island);
 
-			ApplyConstraints(Dt, Island);
+			{
+				SCOPE_CYCLE_COUNTER(STAT_ApplyConstraints);
+				ApplyConstraints(Dt, Island);
+			}
 
 			if (PostApplyCallback != nullptr)
 			{
 				PostApplyCallback(Island);
 			}
 
-			UpdateVelocities(Dt, Island);
+			{
+				SCOPE_CYCLE_COUNTER(STAT_UpdateVelocites);
+				UpdateVelocities(Dt, Island);
+			}
 
-			ApplyPushOut(Dt, Island);
+			{
+				SCOPE_CYCLE_COUNTER(STAT_ApplyPushOut);
+				ApplyPushOut(Dt, Island);
+			}
 
 			if (PostApplyPushOutCallback != nullptr)
 			{
diff --git a/Engine/Source/Runtime/Experimental/Chaos/Private/Chaos/ParticleIterator.cpp b/Engine/Source/Runtime/Experimental/Chaos/Private/Chaos/ParticleIterator.cpp
--- a/Engine/Source/Runtime/Experimental/Chaos/Private/Chaos/ParticleIterator.cpp
+++ b/Engine/Source/Runtime/Experimental/Chaos/Private/Chaos/ParticleIterator.cpp
@@ -4,6 +4,4 @@
 
 namespace Chaos
 {
-	CHAOS_API int32 ChaosParticleParallelFor = 1;
-	FAutoConsoleVariableRef CVarChaosParticleParallelFor(TEXT("p.ChaosParticleParallelFor"), ChaosParticleParallelFor, TEXT("ParticleParallelFor function style for chaos"));
 }
diff --git a/Engine/Source/Runtime/Experimental/Chaos/Public/Chaos/PBDCollisionConstraintImp.h b/Engine/Source/Runtime/Experimental/Chaos/Public/Chaos/PBDCollisionConstraintImp.h
--- a/Engine/Source/Runtime/Experimental/Chaos/Public/Chaos/PBDCollisionConstraintImp.h
+++ b/Engine/Source/Runtime/Experimental/Chaos/Public/Chaos/PBDCollisionConstraintImp.h
@@ -86,6 +86,12 @@ namespace Chaos
 
 #endif
 
+#if UE_BUILD_SHIPPING
+	const bool bDisableCollisionParallelFor = false;
+#else
+	CHAOS_API extern bool bDisableCollisionParallelFor;
+#endif
+
 	template <typename T, int d>
 	struct TSimOverlapVisitor
 	{
@@ -281,7 +287,7 @@ namespace Chaos
 				NarrowPhaseSkipped.Record(NPSkipped);
 				NarrowPhaseRejected.Record(RejectedNP);
 #endif
-			}, bGatherStats || CollisionConstraintsForceSingleThreaded);
+			}, bGatherStats || bDisableCollisionParallelFor);
 			
 			{
 				SCOPE_CYCLE_COUNTER(STAT_ComputeConstraintsSU);
diff --git a/Engine/Source/Runtime/Experimental/Chaos/Public/Chaos/ParticleIterator.h b/Engine/Source/Runtime/Experimental/Chaos/Public/Chaos/ParticleIterator.h
--- a/Engine/Source/Runtime/Experimental/Chaos/Public/Chaos/ParticleIterator.h
+++ b/Engine/Source/Runtime/Experimental/Chaos/Public/Chaos/ParticleIterator.h
@@ -12,7 +12,11 @@ DECLARE_CYCLE_STAT(TEXT("HandleViewParallelForImp"), STAT_HandleViewParallelForI
 
 namespace Chaos
 {
-CHAOS_API extern int32 ChaosParticleParallelFor;
+#if UE_BUILD_SHIPPING
+	const bool bDisableParticleParallelFor = false;
+#else
+	CHAOS_API extern bool bDisableParticleParallelFor;
+#endif
 
 template <typename TSOA>
 class TConstParticleView;
@@ -89,15 +93,13 @@ void ParticlesParallelFor(const TView& Particles, const Lambda& Func, bool bForc
 {
 	SCOPE_CYCLE_COUNTER(STAT_ParticlesParallelFor);
 
-	switch (ChaosParticleParallelFor)
+	if (!bForceSingleThreaded && !bDisableParticleParallelFor)
 	{
-		case 0:
-			Chaos::ParticlesSequentialFor(Particles, Func);
-			break;
-
-		case 1:
-			Chaos::ParticlesParallelForImp(Particles, Func);
-			break;
+		Chaos::ParticlesParallelForImp(Particles, Func);
+	}
+	else
+	{
+		Chaos::ParticlesSequentialFor(Particles, Func);
 	}
 }
 
