Fix CIS - duplicate Stat name

diff --git a/Engine/Source/Runtime/Experimental/Chaos/Private/Chaos/PBDRigidsEvolutionGBF.cpp b/Engine/Source/Runtime/Experimental/Chaos/Private/Chaos/PBDRigidsEvolutionGBF.cpp
--- a/Engine/Source/Runtime/Experimental/Chaos/Private/Chaos/PBDRigidsEvolutionGBF.cpp
+++ b/Engine/Source/Runtime/Experimental/Chaos/Private/Chaos/PBDRigidsEvolutionGBF.cpp
@@ -40,17 +40,17 @@ int DisableThreshold = 5;
 FAutoConsoleVariableRef CVarDisableThreshold(TEXT("p.DisableThreshold2"), DisableThreshold, TEXT("Disable threshold frames to transition to sleeping"));
 
 
-DECLARE_CYCLE_STAT(TEXT("TPBDRigidsEvolutionGBF::AdvanceOneTimeStep"), STAT_AdvanceOneTimeStep, STATGROUP_Chaos);
-DECLARE_CYCLE_STAT(TEXT("TPBDRigidsEvolutionGBF::Integrate"), STAT_Integrate, STATGROUP_Chaos);
-DECLARE_CYCLE_STAT(TEXT("TPBDRigidsEvolutionGBF::KinematicTargets"), STAT_KinematicTargets, STATGROUP_Chaos);
-DECLARE_CYCLE_STAT(TEXT("TPBDRigidsEvolutionGBF::ApplyConstraints"), STAT_ApplyConstraints, STATGROUP_Chaos);
-DECLARE_CYCLE_STAT(TEXT("TPBDRigidsEvolutionGBF::UpdateVelocities"), STAT_UpdateVelocites, STATGROUP_Chaos);
-DECLARE_CYCLE_STAT(TEXT("TPBDRigidsEvolutionGBF::ApplyPushOut"), STAT_ApplyPushOut, STATGROUP_Chaos);
-DECLARE_CYCLE_STAT(TEXT("TPBDRigidsEvolutionGBF::UpdateConstraintPositionBasedState"), STAT_UpdateConstraintPositionBasedState, STATGROUP_Chaos);
-DECLARE_CYCLE_STAT(TEXT("TPBDRigidsEvolutionGBF::CreateConstraintGraph"), STAT_CreateConstraintGraph, STATGROUP_Chaos);
-DECLARE_CYCLE_STAT(TEXT("TPBDRigidsEvolutionGBF::CreateIslands"), STAT_CreateIslands, STATGROUP_Chaos);
-DECLARE_CYCLE_STAT(TEXT("TPBDRigidsEvolutionGBF::ParallelSolve"), STAT_ParallelSolve, STATGROUP_Chaos);
-DECLARE_CYCLE_STAT(TEXT("TPBDRigidsEvolutionGBF::DeactivateSleep"), STAT_DeactivateSleep, STATGROUP_Chaos);
+DECLARE_CYCLE_STAT(TEXT("TPBDRigidsEvolutionGBF::AdvanceOneTimeStep"), STAT_Evo_AdvanceOneTimeStep, STATGROUP_Chaos);
+DECLARE_CYCLE_STAT(TEXT("TPBDRigidsEvolutionGBF::Integrate"), STAT_Evo_Integrate, STATGROUP_Chaos);
+DECLARE_CYCLE_STAT(TEXT("TPBDRigidsEvolutionGBF::KinematicTargets"), STAT_Evo_KinematicTargets, STATGROUP_Chaos);
+DECLARE_CYCLE_STAT(TEXT("TPBDRigidsEvolutionGBF::ApplyConstraints"), STAT_Evo_ApplyConstraints, STATGROUP_Chaos);
+DECLARE_CYCLE_STAT(TEXT("TPBDRigidsEvolutionGBF::UpdateVelocities"), STAT_Evo_UpdateVelocites, STATGROUP_Chaos);
+DECLARE_CYCLE_STAT(TEXT("TPBDRigidsEvolutionGBF::ApplyPushOut"), STAT_Evo_ApplyPushOut, STATGROUP_Chaos);
+DECLARE_CYCLE_STAT(TEXT("TPBDRigidsEvolutionGBF::UpdateConstraintPositionBasedState"), STAT_Evo_UpdateConstraintPositionBasedState, STATGROUP_Chaos);
+DECLARE_CYCLE_STAT(TEXT("TPBDRigidsEvolutionGBF::CreateConstraintGraph"), STAT_Evo_CreateConstraintGraph, STATGROUP_Chaos);
+DECLARE_CYCLE_STAT(TEXT("TPBDRigidsEvolutionGBF::CreateIslands"), STAT_Evo_CreateIslands, STATGROUP_Chaos);
+DECLARE_CYCLE_STAT(TEXT("TPBDRigidsEvolutionGBF::ParallelSolve"), STAT_Evo_ParallelSolve, STATGROUP_Chaos);
+DECLARE_CYCLE_STAT(TEXT("TPBDRigidsEvolutionGBF::DeactivateSleep"), STAT_Evo_DeactivateSleep, STATGROUP_Chaos);
 
 int32 SerializeEvolution = 0;
 FAutoConsoleVariableRef CVarSerializeEvolution(TEXT("p.SerializeEvolution"), SerializeEvolution, TEXT(""));
@@ -90,7 +90,7 @@ void SerializeToDisk(TEvolution& Evolution)
 template <typename T, int d>
 void TPBDRigidsEvolutionGBF<T, d>::AdvanceOneTimeStep(T Dt)
 {
-	SCOPE_CYCLE_COUNTER(STAT_AdvanceOneTimeStep);
+	SCOPE_CYCLE_COUNTER(STAT_Evo_AdvanceOneTimeStep);
 
 #if !UE_BUILD_SHIPPING
 	if (SerializeEvolution)
@@ -100,12 +100,12 @@ void TPBDRigidsEvolutionGBF<T, d>::AdvanceOneTimeStep(T Dt)
 #endif
 
 	{
-		SCOPE_CYCLE_COUNTER(STAT_Integrate);
+		SCOPE_CYCLE_COUNTER(STAT_Evo_Integrate);
 		Integrate(Particles.GetNonDisabledDynamicView(), Dt);	//Question: should we use an awake view?
 	}
 
 	{
-		SCOPE_CYCLE_COUNTER(STAT_KinematicTargets);
+		SCOPE_CYCLE_COUNTER(STAT_Evo_KinematicTargets);
 		ApplyKinematicTargets(Dt);
 	}
 
@@ -120,15 +120,15 @@ void TPBDRigidsEvolutionGBF<T, d>::AdvanceOneTimeStep(T Dt)
 	}
 
 	{
-		SCOPE_CYCLE_COUNTER(STAT_UpdateConstraintPositionBasedState);
+		SCOPE_CYCLE_COUNTER(STAT_Evo_UpdateConstraintPositionBasedState);
 		UpdateConstraintPositionBasedState(Dt);
 	}
 	{
-		SCOPE_CYCLE_COUNTER(STAT_CreateConstraintGraph);
+		SCOPE_CYCLE_COUNTER(STAT_Evo_CreateConstraintGraph);
 		CreateConstraintGraph();
 	}
 	{
-		SCOPE_CYCLE_COUNTER(STAT_CreateIslands);
+		SCOPE_CYCLE_COUNTER(STAT_Evo_CreateIslands);
 		CreateIslands();
 	}
 
@@ -144,12 +144,12 @@ void TPBDRigidsEvolutionGBF<T, d>::AdvanceOneTimeStep(T Dt)
 	SleepedIslands.SetNum(ConstraintGraph.NumIslands());
 	if(Dt > 0)
 	{
-		SCOPE_CYCLE_COUNTER(STAT_ParallelSolve);
+		SCOPE_CYCLE_COUNTER(STAT_Evo_ParallelSolve);
 		PhysicsParallelFor(ConstraintGraph.NumIslands(), [&](int32 Island) {
 			const TArray<TGeometryParticleHandle<T, d>*>& IslandParticles = ConstraintGraph.GetIslandParticles(Island);
 
 			{
-				SCOPE_CYCLE_COUNTER(STAT_ApplyConstraints);
+				SCOPE_CYCLE_COUNTER(STAT_Evo_ApplyConstraints);
 				ApplyConstraints(Dt, Island);
 			}
 
@@ -159,12 +159,12 @@ void TPBDRigidsEvolutionGBF<T, d>::AdvanceOneTimeStep(T Dt)
 			}
 
 			{
-				SCOPE_CYCLE_COUNTER(STAT_UpdateVelocites);
+				SCOPE_CYCLE_COUNTER(STAT_Evo_UpdateVelocites);
 				UpdateVelocities(Dt, Island);
 			}
 
 			{
-				SCOPE_CYCLE_COUNTER(STAT_ApplyPushOut);
+				SCOPE_CYCLE_COUNTER(STAT_Evo_ApplyPushOut);
 				ApplyPushOut(Dt, Island);
 			}
 
@@ -210,7 +210,7 @@ void TPBDRigidsEvolutionGBF<T, d>::AdvanceOneTimeStep(T Dt)
 		});
 	}
 	{
-		SCOPE_CYCLE_COUNTER(STAT_DeactivateSleep);
+		SCOPE_CYCLE_COUNTER(STAT_Evo_DeactivateSleep);
 		for (int32 Island = 0; Island < ConstraintGraph.NumIslands(); ++Island)
 		{
 			if (SleepedIslands[Island])
