Optimization for common case in BV sweep

diff --git a/Engine/Source/Runtime/Experimental/Chaos/Public/Chaos/BoundingVolume.h b/Engine/Source/Runtime/Experimental/Chaos/Public/Chaos/BoundingVolume.h
--- a/Engine/Source/Runtime/Experimental/Chaos/Public/Chaos/BoundingVolume.h
+++ b/Engine/Source/Runtime/Experimental/Chaos/Public/Chaos/BoundingVolume.h
@@ -586,19 +586,70 @@ private:
 			T TOI;
 		};
 
+		const TVector<int32, d> StartMinIndex = MGrid.ClampIndex(MGrid.Cell(Start - QueryHalfExtents));
+		const TVector<int32, d> StartMaxIndex = MGrid.ClampIndex(MGrid.Cell(Start + QueryHalfExtents));
+
+		if (StartMinIndex == StartMaxIndex)
+		{
+			const TVector<T, d> End = Start + CurrentLength * Dir;
+			const TVector<int32, d> EndMinIndex = MGrid.ClampIndex(MGrid.Cell(End  - QueryHalfExtents));
+			const TVector<int32, d> EndMaxIndex = MGrid.ClampIndex(MGrid.Cell(End + QueryHalfExtents));
+			if (StartMinIndex == EndMinIndex && StartMinIndex == EndMaxIndex)
+			{
+				//sweep is fully contained within one cell, this is a common special case - just test all elements
+				const auto& Elems = MElements(StartMinIndex);
+				TVector<T, d> TmpPoint;
+				for (const auto& Elem : Elems)
+				{
+					const TBox<T, d>& InstanceBounds = Elem.Bounds;
+					const TVector<T, d> Min = InstanceBounds.Min() - QueryHalfExtents;
+					const TVector<T, d> Max = InstanceBounds.Max() + QueryHalfExtents;
+					if (TBox<T, d>::RaycastFast(Min, Max, Start, Dir, InvDir, bParallel, CurrentLength, InvCurrentLength, TOI, TmpPoint))
+					{
+						TSpatialVisitorData<TPayloadType> VisitData(Elem.Payload, true, InstanceBounds);
+						const bool bContinue = Visitor.VisitSweep(VisitData, CurrentLength);
+						if (!bContinue)
+						{
+							return false;
+						}
+						InvCurrentLength = 1 / CurrentLength;
+					}
+				}
+
+				return true;
+			}
+		}
 
-		TVector<T, d> HitPoint;
 		FGridSet IdxsSeen(MGrid.Counts());
 		FGridSet CellsVisited(MGrid.Counts());
+		TArray<FCellIntersection> IdxsQueue;	//cells we need to visit
+
+		for (int32 X = StartMinIndex[0]; X <= StartMaxIndex[0]; ++X)
+		{
+			for (int32 Y = StartMinIndex[1]; Y <= StartMaxIndex[1]; ++Y)
+			{
+				for (int32 Z = StartMinIndex[2]; Z <= StartMaxIndex[2]; ++Z)
+				{
+					const TVector<int32, 3> Idx(X, Y, Z);
+					IdxsQueue.Add({Idx, 0 });
+					IdxsSeen.Add(Idx);
+				}
+			}
+		}
+
+		TVector<T, d> HitPoint;
 		const bool bInitialHit = TBox<T, d>::RaycastFast(GlobalBounds.Min(), GlobalBounds.Max(), Start, Dir, InvDir, bParallel, CurrentLength, InvCurrentLength, TOI, HitPoint);
-		if (bInitialHit)
+		if (bInitialHit)	//NOTE: it's possible to have a non empty IdxsQueue and bInitialHit be false. This is because the IdxsQueue works off clamped cells which we can skip
 		{
 			//Flood fill from inflated cell so that we get all cells along the ray
 			TVector<int32, d> HitCellIdx = MGrid.Cell(HitPoint);
 			HitCellIdx = MGrid.ClampIndex(HitCellIdx);	//inflation means we likely are outside grid, just get closest cell
 
-			TArray<FCellIntersection> IdxsQueue;	//cells we need to visit
-			IdxsQueue.Add({ HitCellIdx, TOI });
+			if (!IdxsSeen.Contains(HitCellIdx))
+			{
+				ensure(TOI > 0);	//Not an initial overlap so TOI must be positive
+				IdxsQueue.Add({ HitCellIdx, TOI });
+			}
 
 			const TVector<T, d> HalfDx = MGrid.Dx() * (T)0.5;
 
