Make sure large objects make it into acceleration bucket 1

diff --git a/Engine/Source/Runtime/Experimental/ChaosSolvers/Private/PBDRigidsSolver.cpp b/Engine/Source/Runtime/Experimental/ChaosSolvers/Private/PBDRigidsSolver.cpp
--- a/Engine/Source/Runtime/Experimental/ChaosSolvers/Private/PBDRigidsSolver.cpp
+++ b/Engine/Source/Runtime/Experimental/ChaosSolvers/Private/PBDRigidsSolver.cpp
@@ -133,6 +133,14 @@ namespace Chaos
 		Reset();
 	}
 
+	float MaxBoundsForTree = 10000;
+	FAutoConsoleVariableRef CVarMaxBoundsForTree(
+		TEXT("p.MaxBoundsForTree"),
+		MaxBoundsForTree,
+		TEXT("The max bounds before moving object into a large objects structure. Only applies on object registration")
+		TEXT(""),
+		ECVF_Default);
+
 	void FPBDRigidsSolver::RegisterObject(TGeometryParticle<float, 3>* GTParticle)
 	{
 		UE_LOG(LogPBDRigidsSolverSolver, Verbose, TEXT("FPBDRigidsSolver::RegisterObject()"));
@@ -140,7 +148,7 @@ namespace Chaos
 		// Make sure this particle doesn't already have a proxy
 		checkSlow(GTParticle->Proxy == nullptr);
 
-		if (GTParticle->Geometry() && GTParticle->Geometry()->HasBoundingBox() && GTParticle->Geometry()->BoundingBox().Extents().Max() >= 100000)
+		if (GTParticle->Geometry() && GTParticle->Geometry()->HasBoundingBox() && GTParticle->Geometry()->BoundingBox().Extents().Max() >= MaxBoundsForTree)
 		{
 			GTParticle->SetSpatialIdx(FSpatialAccelerationIdx{ 1,0 });
 		}
diff --git a/Engine/Source/Runtime/Experimental/ChaosSolvers/Private/PhysicsProxy/SingleParticlePhysicsProxy.cpp b/Engine/Source/Runtime/Experimental/ChaosSolvers/Private/PhysicsProxy/SingleParticlePhysicsProxy.cpp
--- a/Engine/Source/Runtime/Experimental/ChaosSolvers/Private/PhysicsProxy/SingleParticlePhysicsProxy.cpp
+++ b/Engine/Source/Runtime/Experimental/ChaosSolvers/Private/PhysicsProxy/SingleParticlePhysicsProxy.cpp
@@ -115,6 +115,8 @@ void FSingleParticlePhysicsProxy<Chaos::TKinematicGeometryParticle<float, 3>>::P
 		RigidHandle->SetSharedGeometry(Data->Geometry);
 		RigidHandle->SetV(Data->MV);
 		RigidHandle->SetW(Data->MW);
+
+		RigidHandle->SetSpatialIdx(Data->SpatialIdx);	//todo: this needs to only happen once during initialization
 		
 		if (Data->Geometry && Data->Geometry->HasBoundingBox())
 		{
@@ -182,6 +184,7 @@ void FSingleParticlePhysicsProxy<Chaos::TPBDRigidParticle<float, 3>>::PushToPhys
 		RigidHandle->SetInvI(Data->MInvI);
 		RigidHandle->SetExternalForce(Data->MExternalForce);
 		RigidHandle->SetExternalTorque(Data->MExternalTorque);
+		RigidHandle->SetSpatialIdx(Data->SpatialIdx);	//todo: this needs to only happen once during initialization
 		GetSolver()->GetEvolution()->GetGravityForces().SetEnabled(*RigidHandle, Data->MGravityEnabled);
 
 		if (Data->Geometry && Data->Geometry->HasBoundingBox())
