Added an implicit triangle object and a faster TTriangle object compatible with SQ so we can move away from using a convex hull per triangle in trimeshes for overlaps and sweeps. Approx 25% improvement in tri sweep narrowphase costs

diff --git a/Engine/Source/Runtime/Experimental/Chaos/Private/Chaos/TriangleMeshImplicitObject.cpp b/Engine/Source/Runtime/Experimental/Chaos/Private/Chaos/TriangleMeshImplicitObject.cpp
--- a/Engine/Source/Runtime/Experimental/Chaos/Private/Chaos/TriangleMeshImplicitObject.cpp
+++ b/Engine/Source/Runtime/Experimental/Chaos/Private/Chaos/TriangleMeshImplicitObject.cpp
@@ -1,8 +1,8 @@
 // Copyright 1998-2019 Epic Games, Inc. All Rights Reserved.
 #include "Chaos/TriangleMeshImplicitObject.h"
 #include "Chaos/Capsule.h"
-#include "Chaos/Convex.h"
 #include "Chaos/GJK.h"
+#include "Chaos/Triangle.h"
 
 namespace Chaos
 {
@@ -269,17 +269,7 @@ bool TTriangleMeshImplicitObject<T>::OverlapGeom(const TImplicitObject<T, 3>& Qu
 		//However, maybe we should check if it's behind the triangle plane. Also, we should enforce this winding in some way
 		const TVector<T, 3> Offset = TVector<T, 3>::CrossProduct(AB, AC);
 
-		// Ugly but required for now until we have an easier way to do tri collisons 
-		TParticles<T, 3> Particles;
-		Particles.AddParticles(3);
-
-		Particles.X(0) = A;
-		Particles.X(1) = B;
-		Particles.X(2) = C;
-
-		TConvex<T, 3> TriangleConvex(Particles);
-
-		return GJKIntersection(TriangleConvex, QueryGeom, QueryTM, Thickness, Offset);
+		return GJKIntersection(TTriangle<T>(A, B, C), QueryGeom, QueryTM, Thickness, Offset);
 	};
 
 	bool bResult = false;
@@ -322,20 +312,17 @@ struct TTriangleMeshSweepVisitor
 	{
 		const int32 TriIdx = VisitData.Payload;
 
-		TParticles<T, 3> TriParticles;
-		TriParticles.AddParticles(3);
-		TriParticles.X(0) = TriMesh.MParticles.X(TriMesh.MElements[TriIdx][0]);
-		TriParticles.X(1) = TriMesh.MParticles.X(TriMesh.MElements[TriIdx][1]);
-		TriParticles.X(2) = TriMesh.MParticles.X(TriMesh.MElements[TriIdx][2]);
-
-		TConvex<T, 3> TriConvex(TriParticles);
-
 		T Time;
 		TVector<T, 3> HitPosition;
 		TVector<T, 3> HitNormal;
-		if (GJKRaycast<T>(TriConvex, QueryGeom, StartTM, Dir, CurLength, Time, HitPosition, HitNormal, Thickness))
+
+		TTriangle<T> Tri(TriMesh.MParticles.X(TriMesh.MElements[TriIdx][0]),
+			TriMesh.MParticles.X(TriMesh.MElements[TriIdx][1]),
+			TriMesh.MParticles.X(TriMesh.MElements[TriIdx][2]));
+
+		if(GJKRaycast<T>(Tri, QueryGeom, StartTM, Dir, CurLength, Time, HitPosition, HitNormal, Thickness))
 		{
-			if (Time < OutTime)
+			if(Time < OutTime)
 			{
 				OutNormal = HitNormal;
 				OutPosition = HitPosition;
@@ -343,7 +330,7 @@ struct TTriangleMeshSweepVisitor
 				CurLength = Time;
 				OutFaceIndex = TriIdx;
 
-				if (Time == 0)
+				if(Time == 0)
 				{
 					//initial overlap, no one will beat this
 					return false;
diff --git a/Engine/Source/Runtime/Experimental/Chaos/Public/Chaos/ImplicitObject.h b/Engine/Source/Runtime/Experimental/Chaos/Public/Chaos/ImplicitObject.h
--- a/Engine/Source/Runtime/Experimental/Chaos/Public/Chaos/ImplicitObject.h
+++ b/Engine/Source/Runtime/Experimental/Chaos/Public/Chaos/ImplicitObject.h
@@ -43,7 +43,8 @@ enum class ImplicitObjectType : int8
 	Cylinder,
 	TriangleMesh,
 	HeightField,
-	Scaled
+	Scaled,
+	Triangle
 };
 
 namespace EImplicitObject
diff --git a/Engine/Source/Runtime/Experimental/Chaos/Public/Chaos/Triangle.h b/Engine/Source/Runtime/Experimental/Chaos/Public/Chaos/Triangle.h
new file mode 100644
--- /dev/null
+++ b/Engine/Source/Runtime/Experimental/Chaos/Public/Chaos/Triangle.h
@@ -0,0 +1,223 @@
+// Copyright 1998-2019 Epic Games, Inc. All Rights Reserved.
+
+#pragma once
+
+#include "ImplicitObject.h"
+#include "Plane.h"
+
+namespace Chaos
+{
+	template<typename T>
+	class TTriangle
+	{
+	public:
+
+		TTriangle(const TVector<T, 3>& InA, const TVector<T, 3>& InB, const TVector<T, 3>& InC)
+			: A(InA)
+			, B(InB)
+			, C(InC)
+		{}
+
+		TVector<T, 3>& operator[](uint32 InIndex)
+		{
+			check(InIndex < 3);
+			return (&A)[InIndex];
+		}
+
+		const TVector<T, 3>& operator[](uint32 InIndex) const
+		{
+			check(InIndex < 3);
+			return (&A)[InIndex];
+		}
+
+		FORCEINLINE TVector<T, 3> GetNormal() const
+		{
+			return TVector<T, 3>::CrossProduct(B - A, C - A).GetSafeNormal();
+		}
+
+		FORCEINLINE TPlane<T, 3> GetPlane() const
+		{
+			return TPlane<T, 3>(A, GetNormal());
+		}
+
+		FORCEINLINE T PhiWithNormal(const TVector<T, 3>& InSamplePoint, TVector<T, 3>& OutNormal) const
+		{
+			OutNormal = GetNormal();
+			TVector<T, 3> ClosestPoint = FindClosestPointOnTriangle(GetPlane(), A, B, C, InSamplePoint);
+			return TVector<T, 3>::DotProduct((InSamplePoint - ClosestPoint), OutNormal);
+		}
+
+		FORCEINLINE TVector<T, 3> Support(const TVector<T, 3>& Direction, const T Thickness) const
+		{
+			const float DotA = TVector<T, 3>::DotProduct(A, Direction);
+			const float DotB = TVector<T, 3>::DotProduct(B, Direction);
+			const float DotC = TVector<T, 3>::DotProduct(C, Direction);
+
+			if(DotA > DotB && DotA > DotC)
+			{
+				if(Thickness != 0)
+				{
+					return A + A.GetSafeNormal() * Thickness;
+				}
+				return A;
+			}
+			else if(DotB > DotA && DotB > DotC)
+			{
+				if(Thickness != 0)
+				{
+					return B + B.GetSafeNormal() * Thickness;
+				}
+				return B;
+			}
+
+			if(Thickness != 0)
+			{
+				return C + C.GetSafeNormal() * Thickness;
+			}
+			return C;
+		}
+
+		FORCEINLINE bool Raycast(const TVector<T, 3>& StartPoint, const TVector<T, 3>& Dir, const T Length, const T Thickness, T& OutTime, TVector<T, 3>& OutPosition, TVector<T, 3>& OutNormal, int32& OutFaceIndex) const
+		{
+			// No face as this is only one triangle
+			OutFaceIndex = INDEX_NONE;
+
+			// Pass through GJK #BGTODO Maybe specialise if it's possible to be faster
+			const TRigidTransform<T, 3> StartTM(StartPoint, TRotation<T, 3>::FromIdentity());
+			const TSphere<T, 3> Sphere(TVector<T, 3>(0), Thickness);
+			return GJKRaycast(*this, Sphere, StartTM, Dir, Length, OutTime, OutPosition, OutNormal);
+		}
+
+		FORCEINLINE bool Overlap(const TVector<T, 3>& Point, const T Thickness) const 
+		{
+			const TVector<T, 3> ClosestPoint = FindClosestPointOnTriangle(GetPlane(), A, B, C, Point);
+			const float AdjustedThickness = FMath::Max(Thickness, KINDA_SMALL_NUMBER);
+			return (Point - ClosestPoint).SizeSquared() <= (AdjustedThickness * AdjustedThickness);
+		}
+
+		FORCEINLINE bool IsConvex() const
+		{
+			return true;
+		}
+
+	private:
+
+		friend FChaosArchive& operator<<(FChaosArchive& Ar, TTriangle<T>& Value);
+
+		TVector<T, 3> A;
+		TVector<T, 3> B;
+		TVector<T, 3> C;
+	};
+
+	template<typename T>
+	FChaosArchive& operator<<(FChaosArchive& Ar, TTriangle<T>& Value)
+	{
+		Ar << Value.A << Value.B << Value.C;
+		return Ar;
+	}
+
+	template<typename T>
+	class TImplicitTriangle final : public TImplicitObject<T, 3>
+	{
+	public:
+
+		TImplicitTriangle()
+			: TImplicitObject<T, 3>(EImplicitObject::IsConvex | EImplicitObject::HasBoundingBox, ImplicitObjectType::Triangle)
+		{}
+
+		TImplicitTriangle(const TImplicitTriangle&) = delete;
+
+		TImplicitTriangle(TImplicitTriangle&& InToSteal)
+			: TImplicitObject<T, 3>(EImplicitObject::IsConvex | EImplicitObject::HasBoundingBox, ImplicitObjectType::Triangle)
+		{}
+
+		TImplicitTriangle(const TVector<T, 3>& InA, const TVector<T, 3>& InB, const TVector<T, 3>& InC)
+			: TImplicitObject<T, 3>(EImplicitObject::IsConvex | EImplicitObject::HasBoundingBox, ImplicitObjectType::Triangle)
+			, Tri(InA, InB, InC)
+		{
+			Bounds = TBox<T, 3>(Tri[0], Tri[0]);
+			Bounds.GrowToInclude(Tri[1]);
+			Bounds.GrowToInclude(Tri[2]);
+		}
+
+		TVector<T, 3>& operator[](uint32 InIndex)
+		{
+			return Tri[InIndex];
+		}
+
+		const TVector<T, 3>& operator[](uint32 InIndex) const
+		{
+			return Tri[InIndex];
+		}
+
+		TVector<T, 3> GetNormal() const
+		{
+			return Tri.GetNormal();
+		}
+
+		TPlane<T, 3> GetPlane() const
+		{
+			return Tri.GetPlane();
+		}
+
+		static ImplicitObjectType GetType()
+		{
+			return ImplicitObjectType::Triangle;
+		}
+
+		virtual T PhiWithNormal(const TVector<T, 3>& InSamplePoint, TVector<T, 3>& OutNormal) const override
+		{
+			return Tri.PhiWithNormal(InSamplePoint, OutNormal);
+		}
+
+		virtual const class TBox<T, 3>& BoundingBox() const override
+		{
+			return Bounds;
+		}
+
+		virtual TVector<T, 3> Support(const TVector<T, 3>& Direction, const T Thickness) const override
+		{
+			return Tri.Support(Direction, Thickness);
+		}
+
+		virtual bool Raycast(const TVector<T, 3>& StartPoint, const TVector<T, 3>& Dir, const T Length, const T Thickness, T& OutTime, TVector<T, 3>& OutPosition, TVector<T, 3>& OutNormal, int32& OutFaceIndex) const override
+		{
+			return Tri.Raycast(StartPoint, Dir, Length, Thickness, OutTime, OutPosition, OutNormal, OutFaceIndex);
+		}
+
+		virtual TVector<T, 3> FindGeometryOpposingNormal(const TVector<T, 3>& DenormDir, int32 FaceIndex, const TVector<T, 3>& OriginalNormal) const override
+		{
+			return Tri.GetNormal();
+		}
+
+		virtual bool Overlap(const TVector<T, 3>& Point, const T Thickness) const override
+		{
+			return Tri.Overlap(Point, Thickness);
+		}
+
+		virtual FString ToString() const override
+		{
+			return FString::Printf(TEXT("Triangle: A: %s, B: %s, C: %s"), LexToString(Tri[0]), LexToString(Tri[1]), LexToString(Tri[2]));
+		}
+
+		virtual void Serialize(FChaosArchive& Ar) override
+		{
+			Ar << Tri;
+		}
+
+		virtual uint32 GetTypeHash() const override
+		{
+			return HashCombine(Bounds.GetTypeHash(), HashCombine(::GetTypeHash(Tri[0]), HashCombine(::GetTypeHash(Tri[1]), ::GetTypeHash(Tri[2]))));
+		}
+
+		virtual FName GetTypeName() const override
+		{
+			return FName("Triangle");
+		}
+
+	private:
+
+		TTriangle<T> Tri;
+		TBox<T, 3> Bounds;
+	};
+}
\ No newline at end of file
