Chaos : Collision Caching

- Moved the TMap for the Handles back to a TArray that mirrors the Constraints array.
- Added a TMap indexed on <Particle,Levelset> that stores the array of handles associated with the collision pair.
#rnx

diff --git a/Engine/Source/Runtime/Experimental/Chaos/Private/Chaos/PBDCollisionConstraint.cpp b/Engine/Source/Runtime/Experimental/Chaos/Private/Chaos/PBDCollisionConstraint.cpp
--- a/Engine/Source/Runtime/Experimental/Chaos/Private/Chaos/PBDCollisionConstraint.cpp
+++ b/Engine/Source/Runtime/Experimental/Chaos/Private/Chaos/PBDCollisionConstraint.cpp
@@ -34,6 +34,11 @@ FAutoConsoleVariableRef CVarConstraintBPBVHDepth(TEXT("p.ConstraintBPBVHDepth"),
 int32 BPTreeOfGrids = 1;
 FAutoConsoleVariableRef CVarBPTreeOfGrids(TEXT("p.BPTreeOfGrids"), BPTreeOfGrids, TEXT("Whether to use a seperate tree of grids for bp"));
 
+bool bSupportCollisionHistoryCVar = false;
+FAutoConsoleVariableRef CVarSupportCollisionHistory(TEXT("p.SupportCollisionHistory"), bSupportCollisionHistoryCVar, TEXT("Whether to support collision history caching."));
+
+int32 CollisionHistoryLifespanCVar = 1;
+FAutoConsoleVariableRef CVarCollisionHistoryLifespan(TEXT("p.CollisionHistoryLifespan"), CollisionHistoryLifespanCVar, TEXT("Number of iterations to cache a collision point.[def:5]"));
 
 extern int32 UseLevelsetCollision;
 
@@ -94,6 +99,7 @@ TPBDCollisionConstraint<T, d>::TPBDCollisionConstraint(const TPBDRigidsSOAs<T,d>
 	, PostComputeCallback(nullptr)
 	, PostApplyCallback(nullptr)
 	, PostApplyPushOutCallback(nullptr)
+	, bSupportHistory(bSupportCollisionHistoryCVar)
 {
 }
 
@@ -104,7 +110,7 @@ void TPBDCollisionConstraint<T, d>::Reset(/*const TPBDRigidParticles<T, d>& InPa
 {
 	SCOPE_CYCLE_COUNTER(STAT_CollisionConstraintsReset);
 
-	int32 Threshold = LifespanCounter - 1; // Maybe this should be solver time?
+	int32 Threshold = LifespanCounter - CollisionHistoryLifespanCVar; // Maybe this should be solver time?
 	for (int32 Idx = Constraints.Num() - 1; Idx >= 0; Idx--)
 	{
 		if ((Constraints[Idx].Lifespan < Threshold) || !bEnableCollisions)
@@ -120,32 +126,40 @@ void TPBDCollisionConstraint<T, d>::Reset(/*const TPBDRigidParticles<T, d>& InPa
 template<typename T, int d>
 void TPBDCollisionConstraint<T, d>::RemoveConstraint(int32 Idx)
 {
-	HandleAllocator.FreeHandle(Handles.FindAndRemoveChecked(GetConstraintHandleID(Idx)));
+	if (bSupportHistory)
+	{
+		FConstraintHandleID HandleID = GetConstraintHandleID(Constraints[Idx]);
+		History[HandleID]->RemoveHandle(Handles[Idx]);
+		if (!History[HandleID]->GetHandles().Num())
+		{
+			History.Remove(HandleID);
+		}
+	}
+
+	HandleAllocator.FreeHandle(Handles[Idx]);
+	Handles.RemoveAtSwap(Idx);
 	Constraints.RemoveAtSwap(Idx);
 	if (Idx < Constraints.Num())
 	{
-		Handles[GetConstraintHandleID(Idx)]->SetConstraintIndex(Idx);
+		Handles[Idx]->SetConstraintIndex(Idx);
 	}
+	ensure(Handles.Num() == Constraints.Num());
 }
 
-
 template<typename T, int d>
 void TPBDCollisionConstraint<T, d>::RemoveConstraints(const TSet<TGeometryParticleHandle<T, d>*>&  InHandleSet)
 {
 	const TArray<TGeometryParticleHandle<T, d>*> HandleArray = InHandleSet.Array();
 	for (auto ParticleHandle : HandleArray)
 	{
-		TArray<FConstraintHandleID> Keys;
-		Handles.GetKeys(Keys);
-		for (FConstraintHandleID Key : Keys)
+		for (int32 Idx = Constraints.Num() - 1; Idx >= 0; Idx--)
 		{
-			if (Key.Key == ParticleHandle || Key.Value == ParticleHandle)
+			if (Constraints[Idx].Levelset == ParticleHandle || Constraints[Idx].Particle == ParticleHandle)
 			{
-				RemoveConstraint(Handles[Key]->GetConstraintIndex());
+				RemoveConstraint(Idx);
 			}
 		}
 	}
-	Handles.Compact();
 }
 
 template<typename T, int d>
@@ -518,6 +532,8 @@ TVector<T, d> GetEnergyClampedImpulse(const TRigidBodyContactConstraint<T, d>& C
 template<typename T, int d>
 void TPBDCollisionConstraint<T, d>::Apply(const T Dt, FRigidBodyContactConstraint& Constraint)
 {
+	if (Constraint.bDisabled) return;
+
 	TGenericParticleHandle<T, d> Particle0 = TGenericParticleHandle<T, d>(Constraint.Particle);
 	TGenericParticleHandle<T, d> Particle1 = TGenericParticleHandle<T, d>(Constraint.Levelset);
 	TPBDRigidParticleHandle<T, d>* PBDRigid0 = Particle0->AsDynamic();
@@ -738,6 +754,16 @@ void TPBDCollisionConstraint<T, d>::Apply(const T Dt, const TArray<FConstraintHa
 	}
 }
 
+template<typename T, int d>
+void TPBDCollisionConstraint<T, d>::Disable(const TArray<FConstraintHandle*>& InConstraintHandles)
+{
+	for(int32 Idx = InConstraintHandles.Num()-1; Idx>=0; Idx--) 
+	{
+		check(InConstraintHandles[Idx] != nullptr);
+		Constraints[InConstraintHandles[Idx]->GetConstraintIndex()].bDisabled = true;
+	}
+}
+
 template<typename T, int d>
 void TPBDCollisionConstraint<T, d>::ApplyPushOut(const T Dt, FRigidBodyContactConstraint& Constraint, const TSet<TGeometryParticleHandle<T,d>*>& IsTemporarilyStatic, int32 Iteration, int32 NumIterations, bool &NeedsAnotherIteration)
 {
diff --git a/Engine/Source/Runtime/Experimental/Chaos/Public/Chaos/PBDCollisionConstraint.h b/Engine/Source/Runtime/Experimental/Chaos/Public/Chaos/PBDCollisionConstraint.h
--- a/Engine/Source/Runtime/Experimental/Chaos/Public/Chaos/PBDCollisionConstraint.h
+++ b/Engine/Source/Runtime/Experimental/Chaos/Public/Chaos/PBDCollisionConstraint.h
@@ -71,6 +71,34 @@ protected:
 	using Base::ConstraintContainer;
 };
 
+template<class T, int d>
+class CHAOS_API TPBDCollisionConstraintHistory
+{
+public:
+	using FConstraintHandle = TPBDCollisionConstraintHandle<T, d>;
+
+	TPBDCollisionConstraintHistory(const TVector<T, d>& InLocation, const TRotation<T, d>& InRotation)
+		: Location(InLocation)
+		, Rotation(InRotation)
+	{}
+
+	const TVector<T, d>& GetLocation() { return Location; }
+	const TRotation<T, d>& GetRotation() { return Rotation; }
+
+	void AddHandle(const FConstraintHandle* InHandle) {
+		Handles.Push(InHandle);
+	}
+	void RemoveHandle(const FConstraintHandle* InHandle) {
+		Handles.RemoveSingleSwap(InHandle);
+	}
+	const TArray< const FConstraintHandle*>& GetHandles() { return Handles; }
+
+private:
+	TVector<T, d> Location;
+	TRotation<T, d> Rotation;
+	TArray< const FConstraintHandle*> Handles;
+};
+
 template<typename T, int d>
 using TRigidBodyContactConstraintsPostComputeCallback = TFunction<void()>;
 
@@ -98,6 +126,7 @@ public:
 	using FReal = T;
 	static const int Dimensions = d;
 	using FConstraintHandle = TPBDCollisionConstraintHandle<T, d>;
+	using FConstraintHistory = TPBDCollisionConstraintHistory<T, d>;
 	using FConstraintHandleAllocator = TConstraintHandleAllocator<TPBDCollisionConstraint<T, d>>;
 	using FRigidBodyContactConstraint = TRigidBodyContactConstraint<T, d>;
 	typedef TPair<TGeometryParticleHandle<T, d>*, const TGeometryParticleHandle<T, d>*> FConstraintHandleID;
@@ -149,12 +178,12 @@ public:
 		
 	const FConstraintHandle* GetConstraintHandle(int32 ConstraintIndex) const
 	{
-		return Handles[GetConstraintHandleID(ConstraintIndex)];
+		return Handles[ConstraintIndex];
 	}
 
 	FConstraintHandle* GetConstraintHandle(int32 ConstraintIndex)
 	{
-		return Handles[ GetConstraintHandleID(ConstraintIndex) ];
+		return Handles[ ConstraintIndex ];
 	}
 
 	void RemoveConstraint(int32 Idx);
@@ -211,6 +240,11 @@ public:
 	 // @todo(ccaulfield): this runs wide. The serial/parallel decision should be in the ConstraintRule
 	void Apply(const T Dt, const TArray<FConstraintHandle*>& InConstraintHandles, const int32 It, const int32 NumIts);
 
+	/**
+	 * Set disabled attribute on all input handles
+	 */
+	void Disable(const TArray<FConstraintHandle*>& InConstraintHandles);
+		
 	/**
 	 * Generate all contact constraints.
 	 */
@@ -288,8 +322,11 @@ private:
 	TRigidBodyContactConstraintsPostApplyCallback<T, d> PostApplyCallback;
 	TRigidBodyContactConstraintsPostApplyPushOutCallback<T, d> PostApplyPushOutCallback;
 
-	TMap<FConstraintHandleID, FConstraintHandle*> Handles;
+	TArray<FConstraintHandle*> Handles;
 	FConstraintHandleAllocator HandleAllocator;
+
+	bool bSupportHistory;
+	TMap<FConstraintHandleID, FConstraintHistory*> History;
 };
 
 extern template void TPBDCollisionConstraint<float, 3>::UpdateConstraint<ECollisionUpdateType::Any>(const float Thickness, FRigidBodyContactConstraint& Constraint);
diff --git a/Engine/Source/Runtime/Experimental/Chaos/Public/Chaos/PBDCollisionConstraintImp.h b/Engine/Source/Runtime/Experimental/Chaos/Public/Chaos/PBDCollisionConstraintImp.h
--- a/Engine/Source/Runtime/Experimental/Chaos/Public/Chaos/PBDCollisionConstraintImp.h
+++ b/Engine/Source/Runtime/Experimental/Chaos/Public/Chaos/PBDCollisionConstraintImp.h
@@ -294,24 +294,28 @@ namespace Chaos
 
 				while (!Queue.IsEmpty())
 				{
-					const TRigidBodyContactConstraint<T, d> * Constraint = Queue.Peek();
-					FConstraintHandleID HandleID = GetConstraintHandleID(*Constraint);
-					if (Handles.Contains(HandleID))
-					{
-						FConstraintHandle* Handle = Handles[HandleID];
-						int32 Idx = Handle->GetConstraintIndex();
-						FVector Position = Constraints[Idx].Location;
+					int32 Idx = Constraints.AddUninitialized(1);
+					FConstraintHandle * Handle = HandleAllocator.AllocHandle(this, Idx);
 
-						Queue.Dequeue(Constraints[Idx]);
-						Constraints[Idx].PreviousLocation = Position;
-						Constraints[Idx].Lifespan = LifespanCounter;
-					}
-					else
+					Handles.Add(Handle);
+					Queue.Dequeue(Constraints[Idx]);
+					Constraints[Idx].Lifespan = LifespanCounter;
+					ensure(Handles.Num() == Constraints.Num());
+
+					if (bSupportHistory)
 					{
-						int32 Idx = Constraints.AddUninitialized(1);
-						Queue.Dequeue(Constraints[Idx]);
-						Handles.Add(GetConstraintHandleID(Idx), HandleAllocator.AllocHandle(this, Idx));
-						Constraints[Idx].Lifespan = LifespanCounter;
+						FConstraintHandleID HandleID = GetConstraintHandleID(Constraints[Idx]);
+						FConstraintHistory* HistoryElement;
+						if (auto ValPtr = History.Find(HandleID))
+						{
+							HistoryElement = *ValPtr;
+						}
+						else
+						{
+							HistoryElement = new FConstraintHistory(Constraints[Idx].Particle->X(), Constraints[Idx].Particle->R());
+							History.Add(HandleID, HistoryElement);
+						}
+						HistoryElement->AddHandle(Handle);
 					}
 				}
 				LifespanCounter++;
diff --git a/Engine/Source/Runtime/Experimental/Chaos/Public/Chaos/PBDCollisionTypes.h b/Engine/Source/Runtime/Experimental/Chaos/Public/Chaos/PBDCollisionTypes.h
--- a/Engine/Source/Runtime/Experimental/Chaos/Public/Chaos/PBDCollisionTypes.h
+++ b/Engine/Source/Runtime/Experimental/Chaos/Public/Chaos/PBDCollisionTypes.h
@@ -13,7 +13,7 @@ namespace Chaos
 template<class T, int d>
 struct TRigidBodyContactConstraint
 {
-	TRigidBodyContactConstraint() : AccumulatedImpulse(0.f), PreviousLocation( FLT_MAX ), Lifespan(0) {}
+	TRigidBodyContactConstraint() : AccumulatedImpulse(0.f), bDisabled(false), Lifespan(0) {}
 	TGeometryParticleHandle<T, d>* Particle;
 	TGeometryParticleHandle<T, d>* Levelset;
 	TVector<T, d> Normal;
@@ -21,7 +21,7 @@ struct TRigidBodyContactConstraint
 	T Phi;
 	TVector<T, d> AccumulatedImpulse;
 
-	TVector<T, d> PreviousLocation;
+	bool bDisabled;
 	int32 Lifespan;
 
 	FString ToString() const
diff --git a/Engine/Source/Runtime/Experimental/Chaos/Public/Chaos/PBDConstraintRule.h b/Engine/Source/Runtime/Experimental/Chaos/Public/Chaos/PBDConstraintRule.h
--- a/Engine/Source/Runtime/Experimental/Chaos/Public/Chaos/PBDConstraintRule.h
+++ b/Engine/Source/Runtime/Experimental/Chaos/Public/Chaos/PBDConstraintRule.h
@@ -241,6 +241,18 @@ namespace Chaos
 					}
 				}
 			}
+
+			for (int32 Level = 0; Level <= MaxLevel; ++Level)
+			{
+				for (int32 Color = 0; Color <= MaxColor; ++Color)
+				{
+					if (LevelToColorToConstraintListMap[Level].Contains(Color) && LevelToColorToConstraintListMap[Level][Color].Num())
+					{
+						const TArray<typename FConstraints::FConstraintHandle*>& ConstraintHandles = GetLevelColorConstraints(LevelToColorToConstraintListMap, Level, Color);
+						Constraints.Disable( ConstraintHandles );
+					}
+				}
+			}
 		}
 
 		virtual void RemoveConstraints(const TSet<TGeometryParticleHandle<T, d>*>& InConstraints)
