Triangle and convex fixes to properly handle thickness and handle it safely in GJK

diff --git a/Engine/Source/Runtime/Experimental/Chaos/Public/Chaos/Convex.h b/Engine/Source/Runtime/Experimental/Chaos/Public/Chaos/Convex.h
--- a/Engine/Source/Runtime/Experimental/Chaos/Public/Chaos/Convex.h
+++ b/Engine/Source/Runtime/Experimental/Chaos/Public/Chaos/Convex.h
@@ -176,7 +176,7 @@ namespace Chaos
 
 			if (Thickness)
 			{
-				return SurfaceParticles.X(MaxVIdx) + SurfaceParticles.X(MaxVIdx).GetSafeNormal()*Thickness;
+				return SurfaceParticles.X(MaxVIdx) + Direction.GetUnsafeNormal() * Thickness;
 			}
 			return SurfaceParticles.X(MaxVIdx);
 		}
diff --git a/Engine/Source/Runtime/Experimental/Chaos/Public/Chaos/Triangle.h b/Engine/Source/Runtime/Experimental/Chaos/Public/Chaos/Triangle.h
--- a/Engine/Source/Runtime/Experimental/Chaos/Public/Chaos/Triangle.h
+++ b/Engine/Source/Runtime/Experimental/Chaos/Public/Chaos/Triangle.h
@@ -57,7 +57,7 @@ namespace Chaos
 			{
 				if(Thickness != 0)
 				{
-					return A + A.GetSafeNormal() * Thickness;
+					return A + Direction.GetUnsafeNormal() * Thickness;
 				}
 				return A;
 			}
@@ -65,14 +65,14 @@ namespace Chaos
 			{
 				if(Thickness != 0)
 				{
-					return B + B.GetSafeNormal() * Thickness;
+					return B + Direction.GetUnsafeNormal() * Thickness;
 				}
 				return B;
 			}
 
 			if(Thickness != 0)
 			{
-				return C + C.GetSafeNormal() * Thickness;
+				return C + Direction.GetUnsafeNormal() * Thickness;
 			}
 			return C;
 		}
