Fix pending delete operation getting removed.

diff --git a/Engine/Source/Runtime/Experimental/Chaos/Private/Chaos/PBDRigidsEvolution.cpp b/Engine/Source/Runtime/Experimental/Chaos/Private/Chaos/PBDRigidsEvolution.cpp
--- a/Engine/Source/Runtime/Experimental/Chaos/Private/Chaos/PBDRigidsEvolution.cpp
+++ b/Engine/Source/Runtime/Experimental/Chaos/Private/Chaos/PBDRigidsEvolution.cpp
@@ -242,7 +242,7 @@ namespace Chaos
 	}
 
 	template<class FPBDRigidsEvolution, class FPBDCollisionConstraint, class T, int d>
-	void TPBDRigidsEvolutionBase<FPBDRigidsEvolution, FPBDCollisionConstraint, T, d>::ApplyParticlePendingData(TGeometryParticleHandle<T, d>* Particle, const FPendingSpatialData& SpatialData, FAccelerationStructure& AccelerationStructure, bool bAsync)
+	void TPBDRigidsEvolutionBase<FPBDRigidsEvolution, FPBDCollisionConstraint, T, d>::ApplyParticlePendingData(TGeometryParticleHandle<T, d>* Particle, const FPendingSpatialData& SpatialData, FAccelerationStructure& AccelerationStructure, bool bUpdateCache)
 	{
 		//Note: we collapsed several update delete events into one struct. If memory is reused this can lead to problems
 		//Luckily there are only 3 states we care about:
@@ -255,7 +255,7 @@ namespace Chaos
 		{
 			AccelerationStructure.RemoveElementFrom(SpatialData.AccelerationHandle, SpatialData.DeletedSpatialIdx);
 
-			if (bAsync)
+			if (bUpdateCache)
 			{
 				if (uint32* InnerIdxPtr = ParticleToCacheInnerIdx.Find(Particle))
 				{
@@ -278,7 +278,7 @@ namespace Chaos
 		{
 			AccelerationStructure.UpdateElementIn(Particle, Particle->WorldSpaceInflatedBounds(), Particle->HasBounds(), SpatialData.UpdatedSpatialIdx);
 			
-			if (bAsync)
+			if (bUpdateCache)
 			{
 				TSpatialAccelerationCache<T, d>& Cache = SpatialAccelerationCache.FindOrAdd(SpatialData.UpdatedSpatialIdx);
 
@@ -318,8 +318,8 @@ namespace Chaos
 	{
 		for (auto Itr : AsyncAccelerationQueue)
 		{
-			ApplyParticlePendingData(Itr.Key, Itr.Value, *AsyncInternalAcceleration, true);
-			ApplyParticlePendingData(Itr.Key, Itr.Value, *AsyncExternalAcceleration, true);
+			ApplyParticlePendingData(Itr.Key, Itr.Value, *AsyncInternalAcceleration, true); //only the first queue needs to update the cached acceleration
+			ApplyParticlePendingData(Itr.Key, Itr.Value, *AsyncExternalAcceleration, false);
 		}
 		AsyncAccelerationQueue.Empty();
 
diff --git a/Engine/Source/Runtime/Experimental/Chaos/Public/Chaos/PBDRigidsEvolution.h b/Engine/Source/Runtime/Experimental/Chaos/Public/Chaos/PBDRigidsEvolution.h
--- a/Engine/Source/Runtime/Experimental/Chaos/Public/Chaos/PBDRigidsEvolution.h
+++ b/Engine/Source/Runtime/Experimental/Chaos/Public/Chaos/PBDRigidsEvolution.h
@@ -309,7 +309,11 @@ class TPBDRigidsEvolutionBase
 		SpatialData.bUpdate = true;
 		SpatialData.UpdatedSpatialIdx = Particle.SpatialIdx();
 
-		AsyncAccelerationQueue.FindOrAdd(Particle.Handle()) = SpatialData;
+		auto& AsyncSpatialData = AsyncAccelerationQueue.FindOrAdd(Particle.Handle());
+		AsyncSpatialData.AccelerationHandle = TAccelerationStructureHandle<T, d>(Particle);
+		AsyncSpatialData.bUpdate = true;
+		AsyncSpatialData.UpdatedSpatialIdx = Particle.SpatialIdx();
+		//question: is it safe to reuse for external? Should probably avoid it
 		ExternalAccelerationQueue.FindOrAdd(Particle.Handle()) = SpatialData;
 	}
 
@@ -635,7 +639,7 @@ protected:
 
 	TMap<FSpatialAccelerationIdx, TSpatialAccelerationCache<T, d>> SpatialAccelerationCache;
 
-	FORCEINLINE_DEBUGGABLE void ApplyParticlePendingData(TGeometryParticleHandle<T, d>* Particle, const FPendingSpatialData& PendingData, FAccelerationStructure& SpatialAcceleration, bool bAsync);
+	FORCEINLINE_DEBUGGABLE void ApplyParticlePendingData(TGeometryParticleHandle<T, d>* Particle, const FPendingSpatialData& PendingData, FAccelerationStructure& SpatialAcceleration, bool bUpdateCache);
 
 	class FChaosAccelerationStructureTask
 	{
