Add TUniquePtr to AccelerationStructureCache to prevent dangling pointers after map realloc.

diff --git a/Engine/Source/Runtime/Experimental/Chaos/Private/Chaos/PBDRigidsEvolution.cpp b/Engine/Source/Runtime/Experimental/Chaos/Private/Chaos/PBDRigidsEvolution.cpp
--- a/Engine/Source/Runtime/Experimental/Chaos/Private/Chaos/PBDRigidsEvolution.cpp
+++ b/Engine/Source/Runtime/Experimental/Chaos/Private/Chaos/PBDRigidsEvolution.cpp
@@ -152,7 +152,7 @@ namespace Chaos
 	template<class FPBDRigidsEvolution, class FPBDCollisionConstraint, class T, int d>
 	TPBDRigidsEvolutionBase<FPBDRigidsEvolution, FPBDCollisionConstraint, T, d>::FChaosAccelerationStructureTask::FChaosAccelerationStructureTask(
 		ISpatialAccelerationCollectionFactory<T, d>& InSpatialCollectionFactory
-		, const TMap<FSpatialAccelerationIdx, TSpatialAccelerationCache<T,d>>& InSpatialAccelerationCache
+		, const TMap<FSpatialAccelerationIdx, TUniquePtr<TSpatialAccelerationCache<T,d>>>& InSpatialAccelerationCache
 		, TUniquePtr<FAccelerationStructure>& InAccelerationStructure
 		, TUniquePtr<FAccelerationStructure>& InAccelerationStructureCopy)
 		: SpatialCollectionFactory(InSpatialCollectionFactory)
@@ -217,7 +217,7 @@ namespace Chaos
 		for (const auto& Itr : SpatialAccelerationCache)
 		{
 			const FSpatialAccelerationIdx SpatialIdx = Itr.Key;
-			const TSpatialAccelerationCache<T, d>& Cache = Itr.Value;
+			const TSpatialAccelerationCache<T, d>& Cache = *Itr.Value;
 			const uint8 BucketIdx = (1 << SpatialIdx.Bucket) & ActiveBucketsMask ? SpatialIdx.Bucket : 0;
 			ViewsPerBucket[BucketIdx].Add(const_cast<TSpatialAccelerationCache<T, d>*>(&Cache));
 
@@ -260,7 +260,7 @@ namespace Chaos
 				if (uint32* InnerIdxPtr = ParticleToCacheInnerIdx.Find(Particle))
 				{
 					const auto SpatialIdx = SpatialData.DeletedSpatialIdx;
-					TSpatialAccelerationCache<T, d>& Cache = SpatialAccelerationCache.FindChecked(SpatialIdx);	//can't delete from cache that doesn't exist
+					TSpatialAccelerationCache<T, d>& Cache = *SpatialAccelerationCache.FindChecked(SpatialIdx);	//can't delete from cache that doesn't exist
 					const uint32 CacheInnerIdx = *InnerIdxPtr;
 					if (CacheInnerIdx + 1 < Cache.Size())	//will get swapped with last element, so update it
 					{
@@ -280,7 +280,14 @@ namespace Chaos
 			
 			if (bUpdateCache)
 			{
-				TSpatialAccelerationCache<T, d>& Cache = SpatialAccelerationCache.FindOrAdd(SpatialData.UpdatedSpatialIdx);
+				TUniquePtr<TSpatialAccelerationCache<T, d>>* CachePtrPtr = SpatialAccelerationCache.Find(SpatialData.UpdatedSpatialIdx);
+				if (CachePtrPtr == nullptr)
+				{
+					CachePtrPtr = &SpatialAccelerationCache[SpatialData.UpdatedSpatialIdx];
+					CachePtrPtr->Reset(new TSpatialAccelerationCache<T, d>());
+				}
+
+				TSpatialAccelerationCache<T, d>& Cache = **CachePtrPtr;
 
 				//make sure in mapping
 				uint32 CacheInnerIdx;
diff --git a/Engine/Source/Runtime/Experimental/Chaos/Public/Chaos/PBDRigidsEvolution.h b/Engine/Source/Runtime/Experimental/Chaos/Public/Chaos/PBDRigidsEvolution.h
--- a/Engine/Source/Runtime/Experimental/Chaos/Public/Chaos/PBDRigidsEvolution.h
+++ b/Engine/Source/Runtime/Experimental/Chaos/Public/Chaos/PBDRigidsEvolution.h
@@ -637,7 +637,7 @@ protected:
 	/** Used for async acceleration rebuild */
 	TMap<TGeometryParticleHandle<T, d>*, uint32> ParticleToCacheInnerIdx;
 
-	TMap<FSpatialAccelerationIdx, TSpatialAccelerationCache<T, d>> SpatialAccelerationCache;
+	TMap<FSpatialAccelerationIdx, TUniquePtr<TSpatialAccelerationCache<T, d>>> SpatialAccelerationCache;
 
 	FORCEINLINE_DEBUGGABLE void ApplyParticlePendingData(TGeometryParticleHandle<T, d>* Particle, const FPendingSpatialData& PendingData, FAccelerationStructure& SpatialAcceleration, bool bUpdateCache);
 
@@ -645,7 +645,7 @@ protected:
 	{
 	public:
 		FChaosAccelerationStructureTask(ISpatialAccelerationCollectionFactory<T,d>& InSpatialCollectionFactory
-			, const TMap<FSpatialAccelerationIdx, TSpatialAccelerationCache<T,d>>& InSpatialAccelerationCache
+			, const TMap<FSpatialAccelerationIdx, TUniquePtr<TSpatialAccelerationCache<T,d>>>& InSpatialAccelerationCache
 			, TUniquePtr<FAccelerationStructure>& InAccelerationStructure
 			, TUniquePtr<FAccelerationStructure>& InAccelerationStructureCopy);
 		static FORCEINLINE TStatId GetStatId();
@@ -654,7 +654,7 @@ protected:
 		void DoTask(ENamedThreads::Type CurrentThread, const FGraphEventRef& MyCompletionGraphEvent);
 
 		ISpatialAccelerationCollectionFactory<T, d>& SpatialCollectionFactory;
-		const TMap<FSpatialAccelerationIdx, TSpatialAccelerationCache<T, d>>& SpatialAccelerationCache;
+		const TMap<FSpatialAccelerationIdx, TUniquePtr<TSpatialAccelerationCache<T, d>>>& SpatialAccelerationCache;
 		TUniquePtr<FAccelerationStructure>& AccelerationStructure;
 		TUniquePtr<FAccelerationStructure>& AccelerationStructureCopy;
 	};
