Fix TParticles resizing ArrayCollection incorrectly during move assignment. Fix SpatialAccelerationCache not resizing array collection during move.

diff --git a/Engine/Source/Runtime/Experimental/Chaos/Public/Chaos/PBDRigidsEvolution.h b/Engine/Source/Runtime/Experimental/Chaos/Public/Chaos/PBDRigidsEvolution.h
--- a/Engine/Source/Runtime/Experimental/Chaos/Public/Chaos/PBDRigidsEvolution.h
+++ b/Engine/Source/Runtime/Experimental/Chaos/Public/Chaos/PBDRigidsEvolution.h
@@ -80,6 +80,7 @@ public:
 		, MBounds(MoveTemp(Other.MBounds))
 		, MPayloads(MoveTemp(Other.MPayloads))
 	{
+		ResizeHelper(Other.MSize);
 		Other.MSize = 0;
 
 		AddArray(&MHasBoundingBoxes);
@@ -97,6 +98,7 @@ public:
 			MHasBoundingBoxes = MoveTemp(Other.MHasBoundingBoxes);
 			MBounds = MoveTemp(Other.MBounds);
 			MPayloads = MoveTemp(Other.MPayloads);
+			ResizeHelper(Other.MSize);
 			Other.MSize = 0;
 #if PARTICLE_ITERATOR_RANGED_FOR_CHECK
 			MDirtyValidationCount = 0;
diff --git a/Engine/Source/Runtime/Experimental/Chaos/Public/Chaos/Particles.h b/Engine/Source/Runtime/Experimental/Chaos/Public/Chaos/Particles.h
--- a/Engine/Source/Runtime/Experimental/Chaos/Public/Chaos/Particles.h
+++ b/Engine/Source/Runtime/Experimental/Chaos/Public/Chaos/Particles.h
@@ -100,7 +100,7 @@ namespace Chaos
 		TParticles& operator=(TParticles<T, d>&& Other)
 		{
 			MX = MoveTemp(Other.MX);
-			AddParticles(Other.Size());
+			ResizeHelper(Other.Size());
 			Other.MSize = 0;
 
 #if PARTICLE_ITERATOR_RANGED_FOR_CHECK
