Add the option to rebuild acceleration structure for perf testing

diff --git a/Engine/Source/Runtime/Experimental/Chaos/Private/Chaos/PBDRigidsEvolution.cpp b/Engine/Source/Runtime/Experimental/Chaos/Private/Chaos/PBDRigidsEvolution.cpp
--- a/Engine/Source/Runtime/Experimental/Chaos/Private/Chaos/PBDRigidsEvolution.cpp
+++ b/Engine/Source/Runtime/Experimental/Chaos/Private/Chaos/PBDRigidsEvolution.cpp
@@ -425,6 +425,26 @@ namespace Chaos
 		ComputeIntermediateSpatialAcceleration(true);
 	}
 
+	template<class FPBDRigidsEvolution, class FPBDCollisionConstraint, class T, int d>
+	void TPBDRigidsEvolutionBase<FPBDRigidsEvolution, FPBDCollisionConstraint, T, d>::RebuildSpatialAccelerationForPerfTest()
+	{
+		WaitOnAccelerationStructure();
+
+		ParticleToCacheInnerIdx.Empty();
+		AsyncAccelerationQueue.Empty();
+		InternalAccelerationQueue.Empty();
+		ExternalAccelerationQueue.Empty();
+
+		AccelerationStructureTaskComplete = nullptr;
+		const auto& NonDisabled = Particles.GetNonDisabledView();
+		for (auto& Particle : NonDisabled)
+		{
+			DirtyParticle(Particle);
+		}
+
+		FlushSpatialAcceleration();
+	}
+
 	template<class FPBDRigidsEvolution, class FPBDCollisionConstraint, class T, int d>
 	void TPBDRigidsEvolutionBase<FPBDRigidsEvolution, FPBDCollisionConstraint, T, d>::Serialize(FChaosArchive& Ar)
 	{
diff --git a/Engine/Source/Runtime/Experimental/Chaos/Public/Chaos/PBDRigidsEvolution.h b/Engine/Source/Runtime/Experimental/Chaos/Public/Chaos/PBDRigidsEvolution.h
--- a/Engine/Source/Runtime/Experimental/Chaos/Public/Chaos/PBDRigidsEvolution.h
+++ b/Engine/Source/Runtime/Experimental/Chaos/Public/Chaos/PBDRigidsEvolution.h
@@ -479,6 +479,9 @@ class TPBDRigidsEvolutionBase
 	/** Perform a blocking flush of the spatial acceleration structure for situations where we aren't simulating but must have an up to date structure */
 	CHAOS_API void FlushSpatialAcceleration();
 
+	/** Rebuilds the spatial acceleration from scratch. This should only be used for perf testing */
+	CHAOS_API void RebuildSpatialAccelerationForPerfTest();
+
 	const auto& GetRigidClustering() const { return Clustering; }
 	auto& GetRigidClustering() { return Clustering; }
 
