Fix Linux as well

diff --git a/Engine/Source/Runtime/Experimental/Chaos/Public/Chaos/BoundingVolumeHierarchy.h b/Engine/Source/Runtime/Experimental/Chaos/Public/Chaos/BoundingVolumeHierarchy.h
--- a/Engine/Source/Runtime/Experimental/Chaos/Public/Chaos/BoundingVolumeHierarchy.h
+++ b/Engine/Source/Runtime/Experimental/Chaos/Public/Chaos/BoundingVolumeHierarchy.h
@@ -170,7 +170,7 @@ class CHAOS_API TBoundingVolumeHierarchy final : public ISpatialAcceleration<int
 	FCriticalSection CriticalSection;
 };
 
-#if PLATFORM_MAC
+#if PLATFORM_MAC || PLATFORM_LINUX
 extern template class CHAOS_API Chaos::TBoundingVolumeHierarchy<Chaos::TParticles<float, 3>, TArray<int32>, float, 3>;
 extern template class CHAOS_API Chaos::TBoundingVolumeHierarchy<Chaos::TGeometryParticles<float, 3>, TArray<int32>, float, 3>;
 #endif
diff --git a/Engine/Source/Runtime/Experimental/Chaos/Public/Chaos/KinematicGeometryParticles.h b/Engine/Source/Runtime/Experimental/Chaos/Public/Chaos/KinematicGeometryParticles.h
--- a/Engine/Source/Runtime/Experimental/Chaos/Public/Chaos/KinematicGeometryParticles.h
+++ b/Engine/Source/Runtime/Experimental/Chaos/Public/Chaos/KinematicGeometryParticles.h
@@ -128,7 +128,7 @@ FChaosArchive& operator<<(FChaosArchive& Ar, TKinematicGeometryParticlesImp<T, d
 	return Ar;
 }
 
-#if PLATFORM_MAC
+#if PLATFORM_MAC || PLATFORM_LINUX
 extern template class CHAOS_API Chaos::TKinematicGeometryParticlesImp<float, 3, Chaos::EGeometryParticlesSimType::RigidBodySim>;
 extern template class CHAOS_API Chaos::TKinematicGeometryParticlesImp<float, 3, Chaos::EGeometryParticlesSimType::Other>;
 #else
