Fix crash in spatial acceleration cache.

diff --git a/Engine/Source/Runtime/Experimental/Chaos/Private/Chaos/PBDRigidsEvolution.cpp b/Engine/Source/Runtime/Experimental/Chaos/Private/Chaos/PBDRigidsEvolution.cpp
--- a/Engine/Source/Runtime/Experimental/Chaos/Private/Chaos/PBDRigidsEvolution.cpp
+++ b/Engine/Source/Runtime/Experimental/Chaos/Private/Chaos/PBDRigidsEvolution.cpp
@@ -283,8 +283,7 @@ namespace Chaos
 				TUniquePtr<TSpatialAccelerationCache<T, d>>* CachePtrPtr = SpatialAccelerationCache.Find(SpatialData.UpdatedSpatialIdx);
 				if (CachePtrPtr == nullptr)
 				{
-					CachePtrPtr = &SpatialAccelerationCache[SpatialData.UpdatedSpatialIdx];
-					CachePtrPtr->Reset(new TSpatialAccelerationCache<T, d>());
+					CachePtrPtr = &SpatialAccelerationCache.Add(SpatialData.UpdatedSpatialIdx, TUniquePtr<TSpatialAccelerationCache<T, d>>(new TSpatialAccelerationCache<T, d>));
 				}
 
 				TSpatialAccelerationCache<T, d>& Cache = **CachePtrPtr;
