Properly serialize out broadphase type so we can easily swap between what was recorded and what we're testing

diff --git a/Engine/Source/Runtime/Core/Public/UObject/ExternalPhysicsCustomObjectVersion.h b/Engine/Source/Runtime/Core/Public/UObject/ExternalPhysicsCustomObjectVersion.h
--- a/Engine/Source/Runtime/Core/Public/UObject/ExternalPhysicsCustomObjectVersion.h
+++ b/Engine/Source/Runtime/Core/Public/UObject/ExternalPhysicsCustomObjectVersion.h
@@ -42,6 +42,9 @@ struct CORE_API FExternalPhysicsCustomObjectVersion
 		// Allow trimeshes to serialize their acceleration structure
 		TrimeshSerializesBV,
 
+		// Serialize broadphase type
+		SerializeBroadphaseType,
+
 		// -----<new versions can be added above this line>-------------------------------------------------
 		VersionPlusOne,
 		LatestVersion = VersionPlusOne - 1
diff --git a/Engine/Source/Runtime/Experimental/Chaos/Private/Chaos/PBDRigidsEvolution.cpp b/Engine/Source/Runtime/Experimental/Chaos/Private/Chaos/PBDRigidsEvolution.cpp
--- a/Engine/Source/Runtime/Experimental/Chaos/Private/Chaos/PBDRigidsEvolution.cpp
+++ b/Engine/Source/Runtime/Experimental/Chaos/Private/Chaos/PBDRigidsEvolution.cpp
@@ -52,7 +52,7 @@ namespace Chaos
 		{
 			TConstParticleView<TSpatialAccelerationCache<T, d>> Empty;
 
-			const uint16 NumBuckets = ConfigSettings.BroadphaseType == 3 ? 2 : 1;
+			const uint16 NumBuckets = ConfigSettings.BroadphaseType >= 3 ? 2 : 1;
 			auto Collection = new TSpatialAccelerationCollection<AABBTreeType, BVType, AABBTreeOfGridsType>();
 
 			for (uint16 BucketIdx = 0; BucketIdx < NumBuckets; ++BucketIdx)
@@ -65,7 +65,7 @@ namespace Chaos
 
 		virtual uint8 GetActiveBucketsMask() const
 		{
-			return ConfigSettings.BroadphaseType == 3 ? 3 : 1;
+			return ConfigSettings.BroadphaseType >= 3 ? 3 : 1;
 		}
 
 		virtual TUniquePtr<ISpatialAcceleration<TAccelerationStructureHandle<T, d>, T, d>> CreateAccelerationPerBucket_Threaded(const TConstParticleView<TSpatialAccelerationCache<T, d>>& Particles, uint16 BucketIdx) override
@@ -82,14 +82,14 @@ namespace Chaos
 				{
 					return MakeUnique<AABBTreeType>(Particles, ConfigSettings.MaxChildrenInLeaf, ConfigSettings.MaxTreeDepth, ConfigSettings.MaxPayloadSize);
 				}
-				else
+				else if(ConfigSettings.BroadphaseType == 4 || ConfigSettings.BroadphaseType == 2)
 				{
 					return MakeUnique<AABBTreeOfGridsType>(Particles, ConfigSettings.AABBMaxChildrenInLeaf, ConfigSettings.AABBMaxTreeDepth, ConfigSettings.MaxPayloadSize);
 				}
 			}
 			case 1:
 			{
-				ensure(ConfigSettings.BroadphaseType == 3);
+				ensure(ConfigSettings.BroadphaseType == 3 || ConfigSettings.BroadphaseType == 4);
 				return MakeUnique<BVType>(Particles, false, 0, ConfigSettings.BVNumCells, ConfigSettings.MaxPayloadSize);
 			}
 			default:
@@ -447,9 +447,21 @@ namespace Chaos
 	template<class FPBDRigidsEvolution, class FPBDCollisionConstraint, class T, int d>
 	void TPBDRigidsEvolutionBase<FPBDRigidsEvolution, FPBDCollisionConstraint, T, d>::Serialize(FChaosArchive& Ar)
 	{
-		Particles.Serialize(Ar);
+		int32 DefaultBroadphaseType = ConfigSettings.BroadphaseType;
 
 		Ar.UsingCustomVersion(FExternalPhysicsCustomObjectVersion::GUID);
+		if (Ar.CustomVer(FExternalPhysicsCustomObjectVersion::GUID) >= FExternalPhysicsCustomObjectVersion::SerializeBroadphaseType)
+		{
+			Ar << ConfigSettings.BroadphaseType;
+		}
+		else
+		{
+			//older archives just assume type 3
+			ConfigSettings.BroadphaseType = 3;
+		}
+
+		Particles.Serialize(Ar);
+
 		if (Ar.CustomVer(FExternalPhysicsCustomObjectVersion::GUID) >= FExternalPhysicsCustomObjectVersion::SerializeEvolutionBV)
 		{
 			if (Ar.CustomVer(FExternalPhysicsCustomObjectVersion::GUID) < FExternalPhysicsCustomObjectVersion::SerializeMultiStructures)
@@ -490,6 +502,9 @@ namespace Chaos
 
 			FlushSpatialAcceleration();
 		}
+
+		ConfigSettings.BroadphaseType = DefaultBroadphaseType;
+
 	}
 }
 
