Fix macOS build.

#rnx

diff --git a/Engine/Source/Runtime/Experimental/Chaos/Public/Chaos/GeometryParticles.h b/Engine/Source/Runtime/Experimental/Chaos/Public/Chaos/GeometryParticles.h
--- a/Engine/Source/Runtime/Experimental/Chaos/Public/Chaos/GeometryParticles.h
+++ b/Engine/Source/Runtime/Experimental/Chaos/Public/Chaos/GeometryParticles.h
@@ -373,6 +373,12 @@ namespace Chaos
 	template<>
 	TGeometryParticlesImp<float, 3, EGeometryParticlesSimType::Other>* TGeometryParticlesImp<float, 3, EGeometryParticlesSimType::Other>::SerializationFactory(FChaosArchive& Ar, TGeometryParticlesImp<float, 3, EGeometryParticlesSimType::Other>* Particles);
 
+#if PLATFORM_MAC || PLATFORM_LINUX
+	extern template class CHAOS_API TGeometryParticlesImp<float, 3, EGeometryParticlesSimType::RigidBodySim>;
+	extern template class CHAOS_API TGeometryParticlesImp<float, 3, EGeometryParticlesSimType::Other>;
+#else
 	extern template class TGeometryParticlesImp<float, 3, EGeometryParticlesSimType::RigidBodySim>;
 	extern template class TGeometryParticlesImp<float, 3, EGeometryParticlesSimType::Other>;
+#endif
+
 }
