Restore ensure in FindGeometryOpposingNormal

diff --git a/Engine/Source/Runtime/Experimental/Chaos/Public/Chaos/Convex.h b/Engine/Source/Runtime/Experimental/Chaos/Public/Chaos/Convex.h
--- a/Engine/Source/Runtime/Experimental/Chaos/Public/Chaos/Convex.h
+++ b/Engine/Source/Runtime/Experimental/Chaos/Public/Chaos/Convex.h
@@ -133,28 +133,17 @@ namespace Chaos
 			return MostOpposingIdx;
 		}
 
-		TVector<T, d> FindGeometryOpposingNormal(const TVector<T, d>& DenormDir, int32 HintFaceIndex, const TVector<T, d>& OriginalNormal) const
+		TVector<T, d> FindGeometryOpposingNormal(const TVector<T, d>& DenormDir, int32 FaceIndex, const TVector<T, d>& OriginalNormal) const
 		{
-			if (HintFaceIndex != INDEX_NONE)
+			// For convexes, this function must be called with a face index.
+			// If this ensure is getting hit, fix the caller so that it
+			// passes in a valid face index.
+			if (ensure(FaceIndex != INDEX_NONE))
 			{
-				const TPlane<float, 3>& OpposingFace = GetFaces()[HintFaceIndex];
+				const TPlane<float, 3>& OpposingFace = GetFaces()[FaceIndex];
 				return OpposingFace.Normal();
 			}
-		
-			// todo: make a way to call FindMostOpposingFace without a search dist
-			int32 MostOpposingIdx = INDEX_NONE;
-			T MostOpposingDot = TNumericLimits<T>::Max();
-			for (int32 Idx = 0; Idx < Planes.Num(); ++Idx)
-			{
-				const T Dot = TVector<T, d>::DotProduct(Planes[Idx].Normal(), DenormDir);
-				if (Dot < MostOpposingDot)
-				{
-					MostOpposingDot = Dot;
-					MostOpposingIdx = Idx;
-				}
-			}
-			ensure(MostOpposingIdx != INDEX_NONE);
-			return Planes[MostOpposingIdx].Normal();
+			return TVector<float, 3>(0.f, 0.f, 1.f);
 		}
 
 		virtual TVector<T, d> Support(const TVector<T, d>& Direction, const T Thickness) const override
