Chaos 6dof joint
- adding projection (WIP)
- cleanup

diff --git a/Engine/Source/Runtime/AnimGraphRuntime/Private/BoneControllers/AnimNode_RigidBody.cpp b/Engine/Source/Runtime/AnimGraphRuntime/Private/BoneControllers/AnimNode_RigidBody.cpp
--- a/Engine/Source/Runtime/AnimGraphRuntime/Private/BoneControllers/AnimNode_RigidBody.cpp
+++ b/Engine/Source/Runtime/AnimGraphRuntime/Private/BoneControllers/AnimNode_RigidBody.cpp
@@ -27,7 +27,7 @@ TAutoConsoleVariable<int32> CVarEnableRigidBodyNode(TEXT("p.RigidBodyNode"), 1,
 TAutoConsoleVariable<int32> CVarRigidBodyLODThreshold(TEXT("p.RigidBodyLODThreshold"), -1, TEXT("Max LOD that rigid body node is allowed to run on. Provides a global threshold that overrides per-node the LODThreshold property. -1 means no override."), ECVF_Scalability);
 
 int32 RBAN_MaxSubSteps = 4;
-FAutoConsoleVariableRef CVarRigidBodyNodeMaxSteps(TEXT("p.RigidBodyNode.MaxSubSteps"), RBAN_MaxSubSteps, TEXT("Set the maximum number of simulation steps in the update loop"), ECVF_ReadOnly);
+FAutoConsoleVariableRef CVarRigidBodyNodeMaxSteps(TEXT("p.RigidBodyNode.MaxSubSteps"), RBAN_MaxSubSteps, TEXT("Set the maximum number of simulation steps in the update loop"), ECVF_Default);
 
 FAnimNode_RigidBody::FAnimNode_RigidBody():
 	QueryParams(NAME_None, FCollisionQueryParams::GetUnknownStatId())
diff --git a/Engine/Source/Runtime/Engine/Private/PhysicsEngine/ImmediatePhysics/ImmediatePhysicsChaos/ImmediatePhysicsSimulation_Chaos.cpp b/Engine/Source/Runtime/Engine/Private/PhysicsEngine/ImmediatePhysics/ImmediatePhysicsChaos/ImmediatePhysicsSimulation_Chaos.cpp
--- a/Engine/Source/Runtime/Engine/Private/PhysicsEngine/ImmediatePhysics/ImmediatePhysicsChaos/ImmediatePhysicsSimulation_Chaos.cpp
+++ b/Engine/Source/Runtime/Engine/Private/PhysicsEngine/ImmediatePhysics/ImmediatePhysicsChaos/ImmediatePhysicsSimulation_Chaos.cpp
@@ -17,43 +17,54 @@
 #include "PhysicsEngine/BodyInstance.h"
 #include "PhysicsEngine/ConstraintInstance.h"
 
+//////////////////////////////////////////////////////////////////////////
 // @todo(ccaulfield): remove when finished
+//
 float ChaosImmediate_Evolution_DeltaTime = 0.03f;
 int32 ChaosImmediate_Evolution_Iterations = 10;
+FAutoConsoleVariableRef CVarChaosImmPhysDeltaTime(TEXT("p.Chaos.ImmPhys.DeltaTime"), ChaosImmediate_Evolution_DeltaTime, TEXT("Override chaos immediate physics delta time if non-zero"));
+FAutoConsoleVariableRef CVarChaosImmPhysIterations(TEXT("p.Chaos.ImmPhys.Iterations"), ChaosImmediate_Evolution_Iterations, TEXT("Number of constraint solver loops in immediate physics"));
+
 int32 ChaosImmediate_Collision_Enabled = 1;
-int32 ChaosImmediate_Collision_ApplyEnabled = 1;
+int32 ChaosImmediate_Collision_ApplyEnabled = 0;
 int32 ChaosImmediate_Collision_PushOutIterations = 5;
 int32 ChaosImmediate_Collision_PushOutPairIterations = 2;
 float ChaosImmediate_Collision_Thickness = 0;
-
-FAutoConsoleVariableRef CVarChaosImmPhysDeltaTime(TEXT("p.Chaos.ImmPhys.DeltaTime"), ChaosImmediate_Evolution_DeltaTime, TEXT("Override chaos immediate physics delta time if non-zero"));
-FAutoConsoleVariableRef CVarChaosImmPhysIterations(TEXT("p.Chaos.ImmPhys.Iterations"), ChaosImmediate_Evolution_Iterations, TEXT("Number of constraint solver loops in immediate physics"));
-FAutoConsoleVariableRef CVarImmediatePhysicsDisableCollisions(TEXT("p.Chaos.ImmPhys.EnableCollisions"), ChaosImmediate_Collision_Enabled, TEXT("Enable/Disable collisions in Immediate Physics."));
-FAutoConsoleVariableRef CVarCollisionDisableApply(TEXT("p.Chaos.ImmPhys.CollisionEnableApply"), ChaosImmediate_Collision_ApplyEnabled, TEXT("Enable/Disable the Apply() (velocity, friction, restitution) method for collisions"));
-FAutoConsoleVariableRef CVarCollisionPushOutIterations(TEXT("p.Chaos.ImmPhys.CollisionPushOutIterations"), ChaosImmediate_Collision_PushOutIterations, TEXT("Set the ApplyPushOut() (position correction) iteration count (0 to disable)"));
-FAutoConsoleVariableRef CVarCollisionPushOutPairIterations(TEXT("p.Chaos.ImmPhys.CollisionPushOutPairIterations"), ChaosImmediate_Collision_PushOutPairIterations, TEXT("Set the ApplyPushOut() internal pair interations (position correction) iteration count (0 to disable)"));
+FAutoConsoleVariableRef CVarChaosImmPhysDisableCollisions(TEXT("p.Chaos.ImmPhys.EnableCollisions"), ChaosImmediate_Collision_Enabled, TEXT("Enable/Disable collisions in Immediate Physics."));
+FAutoConsoleVariableRef CVarChaosImmPhysCollisionDisableApply(TEXT("p.Chaos.ImmPhys.CollisionEnableApply"), ChaosImmediate_Collision_ApplyEnabled, TEXT("Enable/Disable the Apply() (velocity, friction, restitution) method for collisions"));
+FAutoConsoleVariableRef CVarChaosImmPhysCollisionPushOutIterations(TEXT("p.Chaos.ImmPhys.CollisionPushOutIterations"), ChaosImmediate_Collision_PushOutIterations, TEXT("Set the ApplyPushOut() (position correction) iteration count (0 to disable)"));
+FAutoConsoleVariableRef CVarChaosImmPhysCollisionPushOutPairIterations(TEXT("p.Chaos.ImmPhys.CollisionPushOutPairIterations"), ChaosImmediate_Collision_PushOutPairIterations, TEXT("Set the ApplyPushOut() internal pair interations (position correction) iteration count (0 to disable)"));
 FAutoConsoleVariableRef CVarChaosImmPhysThickness(TEXT("p.Chaos.ImmPhys.CollisionThickness"), ChaosImmediate_Collision_Thickness, TEXT("ChaosImmediateThickness"));
 
 float ChaosImmediate_Joint_SwingTwistAngleTolerance = 1.0e-6f;
 int32 ChaosImmediate_Joint_EnableLinearLimits = 1;
 int32 ChaosImmediate_Joint_EnableTwistLimits = 1;
 int32 ChaosImmediate_Joint_EnableSwingLimits = 1;
+int32 ChaosImmediate_Joint_EnablePositionCorrection = 0;
 int32 ChaosImmediate_Joint_EnableDrives = 1;
-int32 ChaosImmediate_Joint_FreezeIterations = 0;
-int32 ChaosImmediate_Joint_FrozenIterations = 1;
-float ChaosImmediate_Joint_PBDDriveStiffness = 0.0f;
-float ChaosImmediate_Joint_PBDMinParentMassRatio = 0.5f;
-float ChaosImmediate_Joint_PBDMaxInertiaRatio = 5.0f;
-FAutoConsoleVariableRef CVarSwingTwistAngleTolerance(TEXT("p.Chaos.ImmPhys.Joint.SwingTwistAngleTolerance"), ChaosImmediate_Joint_SwingTwistAngleTolerance, TEXT("SwingTwistAngleTolerance."));
-FAutoConsoleVariableRef CVarEnableLinearLimits(TEXT("p.Chaos.ImmPhys.Joint.EnableLinearLimits"), ChaosImmediate_Joint_EnableLinearLimits, TEXT("EnableLinearLimits."));
-FAutoConsoleVariableRef CVarEnableTwistLimits(TEXT("p.Chaos.ImmPhys.Joint.EnableTwistLimits"), ChaosImmediate_Joint_EnableTwistLimits, TEXT("EnableTwistLimits."));
-FAutoConsoleVariableRef CVarEnableSwingLimits(TEXT("p.Chaos.ImmPhys.Joint.EnableSwingLimits"), ChaosImmediate_Joint_EnableSwingLimits, TEXT("EnableSwingLimits."));
-FAutoConsoleVariableRef CVarEnableDrives(TEXT("p.Chaos.ImmPhys.Joint.EnableDrives"), ChaosImmediate_Joint_EnableDrives, TEXT("EnableDrives."));
-FAutoConsoleVariableRef CVarFreezeIterations(TEXT("p.Chaos.ImmPhys.Joint.FreezeIterations"), ChaosImmediate_Joint_FreezeIterations, TEXT("FreezeIterations."));
-FAutoConsoleVariableRef CVarFrozenIterations(TEXT("p.Chaos.ImmPhys.Joint.FrozenIterations"), ChaosImmediate_Joint_FrozenIterations, TEXT("FrozenIterations."));
-FAutoConsoleVariableRef CVarDriveStiffness(TEXT("p.Chaos.ImmPhys.Joint.PBDDriveStiffness"), ChaosImmediate_Joint_PBDDriveStiffness, TEXT("6Dof joint drive stiffness override (if > 0)."));
-FAutoConsoleVariableRef CVarMinParentMassRatio(TEXT("p.Chaos.ImmPhys.Joint.MinParentMassRatio"), ChaosImmediate_Joint_PBDMinParentMassRatio, TEXT("6Dof joint PBDMinParentMassRatio (if > 0)"));
-FAutoConsoleVariableRef CVarMaxInertiaRatio(TEXT("p.Chaos.ImmPhys.Joint.MaxInertiaRatio"), ChaosImmediate_Joint_PBDMaxInertiaRatio, TEXT("6Dof joint PBDMaxInertiaRatio (if > 0)"));
+int32 ChaosImmediate_Joint_ProjectPostApply = 1;
+float ChaosImmediate_Joint_Projection = 0.0f;
+float ChaosImmediate_Joint_Stiffness = 1.0f;
+float ChaosImmediate_Joint_DriveStiffness = 0.0f;
+float ChaosImmediate_Joint_MinParentMassRatio = 0.5f;
+float ChaosImmediate_Joint_MaxInertiaRatio = 5.0f;
+FAutoConsoleVariableRef CVarChaosImmPhysSwingTwistAngleTolerance(TEXT("p.Chaos.ImmPhys.Joint.SwingTwistAngleTolerance"), ChaosImmediate_Joint_SwingTwistAngleTolerance, TEXT("SwingTwistAngleTolerance."));
+FAutoConsoleVariableRef CVarChaosImmPhysEnableLinearLimits(TEXT("p.Chaos.ImmPhys.Joint.EnableLinearLimits"), ChaosImmediate_Joint_EnableLinearLimits, TEXT("EnableLinearLimits."));
+FAutoConsoleVariableRef CVarChaosImmPhysEnableTwistLimits(TEXT("p.Chaos.ImmPhys.Joint.EnableTwistLimits"), ChaosImmediate_Joint_EnableTwistLimits, TEXT("EnableTwistLimits."));
+FAutoConsoleVariableRef CVarChaosImmPhysEnableSwingLimits(TEXT("p.Chaos.ImmPhys.Joint.EnableSwingLimits"), ChaosImmediate_Joint_EnableSwingLimits, TEXT("EnableSwingLimits."));
+FAutoConsoleVariableRef CVarChaosImmPhysEnablePositionCorrection(TEXT("p.Chaos.ImmPhys.Joint.EnablePositionCorrection"), ChaosImmediate_Joint_EnablePositionCorrection, TEXT("EnablePositionCorrection."));
+FAutoConsoleVariableRef CVarChaosImmPhysEnableDrives(TEXT("p.Chaos.ImmPhys.Joint.EnableDrives"), ChaosImmediate_Joint_EnableDrives, TEXT("EnableDrives."));
+FAutoConsoleVariableRef CVarChaosImmPhysProjectPostApply(TEXT("p.Chaos.ImmPhys.Joint.ProjectPostApply"), ChaosImmediate_Joint_ProjectPostApply, TEXT("6Dof joint projection in ApplyPushOut(true) or Apply(false)."));
+FAutoConsoleVariableRef CVarChaosImmPhysProjection(TEXT("p.Chaos.ImmPhys.Joint.Projection"), ChaosImmediate_Joint_Projection, TEXT("6Dof joint projection amount override (if > 0)."));
+FAutoConsoleVariableRef CVarChaosImmPhysStiffness(TEXT("p.Chaos.ImmPhys.Joint.Stiffness"), ChaosImmediate_Joint_Stiffness, TEXT("6Dof joint stiffness override (if > 0)."));
+FAutoConsoleVariableRef CVarChaosImmPhysDriveStiffness(TEXT("p.Chaos.ImmPhys.Joint.DriveStiffness"), ChaosImmediate_Joint_DriveStiffness, TEXT("6Dof joint drive stiffness override (if > 0)."));
+FAutoConsoleVariableRef CVarChaosImmPhysMinParentMassRatio(TEXT("p.Chaos.ImmPhys.Joint.MinParentMassRatio"), ChaosImmediate_Joint_MinParentMassRatio, TEXT("6Dof joint MinParentMassRatio (if > 0)"));
+FAutoConsoleVariableRef CVarChaosImmPhysMaxInertiaRatio(TEXT("p.Chaos.ImmPhys.Joint.MaxInertiaRatio"), ChaosImmediate_Joint_MaxInertiaRatio, TEXT("6Dof joint MaxInertiaRatio (if > 0)"));
+//
+// end remove when finished
+//
+//////////////////////////////////////////////////////////////////////////
+
 
 // DebugDraw CVars
 #if UE_BUILD_DEBUG
@@ -68,10 +79,10 @@ int32 ChaosImmediate_DebugDrawCollisions = 0;
 int32 ChaosImmediate_DebugDrawJoints = 0;
 #endif
 
-FAutoConsoleVariableRef CVarDebugDrawParticles(TEXT("p.Chaos.ImmPhys.DebugDrawParticles"), ChaosImmediate_DebugDrawParticles, TEXT("Draw Particle Transforms (0 = never; 1 = end of frame; 2 = begin and end of frame; 3 = post-integate, post-apply and post-applypushout;)."));
-FAutoConsoleVariableRef CVarDebugDrawShapes(TEXT("p.Chaos.ImmPhys.DebugDrawShapes"), ChaosImmediate_DebugDrawShapes, TEXT("Draw Shapes (0 = never; 1 = end of frame; 2 = begin and end of frame; 3 = post-integate, post-apply and post-applypushout;"));
-FAutoConsoleVariableRef CVarDebugDrawCollisions(TEXT("p.Chaos.ImmPhys.DebugDrawCollisions"), ChaosImmediate_DebugDrawCollisions, TEXT("Draw Collisions (0 = never; 1 = end of frame; 2 = begin and end of frame; 3 = post-integate, post-apply and post-applypushout;)"));
-FAutoConsoleVariableRef CVarDebugDrawJoints(TEXT("p.Chaos.ImmPhys.DebugDrawJoints"), ChaosImmediate_DebugDrawJoints, TEXT("Draw Joints. (0 = never; 1 = end of frame; 2 = begin and end of frame; 3 = post-integate, post-apply and post-applypushout; 4 = each Apply step)."));
+FAutoConsoleVariableRef CVarChaosImmPhysDebugDrawParticles(TEXT("p.Chaos.ImmPhys.DebugDrawParticles"), ChaosImmediate_DebugDrawParticles, TEXT("Draw Particle Transforms (0 = never; 1 = end of frame; 2 = begin and end of frame; 3 = post-integate, post-apply and post-applypushout;)."));
+FAutoConsoleVariableRef CVarChaosImmPhysDebugDrawShapes(TEXT("p.Chaos.ImmPhys.DebugDrawShapes"), ChaosImmediate_DebugDrawShapes, TEXT("Draw Shapes (0 = never; 1 = end of frame; 2 = begin and end of frame; 3 = post-integate, post-apply and post-applypushout;"));
+FAutoConsoleVariableRef CVarChaosImmPhysDebugDrawCollisions(TEXT("p.Chaos.ImmPhys.DebugDrawCollisions"), ChaosImmediate_DebugDrawCollisions, TEXT("Draw Collisions (0 = never; 1 = end of frame; 2 = begin and end of frame; 3 = post-integate, post-apply and post-applypushout;)"));
+FAutoConsoleVariableRef CVarChaosImmPhysDebugDrawJoints(TEXT("p.Chaos.ImmPhys.DebugDrawJoints"), ChaosImmediate_DebugDrawJoints, TEXT("Draw Joints. (0 = never; 1 = end of frame; 2 = begin and end of frame; 3 = post-integate, post-apply and post-applypushout; 4 = each Apply step)."));
 
 namespace ImmediatePhysics_Chaos
 {
@@ -389,18 +400,19 @@ namespace ImmediatePhysics_Chaos
 			}
 			UE_LOG(LogChaosJoint, Verbose, TEXT("Simulate dt = %f"), DeltaTime);
 
-			// Overrides
 			TPBDJointSolverSettings<float, 3> JointsSettings = Joints->GetSettings();
 			JointsSettings.SwingTwistAngleTolerance = ChaosImmediate_Joint_SwingTwistAngleTolerance;
-			JointsSettings.PBDMinParentMassRatio = ChaosImmediate_Joint_PBDMinParentMassRatio;
-			JointsSettings.PBDMaxInertiaRatio = ChaosImmediate_Joint_PBDMaxInertiaRatio;
-			JointsSettings.FreezeIterations = ChaosImmediate_Joint_FreezeIterations;
-			JointsSettings.FrozenIterations = ChaosImmediate_Joint_FrozenIterations;
+			JointsSettings.MinParentMassRatio = ChaosImmediate_Joint_MinParentMassRatio;
+			JointsSettings.MaxInertiaRatio = ChaosImmediate_Joint_MaxInertiaRatio;
 			JointsSettings.bEnableLinearLimits = ChaosImmediate_Joint_EnableLinearLimits != 0;
 			JointsSettings.bEnableTwistLimits = ChaosImmediate_Joint_EnableTwistLimits != 0;
 			JointsSettings.bEnableSwingLimits = ChaosImmediate_Joint_EnableSwingLimits != 0;
+			JointsSettings.bEnablePositionCorrection = ChaosImmediate_Joint_EnablePositionCorrection != 0;
 			JointsSettings.bEnableDrives = ChaosImmediate_Joint_EnableDrives != 0;
-			JointsSettings.PBDDriveStiffness = ChaosImmediate_Joint_PBDDriveStiffness;
+			JointsSettings.bProjectPostApply = ChaosImmediate_Joint_ProjectPostApply != 0;
+			JointsSettings.Projection = ChaosImmediate_Joint_Projection;
+			JointsSettings.Stiffness = ChaosImmediate_Joint_Stiffness;
+			JointsSettings.DriveStiffness = ChaosImmediate_Joint_DriveStiffness;
 			Joints->SetSettings(JointsSettings);
 
 			Evolution->SetNumIterations(ChaosImmediate_Evolution_Iterations);
diff --git a/Engine/Source/Runtime/Experimental/Chaos/Private/Chaos/PBDJointConstraintUtilities.cpp b/Engine/Source/Runtime/Experimental/Chaos/Private/Chaos/PBDJointConstraintUtilities.cpp
--- a/Engine/Source/Runtime/Experimental/Chaos/Private/Chaos/PBDJointConstraintUtilities.cpp
+++ b/Engine/Source/Runtime/Experimental/Chaos/Private/Chaos/PBDJointConstraintUtilities.cpp
@@ -3,12 +3,13 @@
 #include "Chaos/ParticleHandle.h"
 #include "Chaos/Utilities.h"
 
+//#pragma optimize("", off)
+
 namespace Chaos
 {
 	template<typename T, int d>
 	TVector<T, d> TPBDJointUtilities<T, d>::ConditionInertia(const TVector<T, d>& InI, const T MaxRatio)
 	{
-		// @todo(ccaulfield): simd
 		if (MaxRatio > 0)
 		{
 			T IMin = InI.Min();
@@ -29,7 +30,6 @@ namespace Chaos
 	template<typename T, int d>
 	TVector<T, d> TPBDJointUtilities<T, d>::ConditionParentInertia(const TVector<T, d>& IParent, const TVector<T, d>& IChild, const T MinRatio)
 	{
-		// @todo(ccaulfield): simd
 		if (MinRatio > 0)
 		{
 			T IParentMax = IParent.Max();
@@ -344,6 +344,78 @@ namespace Chaos
 		OutCR[(int32)EJointAngularAxisIndex::Swing2] = (T)0;
 	}
 
+	template<typename T, int d>
+	void ApplyPositionDelta(
+		const T Dt,
+		const TPBDJointSolverSettings<T, d>& SolverSettings,
+		const TPBDJointSettings<T, d>& JointSettings,
+		const int32 Index0,
+		const int32 Index1,
+		TVector<T, d>& P0,
+		TVector<T, d>& P1,
+		const TVector<T, d>& V0,
+		const TVector<T, d>& V1)
+	{
+		const float Stiffness = (SolverSettings.Stiffness > (T)0) ? SolverSettings.Stiffness : JointSettings.Motion.Stiffness;
+
+		P0 = P0 + Stiffness * V0;
+		P1 = P1 + Stiffness * V1;
+	}
+
+	template<typename T, int d>
+	void ApplyRotationDelta(
+		const T Dt,
+		const TPBDJointSolverSettings<T, d>& SolverSettings,
+		const TPBDJointSettings<T, d>& JointSettings,
+		const int32 Index0,
+		const int32 Index1,
+		TRotation<T, d>& Q0,
+		TRotation<T, d>& Q1,
+		const TVector<T, d>& W0,
+		const TVector<T, d>& W1)
+	{
+		const float Stiffness = (SolverSettings.Stiffness > (T)0) ? SolverSettings.Stiffness : JointSettings.Motion.Stiffness;
+
+		const TRotation<T, d> DQ0 = (TRotation<T, d>::FromElements(Stiffness * W0, 0) * Q0) * (T)0.5;
+		const TRotation<T, d> DQ1 = (TRotation<T, d>::FromElements(Stiffness * W1, 0) * Q1) * (T)0.5;
+		Q0 = (Q0 + DQ0).GetNormalized();
+		Q1 = (Q1 + DQ1).GetNormalized();
+		Q1.EnforceShortestArcWith(Q0);
+	}
+
+	template<typename T, int d>
+	void ApplyRotationDelta(
+	    const T Dt,
+	    const TPBDJointSolverSettings<T, d>& SolverSettings,
+	    const TPBDJointSettings<T, d>& JointSettings,
+	    const int32 Index0,
+	    const int32 Index1,
+	    TVector<T, d>& P0,
+	    TRotation<T, d>& Q0,
+	    TVector<T, d>& P1,
+	    TRotation<T, d>& Q1,
+	    float InvM0,
+	    const PMatrix<T, d, d>& InvIL0,
+	    float InvM1,
+	    const PMatrix<T, d, d>& InvIL1,
+	    const TVector<T, d>& Axis0,
+	    const float Angle0,
+	    const TVector<T, d>& Axis1,
+	    const float Angle1)
+	{
+		const PMatrix<T, d, d> InvI0 = Utilities::Multiply(Q0.ToMatrix(), Utilities::Multiply(InvIL0, Q0.ToMatrix().GetTransposed()));
+		const PMatrix<T, d, d> InvI1 = Utilities::Multiply(Q1.ToMatrix(), Utilities::Multiply(InvIL1, Q1.ToMatrix().GetTransposed()));
+
+		//const T I0 = TVector<T, d>::DotProduct(Axis0, Utilities::Multiply(InvI0, Axis0));
+		//const T I1 = TVector<T, d>::DotProduct(Axis1, Utilities::Multiply(InvI1, Axis1));
+		//const TVector<T, d> W0 = Axis0 * Angle0 * I0 / (I0 + I1);
+		//const TVector<T, d> W1 = Axis1 * Angle1 * I1 / (I0 + I1);
+		const T L = (T)1 / (TVector<T, d>::DotProduct(Axis0, Utilities::Multiply(InvI0, Axis0)) + TVector<T, d>::DotProduct(Axis1, Utilities::Multiply(InvI1, Axis1)));
+		const TVector<T, d> W0 = Utilities::Multiply(InvI0, Axis0) * L * Angle0;
+		const TVector<T, d> W1 = Utilities::Multiply(InvI1, Axis1) * L * Angle1;
+
+		ApplyRotationDelta(Dt, SolverSettings, JointSettings, Index0, Index1, Q0, Q1, W0, W1);
+	}
 
 	template<typename T, int d>
 	void TPBDJointUtilities<T, d>::ApplyJointPositionConstraint(
@@ -366,11 +438,11 @@ namespace Chaos
 		const TVector<T, d> X0 = P0 + Q0 * XL0.GetTranslation();
 		const TVector<T, d> X1 = P1 + Q1 * XL1.GetTranslation();
 		const TRotation<T, d> R0 = Q0 * XL0.GetRotation();
-		PMatrix<T, d, d> InvI0 = Utilities::Multiply(Q0.ToMatrix(), Utilities::Multiply(InvIL0, Q0.ToMatrix().GetTransposed()));
-		PMatrix<T, d, d> InvI1 = Utilities::Multiply(Q1.ToMatrix(), Utilities::Multiply(InvIL1, Q1.ToMatrix().GetTransposed()));
+		const PMatrix<T, d, d> InvI0 = Utilities::Multiply(Q0.ToMatrix(), Utilities::Multiply(InvIL0, Q0.ToMatrix().GetTransposed()));
+		const PMatrix<T, d, d> InvI1 = Utilities::Multiply(Q1.ToMatrix(), Utilities::Multiply(InvIL1, Q1.ToMatrix().GetTransposed()));
 
 		// Calculate constraint error
-		TVector<T, d> CX = GetLimitedPositionError(JointSettings, R0, X1 - X0);
+		const TVector<T, d> CX = GetLimitedPositionError(JointSettings, R0, X1 - X0);
 
 		// Calculate constraint correction
 		PMatrix<T, d, d> M0 = PMatrix<T, d, d>(0, 0, 0);
@@ -383,21 +455,74 @@ namespace Chaos
 		{
 			M1 = Utilities::ComputeJointFactorMatrix(X1 - P1, InvI1, InvM1);
 		}
-		PMatrix<T, d, d> MI = (M0 + M1).Inverse();
-		TVector<T, d> DX = Utilities::Multiply(MI, CX);
+		const PMatrix<T, d, d> MI = (M0 + M1).Inverse();
+		const TVector<T, d> DX = Utilities::Multiply(MI, CX);
 
 		// Apply constraint correction
-		TVector<T, d> DP0 = InvM0 * DX;
-		TVector<T, d> DP1 = -InvM1 * DX;
-		TVector<T, d> DR0 = Utilities::Multiply(InvI0, TVector<T, d>::CrossProduct(X0 - P0, DX));
-		TVector<T, d> DR1 = Utilities::Multiply(InvI1, TVector<T, d>::CrossProduct(X1 - P1, -DX));
-		TRotation<T, d> DQ0 = (TRotation<T, d>::FromElements(DR0, 0) * Q0) * (T)0.5;
-		TRotation<T, d> DQ1 = (TRotation<T, d>::FromElements(DR1, 0) * Q1) * (T)0.5;
-		P0 = P0 + DP0;
-		P1 = P1 + DP1;
-		Q0 = (Q0 + DQ0).GetNormalized();
-		Q1 = (Q1 + DQ1).GetNormalized();
-		Q1.EnforceShortestArcWith(Q0);
+		const TVector<T, d> V0 = InvM0 * DX;
+		const TVector<T, d> V1 = -InvM1 * DX;
+		const TVector<T, d> W0 = Utilities::Multiply(InvI0, TVector<T, d>::CrossProduct(X0 - P0, DX));
+		const TVector<T, d> W1 = Utilities::Multiply(InvI1, TVector<T, d>::CrossProduct(X1 - P1, -DX));
+		ApplyPositionDelta(Dt, SolverSettings, JointSettings, Index0, Index1, P0, P1, V0, V1);
+		ApplyRotationDelta(Dt, SolverSettings, JointSettings, Index0, Index1, Q0, Q1, W0, W1);
+	}
+
+	template<typename T, int d>
+	void ApplyPostRotationPositionCorrection(
+		const T Dt,
+		const TPBDJointSolverSettings<T, d>& SolverSettings,
+		const TPBDJointSettings<T, d>& JointSettings,
+		const int32 Index0,
+		const int32 Index1,
+		TVector<T, d>& P0,
+		TRotation<T, d>& Q0,
+		TVector<T, d>& P1,
+		TRotation<T, d>& Q1,
+		float InvM0,
+		const PMatrix<T, d, d>& InvIL0,
+		float InvM1,
+		const PMatrix<T, d, d>& InvIL1,
+		const TVector<T, d>& CX_1,
+		const TVector<T, d>& Axis0,
+		const TVector<T, d>& Axis1)
+	{
+		if (!SolverSettings.bEnablePositionCorrection)
+		{
+			return;
+		}
+
+		// Post-rotation constraint positions
+		const TRigidTransform<T, d>& XL0 = JointSettings.ConstraintFrames[Index0];
+		const TRigidTransform<T, d>& XL1 = JointSettings.ConstraintFrames[Index1];
+		const TVector<T, d> X0 = P0 + Q0 * XL0.GetTranslation();
+		const TVector<T, d> X1 = P1 + Q1 * XL1.GetTranslation();
+		const TRotation<T, d> R0 = Q0 * XL0.GetRotation();
+		TVector<T, d> CX_2 = GetLimitedPositionError(JointSettings, R0, X1 - X0);
+
+		// Increase in position error caused by the rotation correction
+		TVector<T, d> DX = CX_2 - CX_1;
+		for (int32 AxisIndex = 0; AxisIndex < d; ++AxisIndex)
+		{
+			if ((CX_1[AxisIndex] > 0) && (CX_2[AxisIndex] < CX_1[AxisIndex]))
+			{
+				DX[AxisIndex] = 0;
+			}
+			if ((CX_1[AxisIndex] < 0) && (CX_2[AxisIndex] > CX_1[AxisIndex]))
+			{
+				DX[AxisIndex] = 0;
+			}
+		}
+		const TVector<T, d> DX0 = DX - TVector<T, d>::DotProduct(DX, Axis0) * Axis0;
+		const TVector<T, d> DX1 = DX - TVector<T, d>::DotProduct(DX, Axis1) * Axis1;
+
+		// Correct the extra position error introduced by the rotation correction. We are effectively treating the
+		// bodies as if they have infinite inertia for this correction, which is only correct if
+		// the position correction exactly opposes the constraint that caused the increased position error
+		// which is only the case for uniform shapes (i.e., spherical inertia). To do this properly, we
+		// would have to solve for position and rotation correction simultaneously.
+		const TVector<T, d> V0 = InvM0 * DX0 / (InvM0 + InvM1);
+		const TVector<T, d> V1 = -InvM1 * DX1 / (InvM0 + InvM1);
+		ApplyPositionDelta(Dt, SolverSettings, JointSettings, Index0, Index1, P0, P1, V0, V1);
 	}
 
 	template<typename T, int d>
@@ -418,12 +543,14 @@ namespace Chaos
 	{
 		const TRigidTransform<T, d>& XL0 = JointSettings.ConstraintFrames[Index0];
 		const TRigidTransform<T, d>& XL1 = JointSettings.ConstraintFrames[Index1];
+		const TVector<T, d> X0 = P0 + Q0 * XL0.GetTranslation();
+		const TVector<T, d> X1 = P1 + Q1 * XL1.GetTranslation();
 		const TRotation<T, d> R0 = Q0 * XL0.GetRotation();
 		const TRotation<T, d> R1 = Q1 * XL1.GetRotation();
-		PMatrix<T, d, d> InvI0 = Utilities::Multiply(Q0.ToMatrix(), Utilities::Multiply(InvIL0, Q0.ToMatrix().GetTransposed()));
-		PMatrix<T, d, d> InvI1 = Utilities::Multiply(Q1.ToMatrix(), Utilities::Multiply(InvIL1, Q1.ToMatrix().GetTransposed()));
+		const TVector<T, d> CX = GetLimitedPositionError(JointSettings, R0, X1 - X0);
 
-		TRotation<T, d> R01 = R0.Inverse() * R1;
+		// Calculate the Twist Axis and Angle for each body
+		const TRotation<T, d> R01 = R0.Inverse() * R1;
 		TRotation<T, d> R01Twist, R01Swing;
 		R01.ToSwingTwist(TJointConstants<T, d>::TwistAxis(), R01Swing, R01Twist);
 		R01Swing = R01Swing.GetNormalized();
@@ -442,8 +569,8 @@ namespace Chaos
 			TwistAngle = -TwistAngle;
 		}
 
-		TVector<T, d> TwistAxis0 = R0 * TwistAxis01;
-		TVector<T, d> TwistAxis1 = R1 * TwistAxis01;
+		const TVector<T, d> TwistAxis0 = R0 * TwistAxis01;
+		const TVector<T, d> TwistAxis1 = R1 * TwistAxis01;
 		T TwistAngleMax = FLT_MAX;
 		if (JointSettings.Motion.AngularMotionTypes[(int32)EJointAngularConstraintIndex::Twist] == EJointMotionType::Limited)
 		{
@@ -454,6 +581,7 @@ namespace Chaos
 			TwistAngleMax = 0;
 		}
 
+		// Calculate the twist correction to applt to each body
 		T DTwistAngle = 0;
 		if (TwistAngle > TwistAngleMax)
 		{
@@ -463,17 +591,14 @@ namespace Chaos
 		{
 			DTwistAngle = TwistAngle + TwistAngleMax;
 		}
-		T DTwistAngle0 = DTwistAngle;
-		T DTwistAngle1 = -DTwistAngle;
+		const T DTwistAngle0 = DTwistAngle;
+		const T DTwistAngle1 = -DTwistAngle;
 
-		T L = (T)1 / (TVector<T, d>::DotProduct(TwistAxis0, Utilities::Multiply(InvI0, TwistAxis0)) + TVector<T, d>::DotProduct(TwistAxis1, Utilities::Multiply(InvI1, TwistAxis1)));
-		TVector<T, d> W0 = Utilities::Multiply(InvI0, TwistAxis0) * L * DTwistAngle0;
-		TVector<T, d> W1 = Utilities::Multiply(InvI1, TwistAxis1) * L * DTwistAngle1;
-		TRotation<T, d> DQ0 = (TRotation<T, d>::FromElements(W0, (T)0.0) * Q0) * (T)0.5;
-		TRotation<T, d> DQ1 = (TRotation<T, d>::FromElements(W1, (T)0.0) * Q1) * (T)0.5;
-		Q0 = (Q0 + DQ0).GetNormalized();
-		Q1 = (Q1 + DQ1).GetNormalized();
-		Q1.EnforceShortestArcWith(Q0);
+		// Apply twist correction
+		ApplyRotationDelta(Dt, SolverSettings, JointSettings, Index0, Index1, P0, Q0, P1, Q1, InvM0, InvIL0, InvM1, InvIL1, TwistAxis0, DTwistAngle0, TwistAxis1, DTwistAngle1);
+
+		// Correct any positional error we may have introduced
+		ApplyPostRotationPositionCorrection<T, d>(Dt, SolverSettings, JointSettings, Index0, Index1, P0, Q0, P1, Q1, InvM0, InvIL0, InvM1, InvIL1, CX, TwistAxis0, TwistAxis1);
 	}
 
 	template<typename T, int d>
@@ -494,13 +619,14 @@ namespace Chaos
 	{
 		const TRigidTransform<T, d>& XL0 = JointSettings.ConstraintFrames[Index0];
 		const TRigidTransform<T, d>& XL1 = JointSettings.ConstraintFrames[Index1];
+		const TVector<T, d> X0 = P0 + Q0 * XL0.GetTranslation();
+		const TVector<T, d> X1 = P1 + Q1 * XL1.GetTranslation();
 		const TRotation<T, d> R0 = Q0 * XL0.GetRotation();
 		const TRotation<T, d> R1 = Q1 * XL1.GetRotation();
-		PMatrix<T, d, d> InvI0 = Utilities::Multiply(Q0.ToMatrix(), Utilities::Multiply(InvIL0, Q0.ToMatrix().GetTransposed()));
-		PMatrix<T, d, d> InvI1 = Utilities::Multiply(Q1.ToMatrix(), Utilities::Multiply(InvIL1, Q1.ToMatrix().GetTransposed()));
+		const TVector<T, d> CX = GetLimitedPositionError(JointSettings, R0, X1 - X0);
 
 		// Calculate Swing axis for each body
-		TRotation<T, d> R01 = R0.Inverse() * R1;
+		const TRotation<T, d> R01 = R0.Inverse() * R1;
 		TRotation<T, d> R01Twist, R01Swing;
 		R01.ToSwingTwist(TJointConstants<T, d>::TwistAxis(), R01Swing, R01Twist);
 		R01Swing = R01Swing.GetNormalized();
@@ -514,13 +640,13 @@ namespace Chaos
 			SwingAngle = SwingAngle - (T)2 * PI;
 		}
 
-		TVector<T, d> SwingAxis0 = R0 * SwingAxis01;
-		TVector<T, d> SwingAxis1 = SwingAxis0;
+		const TVector<T, d> SwingAxis0 = R0 * SwingAxis01;
+		const TVector<T, d> SwingAxis1 = SwingAxis0;
 
 		// Calculate swing limit for the current swing axis
 		T SwingAngleMax = FLT_MAX;
-		T Swing1Limit = JointSettings.Motion.AngularLimits[(int32)EJointAngularConstraintIndex::Swing1];
-		T Swing2Limit = JointSettings.Motion.AngularLimits[(int32)EJointAngularConstraintIndex::Swing2];
+		const T Swing1Limit = JointSettings.Motion.AngularLimits[(int32)EJointAngularConstraintIndex::Swing1];
+		const T Swing2Limit = JointSettings.Motion.AngularLimits[(int32)EJointAngularConstraintIndex::Swing2];
 
 		// Circular swing limit
 		SwingAngleMax = Swing1Limit;
@@ -529,8 +655,8 @@ namespace Chaos
 		if (!FMath::IsNearlyEqual(Swing1Limit, Swing2Limit, KINDA_SMALL_NUMBER))
 		{
 			// Map swing axis to ellipse and calculate limit for this swing axis
-			T DotSwing1 = FMath::Abs(TVector<T, d>::DotProduct(SwingAxis01, TJointConstants<T, d>::Swing1Axis()));
-			T DotSwing2 = FMath::Abs(TVector<T, d>::DotProduct(SwingAxis01, TJointConstants<T, d>::Swing2Axis()));
+			const T DotSwing1 = FMath::Abs(TVector<T, d>::DotProduct(SwingAxis01, TJointConstants<T, d>::Swing1Axis()));
+			const T DotSwing2 = FMath::Abs(TVector<T, d>::DotProduct(SwingAxis01, TJointConstants<T, d>::Swing2Axis()));
 			SwingAngleMax = FMath::Sqrt(Swing1Limit * DotSwing2 * Swing1Limit * DotSwing2 + Swing2Limit * DotSwing1 * Swing2Limit * DotSwing1);
 		}
 
@@ -548,14 +674,10 @@ namespace Chaos
 		T DSwingAngle1 = -DSwingAngle;
 
 		// Apply swing correction
-		T L = (T)1 / (TVector<T, d>::DotProduct(SwingAxis0, Utilities::Multiply(InvI0, SwingAxis0)) + TVector<T, d>::DotProduct(SwingAxis1, Utilities::Multiply(InvI1, SwingAxis1)));
-		TVector<T, d> W0 = Utilities::Multiply(InvI0, SwingAxis0) * L * DSwingAngle0;
-		TVector<T, d> W1 = Utilities::Multiply(InvI1, SwingAxis1) * L * DSwingAngle1;
-		TRotation<T, d> DQ0 = (TRotation<T, d>::FromElements(W0, (T)0.0) * Q0) * (T)0.5;
-		TRotation<T, d> DQ1 = (TRotation<T, d>::FromElements(W1, (T)0.0) * Q1) * (T)0.5;
-		Q0 = (Q0 + DQ0).GetNormalized();
-		Q1 = (Q1 + DQ1).GetNormalized();
-		Q1.EnforceShortestArcWith(Q0);
+		ApplyRotationDelta(Dt, SolverSettings, JointSettings, Index0, Index1, P0, Q0, P1, Q1, InvM0, InvIL0, InvM1, InvIL1, SwingAxis0, DSwingAngle0, SwingAxis1, DSwingAngle1);
+	
+		// Correct any positional error we may have introduced
+		ApplyPostRotationPositionCorrection<T, d>(Dt, SolverSettings, JointSettings, Index0, Index1, P0, Q0, P1, Q1, InvM0, InvIL0, InvM1, InvIL1, CX, SwingAxis0, SwingAxis1);
 	}
 
 	template<typename T, int d>
@@ -577,14 +699,37 @@ namespace Chaos
 	{
 		const TRigidTransform<T, d>& XL0 = JointSettings.ConstraintFrames[Index0];
 		const TRigidTransform<T, d>& XL1 = JointSettings.ConstraintFrames[Index1];
+		const TVector<T, d> X0 = P0 + Q0 * XL0.GetTranslation();
+		const TVector<T, d> X1 = P1 + Q1 * XL1.GetTranslation();
 		const TRotation<T, d> R0 = Q0 * XL0.GetRotation();
 		const TRotation<T, d> R1 = Q1 * XL1.GetRotation();
-		PMatrix<T, d, d> InvI0 = Utilities::Multiply(Q0.ToMatrix(), Utilities::Multiply(InvIL0, Q0.ToMatrix().GetTransposed()));
-		PMatrix<T, d, d> InvI1 = Utilities::Multiply(Q1.ToMatrix(), Utilities::Multiply(InvIL1, Q1.ToMatrix().GetTransposed()));
+		const TVector<T, d> CX = GetLimitedPositionError(JointSettings, R0, X1 - X0);
+
+		// Calculate the swing axis for each body
+		const TRotation<T, d> R01 = R0.Inverse() * R1;
+		TRotation<T, d> R01Twist, R01Swing;
+		R01.ToSwingTwist(TJointConstants<T, d>::TwistAxis(), R01Swing, R01Twist);
+		R01Swing = R01Swing.GetNormalized();
+		R01Twist = R01Twist.GetNormalized();
+
+		TVector<T, d> TwistAxis01;
+		T TwistAngle;
+		R01Twist.ToAxisAndAngleSafe(TwistAxis01, TwistAngle, TJointConstants<T, d>::TwistAxis(), SolverSettings.SwingTwistAngleTolerance);
+		if (TwistAngle > PI)
+		{
+			TwistAngle = TwistAngle - (T)2 * PI;
+		}
+		if (TVector<T, d>::DotProduct(TwistAxis01, TJointConstants<T, d>::TwistAxis()) < 0)
+		{
+			TwistAxis01 = -TwistAxis01;
+			TwistAngle = -TwistAngle;
+		}
+		const TVector<T, d> TwistAxis = R0 * TwistAxis01;
 
 		const PMatrix<T, d, d> Axes0 = R0.ToMatrix();
 		const PMatrix<T, d, d> Axes1 = R1.ToMatrix();
-		const TVector<T, d> SwingCross = TVector<T, d>::CrossProduct(Axes0.GetAxis((int32)SwingConstraint), Axes1.GetAxis((int32)SwingConstraint));
+		TVector<T, d> SwingCross = TVector<T, d>::CrossProduct(Axes0.GetAxis((int32)SwingConstraint), Axes1.GetAxis((int32)SwingConstraint));
+		SwingCross = SwingCross - TVector<T, d>::DotProduct(TwistAxis, SwingCross) * TwistAxis;
 		const T SwingCrossLen = SwingCross.Size();
 		if (SwingCrossLen > KINDA_SMALL_NUMBER)
 		{
@@ -624,14 +769,10 @@ namespace Chaos
 			T DSwingAngle1 = -DSwingAngle;
 
 			// Apply swing correction
-			T L = (T)1 / (TVector<T, d>::DotProduct(SwingAxis0, Utilities::Multiply(InvI0, SwingAxis0)) + TVector<T, d>::DotProduct(SwingAxis1, Utilities::Multiply(InvI1, SwingAxis1)));
-			TVector<T, d> W0 = Utilities::Multiply(InvI0, SwingAxis0) * L * DSwingAngle0;
-			TVector<T, d> W1 = Utilities::Multiply(InvI1, SwingAxis1) * L * DSwingAngle1;
-			TRotation<T, d> DQ0 = (TRotation<T, d>::FromElements(W0, (T)0.0) * Q0) * (T)0.5;
-			TRotation<T, d> DQ1 = (TRotation<T, d>::FromElements(W1, (T)0.0) * Q1) * (T)0.5;
-			Q0 = (Q0 + DQ0).GetNormalized();
-			Q1 = (Q1 + DQ1).GetNormalized();
-			Q1.EnforceShortestArcWith(Q0);
+			ApplyRotationDelta(Dt, SolverSettings, JointSettings, Index0, Index1, P0, Q0, P1, Q1, InvM0, InvIL0, InvM1, InvIL1, SwingAxis0, DSwingAngle0, SwingAxis1, DSwingAngle1);
+
+			// Correct any positional error we may have introduced
+			ApplyPostRotationPositionCorrection<T, d>(Dt, SolverSettings, JointSettings, Index0, Index1, P0, Q0, P1, Q1, InvM0, InvIL0, InvM1, InvIL1, CX, SwingAxis0, SwingAxis1);
 		}
 	}
 
@@ -653,12 +794,13 @@ namespace Chaos
 	{
 		const TRigidTransform<T, d>& XL0 = JointSettings.ConstraintFrames[Index0];
 		const TRigidTransform<T, d>& XL1 = JointSettings.ConstraintFrames[Index1];
+		const TVector<T, d> X0 = P0 + Q0 * XL0.GetTranslation();
+		const TVector<T, d> X1 = P1 + Q1 * XL1.GetTranslation();
 		const TRotation<T, d> R0 = Q0 * XL0.GetRotation();
 		const TRotation<T, d> R1 = Q1 * XL1.GetRotation();
-		PMatrix<T, d, d> InvI0 = Utilities::Multiply(Q0.ToMatrix(), Utilities::Multiply(InvIL0, Q0.ToMatrix().GetTransposed()));
-		PMatrix<T, d, d> InvI1 = Utilities::Multiply(Q1.ToMatrix(), Utilities::Multiply(InvIL1, Q1.ToMatrix().GetTransposed()));
+		const TVector<T, d> CX = GetLimitedPositionError(JointSettings, R0, X1 - X0);
 
-		TRotation<T, d> R01 = R0.Inverse() * R1;
+		const TRotation<T, d> R01 = R0.Inverse() * R1;
 		TRotation<T, d> R01Twist, R01Swing;
 		R01.ToSwingTwist(TJointConstants<T, d>::TwistAxis(), R01Swing, R01Twist);
 		R01Swing = R01Swing.GetNormalized();
@@ -677,23 +819,19 @@ namespace Chaos
 			TwistAngle = -TwistAngle;
 		}
 
-		TVector<T, d> TwistAxis0 = R0 * TwistAxis01;
-		TVector<T, d> TwistAxis1 = R1 * TwistAxis01;
-		T TwistAngleTarget = JointSettings.Motion.AngularDriveTargetAngles[(int32)EJointAngularConstraintIndex::Twist];
-		T Stiffness = (SolverSettings.PBDDriveStiffness > 0) ? SolverSettings.PBDDriveStiffness : JointSettings.Motion.AngularDriveStiffness;
-		T DriveStiffness = FMath::Clamp(Stiffness, (T)0, (T)1);
-		T DTwistAngle = TwistAngle - TwistAngleTarget;
-		T DTwistAngle0 = DriveStiffness * DTwistAngle;
-		T DTwistAngle1 = -DriveStiffness * DTwistAngle;
-
-		T L = (T)1 / (TVector<T, d>::DotProduct(TwistAxis0, Utilities::Multiply(InvI0, TwistAxis0)) + TVector<T, d>::DotProduct(TwistAxis1, Utilities::Multiply(InvI1, TwistAxis1)));
-		TVector<T, d> W0 = Utilities::Multiply(InvI0, TwistAxis0) * L * DTwistAngle0;
-		TVector<T, d> W1 = Utilities::Multiply(InvI1, TwistAxis1) * L * DTwistAngle1;
-		TRotation<T, d> DQ0 = (TRotation<T, d>::FromElements(W0, (T)0.0) * Q0) * (T)0.5;
-		TRotation<T, d> DQ1 = (TRotation<T, d>::FromElements(W1, (T)0.0) * Q1) * (T)0.5;
-		Q0 = (Q0 + DQ0).GetNormalized();
-		Q1 = (Q1 + DQ1).GetNormalized();
-		Q1.EnforceShortestArcWith(Q0);
+		const TVector<T, d> TwistAxis0 = R0 * TwistAxis01;
+		const TVector<T, d> TwistAxis1 = R1 * TwistAxis01;
+		const T TwistAngleTarget = JointSettings.Motion.AngularDriveTargetAngles[(int32)EJointAngularConstraintIndex::Twist];
+		const T DriveStiffnessUnclamped = (SolverSettings.DriveStiffness > 0) ? SolverSettings.DriveStiffness : JointSettings.Motion.AngularDriveStiffness;
+		const T DriveStiffness = FMath::Clamp(DriveStiffnessUnclamped, (T)0, (T)1);
+		const T DTwistAngle = TwistAngle - TwistAngleTarget;
+		const T DTwistAngle0 = DriveStiffness * DTwistAngle;
+		const T DTwistAngle1 = -DriveStiffness * DTwistAngle;
+
+		ApplyRotationDelta(Dt, SolverSettings, JointSettings, Index0, Index1, P0, Q0, P1, Q1, InvM0, InvIL0, InvM1, InvIL1, TwistAxis0, DTwistAngle0, TwistAxis1, DTwistAngle1);
+
+		// Correct any positional error we may have introduced
+		ApplyPostRotationPositionCorrection<T, d>(Dt, SolverSettings, JointSettings, Index0, Index1, P0, Q0, P1, Q1, InvM0, InvIL0, InvM1, InvIL1, CX, TwistAxis0, TwistAxis1);
 	}
 
 	template<typename T, int d>
@@ -714,13 +852,14 @@ namespace Chaos
 	{
 		const TRigidTransform<T, d>& XL0 = JointSettings.ConstraintFrames[Index0];
 		const TRigidTransform<T, d>& XL1 = JointSettings.ConstraintFrames[Index1];
+		const TVector<T, d> X0 = P0 + Q0 * XL0.GetTranslation();
+		const TVector<T, d> X1 = P1 + Q1 * XL1.GetTranslation();
 		const TRotation<T, d> R0 = Q0 * XL0.GetRotation();
 		const TRotation<T, d> R1 = Q1 * XL1.GetRotation();
-		PMatrix<T, d, d> InvI0 = Utilities::Multiply(Q0.ToMatrix(), Utilities::Multiply(InvIL0, Q0.ToMatrix().GetTransposed()));
-		PMatrix<T, d, d> InvI1 = Utilities::Multiply(Q1.ToMatrix(), Utilities::Multiply(InvIL1, Q1.ToMatrix().GetTransposed()));
+		const TVector<T, d> CX = GetLimitedPositionError(JointSettings, R0, X1 - X0);
 
 		// Calculate Swing axis for each body
-		TRotation<T, d> R01 = R0.Inverse() * R1;
+		const TRotation<T, d> R01 = R0.Inverse() * R1;
 		TRotation<T, d> R01Twist, R01Swing;
 		R01.ToSwingTwist(TJointConstants<T, d>::TwistAxis(), R01Swing, R01Twist);
 		R01Swing = R01Swing.GetNormalized();
@@ -734,29 +873,25 @@ namespace Chaos
 			SwingAngle = SwingAngle - (T)2 * PI;
 		}
 
-		TVector<T, d> SwingAxis0 = R0 * SwingAxis01;
-		TVector<T, d> SwingAxis1 = SwingAxis0;
+		const TVector<T, d> SwingAxis0 = R0 * SwingAxis01;
+		const TVector<T, d> SwingAxis1 = SwingAxis0;
 
 		// Circular swing target (max of Swing1, Swing2 targets)
 		T Swing1Target = JointSettings.Motion.AngularDriveTargetAngles[(int32)EJointAngularConstraintIndex::Swing1];
 		T Swing2Target = JointSettings.Motion.AngularDriveTargetAngles[(int32)EJointAngularConstraintIndex::Swing2];
 		T SwingAngleTarget = FMath::Max(Swing1Target, Swing2Target);
 
-		T Stiffness = (SolverSettings.PBDDriveStiffness > 0) ? SolverSettings.PBDDriveStiffness : JointSettings.Motion.AngularDriveStiffness;
-		T DriveStiffness = FMath::Clamp(Stiffness, (T)0, (T)1);
+		T DriveStiffnessUnclamped = (SolverSettings.DriveStiffness > 0) ? SolverSettings.DriveStiffness : JointSettings.Motion.AngularDriveStiffness;
+		T DriveStiffness = FMath::Clamp(DriveStiffnessUnclamped, (T)0, (T)1);
 		T DSwingAngle = SwingAngle - SwingAngleTarget;
 		T DSwingAngle0 = DriveStiffness * DSwingAngle;
 		T DSwingAngle1 = -DriveStiffness * DSwingAngle;
 
 		// Apply swing correction
-		T L = (T)1 / (TVector<T, d>::DotProduct(SwingAxis0, Utilities::Multiply(InvI0, SwingAxis0)) + TVector<T, d>::DotProduct(SwingAxis1, Utilities::Multiply(InvI1, SwingAxis1)));
-		TVector<T, d> W0 = Utilities::Multiply(InvI0, SwingAxis0) * L * DSwingAngle0;
-		TVector<T, d> W1 = Utilities::Multiply(InvI1, SwingAxis1) * L * DSwingAngle1;
-		TRotation<T, d> DQ0 = (TRotation<T, d>::FromElements(W0, (T)0.0) * Q0) * (T)0.5;
-		TRotation<T, d> DQ1 = (TRotation<T, d>::FromElements(W1, (T)0.0) * Q1) * (T)0.5;
-		Q0 = (Q0 + DQ0).GetNormalized();
-		Q1 = (Q1 + DQ1).GetNormalized();
-		Q1.EnforceShortestArcWith(Q0);
+		ApplyRotationDelta(Dt, SolverSettings, JointSettings, Index0, Index1, P0, Q0, P1, Q1, InvM0, InvIL0, InvM1, InvIL1, SwingAxis0, DSwingAngle0, SwingAxis1, DSwingAngle1);
+
+		// Correct any positional error we may have introduced
+		ApplyPostRotationPositionCorrection<T, d>(Dt, SolverSettings, JointSettings, Index0, Index1, P0, Q0, P1, Q1, InvM0, InvIL0, InvM1, InvIL1, CX, SwingAxis0, SwingAxis1);
 	}
 
 	template<typename T, int d>
@@ -777,28 +912,131 @@ namespace Chaos
 	{
 		const TRigidTransform<T, d>& XL0 = JointSettings.ConstraintFrames[Index0];
 		const TRigidTransform<T, d>& XL1 = JointSettings.ConstraintFrames[Index1];
+		const TVector<T, d> X0 = P0 + Q0 * XL0.GetTranslation();
+		const TVector<T, d> X1 = P1 + Q1 * XL1.GetTranslation();
 		const TRotation<T, d> R0 = Q0 * XL0.GetRotation();
 		const TRotation<T, d> R1 = Q1 * XL1.GetRotation();
-		PMatrix<T, d, d> InvI0 = Utilities::Multiply(Q0.ToMatrix(), Utilities::Multiply(InvIL0, Q0.ToMatrix().GetTransposed()));
-		PMatrix<T, d, d> InvI1 = Utilities::Multiply(Q1.ToMatrix(), Utilities::Multiply(InvIL1, Q1.ToMatrix().GetTransposed()));
+		const TVector<T, d> CX = GetLimitedPositionError(JointSettings, R0, X1 - X0);
 
+		// Calculate the rotation we need to apply to resolve the rotation delta
 		const TRotation<T, d> TargetR1 = R0 * JointSettings.Motion.AngularDriveTarget;
 		const TRotation<T, d> DR1 = TargetR1 * R1.Inverse();
 		const TRotation<T, d> TargetQ0 = DR1.Inverse() * Q0;
 		const TRotation<T, d> TargetQ1 = DR1 * Q1;
 
-		T Stiffness = (SolverSettings.PBDDriveStiffness > 0) ? SolverSettings.PBDDriveStiffness : JointSettings.Motion.AngularDriveStiffness;
-		T DriveStiffness = FMath::Clamp(Stiffness, (T)0, (T)1);
+		T DriveStiffnessUnclamped = (SolverSettings.DriveStiffness > 0) ? SolverSettings.DriveStiffness : JointSettings.Motion.AngularDriveStiffness;
+		T DriveStiffness = FMath::Clamp(DriveStiffnessUnclamped, (T)0, (T)1);
 
-		// @todo(ccaulfield): use ang mom in slerp drive
-		T L0 = InvM0;// TVector<T, d>::DotProduct(Axis0, Utilities::Multiply(InvI0, Axis0));
-		T L1 = InvM1;//TVector<T, d>::DotProduct(Axis1, Utilities::Multiply(InvI1, Axis1));
-		const T F0 = DriveStiffness * L0 / (L0 + L1);
-		const T F1 = DriveStiffness * L1 / (L0 + L1);
+		TVector<T, d> SLerpAxis;
+		T SLerpAngle;
+		if (DR1.ToAxisAndAngleSafe(SLerpAxis, SLerpAngle, TVector<T, d>(1, 0, 0)))
+		{
+			const PMatrix<T, d, d> InvI0 = Utilities::Multiply(Q0.ToMatrix(), Utilities::Multiply(InvIL0, Q0.ToMatrix().GetTransposed()));
+			const PMatrix<T, d, d> InvI1 = Utilities::Multiply(Q1.ToMatrix(), Utilities::Multiply(InvIL1, Q1.ToMatrix().GetTransposed()));
+			const T I0 = TVector<T, d>::DotProduct(SLerpAxis, Utilities::Multiply(InvI0, SLerpAxis));
+			const T I1 = TVector<T, d>::DotProduct(SLerpAxis, Utilities::Multiply(InvI1, SLerpAxis));
+			const T F0 = DriveStiffness * I0 / (I0 + I1);
+			const T F1 = DriveStiffness * I1 / (I0 + I1);
+
+			// Apply the rotation delta
+			Q0 = TRotation<T, d>::Slerp(Q0, TargetQ0, F0);
+			Q1 = TRotation<T, d>::Slerp(Q1, TargetQ1, F1);
+			Q1.EnforceShortestArcWith(Q0);
+
+			// Correct any positional error we may have introduced
+			ApplyPostRotationPositionCorrection<T, d>(Dt, SolverSettings, JointSettings, Index0, Index1, P0, Q0, P1, Q1, InvM0, InvIL0, InvM1, InvIL1, CX, SLerpAxis, SLerpAxis);;
+		}
+	}
+
+	template<typename T, int d>
+	void TPBDJointUtilities<T, d>::ApplyJointPositionProjection(
+		const T Dt,
+		const TPBDJointSolverSettings<T, d>& SolverSettings,
+		const TPBDJointSettings<T, d>& JointSettings,
+		const int32 Index0,
+		const int32 Index1,
+		TVector<T, d>& P0,
+		TRotation<T, d>& Q0,
+		TVector<T, d>& P1,
+		TRotation<T, d>& Q1,
+		float InvM0,
+		const PMatrix<T, d, d>& InvIL0,
+		float InvM1,
+		const PMatrix<T, d, d>& InvIL1,
+		const T ProjectionFactor)
+	{
+		const TRigidTransform<T, d>& XL0 = JointSettings.ConstraintFrames[Index0];
+		const TRigidTransform<T, d>& XL1 = JointSettings.ConstraintFrames[Index1];
+		const TVector<T, d> X0 = P0 + Q0 * XL0.GetTranslation();
+		const TVector<T, d> X1 = P1 + Q1 * XL1.GetTranslation();
+		const TRotation<T, d> R0 = Q0 * XL0.GetRotation();
+		const TVector<T, d> CX = GetLimitedPositionError(JointSettings, R0, X1 - X0);
+
+		const TVector<T, d> V0 = CX * (ProjectionFactor * InvM0 / (InvM0 + InvM1));
+		const TVector<T, d> V1 = -CX * (ProjectionFactor * InvM1 / (InvM0 + InvM1));
+		P0 = P0 + V0;
+		P1 = P1 + V1;
+	}
+
+	template<typename T, int d>
+	void TPBDJointUtilities<T, d>::ApplyJointTwistProjection(
+		const T Dt,
+		const TPBDJointSolverSettings<T, d>& SolverSettings,
+		const TPBDJointSettings<T, d>& JointSettings,
+		const int32 Index0,
+		const int32 Index1,
+		TVector<T, d>& P0,
+		TRotation<T, d>& Q0,
+		TVector<T, d>& P1,
+		TRotation<T, d>& Q1,
+		float InvM0,
+		const PMatrix<T, d, d>& InvIL0,
+		float InvM1,
+		const PMatrix<T, d, d>& InvIL1,
+		const T ProjectionFactor)
+	{
+
+	}
+
+	template<typename T, int d>
+	void TPBDJointUtilities<T, d>::ApplyJointConeProjection(
+		const T Dt,
+		const TPBDJointSolverSettings<T, d>& SolverSettings,
+		const TPBDJointSettings<T, d>& JointSettings,
+		const int32 Index0,
+		const int32 Index1,
+		TVector<T, d>& P0,
+		TRotation<T, d>& Q0,
+		TVector<T, d>& P1,
+		TRotation<T, d>& Q1,
+		float InvM0,
+		const PMatrix<T, d, d>& InvIL0,
+		float InvM1,
+		const PMatrix<T, d, d>& InvIL1,
+		const T ProjectionFactor)
+	{
+
+	}
+
+	template<typename T, int d>
+	void TPBDJointUtilities<T, d>::ApplyJointSwingProjection(
+		const T Dt,
+		const TPBDJointSolverSettings<T, d>& SolverSettings,
+		const TPBDJointSettings<T, d>& JointSettings,
+		const int32 Index0,
+		const int32 Index1,
+		const EJointAngularConstraintIndex SwingConstraint,
+		TVector<T, d>& P0,
+		TRotation<T, d>& Q0,
+		TVector<T, d>& P1,
+		TRotation<T, d>& Q1,
+		float InvM0,
+		const PMatrix<T, d, d>& InvIL0,
+		float InvM1,
+		const PMatrix<T, d, d>& InvIL1,
+		const T ProjectionFactor)
+	{
 
-		Q0 = TRotation<T, d>::Slerp(Q0, TargetQ0, F0);
-		Q1 = TRotation<T, d>::Slerp(Q1, TargetQ1, F1);
-		Q1.EnforceShortestArcWith(Q0);
 	}
 }
 
diff --git a/Engine/Source/Runtime/Experimental/Chaos/Private/Chaos/PBDJointConstraints.cpp b/Engine/Source/Runtime/Experimental/Chaos/Private/Chaos/PBDJointConstraints.cpp
--- a/Engine/Source/Runtime/Experimental/Chaos/Private/Chaos/PBDJointConstraints.cpp
+++ b/Engine/Source/Runtime/Experimental/Chaos/Private/Chaos/PBDJointConstraints.cpp
@@ -61,6 +61,7 @@ namespace Chaos
 	template<class T, int d>
 	TPBDJointMotionSettings<T, d>::TPBDJointMotionSettings()
 		: Stiffness((T)1)
+		, Projection((T)0)
 		, LinearMotionTypes({ EJointMotionType::Locked, EJointMotionType::Locked, EJointMotionType::Locked })
 		, LinearLimit(FLT_MAX)
 		, AngularMotionTypes({ EJointMotionType::Free, EJointMotionType::Free, EJointMotionType::Free })
@@ -78,6 +79,7 @@ namespace Chaos
 	template<class T, int d>
 	TPBDJointMotionSettings<T, d>::TPBDJointMotionSettings(const TVector<EJointMotionType, d>& InLinearMotionTypes, const TVector<EJointMotionType, d>& InAngularMotionTypes)
 		: Stiffness((T)1)
+		, Projection((T)0)
 		, LinearMotionTypes(InLinearMotionTypes)
 		, LinearLimit(FLT_MAX)
 		, AngularMotionTypes({ EJointMotionType::Free, EJointMotionType::Free, EJointMotionType::Free })
@@ -111,16 +113,18 @@ namespace Chaos
 
 	template<class T, int d>
 	TPBDJointSolverSettings<T, d>::TPBDJointSolverSettings()
-		: SwingTwistAngleTolerance(1.0e-6f)
-		, PBDMinParentMassRatio(0.5f)
-		, PBDMaxInertiaRatio(5.0f)
-		, FreezeIterations(0)
-		, FrozenIterations(0)
-		, bEnableLinearLimits(1)
-		, bEnableTwistLimits(1)
-		, bEnableSwingLimits(1)
-		, bEnableDrives(1)
-		, PBDDriveStiffness(0.0f)
+		: SwingTwistAngleTolerance((T)1.0e-6)
+		, MinParentMassRatio((T)0.5)
+		, MaxInertiaRatio((T)5.0)
+		, bProjectPostApply(true)
+		, bEnableLinearLimits(true)
+		, bEnableTwistLimits(true)
+		, bEnableSwingLimits(true)
+		, bEnablePositionCorrection(true)
+		, bEnableDrives(true)
+		, Projection((T)0)
+		, Stiffness((T)0)
+		, DriveStiffness((T)0)
 	{
 	}
 
@@ -326,10 +330,13 @@ namespace Chaos
 		// @todo(ccaulfield): make sorting optional
 		// @todo(ccaulfield): handles should be sorted by level by the constraint rule/graph
 		// @todo(ccaulfield): the best sort order depends on whether we are freezing.
-		// If we are freezing, we want the root-most bodies solved first, otherwise we want them last
+		// If we are freezing we want the root-most (nearest to kinematic) bodies solved first.
+		// For normal update we want the root body last, otherwise it gets dragged away from the root by the other bodies
+
 		TArray<FConstraintHandle*> SortedConstraintHandles = InConstraintHandles;
 		SortedConstraintHandles.Sort([](const FConstraintHandle& L, const FConstraintHandle& R)
 			{
+				// Sort bodies from leaf to root
 				return L.GetConstraintLevel() > R.GetConstraintLevel();
 			});
 
@@ -343,6 +350,11 @@ namespace Chaos
 			ApplySingle(Dt, ConstraintHandle->GetConstraintIndex(), It, NumIts);
 		}
 
+		if (!Settings.bProjectPostApply)
+		{
+			ApplyProjection(Dt, InConstraintHandles);
+		}
+
 		if (PostApplyCallback != nullptr)
 		{
 			PostApplyCallback(Dt, SortedConstraintHandles);
@@ -352,6 +364,26 @@ namespace Chaos
 	template<class T, int d>
 	void TPBDJointConstraints<T, d>::ApplyPushOut(const T Dt, const TArray<FConstraintHandle*>& InConstraintHandles)
 	{
+		if (Settings.bProjectPostApply)
+		{
+			ApplyProjection(Dt, InConstraintHandles);
+		}
+	}
+
+	template<class T, int d>
+	void TPBDJointConstraints<T, d>::ApplyProjection(const T Dt, const TArray<FConstraintHandle*>& InConstraintHandles)
+	{
+		TArray<FConstraintHandle*> SortedConstraintHandles = InConstraintHandles;
+		SortedConstraintHandles.Sort([](const FConstraintHandle& L, const FConstraintHandle& R)
+			{
+				// Sort bodies from root to leaf
+				return L.GetConstraintLevel() < R.GetConstraintLevel();
+			});
+
+		for (FConstraintHandle* ConstraintHandle : SortedConstraintHandles)
+		{
+			ApplyProjectionSingle(Dt, ConstraintHandle->GetConstraintIndex(), 0, 1);
+		}
 	}
 
 	template<class T, int d>
@@ -388,15 +420,15 @@ namespace Chaos
 			const int32 Level1 = ConstraintStates[ConstraintIndex].ParticleLevels[Index1];
 			if (Level0 < Level1)
 			{
-				TPBDJointUtilities<T, d>::GetConditionedInverseMass(Rigid0, Rigid1, InvM0, InvM1, InvIL0, InvIL1, Settings.PBDMinParentMassRatio, Settings.PBDMaxInertiaRatio);
+				TPBDJointUtilities<T, d>::GetConditionedInverseMass(Rigid0, Rigid1, InvM0, InvM1, InvIL0, InvIL1, Settings.MinParentMassRatio, Settings.MaxInertiaRatio);
 			}
 			else if (Level0 > Level1)
 			{
-				TPBDJointUtilities<T, d>::GetConditionedInverseMass(Rigid1, Rigid0, InvM1, InvM0, InvIL1, InvIL0, Settings.PBDMinParentMassRatio, Settings.PBDMaxInertiaRatio);
+				TPBDJointUtilities<T, d>::GetConditionedInverseMass(Rigid1, Rigid0, InvM1, InvM0, InvIL1, InvIL0, Settings.MinParentMassRatio, Settings.MaxInertiaRatio);
 			}
 			else
 			{
-				TPBDJointUtilities<T, d>::GetConditionedInverseMass(Rigid1, Rigid0, InvM1, InvM0, InvIL1, InvIL0, (T)0, Settings.PBDMaxInertiaRatio);
+				TPBDJointUtilities<T, d>::GetConditionedInverseMass(Rigid1, Rigid0, InvM1, InvM0, InvIL1, InvIL0, (T)0, Settings.MaxInertiaRatio);
 			}
 		}
 
@@ -493,6 +525,68 @@ namespace Chaos
 			Rigid1->SetQ(Q1);
 		}
 	}
+
+	template<class T, int d>
+	void TPBDJointConstraints<T, d>::ApplyProjectionSingle(const T Dt, const int32 ConstraintIndex, const int32 It, const int32 NumIts)
+	{
+		const FJointSettings& JointSettings = ConstraintSettings[ConstraintIndex];
+
+		// Scale projection up to ProjectionSetting over NumProjectionIts
+		const T ProjectionFactor = (Settings.Projection > 0) ? Settings.Projection : JointSettings.Motion.Projection;
+		if (ProjectionFactor == (T)0)
+		{
+			return;
+		}
+
+		const TVector<TGeometryParticleHandle<T, d>*, 2>& Constraint = ConstraintParticles[ConstraintIndex];
+		UE_LOG(LogChaosJoint, Verbose, TEXT("Project Joint Constraint %d %f (it = %d / %d)"), ConstraintIndex, ProjectionFactor, It, NumIts);
+
+		// Switch particles - internally we assume the first body is the parent (i.e., the space in which constraint limits are specified)
+		const int32 Index0 = 1;
+		const int32 Index1 = 0;
+		TGenericParticleHandle<T, d> Particle0 = TGenericParticleHandle<T, d>(ConstraintParticles[ConstraintIndex][Index0]);
+		TGenericParticleHandle<T, d> Particle1 = TGenericParticleHandle<T, d>(ConstraintParticles[ConstraintIndex][Index1]);
+
+		TVector<T, d> P0 = Particle0->P();
+		TRotation<T, d> Q0 = Particle0->Q();
+		TVector<T, d> P1 = Particle1->P();
+		TRotation<T, d> Q1 = Particle1->Q();
+		float InvM0 = Particle0->InvM();
+		float InvM1 = Particle1->InvM();
+		PMatrix<T, d, d> InvIL0 = Particle0->InvI();
+		PMatrix<T, d, d> InvIL1 = Particle1->InvI();
+
+		// Freeze the closest to kinematic connection if there is a difference
+		const int32 Level0 = ConstraintStates[ConstraintIndex].ParticleLevels[Index0];
+		const int32 Level1 = ConstraintStates[ConstraintIndex].ParticleLevels[Index1];
+		if (Level0 < Level1)
+		{
+			InvM0 = 0;
+			InvIL0 = PMatrix<T, d, d>(0, 0, 0);
+		}
+		else if (Level1 < Level0)
+		{
+			InvM1 = 0;
+			InvIL1 = PMatrix<T, d, d>(0, 0, 0);
+		}
+
+		// Project position error
+		TPBDJointUtilities<T, d>::ApplyJointPositionProjection(Dt, Settings, JointSettings, Index0, Index1, P0, Q0, P1, Q1, InvM0, InvIL0, InvM1, InvIL1, ProjectionFactor);
+
+		// Update the particles
+		TPBDRigidParticleHandle<T, d>* Rigid0 = ConstraintParticles[ConstraintIndex][Index0]->AsDynamic();
+		TPBDRigidParticleHandle<T, d>* Rigid1 = ConstraintParticles[ConstraintIndex][Index1]->AsDynamic();
+		if (Rigid0)
+		{
+			Rigid0->SetP(P0);
+			Rigid0->SetQ(Q0);
+		}
+		if (Rigid1)
+		{
+			Rigid1->SetP(P1);
+			Rigid1->SetQ(Q1);
+		}
+	}
 }
 
 namespace Chaos
diff --git a/Engine/Source/Runtime/Experimental/Chaos/Public/Chaos/PBDJointConstraintTypes.h b/Engine/Source/Runtime/Experimental/Chaos/Public/Chaos/PBDJointConstraintTypes.h
--- a/Engine/Source/Runtime/Experimental/Chaos/Public/Chaos/PBDJointConstraintTypes.h
+++ b/Engine/Source/Runtime/Experimental/Chaos/Public/Chaos/PBDJointConstraintTypes.h
@@ -68,6 +68,7 @@ namespace Chaos
 		TPBDJointMotionSettings(const TVector<EJointMotionType, d>& InLinearMotionTypes, const TVector<EJointMotionType, d>& InAngularMotionTypes);
 
 		T Stiffness;
+		T Projection;
 
 		TVector<EJointMotionType, d> LinearMotionTypes;
 		T LinearLimit;
@@ -113,16 +114,18 @@ namespace Chaos
 		T SwingTwistAngleTolerance;
 
 		// Stability control
-		T PBDMinParentMassRatio;
-		T PBDMaxInertiaRatio;
-		int32 FreezeIterations;
-		int32 FrozenIterations;
+		T MinParentMassRatio;
+		T MaxInertiaRatio;
 
 		// @todo(ccaulfield): remove these TEMP overrides for testing
+		bool bProjectPostApply;
 		bool bEnableLinearLimits;
 		bool bEnableTwistLimits;
 		bool bEnableSwingLimits;
+		bool bEnablePositionCorrection;
 		bool bEnableDrives;
-		T PBDDriveStiffness;
+		T Projection;
+		T Stiffness;
+		T DriveStiffness;
 	};
 }
\ No newline at end of file
diff --git a/Engine/Source/Runtime/Experimental/Chaos/Public/Chaos/PBDJointConstraintUtilities.h b/Engine/Source/Runtime/Experimental/Chaos/Public/Chaos/PBDJointConstraintUtilities.h
--- a/Engine/Source/Runtime/Experimental/Chaos/Public/Chaos/PBDJointConstraintUtilities.h
+++ b/Engine/Source/Runtime/Experimental/Chaos/Public/Chaos/PBDJointConstraintUtilities.h
@@ -194,5 +194,71 @@ namespace Chaos
 			const PMatrix<T, d, d>& InvIL0,
 			float InvM1,
 			const PMatrix<T, d, d>& InvIL1);
+
+		static CHAOS_API void ApplyJointPositionProjection(
+			const T Dt,
+			const TPBDJointSolverSettings<T, d>& SolverSettings,
+			const TPBDJointSettings<T, d>& JointSettings,
+			const int32 Index0,
+			const int32 Index1,
+			TVector<T, d>& P0,
+			TRotation<T, d>& Q0,
+			TVector<T, d>& P1,
+			TRotation<T, d>& Q1,
+			float InvM0,
+			const PMatrix<T, d, d>& InvIL0,
+			float InvM1,
+			const PMatrix<T, d, d>& InvIL1,
+			const T ProjectionFactor);
+
+		static CHAOS_API void ApplyJointTwistProjection(
+			const T Dt,
+			const TPBDJointSolverSettings<T, d>& SolverSettings,
+			const TPBDJointSettings<T, d>& JointSettings,
+			const int32 Index0,
+			const int32 Index1,
+			TVector<T, d>& P0,
+			TRotation<T, d>& Q0,
+			TVector<T, d>& P1,
+			TRotation<T, d>& Q1,
+			float InvM0,
+			const PMatrix<T, d, d>& InvIL0,
+			float InvM1,
+			const PMatrix<T, d, d>& InvIL1,
+			const T ProjectionFactor);
+
+		static CHAOS_API void ApplyJointConeProjection(
+			const T Dt,
+			const TPBDJointSolverSettings<T, d>& SolverSettings,
+			const TPBDJointSettings<T, d>& JointSettings,
+			const int32 Index0,
+			const int32 Index1,
+			TVector<T, d>& P0,
+			TRotation<T, d>& Q0,
+			TVector<T, d>& P1,
+			TRotation<T, d>& Q1,
+			float InvM0,
+			const PMatrix<T, d, d>& InvIL0,
+			float InvM1,
+			const PMatrix<T, d, d>& InvIL1,
+			const T ProjectionFactor);
+
+		static CHAOS_API void ApplyJointSwingProjection(
+			const T Dt,
+			const TPBDJointSolverSettings<T, d>& SolverSettings,
+			const TPBDJointSettings<T, d>& JointSettings,
+			const int32 Index0,
+			const int32 Index1,
+			const EJointAngularConstraintIndex SwingConstraint,
+			TVector<T, d>& P0,
+			TRotation<T, d>& Q0,
+			TVector<T, d>& P1,
+			TRotation<T, d>& Q1,
+			float InvM0,
+			const PMatrix<T, d, d>& InvIL0,
+			float InvM1,
+			const PMatrix<T, d, d>& InvIL1,
+			const T ProjectionFactor);
+
 	};
 }
\ No newline at end of file
diff --git a/Engine/Source/Runtime/Experimental/Chaos/Public/Chaos/PBDJointConstraints.h b/Engine/Source/Runtime/Experimental/Chaos/Public/Chaos/PBDJointConstraints.h
--- a/Engine/Source/Runtime/Experimental/Chaos/Public/Chaos/PBDJointConstraints.h
+++ b/Engine/Source/Runtime/Experimental/Chaos/Public/Chaos/PBDJointConstraints.h
@@ -139,6 +139,9 @@ namespace Chaos
 
 		CHAOS_API void ApplySingle(const T Dt, const int32 ConstraintIndex, const int32 It, const int32 NumIts);
 
+		CHAOS_API void ApplyProjection(const T Dt, const TArray<FConstraintHandle*>& InConstraintHandles);
+		CHAOS_API void ApplyProjectionSingle(const T Dt, const int32 ConstraintIndex, const int32 It, const int32 NumIts);
+
 		TPBDJointSolverSettings<T, d> Settings;
 
 		TArray<FJointSettings> ConstraintSettings;
