Make segment be core shape type for capsule

diff --git a/Engine/Source/Runtime/Experimental/Chaos/Public/Chaos/Capsule.h b/Engine/Source/Runtime/Experimental/Chaos/Public/Chaos/Capsule.h
--- a/Engine/Source/Runtime/Experimental/Chaos/Public/Chaos/Capsule.h
+++ b/Engine/Source/Runtime/Experimental/Chaos/Public/Chaos/Capsule.h
@@ -5,6 +5,7 @@
 #include "Chaos/ImplicitObject.h"
 #include "Chaos/ImplicitObjectUnion.h"
 #include "Chaos/Sphere.h"
+#include "Chaos/Segment.h"
 #include "ChaosArchive.h"
 
 namespace Chaos
@@ -24,9 +25,7 @@ namespace Chaos
 		{}
 		TCapsule(const TVector<T, 3>& x1, const TVector<T, 3>& x2, const T Radius)
 		    : TImplicitObject<T, 3>(EImplicitObject::FiniteConvex, ImplicitObjectType::Capsule)
-		    , MPoint(x1)
-		    , MVector(x2 - x1)
-		    , MHeight(MVector.SafeNormalize())
+			, MSegment(x1, x2)
 		    , MRadius(Radius)
 		    , MLocalBoundingBox(x1, x1)
 		    , MUnionedObjects(nullptr)
@@ -38,9 +37,7 @@ namespace Chaos
 
 		TCapsule(const TCapsule<T>& Other)
 		    : TImplicitObject<T, 3>(EImplicitObject::FiniteConvex, ImplicitObjectType::Capsule)
-		    , MPoint(Other.MPoint)
-		    , MVector(Other.MVector)
-		    , MHeight(Other.MHeight)
+			, MSegment(Other.MSegment)
 		    , MRadius(Other.MRadius)
 		    , MLocalBoundingBox(Other.MLocalBoundingBox)
 		    , MUnionedObjects(nullptr)
@@ -50,9 +47,7 @@ namespace Chaos
 
 		TCapsule(TCapsule<T>&& Other)
 		    : TImplicitObject<T, 3>(EImplicitObject::FiniteConvex, ImplicitObjectType::Capsule)
-		    , MPoint(MoveTemp(Other.MPoint))
-		    , MVector(MoveTemp(Other.MVector))
-		    , MHeight(Other.MHeight)
+			, MSegment(MoveTemp(Other.MSegment))
 		    , MRadius(Other.MRadius)
 		    , MLocalBoundingBox(MoveTemp(Other.MLocalBoundingBox))
 		    , MUnionedObjects(MoveTemp(Other.MUnionedObjects))
@@ -67,9 +62,7 @@ namespace Chaos
 			this->bIgnoreAnalyticCollisions = InSteal.bIgnoreAnalyticCollisions;
 			this->bHasBoundingBox = InSteal.bHasBoundingBox;
 
-			MPoint = MoveTemp(InSteal.MPoint);
-			MVector = MoveTemp(InSteal.MVector);
-			MHeight = InSteal.MHeight;
+			MSegment = MoveTemp(InSteal.MSegment);
 			MRadius = InSteal.MRadius;
 			MLocalBoundingBox = MoveTemp(InSteal.MLocalBoundingBox);
 			MUnionedObjects = MoveTemp(InSteal.MUnionedObjects);
@@ -132,8 +125,8 @@ namespace Chaos
 
 		virtual T PhiWithNormal(const TVector<T, 3>& x, TVector<T, 3>& Normal) const override
 		{
-			auto Dot = FMath::Clamp(TVector<T, 3>::DotProduct(x - MPoint, MVector), (T)0., MHeight);
-			TVector<T, 3> ProjectedPoint = Dot * MVector + MPoint;
+			auto Dot = FMath::Clamp(TVector<T, 3>::DotProduct(x - GetX1(), GetAxis()), (T)0., GetHeight());
+			TVector<T, 3> ProjectedPoint = Dot * GetAxis() + GetX1();
 			Normal = x - ProjectedPoint;
 			return Normal.SafeNormalize() - MRadius;
 		}
@@ -291,29 +284,25 @@ namespace Chaos
 
 		virtual bool Raycast(const TVector<T, 3>& StartPoint, const TVector<T, 3>& Dir, const T Length, const T Thickness, T& OutTime, TVector<T, 3>& OutPosition, TVector<T, 3>& OutNormal, int32& OutFaceIndex) const override
 		{
-			return RaycastFast(MRadius, MHeight, MVector, GetX1(), GetX2(), StartPoint, Dir, Length, Thickness, OutTime, OutPosition, OutNormal, OutFaceIndex);
+			return RaycastFast(MRadius, GetHeight(), GetAxis(), GetX1(), GetX2(), StartPoint, Dir, Length, Thickness, OutTime, OutPosition, OutNormal, OutFaceIndex);
 		}
 
 		TVector<T,3> Support(const TVector<T, 3>& Direction, const T Thickness) const override
 		{
-			const T Dot = TVector<T, 3>::DotProduct(Direction, MVector);
-			const TVector<T, 3> FarthestCap = Dot >= 0 ? GetX2() : GetX1();	//orthogonal we choose either
-			//We want N / ||N|| and to avoid inf
-			//So we want N / ||N|| < 1 / eps => N eps < ||N||, but this is clearly true for all eps < 1 and N > 0
-			T SizeSqr = Direction.SizeSquared();
-			if (SizeSqr <= TNumericLimits<T>::Min())
-			{
-				return FarthestCap;
-			}
-			const TVector<T, 3> Normalized = Direction / sqrt(SizeSqr);
-			return FarthestCap + Normalized * (MRadius + Thickness);
+			return MSegment.Support(Direction, MRadius + Thickness);
+		}
+
+		FORCEINLINE void SerializeImp(FArchive& Ar)
+		{
+			TImplicitObject<T, 3>::SerializeImp(Ar);
+			MSegment.Serialize(Ar);
+			Ar << MRadius << MLocalBoundingBox;
 		}
 
 		virtual void Serialize(FChaosArchive& Ar) override
 		{
 			FChaosArchiveScopedMemory ScopedMemory(Ar, GetTypeName());
-			TImplicitObject<T, 3>::SerializeImp(Ar);
-			Ar << MPoint << MVector << MHeight << MRadius << MLocalBoundingBox;
+			SerializeImp(Ar);
 			Ar << MUnionedObjects;
 		}
 
@@ -323,25 +312,26 @@ namespace Chaos
 		}
 
 		const T& GetRadius() const { return MRadius; }
-		const T& GetHeight() const { return MHeight; }
+		T GetHeight() const { return MSegment.GetLength(); }
 		/** Returns the bottommost point on the capsule. */
-		const TVector<T, 3> GetOrigin() const { return MPoint + MVector * -MRadius; }
+		const TVector<T, 3> GetOrigin() const { return GetX1() + GetAxis() * -MRadius; }
 		/** Returns the topmost point on the capsule. */
-		const TVector<T, 3> GetInsertion() const { return MPoint + MVector * (MHeight + MRadius); }
-		TVector<T, 3> GetCenter() const { return MPoint + MVector * (MHeight / 2); }
+		const TVector<T, 3> GetInsertion() const { return GetX1() + GetAxis() * (GetHeight() + MRadius); }
+		TVector<T, 3> GetCenter() const { return MSegment.GetCenter(); }
 		/** Returns the centroid (center of mass). */
 		TVector<T, 3> GetCenterOfMass() const { return GetCenter(); }
-		const TVector<T, 3>& GetAxis() const { return MVector; }
-		const TVector<T, 3>& GetX1() const { return MPoint; }
-		TVector<T, 3> GetX2() const { return MPoint + MVector * MHeight; }
+		const TVector<T, 3>& GetAxis() const { return MSegment.GetAxis(); }
+		const TVector<T, 3>& GetX1() const { return MSegment.GetX1(); }
+		TVector<T, 3> GetX2() const { return MSegment.GetX2(); }
+		TSegment<T> GetSegment() const { return TSegment<T>(GetX1(), GetX2()); }
 
-		T GetArea() const { return GetArea(MHeight, MRadius); }
-		static T GetArea(const T Height, const T Radius) { static const T PI2 = 2. * PI; return PI2*Radius * (Height + 2.*Radius); }
+		T GetArea() const { return GetArea(GetHeight(), MRadius); }
+		static T GetArea(const T Height, const T Radius) { static const T PI2 = 2. * PI; return PI2 * Radius * (Height + 2.*Radius); }
 
-		T GetVolume() const { return GetVolume(MHeight, MRadius); }
-		static T GetVolume(const T Height, const T Radius) { static const T FourThirds = 4./3; return PI*Radius*Radius * (Height + FourThirds*Radius); }
+		T GetVolume() const { return GetVolume(GetHeight(), MRadius); }
+		static T GetVolume(const T Height, const T Radius) { static const T FourThirds = 4. / 3; return PI * Radius*Radius * (Height + FourThirds * Radius); }
 
-		PMatrix<T, 3, 3> GetInertiaTensor(const T Mass) const { return GetInertiaTensor(Mass, MHeight, MRadius); }
+		PMatrix<T, 3, 3> GetInertiaTensor(const T Mass) const { return GetInertiaTensor(Mass, GetHeight(), MRadius); }
 		static PMatrix<T, 3, 3> GetInertiaTensor(const T Mass, const T Height, const T Radius)
 		{
 			// https://www.wolframalpha.com/input/?i=capsule&assumption=%7B%22C%22,+%22capsule%22%7D+-%3E+%7B%22Solid%22%7D
@@ -365,7 +355,7 @@ namespace Chaos
 
 		virtual uint32 GetTypeHash() const override
 		{
-			return HashCombine(::GetTypeHash(MPoint), ::GetTypeHash(MVector));
+			return HashCombine(::GetTypeHash(GetX1()), ::GetTypeHash(GetAxis()));
 		}
 
 	private:
@@ -373,9 +363,11 @@ namespace Chaos
 		{
 			TArray<TUniquePtr<TImplicitObject<float, 3>>> Objects;
 
-			Objects.Add(MakeUnique<Chaos::TCylinder<float>>(MPoint, MPoint + MVector * MHeight, MRadius));
-			Objects.Add(MakeUnique<Chaos::TSphere<float, 3>>(MPoint, MRadius));
-			Objects.Add(MakeUnique<Chaos::TSphere<float, 3>>(MPoint + MVector * MHeight, MRadius));
+			const TVector<T, 3>& X1 = GetX1();
+			const TVector<T, 3>& X2 = GetX2();
+			Objects.Add(MakeUnique<Chaos::TCylinder<float>>(X1, X2, MRadius));
+			Objects.Add(MakeUnique<Chaos::TSphere<float, 3>>(X1, MRadius));
+			Objects.Add(MakeUnique<Chaos::TSphere<float, 3>>(X2, MRadius));
 
 			MUnionedObjects.Reset(new Chaos::TImplicitObjectUnion<float, 3>(std::move(Objects)));
 		}
@@ -385,8 +377,8 @@ namespace Chaos
 			return MUnionedObjects->FindClosestIntersection(StartPoint, EndPoint, Thickness);
 		}
 
-		TVector<T, 3> MPoint, MVector;
-		T MHeight, MRadius;
+		TSegment<T> MSegment;
+		T MRadius;
 		TBox<T, 3> MLocalBoundingBox;
 		TUniquePtr<TImplicitObjectUnion<T, 3>> MUnionedObjects;
 	};
diff --git a/Engine/Source/Runtime/Experimental/Chaos/Public/Chaos/Segment.h b/Engine/Source/Runtime/Experimental/Chaos/Public/Chaos/Segment.h
new file mode 100644
--- /dev/null
+++ b/Engine/Source/Runtime/Experimental/Chaos/Public/Chaos/Segment.h
@@ -0,0 +1,53 @@
+#pragma once
+
+namespace Chaos
+{
+	template<class T>
+	class TSegment
+	{
+	public:
+		TSegment() {}
+		TSegment(const TVector<T, 3>& x1, const TVector<T, 3>& x2)
+			: MPoint(x1)
+			, MAxis(x2 - x1)
+			, MLength(MAxis.SafeNormalize())
+		{ }
+
+		FORCEINLINE bool IsConvex() const { return true; }
+
+		FORCEINLINE const TVector<T, 3> GetCenter() const { return MPoint + (.5f * MLength * MAxis); }
+
+		FORCEINLINE const TVector<T, 3>& GetX1() const { return MPoint; }
+
+		FORCEINLINE TVector<T, 3> GetX2() const { return MPoint + MAxis * MLength; }
+
+		FORCEINLINE const TVector<T, 3>& GetAxis() const { return MAxis; }
+
+		FORCEINLINE float GetLength() const { return MLength; }
+
+		TVector<T,3> Support(const TVector<T, 3>& Direction, const T Thickness) const
+		{
+			const T Dot = TVector<T, 3>::DotProduct(Direction, MAxis);
+			const TVector<T, 3> FarthestCap = Dot >= 0 ? GetX2() : GetX1();	//orthogonal we choose either
+			//We want N / ||N|| and to avoid inf
+			//So we want N / ||N|| < 1 / eps => N eps < ||N||, but this is clearly true for all eps < 1 and N > 0
+			T SizeSqr = Direction.SizeSquared();
+			if (SizeSqr <= TNumericLimits<T>::Min())
+			{
+				return FarthestCap;
+			}
+			const TVector<T, 3> NormalizedDirection = Direction / sqrt(SizeSqr);
+			return FarthestCap + (NormalizedDirection * Thickness);
+		}
+
+		FORCEINLINE void Serialize(FArchive &Ar) 
+		{
+			Ar << MPoint << MAxis << MLength;
+		}
+
+	private:
+		TVector<T, 3> MPoint;
+		TVector<T, 3> MAxis;
+		T MLength;
+	};
+}
