Fix erroneous reference to MPoint in Capsule.h

diff --git a/Engine/Source/Runtime/Experimental/Chaos/Public/Chaos/Capsule.h b/Engine/Source/Runtime/Experimental/Chaos/Public/Chaos/Capsule.h
--- a/Engine/Source/Runtime/Experimental/Chaos/Public/Chaos/Capsule.h
+++ b/Engine/Source/Runtime/Experimental/Chaos/Public/Chaos/Capsule.h
@@ -92,7 +92,7 @@ namespace Chaos
 		{
 			TArray<TVector<T, 3>> Points;
 			const TVector<T, 3> Mid = GetCenter();
-			const TCapsule<T> Capsule(MPoint - Mid, MPoint + (MVector * MHeight) - Mid, GetRadius());
+			const TCapsule<T> Capsule(GetX1() - Mid, GetX1() + (GetAxis() * GetHeight()) - Mid, GetRadius());
 			TCapsuleSpecializeSamplingHelper<T>::ComputeSamplePoints(Points, Capsule, NumPoints);
 			return Points;
 		}
