Add copyright to Segment.h

diff --git a/Engine/Source/Runtime/Experimental/Chaos/Public/Chaos/Segment.h b/Engine/Source/Runtime/Experimental/Chaos/Public/Chaos/Segment.h
--- a/Engine/Source/Runtime/Experimental/Chaos/Public/Chaos/Segment.h
+++ b/Engine/Source/Runtime/Experimental/Chaos/Public/Chaos/Segment.h
@@ -1,3 +1,4 @@
+// Copyright 1998-2019 Epic Games, Inc. All Rights Reserved.
 #pragma once
 
 namespace Chaos
