MTD for capsules

diff --git a/Engine/Source/Runtime/Engine/Private/PhysicsEngine/Experimental/PhysInterface_Chaos.cpp b/Engine/Source/Runtime/Engine/Private/PhysicsEngine/Experimental/PhysInterface_Chaos.cpp
--- a/Engine/Source/Runtime/Engine/Private/PhysicsEngine/Experimental/PhysInterface_Chaos.cpp
+++ b/Engine/Source/Runtime/Engine/Private/PhysicsEngine/Experimental/PhysInterface_Chaos.cpp
@@ -1656,7 +1656,7 @@ bool FPhysInterface_Chaos::Sweep_Geom(FHitResult& OutHit, const FBodyInstance* I
 							Chaos::TVector<float, 3> WorldPosition;
 							Chaos::TVector<float, 3> WorldNormal;
 							int32 FaceIdx;
-							if(Chaos::SweepQuery<float, 3>(*Shape->Geometry, ActorTM, ShapeAdapter.GetGeometry(), StartTM, Dir, DeltaMag, Hit.Distance, WorldPosition, WorldNormal, FaceIdx))
+							if(Chaos::SweepQuery<float, 3>(*Shape->Geometry, ActorTM, ShapeAdapter.GetGeometry(), StartTM, Dir, DeltaMag, Hit.Distance, WorldPosition, WorldNormal, FaceIdx, 0.f, false))
 							{
 								// we just like to make sure if the hit is made
 								FCollisionFilterData QueryFilter;
diff --git a/Engine/Source/Runtime/Experimental/Chaos/Public/Chaos/GJK.h b/Engine/Source/Runtime/Experimental/Chaos/Public/Chaos/GJK.h
--- a/Engine/Source/Runtime/Experimental/Chaos/Public/Chaos/GJK.h
+++ b/Engine/Source/Runtime/Experimental/Chaos/Public/Chaos/GJK.h
@@ -1,5 +1,9 @@
 // Copyright 1998-2019 Epic Games, Inc. All Rights Reserved.
 #pragma once
+#include "Chaos/Transform.h"
+#include "Chaos/Capsule.h"
+#include "Chaos/Box.h"
+#include "Chaos/Sphere.h"
 
 namespace Chaos
 {
@@ -582,6 +586,19 @@ namespace Chaos
 		return TVector<T, 3>(0, 0, 0);
 	}
 
+	// Overloads for geometry types which don't have centroids.
+	template <typename T, typename TGeometryB>
+	TVector<T, 3> GJKDistanceInitialV(const TImplicitObject<T, 3>& A, const TGeometryB& B, const TRigidTransform<T, 3>& BToATM)
+	{
+		return -BToATM.GetTranslation();
+	}
+
+	template <typename T, typename TGeometryA>
+	TVector<T, 3> GJKDistanceInitialV(TGeometryA A, const TImplicitObject<T, 3>& B, const TRigidTransform<T, 3>& BToATM)
+	{
+		return -BToATM.GetTranslation();
+	}
+
 	/**
 	 * Find the distance and nearest points on two convex geometries A and B.
 	 * All calculations are performed in the local-space of object A, and the transform from B-space to A-space must be provided.
@@ -677,6 +694,65 @@ namespace Chaos
 		return false;
 	}
 
+	// Assumes objects are already intersecting, computes a minimum translation
+	// distance, deepest penetration positions on each body, and approximates
+	// a penetration normal and minimum translation distance.
+	//
+	// TODO: We want to re-visit how these functions work. Probably should be
+	// embedded in GJKOverlap and GJKRaycast so that secondary queries are unnecessary.
+	template <typename T, typename TGeometryA, typename TGeometryB>
+	bool GJKPenetrationTemp(const TGeometryA& A, const TGeometryB& B, const TRigidTransform<T, 3>& BToATM, TVector<T, 3>& OutPositionA, TVector<T, 3>& OutPositionB, TVector<T, 3>& OutNormal, T& OutDistance, const T ThicknessA = 0, const TVector<T, 3>& InitialDir = TVector<T, 3>(1, 0, 0), const T ThicknessB = 0, const T Epsilon = (T)1e-6, const int32 MaxIts = 16)
+	{
+		//
+		// TODO: General case for MTD determination.
+		//
+		ensure(false);
+		OutPositionA = TVector<T, 3>(0.f);
+		OutPositionB = TVector<T, 3>(0.f);
+		OutNormal = TVector<T, 3>(0.f, 0.f, 1.f);
+		OutDistance = 0.f;
+		return GJKIntersection(A, B, BToATM, ThicknessA, InitialDir, ThicknessB);
+	}
+
+	// Specialization for when getting MTD against a capsule.
+	template <typename T, typename TGeometryA>
+	bool GJKPenetrationTemp(const TGeometryA& A, const TCapsule<T>& B, const TRigidTransform<T, 3>& BToATM, TVector<T, 3>& OutPositionA, TVector<T, 3>& OutPositionB, TVector<T, 3>& OutNormal, T& OutDistance, const T ThicknessA = 0, const TVector<T, 3>& InitialDir = TVector<T, 3>(1, 0, 0), const T ThicknessB = 0, const T Epsilon = (T)1e-6, const int32 MaxIts = 16)
+	{
+		float SegmentDistance;
+		const TSegment<T>& Segment = B.GetSegment();
+		const float MarginB = B.GetRadius();
+		TVector<float, 3> PositionBInB;
+		if (GJKDistance(A, Segment, BToATM, SegmentDistance, OutPositionA, PositionBInB, Epsilon, MaxIts))
+		{
+			OutPositionB = BToATM.TransformPosition(PositionBInB);
+			OutNormal = (OutPositionB - OutPositionA) / SegmentDistance; // I think this is safe, if GJKDistance returned true...
+			OutPositionB -= OutNormal * MarginB;
+			OutDistance = SegmentDistance - MarginB;
+
+			if (OutDistance > 0.f)
+			{
+				// In this case, our distance calculation says we're not penetrating.
+				//
+				// TODO: check(false)! This shouldn't happen.
+				// It probably won't happen anymore if we warm-start GJKDistance
+				// with a polytope.
+				//
+				OutDistance = 0.f;
+				return false;
+			}
+
+			return true;
+		}
+		else
+		{
+			// TODO: Deep penetration - do EPA
+			ensure(false);
+			return true;
+		}
+
+		return false;
+	}
+
 	/** Approximate the collision point using Core Shape Reduction.
 	 @A The first geometry
 	 @ATM The transform of the first geometry
diff --git a/Engine/Source/Runtime/Experimental/Chaos/Public/Chaos/GeometryQueries.h b/Engine/Source/Runtime/Experimental/Chaos/Public/Chaos/GeometryQueries.h
--- a/Engine/Source/Runtime/Experimental/Chaos/Public/Chaos/GeometryQueries.h
+++ b/Engine/Source/Runtime/Experimental/Chaos/Public/Chaos/GeometryQueries.h
@@ -7,6 +7,7 @@
 #include "Chaos/TriangleMeshImplicitObject.h"
 #include "Chaos/HeightField.h"
 #include "Chaos/Convex.h"
+#include "Chaos/Capsule.h"
 #include "ImplicitObjectScaled.h"
 
 #include "ChaosArchive.h"
@@ -17,7 +18,7 @@
 namespace Chaos
 {
 	template <typename T, int d>
-	bool OverlapQuery(const TImplicitObject<T, d>& A, const TRigidTransform<T,d>& ATM, const TImplicitObject<T, d>& B, const TRigidTransform<T,d>& BTM, const T Thickness = 0)
+	bool OverlapQuery(const TImplicitObject<T, d>& A, const TRigidTransform<T,d>& ATM, const TImplicitObject<T, d>& B, const TRigidTransform<T,d>& BTM, const T Thickness = 0, TVector<T,d>* OutMTD=nullptr)
 	{
 		const ImplicitObjectType AType = A.GetType(true);
 		const ImplicitObjectType BType = B.GetType(true);
@@ -26,14 +27,14 @@ namespace Chaos
 		{
 			const TImplicitObjectTransformed<T, d>& TransformedA = static_cast<const TImplicitObjectTransformed<T, d>&>(A);
 			const TRigidTransform<T, d> NewATM = TransformedA.GetTransform() * ATM;
-			return OverlapQuery(*TransformedA.GetTransformedObject(), NewATM, B, BTM, Thickness);
+			return OverlapQuery(*TransformedA.GetTransformedObject(), NewATM, B, BTM, Thickness, OutMTD);
 		}
 
 		if (BType == ImplicitObjectType::Transformed)
 		{
 			const TImplicitObjectTransformed<T, d>& TransformedB = static_cast<const TImplicitObjectTransformed<T, d>&>(B);
 			const TRigidTransform<T, d> NewBTM = TransformedB.GetTransform() * BTM;
-			return OverlapQuery(A, ATM, *TransformedB.GetTransformedObject(), NewBTM, Thickness);
+			return OverlapQuery(A, ATM, *TransformedB.GetTransformedObject(), NewBTM, Thickness, OutMTD);
 		}
 
 		check(B.IsConvex());	//Query object must be convex
@@ -83,7 +84,7 @@ namespace Chaos
 	}
 
 	template <typename T, int d>
-	bool SweepQuery(const TImplicitObject<T, d>& A, const TRigidTransform<T,d>& ATM, const TImplicitObject<T, d>& B, const TRigidTransform<T, d>& BTM, const TVector<T,d>& Dir, const T Length, T& OutTime, TVector<T,d>& OutPosition, TVector<T,d>& OutNormal, int32& OutFaceIndex, const T Thickness = 0)
+	bool SweepQuery(const TImplicitObject<T, d>& A, const TRigidTransform<T,d>& ATM, const TImplicitObject<T, d>& B, const TRigidTransform<T, d>& BTM, const TVector<T,d>& Dir, const T Length, T& OutTime, TVector<T,d>& OutPosition, TVector<T,d>& OutNormal, int32& OutFaceIndex, const T Thickness, const bool bComputeMTD)
 	{
 		const ImplicitObjectType AType = A.GetType(true);
 		const ImplicitObjectType BType = B.GetType(true);
@@ -94,14 +95,14 @@ namespace Chaos
 		{
 			const TImplicitObjectTransformed<T, d>& TransformedA = static_cast<const TImplicitObjectTransformed<T, d>&>(A);
 			const TRigidTransform<T, d> NewATM = TransformedA.GetTransform() * ATM;
-			return SweepQuery(*TransformedA.GetTransformedObject(), NewATM, B, BTM, Dir, Length, OutTime, OutPosition, OutNormal, OutFaceIndex, Thickness);
+			return SweepQuery(*TransformedA.GetTransformedObject(), NewATM, B, BTM, Dir, Length, OutTime, OutPosition, OutNormal, OutFaceIndex, Thickness, bComputeMTD);
 		}
 
 		if (BType == ImplicitObjectType::Transformed)
 		{
 			const TImplicitObjectTransformed<T, d>& TransformedB = static_cast<const TImplicitObjectTransformed<T, d>&>(B);
 			const TRigidTransform<T, d> NewBTM = TransformedB.GetTransform() * BTM;
-			return SweepQuery(A, ATM, *TransformedB.GetTransformedObject(), NewBTM, Dir, Length, OutTime, OutPosition, OutNormal, OutFaceIndex, Thickness);
+			return SweepQuery(A, ATM, *TransformedB.GetTransformedObject(), NewBTM, Dir, Length, OutTime, OutPosition, OutNormal, OutFaceIndex, Thickness, bComputeMTD);
 		}
 
 		check(B.IsConvex());	//Object being swept must be convex
@@ -165,6 +166,21 @@ namespace Chaos
 			{
 				//todo: find face index if convex hull
 			}
+
+			// Compute MTD in the case of an initial overlap
+			if (bResult && bComputeMTD && OutTime == 0.f)
+			{
+				TVector<T, d> LocalPositionOnB;
+				if (BType == ImplicitObjectType::Capsule)
+				{
+					const TCapsule<T>& BCapsule = static_cast<const TCapsule<T>&>(B);
+					GJKPenetrationTemp<T>(A, BCapsule, BToATM, LocalPosition, LocalPositionOnB, LocalNormal, OutTime, Thickness, Offset);
+				}
+				else
+				{
+					GJKPenetrationTemp<T>(A, B, BToATM, LocalPosition, LocalPositionOnB, LocalNormal, OutTime, Thickness, Offset);
+				}
+			}
 		}
 		else
 		{
@@ -198,7 +214,7 @@ namespace Chaos
 		}
 
 		//put back into world space
-		if (OutTime > 0)
+		if (OutTime > 0 || bComputeMTD)
 		{
 			OutNormal = ATM.TransformVectorNoScale(LocalNormal);
 			OutPosition = ATM.TransformPositionNoScale(LocalPosition);
diff --git a/Engine/Source/Runtime/PhysicsSQ/Private/SQAccelerator.cpp b/Engine/Source/Runtime/PhysicsSQ/Private/SQAccelerator.cpp
--- a/Engine/Source/Runtime/PhysicsSQ/Private/SQAccelerator.cpp
+++ b/Engine/Source/Runtime/PhysicsSQ/Private/SQAccelerator.cpp
@@ -67,16 +67,16 @@ struct FPreFilterInfo
 	int32 ActorIdx;
 };
 
-void FillHitHelper(ChaosInterface::FLocationHit& Hit, const float Distance, const FVector& WorldPosition, const FVector& WorldNormal, int32 FaceIdx)
+void FillHitHelper(ChaosInterface::FLocationHit& Hit, const float Distance, const FVector& WorldPosition, const FVector& WorldNormal, int32 FaceIdx, bool bComputeMTD)
 {
 	Hit.Distance = Distance;
 	Hit.WorldPosition = WorldPosition;
 	Hit.WorldNormal = WorldNormal;
-	Hit.Flags = Distance > 0 ? EHitFlags::Distance | EHitFlags::Normal | EHitFlags::Position : EHitFlags::Distance | EHitFlags::FaceIndex;
+	Hit.Flags = Distance > 0.f || bComputeMTD ? EHitFlags::Distance | EHitFlags::Normal | EHitFlags::Position : EHitFlags::Distance | EHitFlags::FaceIndex;
 	Hit.FaceIndex = FaceIdx;
 }
 
-void FillHitHelper(ChaosInterface::FOverlapHit& Hit, const float Distance, const FVector& WorldPosition, const FVector& WorldNormal, int32 FaceIdx)
+void FillHitHelper(ChaosInterface::FOverlapHit& Hit, const float Distance, const FVector& WorldPosition, const FVector& WorldNormal, int32 FaceIdx, bool bComputeMTD)
 {
 }
 
@@ -203,6 +203,7 @@ private:
 				TVector<float, 3> WorldPosition, WorldNormal;
 				float Distance = 0;	//not needed but fixes compiler warning for overlap
 				int32 FaceIdx = INDEX_NONE;	//not needed but fixes compiler warning for overlap
+				const bool bComputeMTD = !!((uint16)(OutputFlags & EHitFlags::MTD));
 
 				if (SQ == ESQType::Raycast)
 				{
@@ -220,7 +221,7 @@ private:
 				}
 				else if(SQ == ESQType::Sweep && CurLength > 0)
 				{
-					bHit = SweepQuery<float, 3>(*Geom, ActorTM, *QueryGeom, StartTM, Dir, CurLength, Distance, WorldPosition, WorldNormal, FaceIdx);
+					bHit = SweepQuery<float, 3>(*Geom, ActorTM, *QueryGeom, StartTM, Dir, CurLength, Distance, WorldPosition, WorldNormal, FaceIdx, 0.f, bComputeMTD);
 				}
 				else if (SQ == ESQType::Overlap || (SQ == ESQType::Sweep && CurLength == 0))
 				{
@@ -229,7 +230,7 @@ private:
 
 				if(bHit)
 				{
-					FillHitHelper(Hit, Distance, WorldPosition, WorldNormal, FaceIdx);
+					FillHitHelper(Hit, Distance, WorldPosition, WorldNormal, FaceIdx, bComputeMTD);
 #if WITH_PHYSX
 					HitType = QueryFilterData.flags & PxQueryFlag::ePOSTFILTER ? QueryCallback.PostFilter(P2UFilterData(QueryFilterData.data), Hit) : HitType;
 #else
