Fix heightfield scale issues. Need more thought

diff --git a/Engine/Source/Runtime/Experimental/Chaos/Private/Chaos/HeightField.cpp b/Engine/Source/Runtime/Experimental/Chaos/Private/Chaos/HeightField.cpp
--- a/Engine/Source/Runtime/Experimental/Chaos/Private/Chaos/HeightField.cpp
+++ b/Engine/Source/Runtime/Experimental/Chaos/Private/Chaos/HeightField.cpp
@@ -531,10 +531,10 @@ namespace Chaos
 		if(FlatGrid.IsValid(InCoord))
 		{
 			const TBox<T, 3>& Bound = GeomData.CellBounds[InCoord[1] * (GeomData.NumCols - 1) + InCoord[0]];
-			OutMin = Bound.Min() - InInflate;
-			OutMax = Bound.Max() + InInflate;
-			OutMin *= GeomData.Scale;
-			OutMax *= GeomData.Scale;
+			OutMin = Bound.Min();
+			OutMax = Bound.Max();
+			OutMin = OutMin * GeomData.Scale - InInflate;
+			OutMax = OutMax * GeomData.Scale + InInflate;
 			return true;
 		}
 
@@ -1225,7 +1225,7 @@ namespace Chaos
 		
 		TVector<T, 2> MinCorner(0, 0);
 		TVector<T, 2> MaxCorner(GeomData.NumCols - 1, GeomData.NumRows - 1);
-		MaxCorner *= {GeomData.Scale[0], GeomData.Scale[1]};
+		//MaxCorner *= {GeomData.Scale[0], GeomData.Scale[1]};
 
 		FlatGrid = TUniformGrid<T, 2>(MinCorner, MaxCorner, Cells);
 	}
diff --git a/Engine/Source/Runtime/Experimental/Chaos/Public/Chaos/HeightField.h b/Engine/Source/Runtime/Experimental/Chaos/Public/Chaos/HeightField.h
--- a/Engine/Source/Runtime/Experimental/Chaos/Public/Chaos/HeightField.h
+++ b/Engine/Source/Runtime/Experimental/Chaos/Public/Chaos/HeightField.h
@@ -99,6 +99,7 @@ namespace Chaos
 			if(Ar.IsLoading())
 			{
 				BuildQueryData();
+				BoundingBox();	//temp hack to initialize cache
 			}
 		}
 
