Use triangle in convex hull

diff --git a/Engine/Source/Runtime/Experimental/Chaos/Private/Chaos/HeightField.cpp b/Engine/Source/Runtime/Experimental/Chaos/Private/Chaos/HeightField.cpp
--- a/Engine/Source/Runtime/Experimental/Chaos/Private/Chaos/HeightField.cpp
+++ b/Engine/Source/Runtime/Experimental/Chaos/Private/Chaos/HeightField.cpp
@@ -7,6 +7,7 @@
 #include "Chaos/Capsule.h"
 #include "HAL/IConsoleManager.h"
 #include "HAL/PlatformMisc.h"
+#include "Chaos/Triangle.h"
 
 namespace Chaos
 {
@@ -215,16 +216,7 @@ namespace Chaos
 				{
 					return false;
 				}
-
-				// This isn't great as we build a convex on-the-fly but we have no other
-				// trimesh support. Update later as better trimesh collisions come online
-				TParticles<T, 3> TriParticles;
-				TriParticles.AddParticles(3);
-				TriParticles.X(0) = A;
-				TriParticles.X(1) = B;
-				TriParticles.X(2) = C;
-
-				TConvex<T, 3> Convex(TriParticles);
+				TTriangle<T> Convex(A, B, C);
 
 				T Time;
 				TVector<T, 3> HitPosition;
