Don't construct anything if no particles exist.

diff --git a/Engine/Source/Runtime/Experimental/Chaos/Public/Chaos/BoundingVolume.h b/Engine/Source/Runtime/Experimental/Chaos/Public/Chaos/BoundingVolume.h
--- a/Engine/Source/Runtime/Experimental/Chaos/Public/Chaos/BoundingVolume.h
+++ b/Engine/Source/Runtime/Experimental/Chaos/Public/Chaos/BoundingVolume.h
@@ -837,6 +837,11 @@ private:
 	template <typename ParticleView>
 	void GenerateTree(const ParticleView& Particles, const bool bUseVelocity, const T Dt, const int32 MaxCells)
 	{
+		if (!Particles.Num())
+		{
+			return;
+		}
+
 		SCOPE_CYCLE_COUNTER(STAT_BoundingVolumeGenerateTree);
 		TArray<TBox<T, d>> AllBounds;
 		TArray<bool> HasBounds;
