Chaos RBAN joints
- Joint Velocity solve
- Fix capsule inertia
- Fix race condition in Evolution velocity update

diff --git a/Engine/Source/Runtime/AnimGraphRuntime/Private/BoneControllers/AnimNode_RigidBody_Chaos.cpp b/Engine/Source/Runtime/AnimGraphRuntime/Private/BoneControllers/AnimNode_RigidBody_Chaos.cpp
--- a/Engine/Source/Runtime/AnimGraphRuntime/Private/BoneControllers/AnimNode_RigidBody_Chaos.cpp
+++ b/Engine/Source/Runtime/AnimGraphRuntime/Private/BoneControllers/AnimNode_RigidBody_Chaos.cpp
@@ -444,11 +444,13 @@ void FAnimNode_RigidBody_Chaos::EvaluateSkeletalControl_AnyThread(FComponentSpac
 
 			// Run simulation at a minimum of 30 FPS to prevent system from exploding.
 			// DeltaTime can be higher due to URO, so take multiple iterations in that case.
+			extern int32 RBAN_MaxSubSteps;
+			const int32 MaxSteps = RBAN_MaxSubSteps;
 			const float MaxDeltaSeconds = 1.f / 30.f;
-			const int32 NumIterations = FMath::Clamp(FMath::CeilToInt(DeltaSeconds / MaxDeltaSeconds), 1, 4);
-			const float StepDeltaTime = DeltaSeconds / float(NumIterations);
+			const int32 NumSteps = FMath::Clamp(FMath::CeilToInt(DeltaSeconds / MaxDeltaSeconds), 1, MaxSteps);
+			const float StepDeltaTime = DeltaSeconds / float(NumSteps);
 
-			for (int32 Step = 1; Step <= NumIterations; Step++)
+			for (int32 Step = 1; Step <= NumSteps; Step++)
 			{
 				// We call the _AssumesLocked version here without a lock as the simulation is local to this node and we know
 				// we're not going to alter anything while this is running.
diff --git a/Engine/Source/Runtime/Engine/Private/PhysicsEngine/ImmediatePhysics/ImmediatePhysicsChaos/ImmediatePhysicsActorHandle_Chaos.cpp b/Engine/Source/Runtime/Engine/Private/PhysicsEngine/ImmediatePhysics/ImmediatePhysicsChaos/ImmediatePhysicsActorHandle_Chaos.cpp
--- a/Engine/Source/Runtime/Engine/Private/PhysicsEngine/ImmediatePhysics/ImmediatePhysicsChaos/ImmediatePhysicsActorHandle_Chaos.cpp
+++ b/Engine/Source/Runtime/Engine/Private/PhysicsEngine/ImmediatePhysics/ImmediatePhysicsChaos/ImmediatePhysicsActorHandle_Chaos.cpp
@@ -37,8 +37,8 @@ namespace ImmediatePhysics_Chaos
 	Chaos::PMatrix<T, d, d> CalculateInertia_Solid(const T Mass, const FKSphylElem& SphylElem)
 	{
 		return Chaos::PMatrix<T, d, d>(
-			((T)1 / (T)12) * Mass * ((T)3 * SphylElem.Radius + SphylElem.Length),
-			((T)1 / (T)12) * Mass * ((T)3 * SphylElem.Radius + SphylElem.Length),
+			((T)1 / (T)12) * Mass * ((T)3 * SphylElem.Radius * SphylElem.Radius + SphylElem.Length * SphylElem.Length),
+			((T)1 / (T)12) * Mass * ((T)3 * SphylElem.Radius * SphylElem.Radius + SphylElem.Length * SphylElem.Length),
 			((T)1 / (T)2) * Mass * SphylElem.Radius * SphylElem.Radius
 			);
 	}
diff --git a/Engine/Source/Runtime/Engine/Private/PhysicsEngine/ImmediatePhysics/ImmediatePhysicsChaos/ImmediatePhysicsSimulation_Chaos.cpp b/Engine/Source/Runtime/Engine/Private/PhysicsEngine/ImmediatePhysics/ImmediatePhysicsChaos/ImmediatePhysicsSimulation_Chaos.cpp
--- a/Engine/Source/Runtime/Engine/Private/PhysicsEngine/ImmediatePhysics/ImmediatePhysicsChaos/ImmediatePhysicsSimulation_Chaos.cpp
+++ b/Engine/Source/Runtime/Engine/Private/PhysicsEngine/ImmediatePhysics/ImmediatePhysicsChaos/ImmediatePhysicsSimulation_Chaos.cpp
@@ -21,7 +21,7 @@
 // @todo(ccaulfield): remove when finished
 //
 float ChaosImmediate_Evolution_DeltaTime = 0.03f;
-int32 ChaosImmediate_Evolution_Iterations = 10;
+int32 ChaosImmediate_Evolution_Iterations = 4;
 FAutoConsoleVariableRef CVarChaosImmPhysDeltaTime(TEXT("p.Chaos.ImmPhys.DeltaTime"), ChaosImmediate_Evolution_DeltaTime, TEXT("Override chaos immediate physics delta time if non-zero"));
 FAutoConsoleVariableRef CVarChaosImmPhysIterations(TEXT("p.Chaos.ImmPhys.Iterations"), ChaosImmediate_Evolution_Iterations, TEXT("Number of constraint solver loops in immediate physics"));
 
@@ -37,24 +37,24 @@ FAutoConsoleVariableRef CVarChaosImmPhysCollisionPushOutPairIterations(TEXT("p.C
 FAutoConsoleVariableRef CVarChaosImmPhysThickness(TEXT("p.Chaos.ImmPhys.CollisionThickness"), ChaosImmediate_Collision_Thickness, TEXT("ChaosImmediateThickness"));
 
 float ChaosImmediate_Joint_SwingTwistAngleTolerance = 1.0e-6f;
+int32 ChaosImmediate_Joint_PositionIterations = 4;
 int32 ChaosImmediate_Joint_EnableLinearLimits = 1;
 int32 ChaosImmediate_Joint_EnableTwistLimits = 1;
 int32 ChaosImmediate_Joint_EnableSwingLimits = 1;
 int32 ChaosImmediate_Joint_EnablePositionCorrection = 0;
 int32 ChaosImmediate_Joint_EnableDrives = 1;
-int32 ChaosImmediate_Joint_ProjectPostApply = 1;
-float ChaosImmediate_Joint_Projection = 0.0f;
+float ChaosImmediate_Joint_Projection = 0.1f;
 float ChaosImmediate_Joint_Stiffness = 1.0f;
 float ChaosImmediate_Joint_DriveStiffness = 0.0f;
 float ChaosImmediate_Joint_MinParentMassRatio = 0.5f;
 float ChaosImmediate_Joint_MaxInertiaRatio = 5.0f;
 FAutoConsoleVariableRef CVarChaosImmPhysSwingTwistAngleTolerance(TEXT("p.Chaos.ImmPhys.Joint.SwingTwistAngleTolerance"), ChaosImmediate_Joint_SwingTwistAngleTolerance, TEXT("SwingTwistAngleTolerance."));
+FAutoConsoleVariableRef CVarChaosImmPhysEnablePositionSolve(TEXT("p.Chaos.ImmPhys.Joint.PositionIterations"), ChaosImmediate_Joint_PositionIterations, TEXT("6Dof joint PushOut (Position solve) iterations."));
 FAutoConsoleVariableRef CVarChaosImmPhysEnableLinearLimits(TEXT("p.Chaos.ImmPhys.Joint.EnableLinearLimits"), ChaosImmediate_Joint_EnableLinearLimits, TEXT("EnableLinearLimits."));
 FAutoConsoleVariableRef CVarChaosImmPhysEnableTwistLimits(TEXT("p.Chaos.ImmPhys.Joint.EnableTwistLimits"), ChaosImmediate_Joint_EnableTwistLimits, TEXT("EnableTwistLimits."));
 FAutoConsoleVariableRef CVarChaosImmPhysEnableSwingLimits(TEXT("p.Chaos.ImmPhys.Joint.EnableSwingLimits"), ChaosImmediate_Joint_EnableSwingLimits, TEXT("EnableSwingLimits."));
 FAutoConsoleVariableRef CVarChaosImmPhysEnablePositionCorrection(TEXT("p.Chaos.ImmPhys.Joint.EnablePositionCorrection"), ChaosImmediate_Joint_EnablePositionCorrection, TEXT("EnablePositionCorrection."));
 FAutoConsoleVariableRef CVarChaosImmPhysEnableDrives(TEXT("p.Chaos.ImmPhys.Joint.EnableDrives"), ChaosImmediate_Joint_EnableDrives, TEXT("EnableDrives."));
-FAutoConsoleVariableRef CVarChaosImmPhysProjectPostApply(TEXT("p.Chaos.ImmPhys.Joint.ProjectPostApply"), ChaosImmediate_Joint_ProjectPostApply, TEXT("6Dof joint projection in ApplyPushOut(true) or Apply(false)."));
 FAutoConsoleVariableRef CVarChaosImmPhysProjection(TEXT("p.Chaos.ImmPhys.Joint.Projection"), ChaosImmediate_Joint_Projection, TEXT("6Dof joint projection amount override (if > 0)."));
 FAutoConsoleVariableRef CVarChaosImmPhysStiffness(TEXT("p.Chaos.ImmPhys.Joint.Stiffness"), ChaosImmediate_Joint_Stiffness, TEXT("6Dof joint stiffness override (if > 0)."));
 FAutoConsoleVariableRef CVarChaosImmPhysDriveStiffness(TEXT("p.Chaos.ImmPhys.Joint.DriveStiffness"), ChaosImmediate_Joint_DriveStiffness, TEXT("6Dof joint drive stiffness override (if > 0)."));
@@ -68,9 +68,9 @@ FAutoConsoleVariableRef CVarChaosImmPhysMaxInertiaRatio(TEXT("p.Chaos.ImmPhys.Jo
 
 // DebugDraw CVars
 #if UE_BUILD_DEBUG
-int32 ChaosImmediate_DebugDrawParticles = 1;
+int32 ChaosImmediate_DebugDrawParticles = 0;
 int32 ChaosImmediate_DebugDrawShapes = 1;
-int32 ChaosImmediate_DebugDrawCollisions = 3;
+int32 ChaosImmediate_DebugDrawCollisions = 0;
 int32 ChaosImmediate_DebugDrawJoints = 1;
 #else
 int32 ChaosImmediate_DebugDrawParticles = 0;
@@ -409,7 +409,7 @@ namespace ImmediatePhysics_Chaos
 			JointsSettings.bEnableSwingLimits = ChaosImmediate_Joint_EnableSwingLimits != 0;
 			JointsSettings.bEnablePositionCorrection = ChaosImmediate_Joint_EnablePositionCorrection != 0;
 			JointsSettings.bEnableDrives = ChaosImmediate_Joint_EnableDrives != 0;
-			JointsSettings.bProjectPostApply = ChaosImmediate_Joint_ProjectPostApply != 0;
+			JointsSettings.PositionIterations = ChaosImmediate_Joint_PositionIterations;
 			JointsSettings.Projection = ChaosImmediate_Joint_Projection;
 			JointsSettings.Stiffness = ChaosImmediate_Joint_Stiffness;
 			JointsSettings.DriveStiffness = ChaosImmediate_Joint_DriveStiffness;
diff --git a/Engine/Source/Runtime/Experimental/Chaos/Private/Chaos/PBDJointConstraintUtilities.cpp b/Engine/Source/Runtime/Experimental/Chaos/Private/Chaos/PBDJointConstraintUtilities.cpp
--- a/Engine/Source/Runtime/Experimental/Chaos/Private/Chaos/PBDJointConstraintUtilities.cpp
+++ b/Engine/Source/Runtime/Experimental/Chaos/Private/Chaos/PBDJointConstraintUtilities.cpp
@@ -10,10 +10,10 @@
 	template<typename T, int d>
 	TVector<T, d> TPBDJointUtilities<T, d>::ConditionInertia(const TVector<T, d>& InI, const T MaxRatio)
 	{
-		if (MaxRatio > 0)
+		T IMin = InI.Min();
+		T IMax = InI.Max();
+		if ((MaxRatio > 0) && (IMin > 0))
 		{
-			T IMin = InI.Min();
-			T IMax = InI.Max();
 			T Ratio = IMax / IMin;
 			if (Ratio > MaxRatio)
 			{
@@ -34,11 +34,14 @@
 		{
 			T IParentMax = IParent.Max();
 			T IChildMax = IChild.Max();
-			T Ratio = IParentMax / IChildMax;
-			if (Ratio < MinRatio)
+			if ((IParentMax > 0) && (IChildMax > 0))
 			{
-				T Multiplier = MinRatio / Ratio;
-				return IParent * Multiplier;
+				T Ratio = IParentMax / IChildMax;
+				if (Ratio < MinRatio)
+				{
+					T Multiplier = MinRatio / Ratio;
+					return IParent * Multiplier;
+				}
 			}
 		}
 		return IParent;
@@ -47,7 +50,7 @@
 	template<typename T, int d>
 	T TPBDJointUtilities<T, d>::ConditionParentMass(const T MParent, const T MChild, const T MinRatio)
 	{
-		if (MinRatio > 0)
+		if ((MinRatio > 0) && (MParent > 0) && (MChild > 0))
 		{
 			T Ratio = MParent / MChild;
 			if (Ratio < MinRatio)
@@ -60,30 +63,59 @@
 	}
 
 	template<typename T, int d>
-	void TPBDJointUtilities<T, d>::GetConditionedInverseMass(const TPBDRigidParticleHandle<T, d>* PParent, const TPBDRigidParticleHandle<T, d>* PChild, T& OutInvMParent, T& OutInvMChild, PMatrix<T, d, d>& OutInvIParent, PMatrix<T, d, d>& OutInvIChild, const T MinParentMassRatio, const T MaxInertiaRatio)
+	void TPBDJointUtilities<T, d>::GetConditionedInverseMass(
+		const float InMParent, 
+		const TVector<T, d> InIParent, 
+		const float InMChild, 
+		const TVector<T, d> InIChild, 
+		T& OutInvMParent, 
+		T& OutInvMChild, 
+		PMatrix<T, d, d>& OutInvIParent, 
+		PMatrix<T, d, d>& OutInvIChild, 
+		const T MinParentMassRatio, 
+		const T MaxInertiaRatio)
 	{
-		T MParent = PParent->M();
-		T MChild = PChild->M();
-		MParent = ConditionParentMass(MParent, MChild, MinParentMassRatio);
+		T MParent = ConditionParentMass(InMParent, InMChild, MinParentMassRatio);
+		T MChild = InMChild;
 
-		TVector<T, d> IParent = ConditionInertia(PParent->I().GetDiagonal(), MaxInertiaRatio);
-		TVector<T, d> IChild = ConditionInertia(PChild->I().GetDiagonal(), MaxInertiaRatio);
+		TVector<T, d> IParent = ConditionInertia(InIParent, MaxInertiaRatio);
+		TVector<T, d> IChild = ConditionInertia(InIChild, MaxInertiaRatio);
 		IParent = ConditionParentInertia(IParent, IChild, MinParentMassRatio);
 
-		OutInvMParent = (T)1 / MParent;
-		OutInvMChild = (T)1 / MChild;
-		OutInvIParent = PMatrix<T, d, d>((T)1 / IParent.X, (T)1 / IParent.Y, (T)1 / IParent.Z);
-		OutInvIChild = PMatrix<T, d, d>((T)1 / IChild.X, (T)1 / IChild.Y, (T)1 / IChild.Z);
+		OutInvMParent = 0;
+		OutInvIParent = PMatrix<T, d, d>(0, 0, 0);
+		if (MParent > 0)
+		{
+			OutInvMParent = (T)1 / MParent;
+			OutInvIParent = PMatrix<T, d, d>((T)1 / IParent.X, (T)1 / IParent.Y, (T)1 / IParent.Z);
+		}
+
+		OutInvMChild = 0;
+		OutInvIChild = PMatrix<T, d, d>(0, 0, 0);
+		if (MChild > 0)
+		{
+			OutInvMChild = (T)1 / MChild;
+			OutInvIChild = PMatrix<T, d, d>((T)1 / IChild.X, (T)1 / IChild.Y, (T)1 / IChild.Z);
+		}
 	}
 
 
 	template<typename T, int d>
-	void TPBDJointUtilities<T, d>::GetConditionedInverseMass(const TPBDRigidParticleHandle<T, d>* P0, T& OutInvM0, PMatrix<T, d, d>& OutInvI0, const T MaxInertiaRatio)
+	void TPBDJointUtilities<T, d>::GetConditionedInverseMass(
+		const float InM0,
+		const TVector<T, d> InI0,
+		T& OutInvM0, 
+		PMatrix<T, d, d>& OutInvI0, 
+		const T MaxInertiaRatio)
 	{
-		TVector<T, d> I0 = ConditionInertia(P0->I().GetDiagonal(), MaxInertiaRatio);
-
-		OutInvM0 = P0->InvM();
-		OutInvI0 = PMatrix<T, d, d>((T)1 / I0.X, (T)1 / I0.Y, (T)1 / I0.Z);
+		OutInvM0 = 0;
+		OutInvI0 = PMatrix<T, d, d>(0, 0, 0);
+		if (InM0 > 0)
+		{
+			TVector<T, d> I0 = ConditionInertia(InI0, MaxInertiaRatio);
+			OutInvM0 = (T)1 / InM0;
+			OutInvI0 = PMatrix<T, d, d>((T)1 / I0.X, (T)1 / I0.Y, (T)1 / I0.Z);
+		}
 	}
 
 
@@ -104,5 +136,22 @@
 	}
 
 	template<typename T, int d>
+	TVector<T, d> GetSphereLimitedVelocityError(const TVector<T, d>& CX, const T Radius, const TVector<T, d>& CV)
+	{
+		T CXLen = CX.Size();
+		if (CXLen < Radius)
+		{
+			return TVector<T, d>(0, 0, 0);
+		}
+		else if (CXLen > SMALL_NUMBER)
+		{
+			TVector<T, d> Dir = CX / CXLen;
+			T CVDir = TVector<T, d>::DotProduct(CV, Dir);
+			return FMath::Max((T)0, CVDir) * Dir;
+		}
+		return CV;
+	}
+
+	template<typename T, int d>
 	TVector<T, d> GetCylinderLimitedPositionError(const TVector<T, d>& InCX, const TVector<T, d>& Axis, const T Limit, const EJointMotionType AxisMotion)
 	{
@@ -126,5 +175,32 @@
 	}
 
 	template<typename T, int d>
+	TVector<T, d> GetCylinderLimitedVelocityError(const TVector<T, d>& InCX, const TVector<T, d>& Axis, const T Limit, const EJointMotionType AxisMotion, const TVector<T, d>& CV)
+	{
+		TVector<T, d> CXAxis = TVector<T, d>::DotProduct(InCX, Axis) * Axis;
+		TVector<T, d> CXPlane = InCX - CXAxis;
+		T CXPlaneLen = CXPlane.Size();
+
+		TVector<T, d> CVAxis = TVector<T, d>::DotProduct(CV, Axis) * Axis;;
+		TVector<T, d> CVPlane = CV - CVAxis;
+
+		if (AxisMotion == EJointMotionType::Free)
+		{
+			CVAxis = TVector<T, d>(0, 0, 0);
+		}
+		if (CXPlaneLen < Limit)
+		{
+			CVPlane = TVector<T, d>(0, 0, 0);
+		}
+		else if (CXPlaneLen > KINDA_SMALL_NUMBER)
+		{
+			TVector<T, d> Dir = CXPlane / CXPlaneLen;
+			T CVDir = TVector<T, d>::DotProduct(CV, Dir);
+			CVPlane = FMath::Max((T)0, CVDir) * Dir;
+		}
+		return CVAxis + CVPlane;
+	}
+
+	template<typename T, int d>
 	TVector<T, d> GetLineLimitedPositionError(const TVector<T, d>& CX, const TVector<T, d>& Axis, const T Limit, const EJointMotionType AxisMotion)
 	{
@@ -144,5 +220,25 @@
 	}
 
 	template<typename T, int d>
+	TVector<T, d> GetLineLimitedVelocityError(const TVector<T, d>& CX, const TVector<T, d>& Axis, const T Limit, const EJointMotionType AxisMotion, const TVector<T, d>& CV)
+	{
+		T CXDist = TVector<T, d>::DotProduct(CX, Axis);
+		T CVAxis = TVector<T, d>::DotProduct(CV, Axis);
+		if ((AxisMotion == EJointMotionType::Free) || (FMath::Abs(CXDist) < Limit))
+		{
+			return CV - CVAxis * Axis;
+		}
+		else if (CXDist >= Limit)
+		{
+			return CV - FMath::Min((T)0, CVAxis) * Axis;
+		}
+		else
+		{
+			return CV - FMath::Max((T)0, CVAxis) * Axis;
+		}
+	}
+
+
+	template<typename T, int d>
 	TVector<T, d> GetLimitedPositionError(const TPBDJointSettings<T, d>& JointSettings, const TRotation<T, d>& R0, const TVector<T, d>& InCX)
 	{
@@ -197,6 +293,269 @@
 		}
 	}
 
+	template<typename T, int d>
+	TVector<T, d> GetLimitedVelocityError(const TPBDJointSettings<T, d>& JointSettings, const TRotation<T, d>& R0, const TVector<T, d>& InCX, const TVector<T, d>& InCV)
+	{
+		const TVector<EJointMotionType, d>& Motion = JointSettings.Motion.LinearMotionTypes;
+		if ((Motion[0] == EJointMotionType::Locked) && (Motion[1] == EJointMotionType::Locked) && (Motion[2] == EJointMotionType::Locked))
+		{
+			return InCV;
+		}
+		else if ((Motion[0] == EJointMotionType::Limited) && (Motion[1] == EJointMotionType::Limited) && (Motion[2] == EJointMotionType::Limited))
+		{
+			// Spherical distance constraints
+			return GetSphereLimitedVelocityError(InCX, JointSettings.Motion.LinearLimit, InCV);
+		}
+		else if ((Motion[1] == EJointMotionType::Limited) && (Motion[2] == EJointMotionType::Limited))
+		{
+			// Circular Limit (X Axis)
+			TVector<T, d> Axis = R0 * TVector<T, d>(1, 0, 0);
+			return GetCylinderLimitedVelocityError(InCX, Axis, JointSettings.Motion.LinearLimit, Motion[0], InCV);
+		}
+		else if ((Motion[0] == EJointMotionType::Limited) && (Motion[2] == EJointMotionType::Limited))
+		{
+			// Circular Limit (Y Axis)
+			TVector<T, d> Axis = R0 * TVector<T, d>(0, 1, 0);
+			return GetCylinderLimitedVelocityError(InCX, Axis, JointSettings.Motion.LinearLimit, Motion[1], InCV);
+		}
+		else if ((Motion[0] == EJointMotionType::Limited) && (Motion[1] == EJointMotionType::Limited))
+		{
+			// Circular Limit (Z Axis)
+			TVector<T, d> Axis = R0 * TVector<T, d>(0, 0, 1);
+			return GetCylinderLimitedVelocityError(InCX, Axis, JointSettings.Motion.LinearLimit, Motion[2], InCV);
+		}
+		else
+		{
+			// Line/Square/Cube Limits (no way to author square or cube limits, but would work if we wanted it)
+			TVector<T, d> CV = InCV;
+			if (Motion[0] != EJointMotionType::Locked)
+			{
+				TVector<T, d> Axis = R0 * TVector<T, d>(1, 0, 0);
+				CV = GetLineLimitedVelocityError(InCX, Axis, JointSettings.Motion.LinearLimit, Motion[0], CV);
+			}
+			if (Motion[1] != EJointMotionType::Locked)
+			{
+				TVector<T, d> Axis = R0 * TVector<T, d>(0, 1, 0);
+				CV = GetLineLimitedVelocityError(InCX, Axis, JointSettings.Motion.LinearLimit, Motion[1], CV);
+			}
+			if (Motion[2] != EJointMotionType::Locked)
+			{
+				TVector<T, d> Axis = R0 * TVector<T, d>(0, 0, 1);
+				CV = GetLineLimitedVelocityError(InCX, Axis, JointSettings.Motion.LinearLimit, Motion[2], CV);
+			}
+			return CV;
+		}
+	}
 
+
+	template<typename T, int d>
+	void ApplyPositionDelta(
+		const T Dt,
+		const TPBDJointSolverSettings<T, d>& SolverSettings,
+		const TPBDJointSettings<T, d>& JointSettings,
+		const int32 Index0,
+		const int32 Index1,
+		TVector<T, d>& P0,
+		TVector<T, d>& P1,
+		const TVector<T, d>& DP0,
+		const TVector<T, d>& DP1)
+	{
+		const float Stiffness = (SolverSettings.Stiffness > (T)0) ? SolverSettings.Stiffness : JointSettings.Motion.Stiffness;
+
+		P0 = P0 + Stiffness * DP0;
+		P1 = P1 + Stiffness * DP1;
+	}
+
+	template<typename T, int d>
+	void ApplyVelocityDelta(
+		const T Dt,
+		const TPBDJointSolverSettings<T, d>& SolverSettings,
+		const TPBDJointSettings<T, d>& JointSettings,
+		const int32 Index0,
+		const int32 Index1,
+		TVector<T, d>& P0,
+		TVector<T, d>& V0,
+		TVector<T, d>& P1,
+		TVector<T, d>& V1,
+		const TVector<T, d>& DV0,
+		const TVector<T, d>& DV1)
+	{
+		const float Stiffness = (SolverSettings.Stiffness > (T)0) ? SolverSettings.Stiffness : JointSettings.Motion.Stiffness;
+
+		V0 = V0 + Stiffness * DV0;
+		V1 = V1 + Stiffness * DV1;
+
+		ApplyPositionDelta(Dt, SolverSettings, JointSettings, Index0, Index1, P0, P1, DV0 * Dt, DV1 * Dt);
+	}
+
+	template<typename T, int d>
+	void ApplyRotationDelta(
+		const T Dt,
+		const TPBDJointSolverSettings<T, d>& SolverSettings,
+		const TPBDJointSettings<T, d>& JointSettings,
+		const int32 Index0,
+		const int32 Index1,
+		TRotation<T, d>& Q0,
+		TRotation<T, d>& Q1,
+		const TVector<T, d>& DR0,
+		const TVector<T, d>& DR1)
+	{
+		const float Stiffness = (SolverSettings.Stiffness > (T)0) ? SolverSettings.Stiffness : JointSettings.Motion.Stiffness;
+
+		const TRotation<T, d> DQ0 = (TRotation<T, d>::FromElements(Stiffness * DR0, 0) * Q0) * (T)0.5;
+		const TRotation<T, d> DQ1 = (TRotation<T, d>::FromElements(Stiffness * DR1, 0) * Q1) * (T)0.5;
+		Q0 = (Q0 + DQ0).GetNormalized();
+		Q1 = (Q1 + DQ1).GetNormalized();
+		Q1.EnforceShortestArcWith(Q0);
+	}
+
+
+	template<typename T, int d>
+	void ApplyRotationVelocityDelta(
+		const T Dt,
+		const TPBDJointSolverSettings<T, d>& SolverSettings,
+		const TPBDJointSettings<T, d>& JointSettings,
+		const int32 Index0,
+		const int32 Index1,
+		TRotation<T, d>& Q0,
+		TVector<T, d>& W0,
+		TRotation<T, d>& Q1,
+		TVector<T, d>& W1,
+		const TVector<T, d>& DW0,
+		const TVector<T, d>& DW1)
+	{
+		const float Stiffness = (SolverSettings.Stiffness > (T)0) ? SolverSettings.Stiffness : JointSettings.Motion.Stiffness;
+
+		W0 = W0 + Stiffness * DW0;
+		W1 = W1 + Stiffness * DW1;
+
+		ApplyRotationDelta(Dt, SolverSettings, JointSettings, Index0, Index1, Q0, Q1, DW0 * Dt, DW1 * Dt);
+	}
+
+
+	template<typename T, int d>
+	void ApplyRotationVelocityDelta(
+		const T Dt,
+		const TPBDJointSolverSettings<T, d>& SolverSettings,
+		const TPBDJointSettings<T, d>& JointSettings,
+		const int32 Index0,
+		const int32 Index1,
+		TRotation<T, d>& Q0,
+		TVector<T, d>& W0,
+		TRotation<T, d>& Q1,
+		TVector<T, d>& W1,
+		float InvM0,
+		const PMatrix<T, d, d>& InvIL0,
+		float InvM1,
+		const PMatrix<T, d, d>& InvIL1,
+		const TVector<T, d>& Axis0,
+		const TVector<T, d>& Axis1,
+		const float WC)
+	{
+		const PMatrix<T, d, d> InvI0 = Utilities::Multiply(Q0.ToMatrix(), Utilities::Multiply(InvIL0, Q0.ToMatrix().GetTransposed()));
+		const PMatrix<T, d, d> InvI1 = Utilities::Multiply(Q1.ToMatrix(), Utilities::Multiply(InvIL1, Q1.ToMatrix().GetTransposed()));
+		//const T I0 = TVector<T, d>::DotProduct(Axis0, Utilities::Multiply(InvI0, Axis0));
+		//const T I1 = TVector<T, d>::DotProduct(Axis1, Utilities::Multiply(InvI1, Axis1));
+		//const TVector<T, d> DW0 = Axis0 * WC * I0 / (I0 + I1);
+		//const TVector<T, d> DW1 = -Axis1 * WC * I1 / (I0 + I1);
+		const T L = (T)1 / (TVector<T, d>::DotProduct(Axis0, Utilities::Multiply(InvI0, Axis0)) + TVector<T, d>::DotProduct(Axis1, Utilities::Multiply(InvI1, Axis1)));
+		const TVector<T, d> DW0 = Utilities::Multiply(InvI0, Axis0) * L * WC;
+		const TVector<T, d> DW1 = -Utilities::Multiply(InvI1, Axis1) * L * WC;
+
+		ApplyRotationVelocityDelta(Dt, SolverSettings, JointSettings, Index0, Index1, Q0, W0, Q1, W1, DW0, DW1);
+	}
+
+	template<typename T, int d>
+	void ApplyRotationDelta(
+		const T Dt,
+		const TPBDJointSolverSettings<T, d>& SolverSettings,
+		const TPBDJointSettings<T, d>& JointSettings,
+		const int32 Index0,
+		const int32 Index1,
+		TVector<T, d>& P0,
+		TRotation<T, d>& Q0,
+		TVector<T, d>& P1,
+		TRotation<T, d>& Q1,
+		float InvM0,
+		const PMatrix<T, d, d>& InvIL0,
+		float InvM1,
+		const PMatrix<T, d, d>& InvIL1,
+		const TVector<T, d>& Axis0,
+		const float Angle0,
+		const TVector<T, d>& Axis1,
+		const float Angle1)
+	{
+		const PMatrix<T, d, d> InvI0 = Utilities::Multiply(Q0.ToMatrix(), Utilities::Multiply(InvIL0, Q0.ToMatrix().GetTransposed()));
+		const PMatrix<T, d, d> InvI1 = Utilities::Multiply(Q1.ToMatrix(), Utilities::Multiply(InvIL1, Q1.ToMatrix().GetTransposed()));
+
+		//const T I0 = TVector<T, d>::DotProduct(Axis0, Utilities::Multiply(InvI0, Axis0));
+		//const T I1 = TVector<T, d>::DotProduct(Axis1, Utilities::Multiply(InvI1, Axis1));
+		//const TVector<T, d> DR0 = Axis0 * Angle0 * I0 / (I0 + I1);
+		//const TVector<T, d> DR1 = Axis1 * Angle1 * I1 / (I0 + I1);
+		const T L = (T)1 / (TVector<T, d>::DotProduct(Axis0, Utilities::Multiply(InvI0, Axis0)) + TVector<T, d>::DotProduct(Axis1, Utilities::Multiply(InvI1, Axis1)));
+		const TVector<T, d> DR0 = Utilities::Multiply(InvI0, Axis0) * L * Angle0;
+		const TVector<T, d> DR1 = Utilities::Multiply(InvI1, Axis1) * L * Angle1;
+
+		ApplyRotationDelta(Dt, SolverSettings, JointSettings, Index0, Index1, Q0, Q1, DR0, DR1);
+	}
+
+	template<typename T, int d>
+	void ApplyPostRotationPositionCorrection(
+		const T Dt,
+		const TPBDJointSolverSettings<T, d>& SolverSettings,
+		const TPBDJointSettings<T, d>& JointSettings,
+		const int32 Index0,
+		const int32 Index1,
+		TVector<T, d>& P0,
+		TRotation<T, d>& Q0,
+		TVector<T, d>& P1,
+		TRotation<T, d>& Q1,
+		float InvM0,
+		const PMatrix<T, d, d>& InvIL0,
+		float InvM1,
+		const PMatrix<T, d, d>& InvIL1,
+		const TVector<T, d>& CX_1,
+		const TVector<T, d>& Axis0,
+		const TVector<T, d>& Axis1)
+	{
+		if (!SolverSettings.bEnablePositionCorrection)
+		{
+			return;
+		}
+
+		// Post-rotation constraint positions
+		const TRigidTransform<T, d>& XL0 = JointSettings.ConstraintFrames[Index0];
+		const TRigidTransform<T, d>& XL1 = JointSettings.ConstraintFrames[Index1];
+		const TVector<T, d> X0 = P0 + Q0 * XL0.GetTranslation();
+		const TVector<T, d> X1 = P1 + Q1 * XL1.GetTranslation();
+		const TRotation<T, d> R0 = Q0 * XL0.GetRotation();
+		TVector<T, d> CX_2 = GetLimitedPositionError(JointSettings, R0, X1 - X0);
+
+		// Increase in position error caused by the rotation correction
+		TVector<T, d> DX = CX_2 - CX_1;
+		for (int32 AxisIndex = 0; AxisIndex < d; ++AxisIndex)
+		{
+			if ((CX_1[AxisIndex] > 0) && (CX_2[AxisIndex] < CX_1[AxisIndex]))
+			{
+				DX[AxisIndex] = 0;
+			}
+			if ((CX_1[AxisIndex] < 0) && (CX_2[AxisIndex] > CX_1[AxisIndex]))
+			{
+				DX[AxisIndex] = 0;
+			}
+		}
+		const TVector<T, d> DX0 = DX - TVector<T, d>::DotProduct(DX, Axis0) * Axis0;
+		const TVector<T, d> DX1 = DX - TVector<T, d>::DotProduct(DX, Axis1) * Axis1;
+
+		// Correct the extra position error introduced by the rotation correction. We are effectively treating the
+		// bodies as if they have infinite inertia for this correction, which is only correct if
+		// the position correction exactly opposes the constraint that caused the increased position error
+		// which is only the case for uniform shapes (i.e., spherical inertia). To do this properly, we
+		// would have to solve for position and rotation correction simultaneously.
+		const TVector<T, d> DP0 = InvM0 * DX0 / (InvM0 + InvM1);
+		const TVector<T, d> DP1 = -InvM1 * DX1 / (InvM0 + InvM1);
+		ApplyPositionDelta(Dt, SolverSettings, JointSettings, Index0, Index1, P0, P1, DP0, DP1);
+	}
+
 	template<class T, int d>
 	void TPBDJointUtilities<T, d>::CalculateSwingConstraintSpace(
@@ -345,80 +704,57 @@
 	}
 
 	template<typename T, int d>
-	void ApplyPositionDelta(
+	void TPBDJointUtilities<T, d>::ApplyJointPositionConstraint(
 		const T Dt,
 		const TPBDJointSolverSettings<T, d>& SolverSettings,
 		const TPBDJointSettings<T, d>& JointSettings,
 		const int32 Index0,
 		const int32 Index1,
 		TVector<T, d>& P0,
-		TVector<T, d>& P1,
-		const TVector<T, d>& V0,
-		const TVector<T, d>& V1)
-	{
-		const float Stiffness = (SolverSettings.Stiffness > (T)0) ? SolverSettings.Stiffness : JointSettings.Motion.Stiffness;
-
-		P0 = P0 + Stiffness * V0;
-		P1 = P1 + Stiffness * V1;
-	}
-
-	template<typename T, int d>
-	void ApplyRotationDelta(
-		const T Dt,
-		const TPBDJointSolverSettings<T, d>& SolverSettings,
-		const TPBDJointSettings<T, d>& JointSettings,
-		const int32 Index0,
-		const int32 Index1,
 		TRotation<T, d>& Q0,
+		TVector<T, d>& P1,
 		TRotation<T, d>& Q1,
-		const TVector<T, d>& W0,
-		const TVector<T, d>& W1)
+		float InvM0,
+		const PMatrix<T, d, d>& InvIL0,
+		float InvM1,
+		const PMatrix<T, d, d>& InvIL1)
 	{
-		const float Stiffness = (SolverSettings.Stiffness > (T)0) ? SolverSettings.Stiffness : JointSettings.Motion.Stiffness;
-
-		const TRotation<T, d> DQ0 = (TRotation<T, d>::FromElements(Stiffness * W0, 0) * Q0) * (T)0.5;
-		const TRotation<T, d> DQ1 = (TRotation<T, d>::FromElements(Stiffness * W1, 0) * Q1) * (T)0.5;
-		Q0 = (Q0 + DQ0).GetNormalized();
-		Q1 = (Q1 + DQ1).GetNormalized();
-		Q1.EnforceShortestArcWith(Q0);
-	}
-
-	template<typename T, int d>
-	void ApplyRotationDelta(
-	    const T Dt,
-	    const TPBDJointSolverSettings<T, d>& SolverSettings,
-	    const TPBDJointSettings<T, d>& JointSettings,
-	    const int32 Index0,
-	    const int32 Index1,
-	    TVector<T, d>& P0,
-	    TRotation<T, d>& Q0,
-	    TVector<T, d>& P1,
-	    TRotation<T, d>& Q1,
-	    float InvM0,
-	    const PMatrix<T, d, d>& InvIL0,
-	    float InvM1,
-	    const PMatrix<T, d, d>& InvIL1,
-	    const TVector<T, d>& Axis0,
-	    const float Angle0,
-	    const TVector<T, d>& Axis1,
-	    const float Angle1)
-	{
+		const TRigidTransform<T, d>& XL0 = JointSettings.ConstraintFrames[Index0];
+		const TRigidTransform<T, d>& XL1 = JointSettings.ConstraintFrames[Index1];
+		const TVector<T, d> X0 = P0 + Q0 * XL0.GetTranslation();
+		const TVector<T, d> X1 = P1 + Q1 * XL1.GetTranslation();
+		const TRotation<T, d> R0 = Q0 * XL0.GetRotation();
 		const PMatrix<T, d, d> InvI0 = Utilities::Multiply(Q0.ToMatrix(), Utilities::Multiply(InvIL0, Q0.ToMatrix().GetTransposed()));
 		const PMatrix<T, d, d> InvI1 = Utilities::Multiply(Q1.ToMatrix(), Utilities::Multiply(InvIL1, Q1.ToMatrix().GetTransposed()));
 
-		//const T I0 = TVector<T, d>::DotProduct(Axis0, Utilities::Multiply(InvI0, Axis0));
-		//const T I1 = TVector<T, d>::DotProduct(Axis1, Utilities::Multiply(InvI1, Axis1));
-		//const TVector<T, d> W0 = Axis0 * Angle0 * I0 / (I0 + I1);
-		//const TVector<T, d> W1 = Axis1 * Angle1 * I1 / (I0 + I1);
-		const T L = (T)1 / (TVector<T, d>::DotProduct(Axis0, Utilities::Multiply(InvI0, Axis0)) + TVector<T, d>::DotProduct(Axis1, Utilities::Multiply(InvI1, Axis1)));
-		const TVector<T, d> W0 = Utilities::Multiply(InvI0, Axis0) * L * Angle0;
-		const TVector<T, d> W1 = Utilities::Multiply(InvI1, Axis1) * L * Angle1;
+		// Calculate constraint error
+		const TVector<T, d> CX = GetLimitedPositionError(JointSettings, R0, X1 - X0);
 
-		ApplyRotationDelta(Dt, SolverSettings, JointSettings, Index0, Index1, Q0, Q1, W0, W1);
+		// Calculate constraint correction
+		PMatrix<T, d, d> M0 = PMatrix<T, d, d>(0, 0, 0);
+		PMatrix<T, d, d> M1 = PMatrix<T, d, d>(0, 0, 0);
+		if (InvM0 > 0)
+		{
+			M0 = Utilities::ComputeJointFactorMatrix(X0 - P0, InvI0, InvM0);
+		}
+		if (InvM1 > 0)
+		{
+			M1 = Utilities::ComputeJointFactorMatrix(X1 - P1, InvI1, InvM1);
+		}
+		const PMatrix<T, d, d> MI = (M0 + M1).Inverse();
+		const TVector<T, d> DX = Utilities::Multiply(MI, CX);
+
+		// Apply constraint correction
+		const TVector<T, d> DP0 = InvM0 * DX;
+		const TVector<T, d> DP1 = -InvM1 * DX;
+		const TVector<T, d> DR0 = Utilities::Multiply(InvI0, TVector<T, d>::CrossProduct(X0 - P0, DX));
+		const TVector<T, d> DR1 = Utilities::Multiply(InvI1, TVector<T, d>::CrossProduct(X1 - P1, -DX));
+		ApplyPositionDelta(Dt, SolverSettings, JointSettings, Index0, Index1, P0, P1, DP0, DP1);
+		ApplyRotationDelta(Dt, SolverSettings, JointSettings, Index0, Index1, Q0, Q1, DR0, DR1);
 	}
 
 	template<typename T, int d>
-	void TPBDJointUtilities<T, d>::ApplyJointPositionConstraint(
+	void TPBDJointUtilities<T, d>::ApplyJointVelocityConstraint(
 		const T Dt,
 		const TPBDJointSolverSettings<T, d>& SolverSettings,
 		const TPBDJointSettings<T, d>& JointSettings,
@@ -426,8 +762,12 @@
 		const int32 Index1,
 		TVector<T, d>& P0,
 		TRotation<T, d>& Q0,
+		TVector<T, d>& V0,
+		TVector<T, d>& W0,
 		TVector<T, d>& P1,
 		TRotation<T, d>& Q1,
+		TVector<T, d>& V1,
+		TVector<T, d>& W1,
 		float InvM0,
 		const PMatrix<T, d, d>& InvIL0,
 		float InvM1,
@@ -435,40 +775,43 @@
 	{
 		const TRigidTransform<T, d>& XL0 = JointSettings.ConstraintFrames[Index0];
 		const TRigidTransform<T, d>& XL1 = JointSettings.ConstraintFrames[Index1];
-		const TVector<T, d> X0 = P0 + Q0 * XL0.GetTranslation();
-		const TVector<T, d> X1 = P1 + Q1 * XL1.GetTranslation();
+		const TVector<T, d> XC0 = Q0 * XL0.GetTranslation();
+		const TVector<T, d> XC1 = Q1 * XL1.GetTranslation();
 		const TRotation<T, d> R0 = Q0 * XL0.GetRotation();
 		const PMatrix<T, d, d> InvI0 = Utilities::Multiply(Q0.ToMatrix(), Utilities::Multiply(InvIL0, Q0.ToMatrix().GetTransposed()));
 		const PMatrix<T, d, d> InvI1 = Utilities::Multiply(Q1.ToMatrix(), Utilities::Multiply(InvIL1, Q1.ToMatrix().GetTransposed()));
 
-		// Calculate constraint error
-		const TVector<T, d> CX = GetLimitedPositionError(JointSettings, R0, X1 - X0);
+		const TVector<T, d> VC0 = V0 + TVector<T, d>::CrossProduct(W0, XC0);
+		const TVector<T, d> VC1 = V1 + TVector<T, d>::CrossProduct(W1, XC1);
+		const TVector<T, d> VC = GetLimitedVelocityError(JointSettings, R0, (P1 + XC1) - (P0 + XC0), VC1 - VC0);
 
 		// Calculate constraint correction
 		PMatrix<T, d, d> M0 = PMatrix<T, d, d>(0, 0, 0);
 		PMatrix<T, d, d> M1 = PMatrix<T, d, d>(0, 0, 0);
 		if (InvM0 > 0)
 		{
-			M0 = Utilities::ComputeJointFactorMatrix(X0 - P0, InvI0, InvM0);
+			M0 = Utilities::ComputeJointFactorMatrix(XC0, InvI0, InvM0);
 		}
 		if (InvM1 > 0)
 		{
-			M1 = Utilities::ComputeJointFactorMatrix(X1 - P1, InvI1, InvM1);
+			M1 = Utilities::ComputeJointFactorMatrix(XC1, InvI1, InvM1);
 		}
 		const PMatrix<T, d, d> MI = (M0 + M1).Inverse();
-		const TVector<T, d> DX = Utilities::Multiply(MI, CX);
+		const TVector<T, d> DL = Utilities::Multiply(MI, VC);
 
 		// Apply constraint correction
-		const TVector<T, d> V0 = InvM0 * DX;
-		const TVector<T, d> V1 = -InvM1 * DX;
-		const TVector<T, d> W0 = Utilities::Multiply(InvI0, TVector<T, d>::CrossProduct(X0 - P0, DX));
-		const TVector<T, d> W1 = Utilities::Multiply(InvI1, TVector<T, d>::CrossProduct(X1 - P1, -DX));
-		ApplyPositionDelta(Dt, SolverSettings, JointSettings, Index0, Index1, P0, P1, V0, V1);
-		ApplyRotationDelta(Dt, SolverSettings, JointSettings, Index0, Index1, Q0, Q1, W0, W1);
+		const TVector<T, d> DV0 = InvM0 * DL;
+		const TVector<T, d> DV1 = -InvM1 * DL;
+		const TVector<T, d> DW0 = Utilities::Multiply(InvI0, TVector<T, d>::CrossProduct(XC0, DL));
+		const TVector<T, d> DW1 = -Utilities::Multiply(InvI1, TVector<T, d>::CrossProduct(XC1, DL));
+
+		ApplyVelocityDelta(Dt, SolverSettings, JointSettings, Index0, Index1, P0, V0, P1, V1, DV0, DV1);
+		ApplyRotationVelocityDelta(Dt, SolverSettings, JointSettings, Index0, Index1, Q0, W0, Q1, W1, DW0, DW1);
 	}
 
+
 	template<typename T, int d>
-	void ApplyPostRotationPositionCorrection(
+	void TPBDJointUtilities<T, d>::ApplyJointTwistConstraint(
 		const T Dt,
 		const TPBDJointSolverSettings<T, d>& SolverSettings,
 		const TPBDJointSettings<T, d>& JointSettings,
@@ -481,52 +824,70 @@
 		float InvM0,
 		const PMatrix<T, d, d>& InvIL0,
 		float InvM1,
-		const PMatrix<T, d, d>& InvIL1,
-		const TVector<T, d>& CX_1,
-		const TVector<T, d>& Axis0,
-		const TVector<T, d>& Axis1)
+		const PMatrix<T, d, d>& InvIL1)
 	{
-		if (!SolverSettings.bEnablePositionCorrection)
-		{
-			return;
-		}
-
-		// Post-rotation constraint positions
 		const TRigidTransform<T, d>& XL0 = JointSettings.ConstraintFrames[Index0];
 		const TRigidTransform<T, d>& XL1 = JointSettings.ConstraintFrames[Index1];
 		const TVector<T, d> X0 = P0 + Q0 * XL0.GetTranslation();
 		const TVector<T, d> X1 = P1 + Q1 * XL1.GetTranslation();
 		const TRotation<T, d> R0 = Q0 * XL0.GetRotation();
-		TVector<T, d> CX_2 = GetLimitedPositionError(JointSettings, R0, X1 - X0);
+		const TRotation<T, d> R1 = Q1 * XL1.GetRotation();
+		const TVector<T, d> CX = GetLimitedPositionError(JointSettings, R0, X1 - X0);
 
-		// Increase in position error caused by the rotation correction
-		TVector<T, d> DX = CX_2 - CX_1;
-		for (int32 AxisIndex = 0; AxisIndex < d; ++AxisIndex)
+		// Calculate the Twist Axis and Angle for each body
+		const TRotation<T, d> R01 = R0.Inverse() * R1;
+		TRotation<T, d> R01Twist, R01Swing;
+		R01.ToSwingTwist(TJointConstants<T, d>::TwistAxis(), R01Swing, R01Twist);
+		R01Swing = R01Swing.GetNormalized();
+		R01Twist = R01Twist.GetNormalized();
+
+		TVector<T, d> TwistAxis01;
+		T TwistAngle;
+		R01Twist.ToAxisAndAngleSafe(TwistAxis01, TwistAngle, TJointConstants<T, d>::TwistAxis(), SolverSettings.SwingTwistAngleTolerance);
+		if (TwistAngle > PI)
 		{
-			if ((CX_1[AxisIndex] > 0) && (CX_2[AxisIndex] < CX_1[AxisIndex]))
-			{
-				DX[AxisIndex] = 0;
-			}
-			if ((CX_1[AxisIndex] < 0) && (CX_2[AxisIndex] > CX_1[AxisIndex]))
-			{
-				DX[AxisIndex] = 0;
-			}
+			TwistAngle = TwistAngle - (T)2 * PI;
 		}
-		const TVector<T, d> DX0 = DX - TVector<T, d>::DotProduct(DX, Axis0) * Axis0;
-		const TVector<T, d> DX1 = DX - TVector<T, d>::DotProduct(DX, Axis1) * Axis1;
+		if (TVector<T, d>::DotProduct(TwistAxis01, TJointConstants<T, d>::TwistAxis()) < 0)
+		{
+			TwistAxis01 = -TwistAxis01;
+			TwistAngle = -TwistAngle;
+		}
 
-		// Correct the extra position error introduced by the rotation correction. We are effectively treating the
-		// bodies as if they have infinite inertia for this correction, which is only correct if
-		// the position correction exactly opposes the constraint that caused the increased position error
-		// which is only the case for uniform shapes (i.e., spherical inertia). To do this properly, we
-		// would have to solve for position and rotation correction simultaneously.
-		const TVector<T, d> V0 = InvM0 * DX0 / (InvM0 + InvM1);
-		const TVector<T, d> V1 = -InvM1 * DX1 / (InvM0 + InvM1);
-		ApplyPositionDelta(Dt, SolverSettings, JointSettings, Index0, Index1, P0, P1, V0, V1);
+		const TVector<T, d> TwistAxis0 = R0 * TwistAxis01;
+		const TVector<T, d> TwistAxis1 = R1 * TwistAxis01;
+		T TwistAngleMax = FLT_MAX;
+		if (JointSettings.Motion.AngularMotionTypes[(int32)EJointAngularConstraintIndex::Twist] == EJointMotionType::Limited)
+		{
+			TwistAngleMax = JointSettings.Motion.AngularLimits[(int32)EJointAngularConstraintIndex::Twist];
+		}
+		else if (JointSettings.Motion.AngularMotionTypes[(int32)EJointAngularConstraintIndex::Twist] == EJointMotionType::Locked)
+		{
+			TwistAngleMax = 0;
+		}
+
+		// Calculate the twist correction to apply to each body
+		T DTwistAngle = 0;
+		if (TwistAngle > TwistAngleMax)
+		{
+			DTwistAngle = TwistAngle - TwistAngleMax;
+		}
+		else if (TwistAngle < -TwistAngleMax)
+		{
+			DTwistAngle = TwistAngle + TwistAngleMax;
+		}
+		const T DTwistAngle0 = DTwistAngle;
+		const T DTwistAngle1 = -DTwistAngle;
+
+		// Apply twist correction
+		ApplyRotationDelta(Dt, SolverSettings, JointSettings, Index0, Index1, P0, Q0, P1, Q1, InvM0, InvIL0, InvM1, InvIL1, TwistAxis0, DTwistAngle0, TwistAxis1, DTwistAngle1);
+
+		// Correct any positional error we may have introduced
+		ApplyPostRotationPositionCorrection<T, d>(Dt, SolverSettings, JointSettings, Index0, Index1, P0, Q0, P1, Q1, InvM0, InvIL0, InvM1, InvIL1, CX, TwistAxis0, TwistAxis1);
 	}
 
 	template<typename T, int d>
-	void TPBDJointUtilities<T, d>::ApplyJointTwistConstraint(
+	void TPBDJointUtilities<T, d>::ApplyJointTwistVelocityConstraint(
 		const T Dt,
 		const TPBDJointSolverSettings<T, d>& SolverSettings,
 		const TPBDJointSettings<T, d>& JointSettings,
@@ -534,8 +895,12 @@
 		const int32 Index1,
 		TVector<T, d>& P0,
 		TRotation<T, d>& Q0,
+		TVector<T, d>& V0,
+		TVector<T, d>& W0,
 		TVector<T, d>& P1,
 		TRotation<T, d>& Q1,
+		TVector<T, d>& V1,
+		TVector<T, d>& W1,
 		float InvM0,
 		const PMatrix<T, d, d>& InvIL0,
 		float InvM1,
@@ -571,6 +936,7 @@
 
 		const TVector<T, d> TwistAxis0 = R0 * TwistAxis01;
 		const TVector<T, d> TwistAxis1 = R1 * TwistAxis01;
+
 		T TwistAngleMax = FLT_MAX;
 		if (JointSettings.Motion.AngularMotionTypes[(int32)EJointAngularConstraintIndex::Twist] == EJointMotionType::Limited)
 		{
@@ -581,26 +947,22 @@
 			TwistAngleMax = 0;
 		}
 
-		// Calculate the twist correction to applt to each body
-		T DTwistAngle = 0;
+		const T WC0 = TVector<T, d>::DotProduct(W0, TwistAxis0);
+		const T WC1 = TVector<T, d>::DotProduct(W1, TwistAxis1);
+		T DW = (T)0;
 		if (TwistAngle > TwistAngleMax)
 		{
-			DTwistAngle = TwistAngle - TwistAngleMax;
+			DW = FMath::Max((T)0, WC1 - WC0);
 		}
 		else if (TwistAngle < -TwistAngleMax)
 		{
-			DTwistAngle = TwistAngle + TwistAngleMax;
+			DW = FMath::Min((T)0, WC1 - WC0);
 		}
-		const T DTwistAngle0 = DTwistAngle;
-		const T DTwistAngle1 = -DTwistAngle;
 
-		// Apply twist correction
-		ApplyRotationDelta(Dt, SolverSettings, JointSettings, Index0, Index1, P0, Q0, P1, Q1, InvM0, InvIL0, InvM1, InvIL1, TwistAxis0, DTwistAngle0, TwistAxis1, DTwistAngle1);
-
-		// Correct any positional error we may have introduced
-		ApplyPostRotationPositionCorrection<T, d>(Dt, SolverSettings, JointSettings, Index0, Index1, P0, Q0, P1, Q1, InvM0, InvIL0, InvM1, InvIL1, CX, TwistAxis0, TwistAxis1);
+		ApplyRotationVelocityDelta(Dt, SolverSettings, JointSettings, Index0, Index1, Q0, W0, Q1, W1, InvM0, InvIL0, InvM1, InvIL1, TwistAxis0, TwistAxis1, DW);
 	}
 
+
 	template<typename T, int d>
 	void TPBDJointUtilities<T, d>::ApplyJointConeConstraint(
 		const T Dt,
@@ -681,5 +1043,84 @@
 	}
 
 	template<typename T, int d>
+	void TPBDJointUtilities<T, d>::ApplyJointConeVelocityConstraint(
+		const T Dt,
+		const TPBDJointSolverSettings<T, d>& SolverSettings,
+		const TPBDJointSettings<T, d>& JointSettings,
+		const int32 Index0,
+		const int32 Index1,
+		TVector<T, d>& P0,
+		TRotation<T, d>& Q0,
+		TVector<T, d>& V0,
+		TVector<T, d>& W0,
+		TVector<T, d>& P1,
+		TRotation<T, d>& Q1,
+		TVector<T, d>& V1,
+		TVector<T, d>& W1,
+		float InvM0,
+		const PMatrix<T, d, d>& InvIL0,
+		float InvM1,
+		const PMatrix<T, d, d>& InvIL1)
+	{
+		const TRigidTransform<T, d>& XL0 = JointSettings.ConstraintFrames[Index0];
+		const TRigidTransform<T, d>& XL1 = JointSettings.ConstraintFrames[Index1];
+		const TVector<T, d> X0 = P0 + Q0 * XL0.GetTranslation();
+		const TVector<T, d> X1 = P1 + Q1 * XL1.GetTranslation();
+		const TRotation<T, d> R0 = Q0 * XL0.GetRotation();
+		const TRotation<T, d> R1 = Q1 * XL1.GetRotation();
+
+		// Calculate Swing axis for each body
+		const TRotation<T, d> R01 = R0.Inverse() * R1;
+		TRotation<T, d> R01Twist, R01Swing;
+		R01.ToSwingTwist(TJointConstants<T, d>::TwistAxis(), R01Swing, R01Twist);
+		R01Swing = R01Swing.GetNormalized();
+		R01Twist = R01Twist.GetNormalized();
+
+		TVector<T, d> SwingAxis01;
+		T SwingAngle;
+		R01Swing.ToAxisAndAngleSafe(SwingAxis01, SwingAngle, TJointConstants<T, d>::Swing1Axis(), SolverSettings.SwingTwistAngleTolerance);
+		if (SwingAngle > PI)
+		{
+			SwingAngle = SwingAngle - (T)2 * PI;
+		}
+
+		const TVector<T, d> SwingAxis0 = R0 * SwingAxis01;
+		const TVector<T, d> SwingAxis1 = SwingAxis0;
+
+		// Calculate swing limit for the current swing axis
+		T SwingAngleMax = FLT_MAX;
+		const T Swing1Limit = JointSettings.Motion.AngularLimits[(int32)EJointAngularConstraintIndex::Swing1];
+		const T Swing2Limit = JointSettings.Motion.AngularLimits[(int32)EJointAngularConstraintIndex::Swing2];
+
+		// Circular swing limit
+		SwingAngleMax = Swing1Limit;
+
+		// Elliptical swing limit
+		if (!FMath::IsNearlyEqual(Swing1Limit, Swing2Limit, KINDA_SMALL_NUMBER))
+		{
+			// Map swing axis to ellipse and calculate limit for this swing axis
+			const T DotSwing1 = FMath::Abs(TVector<T, d>::DotProduct(SwingAxis01, TJointConstants<T, d>::Swing1Axis()));
+			const T DotSwing2 = FMath::Abs(TVector<T, d>::DotProduct(SwingAxis01, TJointConstants<T, d>::Swing2Axis()));
+			SwingAngleMax = FMath::Sqrt(Swing1Limit * DotSwing2 * Swing1Limit * DotSwing2 + Swing2Limit * DotSwing1 * Swing2Limit * DotSwing1);
+		}
+
+		// Only clamp veloicity if we are outside the limits and moving to increase the error
+		const T WC0 = TVector<T, d>::DotProduct(W0, SwingAxis0);
+		const T WC1 = TVector<T, d>::DotProduct(W1, SwingAxis1);
+		T DW = 0;
+		if (SwingAngle > SwingAngleMax)
+		{
+			DW = FMath::Max((T)0, WC1 - WC0);
+		}
+		else if (SwingAngle < -SwingAngleMax)
+		{
+			DW = FMath::Min((T)0, WC1 - WC0);
+		}
+
+		ApplyRotationVelocityDelta(Dt, SolverSettings, JointSettings, Index0, Index1, Q0, W0, Q1, W1, InvM0, InvIL0, InvM1, InvIL1, SwingAxis0, SwingAxis1, DW);
+	}
+
+
+	template<typename T, int d>
 	void TPBDJointUtilities<T, d>::ApplyJointSwingConstraint(
 		const T Dt,
@@ -777,5 +1218,100 @@
 	}
 
 	template<typename T, int d>
+	void TPBDJointUtilities<T, d>::ApplyJointSwingVelocityConstraint(
+		const T Dt,
+		const TPBDJointSolverSettings<T, d>& SolverSettings,
+		const TPBDJointSettings<T, d>& JointSettings,
+		const int32 Index0,
+		const int32 Index1,
+		const EJointAngularConstraintIndex SwingConstraint,
+		TVector<T, d>& P0,
+		TRotation<T, d>& Q0,
+		TVector<T, d>& V0,
+		TVector<T, d>& W0,
+		TVector<T, d>& P1,
+		TRotation<T, d>& Q1,
+		TVector<T, d>& V1,
+		TVector<T, d>& W1,
+		float InvM0,
+		const PMatrix<T, d, d>& InvIL0,
+		float InvM1,
+		const PMatrix<T, d, d>& InvIL1)
+	{
+		const TRigidTransform<T, d>& XL0 = JointSettings.ConstraintFrames[Index0];
+		const TRigidTransform<T, d>& XL1 = JointSettings.ConstraintFrames[Index1];
+		const TVector<T, d> X0 = P0 + Q0 * XL0.GetTranslation();
+		const TVector<T, d> X1 = P1 + Q1 * XL1.GetTranslation();
+		const TRotation<T, d> R0 = Q0 * XL0.GetRotation();
+		const TRotation<T, d> R1 = Q1 * XL1.GetRotation();
+
+		// Calculate the swing axis for each body
+		const TRotation<T, d> R01 = R0.Inverse() * R1;
+		TRotation<T, d> R01Twist, R01Swing;
+		R01.ToSwingTwist(TJointConstants<T, d>::TwistAxis(), R01Swing, R01Twist);
+		R01Swing = R01Swing.GetNormalized();
+		R01Twist = R01Twist.GetNormalized();
+
+		TVector<T, d> TwistAxis01;
+		T TwistAngle;
+		R01Twist.ToAxisAndAngleSafe(TwistAxis01, TwistAngle, TJointConstants<T, d>::TwistAxis(), SolverSettings.SwingTwistAngleTolerance);
+		if (TwistAngle > PI)
+		{
+			TwistAngle = TwistAngle - (T)2 * PI;
+		}
+		if (TVector<T, d>::DotProduct(TwistAxis01, TJointConstants<T, d>::TwistAxis()) < 0)
+		{
+			TwistAxis01 = -TwistAxis01;
+			TwistAngle = -TwistAngle;
+		}
+		const TVector<T, d> TwistAxis = R0 * TwistAxis01;
+
+		const PMatrix<T, d, d> Axes0 = R0.ToMatrix();
+		const PMatrix<T, d, d> Axes1 = R1.ToMatrix();
+		TVector<T, d> SwingCross = TVector<T, d>::CrossProduct(Axes0.GetAxis((int32)SwingConstraint), Axes1.GetAxis((int32)SwingConstraint));
+		SwingCross = SwingCross - TVector<T, d>::DotProduct(TwistAxis, SwingCross) * TwistAxis;
+		const T SwingCrossLen = SwingCross.Size();
+		if (SwingCrossLen > KINDA_SMALL_NUMBER)
+		{
+			const TVector<T, d> SwingAxis = SwingCross / SwingCrossLen;
+			TVector<T, d> SwingAxis0 = SwingAxis;
+			TVector<T, d> SwingAxis1 = SwingAxis;
+
+			T SwingAngle = FMath::Asin(FMath::Clamp(SwingCrossLen, (T)0, (T)1));
+			const T SwingDot = TVector<T, d>::DotProduct(Axes0.GetAxis((int32)SwingConstraint), Axes1.GetAxis((int32)SwingConstraint));
+			if (SwingDot < (T)0)
+			{
+				SwingAngle = (T)PI - SwingAngle;
+			}
+
+			T SwingAngleMax = FLT_MAX;
+			if (JointSettings.Motion.AngularMotionTypes[(int32)SwingConstraint] == EJointMotionType::Limited)
+			{
+				T Swing1Limit = JointSettings.Motion.AngularLimits[(int32)SwingConstraint];
+				SwingAngleMax = Swing1Limit;
+			}
+			else if (JointSettings.Motion.AngularMotionTypes[(int32)SwingConstraint] == EJointMotionType::Locked)
+			{
+				SwingAngleMax = 0;
+			}
+
+			// Only clamp veloicity if we are outside the limits and moving to increase the error
+			const T WC0 = TVector<T, d>::DotProduct(W0, SwingAxis0);
+			const T WC1 = TVector<T, d>::DotProduct(W1, SwingAxis1);
+			T DW = 0;
+			if (SwingAngle > SwingAngleMax)
+			{
+				DW = FMath::Max((T)0, WC1 - WC0);
+			}
+			else if (SwingAngle < -SwingAngleMax)
+			{
+				DW = FMath::Min((T)0, WC1 - WC0);
+			}
+
+			ApplyRotationVelocityDelta(Dt, SolverSettings, JointSettings, Index0, Index1, Q0, W0, Q1, W1, InvM0, InvIL0, InvM1, InvIL1, SwingAxis0, SwingAxis1, DW);
+		}
+	}
+
+	template<typename T, int d>
 	void TPBDJointUtilities<T, d>::ApplyJointTwistDrive(
 		const T Dt,
@@ -972,10 +1508,10 @@
 		const TRotation<T, d> R0 = Q0 * XL0.GetRotation();
 		const TVector<T, d> CX = GetLimitedPositionError(JointSettings, R0, X1 - X0);
 
-		const TVector<T, d> V0 = CX * (ProjectionFactor * InvM0 / (InvM0 + InvM1));
-		const TVector<T, d> V1 = -CX * (ProjectionFactor * InvM1 / (InvM0 + InvM1));
-		P0 = P0 + V0;
-		P1 = P1 + V1;
+		const TVector<T, d> DP0 = CX * (ProjectionFactor * InvM0 / (InvM0 + InvM1));
+		const TVector<T, d> DP1 = -CX * (ProjectionFactor * InvM1 / (InvM0 + InvM1));
+		P0 = P0 + DP0;
+		P1 = P1 + DP1;
 	}
 
 	template<typename T, int d>
diff --git a/Engine/Source/Runtime/Experimental/Chaos/Private/Chaos/PBDJointConstraints.cpp b/Engine/Source/Runtime/Experimental/Chaos/Private/Chaos/PBDJointConstraints.cpp
--- a/Engine/Source/Runtime/Experimental/Chaos/Private/Chaos/PBDJointConstraints.cpp
+++ b/Engine/Source/Runtime/Experimental/Chaos/Private/Chaos/PBDJointConstraints.cpp
@@ -116,15 +116,15 @@ namespace Chaos
 		: SwingTwistAngleTolerance((T)1.0e-6)
 		, MinParentMassRatio((T)0.5)
 		, MaxInertiaRatio((T)5.0)
-		, bProjectPostApply(true)
 		, bEnableLinearLimits(true)
 		, bEnableTwistLimits(true)
 		, bEnableSwingLimits(true)
-		, bEnablePositionCorrection(true)
+		, bEnablePositionCorrection(false)
 		, bEnableDrives(true)
 		, Projection((T)0)
 		, Stiffness((T)0)
 		, DriveStiffness((T)0)
+		, PositionIterations((T)1)
 	{
 	}
 
@@ -347,12 +347,7 @@ namespace Chaos
 
 		for (FConstraintHandle* ConstraintHandle : SortedConstraintHandles)
 		{
-			ApplySingle(Dt, ConstraintHandle->GetConstraintIndex(), It, NumIts);
-		}
-
-		if (!Settings.bProjectPostApply)
-		{
-			ApplyProjection(Dt, InConstraintHandles);
+			SolveVelocity(Dt, ConstraintHandle->GetConstraintIndex(), It, NumIts);
 		}
 
 		if (PostApplyCallback != nullptr)
@@ -363,31 +358,33 @@ namespace Chaos
 
 	template<class T, int d>
 	void TPBDJointConstraints<T, d>::ApplyPushOut(const T Dt, const TArray<FConstraintHandle*>& InConstraintHandles)
-	{
-		if (Settings.bProjectPostApply)
-		{
-			ApplyProjection(Dt, InConstraintHandles);
-		}
-	}
-
-	template<class T, int d>
-	void TPBDJointConstraints<T, d>::ApplyProjection(const T Dt, const TArray<FConstraintHandle*>& InConstraintHandles)
 	{
 		TArray<FConstraintHandle*> SortedConstraintHandles = InConstraintHandles;
 		SortedConstraintHandles.Sort([](const FConstraintHandle& L, const FConstraintHandle& R)
+		{
+			// Sort bodies from root to leaf
+			return L.GetConstraintLevel() < R.GetConstraintLevel();
+		});
+
+		for (int32 It = 0; It < Settings.PositionIterations; ++It)
+		{
+			for (FConstraintHandle* ConstraintHandle : SortedConstraintHandles)
 			{
-				// Sort bodies from root to leaf
-				return L.GetConstraintLevel() < R.GetConstraintLevel();
-			});
+				SolvePosition(Dt, ConstraintHandle->GetConstraintIndex(), It, Settings.PositionIterations);
+			}
+		}
 
-		for (FConstraintHandle* ConstraintHandle : SortedConstraintHandles)
+		if (Settings.Projection > 0)
 		{
-			ApplyProjectionSingle(Dt, ConstraintHandle->GetConstraintIndex(), 0, 1);
+			for (FConstraintHandle* ConstraintHandle : SortedConstraintHandles)
+			{
+				ProjectPosition(Dt, ConstraintHandle->GetConstraintIndex(), 0, 1);
+			}
 		}
 	}
 
 	template<class T, int d>
-	void TPBDJointConstraints<T, d>::ApplySingle(const T Dt, const int32 ConstraintIndex, const int32 It, const int32 NumIts)
+	void TPBDJointConstraints<T, d>::SolveVelocity(const T Dt, const int32 ConstraintIndex, const int32 It, const int32 NumIts)
 	{
 		const TVector<TGeometryParticleHandle<T, d>*, 2>& Constraint = ConstraintParticles[ConstraintIndex];
 		UE_LOG(LogChaosJoint, Verbose, TEXT("Solve Joint Constraint %d %s %s (dt = %f; it = %d / %d)"), ConstraintIndex, *Constraint[0]->ToString(), *Constraint[1]->ToString(), Dt, It, NumIts);
@@ -404,8 +401,12 @@ namespace Chaos
 
 		TVector<T, d> P0 = Particle0->P();
 		TRotation<T, d> Q0 = Particle0->Q();
+		TVector<T, d> V0 = Particle0->V();
+		TVector<T, d> W0 = Particle0->W();
 		TVector<T, d> P1 = Particle1->P();
 		TRotation<T, d> Q1 = Particle1->Q();
+		TVector<T, d> V1 = Particle1->V();
+		TVector<T, d> W1 = Particle1->W();
 		float InvM0 = Particle0->InvM();
 		float InvM1 = Particle1->InvM();
 		PMatrix<T, d, d> InvIL0 = Particle0->InvI();
@@ -414,22 +415,19 @@ namespace Chaos
 		Q1.EnforceShortestArcWith(Q0);
 
 		// Adjust mass for stability
-		if (Rigid0 && Rigid1)
+		const int32 Level0 = ConstraintStates[ConstraintIndex].ParticleLevels[Index0];
+		const int32 Level1 = ConstraintStates[ConstraintIndex].ParticleLevels[Index1];
+		if (Level0 < Level1)
 		{
-			const int32 Level0 = ConstraintStates[ConstraintIndex].ParticleLevels[Index0];
-			const int32 Level1 = ConstraintStates[ConstraintIndex].ParticleLevels[Index1];
-			if (Level0 < Level1)
-			{
-				TPBDJointUtilities<T, d>::GetConditionedInverseMass(Rigid0, Rigid1, InvM0, InvM1, InvIL0, InvIL1, Settings.MinParentMassRatio, Settings.MaxInertiaRatio);
-			}
-			else if (Level0 > Level1)
-			{
-				TPBDJointUtilities<T, d>::GetConditionedInverseMass(Rigid1, Rigid0, InvM1, InvM0, InvIL1, InvIL0, Settings.MinParentMassRatio, Settings.MaxInertiaRatio);
-			}
-			else
-			{
-				TPBDJointUtilities<T, d>::GetConditionedInverseMass(Rigid1, Rigid0, InvM1, InvM0, InvIL1, InvIL0, (T)0, Settings.MaxInertiaRatio);
-			}
+			TPBDJointUtilities<T, d>::GetConditionedInverseMass(Particle0->M(), Particle0->I().GetDiagonal(), Particle1->M(), Particle1->I().GetDiagonal(), InvM0, InvM1, InvIL0, InvIL1, Settings.MinParentMassRatio, Settings.MaxInertiaRatio);
+		}
+		else if (Level0 > Level1)
+		{
+			TPBDJointUtilities<T, d>::GetConditionedInverseMass(Particle1->M(), Particle1->I().GetDiagonal(), Particle0->M(), Particle0->I().GetDiagonal(), InvM1, InvM0, InvIL1, InvIL0, Settings.MinParentMassRatio, Settings.MaxInertiaRatio);
+		}
+		else
+		{
+			TPBDJointUtilities<T, d>::GetConditionedInverseMass(Particle0->M(), Particle0->I().GetDiagonal(), Particle1->M(), Particle1->I().GetDiagonal(), InvM0, InvM1, InvIL0, InvIL1, (T)0, Settings.MaxInertiaRatio);
 		}
 
 		const TVector<EJointMotionType, d>& LinearMotion = JointSettings.Motion.LinearMotionTypes;
@@ -437,13 +435,7 @@ namespace Chaos
 		EJointMotionType Swing1Motion = JointSettings.Motion.AngularMotionTypes[(int32)EJointAngularConstraintIndex::Swing1];
 		EJointMotionType Swing2Motion = JointSettings.Motion.AngularMotionTypes[(int32)EJointAngularConstraintIndex::Swing2];
 
-		// Disable a constraint if it has any linear limits?
-		if (!Settings.bEnableLinearLimits && ((LinearMotion[0] == EJointMotionType::Limited) || (LinearMotion[1] == EJointMotionType::Limited) || (LinearMotion[2] == EJointMotionType::Limited)))
-		{
-			return;
-		}
-
-		// Apply angular drives
+		// Apply angular drives (NOTE: modifies position, not velocity)
 		if (Settings.bEnableDrives)
 		{
 			bool bTwistLocked = TwistMotion == EJointMotionType::Locked;
@@ -475,6 +467,127 @@ namespace Chaos
 			}
 		}
 
+		// Apply twist velocity constraint
+		if (Settings.bEnableTwistLimits)
+		{
+			if (TwistMotion != EJointMotionType::Free)
+			{
+				TPBDJointUtilities<T, d>::ApplyJointTwistVelocityConstraint(Dt, Settings, JointSettings, Index0, Index1, P0, Q0, V0, W0, P1, Q1, V1, W1, InvM0, InvIL0, InvM1, InvIL1);
+			}
+		}
+
+		// Apply swing velocity constraints
+		if (Settings.bEnableSwingLimits)
+		{
+			if ((Swing1Motion == EJointMotionType::Limited) && (Swing2Motion == EJointMotionType::Limited))
+			{
+				// Swing Cone
+				TPBDJointUtilities<T, d>::ApplyJointConeVelocityConstraint(Dt, Settings, JointSettings, Index0, Index1, P0, Q0, V0, W0, P1, Q1, V1, W1, InvM0, InvIL0, InvM1, InvIL1);
+			}
+			else
+			{
+				if (Swing1Motion != EJointMotionType::Free)
+				{
+					// Swing Arc/Lock
+					TPBDJointUtilities<T, d>::ApplyJointSwingVelocityConstraint(Dt, Settings, JointSettings, Index0, Index1, EJointAngularConstraintIndex::Swing1, P0, Q0, V0, W0, P1, Q1, V1, W1, InvM0, InvIL0, InvM1, InvIL1);
+				}
+				if (Swing2Motion != EJointMotionType::Free)
+				{
+					// Swing Arc/Lock
+					TPBDJointUtilities<T, d>::ApplyJointSwingVelocityConstraint(Dt, Settings, JointSettings, Index0, Index1, EJointAngularConstraintIndex::Swing2, P0, Q0, V0, W0, P1, Q1, V1, W1, InvM0, InvIL0, InvM1, InvIL1);
+				}
+			}
+		}
+
+		// Apply linear velocity  constraints
+		if ((LinearMotion[0] != EJointMotionType::Free) || (LinearMotion[1] != EJointMotionType::Free) || (LinearMotion[2] != EJointMotionType::Free))
+		{
+			TPBDJointUtilities<T, d>::ApplyJointVelocityConstraint(Dt, Settings, JointSettings, Index0, Index1, P0, Q0, V0, W0, P1, Q1, V1, W1, InvM0, InvIL0, InvM1, InvIL1);
+		}
+
+		// Update the particles
+		if (Rigid0)
+		{
+			Rigid0->SetP(P0);
+			Rigid0->SetQ(Q0);
+			Rigid0->SetV(V0);
+			Rigid0->SetW(W0);
+		}
+		if (Rigid1)
+		{
+			Rigid1->SetP(P1);
+			Rigid1->SetQ(Q1);
+			Rigid1->SetV(V1);
+			Rigid1->SetW(W1);
+		}
+	}
+
+	template<class T, int d>
+	void TPBDJointConstraints<T, d>::SolvePosition(const T Dt, const int32 ConstraintIndex, const int32 It, const int32 NumIts)
+	{
+		const TVector<TGeometryParticleHandle<T, d>*, 2>& Constraint = ConstraintParticles[ConstraintIndex];
+		UE_LOG(LogChaosJoint, Verbose, TEXT("Solve Joint Constraint %d %s %s (dt = %f; it = %d / %d)"), ConstraintIndex, *Constraint[0]->ToString(), *Constraint[1]->ToString(), Dt, It, NumIts);
+
+		const FJointSettings& JointSettings = ConstraintSettings[ConstraintIndex];
+
+		// Switch particles - internally we assume the first body is the parent (i.e., the space in which constraint limits are specified)
+		const int32 Index0 = 1;
+		const int32 Index1 = 0;
+		TGenericParticleHandle<T, d> Particle0 = TGenericParticleHandle<T, d>(ConstraintParticles[ConstraintIndex][Index0]);
+		TGenericParticleHandle<T, d> Particle1 = TGenericParticleHandle<T, d>(ConstraintParticles[ConstraintIndex][Index1]);
+		TPBDRigidParticleHandle<T, d>* Rigid0 = ConstraintParticles[ConstraintIndex][Index0]->AsDynamic();
+		TPBDRigidParticleHandle<T, d>* Rigid1 = ConstraintParticles[ConstraintIndex][Index1]->AsDynamic();
+
+		TVector<T, d> P0 = Particle0->P();
+		TRotation<T, d> Q0 = Particle0->Q();
+		TVector<T, d> P1 = Particle1->P();
+		TRotation<T, d> Q1 = Particle1->Q();
+		float InvM0 = Particle0->InvM();
+		float InvM1 = Particle1->InvM();
+		PMatrix<T, d, d> InvIL0 = Particle0->InvI();
+		PMatrix<T, d, d> InvIL1 = Particle1->InvI();
+
+		Q1.EnforceShortestArcWith(Q0);
+
+		// Adjust mass for stability
+		const int32 Level0 = ConstraintStates[ConstraintIndex].ParticleLevels[Index0];
+		const int32 Level1 = ConstraintStates[ConstraintIndex].ParticleLevels[Index1];
+		if (Level0 < Level1)
+		{
+			TPBDJointUtilities<T, d>::GetConditionedInverseMass(Particle0->M(), Particle0->I().GetDiagonal(), Particle1->M(), Particle1->I().GetDiagonal(), InvM0, InvM1, InvIL0, InvIL1, Settings.MinParentMassRatio, Settings.MaxInertiaRatio);
+		}
+		else if (Level0 > Level1)
+		{
+			TPBDJointUtilities<T, d>::GetConditionedInverseMass(Particle1->M(), Particle1->I().GetDiagonal(), Particle0->M(), Particle0->I().GetDiagonal(), InvM1, InvM0, InvIL1, InvIL0, Settings.MinParentMassRatio, Settings.MaxInertiaRatio);
+		}
+		else
+		{
+			TPBDJointUtilities<T, d>::GetConditionedInverseMass(Particle0->M(), Particle0->I().GetDiagonal(), Particle1->M(), Particle1->I().GetDiagonal(), InvM0, InvM1, InvIL0, InvIL1, (T)0, Settings.MaxInertiaRatio);
+		}
+
+		// Freeze the closest to kinematic connection if there is a difference
+		if (Level0 < Level1)
+		{
+			InvM0 = 0;
+			InvIL0 = PMatrix<T, d, d>(0, 0, 0);
+		}
+		else if (Level1 < Level0)
+		{
+			InvM1 = 0;
+			InvIL1 = PMatrix<T, d, d>(0, 0, 0);
+		}
+
+		const TVector<EJointMotionType, d>& LinearMotion = JointSettings.Motion.LinearMotionTypes;
+		EJointMotionType TwistMotion = JointSettings.Motion.AngularMotionTypes[(int32)EJointAngularConstraintIndex::Twist];
+		EJointMotionType Swing1Motion = JointSettings.Motion.AngularMotionTypes[(int32)EJointAngularConstraintIndex::Swing1];
+		EJointMotionType Swing2Motion = JointSettings.Motion.AngularMotionTypes[(int32)EJointAngularConstraintIndex::Swing2];
+
+		// Disable a constraint if it has any linear limits?
+		if (!Settings.bEnableLinearLimits && ((LinearMotion[0] == EJointMotionType::Limited) || (LinearMotion[1] == EJointMotionType::Limited) || (LinearMotion[2] == EJointMotionType::Limited)))
+		{
+			return;
+		}
+
 		// Apply twist constraint
 		if (Settings.bEnableTwistLimits)
 		{
@@ -527,7 +640,7 @@ namespace Chaos
 	}
 
 	template<class T, int d>
-	void TPBDJointConstraints<T, d>::ApplyProjectionSingle(const T Dt, const int32 ConstraintIndex, const int32 It, const int32 NumIts)
+	void TPBDJointConstraints<T, d>::ProjectPosition(const T Dt, const int32 ConstraintIndex, const int32 It, const int32 NumIts)
 	{
 		const FJointSettings& JointSettings = ConstraintSettings[ConstraintIndex];
 
diff --git a/Engine/Source/Runtime/Experimental/Chaos/Private/Chaos/PBDRigidsEvolutionGBF.cpp b/Engine/Source/Runtime/Experimental/Chaos/Private/Chaos/PBDRigidsEvolutionGBF.cpp
--- a/Engine/Source/Runtime/Experimental/Chaos/Private/Chaos/PBDRigidsEvolutionGBF.cpp
+++ b/Engine/Source/Runtime/Experimental/Chaos/Private/Chaos/PBDRigidsEvolutionGBF.cpp
@@ -239,9 +239,12 @@ TPBDRigidsEvolutionGBF<T, d>::TPBDRigidsEvolutionGBF(TPBDRigidsSOAs<T, d>& InPar
 	, PostApplyCallback(nullptr)
 	, PostApplyPushOutCallback(nullptr)
 {
-	SetParticleUpdateVelocityFunction([PBDUpdateRule = TPerParticlePBDUpdateFromDeltaPosition<float, 3>(), this](const TParticleView<TPBDRigidParticles<T, d>>& ParticlesInput, const T Dt) {
-		ParticlesInput.ParallelFor([&](auto& Particle, int32 Index) {
-			PBDUpdateRule.Apply(Particle, Dt);
+	SetParticleUpdateVelocityFunction([PBDUpdateRule = TPerParticlePBDUpdateFromDeltaPosition<float, 3>(), this](const TArray<TGeometryParticleHandle<T, d>*>& ParticlesInput, const T Dt) {
+		ParticlesParallelFor(ParticlesInput, [&](auto& Particle, int32 Index) {
+			if (Particle->AsDynamic())
+			{
+				PBDUpdateRule.Apply(Particle->AsDynamic(), Dt);
+			}
 		});
 	});
 
diff --git a/Engine/Source/Runtime/Experimental/Chaos/Public/Chaos/PBDJointConstraintTypes.h b/Engine/Source/Runtime/Experimental/Chaos/Public/Chaos/PBDJointConstraintTypes.h
--- a/Engine/Source/Runtime/Experimental/Chaos/Public/Chaos/PBDJointConstraintTypes.h
+++ b/Engine/Source/Runtime/Experimental/Chaos/Public/Chaos/PBDJointConstraintTypes.h
@@ -118,7 +118,6 @@ namespace Chaos
 		T MaxInertiaRatio;
 
 		// @todo(ccaulfield): remove these TEMP overrides for testing
-		bool bProjectPostApply;
 		bool bEnableLinearLimits;
 		bool bEnableTwistLimits;
 		bool bEnableSwingLimits;
@@ -127,5 +126,6 @@ namespace Chaos
 		T Projection;
 		T Stiffness;
 		T DriveStiffness;
+		int32 PositionIterations;
 	};
 }
\ No newline at end of file
diff --git a/Engine/Source/Runtime/Experimental/Chaos/Public/Chaos/PBDJointConstraintUtilities.h b/Engine/Source/Runtime/Experimental/Chaos/Public/Chaos/PBDJointConstraintUtilities.h
--- a/Engine/Source/Runtime/Experimental/Chaos/Public/Chaos/PBDJointConstraintUtilities.h
+++ b/Engine/Source/Runtime/Experimental/Chaos/Public/Chaos/PBDJointConstraintUtilities.h
@@ -44,18 +44,21 @@ namespace Chaos
 			const T MinRatio);
 
 		static CHAOS_API void GetConditionedInverseMass(
-			const TPBDRigidParticleHandle<T, d>* PParent, 
-			const TPBDRigidParticleHandle<T, d>* PChild, 
-			T& OutInvMParent, 
-			T& OutInvMChild, 
-			PMatrix<T, d, d>& OutInvIParent, 
-			PMatrix<T, d, d>& OutInvIChild, 
-			const T MinParentMassRatio, 
+			const float MParent,
+			const TVector<T, d> IParent,
+			const float MChild,
+			const TVector<T, d> IChild,
+			T& OutInvMParent,
+			T& OutInvMChild,
+			PMatrix<T, d, d>& OutInvIParent,
+			PMatrix<T, d, d>& OutInvIChild,
+			const T MinParentMassRatio,
 			const T MaxInertiaRatio);
 
 		static CHAOS_API void GetConditionedInverseMass(
-			const TPBDRigidParticleHandle<T, d>* P0, 
-			T& OutInvM0, 
+			const float M,
+			const TVector<T, d> I,
+			T& OutInvM0,
 			PMatrix<T, d, d>& OutInvI0, 
 			const T MaxInertiaRatio);
 
@@ -104,6 +107,25 @@ namespace Chaos
 			float InvM1,
 			const PMatrix<T, d, d>& InvIL1);
 
+		static CHAOS_API void ApplyJointVelocityConstraint(
+			const T Dt,
+			const TPBDJointSolverSettings<T, d>& SolverSettings,
+			const TPBDJointSettings<T, d>& JointSettings,
+			const int32 Index0,
+			const int32 Index1,
+			TVector<T, d>& P0,
+			TRotation<T, d>& Q0,
+			TVector<T, d>& V0,
+			TVector<T, d>& W0,
+			TVector<T, d>& P1,
+			TRotation<T, d>& Q1,
+			TVector<T, d>& V1,
+			TVector<T, d>& W1,
+			float InvM0,
+			const PMatrix<T, d, d>& InvIL0,
+			float InvM1,
+			const PMatrix<T, d, d>& InvIL1);
+
 		static CHAOS_API void ApplyJointTwistConstraint(
 			const T Dt,
 			const TPBDJointSolverSettings<T, d>& SolverSettings,
@@ -119,6 +141,25 @@ namespace Chaos
 			float InvM1,
 			const PMatrix<T, d, d>& InvIL1);
 
+		static CHAOS_API void ApplyJointTwistVelocityConstraint(
+			const T Dt,
+			const TPBDJointSolverSettings<T, d>& SolverSettings,
+			const TPBDJointSettings<T, d>& JointSettings,
+			const int32 Index0,
+			const int32 Index1,
+			TVector<T, d>& P0,
+			TRotation<T, d>& Q0,
+			TVector<T, d>& V0,
+			TVector<T, d>& W0,
+			TVector<T, d>& P1,
+			TRotation<T, d>& Q1,
+			TVector<T, d>& V1,
+			TVector<T, d>& W1,
+			float InvM0,
+			const PMatrix<T, d, d>& InvIL0,
+			float InvM1,
+			const PMatrix<T, d, d>& InvIL1);
+
 		static CHAOS_API void ApplyJointConeConstraint(
 			const T Dt,
 			const TPBDJointSolverSettings<T, d>& SolverSettings,
@@ -134,6 +175,25 @@ namespace Chaos
 			float InvM1,
 			const PMatrix<T, d, d>& InvIL1);
 
+		static CHAOS_API void ApplyJointConeVelocityConstraint(
+			const T Dt,
+			const TPBDJointSolverSettings<T, d>& SolverSettings,
+			const TPBDJointSettings<T, d>& JointSettings,
+			const int32 Index0,
+			const int32 Index1,
+			TVector<T, d>& P0,
+			TRotation<T, d>& Q0,
+			TVector<T, d>& V0,
+			TVector<T, d>& W0,
+			TVector<T, d>& P1,
+			TRotation<T, d>& Q1,
+			TVector<T, d>& V1,
+			TVector<T, d>& W1,
+			float InvM0,
+			const PMatrix<T, d, d>& InvIL0,
+			float InvM1,
+			const PMatrix<T, d, d>& InvIL1);
+
 		static CHAOS_API void ApplyJointSwingConstraint(
 			const T Dt,
 			const TPBDJointSolverSettings<T, d>& SolverSettings,
@@ -150,6 +210,26 @@ namespace Chaos
 			float InvM1,
 			const PMatrix<T, d, d>& InvIL1);
 
+		static CHAOS_API void ApplyJointSwingVelocityConstraint(
+			const T Dt,
+			const TPBDJointSolverSettings<T, d>& SolverSettings,
+			const TPBDJointSettings<T, d>& JointSettings,
+			const int32 Index0,
+			const int32 Index1,
+			const EJointAngularConstraintIndex SwingConstraint,
+			TVector<T, d>& P0,
+			TRotation<T, d>& Q0,
+			TVector<T, d>& V0,
+			TVector<T, d>& W0,
+			TVector<T, d>& P1,
+			TRotation<T, d>& Q1,
+			TVector<T, d>& V1,
+			TVector<T, d>& W1,
+			float InvM0,
+			const PMatrix<T, d, d>& InvIL0,
+			float InvM1,
+			const PMatrix<T, d, d>& InvIL1);
+
 		static CHAOS_API void ApplyJointTwistDrive(
 			const T Dt,
 			const TPBDJointSolverSettings<T, d>& SolverSettings,
diff --git a/Engine/Source/Runtime/Experimental/Chaos/Public/Chaos/PBDJointConstraints.h b/Engine/Source/Runtime/Experimental/Chaos/Public/Chaos/PBDJointConstraints.h
--- a/Engine/Source/Runtime/Experimental/Chaos/Public/Chaos/PBDJointConstraints.h
+++ b/Engine/Source/Runtime/Experimental/Chaos/Public/Chaos/PBDJointConstraints.h
@@ -137,10 +137,9 @@ namespace Chaos
 
 		CHAOS_API void CalculateConstraintSpace(int32 ConstraintIndex, TVector<T, d>& OutX0, PMatrix<T, d, d>& OutR0, TVector<T, d>& OutX1, PMatrix<T, d, d>& OutR1, TVector<T, d>& OutAngles) const;
 
-		CHAOS_API void ApplySingle(const T Dt, const int32 ConstraintIndex, const int32 It, const int32 NumIts);
-
-		CHAOS_API void ApplyProjection(const T Dt, const TArray<FConstraintHandle*>& InConstraintHandles);
-		CHAOS_API void ApplyProjectionSingle(const T Dt, const int32 ConstraintIndex, const int32 It, const int32 NumIts);
+		CHAOS_API void SolveVelocity(const T Dt, const int32 ConstraintIndex, const int32 It, const int32 NumIts);
+		CHAOS_API void SolvePosition(const T Dt, const int32 ConstraintIndex, const int32 It, const int32 NumIts);
+		CHAOS_API void ProjectPosition(const T Dt, const int32 ConstraintIndex, const int32 It, const int32 NumIts);
 
 		TPBDJointSolverSettings<T, d> Settings;
 
diff --git a/Engine/Source/Runtime/Experimental/Chaos/Public/Chaos/PBDRigidsEvolution.h b/Engine/Source/Runtime/Experimental/Chaos/Public/Chaos/PBDRigidsEvolution.h
--- a/Engine/Source/Runtime/Experimental/Chaos/Public/Chaos/PBDRigidsEvolution.h
+++ b/Engine/Source/Runtime/Experimental/Chaos/Public/Chaos/PBDRigidsEvolution.h
@@ -215,7 +215,7 @@ class TPBDRigidsEvolutionBase
 {
   public:
 	typedef TFunction<void(TTransientPBDRigidParticleHandle<T,d>& Particle, const T)> FForceRule;
-	typedef TFunction<void(const TParticleView<TPBDRigidParticles<T, d>>&, const T)> FUpdateVelocityRule;
+	typedef TFunction<void(const TArray<TGeometryParticleHandle<T, d>*>&, const T)> FUpdateVelocityRule;
 	typedef TFunction<void(const TParticleView<TPBDRigidParticles<T, d>>&, const T)> FUpdatePositionRule;
 	typedef TFunction<void(TPBDRigidParticles<T, d>&, const T, const T, const int32)> FKinematicUpdateRule;
 
@@ -551,7 +551,7 @@ protected:
 	
 	void UpdateVelocities(const T Dt, int32 Island)
 	{
-		ParticleUpdateVelocity(Particles.GetActiveParticlesView(), Dt);
+		ParticleUpdateVelocity(ConstraintGraph.GetIslandParticles(Island), Dt);
 	}
 
 	void ApplyPushOut(const T Dt, int32 Island)
diff --git a/Engine/Source/Runtime/Experimental/Chaos/Public/Chaos/PerParticlePBDUpdateFromDeltaPosition.h b/Engine/Source/Runtime/Experimental/Chaos/Public/Chaos/PerParticlePBDUpdateFromDeltaPosition.h
--- a/Engine/Source/Runtime/Experimental/Chaos/Public/Chaos/PerParticlePBDUpdateFromDeltaPosition.h
+++ b/Engine/Source/Runtime/Experimental/Chaos/Public/Chaos/PerParticlePBDUpdateFromDeltaPosition.h
@@ -33,6 +33,12 @@ class TPerParticlePBDUpdateFromDeltaPosition : public TPerParticleRule<T, d>
 		InParticles.W(Index) = TRotation<T, d>::CalculateAngularVelocity(InParticles.R(Index), InParticles.Q(Index), Dt);
 	}
 
+	inline void Apply(TPBDRigidParticleHandle<T, d>* Particle, const T Dt) const override //-V762
+	{
+		Particle->V() = (Particle->P() - Particle->X()) / Dt;
+		Particle->W() = TRotation<T, d>::CalculateAngularVelocity(Particle->R(), Particle->Q(), Dt);
+	}
+
 	inline void Apply(TTransientPBDRigidParticleHandle<T, d>& Particle, const T Dt) const override //-V762
 	{
 		Particle.V() = (Particle.P() - Particle.X()) / Dt;
diff --git a/Engine/Source/Runtime/Experimental/Chaos/Public/Chaos/PerParticleRule.h b/Engine/Source/Runtime/Experimental/Chaos/Public/Chaos/PerParticleRule.h
--- a/Engine/Source/Runtime/Experimental/Chaos/Public/Chaos/PerParticleRule.h
+++ b/Engine/Source/Runtime/Experimental/Chaos/Public/Chaos/PerParticleRule.h
@@ -47,6 +47,7 @@ class TPerParticleRule : public TParticleRule<T, d>
 	virtual void Apply(TRigidParticles<T, d>& InParticles, const T Dt, const int Index) const { Apply(static_cast<TParticles<T, d>&>(InParticles), Dt, Index); };
 	virtual void Apply(TPBDRigidParticles<T, d>& InParticles, const T Dt, const int Index) const { Apply(static_cast<TRigidParticles<T, d>&>(InParticles), Dt, Index); };
 
+	virtual void Apply(TPBDRigidParticleHandle<T, d>* Particle, const T Dt) const { check(0); }
 	virtual void Apply(TTransientPBDRigidParticleHandle<T, d>& Particle, const T Dt) const { check(0); }
 };
 }
diff --git a/Engine/Source/Runtime/Experimental/ChaosCore/Public/Chaos/Vector.h b/Engine/Source/Runtime/Experimental/ChaosCore/Public/Chaos/Vector.h
--- a/Engine/Source/Runtime/Experimental/ChaosCore/Public/Chaos/Vector.h
+++ b/Engine/Source/Runtime/Experimental/ChaosCore/Public/Chaos/Vector.h
@@ -410,6 +410,10 @@ namespace Chaos
 		{
 			return TVector<float, 3>(-X, -Y, -Z);
 		}
+		TVector<float, 3> operator+(const float Other) const
+		{
+			return TVector<float, 3>(X + Other, Y + Other, Z + Other);
+		}
 		TVector<float, 3> operator-(const float Other) const
 		{
 			return TVector<float, 3>(X - Other, Y - Other, Z - Other);
@@ -426,6 +430,10 @@ namespace Chaos
 		{
 			return TVector<float, 3>(S / V.X, S / V.Y, S / V.Z);
 		}
+		TVector<float, 3> operator+(const TVector<float, 3>& Other) const
+		{
+			return TVector<float, 3>(X + Other[0], Y + Other[1], Z + Other[2]);
+		}
 		TVector<float, 3> operator-(const TVector<float, 3>& Other) const
 		{
 			return TVector<float, 3>(X - Other[0], Y - Other[1], Z - Other[2]);
@@ -439,6 +447,11 @@ namespace Chaos
 			return TVector<float, 3>(X / Other[0], Y / Other[1], Z / Other[2]);
 		}
 		template<class T2>
+		TVector<float, 3> operator+(const TVector<T2, 3>& Other) const
+		{
+			return TVector<float, 3>(X + Other[0], Y + Other[1], Z + Other[2]);
+		}
+		template<class T2>
 		TVector<float, 3> operator-(const TVector<T2, 3>& Other) const
 		{
 			return TVector<float, 3>(X - Other[0], Y - Other[1], Z - Other[2]);
