Chaos 6dof joint
- ability to toggle between velocity and position solve mode

diff --git a/Engine/Source/Runtime/Engine/Private/PhysicsEngine/ImmediatePhysics/ImmediatePhysicsChaos/ImmediatePhysicsSimulation_Chaos.cpp b/Engine/Source/Runtime/Engine/Private/PhysicsEngine/ImmediatePhysics/ImmediatePhysicsChaos/ImmediatePhysicsSimulation_Chaos.cpp
--- a/Engine/Source/Runtime/Engine/Private/PhysicsEngine/ImmediatePhysics/ImmediatePhysicsChaos/ImmediatePhysicsSimulation_Chaos.cpp
+++ b/Engine/Source/Runtime/Engine/Private/PhysicsEngine/ImmediatePhysics/ImmediatePhysicsChaos/ImmediatePhysicsSimulation_Chaos.cpp
@@ -38,6 +38,7 @@ FAutoConsoleVariableRef CVarChaosImmPhysThickness(TEXT("p.Chaos.ImmPhys.Collisio
 
 float ChaosImmediate_Joint_SwingTwistAngleTolerance = 1.0e-6f;
 int32 ChaosImmediate_Joint_PositionIterations = 4;
+int32 ChaosImmediate_Joint_EnableVelocitySolve = 1;
 int32 ChaosImmediate_Joint_EnableLinearLimits = 1;
 int32 ChaosImmediate_Joint_EnableTwistLimits = 1;
 int32 ChaosImmediate_Joint_EnableSwingLimits = 1;
@@ -50,6 +51,7 @@ float ChaosImmediate_Joint_MinParentMassRatio = 0.5f;
 float ChaosImmediate_Joint_MaxInertiaRatio = 5.0f;
 FAutoConsoleVariableRef CVarChaosImmPhysSwingTwistAngleTolerance(TEXT("p.Chaos.ImmPhys.Joint.SwingTwistAngleTolerance"), ChaosImmediate_Joint_SwingTwistAngleTolerance, TEXT("SwingTwistAngleTolerance."));
 FAutoConsoleVariableRef CVarChaosImmPhysEnablePositionSolve(TEXT("p.Chaos.ImmPhys.Joint.PositionIterations"), ChaosImmediate_Joint_PositionIterations, TEXT("6Dof joint PushOut (Position solve) iterations."));
+FAutoConsoleVariableRef CVarChaosImmPhysEnableVelocitySolve(TEXT("p.Chaos.ImmPhys.Joint.EnableVelocitySolve"), ChaosImmediate_Joint_EnableVelocitySolve, TEXT("EnableVelocitySolve."));
 FAutoConsoleVariableRef CVarChaosImmPhysEnableLinearLimits(TEXT("p.Chaos.ImmPhys.Joint.EnableLinearLimits"), ChaosImmediate_Joint_EnableLinearLimits, TEXT("EnableLinearLimits."));
 FAutoConsoleVariableRef CVarChaosImmPhysEnableTwistLimits(TEXT("p.Chaos.ImmPhys.Joint.EnableTwistLimits"), ChaosImmediate_Joint_EnableTwistLimits, TEXT("EnableTwistLimits."));
 FAutoConsoleVariableRef CVarChaosImmPhysEnableSwingLimits(TEXT("p.Chaos.ImmPhys.Joint.EnableSwingLimits"), ChaosImmediate_Joint_EnableSwingLimits, TEXT("EnableSwingLimits."));
@@ -404,6 +406,7 @@ namespace ImmediatePhysics_Chaos
 			JointsSettings.SwingTwistAngleTolerance = ChaosImmediate_Joint_SwingTwistAngleTolerance;
 			JointsSettings.MinParentMassRatio = ChaosImmediate_Joint_MinParentMassRatio;
 			JointsSettings.MaxInertiaRatio = ChaosImmediate_Joint_MaxInertiaRatio;
+			JointsSettings.bEnableVelocitySolve = ChaosImmediate_Joint_EnableVelocitySolve != 0;
 			JointsSettings.bEnableLinearLimits = ChaosImmediate_Joint_EnableLinearLimits != 0;
 			JointsSettings.bEnableTwistLimits = ChaosImmediate_Joint_EnableTwistLimits != 0;
 			JointsSettings.bEnableSwingLimits = ChaosImmediate_Joint_EnableSwingLimits != 0;
diff --git a/Engine/Source/Runtime/Experimental/Chaos/Private/Chaos/PBDJointConstraints.cpp b/Engine/Source/Runtime/Experimental/Chaos/Private/Chaos/PBDJointConstraints.cpp
--- a/Engine/Source/Runtime/Experimental/Chaos/Private/Chaos/PBDJointConstraints.cpp
+++ b/Engine/Source/Runtime/Experimental/Chaos/Private/Chaos/PBDJointConstraints.cpp
@@ -116,6 +116,7 @@ namespace Chaos
 		: SwingTwistAngleTolerance((T)1.0e-6)
 		, MinParentMassRatio((T)0.5)
 		, MaxInertiaRatio((T)5.0)
+		, bEnableVelocitySolve(true)
 		, bEnableLinearLimits(true)
 		, bEnableTwistLimits(true)
 		, bEnableSwingLimits(true)
@@ -347,7 +348,14 @@ namespace Chaos
 
 		for (FConstraintHandle* ConstraintHandle : SortedConstraintHandles)
 		{
-			SolveVelocity(Dt, ConstraintHandle->GetConstraintIndex(), It, NumIts);
+			if (Settings.bEnableVelocitySolve)
+			{
+				SolveVelocity(Dt, ConstraintHandle->GetConstraintIndex(), It, NumIts);
+			}
+			else
+			{
+				SolvePosition(Dt, ConstraintHandle->GetConstraintIndex(), It, NumIts);
+			}
 		}
 
 		if (PostApplyCallback != nullptr)
@@ -366,11 +374,14 @@ namespace Chaos
 			return L.GetConstraintLevel() < R.GetConstraintLevel();
 		});
 
-		for (int32 It = 0; It < Settings.PositionIterations; ++It)
+		if (Settings.bEnableVelocitySolve)
 		{
-			for (FConstraintHandle* ConstraintHandle : SortedConstraintHandles)
+			for (int32 It = 0; It < Settings.PositionIterations; ++It)
 			{
-				SolvePosition(Dt, ConstraintHandle->GetConstraintIndex(), It, Settings.PositionIterations);
+				for (FConstraintHandle* ConstraintHandle : SortedConstraintHandles)
+				{
+					SolvePosition(Dt, ConstraintHandle->GetConstraintIndex(), It, Settings.PositionIterations);
+				}
 			}
 		}
 
@@ -566,15 +577,18 @@ namespace Chaos
 		}
 
 		// Freeze the closest to kinematic connection if there is a difference
-		if (Level0 < Level1)
-		{
-			InvM0 = 0;
-			InvIL0 = PMatrix<T, d, d>(0, 0, 0);
-		}
-		else if (Level1 < Level0)
+		if (Settings.bEnableVelocitySolve)
 		{
-			InvM1 = 0;
-			InvIL1 = PMatrix<T, d, d>(0, 0, 0);
+			if (Level0 < Level1)
+			{
+				InvM0 = 0;
+				InvIL0 = PMatrix<T, d, d>(0, 0, 0);
+			}
+			else if (Level1 < Level0)
+			{
+				InvM1 = 0;
+				InvIL1 = PMatrix<T, d, d>(0, 0, 0);
+			}
 		}
 
 		const TVector<EJointMotionType, d>& LinearMotion = JointSettings.Motion.LinearMotionTypes;
@@ -588,6 +602,38 @@ namespace Chaos
 			return;
 		}
 
+		// Apply angular drives (NOTE: modifies position, not velocity)
+		if (!Settings.bEnableVelocitySolve && Settings.bEnableDrives)
+		{
+			bool bTwistLocked = TwistMotion == EJointMotionType::Locked;
+			bool bSwing1Locked = Swing1Motion == EJointMotionType::Locked;
+			bool bSwing2Locked = Swing2Motion == EJointMotionType::Locked;
+
+			// No SLerp drive if we have a locked rotation (it will be grayed out in the editor in this case, but could still have been set before the rotation was locked)
+			if (JointSettings.Motion.bAngularSLerpDriveEnabled && !bTwistLocked && !bSwing1Locked && !bSwing2Locked)
+			{
+				TPBDJointUtilities<T, d>::ApplyJointSLerpDrive(Dt, Settings, JointSettings, Index0, Index1, P0, Q0, P1, Q1, InvM0, InvIL0, InvM1, InvIL1);
+			}
+
+			if (JointSettings.Motion.bAngularTwistDriveEnabled && !bTwistLocked)
+			{
+				TPBDJointUtilities<T, d>::ApplyJointTwistDrive(Dt, Settings, JointSettings, Index0, Index1, P0, Q0, P1, Q1, InvM0, InvIL0, InvM1, InvIL1);
+			}
+
+			if (JointSettings.Motion.bAngularSwingDriveEnabled && !bSwing1Locked && !bSwing2Locked)
+			{
+				TPBDJointUtilities<T, d>::ApplyJointConeDrive(Dt, Settings, JointSettings, Index0, Index1, P0, Q0, P1, Q1, InvM0, InvIL0, InvM1, InvIL1);
+			}
+			else if (JointSettings.Motion.bAngularSwingDriveEnabled && !bSwing1Locked)
+			{
+				//TPBDJointUtilities<T, d>::ApplyJointSwingDrive(Dt, Settings, JointSettings, Index0, Index1, EJointAngularConstraintIndex::Swing1, P0, Q0, P1, Q1, InvM0, InvIL0, InvM1, InvIL1);
+			}
+			else if (JointSettings.Motion.bAngularSwingDriveEnabled && !bSwing2Locked)
+			{
+				//TPBDJointUtilities<T, d>::ApplyJointSwingDrive(Dt, Settings, JointSettings, Index0, Index1, EJointAngularConstraintIndex::Swing2, P0, Q0, P1, Q1, InvM0, InvIL0, InvM1, InvIL1);
+			}
+		}
+
 		// Apply twist constraint
 		if (Settings.bEnableTwistLimits)
 		{
diff --git a/Engine/Source/Runtime/Experimental/Chaos/Public/Chaos/PBDJointConstraintTypes.h b/Engine/Source/Runtime/Experimental/Chaos/Public/Chaos/PBDJointConstraintTypes.h
--- a/Engine/Source/Runtime/Experimental/Chaos/Public/Chaos/PBDJointConstraintTypes.h
+++ b/Engine/Source/Runtime/Experimental/Chaos/Public/Chaos/PBDJointConstraintTypes.h
@@ -118,6 +118,7 @@ namespace Chaos
 		T MaxInertiaRatio;
 
 		// @todo(ccaulfield): remove these TEMP overrides for testing
+		bool bEnableVelocitySolve;
 		bool bEnableLinearLimits;
 		bool bEnableTwistLimits;
 		bool bEnableSwingLimits;
