Time Slicing of AABB Acceleration structure.

diff --git a/Engine/Source/Runtime/Experimental/Chaos/Private/Chaos/PBDRigidsEvolution.cpp b/Engine/Source/Runtime/Experimental/Chaos/Private/Chaos/PBDRigidsEvolution.cpp
--- a/Engine/Source/Runtime/Experimental/Chaos/Private/Chaos/PBDRigidsEvolution.cpp
+++ b/Engine/Source/Runtime/Experimental/Chaos/Private/Chaos/PBDRigidsEvolution.cpp
@@ -18,6 +18,7 @@ namespace Chaos
 		int32 AABBMaxChildrenInLeaf;
 		int32 AABBMaxTreeDepth;
 		float MaxPayloadSize;
+		int32 IterationsPerTimeSlice;
 
 		FAccelerationConfig()
 		{
@@ -28,6 +29,7 @@ namespace Chaos
 			AABBMaxChildrenInLeaf = 500;
 			AABBMaxTreeDepth = 200;
 			MaxPayloadSize = 20000;
+			IterationsPerTimeSlice = 2000;
 		}
 	} ConfigSettings;
 
@@ -38,6 +40,7 @@ namespace Chaos
 	FAutoConsoleVariableRef CVarAABBMaxChildrenInLeaf(TEXT("p.AABBMaxChildrenInLeaf"), ConfigSettings.AABBMaxChildrenInLeaf, TEXT(""));
 	FAutoConsoleVariableRef CVarAABBMaxTreeDepth(TEXT("p.AABBMaxTreeDepth"), ConfigSettings.AABBMaxTreeDepth, TEXT(""));
 	FAutoConsoleVariableRef CVarMaxPayloadSize(TEXT("p.MaxPayloadSize"), ConfigSettings.MaxPayloadSize, TEXT(""));
+	FAutoConsoleVariableRef CVarIterationsPerTimeSlice(TEXT("p.IterationsPerTimeSlice"), ConfigSettings.IterationsPerTimeSlice, TEXT(""));
 
 	template<typename T, int d>
 	struct TDefaultCollectionFactory : public ISpatialAccelerationCollectionFactory<T, d>
@@ -78,11 +81,11 @@ namespace Chaos
 				{
 					return MakeUnique<BVType>(Particles, false, 0, ConfigSettings.BVNumCells, ConfigSettings.MaxPayloadSize);
 				}
-				else if(ConfigSettings.BroadphaseType == 1 || ConfigSettings.BroadphaseType == 3)
+				else if (ConfigSettings.BroadphaseType == 1 || ConfigSettings.BroadphaseType == 3)
 				{
-					return MakeUnique<AABBTreeType>(Particles, ConfigSettings.MaxChildrenInLeaf, ConfigSettings.MaxTreeDepth, ConfigSettings.MaxPayloadSize);
+					return MakeUnique<AABBTreeType>(Particles, ConfigSettings.MaxChildrenInLeaf, ConfigSettings.MaxTreeDepth, ConfigSettings.MaxPayloadSize, ConfigSettings.IterationsPerTimeSlice);
 				}
-				else if(ConfigSettings.BroadphaseType == 4 || ConfigSettings.BroadphaseType == 2)
+				else if (ConfigSettings.BroadphaseType == 4 || ConfigSettings.BroadphaseType == 2)
 				{
 					return MakeUnique<AABBTreeOfGridsType>(Particles, ConfigSettings.AABBMaxChildrenInLeaf, ConfigSettings.AABBMaxTreeDepth, ConfigSettings.MaxPayloadSize);
 				}
@@ -120,13 +123,13 @@ namespace Chaos
 		, bExternalReady(false)
 		, Clustering(static_cast<FPBDRigidsEvolution&>(*this), Particles.GetClusteredParticles())
 		, NumIterations(InNumIterations)
-		, SpatialCollectionFactory(new TDefaultCollectionFactory<T,d>())
+		, SpatialCollectionFactory(new TDefaultCollectionFactory<T, d>())
 	{
 		Particles.GetParticleHandles().AddArray(&PhysicsMaterials);
 		Particles.GetParticleHandles().AddArray(&PerParticlePhysicsMaterials);
 		Particles.GetParticleHandles().AddArray(&ParticleDisableCount);
 		Particles.GetParticleHandles().AddArray(&Collided);
-		
+
 		for (auto& Particle : InParticles.GetNonDisabledView())
 		{
 			DirtyParticle(Particle);
@@ -145,14 +148,16 @@ namespace Chaos
 		WaitOnAccelerationStructure();
 	}
 
-	
+
 	DECLARE_CYCLE_STAT(TEXT("CacheAccelerationBounds"), STAT_CacheAccelerationBounds, STATGROUP_Chaos);
 	DECLARE_CYCLE_STAT(TEXT("ComputeIntermediateSpatialAcceleration"), STAT_ComputeIntermediateSpatialAcceleration, STATGROUP_Chaos);
+	DECLARE_CYCLE_STAT(TEXT("CopyAccelerationStructure"), STAT_CopyAccelerationStructure, STATGROUP_Chaos);
+	DECLARE_CYCLE_STAT(TEXT("SwapAccelerationStructures"), STAT_SwapAccelerationStructures, STATGROUP_Chaos);
 
 	template<class FPBDRigidsEvolution, class FPBDCollisionConstraint, class T, int d>
 	TPBDRigidsEvolutionBase<FPBDRigidsEvolution, FPBDCollisionConstraint, T, d>::FChaosAccelerationStructureTask::FChaosAccelerationStructureTask(
 		ISpatialAccelerationCollectionFactory<T, d>& InSpatialCollectionFactory
-		, const TMap<FSpatialAccelerationIdx, TUniquePtr<TSpatialAccelerationCache<T,d>>>& InSpatialAccelerationCache
+		, const TMap<FSpatialAccelerationIdx, TUniquePtr<TSpatialAccelerationCache<T, d>>>& InSpatialAccelerationCache
 		, TUniquePtr<FAccelerationStructure>& InAccelerationStructure
 		, TUniquePtr<FAccelerationStructure>& InAccelerationStructureCopy)
 		: SpatialCollectionFactory(InSpatialCollectionFactory)
@@ -181,11 +186,11 @@ namespace Chaos
 	}
 
 	template<typename T, int d>
-	TUniquePtr<ISpatialAccelerationCollection<TAccelerationStructureHandle<T, d>, T, d>> CreateNewSpatialStructureFromSubStructure(TUniquePtr<ISpatialAcceleration<TAccelerationStructureHandle<T,d>, T,d>>&& Substructure)
+	TUniquePtr<ISpatialAccelerationCollection<TAccelerationStructureHandle<T, d>, T, d>> CreateNewSpatialStructureFromSubStructure(TUniquePtr<ISpatialAcceleration<TAccelerationStructureHandle<T, d>, T, d>>&& Substructure)
 	{
 		using BVType = TBoundingVolume<TAccelerationStructureHandle<T, d>, T, d>;
 		using AABBType = TAABBTree<TAccelerationStructureHandle<T, d>, TAABBTreeLeafArray<TAccelerationStructureHandle<T, d>, T>, T>;
-		
+
 		if (Substructure->template As<BVType>())
 		{
 			auto Collection = MakeUnique<TSpatialAccelerationCollection<BVType>>();
@@ -212,18 +217,29 @@ namespace Chaos
 	{
 		uint8 ActiveBucketsMask = SpatialCollectionFactory.GetActiveBucketsMask();
 		TArray<TSOAView<TSpatialAccelerationCache<T, d>>> ViewsPerBucket[8];
-		
+
+		bool IsTimeSlicingProgressing = false;
+
 		//merge buckets. todo: support multiple entries per bucket (i.e. dynamic vs static)
 		for (const auto& Itr : SpatialAccelerationCache)
 		{
 			const FSpatialAccelerationIdx SpatialIdx = Itr.Key;
 			const TSpatialAccelerationCache<T, d>& Cache = *Itr.Value;
 			const uint8 BucketIdx = (1 << SpatialIdx.Bucket) & ActiveBucketsMask ? SpatialIdx.Bucket : 0;
-			ViewsPerBucket[BucketIdx].Add(const_cast<TSpatialAccelerationCache<T, d>*>(&Cache));
+			if (AccelerationStructure->GetSubstructure(SpatialIdx) && !AccelerationStructure->GetSubstructure(SpatialIdx)->IsAsyncTimeSlicingComplete())
+			{
+				AccelerationStructure->GetSubstructure(SpatialIdx)->ProgressAsyncTimeSlicing();
 
-			if (AccelerationStructure->IsBucketActive(SpatialIdx.Bucket))
+				// is it still progressing or now complete
+				IsTimeSlicingProgressing = !AccelerationStructure->GetSubstructure(SpatialIdx)->IsAsyncTimeSlicingComplete();
+			}
+			else
 			{
-				AccelerationStructure->RemoveSubstructure(SpatialIdx);
+				ViewsPerBucket[BucketIdx].Add(const_cast<TSpatialAccelerationCache<T, d>*>(&Cache));
+				if (AccelerationStructure->IsBucketActive(SpatialIdx.Bucket))
+				{
+					AccelerationStructure->RemoveSubstructure(SpatialIdx);
+				}
 			}
 		}
 
@@ -234,11 +250,26 @@ namespace Chaos
 			{
 				auto ParticleView = MakeConstParticleView(MoveTemp(ViewsPerBucket[BucketIdx]));
 				auto NewStruct = SpatialCollectionFactory.CreateAccelerationPerBucket_Threaded(ParticleView, BucketIdx);
+
+				// we kicked of the creation of a new structure and it's going to time-slice the work
+				if (!NewStruct->IsAsyncTimeSlicingComplete())
+				{
+					IsTimeSlicingProgressing = true;
+				}
+
 				AccelerationStructure->AddSubstructure(MoveTemp(NewStruct), BucketIdx);
+
 			}
 		}
-		
-		AccelerationStructureCopy = AsUniqueSpatialAccelerationChecked<FAccelerationStructure>(AccelerationStructure->Copy());
+
+		AccelerationStructure->SetAllAsyncTrasksComplete(!IsTimeSlicingProgressing);
+
+		if (!IsTimeSlicingProgressing)
+		{
+			// This operation is slow!
+			SCOPE_CYCLE_COUNTER(STAT_CopyAccelerationStructure);
+			AccelerationStructureCopy = AsUniqueSpatialAccelerationChecked<FAccelerationStructure>(AccelerationStructure->Copy());
+		}
 	}
 
 	template<class FPBDRigidsEvolution, class FPBDCollisionConstraint, class T, int d>
@@ -274,10 +305,10 @@ namespace Chaos
 			}
 		}
 
-		if(SpatialData.bUpdate)
+		if (SpatialData.bUpdate)
 		{
 			AccelerationStructure.UpdateElementIn(Particle, Particle->WorldSpaceInflatedBounds(), Particle->HasBounds(), SpatialData.UpdatedSpatialIdx);
-			
+
 			if (bUpdateCache)
 			{
 				TUniquePtr<TSpatialAccelerationCache<T, d>>* CachePtrPtr = SpatialAccelerationCache.Find(SpatialData.UpdatedSpatialIdx);
@@ -379,10 +410,17 @@ namespace Chaos
 			WaitOnAccelerationStructure();
 		}
 
-		if (!AccelerationStructureTaskComplete || AccelerationStructureTaskComplete->IsComplete())
+		bool AsyncComplete = !AccelerationStructureTaskComplete || AccelerationStructureTaskComplete->IsComplete();
+
+		if (AsyncComplete)
 		{
-			if (AccelerationStructureTaskComplete)
+			// only copy when the acceleration structures have completed time-slicing
+			if (AccelerationStructureTaskComplete && AsyncInternalAcceleration->IsAllAsyncTrasksComplete())
 			{
+				SCOPE_CYCLE_COUNTER(STAT_SwapAccelerationStructures);
+
+				check(AsyncInternalAcceleration->IsAllAsyncTrasksComplete());
+
 				FlushAsyncAccelerationQueue();
 
 				//swap acceleration structure for new one
@@ -391,6 +429,7 @@ namespace Chaos
 				bExternalReady = true;
 			}
 
+			// we run the task for both starting a new accel structure as well as for the timeslicing
 			AccelerationStructureTaskComplete = TGraphTask<FChaosAccelerationStructureTask>::CreateTask().ConstructAndDispatchWhenReady(*SpatialCollectionFactory, SpatialAccelerationCache, AsyncInternalAcceleration, AsyncExternalAcceleration);
 		}
 		else
diff --git a/Engine/Source/Runtime/Experimental/Chaos/Public/Chaos/AABBTree.h b/Engine/Source/Runtime/Experimental/Chaos/Public/Chaos/AABBTree.h
--- a/Engine/Source/Runtime/Experimental/Chaos/Public/Chaos/AABBTree.h
+++ b/Engine/Source/Runtime/Experimental/Chaos/Public/Chaos/AABBTree.h
@@ -222,6 +222,7 @@ public:
 	static constexpr T DefaultMaxPayloadBounds = 100000;
 	static constexpr int32 DefaultMaxChildrenInLeaf = 12;
 	static constexpr int32 DefaultMaxTreeDepth = 16;
+	static constexpr int32 DefaultMaxNumToProcess = 2000;
 	static constexpr ESpatialAcceleration StaticType = TIsSame<TAABBTreeLeafArray<TPayloadType, T>, TLeafType>::Value ? ESpatialAcceleration::AABBTree : 
 		(TIsSame<TBoundingVolume<TPayloadType, T, 3>, TLeafType>::Value ? ESpatialAcceleration::AABBTreeBV : ESpatialAcceleration::Unknown);
 	TAABBTree()
@@ -229,17 +230,61 @@ public:
 		, MaxChildrenInLeaf(DefaultMaxChildrenInLeaf)
 		, MaxTreeDepth(DefaultMaxTreeDepth)
 		, MaxPayloadBounds(DefaultMaxPayloadBounds)
+		, MaxNumToProcess(DefaultMaxNumToProcess)
 	{
 	}
 
+	virtual void ProgressAsyncTimeSlicing() override
+	{
+		// still has work to complete
+		if (!TimeSliceWorkToComplete.IsEmpty())
+		{
+			FWorkSnapshot Store;
+			TimeSliceWorkToComplete.Dequeue(Store);
+			NumProcessedThisSlice = 0;
+			SplitNode(Store.Bounds, Store.Elems, Store.NodeLevel, Store.NewNodeIdx);
+		}
+
+		// are done now?
+		if (TimeSliceWorkToComplete.IsEmpty())
+		{
+			this->SetAsyncTimeSlicingComplete(true);
+		}
+
+	}
+
 	template <typename TParticles>
-	TAABBTree(const TParticles& Particles, int32 InMaxChildrenInLeaf = 12, int32 InMaxTreeDepth = 16, T InMaxPayloadBounds = DefaultMaxPayloadBounds)
+	TAABBTree(const TParticles& Particles, int32 InMaxChildrenInLeaf = DefaultMaxChildrenInLeaf, int32 InMaxTreeDepth = DefaultMaxTreeDepth, T InMaxPayloadBounds = DefaultMaxPayloadBounds, int32 InMaxNumToProcess = DefaultMaxNumToProcess )
 		: ISpatialAcceleration<TPayloadType, T, 3>(StaticType)
 		, MaxChildrenInLeaf(InMaxChildrenInLeaf)
 		, MaxTreeDepth(InMaxTreeDepth)
 		, MaxPayloadBounds(InMaxPayloadBounds)
+		, MaxNumToProcess(InMaxNumToProcess)
+
 	{
+		this->SetAsyncTimeSlicingComplete(false);
 		GenerateTree(Particles);
+
+		/*
+		TBox<T, 3> Tmp = TBox<T, 3>::ZeroBox();
+		DumpTreeToLog(0, -1, Tmp);
+		*/
+
+	}
+
+	bool DumpTreeToLog(int32 NodeIdx, int32 ParentNode, TBox<T,3>& Bounds)
+	{
+		if (Nodes[NodeIdx].bLeaf)
+		{
+			UE_LOG(LogChaos, Log, TEXT("Bounds %s"), Bounds.ToString().GetCharArray().GetData());
+			return false;
+		}
+		else
+		{		
+			DumpTree(Nodes[NodeIdx].ChildrenNodes[0], NodeIdx, Nodes[NodeIdx].ChildrenBounds[0]);
+			DumpTree(Nodes[NodeIdx].ChildrenNodes[1], NodeIdx, Nodes[NodeIdx].ChildrenBounds[0]);
+		}
+		return true;
 	}
 
 	virtual ~TAABBTree() {}
@@ -591,6 +636,8 @@ private:
 		Nodes.Reset();
 		DirtyElements.Reset();
 		PayloadToInfo.Reset();
+		NumProcessedThisSlice = 0;
+		check(MaxNumToProcess > 0);
 
 		int32 Idx = 0;
 		FullBounds = TBox<T, 3>::EmptyBox();
@@ -616,7 +663,7 @@ private:
 			{
 				ElemBounds = TBox<T,3>(TVector<T, 3>(TNumericLimits<T>::Lowest()), TVector<T, 3>(TNumericLimits<T>::Max()));
 			}
-			
+
 			if (!bHasBoundingBox)
 			{
 				PayloadToInfo.Add(Payload, FAABBTreePayloadInfo{ GlobalPayloads.Num(), INDEX_NONE, INDEX_NONE });
@@ -627,13 +674,37 @@ private:
 			//todo: payload info
 		}
 
-		SplitNode(FullBounds, ElemsWithBounds, 0);
-	}
+		SplitNode(FullBounds, ElemsWithBounds, 0, 0);
 
-	int32 SplitNode(const TBox<T, 3>& Bounds, const TArray<FElement>& Elems, int32 NodeLevel)
+		if (TimeSliceWorkToComplete.IsEmpty())
+		{
+			this->SetAsyncTimeSlicingComplete(true);
+		}
+	}
+	
+	void SplitNode(const TBox<T, 3>& Bounds, const TArray<FElement>& Elems, int32 NodeLevel, int32 NewNodeIdx)
 	{
-		const int32 NewNodeIdx = Nodes.Num();
-		Nodes.AddDefaulted();	//todo: remove TBox
+		// create the actual node space but might no be filled in (YET) due to time slicing exit
+		if (NewNodeIdx >= Nodes.Num())
+		{
+		 	Nodes.AddDefaulted();	//todo: remove TBox
+		}
+
+		if (NumProcessedThisSlice >= MaxNumToProcess)
+		{
+			// done enough work, capture stack
+			FWorkSnapshot Store;
+			Store.Bounds = Bounds;
+			Store.Elems = Elems;			
+			Store.NodeLevel = NodeLevel;
+			Store.NewNodeIdx = NewNodeIdx;
+
+			TimeSliceWorkToComplete.Enqueue(Store);
+
+			return; // done enough
+		}
+
+		NumProcessedThisSlice++;
 
 		auto& PayloadToInfoRef = PayloadToInfo;
 		auto& LeavesRef = Leaves;
@@ -653,7 +724,7 @@ private:
 		if (Elems.Num() <= MaxChildrenInLeaf || NodeLevel >= MaxTreeDepth)
 		{
 			MakeLeaf();
-			return NewNodeIdx;
+			return; // keep working
 		}
 
 		const TVector<T, 3> Extents = Bounds.Extents();
@@ -670,7 +741,7 @@ private:
 		FSplitInfo SplitInfos[2];
 		SplitInfos[0].SplitBounds = TBox<T, 3>(Bounds.Min(), Bounds.Min());
 		SplitInfos[1].SplitBounds = TBox<T, 3>(Bounds.Max(), Bounds.Max());
-		
+
 		const TVector<T, 3> Center = Bounds.Center();
 		for (FSplitInfo& SplitInfo : SplitInfos)
 		{
@@ -678,9 +749,9 @@ private:
 
 			for (int32 Axis = 0; Axis < 3; ++Axis)
 			{
-				TVector<T,3> NewPt0 = Center;
-				TVector<T,3> NewPt1 = Center;
-				if(Axis != MaxAxis)
+				TVector<T, 3> NewPt0 = Center;
+				TVector<T, 3> NewPt1 = Center;
+				if (Axis != MaxAxis)
 				{
 					NewPt0[Axis] = Bounds.Min()[Axis];
 					NewPt1[Axis] = Bounds.Max()[Axis];
@@ -692,6 +763,7 @@ private:
 			SplitInfo.SplitBoundsSize2 = SplitInfo.SplitBounds.Extents().SizeSquared();
 		}
 		
+		// add all elements to one of the two split infos at this level - root level [ not taking into account the max number allowed or anything
 		for (const FElement& Elem : Elems)
 		{
 			int32 MinBoxIdx = INDEX_NONE;
@@ -709,30 +781,36 @@ private:
 				}
 				++BoxIdx;
 			}
-			
+
 			SplitInfos[MinBoxIdx].Children.Add(Elem);
 			SplitInfos[MinBoxIdx].RealBounds.GrowToInclude(Elem.Bounds);
 		}
 
+		// if children in both halves, recurse a level down
 		if (SplitInfos[0].Children.Num() && SplitInfos[1].Children.Num())
 		{
 			Nodes[NewNodeIdx].bLeaf = false;
+
 			for (int32 BoxIdx = 0; BoxIdx < 2; ++BoxIdx)
 			{
-				const int32 ChildIdx = SplitNode(SplitInfos[BoxIdx].RealBounds, SplitInfos[BoxIdx].Children, NodeLevel + 1);
-				Nodes[NewNodeIdx].ChildrenBounds[BoxIdx] = SplitInfos[BoxIdx].RealBounds;
-				Nodes[NewNodeIdx].ChildrenNodes[BoxIdx] = ChildIdx;
+			
+					const int32 ChildIdx = Nodes.Num();
+					Nodes[NewNodeIdx].ChildrenBounds[BoxIdx] = SplitInfos[BoxIdx].RealBounds;
+					Nodes[NewNodeIdx].ChildrenNodes[BoxIdx] = ChildIdx;
+
+					SplitNode(SplitInfos[BoxIdx].RealBounds, SplitInfos[BoxIdx].Children, NodeLevel + 1, ChildIdx);
 			}
 		}
 		else
 		{
-			//couldn't split so just make a leaf
+			//couldn't split so just make a leaf - THIS COULD CONTAIN MORE THAN MaxChildrenInLeaf!!!
 			MakeLeaf();
 		}
 
-		return NewNodeIdx;
+		return; // keep working
 	}
-	
+
+
 	TAABBTree(const TAABBTree<TPayloadType, TLeafType, T>& Other)
 		: ISpatialAcceleration<TPayloadType, T, 3>(StaticType)
 		, FullBounds(Other.FullBounds)
@@ -744,6 +822,8 @@ private:
 		, MaxChildrenInLeaf(Other.MaxChildrenInLeaf)
 		, MaxTreeDepth(Other.MaxTreeDepth)
 		, MaxPayloadBounds(Other.MaxPayloadBounds)
+		, MaxNumToProcess(Other.MaxNumToProcess)
+		, NumProcessedThisSlice(Other.NumProcessedThisSlice)
 	{
 
 	}
@@ -757,6 +837,19 @@ private:
 	int32 MaxChildrenInLeaf;
 	int32 MaxTreeDepth;
 	T MaxPayloadBounds;
+	int32 MaxNumToProcess;
+
+	struct FWorkSnapshot
+	{
+		TBox<T, 3> Bounds;
+		TArray<FElement> Elems;
+
+		int32 NodeLevel;
+		int32 NewNodeIdx;
+	};
+
+	int32 NumProcessedThisSlice;
+	TQueue<FWorkSnapshot> TimeSliceWorkToComplete;
 };
 
 template<typename TPayloadType, typename TLeafType, class T>
diff --git a/Engine/Source/Runtime/Experimental/Chaos/Public/Chaos/ISpatialAcceleration.h b/Engine/Source/Runtime/Experimental/Chaos/Public/Chaos/ISpatialAcceleration.h
--- a/Engine/Source/Runtime/Experimental/Chaos/Public/Chaos/ISpatialAcceleration.h
+++ b/Engine/Source/Runtime/Experimental/Chaos/Public/Chaos/ISpatialAcceleration.h
@@ -147,11 +147,13 @@ class CHAOS_API ISpatialAcceleration
 public:
 
 	ISpatialAcceleration(SpatialAccelerationType InType = ESpatialAcceleration::Unknown)
-		: Type(InType)
+		: Type(InType), AsyncTimeSlicingComplete(true)
 	{
 	}
 	virtual ~ISpatialAcceleration() = default;
 
+	virtual bool IsAsyncTimeSlicingComplete() { return AsyncTimeSlicingComplete; }
+	virtual void ProgressAsyncTimeSlicing() {}
 	virtual TArray<TPayloadType> FindAllIntersections(const TBox<T, d>& Box) const { check(false); return TArray<TPayloadType>(); }
 
 	virtual void Raycast(const TVector<T, d>& Start, const TVector<T, d>& Dir, const T Length, ISpatialVisitor<TPayloadType, T>& Visitor) const { check(false); }
@@ -223,8 +225,12 @@ public:
 		return static_cast<const TConcrete&>(*this);
 	}
 
+protected:
+	virtual void SetAsyncTimeSlicingComplete(bool InState) { AsyncTimeSlicingComplete = InState; }
+
 private:
 	SpatialAccelerationType Type;
+	bool AsyncTimeSlicingComplete;
 };
 
 template <typename TBase, typename TDerived>
diff --git a/Engine/Source/Runtime/Experimental/Chaos/Public/Chaos/ISpatialAccelerationCollection.h b/Engine/Source/Runtime/Experimental/Chaos/Public/Chaos/ISpatialAccelerationCollection.h
--- a/Engine/Source/Runtime/Experimental/Chaos/Public/Chaos/ISpatialAccelerationCollection.h
+++ b/Engine/Source/Runtime/Experimental/Chaos/Public/Chaos/ISpatialAccelerationCollection.h
@@ -16,6 +16,7 @@ public:
 	ISpatialAccelerationCollection()
 	: ISpatialAcceleration<TPayloadType, T, d>(StaticType)
 	, ActiveBucketsMask(0)
+	, AllAsyncTrasksComplete(true)
 	{}
 	static constexpr ESpatialAcceleration StaticType = ESpatialAcceleration::Collection;
 	virtual FSpatialAccelerationIdx AddSubstructure(TUniquePtr<ISpatialAcceleration<TPayloadType, T, d>>&& Substructure, uint16 Bucket) = 0;
@@ -31,8 +32,13 @@ public:
 	{
 		return (1 << BucketIdx) & ActiveBucketsMask;
 	}
+
+	bool IsAllAsyncTrasksComplete() const { return AllAsyncTrasksComplete; }
+	void SetAllAsyncTrasksComplete(bool State) { AllAsyncTrasksComplete = State; }
+
 protected:
 	uint8 ActiveBucketsMask;
+	bool AllAsyncTrasksComplete;
 };
 
 template <typename T>
@@ -405,7 +411,12 @@ public:
 	virtual ISpatialAcceleration<TPayloadType, T, d>* GetSubstructure(FSpatialAccelerationIdx Idx) override
 	{
 		check(Idx.Bucket < MaxBuckets);
-		return Buckets[Idx.Bucket].Objects[Idx.InnerIdx].Acceleration.Get();
+		if (Buckets[Idx.Bucket].Objects.Num() > 0)
+		{
+			return Buckets[Idx.Bucket].Objects[Idx.InnerIdx].Acceleration.Get();
+		}
+
+		return nullptr;
 	}
 
 	virtual void Raycast(const TVector<T, d>& Start, const TVector<T, d>& Dir, const T Length, ISpatialVisitor<TPayloadType, T>& Visitor) const override
