Core types file added to ChaosCore to enable Chaos to move away from the <T,d> pattern attached to most classes.

diff --git a/Engine/Source/Runtime/Experimental/ChaosCore/Public/Chaos/Core.h b/Engine/Source/Runtime/Experimental/ChaosCore/Public/Chaos/Core.h
new file mode 100644
--- /dev/null
+++ b/Engine/Source/Runtime/Experimental/ChaosCore/Public/Chaos/Core.h
@@ -0,0 +1,46 @@
+// Copyright 1998-2019 Epic Games, Inc. All Rights Reserved.
+
+#pragma once
+
+#include "Vector.h"
+#include "Matrix.h"
+
+namespace Chaos
+{
+	/**
+	 * Common data types for the Chaos physics engine. Unless a specific
+	 * precision of real type is required most code should use the R types
+	 * (e.g. FVectorR3) to adapt to global changes in precision.
+	 */
+	using FReal = float;
+
+	using FVectorR2 = TVector<FReal, 2>;
+	using FVectorR3 = TVector<FReal, 3>;
+	using FVectorR4 = TVector<FReal, 4>;
+
+	using FVectorF2 = TVector<float, 2>;
+	using FVectorF3 = TVector<float, 3>;
+	using FVectorF4 = TVector<float, 4>;
+
+	using FVectorD2 = TVector<double, 2>;
+	using FVectorD3 = TVector<double, 3>;
+	using FVectorD4 = TVector<double, 4>;
+
+	using FVectorI2 = TVector<int32, 2>;
+	using FVectorI3 = TVector<int32, 3>;
+
+	using FVectorU2 = TVector<uint32, 2>;
+	using FVectorU3 = TVector<uint32, 3>;
+
+	using FMatrixR22 = PMatrix<FReal, 2, 2>;
+	using FMatrixR33 = PMatrix<FReal, 3, 3>;
+	using FMatrixR44 = PMatrix<FReal, 4, 4>;
+
+	using FMatrixF22 = PMatrix<float, 2, 2>;
+	using FMatrixF33 = PMatrix<float, 3, 3>;
+	using FMatrixF44 = PMatrix<float, 4, 4>;
+
+	using FMatrixD22 = PMatrix<double, 2, 2>;
+	using FMatrixD33 = PMatrix<double, 3, 3>;
+	using FMatrixD44 = PMatrix<double, 4, 4>;
+}
\ No newline at end of file
