Fix crash when BV is empty

diff --git a/Engine/Source/Runtime/Experimental/Chaos/Public/Chaos/BoundingVolume.h b/Engine/Source/Runtime/Experimental/Chaos/Public/Chaos/BoundingVolume.h
--- a/Engine/Source/Runtime/Experimental/Chaos/Public/Chaos/BoundingVolume.h
+++ b/Engine/Source/Runtime/Experimental/Chaos/Public/Chaos/BoundingVolume.h
@@ -837,6 +837,10 @@ private:
 	template <typename ParticleView>
 	void GenerateTree(const ParticleView& Particles, const bool bUseVelocity, const T Dt, const int32 MaxCells)
 	{
+		MGlobalPayloads.Reset();
+		MPayloadInfo.Reset();
+		bIsEmpty = true;
+
 		if (!Particles.Num())
 		{
 			return;
@@ -865,11 +869,9 @@ private:
 		};
 
 		//compute DX and fill global payloads
-		MGlobalPayloads.Reset();
 		auto& GlobalPayloads = MGlobalPayloads;
 		auto& PayloadInfos = MPayloadInfo;
 		T NumObjectsWithBounds = 0;
-		MPayloadInfo.Reset();
 
 		auto ComputeBoxAndDx = [&Particles, &AllBounds, &HasBounds, &GlobalPayloads, &PayloadInfos, &GetValidBounds, &NumObjectsWithBounds](TBox<T,d>& OutGlobalBox, bool bFirstPass) -> T
 		{
