Temp version of GJKRaycast2 requires a new support function

diff --git a/Engine/Source/Runtime/Experimental/Chaos/Public/Chaos/Capsule.h b/Engine/Source/Runtime/Experimental/Chaos/Public/Chaos/Capsule.h
--- a/Engine/Source/Runtime/Experimental/Chaos/Public/Chaos/Capsule.h
+++ b/Engine/Source/Runtime/Experimental/Chaos/Public/Chaos/Capsule.h
@@ -292,6 +292,13 @@ namespace Chaos
 			return MSegment.Support(Direction, MRadius + Thickness);
 		}
 
+		virtual TVector<T, 3> Support2(const TVector<T, 3>& Direction) const override { return MSegment.Support2(Direction); }
+
+		virtual T GetMargin() const override
+		{
+			return MRadius;
+		}
+
 		FORCEINLINE void SerializeImp(FArchive& Ar)
 		{
 			TImplicitObject<T, 3>::SerializeImp(Ar);
diff --git a/Engine/Source/Runtime/Experimental/Chaos/Public/Chaos/GJK.h b/Engine/Source/Runtime/Experimental/Chaos/Public/Chaos/GJK.h
--- a/Engine/Source/Runtime/Experimental/Chaos/Public/Chaos/GJK.h
+++ b/Engine/Source/Runtime/Experimental/Chaos/Public/Chaos/GJK.h
@@ -559,6 +559,148 @@ namespace Chaos
 		return true;
 	}
 
+
+	/** Sweeps one geometry against the other
+	 @A The first geometry
+	 @B The second geometry
+	 @StartTM B's starting configuration in A's local space
+	 @RayDir The ray's direction (normalized)
+	 @RayLength The ray's length
+	 @OutTime The time along the ray when the objects first overlap
+	 @OutPosition The first point of impact (in A's local space) when the objects first overlap. Invalid if time of impact is 0
+	 @OutNormal The impact normal (in A's local space) when the objects first overlap. Invalid if time of impact is 0
+	 @ThicknessA The amount of geometry inflation for Geometry A (for example a capsule with radius 5 could pass in its core segnment and a thickness of 5)
+	 @InitialDir The first direction we use to search the CSO
+	 @ThicknessB The amount of geometry inflation for Geometry B (for example a sphere with radius 5 could pass in its center point and a thickness of 5)
+	 @return True if the geometries overlap during the sweep, False otherwise 
+	 @note If A overlaps B at the start of the ray ("initial overlap" condition) then this function returns true, and sets OutTime = 0, but does not set any other output variables.
+	 */
+
+	template <typename T, typename TGeometryA, typename TGeometryB>
+	bool GJKRaycast2(const TGeometryA& A, const TGeometryB& B, const TRigidTransform<T, 3>& StartTM, const TVector<T, 3>& RayDir, const T RayLength,
+		T& OutTime, TVector<T, 3>& OutPosition, TVector<T, 3>& OutNormal, const T ThicknessA = 0, const TVector<T, 3>& InitialDir = TVector<T, 3>(1, 0, 0), const T ThicknessB = 0)
+	{
+		ensure(FMath::IsNearlyEqual(RayDir.SizeSquared(), 1, KINDA_SMALL_NUMBER));
+		ensure(RayLength > 0);
+		check(A.IsConvex() && B.IsConvex());
+		const TVector<T, 3> StartPoint = StartTM.GetLocation();
+
+		TVector<T, 3> Simplex[4] = { TVector<T,3>(0), TVector<T,3>(0), TVector<T,3>(0), TVector<T,3>(0) };
+		TVector<T, 3> As[4] = { TVector<T,3>(0), TVector<T,3>(0), TVector<T,3>(0), TVector<T,3>(0) };
+		TVector<T, 3> Bs[4] = { TVector<T,3>(0), TVector<T,3>(0), TVector<T,3>(0), TVector<T,3>(0) };
+
+		T Barycentric[4];
+		const T Inflation = ThicknessA + ThicknessB;
+		const T Inflation2 = Inflation*Inflation + 1e-6;
+
+		FSimplex SimplexIDs;
+		const TRotation<T, 3> BToARotation = StartTM.GetRotation();
+		const TRotation<T, 3> AToBRotation = BToARotation.Inverse();
+		TVector<T, 3> SupportA = A.Support2(InitialDir);
+		As[0] = SupportA;
+
+		const TVector<T, 3> InitialDirInB = AToBRotation * (-InitialDir);
+		const TVector<T, 3> InitialSupportBLocal = B.Support2(InitialDirInB);
+		TVector<T, 3> SupportB = BToARotation * InitialSupportBLocal;
+		Bs[0] = SupportB;
+
+		T Lambda = 0;
+		TVector<T, 3> X = StartPoint;
+		TVector<T, 3> V = X - (SupportA - SupportB);
+		TVector<T, 3> Normal(0,0,1);
+
+		bool bCloseEnough = V.SizeSquared() < Inflation2;
+		bool bDegenerate = false;
+		bool bTerminate = bCloseEnough;
+		int NumIterations = 0;
+		T InGJKPreDist2 = TNumericLimits<T>::Max();
+		while (!bTerminate)
+		{
+			//if (!ensure(NumIterations++ < 32))	//todo: take this out
+			if (!(NumIterations++ < 32))	//todo: take this out
+			{
+				break;	//if taking too long just stop. This should never happen
+			}
+
+			SupportA = A.Support2(V);	//todo: add thickness to quadratic geometry to avoid quadratic vs quadratic when possible
+			const TVector<T, 3> VInB = AToBRotation * (-V);
+			const TVector<T, 3> SupportBLocal = B.Support2(VInB);
+			SupportB = BToARotation * SupportBLocal;
+			const TVector<T, 3> P = SupportA - SupportB;
+			const TVector<T, 3> W = X - P;
+			SimplexIDs[SimplexIDs.NumVerts] = SimplexIDs.NumVerts;	//is this needed?
+			As[SimplexIDs.NumVerts] = SupportA;
+			Bs[SimplexIDs.NumVerts] = SupportB;
+
+			V = V.GetUnsafeNormal();
+
+			const T VDotW = TVector<T, 3>::DotProduct(V, W);
+			if (VDotW > Inflation)
+			{
+				const T VDotRayDir = TVector<T, 3>::DotProduct(V, RayDir);
+				if (VDotRayDir >= 0)
+				{
+					return false;
+				}
+
+				const T PreLambda = Lambda;	//use to check for no progress
+				// @todo(ccaulfield): this can still overflow - the comparisons against zero above should be changed (though not sure to what yet)
+				Lambda = Lambda - (VDotW - Inflation) / VDotRayDir;
+				if (Lambda > PreLambda)
+				{
+					if (Lambda > RayLength)
+					{
+						return false;
+					}
+
+					const TVector<T, 3> OldX = X;
+					X = StartPoint + Lambda * RayDir;
+					Normal = V;
+
+					//Update simplex from (OldX - P) to (X - P)
+					const TVector<T, 3> XMinusOldX = X - OldX;
+					Simplex[0] += XMinusOldX;
+					Simplex[1] += XMinusOldX;
+					Simplex[2] += XMinusOldX;
+					Simplex[SimplexIDs.NumVerts++] = X - P;
+
+					InGJKPreDist2 = TNumericLimits<T>::Max();	//translated origin so restart gjk search
+				}
+			}
+			else
+			{
+				Simplex[SimplexIDs.NumVerts++] = W;	//this is really X - P which is what we need for simplex computation
+			}
+
+			V = SimplexFindClosestToOrigin(Simplex, SimplexIDs, Barycentric, As, Bs);
+
+			T NewDist2 = V.SizeSquared();	//todo: relative error
+			bCloseEnough = NewDist2 < Inflation2;
+			bDegenerate = NewDist2 >= InGJKPreDist2;
+			InGJKPreDist2 = NewDist2;
+			bTerminate = bCloseEnough || bDegenerate;
+		}
+
+		OutTime = Lambda;
+
+		if (Lambda > 0)
+		{
+			OutNormal = Normal;
+			TVector<T, 3> ClosestA(0);
+			TVector<T, 3> ClosestB(0);
+
+			for (int i = 0; i < SimplexIDs.NumVerts; ++i)
+			{
+				ClosestB += Bs[i] * Barycentric[i];
+			}
+			const TVector<T, 3> ClosestLocal = ClosestB - OutNormal * ThicknessB;
+
+			OutPosition = StartPoint + RayDir * Lambda + ClosestLocal;
+		}
+
+		return true;
+	}
+
 	/**
 	 * Used by GJKDistance. It must return a vector in the Minkowski sum A - B. In principle this can be the vector of any point
 	 * in A to any point in B, but some choices will cause GJK to minimize faster (e.g., for two spheres, we can easily calculate
diff --git a/Engine/Source/Runtime/Experimental/Chaos/Public/Chaos/GeometryQueries.h b/Engine/Source/Runtime/Experimental/Chaos/Public/Chaos/GeometryQueries.h
--- a/Engine/Source/Runtime/Experimental/Chaos/Public/Chaos/GeometryQueries.h
+++ b/Engine/Source/Runtime/Experimental/Chaos/Public/Chaos/GeometryQueries.h
@@ -156,7 +156,8 @@ namespace Chaos
 
 			const TVector<T, d> Offset = ATM.GetLocation() - BTM.GetLocation();
 
-			bResult = GJKRaycast<T>(A, B, BToATM, LocalDir, Length, OutTime, LocalPosition, LocalNormal, Thickness, Offset);
+			bResult = GJKRaycast2<T>(A, B, BToATM, LocalDir, Length, OutTime, LocalPosition, LocalNormal, Thickness + A.GetMargin(), Offset, Thickness + B.GetMargin());
+			//bResult = GJKRaycast<T>(A, B, BToATM, LocalDir, Length, OutTime, LocalPosition, LocalNormal, Thickness, Offset, Thickness);
 
 			if (AType == ImplicitObjectType::Convex)
 			{
diff --git a/Engine/Source/Runtime/Experimental/Chaos/Public/Chaos/ImplicitObject.h b/Engine/Source/Runtime/Experimental/Chaos/Public/Chaos/ImplicitObject.h
--- a/Engine/Source/Runtime/Experimental/Chaos/Public/Chaos/ImplicitObject.h
+++ b/Engine/Source/Runtime/Experimental/Chaos/Public/Chaos/ImplicitObject.h
@@ -152,6 +152,7 @@ public:
 	virtual T PhiWithNormal(const TVector<T, d>& x, TVector<T, d>& Normal) const = 0;
 	virtual const class TBox<T, d>& BoundingBox() const;
 	virtual TVector<T, d> Support(const TVector<T, d>& Direction, const T Thickness) const;
+	virtual TVector<T, d> Support2(const TVector<T, d>& Direction) const { return Support(Direction, 0); }
 	bool HasBoundingBox() const { return bHasBoundingBox; }
 	bool IsConvex() const { return bIsConvex; }
 	void IgnoreAnalyticCollisions(const bool Ignore = true) { bIgnoreAnalyticCollisions = Ignore; }
@@ -224,6 +225,11 @@ public:
 		Out.Add(MakePair(This, ParentTM));
 	}
 
+	virtual T GetMargin() const
+	{
+		return 0;
+	}
+
 	virtual void FindAllIntersectingObjects(TArray < Pair<const TImplicitObject<T, d>*, TRigidTransform<T, d>>>& Out, const TBox<T, d>& LocalBounds) const;
 
 	virtual FString ToString() const
diff --git a/Engine/Source/Runtime/Experimental/Chaos/Public/Chaos/Segment.h b/Engine/Source/Runtime/Experimental/Chaos/Public/Chaos/Segment.h
--- a/Engine/Source/Runtime/Experimental/Chaos/Public/Chaos/Segment.h
+++ b/Engine/Source/Runtime/Experimental/Chaos/Public/Chaos/Segment.h
@@ -41,6 +41,13 @@ namespace Chaos
 			return FarthestCap + (NormalizedDirection * Thickness);
 		}
 
+		TVector<T, 3> Support2(const TVector<T, 3>& Direction) const
+		{
+			const T Dot = TVector<T, 3>::DotProduct(Direction, MAxis);
+			const TVector<T, 3> FarthestCap = Dot >= 0 ? GetX2() : GetX1();	//orthogonal we choose either
+			return FarthestCap;
+		}
+
 		FORCEINLINE void Serialize(FArchive &Ar) 
 		{
 			Ar << MPoint << MAxis << MLength;
diff --git a/Engine/Source/Runtime/Experimental/Chaos/Public/Chaos/Sphere.h b/Engine/Source/Runtime/Experimental/Chaos/Public/Chaos/Sphere.h
--- a/Engine/Source/Runtime/Experimental/Chaos/Public/Chaos/Sphere.h
+++ b/Engine/Source/Runtime/Experimental/Chaos/Public/Chaos/Sphere.h
@@ -191,6 +191,13 @@ namespace Chaos
 			return Center + Normalized * (Radius + Thickness);
 		}
 
+		virtual TVector<T, d> Support2(const TVector<T, d>& Direction) const override { return Center; }
+
+		virtual T GetMargin() const override
+		{
+			return Radius;
+		}
+
 		virtual const TBox<T, d>& BoundingBox() const 
 		{
 			return LocalBoundingBox; 
diff --git a/Engine/Source/Runtime/Experimental/Chaos/Public/Chaos/Triangle.h b/Engine/Source/Runtime/Experimental/Chaos/Public/Chaos/Triangle.h
--- a/Engine/Source/Runtime/Experimental/Chaos/Public/Chaos/Triangle.h
+++ b/Engine/Source/Runtime/Experimental/Chaos/Public/Chaos/Triangle.h
@@ -47,6 +47,11 @@ namespace Chaos
 			return TVector<T, 3>::DotProduct((InSamplePoint - ClosestPoint), OutNormal);
 		}
 
+		FORCEINLINE TVector<T, 3> Support2(const TVector<T, 3>& Direction) const
+		{
+			return Support(Direction, 0);
+		}
+
 		FORCEINLINE TVector<T, 3> Support(const TVector<T, 3>& Direction, const T Thickness) const
 		{
 			const float DotA = TVector<T, 3>::DotProduct(A, Direction);
