Move geometry type over to uint8 so we can easily use Instanced and Scaled

diff --git a/Engine/Source/Runtime/Engine/Private/PhysicsEngine/Experimental/PhysInterface_Chaos.cpp b/Engine/Source/Runtime/Engine/Private/PhysicsEngine/Experimental/PhysInterface_Chaos.cpp
--- a/Engine/Source/Runtime/Engine/Private/PhysicsEngine/Experimental/PhysInterface_Chaos.cpp
+++ b/Engine/Source/Runtime/Engine/Private/PhysicsEngine/Experimental/PhysInterface_Chaos.cpp
@@ -1220,7 +1220,7 @@ bool FPhysInterface_Chaos::IsShapeType(const FPhysicsShapeReference_Chaos& InSha
 
 ECollisionShapeType FPhysInterface_Chaos::GetShapeType(const FPhysicsShapeReference_Chaos& InShapeRef, bool bGetInnerType)
 {
-	Chaos::ImplicitObjectType Type = InShapeRef.Shape->Geometry->GetType(true);
+	Chaos::EImplicitObjectType Type = InShapeRef.Shape->Geometry->GetType(true);
 	if (bGetInnerType && Type == Chaos::ImplicitObjectType::Scaled)
 	{
 		Type = InShapeRef.Shape->Geometry->GetObject<Chaos::TImplicitObjectScaled<float, 3>>()->GetUnscaledObject()->GetType(true);
diff --git a/Engine/Source/Runtime/Experimental/Chaos/Private/Chaos/ImplicitObject.cpp b/Engine/Source/Runtime/Experimental/Chaos/Private/Chaos/ImplicitObject.cpp
--- a/Engine/Source/Runtime/Experimental/Chaos/Private/Chaos/ImplicitObject.cpp
+++ b/Engine/Source/Runtime/Experimental/Chaos/Private/Chaos/ImplicitObject.cpp
@@ -19,7 +19,7 @@
 using namespace Chaos;
 
 template<class T, int d>
-TImplicitObject<T, d>::TImplicitObject(int32 Flags, ImplicitObjectType InType)
+TImplicitObject<T, d>::TImplicitObject(int32 Flags, EImplicitObjectType InType)
     : Type(InType)
     , bIsConvex(!!(Flags & EImplicitObject::IsConvex))
     , bIgnoreAnalyticCollisions(!!(Flags & EImplicitObject::IgnoreAnalyticCollisions))
@@ -33,7 +33,7 @@ TImplicitObject<T, d>::~TImplicitObject()
 }
 
 template<class T, int d>
-ImplicitObjectType TImplicitObject<T, d>::GetType(bool bGetTrueType) const
+EImplicitObjectType TImplicitObject<T, d>::GetType(bool bGetTrueType) const
 {
 	if (bIgnoreAnalyticCollisions && !bGetTrueType)
 	{
@@ -294,9 +294,9 @@ FArchive& TImplicitObject<T, d>::SerializeLegacyHelper(FArchive& Ar, TUniquePtr<
 	{
 		if (Ar.IsLoading())
 		{
-			int8 ObjectType;
+			uint8 ObjectType;
 			Ar << ObjectType;
-			switch ((ImplicitObjectType)ObjectType)
+			switch ((EImplicitObjectType)ObjectType)
 			{
 			case ImplicitObjectType::Sphere: { Value = TUniquePtr<TSphere<T, d>>(new TSphere<T, d>()); break; }
 			case ImplicitObjectType::Box: { Value = TUniquePtr<TBox<T, d>>(new TBox<T, d>()); break; }
@@ -338,7 +338,7 @@ void TImplicitObject<T, d>::Serialize(FChaosArchive& Ar)
 }
 
 template<typename T, int d>
-const FName TImplicitObject<T, d>::GetTypeName(const ImplicitObjectType InType)
+const FName TImplicitObject<T, d>::GetTypeName(const EImplicitObjectType InType)
 {
 	static const FName SphereName = TEXT("Sphere");
 	static const FName BoxName = TEXT("Box");
@@ -380,7 +380,7 @@ TImplicitObject<T, d>* TImplicitObject<T, d>::SerializationFactory(FChaosArchive
 {
 	int8 ObjectType = Ar.IsLoading() ? 0 : (int8)Obj->Type;
 	Ar << ObjectType;
-	switch ((ImplicitObjectType)ObjectType)
+	switch ((EImplicitObjectType)ObjectType)
 	{
 	case ImplicitObjectType::Sphere: if (Ar.IsLoading()) { return new TSphere<T, d>(); } break;
 	case ImplicitObjectType::Box: if (Ar.IsLoading()) { return new TBox<T, d>(); } break;
diff --git a/Engine/Source/Runtime/Experimental/Chaos/Public/Chaos/Box.h b/Engine/Source/Runtime/Experimental/Chaos/Public/Chaos/Box.h
--- a/Engine/Source/Runtime/Experimental/Chaos/Public/Chaos/Box.h
+++ b/Engine/Source/Runtime/Experimental/Chaos/Public/Chaos/Box.h
@@ -120,7 +120,7 @@ namespace Chaos
 			return AABB.Contains(Point, Tolerance);
 		}
 
-		FORCEINLINE static ImplicitObjectType GetType() { return ImplicitObjectType::Box; }
+		FORCEINLINE static EImplicitObjectType GetType() { return ImplicitObjectType::Box; }
 
 		const TBox<T, d>& BoundingBox() const { return *this; }
 
diff --git a/Engine/Source/Runtime/Experimental/Chaos/Public/Chaos/Capsule.h b/Engine/Source/Runtime/Experimental/Chaos/Public/Chaos/Capsule.h
--- a/Engine/Source/Runtime/Experimental/Chaos/Public/Chaos/Capsule.h
+++ b/Engine/Source/Runtime/Experimental/Chaos/Public/Chaos/Capsule.h
@@ -74,7 +74,7 @@ namespace Chaos
 
 		~TCapsule() {}
 
-		static ImplicitObjectType GetType() { return ImplicitObjectType::Capsule; }
+		static EImplicitObjectType GetType() { return ImplicitObjectType::Capsule; }
 
 		static TCapsule<T> NewFromOriginAndAxis(const TVector<T, 3>& Origin, const TVector<T, 3>& Axis, const T Height, const T Radius)
 		{
diff --git a/Engine/Source/Runtime/Experimental/Chaos/Public/Chaos/Convex.h b/Engine/Source/Runtime/Experimental/Chaos/Public/Chaos/Convex.h
--- a/Engine/Source/Runtime/Experimental/Chaos/Public/Chaos/Convex.h
+++ b/Engine/Source/Runtime/Experimental/Chaos/Public/Chaos/Convex.h
@@ -37,7 +37,7 @@ namespace Chaos
 			TConvexBuilder<T>::Build(InParticles, Planes, SurfaceParticles, LocalBoundingBox);
 		}
 
-		static ImplicitObjectType GetType()
+		static EImplicitObjectType GetType()
 		{
 			return ImplicitObjectType::Convex;
 		}
diff --git a/Engine/Source/Runtime/Experimental/Chaos/Public/Chaos/Cylinder.h b/Engine/Source/Runtime/Experimental/Chaos/Public/Chaos/Cylinder.h
--- a/Engine/Source/Runtime/Experimental/Chaos/Public/Chaos/Cylinder.h
+++ b/Engine/Source/Runtime/Experimental/Chaos/Public/Chaos/Cylinder.h
@@ -50,7 +50,7 @@ namespace Chaos
 		}
 		~TCylinder() {}
 
-		static ImplicitObjectType GetType() { return ImplicitObjectType::Cylinder; }
+		static EImplicitObjectType GetType() { return ImplicitObjectType::Cylinder; }
 
 		/**
 		 * Returns sample points centered about the origin.
diff --git a/Engine/Source/Runtime/Experimental/Chaos/Public/Chaos/GeometryQueries.h b/Engine/Source/Runtime/Experimental/Chaos/Public/Chaos/GeometryQueries.h
--- a/Engine/Source/Runtime/Experimental/Chaos/Public/Chaos/GeometryQueries.h
+++ b/Engine/Source/Runtime/Experimental/Chaos/Public/Chaos/GeometryQueries.h
@@ -20,8 +20,8 @@ namespace Chaos
 	template <typename T, int d>
 	bool OverlapQuery(const TImplicitObject<T, d>& A, const TRigidTransform<T,d>& ATM, const TImplicitObject<T, d>& B, const TRigidTransform<T,d>& BTM, const T Thickness = 0, TVector<T,d>* OutMTD=nullptr)
 	{
-		const ImplicitObjectType AType = A.GetType(true);
-		const ImplicitObjectType BType = B.GetType(true);
+		const EImplicitObjectType AType = A.GetType(true);
+		const EImplicitObjectType BType = B.GetType(true);
 		
 		if (AType == ImplicitObjectType::Transformed)
 		{
@@ -86,8 +86,8 @@ namespace Chaos
 	template <typename T, int d>
 	bool SweepQuery(const TImplicitObject<T, d>& A, const TRigidTransform<T,d>& ATM, const TImplicitObject<T, d>& B, const TRigidTransform<T, d>& BTM, const TVector<T,d>& Dir, const T Length, T& OutTime, TVector<T,d>& OutPosition, TVector<T,d>& OutNormal, int32& OutFaceIndex, const T Thickness, const bool bComputeMTD)
 	{
-		const ImplicitObjectType AType = A.GetType(true);
-		const ImplicitObjectType BType = B.GetType(true);
+		const EImplicitObjectType AType = A.GetType(true);
+		const EImplicitObjectType BType = B.GetType(true);
 
 		bool bResult = false;
 		
diff --git a/Engine/Source/Runtime/Experimental/Chaos/Public/Chaos/HeightField.h b/Engine/Source/Runtime/Experimental/Chaos/Public/Chaos/HeightField.h
--- a/Engine/Source/Runtime/Experimental/Chaos/Public/Chaos/HeightField.h
+++ b/Engine/Source/Runtime/Experimental/Chaos/Public/Chaos/HeightField.h
@@ -70,7 +70,7 @@ namespace Chaos
 			return FCrc::MemCrc32(Bytes.GetData(), Bytes.GetAllocatedSize());
 		}
 
-		static ImplicitObjectType GetType()
+		static EImplicitObjectType GetType()
 		{
 			return ImplicitObjectType::HeightField;
 		}
diff --git a/Engine/Source/Runtime/Experimental/Chaos/Public/Chaos/ImplicitObject.h b/Engine/Source/Runtime/Experimental/Chaos/Public/Chaos/ImplicitObject.h
--- a/Engine/Source/Runtime/Experimental/Chaos/Public/Chaos/ImplicitObject.h
+++ b/Engine/Source/Runtime/Experimental/Chaos/Public/Chaos/ImplicitObject.h
@@ -26,26 +26,34 @@ class TBVHParticles;
 template<class T, int d>
 class TImplicitObject;
 
-
-enum class ImplicitObjectType : int8
+namespace ImplicitObjectType
 {
-	//Note: add entries at the bottom for serialization
-	Sphere = 0,
-	Box,
-	Plane,
-	Capsule,
-	Transformed,
-	Union,
-	LevelSet,
-	Unknown,
-	Convex,
-	TaperedCylinder,
-	Cylinder,
-	TriangleMesh,
-	HeightField,
-	Scaled,
-	Triangle
-};
+	enum
+	{
+		//Note: add entries in order to avoid serialization issues (but before IsInstanced)
+		Sphere = 0,
+		Box,
+		Plane,
+		Capsule,
+		Transformed,
+		Union,
+		LevelSet,
+		Unknown,
+		Convex,
+		TaperedCylinder,
+		Cylinder,
+		TriangleMesh,
+		HeightField,
+		Scaled,
+		Triangle,
+
+		//Add entries above this line for serialization
+		IsInstanced = 1 << 6,
+		IsScaled = 1 << 7
+	};
+}
+
+using EImplicitObjectType = uint8;	//see ImplicitObjectType
 
 namespace EImplicitObject
 {
@@ -96,7 +104,7 @@ class CHAOS_API TImplicitObject
 public:
 	static TImplicitObject<T,d>* SerializationFactory(FChaosArchive& Ar, TImplicitObject<T, d>* Obj);
 
-	TImplicitObject(int32 Flags, ImplicitObjectType InType = ImplicitObjectType::Unknown);
+	TImplicitObject(int32 Flags, EImplicitObjectType InType = ImplicitObjectType::Unknown);
 	TImplicitObject(const TImplicitObject<T, d>&) = delete;
 	TImplicitObject(TImplicitObject<T, d>&&) = delete;
 	virtual ~TImplicitObject();
@@ -135,7 +143,7 @@ public:
 		return static_cast<const T_DERIVED&>(*this);
 	}
 
-	ImplicitObjectType GetType(bool bGetTrueType = false) const;
+	EImplicitObjectType GetType(bool bGetTrueType = false) const;
 
 	virtual bool IsValidGeometry() const;
 
@@ -252,10 +260,10 @@ public:
 
 	virtual FName GetTypeName() const { return GetTypeName(GetType()); }
 
-	static const FName GetTypeName(const ImplicitObjectType InType);
+	static const FName GetTypeName(const EImplicitObjectType InType);
 
 protected:
-	ImplicitObjectType Type;
+	EImplicitObjectType Type;
 	bool bIsConvex;
 	bool bIgnoreAnalyticCollisions;
 	bool bHasBoundingBox;
diff --git a/Engine/Source/Runtime/Experimental/Chaos/Public/Chaos/ImplicitObjectScaled.h b/Engine/Source/Runtime/Experimental/Chaos/Public/Chaos/ImplicitObjectScaled.h
--- a/Engine/Source/Runtime/Experimental/Chaos/Public/Chaos/ImplicitObjectScaled.h
+++ b/Engine/Source/Runtime/Experimental/Chaos/Public/Chaos/ImplicitObjectScaled.h
@@ -61,7 +61,7 @@ public:
 	}
 	~TImplicitObjectScaled() {}
 
-	static ImplicitObjectType GetType()
+	static EImplicitObjectType GetType()
 	{
 		return ImplicitObjectType::Scaled;
 	}
diff --git a/Engine/Source/Runtime/Experimental/Chaos/Public/Chaos/ImplicitObjectTransformed.h b/Engine/Source/Runtime/Experimental/Chaos/Public/Chaos/ImplicitObjectTransformed.h
--- a/Engine/Source/Runtime/Experimental/Chaos/Public/Chaos/ImplicitObjectTransformed.h
+++ b/Engine/Source/Runtime/Experimental/Chaos/Public/Chaos/ImplicitObjectTransformed.h
@@ -86,7 +86,7 @@ public:
 	}
 	~TImplicitObjectTransformed() {}
 
-	static ImplicitObjectType GetType()
+	static EImplicitObjectType GetType()
 	{
 		return ImplicitObjectType::Transformed;
 	}
diff --git a/Engine/Source/Runtime/Experimental/Chaos/Public/Chaos/ImplicitObjectUnion.h b/Engine/Source/Runtime/Experimental/Chaos/Public/Chaos/ImplicitObjectUnion.h
--- a/Engine/Source/Runtime/Experimental/Chaos/Public/Chaos/ImplicitObjectUnion.h
+++ b/Engine/Source/Runtime/Experimental/Chaos/Public/Chaos/ImplicitObjectUnion.h
@@ -69,7 +69,7 @@ class TImplicitObjectUnion : public TImplicitObject<T, d>
 	}
 	virtual ~TImplicitObjectUnion() {}
 
-	FORCEINLINE static ImplicitObjectType GetType()
+	FORCEINLINE static EImplicitObjectType GetType()
 	{
 		return ImplicitObjectType::Union;
 	}
diff --git a/Engine/Source/Runtime/Experimental/Chaos/Public/Chaos/Levelset.h b/Engine/Source/Runtime/Experimental/Chaos/Public/Chaos/Levelset.h
--- a/Engine/Source/Runtime/Experimental/Chaos/Public/Chaos/Levelset.h
+++ b/Engine/Source/Runtime/Experimental/Chaos/Public/Chaos/Levelset.h
@@ -54,7 +54,7 @@ class CHAOS_API TLevelSet final : public TImplicitObject<T, d>
 		}
 	}
 
-	FORCEINLINE static ImplicitObjectType GetType()
+	FORCEINLINE static EImplicitObjectType GetType()
 	{
 		return ImplicitObjectType::LevelSet;
 	}
diff --git a/Engine/Source/Runtime/Experimental/Chaos/Public/Chaos/Plane.h b/Engine/Source/Runtime/Experimental/Chaos/Public/Chaos/Plane.h
--- a/Engine/Source/Runtime/Experimental/Chaos/Public/Chaos/Plane.h
+++ b/Engine/Source/Runtime/Experimental/Chaos/Public/Chaos/Plane.h
@@ -34,7 +34,7 @@ class TPlane final : public TImplicitObject<T, d>
 	}
 	virtual ~TPlane() {}
 
-	static ImplicitObjectType GetType()
+	static EImplicitObjectType GetType()
 	{
 		return ImplicitObjectType::Plane;
 	}
diff --git a/Engine/Source/Runtime/Experimental/Chaos/Public/Chaos/Sphere.h b/Engine/Source/Runtime/Experimental/Chaos/Public/Chaos/Sphere.h
--- a/Engine/Source/Runtime/Experimental/Chaos/Public/Chaos/Sphere.h
+++ b/Engine/Source/Runtime/Experimental/Chaos/Public/Chaos/Sphere.h
@@ -62,7 +62,7 @@ namespace Chaos
 
 		virtual ~TSphere() {}
 
-		static ImplicitObjectType GetType() 
+		static EImplicitObjectType GetType()
 		{ 
 			return ImplicitObjectType::Sphere; 
 		}
diff --git a/Engine/Source/Runtime/Experimental/Chaos/Public/Chaos/TaperedCylinder.h b/Engine/Source/Runtime/Experimental/Chaos/Public/Chaos/TaperedCylinder.h
--- a/Engine/Source/Runtime/Experimental/Chaos/Public/Chaos/TaperedCylinder.h
+++ b/Engine/Source/Runtime/Experimental/Chaos/Public/Chaos/TaperedCylinder.h
@@ -58,7 +58,7 @@ namespace Chaos
 		}
 		~TTaperedCylinder() {}
 
-		static ImplicitObjectType GetType() { return ImplicitObjectType::TaperedCylinder; }
+		static EImplicitObjectType GetType() { return ImplicitObjectType::TaperedCylinder; }
 
 		/**
 		 * Returns sample points centered about the origin.
diff --git a/Engine/Source/Runtime/Experimental/Chaos/Public/Chaos/Triangle.h b/Engine/Source/Runtime/Experimental/Chaos/Public/Chaos/Triangle.h
--- a/Engine/Source/Runtime/Experimental/Chaos/Public/Chaos/Triangle.h
+++ b/Engine/Source/Runtime/Experimental/Chaos/Public/Chaos/Triangle.h
@@ -165,7 +165,7 @@ namespace Chaos
 			return Tri.GetPlane();
 		}
 
-		static ImplicitObjectType GetType()
+		static EImplicitObjectType GetType()
 		{
 			return ImplicitObjectType::Triangle;
 		}
diff --git a/Engine/Source/Runtime/Experimental/Chaos/Public/Chaos/TriangleMeshImplicitObject.h b/Engine/Source/Runtime/Experimental/Chaos/Public/Chaos/TriangleMeshImplicitObject.h
--- a/Engine/Source/Runtime/Experimental/Chaos/Public/Chaos/TriangleMeshImplicitObject.h
+++ b/Engine/Source/Runtime/Experimental/Chaos/Public/Chaos/TriangleMeshImplicitObject.h
@@ -44,7 +44,7 @@ namespace Chaos
 			return MLocalBoundingBox;
 		}
 
-		static ImplicitObjectType GetType()
+		static EImplicitObjectType GetType()
 		{
 			return ImplicitObjectType::TriangleMesh;
 		}
diff --git a/Engine/Source/Runtime/Experimental/ChaosSolvers/Private/PhysicsProxy/SkeletalMeshPhysicsProxy.cpp b/Engine/Source/Runtime/Experimental/ChaosSolvers/Private/PhysicsProxy/SkeletalMeshPhysicsProxy.cpp
--- a/Engine/Source/Runtime/Experimental/ChaosSolvers/Private/PhysicsProxy/SkeletalMeshPhysicsProxy.cpp
+++ b/Engine/Source/Runtime/Experimental/ChaosSolvers/Private/PhysicsProxy/SkeletalMeshPhysicsProxy.cpp
@@ -237,7 +237,7 @@ void FSkeletalMeshPhysicsProxy::CreateRigidBodyCallback(FParticlesType& Particle
 			Particles.CollisionGroup(RigidBodyId) = Parameters.CollisionGroup;
 			
 			// Set implicit surface
-			const Chaos::ImplicitObjectType ShapeType = ImplicitObject->GetType();
+			const Chaos::EImplicitObjectType ShapeType = ImplicitObject->GetType();
 			if (ShapeType == Chaos::ImplicitObjectType::LevelSet)
 			{
 				FCollisionStructureManager::UpdateImplicitFlags(
diff --git a/Engine/Source/Runtime/Experimental/GeometryCollectionEngine/Private/GeometryCollection/GeometryCollectionDebugDrawActor.cpp b/Engine/Source/Runtime/Experimental/GeometryCollectionEngine/Private/GeometryCollection/GeometryCollectionDebugDrawActor.cpp
--- a/Engine/Source/Runtime/Experimental/GeometryCollectionEngine/Private/GeometryCollection/GeometryCollectionDebugDrawActor.cpp
+++ b/Engine/Source/Runtime/Experimental/GeometryCollectionEngine/Private/GeometryCollection/GeometryCollectionDebugDrawActor.cpp
@@ -2428,7 +2428,7 @@ void AGeometryCollectionDebugDrawActor::DrawRigidBodyCollisionNoChecks(const UGe
 	const FColor& ActiveColor = (bIsDisabled && !bIsParentUnion) ? FColor::Black: Color;
 
 	// Draw collision volume
-	const Chaos::ImplicitObjectType GeometryType = ParticlesData.GetGeometryType(TransformIndex);
+	const Chaos::EImplicitObjectType GeometryType = ParticlesData.GetGeometryType(TransformIndex);
 	switch (GeometryType)
 	{
 	case Chaos::ImplicitObjectType::Box:
diff --git a/Engine/Source/Runtime/Experimental/GeometryCollectionEngine/Private/GeometryCollection/GeometryCollectionParticlesData.cpp b/Engine/Source/Runtime/Experimental/GeometryCollectionEngine/Private/GeometryCollection/GeometryCollectionParticlesData.cpp
--- a/Engine/Source/Runtime/Experimental/GeometryCollectionEngine/Private/GeometryCollection/GeometryCollectionParticlesData.cpp
+++ b/Engine/Source/Runtime/Experimental/GeometryCollectionEngine/Private/GeometryCollection/GeometryCollectionParticlesData.cpp
@@ -208,7 +208,7 @@ void TGeometryCollectionParticlesData<T, d>::FData::Copy(EGeometryCollectionPart
 template<class T, int d>
 FString TGeometryCollectionParticlesData<T, d>::FData::ToString(int32 Index, const TCHAR* Separator) const
 {
-	auto TypeText = [](Chaos::ImplicitObjectType Type)
+	auto TypeText = [](Chaos::EImplicitObjectType Type)
 	{
 		switch (Type)
 		{
diff --git a/Engine/Source/Runtime/Experimental/GeometryCollectionEngine/Public/GeometryCollection/GeometryCollectionParticlesData.h b/Engine/Source/Runtime/Experimental/GeometryCollectionEngine/Public/GeometryCollection/GeometryCollectionParticlesData.h
--- a/Engine/Source/Runtime/Experimental/GeometryCollectionEngine/Public/GeometryCollection/GeometryCollectionParticlesData.h
+++ b/Engine/Source/Runtime/Experimental/GeometryCollectionEngine/Public/GeometryCollection/GeometryCollectionParticlesData.h
@@ -83,7 +83,7 @@ public:
 	const Chaos::TVector<T, d>               & GetX                     (int32 Index) const { check(HasSyncedData(EGeometryCollectionParticlesData::X                     )); return BufferedData.GetGameDataForRead().X                     [Index]; }
 	const Chaos::TRotation<T, d>             & GetR                     (int32 Index) const { check(HasSyncedData(EGeometryCollectionParticlesData::R                     )); return BufferedData.GetGameDataForRead().R                     [Index]; }
 	const Chaos::TImplicitObject<T, d>* const& GetGeometry              (int32 Index) const { check(HasSyncedData(EGeometryCollectionParticlesData::Geometry              )); return BufferedData.GetGameDataForRead().Geometry              [Index]; }
-	const Chaos::ImplicitObjectType          & GetGeometryType          (int32 Index) const { check(HasSyncedData(EGeometryCollectionParticlesData::GeometryType          )); return BufferedData.GetGameDataForRead().GeometryType          [Index]; }
+	const Chaos::EImplicitObjectType         & GetGeometryType          (int32 Index) const { check(HasSyncedData(EGeometryCollectionParticlesData::GeometryType          )); return BufferedData.GetGameDataForRead().GeometryType          [Index]; }
 	const bool                               & IsGeometryConvex         (int32 Index) const { check(HasSyncedData(EGeometryCollectionParticlesData::GeometryIsConvex      )); return BufferedData.GetGameDataForRead().GeometryIsConvex      [Index]; }
 	const bool                               & HasGeometryBoundingBoxm  (int32 Index) const { check(HasSyncedData(EGeometryCollectionParticlesData::GeometryHasBoundingBox)); return BufferedData.GetGameDataForRead().GeometryHasBoundingBox[Index]; }
 	const Chaos::TVector<T, d>               & GetGeometryBoxMin        (int32 Index) const { check(HasSyncedData(EGeometryCollectionParticlesData::GeometryBoxMin        )); return BufferedData.GetGameDataForRead().GeometryBoxMin        [Index]; }
@@ -125,7 +125,7 @@ private:
 		Chaos::TArrayCollectionArray<Chaos::TVector<T, d>               > X;
 		Chaos::TArrayCollectionArray<Chaos::TRotation<T, d>             > R;
 		Chaos::TArrayCollectionArray<const Chaos::TImplicitObject<T, d>*> Geometry;
-		Chaos::TArrayCollectionArray<Chaos::ImplicitObjectType          > GeometryType;
+		Chaos::TArrayCollectionArray<Chaos::EImplicitObjectType         > GeometryType;
 		Chaos::TArrayCollectionArray<bool                               > GeometryIsConvex;
 		Chaos::TArrayCollectionArray<bool                               > GeometryHasBoundingBox;
 		Chaos::TArrayCollectionArray<Chaos::TVector<T, d>               > GeometryBoxMin;
