Added Instanced implicit object, not hooked up yet

diff --git a/Engine/Source/Runtime/Experimental/Chaos/Public/Chaos/ImplicitObject.h b/Engine/Source/Runtime/Experimental/Chaos/Public/Chaos/ImplicitObject.h
--- a/Engine/Source/Runtime/Experimental/Chaos/Public/Chaos/ImplicitObject.h
+++ b/Engine/Source/Runtime/Experimental/Chaos/Public/Chaos/ImplicitObject.h
@@ -55,6 +55,17 @@ namespace ImplicitObjectType
 
 using EImplicitObjectType = uint8;	//see ImplicitObjectType
 
+FORCEINLINE bool IsInstanced(EImplicitObjectType Type)
+{
+	return (Type & ImplicitObjectType::IsInstanced) != 0;
+}
+
+FORCEINLINE bool IsScaled(EImplicitObjectType Type)
+{
+	return (Type & ImplicitObjectType::IsScaled) != 0;
+}
+
+
 namespace EImplicitObject
 {
 	enum Flags
diff --git a/Engine/Source/Runtime/Experimental/Chaos/Public/Chaos/ImplicitObjectScaled.h b/Engine/Source/Runtime/Experimental/Chaos/Public/Chaos/ImplicitObjectScaled.h
--- a/Engine/Source/Runtime/Experimental/Chaos/Public/Chaos/ImplicitObjectScaled.h
+++ b/Engine/Source/Runtime/Experimental/Chaos/Public/Chaos/ImplicitObjectScaled.h
@@ -11,12 +11,120 @@
 
 namespace Chaos
 {
+
+template <typename TConcrete, typename T, int d, bool bInstanced>
+class TImplicitObjectInstanced final : public TImplicitObject<T,d>
+{
+public:
+	using ObjectType = typename TChooseClass<bInstanced, TSerializablePtr<TConcrete>, TUniquePtr<TConcrete>>::Result;
+	using TImplicitObject<T, d>::GetTypeName;
+
+	TImplicitObjectInstanced(ObjectType Object)
+		: TImplicitObject<T, d>(EImplicitObject::HasBoundingBox, ObjectType->GetType())
+		, MObject(MoveTemp(Object))
+	{
+		ensure(IsInstanced(MObject->GetType(true)) == false);	//cannot have an instance of an instance
+		this->bIsConvex = MObject->IsConvex();
+		Type |= ImplicitObjectType::IsInstanced;
+	}
+
+	static EImplicitObjectType GetType()
+	{
+		return TConcrete::GetType() | ImplicitObjectType::IsInstanced;
+	}
+
+	const TConcrete* GetInstancedObject() const
+	{
+		return MObject.Get();
+	}
+
+	virtual T PhiWithNormal(const TVector<T, d>& X, TVector<T, d>& Normal) const override
+	{
+		return MObject->PhiWithNormal(X, Normal);
+	}
+
+	virtual bool Raycast(const TVector<T, d>& StartPoint, const TVector<T, d>& Dir, const T Length, const T Thickness, T& OutTime, TVector<T, d>& OutPosition, TVector<T, d>& OutNormal, int32& OutFaceIndex) const override
+	{
+		return MObject->Raycast(StartPoint, Dir, Length, Thickness, OutTime, OutPosition, OutNormal, OutFaceIndex);
+	}
+
+	virtual void Serialize(FChaosArchive& Ar) override
+	{
+		FChaosArchiveScopedMemory ScopedMemory(Ar, GetTypeName(), false);
+		TImplicitObject<T, d>::SerializeImp(Ar);
+		Ar << MObject;
+	}
+
+	/** This is a low level function and assumes the internal object has a SweepGeom function. Should not be called directly. See GeometryQueries.h : SweepQuery */
+	template <typename TGeometry>
+	static bool LowLevelSweepGeom(const TImplicitObjectInstanced<TConcrete, T, d, bInstanced>& OwningInstanced, const TImplicitObject<T, d>& B, const TRigidTransform<T, d>& BToATM, const TVector<T, d>& LocalDir, const T Length, T& OutTime, TVector<T, d>& LocalPosition, TVector<T, d>& LocalNormal, int32& OutFaceIndex, T Thickness = 0)
+	{
+		ensure(false);
+		return false;
+	}
+
+	/** This is a low level function and assumes the internal object has a OverlapGeom function. Should not be called directly. See GeometryQueries.h : OverlapQuery */
+	template <typename TGeometry>
+	static bool LowLevelOverlapGeom(const TImplicitObjectInstanced<TConcrete, T, d, bInstanced>& OwningInstanced, const TImplicitObject<T, d>& B, const TRigidTransform<T, d>& BToATM, T Thickness = 0)
+	{
+		ensure(false);
+		return false;
+	}
+
+	virtual int32 FindMostOpposingFace(const TVector<T, d>& Position, const TVector<T, d>& UnitDir, int32 HintFaceIndex, T SearchDist) const override
+	{
+		return MObject->FindMostOpposingFace(UnscaledPosition, UnscaledDir, HintFaceIndex, UnscaledSearchDist);
+	}
+
+	virtual TVector<T, 3> FindGeometryOpposingNormal(const TVector<T, d>& DenormDir, int32 HintFaceIndex, const TVector<T, d>& OriginalNormal) const override
+	{
+		return MObject->FindGeometryOpposingNormal(DenormDir, HintFaceIndex, OriginalNormal);
+	}
+
+	virtual bool Overlap(const TVector<T, d>& Point, const T Thickness) const override
+	{
+		return MObject->Overlap(Point, Thickness);
+	}
+
+	virtual TVector<T, d> Support(const TVector<T, d>& Direction, const T Thickness) const override
+	{
+		return MObject->Support(Direction, Thickness);
+	}
+
+	virtual const TBox<T, d>& BoundingBox() const override { return MObject->BoundingBox(); }
+
+	const ObjectType Object() const { return MObject; }
+
+	virtual uint32 GetTypeHash() const override
+	{
+		return MObject->GetTypeHash();
+	}
+
+	virtual TUniquePtr<TImplicitObject<T, d>> Copy() const override
+	{
+		return TUniquePtr<TImplicitObject<T, d>>(CopyHelper(this));
+	}
+
+protected:
+	ObjectType MObject;
+
+	static TImplicitObjectInstanced<TConcrete, T, d, true>* CopyHelper(const TImplicitObjectInstanced<TConcrete, T, d, true>* Obj)
+	{
+		return new TImplicitObjectInstanced<TConcrete, true>(Obj->MObject);
+	}
+
+	static TImplicitObjectInstanced<TConcrete, T, d, false>* CopyHelper(const TImplicitObjectInstanced<TConcrete, T, d, false>* Obj)
+	{
+		return new TImplicitObjectInstanced<TConcrete, T, d, false>(Obj->MObject->Copy());
+	}
+};
+
+
 template<class T, int d, bool bInstanced = true>
 class TImplicitObjectScaled final : public TImplicitObject<T, d>
 {
-using ObjectType = typename TChooseClass<bInstanced, TSerializablePtr<TImplicitObject<T, d>>, TUniquePtr<TImplicitObject<T, d>>>::Result;
 public:
-
+	using ObjectType = typename TChooseClass<bInstanced, TSerializablePtr<TImplicitObject<T,d>>, TUniquePtr<TImplicitObject<T, d>>>::Result;
 	using TImplicitObject<T, d>::GetTypeName;
 
 	TImplicitObjectScaled(ObjectType Object, const TVector<T, d>& Scale, T Thickness = 0)
@@ -134,7 +242,7 @@ public:
 			TVector<T, d> UnscaledNormal;
 			float UnscaledTime;
 
-			TUniquePtr<TImplicitObject<T, d>> HackBPtr(const_cast<TImplicitObject<T,d>*>(&B));	//todo: hack, need scaled object to except raw ptr similar to transformed implicit
+			TUniquePtr<TImplicitObject<T, d>> HackBPtr(const_cast<TImplicitObject<T,d>*>(&B));	//todo: hack, need scaled object to accept raw ptr similar to transformed implicit
 			TImplicitObjectScaled<T, d> ScaledB(MakeSerializable(HackBPtr), OwningScaled.MInvScale);
 			HackBPtr.Release();
 
