Fix CIS

diff --git a/Engine/Source/Runtime/Experimental/Chaos/Public/Chaos/ImplicitObjectScaled.h b/Engine/Source/Runtime/Experimental/Chaos/Public/Chaos/ImplicitObjectScaled.h
--- a/Engine/Source/Runtime/Experimental/Chaos/Public/Chaos/ImplicitObjectScaled.h
+++ b/Engine/Source/Runtime/Experimental/Chaos/Public/Chaos/ImplicitObjectScaled.h
@@ -20,12 +20,12 @@ public:
 	using TImplicitObject<T, d>::GetTypeName;
 
 	TImplicitObjectInstanced(ObjectType Object)
-		: TImplicitObject<T, d>(EImplicitObject::HasBoundingBox, ObjectType->GetType())
+		: TImplicitObject<T, d>(EImplicitObject::HasBoundingBox, Object->GetType())
 		, MObject(MoveTemp(Object))
 	{
 		ensure(IsInstanced(MObject->GetType(true)) == false);	//cannot have an instance of an instance
 		this->bIsConvex = MObject->IsConvex();
-		Type |= ImplicitObjectType::IsInstanced;
+		this->Type |= ImplicitObjectType::IsInstanced;
 	}
 
 	static EImplicitObjectType GetType()
@@ -73,7 +73,7 @@ public:
 
 	virtual int32 FindMostOpposingFace(const TVector<T, d>& Position, const TVector<T, d>& UnitDir, int32 HintFaceIndex, T SearchDist) const override
 	{
-		return MObject->FindMostOpposingFace(UnscaledPosition, UnscaledDir, HintFaceIndex, UnscaledSearchDist);
+		return MObject->FindMostOpposingFace(Position, UnitDir, HintFaceIndex, SearchDist);
 	}
 
 	virtual TVector<T, 3> FindGeometryOpposingNormal(const TVector<T, d>& DenormDir, int32 HintFaceIndex, const TVector<T, d>& OriginalNormal) const override
