Chaos 6dof joint
- add kinematic target interpolation
- fix non-cone swing axis index in joint solve
- disable velocity solve by default (Unit tests are set up for position solve. RBAN overrides it anyway)

diff --git a/Engine/Source/Runtime/AnimGraphRuntime/Private/BoneControllers/AnimNode_RigidBody.cpp b/Engine/Source/Runtime/AnimGraphRuntime/Private/BoneControllers/AnimNode_RigidBody.cpp
--- a/Engine/Source/Runtime/AnimGraphRuntime/Private/BoneControllers/AnimNode_RigidBody.cpp
+++ b/Engine/Source/Runtime/AnimGraphRuntime/Private/BoneControllers/AnimNode_RigidBody.cpp
@@ -478,15 +478,19 @@ void FAnimNode_RigidBody::EvaluateSkeletalControl_AnyThread(FComponentSpacePoseC
 			// DeltaTime can be higher due to URO, so take multiple iterations in that case.
 			const int32 MaxSteps = RBAN_MaxSubSteps;
 			const float MaxDeltaSeconds = 1.f / 30.f;
+
+#if !WITH_CHAOS
 			const int32 NumSteps = FMath::Clamp(FMath::CeilToInt(DeltaSeconds / MaxDeltaSeconds), 1, MaxSteps);
 			const float StepDeltaTime = DeltaSeconds / float(NumSteps);
-
 			for (int32 Step = 1; Step <= NumSteps; Step++)
 			{
 				// We call the _AssumesLocked version here without a lock as the simulation is local to this node and we know
 				// we're not going to alter anything while this is running.
 				PhysicsSimulation->Simulate_AssumesLocked(StepDeltaTime, SimSpaceGravity);
 			}
+#else
+			PhysicsSimulation->Simulate_AssumesLocked(DeltaSeconds, MaxDeltaSeconds, MaxSteps, SimSpaceGravity);
+#endif
 		}
 		
 		//write back to animation system
diff --git a/Engine/Source/Runtime/AnimGraphRuntime/Private/BoneControllers/AnimNode_RigidBody_Chaos.cpp b/Engine/Source/Runtime/AnimGraphRuntime/Private/BoneControllers/AnimNode_RigidBody_Chaos.cpp
--- a/Engine/Source/Runtime/AnimGraphRuntime/Private/BoneControllers/AnimNode_RigidBody_Chaos.cpp
+++ b/Engine/Source/Runtime/AnimGraphRuntime/Private/BoneControllers/AnimNode_RigidBody_Chaos.cpp
@@ -447,15 +447,19 @@ void FAnimNode_RigidBody_Chaos::EvaluateSkeletalControl_AnyThread(FComponentSpac
 			extern int32 RBAN_MaxSubSteps;
 			const int32 MaxSteps = RBAN_MaxSubSteps;
 			const float MaxDeltaSeconds = 1.f / 30.f;
+
+#if !WITH_CHAOS
 			const int32 NumSteps = FMath::Clamp(FMath::CeilToInt(DeltaSeconds / MaxDeltaSeconds), 1, MaxSteps);
 			const float StepDeltaTime = DeltaSeconds / float(NumSteps);
-
 			for (int32 Step = 1; Step <= NumSteps; Step++)
 			{
 				// We call the _AssumesLocked version here without a lock as the simulation is local to this node and we know
 				// we're not going to alter anything while this is running.
 				PhysicsSimulation->Simulate_AssumesLocked(StepDeltaTime, SimSpaceGravity);
 			}
+#else
+			PhysicsSimulation->Simulate_AssumesLocked(DeltaSeconds, MaxDeltaSeconds, MaxSteps, SimSpaceGravity);
+#endif
 		}
 		
 		//write back to animation system
diff --git a/Engine/Source/Runtime/Engine/Private/PhysicsEngine/ImmediatePhysics/ImmediatePhysicsChaos/ImmediatePhysicsSimulation_Chaos.cpp b/Engine/Source/Runtime/Engine/Private/PhysicsEngine/ImmediatePhysics/ImmediatePhysicsChaos/ImmediatePhysicsSimulation_Chaos.cpp
--- a/Engine/Source/Runtime/Engine/Private/PhysicsEngine/ImmediatePhysics/ImmediatePhysicsChaos/ImmediatePhysicsSimulation_Chaos.cpp
+++ b/Engine/Source/Runtime/Engine/Private/PhysicsEngine/ImmediatePhysics/ImmediatePhysicsChaos/ImmediatePhysicsSimulation_Chaos.cpp
@@ -389,7 +389,7 @@ namespace ImmediatePhysics_Chaos
 
 	DECLARE_CYCLE_STAT(TEXT("FSimulation::Simulate_Chaos"), STAT_ImmediateSimulate_Chaos, STATGROUP_ImmediatePhysics);
 
-	void FSimulation::Simulate(float DeltaTime, const FVector& InGravity)
+	void FSimulation::Simulate(float DeltaTime, float MaxDeltaTime, int32 MaxSubSteps, const FVector& InGravity)
 	{
 		SCOPE_CYCLE_COUNTER(STAT_ImmediateSimulate_Chaos);
 		using namespace Chaos;
@@ -436,7 +436,7 @@ namespace ImmediatePhysics_Chaos
 
 		Evolution->GetGravityForces().SetAcceleration(InGravity);
 		
-		Evolution->AdvanceOneTimeStep(DeltaTime);
+		Evolution->Advance(DeltaTime, MaxDeltaTime, MaxSubSteps);
 		
 		Evolution->EndFrame(DeltaTime);
 
diff --git a/Engine/Source/Runtime/Engine/Public/Physics/ImmediatePhysics/ImmediatePhysicsChaos/ImmediatePhysicsSimulation_Chaos.h b/Engine/Source/Runtime/Engine/Public/Physics/ImmediatePhysics/ImmediatePhysicsChaos/ImmediatePhysicsSimulation_Chaos.h
--- a/Engine/Source/Runtime/Engine/Public/Physics/ImmediatePhysics/ImmediatePhysicsChaos/ImmediatePhysicsSimulation_Chaos.h
+++ b/Engine/Source/Runtime/Engine/Public/Physics/ImmediatePhysics/ImmediatePhysicsChaos/ImmediatePhysicsSimulation_Chaos.h
@@ -59,8 +59,8 @@ namespace ImmediatePhysics_Chaos
 		//void SetIsKinematic(int32 ActorDataIndex, bool bKinematic);
 
 		/** Advance the simulation by DeltaTime */
-		void Simulate(float DeltaTime, const FVector& InGravity);
-		void Simulate_AssumesLocked(float DeltaTime, const FVector& InGravity) { Simulate(DeltaTime, InGravity); }
+		void Simulate(float DeltaTime, float MaxDeltaTime, int32 MaxSubSteps, const FVector& InGravity);
+		void Simulate_AssumesLocked(float DeltaTime, float MaxDeltaTime, int32 MaxSubSteps, const FVector& InGravity) { Simulate(DeltaTime, MaxDeltaTime, MaxSubSteps, InGravity); }
 
 		///** Whether or not an entity is simulated */
 		//bool IsSimulated(int32 ActorDataIndex) const;
diff --git a/Engine/Source/Runtime/Experimental/Chaos/Private/Chaos/ChaosDebugDraw.cpp b/Engine/Source/Runtime/Experimental/Chaos/Private/Chaos/ChaosDebugDraw.cpp
--- a/Engine/Source/Runtime/Experimental/Chaos/Private/Chaos/ChaosDebugDraw.cpp
+++ b/Engine/Source/Runtime/Experimental/Chaos/Private/Chaos/ChaosDebugDraw.cpp
@@ -170,7 +170,7 @@ namespace Chaos
 		}
 
 		template<typename T, int d>
-		void DrawJointConstraintImpl(const TRigidTransform<float, 3>& SpaceTransform, const TVector<T, d>& InXa, const PMatrix<T, d, d>& Ra, const TVector<T, d>& InXb, const PMatrix<T, d, d>& Rb, const TVector<T, d>& CR, T ColorScale)
+		void DrawJointConstraintImpl(const TRigidTransform<float, 3>& SpaceTransform, const TVector<T, d>& InXa, const PMatrix<T, d, d>& Ra, const TVector<T, d>& InXb, const PMatrix<T, d, d>& Rb, const TVector<T, d>& CR, int32 Level, T ColorScale)
 		{
 			using namespace Chaos::DebugDraw;
 			FColor R = (ColorScale * FColor::Red).ToFColor(false);
@@ -187,6 +187,10 @@ namespace Chaos
 			FDebugDrawQueue::GetInstance().DrawDebugDirectionalArrow(Xb, Xb + DrawScale * ConstraintAxisLen * SpaceTransform.TransformVector(Rb.GetAxis(0)), DrawScale * ArrowSize, C, false, KINDA_SMALL_NUMBER, DrawPriority, LineThickness);
 			FDebugDrawQueue::GetInstance().DrawDebugDirectionalArrow(Xb, Xb + DrawScale * ConstraintAxisLen * SpaceTransform.TransformVector(Rb.GetAxis(1)), DrawScale * ArrowSize, M, false, KINDA_SMALL_NUMBER, DrawPriority, LineThickness);
 			FDebugDrawQueue::GetInstance().DrawDebugDirectionalArrow(Xb, Xb + DrawScale * ConstraintAxisLen * SpaceTransform.TransformVector(Rb.GetAxis(2)), DrawScale * ArrowSize, Y, false, KINDA_SMALL_NUMBER, DrawPriority, LineThickness);
+			if (Level >= 0)
+			{
+				FDebugDrawQueue::GetInstance().DrawDebugString(Xb + FontHeight * TVector<T, d>(0, 0, 1), FString::Format(TEXT("{0}"), { Level }), nullptr, FColor::Red, KINDA_SMALL_NUMBER, false, FontScale);
+			}
 			//FDebugDrawQueue::GetInstance().DrawDebugString(Xb + 3 * FontHeight * TVector<T, d>(0, 0, 1), FString::Format(TEXT("T=  {0}"), { FMath::RadiansToDegrees(CR[(int32)E6DJointAngularConstraintIndex::Twist]) }), nullptr, FColor::Red, KINDA_SMALL_NUMBER, false, FontScale);
 			//FDebugDrawQueue::GetInstance().DrawDebugString(Xb + 2 * FontHeight * TVector<T, d>(0, 0, 1), FString::Format(TEXT("S1= {0}"), { FMath::RadiansToDegrees(CR[(int32)E6DJointAngularConstraintIndex::Swing1]) }), nullptr, FColor::Red, KINDA_SMALL_NUMBER, false, FontScale);
 			//FDebugDrawQueue::GetInstance().DrawDebugString(Xb + 1 * FontHeight * TVector<T, d>(0, 0, 1), FString::Format(TEXT("S2= {0}"), { FMath::RadiansToDegrees(CR[(int32)E6DJointAngularConstraintIndex::Swing2]) }), nullptr, FColor::Red, KINDA_SMALL_NUMBER, false, FontScale);
@@ -201,7 +205,7 @@ namespace Chaos
 				TVector<T, d> Xa, Xb, CR;
 				PMatrix<T, d, d> Ra, Rb;
 				ConstraintHandle->CalculateConstraintSpace(Xa, Ra, Xb, Rb, CR);
-				DrawJointConstraintImpl<T, d>(SpaceTransform, Xa, Ra, Xb, Rb, CR, ColorScale);
+				DrawJointConstraintImpl<T, d>(SpaceTransform, Xa, Ra, Xb, Rb, CR, ConstraintHandle->GetConstraintLevel(), ColorScale);
 			}
 		}
 
@@ -211,7 +215,7 @@ namespace Chaos
 			TVector<T, d> Xa, Xb, CR;
 			PMatrix<T, d, d> Ra, Rb;
 			ConstraintHandle->CalculateConstraintSpace(Xa, Ra, Xb, Rb, CR);
-			DrawJointConstraintImpl<T, d>(SpaceTransform, Xa, Ra, Xb, Rb, CR, ColorScale);
+			DrawJointConstraintImpl<T, d>(SpaceTransform, Xa, Ra, Xb, Rb, CR, INDEX_NONE, ColorScale);
 		}
 
 #endif
diff --git a/Engine/Source/Runtime/Experimental/Chaos/Private/Chaos/PBDJointConstraintUtilities.cpp b/Engine/Source/Runtime/Experimental/Chaos/Private/Chaos/PBDJointConstraintUtilities.cpp
--- a/Engine/Source/Runtime/Experimental/Chaos/Private/Chaos/PBDJointConstraintUtilities.cpp
+++ b/Engine/Source/Runtime/Experimental/Chaos/Private/Chaos/PBDJointConstraintUtilities.cpp
@@ -1128,7 +1128,8 @@ namespace Chaos
 		const TPBDJointSettings<T, d>& JointSettings,
 		const int32 Index0,
 		const int32 Index1,
-		const EJointAngularConstraintIndex SwingConstraint,
+		const EJointAngularConstraintIndex SwingConstraintIndex,
+		const EJointAngularAxisIndex SwingAxisIndex,
 		TVector<T, d>& P0,
 		TRotation<T, d>& Q0,
 		TVector<T, d>& P1,
@@ -1167,9 +1168,10 @@ namespace Chaos
 		}
 		const TVector<T, d> TwistAxis = R0 * TwistAxis01;
 
+		const TRotation<T, d> R1NoTwist = R1 * R01Twist.Inverse();
 		const PMatrix<T, d, d> Axes0 = R0.ToMatrix();
-		const PMatrix<T, d, d> Axes1 = R1.ToMatrix();
-		TVector<T, d> SwingCross = TVector<T, d>::CrossProduct(Axes0.GetAxis((int32)SwingConstraint), Axes1.GetAxis((int32)SwingConstraint));
+		const PMatrix<T, d, d> Axes1 = R1NoTwist.ToMatrix();
+		TVector<T, d> SwingCross = TVector<T, d>::CrossProduct(Axes0.GetAxis((int32)SwingAxisIndex), Axes1.GetAxis((int32)SwingAxisIndex));
 		SwingCross = SwingCross - TVector<T, d>::DotProduct(TwistAxis, SwingCross) * TwistAxis;
 		const T SwingCrossLen = SwingCross.Size();
 		if (SwingCrossLen > KINDA_SMALL_NUMBER)
@@ -1179,19 +1181,19 @@ namespace Chaos
 			TVector<T, d> SwingAxis1 = SwingAxis;
 
 			T SwingAngle = FMath::Asin(FMath::Clamp(SwingCrossLen, (T)0, (T)1));
-			const T SwingDot = TVector<T, d>::DotProduct(Axes0.GetAxis((int32)SwingConstraint), Axes1.GetAxis((int32)SwingConstraint));
+			const T SwingDot = TVector<T, d>::DotProduct(Axes0.GetAxis((int32)SwingAxisIndex), Axes1.GetAxis((int32)SwingAxisIndex));
 			if (SwingDot < (T)0)
 			{
 				SwingAngle = (T)PI - SwingAngle;
 			}
 
 			T SwingAngleMax = FLT_MAX;
-			if (JointSettings.Motion.AngularMotionTypes[(int32)SwingConstraint] == EJointMotionType::Limited)
+			if (JointSettings.Motion.AngularMotionTypes[(int32)SwingConstraintIndex] == EJointMotionType::Limited)
 			{
-				T Swing1Limit = JointSettings.Motion.AngularLimits[(int32)SwingConstraint];
-				SwingAngleMax = Swing1Limit;
+				T SwingLimit = JointSettings.Motion.AngularLimits[(int32)SwingConstraintIndex];
+				SwingAngleMax = SwingLimit;
 			}
-			else if (JointSettings.Motion.AngularMotionTypes[(int32)SwingConstraint] == EJointMotionType::Locked)
+			else if (JointSettings.Motion.AngularMotionTypes[(int32)SwingConstraintIndex] == EJointMotionType::Locked)
 			{
 				SwingAngleMax = 0;
 			}
@@ -1224,7 +1226,8 @@ namespace Chaos
 		const TPBDJointSettings<T, d>& JointSettings,
 		const int32 Index0,
 		const int32 Index1,
-		const EJointAngularConstraintIndex SwingConstraint,
+		const EJointAngularConstraintIndex SwingConstraintIndex,
+		const EJointAngularAxisIndex SwingAxisIndex,
 		TVector<T, d>& P0,
 		TRotation<T, d>& Q0,
 		TVector<T, d>& V0,
@@ -1266,9 +1269,10 @@ namespace Chaos
 		}
 		const TVector<T, d> TwistAxis = R0 * TwistAxis01;
 
+		const TRotation<T, d> R1NoTwist = R1 * R01Twist.Inverse();
 		const PMatrix<T, d, d> Axes0 = R0.ToMatrix();
-		const PMatrix<T, d, d> Axes1 = R1.ToMatrix();
-		TVector<T, d> SwingCross = TVector<T, d>::CrossProduct(Axes0.GetAxis((int32)SwingConstraint), Axes1.GetAxis((int32)SwingConstraint));
+		const PMatrix<T, d, d> Axes1 = R1NoTwist.ToMatrix();
+		TVector<T, d> SwingCross = TVector<T, d>::CrossProduct(Axes0.GetAxis((int32)SwingAxisIndex), Axes1.GetAxis((int32)SwingAxisIndex));
 		SwingCross = SwingCross - TVector<T, d>::DotProduct(TwistAxis, SwingCross) * TwistAxis;
 		const T SwingCrossLen = SwingCross.Size();
 		if (SwingCrossLen > KINDA_SMALL_NUMBER)
@@ -1278,19 +1282,19 @@ namespace Chaos
 			TVector<T, d> SwingAxis1 = SwingAxis;
 
 			T SwingAngle = FMath::Asin(FMath::Clamp(SwingCrossLen, (T)0, (T)1));
-			const T SwingDot = TVector<T, d>::DotProduct(Axes0.GetAxis((int32)SwingConstraint), Axes1.GetAxis((int32)SwingConstraint));
+			const T SwingDot = TVector<T, d>::DotProduct(Axes0.GetAxis((int32)SwingAxisIndex), Axes1.GetAxis((int32)SwingAxisIndex));
 			if (SwingDot < (T)0)
 			{
 				SwingAngle = (T)PI - SwingAngle;
 			}
 
 			T SwingAngleMax = FLT_MAX;
-			if (JointSettings.Motion.AngularMotionTypes[(int32)SwingConstraint] == EJointMotionType::Limited)
+			if (JointSettings.Motion.AngularMotionTypes[(int32)SwingConstraintIndex] == EJointMotionType::Limited)
 			{
-				T Swing1Limit = JointSettings.Motion.AngularLimits[(int32)SwingConstraint];
+				T Swing1Limit = JointSettings.Motion.AngularLimits[(int32)SwingConstraintIndex];
 				SwingAngleMax = Swing1Limit;
 			}
-			else if (JointSettings.Motion.AngularMotionTypes[(int32)SwingConstraint] == EJointMotionType::Locked)
+			else if (JointSettings.Motion.AngularMotionTypes[(int32)SwingConstraintIndex] == EJointMotionType::Locked)
 			{
 				SwingAngleMax = 0;
 			}
diff --git a/Engine/Source/Runtime/Experimental/Chaos/Private/Chaos/PBDJointConstraints.cpp b/Engine/Source/Runtime/Experimental/Chaos/Private/Chaos/PBDJointConstraints.cpp
--- a/Engine/Source/Runtime/Experimental/Chaos/Private/Chaos/PBDJointConstraints.cpp
+++ b/Engine/Source/Runtime/Experimental/Chaos/Private/Chaos/PBDJointConstraints.cpp
@@ -114,9 +114,9 @@ namespace Chaos
 	template<class T, int d>
 	TPBDJointSolverSettings<T, d>::TPBDJointSolverSettings()
 		: SwingTwistAngleTolerance((T)1.0e-6)
-		, MinParentMassRatio((T)0.5)
-		, MaxInertiaRatio((T)5.0)
-		, bEnableVelocitySolve(true)
+		, MinParentMassRatio(0)
+		, MaxInertiaRatio(0)
+		, bEnableVelocitySolve(false)
 		, bEnableLinearLimits(true)
 		, bEnableTwistLimits(true)
 		, bEnableSwingLimits(true)
@@ -125,7 +125,7 @@ namespace Chaos
 		, Projection((T)0)
 		, Stiffness((T)0)
 		, DriveStiffness((T)0)
-		, PositionIterations((T)1)
+		, PositionIterations((T)0)
 	{
 	}
 
@@ -500,12 +500,12 @@ namespace Chaos
 				if (Swing1Motion != EJointMotionType::Free)
 				{
 					// Swing Arc/Lock
-					TPBDJointUtilities<T, d>::ApplyJointSwingVelocityConstraint(Dt, Settings, JointSettings, Index0, Index1, EJointAngularConstraintIndex::Swing1, P0, Q0, V0, W0, P1, Q1, V1, W1, InvM0, InvIL0, InvM1, InvIL1);
+					TPBDJointUtilities<T, d>::ApplyJointSwingVelocityConstraint(Dt, Settings, JointSettings, Index0, Index1, EJointAngularConstraintIndex::Swing1, EJointAngularAxisIndex::Swing1, P0, Q0, V0, W0, P1, Q1, V1, W1, InvM0, InvIL0, InvM1, InvIL1);
 				}
 				if (Swing2Motion != EJointMotionType::Free)
 				{
 					// Swing Arc/Lock
-					TPBDJointUtilities<T, d>::ApplyJointSwingVelocityConstraint(Dt, Settings, JointSettings, Index0, Index1, EJointAngularConstraintIndex::Swing2, P0, Q0, V0, W0, P1, Q1, V1, W1, InvM0, InvIL0, InvM1, InvIL1);
+					TPBDJointUtilities<T, d>::ApplyJointSwingVelocityConstraint(Dt, Settings, JointSettings, Index0, Index1, EJointAngularConstraintIndex::Swing2, EJointAngularAxisIndex::Swing2, P0, Q0, V0, W0, P1, Q1, V1, W1, InvM0, InvIL0, InvM1, InvIL1);
 				}
 			}
 		}
@@ -656,12 +656,12 @@ namespace Chaos
 				if (Swing1Motion != EJointMotionType::Free)
 				{
 					// Swing Arc/Lock
-					TPBDJointUtilities<T, d>::ApplyJointSwingConstraint(Dt, Settings, JointSettings, Index0, Index1, EJointAngularConstraintIndex::Swing1, P0, Q0, P1, Q1, InvM0, InvIL0, InvM1, InvIL1);
+					TPBDJointUtilities<T, d>::ApplyJointSwingConstraint(Dt, Settings, JointSettings, Index0, Index1, EJointAngularConstraintIndex::Swing1, EJointAngularAxisIndex::Swing1, P0, Q0, P1, Q1, InvM0, InvIL0, InvM1, InvIL1);
 				}
 				if (Swing2Motion != EJointMotionType::Free)
 				{
 					// Swing Arc/Lock
-					TPBDJointUtilities<T, d>::ApplyJointSwingConstraint(Dt, Settings, JointSettings, Index0, Index1, EJointAngularConstraintIndex::Swing2, P0, Q0, P1, Q1, InvM0, InvIL0, InvM1, InvIL1);
+					TPBDJointUtilities<T, d>::ApplyJointSwingConstraint(Dt, Settings, JointSettings, Index0, Index1, EJointAngularConstraintIndex::Swing2, EJointAngularAxisIndex::Swing2, P0, Q0, P1, Q1, InvM0, InvIL0, InvM1, InvIL1);
 				}
 			}
 		}
diff --git a/Engine/Source/Runtime/Experimental/Chaos/Private/Chaos/PBDRigidsEvolutionGBF.cpp b/Engine/Source/Runtime/Experimental/Chaos/Private/Chaos/PBDRigidsEvolutionGBF.cpp
--- a/Engine/Source/Runtime/Experimental/Chaos/Private/Chaos/PBDRigidsEvolutionGBF.cpp
+++ b/Engine/Source/Runtime/Experimental/Chaos/Private/Chaos/PBDRigidsEvolutionGBF.cpp
@@ -88,7 +88,21 @@ void SerializeToDisk(TEvolution& Evolution)
 #endif
 
 template <typename T, int d>
-void TPBDRigidsEvolutionGBF<T, d>::AdvanceOneTimeStep(T Dt)
+void TPBDRigidsEvolutionGBF<T, d>::Advance(const T Dt, const T MaxStepDt, const int32 MaxSteps)
+{
+	const int32 NumSteps = FMath::Clamp(FMath::CeilToInt(Dt / MaxStepDt), 1, MaxSteps);
+	const T StepDt = Dt / (T)NumSteps;
+	for (int32 Step = 0; Step < NumSteps; ++Step)
+	{
+		// StepFraction: how much of the remaining time this step represents, used to interpolate kinematic targets
+		// E.g., for 4 steps this will be: 1/4, 1/3, 1/2, 1
+		const float StepFraction = (T)1 / (T)(NumSteps - Step);
+		AdvanceOneTimeStep(StepDt, StepFraction);
+	}
+}
+
+template <typename T, int d>
+void TPBDRigidsEvolutionGBF<T, d>::AdvanceOneTimeStep(const T Dt, const T StepFraction)
 {
 	SCOPE_CYCLE_COUNTER(STAT_Evo_AdvanceOneTimeStep);
 
@@ -106,7 +120,7 @@ void TPBDRigidsEvolutionGBF<T, d>::AdvanceOneTimeStep(T Dt)
 
 	{
 		SCOPE_CYCLE_COUNTER(STAT_Evo_KinematicTargets);
-		ApplyKinematicTargets(Dt);
+		ApplyKinematicTargets(Dt, StepFraction);
 	}
 
 	if (PostIntegrateCallback != nullptr)
diff --git a/Engine/Source/Runtime/Experimental/Chaos/Public/Chaos/PBDJointConstraintUtilities.h b/Engine/Source/Runtime/Experimental/Chaos/Public/Chaos/PBDJointConstraintUtilities.h
--- a/Engine/Source/Runtime/Experimental/Chaos/Public/Chaos/PBDJointConstraintUtilities.h
+++ b/Engine/Source/Runtime/Experimental/Chaos/Public/Chaos/PBDJointConstraintUtilities.h
@@ -200,7 +200,8 @@ namespace Chaos
 			const TPBDJointSettings<T, d>& JointSettings,
 			const int32 Index0,
 			const int32 Index1,
-			const EJointAngularConstraintIndex SwingConstraint,
+			const EJointAngularConstraintIndex SwingConstraintIndex,
+			const EJointAngularAxisIndex SwingAxisIndex,
 			TVector<T, d>& P0,
 			TRotation<T, d>& Q0,
 			TVector<T, d>& P1,
@@ -216,7 +217,8 @@ namespace Chaos
 			const TPBDJointSettings<T, d>& JointSettings,
 			const int32 Index0,
 			const int32 Index1,
-			const EJointAngularConstraintIndex SwingConstraint,
+			const EJointAngularConstraintIndex SwingConstraintIndex,
+			const EJointAngularAxisIndex SwingAxisIndex,
 			TVector<T, d>& P0,
 			TRotation<T, d>& Q0,
 			TVector<T, d>& V0,
diff --git a/Engine/Source/Runtime/Experimental/Chaos/Public/Chaos/PBDRigidsEvolution.h b/Engine/Source/Runtime/Experimental/Chaos/Public/Chaos/PBDRigidsEvolution.h
--- a/Engine/Source/Runtime/Experimental/Chaos/Public/Chaos/PBDRigidsEvolution.h
+++ b/Engine/Source/Runtime/Experimental/Chaos/Public/Chaos/PBDRigidsEvolution.h
@@ -420,8 +420,11 @@ class TPBDRigidsEvolutionBase
 		}
 	}
 
-	void ApplyKinematicTargets(T Dt)
+	void ApplyKinematicTargets(const T Dt, const T StepFraction)
 	{
+		check(StepFraction > (T)0);
+		check(StepFraction <= (T)1);
+
 		// @todo(ccaulfield): optimize. Depending on the number of kinematics relative to the number that have 
 		// targets set, it may be faster to process a command list rather than iterate over them all each frame. 
 		const T MinDt = 1e-6f;
@@ -447,17 +450,29 @@ class TPBDRigidsEvolutionBase
 			{
 				// Move to kinematic target and update velocities to match
 				// Target positions only need to be processed once, and we reset the velocity next frame (if no new target is set)
+				TVector<T, d> TargetPos;
+				TRotation<T, d> TargetRot;
+				if (FMath::IsNearlyEqual(StepFraction, (T)1, KINDA_SMALL_NUMBER))
+				{
+					TargetPos = KinematicTarget.GetTarget().GetLocation();
+					TargetRot = KinematicTarget.GetTarget().GetRotation();
+					KinematicTarget.SetMode(EKinematicTargetMode::Zero);
+				}
+				else
+				{
+					TargetPos = TVector<T, d>::Lerp(Particle.X(), KinematicTarget.GetTarget().GetLocation(), StepFraction);
+					TargetRot = TRotation<T, d>::Slerp(Particle.R(), KinematicTarget.GetTarget().GetRotation(), StepFraction);
+				}
 				if (Dt > MinDt)
 				{
-					TVector<float, 3> V = TVector<float, 3>::CalculateVelocity(Particle.X(), KinematicTarget.GetTarget().GetLocation(), Dt);
+					TVector<float, 3> V = TVector<float, 3>::CalculateVelocity(Particle.X(), TargetPos, Dt);
 					Particle.V() = V;
 
-					TVector<float, 3> W = TRotation<float, 3>::CalculateAngularVelocity(Particle.R(), KinematicTarget.GetTarget().GetRotation(), Dt);
+					TVector<float, 3> W = TRotation<float, 3>::CalculateAngularVelocity(Particle.R(), TargetRot, Dt);
 					Particle.W() = W;
 				}
-				Particle.X() = KinematicTarget.GetTarget().GetTranslation();
-				Particle.R() = KinematicTarget.GetTarget().GetRotation();
-				KinematicTarget.SetMode(EKinematicTargetMode::Zero);
+				Particle.X() = TargetPos;
+				Particle.R() = TargetRot;
 				break;
 			}
 
diff --git a/Engine/Source/Runtime/Experimental/Chaos/Public/Chaos/PBDRigidsEvolutionGBF.h b/Engine/Source/Runtime/Experimental/Chaos/Public/Chaos/PBDRigidsEvolutionGBF.h
--- a/Engine/Source/Runtime/Experimental/Chaos/Public/Chaos/PBDRigidsEvolutionGBF.h
+++ b/Engine/Source/Runtime/Experimental/Chaos/Public/Chaos/PBDRigidsEvolutionGBF.h
@@ -78,7 +78,8 @@ public:
 		PostApplyPushOutCallback = Cb;
 	}
 
-	CHAOS_API void AdvanceOneTimeStep(const T dt);
+	CHAOS_API void Advance(const T Dt, const T MaxStepDt, const int32 MaxSteps);
+	CHAOS_API void AdvanceOneTimeStep(const T dt, const T StepFraction = (T)1.0);
 
 	using Base::ApplyConstraints;
 	using Base::ApplyPushOut;
diff --git a/Engine/Source/Runtime/Experimental/ChaosCore/Public/Chaos/Vector.h b/Engine/Source/Runtime/Experimental/ChaosCore/Public/Chaos/Vector.h
--- a/Engine/Source/Runtime/Experimental/ChaosCore/Public/Chaos/Vector.h
+++ b/Engine/Source/Runtime/Experimental/ChaosCore/Public/Chaos/Vector.h
@@ -396,6 +396,7 @@ namespace Chaos
 			Stream.write(reinterpret_cast<const char*>(&Y), sizeof(Y));
 			Stream.write(reinterpret_cast<const char*>(&Z), sizeof(Z));
 		}
+		static inline TVector<float, 3> Lerp(const TVector<float, 3>& V1, const TVector<float, 3>& V2, const float F) { return FMath::Lerp<FVector, float>(V1, V2, F); }
 		static inline TVector<float, 3> CrossProduct(const TVector<float, 3>& V1, const TVector<float, 3>& V2) { return FVector::CrossProduct(V1, V2); }
 		static inline float DotProduct(const TVector<float, 3>& V1, const TVector<float, 3>& V2) { return FVector::DotProduct(V1, V2); }
 		bool operator<=(const TVector<float, 3>& V) const
