Refactor scaled implicit to have concrete types, take 2

diff --git a/Engine/Source/Runtime/Core/Public/UObject/ExternalPhysicsCustomObjectVersion.h b/Engine/Source/Runtime/Core/Public/UObject/ExternalPhysicsCustomObjectVersion.h
--- a/Engine/Source/Runtime/Core/Public/UObject/ExternalPhysicsCustomObjectVersion.h
+++ b/Engine/Source/Runtime/Core/Public/UObject/ExternalPhysicsCustomObjectVersion.h
@@ -45,6 +45,9 @@ struct CORE_API FExternalPhysicsCustomObjectVersion
 		// Serialize broadphase type
 		SerializeBroadphaseType,
 
+		// Allow scaled geometry to be a concrete type
+		ScaledGeometryIsConcrete,
+
 		// -----<new versions can be added above this line>-------------------------------------------------
 		VersionPlusOne,
 		LatestVersion = VersionPlusOne - 1
diff --git a/Engine/Source/Runtime/Engine/Classes/PhysicsEngine/ConvexElem.h b/Engine/Source/Runtime/Engine/Classes/PhysicsEngine/ConvexElem.h
--- a/Engine/Source/Runtime/Engine/Classes/PhysicsEngine/ConvexElem.h
+++ b/Engine/Source/Runtime/Engine/Classes/PhysicsEngine/ConvexElem.h
@@ -24,6 +24,9 @@ namespace Chaos
 {
 	template <typename T, int d>
 	class TImplicitObject;
+
+	template <typename T, int d>
+	class TConvex;
 }
 
 /** One convex hull, used for simplified collision. */
@@ -52,7 +55,7 @@ private:
 	physx::PxConvexMesh*   ConvexMeshNegX;
 
 #if WITH_CHAOS
-	TUniquePtr<Chaos::TImplicitObject<float, 3>> ChaosConvex;
+	TUniquePtr<Chaos::TConvex<float,3>> ChaosConvex;
 #endif
 
 public:
@@ -104,9 +107,9 @@ public:
 	ENGINE_API void SetMirroredConvexMesh(physx::PxConvexMesh* InMesh);
 
 #if WITH_CHAOS
-	ENGINE_API const TUniquePtr<Chaos::TImplicitObject<float, 3>>& GetChaosConvexMesh() const;
+	ENGINE_API const TUniquePtr<Chaos::TConvex<float, 3>>& GetChaosConvexMesh() const;
 
-	ENGINE_API void SetChaosConvexMesh(TUniquePtr<Chaos::TImplicitObject<float, 3>>&& InChaosConvex);
+	ENGINE_API void SetChaosConvexMesh(TUniquePtr<Chaos::TConvex<float, 3>>&& InChaosConvex);
 
 	ENGINE_API void ResetChaosConvexMesh();
 #endif
diff --git a/Engine/Source/Runtime/Engine/Private/Collision/KAggregateGeom.cpp b/Engine/Source/Runtime/Engine/Private/Collision/KAggregateGeom.cpp
--- a/Engine/Source/Runtime/Engine/Private/Collision/KAggregateGeom.cpp
+++ b/Engine/Source/Runtime/Engine/Private/Collision/KAggregateGeom.cpp
@@ -14,6 +14,7 @@
 #include "PhysicsEngine/AggregateGeom.h"
 #include "Engine/Polys.h"
 #include "PhysXIncludes.h"
+#include "Chaos/Convex.h"
 
 #if WITH_CHAOS
 #include "Chaos/ImplicitObject.h"
diff --git a/Engine/Source/Runtime/Engine/Private/PhysicsEngine/BodyInstance.cpp b/Engine/Source/Runtime/Engine/Private/PhysicsEngine/BodyInstance.cpp
--- a/Engine/Source/Runtime/Engine/Private/PhysicsEngine/BodyInstance.cpp
+++ b/Engine/Source/Runtime/Engine/Private/PhysicsEngine/BodyInstance.cpp
@@ -857,7 +857,7 @@ void FBodyInstance::UpdatePhysicsFilterData()
 				// See if we currently have sim collision
 				const bool bWasSimulationShape = FPhysicsInterface::IsSimulationShape(InnerShape);
 				const bool bSyncShape = ShapeIndex < NumSyncShapes;
-				const bool bIsTrimesh = FPhysicsInterface::IsShapeType(InnerShape, ECollisionShapeType::Trimesh);
+				const bool bIsTrimesh = FPhysicsInterface::GetShapeType(InnerShape) == ECollisionShapeType::Trimesh;
 				const bool bIsStatic = FPhysicsInterface::IsStatic(Actor);
 
 				const FBodyCollisionFlags& CollisionFlags = PerShapeCollisionData.CollisionFlags;
diff --git a/Engine/Source/Runtime/Engine/Private/PhysicsEngine/BodySetup.cpp b/Engine/Source/Runtime/Engine/Private/PhysicsEngine/BodySetup.cpp
--- a/Engine/Source/Runtime/Engine/Private/PhysicsEngine/BodySetup.cpp
+++ b/Engine/Source/Runtime/Engine/Private/PhysicsEngine/BodySetup.cpp
@@ -22,6 +22,7 @@
 #include "UObject/PropertyPortFlags.h"
 #include "Components/SplineMeshComponent.h"
 #include "ChaosCheck.h"
+#include "Chaos/Convex.h"
 
 #include "PhysXCookHelper.h"
 
@@ -1667,12 +1668,12 @@ float FKConvexElem::GetVolume(const FVector& Scale) const
 }
 
 #if WITH_CHAOS
-const TUniquePtr<Chaos::TImplicitObject<float, 3>>& FKConvexElem::GetChaosConvexMesh() const
+const TUniquePtr<Chaos::TConvex<float, 3>>& FKConvexElem::GetChaosConvexMesh() const
 {
 	return ChaosConvex;
 }
 
-void FKConvexElem::SetChaosConvexMesh(TUniquePtr<Chaos::TImplicitObject<float, 3>>&& InChaosConvex)
+void FKConvexElem::SetChaosConvexMesh(TUniquePtr<Chaos::TConvex<float, 3>>&& InChaosConvex)
 {
 	ChaosConvex = MoveTemp(InChaosConvex);
 }
diff --git a/Engine/Source/Runtime/Engine/Private/PhysicsEngine/Experimental/ChaosDerivedData.cpp b/Engine/Source/Runtime/Engine/Private/PhysicsEngine/Experimental/ChaosDerivedData.cpp
--- a/Engine/Source/Runtime/Engine/Private/PhysicsEngine/Experimental/ChaosDerivedData.cpp
+++ b/Engine/Source/Runtime/Engine/Private/PhysicsEngine/Experimental/ChaosDerivedData.cpp
@@ -21,7 +21,7 @@ const TCHAR* FChaosDerivedDataCooker::GetPluginName() const
 
 const TCHAR* FChaosDerivedDataCooker::GetVersionString() const
 {
-	return TEXT("DFE79D57413E4527AD10CEEF771A4654");
+	return TEXT("EF20258F42C54D658B5BEE64EB779C6A");
 }
 
 FString FChaosDerivedDataCooker::GetPluginSpecificCacheKeySuffix() const
diff --git a/Engine/Source/Runtime/Engine/Private/PhysicsEngine/Experimental/ChaosInterfaceUtils.cpp b/Engine/Source/Runtime/Engine/Private/PhysicsEngine/Experimental/ChaosInterfaceUtils.cpp
--- a/Engine/Source/Runtime/Engine/Private/PhysicsEngine/Experimental/ChaosInterfaceUtils.cpp
+++ b/Engine/Source/Runtime/Engine/Private/PhysicsEngine/Experimental/ChaosInterfaceUtils.cpp
@@ -12,6 +12,7 @@
 #include "Chaos/TriangleMesh.h"
 #include "Chaos/TriangleMeshImplicitObject.h"
 #include "Chaos/UniformGrid.h"
+#include "Chaos/Convex.h"
 
 #include "PhysicsEngine/BoxElem.h"
 #include "PhysicsEngine/ConvexElem.h"
@@ -257,7 +258,7 @@ namespace ChaosInterface
 				//}
 				//else
 				{
-					TUniquePtr<Chaos::TImplicitObjectScaled<float, 3, true>> Implicit = MakeUnique<Chaos::TImplicitObjectScaled<float, 3>>(MakeSerializable(ConvexImplicit), Scale);
+					auto Implicit = MakeUnique<Chaos::TImplicitObjectScaled<Chaos::TConvex<float, 3>>>(MakeSerializable(ConvexImplicit), Scale);
 					auto NewShape = NewShapeHelper(MakeSerializable(Implicit));
 					Shapes.Emplace(MoveTemp(NewShape));
 					Geoms.Add(MoveTemp(Implicit));
@@ -267,7 +268,7 @@ namespace ChaosInterface
 
 		for (const auto& ChaosTriMesh : InParams.ChaosTriMeshes)
 		{
-			auto Implicit = MakeUnique<Chaos::TImplicitObjectScaled<float, 3>>(MakeSerializable(ChaosTriMesh), Scale);
+			auto Implicit = MakeUnique<Chaos::TImplicitObjectScaled<Chaos::TTriangleMeshImplicitObject<float>>>(MakeSerializable(ChaosTriMesh), Scale);
 			auto NewShape = NewShapeHelper(MakeSerializable(Implicit), true);
 			Shapes.Emplace(MoveTemp(NewShape));
 			Geoms.Add(MoveTemp(Implicit));
diff --git a/Engine/Source/Runtime/Engine/Private/PhysicsEngine/Experimental/PhysInterface_Chaos.cpp b/Engine/Source/Runtime/Engine/Private/PhysicsEngine/Experimental/PhysInterface_Chaos.cpp
--- a/Engine/Source/Runtime/Engine/Private/PhysicsEngine/Experimental/PhysInterface_Chaos.cpp
+++ b/Engine/Source/Runtime/Engine/Private/PhysicsEngine/Experimental/PhysInterface_Chaos.cpp
@@ -1185,86 +1185,9 @@ bool FPhysInterface_Chaos::IsQueryShape(const FPhysicsShapeHandle& InShape)
     return InShape.bQuery;
 }
 
-bool FPhysInterface_Chaos::IsShapeType(const FPhysicsShapeReference_Chaos& InShapeRef, ECollisionShapeType InType)
+ECollisionShapeType FPhysInterface_Chaos::GetShapeType(const FPhysicsShapeReference_Chaos& InShapeRef)
 {
-    if (InType == ECollisionShapeType::Box && InShapeRef.Shape->Geometry->GetType() == Chaos::ImplicitObjectType::Box)
-    {
-        return true;
-    }
-    if (InType == ECollisionShapeType::Sphere && InShapeRef.Shape->Geometry->GetType() == Chaos::ImplicitObjectType::Sphere)
-    {
-        return true;
-    }
-    if (InType == ECollisionShapeType::Plane && InShapeRef.Shape->Geometry->GetType() == Chaos::ImplicitObjectType::Plane)
-    {
-        return true;
-    }
-    if (InType == ECollisionShapeType::Capsule && InShapeRef.Shape->Geometry->GetType() == Chaos::ImplicitObjectType::Capsule)
-    {
-        return true;
-    }
-    if (InType == ECollisionShapeType::Convex && InShapeRef.Shape->Geometry->GetType() == Chaos::ImplicitObjectType::Convex)
-    {
-        return true;
-    }
-    if (InType == ECollisionShapeType::Trimesh && InShapeRef.Shape->Geometry->GetType() == Chaos::ImplicitObjectType::TriangleMesh)
-    {
-        return true;
-    }
-    if (InType == ECollisionShapeType::Heightfield && InShapeRef.Shape->Geometry->GetType() == Chaos::ImplicitObjectType::HeightField)
-    {
-        return true;
-    }
-    if (InType == ECollisionShapeType::Scaled && InShapeRef.Shape->Geometry->GetType() == Chaos::ImplicitObjectType::Scaled)
-    {
-        return true;
-    }
-
-    return false;
-}
-
-ECollisionShapeType FPhysInterface_Chaos::GetShapeType(const FPhysicsShapeReference_Chaos& InShapeRef, bool bGetInnerType)
-{
-	Chaos::EImplicitObjectType Type = InShapeRef.Shape->Geometry->GetType(true);
-	if (bGetInnerType && Type == Chaos::ImplicitObjectType::Scaled)
-	{
-		Type = InShapeRef.Shape->Geometry->GetObject<Chaos::TImplicitObjectScaled<float, 3>>()->GetUnscaledObject()->GetType(true);
-	}
-
-    if (Type == Chaos::ImplicitObjectType::Box)
-    {
-        return ECollisionShapeType::Box;
-    }
-    if (Type == Chaos::ImplicitObjectType::Sphere)
-    {
-        return ECollisionShapeType::Sphere;
-    }
-    if (Type == Chaos::ImplicitObjectType::Plane)
-    {
-        return ECollisionShapeType::Plane;
-    }
-    if (Type == Chaos::ImplicitObjectType::Capsule)
-    {
-        return ECollisionShapeType::Capsule;
-    }
-    if (Type == Chaos::ImplicitObjectType::Convex)
-    {
-        return ECollisionShapeType::Convex;
-    }
-    if (Type == Chaos::ImplicitObjectType::TriangleMesh)
-    {
-        return ECollisionShapeType::Trimesh;
-    }
-    if (Type == Chaos::ImplicitObjectType::HeightField)
-    {
-        return ECollisionShapeType::Heightfield;
-    }
-    if (Type == Chaos::ImplicitObjectType::Scaled)
-    {
-        return ECollisionShapeType::Scaled;
-    }
-
-    return ECollisionShapeType::None;
+	return GetImplicitType(*InShapeRef.Shape->Geometry);
 }
 
 FTransform FPhysInterface_Chaos::GetLocalTransform(const FPhysicsShapeReference_Chaos& InShapeRef)
@@ -1439,9 +1362,8 @@ void CalculateMassPropertiesOfImplicitType(
 			MassProperties.Add(MassProperty);
 			TotalMass += Mass;
 		}
-		else if (ImplicitObject->GetType() == Chaos::ImplicitObjectType::Scaled)
+		else if (const auto ScaledObject = Chaos::TImplicitObjectScaledGeneric<float,3>::AsScaled(*ImplicitObject))
 		{
-			const Chaos::TImplicitObjectScaled<float, 3> * ScaledObject = ImplicitObject->template GetObject<Chaos::TImplicitObjectScaled<float, 3>>();
 			if (ensureMsgf(ScaledObject->GetUnscaledObject(), TEXT("Error : Null internal MObject type within TImplicitObjectScaled.")))
 			{
 				Chaos::TRigidTransform<float, 3> ScaledWorldTransform(WorldTransform.GetTranslation(), WorldTransform.GetRotation(), ScaledObject->GetScale());
@@ -1798,8 +1720,7 @@ bool FPhysInterface_Chaos::GetSquaredDistanceToBody(const FBodyInstance* InInsta
 				continue;
 			}
 
-			bool bGetInnerType = true;
-			ECollisionShapeType GeomType = FPhysicsInterface::GetShapeType(Shape, bGetInnerType);
+			ECollisionShapeType GeomType = FPhysicsInterface::GetShapeType(Shape);
 
 			if (GeomType == ECollisionShapeType::Trimesh)
 			{
diff --git a/Engine/Source/Runtime/Engine/Private/PhysicsEngine/Experimental/PhysicsInterfaceImmediatePhysX.cpp b/Engine/Source/Runtime/Engine/Private/PhysicsEngine/Experimental/PhysicsInterfaceImmediatePhysX.cpp
--- a/Engine/Source/Runtime/Engine/Private/PhysicsEngine/Experimental/PhysicsInterfaceImmediatePhysX.cpp
+++ b/Engine/Source/Runtime/Engine/Private/PhysicsEngine/Experimental/PhysicsInterfaceImmediatePhysX.cpp
@@ -706,25 +706,6 @@ bool FPhysicsInterface_ImmediatePhysX::IsQueryShape(const FPhysicsShapeHandle& I
     return InShape.Actor != nullptr;
 }
 
-bool FPhysicsInterface_ImmediatePhysX::IsShapeType(const FPhysicsShapeHandle& InShape, ECollisionShapeType InType)
-{
-    if (!InShape.Actor)
-    {
-        return false;
-    }
-    const FPhysScene_ImmediatePhysX::FShape& Shape = InShape.Actor->Scene->Actors[InShape.Actor->Index].Shapes[InShape.Index];
-    if ((Shape.Geometry->getType() == PxGeometryType::Enum::eSPHERE && InType == ECollisionShapeType::Sphere)
-        || (Shape.Geometry->getType() == PxGeometryType::Enum::eBOX && InType == ECollisionShapeType::Box)
-        || (Shape.Geometry->getType() == PxGeometryType::Enum::eCONVEXMESH && InType == ECollisionShapeType::Convex)
-        || (Shape.Geometry->getType() == PxGeometryType::Enum::eTRIANGLEMESH && InType == ECollisionShapeType::Trimesh)
-        || (Shape.Geometry->getType() == PxGeometryType::Enum::eHEIGHTFIELD && InType == ECollisionShapeType::Heightfield)
-        || (Shape.Geometry->getType() == PxGeometryType::Enum::eCAPSULE && InType == ECollisionShapeType::TCapsule))
-    {
-        return true;
-    }
-    return false;
-}
-
 ECollisionShapeType FPhysicsInterface_ImmediatePhysX::GetShapeType(const FPhysicsShapeHandle& InShape)
 {
     if (!InShape.Actor)
diff --git a/Engine/Source/Runtime/Engine/Private/PhysicsEngine/PhysicsInterfacePhysX.cpp b/Engine/Source/Runtime/Engine/Private/PhysicsEngine/PhysicsInterfacePhysX.cpp
--- a/Engine/Source/Runtime/Engine/Private/PhysicsEngine/PhysicsInterfacePhysX.cpp
+++ b/Engine/Source/Runtime/Engine/Private/PhysicsEngine/PhysicsInterfacePhysX.cpp
@@ -972,16 +972,6 @@ bool FPhysicsInterface_PhysX::IsQueryShape(const FPhysicsShapeHandle_PhysX& InSh
 	return false;
 }
 
-bool FPhysicsInterface_PhysX::IsShapeType(const FPhysicsShapeHandle_PhysX& InShape, ECollisionShapeType InType)
-{
-	if(InShape.IsValid())
-	{
-		return InShape.Shape->getGeometryType() == U2PCollisionShapeType(InType);
-	}
-
-	return false;
-}
-
 ECollisionShapeType FPhysicsInterface_PhysX::GetShapeType(const FPhysicsShapeHandle_PhysX& InShape)
 {
 	if(InShape.IsValid())
diff --git a/Engine/Source/Runtime/Engine/Public/Physics/Experimental/ChaosDerivedDataReader.h b/Engine/Source/Runtime/Engine/Public/Physics/Experimental/ChaosDerivedDataReader.h
--- a/Engine/Source/Runtime/Engine/Public/Physics/Experimental/ChaosDerivedDataReader.h
+++ b/Engine/Source/Runtime/Engine/Public/Physics/Experimental/ChaosDerivedDataReader.h
@@ -11,6 +11,9 @@ namespace Chaos
 
 	template <typename T>
 	class TTriangleMeshImplicitObject;
+
+	template <typename T, int d>
+	class TConvex;
 }
 
 struct FUntypedBulkData;
@@ -23,7 +26,7 @@ public:
 	// Only valid use is to explicitly read chaos bulk data
 	explicit FChaosDerivedDataReader(FUntypedBulkData* InBulkData);
 
-	TArray<TUniquePtr<Chaos::TImplicitObject<T, d>>> ConvexImplicitObjects;
+	TArray<TUniquePtr<Chaos::TConvex<T, d>>> ConvexImplicitObjects;
 	TArray<TUniquePtr<Chaos::TTriangleMeshImplicitObject<T>>> TrimeshImplicitObjects;
 	FBodySetupUVInfo UVInfo;
 
diff --git a/Engine/Source/Runtime/Engine/Public/Physics/Experimental/PhysInterface_Chaos.h b/Engine/Source/Runtime/Engine/Public/Physics/Experimental/PhysInterface_Chaos.h
--- a/Engine/Source/Runtime/Engine/Public/Physics/Experimental/PhysInterface_Chaos.h
+++ b/Engine/Source/Runtime/Engine/Public/Physics/Experimental/PhysInterface_Chaos.h
@@ -374,8 +374,7 @@ public:
 	static FCollisionFilterData GetQueryFilter(const FPhysicsShapeHandle& InShape);
 	static bool IsSimulationShape(const FPhysicsShapeHandle& InShape);
 	static bool IsQueryShape(const FPhysicsShapeHandle& InShape);
-	static bool IsShapeType(const FPhysicsShapeHandle& InShape, ECollisionShapeType InType);
-	static ECollisionShapeType GetShapeType(const FPhysicsShapeHandle& InShape, bool GetInnerType = false);
+	static ECollisionShapeType GetShapeType(const FPhysicsShapeHandle& InShape);
 	static FTransform GetLocalTransform(const FPhysicsShapeHandle& InShape);
     static void* GetUserData(const FPhysicsShapeHandle& InShape) { return nullptr; }
 
diff --git a/Engine/Source/Runtime/Engine/Public/Physics/Experimental/PhysicsInterfaceImmediatePhysX.h b/Engine/Source/Runtime/Engine/Public/Physics/Experimental/PhysicsInterfaceImmediatePhysX.h
--- a/Engine/Source/Runtime/Engine/Public/Physics/Experimental/PhysicsInterfaceImmediatePhysX.h
+++ b/Engine/Source/Runtime/Engine/Public/Physics/Experimental/PhysicsInterfaceImmediatePhysX.h
@@ -208,7 +208,6 @@ struct ENGINE_API FPhysicsInterface_ImmediatePhysX
 
 	static bool IsSimulationShape(const FPhysicsShapeReference_ImmediatePhysX& InShape);
 	static bool IsQueryShape(const FPhysicsShapeReference_ImmediatePhysX& InShape);
-	static bool IsShapeType(const FPhysicsShapeReference_ImmediatePhysX& InShape, ECollisionShapeType InType);
     // @todo(mlentine): We don't keep track of what is shared but anything can be
     static bool IsShared(const FPhysicsShapeHandle& InShape) { return true; }
 	static ECollisionShapeType GetShapeType(const FPhysicsShapeReference_ImmediatePhysX& InShape);
diff --git a/Engine/Source/Runtime/Engine/Public/Physics/PhysicsInterfacePhysX.h b/Engine/Source/Runtime/Engine/Public/Physics/PhysicsInterfacePhysX.h
--- a/Engine/Source/Runtime/Engine/Public/Physics/PhysicsInterfacePhysX.h
+++ b/Engine/Source/Runtime/Engine/Public/Physics/PhysicsInterfacePhysX.h
@@ -211,7 +211,6 @@ struct ENGINE_API FPhysicsInterface_PhysX : public FGenericPhysicsInterface
 	static FCollisionFilterData GetQueryFilter(const FPhysicsShapeHandle_PhysX& InShape);
 	static bool IsSimulationShape(const FPhysicsShapeHandle_PhysX& InShape);
 	static bool IsQueryShape(const FPhysicsShapeHandle_PhysX& InShape);
-	static bool IsShapeType(const FPhysicsShapeHandle_PhysX& InShape, ECollisionShapeType InType);
 	static ECollisionShapeType GetShapeType(const FPhysicsShapeHandle_PhysX& InShape);
 	static FPhysicsGeometryCollection_PhysX GetGeometryCollection(const FPhysicsShapeHandle_PhysX& InShape);
 	static FTransform GetLocalTransform(const FPhysicsShapeHandle_PhysX& InShape);
diff --git a/Engine/Source/Runtime/Experimental/Chaos/Private/Chaos/ChaosDebugDraw.cpp b/Engine/Source/Runtime/Experimental/Chaos/Private/Chaos/ChaosDebugDraw.cpp
--- a/Engine/Source/Runtime/Experimental/Chaos/Private/Chaos/ChaosDebugDraw.cpp
+++ b/Engine/Source/Runtime/Experimental/Chaos/Private/Chaos/ChaosDebugDraw.cpp
@@ -105,8 +105,6 @@ namespace Chaos
 				break;
 			case ImplicitObjectType::HeightField:
 				break;
-			case ImplicitObjectType::Scaled:
-				break;
 			default:
 				break;
 			}
diff --git a/Engine/Source/Runtime/Experimental/Chaos/Private/Chaos/ImplicitObject.cpp b/Engine/Source/Runtime/Experimental/Chaos/Private/Chaos/ImplicitObject.cpp
--- a/Engine/Source/Runtime/Experimental/Chaos/Private/Chaos/ImplicitObject.cpp
+++ b/Engine/Source/Runtime/Experimental/Chaos/Private/Chaos/ImplicitObject.cpp
@@ -353,9 +353,8 @@ const FName TImplicitObject<T, d>::GetTypeName(const EImplicitObjectType InType)
 	static const FName CylinderName = TEXT("Cylinder");
 	static const FName TriangleMeshName = TEXT("TriangleMesh");
 	static const FName HeightFieldName = TEXT("HeightField");
-	static const FName ScaledName = TEXT("Scaled");
 
-	switch (InType)
+	switch (GetInnerType(InType))
 	{
 		case ImplicitObjectType::Sphere: return SphereName;
 		case ImplicitObjectType::Box: return BoxName;
@@ -370,7 +369,6 @@ const FName TImplicitObject<T, d>::GetTypeName(const EImplicitObjectType InType)
 		case ImplicitObjectType::Cylinder: return CylinderName;
 		case ImplicitObjectType::TriangleMesh: return TriangleMeshName;
 		case ImplicitObjectType::HeightField: return HeightFieldName;
-		case ImplicitObjectType::Scaled: return ScaledName;
 	}
 	return NAME_None;
 }
@@ -380,6 +378,24 @@ TImplicitObject<T, d>* TImplicitObject<T, d>::SerializationFactory(FChaosArchive
 {
 	int8 ObjectType = Ar.IsLoading() ? 0 : (int8)Obj->Type;
 	Ar << ObjectType;
+
+	Ar.UsingCustomVersion(FExternalPhysicsCustomObjectVersion::GUID);
+	if (Ar.CustomVer(FExternalPhysicsCustomObjectVersion::GUID) >= FExternalPhysicsCustomObjectVersion::ScaledGeometryIsConcrete)
+	{
+		if (IsScaled(ObjectType))
+		{
+			EImplicitObjectType InnerType = GetInnerType(ObjectType);
+			switch (InnerType)
+			{
+			case ImplicitObjectType::Convex: if (Ar.IsLoading()) { return new TImplicitObjectScaled<TConvex<T, d>>(); } break;
+			case ImplicitObjectType::TriangleMesh: if (Ar.IsLoading()) { return new TImplicitObjectScaled<TTriangleMeshImplicitObject<T>>(); } break;
+			default: check(false);
+			}
+
+			return nullptr;
+		}
+	}
+
 	switch ((EImplicitObjectType)ObjectType)
 	{
 	case ImplicitObjectType::Sphere: if (Ar.IsLoading()) { return new TSphere<T, d>(); } break;
@@ -392,7 +408,11 @@ TImplicitObject<T, d>* TImplicitObject<T, d>::SerializationFactory(FChaosArchive
 	case ImplicitObjectType::Convex: if (Ar.IsLoading()) { return new TConvex<T, d>(); } break;
 	case ImplicitObjectType::TaperedCylinder: if (Ar.IsLoading()) { return new TTaperedCylinder<T>(); } break;
 	case ImplicitObjectType::TriangleMesh: if (Ar.IsLoading()) { return new TTriangleMeshImplicitObject<T>(); } break;
-	case ImplicitObjectType::Scaled: if (Ar.IsLoading()) { return new TImplicitObjectScaled<T,d>(); } break;
+	case ImplicitObjectType::DEPRECATED_Scaled:
+	{
+		ensure(Ar.IsLoading() && (Ar.CustomVer(FExternalPhysicsCustomObjectVersion::GUID) < FExternalPhysicsCustomObjectVersion::ScaledGeometryIsConcrete));
+		return new TImplicitObjectScaledGeneric<T, d>();
+	}
 	case ImplicitObjectType::HeightField: if (Ar.IsLoading()) { return new THeightField<T>(); } break;
 	case ImplicitObjectType::Cylinder: if (Ar.IsLoading()) { return new TCylinder<T>(); } break;
 	default:
diff --git a/Engine/Source/Runtime/Experimental/Chaos/Private/Chaos/PBDCollisionConstraint.cpp b/Engine/Source/Runtime/Experimental/Chaos/Private/Chaos/PBDCollisionConstraint.cpp
--- a/Engine/Source/Runtime/Experimental/Chaos/Private/Chaos/PBDCollisionConstraint.cpp
+++ b/Engine/Source/Runtime/Experimental/Chaos/Private/Chaos/PBDCollisionConstraint.cpp
@@ -1998,59 +1998,59 @@ typename TPBDCollisionConstraint<T, d>::FRigidBodyContactConstraint TPBDCollisio
 	{
 		return ComputeLevelsetConstraint(Particle0, Particle1, Thickness);
 	}
-	if (Particle0->Geometry()->GetType() == TBox<T, d>::GetType() && Particle1->Geometry()->GetType() == TBox<T, d>::GetType())
+	if (Particle0->Geometry()->GetType() == TBox<T, d>::StaticType() && Particle1->Geometry()->GetType() == TBox<T, d>::StaticType())
 	{
 		return ComputeBoxConstraint(Particle0, Particle1, Thickness);
 	}
-	else if (Particle0->Geometry()->GetType() == TSphere<T, d>::GetType() && Particle1->Geometry()->GetType() == TSphere<T, d>::GetType())
+	else if (Particle0->Geometry()->GetType() == TSphere<T, d>::StaticType() && Particle1->Geometry()->GetType() == TSphere<T, d>::StaticType())
 	{
 		return ComputeSphereConstraint(Particle0, Particle1, Thickness);
 	}
-	else if (Particle0->Geometry()->GetType() == TBox<T, d>::GetType() && Particle1->Geometry()->GetType() == TPlane<T, d>::GetType())
+	else if (Particle0->Geometry()->GetType() == TBox<T, d>::StaticType() && Particle1->Geometry()->GetType() == TPlane<T, d>::StaticType())
 	{
 		return ComputeBoxPlaneConstraint(Particle0, Particle1, Thickness);
 	}
-	else if (Particle1->Geometry()->GetType() == TBox<T, d>::GetType() && Particle0->Geometry()->GetType() == TPlane<T, d>::GetType())
+	else if (Particle1->Geometry()->GetType() == TBox<T, d>::StaticType() && Particle0->Geometry()->GetType() == TPlane<T, d>::StaticType())
 	{
 		return ComputeBoxPlaneConstraint(Particle1, Particle0, Thickness);
 	}
-	else if (Particle0->Geometry()->GetType() == TSphere<T, d>::GetType() && Particle1->Geometry()->GetType() == TPlane<T, d>::GetType())
+	else if (Particle0->Geometry()->GetType() == TSphere<T, d>::StaticType() && Particle1->Geometry()->GetType() == TPlane<T, d>::StaticType())
 	{
 		return ComputeSpherePlaneConstraint(Particle0, Particle1, Thickness);
 	}
-	else if (Particle1->Geometry()->GetType() == TSphere<T, d>::GetType() && Particle0->Geometry()->GetType() == TPlane<T, d>::GetType())
+	else if (Particle1->Geometry()->GetType() == TSphere<T, d>::StaticType() && Particle0->Geometry()->GetType() == TPlane<T, d>::StaticType())
 	{
 		return ComputeSpherePlaneConstraint(Particle1, Particle0, Thickness);
 	}
-	else if (Particle0->Geometry()->GetType() == TSphere<T, d>::GetType() && Particle1->Geometry()->GetType() == TBox<T, d>::GetType())
+	else if (Particle0->Geometry()->GetType() == TSphere<T, d>::StaticType() && Particle1->Geometry()->GetType() == TBox<T, d>::StaticType())
 	{
 		return ComputeSphereBoxConstraint(Particle0, Particle1, Thickness);
 	}
-	else if (Particle1->Geometry()->GetType() == TSphere<T, d>::GetType() && Particle0->Geometry()->GetType() == TBox<T, d>::GetType())
+	else if (Particle1->Geometry()->GetType() == TSphere<T, d>::StaticType() && Particle0->Geometry()->GetType() == TBox<T, d>::StaticType())
 	{
 		return ComputeSphereBoxConstraint(Particle1, Particle0, Thickness);
 	}
-	else if (Particle0->Geometry()->GetType() == TCapsule<T>::GetType() && Particle1->Geometry()->GetType() == TCapsule<T>::GetType())
+	else if (Particle0->Geometry()->GetType() == TCapsule<T>::StaticType() && Particle1->Geometry()->GetType() == TCapsule<T>::StaticType())
 	{
 		return ComputeCapsuleCapsuleConstraint(Particle0, Particle1, Thickness);
 	}
-	else if (Particle0->Geometry()->GetType() == TCapsule<T>::GetType() && Particle1->Geometry()->GetType() == TBox<T, d>::GetType())
+	else if (Particle0->Geometry()->GetType() == TCapsule<T>::StaticType() && Particle1->Geometry()->GetType() == TBox<T, d>::StaticType())
 	{
 		return ComputeCapsuleBoxConstraint(Particle0, Particle1, Thickness);
 	}
-	else if (Particle1->Geometry()->GetType() == TCapsule<T>::GetType() && Particle0->Geometry()->GetType() == TBox<T, d>::GetType())
+	else if (Particle1->Geometry()->GetType() == TCapsule<T>::StaticType() && Particle0->Geometry()->GetType() == TBox<T, d>::StaticType())
 	{
 		return ComputeCapsuleBoxConstraint(Particle1, Particle0, Thickness);
 	}
-	else if (Particle0->Geometry()->GetType() < TImplicitObjectUnion<T, d>::GetType() && Particle1->Geometry()->GetType() == TImplicitObjectUnion<T, d>::GetType())
+	else if (Particle0->Geometry()->GetType() < TImplicitObjectUnion<T, d>::StaticType() && Particle1->Geometry()->GetType() == TImplicitObjectUnion<T, d>::StaticType())
 	{
 		return ComputeSingleUnionConstraint(Particle0, Particle1, Thickness);
 	}
-	else if (Particle0->Geometry()->GetType() == TImplicitObjectUnion<T, d>::GetType() && Particle1->Geometry()->GetType() < TImplicitObjectUnion<T, d>::GetType())
+	else if (Particle0->Geometry()->GetType() == TImplicitObjectUnion<T, d>::StaticType() && Particle1->Geometry()->GetType() < TImplicitObjectUnion<T, d>::StaticType())
 	{
 		return ComputeSingleUnionConstraint(Particle1, Particle0, Thickness);
 	}
-	else if(Particle0->Geometry()->GetType() == TImplicitObjectUnion<T, d>::GetType() && Particle1->Geometry()->GetType() == TImplicitObjectUnion<T, d>::GetType())
+	else if(Particle0->Geometry()->GetType() == TImplicitObjectUnion<T, d>::StaticType() && Particle1->Geometry()->GetType() == TImplicitObjectUnion<T, d>::StaticType())
 	{
 		return ComputeUnionUnionConstraint(Particle0, Particle1, Thickness);
 	}
@@ -2080,35 +2080,35 @@ void UpdateLevelsetConstraintImp(const float Thickness, TRigidBodyContactConstra
 template<ECollisionUpdateType UpdateType, typename T, int d>
 void UpdateConstraintImp2(const TImplicitObject<T, d>& ParticleObject, const TRigidTransform<T, d>& ParticleTM, const TImplicitObject<T, d>& LevelsetObject, const TRigidTransform<T, d>& LevelsetTM, const T Thickness, TRigidBodyContactConstraint<T, d>& Constraint)
 {
-	if (ParticleObject.GetType() == TBox<T, d>::GetType() && LevelsetObject.GetType() == TBox<T, d>::GetType())
+	if (ParticleObject.GetType() == TBox<T, d>::StaticType() && LevelsetObject.GetType() == TBox<T, d>::StaticType())
 	{
 		UpdateBoxConstraint(*ParticleObject.template GetObject<TBox<T,d>>(), ParticleTM, *LevelsetObject.template GetObject<TBox<T, d>>(), LevelsetTM, Thickness, Constraint);
 	}
-	else if (ParticleObject.GetType() == TSphere<T, d>::GetType() && LevelsetObject.GetType() == TSphere<T, d>::GetType())
+	else if (ParticleObject.GetType() == TSphere<T, d>::StaticType() && LevelsetObject.GetType() == TSphere<T, d>::StaticType())
 	{
 		UpdateSphereConstraint(*ParticleObject.template GetObject<TSphere<T, d>>(), ParticleTM, *LevelsetObject.template GetObject<TSphere<T, d>>(), LevelsetTM, Thickness, Constraint);
 	}
-	else if (ParticleObject.GetType() == TBox<T, d>::GetType() && LevelsetObject.GetType() == TPlane<T, d>::GetType())
+	else if (ParticleObject.GetType() == TBox<T, d>::StaticType() && LevelsetObject.GetType() == TPlane<T, d>::StaticType())
 	{
 		UpdateBoxPlaneConstraint(*ParticleObject.template GetObject<TBox<T, d>>(), ParticleTM, *LevelsetObject.template GetObject<TPlane<T, d>>(), LevelsetTM, Thickness, Constraint);
 	}
-	else if (ParticleObject.GetType() == TSphere<T, d>::GetType() && LevelsetObject.GetType() == TPlane<T, d>::GetType())
+	else if (ParticleObject.GetType() == TSphere<T, d>::StaticType() && LevelsetObject.GetType() == TPlane<T, d>::StaticType())
 	{
 		UpdateSpherePlaneConstraint(*ParticleObject.template GetObject<TSphere<T, d>>(), ParticleTM, *LevelsetObject.template GetObject<TPlane<T, d>>(), LevelsetTM, Thickness, Constraint);
 	}
-	else if (ParticleObject.GetType() == TSphere<T, d>::GetType() && LevelsetObject.GetType() == TBox<T, d>::GetType())
+	else if (ParticleObject.GetType() == TSphere<T, d>::StaticType() && LevelsetObject.GetType() == TBox<T, d>::StaticType())
 	{
 		UpdateSphereBoxConstraint(*ParticleObject.template GetObject<TSphere<T, d>>(), ParticleTM, *LevelsetObject.template GetObject<TBox<T, d>>(), LevelsetTM, Thickness, Constraint);
 	}
-	else if (ParticleObject.GetType() == TCapsule<T>::GetType() && LevelsetObject.GetType() == TCapsule<T>::GetType())
+	else if (ParticleObject.GetType() == TCapsule<T>::StaticType() && LevelsetObject.GetType() == TCapsule<T>::StaticType())
 	{
 		UpdateCapsuleCapsuleConstraint(*ParticleObject.template GetObject<TCapsule<T>>(), ParticleTM, *LevelsetObject.template GetObject<TCapsule<T>>(), LevelsetTM, Thickness, Constraint);
 	}
-	else if (ParticleObject.GetType() == TCapsule<T>::GetType() && LevelsetObject.GetType() == TBox<T, d>::GetType())
+	else if (ParticleObject.GetType() == TCapsule<T>::StaticType() && LevelsetObject.GetType() == TBox<T, d>::StaticType())
 	{
 		UpdateCapsuleBoxConstraint(*ParticleObject.template GetObject<TCapsule<T>>(), ParticleTM, *LevelsetObject.template GetObject<TBox<T, d>>(), LevelsetTM, Thickness, Constraint);
 	}
-	else if (ParticleObject.GetType() == TPlane<T, d>::GetType() && LevelsetObject.GetType() == TBox<T, d>::GetType())
+	else if (ParticleObject.GetType() == TPlane<T, d>::StaticType() && LevelsetObject.GetType() == TBox<T, d>::StaticType())
 	{
 		TRigidBodyContactConstraint<T, d> TmpConstraint = Constraint;
 		UpdateBoxPlaneConstraint(*LevelsetObject.template GetObject<TBox<T, d>>(), LevelsetTM, *ParticleObject.template GetObject<TPlane<T, d>>(), ParticleTM, Thickness, TmpConstraint);
@@ -2118,7 +2118,7 @@ void UpdateConstraintImp2(const TImplicitObject<T, d>& ParticleObject, const TRi
 			Constraint.Normal = -Constraint.Normal;
 		}
 	}
-	else if (ParticleObject.GetType() == TPlane<T, d>::GetType() && LevelsetObject.GetType() == TSphere<T, d>::GetType())
+	else if (ParticleObject.GetType() == TPlane<T, d>::StaticType() && LevelsetObject.GetType() == TSphere<T, d>::StaticType())
 	{
 		TRigidBodyContactConstraint<T, d> TmpConstraint = Constraint;
 		UpdateSpherePlaneConstraint(*LevelsetObject.template GetObject<TSphere<T, d>>(), LevelsetTM, *ParticleObject.template GetObject<TPlane<T, d>>(), ParticleTM, Thickness, TmpConstraint);
@@ -2128,7 +2128,7 @@ void UpdateConstraintImp2(const TImplicitObject<T, d>& ParticleObject, const TRi
 			Constraint.Normal = -Constraint.Normal;
 		}
 	}
-	else if (ParticleObject.GetType() == TBox<T, d>::GetType() && LevelsetObject.GetType() == TSphere<T, d>::GetType())
+	else if (ParticleObject.GetType() == TBox<T, d>::StaticType() && LevelsetObject.GetType() == TSphere<T, d>::StaticType())
 	{
 		TRigidBodyContactConstraint<T, d> TmpConstraint = Constraint;
 		UpdateSphereBoxConstraint(*LevelsetObject.template GetObject<TSphere<T, d>>(), LevelsetTM, *ParticleObject.template GetObject<TBox<T, d>>(), ParticleTM, Thickness, TmpConstraint);
@@ -2138,7 +2138,7 @@ void UpdateConstraintImp2(const TImplicitObject<T, d>& ParticleObject, const TRi
 			Constraint.Normal = -Constraint.Normal;
 		}
 	}
-	else if (ParticleObject.GetType() == TBox<T, d>::GetType() && LevelsetObject.GetType() == TCapsule<T>::GetType())
+	else if (ParticleObject.GetType() == TBox<T, d>::StaticType() && LevelsetObject.GetType() == TCapsule<T>::StaticType())
 	{
 		TRigidBodyContactConstraint<T, d> TmpConstraint = Constraint;
 		UpdateCapsuleBoxConstraint(*LevelsetObject.template GetObject<TCapsule<T>>(), LevelsetTM, *ParticleObject.template GetObject<TBox<T, d>>(), ParticleTM, Thickness, TmpConstraint);
@@ -2148,15 +2148,15 @@ void UpdateConstraintImp2(const TImplicitObject<T, d>& ParticleObject, const TRi
 			Constraint.Normal = -Constraint.Normal;
 		}
 	}
-	else if (ParticleObject.GetType() < TImplicitObjectUnion<T, d>::GetType() && LevelsetObject.GetType() == TImplicitObjectUnion<T, d>::GetType())
+	else if (ParticleObject.GetType() < TImplicitObjectUnion<T, d>::StaticType() && LevelsetObject.GetType() == TImplicitObjectUnion<T, d>::StaticType())
 	{
 		return UpdateSingleUnionConstraint<UpdateType>(Thickness, Constraint);
 	}
-	else if (ParticleObject.GetType() == TImplicitObjectUnion<T, d>::GetType() && LevelsetObject.GetType() < TImplicitObjectUnion<T, d>::GetType())
+	else if (ParticleObject.GetType() == TImplicitObjectUnion<T, d>::StaticType() && LevelsetObject.GetType() < TImplicitObjectUnion<T, d>::StaticType())
 	{
 		check(false);	//should not be possible to get this ordering (see ComputeConstraint)
 	}
-	else if (ParticleObject.GetType() == TImplicitObjectUnion<T, d>::GetType() && LevelsetObject.GetType() == TImplicitObjectUnion<T, d>::GetType())
+	else if (ParticleObject.GetType() == TImplicitObjectUnion<T, d>::StaticType() && LevelsetObject.GetType() == TImplicitObjectUnion<T, d>::StaticType())
 	{
 		return UpdateUnionUnionConstraint<UpdateType>(Thickness, Constraint);
 	}
diff --git a/Engine/Source/Runtime/Experimental/Chaos/Public/Chaos/Box.h b/Engine/Source/Runtime/Experimental/Chaos/Public/Chaos/Box.h
--- a/Engine/Source/Runtime/Experimental/Chaos/Public/Chaos/Box.h
+++ b/Engine/Source/Runtime/Experimental/Chaos/Public/Chaos/Box.h
@@ -120,7 +120,7 @@ namespace Chaos
 			return AABB.Contains(Point, Tolerance);
 		}
 
-		FORCEINLINE static EImplicitObjectType GetType() { return ImplicitObjectType::Box; }
+		FORCEINLINE static EImplicitObjectType StaticType() { return ImplicitObjectType::Box; }
 
 		const TBox<T, d>& BoundingBox() const { return *this; }
 
diff --git a/Engine/Source/Runtime/Experimental/Chaos/Public/Chaos/Capsule.h b/Engine/Source/Runtime/Experimental/Chaos/Public/Chaos/Capsule.h
--- a/Engine/Source/Runtime/Experimental/Chaos/Public/Chaos/Capsule.h
+++ b/Engine/Source/Runtime/Experimental/Chaos/Public/Chaos/Capsule.h
@@ -74,7 +74,7 @@ namespace Chaos
 
 		~TCapsule() {}
 
-		static EImplicitObjectType GetType() { return ImplicitObjectType::Capsule; }
+		static EImplicitObjectType StaticType() { return ImplicitObjectType::Capsule; }
 
 		static TCapsule<T> NewFromOriginAndAxis(const TVector<T, 3>& Origin, const TVector<T, 3>& Axis, const T Height, const T Radius)
 		{
diff --git a/Engine/Source/Runtime/Experimental/Chaos/Public/Chaos/Convex.h b/Engine/Source/Runtime/Experimental/Chaos/Public/Chaos/Convex.h
--- a/Engine/Source/Runtime/Experimental/Chaos/Public/Chaos/Convex.h
+++ b/Engine/Source/Runtime/Experimental/Chaos/Public/Chaos/Convex.h
@@ -37,7 +37,7 @@ namespace Chaos
 			TConvexBuilder<T>::Build(InParticles, Planes, SurfaceParticles, LocalBoundingBox);
 		}
 
-		static EImplicitObjectType GetType()
+		static EImplicitObjectType StaticType()
 		{
 			return ImplicitObjectType::Convex;
 		}
diff --git a/Engine/Source/Runtime/Experimental/Chaos/Public/Chaos/Cylinder.h b/Engine/Source/Runtime/Experimental/Chaos/Public/Chaos/Cylinder.h
--- a/Engine/Source/Runtime/Experimental/Chaos/Public/Chaos/Cylinder.h
+++ b/Engine/Source/Runtime/Experimental/Chaos/Public/Chaos/Cylinder.h
@@ -50,7 +50,7 @@ namespace Chaos
 		}
 		~TCylinder() {}
 
-		static EImplicitObjectType GetType() { return ImplicitObjectType::Cylinder; }
+		static EImplicitObjectType StaticType() { return ImplicitObjectType::Cylinder; }
 
 		/**
 		 * Returns sample points centered about the origin.
diff --git a/Engine/Source/Runtime/Experimental/Chaos/Public/Chaos/GeometryQueries.h b/Engine/Source/Runtime/Experimental/Chaos/Public/Chaos/GeometryQueries.h
--- a/Engine/Source/Runtime/Experimental/Chaos/Public/Chaos/GeometryQueries.h
+++ b/Engine/Source/Runtime/Experimental/Chaos/Public/Chaos/GeometryQueries.h
@@ -66,16 +66,16 @@ namespace Chaos
 				const TTriangleMeshImplicitObject<T>& ATriangleMesh = static_cast<const TTriangleMeshImplicitObject<T>&>(A);
 				return ATriangleMesh.OverlapGeom(B, BToATM, Thickness);
 			}
-			case ImplicitObjectType::Scaled:
-			{
-				const TImplicitObjectScaled<T, d>& AScaled = static_cast<const TImplicitObjectScaled<T, d>&>(A);
-				const TImplicitObject<T, d>& UnscaledObj = *AScaled.GetUnscaledObject();
-				check(UnscaledObj.GetType(true) == ImplicitObjectType::TriangleMesh);
-				const TTriangleMeshImplicitObject<T>& ATriangleMesh = static_cast<const TTriangleMeshImplicitObject<T>&>(UnscaledObj);
-				return TImplicitObjectScaled<T,d>::LowLevelOverlapGeom(AScaled, ATriangleMesh, B, BToATM, Thickness);
-			}
 			default:
-				check(false);	//unsupported query type
+				if (IsScaled(AType))
+				{
+					const auto& AScaled = TImplicitObjectScaled<TTriangleMeshImplicitObject<T>>::AsScaledChecked(A);
+					return AScaled.LowLevelOverlapGeom(B, BToATM, Thickness);
+				}
+				else
+				{
+					check(false);	//unsupported query type
+				}
 
 			}
 		}
@@ -115,9 +115,9 @@ namespace Chaos
 		const TVector<T, d> LocalDir = ATM.InverseTransformVectorNoScale(Dir);
 
 		bool bSweepAsRaycast = BType == ImplicitObjectType::Sphere;
-		if (bSweepAsRaycast && AType == ImplicitObjectType::Scaled)
+		if (bSweepAsRaycast && IsScaled(AType))
 		{
-			const auto& Scaled = A. template GetObjectChecked<TImplicitObjectScaled<T, d>>();
+			const auto& Scaled = TImplicitObjectScaledGeneric<T, d>::AsScaledChecked(A);
 			const TVector<T, d>& Scale = Scaled.GetScale();
 			bSweepAsRaycast = FMath::IsNearlyEqual(Scale[0], Scale[1]) && FMath::IsNearlyEqual(Scale[0], Scale[2]);
 		}
@@ -134,15 +134,10 @@ namespace Chaos
 			auto IsValidConvex = [](const TImplicitObject<T, 3>& InObject) -> bool
 			{
 				//todo: move this out of here
-				const TImplicitObject<T, 3>* Obj = &InObject;
-				if (const TImplicitObjectScaled<T, 3>* Scaled = InObject.template GetObject<TImplicitObjectScaled<T, 3>>())
-				{
-					Obj = Scaled->GetUnscaledObject();
-				}
-				if (const TConvex<T, 3>* Convex = Obj->template GetObject<TConvex<T, 3>>())
+				if (const auto Convex = TImplicitObjectScaled<TConvex<T,3>>::AsScaled(InObject))
 				{
-					return Convex->GetSurfaceParticles().Size() > 0;
-				}
+					return Convex->GetUnscaledObject()->GetSurfaceParticles().Size() > 0;
+				}				
 
 				return true;
 			};
@@ -163,8 +158,9 @@ namespace Chaos
 			{
 				//todo: find face index
 			}
-			else if (AType == ImplicitObjectType::Scaled)
+			else if (AType == ImplicitObjectType::DEPRECATED_Scaled)
 			{
+				ensure(false);
 				//todo: find face index if convex hull
 			}
 
@@ -199,17 +195,17 @@ namespace Chaos
 				bResult = ATriangleMesh.SweepGeom(B, BToATM, LocalDir, Length, OutTime, LocalPosition, LocalNormal, OutFaceIndex, Thickness);
 				break;
 			}
-			case ImplicitObjectType::Scaled:
-			{
-				const TImplicitObjectScaled<T, d>& AScaled = static_cast<const TImplicitObjectScaled<T, d>&>(A);
-				const TImplicitObject<T, d>& UnscaledObj = *AScaled.GetUnscaledObject();
-				check(UnscaledObj.GetType(true) == ImplicitObjectType::TriangleMesh);
-				const TTriangleMeshImplicitObject<T>& ATriangleMesh = static_cast<const TTriangleMeshImplicitObject<T>&>(UnscaledObj);
-				bResult = TImplicitObjectScaled<T,d>::LowLevelSweepGeom(AScaled, ATriangleMesh, B, BToATM, LocalDir, Length, OutTime, LocalPosition, LocalNormal, OutFaceIndex, Thickness);
-				break;
-			}
 			default:
-				check(false);	//unsupported query type
+				if (IsScaled(AType))
+				{
+					const auto& AScaled = TImplicitObjectScaled<TTriangleMeshImplicitObject<T>>::AsScaledChecked(A);
+					bResult = AScaled.LowLevelSweepGeom(B, BToATM, LocalDir, Length, OutTime, LocalPosition, LocalNormal, OutFaceIndex, Thickness);
+					break;
+				}
+				else
+				{
+					check(false);	//unsupported query type
+				}
 
 			}
 		}
diff --git a/Engine/Source/Runtime/Experimental/Chaos/Public/Chaos/HeightField.h b/Engine/Source/Runtime/Experimental/Chaos/Public/Chaos/HeightField.h
--- a/Engine/Source/Runtime/Experimental/Chaos/Public/Chaos/HeightField.h
+++ b/Engine/Source/Runtime/Experimental/Chaos/Public/Chaos/HeightField.h
@@ -70,7 +70,7 @@ namespace Chaos
 			return FCrc::MemCrc32(Bytes.GetData(), Bytes.GetAllocatedSize());
 		}
 
-		static EImplicitObjectType GetType()
+		static EImplicitObjectType StaticType()
 		{
 			return ImplicitObjectType::HeightField;
 		}
diff --git a/Engine/Source/Runtime/Experimental/Chaos/Public/Chaos/ImplicitObject.h b/Engine/Source/Runtime/Experimental/Chaos/Public/Chaos/ImplicitObject.h
--- a/Engine/Source/Runtime/Experimental/Chaos/Public/Chaos/ImplicitObject.h
+++ b/Engine/Source/Runtime/Experimental/Chaos/Public/Chaos/ImplicitObject.h
@@ -44,7 +44,7 @@ namespace ImplicitObjectType
 		Cylinder,
 		TriangleMesh,
 		HeightField,
-		Scaled,
+		DEPRECATED_Scaled,	//needed for serialization of existing data
 		Triangle,
 
 		//Add entries above this line for serialization
@@ -55,6 +55,22 @@ namespace ImplicitObjectType
 
 using EImplicitObjectType = uint8;	//see ImplicitObjectType
 
+FORCEINLINE bool IsInstanced(EImplicitObjectType Type)
+{
+	return (Type & ImplicitObjectType::IsInstanced) != 0;
+}
+
+FORCEINLINE bool IsScaled(EImplicitObjectType Type)
+{
+	return (Type & ImplicitObjectType::IsScaled) != 0;
+}
+
+FORCEINLINE EImplicitObjectType GetInnerType(EImplicitObjectType Type)
+{
+	return Type & (~(ImplicitObjectType::IsScaled | ImplicitObjectType::IsInstanced));
+}
+
+
 namespace EImplicitObject
 {
 	enum Flags
@@ -102,6 +118,8 @@ template<class T, int d>
 class CHAOS_API TImplicitObject
 {
 public:
+	using TType = T;
+	static constexpr int D = d;
 	static TImplicitObject<T,d>* SerializationFactory(FChaosArchive& Ar, TImplicitObject<T, d>* Obj);
 
 	TImplicitObject(int32 Flags, EImplicitObjectType InType = ImplicitObjectType::Unknown);
@@ -112,7 +130,7 @@ public:
 	template<class T_DERIVED>
 	T_DERIVED* GetObject()
 	{
-		if (T_DERIVED::GetType() == Type)
+		if (T_DERIVED::StaticType() == Type)
 		{
 			return static_cast<T_DERIVED*>(this);
 		}
@@ -122,7 +140,7 @@ public:
 	template<class T_DERIVED>
 	const T_DERIVED* GetObject() const
 	{
-		if (T_DERIVED::GetType() == Type)
+		if (T_DERIVED::StaticType() == Type)
 		{
 			return static_cast<const T_DERIVED*>(this);
 		}
@@ -132,14 +150,14 @@ public:
 	template<class T_DERIVED>
 	const T_DERIVED& GetObjectChecked() const
 	{
-		check(T_DERIVED::GetType() == Type);
+		check(T_DERIVED::StaticType() == Type);
 		return static_cast<const T_DERIVED&>(*this);
 	}
 
 	template<class T_DERIVED>
 	T_DERIVED& GetObjectChecked()
 	{
-		check(T_DERIVED::GetType() == Type);
+		check(T_DERIVED::StaticType() == Type);
 		return static_cast<const T_DERIVED&>(*this);
 	}
 
@@ -246,6 +264,11 @@ public:
 	}
 
 	void SerializeImp(FArchive& Ar);
+
+	static EImplicitObjectType StaticType()
+	{
+		return ImplicitObjectType::Unknown;
+	}
 	
 	virtual void Serialize(FArchive& Ar)
 	{
diff --git a/Engine/Source/Runtime/Experimental/Chaos/Public/Chaos/ImplicitObjectScaled.h b/Engine/Source/Runtime/Experimental/Chaos/Public/Chaos/ImplicitObjectScaled.h
--- a/Engine/Source/Runtime/Experimental/Chaos/Public/Chaos/ImplicitObjectScaled.h
+++ b/Engine/Source/Runtime/Experimental/Chaos/Public/Chaos/ImplicitObjectScaled.h
@@ -11,23 +11,139 @@
 
 namespace Chaos
 {
-template<class T, int d, bool bInstanced = true>
-class TImplicitObjectScaled final : public TImplicitObject<T, d>
+
+template <typename TConcrete, bool bInstanced = true>
+class TImplicitObjectInstanced final : public TImplicitObject<typename TConcrete::TType,TConcrete::D>
 {
-using ObjectType = typename TChooseClass<bInstanced, TSerializablePtr<TImplicitObject<T, d>>, TUniquePtr<TImplicitObject<T, d>>>::Result;
 public:
+	using T = typename TConcrete::TType;
+	using TType = T;
+	static constexpr int d = TConcrete::D;
+	static constexpr int D = d;
 
+	using ObjectType = typename TChooseClass<bInstanced, TSerializablePtr<TConcrete>, TUniquePtr<TConcrete>>::Result;
+	using TImplicitObject<T, d>::GetTypeName;
+
+	TImplicitObjectInstanced(ObjectType Object)
+		: TImplicitObject<T, d>(EImplicitObject::HasBoundingBox, Object->GetType() | ImplicitObjectType::IsInstanced)
+		, MObject(MoveTemp(Object))
+	{
+		ensure(IsInstanced(MObject->GetType(true)) == false);	//cannot have an instance of an instance
+		this->bIsConvex = MObject->IsConvex();
+	}
+
+	static EImplicitObjectType StaticType()
+	{
+		return TConcrete::StaticType() | ImplicitObjectType::IsInstanced;
+	}
+
+	const TConcrete* GetInstancedObject() const
+	{
+		return MObject.Get();
+	}
+
+	virtual T PhiWithNormal(const TVector<T, d>& X, TVector<T, d>& Normal) const override
+	{
+		return MObject->PhiWithNormal(X, Normal);
+	}
+
+	virtual bool Raycast(const TVector<T, d>& StartPoint, const TVector<T, d>& Dir, const T Length, const T Thickness, T& OutTime, TVector<T, d>& OutPosition, TVector<T, d>& OutNormal, int32& OutFaceIndex) const override
+	{
+		return MObject->Raycast(StartPoint, Dir, Length, Thickness, OutTime, OutPosition, OutNormal, OutFaceIndex);
+	}
+
+	virtual void Serialize(FChaosArchive& Ar) override
+	{
+		FChaosArchiveScopedMemory ScopedMemory(Ar, GetTypeName(), false);
+		TImplicitObject<T, d>::SerializeImp(Ar);
+		Ar << MObject;
+	}
+
+	/** This is a low level function and assumes the internal object has a SweepGeom function. Should not be called directly. See GeometryQueries.h : SweepQuery */
+	template <typename TGeometry>
+	static bool LowLevelSweepGeom(const TImplicitObjectInstanced<TConcrete, bInstanced>& OwningInstanced, const TImplicitObject<T, d>& B, const TRigidTransform<T, d>& BToATM, const TVector<T, d>& LocalDir, const T Length, T& OutTime, TVector<T, d>& LocalPosition, TVector<T, d>& LocalNormal, int32& OutFaceIndex, T Thickness = 0)
+	{
+		ensure(false);
+		return false;
+	}
+
+	/** This is a low level function and assumes the internal object has a OverlapGeom function. Should not be called directly. See GeometryQueries.h : OverlapQuery */
+	bool LowLevelOverlapGeom(const TImplicitObjectInstanced<TConcrete, bInstanced>& OwningInstanced, const TImplicitObject<T, d>& B, const TRigidTransform<T, d>& BToATM, T Thickness = 0) const
+	{
+		ensure(false);
+		return false;
+	}
+
+	virtual int32 FindMostOpposingFace(const TVector<T, d>& Position, const TVector<T, d>& UnitDir, int32 HintFaceIndex, T SearchDist) const override
+	{
+		return MObject->FindMostOpposingFace(Position, UnitDir, HintFaceIndex, SearchDist);
+	}
+
+	virtual TVector<T, 3> FindGeometryOpposingNormal(const TVector<T, d>& DenormDir, int32 HintFaceIndex, const TVector<T, d>& OriginalNormal) const override
+	{
+		return MObject->FindGeometryOpposingNormal(DenormDir, HintFaceIndex, OriginalNormal);
+	}
+
+	virtual bool Overlap(const TVector<T, d>& Point, const T Thickness) const override
+	{
+		return MObject->Overlap(Point, Thickness);
+	}
+
+	virtual TVector<T, d> Support(const TVector<T, d>& Direction, const T Thickness) const override
+	{
+		return MObject->Support(Direction, Thickness);
+	}
+
+	virtual const TBox<T, d>& BoundingBox() const override { return MObject->BoundingBox(); }
+
+	const ObjectType Object() const { return MObject; }
+
+	virtual uint32 GetTypeHash() const override
+	{
+		return MObject->GetTypeHash();
+	}
+
+	virtual TUniquePtr<TImplicitObject<T, d>> Copy() const override
+	{
+		return TUniquePtr<TImplicitObject<T, d>>(CopyHelper(this));
+	}
+
+protected:
+	ObjectType MObject;
+
+	static TImplicitObjectInstanced<TConcrete, true>* CopyHelper(const TImplicitObjectInstanced<TConcrete, true>* Obj)
+	{
+		return new TImplicitObjectInstanced<TConcrete, true>(Obj->MObject);
+	}
+
+	static TImplicitObjectInstanced<TConcrete, false>* CopyHelper(const TImplicitObjectInstanced<TConcrete, false>* Obj)
+	{
+		return new TImplicitObjectInstanced<TConcrete, false>(Obj->MObject->Copy());
+	}
+};
+
+
+template<typename TConcrete, bool bInstanced = true>
+class TImplicitObjectScaled final : public TImplicitObject<typename TConcrete::TType, TConcrete::D>
+{
+public:
+	using T = typename TConcrete::TType;
+	using TType = T;
+	static constexpr int d = TConcrete::D;
+	static constexpr int D = d;
+
+	using ObjectType = typename TChooseClass<bInstanced, TSerializablePtr<TConcrete>, TUniquePtr<TConcrete>>::Result;
 	using TImplicitObject<T, d>::GetTypeName;
 
 	TImplicitObjectScaled(ObjectType Object, const TVector<T, d>& Scale, T Thickness = 0)
-	    : TImplicitObject<T, d>(EImplicitObject::HasBoundingBox, ImplicitObjectType::Scaled)
+	    : TImplicitObject<T, d>(EImplicitObject::HasBoundingBox, Object->GetType() | ImplicitObjectType::IsScaled)
 	    , MObject(MoveTemp(Object))
 		, MInternalThickness(Thickness)
 	{
-		ensureMsgf(MObject->GetType(true) != ImplicitObjectType::Scaled, TEXT("Scaled objects should not contain each other."));
+		ensureMsgf((IsScaled(MObject->GetType(true)) == false), TEXT("Scaled objects should not contain each other."));
+		ensureMsgf((IsInstanced(MObject->GetType(true)) == false), TEXT("Scaled objects should not contain instances."));
 		switch (MObject->GetType(true))
 		{
-		case ImplicitObjectType::Scaled:
 		case ImplicitObjectType::Transformed:
 		case ImplicitObjectType::Union:
 			check(false);	//scale is only supported for concrete types like sphere, capsule, convex, levelset, etc... Nothing that contains other objects
@@ -38,35 +154,91 @@ public:
 		SetScale(Scale);
 	}
 	TImplicitObjectScaled(ObjectType Object, TUniquePtr<Chaos::TImplicitObject<T,d>> &&ObjectOwner, const TVector<T, d>& Scale, T Thickness = 0)
-	    : TImplicitObject<T, d>(EImplicitObject::HasBoundingBox, ImplicitObjectType::Scaled)
+	    : TImplicitObject<T, d>(EImplicitObject::HasBoundingBox, Object->GetType() | ImplicitObjectType::IsScaled)
 	    , MObject(Object)
 		, MInternalThickness(Thickness)
 	{
-		ensureMsgf(MObject->GetType(true) != ImplicitObjectType::Scaled, TEXT("Scaled objects should not contain each other."));
+		ensureMsgf((IsScaled(MObject->GetType(true)) == false), TEXT("Scaled objects should not contain each other."));
+		ensureMsgf((IsInstanced(MObject->GetType(true)) == false), TEXT("Scaled objects should not contain instances."));
 		this->bIsConvex = Object->IsConvex();
 		SetScale(Scale);
 	}
 
-	TImplicitObjectScaled(const TImplicitObjectScaled<T, d, bInstanced>& Other) = delete;
-	TImplicitObjectScaled(TImplicitObjectScaled<T, d, bInstanced>&& Other)
-	    : TImplicitObject<T, d>(EImplicitObject::HasBoundingBox, ImplicitObjectType::Scaled)
-	    , MObject(Other.MObject)
+	TImplicitObjectScaled(const TImplicitObjectScaled<TConcrete, bInstanced>& Other) = delete;
+	TImplicitObjectScaled(TImplicitObjectScaled<TConcrete, bInstanced>&& Other)
+	    : TImplicitObject<T, d>(EImplicitObject::HasBoundingBox, Other.MObject->GetType() | ImplicitObjectType::IsScaled)
+	    , MObject(MoveTemp(Other.MObject))
 	    , MScale(Other.MScale)
 		, MInvScale(Other.MInvScale)
 		, MInternalThickness(Other.MInternalThickness)
 	    , MLocalBoundingBox(MoveTemp(Other.MLocalBoundingBox))
 	{
-		ensureMsgf(MObject->GetType(true) != ImplicitObjectType::Scaled, TEXT("Scaled objects should not contain each other."));
+		ensureMsgf((IsScaled(MObject->GetType(true)) == false), TEXT("Scaled objects should not contain each other."));
+		ensureMsgf((IsInstanced(MObject->GetType(true)) == false), TEXT("Scaled objects should not contain instances."));
 		this->bIsConvex = Other.MObject->IsConvex();
 	}
 	~TImplicitObjectScaled() {}
 
-	static EImplicitObjectType GetType()
+	static EImplicitObjectType StaticType()
+	{
+		return TConcrete::StaticType() | ImplicitObjectType::IsScaled;
+	}
+
+	static const TImplicitObjectScaled<TConcrete>& AsScaledChecked(const TImplicitObject<T, d>& Obj)
+	{
+		if (TIsSame<TConcrete, TImplicitObject<T, d>>::Value)
+		{
+			//can cast any scaled to ImplicitObject base
+			check(IsScaled(Obj.GetType()));
+		}
+		else
+		{
+			check(StaticType() == Obj.GetType());
+		}
+		return static_cast<const TImplicitObjectScaled<TConcrete>&>(Obj);
+	}
+
+	static TImplicitObjectScaled<TConcrete>& AsScaledChecked(TImplicitObject<T, d>& Obj)
+	{
+		if (TIsSame<TConcrete, TImplicitObject<T, d>>::Value)
+		{
+			//can cast any scaled to ImplicitObject base
+			check(IsScaled(Obj.GetType()));
+		}
+		else
+		{
+			check(StaticType() == Obj.GetType());
+		}
+		return static_cast<TImplicitObjectScaled<TConcrete>&>(Obj);
+	}
+
+	static const TImplicitObjectScaled<TConcrete>* AsScaled(const TImplicitObject<T, d>& Obj)
+	{
+		if (TIsSame<TConcrete, TImplicitObject<T, d>>::Value)
+		{
+			//can cast any scaled to ImplicitObject base
+			return IsScaled(Obj.GetType()) ? static_cast<const TImplicitObjectScaled<TConcrete>*>(&Obj) : nullptr;
+		}
+		else
+		{
+			return StaticType() == Obj.GetType() ? static_cast<const TImplicitObjectScaled<TConcrete>*>(&Obj) : nullptr;
+		}
+	}
+
+	static TImplicitObjectScaled<TConcrete>* AsScaled(TImplicitObject<T, d>& Obj)
 	{
-		return ImplicitObjectType::Scaled;
+		if (TIsSame<TConcrete, TImplicitObject<T, d>>::Value)
+		{
+			//can cast any scaled to ImplicitObject base
+			return IsScaled(Obj.GetType()) ? static_cast<TImplicitObjectScaled<TConcrete>*>(&Obj) : nullptr;
+		}
+		else
+		{
+			return StaticType() == Obj.GetType() ? static_cast<TImplicitObjectScaled<TConcrete>*>(&Obj) : nullptr;
+		}
 	}
 
-	const TImplicitObject<T, d>* GetUnscaledObject() const
+	const TConcrete* GetUnscaledObject() const
 	{
 		return MObject.Get();
 	}
@@ -115,14 +287,13 @@ public:
 	}
 
 	/** This is a low level function and assumes the internal object has a SweepGeom function. Should not be called directly. See GeometryQueries.h : SweepQuery */
-	template <typename TGeometry>
-	static bool LowLevelSweepGeom(const TImplicitObjectScaled<T, d>& OwningScaled, const TGeometry& InternalGeom, const TImplicitObject<T, d>& B, const TRigidTransform<T, d>& BToATM, const TVector<T, d>& LocalDir, const T Length, T& OutTime, TVector<T, d>& LocalPosition, TVector<T, d>& LocalNormal, int32& OutFaceIndex, T Thickness = 0)
+	bool LowLevelSweepGeom(const TImplicitObject<T, d>& B, const TRigidTransform<T, d>& BToATM, const TVector<T, d>& LocalDir, const T Length, T& OutTime, TVector<T, d>& LocalPosition, TVector<T, d>& LocalNormal, int32& OutFaceIndex, T Thickness = 0) const
 	{
 		ensure(Length > 0);
 		ensure(FMath::IsNearlyEqual(LocalDir.SizeSquared(), 1, KINDA_SMALL_NUMBER));
-		ensure(Thickness == 0 || (FMath::IsNearlyEqual(OwningScaled.MScale[0], OwningScaled.MScale[1]) && FMath::IsNearlyEqual(OwningScaled.MScale[0], OwningScaled.MScale[2])));
+		ensure(Thickness == 0 || (FMath::IsNearlyEqual(MScale[0], MScale[1]) && FMath::IsNearlyEqual(MScale[0], MScale[2])));
 
-		const TVector<T, d> UnscaledDirDenorm = OwningScaled.MInvScale * LocalDir;
+		const TVector<T, d> UnscaledDirDenorm = MInvScale * LocalDir;
 		const T LengthScale = UnscaledDirDenorm.Size();
 		if (ensure(LengthScale > TNumericLimits<T>::Min()))
 		{
@@ -134,17 +305,17 @@ public:
 			TVector<T, d> UnscaledNormal;
 			float UnscaledTime;
 
-			TUniquePtr<TImplicitObject<T, d>> HackBPtr(const_cast<TImplicitObject<T,d>*>(&B));	//todo: hack, need scaled object to except raw ptr similar to transformed implicit
-			TImplicitObjectScaled<T, d> ScaledB(MakeSerializable(HackBPtr), OwningScaled.MInvScale);
+			TUniquePtr<TImplicitObject<T, d>> HackBPtr(const_cast<TImplicitObject<T,d>*>(&B));	//todo: hack, need scaled object to accept raw ptr similar to transformed implicit
+			TImplicitObjectScaled<TImplicitObject<T, d>> ScaledB(MakeSerializable(HackBPtr), MInvScale);
 			HackBPtr.Release();
 
-			TRigidTransform<T, d> BToATMNoScale(BToATM.GetLocation() * OwningScaled.MInvScale, BToATM.GetRotation());
+			TRigidTransform<T, d> BToATMNoScale(BToATM.GetLocation() * MInvScale, BToATM.GetRotation());
 			
-			if (InternalGeom.SweepGeom(ScaledB, BToATMNoScale, UnscaledDir, UnscaledLength, UnscaledTime, UnscaledPosition, UnscaledNormal, OutFaceIndex, OwningScaled.MInternalThickness + Thickness))
+			if (MObject->SweepGeom(ScaledB, BToATMNoScale, UnscaledDir, UnscaledLength, UnscaledTime, UnscaledPosition, UnscaledNormal, OutFaceIndex, MInternalThickness + Thickness))
 			{
 				OutTime = LengthScaleInv * UnscaledTime;
-				LocalPosition = OwningScaled.MScale * UnscaledPosition;
-				LocalNormal = (OwningScaled.MInvScale * UnscaledNormal).GetSafeNormal();
+				LocalPosition = MScale * UnscaledPosition;
+				LocalNormal = (MInvScale * UnscaledNormal).GetSafeNormal();
 				ensure(OutTime <= Length);
 				return true;
 			}
@@ -154,18 +325,16 @@ public:
 	}
 
 	/** This is a low level function and assumes the internal object has a OverlapGeom function. Should not be called directly. See GeometryQueries.h : OverlapQuery */
-	template <typename TGeometry>
-	static bool LowLevelOverlapGeom(const TImplicitObjectScaled<T, d>& OwningScaled, const TGeometry& InternalGeom, const TImplicitObject<T, d>& B, const TRigidTransform<T, d>& BToATM, T Thickness = 0)
+	bool LowLevelOverlapGeom(const TImplicitObject<T, d>& B, const TRigidTransform<T, d>& BToATM, T Thickness = 0) const
 	{
-		ensure(Thickness == 0 || (FMath::IsNearlyEqual(OwningScaled.MScale[0], OwningScaled.MScale[1]) && FMath::IsNearlyEqual(OwningScaled.MScale[0], OwningScaled.MScale[2])));
+		ensure(Thickness == 0 || (FMath::IsNearlyEqual(MScale[0], MScale[1]) && FMath::IsNearlyEqual(MScale[0], MScale[2])));
 
 		TUniquePtr<TImplicitObject<T, d>> HackBPtr(const_cast<TImplicitObject<T, d>*>(&B));	//todo: hack, need scaled object to except raw ptr similar to transformed implicit
-		TImplicitObjectScaled<T, d> ScaledB(MakeSerializable(HackBPtr), OwningScaled.MInvScale);
+		TImplicitObjectScaled<TImplicitObject<T, d>> ScaledB(MakeSerializable(HackBPtr), MInvScale);
 		HackBPtr.Release();
 
-		TRigidTransform<T, d> BToATMNoScale(BToATM.GetLocation() * OwningScaled.MInvScale, BToATM.GetRotation());
-
-		return InternalGeom.OverlapGeom(ScaledB, BToATMNoScale, OwningScaled.MInternalThickness + Thickness);
+		TRigidTransform<T, d> BToATMNoScale(BToATM.GetLocation() * MInvScale, BToATM.GetRotation());
+		return MObject->OverlapGeom(ScaledB, BToATMNoScale, MInternalThickness + Thickness);
 	}
 
 	virtual int32 FindMostOpposingFace(const TVector<T, d>& Position, const TVector<T, d>& UnitDir, int32 HintFaceIndex, T SearchDist) const override
@@ -277,11 +446,12 @@ public:
 		return HashCombine(MObject->GetTypeHash(), ::GetTypeHash(MScale));
 	}
 
+#if 0
 	virtual TUniquePtr<TImplicitObject<T, d>> Copy() const override
 	{
 		return TUniquePtr<TImplicitObject<T, d>>(CopyHelper(this));
 	}
-
+#endif
 private:
 	ObjectType MObject;
 	TVector<T, d> MScale;
@@ -291,22 +461,22 @@ private:
 
 	//needed for serialization
 	TImplicitObjectScaled()
-	: TImplicitObject<T, d>(EImplicitObject::HasBoundingBox, ImplicitObjectType::Scaled)
+	: TImplicitObject<T, d>(EImplicitObject::HasBoundingBox, StaticType())
 	, MInternalThickness(0)
 	{}
 	friend TImplicitObject<T, d>;	//needed for serialization
 
-	template <typename T2, int d2, bool bInstanced2>
+	template <typename TConcrete2, bool bInstanced2>
 	friend class TImplicitObjectScaled;
 
-	static TImplicitObjectScaled<T, d, true>* CopyHelper(const TImplicitObjectScaled<T, d, true>* Obj)
+	static TImplicitObjectScaled<TConcrete, true>* CopyHelper(const TImplicitObjectScaled<TConcrete, true>* Obj)
 	{
-		return new TImplicitObjectScaled<T, d, true>(Obj->MObject, Obj->MScale, Obj->MInternalThickness);
+		return new TImplicitObjectScaled<TConcrete, true>(Obj->MObject, Obj->MScale, Obj->MInternalThickness);
 	}
 
-	static TImplicitObjectScaled<T, d, false>* CopyHelper(const TImplicitObjectScaled<T, d, false>* Obj)
+	static TImplicitObjectScaled<TConcrete, false>* CopyHelper(const TImplicitObjectScaled<TConcrete, false>* Obj)
 	{
-		return new TImplicitObjectScaled<T, d, false>(Obj->MObject->Copy(), Obj->MScale, Obj->MInternalThickness);
+		return new TImplicitObjectScaled<TConcrete, false>(Obj->MObject->Copy(), Obj->MScale, Obj->MInternalThickness);
 	}
 
 	void UpdateBounds()
@@ -319,7 +489,10 @@ private:
 	}
 };
 
+template <typename TConcrete>
+using TImplicitObjectScaledNonSerializable = TImplicitObjectScaled<TConcrete, false>;
+
 template <typename T, int d>
-using TImplicitObjectScaledNonSerializable = TImplicitObjectScaled<T, d, false>;
+using TImplicitObjectScaledGeneric = TImplicitObjectScaled<TImplicitObject<T, d>>;
 
 }
diff --git a/Engine/Source/Runtime/Experimental/Chaos/Public/Chaos/ImplicitObjectTransformed.h b/Engine/Source/Runtime/Experimental/Chaos/Public/Chaos/ImplicitObjectTransformed.h
--- a/Engine/Source/Runtime/Experimental/Chaos/Public/Chaos/ImplicitObjectTransformed.h
+++ b/Engine/Source/Runtime/Experimental/Chaos/Public/Chaos/ImplicitObjectTransformed.h
@@ -86,7 +86,7 @@ public:
 	}
 	~TImplicitObjectTransformed() {}
 
-	static EImplicitObjectType GetType()
+	static EImplicitObjectType StaticType()
 	{
 		return ImplicitObjectType::Transformed;
 	}
diff --git a/Engine/Source/Runtime/Experimental/Chaos/Public/Chaos/ImplicitObjectUnion.h b/Engine/Source/Runtime/Experimental/Chaos/Public/Chaos/ImplicitObjectUnion.h
--- a/Engine/Source/Runtime/Experimental/Chaos/Public/Chaos/ImplicitObjectUnion.h
+++ b/Engine/Source/Runtime/Experimental/Chaos/Public/Chaos/ImplicitObjectUnion.h
@@ -69,7 +69,7 @@ class TImplicitObjectUnion : public TImplicitObject<T, d>
 	}
 	virtual ~TImplicitObjectUnion() {}
 
-	FORCEINLINE static EImplicitObjectType GetType()
+	FORCEINLINE static EImplicitObjectType StaticType()
 	{
 		return ImplicitObjectType::Union;
 	}
diff --git a/Engine/Source/Runtime/Experimental/Chaos/Public/Chaos/Levelset.h b/Engine/Source/Runtime/Experimental/Chaos/Public/Chaos/Levelset.h
--- a/Engine/Source/Runtime/Experimental/Chaos/Public/Chaos/Levelset.h
+++ b/Engine/Source/Runtime/Experimental/Chaos/Public/Chaos/Levelset.h
@@ -54,7 +54,7 @@ class CHAOS_API TLevelSet final : public TImplicitObject<T, d>
 		}
 	}
 
-	FORCEINLINE static EImplicitObjectType GetType()
+	FORCEINLINE static EImplicitObjectType StaticType()
 	{
 		return ImplicitObjectType::LevelSet;
 	}
diff --git a/Engine/Source/Runtime/Experimental/Chaos/Public/Chaos/Plane.h b/Engine/Source/Runtime/Experimental/Chaos/Public/Chaos/Plane.h
--- a/Engine/Source/Runtime/Experimental/Chaos/Public/Chaos/Plane.h
+++ b/Engine/Source/Runtime/Experimental/Chaos/Public/Chaos/Plane.h
@@ -34,7 +34,7 @@ class TPlane final : public TImplicitObject<T, d>
 	}
 	virtual ~TPlane() {}
 
-	static EImplicitObjectType GetType()
+	static EImplicitObjectType StaticType()
 	{
 		return ImplicitObjectType::Plane;
 	}
diff --git a/Engine/Source/Runtime/Experimental/Chaos/Public/Chaos/Sphere.h b/Engine/Source/Runtime/Experimental/Chaos/Public/Chaos/Sphere.h
--- a/Engine/Source/Runtime/Experimental/Chaos/Public/Chaos/Sphere.h
+++ b/Engine/Source/Runtime/Experimental/Chaos/Public/Chaos/Sphere.h
@@ -62,7 +62,7 @@ namespace Chaos
 
 		virtual ~TSphere() {}
 
-		static EImplicitObjectType GetType()
+		static EImplicitObjectType StaticType()
 		{ 
 			return ImplicitObjectType::Sphere; 
 		}
diff --git a/Engine/Source/Runtime/Experimental/Chaos/Public/Chaos/TaperedCylinder.h b/Engine/Source/Runtime/Experimental/Chaos/Public/Chaos/TaperedCylinder.h
--- a/Engine/Source/Runtime/Experimental/Chaos/Public/Chaos/TaperedCylinder.h
+++ b/Engine/Source/Runtime/Experimental/Chaos/Public/Chaos/TaperedCylinder.h
@@ -58,7 +58,7 @@ namespace Chaos
 		}
 		~TTaperedCylinder() {}
 
-		static EImplicitObjectType GetType() { return ImplicitObjectType::TaperedCylinder; }
+		static EImplicitObjectType StaticType() { return ImplicitObjectType::TaperedCylinder; }
 
 		/**
 		 * Returns sample points centered about the origin.
diff --git a/Engine/Source/Runtime/Experimental/Chaos/Public/Chaos/Triangle.h b/Engine/Source/Runtime/Experimental/Chaos/Public/Chaos/Triangle.h
--- a/Engine/Source/Runtime/Experimental/Chaos/Public/Chaos/Triangle.h
+++ b/Engine/Source/Runtime/Experimental/Chaos/Public/Chaos/Triangle.h
@@ -165,7 +165,7 @@ namespace Chaos
 			return Tri.GetPlane();
 		}
 
-		static EImplicitObjectType GetType()
+		static EImplicitObjectType StaticType()
 		{
 			return ImplicitObjectType::Triangle;
 		}
diff --git a/Engine/Source/Runtime/Experimental/Chaos/Public/Chaos/TriangleMeshImplicitObject.h b/Engine/Source/Runtime/Experimental/Chaos/Public/Chaos/TriangleMeshImplicitObject.h
--- a/Engine/Source/Runtime/Experimental/Chaos/Public/Chaos/TriangleMeshImplicitObject.h
+++ b/Engine/Source/Runtime/Experimental/Chaos/Public/Chaos/TriangleMeshImplicitObject.h
@@ -44,7 +44,7 @@ namespace Chaos
 			return MLocalBoundingBox;
 		}
 
-		static EImplicitObjectType GetType()
+		static EImplicitObjectType StaticType()
 		{
 			return ImplicitObjectType::TriangleMesh;
 		}
diff --git a/Engine/Source/Runtime/PhysicsCore/Private/ChaosInterfaceWrapperCore.cpp b/Engine/Source/Runtime/PhysicsCore/Private/ChaosInterfaceWrapperCore.cpp
--- a/Engine/Source/Runtime/PhysicsCore/Private/ChaosInterfaceWrapperCore.cpp
+++ b/Engine/Source/Runtime/PhysicsCore/Private/ChaosInterfaceWrapperCore.cpp
@@ -11,7 +11,7 @@ namespace ChaosInterface
 {
 	ECollisionShapeType GetImplicitType(const Chaos::TImplicitObject<float, 3>& InGeometry)
 	{
-		switch (InGeometry.GetType())
+		switch (Chaos::GetInnerType(InGeometry.GetType()))
 		{
 		case Chaos::ImplicitObjectType::Sphere: return ECollisionShapeType::Sphere;
 		case Chaos::ImplicitObjectType::Box: return ECollisionShapeType::Box;
@@ -19,7 +19,6 @@ namespace ChaosInterface
 		case Chaos::ImplicitObjectType::Convex: return ECollisionShapeType::Convex;
 		case Chaos::ImplicitObjectType::TriangleMesh: return ECollisionShapeType::Trimesh;
 		case Chaos::ImplicitObjectType::HeightField: return ECollisionShapeType::Heightfield;
-		case Chaos::ImplicitObjectType::Scaled: return ECollisionShapeType::Scaled;
 		default: break;
 		}
 
diff --git a/Engine/Source/Runtime/PhysicsCore/Public/PhysXToChaosUtil.h b/Engine/Source/Runtime/PhysicsCore/Public/PhysXToChaosUtil.h
--- a/Engine/Source/Runtime/PhysicsCore/Public/PhysXToChaosUtil.h
+++ b/Engine/Source/Runtime/PhysicsCore/Public/PhysXToChaosUtil.h
@@ -71,7 +71,7 @@ inline TUniquePtr<Chaos::TImplicitObjectTransformed<float, 3>> PxShapeToChaosGeo
 		else
 		{
 			ensure(ConvexMeshGeom.scale.rotation == PxQuat(PxIdentity));
-			InnerObj = MakeUnique<TImplicitObjectScaled<float, 3, /*bInstanced=*/false>>(MoveTemp(ConvexObj), P2UVector(ConvexMeshGeom.scale.scale));	//todo: make instanced
+			InnerObj = MakeUnique<TImplicitObjectScaled<TConvex<float, 3>, /*bInstanced=*/false>>(MoveTemp(ConvexObj), P2UVector(ConvexMeshGeom.scale.scale));	//todo: make instanced
 		}
 		break;
 	}
@@ -167,7 +167,7 @@ inline TUniquePtr<Chaos::TImplicitObjectTransformed<float, 3>> PxShapeToChaosGeo
 		else
 		{
 			ensure(TriMeshGeom.scale.rotation == PxQuat(PxIdentity));
-			InnerObj = MakeUnique<TImplicitObjectScaled<float, 3, /*bInstanced=*/false>>(MoveTemp(TriMeshObj), P2UVector(TriMeshGeom.scale.scale));	//todo: make instanced
+			InnerObj = MakeUnique<TImplicitObjectScaled<TTriangleMeshImplicitObject<float>, /*bInstanced=*/false>>(MoveTemp(TriMeshObj), P2UVector(TriMeshGeom.scale.scale));	//todo: make instanced
 		}
 		break;
 	}
diff --git a/Engine/Source/Runtime/PhysicsCore/Public/PhysicsInterfaceTypesCore.h b/Engine/Source/Runtime/PhysicsCore/Public/PhysicsInterfaceTypesCore.h
--- a/Engine/Source/Runtime/PhysicsCore/Public/PhysicsInterfaceTypesCore.h
+++ b/Engine/Source/Runtime/PhysicsCore/Public/PhysicsInterfaceTypesCore.h
@@ -49,7 +49,6 @@ enum class ECollisionShapeType : uint8
 	Convex,
 	Trimesh,
 	Heightfield,
-	Scaled,
 	None
 };
 
