Speculative CIS Fix for PBDConstraintGraph on Mac and Linux

diff --git a/Engine/Source/Runtime/Experimental/Chaos/Public/Chaos/PBDConstraintGraph.h b/Engine/Source/Runtime/Experimental/Chaos/Public/Chaos/PBDConstraintGraph.h
--- a/Engine/Source/Runtime/Experimental/Chaos/Public/Chaos/PBDConstraintGraph.h
+++ b/Engine/Source/Runtime/Experimental/Chaos/Public/Chaos/PBDConstraintGraph.h
@@ -234,3 +234,8 @@ namespace Chaos
 	};
 
 }
+
+#if PLATFORM_MAC || PLATFORM_LINUX
+extern template class CHAOS_API Chaos::TPBDConstraintGraph<float, 3>;
+extern template class CHAOS_API Chaos::TPBDConstraintGraph<float, 3>;
+#endif
