Perf improvement regarding timeslicing acceleration structure

diff --git a/Engine/Source/Runtime/Experimental/Chaos/Private/Chaos/PBDRigidsEvolution.cpp b/Engine/Source/Runtime/Experimental/Chaos/Private/Chaos/PBDRigidsEvolution.cpp
--- a/Engine/Source/Runtime/Experimental/Chaos/Private/Chaos/PBDRigidsEvolution.cpp
+++ b/Engine/Source/Runtime/Experimental/Chaos/Private/Chaos/PBDRigidsEvolution.cpp
@@ -153,6 +153,9 @@ namespace Chaos
 	DECLARE_CYCLE_STAT(TEXT("ComputeIntermediateSpatialAcceleration"), STAT_ComputeIntermediateSpatialAcceleration, STATGROUP_Chaos);
 	DECLARE_CYCLE_STAT(TEXT("CopyAccelerationStructure"), STAT_CopyAccelerationStructure, STATGROUP_Chaos);
 	DECLARE_CYCLE_STAT(TEXT("SwapAccelerationStructures"), STAT_SwapAccelerationStructures, STATGROUP_Chaos);
+	DECLARE_CYCLE_STAT(TEXT("AccelerationStructureTimeSlice"), STAT_AccelerationStructureTimeSlice, STATGROUP_Chaos);
+	DECLARE_CYCLE_STAT(TEXT("CreateInitialAccelerationStructure"), STAT_CreateInitialAccelerationStructure, STATGROUP_Chaos);
+
 
 	template<class FPBDRigidsEvolution, class FPBDCollisionConstraint, class T, int d>
 	TPBDRigidsEvolutionBase<FPBDRigidsEvolution, FPBDCollisionConstraint, T, d>::FChaosAccelerationStructureTask::FChaosAccelerationStructureTask(
@@ -228,6 +231,8 @@ namespace Chaos
 			const uint8 BucketIdx = (1 << SpatialIdx.Bucket) & ActiveBucketsMask ? SpatialIdx.Bucket : 0;
 			if (AccelerationStructure->GetSubstructure(SpatialIdx) && !AccelerationStructure->GetSubstructure(SpatialIdx)->IsAsyncTimeSlicingComplete())
 			{
+				SCOPE_CYCLE_COUNTER(STAT_AccelerationStructureTimeSlice);
+
 				AccelerationStructure->GetSubstructure(SpatialIdx)->ProgressAsyncTimeSlicing();
 
 				// is it still progressing or now complete
@@ -248,6 +253,8 @@ namespace Chaos
 		{
 			if (ViewsPerBucket[BucketIdx].Num())
 			{
+				SCOPE_CYCLE_COUNTER(STAT_CreateInitialAccelerationStructure);
+
 				auto ParticleView = MakeConstParticleView(MoveTemp(ViewsPerBucket[BucketIdx]));
 				auto NewStruct = SpatialCollectionFactory.CreateAccelerationPerBucket_Threaded(ParticleView, BucketIdx);
 
diff --git a/Engine/Source/Runtime/Experimental/Chaos/Public/Chaos/AABBTree.h b/Engine/Source/Runtime/Experimental/Chaos/Public/Chaos/AABBTree.h
--- a/Engine/Source/Runtime/Experimental/Chaos/Public/Chaos/AABBTree.h
+++ b/Engine/Source/Runtime/Experimental/Chaos/Public/Chaos/AABBTree.h
@@ -19,6 +19,10 @@ enum class EAABBQueryType
 	Overlap
 };
 
+DECLARE_CYCLE_STAT(TEXT("AABBTreeGenerateTree"), STAT_AABBTreeGenerateTree, STATGROUP_Chaos);
+DECLARE_CYCLE_STAT(TEXT("AABBTreeTimeSliceSetup"), STAT_AABBTreeTimeSliceSetup, STATGROUP_Chaos);
+DECLARE_CYCLE_STAT(TEXT("AABBTreeInitialTimeSlice"), STAT_AABBTreeInitialTimeSlice, STATGROUP_Chaos);
+
 template <typename T, EAABBQueryType Query>
 struct TAABBTreeIntersectionHelper
 {
@@ -222,7 +226,7 @@ public:
 	static constexpr T DefaultMaxPayloadBounds = 100000;
 	static constexpr int32 DefaultMaxChildrenInLeaf = 12;
 	static constexpr int32 DefaultMaxTreeDepth = 16;
-	static constexpr int32 DefaultMaxNumToProcess = 2000;
+	static constexpr int32 DefaultMaxNumToProcess = 4000;
 	static constexpr ESpatialAcceleration StaticType = TIsSame<TAABBTreeLeafArray<TPayloadType, T>, TLeafType>::Value ? ESpatialAcceleration::AABBTree : 
 		(TIsSame<TBoundingVolume<TPayloadType, T, 3>, TLeafType>::Value ? ESpatialAcceleration::AABBTreeBV : ESpatialAcceleration::Unknown);
 	TAABBTree()
@@ -627,8 +631,9 @@ private:
 	template <typename TParticles>
 	void GenerateTree(const TParticles& Particles)
 	{
-		TArray<FElement> ElemsWithBounds;
+		SCOPE_CYCLE_COUNTER(STAT_AABBTreeGenerateTree);
 
+		TArray<FElement> ElemsWithBounds;
 		ElemsWithBounds.Reserve(Particles.Num());
 
 		GlobalPayloads.Reset();
@@ -639,58 +644,78 @@ private:
 		NumProcessedThisSlice = 0;
 		check(MaxNumToProcess > 0);
 
-		int32 Idx = 0;
 		FullBounds = TBox<T, 3>::EmptyBox();
-		for (auto& Particle : Particles)
+
 		{
-			bool bHasBoundingBox = HasBoundingBox(Particle);
-			auto Payload = Particle.template GetPayload<TPayloadType>(Idx);
-			TBox<T, 3> ElemBounds = ComputeWorldSpaceBoundingBox(Particle, false, (T)0);
+			SCOPE_CYCLE_COUNTER(STAT_AABBTreeTimeSliceSetup);
+
+			int32 Idx = 0;
 
-			if (bHasBoundingBox)
+			for (auto& Particle : Particles)
 			{
-				if (ElemBounds.Extents().Max() > MaxPayloadBounds)
+				bool bHasBoundingBox = HasBoundingBox(Particle);
+				auto Payload = Particle.template GetPayload<TPayloadType>(Idx);
+				TBox<T, 3> ElemBounds = ComputeWorldSpaceBoundingBox(Particle, false, (T)0);
+
+				if (bHasBoundingBox)
 				{
-					bHasBoundingBox = false;
+					if (ElemBounds.Extents().Max() > MaxPayloadBounds)
+					{
+						bHasBoundingBox = false;
+					}
+					else
+					{
+						ElemsWithBounds.Add(FElement{ Payload, ElemBounds });
+						FullBounds.GrowToInclude(ElemBounds);
+					}
 				}
 				else
 				{
-					ElemsWithBounds.Add(FElement{ Payload, ElemBounds });
-					FullBounds.GrowToInclude(ElemBounds);
+					ElemBounds = TBox<T, 3>(TVector<T, 3>(TNumericLimits<T>::Lowest()), TVector<T, 3>(TNumericLimits<T>::Max()));
 				}
-			}
-			else
-			{
-				ElemBounds = TBox<T,3>(TVector<T, 3>(TNumericLimits<T>::Lowest()), TVector<T, 3>(TNumericLimits<T>::Max()));
-			}
 
-			if (!bHasBoundingBox)
-			{
-				PayloadToInfo.Add(Payload, FAABBTreePayloadInfo{ GlobalPayloads.Num(), INDEX_NONE, INDEX_NONE });
-				GlobalPayloads.Add(FElement{ Payload, ElemBounds });
+				if (!bHasBoundingBox)
+				{
+					PayloadToInfo.Add(Payload, FAABBTreePayloadInfo{ GlobalPayloads.Num(), INDEX_NONE, INDEX_NONE });
+					GlobalPayloads.Add(FElement{ Payload, ElemBounds });
+				}
+
+				++Idx;
+				//todo: payload info
 			}
 
-			++Idx;
-			//todo: payload info
 		}
 
-		SplitNode(FullBounds, ElemsWithBounds, 0, 0);
-
-		if (TimeSliceWorkToComplete.IsEmpty())
 		{
-			this->SetAsyncTimeSlicingComplete(true);
+			SCOPE_CYCLE_COUNTER(STAT_AABBTreeInitialTimeSlice);
+
+			SplitNode(FullBounds, ElemsWithBounds, 0, 0, true);
+
+			if (TimeSliceWorkToComplete.IsEmpty())
+			{
+				this->SetAsyncTimeSlicingComplete(true);
+			}
 		}
+
 	}
-	
-	void SplitNode(const TBox<T, 3>& Bounds, const TArray<FElement>& Elems, int32 NodeLevel, int32 NewNodeIdx)
+
+
+	void SplitNode(const TBox<T, 3>& Bounds, const TArray<FElement>& Elems, int32 NodeLevel, int32 NewNodeIdx, bool Initial=false)
 	{
+		const int32 InitialMaxNumToProcess = 100;
+		int32 NumToProcess = MaxNumToProcess;
+		if (Initial)
+		{
+			NumToProcess = InitialMaxNumToProcess;
+		}
+
 		// create the actual node space but might no be filled in (YET) due to time slicing exit
 		if (NewNodeIdx >= Nodes.Num())
 		{
 		 	Nodes.AddDefaulted();	//todo: remove TBox
 		}
 
-		if (NumProcessedThisSlice >= MaxNumToProcess)
+		if (NumProcessedThisSlice >= NumToProcess)
 		{
 			// done enough work, capture stack
 			FWorkSnapshot Store;
@@ -704,8 +729,6 @@ private:
 			return; // done enough
 		}
 
-		NumProcessedThisSlice++;
-
 		auto& PayloadToInfoRef = PayloadToInfo;
 		auto& LeavesRef = Leaves;
 		auto& NodesRef = Nodes;
@@ -791,9 +814,11 @@ private:
 		{
 			Nodes[NewNodeIdx].bLeaf = false;
 
+			// increment this early so the amount of recursion work left on the stack is smaller
+			NumProcessedThisSlice += SplitInfos[0].Children.Num() + SplitInfos[1].Children.Num();
+
 			for (int32 BoxIdx = 0; BoxIdx < 2; ++BoxIdx)
 			{
-			
 					const int32 ChildIdx = Nodes.Num();
 					Nodes[NewNodeIdx].ChildrenBounds[BoxIdx] = SplitInfos[BoxIdx].RealBounds;
 					Nodes[NewNodeIdx].ChildrenNodes[BoxIdx] = ChildIdx;
