Change GJK.h default switch case from check(false) to ensure(false).

diff --git a/Engine/Source/Runtime/Experimental/Chaos/Public/Chaos/GJK.h b/Engine/Source/Runtime/Experimental/Chaos/Public/Chaos/GJK.h
--- a/Engine/Source/Runtime/Experimental/Chaos/Public/Chaos/GJK.h
+++ b/Engine/Source/Runtime/Experimental/Chaos/Public/Chaos/GJK.h
@@ -336,7 +336,7 @@ namespace Chaos
 			break;
 		}
 		default:
-			check(false);
+			ensure(false);
 			ClosestPoint = TVector<T, 3>(0);
 		}
 
