Added a new Chaos::TConvex constructor from planes and surface points.

diff --git a/Engine/Source/Runtime/Experimental/Chaos/Public/Chaos/Convex.h b/Engine/Source/Runtime/Experimental/Chaos/Public/Chaos/Convex.h
--- a/Engine/Source/Runtime/Experimental/Chaos/Public/Chaos/Convex.h
+++ b/Engine/Source/Runtime/Experimental/Chaos/Public/Chaos/Convex.h
@@ -25,6 +25,17 @@ namespace Chaos
 		    , SurfaceParticles(MoveTemp(Other.SurfaceParticles))
 		    , LocalBoundingBox(MoveTemp(Other.LocalBoundingBox))
 		{}
+		TConvex(TArray<TPlane<T, d>>&& InPlanes, TParticles<T, d>&& InSurfaceParticles)
+		    : TImplicitObject<T, d>(EImplicitObject::IsConvex | EImplicitObject::HasBoundingBox, ImplicitObjectType::Convex)
+			, Planes(MoveTemp(InPlanes))
+		    , SurfaceParticles(MoveTemp(InSurfaceParticles))
+		    , LocalBoundingBox(TBox<T, d>::EmptyBox())
+		{
+			for (uint32 ParticleIndex = 0; ParticleIndex < SurfaceParticles.Size(); ++ParticleIndex)
+			{
+				LocalBoundingBox.GrowToInclude(SurfaceParticles.X(ParticleIndex));
+			}
+		}
 		TConvex(const TParticles<T, 3>& InParticles)
 		    : TImplicitObject<T, d>(EImplicitObject::IsConvex | EImplicitObject::HasBoundingBox, ImplicitObjectType::Convex)
 		{
