Copy fixes for several major UI bugs in Modeling Tools from Dev-Editor
CL 9875272: Handle case where user disables Gizmo in DrawPolygon, resulting in it being nullptr in ctrl+click handler. Also update Gizmo position after it is re-shown.
CL 9888557: resolve flickering/instability with Gizmo constant-size calculation and axis filtering
CL 9888589: snap test in PolyEdit needs to be performed on world-space points, not local-space
#rnx

diff --git a/Engine/Plugins/Experimental/MeshModelingToolset/Source/MeshModelingTools/Private/DrawPolygonTool.cpp b/Engine/Plugins/Experimental/MeshModelingToolset/Source/MeshModelingTools/Private/DrawPolygonTool.cpp
--- a/Engine/Plugins/Experimental/MeshModelingToolset/Source/MeshModelingTools/Private/DrawPolygonTool.cpp
+++ b/Engine/Plugins/Experimental/MeshModelingToolset/Source/MeshModelingTools/Private/DrawPolygonTool.cpp
@@ -930,7 +930,10 @@ void UDrawPolygonTool::SetDrawPlaneFromWorldPos(const FVector& Position, const F
 
 	SnapEngine.Plane = FFrame3d((FVector3d)DrawPlane.Origin, (FQuaterniond)DrawPlane.Rotation);
 
-	PlaneTransformGizmo->SetNewGizmoTransform(FTransform(DrawPlaneOrientation, DrawPlaneOrigin));
+	if (PlaneTransformGizmo != nullptr)
+	{
+		PlaneTransformGizmo->SetNewGizmoTransform(FTransform(DrawPlaneOrientation, DrawPlaneOrigin));
+	}
 }
 
 
@@ -946,11 +949,13 @@ void UDrawPolygonTool::UpdateShowGizmoState(bool bNewVisibility)
 	if (bNewVisibility == false)
 	{
 		GetToolManager()->GetPairedGizmoManager()->DestroyAllGizmosByOwner(this);
+		PlaneTransformGizmo = nullptr;
 	}
 	else
 	{
 		PlaneTransformGizmo = GetToolManager()->GetPairedGizmoManager()->Create3AxisTransformGizmo(this);
 		PlaneTransformGizmo->SetActiveTarget(PlaneTransformProxy, GetToolManager());
+		PlaneTransformGizmo->SetNewGizmoTransform(FTransform(DrawPlaneOrientation, DrawPlaneOrigin));
 	}
 }
 
diff --git a/Engine/Plugins/Experimental/MeshModelingToolset/Source/MeshModelingTools/Private/EditMeshPolygonsTool.cpp b/Engine/Plugins/Experimental/MeshModelingToolset/Source/MeshModelingTools/Private/EditMeshPolygonsTool.cpp
--- a/Engine/Plugins/Experimental/MeshModelingToolset/Source/MeshModelingTools/Private/EditMeshPolygonsTool.cpp
+++ b/Engine/Plugins/Experimental/MeshModelingToolset/Source/MeshModelingTools/Private/EditMeshPolygonsTool.cpp
@@ -657,7 +657,9 @@ void UEditMeshPolygonsTool::Setup()
 	TopoSelector.Initialize(DynamicMeshComponent->GetMesh(), &Topology);
 	TopoSelector.SetSpatialSource([this]() {return &GetSpatial(); });
 	TopoSelector.PointsWithinToleranceTest = [this](const FVector3d& Position1, const FVector3d& Position2) {
-		return ToolSceneQueriesUtil::PointSnapQuery(this->CameraState, Position1, Position2, VisualAngleSnapThreshold);
+		FTransform Transform = ComponentTarget->GetWorldTransform();
+		return ToolSceneQueriesUtil::PointSnapQuery(this->CameraState, 
+			Transform.TransformPosition(Position1), Transform.TransformPosition(Position2), VisualAngleSnapThreshold);
 	};
 
 	// hide input StaticMeshComponent
diff --git a/Engine/Plugins/Experimental/MeshModelingToolset/Source/ModelingComponents/Private/Transforms/QuickAxisTranslater.cpp b/Engine/Plugins/Experimental/MeshModelingToolset/Source/ModelingComponents/Private/Transforms/QuickAxisTranslater.cpp
--- a/Engine/Plugins/Experimental/MeshModelingToolset/Source/ModelingComponents/Private/Transforms/QuickAxisTranslater.cpp
+++ b/Engine/Plugins/Experimental/MeshModelingToolset/Source/ModelingComponents/Private/Transforms/QuickAxisTranslater.cpp
@@ -90,11 +90,12 @@ bool ProjectToCameraPlanePos(const FVector3d& ScenePos, const FFrame3d& CameraPl
 	return true;
 }
 
+
 void FQuickAxisTranslater::UpdateSnapAxes()
 {
 	// cos(angle) tolerances 
-	static const double ForwardDotTol = FMathd::Cos(20.0 * FMathd::DegToRad);
-	static const double OverlapDotTol = FMathd::Cos(20.0 * FMathd::DegToRad);
+	static const double ForwardDotTol = FMathd::Cos(15.0 * FMathd::DegToRad);
+	static const double OverlapDotTol = FMathd::Cos(15.0 * FMathd::DegToRad);
 
 	MoveAxisSolver.Reset();
 
@@ -147,9 +148,9 @@ void FQuickAxisTranslater::UpdateSnapAxes()
 
 	// if angle between these 2D projections of axes is too small, then axes are visually overlapping
 	// from this view position, and snapping will be unstable. So only keep the one pointing "towards" camera.
-	double AngleXYDot = FMathd::Abs(PlaneAxisX.Dot(PlaneAxisY));
-	double AngleXZDot = FMathd::Abs(PlaneAxisX.Dot(PlaneAxisZ));
-	double AngleYZDot = FMathd::Abs(PlaneAxisY.Dot(PlaneAxisZ));
+	double AngleXYDot = (IgnoreAxis[0] || IgnoreAxis[1]) ? 0.0 : FMathd::Abs(PlaneAxisX.Dot(PlaneAxisY));
+	double AngleXZDot = (IgnoreAxis[0] || IgnoreAxis[2]) ? 0.0 : FMathd::Abs(PlaneAxisX.Dot(PlaneAxisZ));
+	double AngleYZDot = (IgnoreAxis[1] || IgnoreAxis[2]) ? 0.0 : FMathd::Abs(PlaneAxisY.Dot(PlaneAxisZ));
 	if (AngleXYDot > OverlapDotTol)
 	{
 		int which = (FMathd::Sign(DirectionDotX) < FMathd::Sign(DirectionDotY)) ? 0 : 1;
diff --git a/Engine/Source/Runtime/Experimental/InteractiveToolsFramework/InteractiveToolsFramework.Build.cs b/Engine/Source/Runtime/Experimental/InteractiveToolsFramework/InteractiveToolsFramework.Build.cs
--- a/Engine/Source/Runtime/Experimental/InteractiveToolsFramework/InteractiveToolsFramework.Build.cs
+++ b/Engine/Source/Runtime/Experimental/InteractiveToolsFramework/InteractiveToolsFramework.Build.cs
@@ -37,7 +37,8 @@ public class InteractiveToolsFramework : ModuleRules
 		PrivateDependencyModuleNames.AddRange(
 			new string[]
 			{
-				"Engine"
+				"Engine",
+                "RHI",
 				//"Slate",
 				//"SlateCore",
 				// ... add private dependencies that you statically link with here ...	
diff --git a/Engine/Source/Runtime/Experimental/InteractiveToolsFramework/Private/BaseGizmos/GizmoRenderingUtil.cpp b/Engine/Source/Runtime/Experimental/InteractiveToolsFramework/Private/BaseGizmos/GizmoRenderingUtil.cpp
--- a/Engine/Source/Runtime/Experimental/InteractiveToolsFramework/Private/BaseGizmos/GizmoRenderingUtil.cpp
+++ b/Engine/Source/Runtime/Experimental/InteractiveToolsFramework/Private/BaseGizmos/GizmoRenderingUtil.cpp
@@ -1,6 +1,7 @@
 // Copyright 1998-2019 Epic Games, Inc. All Rights Reserved.
 
 #include "BaseGizmos/GizmoRenderingUtil.h"
+#include "RHI.h"
 
 
 const FSceneView* GizmoRenderingUtil::FindActiveSceneView(
@@ -44,20 +45,59 @@ const FSceneView* GizmoRenderingUtil::FindActiveSceneView(
 }
 
 
+static double VectorDifferenceSqr(const FVector2D& A, const FVector2D& B)
+{
+	double ax = A.X, ay = A.Y;
+	double bx = B.X, by = B.Y;
+	ax -= bx;
+	ay -= by;
+	return ax*ax + ay*ay;
+}
+
+static double VectorDifferenceSqr(const FVector& A, const FVector& B)
+{
+	double ax = A.X, ay = A.Y, az = A.Z;
+	double bx = B.X, by = B.Y, bz = B.Z;
+	ax -= bx; 
+	ay -= by;
+	az -= bz;
+	return ax*ax + ay*ay + az*az;
+}
+
+// duplicates FSceneView::WorldToPixel but in double where possible (unfortunately WorldToScreen still in float)
+static FVector2D WorldToPixelDouble(const FSceneView* View, const FVector& Location)
+{
+	FVector4 ScreenPoint = View->WorldToScreen(Location);
+
+	double InvW = (ScreenPoint.W > 0.0 ? 1.0 : -1.0) / (double)ScreenPoint.W;
+	double Y = (GProjectionSignY > 0.0) ? (double)ScreenPoint.Y : 1.0 - (double)ScreenPoint.Y;
+
+	const FIntRect& UnscaledViewRect = View->UnscaledViewRect;
+	double PosX = (double)UnscaledViewRect.Min.X + (0.5 + (double)ScreenPoint.X * 0.5 * InvW) * (double)UnscaledViewRect.Width();
+	double PosY = (double)UnscaledViewRect.Min.Y + (0.5 - Y * 0.5 * InvW) * (double)UnscaledViewRect.Height();
+
+	return FVector2D((float)PosX, (float)PosY);
+}
+
 
 float GizmoRenderingUtil::CalculateLocalPixelToWorldScale(
 	const FSceneView* View,
 	const FVector& Location)
 {
-	FVector2D PixelA, PixelB;
-	View->WorldToPixel(Location, PixelA);
-	FVector OffsetPointWorld = Location + 0.01f*View->GetViewRight() + 0.01f*View->GetViewUp();
-	View->WorldToPixel(OffsetPointWorld, PixelB);
-
-	float PixelDelta = (PixelA - PixelB).Size();
-	float WorldDelta = (Location - OffsetPointWorld).Size();
+	// To calculate this scale at Location, we project Location to screen and also project a second
+	// point at a small distance from Location in a camera-perpendicular plane, then measure 2D/3D distance ratio. 
+	// However, because some of the computations are done in float, there will be enormous numerical error 
+	// when the camera is very far from the location if the distance is relatively small. The "W" value
+	// below gives us a sense of this distance, so we make the offset relative to that
+	// (this does do one redundant WorldToScreen)
+	FVector4 LocationScreenPoint = View->WorldToScreen(Location);
+	float OffsetDelta = LocationScreenPoint.W * 0.01f;
 
-	float PixelToWorldScale = WorldDelta / PixelDelta;
+	FVector2D PixelA = WorldToPixelDouble(View, Location);
+	FVector OffsetPointWorld = Location + OffsetDelta*View->GetViewRight() + OffsetDelta*View->GetViewUp();
+	FVector2D PixelB = WorldToPixelDouble(View, OffsetPointWorld);
 
-	return PixelToWorldScale;
+	double PixelDeltaSqr = VectorDifferenceSqr(PixelA, PixelB);
+	double WorldDeltaSqr = VectorDifferenceSqr(Location, OffsetPointWorld);
+	return (float)(sqrt(WorldDeltaSqr / PixelDeltaSqr));
 }
\ No newline at end of file
