Chaos :  Implicit

Fix for CIS ( Debug)

diff --git a/Engine/Source/Runtime/Experimental/Chaos/Private/Chaos/ImplicitObject.cpp b/Engine/Source/Runtime/Experimental/Chaos/Private/Chaos/ImplicitObject.cpp
--- a/Engine/Source/Runtime/Experimental/Chaos/Private/Chaos/ImplicitObject.cpp
+++ b/Engine/Source/Runtime/Experimental/Chaos/Private/Chaos/ImplicitObject.cpp
@@ -162,6 +162,12 @@ Pair<TVector<T, d>, bool> TImplicitObject<T, d>::FindDeepestIntersection(const T
 	}
 	return MakePair(Point, Phi < Thickness);
 }
+template<class T, int d>
+bool TImplicitObject<T, d>::Overlap(const TVector<T, d>& Point, const T Thickness) const
+{
+	return SignedDistance(Point) <= Thickness;
+}
+
 
 template<class T, int d>
 Pair<TVector<T, d>, bool> TImplicitObject<T, d>::FindClosestIntersection(const TVector<T, d>& StartPoint, const TVector<T, d>& EndPoint, const T Thickness) const
diff --git a/Engine/Source/Runtime/Experimental/Chaos/Public/Chaos/ImplicitObject.h b/Engine/Source/Runtime/Experimental/Chaos/Public/Chaos/ImplicitObject.h
--- a/Engine/Source/Runtime/Experimental/Chaos/Public/Chaos/ImplicitObject.h
+++ b/Engine/Source/Runtime/Experimental/Chaos/Public/Chaos/ImplicitObject.h
@@ -236,10 +236,7 @@ public:
 	}
 
 	//This gives derived types a way to do an overlap check without calling PhiWithNormal todo: this api is confusing
-	virtual bool Overlap(const TVector<T, d>& Point, const T Thickness) const
-	{
-		return SignedDistance(Point) <= Thickness;
-	}
+	virtual bool Overlap(const TVector<T, d>& Point, const T Thickness) const; 
 
 	virtual void AccumulateAllImplicitObjects(TArray<Pair<const TImplicitObject<T, d>*, TRigidTransform<T, d>>>& Out, const TRigidTransform<T, d>& ParentTM) const
 	{
