Chaos 6dof joint
- Soft constraint limits
- Constraint rule sorting

diff --git a/Engine/Source/Runtime/Engine/Classes/PhysicsEngine/ConstraintInstance.h b/Engine/Source/Runtime/Engine/Classes/PhysicsEngine/ConstraintInstance.h
--- a/Engine/Source/Runtime/Engine/Classes/PhysicsEngine/ConstraintInstance.h
+++ b/Engine/Source/Runtime/Engine/Classes/PhysicsEngine/ConstraintInstance.h
@@ -358,6 +358,60 @@ public:
 		UpdateAngularLimit();
 	}
 
+	/** Whether the linear limits are soft (only if at least one axis if Limited) */
+	bool GetIsSoftLinearLimit() const
+	{
+		return ProfileInstance.LinearLimit.bSoftConstraint;
+	}
+
+	/** Linear stiffness if the constraint is set to use soft linear limits */
+	float GetSoftLinearLimitStiffness() const
+	{
+		return ProfileInstance.LinearLimit.Stiffness;
+	}
+
+	/** Linear damping if the constraint is set to use soft linear limits */
+	float GetSoftLinearLimitDamping() const
+	{
+		return ProfileInstance.LinearLimit.Damping;
+	}
+
+	/** Whether the twist limits are soft (only available if twist is Limited) */
+	bool GetIsSoftTwistLimit() const
+	{
+		return ProfileInstance.TwistLimit.bSoftConstraint;
+	}
+
+	/** Twist stiffness if the constraint is set to use soft limits */
+	float GetSoftTwistLimitStiffness() const
+	{
+		return ProfileInstance.TwistLimit.Stiffness;
+	}
+
+	/** Twist damping if the constraint is set to use soft limits */
+	float GetSoftTwistLimitDamping() const
+	{
+		return ProfileInstance.TwistLimit.Damping;
+	}
+
+	/** Whether the swing limits are soft (only available if swing1 and/or swing2 is Limited) */
+	bool GetIsSoftSwingLimit() const
+	{
+		return ProfileInstance.ConeLimit.bSoftConstraint;
+	}
+
+	/** Swing stiffness if the constraint is set to use soft limits */
+	float GetSoftSwingLimitStiffness() const
+	{
+		return ProfileInstance.ConeLimit.Stiffness;
+	}
+
+	/** Swing damping if the constraint is set to use soft limits */
+	float GetSoftSwingLimitDamping() const
+	{
+		return ProfileInstance.ConeLimit.Damping;
+	}
+
 	/** Sets the Linear Breakable properties
 	*	@param bInLinearBreakable		Whether it is possible to break the joint with linear force
 	*	@param InLinearBreakThreshold	Force needed to break the joint
diff --git a/Engine/Source/Runtime/Engine/Private/PhysicsEngine/ImmediatePhysics/ImmediatePhysicsChaos/ImmediatePhysicsJointHandle_Chaos.cpp b/Engine/Source/Runtime/Engine/Private/PhysicsEngine/ImmediatePhysics/ImmediatePhysicsChaos/ImmediatePhysicsJointHandle_Chaos.cpp
--- a/Engine/Source/Runtime/Engine/Private/PhysicsEngine/ImmediatePhysics/ImmediatePhysicsChaos/ImmediatePhysicsJointHandle_Chaos.cpp
+++ b/Engine/Source/Runtime/Engine/Private/PhysicsEngine/ImmediatePhysics/ImmediatePhysicsChaos/ImmediatePhysicsJointHandle_Chaos.cpp
@@ -13,36 +13,72 @@ static_assert((int32)Chaos::EJointMotionType::Locked == (int32)EAngularConstrain
 
 namespace ImmediatePhysics_Chaos
 {
+	const float ChaosImmediate_StiffnessDt = 0.03f;
+
 	float ChaosImmediate_JointStiffness = 1.0f;
 	FAutoConsoleVariableRef CVarJointStiffness(TEXT("p.Chaos.ImmPhys.JointStiffness"), ChaosImmediate_JointStiffness, TEXT("Joint solver stiffness."));
 
-	float ChaosImmediate_MaxDriveStiffness = 2000.0f;
-	FAutoConsoleVariableRef CVarMaxDriveStiffness(TEXT("p.Chaos.ImmPhys.MaxDriveStiffness"), ChaosImmediate_MaxDriveStiffness, TEXT("The value of drive stiffness per unit mass that equates to full stiffness in the solver."));
-
-	int32 ChaosImmediate_ScaleDriveStiffnessByMass = 0;
-	FAutoConsoleVariableRef CVarScaleDriveStiffnessByMass(TEXT("p.Chaos.ImmPhys.ScaleDriveStiffnessByMass"), ChaosImmediate_ScaleDriveStiffnessByMass, TEXT("If true, converted stiffness is multiplied by inertia."));
+	bool ChaosImmediate_StiffnessUseMass = true;
+	FAutoConsoleVariableRef CVarStiffnessUseMass(TEXT("p.Chaos.ImmPhys.StiffnessUseMass"), ChaosImmediate_StiffnessUseMass, TEXT("Whether to use mass to scale stiffness in the conversion."));
+
+	float ChaosImmediate_DriveStiffnessSourceMin = 100.0f;
+	float ChaosImmediate_DriveStiffnessSourceMax = 2000.0f;
+	float ChaosImmediate_DriveStiffnessTargetMin = 0.1f;
+	float ChaosImmediate_DriveStiffnessTargetMax = 1.0f;
+	FAutoConsoleVariableRef CVarDriveStiffnessSourceMin(TEXT("p.Chaos.ImmPhys.DriveStiffnessSourceMin"), ChaosImmediate_DriveStiffnessSourceMin, TEXT("Conversion factor for drive stiffness."));
+	FAutoConsoleVariableRef CVarDriveStiffnessSourceMax(TEXT("p.Chaos.ImmPhys.DriveStiffnessSourceMax"), ChaosImmediate_DriveStiffnessSourceMax, TEXT("Conversion factor for drive stiffness."));
+
+	float ChaosImmediate_SoftLinearStiffnessSourceMin = 100.0f;
+	float ChaosImmediate_SoftLinearStiffnessSourceMax = 500.0f;
+	float ChaosImmediate_SoftLinearStiffnessTargetMin = 0.5f;
+	float ChaosImmediate_SoftLinearStiffnessTargetMax = 1.0f;
+	FAutoConsoleVariableRef CVarSoftLinearStiffnessSourceMin(TEXT("p.Chaos.ImmPhys.SoftyLinearStiffnessSourceMin"), ChaosImmediate_SoftLinearStiffnessSourceMin, TEXT("Conversion factor for soft linear stiffness."));
+	FAutoConsoleVariableRef CVarSoftLinearStiffnessSourceMax(TEXT("p.Chaos.ImmPhys.SoftyLinearStiffnessSourceMax"), ChaosImmediate_SoftLinearStiffnessSourceMax, TEXT("Conversion factor for soft linear stiffness."));
+
+	float ChaosImmediate_SoftAngularStiffnessSourceMin = 100.0f;
+	float ChaosImmediate_SoftAngularStiffnessSourceMax = 500.0f;
+	float ChaosImmediate_SoftAngularStiffnessTargetMin = 0.5f;
+	float ChaosImmediate_SoftAngularStiffnessTargetMax = 1.0f;
+	FAutoConsoleVariableRef CVarSoftAngularStiffnessSourceMin(TEXT("p.Chaos.ImmPhys.SoftAngularStiffnessSourceMin"), ChaosImmediate_SoftAngularStiffnessSourceMin, TEXT("Conversion factor for soft angular stiffness."));
+	FAutoConsoleVariableRef CVarSoftAngularStiffnessSourceMax(TEXT("p.Chaos.ImmPhys.SoftAngularStiffnessSourceMax"), ChaosImmediate_SoftAngularStiffnessSourceMax, TEXT("Conversion factor for soft angular stiffness."));
 
 	// Convert the UE (PhysX) drive spring stiffness to a joint stiffness [0,1] value for use in the solver.
-	float ConvertAngularDriveStiffness(float InStiffness, const Chaos::TVector<float, 3>& IIa, const Chaos::TVector<float, 3>& IIb)
+	// We linearly map a range of UE stiffness values to a range of Chaos values, optionally scaling by the effective mass.
+	// Ranges and effective mass use are vcar controlled.
+	float ConvertDriveAngularStiffness(const float InStiffness, const Chaos::TVector<float, 3>& InvI0, const Chaos::TVector<float, 3>& InvI1)
+	{
+		if (ChaosImmediate_DriveStiffnessSourceMax > ChaosImmediate_DriveStiffnessSourceMin)
+		{
+			const float InvI = ChaosImmediate_StiffnessUseMass ? (InvI0.Min() + InvI1.Min()) : 1.0f;
+			const float F = (InStiffness - ChaosImmediate_DriveStiffnessSourceMin) * InvI  / (ChaosImmediate_DriveStiffnessSourceMax - ChaosImmediate_DriveStiffnessSourceMin);
+			return FMath::Lerp(ChaosImmediate_DriveStiffnessTargetMin, ChaosImmediate_DriveStiffnessTargetMax, FMath::Clamp(F, 0.0f, 1.0f));
+		}
+		return 0;
+	}
+
+	float ConvertSoftLinearStiffness(const float InStiffness, const float InvM0, const float InvM1)
 	{
-		if (ChaosImmediate_MaxDriveStiffness > 0)
+		if (ChaosImmediate_SoftLinearStiffnessSourceMax > ChaosImmediate_SoftLinearStiffnessSourceMin)
 		{
-			float IIMax = 1.0f;
-			if (ChaosImmediate_ScaleDriveStiffnessByMass)
-			{
-				IIMax = IIa.Min() + IIb.Min();
-			}
-			float Stiffness = (InStiffness / ChaosImmediate_MaxDriveStiffness) * IIMax;
-			return FMath::Clamp(Stiffness, 0.0f, 1.0f);
+			const float InvM = ChaosImmediate_StiffnessUseMass ? (InvM0 + InvM1) : 1.0f;
+			const float F = (InStiffness - ChaosImmediate_SoftLinearStiffnessSourceMin) * InvM / (ChaosImmediate_SoftLinearStiffnessSourceMax - ChaosImmediate_SoftLinearStiffnessSourceMin);
+			return FMath::Lerp(ChaosImmediate_SoftLinearStiffnessTargetMin, ChaosImmediate_SoftLinearStiffnessTargetMax, FMath::Clamp(F, 0.0f, 1.0f));
 		}
 		return 0;
 	}
 
-	float ConvertAngularDriveDamping(float InDamping, const Chaos::TVector<float, 3>& IIa, const Chaos::TVector<float, 3>& IIb)
+	float ConvertSoftAngularStiffness(const float InStiffness, const Chaos::TVector<float, 3>& InvI0, const Chaos::TVector<float, 3>& InvI1)
 	{
+		if (ChaosImmediate_SoftAngularStiffnessSourceMax > ChaosImmediate_SoftAngularStiffnessSourceMin)
+		{
+			const float InvI = ChaosImmediate_StiffnessUseMass ? (InvI0.Min() + InvI1.Min()) : 1.0f;
+			const float F = (InStiffness - ChaosImmediate_SoftAngularStiffnessSourceMin) * InvI / (ChaosImmediate_SoftAngularStiffnessSourceMax - ChaosImmediate_SoftAngularStiffnessSourceMin);
+			return FMath::Lerp(ChaosImmediate_SoftAngularStiffnessTargetMin, ChaosImmediate_SoftAngularStiffnessTargetMax, FMath::Clamp(F, 0.0f, 1.0f));
+		}
 		return 0;
 	}
 
+
 	FJointHandle::FJointHandle(FChaosConstraintContainer* InConstraints, FConstraintInstance* ConstraintInstance, FActorHandle* Actor1, FActorHandle* Actor2)
 		: ActorHandles({ Actor1, Actor2 })
 		, Constraints(InConstraints)
@@ -80,7 +116,13 @@ namespace ImmediatePhysics_Chaos
 			FMath::DegreesToRadians(ConstraintInstance->GetAngularSwing1Limit()),
 			FMath::DegreesToRadians(ConstraintInstance->GetAngularSwing2Limit())
 		};
-
+		ConstraintSettings.Motion.bSoftLinearLimitsEnabled = ConstraintInstance->GetIsSoftLinearLimit();
+		ConstraintSettings.Motion.bSoftTwistLimitsEnabled = ConstraintInstance->GetIsSoftTwistLimit();
+		ConstraintSettings.Motion.bSoftSwingLimitsEnabled = ConstraintInstance->GetIsSoftSwingLimit();
+		ConstraintSettings.Motion.SoftLinearStiffness = ConvertSoftLinearStiffness(ConstraintInstance->GetSoftLinearLimitStiffness(), Actor1->GetInverseMass(), Actor2->GetInverseMass());
+		ConstraintSettings.Motion.SoftTwistStiffness = ConvertSoftAngularStiffness(ConstraintInstance->GetSoftTwistLimitStiffness(), Actor1->GetInverseInertia(), Actor2->GetInverseInertia());
+		ConstraintSettings.Motion.SoftSwingStiffness = ConvertSoftAngularStiffness(ConstraintInstance->GetSoftSwingLimitStiffness(), Actor1->GetInverseInertia(), Actor2->GetInverseInertia());
+		 
 		// @todo(ccaulfield): Remove one of these
 		ConstraintSettings.Motion.AngularDriveTarget = FQuat(ConstraintInstance->ProfileInstance.AngularDrive.OrientationTarget);
 		ConstraintSettings.Motion.AngularDriveTargetAngles = TVector<float, 3>(
@@ -102,8 +144,7 @@ namespace ImmediatePhysics_Chaos
 			ConstraintSettings.Motion.bAngularSwingDriveEnabled = ConstraintInstance->ProfileInstance.AngularDrive.SwingDrive.bEnablePositionDrive;
 		}
 
-		ConstraintSettings.Motion.AngularDriveStiffness = ConvertAngularDriveStiffness(ConstraintInstance->ProfileInstance.AngularDrive.TwistDrive.Stiffness, Actor1->GetInverseInertia(), Actor2->GetInverseInertia());
-		ConstraintSettings.Motion.AngularDriveDamping = ConvertAngularDriveDamping(ConstraintInstance->ProfileInstance.AngularDrive.TwistDrive.Damping, Actor1->GetInverseInertia(), Actor2->GetInverseInertia());
+		ConstraintSettings.Motion.AngularDriveStiffness = ConvertDriveAngularStiffness(ConstraintInstance->ProfileInstance.AngularDrive.TwistDrive.Stiffness, Actor1->GetInverseInertia(), Actor2->GetInverseInertia());
 
 		ConstraintHandle = Constraints->AddConstraint({ Actor1->ParticleHandle, Actor2->ParticleHandle }, ConstraintSettings);
 	}
diff --git a/Engine/Source/Runtime/Engine/Private/PhysicsEngine/ImmediatePhysics/ImmediatePhysicsChaos/ImmediatePhysicsSimulation_Chaos.cpp b/Engine/Source/Runtime/Engine/Private/PhysicsEngine/ImmediatePhysics/ImmediatePhysicsChaos/ImmediatePhysicsSimulation_Chaos.cpp
--- a/Engine/Source/Runtime/Engine/Private/PhysicsEngine/ImmediatePhysics/ImmediatePhysicsChaos/ImmediatePhysicsSimulation_Chaos.cpp
+++ b/Engine/Source/Runtime/Engine/Private/PhysicsEngine/ImmediatePhysics/ImmediatePhysicsChaos/ImmediatePhysicsSimulation_Chaos.cpp
@@ -21,12 +21,12 @@
 // @todo(ccaulfield): remove when finished
 //
 float ChaosImmediate_Evolution_DeltaTime = 0.03f;
-int32 ChaosImmediate_Evolution_Iterations = 4;
+int32 ChaosImmediate_Evolution_Iterations = 6;
 FAutoConsoleVariableRef CVarChaosImmPhysDeltaTime(TEXT("p.Chaos.ImmPhys.DeltaTime"), ChaosImmediate_Evolution_DeltaTime, TEXT("Override chaos immediate physics delta time if non-zero"));
 FAutoConsoleVariableRef CVarChaosImmPhysIterations(TEXT("p.Chaos.ImmPhys.Iterations"), ChaosImmediate_Evolution_Iterations, TEXT("Number of constraint solver loops in immediate physics"));
 
 int32 ChaosImmediate_Collision_Enabled = 1;
-int32 ChaosImmediate_Collision_ApplyEnabled = 0;
+int32 ChaosImmediate_Collision_ApplyEnabled = 1;
 int32 ChaosImmediate_Collision_PushOutIterations = 5;
 int32 ChaosImmediate_Collision_PushOutPairIterations = 2;
 float ChaosImmediate_Collision_Thickness = 0;
@@ -37,15 +37,16 @@ FAutoConsoleVariableRef CVarChaosImmPhysCollisionPushOutPairIterations(TEXT("p.C
 FAutoConsoleVariableRef CVarChaosImmPhysThickness(TEXT("p.Chaos.ImmPhys.CollisionThickness"), ChaosImmediate_Collision_Thickness, TEXT("ChaosImmediateThickness"));
 
 float ChaosImmediate_Joint_SwingTwistAngleTolerance = 1.0e-6f;
-int32 ChaosImmediate_Joint_PositionIterations = 4;
+int32 ChaosImmediate_Joint_PositionIterations = 6;
 int32 ChaosImmediate_Joint_EnableVelocitySolve = 1;
 int32 ChaosImmediate_Joint_EnableLinearLimits = 1;
 int32 ChaosImmediate_Joint_EnableTwistLimits = 1;
 int32 ChaosImmediate_Joint_EnableSwingLimits = 1;
-int32 ChaosImmediate_Joint_EnablePositionCorrection = 0;
 int32 ChaosImmediate_Joint_EnableDrives = 1;
 float ChaosImmediate_Joint_Projection = 0.1f;
 float ChaosImmediate_Joint_Stiffness = 1.0f;
+float ChaosImmediate_Joint_SoftLinearStiffness = 0.0f;
+float ChaosImmediate_Joint_SoftAngularStiffness = 0.0f;
 float ChaosImmediate_Joint_DriveStiffness = 0.0f;
 float ChaosImmediate_Joint_MinParentMassRatio = 0.5f;
 float ChaosImmediate_Joint_MaxInertiaRatio = 5.0f;
@@ -55,10 +56,11 @@ FAutoConsoleVariableRef CVarChaosImmPhysEnableVelocitySolve(TEXT("p.Chaos.ImmPhy
 FAutoConsoleVariableRef CVarChaosImmPhysEnableLinearLimits(TEXT("p.Chaos.ImmPhys.Joint.EnableLinearLimits"), ChaosImmediate_Joint_EnableLinearLimits, TEXT("EnableLinearLimits."));
 FAutoConsoleVariableRef CVarChaosImmPhysEnableTwistLimits(TEXT("p.Chaos.ImmPhys.Joint.EnableTwistLimits"), ChaosImmediate_Joint_EnableTwistLimits, TEXT("EnableTwistLimits."));
 FAutoConsoleVariableRef CVarChaosImmPhysEnableSwingLimits(TEXT("p.Chaos.ImmPhys.Joint.EnableSwingLimits"), ChaosImmediate_Joint_EnableSwingLimits, TEXT("EnableSwingLimits."));
-FAutoConsoleVariableRef CVarChaosImmPhysEnablePositionCorrection(TEXT("p.Chaos.ImmPhys.Joint.EnablePositionCorrection"), ChaosImmediate_Joint_EnablePositionCorrection, TEXT("EnablePositionCorrection."));
 FAutoConsoleVariableRef CVarChaosImmPhysEnableDrives(TEXT("p.Chaos.ImmPhys.Joint.EnableDrives"), ChaosImmediate_Joint_EnableDrives, TEXT("EnableDrives."));
 FAutoConsoleVariableRef CVarChaosImmPhysProjection(TEXT("p.Chaos.ImmPhys.Joint.Projection"), ChaosImmediate_Joint_Projection, TEXT("6Dof joint projection amount override (if > 0)."));
 FAutoConsoleVariableRef CVarChaosImmPhysStiffness(TEXT("p.Chaos.ImmPhys.Joint.Stiffness"), ChaosImmediate_Joint_Stiffness, TEXT("6Dof joint stiffness override (if > 0)."));
+FAutoConsoleVariableRef CVarChaosImmPhysSoftLinearStiffness(TEXT("p.Chaos.ImmPhys.Joint.SoftLinearStiffness"), ChaosImmediate_Joint_SoftLinearStiffness, TEXT("6Dof joint soft linear stiffness override (if > 0)."));
+FAutoConsoleVariableRef CVarChaosImmPhysSoftAngularStiffness(TEXT("p.Chaos.ImmPhys.Joint.SoftAngularStiffness"), ChaosImmediate_Joint_SoftAngularStiffness, TEXT("6Dof joint soft angular stiffness override (if > 0)."));
 FAutoConsoleVariableRef CVarChaosImmPhysDriveStiffness(TEXT("p.Chaos.ImmPhys.Joint.DriveStiffness"), ChaosImmediate_Joint_DriveStiffness, TEXT("6Dof joint drive stiffness override (if > 0)."));
 FAutoConsoleVariableRef CVarChaosImmPhysMinParentMassRatio(TEXT("p.Chaos.ImmPhys.Joint.MinParentMassRatio"), ChaosImmediate_Joint_MinParentMassRatio, TEXT("6Dof joint MinParentMassRatio (if > 0)"));
 FAutoConsoleVariableRef CVarChaosImmPhysMaxInertiaRatio(TEXT("p.Chaos.ImmPhys.Joint.MaxInertiaRatio"), ChaosImmediate_Joint_MaxInertiaRatio, TEXT("6Dof joint MaxInertiaRatio (if > 0)"));
@@ -129,6 +131,10 @@ namespace ImmediatePhysics_Chaos
 
 		Evolution->AddConstraintRule(JointsRule.Get());
 
+		// We want collisions to override joints
+		JointsRule->SetPriority(0);
+		Evolution->GetCollisionConstraintsRule().SetPriority(1);
+
 		// Filter collisions after detection
 		// @todo(ccaulfield): Eventually we will build lists of potentially colliding pairs and won't need this
 		Collisions.SetPostComputeCallback(
@@ -398,7 +404,9 @@ namespace ImmediatePhysics_Chaos
 		{
 			if (ChaosImmediate_Evolution_DeltaTime > 0)
 			{
-				DeltaTime = ChaosImmediate_Evolution_DeltaTime;
+				// Round Dt to the nearest multiple of fixed step size...
+				const float NumDts = FMath::RoundToFloat(DeltaTime / ChaosImmediate_Evolution_DeltaTime);
+				DeltaTime = NumDts * ChaosImmediate_Evolution_DeltaTime;
 			}
 			UE_LOG(LogChaosJoint, Verbose, TEXT("Simulate dt = %f"), DeltaTime);
 
@@ -410,11 +418,12 @@ namespace ImmediatePhysics_Chaos
 			JointsSettings.bEnableLinearLimits = ChaosImmediate_Joint_EnableLinearLimits != 0;
 			JointsSettings.bEnableTwistLimits = ChaosImmediate_Joint_EnableTwistLimits != 0;
 			JointsSettings.bEnableSwingLimits = ChaosImmediate_Joint_EnableSwingLimits != 0;
-			JointsSettings.bEnablePositionCorrection = ChaosImmediate_Joint_EnablePositionCorrection != 0;
 			JointsSettings.bEnableDrives = ChaosImmediate_Joint_EnableDrives != 0;
 			JointsSettings.PositionIterations = ChaosImmediate_Joint_PositionIterations;
 			JointsSettings.Projection = ChaosImmediate_Joint_Projection;
 			JointsSettings.Stiffness = ChaosImmediate_Joint_Stiffness;
+			JointsSettings.SoftLinearStiffness = ChaosImmediate_Joint_SoftLinearStiffness;
+			JointsSettings.SoftAngularStiffness = ChaosImmediate_Joint_SoftAngularStiffness;
 			JointsSettings.DriveStiffness = ChaosImmediate_Joint_DriveStiffness;
 			Joints->SetSettings(JointsSettings);
 
diff --git a/Engine/Source/Runtime/Experimental/Chaos/Private/Chaos/PBDConstraintRule.cpp b/Engine/Source/Runtime/Experimental/Chaos/Private/Chaos/PBDConstraintRule.cpp
--- a/Engine/Source/Runtime/Experimental/Chaos/Private/Chaos/PBDConstraintRule.cpp
+++ b/Engine/Source/Runtime/Experimental/Chaos/Private/Chaos/PBDConstraintRule.cpp
@@ -11,7 +11,6 @@
 
 namespace Chaos
 {
-	template class TPBDConstraintRule<float, 3>;
 	template class TPBDConstraintColorRule<TPBDCollisionConstraint<float, 3>, float, 3>;
 	template class TPBDConstraintIslandRule<TPBD6DJointConstraints<float, 3>, float, 3>;
 	template class TPBDConstraintIslandRule<TPBDJointConstraints<float, 3>, float, 3>;
diff --git a/Engine/Source/Runtime/Experimental/Chaos/Private/Chaos/PBDJointConstraintUtilities.cpp b/Engine/Source/Runtime/Experimental/Chaos/Private/Chaos/PBDJointConstraintUtilities.cpp
--- a/Engine/Source/Runtime/Experimental/Chaos/Private/Chaos/PBDJointConstraintUtilities.cpp
+++ b/Engine/Source/Runtime/Experimental/Chaos/Private/Chaos/PBDJointConstraintUtilities.cpp
@@ -347,6 +347,41 @@ namespace Chaos
 		}
 	}
 
+	template<typename T, int d>
+	T GetLinearStiffness(
+		const TPBDJointSolverSettings<T, d>& SolverSettings,
+		const TPBDJointSettings<T, d>& JointSettings)
+	{
+		const float SolverStiffness = (SolverSettings.Stiffness > (T)0) ? SolverSettings.Stiffness : JointSettings.Motion.Stiffness;
+		const float SoftSolverStiffness = (SolverSettings.SoftLinearStiffness > (T)0) ? SolverSettings.SoftLinearStiffness : JointSettings.Motion.SoftLinearStiffness;
+		const bool bIsSoft = JointSettings.Motion.bSoftLinearLimitsEnabled && ((JointSettings.Motion.LinearMotionTypes[0] == EJointMotionType::Limited) || (JointSettings.Motion.LinearMotionTypes[1] == EJointMotionType::Limited) || (JointSettings.Motion.LinearMotionTypes[2] == EJointMotionType::Limited));
+		const float Stiffness = bIsSoft ? SolverStiffness * SoftSolverStiffness : SolverStiffness;
+		return Stiffness;
+	}
+
+	template<typename T, int d>
+	T GetTwistStiffness(
+		const TPBDJointSolverSettings<T, d>& SolverSettings,
+		const TPBDJointSettings<T, d>& JointSettings)
+	{
+		const float SolverStiffness = (SolverSettings.Stiffness > (T)0) ? SolverSettings.Stiffness : JointSettings.Motion.Stiffness;
+		const float SoftSolverStiffness = (SolverSettings.SoftAngularStiffness > (T)0) ? SolverSettings.SoftAngularStiffness : JointSettings.Motion.SoftTwistStiffness;
+		const bool bIsSoft = JointSettings.Motion.bSoftTwistLimitsEnabled && (JointSettings.Motion.AngularMotionTypes[(int32)EJointAngularConstraintIndex::Twist] == EJointMotionType::Limited);
+		const float Stiffness = bIsSoft ? SolverStiffness * SoftSolverStiffness : SolverStiffness;
+		return Stiffness;
+	}
+
+	template<typename T, int d>
+	T GetSwingStiffness(
+		const TPBDJointSolverSettings<T, d>& SolverSettings,
+		const TPBDJointSettings<T, d>& JointSettings)
+	{
+		const float SolverStiffness = (SolverSettings.Stiffness > (T)0) ? SolverSettings.Stiffness : JointSettings.Motion.Stiffness;
+		const float SoftSolverStiffness = (SolverSettings.SoftAngularStiffness > (T)0) ? SolverSettings.SoftAngularStiffness : JointSettings.Motion.SoftSwingStiffness;
+		const bool bIsSoft = JointSettings.Motion.bSoftSwingLimitsEnabled && ((JointSettings.Motion.AngularMotionTypes[(int32)EJointAngularConstraintIndex::Swing1] == EJointMotionType::Limited) || (JointSettings.Motion.AngularMotionTypes[(int32)EJointAngularConstraintIndex::Swing2] == EJointMotionType::Limited));
+		const float Stiffness = bIsSoft ? SolverStiffness * SoftSolverStiffness : SolverStiffness;
+		return Stiffness;
+	}
 
 	template<typename T, int d>
 	void ApplyPositionDelta(
@@ -358,10 +393,9 @@ namespace Chaos
 		TVector<T, d>& P0,
 		TVector<T, d>& P1,
 		const TVector<T, d>& DP0,
-		const TVector<T, d>& DP1)
+		const TVector<T, d>& DP1,
+		const T Stiffness)
 	{
-		const float Stiffness = (SolverSettings.Stiffness > (T)0) ? SolverSettings.Stiffness : JointSettings.Motion.Stiffness;
-
 		P0 = P0 + Stiffness * DP0;
 		P1 = P1 + Stiffness * DP1;
 	}
@@ -378,14 +412,13 @@ namespace Chaos
 		TVector<T, d>& P1,
 		TVector<T, d>& V1,
 		const TVector<T, d>& DV0,
-		const TVector<T, d>& DV1)
+		const TVector<T, d>& DV1, 
+		const T Stiffness)
 	{
-		const float Stiffness = (SolverSettings.Stiffness > (T)0) ? SolverSettings.Stiffness : JointSettings.Motion.Stiffness;
-
 		V0 = V0 + Stiffness * DV0;
 		V1 = V1 + Stiffness * DV1;
 
-		ApplyPositionDelta(Dt, SolverSettings, JointSettings, Index0, Index1, P0, P1, DV0 * Dt, DV1 * Dt);
+		ApplyPositionDelta(Dt, SolverSettings, JointSettings, Index0, Index1, P0, P1, DV0 * Dt, DV1 * Dt, Stiffness);
 	}
 
 	template<typename T, int d>
@@ -398,10 +431,9 @@ namespace Chaos
 		TRotation<T, d>& Q0,
 		TRotation<T, d>& Q1,
 		const TVector<T, d>& DR0,
-		const TVector<T, d>& DR1)
+		const TVector<T, d>& DR1,
+		const T Stiffness)
 	{
-		const float Stiffness = (SolverSettings.Stiffness > (T)0) ? SolverSettings.Stiffness : JointSettings.Motion.Stiffness;
-
 		const TRotation<T, d> DQ0 = (TRotation<T, d>::FromElements(Stiffness * DR0, 0) * Q0) * (T)0.5;
 		const TRotation<T, d> DQ1 = (TRotation<T, d>::FromElements(Stiffness * DR1, 0) * Q1) * (T)0.5;
 		Q0 = (Q0 + DQ0).GetNormalized();
@@ -422,14 +454,13 @@ namespace Chaos
 		TRotation<T, d>& Q1,
 		TVector<T, d>& W1,
 		const TVector<T, d>& DW0,
-		const TVector<T, d>& DW1)
+		const TVector<T, d>& DW1,
+		const T Stiffness)
 	{
-		const float Stiffness = (SolverSettings.Stiffness > (T)0) ? SolverSettings.Stiffness : JointSettings.Motion.Stiffness;
-
 		W0 = W0 + Stiffness * DW0;
 		W1 = W1 + Stiffness * DW1;
 
-		ApplyRotationDelta(Dt, SolverSettings, JointSettings, Index0, Index1, Q0, Q1, DW0 * Dt, DW1 * Dt);
+		ApplyRotationDelta(Dt, SolverSettings, JointSettings, Index0, Index1, Q0, Q1, DW0 * Dt, DW1 * Dt, Stiffness);
 	}
 
 
@@ -444,13 +475,14 @@ namespace Chaos
 		TVector<T, d>& W0,
 		TRotation<T, d>& Q1,
 		TVector<T, d>& W1,
-		float InvM0,
+		const T InvM0,
 		const PMatrix<T, d, d>& InvIL0,
-		float InvM1,
+		const T InvM1,
 		const PMatrix<T, d, d>& InvIL1,
 		const TVector<T, d>& Axis0,
 		const TVector<T, d>& Axis1,
-		const float WC)
+		const T WC,
+		const T Stiffness)
 	{
 		const PMatrix<T, d, d> InvI0 = Utilities::Multiply(Q0.ToMatrix(), Utilities::Multiply(InvIL0, Q0.ToMatrix().GetTransposed()));
 		const PMatrix<T, d, d> InvI1 = Utilities::Multiply(Q1.ToMatrix(), Utilities::Multiply(InvIL1, Q1.ToMatrix().GetTransposed()));
@@ -462,7 +494,7 @@ namespace Chaos
 		const TVector<T, d> DW0 = Utilities::Multiply(InvI0, Axis0) * L * WC;
 		const TVector<T, d> DW1 = -Utilities::Multiply(InvI1, Axis1) * L * WC;
 
-		ApplyRotationVelocityDelta(Dt, SolverSettings, JointSettings, Index0, Index1, Q0, W0, Q1, W1, DW0, DW1);
+		ApplyRotationVelocityDelta(Dt, SolverSettings, JointSettings, Index0, Index1, Q0, W0, Q1, W1, DW0, DW1, Stiffness);
 	}
 
 	template<typename T, int d>
@@ -476,14 +508,15 @@ namespace Chaos
 		TRotation<T, d>& Q0,
 		TVector<T, d>& P1,
 		TRotation<T, d>& Q1,
-		float InvM0,
+		const T InvM0,
 		const PMatrix<T, d, d>& InvIL0,
-		float InvM1,
+		const T InvM1,
 		const PMatrix<T, d, d>& InvIL1,
 		const TVector<T, d>& Axis0,
-		const float Angle0,
+		const T Angle0,
 		const TVector<T, d>& Axis1,
-		const float Angle1)
+		const T Angle1, 
+		const T Stiffness)
 	{
 		const PMatrix<T, d, d> InvI0 = Utilities::Multiply(Q0.ToMatrix(), Utilities::Multiply(InvIL0, Q0.ToMatrix().GetTransposed()));
 		const PMatrix<T, d, d> InvI1 = Utilities::Multiply(Q1.ToMatrix(), Utilities::Multiply(InvIL1, Q1.ToMatrix().GetTransposed()));
@@ -496,65 +529,7 @@ namespace Chaos
 		const TVector<T, d> DR0 = Utilities::Multiply(InvI0, Axis0) * L * Angle0;
 		const TVector<T, d> DR1 = Utilities::Multiply(InvI1, Axis1) * L * Angle1;
 
-		ApplyRotationDelta(Dt, SolverSettings, JointSettings, Index0, Index1, Q0, Q1, DR0, DR1);
-	}
-
-	template<typename T, int d>
-	void ApplyPostRotationPositionCorrection(
-		const T Dt,
-		const TPBDJointSolverSettings<T, d>& SolverSettings,
-		const TPBDJointSettings<T, d>& JointSettings,
-		const int32 Index0,
-		const int32 Index1,
-		TVector<T, d>& P0,
-		TRotation<T, d>& Q0,
-		TVector<T, d>& P1,
-		TRotation<T, d>& Q1,
-		float InvM0,
-		const PMatrix<T, d, d>& InvIL0,
-		float InvM1,
-		const PMatrix<T, d, d>& InvIL1,
-		const TVector<T, d>& CX_1,
-		const TVector<T, d>& Axis0,
-		const TVector<T, d>& Axis1)
-	{
-		if (!SolverSettings.bEnablePositionCorrection)
-		{
-			return;
-		}
-
-		// Post-rotation constraint positions
-		const TRigidTransform<T, d>& XL0 = JointSettings.ConstraintFrames[Index0];
-		const TRigidTransform<T, d>& XL1 = JointSettings.ConstraintFrames[Index1];
-		const TVector<T, d> X0 = P0 + Q0 * XL0.GetTranslation();
-		const TVector<T, d> X1 = P1 + Q1 * XL1.GetTranslation();
-		const TRotation<T, d> R0 = Q0 * XL0.GetRotation();
-		TVector<T, d> CX_2 = GetLimitedPositionError(JointSettings, R0, X1 - X0);
-
-		// Increase in position error caused by the rotation correction
-		TVector<T, d> DX = CX_2 - CX_1;
-		for (int32 AxisIndex = 0; AxisIndex < d; ++AxisIndex)
-		{
-			if ((CX_1[AxisIndex] > 0) && (CX_2[AxisIndex] < CX_1[AxisIndex]))
-			{
-				DX[AxisIndex] = 0;
-			}
-			if ((CX_1[AxisIndex] < 0) && (CX_2[AxisIndex] > CX_1[AxisIndex]))
-			{
-				DX[AxisIndex] = 0;
-			}
-		}
-		const TVector<T, d> DX0 = DX - TVector<T, d>::DotProduct(DX, Axis0) * Axis0;
-		const TVector<T, d> DX1 = DX - TVector<T, d>::DotProduct(DX, Axis1) * Axis1;
-
-		// Correct the extra position error introduced by the rotation correction. We are effectively treating the
-		// bodies as if they have infinite inertia for this correction, which is only correct if
-		// the position correction exactly opposes the constraint that caused the increased position error
-		// which is only the case for uniform shapes (i.e., spherical inertia). To do this properly, we
-		// would have to solve for position and rotation correction simultaneously.
-		const TVector<T, d> DP0 = InvM0 * DX0 / (InvM0 + InvM1);
-		const TVector<T, d> DP1 = -InvM1 * DX1 / (InvM0 + InvM1);
-		ApplyPositionDelta(Dt, SolverSettings, JointSettings, Index0, Index1, P0, P1, DP0, DP1);
+		ApplyRotationDelta(Dt, SolverSettings, JointSettings, Index0, Index1, Q0, Q1, DR0, DR1, Stiffness);
 	}
 
 	template<class T, int d>
@@ -749,8 +724,8 @@ namespace Chaos
 		const TVector<T, d> DP1 = -InvM1 * DX;
 		const TVector<T, d> DR0 = Utilities::Multiply(InvI0, TVector<T, d>::CrossProduct(X0 - P0, DX));
 		const TVector<T, d> DR1 = Utilities::Multiply(InvI1, TVector<T, d>::CrossProduct(X1 - P1, -DX));
-		ApplyPositionDelta(Dt, SolverSettings, JointSettings, Index0, Index1, P0, P1, DP0, DP1);
-		ApplyRotationDelta(Dt, SolverSettings, JointSettings, Index0, Index1, Q0, Q1, DR0, DR1);
+		ApplyPositionDelta(Dt, SolverSettings, JointSettings, Index0, Index1, P0, P1, DP0, DP1, GetLinearStiffness(SolverSettings, JointSettings));
+		ApplyRotationDelta(Dt, SolverSettings, JointSettings, Index0, Index1, Q0, Q1, DR0, DR1, GetLinearStiffness(SolverSettings, JointSettings));
 	}
 
 	template<typename T, int d>
@@ -805,8 +780,8 @@ namespace Chaos
 		const TVector<T, d> DW0 = Utilities::Multiply(InvI0, TVector<T, d>::CrossProduct(XC0, DL));
 		const TVector<T, d> DW1 = -Utilities::Multiply(InvI1, TVector<T, d>::CrossProduct(XC1, DL));
 
-		ApplyVelocityDelta(Dt, SolverSettings, JointSettings, Index0, Index1, P0, V0, P1, V1, DV0, DV1);
-		ApplyRotationVelocityDelta(Dt, SolverSettings, JointSettings, Index0, Index1, Q0, W0, Q1, W1, DW0, DW1);
+		ApplyVelocityDelta(Dt, SolverSettings, JointSettings, Index0, Index1, P0, V0, P1, V1, DV0, DV1, GetLinearStiffness(SolverSettings, JointSettings));
+		ApplyRotationVelocityDelta(Dt, SolverSettings, JointSettings, Index0, Index1, Q0, W0, Q1, W1, DW0, DW1, GetLinearStiffness(SolverSettings, JointSettings));
 	}
 
 
@@ -880,10 +855,7 @@ namespace Chaos
 		const T DTwistAngle1 = -DTwistAngle;
 
 		// Apply twist correction
-		ApplyRotationDelta(Dt, SolverSettings, JointSettings, Index0, Index1, P0, Q0, P1, Q1, InvM0, InvIL0, InvM1, InvIL1, TwistAxis0, DTwistAngle0, TwistAxis1, DTwistAngle1);
-
-		// Correct any positional error we may have introduced
-		ApplyPostRotationPositionCorrection<T, d>(Dt, SolverSettings, JointSettings, Index0, Index1, P0, Q0, P1, Q1, InvM0, InvIL0, InvM1, InvIL1, CX, TwistAxis0, TwistAxis1);
+		ApplyRotationDelta(Dt, SolverSettings, JointSettings, Index0, Index1, P0, Q0, P1, Q1, InvM0, InvIL0, InvM1, InvIL1, TwistAxis0, DTwistAngle0, TwistAxis1, DTwistAngle1, GetTwistStiffness(SolverSettings, JointSettings));
 	}
 
 	template<typename T, int d>
@@ -959,7 +931,7 @@ namespace Chaos
 			DW = FMath::Min((T)0, WC1 - WC0);
 		}
 
-		ApplyRotationVelocityDelta(Dt, SolverSettings, JointSettings, Index0, Index1, Q0, W0, Q1, W1, InvM0, InvIL0, InvM1, InvIL1, TwistAxis0, TwistAxis1, DW);
+		ApplyRotationVelocityDelta(Dt, SolverSettings, JointSettings, Index0, Index1, Q0, W0, Q1, W1, InvM0, InvIL0, InvM1, InvIL1, TwistAxis0, TwistAxis1, DW, GetTwistStiffness(SolverSettings, JointSettings));
 	}
 
 
@@ -1036,10 +1008,7 @@ namespace Chaos
 		T DSwingAngle1 = -DSwingAngle;
 
 		// Apply swing correction
-		ApplyRotationDelta(Dt, SolverSettings, JointSettings, Index0, Index1, P0, Q0, P1, Q1, InvM0, InvIL0, InvM1, InvIL1, SwingAxis0, DSwingAngle0, SwingAxis1, DSwingAngle1);
-	
-		// Correct any positional error we may have introduced
-		ApplyPostRotationPositionCorrection<T, d>(Dt, SolverSettings, JointSettings, Index0, Index1, P0, Q0, P1, Q1, InvM0, InvIL0, InvM1, InvIL1, CX, SwingAxis0, SwingAxis1);
+		ApplyRotationDelta(Dt, SolverSettings, JointSettings, Index0, Index1, P0, Q0, P1, Q1, InvM0, InvIL0, InvM1, InvIL1, SwingAxis0, DSwingAngle0, SwingAxis1, DSwingAngle1, GetSwingStiffness(SolverSettings, JointSettings));
 	}
 
 	template<typename T, int d>
@@ -1117,7 +1086,7 @@ namespace Chaos
 			DW = FMath::Min((T)0, WC1 - WC0);
 		}
 
-		ApplyRotationVelocityDelta(Dt, SolverSettings, JointSettings, Index0, Index1, Q0, W0, Q1, W1, InvM0, InvIL0, InvM1, InvIL1, SwingAxis0, SwingAxis1, DW);
+		ApplyRotationVelocityDelta(Dt, SolverSettings, JointSettings, Index0, Index1, Q0, W0, Q1, W1, InvM0, InvIL0, InvM1, InvIL1, SwingAxis0, SwingAxis1, DW, GetSwingStiffness(SolverSettings, JointSettings));
 	}
 
 
@@ -1212,10 +1181,7 @@ namespace Chaos
 			T DSwingAngle1 = -DSwingAngle;
 
 			// Apply swing correction
-			ApplyRotationDelta(Dt, SolverSettings, JointSettings, Index0, Index1, P0, Q0, P1, Q1, InvM0, InvIL0, InvM1, InvIL1, SwingAxis0, DSwingAngle0, SwingAxis1, DSwingAngle1);
-
-			// Correct any positional error we may have introduced
-			ApplyPostRotationPositionCorrection<T, d>(Dt, SolverSettings, JointSettings, Index0, Index1, P0, Q0, P1, Q1, InvM0, InvIL0, InvM1, InvIL1, CX, SwingAxis0, SwingAxis1);
+			ApplyRotationDelta(Dt, SolverSettings, JointSettings, Index0, Index1, P0, Q0, P1, Q1, InvM0, InvIL0, InvM1, InvIL1, SwingAxis0, DSwingAngle0, SwingAxis1, DSwingAngle1, GetSwingStiffness(SolverSettings, JointSettings));
 		}
 	}
 
@@ -1312,7 +1278,7 @@ namespace Chaos
 				DW = FMath::Min((T)0, WC1 - WC0);
 			}
 
-			ApplyRotationVelocityDelta(Dt, SolverSettings, JointSettings, Index0, Index1, Q0, W0, Q1, W1, InvM0, InvIL0, InvM1, InvIL1, SwingAxis0, SwingAxis1, DW);
+			ApplyRotationVelocityDelta(Dt, SolverSettings, JointSettings, Index0, Index1, Q0, W0, Q1, W1, InvM0, InvIL0, InvM1, InvIL1, SwingAxis0, SwingAxis1, DW, GetSwingStiffness(SolverSettings, JointSettings));
 		}
 	}
 
@@ -1365,13 +1331,10 @@ namespace Chaos
 		const T DriveStiffnessUnclamped = (SolverSettings.DriveStiffness > 0) ? SolverSettings.DriveStiffness : JointSettings.Motion.AngularDriveStiffness;
 		const T DriveStiffness = FMath::Clamp(DriveStiffnessUnclamped, (T)0, (T)1);
 		const T DTwistAngle = TwistAngle - TwistAngleTarget;
-		const T DTwistAngle0 = DriveStiffness * DTwistAngle;
-		const T DTwistAngle1 = -DriveStiffness * DTwistAngle;
-
-		ApplyRotationDelta(Dt, SolverSettings, JointSettings, Index0, Index1, P0, Q0, P1, Q1, InvM0, InvIL0, InvM1, InvIL1, TwistAxis0, DTwistAngle0, TwistAxis1, DTwistAngle1);
+		const T DTwistAngle0 = DTwistAngle;
+		const T DTwistAngle1 = -DTwistAngle;
 
-		// Correct any positional error we may have introduced
-		ApplyPostRotationPositionCorrection<T, d>(Dt, SolverSettings, JointSettings, Index0, Index1, P0, Q0, P1, Q1, InvM0, InvIL0, InvM1, InvIL1, CX, TwistAxis0, TwistAxis1);
+		ApplyRotationDelta(Dt, SolverSettings, JointSettings, Index0, Index1, P0, Q0, P1, Q1, InvM0, InvIL0, InvM1, InvIL1, TwistAxis0, DTwistAngle0, TwistAxis1, DTwistAngle1, DriveStiffness);
 	}
 
 	template<typename T, int d>
@@ -1424,14 +1387,11 @@ namespace Chaos
 		T DriveStiffnessUnclamped = (SolverSettings.DriveStiffness > 0) ? SolverSettings.DriveStiffness : JointSettings.Motion.AngularDriveStiffness;
 		T DriveStiffness = FMath::Clamp(DriveStiffnessUnclamped, (T)0, (T)1);
 		T DSwingAngle = SwingAngle - SwingAngleTarget;
-		T DSwingAngle0 = DriveStiffness * DSwingAngle;
-		T DSwingAngle1 = -DriveStiffness * DSwingAngle;
+		T DSwingAngle0 = DSwingAngle;
+		T DSwingAngle1 = -DSwingAngle;
 
 		// Apply swing correction
-		ApplyRotationDelta(Dt, SolverSettings, JointSettings, Index0, Index1, P0, Q0, P1, Q1, InvM0, InvIL0, InvM1, InvIL1, SwingAxis0, DSwingAngle0, SwingAxis1, DSwingAngle1);
-
-		// Correct any positional error we may have introduced
-		ApplyPostRotationPositionCorrection<T, d>(Dt, SolverSettings, JointSettings, Index0, Index1, P0, Q0, P1, Q1, InvM0, InvIL0, InvM1, InvIL1, CX, SwingAxis0, SwingAxis1);
+		ApplyRotationDelta(Dt, SolverSettings, JointSettings, Index0, Index1, P0, Q0, P1, Q1, InvM0, InvIL0, InvM1, InvIL1, SwingAxis0, DSwingAngle0, SwingAxis1, DSwingAngle1, DriveStiffness);
 	}
 
 	template<typename T, int d>
@@ -1482,9 +1442,6 @@ namespace Chaos
 			Q0 = TRotation<T, d>::Slerp(Q0, TargetQ0, F0);
 			Q1 = TRotation<T, d>::Slerp(Q1, TargetQ1, F1);
 			Q1.EnforceShortestArcWith(Q0);
-
-			// Correct any positional error we may have introduced
-			ApplyPostRotationPositionCorrection<T, d>(Dt, SolverSettings, JointSettings, Index0, Index1, P0, Q0, P1, Q1, InvM0, InvIL0, InvM1, InvIL1, CX, SLerpAxis, SLerpAxis);;
 		}
 	}
 
diff --git a/Engine/Source/Runtime/Experimental/Chaos/Private/Chaos/PBDJointConstraints.cpp b/Engine/Source/Runtime/Experimental/Chaos/Private/Chaos/PBDJointConstraints.cpp
--- a/Engine/Source/Runtime/Experimental/Chaos/Private/Chaos/PBDJointConstraints.cpp
+++ b/Engine/Source/Runtime/Experimental/Chaos/Private/Chaos/PBDJointConstraints.cpp
@@ -10,6 +10,7 @@
 
 #include "HAL/IConsoleManager.h"
 
+//#pragma optimize("", off)
 
 namespace Chaos
 {
@@ -66,13 +67,18 @@ namespace Chaos
 		, LinearLimit(FLT_MAX)
 		, AngularMotionTypes({ EJointMotionType::Free, EJointMotionType::Free, EJointMotionType::Free })
 		, AngularLimits(TVector<T, d>(FLT_MAX, FLT_MAX, FLT_MAX))
+		, bSoftLinearLimitsEnabled(false)
+		, bSoftTwistLimitsEnabled(false)
+		, bSoftSwingLimitsEnabled(false)
+		, SoftLinearStiffness(0)
+		, SoftTwistStiffness(0)
+		, SoftSwingStiffness(0)
 		, AngularDriveTarget(TRotation<T, d>::FromIdentity())
 		, AngularDriveTargetAngles(TVector<T, d>(0, 0, 0))
 		, bAngularSLerpDriveEnabled(false)
 		, bAngularTwistDriveEnabled(false)
 		, bAngularSwingDriveEnabled(false)
 		, AngularDriveStiffness(0)
-		, AngularDriveDamping(0)
 	{
 	}
 
@@ -84,13 +90,18 @@ namespace Chaos
 		, LinearLimit(FLT_MAX)
 		, AngularMotionTypes({ EJointMotionType::Free, EJointMotionType::Free, EJointMotionType::Free })
 		, AngularLimits(TVector<T, d>(FLT_MAX, FLT_MAX, FLT_MAX))
+		, bSoftLinearLimitsEnabled(false)
+		, bSoftTwistLimitsEnabled(false)
+		, bSoftSwingLimitsEnabled(false)
+		, SoftLinearStiffness(0)
+		, SoftTwistStiffness(0)
+		, SoftSwingStiffness(0)
 		, AngularDriveTarget(TRotation<T, d>::FromIdentity())
 		, AngularDriveTargetAngles(TVector<T, d>(0, 0, 0))
 		, bAngularSLerpDriveEnabled(false)
 		, bAngularTwistDriveEnabled(false)
 		, bAngularSwingDriveEnabled(false)
 		, AngularDriveStiffness(0)
-		, AngularDriveDamping(0)
 	{
 	}
 
@@ -120,11 +131,12 @@ namespace Chaos
 		, bEnableLinearLimits(true)
 		, bEnableTwistLimits(true)
 		, bEnableSwingLimits(true)
-		, bEnablePositionCorrection(false)
 		, bEnableDrives(true)
 		, Projection((T)0)
 		, Stiffness((T)0)
 		, DriveStiffness((T)0)
+		, SoftLinearStiffness((T)0)
+		, SoftAngularStiffness((T)0)
 		, PositionIterations((T)0)
 	{
 	}
@@ -367,15 +379,15 @@ namespace Chaos
 	template<class T, int d>
 	void TPBDJointConstraints<T, d>::ApplyPushOut(const T Dt, const TArray<FConstraintHandle*>& InConstraintHandles)
 	{
-		TArray<FConstraintHandle*> SortedConstraintHandles = InConstraintHandles;
-		SortedConstraintHandles.Sort([](const FConstraintHandle& L, const FConstraintHandle& R)
-		{
-			// Sort bodies from root to leaf
-			return L.GetConstraintLevel() < R.GetConstraintLevel();
-		});
-
 		if (Settings.bEnableVelocitySolve)
 		{
+			TArray<FConstraintHandle*> SortedConstraintHandles = InConstraintHandles;
+			SortedConstraintHandles.Sort([](const FConstraintHandle& L, const FConstraintHandle& R)
+				{
+					// Sort bodies from root to leaf
+					return L.GetConstraintLevel() < R.GetConstraintLevel();
+				});
+
 			for (int32 It = 0; It < Settings.PositionIterations; ++It)
 			{
 				for (FConstraintHandle* ConstraintHandle : SortedConstraintHandles)
@@ -385,8 +397,22 @@ namespace Chaos
 			}
 		}
 
+		// @todo(ccaulfield): should be called constraint rule
+		ApplyProjection(Dt, InConstraintHandles);
+	}
+
+	template<class T, int d>
+	void TPBDJointConstraints<T, d>::ApplyProjection(const T Dt, const TArray<FConstraintHandle*>& InConstraintHandles)
+	{
 		if (Settings.Projection > 0)
 		{
+			TArray<FConstraintHandle*> SortedConstraintHandles = InConstraintHandles;
+			SortedConstraintHandles.Sort([](const FConstraintHandle& L, const FConstraintHandle& R)
+				{
+					// Sort bodies from root to leaf
+					return L.GetConstraintLevel() < R.GetConstraintLevel();
+				});
+
 			for (FConstraintHandle* ConstraintHandle : SortedConstraintHandles)
 			{
 				ProjectPosition(Dt, ConstraintHandle->GetConstraintIndex(), 0, 1);
@@ -576,18 +602,19 @@ namespace Chaos
 			TPBDJointUtilities<T, d>::GetConditionedInverseMass(Particle0->M(), Particle0->I().GetDiagonal(), Particle1->M(), Particle1->I().GetDiagonal(), InvM0, InvM1, InvIL0, InvIL1, (T)0, Settings.MaxInertiaRatio);
 		}
 
-		// Freeze the closest to kinematic connection if there is a difference
-		if (Settings.bEnableVelocitySolve)
+		// Freeze the closest to kinematic connection (if one is closer than the other)
+		if (Settings.bEnableVelocitySolve && (Level0 != Level1))
 		{
+			const T FreezeFactor = (T)(NumIts - (It + 1)) / (T)NumIts;
 			if (Level0 < Level1)
 			{
-				InvM0 = 0;
-				InvIL0 = PMatrix<T, d, d>(0, 0, 0);
+				InvM0 = InvM0 * FreezeFactor * FreezeFactor;
+				InvIL0 = InvIL0 * FreezeFactor * FreezeFactor;
 			}
 			else if (Level1 < Level0)
 			{
-				InvM1 = 0;
-				InvIL1 = PMatrix<T, d, d>(0, 0, 0);
+				InvM1 = InvM1 * FreezeFactor * FreezeFactor;
+				InvIL1 = InvIL1 * FreezeFactor * FreezeFactor;
 			}
 		}
 
diff --git a/Engine/Source/Runtime/Experimental/Chaos/Private/Chaos/PBDRigidsEvolutionGBF.cpp b/Engine/Source/Runtime/Experimental/Chaos/Private/Chaos/PBDRigidsEvolutionGBF.cpp
--- a/Engine/Source/Runtime/Experimental/Chaos/Private/Chaos/PBDRigidsEvolutionGBF.cpp
+++ b/Engine/Source/Runtime/Experimental/Chaos/Private/Chaos/PBDRigidsEvolutionGBF.cpp
@@ -106,6 +106,9 @@ void TPBDRigidsEvolutionGBF<T, d>::AdvanceOneTimeStep(const T Dt, const T StepFr
 {
 	SCOPE_CYCLE_COUNTER(STAT_Evo_AdvanceOneTimeStep);
 
+	// @todo(ccaulfield): only need this when list has changed
+	ConstraintRules.StableSort();
+
 #if !UE_BUILD_SHIPPING
 	if (SerializeEvolution)
 	{
diff --git a/Engine/Source/Runtime/Experimental/Chaos/Public/Chaos/PBDConstraintRule.h b/Engine/Source/Runtime/Experimental/Chaos/Public/Chaos/PBDConstraintRule.h
--- a/Engine/Source/Runtime/Experimental/Chaos/Public/Chaos/PBDConstraintRule.h
+++ b/Engine/Source/Runtime/Experimental/Chaos/Public/Chaos/PBDConstraintRule.h
@@ -16,12 +16,25 @@ namespace Chaos
 	/**
 	 * Constraint Rules bind constraint collections to the evolution and provide their update algorithm.
 	 */
-	template<typename T, int d>
-	class CHAOS_API TPBDConstraintRule
+	class CHAOS_API FConstraintRule
 	{
 	public:
-		TPBDConstraintRule() {}
-		virtual ~TPBDConstraintRule() {}
+		FConstraintRule(int32 InPriority) : Priority(InPriority) {}
+		virtual ~FConstraintRule() {}
+
+		/** Determines the order in which constraints are resolved. Higher priority constraints override lower priority ones. */
+		int32 GetPriority() const { return Priority; }
+
+		/** Set the constraint resolution priority. Higher priority constraints override lower priority ones. */
+		void SetPriority(const int32 InPriority) { Priority = InPriority; }
+
+		friend bool operator<(const FConstraintRule& L, const FConstraintRule& R)
+		{
+			return L.GetPriority() < R.GetPriority();
+		}
+
+	protected:
+		int32 Priority;
 	};
 
 	/**
@@ -32,12 +45,12 @@ namespace Chaos
 	 * @see TPBDConstraintGraphRuleImpl
 	 */
 	template<typename T, int d>
-	class CHAOS_API TPBDConstraintGraphRule
+	class CHAOS_API TPBDConstraintGraphRule : public FConstraintRule
 	{
 	public:
 		typedef TPBDConstraintGraph<T, d> FConstraintGraph;
 
-		TPBDConstraintGraphRule() {}
+		TPBDConstraintGraphRule(int32 InPriority) : FConstraintRule(InPriority) {}
 		virtual ~TPBDConstraintGraphRule() {}
 
 		virtual void BindToGraph(FConstraintGraph& InContactGraph, uint32 InContainerId) {}
@@ -80,8 +93,9 @@ namespace Chaos
 		typedef T_CONSTRAINTS FConstraints;
 		typedef TPBDConstraintGraph<T, d> FConstraintGraph;
 
-		TPBDConstraintGraphRuleImpl(FConstraints& InConstraints)
-			: Constraints(InConstraints)
+		TPBDConstraintGraphRuleImpl(FConstraints& InConstraints, int32 InPriority)
+			: TPBDConstraintGraphRule<T, d>(InPriority)
+			, Constraints(InConstraints)
 			, ConstraintGraph(nullptr)
 		{
 		}
@@ -128,8 +142,8 @@ namespace Chaos
 		using FConstraintList = TArray<TConstraintHandle<T, d>*>;
 		using FConstraintGraph = typename Base::FConstraintGraph;
 
-		TPBDConstraintIslandRule(FConstraints& InConstraints)
-			: TPBDConstraintGraphRuleImpl<T_CONSTRAINTS, T, d>(InConstraints)
+		TPBDConstraintIslandRule(FConstraints& InConstraints, int32 InPriority = 0)
+			: TPBDConstraintGraphRuleImpl<T_CONSTRAINTS, T, d>(InConstraints, InPriority)
 		{
 		}
 
@@ -213,8 +227,8 @@ namespace Chaos
 		typedef typename Base::FConstraintGraph FConstraintGraph;
 		typedef TPBDConstraintColor<T, d> FConstraintColor;
 
-		TPBDConstraintColorRule(FConstraints& InConstraints, const int32 InPushOutIterations)
-			: TPBDConstraintGraphRuleImpl<T_CONSTRAINTS, T, d>(InConstraints)
+		TPBDConstraintColorRule(FConstraints& InConstraints, const int32 InPushOutIterations, int32 InPriority = 0)
+			: TPBDConstraintGraphRuleImpl<T_CONSTRAINTS, T, d>(InConstraints, InPriority)
 			, PushOutIterations(InPushOutIterations)
 		{
 		}
@@ -382,15 +396,15 @@ namespace Chaos
 		typedef TPBDConstraintGraph<T, d> FConstraintGraph;
 
 		template<typename T_CONSTRAINTS>
-		static TPBDConstraintIslandRule<T_CONSTRAINTS, T, d> CreateIslandRule(T_CONSTRAINTS& Constraints)
+		static TPBDConstraintIslandRule<T_CONSTRAINTS, T, d> CreateIslandRule(T_CONSTRAINTS& Constraints, const int32 Priority = 0)
 		{
-			return TPBDConstraintIslandRule<T_CONSTRAINTS, T, d>(Constraints);
+			return TPBDConstraintIslandRule<T_CONSTRAINTS, T, d>(Constraints, Priority);
 		}
 
 		template<typename T_CONSTRAINTS>
-		static TPBDConstraintColorRule<T_CONSTRAINTS, T, d> CreateColorRule(T_CONSTRAINTS& Constraints, const int32 InPushOutIterations)
+		static TPBDConstraintColorRule<T_CONSTRAINTS, T, d> CreateColorRule(T_CONSTRAINTS& Constraints, const int32 InPushOutIterations, const int32 Priority = 0)
 		{
-			return TPBDConstraintColorRule<T_CONSTRAINTS, T, d>(Constraints, InPushOutIterations);
+			return TPBDConstraintColorRule<T_CONSTRAINTS, T, d>(Constraints, InPushOutIterations, Priority);
 		}
 	};
 
diff --git a/Engine/Source/Runtime/Experimental/Chaos/Public/Chaos/PBDJointConstraintTypes.h b/Engine/Source/Runtime/Experimental/Chaos/Public/Chaos/PBDJointConstraintTypes.h
--- a/Engine/Source/Runtime/Experimental/Chaos/Public/Chaos/PBDJointConstraintTypes.h
+++ b/Engine/Source/Runtime/Experimental/Chaos/Public/Chaos/PBDJointConstraintTypes.h
@@ -76,6 +76,13 @@ namespace Chaos
 		TVector<EJointMotionType, d> AngularMotionTypes;
 		TVector<T, d> AngularLimits;
 
+		bool bSoftLinearLimitsEnabled;
+		bool bSoftTwistLimitsEnabled;
+		bool bSoftSwingLimitsEnabled;
+		T SoftLinearStiffness;
+		T SoftTwistStiffness;
+		T SoftSwingStiffness;
+
 		// @todo(ccaulfield): remove one of these
 		TRotation<T, d> AngularDriveTarget;
 		TVector<T, d> AngularDriveTargetAngles;
@@ -83,9 +90,7 @@ namespace Chaos
 		bool bAngularSLerpDriveEnabled;
 		bool bAngularTwistDriveEnabled;
 		bool bAngularSwingDriveEnabled;
-
 		T AngularDriveStiffness;
-		T AngularDriveDamping;
 	};
 
 
@@ -122,11 +127,12 @@ namespace Chaos
 		bool bEnableLinearLimits;
 		bool bEnableTwistLimits;
 		bool bEnableSwingLimits;
-		bool bEnablePositionCorrection;
 		bool bEnableDrives;
 		T Projection;
 		T Stiffness;
 		T DriveStiffness;
+		T SoftLinearStiffness;
+		T SoftAngularStiffness;
 		int32 PositionIterations;
 	};
 }
\ No newline at end of file
diff --git a/Engine/Source/Runtime/Experimental/Chaos/Public/Chaos/PBDJointConstraints.h b/Engine/Source/Runtime/Experimental/Chaos/Public/Chaos/PBDJointConstraints.h
--- a/Engine/Source/Runtime/Experimental/Chaos/Public/Chaos/PBDJointConstraints.h
+++ b/Engine/Source/Runtime/Experimental/Chaos/Public/Chaos/PBDJointConstraints.h
@@ -125,9 +125,10 @@ namespace Chaos
 
 		CHAOS_API void Apply(const T Dt, const TArray<FConstraintHandle*>& InConstraintHandles, const int32 It, const int32 NumIts);
 
-		// @todo(ccaulfield): remove  this
 		CHAOS_API void ApplyPushOut(const T Dt, const TArray<FConstraintHandle*>& InConstraintHandles);
 
+		CHAOS_API void ApplyProjection(const T Dt, const TArray<FConstraintHandle*>& InConstraintHandles);
+
 	protected:
 		using Base::GetConstraintIndex;
 		using Base::SetConstraintIndex;
