Fix strict CIS in PBDRigidsEvolution.cpp

diff --git a/Engine/Source/Runtime/Experimental/Chaos/Public/Chaos/PBDRigidsEvolutionGBF.h b/Engine/Source/Runtime/Experimental/Chaos/Public/Chaos/PBDRigidsEvolutionGBF.h
--- a/Engine/Source/Runtime/Experimental/Chaos/Public/Chaos/PBDRigidsEvolutionGBF.h
+++ b/Engine/Source/Runtime/Experimental/Chaos/Public/Chaos/PBDRigidsEvolutionGBF.h
@@ -180,6 +180,7 @@ protected:
 	using Base::CreateConstraintGraph;
 	using Base::CreateIslands;
 	using Base::ConstraintGraph;
+	using Base::ConstraintRules;
 	using Base::UpdateVelocities;
 	using Base::PhysicsMaterials;
 	using Base::ParticleDisableCount;
