Removed non-real types in Chaos core to prevent confusion and accidental use of incorrect types (if needed can explicitly use TVector<T, d>)

diff --git a/Engine/Source/Runtime/Experimental/ChaosCore/Public/Chaos/Core.h b/Engine/Source/Runtime/Experimental/ChaosCore/Public/Chaos/Core.h
--- a/Engine/Source/Runtime/Experimental/ChaosCore/Public/Chaos/Core.h
+++ b/Engine/Source/Runtime/Experimental/ChaosCore/Public/Chaos/Core.h
@@ -18,14 +18,6 @@ namespace Chaos
 	using FVectorR3 = TVector<FReal, 3>;
 	using FVectorR4 = TVector<FReal, 4>;
 
-	using FVectorF2 = TVector<float, 2>;
-	using FVectorF3 = TVector<float, 3>;
-	using FVectorF4 = TVector<float, 4>;
-
-	using FVectorD2 = TVector<double, 2>;
-	using FVectorD3 = TVector<double, 3>;
-	using FVectorD4 = TVector<double, 4>;
-
 	using FVectorI2 = TVector<int32, 2>;
 	using FVectorI3 = TVector<int32, 3>;
 
@@ -35,12 +27,4 @@ namespace Chaos
 	using FMatrixR22 = PMatrix<FReal, 2, 2>;
 	using FMatrixR33 = PMatrix<FReal, 3, 3>;
 	using FMatrixR44 = PMatrix<FReal, 4, 4>;
-
-	using FMatrixF22 = PMatrix<float, 2, 2>;
-	using FMatrixF33 = PMatrix<float, 3, 3>;
-	using FMatrixF44 = PMatrix<float, 4, 4>;
-
-	using FMatrixD22 = PMatrix<double, 2, 2>;
-	using FMatrixD33 = PMatrix<double, 3, 3>;
-	using FMatrixD44 = PMatrix<double, 4, 4>;
 }
\ No newline at end of file
