Bug fix, Color node was not being reset every frame

diff --git a/Engine/Source/Runtime/Experimental/Chaos/Private/Chaos/PBDConstraintColor.cpp b/Engine/Source/Runtime/Experimental/Chaos/Private/Chaos/PBDConstraintColor.cpp
--- a/Engine/Source/Runtime/Experimental/Chaos/Private/Chaos/PBDConstraintColor.cpp
+++ b/Engine/Source/Runtime/Experimental/Chaos/Private/Chaos/PBDConstraintColor.cpp
@@ -260,6 +260,12 @@ void TPBDConstraintColor<T, d>::ComputeContactGraph(const int32 Island, const FC
 template<typename T, int d>
 void TPBDConstraintColor<T, d>::InitializeColor(const FConstraintGraph& ConstraintGraph)
 {
+	// Reset the existing Nodes - so colors are all reset to zero
+	ParallelFor(Nodes.Num(), [&](int32 Index)
+	{
+		Nodes[Index].NextColor = 0;
+		Nodes[Index].UsedColors.Empty();
+	});
 
 	// The Number of nodes is large and fairly constant so persist rather than resetting every frame
 	if (Nodes.Num() != ConstraintGraph.Nodes.Num())
