Move mesh sculpting back to local position space to avoid numerical issues far from origin (keep bake of scale and rotation). Enable shift-modifier angle snap for rotation gizmo. Copy of Dev-editor CL 9898657 and 9898698

diff --git a/Engine/Plugins/Experimental/MeshModelingToolset/Source/MeshModelingTools/Private/DynamicMeshSculptTool.cpp b/Engine/Plugins/Experimental/MeshModelingToolset/Source/MeshModelingTools/Private/DynamicMeshSculptTool.cpp
--- a/Engine/Plugins/Experimental/MeshModelingToolset/Source/MeshModelingTools/Private/DynamicMeshSculptTool.cpp
+++ b/Engine/Plugins/Experimental/MeshModelingToolset/Source/MeshModelingTools/Private/DynamicMeshSculptTool.cpp
@@ -116,10 +116,11 @@ void UDynamicMeshSculptTool::Setup()
 	InitialTargetTransform = FTransform3d(ComponentTarget->GetWorldTransform());
 	// clamp scaling because if we allow zero-scale we cannot invert this transform on Accept
 	InitialTargetTransform.ClampMinimumScale(0.01);
+	FVector3d Translation = InitialTargetTransform.GetTranslation();
+	InitialTargetTransform.SetTranslation(FVector3d::Zero());
 	DynamicMeshComponent->ApplyTransform(InitialTargetTransform, false);
 	// since we moved to World coords there is not a current transform anymore.
-	// @todo: only remove scaling, keep rotation and translation in CurTargetTransform. This will improve numerical precision.
-	CurTargetTransform = FTransform3d::Identity();
+	CurTargetTransform = FTransform3d(Translation);
 	DynamicMeshComponent->SetWorldTransform((FTransform)CurTargetTransform);
 
 	// copy material if there is one
diff --git a/Engine/Source/Runtime/Experimental/InteractiveToolsFramework/Private/BaseGizmos/AxisAngleGizmo.cpp b/Engine/Source/Runtime/Experimental/InteractiveToolsFramework/Private/BaseGizmos/AxisAngleGizmo.cpp
--- a/Engine/Source/Runtime/Experimental/InteractiveToolsFramework/Private/BaseGizmos/AxisAngleGizmo.cpp
+++ b/Engine/Source/Runtime/Experimental/InteractiveToolsFramework/Private/BaseGizmos/AxisAngleGizmo.cpp
@@ -23,6 +23,7 @@ void UAxisAngleGizmo::Setup()
 
 	// Add default mouse input behavior
 	UClickDragInputBehavior* MouseBehavior = NewObject<UClickDragInputBehavior>();
+	MouseBehavior->Modifiers.RegisterModifier(SnapAngleModifierID, FInputDeviceState::IsShiftKeyDown);
 	MouseBehavior->Initialize(this);
 	MouseBehavior->SetDefaultPriority(FInputCapturePriority(FInputCapturePriority::DEFAULT_GIZMO_PRIORITY));
 	AddInputBehavior(MouseBehavior);
@@ -40,6 +41,15 @@ void UAxisAngleGizmo::Setup()
 }
 
 
+void UAxisAngleGizmo::OnUpdateModifierState(int ModifierID, bool bIsOn)
+{
+	if (ModifierID == SnapAngleModifierID)
+	{
+		bEnableSnapAngleModifier = bIsOn;
+	}
+}
+
+
 FInputRayHit UAxisAngleGizmo::CanBeginClickDragSequence(const FInputDeviceRay& PressPos)
 {
 	FInputRayHit GizmoHit;
@@ -96,6 +106,39 @@ void UAxisAngleGizmo::OnClickPress(const FInputDeviceRay& PressPos)
 }
 
 
+// these functions should move to a utility namespace
+
+static float SnapToIncrement(float fValue, float fIncrement, float offset = 0)
+{
+	if (!FMath::IsFinite(fValue))
+	{
+		return 0;
+	}
+	fValue -= offset;
+	float sign = FMath::Sign(fValue);
+	fValue = FMath::Abs(fValue);
+	int nInc = (int)(fValue / fIncrement);
+	float fRem = (float)fmod(fValue, fIncrement);
+	if (fRem > fIncrement / 2)
+	{
+		++nInc;
+	}
+	return sign * (float)nInc * fIncrement + offset;
+}
+
+
+static float SnapToNearbyIncrement(float fValue, float fIncrement, float fTolerance)
+{
+	float snapped = SnapToIncrement(fValue, fIncrement);
+	if (FMath::Abs(snapped - fValue) < fTolerance)
+	{
+		return snapped;
+	}
+	return fValue;
+}
+
+
+
 void UAxisAngleGizmo::OnClickDrag(const FInputDeviceRay& DragPos)
 {
 	check(bInInteraction);
@@ -117,6 +160,10 @@ void UAxisAngleGizmo::OnClickDrag(const FInputDeviceRay& DragPos)
 		RotationOrigin, RotationAxis, RotationPlaneX, RotationPlaneY);
 
 	float DeltaAngle = InteractionCurAngle - InteractionStartAngle;
+	if (bEnableSnapAngleModifier)
+	{
+		DeltaAngle = SnapToIncrement(DeltaAngle, 3.14159f / 4.0f);				// hardcoded 45-degree snap
+	}
 	float NewAngle = InitialTargetAngle + DeltaAngle;
 
 	AngleSource->SetParameter(NewAngle);
diff --git a/Engine/Source/Runtime/Experimental/InteractiveToolsFramework/Public/BaseGizmos/AxisAngleGizmo.h b/Engine/Source/Runtime/Experimental/InteractiveToolsFramework/Public/BaseGizmos/AxisAngleGizmo.h
--- a/Engine/Source/Runtime/Experimental/InteractiveToolsFramework/Public/BaseGizmos/AxisAngleGizmo.h
+++ b/Engine/Source/Runtime/Experimental/InteractiveToolsFramework/Public/BaseGizmos/AxisAngleGizmo.h
@@ -49,6 +49,9 @@ public:
 	virtual bool OnUpdateHover(const FInputDeviceRay& DevicePos) override;
 	virtual void OnEndHover() override;
 
+	// IModifierToggleBehaviorTarget implementation (inherited via IClickDragBehaviorTarget)
+	virtual void OnUpdateModifierState(int ModifierID, bool bIsOn) override;
+
 public:
 	UPROPERTY()
 	TScriptInterface<IGizmoAxisSource> AxisSource;
@@ -94,5 +97,9 @@ public:
 protected:
 	FVector LastHitPosition;
 	float InitialTargetAngle;
+
+	static const int SnapAngleModifierID = 1;		// identifier we associate with the shift key
+	bool bEnableSnapAngleModifier = false;			// flag we use to keep track of modifier state
+
 };
 
