Make OutNormal safer in GJKPenetrationTemp

diff --git a/Engine/Source/Runtime/Experimental/Chaos/Public/Chaos/GJK.h b/Engine/Source/Runtime/Experimental/Chaos/Public/Chaos/GJK.h
--- a/Engine/Source/Runtime/Experimental/Chaos/Public/Chaos/GJK.h
+++ b/Engine/Source/Runtime/Experimental/Chaos/Public/Chaos/GJK.h
@@ -867,7 +867,10 @@ namespace Chaos
 		if (GJKDistance(A, Segment, BToATM, SegmentDistance, OutPositionA, PositionBInB, Epsilon, MaxIts))
 		{
 			OutPositionB = BToATM.TransformPosition(PositionBInB);
-			OutNormal = (OutPositionB - OutPositionA) / SegmentDistance; // I think this is safe, if GJKDistance returned true...
+			OutNormal
+				= ensure(SegmentDistance > TNumericLimits<T>::Min())
+				? (OutPositionB - OutPositionA) / SegmentDistance
+				: TVector<T, 3>(0.f, 0.f, 1.f);
 			OutPositionB -= OutNormal * MarginB;
 			OutDistance = SegmentDistance - MarginB;
 
