First pass removing T,d from Chaos. Only ImplicitObject

diff --git a/Engine/Source/Runtime/Engine/Classes/PhysicsEngine/BodySetup.h b/Engine/Source/Runtime/Engine/Classes/PhysicsEngine/BodySetup.h
--- a/Engine/Source/Runtime/Engine/Classes/PhysicsEngine/BodySetup.h
+++ b/Engine/Source/Runtime/Engine/Classes/PhysicsEngine/BodySetup.h
@@ -42,8 +42,7 @@ namespace physx
 #if WITH_CHAOS
 namespace Chaos
 {
-	template<typename T, int d>
-	class TImplicitObject;
+	class FImplicitObject;
 
 	template <typename T>
 	class TTriangleMeshImplicitObject;
diff --git a/Engine/Source/Runtime/Engine/Classes/PhysicsEngine/ConvexElem.h b/Engine/Source/Runtime/Engine/Classes/PhysicsEngine/ConvexElem.h
--- a/Engine/Source/Runtime/Engine/Classes/PhysicsEngine/ConvexElem.h
+++ b/Engine/Source/Runtime/Engine/Classes/PhysicsEngine/ConvexElem.h
@@ -22,8 +22,7 @@ namespace physx
 
 namespace Chaos
 {
-	template <typename T, int d>
-	class TImplicitObject;
+	class FImplicitObject;
 
 	template <typename T, int d>
 	class TConvex;
diff --git a/Engine/Source/Runtime/Engine/Private/Collision/CollisionDebugDrawing.cpp b/Engine/Source/Runtime/Engine/Private/Collision/CollisionDebugDrawing.cpp
--- a/Engine/Source/Runtime/Engine/Private/Collision/CollisionDebugDrawing.cpp
+++ b/Engine/Source/Runtime/Engine/Private/Collision/CollisionDebugDrawing.cpp
@@ -250,7 +250,7 @@ void DrawGeomSweeps(const UWorld* InWorld, const FVector& Start, const FVector&
 
 #endif // WITH_PHYSX
 
-void DrawGeomOverlaps(const UWorld* InWorld, const Chaos::TImplicitObject<float, 3>& Geom, const FTransform& GeomPose, TArray<struct FOverlapResult>& Overlaps, float Lifetime)
+void DrawGeomOverlaps(const UWorld* InWorld, const Chaos::FImplicitObject& Geom, const FTransform& GeomPose, TArray<struct FOverlapResult>& Overlaps, float Lifetime)
 {
 	using namespace Chaos;
 	FVector Pos = GeomPose.GetLocation();
@@ -271,7 +271,7 @@ void DrawGeomOverlaps(const UWorld* InWorld, const Chaos::TImplicitObject<float,
 	}
 }
 
-void DrawGeomSweeps(const UWorld* InWorld, const FVector& Start, const FVector& End, const Chaos::TImplicitObject<float, 3>& Geom, const FQuat& Rotation, const TArray<FHitResult>& Hits, float Lifetime)
+void DrawGeomSweeps(const UWorld* InWorld, const FVector& Start, const FVector& End, const Chaos::FImplicitObject& Geom, const FQuat& Rotation, const TArray<FHitResult>& Hits, float Lifetime)
 {
 	using namespace Chaos;
 	if (const auto Box = Geom.template GetObject<TBox<float, 3>>())
diff --git a/Engine/Source/Runtime/Engine/Private/Collision/CollisionDebugDrawing.h b/Engine/Source/Runtime/Engine/Private/Collision/CollisionDebugDrawing.h
--- a/Engine/Source/Runtime/Engine/Private/Collision/CollisionDebugDrawing.h
+++ b/Engine/Source/Runtime/Engine/Private/Collision/CollisionDebugDrawing.h
@@ -18,5 +18,5 @@ void DrawGeomOverlaps(const UWorld* InWorld, const PxGeometry& PGeom, const PxTr
 void DrawGeomSweeps(const UWorld* InWorld, const FVector& Start, const FVector& End, const PxGeometry& PGeom, const PxQuat& Rotation, const TArray<FHitResult>& Hits, float Lifetime);
 #endif
 
-void DrawGeomOverlaps(const UWorld* InWorld, const Chaos::TImplicitObject<float, 3>& Geom, const FTransform& GeomPose, TArray<struct FOverlapResult>& Overlaps, float Lifetime);
-void DrawGeomSweeps(const UWorld* InWorld, const FVector& Start, const FVector& End, const Chaos::TImplicitObject<float,3>& Geom, const FQuat& Rotation, const TArray<FHitResult>& Hits, float Lifetime);
+void DrawGeomOverlaps(const UWorld* InWorld, const Chaos::FImplicitObject& Geom, const FTransform& GeomPose, TArray<struct FOverlapResult>& Overlaps, float Lifetime);
+void DrawGeomSweeps(const UWorld* InWorld, const FVector& Start, const FVector& End, const Chaos::FImplicitObject& Geom, const FQuat& Rotation, const TArray<FHitResult>& Hits, float Lifetime);
diff --git a/Engine/Source/Runtime/Engine/Private/PhysicsEngine/Experimental/ChaosDerivedData.cpp b/Engine/Source/Runtime/Engine/Private/PhysicsEngine/Experimental/ChaosDerivedData.cpp
--- a/Engine/Source/Runtime/Engine/Private/PhysicsEngine/Experimental/ChaosDerivedData.cpp
+++ b/Engine/Source/Runtime/Engine/Private/PhysicsEngine/Experimental/ChaosDerivedData.cpp
@@ -144,16 +144,15 @@ void FChaosDerivedDataCooker::BuildTriangleMeshes(TArray<TUniquePtr<Chaos::TTria
 
 template void FChaosDerivedDataCooker::BuildTriangleMeshes(TArray<TUniquePtr<Chaos::TTriangleMeshImplicitObject<float>>>& OutTriangleMeshes, const FCookBodySetupInfo& InParams);
 //#BGTODO When it's possible to build with doubles, re-enable this. (Currently at least TRigidTransform cannot build with double precision because we don't have a base transform implementation using them)
-//template void FChaosDerivedDataCooker::BuildTriangleMeshes(TArray<TUniquePtr<Chaos::TImplicitObject<float, 3>>>& OutTriangleMeshes, const FCookBodySetupInfo& InParams);
+//template void FChaosDerivedDataCooker::BuildTriangleMeshes(TArray<TUniquePtr<Chaos::FImplicitObject>>& OutTriangleMeshes, const FCookBodySetupInfo& InParams);
 
-template<typename Precision>
-void FChaosDerivedDataCooker::BuildConvexMeshes(TArray<TUniquePtr<Chaos::TImplicitObject<Precision, 3>>>& OutConvexMeshes, const FCookBodySetupInfo& InParams)
+void FChaosDerivedDataCooker::BuildConvexMeshes(TArray<TUniquePtr<Chaos::FImplicitObject>>& OutConvexMeshes, const FCookBodySetupInfo& InParams)
 {
-	auto BuildConvexFromVerts = [](TArray<TUniquePtr<Chaos::TImplicitObject<Precision, 3>>>& OutConvexes, const TArray<TArray<FVector>>& InMeshVerts, const bool bMirrored)
+	auto BuildConvexFromVerts = [](TArray<TUniquePtr<Chaos::FImplicitObject>>& OutConvexes, const TArray<TArray<FVector>>& InMeshVerts, const bool bMirrored)
 	{
 		for(const TArray<FVector>& HullVerts : InMeshVerts)
 		{
-			Chaos::TParticles<Precision, 3> ConvexParticles;
+			Chaos::TParticles<Chaos::FReal, 3> ConvexParticles;
 
 			const int32 NumHullVerts = HullVerts.Num();
 
@@ -165,7 +164,7 @@ void FChaosDerivedDataCooker::BuildConvexMeshes(TArray<TUniquePtr<Chaos::TImplic
 				ConvexParticles.X(VertIndex) = FVector(bMirrored ? -HullVert.X : HullVert.X, HullVert.Y, HullVert.Z);
 			}
 
-			OutConvexes.Emplace(new Chaos::TConvex<Precision, 3>(ConvexParticles));
+			OutConvexes.Emplace(new Chaos::TConvex<Chaos::FReal, 3>(ConvexParticles));
 		}
 	};
 
@@ -181,12 +180,8 @@ void FChaosDerivedDataCooker::BuildConvexMeshes(TArray<TUniquePtr<Chaos::TImplic
 
 }
 
-template void FChaosDerivedDataCooker::BuildConvexMeshes(TArray<TUniquePtr<Chaos::TImplicitObject<float, 3>>>& OutConvexMeshes, const FCookBodySetupInfo& InParams);
-//#BGTODO When it's possible to build with doubles, re-enable this. (Currently at least TRigidTransform cannot build with double precision because we don't have a base transform implementation using them)
-//template void FChaosDerivedDataCooker::BuildConvexMeshes(TArray<Chaos::TCollisionConvexMesh<double>>& OutTriangleMeshes, const FCookBodySetupInfo& InParams);
-
 template<typename Precision>
-void BuildSimpleShapes(TArray<TUniquePtr<Chaos::TImplicitObject<Precision, 3>>>& OutImplicits, UBodySetup* InSetup)
+void BuildSimpleShapes(TArray<TUniquePtr<Chaos::FImplicitObject>>& OutImplicits, UBodySetup* InSetup)
 {
 	check(InSetup);
 
@@ -222,7 +217,7 @@ void BuildSimpleShapes(TArray<TUniquePtr<Chaos::TImplicitObject<Precision, 3>>>&
 template<typename Precision>
 void FChaosDerivedDataCooker::BuildInternal(Chaos::FChaosArchive& Ar, FCookBodySetupInfo& InInfo)
 {
-	TArray<TUniquePtr<Chaos::TImplicitObject<Precision, 3>>> SimpleImplicits;
+	TArray<TUniquePtr<Chaos::FImplicitObject>> SimpleImplicits;
 	TArray<TUniquePtr<Chaos::TTriangleMeshImplicitObject<Precision>>> ComplexImplicits;
 
 	//BuildSimpleShapes(SimpleImplicits, Setup);
diff --git a/Engine/Source/Runtime/Engine/Private/PhysicsEngine/Experimental/ChaosDerivedData.h b/Engine/Source/Runtime/Engine/Private/PhysicsEngine/Experimental/ChaosDerivedData.h
--- a/Engine/Source/Runtime/Engine/Private/PhysicsEngine/Experimental/ChaosDerivedData.h
+++ b/Engine/Source/Runtime/Engine/Private/PhysicsEngine/Experimental/ChaosDerivedData.h
@@ -18,8 +18,7 @@ namespace Chaos
 {
 	class FChaosArchive;
 
-	template<typename T, int d>
-	class TImplicitObject;
+	class FImplicitObject;
 
 	template<typename T>
 	class TTriangleMeshImplicitObject;
@@ -72,8 +71,7 @@ private:
 	template<typename Precision>
 	void BuildTriangleMeshes(TArray<TUniquePtr<Chaos::TTriangleMeshImplicitObject<Precision>>>& OutTriangleMeshes, const FCookBodySetupInfo& InParams);
 
-	template<typename Precision>
-	void BuildConvexMeshes(TArray<TUniquePtr<Chaos::TImplicitObject<Precision, 3>>>& OutTriangleMeshes, const FCookBodySetupInfo& InParams);
+	void BuildConvexMeshes(TArray<TUniquePtr<Chaos::FImplicitObject>>& OutTriangleMeshes, const FCookBodySetupInfo& InParams);
 
 	UBodySetup* Setup;
 	FName RequestedFormat;
@@ -83,10 +81,6 @@ extern template void FChaosDerivedDataCooker::BuildTriangleMeshes(TArray<TUnique
 // #BGTODO When it's possible to build with doubles, re-enable this. (Currently at least TRigidTransform cannot build with double precision because we don't have a base transform implementation using them)
 //extern template void FChaosDerivedDataCooker::BuildTriangleMeshes(TArray<Chaos::TCollisionTriangleMesh<double>>& OutTriangleMeshes, const FCookBodySetupInfo& InParams);
 
-extern template void FChaosDerivedDataCooker::BuildConvexMeshes(TArray<TUniquePtr<Chaos::TImplicitObject<float, 3>>>& OutTriangleMeshes, const FCookBodySetupInfo& InParams);
-// #BGTODO When it's possible to build with doubles, re-enable this. (Currently at least TRigidTransform cannot build with double precision because we don't have a base transform implementation using them)
-//extern template void FChaosDerivedDataCooker::BuildConvexMeshes(TArray<Chaos::TCollisionConvexMesh<double>>& OutTriangleMeshes, const FCookBodySetupInfo& InParams);
-
 extern template void FChaosDerivedDataCooker::BuildInternal<float>(Chaos::FChaosArchive& Ar, FCookBodySetupInfo& InInfo);
 // #BGTODO When it's possible to build with doubles, re-enable this. (Currently at least TRigidTransform cannot build with double precision because we don't have a base transform implementation using them)
 //extern template void FChaosDerivedDataCooker::BuildInternal<float>(FArchive& Ar, FCookBodySetupInfo& InInfo);
diff --git a/Engine/Source/Runtime/Engine/Private/PhysicsEngine/Experimental/ChaosInterfaceUtils.cpp b/Engine/Source/Runtime/Engine/Private/PhysicsEngine/Experimental/ChaosInterfaceUtils.cpp
--- a/Engine/Source/Runtime/Engine/Private/PhysicsEngine/Experimental/ChaosInterfaceUtils.cpp
+++ b/Engine/Source/Runtime/Engine/Private/PhysicsEngine/Experimental/ChaosInterfaceUtils.cpp
@@ -82,7 +82,7 @@ namespace ChaosInterface
 	}
 
 	template<class PHYSX_MESH>
-	TUniquePtr<Chaos::TImplicitObject<float, 3>> ConvertPhysXMeshToLevelset(const PHYSX_MESH* PhysXMesh, const FVector& Scale)
+	TUniquePtr<Chaos::FImplicitObject> ConvertPhysXMeshToLevelset(const PHYSX_MESH* PhysXMesh, const FVector& Scale)
 	{
 #if WITH_CHAOS && !WITH_CHAOS_NEEDS_TO_BE_FIXED
 		TArray<Chaos::TVector<int32, 3>> CollisionMeshElements = GetMeshElements(PhysXMesh);
@@ -99,7 +99,7 @@ namespace ChaosInterface
 			BoundingBox.GrowToInclude(CollisionMeshParticles.X(j));
 		}
 #if FORCE_ANALYTICS
-		return TUniquePtr<Chaos::TImplicitObject<float, 3>>(new Chaos::TBox<float, 3>(BoundingBox));
+		return TUniquePtr<Chaos::FImplicitObject>(new Chaos::TBox<float, 3>(BoundingBox));
 #else
 		int32 MaxAxisSize = 10;
 		int32 MaxAxis;
@@ -122,11 +122,11 @@ namespace ChaosInterface
 		Counts[2] = Counts[2] < 1 ? 1 : Counts[2];
 		Chaos::TUniformGrid<float, 3> Grid(BoundingBox.Min(), BoundingBox.Max(), Counts, 1);
 		Chaos::TTriangleMesh<float> CollisionMesh(MoveTemp(CollisionMeshElements));
-		return TUniquePtr<Chaos::TImplicitObject<float, 3>>(new Chaos::TLevelSet<float, 3>(Grid, CollisionMeshParticles, CollisionMesh));
+		return TUniquePtr<Chaos::FImplicitObject>(new Chaos::TLevelSet<float, 3>(Grid, CollisionMeshParticles, CollisionMesh));
 #endif
 
 #else
-		return TUniquePtr<Chaos::TImplicitObject<float, 3>>();
+		return TUniquePtr<Chaos::FImplicitObject>();
 #endif // !WITH_CHAOS_NEEDS_TO_BE_FIXED
 
 	}
@@ -134,13 +134,13 @@ namespace ChaosInterface
 #endif
 
 
-	void CreateGeometry(const FGeometryAddParams& InParams, TArray<TUniquePtr<Chaos::TImplicitObject<float, 3>>>& OutGeoms, Chaos::TShapesArray<float, 3>& OutShapes)
+	void CreateGeometry(const FGeometryAddParams& InParams, TArray<TUniquePtr<Chaos::FImplicitObject>>& OutGeoms, Chaos::TShapesArray<float, 3>& OutShapes)
 	{
 		const FVector& Scale = InParams.Scale;
-		TArray<TUniquePtr<Chaos::TImplicitObject<float, 3>>>& Geoms = OutGeoms;
+		TArray<TUniquePtr<Chaos::FImplicitObject>>& Geoms = OutGeoms;
 		Chaos::TShapesArray<float, 3>& Shapes = OutShapes;
 
-		auto NewShapeHelper = [&InParams](Chaos::TSerializablePtr<Chaos::TImplicitObject<float, 3>> InGeom, bool bComplexShape = false)
+		auto NewShapeHelper = [&InParams](Chaos::TSerializablePtr<Chaos::FImplicitObject> InGeom, bool bComplexShape = false)
 		{
 			auto NewShape = Chaos::TPerShapeData<float, 3>::CreatePerShapeData();
 			NewShape->Geometry = InGeom;
@@ -171,11 +171,11 @@ namespace ChaosInterface
 			HalfExtents.Z = FMath::Max(HalfExtents.Z, KINDA_SMALL_NUMBER);
 
 			// TBox can handle translations internally but if we have a rotation we need to wrap it in a transform
-			TUniquePtr<Chaos::TImplicitObject<float, 3>> Implicit;
+			TUniquePtr<Chaos::FImplicitObject> Implicit;
 			if (!BoxTransform.GetRotation().IsIdentity())
 			{
 				auto ImplicitBox = MakeUnique<Chaos::TBox<float, 3>>(-HalfExtents, HalfExtents);
-				Implicit = TUniquePtr<Chaos::TImplicitObject<float, 3>>(new Chaos::TImplicitObjectTransformed<float, 3>(MoveTemp(ImplicitBox), BoxTransform));
+				Implicit = TUniquePtr<Chaos::FImplicitObject>(new Chaos::TImplicitObjectTransformed<float, 3>(MoveTemp(ImplicitBox), BoxTransform));
 			}
 			else
 			{
@@ -250,7 +250,7 @@ namespace ChaosInterface
 			{
 				//if (!ConvexTransform.GetTranslation().IsNearlyZero() || !ConvexTransform.GetRotation().IsIdentity())
 				//{
-				//	TUniquePtr<Chaos::TImplicitObject<float, 3>> TransformImplicit = TUniquePtr<Chaos::TImplicitObject<float, 3>>(new Chaos::TImplicitObjectTransformed<float, 3>(MakeSerializable(ConvexImplicit), ConvexTransform));
+				//	TUniquePtr<Chaos::FImplicitObject> TransformImplicit = TUniquePtr<Chaos::FImplicitObject>(new Chaos::TImplicitObjectTransformed<float, 3>(MakeSerializable(ConvexImplicit), ConvexTransform));
 				//	TUniquePtr<Chaos::TImplicitObjectScaled<float, 3, false>> Implicit = MakeUnique<Chaos::TImplicitObjectScaled<float, 3, false>>(MoveTemp(TransformImplicit), Scale);
 				//	auto NewShape = NewShapeHelper(MakeSerializable(Implicit));
 				//	Shapes.Emplace(MoveTemp(NewShape));
diff --git a/Engine/Source/Runtime/Engine/Private/PhysicsEngine/Experimental/PhysInterface_Chaos.cpp b/Engine/Source/Runtime/Engine/Private/PhysicsEngine/Experimental/PhysInterface_Chaos.cpp
--- a/Engine/Source/Runtime/Engine/Private/PhysicsEngine/Experimental/PhysInterface_Chaos.cpp
+++ b/Engine/Source/Runtime/Engine/Private/PhysicsEngine/Experimental/PhysInterface_Chaos.cpp
@@ -82,7 +82,7 @@ ECollisionShapeType GetGeometryType(const Chaos::TPerShapeData<float, 3>& Shape)
 	return GetType(*Shape.Geometry);
 }
 
-const Chaos::TImplicitObject<float, 3>& FPhysicsShapeReference_Chaos::GetGeometry() const
+const Chaos::FImplicitObject& FPhysicsShapeReference_Chaos::GetGeometry() const
 {
 	check(IsValid()); return *Shape->Geometry;
 }
@@ -1075,7 +1075,7 @@ const FBodyInstance* FPhysInterface_Chaos::ShapeToOriginalBodyInstance(const FBo
 
 void FPhysInterface_Chaos::AddGeometry(FPhysicsActorHandle& InActor, const FGeometryAddParams& InParams, TArray<FPhysicsShapeHandle>* OutOptShapes)
 {
-	TArray<TUniquePtr<Chaos::TImplicitObject<float, 3>>> Geoms;
+	TArray<TUniquePtr<Chaos::FImplicitObject>> Geoms;
 	Chaos::TShapesArray<float, 3> Shapes;
 	ChaosInterface::CreateGeometry(InParams, Geoms, Shapes);
 
@@ -1121,7 +1121,7 @@ ECollisionShapeType FPhysicsGeometryCollection_Chaos::GetType() const
 	return GetImplicitType(Geom);
 }
 
-const Chaos::TImplicitObject<float, 3>& FPhysicsGeometryCollection_Chaos::GetGeometry() const
+const Chaos::FImplicitObject& FPhysicsGeometryCollection_Chaos::GetGeometry() const
 {
 	return Geom;
 }
@@ -1243,7 +1243,7 @@ void CalculateMassPropertiesOfImplicitType(
 	TArray< Chaos::TMassProperties<float, 3> > & MassProperties,
 	float & TotalMass,
 	const Chaos::TRigidTransform<float, 3> & WorldTransform,
-	const Chaos::TImplicitObject<float, 3>* ImplicitObject, 
+	const Chaos::FImplicitObject* ImplicitObject, 
 	float InDensityKGPerCM)
 {
 	// WIP
@@ -1374,7 +1374,7 @@ void CalculateMassPropertiesOfImplicitType(
 		{
 			const Chaos::TImplicitObjectUnion<float,3> * ImplicitUnion = ImplicitObject->template GetObject<Chaos::TImplicitObjectUnion<float,3>>();
 
-			for (const TUniquePtr < Chaos::TImplicitObject<float,3> > & ImplicitSubObject : ImplicitUnion->GetObjects())
+			for (const TUniquePtr < Chaos::FImplicitObject > & ImplicitSubObject : ImplicitUnion->GetObjects())
 			{
 				CalculateMassPropertiesOfImplicitType(MassProperties, TotalMass, WorldTransform, ImplicitSubObject.Get(), InDensityKGPerCM);
 			}
@@ -1396,7 +1396,7 @@ void FPhysInterface_Chaos::CalculateMassPropertiesFromShapeCollection(physx::PxM
 	{
 		if (const Chaos::TPerShapeData<float, 3>* Shape = ShapeHandle.Shape)
 		{
-			if (const Chaos::TImplicitObject<float, 3> * ImplicitObject = Shape->Geometry.Get())
+			if (const Chaos::FImplicitObject * ImplicitObject = Shape->Geometry.Get())
 			{
 				FTransform WorldTransform(ShapeHandle.ActorRef->R(), ShapeHandle.ActorRef->X());
 				CalculateMassPropertiesOfImplicitType(MassProperties, TotalMass, WorldTransform, ImplicitObject, InDensityKGPerCM);
@@ -1614,7 +1614,7 @@ bool FPhysInterface_Chaos::Sweep_Geom(FHitResult& OutHit, const FBodyInstance* I
 	return bSweepHit;
 }
 
-bool Overlap_GeomInternal(const FBodyInstance* InInstance, const Chaos::TImplicitObject<float, 3>& InGeom, const FTransform& GeomTransform, FMTDResult* OutOptResult)
+bool Overlap_GeomInternal(const FBodyInstance* InInstance, const Chaos::FImplicitObject& InGeom, const FTransform& GeomTransform, FMTDResult* OutOptResult)
 {
 	const FBodyInstance* TargetInstance = InInstance->WeldParent ? InInstance->WeldParent : InInstance;
 	Chaos::TGeometryParticle<float, 3>* RigidBody = TargetInstance->ActorHandle;
diff --git a/Engine/Source/Runtime/Engine/Private/PhysicsEngine/ImmediatePhysics/ImmediatePhysicsChaos/ImmediatePhysicsActorHandle_Chaos.cpp b/Engine/Source/Runtime/Engine/Private/PhysicsEngine/ImmediatePhysics/ImmediatePhysicsChaos/ImmediatePhysicsActorHandle_Chaos.cpp
--- a/Engine/Source/Runtime/Engine/Private/PhysicsEngine/ImmediatePhysics/ImmediatePhysicsChaos/ImmediatePhysicsActorHandle_Chaos.cpp
+++ b/Engine/Source/Runtime/Engine/Private/PhysicsEngine/ImmediatePhysics/ImmediatePhysicsChaos/ImmediatePhysicsActorHandle_Chaos.cpp
@@ -138,7 +138,7 @@ namespace ImmediatePhysics_Chaos
 	}
 
 
-	bool CreateGeometry(FBodyInstance* BodyInstance, const FVector& Scale, float& OutMass, Chaos::TVector<float, 3>& OutInertia, Chaos::TRigidTransform<float, 3>& OutCoMTransform, TUniquePtr<Chaos::TImplicitObject<FReal, Dimensions>>& OutGeom, TArray<TUniquePtr<Chaos::TPerShapeData<float, 3>>>& OutShapes)
+	bool CreateGeometry(FBodyInstance* BodyInstance, const FVector& Scale, float& OutMass, Chaos::TVector<float, 3>& OutInertia, Chaos::TRigidTransform<float, 3>& OutCoMTransform, TUniquePtr<Chaos::FImplicitObject>& OutGeom, TArray<TUniquePtr<Chaos::TPerShapeData<float, 3>>>& OutShapes)
 	{
 		UBodySetup* BodySetup = BodyInstance->BodySetup.Get();
 
@@ -175,7 +175,7 @@ namespace ImmediatePhysics_Chaos
 		AddParams.ChaosTriMeshes = MakeArrayView(BodySetup->ChaosTriMeshes);
 #endif
 
-		TArray<TUniquePtr<Chaos::TImplicitObject<float, 3>>> Geoms;
+		TArray<TUniquePtr<Chaos::FImplicitObject>> Geoms;
 		TArray<TUniquePtr<Chaos::TPerShapeData<float, 3>>, TInlineAllocator<1>> Shapes;
 		ChaosInterface::CreateGeometry(AddParams, Geoms, Shapes);
 
diff --git a/Engine/Source/Runtime/Engine/Public/Physics/Experimental/ChaosDerivedDataReader.h b/Engine/Source/Runtime/Engine/Public/Physics/Experimental/ChaosDerivedDataReader.h
--- a/Engine/Source/Runtime/Engine/Public/Physics/Experimental/ChaosDerivedDataReader.h
+++ b/Engine/Source/Runtime/Engine/Public/Physics/Experimental/ChaosDerivedDataReader.h
@@ -6,8 +6,7 @@
 #include "PhysicsEngine/BodySetup.h"
 namespace Chaos
 {
-	template<typename T, int d>
-	class TImplicitObject;
+	class FImplicitObject;
 
 	template <typename T>
 	class TTriangleMeshImplicitObject;
diff --git a/Engine/Source/Runtime/Engine/Public/Physics/Experimental/ChaosInterfaceUtils.h b/Engine/Source/Runtime/Engine/Public/Physics/Experimental/ChaosInterfaceUtils.h
--- a/Engine/Source/Runtime/Engine/Public/Physics/Experimental/ChaosInterfaceUtils.h
+++ b/Engine/Source/Runtime/Engine/Public/Physics/Experimental/ChaosInterfaceUtils.h
@@ -13,5 +13,5 @@ namespace ChaosInterface
 	/**
 	 * Create the Chaos Geometry based on the geometry parameters.
 	 */
-	void CreateGeometry(const FGeometryAddParams& InParams, TArray<TUniquePtr<Chaos::TImplicitObject<float, 3>>>& OutGeoms, TArray<TUniquePtr<Chaos::TPerShapeData<float, 3>>, TInlineAllocator<1>>& OutShapes);
+	void CreateGeometry(const FGeometryAddParams& InParams, TArray<TUniquePtr<Chaos::FImplicitObject>>& OutGeoms, TArray<TUniquePtr<Chaos::TPerShapeData<float, 3>>, TInlineAllocator<1>>& OutShapes);
 }
diff --git a/Engine/Source/Runtime/Engine/Public/Physics/Experimental/PhysInterface_Chaos.h b/Engine/Source/Runtime/Engine/Public/Physics/Experimental/PhysInterface_Chaos.h
--- a/Engine/Source/Runtime/Engine/Public/Physics/Experimental/PhysInterface_Chaos.h
+++ b/Engine/Source/Runtime/Engine/Public/Physics/Experimental/PhysInterface_Chaos.h
@@ -37,8 +37,7 @@ namespace Chaos
 	template <typename T, int>
 	class TBVHParticles;
 
-	template <typename T, int>
-	class TImplicitObject;
+	class FImplicitObject;
 
 	template <typename T, int>
 	class TPBDRigidParticles;
@@ -99,7 +98,7 @@ public:
 	bool Equals(const FPhysicsShapeReference_Chaos& Other) const { return Shape == Other.Shape; }
     bool operator==(const FPhysicsShapeReference_Chaos& Other) const { return Equals(Other); }
 	
-	const Chaos::TImplicitObject<float, 3>& GetGeometry() const;
+	const Chaos::FImplicitObject& GetGeometry() const;
 
 
 	Chaos::TPerShapeData<float, 3>* Shape;
@@ -143,7 +142,7 @@ struct ENGINE_API FPhysicsGeometryCollection_Chaos
 	~FPhysicsGeometryCollection_Chaos();
 
 	ECollisionShapeType GetType() const;
-	const Chaos::TImplicitObject<float, 3>& GetGeometry() const;
+	const Chaos::FImplicitObject& GetGeometry() const;
 	const Chaos::TBox<float, 3>& GetBoxGeometry() const;
 	const Chaos::TSphere<float, 3>&  GetSphereGeometry() const;
 	const Chaos::TCapsule<float>&  GetCapsuleGeometry() const;
@@ -154,7 +153,7 @@ private:
 	friend class FPhysInterface_Chaos;
 	explicit FPhysicsGeometryCollection_Chaos(const FPhysicsShapeReference_Chaos& InShape);
 
-	const Chaos::TImplicitObject<float, 3>& Geom;
+	const Chaos::FImplicitObject& Geom;
 };
 
 
@@ -365,7 +364,7 @@ public:
 	// Shape interface functions
 	static FPhysicsShapeHandle CreateShape(physx::PxGeometry* InGeom, bool bSimulation = true, bool bQuery = true, UPhysicalMaterial* InSimpleMaterial = nullptr, TArray<UPhysicalMaterial*>* InComplexMaterials = nullptr);
 	
-	static void CreateGeometry(const FGeometryAddParams& InParams, TArray<TUniquePtr<Chaos::TImplicitObject<float, 3>>>& OutGeoms, Chaos::TShapesArray<float, 3>& OutShapes, TArray<FPhysicsShapeHandle>* OutOptShapes);
+	static void CreateGeometry(const FGeometryAddParams& InParams, TArray<TUniquePtr<Chaos::FImplicitObject>>& OutGeoms, Chaos::TShapesArray<float, 3>& OutShapes, TArray<FPhysicsShapeHandle>* OutOptShapes);
 	static void AddGeometry(FPhysicsActorHandle& InActor, const FGeometryAddParams& InParams, TArray<FPhysicsShapeHandle>* OutOptShapes = nullptr);
 	static FPhysicsShapeHandle CloneShape(const FPhysicsShapeHandle& InShape);
 	static FPhysicsGeometryCollection_Chaos GetGeometryCollection(const FPhysicsShapeHandle& InShape);
@@ -399,7 +398,7 @@ public:
 };
 
 /*
-FORCEINLINE ECollisionShapeType GetType(const Chaos::TImplicitObject<float, 3>& Geom)
+FORCEINLINE ECollisionShapeType GetType(const Chaos::FImplicitObject& Geom)
 {
 	if (Geom.GetType() == Chaos::ImplicitObjectType::Box)
 	{
@@ -430,12 +429,12 @@ FORCEINLINE float GetHalfHeight(const Chaos::TCapsule<float>& Capsule)
 }
 */
 
-FORCEINLINE void DrawOverlappingTris(const UWorld* World, const FLocationHit& Hit, const Chaos::TImplicitObject<float, 3>& Geom, const FTransform& QueryTM)
+FORCEINLINE void DrawOverlappingTris(const UWorld* World, const FLocationHit& Hit, const Chaos::FImplicitObject& Geom, const FTransform& QueryTM)
 {
 	//TODO_SQ_IMPLEMENTATION
 }
 
-FORCEINLINE void ComputeZeroDistanceImpactNormalAndPenetration(const UWorld* World, const FLocationHit& Hit, const Chaos::TImplicitObject<float, 3>& Geom, const FTransform& QueryTM, FHitResult& OutResult)
+FORCEINLINE void ComputeZeroDistanceImpactNormalAndPenetration(const UWorld* World, const FLocationHit& Hit, const Chaos::FImplicitObject& Geom, const FTransform& QueryTM, FHitResult& OutResult)
 {
 	//TODO_SQ_IMPLEMENTATION
 }
diff --git a/Engine/Source/Runtime/Engine/Public/Physics/ImmediatePhysics/ImmediatePhysicsChaos/ImmediatePhysicsActorHandle_Chaos.h b/Engine/Source/Runtime/Engine/Public/Physics/ImmediatePhysics/ImmediatePhysicsChaos/ImmediatePhysicsActorHandle_Chaos.h
--- a/Engine/Source/Runtime/Engine/Public/Physics/ImmediatePhysics/ImmediatePhysicsChaos/ImmediatePhysicsActorHandle_Chaos.h
+++ b/Engine/Source/Runtime/Engine/Public/Physics/ImmediatePhysics/ImmediatePhysicsChaos/ImmediatePhysicsActorHandle_Chaos.h
@@ -125,7 +125,7 @@ namespace ImmediatePhysics_Chaos
 
 		Chaos::TPBDRigidsEvolutionGBF<FReal, Dimensions>* Evolution;
 		Chaos::TGeometryParticleHandle<FReal, Dimensions>* ParticleHandle;
-		TUniquePtr<Chaos::TImplicitObject<float, 3>> Geometry;
+		TUniquePtr<Chaos::FImplicitObject> Geometry;
 		TArray<TUniquePtr<Chaos::TPerShapeData<float, 3>>> Shapes;
 		FTransform ActorToCoMTransform;
 		int32 Level;
diff --git a/Engine/Source/Runtime/Engine/Public/Physics/ImmediatePhysics/ImmediatePhysicsChaos/ImmediatePhysicsCore_Chaos.h b/Engine/Source/Runtime/Engine/Public/Physics/ImmediatePhysics/ImmediatePhysicsChaos/ImmediatePhysicsCore_Chaos.h
--- a/Engine/Source/Runtime/Engine/Public/Physics/ImmediatePhysics/ImmediatePhysicsChaos/ImmediatePhysicsCore_Chaos.h
+++ b/Engine/Source/Runtime/Engine/Public/Physics/ImmediatePhysics/ImmediatePhysicsChaos/ImmediatePhysicsCore_Chaos.h
@@ -9,7 +9,7 @@
 
 namespace Chaos
 {
-	template<typename T, int D> class TImplicitObject;
+	class FImplicitObject;
 	template<typename T, int D> struct TKinematicGeometryParticleParameters;
 	template<typename T, int D> class TKinematicTarget;
 	template<typename P, typename T, int D> class TPBDConstraintIslandRule;
diff --git a/Engine/Source/Runtime/Engine/Public/Physics/PhysicsInterfaceTypes.h b/Engine/Source/Runtime/Engine/Public/Physics/PhysicsInterfaceTypes.h
--- a/Engine/Source/Runtime/Engine/Public/Physics/PhysicsInterfaceTypes.h
+++ b/Engine/Source/Runtime/Engine/Public/Physics/PhysicsInterfaceTypes.h
@@ -16,8 +16,7 @@ namespace Chaos
 	template<typename T, int d>
 	class TImplicitObjectUnion;
 
-	template<typename T, int d>
-	class TImplicitObject;
+	class FImplicitObject;
 
 	template <typename T>
 	class TTriangleMeshImplicitObject;
diff --git a/Engine/Source/Runtime/Experimental/Chaos/Private/Chaos/ChaosArchive.cpp b/Engine/Source/Runtime/Experimental/Chaos/Private/Chaos/ChaosArchive.cpp
--- a/Engine/Source/Runtime/Experimental/Chaos/Private/Chaos/ChaosArchive.cpp
+++ b/Engine/Source/Runtime/Experimental/Chaos/Private/Chaos/ChaosArchive.cpp
@@ -6,9 +6,9 @@
 namespace Chaos
 {
 
-void FChaosArchive::SerializeLegacy(TUniquePtr<TImplicitObject<float, 3>>& Obj)
+void FChaosArchive::SerializeLegacy(TUniquePtr<FImplicitObject>& Obj)
 {
-	TImplicitObject<float, 3>::SerializeLegacyHelper(InnerArchive, Obj);
+	FImplicitObject::SerializeLegacyHelper(InnerArchive, Obj);
 }
 
 #if CHAOS_MEMORY_TRACKING
diff --git a/Engine/Source/Runtime/Experimental/Chaos/Private/Chaos/ChaosDebugDraw.cpp b/Engine/Source/Runtime/Experimental/Chaos/Private/Chaos/ChaosDebugDraw.cpp
--- a/Engine/Source/Runtime/Experimental/Chaos/Private/Chaos/ChaosDebugDraw.cpp
+++ b/Engine/Source/Runtime/Experimental/Chaos/Private/Chaos/ChaosDebugDraw.cpp
@@ -46,7 +46,7 @@ namespace Chaos
 #if CHAOS_DEBUG_DRAW
 
 		template <typename T, int d>
-		void DrawShapesImpl(const TRigidTransform<float, 3>& ShapeTransform, const TImplicitObject<T, d>* Shape, FColor Color)
+		void DrawShapesImpl(const TRigidTransform<float, 3>& ShapeTransform, const FImplicitObject* Shape, FColor Color)
 		{
 			// @todo(ccaulfield): handle scale throughout
 			switch (Shape->GetType(false))
diff --git a/Engine/Source/Runtime/Experimental/Chaos/Private/Chaos/GeometryParticles.cpp b/Engine/Source/Runtime/Experimental/Chaos/Private/Chaos/GeometryParticles.cpp
--- a/Engine/Source/Runtime/Experimental/Chaos/Private/Chaos/GeometryParticles.cpp
+++ b/Engine/Source/Runtime/Experimental/Chaos/Private/Chaos/GeometryParticles.cpp
@@ -8,7 +8,7 @@
 namespace Chaos
 {
 	template <typename T, int d>
-	void UpdateShapesArrayFromGeometry(TShapesArray<T, d>& ShapesArray, TSerializablePtr<TImplicitObject<T, d>> Geometry)
+	void UpdateShapesArrayFromGeometry(TShapesArray<T, d>& ShapesArray, TSerializablePtr<FImplicitObject> Geometry)
 	{
 		if(Geometry)
 		{
@@ -114,5 +114,5 @@ namespace Chaos
 	template class TGeometryParticlesImp<float, 3, EGeometryParticlesSimType::RigidBodySim>;
 	template class TGeometryParticlesImp<float, 3, EGeometryParticlesSimType::Other>;
 	template class TPerShapeData<float, 3>;
-	template void UpdateShapesArrayFromGeometry(TShapesArray<float, 3>& ShapesArray, TSerializablePtr<TImplicitObject<float, 3>> Geometry);
+	template void UpdateShapesArrayFromGeometry(TShapesArray<float, 3>& ShapesArray, TSerializablePtr<FImplicitObject> Geometry);
 }
diff --git a/Engine/Source/Runtime/Experimental/Chaos/Private/Chaos/HeightField.cpp b/Engine/Source/Runtime/Experimental/Chaos/Private/Chaos/HeightField.cpp
--- a/Engine/Source/Runtime/Experimental/Chaos/Private/Chaos/HeightField.cpp
+++ b/Engine/Source/Runtime/Experimental/Chaos/Private/Chaos/HeightField.cpp
@@ -193,7 +193,7 @@ namespace Chaos
 	{
 	public:
 
-		THeightfieldSweepVisitor(const typename THeightField<T>::FDataType* InData, const TImplicitObject<T, 3>& InQueryGeom, const TRigidTransform<T, 3>& InStartTM, const TVector<T, 3>& InDir, const T InThickness)
+		THeightfieldSweepVisitor(const typename THeightField<T>::FDataType* InData, const FImplicitObject& InQueryGeom, const TRigidTransform<T, 3>& InStartTM, const TVector<T, 3>& InDir, const T InThickness)
 			: OutTime(TNumericLimits<T>::Max())
 			, OutFaceIndex(INDEX_NONE)
 			, HfData(InData)
@@ -259,7 +259,7 @@ namespace Chaos
 
 		const typename THeightField<T>::FDataType* HfData;
 		const TRigidTransform<T, 3> StartTM;
-		const TImplicitObject<T, 3>& OtherGeom;
+		const FImplicitObject& OtherGeom;
 		const TVector<T, 3>& Dir;
 		const T Thickness;
 
@@ -409,7 +409,7 @@ namespace Chaos
 
 	template <typename T>
 	THeightField<T>::THeightField(TArray<T>&& Height, int32 NumRows, int32 NumCols, const TVector<T, 3>& InScale)
-		: TImplicitObject<T, 3>(EImplicitObject::HasBoundingBox, ImplicitObjectType::HeightField)
+		: FImplicitObject(EImplicitObject::HasBoundingBox, ImplicitObjectType::HeightField)
 	{
 		BuildGeomData<T, T>(MakeArrayView(Height), NumRows, NumCols, TVector<T, 3>(1), [](const T InVal) {return InVal; }, GeomData, LocalBounds);
 		CalcBounds();
@@ -418,7 +418,7 @@ namespace Chaos
 
 	template<typename T>
 	Chaos::THeightField<T>::THeightField(TArrayView<const uint16> InHeights, int32 InNumRows, int32 InNumCols, const TVector<T, 3>& InScale)
-		: TImplicitObject<T, 3>(EImplicitObject::HasBoundingBox, ImplicitObjectType::HeightField)
+		: FImplicitObject(EImplicitObject::HasBoundingBox, ImplicitObjectType::HeightField)
 	{
 		TUniqueFunction<T(const uint16)> ConversionFunc = [](const T InVal) -> float
 		{
@@ -1002,7 +1002,7 @@ namespace Chaos
 
 
 	template <typename T>
-	bool THeightField<T>::OverlapGeom(const TImplicitObject<T, 3>& QueryGeom, const TRigidTransform<T, 3>& QueryTM, const T Thickness) const
+	bool THeightField<T>::OverlapGeom(const FImplicitObject& QueryGeom, const TRigidTransform<T, 3>& QueryTM, const T Thickness) const
 	{
 		auto OverlapTriangle = [&](const TVector<T, 3>& A, const TVector<T, 3>& B, const TVector<T, 3>& C) -> bool
 		{
@@ -1060,7 +1060,7 @@ namespace Chaos
 	}
 
 	template <typename T>
-	bool THeightField<T>::SweepGeom(const TImplicitObject<T, 3>& QueryGeom, const TRigidTransform<T, 3>& StartTM, const TVector<T, 3>& Dir, const T Length, T& OutTime, TVector<T, 3>& OutPosition, TVector<T, 3>& OutNormal, int32& OutFaceIndex, const T Thickness) const
+	bool THeightField<T>::SweepGeom(const FImplicitObject& QueryGeom, const TRigidTransform<T, 3>& StartTM, const TVector<T, 3>& Dir, const T Length, T& OutTime, TVector<T, 3>& OutPosition, TVector<T, 3>& OutNormal, int32& OutFaceIndex, const T Thickness) const
 	{
 		bool bHit = false;
 		THeightfieldSweepVisitor<T> SQVisitor(&GeomData, QueryGeom, StartTM, Dir, Thickness);
diff --git a/Engine/Source/Runtime/Experimental/Chaos/Private/Chaos/ImplicitObject.cpp b/Engine/Source/Runtime/Experimental/Chaos/Private/Chaos/ImplicitObject.cpp
--- a/Engine/Source/Runtime/Experimental/Chaos/Private/Chaos/ImplicitObject.cpp
+++ b/Engine/Source/Runtime/Experimental/Chaos/Private/Chaos/ImplicitObject.cpp
@@ -18,8 +18,7 @@
 
 using namespace Chaos;
 
-template<class T, int d>
-TImplicitObject<T, d>::TImplicitObject(int32 Flags, EImplicitObjectType InType)
+FImplicitObject::FImplicitObject(int32 Flags, EImplicitObjectType InType)
     : Type(InType)
     , bIsConvex(!!(Flags & EImplicitObject::IsConvex))
     , bIgnoreAnalyticCollisions(!!(Flags & EImplicitObject::IgnoreAnalyticCollisions))
@@ -27,13 +26,11 @@ TImplicitObject<T, d>::TImplicitObject(int32 Flags, EImplicitObjectType InType)
 {
 }
 
-template<class T, int d>
-TImplicitObject<T, d>::~TImplicitObject()
+FImplicitObject::~FImplicitObject()
 {
 }
 
-template<class T, int d>
-EImplicitObjectType TImplicitObject<T, d>::GetType(bool bGetTrueType) const
+EImplicitObjectType FImplicitObject::GetType(bool bGetTrueType) const
 {
 	if (bIgnoreAnalyticCollisions && !bGetTrueType)
 	{
@@ -42,50 +39,44 @@ EImplicitObjectType TImplicitObject<T, d>::GetType(bool bGetTrueType) const
 	return Type;
 }
 
-template<class T, int d>
-bool TImplicitObject<T, d>::IsValidGeometry() const
+bool FImplicitObject::IsValidGeometry() const
 {
 	return true;
 }
 
-template<class T, int d>
-TUniquePtr<TImplicitObject<T, d>> TImplicitObject<T, d>::Copy() const
+TUniquePtr<FImplicitObject> FImplicitObject::Copy() const
 {
 	check(false);
 	return nullptr;
 }
 
-template<class T, int d>
-bool TImplicitObject<T, d>::IsUnderlyingUnion() const
+bool FImplicitObject::IsUnderlyingUnion() const
 {
 	return Type == ImplicitObjectType::Union;
 }
 
-template<class T, int d>
-T TImplicitObject<T, d>::SignedDistance(const TVector<T, d>& x) const
+FReal FImplicitObject::SignedDistance(const FVec3& x) const
 {
-	TVector<T, d> Normal;
+	FVec3 Normal;
 	return PhiWithNormal(x, Normal);
 }
 
-template<class T, int d>
-TVector<T, d> TImplicitObject<T, d>::Normal(const TVector<T, d>& x) const
+FVec3 FImplicitObject::Normal(const FVec3& x) const
 {
-	TVector<T, d> Normal;
+	FVec3 Normal;
 	PhiWithNormal(x, Normal);
 	return Normal;
 }
 
-template<class T, int d>
-TVector<T, d> TImplicitObject<T, d>::Support(const TVector<T, d>& Direction, const T Thickness) const
+FVec3 FImplicitObject::Support(const FVec3& Direction, const FReal Thickness) const
 {
 	check(false); //not a good implementation, don't use this
-	return TVector<T, d>(0);
+	return FVec3(0);
 #if 0
 	check(bHasBoundingBox);
-	const TBox<T, d> Box = BoundingBox();
-	TVector<T, d> EndPoint = Box.Center();
-	TVector<T, d> StartPoint = EndPoint + Direction.GetSafeNormal() * (Box.Extents().Max() + Thickness);
+	const TBox<FReal, 3> Box = BoundingBox();
+	FVec3 EndPoint = Box.Center();
+	FVec3 StartPoint = EndPoint + Direction.GetSafeNormal() * (Box.Extents().Max() + Thickness);
 	checkSlow(SignedDistance(StartPoint) > 0);
 	checkSlow(SignedDistance(EndPoint) < 0);
 	// @todo(mlentine): The termination condition is slightly different here so we can probably optimize by reimplementing for this function.
@@ -95,35 +86,33 @@ TVector<T, d> TImplicitObject<T, d>::Support(const TVector<T, d>& Direction, con
 #endif
 }
 
-template<typename T, int d>
-const TBox<T, d>& TImplicitObject<T, d>::BoundingBox() const
+const TBox<FReal, 3>& FImplicitObject::BoundingBox() const
 {
 	check(false);
-	static const TBox<T, d> Unbounded(TVector<T, d>(-FLT_MAX), TVector<T, d>(FLT_MAX));
+	static const TBox<FReal, 3> Unbounded(FVec3(-FLT_MAX), FVec3(FLT_MAX));
 	return Unbounded;
 }
 
 // @todo(mlentine): This is a lot of duplication from the collisions code that should be reduced
-template<class T, int d>
-Pair<TVector<T, d>, bool> TImplicitObject<T, d>::FindDeepestIntersection(const TImplicitObject<T, d>* Other, const TBVHParticles<float, d>* Particles, const PMatrix<T, d, d>& OtherToLocalTransform, const T Thickness) const
+Pair<FVec3, bool> FImplicitObject::FindDeepestIntersection(const FImplicitObject* Other, const TBVHParticles<FReal, 3>* Particles, const FMatrix33& OtherToLocalTransform, const FReal Thickness) const
 {
 	// Do analytics
 	// @todo(mlentine): Should we do a convex pass here?
 	if (!Particles)
 	{
-		return MakePair(TVector<T, d>(0), false);
+		return MakePair(FVec3(0), false);
 	}
-	TVector<T, d> Point;
-	T Phi = Thickness;
+	FVec3 Point;
+	FReal Phi = Thickness;
 	if (HasBoundingBox())
 	{
-		TBox<T, d> ImplicitBox = BoundingBox().TransformedBox(OtherToLocalTransform.Inverse());
+		TBox<FReal, 3> ImplicitBox = BoundingBox().TransformedBox(OtherToLocalTransform.Inverse());
 		ImplicitBox.Thicken(Thickness);
 		TArray<int32> PotentialParticles = Particles->FindAllIntersections(ImplicitBox);
 		for (int32 i : PotentialParticles)
 		{
-			TVector<T, d> LocalPoint = OtherToLocalTransform.TransformPosition(Particles->X(i));
-			T LocalPhi = SignedDistance(LocalPoint);
+			FVec3 LocalPoint = OtherToLocalTransform.TransformPosition(Particles->X(i));
+			FReal LocalPhi = SignedDistance(LocalPoint);
 			if (LocalPhi < Phi)
 			{
 				Phi = LocalPhi;
@@ -133,27 +122,26 @@ Pair<TVector<T, d>, bool> TImplicitObject<T, d>::FindDeepestIntersection(const T
 	}
 	else
 	{
-		return FindDeepestIntersection(Other, static_cast<const TParticles<float, d>*>(Particles), OtherToLocalTransform, Thickness);
+		return FindDeepestIntersection(Other, static_cast<const TParticles<FReal, 3>*>(Particles), OtherToLocalTransform, Thickness);
 	}
 	return MakePair(Point, Phi < Thickness);
 }
 
-template<class T, int d>
-Pair<TVector<T, d>, bool> TImplicitObject<T, d>::FindDeepestIntersection(const TImplicitObject<T, d>* Other, const TParticles<float, d>* Particles, const PMatrix<T, d, d>& OtherToLocalTransform, const T Thickness) const
+Pair<FVec3, bool> FImplicitObject::FindDeepestIntersection(const FImplicitObject* Other, const TParticles<FReal, 3>* Particles, const FMatrix33& OtherToLocalTransform, const FReal Thickness) const
 {
 	// Do analytics
 	// @todo(mlentine): Should we do a convex pass here?
 	if (!Particles)
 	{
-		return MakePair(TVector<T, d>(0), false);
+		return MakePair(FVec3(0), false);
 	}
-	TVector<T, d> Point;
-	T Phi = Thickness;
+	FVec3 Point;
+	FReal Phi = Thickness;
 	int32 NumParticles = Particles->Size();
 	for (int32 i = 0; i < NumParticles; ++i)
 	{
-		TVector<T, d> LocalPoint = OtherToLocalTransform.TransformPosition(Particles->X(i));
-		T LocalPhi = SignedDistance(LocalPoint);
+		FVec3 LocalPoint = OtherToLocalTransform.TransformPosition(Particles->X(i));
+		FReal LocalPhi = SignedDistance(LocalPoint);
 		if (LocalPhi < Phi)
 		{
 			Phi = LocalPhi;
@@ -162,52 +150,45 @@ Pair<TVector<T, d>, bool> TImplicitObject<T, d>::FindDeepestIntersection(const T
 	}
 	return MakePair(Point, Phi < Thickness);
 }
-template<class T, int d>
-bool TImplicitObject<T, d>::Overlap(const TVector<T, d>& Point, const T Thickness) const
-{
-	return SignedDistance(Point) <= Thickness;
-}
-
 
-template<class T, int d>
-Pair<TVector<T, d>, bool> TImplicitObject<T, d>::FindClosestIntersection(const TVector<T, d>& StartPoint, const TVector<T, d>& EndPoint, const T Thickness) const
+Pair<FVec3, bool> FImplicitObject::FindClosestIntersection(const FVec3& StartPoint, const FVec3& EndPoint, const FReal Thickness) const
 {
-	constexpr T Epsilon = (T)1e-4;
-	constexpr T EpsilonSquared = Epsilon * Epsilon;
+	constexpr FReal Epsilon = (FReal)1e-4;
+	constexpr FReal EpsilonSquared = Epsilon * Epsilon;
 
 	//Consider 0 thickness with Start sitting on abs(Phi) < Epsilon. This is a common case; for example a particle sitting perfectly on a floor. In this case intersection could return false.
 	//If start is in this fuzzy region we simply return that spot snapped onto the surface. This is valid because low precision means we don't really know where we are, so let's take the cheapest option
 	//If end is in this fuzzy region it is also a valid hit. However, there could be multiple hits between start and end and since we want the first one, we can't simply return this point.
 	//As such we move end away from start (and out of the fuzzy region) so that we always get a valid intersection if no earlier ones exist
 	//When Thickness > 0 the same idea applies, but we must consider Phi = (Thickness - Epsilon, Thickness + Epsilon)
-	TVector<T, d> Normal;
-	const T Phi = PhiWithNormal(StartPoint, Normal);
+	FVec3 Normal;
+	const FReal Phi = PhiWithNormal(StartPoint, Normal);
 	if (FMath::IsNearlyEqual(Phi, Thickness, Epsilon))
 	{
-		return MakePair(TVector<T, d>(StartPoint - Normal * Phi), true); //snap to surface
+		return MakePair(FVec3(StartPoint - Normal * Phi), true); //snap to surface
 	}
 
-	TVector<T, d> ModifiedEnd = EndPoint;
+	FVec3 ModifiedEnd = EndPoint;
 	{
-		const TVector<T, d> OriginalStartToEnd = (EndPoint - StartPoint);
-		const T OriginalLength2 = OriginalStartToEnd.SizeSquared();
+		const FVec3 OriginalStartToEnd = (EndPoint - StartPoint);
+		const FReal OriginalLength2 = OriginalStartToEnd.SizeSquared();
 		if (OriginalLength2 < EpsilonSquared)
 		{
-			return MakePair(TVector<T, d>(0), false); //start was not close to surface, and end is very close to start so no hit
+			return MakePair(FVec3(0), false); //start was not close to surface, and end is very close to start so no hit
 		}
 
-		TVector<T, d> EndNormal;
-		const T EndPhi = PhiWithNormal(EndPoint, EndNormal);
+		FVec3 EndNormal;
+		const FReal EndPhi = PhiWithNormal(EndPoint, EndNormal);
 		if (FMath::IsNearlyEqual(EndPhi, Thickness, Epsilon))
 		{
 			//We want to push End out of the fuzzy region. Moving along the normal direction is best since direction could be nearly parallel with fuzzy band
 			//To ensure an intersection, we must go along the normal, but in the same general direction as the ray.
-			const TVector<T, d> OriginalDir = OriginalStartToEnd / FMath::Sqrt(OriginalLength2);
-			const T Dot = TVector<T, d>::DotProduct(OriginalDir, EndNormal);
+			const FVec3 OriginalDir = OriginalStartToEnd / FMath::Sqrt(OriginalLength2);
+			const FReal Dot = FVec3::DotProduct(OriginalDir, EndNormal);
 			if (FMath::IsNearlyZero(Dot, Epsilon))
 			{
 				//End is in the fuzzy region, and the direction from start to end is nearly parallel with this fuzzy band, so we should just return End since no other hits will occur
-				return MakePair(TVector<T, d>(EndPoint - Normal * Phi), true); //snap to surface
+				return MakePair(FVec3(EndPoint - Normal * Phi), true); //snap to surface
 			}
 			else
 			{
@@ -222,56 +203,55 @@ Pair<TVector<T, d>, bool> TImplicitObject<T, d>::FindClosestIntersection(const T
 float ClosestIntersectionStepSizeMultiplier = 0.5f;
 FAutoConsoleVariableRef CVarClosestIntersectionStepSizeMultiplier(TEXT("p.ClosestIntersectionStepSizeMultiplier"), ClosestIntersectionStepSizeMultiplier, TEXT("When raycasting we use this multiplier to substep the travel distance along the ray. Smaller number gives better accuracy at higher cost"));
 
-template<class T, int d>
-Pair<TVector<T, d>, bool> TImplicitObject<T, d>::FindClosestIntersectionImp(const TVector<T, d>& StartPoint, const TVector<T, d>& EndPoint, const T Thickness) const
+Pair<FVec3, bool> FImplicitObject::FindClosestIntersectionImp(const FVec3& StartPoint, const FVec3& EndPoint, const FReal Thickness) const
 {
-	T Epsilon = (T)1e-4;
+	FReal Epsilon = (FReal)1e-4;
 
-	TVector<T, d> Ray = EndPoint - StartPoint;
-	T Length = Ray.Size();
-	TVector<T, d> Direction = Ray.GetUnsafeNormal(); //this is safe because StartPoint and EndPoint were already tested to be far enough away. In the case where ModifiedEnd is pushed, we push it along the direction so it can only get farther
-	TVector<T, d> EndNormal;
-	const T EndPhi = PhiWithNormal(EndPoint, EndNormal);
-	TVector<T, d> ClosestPoint = StartPoint;
+	FVec3 Ray = EndPoint - StartPoint;
+	FReal Length = Ray.Size();
+	FVec3 Direction = Ray.GetUnsafeNormal(); //this is safe because StartPoint and EndPoint were already tested to be far enough away. In the case where ModifiedEnd is pushed, we push it along the direction so it can only get farther
+	FVec3 EndNormal;
+	const FReal EndPhi = PhiWithNormal(EndPoint, EndNormal);
+	FVec3 ClosestPoint = StartPoint;
 
-	TVector<T, d> Normal;
-	T Phi = PhiWithNormal(ClosestPoint, Normal);
+	FVec3 Normal;
+	FReal Phi = PhiWithNormal(ClosestPoint, Normal);
 
 	while (Phi > Thickness + Epsilon)
 	{
-		ClosestPoint += Direction * (Phi - Thickness) * (T)ClosestIntersectionStepSizeMultiplier;
+		ClosestPoint += Direction * (Phi - Thickness) * (FReal)ClosestIntersectionStepSizeMultiplier;
 		if ((ClosestPoint - StartPoint).Size() > Length)
 		{
 			if (EndPhi < Thickness + Epsilon)
 			{
-				return MakePair(TVector<T, d>(EndPoint + EndNormal * (-EndPhi + Thickness)), true);
+				return MakePair(FVec3(EndPoint + EndNormal * (-EndPhi + Thickness)), true);
 			}
-			return MakePair(TVector<T, d>(0), false);
+			return MakePair(FVec3(0), false);
 		}
 		// If the Change is too small we want to nudge it forward. This makes it possible to miss intersections very close to the surface but is more efficient and shouldn't matter much.
-		if ((Phi - Thickness) < (T)1e-2)
+		if ((Phi - Thickness) < (FReal)1e-2)
 		{
-			ClosestPoint += Direction * (T)1e-2;
+			ClosestPoint += Direction * (FReal)1e-2;
 			if ((ClosestPoint - StartPoint).Size() > Length)
 			{
 				if (EndPhi < Thickness + Epsilon)
 				{
-					return MakePair(TVector<T, d>(EndPoint + EndNormal * (-EndPhi + Thickness)), true);
+					return MakePair(FVec3(EndPoint + EndNormal * (-EndPhi + Thickness)), true);
 				}
 				else
 				{
-					return MakePair(TVector<T, d>(0), false);
+					return MakePair(FVec3(0), false);
 				}
 			}
 		}
-		T NewPhi = PhiWithNormal(ClosestPoint, Normal);
+		FReal NewPhi = PhiWithNormal(ClosestPoint, Normal);
 		if (NewPhi >= Phi)
 		{
 			if (EndPhi < Thickness + Epsilon)
 			{
-				return MakePair(TVector<T, d>(EndPoint + EndNormal * (-EndPhi + Thickness)), true);
+				return MakePair(FVec3(EndPoint + EndNormal * (-EndPhi + Thickness)), true);
 			}
-			return MakePair(TVector<T, d>(0), false);
+			return MakePair(FVec3(0), false);
 		}
 		Phi = NewPhi;
 	}
@@ -282,17 +262,15 @@ Pair<TVector<T, d>, bool> TImplicitObject<T, d>::FindClosestIntersectionImp(cons
 	return MakePair(ClosestPoint, true);
 }
 
-template<typename T, int d>
-void TImplicitObject<T, d>::FindAllIntersectingObjects(TArray<Pair<const TImplicitObject<T, d>*, TRigidTransform<T, d>>>& Out, const TBox<T, d>& LocalBounds) const
+void FImplicitObject::FindAllIntersectingObjects(TArray<Pair<const FImplicitObject*, FRigidTransform3>>& Out, const TBox<FReal, 3>& LocalBounds) const
 {
 	if (!HasBoundingBox() || LocalBounds.Intersects(BoundingBox()))
 	{
-		Out.Add(MakePair(this, TRigidTransform<T, d>(TVector<T, d>(0), TRotation<T, d>::FromElements(TVector<T, d>(0), (T)1))));
+		Out.Add(MakePair(this, FRigidTransform3(FVec3(0), FRotation3::FromElements(FVec3(0), (FReal)1))));
 	}
 }
 
-template<typename T, int d>
-FArchive& TImplicitObject<T, d>::SerializeLegacyHelper(FArchive& Ar, TUniquePtr<TImplicitObject<T, d>>& Value)
+FArchive& FImplicitObject::SerializeLegacyHelper(FArchive& Ar, TUniquePtr<FImplicitObject>& Value)
 {
 	bool bExists = Value.Get() != nullptr;
 	Ar << bExists;
@@ -304,10 +282,10 @@ FArchive& TImplicitObject<T, d>::SerializeLegacyHelper(FArchive& Ar, TUniquePtr<
 			Ar << ObjectType;
 			switch ((EImplicitObjectType)ObjectType)
 			{
-			case ImplicitObjectType::Sphere: { Value = TUniquePtr<TSphere<T, d>>(new TSphere<T, d>()); break; }
-			case ImplicitObjectType::Box: { Value = TUniquePtr<TBox<T, d>>(new TBox<T, d>()); break; }
-			case ImplicitObjectType::Plane: { Value = TUniquePtr<TPlane<T, d>>(new TPlane<T, d>()); break; }
-			case ImplicitObjectType::LevelSet: { Value = TUniquePtr<TLevelSet<T, d>>(new TLevelSet<T, d>()); break; }
+			case ImplicitObjectType::Sphere: { Value = TUniquePtr<TSphere<FReal,3>>(new TSphere<FReal, 3>()); break; }
+			case ImplicitObjectType::Box: { Value = TUniquePtr<TBox<FReal,3>>(new TBox<FReal, 3>()); break; }
+			case ImplicitObjectType::Plane: { Value = TUniquePtr<TPlane<FReal,3>>(new TPlane<FReal, 3>()); break; }
+			case ImplicitObjectType::LevelSet: { Value = TUniquePtr<TLevelSet<FReal,3>>(new TLevelSet<FReal, 3>()); break; }
 			default: check(false);
 			}
 		}
@@ -327,8 +305,7 @@ FArchive& TImplicitObject<T, d>::SerializeLegacyHelper(FArchive& Ar, TUniquePtr<
 	return Ar;
 }
 
-template<typename T, int d>
-void TImplicitObject<T, d>::SerializeImp(FArchive& Ar)
+void FImplicitObject::SerializeImp(FArchive& Ar)
 {
 	Ar.UsingCustomVersion(FDestructionObjectVersion::GUID);
 	if (Ar.CustomVer(FDestructionObjectVersion::GUID) >= FDestructionObjectVersion::ChaosArchiveAdded)
@@ -337,14 +314,12 @@ void TImplicitObject<T, d>::SerializeImp(FArchive& Ar)
 	}
 }
 
-template<typename T, int d>
-void TImplicitObject<T, d>::Serialize(FChaosArchive& Ar)
+void FImplicitObject::Serialize(FChaosArchive& Ar)
 {
 	SerializeImp(Ar);
 }
 
-template<typename T, int d>
-const FName TImplicitObject<T, d>::GetTypeName(const EImplicitObjectType InType)
+const FName FImplicitObject::GetTypeName(const EImplicitObjectType InType)
 {
 	static const FName SphereName = TEXT("Sphere");
 	static const FName BoxName = TEXT("Box");
@@ -379,8 +354,7 @@ const FName TImplicitObject<T, d>::GetTypeName(const EImplicitObjectType InType)
 	return NAME_None;
 }
 
-template<typename T, int d>
-TImplicitObject<T, d>* TImplicitObject<T, d>::SerializationFactory(FChaosArchive& Ar, TImplicitObject<T, d>* Obj)
+FImplicitObject* FImplicitObject::SerializationFactory(FChaosArchive& Ar, FImplicitObject* Obj)
 {
 	int8 ObjectType = Ar.IsLoading() ? 0 : (int8)Obj->Type;
 	Ar << ObjectType;
@@ -393,8 +367,8 @@ TImplicitObject<T, d>* TImplicitObject<T, d>::SerializationFactory(FChaosArchive
 			EImplicitObjectType InnerType = GetInnerType(ObjectType);
 			switch (InnerType)
 			{
-			case ImplicitObjectType::Convex: if (Ar.IsLoading()) { return new TImplicitObjectScaled<TConvex<T, d>>(); } break;
-			case ImplicitObjectType::TriangleMesh: if (Ar.IsLoading()) { return new TImplicitObjectScaled<TTriangleMeshImplicitObject<T>>(); } break;
+			case ImplicitObjectType::Convex: if (Ar.IsLoading()) { return new TImplicitObjectScaled<TConvex<FReal, 3>>(); } break;
+			case ImplicitObjectType::TriangleMesh: if (Ar.IsLoading()) { return new TImplicitObjectScaled<TTriangleMeshImplicitObject<FReal>>(); } break;
 			default: check(false);
 			}
 
@@ -404,28 +378,25 @@ TImplicitObject<T, d>* TImplicitObject<T, d>::SerializationFactory(FChaosArchive
 
 	switch ((EImplicitObjectType)ObjectType)
 	{
-	case ImplicitObjectType::Sphere: if (Ar.IsLoading()) { return new TSphere<T, d>(); } break;
-	case ImplicitObjectType::Box: if (Ar.IsLoading()) { return new TBox<T, d>(); } break;
-	case ImplicitObjectType::Plane: if (Ar.IsLoading()) { return new TPlane<T, d>(); } break;
-	case ImplicitObjectType::Capsule: if (Ar.IsLoading()) { return new TCapsule<T>(); } break;
-	case ImplicitObjectType::Transformed: if (Ar.IsLoading()) { return new TImplicitObjectTransformed<T, d>(); } break;
-	case ImplicitObjectType::Union: if (Ar.IsLoading()) { return new TImplicitObjectUnion<T, d>(); } break;
-	case ImplicitObjectType::LevelSet: if (Ar.IsLoading()) { return new TLevelSet<T, d>(); } break;
-	case ImplicitObjectType::Convex: if (Ar.IsLoading()) { return new TConvex<T, d>(); } break;
-	case ImplicitObjectType::TaperedCylinder: if (Ar.IsLoading()) { return new TTaperedCylinder<T>(); } break;
-	case ImplicitObjectType::TriangleMesh: if (Ar.IsLoading()) { return new TTriangleMeshImplicitObject<T>(); } break;
+	case ImplicitObjectType::Sphere: if (Ar.IsLoading()) { return new TSphere<FReal, 3>(); } break;
+	case ImplicitObjectType::Box: if (Ar.IsLoading()) { return new TBox<FReal, 3>(); } break;
+	case ImplicitObjectType::Plane: if (Ar.IsLoading()) { return new TPlane<FReal, 3>(); } break;
+	case ImplicitObjectType::Capsule: if (Ar.IsLoading()) { return new TCapsule<FReal>(); } break;
+	case ImplicitObjectType::Transformed: if (Ar.IsLoading()) { return new TImplicitObjectTransformed<FReal, 3>(); } break;
+	case ImplicitObjectType::Union: if (Ar.IsLoading()) { return new TImplicitObjectUnion<FReal, 3>(); } break;
+	case ImplicitObjectType::LevelSet: if (Ar.IsLoading()) { return new TLevelSet<FReal, 3>(); } break;
+	case ImplicitObjectType::Convex: if (Ar.IsLoading()) { return new TConvex<FReal, 3>(); } break;
+	case ImplicitObjectType::TaperedCylinder: if (Ar.IsLoading()) { return new TTaperedCylinder<FReal>(); } break;
+	case ImplicitObjectType::TriangleMesh: if (Ar.IsLoading()) { return new TTriangleMeshImplicitObject<FReal>(); } break;
 	case ImplicitObjectType::DEPRECATED_Scaled:
 	{
 		ensure(Ar.IsLoading() && (Ar.CustomVer(FExternalPhysicsCustomObjectVersion::GUID) < FExternalPhysicsCustomObjectVersion::ScaledGeometryIsConcrete));
-		return new TImplicitObjectScaledGeneric<T, d>();
+		return new TImplicitObjectScaledGeneric<FReal, 3>();
 	}
-	case ImplicitObjectType::HeightField: if (Ar.IsLoading()) { return new THeightField<T>(); } break;
-	case ImplicitObjectType::Cylinder: if (Ar.IsLoading()) { return new TCylinder<T>(); } break;
+	case ImplicitObjectType::HeightField: if (Ar.IsLoading()) { return new THeightField<FReal>(); } break;
+	case ImplicitObjectType::Cylinder: if (Ar.IsLoading()) { return new TCylinder<FReal>(); } break;
 	default:
 		check(false);
 	}
 	return nullptr;
-}
-
-//template class Chaos::TImplicitObject<float, 2>; // Missing 2D Rotation class
-template class Chaos::TImplicitObject<float, 3>;
+}
\ No newline at end of file
diff --git a/Engine/Source/Runtime/Experimental/Chaos/Private/Chaos/Levelset.cpp b/Engine/Source/Runtime/Experimental/Chaos/Private/Chaos/Levelset.cpp
--- a/Engine/Source/Runtime/Experimental/Chaos/Private/Chaos/Levelset.cpp
+++ b/Engine/Source/Runtime/Experimental/Chaos/Private/Chaos/Levelset.cpp
@@ -34,7 +34,7 @@ using namespace Chaos;
 
 template<class T, int d>
 TLevelSet<T, d>::TLevelSet(FErrorReporter& ErrorReporter, const TUniformGrid<T, d>& InGrid, const TParticles<T, d>& InParticles, const TTriangleMesh<T>& Mesh, const int32 BandWidth)
-    : TImplicitObject<T, d>(EImplicitObject::HasBoundingBox, ImplicitObjectType::LevelSet)
+    : FImplicitObject(EImplicitObject::HasBoundingBox, ImplicitObjectType::LevelSet)
     , MGrid(InGrid)
     , MPhi(MGrid)
     , MNormals(MGrid)
@@ -91,8 +91,8 @@ TLevelSet<T, d>::TLevelSet(FErrorReporter& ErrorReporter, const TUniformGrid<T,
 }
 
 template<class T, int d>
-TLevelSet<T, d>::TLevelSet(FErrorReporter& ErrorReporter, const TUniformGrid<T, d>& InGrid, const TImplicitObject<T, d>& InObject, const int32 BandWidth, const bool bUseObjectPhi)
-    : TImplicitObject<T, d>(EImplicitObject::HasBoundingBox, ImplicitObjectType::LevelSet)
+TLevelSet<T, d>::TLevelSet(FErrorReporter& ErrorReporter, const TUniformGrid<T, d>& InGrid, const FImplicitObject& InObject, const int32 BandWidth, const bool bUseObjectPhi)
+    : FImplicitObject(EImplicitObject::HasBoundingBox, ImplicitObjectType::LevelSet)
     , MGrid(InGrid)
     , MPhi(MGrid)
     , MNormals(MGrid)
@@ -149,7 +149,7 @@ TLevelSet<T, d>::TLevelSet(FErrorReporter& ErrorReporter, const TUniformGrid<T,
 
 template<class T, int d>
 TLevelSet<T, d>::TLevelSet(std::istream& Stream)
-    : TImplicitObject<T, d>(EImplicitObject::HasBoundingBox, ImplicitObjectType::LevelSet)
+    : FImplicitObject(EImplicitObject::HasBoundingBox, ImplicitObjectType::LevelSet)
     , MGrid(Stream)
     , MPhi(Stream)
     , MLocalBoundingBox(MGrid.MinCorner(), MGrid.MaxCorner())
@@ -160,7 +160,7 @@ TLevelSet<T, d>::TLevelSet(std::istream& Stream)
 
 template<class T, int d>
 TLevelSet<T, d>::TLevelSet(TLevelSet<T, d>&& Other)
-    : TImplicitObject<T, d>(EImplicitObject::HasBoundingBox, ImplicitObjectType::LevelSet)
+    : FImplicitObject(EImplicitObject::HasBoundingBox, ImplicitObjectType::LevelSet)
     , MGrid(MoveTemp(Other.MGrid))
     , MPhi(MoveTemp(Other.MPhi))
     , MLocalBoundingBox(MoveTemp(Other.MLocalBoundingBox))
@@ -637,7 +637,7 @@ bool TLevelSet<T, d>::ComputeDistancesNearZeroIsocontour(FErrorReporter& ErrorRe
 }
 
 template<class T, int d>
-void TLevelSet<T, d>::ComputeDistancesNearZeroIsocontour(const TImplicitObject<T, d>& Object, const TArrayND<T, d>& ObjectPhi, TArray<TVector<int32, d>>& InterfaceIndices)
+void TLevelSet<T, d>::ComputeDistancesNearZeroIsocontour(const FImplicitObject& Object, const TArrayND<T, d>& ObjectPhi, TArray<TVector<int32, d>>& InterfaceIndices)
 {
 	MPhi.Fill(FLT_MAX);
 	const auto& Counts = MGrid.Counts();
diff --git a/Engine/Source/Runtime/Experimental/Chaos/Private/Chaos/PBDCollisionConstraint.cpp b/Engine/Source/Runtime/Experimental/Chaos/Private/Chaos/PBDCollisionConstraint.cpp
--- a/Engine/Source/Runtime/Experimental/Chaos/Private/Chaos/PBDCollisionConstraint.cpp
+++ b/Engine/Source/Runtime/Experimental/Chaos/Private/Chaos/PBDCollisionConstraint.cpp
@@ -55,7 +55,7 @@ CHAOS_API int32 EnableCollisions = 1;
 FAutoConsoleVariableRef CVarEnableCollisions(TEXT("p.EnableCollisions"), EnableCollisions, TEXT("Enable/Disable collisions on the Chaos solver."));
 
 template<ECollisionUpdateType, typename T, int d>
-void UpdateConstraintImp2(const TRigidTransform<T, d>& ParticleTM, const TImplicitObject<T, d>& LevelsetObject, const TRigidTransform<T, d>& LevelsetTM, const T Thickness, TRigidBodyContactConstraint<T, d>& Constraint);
+void UpdateConstraintImp2(const TRigidTransform<T, d>& ParticleTM, const FImplicitObject& LevelsetObject, const TRigidTransform<T, d>& LevelsetTM, const T Thickness, TRigidBodyContactConstraint<T, d>& Constraint);
 
 //
 // Collision Constraint Handle
@@ -939,7 +939,7 @@ void UpdateLevelsetConstraintHelperCCD(const TRigidParticles<T, d>& InParticles,
 #endif
 
 template <typename T, int d>
-bool SampleObjectHelper2(const TImplicitObject<T, d>& Object, const TRigidTransform<T,d>& ObjectTransform, const TRigidTransform<T, d>& SampleToObjectTransform, const TVector<T, d>& SampleParticle, T Thickness, TRigidBodyContactConstraint<T, d>& Constraint)
+bool SampleObjectHelper2(const FImplicitObject& Object, const TRigidTransform<T,d>& ObjectTransform, const TRigidTransform<T, d>& SampleToObjectTransform, const TVector<T, d>& SampleParticle, T Thickness, TRigidBodyContactConstraint<T, d>& Constraint)
 {
 	TVector<T, d> LocalPoint = SampleToObjectTransform.TransformPositionNoScale(SampleParticle);
 	TVector<T, d> LocalNormal;
@@ -955,7 +955,7 @@ bool SampleObjectHelper2(const TImplicitObject<T, d>& Object, const TRigidTransf
 }
 
 template <typename T, int d>
-bool SampleObjectNoNormal2(const TImplicitObject<T, d>& Object, const TRigidTransform<T,d>& ObjectTransform, const TRigidTransform<T, d>& SampleToObjectTransform, const TVector<T, d>& SampleParticle, T Thickness, TRigidBodyContactConstraint<T, d>& Constraint)
+bool SampleObjectNoNormal2(const FImplicitObject& Object, const TRigidTransform<T,d>& ObjectTransform, const TRigidTransform<T, d>& SampleToObjectTransform, const TVector<T, d>& SampleParticle, T Thickness, TRigidBodyContactConstraint<T, d>& Constraint)
 {
 	TVector<T, d> LocalPoint = SampleToObjectTransform.TransformPositionNoScale(SampleParticle);
 	TVector<T, d> LocalNormal;
@@ -969,7 +969,7 @@ bool SampleObjectNoNormal2(const TImplicitObject<T, d>& Object, const TRigidTran
 }
 
 template <typename T, int d>
-bool SampleObjectNormalAverageHelper2(const TImplicitObject<T, d>& Object, const TRigidTransform<T, d>& ObjectTransform, const TRigidTransform<T, d>& SampleToObjectTransform, const TVector<T, d>& SampleParticle, T Thickness, T& TotalThickness, TRigidBodyContactConstraint<T, d>& Constraint)
+bool SampleObjectNormalAverageHelper2(const FImplicitObject& Object, const TRigidTransform<T, d>& ObjectTransform, const TRigidTransform<T, d>& SampleToObjectTransform, const TVector<T, d>& SampleParticle, T Thickness, T& TotalThickness, TRigidBodyContactConstraint<T, d>& Constraint)
 {
 	TVector<T, d> LocalPoint = SampleToObjectTransform.TransformPositionNoScale(SampleParticle);
 	TVector<T, d> LocalNormal;
@@ -998,7 +998,7 @@ int32 SampleMinParticlesForAcceleration2 = 2048;
 FAutoConsoleVariableRef CVarSampleMinParticlesForAcceleration2(TEXT("p.SampleMinParticlesForAcceleration2"), SampleMinParticlesForAcceleration2, TEXT("The minimum number of particles needed before using an acceleration structure when sampling"));
 
 template <ECollisionUpdateType UpdateType, typename T, int d>
-void SampleObject2(const TImplicitObject<T, d>& Object, const TRigidTransform<T, d>& ObjectTransform, const TBVHParticles<T, d>& SampleParticles, const TRigidTransform<T, d>& SampleParticlesTransform, T Thickness, TRigidBodyContactConstraint<T, d>& Constraint)
+void SampleObject2(const FImplicitObject& Object, const TRigidTransform<T, d>& ObjectTransform, const TBVHParticles<T, d>& SampleParticles, const TRigidTransform<T, d>& SampleParticlesTransform, T Thickness, TRigidBodyContactConstraint<T, d>& Constraint)
 {
 	SCOPE_CYCLE_COUNTER(STAT_SampleObject);
 	TRigidBodyContactConstraint<T, d> AvgConstraint;
@@ -1103,7 +1103,7 @@ void SampleObject2(const TImplicitObject<T, d>& Object, const TRigidTransform<T,
 
 #if INTEL_ISPC
 template<ECollisionUpdateType UpdateType>
-void SampleObject2(const TImplicitObject<float, 3>& Object, const TRigidTransform<float, 3>& ObjectTransform, const TBVHParticles<float, 3>& SampleParticles, const TRigidTransform<float, 3>& SampleParticlesTransform, float Thickness, TRigidBodyContactConstraint<float, 3>& Constraint)
+void SampleObject2(const FImplicitObject& Object, const TRigidTransform<float, 3>& ObjectTransform, const TBVHParticles<float, 3>& SampleParticles, const TRigidTransform<float, 3>& SampleParticlesTransform, float Thickness, TRigidBodyContactConstraint<float, 3>& Constraint)
 {
 	SCOPE_CYCLE_COUNTER(STAT_SampleObject);
 	TRigidBodyContactConstraint<float, 3> AvgConstraint;
@@ -1578,10 +1578,10 @@ void UpdateCapsuleBoxConstraint(const TCapsule<T>& A, const TRigidTransform<T, d
 
 DECLARE_CYCLE_STAT(TEXT("TPBDCollisionConstraint::FindRelevantShapes"), STAT_FindRelevantShapes, STATGROUP_ChaosWide);
 template <typename T, int d>
-TArray<Pair<const TImplicitObject<T, d>*, TRigidTransform<T, d>>> FindRelevantShapes2(const TImplicitObject<T,d>* ParticleObj, const TRigidTransform<T,d>& ParticlesTM, const TImplicitObject<T,d>& LevelsetObj, const TRigidTransform<T,d>& LevelsetTM, const T Thickness)
+TArray<Pair<const FImplicitObject*, TRigidTransform<T, d>>> FindRelevantShapes2(const FImplicitObject* ParticleObj, const TRigidTransform<T,d>& ParticlesTM, const FImplicitObject& LevelsetObj, const TRigidTransform<T,d>& LevelsetTM, const T Thickness)
 {
 	SCOPE_CYCLE_COUNTER(STAT_FindRelevantShapes);
-	TArray<Pair<const TImplicitObject<T, d>*, TRigidTransform<T, d>>> RelevantShapes;
+	TArray<Pair<const FImplicitObject*, TRigidTransform<T, d>>> RelevantShapes;
 	//find all levelset inner objects
 	if (ParticleObj)
 	{
@@ -1620,22 +1620,22 @@ void UpdateUnionUnionConstraint(const T Thickness, TRigidBodyContactConstraint<T
 	TRigidTransform<T, d> ParticlesTM = TRigidTransform<T, d>(Particle0->P(), Particle0->Q());
 	TRigidTransform<T, d> LevelsetTM = TRigidTransform<T, d>(Particle1->P(), Particle1->Q());
 
-	const TImplicitObject<T, d>* ParticleObj = Particle0->Geometry().Get();
-	const TImplicitObject<T,d>* LevelsetObj = Particle1->Geometry().Get();
-	const TArray<Pair<const TImplicitObject<T, d>*, TRigidTransform<T, d>>> LevelsetShapes = FindRelevantShapes2(ParticleObj, ParticlesTM, *LevelsetObj, LevelsetTM, Thickness);
+	const FImplicitObject* ParticleObj = Particle0->Geometry().Get();
+	const FImplicitObject* LevelsetObj = Particle1->Geometry().Get();
+	const TArray<Pair<const FImplicitObject*, TRigidTransform<T, d>>> LevelsetShapes = FindRelevantShapes2(ParticleObj, ParticlesTM, *LevelsetObj, LevelsetTM, Thickness);
 
-	for (const Pair<const TImplicitObject<T, d>*, TRigidTransform<T, d>>& LevelsetObjPair : LevelsetShapes)
+	for (const Pair<const FImplicitObject*, TRigidTransform<T, d>>& LevelsetObjPair : LevelsetShapes)
 	{
-		const TImplicitObject<T, d>& LevelsetInnerObj = *LevelsetObjPair.First;
+		const FImplicitObject& LevelsetInnerObj = *LevelsetObjPair.First;
 		const TRigidTransform<T, d>& LevelsetInnerObjTM = LevelsetObjPair.Second * LevelsetTM;
 
 		//now find all particle inner objects
-		const TArray<Pair<const TImplicitObject<T, d>*, TRigidTransform<T, d>>> ParticleShapes = FindRelevantShapes2(&LevelsetInnerObj, LevelsetInnerObjTM, *ParticleObj, ParticlesTM, Thickness);
+		const TArray<Pair<const FImplicitObject*, TRigidTransform<T, d>>> ParticleShapes = FindRelevantShapes2(&LevelsetInnerObj, LevelsetInnerObjTM, *ParticleObj, ParticlesTM, Thickness);
 
 		//for each inner obj pair, update constraint
-		for (const Pair<const TImplicitObject<T, d>*, TRigidTransform<T, d>>& ParticlePair : ParticleShapes)
+		for (const Pair<const FImplicitObject*, TRigidTransform<T, d>>& ParticlePair : ParticleShapes)
 		{
-			const TImplicitObject<T, d>& ParticleInnerObj = *ParticlePair.First;
+			const FImplicitObject& ParticleInnerObj = *ParticlePair.First;
 			const TRigidTransform<T, d> ParticleInnerObjTM = ParticlePair.Second * ParticlesTM;
 			UpdateConstraintImp2<UpdateType>(ParticleInnerObj, ParticleInnerObjTM, LevelsetInnerObj, LevelsetInnerObjTM, Thickness, Constraint);
 		}
@@ -1654,13 +1654,13 @@ void UpdateSingleUnionConstraint(const T Thickness, TRigidBodyContactConstraint<
 	TRigidTransform<T, d> ParticlesTM = TRigidTransform<T, d>(Particle0->P(), Particle0->Q());
 	TRigidTransform<T, d> LevelsetTM = TRigidTransform<T, d>(Particle1->P(), Particle1->Q());
 
-	const TImplicitObject<T, d>* ParticleObj = Particle0->Geometry().Get();
-	const TImplicitObject<T, d>* LevelsetObj = Particle1->Geometry().Get();
-	const TArray<Pair<const TImplicitObject<T, d>*, TRigidTransform<T, d>>> LevelsetShapes = FindRelevantShapes2(ParticleObj, ParticlesTM, *LevelsetObj, LevelsetTM, Thickness);
+	const FImplicitObject* ParticleObj = Particle0->Geometry().Get();
+	const FImplicitObject* LevelsetObj = Particle1->Geometry().Get();
+	const TArray<Pair<const FImplicitObject*, TRigidTransform<T, d>>> LevelsetShapes = FindRelevantShapes2(ParticleObj, ParticlesTM, *LevelsetObj, LevelsetTM, Thickness);
 	
-	for (const Pair<const TImplicitObject<T, d>*, TRigidTransform<T, d>>& LevelsetObjPair : LevelsetShapes)
+	for (const Pair<const FImplicitObject*, TRigidTransform<T, d>>& LevelsetObjPair : LevelsetShapes)
 	{
-		const TImplicitObject<T, d>& LevelsetInnerObj = *LevelsetObjPair.First;
+		const FImplicitObject& LevelsetInnerObj = *LevelsetObjPair.First;
 		const TRigidTransform<T, d> LevelsetInnerObjTM = LevelsetObjPair.Second * LevelsetTM;
 		UpdateConstraintImp2<UpdateType>(*ParticleObj, ParticlesTM, LevelsetInnerObj, LevelsetInnerObjTM, Thickness, Constraint);
 	}
@@ -1698,7 +1698,7 @@ void TPBDCollisionConstraint<T, d>::UpdateLevelsetConstraint(const T Thickness,
 
 DECLARE_CYCLE_STAT(TEXT("UpdateConvexConstraintsUsingCoreShapes"), UpdateConvexConstraintsUsingCoreShapes, STATGROUP_ChaosWide);
 template<ECollisionUpdateType UpdateType, typename T, int d>
-void UpdateConvexConstraintsUsingCoreShapes(const TImplicitObject<T, d> & AObj, const TRigidTransform<T, d>& ATM, const TImplicitObject<T, d> & BObj, const TRigidTransform<T, d>& BTM, const T Thickness, TRigidBodyContactConstraint<T, d>& Constraint)
+void UpdateConvexConstraintsUsingCoreShapes(const FImplicitObject & AObj, const TRigidTransform<T, d>& ATM, const FImplicitObject & BObj, const TRigidTransform<T, d>& BTM, const T Thickness, TRigidBodyContactConstraint<T, d>& Constraint)
 {
 	SCOPE_CYCLE_COUNTER(UpdateConvexConstraintsUsingCoreShapes);
 
@@ -1727,18 +1727,18 @@ void UpdateUnionLevelsetConstraint(const T Thickness, TRigidBodyContactConstrain
 		return;
 	}
 
-	const TImplicitObject<T, d>* ParticleObj = Particle0->Geometry().Get();
-	const TImplicitObject<T, d>* LevelsetObj = Particle1->Geometry().Get();
-	TArray<Pair<const TImplicitObject<T, d>*, TRigidTransform<T, d>>> LevelsetShapes = FindRelevantShapes2(ParticleObj, ParticlesTM, *LevelsetObj, LevelsetTM, Thickness);
+	const FImplicitObject* ParticleObj = Particle0->Geometry().Get();
+	const FImplicitObject* LevelsetObj = Particle1->Geometry().Get();
+	TArray<Pair<const FImplicitObject*, TRigidTransform<T, d>>> LevelsetShapes = FindRelevantShapes2(ParticleObj, ParticlesTM, *LevelsetObj, LevelsetTM, Thickness);
 
 	if (LevelsetShapes.Num() && Particle0->CollisionParticles().Get())
 	{
 		const TBVHParticles<T, d>& SampleParticles = *Particle0->CollisionParticles().Get();
 		if (SampleParticles.Size())
 		{
-			for (const Pair<const TImplicitObject<T, d>*, TRigidTransform<T, d>>& LevelsetObjPair : LevelsetShapes)
+			for (const Pair<const FImplicitObject*, TRigidTransform<T, d>>& LevelsetObjPair : LevelsetShapes)
 			{
-				const TImplicitObject<T, d>* Object = LevelsetObjPair.First;
+				const FImplicitObject* Object = LevelsetObjPair.First;
 				const TRigidTransform<T, d> ObjectTM = LevelsetObjPair.Second * LevelsetTM;
 				SampleObject2<UpdateType>(*Object, ObjectTM, SampleParticles, ParticlesTM, Thickness, Constraint);
 				if (UpdateType == ECollisionUpdateType::Any && Constraint.Phi < Thickness)
@@ -1752,15 +1752,15 @@ void UpdateUnionLevelsetConstraint(const T Thickness, TRigidBodyContactConstrain
 		{
 			const TImplicitObjectUnion<T, d>* UnionObj = static_cast<const TImplicitObjectUnion<T, d>*>(ParticleObj);
 			//need to traverse shapes to get their collision particles
-			for (const Pair<const TImplicitObject<T, d>*, TRigidTransform<T, d>>& LevelsetObjPair : LevelsetShapes)
+			for (const Pair<const FImplicitObject*, TRigidTransform<T, d>>& LevelsetObjPair : LevelsetShapes)
 			{
-				const TImplicitObject<T, d>* LevelsetInnerObject = LevelsetObjPair.First;
+				const FImplicitObject* LevelsetInnerObject = LevelsetObjPair.First;
 				const TRigidTransform<T, d> LevelsetInnerObjectTM = LevelsetObjPair.Second * LevelsetTM;
 
-				TArray<Pair<const TImplicitObject<T, d>*, TRigidTransform<T, d>>> ParticleShapes = FindRelevantShapes2(LevelsetInnerObject, LevelsetInnerObjectTM, *ParticleObj, ParticlesTM, Thickness);
-				for (const Pair<const TImplicitObject<T, d>*, TRigidTransform<T, d>>& ParticleObjPair : ParticleShapes)
+				TArray<Pair<const FImplicitObject*, TRigidTransform<T, d>>> ParticleShapes = FindRelevantShapes2(LevelsetInnerObject, LevelsetInnerObjectTM, *ParticleObj, ParticlesTM, Thickness);
+				for (const Pair<const FImplicitObject*, TRigidTransform<T, d>>& ParticleObjPair : ParticleShapes)
 				{
-					const TImplicitObject<T, d>* ParticleInnerObject = ParticleObjPair.First;
+					const FImplicitObject* ParticleInnerObject = ParticleObjPair.First;
 					const TRigidTransform<T, d> ParticleInnerObjectTM = ParticleObjPair.Second * ParticlesTM;
 
 					if (const int32* OriginalIdx = UnionObj->MCollisionParticleLookupHack.Find(ParticleInnerObject))
@@ -1793,8 +1793,8 @@ void UpdateLevelsetUnionConstraint(const T Thickness, TRigidBodyContactConstrain
 	TRigidTransform<T, d> ParticlesTM = TRigidTransform<T, d>(Particle0->P(), Particle0->Q());
 	TRigidTransform<T, d> LevelsetTM = TRigidTransform<T, d>(Particle1->P(), Particle1->Q());
 
-	const TImplicitObject<T, d>* ParticleObj = Particle0->Geometry().Get();
-	const TImplicitObject<T, d>* LevelsetObj = Particle1->Geometry().Get();
+	const FImplicitObject* ParticleObj = Particle0->Geometry().Get();
+	const FImplicitObject* LevelsetObj = Particle1->Geometry().Get();
 
 	if (!(ensure(!FMath::IsNaN(ParticlesTM.GetTranslation().X)) && ensure(!FMath::IsNaN(ParticlesTM.GetTranslation().Y)) && ensure(!FMath::IsNaN(ParticlesTM.GetTranslation().Z))))
 	{
@@ -1807,12 +1807,12 @@ void UpdateLevelsetUnionConstraint(const T Thickness, TRigidBodyContactConstrain
 	}
 
 #if CHAOS_PARTICLEHANDLE_TODO
-	TArray<Pair<const TImplicitObject<T, d>*, TRigidTransform<T, d>>> ParticleShapes = FindRelevantShapes2(LevelsetObj, LevelsetTM, *ParticleObj, ParticlesTM, Thickness);
+	TArray<Pair<const FImplicitObject*, TRigidTransform<T, d>>> ParticleShapes = FindRelevantShapes2(LevelsetObj, LevelsetTM, *ParticleObj, ParticlesTM, Thickness);
 	check(ParticleObj->IsUnderlyingUnion());
 	const TImplicitObjectUnion<T, d>* UnionObj = static_cast<const TImplicitObjectUnion<T, d>*>(ParticleObj);
-	for (const Pair<const TImplicitObject<T, d>*, TRigidTransform<T, d>>& ParticleObjPair : ParticleShapes)
+	for (const Pair<const FImplicitObject*, TRigidTransform<T, d>>& ParticleObjPair : ParticleShapes)
 	{
-		const TImplicitObject<T, d>* Object = ParticleObjPair.First;
+		const FImplicitObject* Object = ParticleObjPair.First;
 
 		if (const int32* OriginalIdx = UnionObj->MCollisionParticleLookupHack.Find(Object))
 		{
@@ -2078,7 +2078,7 @@ void UpdateLevelsetConstraintImp(const float Thickness, TRigidBodyContactConstra
 }
 
 template<ECollisionUpdateType UpdateType, typename T, int d>
-void UpdateConstraintImp2(const TImplicitObject<T, d>& ParticleObject, const TRigidTransform<T, d>& ParticleTM, const TImplicitObject<T, d>& LevelsetObject, const TRigidTransform<T, d>& LevelsetTM, const T Thickness, TRigidBodyContactConstraint<T, d>& Constraint)
+void UpdateConstraintImp2(const FImplicitObject& ParticleObject, const TRigidTransform<T, d>& ParticleTM, const FImplicitObject& LevelsetObject, const TRigidTransform<T, d>& LevelsetTM, const T Thickness, TRigidBodyContactConstraint<T, d>& Constraint)
 {
 	if (ParticleObject.GetType() == TBox<T, d>::StaticType() && LevelsetObject.GetType() == TBox<T, d>::StaticType())
 	{
@@ -2221,6 +2221,6 @@ template void TPBDCollisionConstraint<float, 3>::UpdateLevelsetConstraint<EColli
 template void TPBDCollisionConstraint<float, 3>::ComputeConstraints<false>(const TPBDCollisionConstraint<float,3>::FAccelerationStructure&, float Dt);
 template void TPBDCollisionConstraint<float, 3>::ComputeConstraints<true>(const TPBDCollisionConstraint<float, 3>::FAccelerationStructure&, float Dt);
 
-template void UpdateConstraintImp2<ECollisionUpdateType::Any, float, 3>(const TImplicitObject<float, 3>& ParticleObject, const TRigidTransform<float, 3>& ParticleTM, const TImplicitObject<float, 3>& LevelsetObject, const TRigidTransform<float, 3>& LevelsetTM, const float Thickness, TRigidBodyContactConstraint<float, 3>& Constraint);
-template void UpdateConstraintImp2<ECollisionUpdateType::Deepest, float, 3>(const TImplicitObject<float, 3>& ParticleObject, const TRigidTransform<float, 3>& ParticleTM, const TImplicitObject<float, 3>& LevelsetObject, const TRigidTransform<float, 3>& LevelsetTM, const float Thickness, TRigidBodyContactConstraint<float, 3>& Constraint);
+template void UpdateConstraintImp2<ECollisionUpdateType::Any, float, 3>(const FImplicitObject& ParticleObject, const TRigidTransform<float, 3>& ParticleTM, const FImplicitObject& LevelsetObject, const TRigidTransform<float, 3>& LevelsetTM, const float Thickness, TRigidBodyContactConstraint<float, 3>& Constraint);
+template void UpdateConstraintImp2<ECollisionUpdateType::Deepest, float, 3>(const FImplicitObject& ParticleObject, const TRigidTransform<float, 3>& ParticleTM, const FImplicitObject& LevelsetObject, const TRigidTransform<float, 3>& LevelsetTM, const float Thickness, TRigidBodyContactConstraint<float, 3>& Constraint);
 }
diff --git a/Engine/Source/Runtime/Experimental/Chaos/Private/Chaos/PBDRigidClustering.cpp b/Engine/Source/Runtime/Experimental/Chaos/Private/Chaos/PBDRigidClustering.cpp
--- a/Engine/Source/Runtime/Experimental/Chaos/Private/Chaos/PBDRigidClustering.cpp
+++ b/Engine/Source/Runtime/Experimental/Chaos/Private/Chaos/PBDRigidClustering.cpp
@@ -91,7 +91,7 @@ namespace Chaos
 	int32 TPBDRigidClustering<FPBDRigidsEvolution, FPBDCollisionConstraint, T, d>::CreateClusterParticle(
 		int32 ClusterGroupIndex,
 		TArray<uint32>&& Children,
-		/*const*/ TSerializablePtr<TImplicitObject<T, d>> ProxyGeometry,
+		/*const*/ TSerializablePtr<FImplicitObject> ProxyGeometry,
 		const TRigidTransform<T, d>* ForceMassOrientation,
 		const FClusterCreationParameters<T>& Parameters)
 	{
@@ -236,7 +236,7 @@ namespace Chaos
 		NoCleanParams.bCopyCollisionParticles = !!UnionsHaveCollisionParticles;
 
 		UpdateMassProperties(ChildrenArray, NewIndex, nullptr);
-		UpdateGeometry(ChildrenArray, NewIndex, TSerializablePtr<TImplicitObject<T, 3>>(), NoCleanParams);
+		UpdateGeometry(ChildrenArray, NewIndex, TSerializablePtr<FImplicitObject>(), NoCleanParams);
 
 		return NewIndex;
 	}
@@ -315,7 +315,7 @@ namespace Chaos
 
 				FClusterCreationParameters<T> Parameters(0.3, 100, false, !!UnionsHaveCollisionParticles);
 				Parameters.ConnectionMethod = MClusterUnionConnectionType;
-				int32 NewIndex = CreateClusterParticle(-Group.Key, MoveTemp(ClusterChildren), TSerializablePtr<TImplicitObject<T, d>>(), nullptr, Parameters);
+				int32 NewIndex = CreateClusterParticle(-Group.Key, MoveTemp(ClusterChildren), TSerializablePtr<FImplicitObject>(), nullptr, Parameters);
 				MParticles.InternalCluster(NewIndex) = true;
 				MEvolution.SetPhysicsMaterial(MParticles.Handle(NewIndex), MEvolution.GetPhysicsMaterial(MParticles.Handle(Group.Value[0])));
 
@@ -1285,12 +1285,12 @@ namespace Chaos
 
 	template<class FPBDRigidsEvolution, class FPBDCollisionConstraint, class T, int d>
 	void TPBDRigidClustering<FPBDRigidsEvolution, FPBDCollisionConstraint, T, d>::UpdateGeometry(const TArray<uint32>& Children,
-		const uint32 NewIndex, TSerializablePtr<TImplicitObject<T, d>> ProxyGeometry, const FClusterCreationParameters<T>& Parameters)
+		const uint32 NewIndex, TSerializablePtr<FImplicitObject> ProxyGeometry, const FClusterCreationParameters<T>& Parameters)
 	{
 		SCOPE_CYCLE_COUNTER(STAT_UpdateGeometry);
 		const bool bUseCollisionPoints = (ProxyGeometry || Parameters.bCopyCollisionParticles) && !Parameters.CollisionParticles;
-		TArray<TUniquePtr<TImplicitObject<T, d>>> Objects;
-		TArray<TUniquePtr<TImplicitObject<T, d>>> Objects2; //todo: find a better way to reuse this
+		TArray<TUniquePtr<FImplicitObject>> Objects;
+		TArray<TUniquePtr<FImplicitObject>> Objects2; //todo: find a better way to reuse this
 
 		ensure(!MParticles.Geometry(NewIndex)); //we should never update existing geometry since this is used by SQ threads.
 		ensure(!MParticles.DynamicGeometry(NewIndex)); //we should never update existing geometry since this is used by SQ threads.
@@ -1314,8 +1314,8 @@ namespace Chaos
 					const int32 MultiChildProxyId = MParticles.MultiChildProxyId(Child).Id;
 					if (UseLevelsetCollision || MultiChildProxyId == INDEX_NONE || !MParticles.MultiChildProxyData(MultiChildProxyId))
 					{
-						Objects.Add(TUniquePtr<TImplicitObject<T, d>>(new TImplicitObjectTransformed<T, d>(MParticles.Geometry(Child), Frame)));
-						Objects2.Add(TUniquePtr<TImplicitObject<T, d>>(new TImplicitObjectTransformed<T, d>(MParticles.Geometry(Child), Frame)));
+						Objects.Add(TUniquePtr<FImplicitObject>(new TImplicitObjectTransformed<T, d>(MParticles.Geometry(Child), Frame)));
+						Objects2.Add(TUniquePtr<FImplicitObject>(new TImplicitObjectTransformed<T, d>(MParticles.Geometry(Child), Frame)));
 						GeomToOriginalParticlesHack.Add(Child);
 					}
 					else if (MParticles.MultiChildProxyData(MultiChildProxyId)->KeyChild == Child)
@@ -1323,8 +1323,8 @@ namespace Chaos
 						//using multi child proxy and this child is the key
 						const TRigidTransform<T, d> ProxyWorldTM = MParticles.MultiChildProxyData(MultiChildProxyId)->RelativeToKeyChild * ChildWorldTM;
 						const TRigidTransform<T, d> ProxyRelativeTM = ProxyWorldTM.GetRelativeTransform(ClusterWorldTM);
-						Objects.Add(TUniquePtr<TImplicitObject<T, d>>(new TImplicitObjectTransformed<T, d>(MParticles.Geometry(MultiChildProxyId), ProxyRelativeTM)));
-						Objects2.Add(TUniquePtr<TImplicitObject<T, d>>(new TImplicitObjectTransformed<T, d>(MParticles.Geometry(MultiChildProxyId), ProxyRelativeTM)));
+						Objects.Add(TUniquePtr<FImplicitObject>(new TImplicitObjectTransformed<T, d>(MParticles.Geometry(MultiChildProxyId), ProxyRelativeTM)));
+						Objects2.Add(TUniquePtr<FImplicitObject>(new TImplicitObjectTransformed<T, d>(MParticles.Geometry(MultiChildProxyId), ProxyRelativeTM)));
 						UsedGeomChild = MultiChildProxyId;
 						GeomToOriginalParticlesHack.Add(UsedGeomChild);
 						bUsingMultiChildProxy = true;
@@ -1383,7 +1383,7 @@ namespace Chaos
 		{
 			//ensureMsgf(false, TEXT("Checking usage with no proxy and no objects"));
 			//@coverage : {production}
-			MParticles.SetGeometry(NewIndex, Chaos::TSerializablePtr<Chaos::TImplicitObject<float, 3>>());
+			MParticles.SetGeometry(NewIndex, Chaos::TSerializablePtr<Chaos::FImplicitObject>());
 		}
 		else
 		{
@@ -1864,7 +1864,7 @@ namespace Chaos
 		ResourceLock.WriteLock();
 		//BufferResource.MChildren.Reset();
 		//BufferResource.ClusterParentTransforms.Reset();	//todo: once everything is atomic this should get reset
-		const TArray<TSerializablePtr<TImplicitObject<T, d>>>& AllGeom = MParticles.GetAllGeometry();
+		const TArray<TSerializablePtr<FImplicitObject>>& AllGeom = MParticles.GetAllGeometry();
 		BufferResource.GeometryPtrs.SetNum(AllGeom.Num());
 
 		const auto& NonDisabledClusteredParticles = MEvolution.GetNonDisabledClusteredArray();
diff --git a/Engine/Source/Runtime/Experimental/Chaos/Private/Chaos/TriangleMeshImplicitObject.cpp b/Engine/Source/Runtime/Experimental/Chaos/Private/Chaos/TriangleMeshImplicitObject.cpp
--- a/Engine/Source/Runtime/Experimental/Chaos/Private/Chaos/TriangleMeshImplicitObject.cpp
+++ b/Engine/Source/Runtime/Experimental/Chaos/Private/Chaos/TriangleMeshImplicitObject.cpp
@@ -8,7 +8,7 @@ namespace Chaos
 {
 template <typename T>
 TTriangleMeshImplicitObject<T>::TTriangleMeshImplicitObject(TParticles<T, 3>&& Particles, TArray<TVector<int32, 3>>&& Elements)
-	: TImplicitObject<T, 3>(EImplicitObject::HasBoundingBox, ImplicitObjectType::TriangleMesh)
+	: FImplicitObject(EImplicitObject::HasBoundingBox, ImplicitObjectType::TriangleMesh)
 	, MParticles(MoveTemp(Particles))
 	, MElements(MoveTemp(Elements))
 	, MLocalBoundingBox(MParticles.X(0), MParticles.X(0))
@@ -258,7 +258,7 @@ bool TTriangleMeshImplicitObject<T>::Overlap(const TVector<T, 3>& Point, const T
 }
 
 template <typename T>
-bool TTriangleMeshImplicitObject<T>::OverlapGeom(const TImplicitObject<T, 3>& QueryGeom, const TRigidTransform<T, 3>& QueryTM, const T Thickness) const
+bool TTriangleMeshImplicitObject<T>::OverlapGeom(const FImplicitObject& QueryGeom, const TRigidTransform<T, 3>& QueryTM, const T Thickness) const
 {
 	auto OverlapTriangle = [&](const TVector<T, 3>& A, const TVector<T, 3>& B, const TVector<T, 3>& C) -> bool
 	{
@@ -298,7 +298,7 @@ bool TTriangleMeshImplicitObject<T>::OverlapGeom(const TImplicitObject<T, 3>& Qu
 template <typename T>
 struct TTriangleMeshSweepVisitor
 {
-	TTriangleMeshSweepVisitor(const TTriangleMeshImplicitObject<T>& InTriMesh, const TImplicitObject<T,3>& InQueryGeom, const TRigidTransform<T,3>& InStartTM, const TVector<T,3>& InDir, const T InThickness)
+	TTriangleMeshSweepVisitor(const TTriangleMeshImplicitObject<T>& InTriMesh, const FImplicitObject& InQueryGeom, const TRigidTransform<T,3>& InStartTM, const TVector<T,3>& InDir, const T InThickness)
 	: TriMesh(InTriMesh)
 	, StartTM(InStartTM)
 	, QueryGeom(InQueryGeom)
@@ -343,7 +343,7 @@ struct TTriangleMeshSweepVisitor
 
 	const TTriangleMeshImplicitObject<T>& TriMesh;
 	const TRigidTransform<T, 3> StartTM;
-	const TImplicitObject<T, 3>& QueryGeom;
+	const FImplicitObject& QueryGeom;
 	const TVector<T, 3>& Dir;
 	const T Thickness;
 
@@ -354,7 +354,7 @@ struct TTriangleMeshSweepVisitor
 };
 
 template <typename T>
-bool TTriangleMeshImplicitObject<T>::SweepGeom(const TImplicitObject<T, 3>& QueryGeom, const TRigidTransform<T, 3>& StartTM, const TVector<T, 3>& Dir, const T Length, T& OutTime, TVector<T, 3>& OutPosition, TVector<T, 3>& OutNormal, int32& OutFaceIndex, const T Thickness) const
+bool TTriangleMeshImplicitObject<T>::SweepGeom(const FImplicitObject& QueryGeom, const TRigidTransform<T, 3>& StartTM, const TVector<T, 3>& Dir, const T Length, T& OutTime, TVector<T, 3>& OutPosition, TVector<T, 3>& OutNormal, int32& OutFaceIndex, const T Thickness) const
 {
 	bool bHit = false;
 	TTriangleMeshSweepVisitor<T> SQVisitor(*this, QueryGeom, StartTM, Dir, Thickness);
diff --git a/Engine/Source/Runtime/Experimental/Chaos/Public/Chaos/Box.h b/Engine/Source/Runtime/Experimental/Chaos/Public/Chaos/Box.h
--- a/Engine/Source/Runtime/Experimental/Chaos/Public/Chaos/Box.h
+++ b/Engine/Source/Runtime/Experimental/Chaos/Public/Chaos/Box.h
@@ -11,17 +11,17 @@
 namespace Chaos
 {
 	template<class T, int d>
-	class TBox final : public TImplicitObject<T, d>
+	class TBox final : public FImplicitObject
 	{
 	public:
-		using TImplicitObject<T, d>::SignedDistance;
-		using TImplicitObject<T, d>::GetTypeName;
+		using FImplicitObject::SignedDistance;
+		using FImplicitObject::GetTypeName;
 
 		// This should never be used outside of creating a default for arrays
 		FORCEINLINE TBox()
-		    : TImplicitObject<T, d>(EImplicitObject::FiniteConvex, ImplicitObjectType::Box){};
+		    : FImplicitObject(EImplicitObject::FiniteConvex, ImplicitObjectType::Box){};
 		FORCEINLINE TBox(const TVector<T, d>& Min, const TVector<T, d>& Max)
-		    : TImplicitObject<T, d>(EImplicitObject::FiniteConvex, ImplicitObjectType::Box)
+		    : FImplicitObject(EImplicitObject::FiniteConvex, ImplicitObjectType::Box)
 			, AABB(Min, Max)
 		{
 			//todo: turn back on
@@ -32,19 +32,19 @@ namespace Chaos
 		}
 
 		FORCEINLINE TBox(const TBox<T, d>& Other)
-		    : TImplicitObject<T, d>(EImplicitObject::FiniteConvex, ImplicitObjectType::Box)
+		    : FImplicitObject(EImplicitObject::FiniteConvex, ImplicitObjectType::Box)
 			, AABB(Other.AABB)
 		{
 		}
 
 		FORCEINLINE TBox(const TAABB<T, d>& AABB)
-		    : TImplicitObject<T, d>(EImplicitObject::FiniteConvex, ImplicitObjectType::Box)
+		    : FImplicitObject(EImplicitObject::FiniteConvex, ImplicitObjectType::Box)
 			, AABB(AABB)
 		{
 		}
 
 		FORCEINLINE TBox(TBox<T, d>&& Other)
-		    : TImplicitObject<T, d>(EImplicitObject::FiniteConvex, ImplicitObjectType::Box)
+		    : FImplicitObject(EImplicitObject::FiniteConvex, ImplicitObjectType::Box)
 		    , AABB(MoveTemp(Other.AABB))
 		{
 		}
@@ -73,9 +73,9 @@ namespace Chaos
 
 		virtual ~TBox() {}
 
-		virtual TUniquePtr<TImplicitObject<T, d>> Copy() const override
+		virtual TUniquePtr<FImplicitObject> Copy() const override
 		{
-			return TUniquePtr<TImplicitObject<T, d>>(new TBox<T,d>(*this));
+			return TUniquePtr<FImplicitObject>(new TBox<T,d>(*this));
 		}
 
 		/**
@@ -230,7 +230,7 @@ namespace Chaos
 
 		FORCEINLINE void SerializeImp(FArchive& Ar)
 		{
-			TImplicitObject<T, d>::SerializeImp(Ar);
+			FImplicitObject::SerializeImp(Ar);
 			AABB.Serialize(Ar);
 		}
 
diff --git a/Engine/Source/Runtime/Experimental/Chaos/Public/Chaos/Capsule.h b/Engine/Source/Runtime/Experimental/Chaos/Public/Chaos/Capsule.h
--- a/Engine/Source/Runtime/Experimental/Chaos/Public/Chaos/Capsule.h
+++ b/Engine/Source/Runtime/Experimental/Chaos/Public/Chaos/Capsule.h
@@ -14,17 +14,17 @@ namespace Chaos
 	struct TCapsuleSpecializeSamplingHelper;
 
 	template<class T>
-	class TCapsule final : public TImplicitObject<T, 3>
+	class TCapsule final : public FImplicitObject
 	{
 	public:
-		using TImplicitObject<T, 3>::SignedDistance;
-		using TImplicitObject<T, 3>::GetTypeName;
+		using FImplicitObject::SignedDistance;
+		using FImplicitObject::GetTypeName;
 
 		TCapsule()
-		    : TImplicitObject<T, 3>(EImplicitObject::FiniteConvex, ImplicitObjectType::Capsule)
+		    : FImplicitObject(EImplicitObject::FiniteConvex, ImplicitObjectType::Capsule)
 		{}
 		TCapsule(const TVector<T, 3>& x1, const TVector<T, 3>& x2, const T Radius)
-		    : TImplicitObject<T, 3>(EImplicitObject::FiniteConvex, ImplicitObjectType::Capsule)
+		    : FImplicitObject(EImplicitObject::FiniteConvex, ImplicitObjectType::Capsule)
 			, MSegment(x1, x2)
 		    , MRadius(Radius)
 		    , MLocalBoundingBox(x1, x1)
@@ -36,7 +36,7 @@ namespace Chaos
 		}
 
 		TCapsule(const TCapsule<T>& Other)
-		    : TImplicitObject<T, 3>(EImplicitObject::FiniteConvex, ImplicitObjectType::Capsule)
+		    : FImplicitObject(EImplicitObject::FiniteConvex, ImplicitObjectType::Capsule)
 			, MSegment(Other.MSegment)
 		    , MRadius(Other.MRadius)
 		    , MLocalBoundingBox(Other.MLocalBoundingBox)
@@ -46,7 +46,7 @@ namespace Chaos
 		}
 
 		TCapsule(TCapsule<T>&& Other)
-		    : TImplicitObject<T, 3>(EImplicitObject::FiniteConvex, ImplicitObjectType::Capsule)
+		    : FImplicitObject(EImplicitObject::FiniteConvex, ImplicitObjectType::Capsule)
 			, MSegment(MoveTemp(Other.MSegment))
 		    , MRadius(Other.MRadius)
 		    , MLocalBoundingBox(MoveTemp(Other.MLocalBoundingBox))
@@ -301,7 +301,7 @@ namespace Chaos
 
 		FORCEINLINE void SerializeImp(FArchive& Ar)
 		{
-			TImplicitObject<T, 3>::SerializeImp(Ar);
+			FImplicitObject::SerializeImp(Ar);
 			MSegment.Serialize(Ar);
 			Ar << MRadius << MLocalBoundingBox;
 		}
@@ -313,9 +313,9 @@ namespace Chaos
 			Ar << MUnionedObjects;
 		}
 
-		virtual TUniquePtr<TImplicitObject<T, 3>> Copy() const override
+		virtual TUniquePtr<FImplicitObject> Copy() const override
 		{
-			return TUniquePtr<TImplicitObject<T,3>>(new TCapsule<T>(*this));
+			return TUniquePtr<FImplicitObject>(new TCapsule<T>(*this));
 		}
 
 		const T& GetRadius() const { return MRadius; }
@@ -368,7 +368,7 @@ namespace Chaos
 	private:
 		void InitUnionedObjects()
 		{
-			TArray<TUniquePtr<TImplicitObject<float, 3>>> Objects;
+			TArray<TUniquePtr<FImplicitObject>> Objects;
 
 			const TVector<T, 3>& X1 = GetX1();
 			const TVector<T, 3>& X2 = GetX2();
diff --git a/Engine/Source/Runtime/Experimental/Chaos/Public/Chaos/ChaosArchive.h b/Engine/Source/Runtime/Experimental/Chaos/Public/Chaos/ChaosArchive.h
--- a/Engine/Source/Runtime/Experimental/Chaos/Public/Chaos/ChaosArchive.h
+++ b/Engine/Source/Runtime/Experimental/Chaos/Public/Chaos/ChaosArchive.h
@@ -9,8 +9,7 @@
 
 namespace Chaos
 {
-template <typename T, int d>
-class CHAOS_API TImplicitObject; //needed for legacy serializer
+class CHAOS_API FImplicitObject; //needed for legacy serializer
 
 
 #if CHAOS_MEMORY_TRACKING
@@ -257,7 +256,7 @@ private:
 		check(false);
 	}
 
-	void SerializeLegacy(TUniquePtr<TImplicitObject<float, 3>>& Obj);
+	void SerializeLegacy(TUniquePtr<FImplicitObject>& Obj);
 
 	template <typename T>
 	void StaticSerialize(TSerializablePtr<T>& Serializable)
diff --git a/Engine/Source/Runtime/Experimental/Chaos/Public/Chaos/Convex.h b/Engine/Source/Runtime/Experimental/Chaos/Public/Chaos/Convex.h
--- a/Engine/Source/Runtime/Experimental/Chaos/Public/Chaos/Convex.h
+++ b/Engine/Source/Runtime/Experimental/Chaos/Public/Chaos/Convex.h
@@ -10,23 +10,23 @@
 namespace Chaos
 {
 	template<class T, int d>
-	class TConvex final : public TImplicitObject<T, d>
+	class TConvex final : public FImplicitObject
 	{
 	public:
-		using TImplicitObject<T, d>::GetTypeName;
+		using FImplicitObject::GetTypeName;
 
 		TConvex()
-		    : TImplicitObject<T,3>(EImplicitObject::IsConvex | EImplicitObject::HasBoundingBox, ImplicitObjectType::Convex)
+		    : FImplicitObject(EImplicitObject::IsConvex | EImplicitObject::HasBoundingBox, ImplicitObjectType::Convex)
 		{}
 		TConvex(const TConvex&) = delete;
 		TConvex(TConvex&& Other)
-		    : TImplicitObject<T,3>(EImplicitObject::IsConvex | EImplicitObject::HasBoundingBox, ImplicitObjectType::Convex)
+		    : FImplicitObject(EImplicitObject::IsConvex | EImplicitObject::HasBoundingBox, ImplicitObjectType::Convex)
 			, Planes(MoveTemp(Other.Planes))
 		    , SurfaceParticles(MoveTemp(Other.SurfaceParticles))
 		    , LocalBoundingBox(MoveTemp(Other.LocalBoundingBox))
 		{}
 		TConvex(TArray<TPlane<T, d>>&& InPlanes, TParticles<T, d>&& InSurfaceParticles)
-		    : TImplicitObject<T, d>(EImplicitObject::IsConvex | EImplicitObject::HasBoundingBox, ImplicitObjectType::Convex)
+		    : FImplicitObject(EImplicitObject::IsConvex | EImplicitObject::HasBoundingBox, ImplicitObjectType::Convex)
 			, Planes(MoveTemp(InPlanes))
 		    , SurfaceParticles(MoveTemp(InSurfaceParticles))
 		    , LocalBoundingBox(TBox<T, d>::EmptyBox())
@@ -37,7 +37,7 @@ namespace Chaos
 			}
 		}
 		TConvex(const TParticles<T, 3>& InParticles)
-		    : TImplicitObject<T, d>(EImplicitObject::IsConvex | EImplicitObject::HasBoundingBox, ImplicitObjectType::Convex)
+		    : FImplicitObject(EImplicitObject::IsConvex | EImplicitObject::HasBoundingBox, ImplicitObjectType::Convex)
 		{
 			const uint32 NumParticles = InParticles.Size();
 			if (NumParticles == 0)
@@ -212,7 +212,7 @@ namespace Chaos
 
 		FORCEINLINE void SerializeImp(FArchive& Ar)
 		{
-			TImplicitObject<T, 3>::SerializeImp(Ar);
+			FImplicitObject::SerializeImp(Ar);
 			Ar << Planes << SurfaceParticles << LocalBoundingBox;
 		}
 
diff --git a/Engine/Source/Runtime/Experimental/Chaos/Public/Chaos/Cylinder.h b/Engine/Source/Runtime/Experimental/Chaos/Public/Chaos/Cylinder.h
--- a/Engine/Source/Runtime/Experimental/Chaos/Public/Chaos/Cylinder.h
+++ b/Engine/Source/Runtime/Experimental/Chaos/Public/Chaos/Cylinder.h
@@ -13,14 +13,14 @@ namespace Chaos
 	struct TCylinderSpecializeSamplingHelper;
 
 	template<class T>
-	class TCylinder final : public TImplicitObject<T, 3>
+	class TCylinder final : public FImplicitObject
 	{
 	public:
-		using TImplicitObject<T, 3>::SignedDistance;
-		using TImplicitObject<T, 3>::GetTypeName;
+		using FImplicitObject::SignedDistance;
+		using FImplicitObject::GetTypeName;
 
 		TCylinder(const TVector<T, 3>& x1, const TVector<T, 3>& x2, const T Radius)
-		    : TImplicitObject<T, 3>(EImplicitObject::FiniteConvex, ImplicitObjectType::Cylinder)
+		    : FImplicitObject(EImplicitObject::FiniteConvex, ImplicitObjectType::Cylinder)
 		    , MPlane1(x1, (x2 - x1).GetSafeNormal()) // Plane normals point inward
 		    , MPlane2(x2, -MPlane1.Normal())
 		    , MHeight((x2 - x1).Size())
@@ -31,7 +31,7 @@ namespace Chaos
 			MLocalBoundingBox = TBox<T, 3>(MLocalBoundingBox.Min() - TVector<T, 3>(MRadius), MLocalBoundingBox.Max() + TVector<T, 3>(MRadius));
 		}
 		TCylinder(const TCylinder<T>& Other)
-		    : TImplicitObject<T, 3>(EImplicitObject::FiniteConvex, ImplicitObjectType::Cylinder)
+		    : FImplicitObject(EImplicitObject::FiniteConvex, ImplicitObjectType::Cylinder)
 		    , MPlane1(Other.MPlane1)
 		    , MPlane2(Other.MPlane2)
 		    , MHeight(Other.MHeight)
@@ -40,7 +40,7 @@ namespace Chaos
 		{
 		}
 		TCylinder(TCylinder<T>&& Other)
-		    : TImplicitObject<T, 3>(EImplicitObject::FiniteConvex, ImplicitObjectType::Cylinder)
+		    : FImplicitObject(EImplicitObject::FiniteConvex, ImplicitObjectType::Cylinder)
 		    , MPlane1(MoveTemp(Other.MPlane1))
 		    , MPlane2(MoveTemp(Other.MPlane2))
 		    , MHeight(Other.MHeight)
@@ -213,7 +213,7 @@ namespace Chaos
 		virtual void Serialize(FChaosArchive& Ar)
 		{
 			FChaosArchiveScopedMemory ScopedMemory(Ar, GetTypeName());
-			TImplicitObject<T, 3>::SerializeImp(Ar);
+			FImplicitObject::SerializeImp(Ar);
 			Ar << MPlane1;
 			Ar << MPlane2;
 			Ar << MHeight;
@@ -263,8 +263,8 @@ namespace Chaos
 		}
 
 		//needed for serialization
-		TCylinder() : TImplicitObject<T, 3>(EImplicitObject::HasBoundingBox, ImplicitObjectType::Cylinder) {}
-		friend TImplicitObject<T, 3>;	//needed for serialization
+		TCylinder() : FImplicitObject(EImplicitObject::HasBoundingBox, ImplicitObjectType::Cylinder) {}
+		friend FImplicitObject;	//needed for serialization
 
 	private:
 		TPlane<T, 3> MPlane1, MPlane2;
diff --git a/Engine/Source/Runtime/Experimental/Chaos/Public/Chaos/Declares.h b/Engine/Source/Runtime/Experimental/Chaos/Public/Chaos/Declares.h
--- a/Engine/Source/Runtime/Experimental/Chaos/Public/Chaos/Declares.h
+++ b/Engine/Source/Runtime/Experimental/Chaos/Public/Chaos/Declares.h
@@ -10,7 +10,7 @@ namespace Chaos
 
 	typedef FPBDRigidsSolver FPhysicsSolver;
 
-	template <typename T, int d> class TImplicitObject;
+	class FImplicitObject;
 	template <typename T, int d> class TPerShapeData;
 	template <typename TSOA> class TParticleView;
 	template<class T, int d> class TPBDRigidParticles;
diff --git a/Engine/Source/Runtime/Experimental/Chaos/Public/Chaos/GJK.h b/Engine/Source/Runtime/Experimental/Chaos/Public/Chaos/GJK.h
--- a/Engine/Source/Runtime/Experimental/Chaos/Public/Chaos/GJK.h
+++ b/Engine/Source/Runtime/Experimental/Chaos/Public/Chaos/GJK.h
@@ -730,13 +730,13 @@ namespace Chaos
 
 	// Overloads for geometry types which don't have centroids.
 	template <typename T, typename TGeometryB>
-	TVector<T, 3> GJKDistanceInitialV(const TImplicitObject<T, 3>& A, const TGeometryB& B, const TRigidTransform<T, 3>& BToATM)
+	TVector<T, 3> GJKDistanceInitialV(const FImplicitObject& A, const TGeometryB& B, const TRigidTransform<T, 3>& BToATM)
 	{
 		return -BToATM.GetTranslation();
 	}
 
 	template <typename T, typename TGeometryA>
-	TVector<T, 3> GJKDistanceInitialV(TGeometryA A, const TImplicitObject<T, 3>& B, const TRigidTransform<T, 3>& BToATM)
+	TVector<T, 3> GJKDistanceInitialV(TGeometryA A, const FImplicitObject& B, const TRigidTransform<T, 3>& BToATM)
 	{
 		return -BToATM.GetTranslation();
 	}
diff --git a/Engine/Source/Runtime/Experimental/Chaos/Public/Chaos/GeometryParticles.h b/Engine/Source/Runtime/Experimental/Chaos/Public/Chaos/GeometryParticles.h
--- a/Engine/Source/Runtime/Experimental/Chaos/Public/Chaos/GeometryParticles.h
+++ b/Engine/Source/Runtime/Experimental/Chaos/Public/Chaos/GeometryParticles.h
@@ -31,7 +31,7 @@ namespace Chaos
 		FCollisionFilterData QueryData;
 		FCollisionFilterData SimData;
 		void* UserData;
-		TSerializablePtr<TImplicitObject<T, d>> Geometry;
+		TSerializablePtr<FImplicitObject> Geometry;
 		
 
 		static TPerShapeData<T, d>* SerializationFactory(FChaosArchive& Ar, TPerShapeData<T, d>*);
@@ -53,9 +53,9 @@ namespace Chaos
 	using TShapesArray = TArray<TUniquePtr<TPerShapeData<T, d>>, TInlineAllocator<1>>;
 
 	template <typename T, int d>
-	void UpdateShapesArrayFromGeometry(TShapesArray<T, d>& ShapesArray, TSerializablePtr<TImplicitObject<T, d>> Geometry);
+	void UpdateShapesArrayFromGeometry(TShapesArray<T, d>& ShapesArray, TSerializablePtr<FImplicitObject> Geometry);
 
-	extern template void CHAOS_API UpdateShapesArrayFromGeometry(TShapesArray<float, 3>& ShapesArray, TSerializablePtr<TImplicitObject<float, 3>> Geometry);
+	extern template void CHAOS_API UpdateShapesArrayFromGeometry(TShapesArray<float, 3>& ShapesArray, TSerializablePtr<FImplicitObject> Geometry);
 
 #if CHAOS_DETERMINISTIC
 	using FParticleID = int32;	//Used to break ties when determinism is needed. Should not be used for anything else
@@ -178,9 +178,9 @@ namespace Chaos
 
 		CHAOS_API TSerializablePtr<TImplicitObject<T, d>> Geometry(const int32 Index) const { return MGeometry[Index]; }
 
-		CHAOS_API const TUniquePtr<TImplicitObject<T, d>>& DynamicGeometry(const int32 Index) const { return MDynamicGeometry[Index]; }
+		CHAOS_API const TUniquePtr<FImplicitObject>& DynamicGeometry(const int32 Index) const { return MDynamicGeometry[Index]; }
 
-		const TSharedPtr<TImplicitObject<T, d>, ESPMode::ThreadSafe>& SharedGeometry(const int32 Index) const { return MSharedGeometry[Index]; }
+		const TSharedPtr<FImplicitObject, ESPMode::ThreadSafe>& SharedGeometry(const int32 Index) const { return MSharedGeometry[Index]; }
 
 		const TShapesArray<T, d>& ShapesArray(const int32 Index) const { return MShapesArray[Index]; }
 
@@ -189,7 +189,7 @@ namespace Chaos
 		CHAOS_API FParticleID& ParticleID(const int32 Idx) { return MParticleIDs[Idx]; }
 #endif
 
-		CHAOS_API void SetDynamicGeometry(const int32 Index, TUniquePtr<TImplicitObject<T, d>>&& InUnique)
+		CHAOS_API void SetDynamicGeometry(const int32 Index, TUniquePtr<FImplicitObject>&& InUnique)
 		{
 			check(!SharedGeometry(Index));	// If shared geometry exists we should not be setting dynamic geometry on top
 			MGeometry[Index] = MakeSerializable(InUnique);
@@ -197,7 +197,7 @@ namespace Chaos
 			UpdateShapesArray(Index);
 		}
 
-		CHAOS_API void SetSharedGeometry(const int32 Index, TSharedPtr<TImplicitObject<T, d>, ESPMode::ThreadSafe> InShared)
+		CHAOS_API void SetSharedGeometry(const int32 Index, TSharedPtr<FImplicitObject, ESPMode::ThreadSafe> InShared)
 		{
 			check(!DynamicGeometry(Index));	// If dynamic geometry exists we should not be setting shared geometry on top
 			MGeometry[Index] = MakeSerializable(InShared);
@@ -205,7 +205,7 @@ namespace Chaos
 			UpdateShapesArray(Index);
 		}
 		
-		CHAOS_API void SetGeometry(const int32 Index, TSerializablePtr<TImplicitObject<T, d>> InGeometry)
+		CHAOS_API void SetGeometry(const int32 Index, TSerializablePtr<FImplicitObject> InGeometry)
 		{
 			check(!DynamicGeometry(Index));
 			check(!SharedGeometry(Index));
@@ -253,7 +253,7 @@ namespace Chaos
 			return MWorldSpaceInflatedBounds[Index];
 		}
 
-		const TArray<TSerializablePtr<TImplicitObject<T, d>>>& GetAllGeometry() const { return MGeometry; }
+		const TArray<TSerializablePtr<FImplicitObject>>& GetAllGeometry() const { return MGeometry; }
 
 		typedef TGeometryParticleHandle<T, d> THandleType;
 		CHAOS_API THandleType* Handle(int32 Index) const { return const_cast<THandleType*>(MGeometryParticleHandle[Index].Get()); }
@@ -324,13 +324,13 @@ namespace Chaos
 		TArrayCollectionArray<TRotation<T, d>> MR;
 		// MGeometry contains raw ptrs to every entry in both MSharedGeometry and MDynamicGeometry.
 		// It may also contain raw ptrs to geometry which is managed outside of Chaos.
-		TArrayCollectionArray<TSerializablePtr<TImplicitObject<T, d>>> MGeometry;
+		TArrayCollectionArray<TSerializablePtr<FImplicitObject>> MGeometry;
 		// MSharedGeometry entries are owned by the solver, shared between *representations* of a particle.
 		// This is NOT for sharing geometry resources between particle's A and B, this is for sharing the
 		// geometry between particle A's various representations.
-		TArrayCollectionArray<TSharedPtr<TImplicitObject<T, d>, ESPMode::ThreadSafe>> MSharedGeometry;
+		TArrayCollectionArray<TSharedPtr<FImplicitObject, ESPMode::ThreadSafe>> MSharedGeometry;
 		// MDynamicGeometry entries are used for geo which is by the evolution. It is not set from the game side.
-		TArrayCollectionArray<TUniquePtr<TImplicitObject<T, d>>> MDynamicGeometry;
+		TArrayCollectionArray<TUniquePtr<FImplicitObject>> MDynamicGeometry;
 		TArrayCollectionArray<TSerializablePtr<TGeometryParticleHandle<T, d>>> MGeometryParticleHandle;
 		TArrayCollectionArray<TGeometryParticle<T, d>*> MGeometryParticle;
 		TArrayCollectionArray<TShapesArray<T,d>> MShapesArray;
diff --git a/Engine/Source/Runtime/Experimental/Chaos/Public/Chaos/GeometryQueries.h b/Engine/Source/Runtime/Experimental/Chaos/Public/Chaos/GeometryQueries.h
--- a/Engine/Source/Runtime/Experimental/Chaos/Public/Chaos/GeometryQueries.h
+++ b/Engine/Source/Runtime/Experimental/Chaos/Public/Chaos/GeometryQueries.h
@@ -18,7 +18,7 @@
 namespace Chaos
 {
 	template <typename T, int d>
-	bool OverlapQuery(const TImplicitObject<T, d>& A, const TRigidTransform<T,d>& ATM, const TImplicitObject<T, d>& B, const TRigidTransform<T,d>& BTM, const T Thickness = 0, TVector<T,d>* OutMTD=nullptr)
+	bool OverlapQuery(const FImplicitObject& A, const TRigidTransform<T,d>& ATM, const FImplicitObject& B, const TRigidTransform<T,d>& BTM, const T Thickness = 0, TVector<T,d>* OutMTD=nullptr)
 	{
 		const EImplicitObjectType AType = A.GetType(true);
 		const EImplicitObjectType BType = B.GetType(true);
@@ -84,7 +84,7 @@ namespace Chaos
 	}
 
 	template <typename T, int d>
-	bool SweepQuery(const TImplicitObject<T, d>& A, const TRigidTransform<T,d>& ATM, const TImplicitObject<T, d>& B, const TRigidTransform<T, d>& BTM, const TVector<T,d>& Dir, const T Length, T& OutTime, TVector<T,d>& OutPosition, TVector<T,d>& OutNormal, int32& OutFaceIndex, const T Thickness, const bool bComputeMTD)
+	bool SweepQuery(const FImplicitObject& A, const TRigidTransform<T,d>& ATM, const FImplicitObject& B, const TRigidTransform<T, d>& BTM, const TVector<T,d>& Dir, const T Length, T& OutTime, TVector<T,d>& OutPosition, TVector<T,d>& OutNormal, int32& OutFaceIndex, const T Thickness, const bool bComputeMTD)
 	{
 		const EImplicitObjectType AType = A.GetType(true);
 		const EImplicitObjectType BType = B.GetType(true);
@@ -131,7 +131,7 @@ namespace Chaos
 		//todo: handle case where A is a sphere
 		else if (A.IsConvex())
 		{
-			auto IsValidConvex = [](const TImplicitObject<T, 3>& InObject) -> bool
+			auto IsValidConvex = [](const FImplicitObject& InObject) -> bool
 			{
 				//todo: move this out of here
 				if (const auto Convex = TImplicitObjectScaled<TConvex<T,3>>::AsScaled(InObject))
diff --git a/Engine/Source/Runtime/Experimental/Chaos/Public/Chaos/HeightField.h b/Engine/Source/Runtime/Experimental/Chaos/Public/Chaos/HeightField.h
--- a/Engine/Source/Runtime/Experimental/Chaos/Public/Chaos/HeightField.h
+++ b/Engine/Source/Runtime/Experimental/Chaos/Public/Chaos/HeightField.h
@@ -22,10 +22,10 @@ namespace Chaos
 namespace Chaos
 {
 	template<typename T>
-	class CHAOS_API THeightField final : public TImplicitObject<T,3>
+	class CHAOS_API THeightField final : public FImplicitObject
 	{
 	public:
-		using TImplicitObject<T, 3>::GetTypeName;
+		using FImplicitObject::GetTypeName;
 
 		THeightField(TArray<T>&& Height, int32 InNumRows, int32 InNumCols, const TVector<T,3>& InScale);
 		THeightField(TArrayView<const uint16> InHeights, int32 InNumRows, int32 InNumCols, const TVector<T, 3>& InScale);
@@ -46,8 +46,8 @@ namespace Chaos
 
 		virtual bool Raycast(const TVector<T, 3>& StartPoint, const TVector<T, 3>& Dir, const T Length, const T Thickness, T& OutTime, TVector<T,3>& OutPosition, TVector<T,3>& OutNormal, int32& OutFaceIndex) const override;
 		virtual bool Overlap(const TVector<T, 3>& Point, const T Thickness) const override;
-		bool OverlapGeom(const TImplicitObject<T, 3>& QueryGeom, const TRigidTransform<T, 3>& QueryTM, const T Thickness) const;
-		bool SweepGeom(const TImplicitObject<T, 3>& QueryGeom, const TRigidTransform<T, 3>& StartTM, const TVector<T, 3>& Dir, const T Length, T& OutTime, TVector<T, 3>& OutPosition, TVector<T, 3>& OutNormal, int32& OutFaceIndex, const T Thickness = 0) const;
+		bool OverlapGeom(const FImplicitObject& QueryGeom, const TRigidTransform<T, 3>& QueryTM, const T Thickness) const;
+		bool SweepGeom(const FImplicitObject& QueryGeom, const TRigidTransform<T, 3>& StartTM, const TVector<T, 3>& Dir, const T Length, T& OutTime, TVector<T, 3>& OutPosition, TVector<T, 3>& OutNormal, int32& OutFaceIndex, const T Thickness = 0) const;
 		virtual int32 FindMostOpposingFace(const TVector<T, 3>& Position, const TVector<T, 3>& UnitDir, int32 HintFaceIndex, T SearchDist) const override;
 		virtual TVector<T, 3> FindGeometryOpposingNormal(const TVector<T, 3>& DenormDir, int32 FaceIndex, const TVector<T, 3>& OriginalNormal) const override;
 
@@ -78,7 +78,7 @@ namespace Chaos
 		virtual void Serialize(FChaosArchive& Ar) override
 		{
 			FChaosArchiveScopedMemory ScopedMemory(Ar, GetTypeName());
-			TImplicitObject<T, 3>::SerializeImp(Ar);
+			FImplicitObject::SerializeImp(Ar);
 			
 			GeomData.Serialize(Ar);
 
@@ -394,7 +394,7 @@ namespace Chaos
 		void BuildQueryData();
 		
 		// Needed for serialization
-		THeightField() : TImplicitObject<T, 3>(EImplicitObject::HasBoundingBox, ImplicitObjectType::HeightField) {}
-		friend TImplicitObject<T, 3>;
+		THeightField() : FImplicitObject(EImplicitObject::HasBoundingBox, ImplicitObjectType::HeightField) {}
+		friend FImplicitObject;
 	};
 }
diff --git a/Engine/Source/Runtime/Experimental/Chaos/Public/Chaos/ImplicitObject.h b/Engine/Source/Runtime/Experimental/Chaos/Public/Chaos/ImplicitObject.h
--- a/Engine/Source/Runtime/Experimental/Chaos/Public/Chaos/ImplicitObject.h
+++ b/Engine/Source/Runtime/Experimental/Chaos/Public/Chaos/ImplicitObject.h
@@ -1,13 +1,11 @@
 // Copyright 1998-2019 Epic Games, Inc. All Rights Reserved.
 #pragma once
 
-#include "Chaos/Matrix.h"
 #include "Chaos/Pair.h"
 #include "Chaos/Serializable.h"
-#include "Chaos/Vector.h"
+#include "Chaos/Core.h"
 
 #include <functional>
-#include "Transform.h"
 
 namespace Chaos
 {
@@ -23,8 +21,7 @@ template<class T, int d>
 class TParticles;
 template<class T, int d>
 class TBVHParticles;
-template<class T, int d>
-class TImplicitObject;
+class FImplicitObject;
 
 namespace ImplicitObjectType
 {
@@ -93,9 +90,9 @@ struct TImplicitObjectPtrStorage
 template<class T, int d>
 struct TImplicitObjectPtrStorage<T, d, false>
 {
-	using PtrType = TImplicitObject<T, d>*;
+	using PtrType = FImplicitObject*;
 
-	static PtrType Convert(const TUniquePtr<TImplicitObject<T, d>>& Object)
+	static PtrType Convert(const TUniquePtr<FImplicitObject>& Object)
 	{
 		return Object.Get();
 	}
@@ -104,28 +101,25 @@ struct TImplicitObjectPtrStorage<T, d, false>
 template<class T, int d>
 struct TImplicitObjectPtrStorage<T, d, true>
 {
-	using PtrType = TSerializablePtr<TImplicitObject<T, d>>;
+	using PtrType = TSerializablePtr<FImplicitObject>;
 
-	static PtrType Convert(const TUniquePtr<TImplicitObject<T, d>>& Object)
+	static PtrType Convert(const TUniquePtr<FImplicitObject>& Object)
 	{
 		return MakeSerializable(Object);
 	}
 };
 
-
-
-template<class T, int d>
-class CHAOS_API TImplicitObject
+class CHAOS_API FImplicitObject
 {
 public:
-	using TType = T;
-	static constexpr int D = d;
-	static TImplicitObject<T,d>* SerializationFactory(FChaosArchive& Ar, TImplicitObject<T, d>* Obj);
+	using TType = FReal;
+	static constexpr int D = 3;
+	static FImplicitObject* SerializationFactory(FChaosArchive& Ar, FImplicitObject* Obj);
 
-	TImplicitObject(int32 Flags, EImplicitObjectType InType = ImplicitObjectType::Unknown);
-	TImplicitObject(const TImplicitObject<T, d>&) = delete;
-	TImplicitObject(TImplicitObject<T, d>&&) = delete;
-	virtual ~TImplicitObject();
+	FImplicitObject(int32 Flags, EImplicitObjectType InType = ImplicitObjectType::Unknown);
+	FImplicitObject(const FImplicitObject&) = delete;
+	FImplicitObject(FImplicitObject&&) = delete;
+	virtual ~FImplicitObject();
 
 	template<class T_DERIVED>
 	T_DERIVED* GetObject()
@@ -165,20 +159,24 @@ public:
 
 	virtual bool IsValidGeometry() const;
 
-	virtual TUniquePtr<TImplicitObject<T, d>> Copy() const;
+	virtual TUniquePtr<FImplicitObject> Copy() const;
 
 	//This is strictly used for optimization purposes
 	bool IsUnderlyingUnion() const;
 
 	// Explicitly non-virtual.  Must cast to derived types to target their implementation.
-	T SignedDistance(const TVector<T, d>& x) const;
+	FReal SignedDistance(const FVec3& x) const;
 
 	// Explicitly non-virtual.  Must cast to derived types to target their implementation.
-	TVector<T, d> Normal(const TVector<T, d>& x) const;
-	virtual T PhiWithNormal(const TVector<T, d>& x, TVector<T, d>& Normal) const = 0;
-	virtual const class TBox<T, d>& BoundingBox() const;
-	virtual TVector<T, d> Support(const TVector<T, d>& Direction, const T Thickness) const;
-	virtual TVector<T, d> Support2(const TVector<T, d>& Direction) const { return Support(Direction, 0); }
+	FVec3 Normal(const FVec3& x) const;
+	virtual FReal PhiWithNormal(const FVec3& x, FVec3& Normal) const = 0;
+	virtual const class TBox<FReal, 3>& BoundingBox() const;
+	virtual FReal GetMargin() const
+	{
+		return 0;
+	}
+	virtual FVec3 Support(const FVec3& Direction, const FReal Thickness) const;
+	virtual FVec3 Support2(const FVec3& Direction) const { return Support(Direction, 0); }
 	bool HasBoundingBox() const { return bHasBoundingBox; }
 	bool IsConvex() const { return bIsConvex; }
 	void IgnoreAnalyticCollisions(const bool Ignore = true) { bIgnoreAnalyticCollisions = Ignore; }
@@ -190,16 +188,16 @@ public:
 		return FString::Printf(TEXT("ImplicitObject - No Performance String"));
 	};
 
-	Pair<TVector<T, d>, bool> FindDeepestIntersection(const TImplicitObject<T, d>* Other, const TBVHParticles<float, d>* Particles, const PMatrix<T, d, d>& OtherToLocalTransform, const T Thickness) const;
-	Pair<TVector<T, d>, bool> FindDeepestIntersection(const TImplicitObject<T, d>* Other, const TParticles<float, d>* Particles, const PMatrix<T, d, d>& OtherToLocalTransform, const T Thickness) const;
-	Pair<TVector<T, d>, bool> FindClosestIntersection(const TVector<T, d>& StartPoint, const TVector<T, d>& EndPoint, const T Thickness) const;
+	Pair<FVec3, bool> FindDeepestIntersection(const FImplicitObject* Other, const TBVHParticles<FReal, 3>* Particles, const FMatrix33& OtherToLocalTransform, const FReal Thickness) const;
+	Pair<FVec3, bool> FindDeepestIntersection(const FImplicitObject* Other, const TParticles<FReal, 3>* Particles, const FMatrix33& OtherToLocalTransform, const FReal Thickness) const;
+	Pair<FVec3, bool> FindClosestIntersection(const FVec3& StartPoint, const FVec3& EndPoint, const FReal Thickness) const;
 
 	//This gives derived types a way to avoid calling PhiWithNormal todo: this api is confusing
-	virtual bool Raycast(const TVector<T, d>& StartPoint, const TVector<T,d>& Dir, const T Length, const T Thickness, T& OutTime, TVector<T,d>& OutPosition, TVector<T,d>& OutNormal, int32& OutFaceIndex) const
+	virtual bool Raycast(const FVec3& StartPoint, const FVec3& Dir, const FReal Length, const FReal Thickness, FReal& OutTime, FVec3& OutPosition, FVec3& OutNormal, int32& OutFaceIndex) const
 	{
 		OutFaceIndex = INDEX_NONE;
-		const TVector<T, d> EndPoint = StartPoint + Dir * Length;
-		Pair<TVector<T,d>, bool> Result = FindClosestIntersection(StartPoint, EndPoint, Thickness);
+		const FVec3 EndPoint = StartPoint + Dir * Length;
+		Pair<FVec3, bool> Result = FindClosestIntersection(StartPoint, EndPoint, Thickness);
 		if (Result.Second)
 		{
 			OutPosition = Result.First;
@@ -216,7 +214,7 @@ public:
 		@param HintFaceIndex - for certain geometry we can use this to accelerate the search.
 		@return Index of the most opposing face
 	*/
-	virtual int32 FindMostOpposingFace(const TVector<T, d>& Position, const TVector<T, d>& UnitDir, int32 HintFaceIndex, T SearchDist) const
+	virtual int32 FindMostOpposingFace(const FVec3& Position, const FVec3& UnitDir, int32 HintFaceIndex, FReal SearchDist) const
 	{
 		//Many objects have no concept of a face
 		return INDEX_NONE;
@@ -229,31 +227,29 @@ public:
 		@param OriginalNormal - the original normal given by something like a sphere sweep
 		@return The most opposing normal associated with the underlying geometry's feature (like a face)
 	*/
-	virtual TVector<T,d> FindGeometryOpposingNormal(const TVector<T, d>& DenormDir, int32 FaceIndex, const TVector<T,d>& OriginalNormal) const
+	virtual FVec3 FindGeometryOpposingNormal(const FVec3& DenormDir, int32 FaceIndex, const FVec3& OriginalNormal) const
 	{
 		//Many objects have no concept of a face
 		return OriginalNormal;
 	}
 
 	//This gives derived types a way to do an overlap check without calling PhiWithNormal todo: this api is confusing
-	virtual bool Overlap(const TVector<T, d>& Point, const T Thickness) const; 
-
-	virtual void AccumulateAllImplicitObjects(TArray<Pair<const TImplicitObject<T, d>*, TRigidTransform<T, d>>>& Out, const TRigidTransform<T, d>& ParentTM) const
+	virtual bool Overlap(const FVec3& Point, const FReal Thickness) const
 	{
-		Out.Add(MakePair(this, ParentTM));
+		return SignedDistance(Point) <= Thickness;
 	}
 
-	virtual void AccumulateAllSerializableImplicitObjects(TArray<Pair<TSerializablePtr<TImplicitObject<T, d>>, TRigidTransform<T, d>>>& Out, const TRigidTransform<T, d>& ParentTM, TSerializablePtr<TImplicitObject<T,d>> This) const
+	virtual void AccumulateAllImplicitObjects(TArray<Pair<const FImplicitObject*, FRigidTransform3>>& Out, const FRigidTransform3& ParentTM) const
 	{
-		Out.Add(MakePair(This, ParentTM));
+		Out.Add(MakePair(this, ParentTM));
 	}
 
-	virtual T GetMargin() const
+	virtual void AccumulateAllSerializableImplicitObjects(TArray<Pair<TSerializablePtr<FImplicitObject>, FRigidTransform3>>& Out, const FRigidTransform3& ParentTM, TSerializablePtr<FImplicitObject> This) const
 	{
-		return 0;
+		Out.Add(MakePair(This, ParentTM));
 	}
 
-	virtual void FindAllIntersectingObjects(TArray < Pair<const TImplicitObject<T, d>*, TRigidTransform<T, d>>>& Out, const TBox<T, d>& LocalBounds) const;
+	virtual void FindAllIntersectingObjects(TArray < Pair<const FImplicitObject*, FRigidTransform3>>& Out, const TBox<FReal, 3>& LocalBounds) const;
 
 	virtual FString ToString() const
 	{
@@ -274,7 +270,7 @@ public:
 
 	virtual void Serialize(FChaosArchive& Ar);
 	
-	static FArchive& SerializeLegacyHelper(FArchive& Ar, TUniquePtr<TImplicitObject<T, d>>& Value);
+	static FArchive& SerializeLegacyHelper(FArchive& Ar, TUniquePtr<FImplicitObject>& Value);
 
 	virtual uint32 GetTypeHash() const = 0;
 
@@ -289,22 +285,20 @@ protected:
 	bool bHasBoundingBox;
 
 private:
-	virtual Pair<TVector<T, d>, bool> FindClosestIntersectionImp(const TVector<T, d>& StartPoint, const TVector<T, d>& EndPoint, const T Thickness) const;
+	virtual Pair<FVec3, bool> FindClosestIntersectionImp(const FVec3& StartPoint, const FVec3& EndPoint, const FReal Thickness) const;
 };
 
-template <typename T, int d>
-FORCEINLINE FChaosArchive& operator<<(FChaosArchive& Ar, TImplicitObject<T, d>& Value)
+FORCEINLINE FChaosArchive& operator<<(FChaosArchive& Ar, FImplicitObject& Value)
 {
 	Value.Serialize(Ar);
 	return Ar;
 }
 
-template <typename T, int d>
-FORCEINLINE FArchive& operator<<(FArchive& Ar, TImplicitObject<T, d>& Value)
+FORCEINLINE FArchive& operator<<(FArchive& Ar, FImplicitObject& Value)
 {
 	Value.Serialize(Ar);
 	return Ar;
 }
 
-typedef TImplicitObject<float, 3> FImplicitObject3;
+typedef FImplicitObject FImplicitObject3;
 }
\ No newline at end of file
diff --git a/Engine/Source/Runtime/Experimental/Chaos/Public/Chaos/ImplicitObjectIntersection.h b/Engine/Source/Runtime/Experimental/Chaos/Public/Chaos/ImplicitObjectIntersection.h
--- a/Engine/Source/Runtime/Experimental/Chaos/Public/Chaos/ImplicitObjectIntersection.h
+++ b/Engine/Source/Runtime/Experimental/Chaos/Public/Chaos/ImplicitObjectIntersection.h
@@ -9,13 +9,13 @@
 namespace Chaos
 {
 template<class T, int d>
-class TImplicitObjectIntersection : public TImplicitObject<T, d>
+class TImplicitObjectIntersection : public FImplicitObject
 {
   public:
-	using TImplicitObject<T, d>::SignedDistance;
+	using FImplicitObject::SignedDistance;
 
-	TImplicitObjectIntersection(TArray<TUniquePtr<TImplicitObject<T, d>>>&& Objects)
-	    : TImplicitObject<T, d>(EImplicitObject::HasBoundingBox)
+	TImplicitObjectIntersection(TArray<TUniquePtr<FImplicitObject>>&& Objects)
+	    : FImplicitObject(EImplicitObject::HasBoundingBox)
 	    , MObjects(MoveTemp(Objects))
 	    , MLocalBoundingBox(MObjects[0]->BoundingBox())
 	{
@@ -26,7 +26,7 @@ class TImplicitObjectIntersection : public TImplicitObject<T, d>
 	}
 	TImplicitObjectIntersection(const TImplicitObjectIntersection<T, d>& Other) = delete;
 	TImplicitObjectIntersection(TImplicitObjectIntersection<T, d>&& Other)
-	    : TImplicitObject<T, d>(EImplicitObject::HasBoundingBox)
+	    : FImplicitObject(EImplicitObject::HasBoundingBox)
 	    , MObjects(MoveTemp(Other.MObjects))
 	    , MLocalBoundingBox(MoveTemp(Other.MLocalBoundingBox))
 	{
@@ -62,7 +62,7 @@ class TImplicitObjectIntersection : public TImplicitObject<T, d>
 	{
 		uint32 OutHash = MObjects.Num() > 0 ? MObjects[0]->GetTypeHash() : 0;
 
-		for(const TUniquePtr<TImplicitObject<T, d>>& Ptr : MObjects)
+		for(const TUniquePtr<FImplicitObject>& Ptr : MObjects)
 		{
 			OutHash = HashCombine(Ptr->GetTypeHash(), OutHash);
 		}
@@ -93,7 +93,7 @@ private:
 	}
 
   private:
-	TArray<TUniquePtr<TImplicitObject<T, d>>> MObjects;
+	TArray<TUniquePtr<FImplicitObject>> MObjects;
 	TBox<T, d> MLocalBoundingBox;
 };
 }
diff --git a/Engine/Source/Runtime/Experimental/Chaos/Public/Chaos/ImplicitObjectScaled.h b/Engine/Source/Runtime/Experimental/Chaos/Public/Chaos/ImplicitObjectScaled.h
--- a/Engine/Source/Runtime/Experimental/Chaos/Public/Chaos/ImplicitObjectScaled.h
+++ b/Engine/Source/Runtime/Experimental/Chaos/Public/Chaos/ImplicitObjectScaled.h
@@ -13,7 +13,7 @@ namespace Chaos
 {
 
 template <typename TConcrete, bool bInstanced = true>
-class TImplicitObjectInstanced final : public TImplicitObject<typename TConcrete::TType,TConcrete::D>
+class TImplicitObjectInstanced final : public FImplicitObject
 {
 public:
 	using T = typename TConcrete::TType;
@@ -22,10 +22,10 @@ public:
 	static constexpr int D = d;
 
 	using ObjectType = typename TChooseClass<bInstanced, TSerializablePtr<TConcrete>, TUniquePtr<TConcrete>>::Result;
-	using TImplicitObject<T, d>::GetTypeName;
+	using FImplicitObject::GetTypeName;
 
 	TImplicitObjectInstanced(ObjectType Object)
-		: TImplicitObject<T, d>(EImplicitObject::HasBoundingBox, Object->GetType() | ImplicitObjectType::IsInstanced)
+		: FImplicitObject(EImplicitObject::HasBoundingBox, Object->GetType() | ImplicitObjectType::IsInstanced)
 		, MObject(MoveTemp(Object))
 	{
 		ensure(IsInstanced(MObject->GetType(true)) == false);	//cannot have an instance of an instance
@@ -55,20 +55,20 @@ public:
 	virtual void Serialize(FChaosArchive& Ar) override
 	{
 		FChaosArchiveScopedMemory ScopedMemory(Ar, GetTypeName(), false);
-		TImplicitObject<T, d>::SerializeImp(Ar);
+		FImplicitObject::SerializeImp(Ar);
 		Ar << MObject;
 	}
 
 	/** This is a low level function and assumes the internal object has a SweepGeom function. Should not be called directly. See GeometryQueries.h : SweepQuery */
 	template <typename TGeometry>
-	static bool LowLevelSweepGeom(const TImplicitObjectInstanced<TConcrete, bInstanced>& OwningInstanced, const TImplicitObject<T, d>& B, const TRigidTransform<T, d>& BToATM, const TVector<T, d>& LocalDir, const T Length, T& OutTime, TVector<T, d>& LocalPosition, TVector<T, d>& LocalNormal, int32& OutFaceIndex, T Thickness = 0)
+	static bool LowLevelSweepGeom(const TImplicitObjectInstanced<TConcrete, bInstanced>& OwningInstanced, const FImplicitObject& B, const TRigidTransform<T, d>& BToATM, const TVector<T, d>& LocalDir, const T Length, T& OutTime, TVector<T, d>& LocalPosition, TVector<T, d>& LocalNormal, int32& OutFaceIndex, T Thickness = 0)
 	{
 		ensure(false);
 		return false;
 	}
 
 	/** This is a low level function and assumes the internal object has a OverlapGeom function. Should not be called directly. See GeometryQueries.h : OverlapQuery */
-	bool LowLevelOverlapGeom(const TImplicitObjectInstanced<TConcrete, bInstanced>& OwningInstanced, const TImplicitObject<T, d>& B, const TRigidTransform<T, d>& BToATM, T Thickness = 0) const
+	bool LowLevelOverlapGeom(const TImplicitObjectInstanced<TConcrete, bInstanced>& OwningInstanced, const FImplicitObject& B, const TRigidTransform<T, d>& BToATM, T Thickness = 0) const
 	{
 		ensure(false);
 		return false;
@@ -103,9 +103,9 @@ public:
 		return MObject->GetTypeHash();
 	}
 
-	virtual TUniquePtr<TImplicitObject<T, d>> Copy() const override
+	virtual TUniquePtr<FImplicitObject> Copy() const override
 	{
-		return TUniquePtr<TImplicitObject<T, d>>(CopyHelper(this));
+		return TUniquePtr<FImplicitObject>(CopyHelper(this));
 	}
 
 protected:
@@ -124,7 +124,7 @@ protected:
 
 
 template<typename TConcrete, bool bInstanced = true>
-class TImplicitObjectScaled final : public TImplicitObject<typename TConcrete::TType, TConcrete::D>
+class TImplicitObjectScaled final : public FImplicitObject
 {
 public:
 	using T = typename TConcrete::TType;
@@ -133,10 +133,10 @@ public:
 	static constexpr int D = d;
 
 	using ObjectType = typename TChooseClass<bInstanced, TSerializablePtr<TConcrete>, TUniquePtr<TConcrete>>::Result;
-	using TImplicitObject<T, d>::GetTypeName;
+	using FImplicitObject::GetTypeName;
 
 	TImplicitObjectScaled(ObjectType Object, const TVector<T, d>& Scale, T Thickness = 0)
-	    : TImplicitObject<T, d>(EImplicitObject::HasBoundingBox, Object->GetType() | ImplicitObjectType::IsScaled)
+	    : FImplicitObject(EImplicitObject::HasBoundingBox, Object->GetType() | ImplicitObjectType::IsScaled)
 	    , MObject(MoveTemp(Object))
 		, MInternalThickness(Thickness)
 	{
@@ -153,8 +153,8 @@ public:
 		this->bIsConvex = MObject->IsConvex();
 		SetScale(Scale);
 	}
-	TImplicitObjectScaled(ObjectType Object, TUniquePtr<Chaos::TImplicitObject<T,d>> &&ObjectOwner, const TVector<T, d>& Scale, T Thickness = 0)
-	    : TImplicitObject<T, d>(EImplicitObject::HasBoundingBox, Object->GetType() | ImplicitObjectType::IsScaled)
+	TImplicitObjectScaled(ObjectType Object, TUniquePtr<Chaos::FImplicitObject> &&ObjectOwner, const TVector<T, d>& Scale, T Thickness = 0)
+	    : FImplicitObject(EImplicitObject::HasBoundingBox, Object->GetType() | ImplicitObjectType::IsScaled)
 	    , MObject(Object)
 		, MInternalThickness(Thickness)
 	{
@@ -166,7 +166,7 @@ public:
 
 	TImplicitObjectScaled(const TImplicitObjectScaled<TConcrete, bInstanced>& Other) = delete;
 	TImplicitObjectScaled(TImplicitObjectScaled<TConcrete, bInstanced>&& Other)
-	    : TImplicitObject<T, d>(EImplicitObject::HasBoundingBox, Other.MObject->GetType() | ImplicitObjectType::IsScaled)
+	    : FImplicitObject(EImplicitObject::HasBoundingBox, Other.MObject->GetType() | ImplicitObjectType::IsScaled)
 	    , MObject(MoveTemp(Other.MObject))
 	    , MScale(Other.MScale)
 		, MInvScale(Other.MInvScale)
@@ -184,9 +184,9 @@ public:
 		return TConcrete::StaticType() | ImplicitObjectType::IsScaled;
 	}
 
-	static const TImplicitObjectScaled<TConcrete>& AsScaledChecked(const TImplicitObject<T, d>& Obj)
+	static const TImplicitObjectScaled<TConcrete>& AsScaledChecked(const FImplicitObject& Obj)
 	{
-		if (TIsSame<TConcrete, TImplicitObject<T, d>>::Value)
+		if (TIsSame<TConcrete, FImplicitObject>::Value)
 		{
 			//can cast any scaled to ImplicitObject base
 			check(IsScaled(Obj.GetType()));
@@ -198,9 +198,9 @@ public:
 		return static_cast<const TImplicitObjectScaled<TConcrete>&>(Obj);
 	}
 
-	static TImplicitObjectScaled<TConcrete>& AsScaledChecked(TImplicitObject<T, d>& Obj)
+	static TImplicitObjectScaled<TConcrete>& AsScaledChecked(FImplicitObject& Obj)
 	{
-		if (TIsSame<TConcrete, TImplicitObject<T, d>>::Value)
+		if (TIsSame<TConcrete, FImplicitObject>::Value)
 		{
 			//can cast any scaled to ImplicitObject base
 			check(IsScaled(Obj.GetType()));
@@ -212,9 +212,9 @@ public:
 		return static_cast<TImplicitObjectScaled<TConcrete>&>(Obj);
 	}
 
-	static const TImplicitObjectScaled<TConcrete>* AsScaled(const TImplicitObject<T, d>& Obj)
+	static const TImplicitObjectScaled<TConcrete>* AsScaled(const FImplicitObject& Obj)
 	{
-		if (TIsSame<TConcrete, TImplicitObject<T, d>>::Value)
+		if (TIsSame<TConcrete, FImplicitObject>::Value)
 		{
 			//can cast any scaled to ImplicitObject base
 			return IsScaled(Obj.GetType()) ? static_cast<const TImplicitObjectScaled<TConcrete>*>(&Obj) : nullptr;
@@ -225,9 +225,9 @@ public:
 		}
 	}
 
-	static TImplicitObjectScaled<TConcrete>* AsScaled(TImplicitObject<T, d>& Obj)
+	static TImplicitObjectScaled<TConcrete>* AsScaled(FImplicitObject& Obj)
 	{
-		if (TIsSame<TConcrete, TImplicitObject<T, d>>::Value)
+		if (TIsSame<TConcrete, FImplicitObject>::Value)
 		{
 			//can cast any scaled to ImplicitObject base
 			return IsScaled(Obj.GetType()) ? static_cast<TImplicitObjectScaled<TConcrete>*>(&Obj) : nullptr;
@@ -287,7 +287,7 @@ public:
 	}
 
 	/** This is a low level function and assumes the internal object has a SweepGeom function. Should not be called directly. See GeometryQueries.h : SweepQuery */
-	bool LowLevelSweepGeom(const TImplicitObject<T, d>& B, const TRigidTransform<T, d>& BToATM, const TVector<T, d>& LocalDir, const T Length, T& OutTime, TVector<T, d>& LocalPosition, TVector<T, d>& LocalNormal, int32& OutFaceIndex, T Thickness = 0) const
+	bool LowLevelSweepGeom(const FImplicitObject& B, const TRigidTransform<T, d>& BToATM, const TVector<T, d>& LocalDir, const T Length, T& OutTime, TVector<T, d>& LocalPosition, TVector<T, d>& LocalNormal, int32& OutFaceIndex, T Thickness = 0) const
 	{
 		ensure(Length > 0);
 		ensure(FMath::IsNearlyEqual(LocalDir.SizeSquared(), 1, KINDA_SMALL_NUMBER));
@@ -305,8 +305,8 @@ public:
 			TVector<T, d> UnscaledNormal;
 			float UnscaledTime;
 
-			TUniquePtr<TImplicitObject<T, d>> HackBPtr(const_cast<TImplicitObject<T,d>*>(&B));	//todo: hack, need scaled object to accept raw ptr similar to transformed implicit
-			TImplicitObjectScaled<TImplicitObject<T, d>> ScaledB(MakeSerializable(HackBPtr), MInvScale);
+			TUniquePtr<FImplicitObject> HackBPtr(const_cast<FImplicitObject*>(&B));	//todo: hack, need scaled object to accept raw ptr similar to transformed implicit
+			TImplicitObjectScaled<FImplicitObject> ScaledB(MakeSerializable(HackBPtr), MInvScale);
 			HackBPtr.Release();
 
 			TRigidTransform<T, d> BToATMNoScale(BToATM.GetLocation() * MInvScale, BToATM.GetRotation());
@@ -325,12 +325,12 @@ public:
 	}
 
 	/** This is a low level function and assumes the internal object has a OverlapGeom function. Should not be called directly. See GeometryQueries.h : OverlapQuery */
-	bool LowLevelOverlapGeom(const TImplicitObject<T, d>& B, const TRigidTransform<T, d>& BToATM, T Thickness = 0) const
+	bool LowLevelOverlapGeom(const FImplicitObject& B, const TRigidTransform<T, d>& BToATM, T Thickness = 0) const
 	{
 		ensure(Thickness == 0 || (FMath::IsNearlyEqual(MScale[0], MScale[1]) && FMath::IsNearlyEqual(MScale[0], MScale[2])));
 
-		TUniquePtr<TImplicitObject<T, d>> HackBPtr(const_cast<TImplicitObject<T, d>*>(&B));	//todo: hack, need scaled object to except raw ptr similar to transformed implicit
-		TImplicitObjectScaled<TImplicitObject<T, d>> ScaledB(MakeSerializable(HackBPtr), MInvScale);
+		TUniquePtr<FImplicitObject> HackBPtr(const_cast<FImplicitObject*>(&B));	//todo: hack, need scaled object to except raw ptr similar to transformed implicit
+		TImplicitObjectScaled<FImplicitObject> ScaledB(MakeSerializable(HackBPtr), MInvScale);
 		HackBPtr.Release();
 
 		TRigidTransform<T, d> BToATMNoScale(BToATM.GetLocation() * MInvScale, BToATM.GetRotation());
@@ -436,7 +436,7 @@ public:
 	virtual void Serialize(FChaosArchive& Ar) override
 	{
 		FChaosArchiveScopedMemory ScopedMemory(Ar, GetTypeName(), false);
-		TImplicitObject<T, d>::SerializeImp(Ar);
+		FImplicitObject::SerializeImp(Ar);
 		Ar << MObject << MScale << MInvScale << MLocalBoundingBox;
 		ensure(MInternalThickness == 0);	//not supported: do we care?
 	}
@@ -447,9 +447,9 @@ public:
 	}
 
 #if 0
-	virtual TUniquePtr<TImplicitObject<T, d>> Copy() const override
+	virtual TUniquePtr<FImplicitObject> Copy() const override
 	{
-		return TUniquePtr<TImplicitObject<T, d>>(CopyHelper(this));
+		return TUniquePtr<FImplicitObject>(CopyHelper(this));
 	}
 #endif
 private:
@@ -461,13 +461,12 @@ private:
 
 	//needed for serialization
 	TImplicitObjectScaled()
-	: TImplicitObject<T, d>(EImplicitObject::HasBoundingBox, StaticType())
+	: FImplicitObject(EImplicitObject::HasBoundingBox, StaticType())
 	, MInternalThickness(0)
 	{}
-	friend TImplicitObject<T, d>;	//needed for serialization
+	friend FImplicitObject;	//needed for serialization
 
-	template <typename TConcrete2, bool bInstanced2>
-	friend class TImplicitObjectScaled;
+	friend class FImplicitObjectScaled;
 
 	static TImplicitObjectScaled<TConcrete, true>* CopyHelper(const TImplicitObjectScaled<TConcrete, true>* Obj)
 	{
@@ -493,6 +492,6 @@ template <typename TConcrete>
 using TImplicitObjectScaledNonSerializable = TImplicitObjectScaled<TConcrete, false>;
 
 template <typename T, int d>
-using TImplicitObjectScaledGeneric = TImplicitObjectScaled<TImplicitObject<T, d>>;
+using TImplicitObjectScaledGeneric = TImplicitObjectScaled<FImplicitObject>;
 
 }
diff --git a/Engine/Source/Runtime/Experimental/Chaos/Public/Chaos/ImplicitObjectTransformed.h b/Engine/Source/Runtime/Experimental/Chaos/Public/Chaos/ImplicitObjectTransformed.h
--- a/Engine/Source/Runtime/Experimental/Chaos/Public/Chaos/ImplicitObjectTransformed.h
+++ b/Engine/Source/Runtime/Experimental/Chaos/Public/Chaos/ImplicitObjectTransformed.h
@@ -10,26 +10,22 @@
 namespace Chaos
 {
 
-template <typename T, int d>
-void TImplicitObjectTransformSerializeHelper(FChaosArchive& Ar, TSerializablePtr<TImplicitObject<T, d>>& Obj)
+inline void TImplicitObjectTransformSerializeHelper(FChaosArchive& Ar, TSerializablePtr<FImplicitObject>& Obj)
 {
 	Ar << Obj;
 }
 
-template <typename T, int d>
-void TImplicitObjectTransformSerializeHelper(FChaosArchive& Ar, const TImplicitObject<T, d>* Obj)
+inline void TImplicitObjectTransformSerializeHelper(FChaosArchive& Ar, const FImplicitObject* Obj)
 {
 	check(false);
 }
 
-template <typename T, int d>
-void TImplicitObjectTransformAccumulateSerializableHelper(TArray<Pair<TSerializablePtr<TImplicitObject<T, d>>, TRigidTransform<T, d>>>& Out, TSerializablePtr<TImplicitObject<T, d>> Obj, const TRigidTransform<T, d>& NewTM)
+inline void TImplicitObjectTransformAccumulateSerializableHelper(TArray<Pair<TSerializablePtr<FImplicitObject>, FRigidTransform3>>& Out, TSerializablePtr<FImplicitObject> Obj, const FRigidTransform3& NewTM)
 {
 	Obj->AccumulateAllSerializableImplicitObjects(Out, NewTM, Obj);
 }
 
-template <typename T, int d>
-void TImplicitObjectTransformAccumulateSerializableHelper(TArray<Pair<TSerializablePtr<TImplicitObject<T, d>>, TRigidTransform<T, d>>>& Out, const TImplicitObject<T, d>* Obj, const TRigidTransform<T, d>& NewTM)
+inline void TImplicitObjectTransformAccumulateSerializableHelper(TArray<Pair<TSerializablePtr<FImplicitObject>, FRigidTransform3>>& Out, const FImplicitObject* Obj, const FRigidTransform3& NewTM)
 {
 	check(false);
 }
@@ -40,19 +36,19 @@ void TImplicitObjectTransformAccumulateSerializableHelper(TArray<Pair<TSerializa
  * @template bSerializable Whether the shape can be serialized (usually true). Set to false for transient/stack-allocated objects. 
  */
 template<class T, int d, bool bSerializable = true>
-class TImplicitObjectTransformed final : public TImplicitObject<T, d>
+class TImplicitObjectTransformed final : public FImplicitObject
 {
 	using FStorage = TImplicitObjectPtrStorage<T, d, bSerializable>;
 	using ObjectType = typename FStorage::PtrType;
 
 public:
-	using TImplicitObject<T, d>::GetTypeName;
+	using FImplicitObject::GetTypeName;
 
 	/**
 	 * Create a transform around an ImplicitObject. Lifetime of the wrapped object is managed externally.
 	 */
 	TImplicitObjectTransformed(ObjectType Object, const TRigidTransform<T, d>& InTransform)
-	    : TImplicitObject<T, d>(EImplicitObject::HasBoundingBox, ImplicitObjectType::Transformed)
+	    : FImplicitObject(EImplicitObject::HasBoundingBox, ImplicitObjectType::Transformed)
 	    , MObject(Object)
 	    , MTransform(InTransform)
 	    , MLocalBoundingBox(Object->BoundingBox().TransformedBox(InTransform))
@@ -63,8 +59,8 @@ public:
 	/**
 	 * Create a transform around an ImplicitObject and take control of its lifetime.
 	 */
-	TImplicitObjectTransformed(TUniquePtr<Chaos::TImplicitObject<T,d>> &&ObjectOwner, const TRigidTransform<T, d>& InTransform)
-	    : TImplicitObject<T, d>(EImplicitObject::HasBoundingBox, ImplicitObjectType::Transformed)
+	TImplicitObjectTransformed(TUniquePtr<Chaos::FImplicitObject> &&ObjectOwner, const TRigidTransform<T, d>& InTransform)
+	    : FImplicitObject(EImplicitObject::HasBoundingBox, ImplicitObjectType::Transformed)
 		, MObjectOwner(MoveTemp(ObjectOwner))
 	    , MTransform(InTransform)
 	{
@@ -76,7 +72,7 @@ public:
 
 	TImplicitObjectTransformed(const TImplicitObjectTransformed<T, d, bSerializable>& Other) = delete;
 	TImplicitObjectTransformed(TImplicitObjectTransformed<T, d, bSerializable>&& Other)
-	    : TImplicitObject<T, d>(EImplicitObject::HasBoundingBox, ImplicitObjectType::Transformed)
+	    : FImplicitObject(EImplicitObject::HasBoundingBox, ImplicitObjectType::Transformed)
 	    , MObject(Other.MObject)
 		, MObjectOwner(MoveTemp(Other.MObjectOwner))
 	    , MTransform(Other.MTransform)
@@ -91,7 +87,7 @@ public:
 		return ImplicitObjectType::Transformed;
 	}
 
-	const TImplicitObject<T, d>* GetTransformedObject() const
+	const FImplicitObject* GetTransformedObject() const
 	{
 		return MObject.Get();
 	}
@@ -166,20 +162,20 @@ public:
 		MTransform = InTransform;
 	}
 
-	virtual void AccumulateAllImplicitObjects(TArray<Pair<const TImplicitObject<T, d>*, TRigidTransform<T, d>>>& Out, const TRigidTransform<T, d>& ParentTM) const
+	virtual void AccumulateAllImplicitObjects(TArray<Pair<const FImplicitObject*, TRigidTransform<T, d>>>& Out, const TRigidTransform<T, d>& ParentTM) const
 	{
 		const TRigidTransform<T, d> NewTM = MTransform * ParentTM;
 		MObject->AccumulateAllImplicitObjects(Out, NewTM);
 	}
 
-	virtual void AccumulateAllSerializableImplicitObjects(TArray<Pair<TSerializablePtr<TImplicitObject<T, d>>, TRigidTransform<T, d>>>& Out, const TRigidTransform<T, d>& ParentTM, TSerializablePtr<TImplicitObject<T,d>> This) const
+	virtual void AccumulateAllSerializableImplicitObjects(TArray<Pair<TSerializablePtr<FImplicitObject>, TRigidTransform<T, d>>>& Out, const TRigidTransform<T, d>& ParentTM, TSerializablePtr<FImplicitObject> This) const
 	{
 		check(bSerializable);
 		const TRigidTransform<T, d> NewTM = MTransform * ParentTM;
 		TImplicitObjectTransformAccumulateSerializableHelper(Out, MObject, NewTM);
 	}
 
-	virtual void FindAllIntersectingObjects(TArray < Pair<const TImplicitObject<T, d>*, TRigidTransform<T, d>>>& Out, const TBox<T, d>& LocalBounds) const
+	virtual void FindAllIntersectingObjects(TArray < Pair<const FImplicitObject*, TRigidTransform<T, d>>>& Out, const TBox<T, d>& LocalBounds) const
 	{
 		const TBox<T, d> SubobjectBounds = LocalBounds.TransformedBox(MTransform.Inverse());
 		int32 NumOut = Out.Num();
@@ -198,7 +194,7 @@ public:
 	{
 		check(bSerializable);
 		FChaosArchiveScopedMemory ScopedMemory(Ar, GetTypeName(), false);
-		TImplicitObject<T, d>::SerializeImp(Ar);
+		FImplicitObject::SerializeImp(Ar);
 		TImplicitObjectTransformSerializeHelper(Ar, MObject);
 		Ar << MTransform;
 		Ar << MLocalBoundingBox;
@@ -212,13 +208,13 @@ public:
 
 private:
 	ObjectType MObject;
-	TUniquePtr<Chaos::TImplicitObject<T, d>> MObjectOwner;
+	TUniquePtr<Chaos::FImplicitObject> MObjectOwner;
 	TRigidTransform<T, d> MTransform;
 	TBox<T, d> MLocalBoundingBox;
 
 	//needed for serialization
-	TImplicitObjectTransformed() : TImplicitObject<T, d>(EImplicitObject::HasBoundingBox, ImplicitObjectType::Transformed) {}
-	friend TImplicitObject<T, d>;	//needed for serialization
+	TImplicitObjectTransformed() : FImplicitObject(EImplicitObject::HasBoundingBox, ImplicitObjectType::Transformed) {}
+	friend FImplicitObject;	//needed for serialization
 };
 
 template <typename T, int d>
diff --git a/Engine/Source/Runtime/Experimental/Chaos/Public/Chaos/ImplicitObjectUnion.h b/Engine/Source/Runtime/Experimental/Chaos/Public/Chaos/ImplicitObjectUnion.h
--- a/Engine/Source/Runtime/Experimental/Chaos/Public/Chaos/ImplicitObjectUnion.h
+++ b/Engine/Source/Runtime/Experimental/Chaos/Public/Chaos/ImplicitObjectUnion.h
@@ -12,14 +12,14 @@
 namespace Chaos
 {
 template<class T, int d>
-class TImplicitObjectUnion : public TImplicitObject<T, d>
+class TImplicitObjectUnion : public FImplicitObject
 {
   public:
 
-	using TImplicitObject<T, d>::GetTypeName;
+	using FImplicitObject::GetTypeName;
 
-	TImplicitObjectUnion(TArray<TUniquePtr<TImplicitObject<T, d>>>&& Objects, const TArray<int32>& OriginalParticleLookupHack = TArray<int32>())
-	    : TImplicitObject<T, d>(EImplicitObject::HasBoundingBox, ImplicitObjectType::Union)
+	TImplicitObjectUnion(TArray<TUniquePtr<FImplicitObject>>&& Objects, const TArray<int32>& OriginalParticleLookupHack = TArray<int32>())
+	    : FImplicitObject(EImplicitObject::HasBoundingBox, ImplicitObjectType::Union)
 	    , MObjects(MoveTemp(Objects))
 		, Hierarchy(GeomParticles)
 	    , MLocalBoundingBox()
@@ -57,7 +57,7 @@ class TImplicitObjectUnion : public TImplicitObject<T, d>
 	}
 	TImplicitObjectUnion(const TImplicitObjectUnion<T, d>& Other) = delete;
 	TImplicitObjectUnion(TImplicitObjectUnion<T, d>&& Other)
-	    : TImplicitObject<T, d>(EImplicitObject::HasBoundingBox)
+	    : FImplicitObject(EImplicitObject::HasBoundingBox)
 	    , MObjects(MoveTemp(Other.MObjects))
 		, GeomParticles(MoveTemp(Other.GeomParticles))
 		, Hierarchy(MoveTemp(Other.Hierarchy))
@@ -107,28 +107,28 @@ class TImplicitObjectUnion : public TImplicitObject<T, d>
 
 	virtual const TBox<T, d>& BoundingBox() const override { return MLocalBoundingBox; }
 
-	virtual void AccumulateAllImplicitObjects(TArray<Pair<const TImplicitObject<T, d>*, TRigidTransform<T, d>>>& Out, const TRigidTransform<T, d>& ParentTM) const
+	virtual void AccumulateAllImplicitObjects(TArray<Pair<const FImplicitObject*, TRigidTransform<T, d>>>& Out, const TRigidTransform<T, d>& ParentTM) const
 	{
-		for (const TUniquePtr<TImplicitObject<T, d>>& Object : MObjects)
+		for (const TUniquePtr<FImplicitObject>& Object : MObjects)
 		{
 			Object->AccumulateAllImplicitObjects(Out, ParentTM);
 		}
 	}
 
-	virtual void AccumulateAllSerializableImplicitObjects(TArray<Pair<TSerializablePtr<TImplicitObject<T, d>>, TRigidTransform<T, d>>>& Out, const TRigidTransform<T, d>& ParentTM, TSerializablePtr<TImplicitObject<T,d>> This) const
+	virtual void AccumulateAllSerializableImplicitObjects(TArray<Pair<TSerializablePtr<FImplicitObject>, TRigidTransform<T, d>>>& Out, const TRigidTransform<T, d>& ParentTM, TSerializablePtr<FImplicitObject> This) const
 	{
 		AccumulateAllSerializableImplicitObjectsHelper(Out, ParentTM);
 	}
 
-	void AccumulateAllSerializableImplicitObjectsHelper(TArray<Pair<TSerializablePtr<TImplicitObject<T, d>>, TRigidTransform<T, d>>>& Out, const TRigidTransform<T, d>& ParentTM) const
+	void AccumulateAllSerializableImplicitObjectsHelper(TArray<Pair<TSerializablePtr<FImplicitObject>, TRigidTransform<T, d>>>& Out, const TRigidTransform<T, d>& ParentTM) const
 	{
-		for (const TUniquePtr<TImplicitObject<T, d>>& Object : MObjects)
+		for (const TUniquePtr<FImplicitObject>& Object : MObjects)
 		{
 			Object->AccumulateAllSerializableImplicitObjects(Out, ParentTM, MakeSerializable(Object));
 		}
 	}
 
-	virtual void FindAllIntersectingObjects(TArray < Pair<const TImplicitObject<T, d>*, TRigidTransform<T, d>>>& Out, const TBox<T, d>& LocalBounds) const
+	virtual void FindAllIntersectingObjects(TArray < Pair<const FImplicitObject*, TRigidTransform<T, d>>>& Out, const TBox<T, d>& LocalBounds) const
 	{
 		if (bHierarchyBuilt)
 		{
@@ -136,13 +136,13 @@ class TImplicitObjectUnion : public TImplicitObject<T, d>
 			Out.Reserve(Out.Num() + Overlaps.Num());
 			for (int32 Idx : Overlaps)
 			{
-				const TImplicitObject<T, d>* Obj = GeomParticles.Geometry(Idx).Get();
+				const FImplicitObject* Obj = GeomParticles.Geometry(Idx).Get();
 				Out.Add(MakePair(Obj, TRigidTransform<T, d>(GeomParticles.X(Idx), GeomParticles.R(Idx))));
 			}
 		}
 		else
 		{
-			for (const TUniquePtr<TImplicitObject<T, d>>& Object : MObjects)
+			for (const TUniquePtr<FImplicitObject>& Object : MObjects)
 			{
 				Object->FindAllIntersectingObjects(Out, LocalBounds);
 			}
@@ -213,7 +213,7 @@ class TImplicitObjectUnion : public TImplicitObject<T, d>
 
 	virtual void CacheAllImplicitObjects()
 	{
-		TArray < Pair<TSerializablePtr<TImplicitObject<T, d>>, TRigidTransform<T, d>>> SubObjects;
+		TArray < Pair<TSerializablePtr<FImplicitObject>, TRigidTransform<T, d>>> SubObjects;
 		AccumulateAllSerializableImplicitObjectsHelper(SubObjects, TRigidTransform<T, d>::Identity);
 		//build hierarchy
 		{
@@ -242,7 +242,7 @@ class TImplicitObjectUnion : public TImplicitObject<T, d>
 		T MinTime = 0;	//initialization not needed, but doing it to avoid warning
 		bool bFound = false;
 
-		for (const TUniquePtr<TImplicitObject<T, d>>& Obj : MObjects)
+		for (const TUniquePtr<FImplicitObject>& Obj : MObjects)
 		{
 			TVector<T, d> Position;
 			TVector<T, d> Normal;
@@ -267,7 +267,7 @@ class TImplicitObjectUnion : public TImplicitObject<T, d>
 
 	virtual bool Overlap(const TVector<T, d>& Point, const T Thickness) const override
 	{
-		for (const TUniquePtr<TImplicitObject<T, d>>& Obj : MObjects)
+		for (const TUniquePtr<FImplicitObject>& Obj : MObjects)
 		{
 			if (Obj->Overlap(Point, Thickness))
 			{
@@ -281,25 +281,25 @@ class TImplicitObjectUnion : public TImplicitObject<T, d>
 	virtual void Serialize(FChaosArchive& Ar) override
 	{
 		FChaosArchiveScopedMemory ScopedMemory(Ar, GetTypeName(), false);
-		TImplicitObject<T, d>::SerializeImp(Ar);
+		FImplicitObject::SerializeImp(Ar);
 		Ar << MObjects << MLocalBoundingBox << GeomParticles << Hierarchy << bHierarchyBuilt;
 	}
 
 	virtual bool IsValidGeometry() const
 	{
-		bool bValid = TImplicitObject<T, d>::IsValidGeometry();
+		bool bValid = FImplicitObject::IsValidGeometry();
 		bValid = bValid && MObjects.Num();
 		return bValid;
 	}
 
-	const TArray<TUniquePtr<TImplicitObject<T, d>>>& GetObjects() const { return MObjects; }
+	const TArray<TUniquePtr<FImplicitObject>>& GetObjects() const { return MObjects; }
 
 	virtual uint32 GetTypeHash() const override
 	{
 		uint32 Result = 0;
 
 		// Union hash is just the hash of all internal objects
-		for(const TUniquePtr<TImplicitObject<T, d>>& InnerObj : MObjects)
+		for(const TUniquePtr<FImplicitObject>& InnerObj : MObjects)
 		{
 			Result = HashCombine(Result, InnerObj->GetTypeHash());
 		}
@@ -329,7 +329,7 @@ private:
 	}
 
   private:
-	TArray<TUniquePtr<TImplicitObject<T, d>>> MObjects;
+	TArray<TUniquePtr<FImplicitObject>> MObjects;
 	TGeometryParticles<T, d> GeomParticles;
 	TBoundingVolumeHierarchy<TGeometryParticles<T, d>, TArray<int32>, T, d> Hierarchy;
 	TBox<T, d> MLocalBoundingBox;
@@ -338,9 +338,9 @@ private:
 
 
 	//needed for serialization
-	TImplicitObjectUnion() : TImplicitObject<T, d>(EImplicitObject::HasBoundingBox, ImplicitObjectType::Union), Hierarchy(GeomParticles, 1){}
-	friend TImplicitObject<T, d>;	//needed for serialization
+	TImplicitObjectUnion() : FImplicitObject(EImplicitObject::HasBoundingBox, ImplicitObjectType::Union), Hierarchy(GeomParticles, 1){}
+	friend FImplicitObject;	//needed for serialization
 public:
-	TMap<const TImplicitObject<T,d>*, int32> MCollisionParticleLookupHack;	//temp hack for finding collision particles
+	TMap<const FImplicitObject*, int32> MCollisionParticleLookupHack;	//temp hack for finding collision particles
 };
 }
diff --git a/Engine/Source/Runtime/Experimental/Chaos/Public/Chaos/Levelset.h b/Engine/Source/Runtime/Experimental/Chaos/Public/Chaos/Levelset.h
--- a/Engine/Source/Runtime/Experimental/Chaos/Public/Chaos/Levelset.h
+++ b/Engine/Source/Runtime/Experimental/Chaos/Public/Chaos/Levelset.h
@@ -19,13 +19,13 @@ template<class T, int D>
 class TPlane;
 
 template<class T, int d>
-class CHAOS_API TLevelSet final : public TImplicitObject<T, d>
+class CHAOS_API TLevelSet final : public FImplicitObject
 {
   public:
-	using TImplicitObject<T, d>::SignedDistance;
+	using FImplicitObject::SignedDistance;
 
 	TLevelSet(FErrorReporter& ErrorReporter, const TUniformGrid<T, d>& InGrid, const TParticles<T, d>& InParticles, const TTriangleMesh<T>& Mesh, const int32 BandWidth = 0);
-	TLevelSet(FErrorReporter& ErrorReporter, const TUniformGrid<T, d>& InGrid, const TImplicitObject<T, d>& InObject, const int32 BandWidth = 0, const bool bUseObjectPhi = false);
+	TLevelSet(FErrorReporter& ErrorReporter, const TUniformGrid<T, d>& InGrid, const FImplicitObject& InObject, const int32 BandWidth = 0, const bool bUseObjectPhi = false);
 	TLevelSet(std::istream& Stream);
 	TLevelSet(const TLevelSet<T, d>& Other) = delete;
 	TLevelSet(TLevelSet<T, d>&& Other);
@@ -61,7 +61,7 @@ class CHAOS_API TLevelSet final : public TImplicitObject<T, d>
 
 	FORCEINLINE void SerializeImp(FArchive& Ar)
 	{
-		TImplicitObject<T, d>::SerializeImp(Ar);
+		FImplicitObject::SerializeImp(Ar);
 		Ar << MGrid;
 		Ar << MPhi;
 		Ar << MNormals;
@@ -105,7 +105,7 @@ class CHAOS_API TLevelSet final : public TImplicitObject<T, d>
 
   private:
 	bool ComputeDistancesNearZeroIsocontour(FErrorReporter& ErrorReporter, const TParticles<T, d>& InParticles, const TArray<TVector<T, 3>> &Normals, const TTriangleMesh<T>& Mesh, TArrayND<bool, d>& BlockedFaceX, TArrayND<bool, d>& BlockedFaceY, TArrayND<bool, d>& BlockedFaceZ, TArray<TVector<int32, d>>& InterfaceIndices);
-	void ComputeDistancesNearZeroIsocontour(const TImplicitObject<T, d>& Object, const TArrayND<T, d>& ObjectPhi, TArray<TVector<int32, d>>& InterfaceIndices);
+	void ComputeDistancesNearZeroIsocontour(const FImplicitObject& Object, const TArrayND<T, d>& ObjectPhi, TArray<TVector<int32, d>>& InterfaceIndices);
 	void CorrectSign(const TArrayND<bool, d>& BlockedFaceX, const TArrayND<bool, d>& BlockedFaceY, const TArrayND<bool, d>& BlockedFaceZ, TArray<TVector<int32, d>>& InterfaceIndices);
 	T ComputePhi(const TArrayND<bool, d>& Done, const TVector<int32, d>& CellIndex);
 	void FillWithFastMarchingMethod(const T StoppingDistance, const TArray<TVector<int32, d>>& InterfaceIndices);
@@ -124,7 +124,7 @@ class CHAOS_API TLevelSet final : public TImplicitObject<T, d>
 	TBox<T, d> MOriginalLocalBoundingBox;
 	int32 MBandWidth;
 private:
-	TLevelSet() : TImplicitObject<T, d>(EImplicitObject::HasBoundingBox, ImplicitObjectType::LevelSet) {}	//needed for serialization
-	friend TImplicitObject<T, d>;	//needed for serialization
+	TLevelSet() : FImplicitObject(EImplicitObject::HasBoundingBox, ImplicitObjectType::LevelSet) {}	//needed for serialization
+	friend FImplicitObject;	//needed for serialization
 };
 }
diff --git a/Engine/Source/Runtime/Experimental/Chaos/Public/Chaos/PBDCollisionConstraint.h b/Engine/Source/Runtime/Experimental/Chaos/Public/Chaos/PBDCollisionConstraint.h
--- a/Engine/Source/Runtime/Experimental/Chaos/Public/Chaos/PBDCollisionConstraint.h
+++ b/Engine/Source/Runtime/Experimental/Chaos/Public/Chaos/PBDCollisionConstraint.h
@@ -23,8 +23,7 @@ class TPBDCollisionConstraint;
 template <typename T, int d>
 class TRigidTransform;
 
-template <typename T, int d>
-class TImplicitObject;
+class FImplicitObject;
 
 template <typename T, int d>
 class TBVHParticles;
diff --git a/Engine/Source/Runtime/Experimental/Chaos/Public/Chaos/PBDCollisionSphereConstraints.h b/Engine/Source/Runtime/Experimental/Chaos/Public/Chaos/PBDCollisionSphereConstraints.h
--- a/Engine/Source/Runtime/Experimental/Chaos/Public/Chaos/PBDCollisionSphereConstraints.h
+++ b/Engine/Source/Runtime/Experimental/Chaos/Public/Chaos/PBDCollisionSphereConstraints.h
@@ -24,9 +24,9 @@ class TPBDCollisionSphereConstraints : public TPerParticleRule<T, d>
 	{
 		for (int32 i = 0; i < InParticles.Size(); ++i)
 		{
-			MObjects.Add(TUniquePtr<TImplicitObject<T, d>>(new TSphere<T, d>(InParticles.P(i), Height)));
+			MObjects.Add(TUniquePtr<FImplicitObject>(new TSphere<T, d>(InParticles.P(i), Height)));
 		}
-		TBoundingVolumeHierarchy<TArray<TUniquePtr<TImplicitObject<T, d>>>, TArray<int32>, T, d> Hierarchy(MObjects);
+		TBoundingVolumeHierarchy<TArray<TUniquePtr<FImplicitObject>>, TArray<int32>, T, d> Hierarchy(MObjects);
 		FCriticalSection CriticalSection;
 		PhysicsParallelFor(InParticles.Size(), [&](int32 Index) {
 			TArray<int32> PotentialIntersections = Hierarchy.FindAllIntersections(InParticles.P(Index));
@@ -64,7 +64,7 @@ class TPBDCollisionSphereConstraints : public TPerParticleRule<T, d>
   private:
 	T MH;
 	TMap<int32, TArray<int32>> MConstraints;
-	TArray<TUniquePtr<TImplicitObject<T, d>>> MObjects;
+	TArray<TUniquePtr<FImplicitObject>> MObjects;
 };
 }
 #endif
diff --git a/Engine/Source/Runtime/Experimental/Chaos/Public/Chaos/PBDRigidClustering.h b/Engine/Source/Runtime/Experimental/Chaos/Public/Chaos/PBDRigidClustering.h
--- a/Engine/Source/Runtime/Experimental/Chaos/Public/Chaos/PBDRigidClustering.h
+++ b/Engine/Source/Runtime/Experimental/Chaos/Public/Chaos/PBDRigidClustering.h
@@ -24,7 +24,7 @@ public:
 
 	FClusterChildrenMap MChildren;
 	FClusterTransformMap ClusterParentTransforms;
-	TArray<Chaos::TSerializablePtr<TImplicitObject<T, d>>> GeometryPtrs;
+	TArray<Chaos::TSerializablePtr<FImplicitObject>> GeometryPtrs;
 };
 
 /* 
@@ -53,7 +53,7 @@ public:
 	*    ForceMassOrientation : Inertial alignment into mass space.
 	*/
 	int32 CreateClusterParticle(int32 ClusterGroupIndex, TArray<uint32>&& Children,
-		TSerializablePtr<TImplicitObject<T, d>> ProxyGeometry = TSerializablePtr<TImplicitObject<T, 3>>(),
+		TSerializablePtr<FImplicitObject> ProxyGeometry = TSerializablePtr<FImplicitObject>(),
 		const TRigidTransform<T, d>* ForceMassOrientation = nullptr, const FClusterCreationParameters<T>& Parameters = FClusterCreationParameters<T>());
 
 	/*
@@ -264,7 +264,7 @@ public:
  protected:
 
 	void UpdateMassProperties(const TArray<uint32>& Children, const uint32 NewIndex, const TRigidTransform<T, d>* ForceMassOrientation);
-	void UpdateGeometry(const TArray<uint32>& Children, const uint32 NewIndex, TSerializablePtr<TImplicitObject<T, d>> ProxyGeometry, const FClusterCreationParameters<T>& Parameters);
+	void UpdateGeometry(const TArray<uint32>& Children, const uint32 NewIndex, TSerializablePtr<FImplicitObject> ProxyGeometry, const FClusterCreationParameters<T>& Parameters);
 	void ComputeStrainFromCollision(const FPBDCollisionConstraint& CollisionRule);
 	void ResetCollisionImpulseArray();
 	void DisableCluster(TPBDRigidClusteredParticleHandle<T, d>* ClusteredParticle);
diff --git a/Engine/Source/Runtime/Experimental/Chaos/Public/Chaos/ParticleHandle.h b/Engine/Source/Runtime/Experimental/Chaos/Public/Chaos/ParticleHandle.h
--- a/Engine/Source/Runtime/Experimental/Chaos/Public/Chaos/ParticleHandle.h
+++ b/Engine/Source/Runtime/Experimental/Chaos/Public/Chaos/ParticleHandle.h
@@ -346,14 +346,14 @@ public:
 	TRotation<T, d>& R() { return GeometryParticles->R(ParticleIdx); }
 	void SetR(const TRotation<T, d>& InR) { GeometryParticles->R(ParticleIdx) = InR; }
 
-	TSerializablePtr<TImplicitObject<T, d>> Geometry() const { return GeometryParticles->Geometry(ParticleIdx); }
-	void SetGeometry(TSerializablePtr<TImplicitObject<T, d>> InGeometry) { GeometryParticles->SetGeometry(ParticleIdx, InGeometry); }
+	TSerializablePtr<FImplicitObject> Geometry() const { return GeometryParticles->Geometry(ParticleIdx); }
+	void SetGeometry(TSerializablePtr<FImplicitObject> InGeometry) { GeometryParticles->SetGeometry(ParticleIdx, InGeometry); }
 
-	TSerializablePtr<TImplicitObject<T, d>> SharedGeometry() const { return GeometryParticles->SharedGeometry(ParticleIdx); }
-	void SetSharedGeometry(TSharedPtr<TImplicitObject<T, d>, ESPMode::ThreadSafe> InGeometry) { GeometryParticles->SetSharedGeometry(ParticleIdx, InGeometry); }
+	TSerializablePtr<FImplicitObject> SharedGeometry() const { return GeometryParticles->SharedGeometry(ParticleIdx); }
+	void SetSharedGeometry(TSharedPtr<FImplicitObject, ESPMode::ThreadSafe> InGeometry) { GeometryParticles->SetSharedGeometry(ParticleIdx, InGeometry); }
 
-	const TUniquePtr<TImplicitObject<T, d>>& DynamicGeometry() const { return GeometryParticles->DynamicGeometry(ParticleIdx); }
-	void SetDynamicGeometry(TUniquePtr<TImplicitObject<T, d>>&& Unique) { GeometryParticles->SetDynamicGeometry(ParticleIdx, MoveTemp(Unique)); }
+	const TUniquePtr<FImplicitObject>& DynamicGeometry() const { return GeometryParticles->DynamicGeometry(ParticleIdx); }
+	void SetDynamicGeometry(TUniquePtr<FImplicitObject>&& Unique) { GeometryParticles->SetDynamicGeometry(ParticleIdx, MoveTemp(Unique)); }
 
 	const TShapesArray<T,d>& ShapesArray() const { return GeometryParticles->ShapesArray(ParticleIdx); }
 
@@ -805,8 +805,8 @@ public:
 		// Static Particles
 		const TVector<T, d>& X() const { return Handle->X(); }
 		const TRotation<T, d>& R() const { return Handle->R(); }
-		TSerializablePtr<TImplicitObject<T, d>> Geometry() const { return Handle->Geometry(); }
-		const TUniquePtr<TImplicitObject<T, d>>& DynamicGeometry() const { return Handle->DynamicGeometry(); }
+		TSerializablePtr<FImplicitObject> Geometry() const { return Handle->Geometry(); }
+		const TUniquePtr<FImplicitObject>& DynamicGeometry() const { return Handle->DynamicGeometry(); }
 		bool Sleeping() const { return Handle->Sleeping(); }
 		FString ToString() const { return Handle->ToString(); }
 
@@ -999,26 +999,26 @@ public:
 	}
 
 	//todo: geometry should not be owned by particle
-	void SetGeometry(TUniquePtr<TImplicitObject<T, d>>&& UniqueGeometry)
+	void SetGeometry(TUniquePtr<FImplicitObject>&& UniqueGeometry)
 	{
 		// Take ownership of the geometry, putting it into a shared ptr.
 		// This is necessary because we cannot be sure whether the particle
 		// will be destroyed on the game thread or physics thread first,
 		// but geometry data is shared between them.
-		TImplicitObject<T, d>* RawGeometry = UniqueGeometry.Release();
-		SetGeometry(TSharedPtr<TImplicitObject<T, d>, ESPMode::ThreadSafe>(RawGeometry));
+		FImplicitObject* RawGeometry = UniqueGeometry.Release();
+		SetGeometry(TSharedPtr<FImplicitObject, ESPMode::ThreadSafe>(RawGeometry));
 	}
 
 	// TODO: Right now this method exists so we can do things like FPhysTestSerializer::CreateChaosData.
 	//       We should replace this with a method for supporting SetGeometry(RawGeometry).
-	void SetGeometry(TSharedPtr<TImplicitObject<T, d>, ESPMode::ThreadSafe> SharedGeometry)
+	void SetGeometry(TSharedPtr<FImplicitObject, ESPMode::ThreadSafe> SharedGeometry)
 	{
 		this->MarkDirty(EParticleFlags::Geometry);
 		this->MGeometry = SharedGeometry;
 		UpdateShapesArray();
 	}
 
-	void SetGeometry(TSerializablePtr<TImplicitObject<T, d>> RawGeometry)
+	void SetGeometry(TSerializablePtr<FImplicitObject> RawGeometry)
 	{
 		// Ultimately this method should replace SetGeometry(SharedPtr).
 		// We don't really want people making shared ptrs to geometry everywhere.
@@ -1038,7 +1038,7 @@ public:
 		MShapesArray = MoveTemp(InShapesArray);
 	}
 
-	TSerializablePtr<TImplicitObject<T, d>> Geometry() const { return MakeSerializable(MGeometry); }
+	TSerializablePtr<FImplicitObject> Geometry() const { return MakeSerializable(MGeometry); }
 
 	const TShapesArray<T,d>& ShapesArray() const { return MShapesArray; }
 
@@ -1097,7 +1097,7 @@ public:
 private:
 	TVector<T, d> MX;
 	TRotation<T, d> MR;
-	TSharedPtr<TImplicitObject<T, d>, ESPMode::ThreadSafe> MGeometry;	//TODO: geometry should live in bodysetup
+	TSharedPtr<FImplicitObject, ESPMode::ThreadSafe> MGeometry;	//TODO: geometry should live in bodysetup
 	TShapesArray<T,d> MShapesArray;
 	FSpatialAccelerationIdx MSpatialIdx;
 
@@ -1106,7 +1106,7 @@ private:
 
 	// This is only for use by ParticleData. This should be called only in one place,
 	// when the geometry is being copied from GT to PT.
-	TSharedPtr<TImplicitObject<T, d>, ESPMode::ThreadSafe> GeometrySharedLowLevel() const
+	TSharedPtr<FImplicitObject, ESPMode::ThreadSafe> GeometrySharedLowLevel() const
 	{
 		return MGeometry;
 	}
@@ -1162,7 +1162,7 @@ public:
 
 	TVector<T, d> X;
 	TRotation<T, d> R;
-	TSharedPtr<TImplicitObject<T, d>, ESPMode::ThreadSafe> Geometry;
+	TSharedPtr<FImplicitObject, ESPMode::ThreadSafe> Geometry;
 	FSpatialAccelerationIdx SpatialIdx;
 };
 
diff --git a/Engine/Source/Runtime/Experimental/Chaos/Public/Chaos/Plane.h b/Engine/Source/Runtime/Experimental/Chaos/Public/Chaos/Plane.h
--- a/Engine/Source/Runtime/Experimental/Chaos/Public/Chaos/Plane.h
+++ b/Engine/Source/Runtime/Experimental/Chaos/Public/Chaos/Plane.h
@@ -7,27 +7,27 @@
 namespace Chaos
 {
 template<class T, int d>
-class TPlane final : public TImplicitObject<T, d>
+class TPlane final : public FImplicitObject
 {
   public:
-	using TImplicitObject<T, d>::GetTypeName;
+	using FImplicitObject::GetTypeName;
 
 
-	TPlane() : TImplicitObject<T, d>(0, ImplicitObjectType::Plane) {}	//needed for serialization
+	TPlane() : FImplicitObject(0, ImplicitObjectType::Plane) {}	//needed for serialization
 	TPlane(const TVector<T, d>& InX, const TVector<T, d>& InNormal)
-	    : TImplicitObject<T, d>(0, ImplicitObjectType::Plane)
+	    : FImplicitObject(0, ImplicitObjectType::Plane)
 	    , MX(InX)
 	    , MNormal(InNormal)
 	{
 	}
 	TPlane(const TPlane<T, d>& Other)
-	    : TImplicitObject<T, d>(0, ImplicitObjectType::Plane)
+	    : FImplicitObject(0, ImplicitObjectType::Plane)
 	    , MX(Other.MX)
 	    , MNormal(Other.MNormal)
 	{
 	}
 	TPlane(TPlane<T, d>&& Other)
-	    : TImplicitObject<T, d>(0, ImplicitObjectType::Plane)
+	    : FImplicitObject(0, ImplicitObjectType::Plane)
 	    , MX(MoveTemp(Other.MX))
 	    , MNormal(MoveTemp(Other.MNormal))
 	{
@@ -133,7 +133,7 @@ class TPlane final : public TImplicitObject<T, d>
 	
 	FORCEINLINE void SerializeImp(FArchive& Ar)
 	{
-		TImplicitObject<T, d>::SerializeImp(Ar);
+		FImplicitObject::SerializeImp(Ar);
 		Ar << MX << MNormal;
 	}
 
diff --git a/Engine/Source/Runtime/Experimental/Chaos/Public/Chaos/Sphere.h b/Engine/Source/Runtime/Experimental/Chaos/Public/Chaos/Sphere.h
--- a/Engine/Source/Runtime/Experimental/Chaos/Public/Chaos/Sphere.h
+++ b/Engine/Source/Runtime/Experimental/Chaos/Public/Chaos/Sphere.h
@@ -17,14 +17,14 @@ namespace Chaos
 	};
 
 	template<class T, int d>
-	class TSphere final : public TImplicitObject<T, d>
+	class TSphere final : public FImplicitObject
 	{
 	public:
 
-		using TImplicitObject<T, d>::GetTypeName;
+		using FImplicitObject::GetTypeName;
 
 		TSphere(const TVector<T, d>& InCenter, const T InRadius)
-		    : TImplicitObject<T, d>(EImplicitObject::IsConvex | EImplicitObject::HasBoundingBox, ImplicitObjectType::Sphere)
+		    : FImplicitObject(EImplicitObject::IsConvex | EImplicitObject::HasBoundingBox, ImplicitObjectType::Sphere)
 		    , Center(InCenter)
 		    , Radius(InRadius)
 		    , LocalBoundingBox(Center - Radius, Center + Radius)
@@ -32,7 +32,7 @@ namespace Chaos
 		}
 
 		TSphere(const TSphere<T, d>& Other)
-		    : TImplicitObject<T, d>(EImplicitObject::IsConvex | EImplicitObject::HasBoundingBox, ImplicitObjectType::Sphere)
+		    : FImplicitObject(EImplicitObject::IsConvex | EImplicitObject::HasBoundingBox, ImplicitObjectType::Sphere)
 		    , Center(Other.Center)
 		    , Radius(Other.Radius)
 		    , LocalBoundingBox(Other.LocalBoundingBox)
@@ -40,7 +40,7 @@ namespace Chaos
 		}
 
 		TSphere(TSphere<T, d>&& Other)
-		    : TImplicitObject<T, d>(EImplicitObject::IsConvex | EImplicitObject::HasBoundingBox, ImplicitObjectType::Sphere)
+		    : FImplicitObject(EImplicitObject::IsConvex | EImplicitObject::HasBoundingBox, ImplicitObjectType::Sphere)
 		    , Center(MoveTemp(Other.Center))
 		    , Radius(Other.Radius)
 		    , LocalBoundingBox(MoveTemp(Other.LocalBoundingBox))
@@ -249,7 +249,7 @@ namespace Chaos
 
 		FORCEINLINE void SerializeImp(FArchive& Ar)
 		{
-			TImplicitObject<T, d>::SerializeImp(Ar);
+			FImplicitObject::SerializeImp(Ar);
 			Ar << Center << Radius;
 			if (Ar.IsLoading())
 			{
@@ -325,9 +325,9 @@ namespace Chaos
 			return HashCombine(CenterHash, HashCombine(RadiusHash, BoundsHash));
 		}
 
-		virtual TUniquePtr<TImplicitObject<T, d>> Copy() const override
+		virtual TUniquePtr<FImplicitObject> Copy() const override
 		{
-			return TUniquePtr<TImplicitObject<T, d>>(new TSphere<T,d>(Center, Radius));
+			return TUniquePtr<FImplicitObject>(new TSphere<T,d>(Center, Radius));
 		}
 
 	private:
@@ -338,9 +338,9 @@ namespace Chaos
 	private:
 
 		// TImplicitObject requires ability to default construct when deserializing shapes
-		friend TImplicitObject<T, d>;
+		friend FImplicitObject;
 		TSphere()
-		    : TImplicitObject<T, d>(EImplicitObject::IsConvex | EImplicitObject::HasBoundingBox, ImplicitObjectType::Sphere) 
+		    : FImplicitObject(EImplicitObject::IsConvex | EImplicitObject::HasBoundingBox, ImplicitObjectType::Sphere) 
 		{}
 
 	};
diff --git a/Engine/Source/Runtime/Experimental/Chaos/Public/Chaos/TaperedCylinder.h b/Engine/Source/Runtime/Experimental/Chaos/Public/Chaos/TaperedCylinder.h
--- a/Engine/Source/Runtime/Experimental/Chaos/Public/Chaos/TaperedCylinder.h
+++ b/Engine/Source/Runtime/Experimental/Chaos/Public/Chaos/TaperedCylinder.h
@@ -10,16 +10,16 @@ namespace Chaos
 	struct TTaperedCylinderSpecializeSamplingHelper;
 
 	template<class T>
-	class TTaperedCylinder : public TImplicitObject<T, 3>
+	class TTaperedCylinder : public FImplicitObject
 	{
 	public:
 		TTaperedCylinder()
-		    : TImplicitObject<T, 3>(EImplicitObject::FiniteConvex, ImplicitObjectType::TaperedCylinder)
+		    : FImplicitObject(EImplicitObject::FiniteConvex, ImplicitObjectType::TaperedCylinder)
 		{
 			this->bIsConvex = true;
 		}
 		TTaperedCylinder(const TVector<T, 3>& x1, const TVector<T, 3>& x2, const T Radius1, const T Radius2)
-		    : TImplicitObject<T, 3>(EImplicitObject::FiniteConvex, ImplicitObjectType::TaperedCylinder)
+		    : FImplicitObject(EImplicitObject::FiniteConvex, ImplicitObjectType::TaperedCylinder)
 		    , MPlane1(x1, (x2 - x1).GetSafeNormal())
 		    , MPlane2(x2, -MPlane1.Normal())
 		    , MHeight((x2 - x1).Size())
@@ -35,7 +35,7 @@ namespace Chaos
 			MLocalBoundingBox = TBox<T, 3>(MLocalBoundingBox.Min() - TVector<T, 3>(MaxRadius), MLocalBoundingBox.Max() + TVector<T, 3>(MaxRadius));
 		}
 		TTaperedCylinder(const TTaperedCylinder<T>& Other)
-		    : TImplicitObject<T, 3>(EImplicitObject::FiniteConvex, ImplicitObjectType::TaperedCylinder)
+		    : FImplicitObject(EImplicitObject::FiniteConvex, ImplicitObjectType::TaperedCylinder)
 		    , MPlane1(Other.MPlane1)
 		    , MPlane2(Other.MPlane2)
 		    , MHeight(Other.MHeight)
@@ -46,7 +46,7 @@ namespace Chaos
 			this->bIsConvex = true;
 		}
 		TTaperedCylinder(TTaperedCylinder<T>&& Other)
-		    : TImplicitObject<T, 3>(EImplicitObject::FiniteConvex, ImplicitObjectType::TaperedCylinder)
+		    : FImplicitObject(EImplicitObject::FiniteConvex, ImplicitObjectType::TaperedCylinder)
 		    , MPlane1(MoveTemp(Other.MPlane1))
 		    , MPlane2(MoveTemp(Other.MPlane2))
 		    , MHeight(Other.MHeight)
diff --git a/Engine/Source/Runtime/Experimental/Chaos/Public/Chaos/Triangle.h b/Engine/Source/Runtime/Experimental/Chaos/Public/Chaos/Triangle.h
--- a/Engine/Source/Runtime/Experimental/Chaos/Public/Chaos/Triangle.h
+++ b/Engine/Source/Runtime/Experimental/Chaos/Public/Chaos/Triangle.h
@@ -122,22 +122,22 @@ namespace Chaos
 	}
 
 	template<typename T>
-	class TImplicitTriangle final : public TImplicitObject<T, 3>
+	class TImplicitTriangle final : public FImplicitObject
 	{
 	public:
 
 		TImplicitTriangle()
-			: TImplicitObject<T, 3>(EImplicitObject::IsConvex | EImplicitObject::HasBoundingBox, ImplicitObjectType::Triangle)
+			: FImplicitObject(EImplicitObject::IsConvex | EImplicitObject::HasBoundingBox, ImplicitObjectType::Triangle)
 		{}
 
 		TImplicitTriangle(const TImplicitTriangle&) = delete;
 
 		TImplicitTriangle(TImplicitTriangle&& InToSteal)
-			: TImplicitObject<T, 3>(EImplicitObject::IsConvex | EImplicitObject::HasBoundingBox, ImplicitObjectType::Triangle)
+			: FImplicitObject(EImplicitObject::IsConvex | EImplicitObject::HasBoundingBox, ImplicitObjectType::Triangle)
 		{}
 
 		TImplicitTriangle(const TVector<T, 3>& InA, const TVector<T, 3>& InB, const TVector<T, 3>& InC)
-			: TImplicitObject<T, 3>(EImplicitObject::IsConvex | EImplicitObject::HasBoundingBox, ImplicitObjectType::Triangle)
+			: FImplicitObject(EImplicitObject::IsConvex | EImplicitObject::HasBoundingBox, ImplicitObjectType::Triangle)
 			, Tri(InA, InB, InC)
 		{
 			Bounds = TBox<T, 3>(Tri[0], Tri[0]);
diff --git a/Engine/Source/Runtime/Experimental/Chaos/Public/Chaos/TriangleMeshImplicitObject.h b/Engine/Source/Runtime/Experimental/Chaos/Public/Chaos/TriangleMeshImplicitObject.h
--- a/Engine/Source/Runtime/Experimental/Chaos/Public/Chaos/TriangleMeshImplicitObject.h
+++ b/Engine/Source/Runtime/Experimental/Chaos/Public/Chaos/TriangleMeshImplicitObject.h
@@ -15,10 +15,10 @@ namespace Chaos
 {
 
 	template<typename T>
-	class CHAOS_API TTriangleMeshImplicitObject final : public TImplicitObject<T,3>
+	class CHAOS_API TTriangleMeshImplicitObject final : public FImplicitObject
 	{
 	public:
-		using TImplicitObject<T, 3>::GetTypeName;
+		using FImplicitObject::GetTypeName;
 
 		TTriangleMeshImplicitObject(TParticles<T,3>&& Particles, TArray<TVector<int32, 3>>&& Elements);
 		TTriangleMeshImplicitObject(const TTriangleMeshImplicitObject& Other) = delete;
@@ -34,8 +34,8 @@ namespace Chaos
 		virtual bool Raycast(const TVector<T, 3>& StartPoint, const TVector<T, 3>& Dir, const T Length, const T Thickness, T& OutTime, TVector<T,3>& OutPosition, TVector<T,3>& OutNormal, int32& OutFaceIndex) const override;
 		virtual bool Overlap(const TVector<T, 3>& Point, const T Thickness) const override;
 
-		bool OverlapGeom(const TImplicitObject<T,3>& QueryGeom, const TRigidTransform<T, 3>& QueryTM, const T Thickness) const;
-		bool SweepGeom(const TImplicitObject<T, 3>& QueryGeom, const TRigidTransform<T, 3>& StartTM, const TVector<T,3>& Dir, const T Length, T& OutTime, TVector<T,3>& OutPosition, TVector<T,3>& OutNormal, int32& OutFaceIndex, const T Thickness = 0) const;
+		bool OverlapGeom(const FImplicitObject& QueryGeom, const TRigidTransform<T, 3>& QueryTM, const T Thickness) const;
+		bool SweepGeom(const FImplicitObject& QueryGeom, const TRigidTransform<T, 3>& StartTM, const TVector<T,3>& Dir, const T Length, T& OutTime, TVector<T,3>& OutPosition, TVector<T,3>& OutNormal, int32& OutFaceIndex, const T Thickness = 0) const;
 		virtual int32 FindMostOpposingFace(const TVector<T, 3>& Position, const TVector<T, 3>& UnitDir, int32 HintFaceIndex, T SearchDistance) const override;
 		virtual TVector<T, 3> FindGeometryOpposingNormal(const TVector<T, 3>& DenormDir, int32 FaceIndex, const TVector<T, 3>& OriginalNormal) const override;
 
@@ -53,7 +53,7 @@ namespace Chaos
 		{
 			Ar.UsingCustomVersion(FExternalPhysicsCustomObjectVersion::GUID);
 
-			TImplicitObject<T, 3>::SerializeImp(Ar);
+			FImplicitObject::SerializeImp(Ar);
 			Ar << MParticles;
 			Ar << MElements;
 			Ar << MLocalBoundingBox;
@@ -142,11 +142,10 @@ namespace Chaos
 		friend struct TTriangleMeshSweepVisitor;
 
 		// Required by implicit object serialization, disabled for general use.
-		template<typename InnerT, int InnerD>
-		friend class TImplicitObject;
+		friend class FImplicitObject;
 
 		TTriangleMeshImplicitObject() 
-			: TImplicitObject<T, 3>(EImplicitObject::HasBoundingBox, ImplicitObjectType::TriangleMesh)
+			: FImplicitObject(EImplicitObject::HasBoundingBox, ImplicitObjectType::TriangleMesh)
 		{};
 	};
 }
diff --git a/Engine/Source/Runtime/Experimental/ChaosCore/Public/Chaos/Core.h b/Engine/Source/Runtime/Experimental/ChaosCore/Public/Chaos/Core.h
--- a/Engine/Source/Runtime/Experimental/ChaosCore/Public/Chaos/Core.h
+++ b/Engine/Source/Runtime/Experimental/ChaosCore/Public/Chaos/Core.h
@@ -2,29 +2,22 @@
 
 #pragma once
 
-#include "Vector.h"
-#include "Matrix.h"
+#include "Chaos/Vector.h"
+#include "Chaos/Matrix.h"
+#include "Chaos/Rotation.h"
+#include "Chaos/Transform.h"
 
 namespace Chaos
 {
 	/**
 	 * Common data types for the Chaos physics engine. Unless a specific
-	 * precision of real type is required most code should use the R types
-	 * (e.g. FVectorR3) to adapt to global changes in precision.
+	 * precision of type is required most code should use these existing types
+	 * (e.g. FVec3) to adapt to global changes in precision.
 	 */
 	using FReal = float;
 
-	using FVectorR2 = TVector<FReal, 2>;
-	using FVectorR3 = TVector<FReal, 3>;
-	using FVectorR4 = TVector<FReal, 4>;
-
-	using FVectorI2 = TVector<int32, 2>;
-	using FVectorI3 = TVector<int32, 3>;
-
-	using FVectorU2 = TVector<uint32, 2>;
-	using FVectorU3 = TVector<uint32, 3>;
-
-	using FMatrixR22 = PMatrix<FReal, 2, 2>;
-	using FMatrixR33 = PMatrix<FReal, 3, 3>;
-	using FMatrixR44 = PMatrix<FReal, 4, 4>;
+	using FVec3 = TVector<FReal, 3>;
+	using FMatrix33 = PMatrix<FReal, 3, 3>;
+	using FRigidTransform3 = TRigidTransform<FReal, 3>;
+	using FRotation3 = TRotation<FReal, 3>;
 }
\ No newline at end of file
diff --git a/Engine/Source/Runtime/Experimental/ChaosSolvers/Private/PhysicsProxy/GeometryCollectionPhysicsProxy.cpp b/Engine/Source/Runtime/Experimental/ChaosSolvers/Private/PhysicsProxy/GeometryCollectionPhysicsProxy.cpp
--- a/Engine/Source/Runtime/Experimental/ChaosSolvers/Private/PhysicsProxy/GeometryCollectionPhysicsProxy.cpp
+++ b/Engine/Source/Runtime/Experimental/ChaosSolvers/Private/PhysicsProxy/GeometryCollectionPhysicsProxy.cpp
@@ -173,7 +173,7 @@ DECLARE_CYCLE_STAT(TEXT("FGeometryCollectionPhysicsProxy::PopulateSimulatedParti
 void PopulateSimulatedParticle(FGeometryCollectionPhysicsProxy::FParticlesType& Particles,
 	const FSharedSimulationParameters& SharedParams,
 	const FCollisionStructureManager::FSimplicial* Simplicial,
-	Chaos::TSerializablePtr<Chaos::TImplicitObject<float, 3>> Implicit,
+	Chaos::TSerializablePtr<Chaos::FImplicitObject> Implicit,
 	float MassIn,
 	const FVector& InertiaTensorVec,
 	int32 RigidBodyIndex,
@@ -2024,7 +2024,7 @@ void FGeometryCollectionPhysicsProxy::CreateDynamicAttributes()
 		DynamicCollection->AddExternalAttribute("InitialLinearVelocity", FTransformCollection::TransformGroup, InitialLinearVelocity);
 	}
 
-	const auto& RestImplicits = Parameters.RestCollection->GetAttribute<TUniquePtr<Chaos::TImplicitObject<float, 3>>>(ImplicitsAttribute, FTransformCollection::TransformGroup);
+	const auto& RestImplicits = Parameters.RestCollection->GetAttribute<TUniquePtr<Chaos::FImplicitObject>>(ImplicitsAttribute, FTransformCollection::TransformGroup);
 	for (int32 Index = DynamicCollection->NumElements(FTransformCollection::TransformGroup) - 1; 0 <= Index; Index--)
 	{
 		Simplicials[Index] = TUniquePtr<FSimplicial>(nullptr);
@@ -2127,7 +2127,7 @@ void FGeometryCollectionPhysicsProxy::InitializeSharedCollisionStructures(Chaos:
 	TManagedArray<FTransform>& CollectionMassToLocal = RestCollection.AddAttribute<FTransform>(TEXT("MassToLocal"), FTransformCollection::TransformGroup);
 	TManagedArray<float>& CollectionMass = RestCollection.AddAttribute<float>(TEXT("Mass"), FTransformCollection::TransformGroup);
 	TManagedArray<TUniquePtr<FSimplicial>>& CollectionSimplicials = RestCollection.AddAttribute<TUniquePtr<FSimplicial>>(SimplicialsAttribute, FTransformCollection::TransformGroup);
-	TManagedArray<TUniquePtr<TImplicitObject<float, 3>>>& CollectionImplicits = RestCollection.AddAttribute<TUniquePtr<TImplicitObject<float, 3>>>(ImplicitsAttribute, FTransformCollection::TransformGroup);
+	TManagedArray<TUniquePtr<FImplicitObject>>& CollectionImplicits = RestCollection.AddAttribute<TUniquePtr<FImplicitObject>>(ImplicitsAttribute, FTransformCollection::TransformGroup);
 
 	for (int32 Index = 0; Index < CollectionMassToLocal.Num(); ++Index)
 	{
@@ -2322,7 +2322,7 @@ void FGeometryCollectionPhysicsProxy::InitializeSharedCollisionStructures(Chaos:
 
 
 			ErrorReporter.SetPrefix(BaseErrorPrefix + " | Transform Index: " + FString::FromInt(TransformGroupIndex));
-			CollectionImplicits[TransformGroupIndex] = TUniquePtr<TImplicitObject<float, 3>>(
+			CollectionImplicits[TransformGroupIndex] = TUniquePtr<FImplicitObject>(
 				FCollisionStructureManager::NewImplicit(ErrorReporter, MassSpaceParticles, *TriMesh, InstanceBoundingBox,
 					InnerRadius[GeometryIndex], SizeSpecificData.MinLevelSetResolution, SizeSpecificData.MaxLevelSetResolution,
 					SizeSpecificData.CollisionObjectReductionPercentage, SizeSpecificData.CollisionType,
@@ -2469,7 +2469,7 @@ void FGeometryCollectionPhysicsProxy::InitializeSharedCollisionStructures(Chaos:
 
 					//don't support non level-set serialization
 					ErrorReporter.SetPrefix(BaseErrorPrefix + " | Cluster Transform Index: " + FString::FromInt(ClusterTransformIdx));
-					CollectionImplicits[ClusterTransformIdx] = TUniquePtr<TImplicitObject<float, 3>>(
+					CollectionImplicits[ClusterTransformIdx] = TUniquePtr<FImplicitObject>(
 						FCollisionStructureManager::NewImplicit(ErrorReporter, MassSpaceParticles, *UnionMesh,
 							InstanceBoundingBox, 0, MinResolution, MaxResolution,
 							SizeSpecificData.CollisionObjectReductionPercentage, SizeSpecificData.CollisionType,
diff --git a/Engine/Source/Runtime/Experimental/ChaosSolvers/Private/PhysicsProxy/SkeletalMeshPhysicsProxy.cpp b/Engine/Source/Runtime/Experimental/ChaosSolvers/Private/PhysicsProxy/SkeletalMeshPhysicsProxy.cpp
--- a/Engine/Source/Runtime/Experimental/ChaosSolvers/Private/PhysicsProxy/SkeletalMeshPhysicsProxy.cpp
+++ b/Engine/Source/Runtime/Experimental/ChaosSolvers/Private/PhysicsProxy/SkeletalMeshPhysicsProxy.cpp
@@ -152,13 +152,13 @@ void FSkeletalMeshPhysicsProxy::CreateRigidBodyCallback(FParticlesType& Particle
 			//Chaos::TTriangleMesh<T> TriMesh = Chaos::TTriangleMesh<T>::GetConvexHullFromParticles(
 			//	TArrayView<const Chaos::TVector<T, 3>>(SamplePoints));
 
-			TUniquePtr<Chaos::TImplicitObject<float, 3>> ImplicitObject(Group->BuildSimImplicitObject());
+			TUniquePtr<Chaos::FImplicitObject> ImplicitObject(Group->BuildSimImplicitObject());
 			check(ImplicitObject);
 
 			if (SamplePoints)
 			{
 				auto PointVolumeRegistrationCheck = 
-					[&](const TArray<Chaos::TVector<float, 3>> &InSamplePoints, const Chaos::TImplicitObject<float, 3> &InImplicitObject, const float InTolerance) -> bool
+					[&](const TArray<Chaos::TVector<float, 3>> &InSamplePoints, const Chaos::FImplicitObject &InImplicitObject, const float InTolerance) -> bool
 					{
 						for (int32 i = 0; i < InSamplePoints.Num(); i++)
 						{
diff --git a/Engine/Source/Runtime/Experimental/ChaosSolvers/Private/PhysicsProxy/StaticMeshPhysicsObject.cpp b/Engine/Source/Runtime/Experimental/ChaosSolvers/Private/PhysicsProxy/StaticMeshPhysicsObject.cpp
--- a/Engine/Source/Runtime/Experimental/ChaosSolvers/Private/PhysicsProxy/StaticMeshPhysicsObject.cpp
+++ b/Engine/Source/Runtime/Experimental/ChaosSolvers/Private/PhysicsProxy/StaticMeshPhysicsObject.cpp
@@ -140,7 +140,7 @@ void FStaticMeshPhysicsProxy::CreateRigidBodyCallback(FParticlesType& Particles)
 			Chaos::FErrorReporter ErrorReporter(Parameters.Name + " | RigidBodyId: " + FString::FromInt(RigidBodyId));;
 			Particles.SetDynamicGeometry(
 				RigidBodyId,
-				TUniquePtr<Chaos::TImplicitObject<float, 3>>(
+				TUniquePtr<Chaos::FImplicitObject>(
 					FCollisionStructureManager::NewImplicit(
 						ErrorReporter,
 						MeshParticles, *TriangleMesh,
@@ -165,7 +165,7 @@ void FStaticMeshPhysicsProxy::CreateRigidBodyCallback(FParticlesType& Particles)
 			const Chaos::TBox<float, 3> BBox = Sphere->BoundingBox();
 			Bounds.Min = BBox.Min();
 			Bounds.Max = BBox.Max();
-			Particles.SetDynamicGeometry(RigidBodyId, TUniquePtr<Chaos::TImplicitObject<float, 3>>(Sphere));
+			Particles.SetDynamicGeometry(RigidBodyId, TUniquePtr<Chaos::FImplicitObject>(Sphere));
 			if (!Parameters.MeshVertexPositions.Size())
 			{
 				float Radius = Parameters.ShapeParams.SphereRadius;
@@ -185,7 +185,7 @@ void FStaticMeshPhysicsProxy::CreateRigidBodyCallback(FParticlesType& Particles)
 			Chaos::TBox<float,3>* Box = new Chaos::TBox<float, 3>(-HalfExtents, HalfExtents);
 			Bounds.Min = Box->Min();
 			Bounds.Max = Box->Max();
-			Particles.SetDynamicGeometry(RigidBodyId, TUniquePtr<Chaos::TImplicitObject<float, 3>>(Box));
+			Particles.SetDynamicGeometry(RigidBodyId, TUniquePtr<Chaos::FImplicitObject>(Box));
 			if (!Parameters.MeshVertexPositions.Size())
 			{
 				Chaos::TVector<float, 3> x1(-HalfExtents);
@@ -210,7 +210,7 @@ void FStaticMeshPhysicsProxy::CreateRigidBodyCallback(FParticlesType& Particles)
 			const Chaos::TBox<float, 3> BBox = Capsule->BoundingBox();
 			Bounds.Min = BBox.Min();
 			Bounds.Max = BBox.Max();
-			Particles.SetDynamicGeometry(RigidBodyId, TUniquePtr<Chaos::TImplicitObject<float, 3>>(Capsule));
+			Particles.SetDynamicGeometry(RigidBodyId, TUniquePtr<Chaos::FImplicitObject>(Capsule));
 			if (!Parameters.MeshVertexPositions.Size())
 			{
 				float HalfHeight = Parameters.ShapeParams.CapsuleHalfHeightAndRadius.X;
@@ -237,7 +237,7 @@ void FStaticMeshPhysicsProxy::CreateRigidBodyCallback(FParticlesType& Particles)
 		{
 			Bounds.Min = Chaos::TVector<float, 3>(-KINDA_SMALL_NUMBER);
 			Bounds.Max = Chaos::TVector<float, 3>(KINDA_SMALL_NUMBER);
-			Particles.SetSharedGeometry(RigidBodyId, TSharedPtr<Chaos::TImplicitObject<float, 3>, ESPMode::ThreadSafe>());
+			Particles.SetSharedGeometry(RigidBodyId, TSharedPtr<Chaos::FImplicitObject, ESPMode::ThreadSafe>());
 		}
 
 		FTransform WorldTransform = Parameters.InitialTransform;
diff --git a/Engine/Source/Runtime/Experimental/ChaosSolvers/Public/PhysicsProxy/AnalyticImplicitGroup.h b/Engine/Source/Runtime/Experimental/ChaosSolvers/Public/PhysicsProxy/AnalyticImplicitGroup.h
--- a/Engine/Source/Runtime/Experimental/ChaosSolvers/Public/PhysicsProxy/AnalyticImplicitGroup.h
+++ b/Engine/Source/Runtime/Experimental/ChaosSolvers/Public/PhysicsProxy/AnalyticImplicitGroup.h
@@ -405,7 +405,7 @@ public:
 	 *
 	 * Transfers ownership of sub structures to the returned implicit object.
 	 */
-	Chaos::TImplicitObject<float, 3>* BuildSimImplicitObject()
+	Chaos::FImplicitObject* BuildSimImplicitObject()
 	{
 		// TODO: We make copies of implicit objects owned by this class, so we
 		// give the solver memory it can own.  It'd be nice if we could transfer
@@ -425,7 +425,7 @@ public:
 			else
 			{
 				// Make a copy and transfer ownership to the transformed implicit.
-				TUniquePtr<Chaos::TImplicitObject<float, 3>> ObjPtr(TransferImplicitObj(0));
+				TUniquePtr<Chaos::FImplicitObject> ObjPtr(TransferImplicitObj(0));
 				return new Chaos::TImplicitObjectTransformed<float, 3, true>(
 					MoveTemp(ObjPtr),
 					Chaos::TRigidTransform<float, 3>(Transforms[0]));
@@ -435,11 +435,11 @@ public:
 		{
 			// Make copies of the implicits owned by transformed immplicits, and 
 			// transfer ownership of the transformed implicits to the implicit union.
-			TArray<TUniquePtr<Chaos::TImplicitObject<float, 3>>> ImplicitObjects;
+			TArray<TUniquePtr<Chaos::FImplicitObject>> ImplicitObjects;
 			ImplicitObjects.Reserve(Num);
 			for (int i = 0; i < Num; i++)
 			{
-				TUniquePtr<Chaos::TImplicitObject<float, 3>> ObjPtr(TransferImplicitObj(i));
+				TUniquePtr<Chaos::FImplicitObject> ObjPtr(TransferImplicitObj(i));
 
 				const FTransform &Xf = Transforms[i];
 				if (Xf.Equals(FTransform::Identity))
@@ -451,7 +451,7 @@ public:
 				else
 				{
 					ImplicitObjects.Add(
-						TUniquePtr<Chaos::TImplicitObject<float, 3>>(
+						TUniquePtr<Chaos::FImplicitObject>(
 							new Chaos::TImplicitObjectTransformed<float, 3, true>(
 								MoveTemp(ObjPtr),
 								Chaos::TRigidTransform<float, 3>(Xf))));
@@ -562,9 +562,9 @@ protected:
 		}
 	}
 
-	Chaos::TImplicitObject<float, 3>* TransferImplicitObj(int32 Idx)
+	Chaos::FImplicitObject* TransferImplicitObj(int32 Idx)
 	{
-		Chaos::TImplicitObject<float, 3>* Obj = nullptr;
+		Chaos::FImplicitObject* Obj = nullptr;
 		if (Idx < Spheres.Num())
 		{
 			Obj = Spheres[Idx]; 
diff --git a/Engine/Source/Runtime/Experimental/ChaosSolvers/Public/PhysicsProxy/GeometryCollectionPhysicsProxy.h b/Engine/Source/Runtime/Experimental/ChaosSolvers/Public/PhysicsProxy/GeometryCollectionPhysicsProxy.h
--- a/Engine/Source/Runtime/Experimental/ChaosSolvers/Public/PhysicsProxy/GeometryCollectionPhysicsProxy.h
+++ b/Engine/Source/Runtime/Experimental/ChaosSolvers/Public/PhysicsProxy/GeometryCollectionPhysicsProxy.h
@@ -153,7 +153,7 @@ private:
 	TManagedArray<FVector> InitialLinearVelocity;
 	TManagedArray<bool> SimulatableParticles;
 	TManagedArray<TUniquePtr<FSimplicial> > Simplicials;
-	TManagedArray<Chaos::TSerializablePtr<Chaos::TImplicitObject<float, 3>>> Implicits;
+	TManagedArray<Chaos::TSerializablePtr<Chaos::FImplicitObject>> Implicits;
 	TArray<int32> EndFrameUnparentingBuffer;
 
 	// This is a subset of the geometry group that are used in the transform hierarchy to represent geometry
diff --git a/Engine/Source/Runtime/Experimental/GeometryCollectionCore/Public/GeometryCollection/ManagedArrayTypes.h b/Engine/Source/Runtime/Experimental/GeometryCollectionCore/Public/GeometryCollection/ManagedArrayTypes.h
--- a/Engine/Source/Runtime/Experimental/GeometryCollectionCore/Public/GeometryCollection/ManagedArrayTypes.h
+++ b/Engine/Source/Runtime/Experimental/GeometryCollectionCore/Public/GeometryCollection/ManagedArrayTypes.h
@@ -34,8 +34,7 @@ inline FArchive& operator<<(FArchive& Ar, TUniquePtr<TArray<FVector>>& ValueIn)
 	return Ar;
 }
 
-template <typename T, int d>
-inline FArchive& operator<<(FArchive& Ar, Chaos::TImplicitObject<T,d>*& ValueIn)
+inline FArchive& operator<<(FArchive& Ar, Chaos::FImplicitObject*& ValueIn)
 {
 	check(false);	//We don't serialize raw pointers to implicit objects. Use unique ptr
 	return Ar;
diff --git a/Engine/Source/Runtime/Experimental/GeometryCollectionEngine/Private/GeometryCollection/GeometryCollectionActor.cpp b/Engine/Source/Runtime/Experimental/GeometryCollectionEngine/Private/GeometryCollection/GeometryCollectionActor.cpp
--- a/Engine/Source/Runtime/Experimental/GeometryCollectionEngine/Private/GeometryCollection/GeometryCollectionActor.cpp
+++ b/Engine/Source/Runtime/Experimental/GeometryCollectionEngine/Private/GeometryCollection/GeometryCollectionActor.cpp
@@ -88,7 +88,7 @@ bool LowLevelRaycastImp(const Chaos::TVector<float, 3>& Start, const Chaos::TVec
 			const TVector<float, 3> DirLocal = TM.InverseTransformVectorNoScale(Dir);
 			const TVector<float, 3> EndLocal = StartLocal + DirLocal * DeltaMag;	//todo(ocohen): apeiron just undoes this later, we should fix the API
 
-			const TImplicitObject<float, 3>* Object = Particles.Geometry(RigidBodyIdx).Get();	//todo(ocohen): can this ever be null?
+			const FImplicitObject* Object = Particles.Geometry(RigidBodyIdx).Get();	//todo(ocohen): can this ever be null?
 			Pair<TVector<float, 3>, bool> Result = Object->FindClosestIntersection(StartLocal, EndLocal, /*Thickness=*/0.f);
 			if(Result.Second)	//todo(ocohen): once we do more than just a bool we need to get the closest point
 			{
diff --git a/Engine/Source/Runtime/Experimental/GeometryCollectionEngine/Private/GeometryCollection/GeometryCollectionParticlesData.cpp b/Engine/Source/Runtime/Experimental/GeometryCollectionEngine/Private/GeometryCollection/GeometryCollectionParticlesData.cpp
--- a/Engine/Source/Runtime/Experimental/GeometryCollectionEngine/Private/GeometryCollection/GeometryCollectionParticlesData.cpp
+++ b/Engine/Source/Runtime/Experimental/GeometryCollectionEngine/Private/GeometryCollection/GeometryCollectionParticlesData.cpp
@@ -143,10 +143,10 @@ void TGeometryCollectionParticlesData<T, d>::FData::Copy(EGeometryCollectionPart
 	const Chaos::TPBDRigidClustering<FPBDRigidsEvolution, FPBDCollisionConstraints, T, d>& Clustering = Solver->GetRigidClustering();
 
 	// Lambdas used to flatten the particle structure
-	// Can't rely on TImplicitObject<T, d>::GetType, as it returns Unknown
+	// Can't rely on FImplicitObject::GetType, as it returns Unknown
 	// TODO: Would there be something to fix in TImplicitObject so we can use this simpler line of code instead?:
-	//    auto GetType = [](const Chaos::TImplicitObject<T, d>* Geometry) { return Geometry ? Geometry->GetType(): Chaos::ImplicitObjectType::Unknown; };
-	auto GetType = [](Chaos::TSerializablePtr<Chaos::TImplicitObject<T, d>> ImplicitObject)
+	//    auto GetType = [](const Chaos::FImplicitObject* Geometry) { return Geometry ? Geometry->GetType(): Chaos::ImplicitObjectType::Unknown; };
+	auto GetType = [](Chaos::TSerializablePtr<Chaos::FImplicitObject> ImplicitObject)
 	{
 		if (ImplicitObject)
 		{
@@ -159,13 +159,13 @@ void TGeometryCollectionParticlesData<T, d>::FData::Copy(EGeometryCollectionPart
 		}
 		return Chaos::ImplicitObjectType::Unknown;
 	};
-	auto IsConvex       = [](Chaos::TSerializablePtr<Chaos::TImplicitObject<T, d>> ImplicitObject) { return ImplicitObject ? ImplicitObject->IsConvex      (): false; };
-	auto HasBoundingBox = [](Chaos::TSerializablePtr<Chaos::TImplicitObject<T, d>> ImplicitObject) { return ImplicitObject ? ImplicitObject->HasBoundingBox(): false; };
-	auto BoxMin         = [](Chaos::TSerializablePtr<Chaos::TImplicitObject<T, d>> ImplicitObject) { const Chaos::TBox     <T, d>* Box     ; return ImplicitObject && (Box      = ImplicitObject->template GetObject<Chaos::TBox     <T, d>>()) != nullptr ? Box     ->Min    (): Chaos::TVector<T, d>(T(0)); };
-	auto BoxMax         = [](Chaos::TSerializablePtr<Chaos::TImplicitObject<T, d>> ImplicitObject) { const Chaos::TBox     <T, d>* Box     ; return ImplicitObject && (Box      = ImplicitObject->template GetObject<Chaos::TBox     <T, d>>()) != nullptr ? Box     ->Max    (): Chaos::TVector<T, d>(T(0)); };
-	auto SphereCenter   = [](Chaos::TSerializablePtr<Chaos::TImplicitObject<T, d>> ImplicitObject) { const Chaos::TSphere  <T, d>* Sphere  ; return ImplicitObject && (Sphere   = ImplicitObject->template GetObject<Chaos::TSphere  <T, d>>()) != nullptr ? Sphere  ->GetCenter (): Chaos::TVector<T, d>(T(0)); };
-	auto SphereRadius   = [](Chaos::TSerializablePtr<Chaos::TImplicitObject<T, d>> ImplicitObject) { const Chaos::TSphere  <T, d>* Sphere  ; return ImplicitObject && (Sphere   = ImplicitObject->template GetObject<Chaos::TSphere  <T, d>>()) != nullptr ? Sphere  ->GetRadius (): T(0); };
-	auto LevelSetGrid   = [](Chaos::TSerializablePtr<Chaos::TImplicitObject<T, d>> ImplicitObject) { const Chaos::TLevelSet<T, d>* LevelSet; return ImplicitObject && (LevelSet = ImplicitObject->template GetObject<Chaos::TLevelSet<T, d>>()) != nullptr ? LevelSet->GetGrid(): Chaos::TUniformGrid<T, d>(); };
+	auto IsConvex       = [](Chaos::TSerializablePtr<Chaos::FImplicitObject> ImplicitObject) { return ImplicitObject ? ImplicitObject->IsConvex      (): false; };
+	auto HasBoundingBox = [](Chaos::TSerializablePtr<Chaos::FImplicitObject> ImplicitObject) { return ImplicitObject ? ImplicitObject->HasBoundingBox(): false; };
+	auto BoxMin         = [](Chaos::TSerializablePtr<Chaos::FImplicitObject> ImplicitObject) { const Chaos::TBox     <T, d>* Box     ; return ImplicitObject && (Box      = ImplicitObject->template GetObject<Chaos::TBox     <T, d>>()) != nullptr ? Box     ->Min    (): Chaos::TVector<T, d>(T(0)); };
+	auto BoxMax         = [](Chaos::TSerializablePtr<Chaos::FImplicitObject> ImplicitObject) { const Chaos::TBox     <T, d>* Box     ; return ImplicitObject && (Box      = ImplicitObject->template GetObject<Chaos::TBox     <T, d>>()) != nullptr ? Box     ->Max    (): Chaos::TVector<T, d>(T(0)); };
+	auto SphereCenter   = [](Chaos::TSerializablePtr<Chaos::FImplicitObject> ImplicitObject) { const Chaos::TSphere  <T, d>* Sphere  ; return ImplicitObject && (Sphere   = ImplicitObject->template GetObject<Chaos::TSphere  <T, d>>()) != nullptr ? Sphere  ->GetCenter (): Chaos::TVector<T, d>(T(0)); };
+	auto SphereRadius   = [](Chaos::TSerializablePtr<Chaos::FImplicitObject> ImplicitObject) { const Chaos::TSphere  <T, d>* Sphere  ; return ImplicitObject && (Sphere   = ImplicitObject->template GetObject<Chaos::TSphere  <T, d>>()) != nullptr ? Sphere  ->GetRadius (): T(0); };
+	auto LevelSetGrid   = [](Chaos::TSerializablePtr<Chaos::FImplicitObject> ImplicitObject) { const Chaos::TLevelSet<T, d>* LevelSet; return ImplicitObject && (LevelSet = ImplicitObject->template GetObject<Chaos::TLevelSet<T, d>>()) != nullptr ? LevelSet->GetGrid(): Chaos::TUniformGrid<T, d>(); };
 
 	// Data type copy for all particles
 	const int32 Count = RigidBodyIds.Num();
diff --git a/Engine/Source/Runtime/Experimental/GeometryCollectionEngine/Private/GeometryCollection/GeometryCollectionSQAccelerator.cpp b/Engine/Source/Runtime/Experimental/GeometryCollectionEngine/Private/GeometryCollection/GeometryCollectionSQAccelerator.cpp
--- a/Engine/Source/Runtime/Experimental/GeometryCollectionEngine/Private/GeometryCollection/GeometryCollectionSQAccelerator.cpp
+++ b/Engine/Source/Runtime/Experimental/GeometryCollectionEngine/Private/GeometryCollection/GeometryCollectionSQAccelerator.cpp
@@ -104,7 +104,7 @@ bool LowLevelRaycastSingleElement(int32 InParticleIndex, const Chaos::FPhysicsSo
 		const TVector<float, 3> DirLocal = TM.InverseTransformVectorNoScale(Dir);
 		const TVector<float, 3> EndLocal = StartLocal + DirLocal * DeltaMag;
 
-		const TImplicitObject<float, 3>* Object = ClusterBuffer.GeometryPtrs[InParticleIndex].Get();	//todo(ocohen): can this ever be null?
+		const FImplicitObject* Object = ClusterBuffer.GeometryPtrs[InParticleIndex].Get();	//todo(ocohen): can this ever be null?
 
 		if (!Object)
 		{
@@ -138,7 +138,7 @@ static TAutoConsoleVariable<int32> CVarMaxSweepSteps(
 	TEXT("Number of steps during a sweep"),
 	ECVF_Default);
 
-bool LowLevelSweepSingleElement(int32 InParticleIndex, const Chaos::FPhysicsSolver* InSolver, const Chaos::TClusterBuffer<float, 3>& ClusterBuffer, const FGeometryCollectionPhysicsProxy* InObject, const Chaos::TImplicitObject<float, 3>& QueryGeom, const Chaos::TParticles<float, 3>& CollisionParticles, const FTransform& StartPose, const FVector& Dir, float DeltaMag, const bool bCanBeDisabled, FHitSweep& OutHit)
+bool LowLevelSweepSingleElement(int32 InParticleIndex, const Chaos::FPhysicsSolver* InSolver, const Chaos::TClusterBuffer<float, 3>& ClusterBuffer, const FGeometryCollectionPhysicsProxy* InObject, const Chaos::FImplicitObject& QueryGeom, const Chaos::TParticles<float, 3>& CollisionParticles, const FTransform& StartPose, const FVector& Dir, float DeltaMag, const bool bCanBeDisabled, FHitSweep& OutHit)
 {
 	using namespace Chaos;
 
@@ -161,7 +161,7 @@ bool LowLevelSweepSingleElement(int32 InParticleIndex, const Chaos::FPhysicsSolv
 	const int32 LocalBodyIndex = InParticleIndex - PhysResult.BaseIndex;
 	const TRigidTransform<float, 3>& TM = PhysResult.ParticleToWorldTransforms[LocalBodyIndex];
 
-	const TImplicitObject<float, 3>* Object = ClusterBuffer.GeometryPtrs[InParticleIndex].Get();
+	const FImplicitObject* Object = ClusterBuffer.GeometryPtrs[InParticleIndex].Get();
 	if (!Object)
 	{
 		return false;
@@ -197,7 +197,7 @@ bool LowLevelSweepSingleElement(int32 InParticleIndex, const Chaos::FPhysicsSolv
 	return bFound;
 }
 
-bool LowLevelOverlap(const UGeometryCollectionComponent& GeomCollectionComponent, const TArray<int32>& InPotentialIntersections, const Chaos::TClusterBuffer<float, 3>& ClusterBuffer, const Chaos::TImplicitObject<float, 3>& QueryGeom, const FTransform& GeomPose, FHitOverlap& OutHit)
+bool LowLevelOverlap(const UGeometryCollectionComponent& GeomCollectionComponent, const TArray<int32>& InPotentialIntersections, const Chaos::TClusterBuffer<float, 3>& ClusterBuffer, const Chaos::FImplicitObject& QueryGeom, const FTransform& GeomPose, FHitOverlap& OutHit)
 {
 	using namespace Chaos;
 
@@ -233,7 +233,7 @@ bool LowLevelOverlap(const UGeometryCollectionComponent& GeomCollectionComponent
 			const int32 LocalBodyIndex = RigidBodyIdx - PhysResult.BaseIndex;
 			const TRigidTransform<float, 3>& TM = PhysResult.ParticleToWorldTransforms[LocalBodyIndex];
 
-			const TImplicitObject<float, 3>* Object = ClusterBuffer.GeometryPtrs[RigidBodyIdx].Get();
+			const FImplicitObject* Object = ClusterBuffer.GeometryPtrs[RigidBodyIdx].Get();
 			if (!Object)
 			{
 				return;
@@ -294,7 +294,7 @@ void FGeometryCollectionSQAccelerator::Raycast(const FVector& Start, const FVect
 			const PhysicsProxyWrapper& ObjectWrapper = ObjectMap.PhysicsProxyReverseMappingArray[IntersectParticleIndex];
 			if (!ObjectWrapper.PhysicsProxy)
 			{
-				const TImplicitObject<float, 3>* Object = Buffer.GeometryPtrs[IntersectParticleIndex].Get();
+				const FImplicitObject* Object = Buffer.GeometryPtrs[IntersectParticleIndex].Get();
 				// Ignore the ground plane
 				if (IntersectParticleIndex == 0 && Object->GetType(true) == ImplicitObjectType::Plane)
 				{
@@ -385,7 +385,7 @@ DECLARE_CYCLE_STAT(TEXT("Sweep Broadphase"), STAT_SQSweepBroadPhase, STATGROUP_C
 DECLARE_CYCLE_STAT(TEXT("Sweep Narrowphase"), STAT_SQSweepNarrowPhase, STATGROUP_Chaos);
 
 //@todo(mlentine): Avoid duplicated code between this and overlap
-void FGeometryCollectionSQAccelerator::Sweep(const Chaos::TImplicitObject<float, 3>& QueryGeom, const FTransform& StartTM, const FVector& Dir, const float DeltaMagnitude, ChaosInterface::FSQHitBuffer<ChaosInterface::FSweepHit>& HitBuffer, EHitFlags OutputFlags, const FQueryFilterData& QueryFilterData, ICollisionQueryFilterCallbackBase& QueryCallback) const
+void FGeometryCollectionSQAccelerator::Sweep(const Chaos::FImplicitObject& QueryGeom, const FTransform& StartTM, const FVector& Dir, const float DeltaMagnitude, ChaosInterface::FSQHitBuffer<ChaosInterface::FSweepHit>& HitBuffer, EHitFlags OutputFlags, const FQueryFilterData& QueryFilterData, ICollisionQueryFilterCallbackBase& QueryCallback) const
 {
 	SCOPE_CYCLE_COUNTER(STAT_GCSweep);
 	FChaosScopeSolverLock SolverScopeLock;
@@ -445,7 +445,7 @@ void FGeometryCollectionSQAccelerator::Sweep(const Chaos::TImplicitObject<float,
 		Chaos::TBox<float, 3> Box;
 	};
 	FLocalImplicitStorage ImplicitStorage;
-	Chaos::TImplicitObject<float, 3>* Implicit = nullptr;
+	Chaos::FImplicitObject* Implicit = nullptr;
 
 	bool bHit = false;
 	FHitSweep Hit;
@@ -558,7 +558,7 @@ void FGeometryCollectionSQAccelerator::Sweep(const Chaos::TImplicitObject<float,
 
 			if (!ObjectWrapper.PhysicsProxy)
 			{
-				const TImplicitObject<float, 3>* Object = Buffer.GeometryPtrs[ParticleIndex].Get();
+				const FImplicitObject* Object = Buffer.GeometryPtrs[ParticleIndex].Get();
 				
 				// Ignore ground plane
 				if (ParticleIndex == 0 && Object->GetType(true) == ImplicitObjectType::Plane)
@@ -645,7 +645,7 @@ void FGeometryCollectionSQAccelerator::Sweep(const Chaos::TImplicitObject<float,
 #endif
 }
 
-bool LowLevelOverlapSingleElement(int32 InParticleIndex, const Chaos::FPhysicsSolver* InSolver, const Chaos::TClusterBuffer<float, 3>& ClusterBuffer, const FGeometryCollectionPhysicsProxy* InObject, const Chaos::TImplicitObject<float, 3>& QueryGeom, const FTransform& InPose, FHitOverlap& OutHit)
+bool LowLevelOverlapSingleElement(int32 InParticleIndex, const Chaos::FPhysicsSolver* InSolver, const Chaos::TClusterBuffer<float, 3>& ClusterBuffer, const FGeometryCollectionPhysicsProxy* InObject, const Chaos::FImplicitObject& QueryGeom, const FTransform& InPose, FHitOverlap& OutHit)
 {
 	using namespace Chaos;
 
@@ -667,7 +667,7 @@ bool LowLevelOverlapSingleElement(int32 InParticleIndex, const Chaos::FPhysicsSo
 
 	const int32 LocalBodyIndex = InParticleIndex - PhysResult.BaseIndex;
 	const TRigidTransform<float, 3>& TM = PhysResult.ParticleToWorldTransforms[LocalBodyIndex];
-	const TImplicitObject<float, 3>* Object = ClusterBuffer.GeometryPtrs[InParticleIndex].Get();
+	const FImplicitObject* Object = ClusterBuffer.GeometryPtrs[InParticleIndex].Get();
 	
 	if(!Object)
 	{
@@ -680,7 +680,7 @@ bool LowLevelOverlapSingleElement(int32 InParticleIndex, const Chaos::FPhysicsSo
 }
 
 
-void FGeometryCollectionSQAccelerator::Overlap(const Chaos::TImplicitObject<float, 3>& QueryGeom, const FTransform& GeomPose, ChaosInterface::FSQHitBuffer<ChaosInterface::FOverlapHit>& HitBuffer, const FQueryFilterData& QueryFilterData, ICollisionQueryFilterCallbackBase& QueryCallback) const
+void FGeometryCollectionSQAccelerator::Overlap(const Chaos::FImplicitObject& QueryGeom, const FTransform& GeomPose, ChaosInterface::FSQHitBuffer<ChaosInterface::FOverlapHit>& HitBuffer, const FQueryFilterData& QueryFilterData, ICollisionQueryFilterCallbackBase& QueryCallback) const
 {
 	return;	//todo: This is currently broken because it doesn'thandle cluster unions. Need to fix function - disabling entirely for now
 	SCOPE_CYCLE_COUNTER(STAT_GCOverlap);
@@ -691,7 +691,7 @@ void FGeometryCollectionSQAccelerator::Overlap(const Chaos::TImplicitObject<floa
 	TMap<const Chaos::FPhysicsSolver*, TArray<int32>> SolverIntersectionSets;
 
 	// Getter for intersections from the mapping above
-	auto GetIntersectionsFunc = [&](const Chaos::FPhysicsSolver* InSolver, Chaos::TImplicitObject<float, 3>* InGeometry, const FTransform& InPose) -> TArray<int32>*
+	auto GetIntersectionsFunc = [&](const Chaos::FPhysicsSolver* InSolver, Chaos::FImplicitObject* InGeometry, const FTransform& InPose) -> TArray<int32>*
 	{
 		TArray<int32>* CurrentIntersections = SolverIntersectionSets.Find(InSolver);
 
@@ -723,7 +723,7 @@ void FGeometryCollectionSQAccelerator::Overlap(const Chaos::TImplicitObject<floa
 	bool bHit = false;
 	FHitOverlap Hit;
 	PxGeometryHolder Holder(QueryGeom);
-	Chaos::TImplicitObject<float, 3>* Implicit = nullptr;
+	Chaos::FImplicitObject* Implicit = nullptr;
 
 	if(Holder.getType() == PxGeometryType::eCAPSULE)
 	{
diff --git a/Engine/Source/Runtime/Experimental/GeometryCollectionEngine/Private/GeometryCollection/GeometryCollectionSQAccelerator.h b/Engine/Source/Runtime/Experimental/GeometryCollectionEngine/Private/GeometryCollection/GeometryCollectionSQAccelerator.h
--- a/Engine/Source/Runtime/Experimental/GeometryCollectionEngine/Private/GeometryCollection/GeometryCollectionSQAccelerator.h
+++ b/Engine/Source/Runtime/Experimental/GeometryCollectionEngine/Private/GeometryCollection/GeometryCollectionSQAccelerator.h
@@ -11,8 +11,8 @@ class FGeometryCollectionSQAccelerator
 {
 public:
 	void Raycast(const FVector& Start, const FVector& Dir, const float DeltaMagnitude, ChaosInterface::FSQHitBuffer<ChaosInterface::FRaycastHit>& HitBuffer, EHitFlags OutputFlags, const FQueryFilterData& QueryFilterData, ICollisionQueryFilterCallbackBase& QueryCallback) const;
-	void Sweep(const Chaos::TImplicitObject<float, 3>& QueryGeom, const FTransform& StartTM, const FVector& Dir, const float DeltaMagnitude, ChaosInterface::FSQHitBuffer<ChaosInterface::FSweepHit>& HitBuffer, EHitFlags OutputFlags, const FQueryFilterData& QueryFilterData, ICollisionQueryFilterCallbackBase& QueryCallback) const;
-	void Overlap(const Chaos::TImplicitObject<float, 3>& QueryGeom, const FTransform& GeomPose, ChaosInterface::FSQHitBuffer<ChaosInterface::FOverlapHit>& HitBuffer, const FQueryFilterData& QueryFilterData, ICollisionQueryFilterCallbackBase& QueryCallback) const;
+	void Sweep(const Chaos::FImplicitObject& QueryGeom, const FTransform& StartTM, const FVector& Dir, const float DeltaMagnitude, ChaosInterface::FSQHitBuffer<ChaosInterface::FSweepHit>& HitBuffer, EHitFlags OutputFlags, const FQueryFilterData& QueryFilterData, ICollisionQueryFilterCallbackBase& QueryCallback) const;
+	void Overlap(const Chaos::FImplicitObject& QueryGeom, const FTransform& GeomPose, ChaosInterface::FSQHitBuffer<ChaosInterface::FOverlapHit>& HitBuffer, const FQueryFilterData& QueryFilterData, ICollisionQueryFilterCallbackBase& QueryCallback) const;
 	virtual ~FGeometryCollectionSQAccelerator() {}
 
 	void AddComponent(UGeometryCollectionComponent* Component);
diff --git a/Engine/Source/Runtime/Experimental/GeometryCollectionEngine/Public/GeometryCollection/GeometryCollectionDebugDrawComponent.h b/Engine/Source/Runtime/Experimental/GeometryCollectionEngine/Public/GeometryCollection/GeometryCollectionDebugDrawComponent.h
--- a/Engine/Source/Runtime/Experimental/GeometryCollectionEngine/Public/GeometryCollection/GeometryCollectionDebugDrawComponent.h
+++ b/Engine/Source/Runtime/Experimental/GeometryCollectionEngine/Public/GeometryCollection/GeometryCollectionDebugDrawComponent.h
@@ -20,7 +20,7 @@ class AGeometryCollectionDebugDrawActor;
 class AChaosSolverActor;
 
 #if GEOMETRYCOLLECTION_DEBUG_DRAW
-namespace Chaos { template<class T, int d> class TImplicitObject; }
+namespace Chaos { class FImplicitObject; }
 namespace Chaos { template<class T, int d> class TPBDRigidParticles; }
 #endif  // #if GEOMETRYCOLLECTION_DEBUG_DRAW
 
diff --git a/Engine/Source/Runtime/Experimental/GeometryCollectionEngine/Public/GeometryCollection/GeometryCollectionParticlesData.h b/Engine/Source/Runtime/Experimental/GeometryCollectionEngine/Public/GeometryCollection/GeometryCollectionParticlesData.h
--- a/Engine/Source/Runtime/Experimental/GeometryCollectionEngine/Public/GeometryCollection/GeometryCollectionParticlesData.h
+++ b/Engine/Source/Runtime/Experimental/GeometryCollectionEngine/Public/GeometryCollection/GeometryCollectionParticlesData.h
@@ -82,7 +82,7 @@ public:
 
 	const Chaos::TVector<T, d>               & GetX                     (int32 Index) const { check(HasSyncedData(EGeometryCollectionParticlesData::X                     )); return BufferedData.GetGameDataForRead().X                     [Index]; }
 	const Chaos::TRotation<T, d>             & GetR                     (int32 Index) const { check(HasSyncedData(EGeometryCollectionParticlesData::R                     )); return BufferedData.GetGameDataForRead().R                     [Index]; }
-	const Chaos::TImplicitObject<T, d>* const& GetGeometry              (int32 Index) const { check(HasSyncedData(EGeometryCollectionParticlesData::Geometry              )); return BufferedData.GetGameDataForRead().Geometry              [Index]; }
+	const Chaos::FImplicitObject* const& GetGeometry              (int32 Index) const { check(HasSyncedData(EGeometryCollectionParticlesData::Geometry              )); return BufferedData.GetGameDataForRead().Geometry              [Index]; }
 	const Chaos::EImplicitObjectType         & GetGeometryType          (int32 Index) const { check(HasSyncedData(EGeometryCollectionParticlesData::GeometryType          )); return BufferedData.GetGameDataForRead().GeometryType          [Index]; }
 	const bool                               & IsGeometryConvex         (int32 Index) const { check(HasSyncedData(EGeometryCollectionParticlesData::GeometryIsConvex      )); return BufferedData.GetGameDataForRead().GeometryIsConvex      [Index]; }
 	const bool                               & HasGeometryBoundingBoxm  (int32 Index) const { check(HasSyncedData(EGeometryCollectionParticlesData::GeometryHasBoundingBox)); return BufferedData.GetGameDataForRead().GeometryHasBoundingBox[Index]; }
@@ -124,7 +124,7 @@ private:
 		FDataFlags SyncedDataFlags;
 		Chaos::TArrayCollectionArray<Chaos::TVector<T, d>               > X;
 		Chaos::TArrayCollectionArray<Chaos::TRotation<T, d>             > R;
-		Chaos::TArrayCollectionArray<const Chaos::TImplicitObject<T, d>*> Geometry;
+		Chaos::TArrayCollectionArray<const Chaos::FImplicitObject*> Geometry;
 		Chaos::TArrayCollectionArray<Chaos::EImplicitObjectType         > GeometryType;
 		Chaos::TArrayCollectionArray<bool                               > GeometryIsConvex;
 		Chaos::TArrayCollectionArray<bool                               > GeometryHasBoundingBox;
diff --git a/Engine/Source/Runtime/Experimental/GeometryCollectionSimulationCore/Private/GeometryCollection/GeometryCollectionCollisionStructureManager.cpp b/Engine/Source/Runtime/Experimental/GeometryCollectionSimulationCore/Private/GeometryCollection/GeometryCollectionCollisionStructureManager.cpp
--- a/Engine/Source/Runtime/Experimental/GeometryCollectionSimulationCore/Private/GeometryCollection/GeometryCollectionCollisionStructureManager.cpp
+++ b/Engine/Source/Runtime/Experimental/GeometryCollectionSimulationCore/Private/GeometryCollection/GeometryCollectionCollisionStructureManager.cpp
@@ -34,7 +34,7 @@ FCollisionStructureManager::FSimplicial*
 FCollisionStructureManager::NewSimplicial(
 	const Chaos::TParticles<float, 3>& Vertices,
 	Chaos::TTriangleMesh<float>& TriMesh,
-	const Chaos::TImplicitObject<float, 3>* Implicit,
+	const Chaos::FImplicitObject* Implicit,
 	int32 CollisionParticlesMaxInput)
 {
 	FCollisionStructureManager::FSimplicial * Simplicial = new FCollisionStructureManager::FSimplicial();
@@ -249,7 +249,7 @@ FCollisionStructureManager::NewImplicit(
 	const ECollisionTypeEnum CollisionType,
 	const EImplicitTypeEnum ImplicitType)
 {
-	Chaos::TImplicitObject<float, 3>* Implicit = nullptr;
+	Chaos::FImplicitObject* Implicit = nullptr;
 	if (ImplicitType == EImplicitTypeEnum::Chaos_Implicit_Box)
 	{
 		FVector HalfExtents = CollisionBounds.GetExtent() * (1 - CollisionObjectReduction / 100.f);
diff --git a/Engine/Source/Runtime/Experimental/GeometryCollectionSimulationCore/Public/GeometryCollection/GeometryCollectionCollisionStructureManager.h b/Engine/Source/Runtime/Experimental/GeometryCollectionSimulationCore/Public/GeometryCollection/GeometryCollectionCollisionStructureManager.h
--- a/Engine/Source/Runtime/Experimental/GeometryCollectionSimulationCore/Public/GeometryCollection/GeometryCollectionCollisionStructureManager.h
+++ b/Engine/Source/Runtime/Experimental/GeometryCollectionSimulationCore/Public/GeometryCollection/GeometryCollectionCollisionStructureManager.h
@@ -28,12 +28,12 @@ public:
 
 	typedef TArray<Chaos::TVector<float, 3>> FPoints;
 	typedef Chaos::TBVHParticles<float,3> FSimplicial;
-	typedef Chaos::TImplicitObject<float, 3> FImplicit;
+	typedef Chaos::FImplicitObject FImplicit;
 
 	static FSimplicial* NewSimplicial(
 		const Chaos::TParticles<float, 3>& Vertices,
 		Chaos::TTriangleMesh<float>& TriMesh,
-		const Chaos::TImplicitObject<float, 3>* Implicit,
+		const Chaos::FImplicitObject* Implicit,
 		int32 CollisionParticlesMaxInput
 	);
 	
diff --git a/Engine/Source/Runtime/Landscape/Private/LandscapeCollision.cpp b/Engine/Source/Runtime/Landscape/Private/LandscapeCollision.cpp
--- a/Engine/Source/Runtime/Landscape/Private/LandscapeCollision.cpp
+++ b/Engine/Source/Runtime/Landscape/Private/LandscapeCollision.cpp
@@ -404,7 +404,7 @@ void ULandscapeHeightfieldCollisionComponent::OnCreatePhysicsState()
 				FPhysicsInterface::CreateActor(Params, PhysHandle);
 
 				Chaos::TShapesArray<float, 3> ShapeArray;
-				TArray<TUniquePtr<Chaos::TImplicitObject<float, 3>>> Geoms;
+				TArray<TUniquePtr<Chaos::FImplicitObject>> Geoms;
 
 				// First add complex geometry
 				TUniquePtr<Chaos::TPerShapeData<float, 3>> NewShape = Chaos::TPerShapeData<float, 3>::CreatePerShapeData();
diff --git a/Engine/Source/Runtime/PhysicsCore/Private/ChaosInterfaceWrapperCore.cpp b/Engine/Source/Runtime/PhysicsCore/Private/ChaosInterfaceWrapperCore.cpp
--- a/Engine/Source/Runtime/PhysicsCore/Private/ChaosInterfaceWrapperCore.cpp
+++ b/Engine/Source/Runtime/PhysicsCore/Private/ChaosInterfaceWrapperCore.cpp
@@ -9,7 +9,7 @@
 
 namespace ChaosInterface
 {
-	ECollisionShapeType GetImplicitType(const Chaos::TImplicitObject<float, 3>& InGeometry)
+	ECollisionShapeType GetImplicitType(const Chaos::FImplicitObject& InGeometry)
 	{
 		switch (Chaos::GetInnerType(InGeometry.GetType()))
 		{
diff --git a/Engine/Source/Runtime/PhysicsCore/Private/PhysTestSerializer.cpp b/Engine/Source/Runtime/PhysicsCore/Private/PhysTestSerializer.cpp
--- a/Engine/Source/Runtime/PhysicsCore/Private/PhysTestSerializer.cpp
+++ b/Engine/Source/Runtime/PhysicsCore/Private/PhysTestSerializer.cpp
@@ -281,7 +281,7 @@ void FPhysTestSerializer::CreateChaosData()
 			PxActorToChaosHandle.Add(Act, Particle.Get());
 
 			//geometry
-			TArray<TUniquePtr<TImplicitObject<float, 3>>> Geoms;
+			TArray<TUniquePtr<FImplicitObject>> Geoms;
 			const int32 NumShapes = Actor->getNbShapes();
 			TArray<PxShape*> Shapes;
 			Shapes.AddUninitialized(NumShapes);
@@ -298,7 +298,7 @@ void FPhysTestSerializer::CreateChaosData()
 			{
 				if (Geoms.Num() == 1)
 				{
-					auto SharedGeom = TSharedPtr<TImplicitObject<float, 3>, ESPMode::ThreadSafe>(Geoms[0].Release());
+					auto SharedGeom = TSharedPtr<FImplicitObject, ESPMode::ThreadSafe>(Geoms[0].Release());
 					GTParticle->SetGeometry(SharedGeom);
 					Particle->SetSharedGeometry(SharedGeom);
 				}
diff --git a/Engine/Source/Runtime/PhysicsCore/Private/SQCapture.cpp b/Engine/Source/Runtime/PhysicsCore/Private/SQCapture.cpp
--- a/Engine/Source/Runtime/PhysicsCore/Private/SQCapture.cpp
+++ b/Engine/Source/Runtime/PhysicsCore/Private/SQCapture.cpp
@@ -340,7 +340,7 @@ void FSQCapture::Serialize(Chaos::FChaosArchive& Ar)
 	}
 }
 
-void FSQCapture::StartCaptureChaosSweep(const Chaos::TPBDRigidsEvolutionGBF<float, 3>& Evolution, const Chaos::TImplicitObject<float,3>& InQueryGeom, const FTransform& InStartTM, const FVector& InDir, float InDeltaMag, FHitFlags InOutputFlags, const FQueryFilterData& QueryFilter, const FCollisionFilterData& FilterData, ICollisionQueryFilterCallbackBase& Callback)
+void FSQCapture::StartCaptureChaosSweep(const Chaos::TPBDRigidsEvolutionGBF<float, 3>& Evolution, const Chaos::FImplicitObject& InQueryGeom, const FTransform& InStartTM, const FVector& InDir, float InDeltaMag, FHitFlags InOutputFlags, const FQueryFilterData& QueryFilter, const FCollisionFilterData& FilterData, ICollisionQueryFilterCallbackBase& Callback)
 {
 	if (IsInGameThread())
 	{
@@ -394,7 +394,7 @@ void FSQCapture::EndCaptureChaosRaycast(const ChaosInterface::FSQHitBuffer<Chaos
 	}
 }
 
-void FSQCapture::StartCaptureChaosOverlap(const Chaos::TPBDRigidsEvolutionGBF<float, 3>& Evolution, const Chaos::TImplicitObject<float, 3>& InQueryGeom, const FTransform& InStartTM, const FQueryFilterData& QueryFilter, const FCollisionFilterData& FilterData, ICollisionQueryFilterCallbackBase& Callback)
+void FSQCapture::StartCaptureChaosOverlap(const Chaos::TPBDRigidsEvolutionGBF<float, 3>& Evolution, const Chaos::FImplicitObject& InQueryGeom, const FTransform& InStartTM, const FQueryFilterData& QueryFilter, const FCollisionFilterData& FilterData, ICollisionQueryFilterCallbackBase& Callback)
 {
 	if (IsInGameThread())
 	{
diff --git a/Engine/Source/Runtime/PhysicsCore/Public/ChaosInterfaceWrapperCore.h b/Engine/Source/Runtime/PhysicsCore/Public/ChaosInterfaceWrapperCore.h
--- a/Engine/Source/Runtime/PhysicsCore/Public/ChaosInterfaceWrapperCore.h
+++ b/Engine/Source/Runtime/PhysicsCore/Public/ChaosInterfaceWrapperCore.h
@@ -14,8 +14,7 @@ class UPhysicalMaterial;
 
 namespace Chaos
 {
-	template <class T, int d>
-	class TImplicitObject;
+	class FImplicitObject;
 
 	template <class T>
 	class TCapsule;
@@ -97,9 +96,9 @@ extern PHYSICSCORE_API FCollisionFilterData GetQueryFilterData(const Chaos::TPer
 extern PHYSICSCORE_API FCollisionFilterData GetSimulationFilterData(const Chaos::TPerShapeData<float, 3>& Shape);
 
 
-PHYSICSCORE_API ECollisionShapeType GetImplicitType(const Chaos::TImplicitObject<float, 3>& InGeometry);
+PHYSICSCORE_API ECollisionShapeType GetImplicitType(const Chaos::FImplicitObject& InGeometry);
 
-FORCEINLINE ECollisionShapeType GetType(const Chaos::TImplicitObject<float, 3>& InGeometry)
+FORCEINLINE ECollisionShapeType GetType(const Chaos::FImplicitObject& InGeometry)
 {
 	return GetImplicitType(InGeometry);
 }
@@ -185,7 +184,7 @@ inline uint32 GetNumShapes(const FDummyPhysActor& RigidActor)
 	return 0;
 }
 
-inline void GetShapes(const FDummyPhysActor& RigidActor, Chaos::TImplicitObject<float, 3>** ShapesBuffer, uint32 NumShapes)
+inline void GetShapes(const FDummyPhysActor& RigidActor, Chaos::FImplicitObject** ShapesBuffer, uint32 NumShapes)
 {
 
 }
@@ -195,7 +194,7 @@ inline void SetActor(FDummyPhysType& Hit, FDummyPhysActor* Actor)
 
 }
 
-inline void SetShape(FDummyPhysType& Hit, Chaos::TImplicitObject<float, 3>* Shape)
+inline void SetShape(FDummyPhysType& Hit, Chaos::FImplicitObject* Shape)
 {
 
 }
diff --git a/Engine/Source/Runtime/PhysicsCore/Public/CollisionQueryFilterCallbackCore.h b/Engine/Source/Runtime/PhysicsCore/Public/CollisionQueryFilterCallbackCore.h
--- a/Engine/Source/Runtime/PhysicsCore/Public/CollisionQueryFilterCallbackCore.h
+++ b/Engine/Source/Runtime/PhysicsCore/Public/CollisionQueryFilterCallbackCore.h
@@ -26,8 +26,7 @@ enum class ECollisionQueryHitType : uint8
 
 namespace Chaos
 {
-	template <typename, int>
-	class TImplicitObject;
+	class FImplicitObject;
 }
 
 class ICollisionQueryFilterCallbackBase
diff --git a/Engine/Source/Runtime/PhysicsCore/Public/PhysXToChaosUtil.h b/Engine/Source/Runtime/PhysicsCore/Public/PhysXToChaosUtil.h
--- a/Engine/Source/Runtime/PhysicsCore/Public/PhysXToChaosUtil.h
+++ b/Engine/Source/Runtime/PhysicsCore/Public/PhysXToChaosUtil.h
@@ -23,7 +23,7 @@ inline TUniquePtr<Chaos::TImplicitObjectTransformed<float, 3>> PxShapeToChaosGeo
 {
 	using namespace Chaos;
 
-	TUniquePtr<TImplicitObject<float, 3>> InnerObj;
+	TUniquePtr<FImplicitObject> InnerObj;
 	PxTransform ShapeTM = Shape->getLocalPose();
 	PxGeometryHolder Geom = Shape->getGeometry();
 	switch (Geom.getType())
diff --git a/Engine/Source/Runtime/PhysicsCore/Public/PhysicsInterfaceDeclaresCore.h b/Engine/Source/Runtime/PhysicsCore/Public/PhysicsInterfaceDeclaresCore.h
--- a/Engine/Source/Runtime/PhysicsCore/Public/PhysicsInterfaceDeclaresCore.h
+++ b/Engine/Source/Runtime/PhysicsCore/Public/PhysicsInterfaceDeclaresCore.h
@@ -69,8 +69,7 @@ struct FQueryDebugParams {};
 
 namespace Chaos
 {
-	template<class T, int d>
-	class TImplicitObject;
+	class FImplicitObject;
 
 	template<class T>
 	class TCapsule;
@@ -106,7 +105,7 @@ using FPhysicsQueryHit = ChaosInterface::FQueryHit;
 using FPhysicsTransform = FTransform;
 
 using FPhysicsShape = Chaos::TPerShapeData<float, 3>;
-using FPhysicsGeometry = Chaos::TImplicitObject<float, 3>;
+using FPhysicsGeometry = Chaos::FImplicitObject;
 using FPhysicsCapsuleGeometry = Chaos::TCapsule<float>;
 using FPhysicsMaterial = FPhysTypeDummy;
 using FPhysicsActor = Chaos::TGeometryParticle<float,3>;
diff --git a/Engine/Source/Runtime/PhysicsCore/Public/SQCapture.h b/Engine/Source/Runtime/PhysicsCore/Public/SQCapture.h
--- a/Engine/Source/Runtime/PhysicsCore/Public/SQCapture.h
+++ b/Engine/Source/Runtime/PhysicsCore/Public/SQCapture.h
@@ -51,13 +51,13 @@ struct PHYSICSCORE_API FSQCapture
 	void EndCapturePhysXOverlap(const PxHitCallback<PxOverlapHit>& Results);
 #endif
 
-	void StartCaptureChaosSweep(const Chaos::TPBDRigidsEvolutionGBF<float, 3>& Evolution, const Chaos::TImplicitObject<float, 3>& InQueryGeom, const FTransform& InStartTM, const FVector& InDir, float InDeltaMag, FHitFlags InOutputFlags, const FQueryFilterData& QueryFilter, const FCollisionFilterData& FilterData, ICollisionQueryFilterCallbackBase& Callback);
+	void StartCaptureChaosSweep(const Chaos::TPBDRigidsEvolutionGBF<float, 3>& Evolution, const Chaos::FImplicitObject& InQueryGeom, const FTransform& InStartTM, const FVector& InDir, float InDeltaMag, FHitFlags InOutputFlags, const FQueryFilterData& QueryFilter, const FCollisionFilterData& FilterData, ICollisionQueryFilterCallbackBase& Callback);
 	void EndCaptureChaosSweep(const ChaosInterface::FSQHitBuffer<ChaosInterface::FSweepHit>& Results);
 
 	void StartCaptureChaosRaycast(const Chaos::TPBDRigidsEvolutionGBF<float,3>& Evolution, const FVector& InStartPoint, const FVector& InDir, float InDeltaMag, FHitFlags InOutputFlags, const FQueryFilterData& QueryFilter, const FCollisionFilterData& FilterData, ICollisionQueryFilterCallbackBase& Callback);
 	void EndCaptureChaosRaycast(const ChaosInterface::FSQHitBuffer<ChaosInterface::FRaycastHit>& Results);
 
-	void StartCaptureChaosOverlap(const Chaos::TPBDRigidsEvolutionGBF<float, 3>& Evolution, const Chaos::TImplicitObject<float, 3>& InQueryGeom, const FTransform& InStartTM, const FQueryFilterData& QueryFilter, const FCollisionFilterData& FilterData, ICollisionQueryFilterCallbackBase& Callback);
+	void StartCaptureChaosOverlap(const Chaos::TPBDRigidsEvolutionGBF<float, 3>& Evolution, const Chaos::FImplicitObject& InQueryGeom, const FTransform& InStartTM, const FQueryFilterData& QueryFilter, const FCollisionFilterData& FilterData, ICollisionQueryFilterCallbackBase& Callback);
 	void EndCaptureChaosOverlap(const ChaosInterface::FSQHitBuffer<ChaosInterface::FOverlapHit>& Results);
 
 	ECollisionQueryHitType GetFilterResult(const Chaos::TPerShapeData<float,3>* Shape, const Chaos::TGeometryParticle<float,3>* Actor) const;
@@ -82,9 +82,9 @@ struct PHYSICSCORE_API FSQCapture
 	PxGeometryHolder PhysXGeometry;
 #endif
 
-	TUniquePtr<Chaos::TImplicitObject<float, 3>> ChaosOwnerObject;	//should be private, do not access directly
-	const Chaos::TImplicitObject<float, 3>* ChaosGeometry;
-	TUniquePtr<Chaos::TImplicitObject<float, 3>> SerializableChaosGeometry;
+	TUniquePtr<Chaos::FImplicitObject> ChaosOwnerObject;	//should be private, do not access directly
+	const Chaos::FImplicitObject* ChaosGeometry;
+	TUniquePtr<Chaos::FImplicitObject> SerializableChaosGeometry;
 #if WITH_PHYSX
 	//for now just use physx hit buffer
 	ChaosInterface::FSQHitBuffer<ChaosInterface::FSweepHit> ChaosSweepBuffer;
diff --git a/Engine/Source/Runtime/PhysicsSQ/Private/Experimental/SceneQueryChaosImp.h b/Engine/Source/Runtime/PhysicsSQ/Private/Experimental/SceneQueryChaosImp.h
--- a/Engine/Source/Runtime/PhysicsSQ/Private/Experimental/SceneQueryChaosImp.h
+++ b/Engine/Source/Runtime/PhysicsSQ/Private/Experimental/SceneQueryChaosImp.h
@@ -4,19 +4,19 @@
 
 #if WITH_CHAOS
 
-inline bool LowLevelRaycastImp(const FVector& Start, const FVector& Dir, float DeltaMag, const Chaos::TImplicitObject<float, 3>& Shape, const FTransform ActorTM, EHitFlags OutputFlags, FHitRaycast& Hit)
+inline bool LowLevelRaycastImp(const FVector& Start, const FVector& Dir, float DeltaMag, const Chaos::FImplicitObject& Shape, const FTransform ActorTM, EHitFlags OutputFlags, FHitRaycast& Hit)
 {
 	//TODO_SQ_IMPLEMENTATION
 	return false;			
 }
 
-inline bool LowLevelSweepImp(const FTransform& StartTM, const FVector& Dir, float DeltaMag, const FPhysicsGeometry& SweepGeom, const Chaos::TImplicitObject<float, 3>& Shape, const FTransform ActorTM, EHitFlags OutputFlags, FHitSweep& Hit)
+inline bool LowLevelSweepImp(const FTransform& StartTM, const FVector& Dir, float DeltaMag, const FPhysicsGeometry& SweepGeom, const Chaos::FImplicitObject& Shape, const FTransform ActorTM, EHitFlags OutputFlags, FHitSweep& Hit)
 {
 	//TODO_SQ_IMPLEMENTATION
 	return false;
 }
 
-inline bool LowLevelOverlapImp(const FTransform& GeomPose, const FPhysicsGeometry& OverlapGeom, const Chaos::TImplicitObject<float, 3>& Shape, const FTransform ActorTM, FHitOverlap& Overlap)
+inline bool LowLevelOverlapImp(const FTransform& GeomPose, const FPhysicsGeometry& OverlapGeom, const Chaos::FImplicitObject& Shape, const FTransform ActorTM, FHitOverlap& Overlap)
 {
 	//TODO_SQ_IMPLEMENTATION
 	return false;
diff --git a/Engine/Source/Runtime/PhysicsSQ/Private/SQAccelerator.cpp b/Engine/Source/Runtime/PhysicsSQ/Private/SQAccelerator.cpp
--- a/Engine/Source/Runtime/PhysicsSQ/Private/SQAccelerator.cpp
+++ b/Engine/Source/Runtime/PhysicsSQ/Private/SQAccelerator.cpp
@@ -63,7 +63,7 @@ FChaosSQAccelerator::FChaosSQAccelerator(const Chaos::ISpatialAcceleration<Chaos
 
 struct FPreFilterInfo
 {
-	const Chaos::TImplicitObject<float, 3>* Geom;
+	const Chaos::FImplicitObject* Geom;
 	int32 ActorIdx;
 };
 
@@ -101,7 +101,7 @@ struct TSQVisitor : public Chaos::ISpatialVisitor<TPayload, float>
 	}
 
 	TSQVisitor(const FTransform& InStartTM, const FVector& InDir, ChaosInterface::FSQHitBuffer<ChaosInterface::FSweepHit>& InHitBuffer, EHitFlags InOutputFlags,
-		const FQueryFilterData& InQueryFilterData, ICollisionQueryFilterCallbackBase& InQueryCallback, const Chaos::TImplicitObject<float,3>& InQueryGeom, const FQueryDebugParams& InDebugParams)
+		const FQueryFilterData& InQueryFilterData, ICollisionQueryFilterCallbackBase& InQueryCallback, const Chaos::FImplicitObject& InQueryGeom, const FQueryDebugParams& InDebugParams)
 		: StartTM(InStartTM)
 		, Dir(InDir)
 		, HitBuffer(InHitBuffer)
@@ -120,7 +120,7 @@ struct TSQVisitor : public Chaos::ISpatialVisitor<TPayload, float>
 	}
 
 	TSQVisitor(const FTransform& InWorldTM, ChaosInterface::FSQHitBuffer<ChaosInterface::FOverlapHit>& InHitBuffer,
-		const FQueryFilterData& InQueryFilterData, ICollisionQueryFilterCallbackBase& InQueryCallback, const Chaos::TImplicitObject<float, 3>& InQueryGeom, const FQueryDebugParams& InDebugParams)
+		const FQueryFilterData& InQueryFilterData, ICollisionQueryFilterCallbackBase& InQueryCallback, const Chaos::FImplicitObject& InQueryGeom, const FQueryDebugParams& InDebugParams)
 		: StartTM(InWorldTM)
 		, HitBuffer(InHitBuffer)
 		, QueryFilterData(InQueryFilterData)
@@ -180,7 +180,7 @@ private:
 
 		for (const auto& Shape : Shapes)
 		{
-			const TImplicitObject<float, 3>* Geom = Shape->Geometry.Get();
+			const FImplicitObject* Geom = Shape->Geometry.Get();
 			//TODO: use gt particles directly
 			//#TODO alternative to px flags
 #if WITH_PHYSX
@@ -290,7 +290,7 @@ private:
 	const FQueryFilterData& QueryFilterData;
 	ICollisionQueryFilterCallbackBase& QueryCallback;
 	bool bAnyHit;
-	const Chaos::TImplicitObject<float, 3>* QueryGeom;
+	const Chaos::FImplicitObject* QueryGeom;
 	const FQueryDebugParams DebugParams;
 };
 
@@ -305,7 +305,7 @@ void FChaosSQAccelerator::Raycast(const FVector& Start, const FVector& Dir, cons
 	HitBuffer.DecFlushCount();
 }
 
-void FChaosSQAccelerator::Sweep(const Chaos::TImplicitObject<float, 3>& QueryGeom, const FTransform& StartTM, const FVector& Dir, const float DeltaMagnitude, ChaosInterface::FSQHitBuffer<ChaosInterface::FSweepHit>& HitBuffer, EHitFlags OutputFlags, const FQueryFilterData& QueryFilterData, ICollisionQueryFilterCallbackBase& QueryCallback, const FQueryDebugParams& DebugParams) const
+void FChaosSQAccelerator::Sweep(const Chaos::FImplicitObject& QueryGeom, const FTransform& StartTM, const FVector& Dir, const float DeltaMagnitude, ChaosInterface::FSQHitBuffer<ChaosInterface::FSweepHit>& HitBuffer, EHitFlags OutputFlags, const FQueryFilterData& QueryFilterData, ICollisionQueryFilterCallbackBase& QueryCallback, const FQueryDebugParams& DebugParams) const
 {
 	using namespace Chaos;
 	using namespace ChaosInterface;
@@ -318,7 +318,7 @@ void FChaosSQAccelerator::Sweep(const Chaos::TImplicitObject<float, 3>& QueryGeo
 	HitBuffer.DecFlushCount();
 }
 
-void FChaosSQAccelerator::Overlap(const Chaos::TImplicitObject<float, 3>& QueryGeom, const FTransform& GeomPose, ChaosInterface::FSQHitBuffer<ChaosInterface::FOverlapHit>& HitBuffer, const FQueryFilterData& QueryFilterData, ICollisionQueryFilterCallbackBase& QueryCallback, const FQueryDebugParams& DebugParams) const
+void FChaosSQAccelerator::Overlap(const Chaos::FImplicitObject& QueryGeom, const FTransform& GeomPose, ChaosInterface::FSQHitBuffer<ChaosInterface::FOverlapHit>& HitBuffer, const FQueryFilterData& QueryFilterData, ICollisionQueryFilterCallbackBase& QueryCallback, const FQueryDebugParams& DebugParams) const
 {
 	using namespace Chaos;
 	using namespace ChaosInterface;
diff --git a/Engine/Source/Runtime/PhysicsSQ/Public/SQAccelerator.h b/Engine/Source/Runtime/PhysicsSQ/Public/SQAccelerator.h
--- a/Engine/Source/Runtime/PhysicsSQ/Public/SQAccelerator.h
+++ b/Engine/Source/Runtime/PhysicsSQ/Public/SQAccelerator.h
@@ -21,8 +21,7 @@ namespace physx
 
 namespace Chaos
 {
-	template <typename, int>
-	class TImplicitObject;
+	class FImplicitObject;
 
 	template <typename T, int d>
 	class TAccelerationStructureHandle;
@@ -45,8 +44,8 @@ public:
 	virtual ~FChaosSQAccelerator() {};
 
 	void Raycast(const FVector& Start, const FVector& Dir, const float DeltaMagnitude, ChaosInterface::FSQHitBuffer<ChaosInterface::FRaycastHit>& HitBuffer, EHitFlags OutputFlags, const FQueryFilterData& QueryFilterData, ICollisionQueryFilterCallbackBase& QueryCallback, const FQueryDebugParams& DebugParams = FQueryDebugParams()) const;
-	void Sweep(const Chaos::TImplicitObject<float, 3>& QueryGeom, const FTransform& StartTM, const FVector& Dir, const float DeltaMagnitude, ChaosInterface::FSQHitBuffer<ChaosInterface::FSweepHit>& HitBuffer, EHitFlags OutputFlags, const FQueryFilterData& QueryFilterData, ICollisionQueryFilterCallbackBase& QueryCallback, const FQueryDebugParams& DebugParams = FQueryDebugParams()) const;
-	void Overlap(const Chaos::TImplicitObject<float, 3>& QueryGeom, const FTransform& GeomPose, ChaosInterface::FSQHitBuffer<ChaosInterface::FOverlapHit>& HitBuffer, const FQueryFilterData& QueryFilterData, ICollisionQueryFilterCallbackBase& QueryCallback, const FQueryDebugParams& DebugParams = FQueryDebugParams()) const;
+	void Sweep(const Chaos::FImplicitObject& QueryGeom, const FTransform& StartTM, const FVector& Dir, const float DeltaMagnitude, ChaosInterface::FSQHitBuffer<ChaosInterface::FSweepHit>& HitBuffer, EHitFlags OutputFlags, const FQueryFilterData& QueryFilterData, ICollisionQueryFilterCallbackBase& QueryCallback, const FQueryDebugParams& DebugParams = FQueryDebugParams()) const;
+	void Overlap(const Chaos::FImplicitObject& QueryGeom, const FTransform& GeomPose, ChaosInterface::FSQHitBuffer<ChaosInterface::FOverlapHit>& HitBuffer, const FQueryFilterData& QueryFilterData, ICollisionQueryFilterCallbackBase& QueryCallback, const FQueryDebugParams& DebugParams = FQueryDebugParams()) const;
 
 private:
 	const Chaos::ISpatialAcceleration<Chaos::TAccelerationStructureHandle<float, 3>, float, 3>& SpatialAcceleration;
