Default OutNormal in the case of MTDs against convexes - not supported yet

diff --git a/Engine/Source/Runtime/Experimental/Chaos/Public/Chaos/GeometryQueries.h b/Engine/Source/Runtime/Experimental/Chaos/Public/Chaos/GeometryQueries.h
--- a/Engine/Source/Runtime/Experimental/Chaos/Public/Chaos/GeometryQueries.h
+++ b/Engine/Source/Runtime/Experimental/Chaos/Public/Chaos/GeometryQueries.h
@@ -208,6 +208,13 @@ namespace Chaos
 				}
 
 			}
+
+			// Compute MTD in the case of an initial overlap
+			if (bResult && bComputeMTD && OutTime == 0.f)
+			{
+				ensure(false); // We don't support MTD for non-convex types yet!
+				OutNormal = TVector<float, 3>(0.f, 0.f, 1.f);
+			}
 		}
 
 		//put back into world space
