Don't access OutNormal and OutPosition in ImplicitObjectScaled/Transformed when TOI is 0, outputs may be uninitialzied.

diff --git a/Engine/Source/Runtime/Experimental/Chaos/Public/Chaos/ImplicitObjectScaled.h b/Engine/Source/Runtime/Experimental/Chaos/Public/Chaos/ImplicitObjectScaled.h
--- a/Engine/Source/Runtime/Experimental/Chaos/Public/Chaos/ImplicitObjectScaled.h
+++ b/Engine/Source/Runtime/Experimental/Chaos/Public/Chaos/ImplicitObjectScaled.h
@@ -276,8 +276,11 @@ public:
 			if (MObject->Raycast(UnscaledStart, UnscaledDir, UnscaledLength, MInternalThickness + Thickness * MInvScale[0], UnscaledTime, UnscaledPosition, UnscaledNormal, OutFaceIndex))
 			{
 				OutTime = LengthScaleInv * UnscaledTime;
-				OutPosition = MScale * UnscaledPosition;
-				OutNormal = (MInvScale * UnscaledNormal).GetSafeNormal();
+				if (OutTime != 0) // Normal/Position output may be uninitialized with TOI 0.
+				{
+					OutPosition = MScale * UnscaledPosition;
+					OutNormal = (MInvScale * UnscaledNormal).GetSafeNormal();
+				}
 				ensure(OutTime <= Length);
 				return true;
 			}
diff --git a/Engine/Source/Runtime/Experimental/Chaos/Public/Chaos/ImplicitObjectTransformed.h b/Engine/Source/Runtime/Experimental/Chaos/Public/Chaos/ImplicitObjectTransformed.h
--- a/Engine/Source/Runtime/Experimental/Chaos/Public/Chaos/ImplicitObjectTransformed.h
+++ b/Engine/Source/Runtime/Experimental/Chaos/Public/Chaos/ImplicitObjectTransformed.h
@@ -109,8 +109,11 @@ public:
 
 		if (MObject->Raycast(LocalStart, LocalDir, Length, Thickness, OutTime, LocalPosition, LocalNormal, OutFaceIndex))
 		{
-			OutPosition = MTransform.TransformPosition(LocalPosition);
-			OutNormal = MTransform.TransformVector(LocalNormal);
+			if (OutTime != 0.0f)
+			{
+				OutPosition = MTransform.TransformPosition(LocalPosition);
+				OutNormal = MTransform.TransformVector(LocalNormal);
+			}
 			return true;
 		}
 		
