Refactor scene queries to be more templatized so we can avoid several virtual calls inside GJK

diff --git a/Engine/Source/Runtime/Engine/Private/PhysicsEngine/Experimental/PhysInterface_Chaos.cpp b/Engine/Source/Runtime/Engine/Private/PhysicsEngine/Experimental/PhysInterface_Chaos.cpp
--- a/Engine/Source/Runtime/Engine/Private/PhysicsEngine/Experimental/PhysInterface_Chaos.cpp
+++ b/Engine/Source/Runtime/Engine/Private/PhysicsEngine/Experimental/PhysInterface_Chaos.cpp
@@ -1583,7 +1583,7 @@ bool FPhysInterface_Chaos::Sweep_Geom(FHitResult& OutHit, const FBodyInstance* I
 							Chaos::TVector<float, 3> WorldPosition;
 							Chaos::TVector<float, 3> WorldNormal;
 							int32 FaceIdx;
-							if(Chaos::SweepQuery<float, 3>(*Shape->Geometry, ActorTM, ShapeAdapter.GetGeometry(), StartTM, Dir, DeltaMag, Hit.Distance, WorldPosition, WorldNormal, FaceIdx, 0.f, false))
+							if (Chaos::CastHelper(ShapeAdapter.GetGeometry(), [&](const auto& Downcast) { return Chaos::SweepQuery(*Shape->Geometry, ActorTM, Downcast, StartTM, Dir, DeltaMag, Hit.Distance, WorldPosition, WorldNormal, FaceIdx, 0.f, false); }))
 							{
 								// we just like to make sure if the hit is made
 								FCollisionFilterData QueryFilter;
@@ -1666,7 +1666,7 @@ bool Overlap_GeomInternal(const FBodyInstance* InInstance, const Chaos::FImplici
 			}
 			//else	//todo: don't bother with this once MTD is implemented
 			{
-				if(Chaos::OverlapQuery<float, 3>(*Shape->Geometry, ActorTM, InGeom, GeomTransform))
+				if (Chaos::CastHelper(InGeom, [&](const auto& Downcast) { return Chaos::OverlapQuery(*Shape->Geometry, ActorTM, Downcast, GeomTransform); }))
 				{
 					return true;
 				}
diff --git a/Engine/Source/Runtime/Experimental/Chaos/Private/Chaos/HeightField.cpp b/Engine/Source/Runtime/Experimental/Chaos/Private/Chaos/HeightField.cpp
--- a/Engine/Source/Runtime/Experimental/Chaos/Private/Chaos/HeightField.cpp
+++ b/Engine/Source/Runtime/Experimental/Chaos/Private/Chaos/HeightField.cpp
@@ -188,12 +188,12 @@ namespace Chaos
 		T Thickness;
 	};
 
-	template<typename T>
+	template<typename GeomQueryType, typename T>
 	class THeightfieldSweepVisitor
 	{
 	public:
 
-		THeightfieldSweepVisitor(const typename THeightField<T>::FDataType* InData, const FImplicitObject& InQueryGeom, const TRigidTransform<T, 3>& InStartTM, const TVector<T, 3>& InDir, const T InThickness)
+		THeightfieldSweepVisitor(const typename THeightField<T>::FDataType* InData, const GeomQueryType& InQueryGeom, const TRigidTransform<T, 3>& InStartTM, const TVector<T, 3>& InDir, const T InThickness)
 			: OutTime(TNumericLimits<T>::Max())
 			, OutFaceIndex(INDEX_NONE)
 			, HfData(InData)
@@ -221,7 +221,7 @@ namespace Chaos
 				T Time;
 				TVector<T, 3> HitPosition;
 				TVector<T, 3> HitNormal;
-				if(GJKRaycast<T>(Convex, OtherGeom, StartTM, Dir, CurrentLength, Time, HitPosition, HitNormal, Thickness))
+				if(GJKRaycast2<T>(Convex, OtherGeom, StartTM, Dir, CurrentLength, Time, HitPosition, HitNormal, Thickness))
 				{
 					if(Time < OutTime)
 					{
@@ -259,7 +259,7 @@ namespace Chaos
 
 		const typename THeightField<T>::FDataType* HfData;
 		const TRigidTransform<T, 3> StartTM;
-		const FImplicitObject& OtherGeom;
+		const GeomQueryType& OtherGeom;
 		const TVector<T, 3>& Dir;
 		const T Thickness;
 
@@ -1002,7 +1002,8 @@ namespace Chaos
 
 
 	template <typename T>
-	bool THeightField<T>::OverlapGeom(const FImplicitObject& QueryGeom, const TRigidTransform<T, 3>& QueryTM, const T Thickness) const
+	template <typename QueryGeomType>
+	bool THeightField<T>::OverlapGeomImp(const QueryGeomType& QueryGeom, const TRigidTransform<T, 3>& QueryTM, const T Thickness) const
 	{
 		auto OverlapTriangle = [&](const TVector<T, 3>& A, const TVector<T, 3>& B, const TVector<T, 3>& C) -> bool
 		{
@@ -1060,10 +1061,59 @@ namespace Chaos
 	}
 
 	template <typename T>
-	bool THeightField<T>::SweepGeom(const FImplicitObject& QueryGeom, const TRigidTransform<T, 3>& StartTM, const TVector<T, 3>& Dir, const T Length, T& OutTime, TVector<T, 3>& OutPosition, TVector<T, 3>& OutNormal, int32& OutFaceIndex, const T Thickness) const
+	bool THeightField<T>::OverlapGeom(const TSphere<T, 3>& QueryGeom, const TRigidTransform<T, 3>& QueryTM, const T Thickness) const
+	{
+		return OverlapGeomImp(QueryGeom, QueryTM, Thickness);
+	}
+
+	template <typename T>
+	bool THeightField<T>::OverlapGeom(const TBox<T, 3>& QueryGeom, const TRigidTransform<T, 3>& QueryTM, const T Thickness) const
+	{
+		return OverlapGeomImp(QueryGeom, QueryTM, Thickness);
+	}
+
+	template <typename T>
+	bool THeightField<T>::OverlapGeom(const TCapsule<T>& QueryGeom, const TRigidTransform<T, 3>& QueryTM, const T Thickness) const
+	{
+		return OverlapGeomImp(QueryGeom, QueryTM, Thickness);
+	}
+
+	template <typename T>
+	bool THeightField<T>::OverlapGeom(const TConvex<T, 3>& QueryGeom, const TRigidTransform<T, 3>& QueryTM, const T Thickness) const
+	{
+		return OverlapGeomImp(QueryGeom, QueryTM, Thickness);
+	}
+
+	template <typename T>
+	bool THeightField<T>::OverlapGeom(const TImplicitObjectScaled<TSphere<T, 3>>& QueryGeom, const TRigidTransform<T, 3>& QueryTM, const T Thickness) const
+	{
+		return OverlapGeomImp(QueryGeom, QueryTM, Thickness);
+	}
+
+	template <typename T>
+	bool THeightField<T>::OverlapGeom(const TImplicitObjectScaled<TBox<T, 3>>& QueryGeom, const TRigidTransform<T, 3>& QueryTM, const T Thickness) const
+	{
+		return OverlapGeomImp(QueryGeom, QueryTM, Thickness);
+	}
+
+	template <typename T>
+	bool THeightField<T>::OverlapGeom(const TImplicitObjectScaled<TCapsule<T>>& QueryGeom, const TRigidTransform<T, 3>& QueryTM, const T Thickness) const
+	{
+		return OverlapGeomImp(QueryGeom, QueryTM, Thickness);
+	}
+
+	template <typename T>
+	bool THeightField<T>::OverlapGeom(const TImplicitObjectScaled<TConvex<T, 3>>& QueryGeom, const TRigidTransform<T, 3>& QueryTM, const T Thickness) const
+	{
+		return OverlapGeomImp(QueryGeom, QueryTM, Thickness);
+	}
+
+	template <typename T>
+	template <typename QueryGeomType>
+	bool THeightField<T>::SweepGeomImp(const QueryGeomType& QueryGeom, const TRigidTransform<T, 3>& StartTM, const TVector<T, 3>& Dir, const T Length, T& OutTime, TVector<T, 3>& OutPosition, TVector<T, 3>& OutNormal, int32& OutFaceIndex, const T Thickness) const
 	{
 		bool bHit = false;
-		THeightfieldSweepVisitor<T> SQVisitor(&GeomData, QueryGeom, StartTM, Dir, Thickness);
+		THeightfieldSweepVisitor<QueryGeomType, T> SQVisitor(&GeomData, QueryGeom, StartTM, Dir, Thickness);
 		const TBox<T, 3> QueryBounds = QueryGeom.BoundingBox();
 		const TVector<T, 3> StartPoint = StartTM.TransformPositionNoScale(QueryBounds.Center());
 
@@ -1082,6 +1132,54 @@ namespace Chaos
 		return bHit;
 	}
 
+	template <typename T>
+	bool THeightField<T>::SweepGeom(const TSphere<T, 3>& QueryGeom, const TRigidTransform<T, 3>& StartTM, const TVector<T, 3>& Dir, const T Length, T& OutTime, TVector<T, 3>& OutPosition, TVector<T, 3>& OutNormal, int32& OutFaceIndex, const T Thickness) const
+	{
+		return SweepGeomImp(QueryGeom, StartTM, Dir, Length, OutTime, OutPosition, OutNormal, OutFaceIndex, Thickness);
+	}
+
+	template <typename T>
+	bool THeightField<T>::SweepGeom(const TBox<T, 3>& QueryGeom, const TRigidTransform<T, 3>& StartTM, const TVector<T, 3>& Dir, const T Length, T& OutTime, TVector<T, 3>& OutPosition, TVector<T, 3>& OutNormal, int32& OutFaceIndex, const T Thickness) const
+	{
+		return SweepGeomImp(QueryGeom, StartTM, Dir, Length, OutTime, OutPosition, OutNormal, OutFaceIndex, Thickness);
+	}
+
+	template <typename T>
+	bool THeightField<T>::SweepGeom(const TCapsule<T>& QueryGeom, const TRigidTransform<T, 3>& StartTM, const TVector<T, 3>& Dir, const T Length, T& OutTime, TVector<T, 3>& OutPosition, TVector<T, 3>& OutNormal, int32& OutFaceIndex, const T Thickness) const
+	{
+		return SweepGeomImp(QueryGeom, StartTM, Dir, Length, OutTime, OutPosition, OutNormal, OutFaceIndex, Thickness);
+	}
+
+	template <typename T>
+	bool THeightField<T>::SweepGeom(const TConvex<T, 3>& QueryGeom, const TRigidTransform<T, 3>& StartTM, const TVector<T, 3>& Dir, const T Length, T& OutTime, TVector<T, 3>& OutPosition, TVector<T, 3>& OutNormal, int32& OutFaceIndex, const T Thickness) const
+	{
+		return SweepGeomImp(QueryGeom, StartTM, Dir, Length, OutTime, OutPosition, OutNormal, OutFaceIndex, Thickness);
+	}
+
+	template <typename T>
+	bool THeightField<T>::SweepGeom(const TImplicitObjectScaled<TSphere<T, 3>>& QueryGeom, const TRigidTransform<T, 3>& StartTM, const TVector<T, 3>& Dir, const T Length, T& OutTime, TVector<T, 3>& OutPosition, TVector<T, 3>& OutNormal, int32& OutFaceIndex, const T Thickness) const
+	{
+		return SweepGeomImp(QueryGeom, StartTM, Dir, Length, OutTime, OutPosition, OutNormal, OutFaceIndex, Thickness);
+	}
+
+	template <typename T>
+	bool THeightField<T>::SweepGeom(const TImplicitObjectScaled<TBox<T, 3>>& QueryGeom, const TRigidTransform<T, 3>& StartTM, const TVector<T, 3>& Dir, const T Length, T& OutTime, TVector<T, 3>& OutPosition, TVector<T, 3>& OutNormal, int32& OutFaceIndex, const T Thickness) const
+	{
+		return SweepGeomImp(QueryGeom, StartTM, Dir, Length, OutTime, OutPosition, OutNormal, OutFaceIndex, Thickness);
+	}
+
+	template <typename T>
+	bool THeightField<T>::SweepGeom(const TImplicitObjectScaled<TCapsule<T>>& QueryGeom, const TRigidTransform<T, 3>& StartTM, const TVector<T, 3>& Dir, const T Length, T& OutTime, TVector<T, 3>& OutPosition, TVector<T, 3>& OutNormal, int32& OutFaceIndex, const T Thickness) const
+	{
+		return SweepGeomImp(QueryGeom, StartTM, Dir, Length, OutTime, OutPosition, OutNormal, OutFaceIndex, Thickness);
+	}
+
+	template <typename T>
+	bool THeightField<T>::SweepGeom(const TImplicitObjectScaled<TConvex<T, 3>>& QueryGeom, const TRigidTransform<T, 3>& StartTM, const TVector<T, 3>& Dir, const T Length, T& OutTime, TVector<T, 3>& OutPosition, TVector<T, 3>& OutNormal, int32& OutFaceIndex, const T Thickness) const
+	{
+		return SweepGeomImp(QueryGeom, StartTM, Dir, Length, OutTime, OutPosition, OutNormal, OutFaceIndex, Thickness);
+	}
+
 	template <typename T>
 	int32 THeightField<T>::FindMostOpposingFace(const TVector<T, 3>& Position, const TVector<T, 3>& UnitDir, int32 HintFaceIndex, T SearchDist) const
 	{
diff --git a/Engine/Source/Runtime/Experimental/Chaos/Private/Chaos/ImplicitObject.cpp b/Engine/Source/Runtime/Experimental/Chaos/Private/Chaos/ImplicitObject.cpp
--- a/Engine/Source/Runtime/Experimental/Chaos/Private/Chaos/ImplicitObject.cpp
+++ b/Engine/Source/Runtime/Experimental/Chaos/Private/Chaos/ImplicitObject.cpp
@@ -68,24 +68,6 @@ FVec3 FImplicitObject::Normal(const FVec3& x) const
 	return Normal;
 }
 
-FVec3 FImplicitObject::Support(const FVec3& Direction, const FReal Thickness) const
-{
-	check(false); //not a good implementation, don't use this
-	return FVec3(0);
-#if 0
-	check(bHasBoundingBox);
-	const TBox<FReal, 3> Box = BoundingBox();
-	FVec3 EndPoint = Box.Center();
-	FVec3 StartPoint = EndPoint + Direction.GetSafeNormal() * (Box.Extents().Max() + Thickness);
-	checkSlow(SignedDistance(StartPoint) > 0);
-	checkSlow(SignedDistance(EndPoint) < 0);
-	// @todo(mlentine): The termination condition is slightly different here so we can probably optimize by reimplementing for this function.
-	const auto& Intersection = FindClosestIntersection(StartPoint, EndPoint, Thickness);
-	check(Intersection.Second);
-	return Intersection.First;
-#endif
-}
-
 const TBox<FReal, 3>& FImplicitObject::BoundingBox() const
 {
 	check(false);
@@ -399,4 +381,4 @@ FImplicitObject* FImplicitObject::SerializationFactory(FChaosArchive& Ar, FImpli
 		check(false);
 	}
 	return nullptr;
-}
\ No newline at end of file
+}
diff --git a/Engine/Source/Runtime/Experimental/Chaos/Private/Chaos/PBDCollisionConstraint.cpp b/Engine/Source/Runtime/Experimental/Chaos/Private/Chaos/PBDCollisionConstraint.cpp
--- a/Engine/Source/Runtime/Experimental/Chaos/Private/Chaos/PBDCollisionConstraint.cpp
+++ b/Engine/Source/Runtime/Experimental/Chaos/Private/Chaos/PBDCollisionConstraint.cpp
@@ -1701,8 +1701,13 @@ template<ECollisionUpdateType UpdateType, typename T, int d>
 void UpdateConvexConstraintsUsingCoreShapes(const FImplicitObject & AObj, const TRigidTransform<T, d>& ATM, const FImplicitObject & BObj, const TRigidTransform<T, d>& BTM, const T Thickness, TRigidBodyContactConstraint<T, d>& Constraint)
 {
 	SCOPE_CYCLE_COUNTER(UpdateConvexConstraintsUsingCoreShapes);
-
-	GJKCoreShapeIntersection<T, 3>(AObj, ATM, BObj, BTM, Constraint.Location, Constraint.Phi, Constraint.Normal, Thickness);
+	CastHelper(AObj, [&](const auto& ADowncast)
+	{
+		CastHelper(BObj, [&](const auto& BDowncast)
+		{
+			GJKCoreShapeIntersection<T, 3>(ADowncast, ATM, BDowncast, BTM, Constraint.Location, Constraint.Phi, Constraint.Normal, Thickness);
+		});
+	});
 }
 
 DECLARE_CYCLE_STAT(TEXT("TPBDCollisionConstraint::UpdateUnionLevelsetConstraint"), STAT_UpdateUnionLevelsetConstraint, STATGROUP_ChaosWide);
diff --git a/Engine/Source/Runtime/Experimental/Chaos/Private/Chaos/TriangleMeshImplicitObject.cpp b/Engine/Source/Runtime/Experimental/Chaos/Private/Chaos/TriangleMeshImplicitObject.cpp
--- a/Engine/Source/Runtime/Experimental/Chaos/Private/Chaos/TriangleMeshImplicitObject.cpp
+++ b/Engine/Source/Runtime/Experimental/Chaos/Private/Chaos/TriangleMeshImplicitObject.cpp
@@ -3,6 +3,8 @@
 #include "Chaos/Capsule.h"
 #include "Chaos/GJK.h"
 #include "Chaos/Triangle.h"
+#include "Chaos/Convex.h"
+#include "Chaos/ImplicitObjectScaled.h"
 
 namespace Chaos
 {
@@ -258,7 +260,8 @@ bool TTriangleMeshImplicitObject<T>::Overlap(const TVector<T, 3>& Point, const T
 }
 
 template <typename T>
-bool TTriangleMeshImplicitObject<T>::OverlapGeom(const FImplicitObject& QueryGeom, const TRigidTransform<T, 3>& QueryTM, const T Thickness) const
+template <typename QueryGeomType>
+bool TTriangleMeshImplicitObject<T>::OverlapGeomImp(const QueryGeomType& QueryGeom, const TRigidTransform<T, 3>& QueryTM, const T Thickness) const
 {
 	auto OverlapTriangle = [&](const TVector<T, 3>& A, const TVector<T, 3>& B, const TVector<T, 3>& C) -> bool
 	{
@@ -294,11 +297,66 @@ bool TTriangleMeshImplicitObject<T>::OverlapGeom(const FImplicitObject& QueryGeo
 	return false;
 }
 
+template <typename T>
+bool TTriangleMeshImplicitObject<T>::OverlapGeom(const TSphere<T, 3>& QueryGeom, const TRigidTransform<T, 3>& QueryTM, const T Thickness) const
+{
+	return OverlapGeomImp(QueryGeom, QueryTM, Thickness);
+}
+
+template <typename T>
+bool TTriangleMeshImplicitObject<T>::OverlapGeom(const TBox<T, 3>& QueryGeom, const TRigidTransform<T, 3>& QueryTM, const T Thickness) const
+{
+	return OverlapGeomImp(QueryGeom, QueryTM, Thickness);
+}
+
+template <typename T>
+bool TTriangleMeshImplicitObject<T>::OverlapGeom(const TCapsule<T>& QueryGeom, const TRigidTransform<T, 3>& QueryTM, const T Thickness) const
+{
+	return OverlapGeomImp(QueryGeom, QueryTM, Thickness);
+}
+
+template <typename T>
+bool TTriangleMeshImplicitObject<T>::OverlapGeom(const TConvex<T, 3>& QueryGeom, const TRigidTransform<T, 3>& QueryTM, const T Thickness) const
+{
+	return OverlapGeomImp(QueryGeom, QueryTM, Thickness);
+}
+
+template <typename T>
+bool TTriangleMeshImplicitObject<T>::OverlapGeom(const TImplicitObjectScaled<TSphere<T, 3>>& QueryGeom, const TRigidTransform<T, 3>& QueryTM, const T Thickness) const
+{
+	return OverlapGeomImp(QueryGeom, QueryTM, Thickness);
+}
 
 template <typename T>
+bool TTriangleMeshImplicitObject<T>::OverlapGeom(const TImplicitObjectScaled<TBox<T, 3>>& QueryGeom, const TRigidTransform<T, 3>& QueryTM, const T Thickness) const
+{
+	return OverlapGeomImp(QueryGeom, QueryTM, Thickness);
+}
+
+template <typename T>
+bool TTriangleMeshImplicitObject<T>::OverlapGeom(const TImplicitObjectScaled<TCapsule<T>>& QueryGeom, const TRigidTransform<T, 3>& QueryTM, const T Thickness) const
+{
+	return OverlapGeomImp(QueryGeom, QueryTM, Thickness);
+}
+
+template <typename T>
+bool TTriangleMeshImplicitObject<T>::OverlapGeom(const TImplicitObjectScaled<TConvex<T, 3>>& QueryGeom, const TRigidTransform<T, 3>& QueryTM, const T Thickness) const
+{
+	return OverlapGeomImp(QueryGeom, QueryTM, Thickness);
+}
+
+template <typename T>
+bool TTriangleMeshImplicitObject<T>::OverlapGeom(const TImplicitObjectScaled<TImplicitObjectScaled<TConvex<T, 3>>>& QueryGeom, const TRigidTransform<T, 3>& QueryTM, const T Thickness) const
+{
+	return OverlapGeomImp(QueryGeom, QueryTM, Thickness);
+}
+
+
+
+template <typename QueryGeomType, typename T>
 struct TTriangleMeshSweepVisitor
 {
-	TTriangleMeshSweepVisitor(const TTriangleMeshImplicitObject<T>& InTriMesh, const FImplicitObject& InQueryGeom, const TRigidTransform<T,3>& InStartTM, const TVector<T,3>& InDir, const T InThickness)
+	TTriangleMeshSweepVisitor(const TTriangleMeshImplicitObject<T>& InTriMesh, const QueryGeomType& InQueryGeom, const TRigidTransform<T,3>& InStartTM, const TVector<T,3>& InDir, const T InThickness)
 	: TriMesh(InTriMesh)
 	, StartTM(InStartTM)
 	, QueryGeom(InQueryGeom)
@@ -343,7 +401,7 @@ struct TTriangleMeshSweepVisitor
 
 	const TTriangleMeshImplicitObject<T>& TriMesh;
 	const TRigidTransform<T, 3> StartTM;
-	const FImplicitObject& QueryGeom;
+	const QueryGeomType& QueryGeom;
 	const TVector<T, 3>& Dir;
 	const T Thickness;
 
@@ -354,14 +412,15 @@ struct TTriangleMeshSweepVisitor
 };
 
 template <typename T>
-bool TTriangleMeshImplicitObject<T>::SweepGeom(const FImplicitObject& QueryGeom, const TRigidTransform<T, 3>& StartTM, const TVector<T, 3>& Dir, const T Length, T& OutTime, TVector<T, 3>& OutPosition, TVector<T, 3>& OutNormal, int32& OutFaceIndex, const T Thickness) const
+template <typename QueryGeomType>
+bool TTriangleMeshImplicitObject<T>::SweepGeomImp(const QueryGeomType& QueryGeom, const TRigidTransform<T, 3>& StartTM, const TVector<T, 3>& Dir, const T Length, T& OutTime, TVector<T, 3>& OutPosition, TVector<T, 3>& OutNormal, int32& OutFaceIndex, const T Thickness) const
 {
 	bool bHit = false;
-	TTriangleMeshSweepVisitor<T> SQVisitor(*this, QueryGeom, StartTM, Dir, Thickness);
+	TTriangleMeshSweepVisitor<QueryGeomType, T> SQVisitor(*this, QueryGeom, StartTM, Dir, Thickness);
 	const TBox<T, 3> QueryBounds = QueryGeom.BoundingBox();
 	const TVector<T, 3> StartPoint = StartTM.TransformPositionNoScale(QueryBounds.Center());
 	const TVector<T, 3> Inflation = QueryBounds.Extents() * 0.5 + TVector<T, 3>(Thickness);
-	BVH.template Sweep<TTriangleMeshSweepVisitor<T>, false>(StartPoint, Dir, Length, Inflation, SQVisitor);
+	BVH.template Sweep<TTriangleMeshSweepVisitor<QueryGeomType, T>, false>(StartPoint, Dir, Length, Inflation, SQVisitor);
 
 	if (SQVisitor.OutTime <= Length)
 	{
@@ -374,6 +433,54 @@ bool TTriangleMeshImplicitObject<T>::SweepGeom(const FImplicitObject& QueryGeom,
 	return bHit;
 }
 
+template <typename T>
+bool TTriangleMeshImplicitObject<T>::SweepGeom(const TSphere<T,3>& QueryGeom, const TRigidTransform<T, 3>& StartTM, const TVector<T, 3>& Dir, const T Length, T& OutTime, TVector<T, 3>& OutPosition, TVector<T, 3>& OutNormal, int32& OutFaceIndex, const T Thickness) const
+{
+	return SweepGeomImp(QueryGeom, StartTM, Dir, Length, OutTime, OutPosition, OutNormal, OutFaceIndex, Thickness);
+}
+
+template <typename T>
+bool TTriangleMeshImplicitObject<T>::SweepGeom(const TBox<T, 3>& QueryGeom, const TRigidTransform<T, 3>& StartTM, const TVector<T, 3>& Dir, const T Length, T& OutTime, TVector<T, 3>& OutPosition, TVector<T, 3>& OutNormal, int32& OutFaceIndex, const T Thickness) const
+{
+	return SweepGeomImp(QueryGeom, StartTM, Dir, Length, OutTime, OutPosition, OutNormal, OutFaceIndex, Thickness);
+}
+
+template <typename T>
+bool TTriangleMeshImplicitObject<T>::SweepGeom(const TCapsule<T>& QueryGeom, const TRigidTransform<T, 3>& StartTM, const TVector<T, 3>& Dir, const T Length, T& OutTime, TVector<T, 3>& OutPosition, TVector<T, 3>& OutNormal, int32& OutFaceIndex, const T Thickness) const
+{
+	return SweepGeomImp(QueryGeom, StartTM, Dir, Length, OutTime, OutPosition, OutNormal, OutFaceIndex, Thickness);
+}
+
+template <typename T>
+bool TTriangleMeshImplicitObject<T>::SweepGeom(const TConvex<T, 3>& QueryGeom, const TRigidTransform<T, 3>& StartTM, const TVector<T, 3>& Dir, const T Length, T& OutTime, TVector<T, 3>& OutPosition, TVector<T, 3>& OutNormal, int32& OutFaceIndex, const T Thickness) const
+{
+	return SweepGeomImp(QueryGeom, StartTM, Dir, Length, OutTime, OutPosition, OutNormal, OutFaceIndex, Thickness);
+}
+
+template <typename T>
+bool TTriangleMeshImplicitObject<T>::SweepGeom(const TImplicitObjectScaled<TSphere<T, 3>>& QueryGeom, const TRigidTransform<T, 3>& StartTM, const TVector<T, 3>& Dir, const T Length, T& OutTime, TVector<T, 3>& OutPosition, TVector<T, 3>& OutNormal, int32& OutFaceIndex, const T Thickness) const
+{
+	return SweepGeomImp(QueryGeom, StartTM, Dir, Length, OutTime, OutPosition, OutNormal, OutFaceIndex, Thickness);
+}
+
+template <typename T>
+bool TTriangleMeshImplicitObject<T>::SweepGeom(const TImplicitObjectScaled<TBox<T, 3>>& QueryGeom, const TRigidTransform<T, 3>& StartTM, const TVector<T, 3>& Dir, const T Length, T& OutTime, TVector<T, 3>& OutPosition, TVector<T, 3>& OutNormal, int32& OutFaceIndex, const T Thickness) const
+{
+	return SweepGeomImp(QueryGeom, StartTM, Dir, Length, OutTime, OutPosition, OutNormal, OutFaceIndex, Thickness);
+}
+
+template <typename T>
+bool TTriangleMeshImplicitObject<T>::SweepGeom(const TImplicitObjectScaled<TCapsule<T>>& QueryGeom, const TRigidTransform<T, 3>& StartTM, const TVector<T, 3>& Dir, const T Length, T& OutTime, TVector<T, 3>& OutPosition, TVector<T, 3>& OutNormal, int32& OutFaceIndex, const T Thickness) const
+{
+	return SweepGeomImp(QueryGeom, StartTM, Dir, Length, OutTime, OutPosition, OutNormal, OutFaceIndex, Thickness);
+}
+
+template <typename T>
+bool TTriangleMeshImplicitObject<T>::SweepGeom(const TImplicitObjectScaled<TConvex<T, 3>>& QueryGeom, const TRigidTransform<T, 3>& StartTM, const TVector<T, 3>& Dir, const T Length, T& OutTime, TVector<T, 3>& OutPosition, TVector<T, 3>& OutNormal, int32& OutFaceIndex, const T Thickness) const
+{
+	return SweepGeomImp(QueryGeom, StartTM, Dir, Length, OutTime, OutPosition, OutNormal, OutFaceIndex, Thickness);
+}
+
 template <typename T>
 int32 TTriangleMeshImplicitObject<T>::FindMostOpposingFace(const TVector<T, 3>& Position, const TVector<T, 3>& UnitDir, int32 HintFaceIndex, T SearchDist) const
 {
diff --git a/Engine/Source/Runtime/Experimental/Chaos/Public/Chaos/AABB.h b/Engine/Source/Runtime/Experimental/Chaos/Public/Chaos/AABB.h
--- a/Engine/Source/Runtime/Experimental/Chaos/Public/Chaos/AABB.h
+++ b/Engine/Source/Runtime/Experimental/Chaos/Public/Chaos/AABB.h
@@ -387,6 +387,17 @@ namespace Chaos
 			return ChosenPt;
 		}
 
+		FORCEINLINE TVector<T, d> Support2(const TVector<T, d>& Direction) const
+		{
+			TVector<T, d> ChosenPt;
+			for (int Axis = 0; Axis < d; ++Axis)
+			{
+				ChosenPt[Axis] = Direction[Axis] < 0 ? MMin[Axis] : MMax[Axis];
+			}
+
+			return ChosenPt;
+		}
+
 		FORCEINLINE void GrowToInclude(const TVector<T, d>& V)
 		{
 			MMin = TVector<T, d>(FGenericPlatformMath::Min(MMin[0], V[0]), FGenericPlatformMath::Min(MMin[1], V[1]), FGenericPlatformMath::Min(MMin[2], V[2]));
diff --git a/Engine/Source/Runtime/Experimental/Chaos/Public/Chaos/Box.h b/Engine/Source/Runtime/Experimental/Chaos/Public/Chaos/Box.h
--- a/Engine/Source/Runtime/Experimental/Chaos/Public/Chaos/Box.h
+++ b/Engine/Source/Runtime/Experimental/Chaos/Public/Chaos/Box.h
@@ -120,7 +120,7 @@ namespace Chaos
 			return AABB.Contains(Point, Tolerance);
 		}
 
-		FORCEINLINE static EImplicitObjectType StaticType() { return ImplicitObjectType::Box; }
+		FORCEINLINE static constexpr EImplicitObjectType StaticType() { return ImplicitObjectType::Box; }
 
 		const TBox<T, d>& BoundingBox() const { return *this; }
 
@@ -156,11 +156,18 @@ namespace Chaos
 			return AABB.FindGeometryOpposingNormal(DenormDir, FaceIndex, OriginalNormal);
 		}
 
-		virtual TVector<T, d> Support(const TVector<T, d>& Direction, const T Thickness) const override
+		FORCEINLINE T GetMargin() const { return 0; }
+
+		FORCEINLINE TVector<T, d> Support(const TVector<T, d>& Direction, const T Thickness) const
 		{
 			return AABB.Support(Direction, Thickness);
 		}
 
+		FORCEINLINE TVector<T, d> Support2(const TVector<T, d>& Direction) const
+		{
+			return AABB.Support2(Direction);
+		}
+
 		FORCEINLINE void GrowToInclude(const TVector<T, d>& V)
 		{
 			AABB.GrowToInclude(V);
diff --git a/Engine/Source/Runtime/Experimental/Chaos/Public/Chaos/Capsule.h b/Engine/Source/Runtime/Experimental/Chaos/Public/Chaos/Capsule.h
--- a/Engine/Source/Runtime/Experimental/Chaos/Public/Chaos/Capsule.h
+++ b/Engine/Source/Runtime/Experimental/Chaos/Public/Chaos/Capsule.h
@@ -74,7 +74,7 @@ namespace Chaos
 
 		~TCapsule() {}
 
-		static EImplicitObjectType StaticType() { return ImplicitObjectType::Capsule; }
+		static constexpr EImplicitObjectType StaticType() { return ImplicitObjectType::Capsule; }
 
 		static TCapsule<T> NewFromOriginAndAxis(const TVector<T, 3>& Origin, const TVector<T, 3>& Axis, const T Height, const T Radius)
 		{
@@ -287,14 +287,14 @@ namespace Chaos
 			return RaycastFast(MRadius, GetHeight(), GetAxis(), GetX1(), GetX2(), StartPoint, Dir, Length, Thickness, OutTime, OutPosition, OutNormal, OutFaceIndex);
 		}
 
-		TVector<T,3> Support(const TVector<T, 3>& Direction, const T Thickness) const override
+		FORCEINLINE TVector<T,3> Support(const TVector<T, 3>& Direction, const T Thickness) const
 		{
 			return MSegment.Support(Direction, MRadius + Thickness);
 		}
 
-		virtual TVector<T, 3> Support2(const TVector<T, 3>& Direction) const override { return MSegment.Support2(Direction); }
+		FORCEINLINE TVector<T, 3> Support2(const TVector<T, 3>& Direction) const { return MSegment.Support2(Direction); }
 
-		virtual T GetMargin() const override
+		FORCEINLINE T GetMargin() const
 		{
 			return MRadius;
 		}
diff --git a/Engine/Source/Runtime/Experimental/Chaos/Public/Chaos/Convex.h b/Engine/Source/Runtime/Experimental/Chaos/Public/Chaos/Convex.h
--- a/Engine/Source/Runtime/Experimental/Chaos/Public/Chaos/Convex.h
+++ b/Engine/Source/Runtime/Experimental/Chaos/Public/Chaos/Convex.h
@@ -48,7 +48,7 @@ namespace Chaos
 			TConvexBuilder<T>::Build(InParticles, Planes, SurfaceParticles, LocalBoundingBox);
 		}
 
-		static EImplicitObjectType StaticType()
+		static constexpr EImplicitObjectType StaticType()
 		{
 			return ImplicitObjectType::Convex;
 		}
@@ -157,7 +157,11 @@ namespace Chaos
 			return TVector<float, 3>(0.f, 0.f, 1.f);
 		}
 
-		virtual TVector<T, d> Support(const TVector<T, d>& Direction, const T Thickness) const override
+		FORCEINLINE T GetMargin() const { return 0; }
+
+		FORCEINLINE TVector<T, d> Support2(const TVector<T, d>& Direction) const { return Support(Direction, 0); }
+
+		TVector<T, d> Support(const TVector<T, d>& Direction, const T Thickness) const
 		{
 			T MaxDot = TNumericLimits<T>::Lowest();
 			int32 MaxVIdx = 0;
diff --git a/Engine/Source/Runtime/Experimental/Chaos/Public/Chaos/Cylinder.h b/Engine/Source/Runtime/Experimental/Chaos/Public/Chaos/Cylinder.h
--- a/Engine/Source/Runtime/Experimental/Chaos/Public/Chaos/Cylinder.h
+++ b/Engine/Source/Runtime/Experimental/Chaos/Public/Chaos/Cylinder.h
@@ -50,7 +50,7 @@ namespace Chaos
 		}
 		~TCylinder() {}
 
-		static EImplicitObjectType StaticType() { return ImplicitObjectType::Cylinder; }
+		static constexpr EImplicitObjectType StaticType() { return ImplicitObjectType::Cylinder; }
 
 		/**
 		 * Returns sample points centered about the origin.
diff --git a/Engine/Source/Runtime/Experimental/Chaos/Public/Chaos/GJK.h b/Engine/Source/Runtime/Experimental/Chaos/Public/Chaos/GJK.h
--- a/Engine/Source/Runtime/Experimental/Chaos/Public/Chaos/GJK.h
+++ b/Engine/Source/Runtime/Experimental/Chaos/Public/Chaos/GJK.h
@@ -578,11 +578,13 @@ namespace Chaos
 
 	template <typename T, typename TGeometryA, typename TGeometryB>
 	bool GJKRaycast2(const TGeometryA& A, const TGeometryB& B, const TRigidTransform<T, 3>& StartTM, const TVector<T, 3>& RayDir, const T RayLength,
-		T& OutTime, TVector<T, 3>& OutPosition, TVector<T, 3>& OutNormal, const T ThicknessA = 0, const TVector<T, 3>& InitialDir = TVector<T, 3>(1, 0, 0), const T ThicknessB = 0)
+		T& OutTime, TVector<T, 3>& OutPosition, TVector<T, 3>& OutNormal, const T GivenThicknessA = 0, bool bComputeMTD = false, const TVector<T, 3>& InitialDir = TVector<T, 3>(1, 0, 0), const T GivenThicknessB = 0)
 	{
 		ensure(FMath::IsNearlyEqual(RayDir.SizeSquared(), 1, KINDA_SMALL_NUMBER));
 		ensure(RayLength > 0);
-		check(A.IsConvex() && B.IsConvex());
+		const T ThicknessA = A.GetMargin();
+		const T ThicknessB = B.GetMargin();
+
 		const TVector<T, 3> StartPoint = StartTM.GetLocation();
 
 		TVector<T, 3> Simplex[4] = { TVector<T,3>(0), TVector<T,3>(0), TVector<T,3>(0), TVector<T,3>(0) };
@@ -686,7 +688,6 @@ namespace Chaos
 		if (Lambda > 0)
 		{
 			OutNormal = Normal;
-			TVector<T, 3> ClosestA(0);
 			TVector<T, 3> ClosestB(0);
 
 			for (int i = 0; i < SimplexIDs.NumVerts; ++i)
@@ -697,6 +698,35 @@ namespace Chaos
 
 			OutPosition = StartPoint + RayDir * Lambda + ClosestLocal;
 		}
+		else if (bComputeMTD)
+		{
+			if (InGJKPreDist2)
+			{
+				OutNormal = Normal;
+				TVector<T, 3> ClosestA(0);
+				TVector<T, 3> ClosestB(0);
+
+				for (int i = 0; i < SimplexIDs.NumVerts; ++i)
+				{
+					ClosestA += As[i] * Barycentric[i];
+					ClosestB += Bs[i] * Barycentric[i];
+				}
+
+				const T InGJKPreDist = FMath::Sqrt(InGJKPreDist2);
+				OutNormal = (ClosestA - ClosestB).GetUnsafeNormal();	//question: should we just use InGJKPreDist2?
+				const TVector<T, 3> ClosestLocal = ClosestB - OutNormal * ThicknessB;
+
+				OutPosition = StartPoint + ClosestLocal;
+				OutTime = InGJKPreDist - ThicknessA - ThicknessB;
+			}
+			else
+			{
+				ensure(false);	//todo: use EPA
+				OutTime = 0;
+				OutPosition = TVector<T, 3>(0);
+				OutNormal = { 0,0,1 };
+			}
+		}
 
 		return true;
 	}
diff --git a/Engine/Source/Runtime/Experimental/Chaos/Public/Chaos/GeometryQueries.h b/Engine/Source/Runtime/Experimental/Chaos/Public/Chaos/GeometryQueries.h
--- a/Engine/Source/Runtime/Experimental/Chaos/Public/Chaos/GeometryQueries.h
+++ b/Engine/Source/Runtime/Experimental/Chaos/Public/Chaos/GeometryQueries.h
@@ -9,6 +9,8 @@
 #include "Chaos/Convex.h"
 #include "Chaos/Capsule.h"
 #include "ImplicitObjectScaled.h"
+#include "Chaos/Box.h"
+#include "Chaos/Sphere.h"
 
 #include "ChaosArchive.h"
 #include <algorithm>
@@ -17,40 +19,52 @@
 
 namespace Chaos
 {
-	template <typename T, int d>
-	bool OverlapQuery(const FImplicitObject& A, const TRigidTransform<T,d>& ATM, const FImplicitObject& B, const TRigidTransform<T,d>& BTM, const T Thickness = 0, TVector<T,d>* OutMTD=nullptr)
+	template <typename Lambda>
+	FORCEINLINE_DEBUGGABLE auto CastHelper(const FImplicitObject& Geom, const Lambda& Func)
+	{
+		const EImplicitObjectType Type = Geom.GetType(true);
+		switch (Type)
+		{
+		case ImplicitObjectType::Sphere: return Func(Geom.template GetObjectChecked<TSphere<FReal, 3>>());
+		case ImplicitObjectType::Box: return Func(Geom.template GetObjectChecked<TBox<FReal, 3>>());
+		case ImplicitObjectType::Capsule: return Func(Geom.template GetObjectChecked<TCapsule<FReal>>());
+		case ImplicitObjectType::Convex: return Func(Geom.template GetObjectChecked<TConvex<FReal, 3>>());
+		case ImplicitObjectType::IsScaled | ImplicitObjectType::Sphere: return Func(Geom.template GetObjectChecked<TImplicitObjectScaled<TSphere<FReal, 3>>>());
+		case ImplicitObjectType::IsScaled | ImplicitObjectType::Box: return Func(Geom.template GetObjectChecked< TImplicitObjectScaled<TBox<FReal, 3>>>());
+		case ImplicitObjectType::IsScaled | ImplicitObjectType::Capsule: return Func(Geom.template GetObjectChecked< TImplicitObjectScaled<TCapsule<FReal>>>());
+		case ImplicitObjectType::IsScaled | ImplicitObjectType::Convex: return Func(Geom.template GetObjectChecked< TImplicitObjectScaled<TConvex<FReal, 3>>>());
+		default: check(false);
+		}
+		return Func(Geom.template GetObjectChecked<TSphere<FReal, 3>>());	//needed for return type
+	}
+
+	template <typename QueryGeometry>
+	bool OverlapQuery(const FImplicitObject& A, const FRigidTransform3& ATM, const QueryGeometry& B, const FRigidTransform3& BTM, const FReal Thickness = 0, FVec3* OutMTD=nullptr)
 	{
 		const EImplicitObjectType AType = A.GetType(true);
-		const EImplicitObjectType BType = B.GetType(true);
+		constexpr EImplicitObjectType BType = QueryGeometry::StaticType();
 		
 		if (AType == ImplicitObjectType::Transformed)
 		{
-			const TImplicitObjectTransformed<T, d>& TransformedA = static_cast<const TImplicitObjectTransformed<T, d>&>(A);
-			const TRigidTransform<T, d> NewATM = TransformedA.GetTransform() * ATM;
+			const TImplicitObjectTransformed<FReal, 3>& TransformedA = static_cast<const TImplicitObjectTransformed<FReal, 3>&>(A);
+			const FRigidTransform3 NewATM = TransformedA.GetTransform() * ATM;
 			return OverlapQuery(*TransformedA.GetTransformedObject(), NewATM, B, BTM, Thickness, OutMTD);
 		}
 
-		if (BType == ImplicitObjectType::Transformed)
-		{
-			const TImplicitObjectTransformed<T, d>& TransformedB = static_cast<const TImplicitObjectTransformed<T, d>&>(B);
-			const TRigidTransform<T, d> NewBTM = TransformedB.GetTransform() * BTM;
-			return OverlapQuery(A, ATM, *TransformedB.GetTransformedObject(), NewBTM, Thickness, OutMTD);
-		}
+		const FRigidTransform3 BToATM = BTM.GetRelativeTransform(ATM);
 
-		check(B.IsConvex());	//Query object must be convex
-		const TRigidTransform<T, d> BToATM = BTM.GetRelativeTransform(ATM);
-
-		if(BType == ImplicitObjectType::Sphere)
+		if(false && BType == ImplicitObjectType::Sphere)	//todo: fix
 		{
-			const TSphere<T, d>& BSphere = static_cast<const TSphere<T, d>&>(B);
+			ensure(false);	//todo: override helper
+			/*const TSphere<T, d>& BSphere = static_cast<const TSphere<T, d>&>(B);
 			const TVector<T, d> PtInA = BToATM.TransformPositionNoScale(BSphere.GetCenter());
-			return A.Overlap(PtInA, Thickness + BSphere.GetRadius());
+			return A.Overlap(PtInA, Thickness + BSphere.GetRadius());*/
 		}
 		//todo: A is a sphere
 		else if (A.IsConvex())
 		{
-			const TVector<T, d> Offset = ATM.GetLocation() - BTM.GetLocation();
-			return GJKIntersection(A, B, BToATM, Thickness, Offset.SizeSquared() < 1e-4 ? TVector<T, d>(1, 0, 0) : Offset);
+			const FVec3 Offset = ATM.GetLocation() - BTM.GetLocation();
+			return CastHelper(A, [&](const auto& AConcrete) { return GJKIntersection<FReal>(AConcrete, B, BToATM, Thickness, Offset.SizeSquared() < 1e-4 ? FVec3(1, 0, 0) : Offset); });
 		}
 		else
 		{
@@ -58,75 +72,68 @@ namespace Chaos
 			{
 			case ImplicitObjectType::HeightField:
 			{
-				const THeightField<T>& AHeightField = static_cast<const THeightField<T>&>(A);
+				const THeightField<FReal>& AHeightField = static_cast<const THeightField<FReal>&>(A);
 				return AHeightField.OverlapGeom(B, BToATM, Thickness);
 			}
 			case ImplicitObjectType::TriangleMesh:
 			{
-				const TTriangleMeshImplicitObject<T>& ATriangleMesh = static_cast<const TTriangleMeshImplicitObject<T>&>(A);
+				const TTriangleMeshImplicitObject<FReal>& ATriangleMesh = static_cast<const TTriangleMeshImplicitObject<FReal>&>(A);
 				return ATriangleMesh.OverlapGeom(B, BToATM, Thickness);
 			}
 			default:
 				if (IsScaled(AType))
 				{
-					const auto& AScaled = TImplicitObjectScaled<TTriangleMeshImplicitObject<T>>::AsScaledChecked(A);
+					const auto& AScaled = TImplicitObjectScaled<TTriangleMeshImplicitObject<FReal>>::AsScaledChecked(A);
 					return AScaled.LowLevelOverlapGeom(B, BToATM, Thickness);
 				}
 				else
 				{
 					check(false);	//unsupported query type
 				}
-
 			}
 		}
 
 		return false;
 	}
 
-	template <typename T, int d>
-	bool SweepQuery(const FImplicitObject& A, const TRigidTransform<T,d>& ATM, const FImplicitObject& B, const TRigidTransform<T, d>& BTM, const TVector<T,d>& Dir, const T Length, T& OutTime, TVector<T,d>& OutPosition, TVector<T,d>& OutNormal, int32& OutFaceIndex, const T Thickness, const bool bComputeMTD)
+	template <typename SweptGeometry>
+	bool SweepQuery(const FImplicitObject& A, const FRigidTransform3& ATM, const SweptGeometry& B, const FRigidTransform3& BTM, const FVec3& Dir, const FReal Length, FReal& OutTime, FVec3& OutPosition, FVec3& OutNormal, int32& OutFaceIndex, const FReal Thickness, const bool bComputeMTD)
 	{
 		const EImplicitObjectType AType = A.GetType(true);
-		const EImplicitObjectType BType = B.GetType(true);
+		constexpr EImplicitObjectType BType = SweptGeometry::StaticType();
 
 		bool bResult = false;
 		
 		if (AType == ImplicitObjectType::Transformed)
 		{
-			const TImplicitObjectTransformed<T, d>& TransformedA = static_cast<const TImplicitObjectTransformed<T, d>&>(A);
-			const TRigidTransform<T, d> NewATM = TransformedA.GetTransform() * ATM;
+			const TImplicitObjectTransformed<FReal, 3>& TransformedA = static_cast<const TImplicitObjectTransformed<FReal, 3>&>(A);
+			const FRigidTransform3 NewATM = TransformedA.GetTransform() * ATM;
 			return SweepQuery(*TransformedA.GetTransformedObject(), NewATM, B, BTM, Dir, Length, OutTime, OutPosition, OutNormal, OutFaceIndex, Thickness, bComputeMTD);
 		}
 
-		if (BType == ImplicitObjectType::Transformed)
-		{
-			const TImplicitObjectTransformed<T, d>& TransformedB = static_cast<const TImplicitObjectTransformed<T, d>&>(B);
-			const TRigidTransform<T, d> NewBTM = TransformedB.GetTransform() * BTM;
-			return SweepQuery(A, ATM, *TransformedB.GetTransformedObject(), NewBTM, Dir, Length, OutTime, OutPosition, OutNormal, OutFaceIndex, Thickness, bComputeMTD);
-		}
-
-		check(B.IsConvex());	//Object being swept must be convex
 		OutFaceIndex = INDEX_NONE;
 		
-		TVector<T, d> LocalPosition(-TNumericLimits<float>::Max()); // Make it obvious when things go wrong
-		TVector<T, d> LocalNormal(0);
+		FVec3 LocalPosition(-TNumericLimits<float>::Max()); // Make it obvious when things go wrong
+		FVec3 LocalNormal(0);
 
-		const TRigidTransform<T, d> BToATM = BTM.GetRelativeTransform(ATM);
-		const TVector<T, d> LocalDir = ATM.InverseTransformVectorNoScale(Dir);
+		const FRigidTransform3 BToATM = BTM.GetRelativeTransform(ATM);
+		const FVec3 LocalDir = ATM.InverseTransformVectorNoScale(Dir);
 
-		bool bSweepAsRaycast = BType == ImplicitObjectType::Sphere;
+		bool bSweepAsRaycast = false;//todo: fix BType == ImplicitObjectType::Sphere;
 		if (bSweepAsRaycast && IsScaled(AType))
 		{
-			const auto& Scaled = TImplicitObjectScaledGeneric<T, d>::AsScaledChecked(A);
-			const TVector<T, d>& Scale = Scaled.GetScale();
+			const auto& Scaled = TImplicitObjectScaledGeneric<FReal, 3>::AsScaledChecked(A);
+			const FVec3& Scale = Scaled.GetScale();
 			bSweepAsRaycast = FMath::IsNearlyEqual(Scale[0], Scale[1]) && FMath::IsNearlyEqual(Scale[0], Scale[2]);
 		}
 
 		if (bSweepAsRaycast)
 		{
+			/*
 			const TSphere<T, d>& BSphere = static_cast<const TSphere<T, d>&>(B);
 			const TVector<T, d> Start = BToATM.TransformPositionNoScale(BSphere.GetCenter());
-			bResult = A.Raycast(Start, LocalDir, Length, Thickness + BSphere.GetRadius(), OutTime, LocalPosition, LocalNormal, OutFaceIndex);
+			bResult = A.Raycast(Start, LocalDir, Length, Thickness + BSphere.GetRadius(), OutTime, LocalPosition, LocalNormal, OutFaceIndex);*/
+			ensure(false);	//todo: override helper
 		}
 		//todo: handle case where A is a sphere
 		else if (A.IsConvex())
@@ -134,7 +141,7 @@ namespace Chaos
 			auto IsValidConvex = [](const FImplicitObject& InObject) -> bool
 			{
 				//todo: move this out of here
-				if (const auto Convex = TImplicitObjectScaled<TConvex<T,3>>::AsScaled(InObject))
+				if (const auto Convex = TImplicitObjectScaled<TConvex<FReal,3>>::AsScaled(InObject))
 				{
 					return Convex->GetUnscaledObject()->GetSurfaceParticles().Size() > 0;
 				}				
@@ -149,11 +156,8 @@ namespace Chaos
 				return false;
 			}
 
-			const TVector<T, d> Offset = ATM.GetLocation() - BTM.GetLocation();
-
-			bResult = GJKRaycast2<T>(A, B, BToATM, LocalDir, Length, OutTime, LocalPosition, LocalNormal, Thickness + A.GetMargin(), Offset, Thickness + B.GetMargin());
-			//bResult = GJKRaycast<T>(A, B, BToATM, LocalDir, Length, OutTime, LocalPosition, LocalNormal, Thickness, Offset, Thickness);
-
+			const FVec3 Offset = ATM.GetLocation() - BTM.GetLocation();
+			bResult = CastHelper(A, [&](const auto& ADowncast) { return GJKRaycast2(ADowncast, B, BToATM, LocalDir, Length, OutTime, LocalPosition, LocalNormal, Thickness, bComputeMTD, Offset, Thickness); });
 			if (AType == ImplicitObjectType::Convex)
 			{
 				//todo: find face index
@@ -163,42 +167,28 @@ namespace Chaos
 				ensure(false);
 				//todo: find face index if convex hull
 			}
-
-			// Compute MTD in the case of an initial overlap
-			if (bResult && bComputeMTD && OutTime == 0.f)
-			{
-				TVector<T, d> LocalPositionOnB;
-				if (BType == ImplicitObjectType::Capsule)
-				{
-					const TCapsule<T>& BCapsule = static_cast<const TCapsule<T>&>(B);
-					GJKPenetrationTemp<T>(A, BCapsule, BToATM, LocalPosition, LocalPositionOnB, LocalNormal, OutTime, Thickness, Offset);
-				}
-				else
-				{
-					GJKPenetrationTemp<T>(A, B, BToATM, LocalPosition, LocalPositionOnB, LocalNormal, OutTime, Thickness, Offset);
-				}
-			}
 		}
 		else
 		{
+			//todo: pass bComputeMTD into these functions
 			switch (AType)
 			{
 			case ImplicitObjectType::HeightField:
 			{
-				const THeightField<T>& AHeightField = static_cast<const THeightField<T>&>(A);
+				const auto& AHeightField = static_cast<const THeightField<FReal>&>(A);
 				bResult = AHeightField.SweepGeom(B, BToATM, LocalDir, Length, OutTime, LocalPosition, LocalNormal, OutFaceIndex, Thickness);
 				break;
 			}
 			case ImplicitObjectType::TriangleMesh:
 			{
-				const TTriangleMeshImplicitObject<T>& ATriangleMesh = static_cast<const TTriangleMeshImplicitObject<T>&>(A);
+				const auto& ATriangleMesh = static_cast<const TTriangleMeshImplicitObject<FReal>&>(A);
 				bResult = ATriangleMesh.SweepGeom(B, BToATM, LocalDir, Length, OutTime, LocalPosition, LocalNormal, OutFaceIndex, Thickness);
 				break;
 			}
 			default:
 				if (IsScaled(AType))
 				{
-					const auto& AScaled = TImplicitObjectScaled<TTriangleMeshImplicitObject<T>>::AsScaledChecked(A);
+					const auto& AScaled = TImplicitObjectScaled<TTriangleMeshImplicitObject<FReal>>::AsScaledChecked(A);
 					bResult = AScaled.LowLevelSweepGeom(B, BToATM, LocalDir, Length, OutTime, LocalPosition, LocalNormal, OutFaceIndex, Thickness);
 					break;
 				}
@@ -213,7 +203,7 @@ namespace Chaos
 			if (bResult && bComputeMTD && OutTime == 0.f)
 			{
 				ensure(false); // We don't support MTD for non-convex types yet!
-				OutNormal = TVector<float, 3>(0.f, 0.f, 1.f);
+				OutNormal = FVec3(0, 0, 1);
 			}
 		}
 
diff --git a/Engine/Source/Runtime/Experimental/Chaos/Public/Chaos/HeightField.h b/Engine/Source/Runtime/Experimental/Chaos/Public/Chaos/HeightField.h
--- a/Engine/Source/Runtime/Experimental/Chaos/Public/Chaos/HeightField.h
+++ b/Engine/Source/Runtime/Experimental/Chaos/Public/Chaos/HeightField.h
@@ -46,8 +46,25 @@ namespace Chaos
 
 		virtual bool Raycast(const TVector<T, 3>& StartPoint, const TVector<T, 3>& Dir, const T Length, const T Thickness, T& OutTime, TVector<T,3>& OutPosition, TVector<T,3>& OutNormal, int32& OutFaceIndex) const override;
 		virtual bool Overlap(const TVector<T, 3>& Point, const T Thickness) const override;
-		bool OverlapGeom(const FImplicitObject& QueryGeom, const TRigidTransform<T, 3>& QueryTM, const T Thickness) const;
-		bool SweepGeom(const FImplicitObject& QueryGeom, const TRigidTransform<T, 3>& StartTM, const TVector<T, 3>& Dir, const T Length, T& OutTime, TVector<T, 3>& OutPosition, TVector<T, 3>& OutNormal, int32& OutFaceIndex, const T Thickness = 0) const;
+		
+		bool OverlapGeom(const TSphere<T, 3>& QueryGeom, const TRigidTransform<T, 3>& QueryTM, const T Thickness) const;
+		bool OverlapGeom(const TBox<T, 3>& QueryGeom, const TRigidTransform<T, 3>& QueryTM, const T Thickness) const;
+		bool OverlapGeom(const TCapsule<T>& QueryGeom, const TRigidTransform<T, 3>& QueryTM, const T Thickness) const;
+		bool OverlapGeom(const TConvex<T, 3>& QueryGeom, const TRigidTransform<T, 3>& QueryTM, const T Thickness) const;
+		bool OverlapGeom(const TImplicitObjectScaled<TSphere<T, 3>>& QueryGeom, const TRigidTransform<T, 3>& QueryTM, const T Thickness) const;
+		bool OverlapGeom(const TImplicitObjectScaled<TBox<T, 3>>& QueryGeom, const TRigidTransform<T, 3>& QueryTM, const T Thickness) const;
+		bool OverlapGeom(const TImplicitObjectScaled<TCapsule<T>>& QueryGeom, const TRigidTransform<T, 3>& QueryTM, const T Thickness) const;
+		bool OverlapGeom(const TImplicitObjectScaled<TConvex<T, 3>>& QueryGeom, const TRigidTransform<T, 3>& QueryTM, const T Thickness) const;
+
+		bool SweepGeom(const TSphere<T, 3>& QueryGeom, const TRigidTransform<T, 3>& StartTM, const TVector<T, 3>& Dir, const T Length, T& OutTime, TVector<T, 3>& OutPosition, TVector<T, 3>& OutNormal, int32& OutFaceIndex, const T Thickness = 0) const;
+		bool SweepGeom(const TBox<T, 3>& QueryGeom, const TRigidTransform<T, 3>& StartTM, const TVector<T, 3>& Dir, const T Length, T& OutTime, TVector<T, 3>& OutPosition, TVector<T, 3>& OutNormal, int32& OutFaceIndex, const T Thickness = 0) const;
+		bool SweepGeom(const TCapsule<T>& QueryGeom, const TRigidTransform<T, 3>& StartTM, const TVector<T, 3>& Dir, const T Length, T& OutTime, TVector<T, 3>& OutPosition, TVector<T, 3>& OutNormal, int32& OutFaceIndex, const T Thickness = 0) const;
+		bool SweepGeom(const TConvex<T, 3>& QueryGeom, const TRigidTransform<T, 3>& StartTM, const TVector<T, 3>& Dir, const T Length, T& OutTime, TVector<T, 3>& OutPosition, TVector<T, 3>& OutNormal, int32& OutFaceIndex, const T Thickness = 0) const;
+		bool SweepGeom(const TImplicitObjectScaled<TSphere<T, 3>>& QueryGeom, const TRigidTransform<T, 3>& StartTM, const TVector<T, 3>& Dir, const T Length, T& OutTime, TVector<T, 3>& OutPosition, TVector<T, 3>& OutNormal, int32& OutFaceIndex, const T Thickness = 0) const;
+		bool SweepGeom(const TImplicitObjectScaled<TBox<T, 3>>& QueryGeom, const TRigidTransform<T, 3>& StartTM, const TVector<T, 3>& Dir, const T Length, T& OutTime, TVector<T, 3>& OutPosition, TVector<T, 3>& OutNormal, int32& OutFaceIndex, const T Thickness = 0) const;
+		bool SweepGeom(const TImplicitObjectScaled<TCapsule<T>>& QueryGeom, const TRigidTransform<T, 3>& StartTM, const TVector<T, 3>& Dir, const T Length, T& OutTime, TVector<T, 3>& OutPosition, TVector<T, 3>& OutNormal, int32& OutFaceIndex, const T Thickness = 0) const;
+		bool SweepGeom(const TImplicitObjectScaled<TConvex<T, 3>>& QueryGeom, const TRigidTransform<T, 3>& StartTM, const TVector<T, 3>& Dir, const T Length, T& OutTime, TVector<T, 3>& OutPosition, TVector<T, 3>& OutNormal, int32& OutFaceIndex, const T Thickness = 0) const;
+
 		virtual int32 FindMostOpposingFace(const TVector<T, 3>& Position, const TVector<T, 3>& UnitDir, int32 HintFaceIndex, T SearchDist) const override;
 		virtual TVector<T, 3> FindGeometryOpposingNormal(const TVector<T, 3>& DenormDir, int32 FaceIndex, const TVector<T, 3>& OriginalNormal) const override;
 
@@ -70,7 +87,7 @@ namespace Chaos
 			return FCrc::MemCrc32(Bytes.GetData(), Bytes.GetAllocatedSize());
 		}
 
-		static EImplicitObjectType StaticType()
+		static constexpr EImplicitObjectType StaticType()
 		{
 			return ImplicitObjectType::HeightField;
 		}
@@ -396,5 +413,11 @@ namespace Chaos
 		// Needed for serialization
 		THeightField() : FImplicitObject(EImplicitObject::HasBoundingBox, ImplicitObjectType::HeightField) {}
 		friend FImplicitObject;
+
+		template <typename QueryGeomType>
+		bool OverlapGeomImp(const QueryGeomType& QueryGeom, const TRigidTransform<T, 3>& QueryTM, const T Thickness) const;
+
+		template <typename QueryGeomType>
+		bool SweepGeomImp(const QueryGeomType& QueryGeom, const TRigidTransform<T, 3>& StartTM, const TVector<T, 3>& Dir, const T Length, T& OutTime, TVector<T, 3>& OutPosition, TVector<T, 3>& OutNormal, int32& OutFaceIndex, const T Thickness) const;
 	};
 }
diff --git a/Engine/Source/Runtime/Experimental/Chaos/Public/Chaos/ImplicitObject.h b/Engine/Source/Runtime/Experimental/Chaos/Public/Chaos/ImplicitObject.h
--- a/Engine/Source/Runtime/Experimental/Chaos/Public/Chaos/ImplicitObject.h
+++ b/Engine/Source/Runtime/Experimental/Chaos/Public/Chaos/ImplicitObject.h
@@ -171,12 +171,6 @@ public:
 	FVec3 Normal(const FVec3& x) const;
 	virtual FReal PhiWithNormal(const FVec3& x, FVec3& Normal) const = 0;
 	virtual const class TBox<FReal, 3>& BoundingBox() const;
-	virtual FReal GetMargin() const
-	{
-		return 0;
-	}
-	virtual FVec3 Support(const FVec3& Direction, const FReal Thickness) const;
-	virtual FVec3 Support2(const FVec3& Direction) const { return Support(Direction, 0); }
 	bool HasBoundingBox() const { return bHasBoundingBox; }
 	bool IsConvex() const { return bIsConvex; }
 	void IgnoreAnalyticCollisions(const bool Ignore = true) { bIgnoreAnalyticCollisions = Ignore; }
@@ -258,7 +252,7 @@ public:
 
 	void SerializeImp(FArchive& Ar);
 
-	static EImplicitObjectType StaticType()
+	constexpr static EImplicitObjectType StaticType()
 	{
 		return ImplicitObjectType::Unknown;
 	}
diff --git a/Engine/Source/Runtime/Experimental/Chaos/Public/Chaos/ImplicitObjectScaled.h b/Engine/Source/Runtime/Experimental/Chaos/Public/Chaos/ImplicitObjectScaled.h
--- a/Engine/Source/Runtime/Experimental/Chaos/Public/Chaos/ImplicitObjectScaled.h
+++ b/Engine/Source/Runtime/Experimental/Chaos/Public/Chaos/ImplicitObjectScaled.h
@@ -32,7 +32,7 @@ public:
 		this->bIsConvex = MObject->IsConvex();
 	}
 
-	static EImplicitObjectType StaticType()
+	static constexpr EImplicitObjectType StaticType()
 	{
 		return TConcrete::StaticType() | ImplicitObjectType::IsInstanced;
 	}
@@ -89,10 +89,10 @@ public:
 		return MObject->Overlap(Point, Thickness);
 	}
 
-	virtual TVector<T, d> Support(const TVector<T, d>& Direction, const T Thickness) const override
-	{
-		return MObject->Support(Direction, Thickness);
-	}
+	FORCEINLINE T GetMargin() const { return MObject->GetMargin(); }
+
+	FORCEINLINE TVector<T, d> Support(const TVector<T, d>& Direction, const T Thickness) const { return MObject->Support(Direction, Thickness); }
+	FORCEINLINE TVector<T, d> Support2(const TVector<T, d>& Direction, const T Thickness) const { return MObject->Support2(Direction); }
 
 	virtual const TBox<T, d>& BoundingBox() const override { return MObject->BoundingBox(); }
 
@@ -179,7 +179,7 @@ public:
 	}
 	~TImplicitObjectScaled() {}
 
-	static EImplicitObjectType StaticType()
+	static constexpr EImplicitObjectType StaticType()
 	{
 		return TConcrete::StaticType() | ImplicitObjectType::IsScaled;
 	}
@@ -290,7 +290,8 @@ public:
 	}
 
 	/** This is a low level function and assumes the internal object has a SweepGeom function. Should not be called directly. See GeometryQueries.h : SweepQuery */
-	bool LowLevelSweepGeom(const FImplicitObject& B, const TRigidTransform<T, d>& BToATM, const TVector<T, d>& LocalDir, const T Length, T& OutTime, TVector<T, d>& LocalPosition, TVector<T, d>& LocalNormal, int32& OutFaceIndex, T Thickness = 0) const
+	template <typename QueryGeomType>
+	bool LowLevelSweepGeom(const QueryGeomType& B, const TRigidTransform<T, d>& BToATM, const TVector<T, d>& LocalDir, const T Length, T& OutTime, TVector<T, d>& LocalPosition, TVector<T, d>& LocalNormal, int32& OutFaceIndex, T Thickness = 0) const
 	{
 		ensure(Length > 0);
 		ensure(FMath::IsNearlyEqual(LocalDir.SizeSquared(), 1, KINDA_SMALL_NUMBER));
@@ -308,9 +309,7 @@ public:
 			TVector<T, d> UnscaledNormal;
 			float UnscaledTime;
 
-			TUniquePtr<FImplicitObject> HackBPtr(const_cast<FImplicitObject*>(&B));	//todo: hack, need scaled object to accept raw ptr similar to transformed implicit
-			TImplicitObjectScaled<FImplicitObject> ScaledB(MakeSerializable(HackBPtr), MInvScale);
-			HackBPtr.Release();
+			auto ScaledB = MakeScaledHelper(B, MInvScale);
 
 			TRigidTransform<T, d> BToATMNoScale(BToATM.GetLocation() * MInvScale, BToATM.GetRotation());
 			
@@ -328,14 +327,12 @@ public:
 	}
 
 	/** This is a low level function and assumes the internal object has a OverlapGeom function. Should not be called directly. See GeometryQueries.h : OverlapQuery */
-	bool LowLevelOverlapGeom(const FImplicitObject& B, const TRigidTransform<T, d>& BToATM, T Thickness = 0) const
+	template <typename QueryGeomType>
+	bool LowLevelOverlapGeom(const QueryGeomType& B, const TRigidTransform<T, d>& BToATM, T Thickness = 0) const
 	{
 		ensure(Thickness == 0 || (FMath::IsNearlyEqual(MScale[0], MScale[1]) && FMath::IsNearlyEqual(MScale[0], MScale[2])));
 
-		TUniquePtr<FImplicitObject> HackBPtr(const_cast<FImplicitObject*>(&B));	//todo: hack, need scaled object to except raw ptr similar to transformed implicit
-		TImplicitObjectScaled<FImplicitObject> ScaledB(MakeSerializable(HackBPtr), MInvScale);
-		HackBPtr.Release();
-
+		auto ScaledB = MakeScaledHelper(B, MInvScale);
 		TRigidTransform<T, d> BToATMNoScale(BToATM.GetLocation() * MInvScale, BToATM.GetRotation());
 		return MObject->OverlapGeom(ScaledB, BToATMNoScale, MInternalThickness + Thickness);
 	}
@@ -400,7 +397,7 @@ public:
 		return ClosestIntersection;
 	}
 
-	virtual TVector<T, d> Support(const TVector<T, d>& Direction, const T Thickness) const override
+	FORCEINLINE_DEBUGGABLE TVector<T, d> Support(const TVector<T, d>& Direction, const T Thickness) const
 	{
 		// Support_obj(dir) = pt => for all x in obj, pt \dot dir >= x \dot dir
 		// We want Support_objScaled(dir) = Support_obj(dir') where dir' is some modification of dir so we can use the unscaled support function
@@ -412,6 +409,13 @@ public:
 		return Thickness > 0 ? TVector<T, d>(UnthickenedPt + Direction.GetSafeNormal() * Thickness) : UnthickenedPt;
 	}
 
+	FORCEINLINE_DEBUGGABLE TVector<T, d> Support2(const TVector<T, d>& Direction) const
+	{
+		return MObject->Support2(Direction * MScale) * MScale;
+	}
+
+	FORCEINLINE T GetMargin() const { return MObject->GetMargin(); }
+
 	const TVector<T, d>& GetScale() const { return MScale; }
 	void SetScale(const TVector<T, d>& Scale)
 	{
@@ -442,6 +446,12 @@ public:
 		FImplicitObject::SerializeImp(Ar);
 		Ar << MObject << MScale << MInvScale << MLocalBoundingBox;
 		ensure(MInternalThickness == 0);	//not supported: do we care?
+
+		Ar.UsingCustomVersion(FExternalPhysicsCustomObjectVersion::GUID);
+		if (Ar.CustomVer(FExternalPhysicsCustomObjectVersion::GUID) < FExternalPhysicsCustomObjectVersion::ScaledGeometryIsConcrete)
+		{
+			this->Type = MObject->GetType() | ImplicitObjectType::IsScaled;	//update type so downcasts work
+		}
 	}
 
 	virtual uint32 GetTypeHash() const override
@@ -489,6 +499,23 @@ private:
 		const TVector<T, d> Vector2 = UnscaledBounds.Max() *MScale;
 		MLocalBoundingBox.GrowToInclude(Vector2);
 	}
+
+	template <typename QueryGeomType>
+	static auto MakeScaledHelper(const QueryGeomType& B, const TVector<T,d>& InvScale )
+	{
+		TUniquePtr<QueryGeomType> HackBPtr(const_cast<QueryGeomType*>(&B));	//todo: hack, need scaled object to accept raw ptr similar to transformed implicit
+		TImplicitObjectScaled<QueryGeomType> ScaledB(MakeSerializable(HackBPtr), InvScale);
+		HackBPtr.Release();
+		return ScaledB;
+	}
+
+	template <typename QueryGeomType>
+	static auto MakeScaledHelper(const TImplicitObjectScaled<QueryGeomType>& B, const TVector<T,d>& InvScale)
+	{
+		//if scaled of scaled just collapse into one scaled
+		TImplicitObjectScaled<QueryGeomType> ScaledB(B.Object(), InvScale * B.GetScale());
+		return ScaledB;
+	}
 };
 
 template <typename TConcrete>
diff --git a/Engine/Source/Runtime/Experimental/Chaos/Public/Chaos/ImplicitObjectTransformed.h b/Engine/Source/Runtime/Experimental/Chaos/Public/Chaos/ImplicitObjectTransformed.h
--- a/Engine/Source/Runtime/Experimental/Chaos/Public/Chaos/ImplicitObjectTransformed.h
+++ b/Engine/Source/Runtime/Experimental/Chaos/Public/Chaos/ImplicitObjectTransformed.h
@@ -82,7 +82,7 @@ public:
 	}
 	~TImplicitObjectTransformed() {}
 
-	static EImplicitObjectType StaticType()
+	static constexpr EImplicitObjectType StaticType()
 	{
 		return ImplicitObjectType::Transformed;
 	}
@@ -153,11 +153,6 @@ public:
 		return ClosestIntersection;
 	}
 
-	virtual TVector<T, d> Support(const TVector<T, d>& Direction, const T Thickness) const override
-	{
-		return MTransform.TransformPosition(MObject->Support(MTransform.InverseTransformVector(Direction), Thickness));
-	}
-
 	const TRigidTransform<T, d>& GetTransform() const { return MTransform; }
 	void SetTransform(const TRigidTransform<T, d>& InTransform)
 	{
diff --git a/Engine/Source/Runtime/Experimental/Chaos/Public/Chaos/ImplicitObjectUnion.h b/Engine/Source/Runtime/Experimental/Chaos/Public/Chaos/ImplicitObjectUnion.h
--- a/Engine/Source/Runtime/Experimental/Chaos/Public/Chaos/ImplicitObjectUnion.h
+++ b/Engine/Source/Runtime/Experimental/Chaos/Public/Chaos/ImplicitObjectUnion.h
@@ -69,7 +69,7 @@ class TImplicitObjectUnion : public FImplicitObject
 	}
 	virtual ~TImplicitObjectUnion() {}
 
-	FORCEINLINE static EImplicitObjectType StaticType()
+	FORCEINLINE static constexpr EImplicitObjectType StaticType()
 	{
 		return ImplicitObjectType::Union;
 	}
diff --git a/Engine/Source/Runtime/Experimental/Chaos/Public/Chaos/Levelset.h b/Engine/Source/Runtime/Experimental/Chaos/Public/Chaos/Levelset.h
--- a/Engine/Source/Runtime/Experimental/Chaos/Public/Chaos/Levelset.h
+++ b/Engine/Source/Runtime/Experimental/Chaos/Public/Chaos/Levelset.h
@@ -54,7 +54,7 @@ class CHAOS_API TLevelSet final : public FImplicitObject
 		}
 	}
 
-	FORCEINLINE static EImplicitObjectType StaticType()
+	FORCEINLINE static constexpr EImplicitObjectType StaticType()
 	{
 		return ImplicitObjectType::LevelSet;
 	}
diff --git a/Engine/Source/Runtime/Experimental/Chaos/Public/Chaos/Plane.h b/Engine/Source/Runtime/Experimental/Chaos/Public/Chaos/Plane.h
--- a/Engine/Source/Runtime/Experimental/Chaos/Public/Chaos/Plane.h
+++ b/Engine/Source/Runtime/Experimental/Chaos/Public/Chaos/Plane.h
@@ -34,7 +34,7 @@ class TPlane final : public FImplicitObject
 	}
 	virtual ~TPlane() {}
 
-	static EImplicitObjectType StaticType()
+	static constexpr EImplicitObjectType StaticType()
 	{
 		return ImplicitObjectType::Plane;
 	}
diff --git a/Engine/Source/Runtime/Experimental/Chaos/Public/Chaos/Segment.h b/Engine/Source/Runtime/Experimental/Chaos/Public/Chaos/Segment.h
--- a/Engine/Source/Runtime/Experimental/Chaos/Public/Chaos/Segment.h
+++ b/Engine/Source/Runtime/Experimental/Chaos/Public/Chaos/Segment.h
@@ -41,7 +41,7 @@ namespace Chaos
 			return FarthestCap + (NormalizedDirection * Thickness);
 		}
 
-		TVector<T, 3> Support2(const TVector<T, 3>& Direction) const
+		FORCEINLINE_DEBUGGABLE TVector<T, 3> Support2(const TVector<T, 3>& Direction) const
 		{
 			const T Dot = TVector<T, 3>::DotProduct(Direction, MAxis);
 			const TVector<T, 3> FarthestCap = Dot >= 0 ? GetX2() : GetX1();	//orthogonal we choose either
diff --git a/Engine/Source/Runtime/Experimental/Chaos/Public/Chaos/Sphere.h b/Engine/Source/Runtime/Experimental/Chaos/Public/Chaos/Sphere.h
--- a/Engine/Source/Runtime/Experimental/Chaos/Public/Chaos/Sphere.h
+++ b/Engine/Source/Runtime/Experimental/Chaos/Public/Chaos/Sphere.h
@@ -62,7 +62,7 @@ namespace Chaos
 
 		virtual ~TSphere() {}
 
-		static EImplicitObjectType StaticType()
+		static constexpr EImplicitObjectType StaticType()
 		{ 
 			return ImplicitObjectType::Sphere; 
 		}
@@ -177,7 +177,7 @@ namespace Chaos
 			return MakePair(TVector<T, d>(Root2 * Direction + StartPoint), true);
 		}
 
-		virtual TVector<T, d> Support(const TVector<T, d>& Direction, const T Thickness) const override
+		TVector<T, d> Support(const TVector<T, d>& Direction, const T Thickness) const
 		{
 			//We want N / ||N|| and to avoid inf
 			//So we want N / ||N|| < 1 / eps => N eps < ||N||, but this is clearly true for all eps < 1 and N > 0
@@ -191,9 +191,9 @@ namespace Chaos
 			return Center + Normalized * (Radius + Thickness);
 		}
 
-		virtual TVector<T, d> Support2(const TVector<T, d>& Direction) const override { return Center; }
+		FORCEINLINE const TVector<T, d>& Support2(const TVector<T, d>& Direction) const { return Center; }
 
-		virtual T GetMargin() const override
+		FORCEINLINE T GetMargin() const
 		{
 			return Radius;
 		}
diff --git a/Engine/Source/Runtime/Experimental/Chaos/Public/Chaos/TaperedCylinder.h b/Engine/Source/Runtime/Experimental/Chaos/Public/Chaos/TaperedCylinder.h
--- a/Engine/Source/Runtime/Experimental/Chaos/Public/Chaos/TaperedCylinder.h
+++ b/Engine/Source/Runtime/Experimental/Chaos/Public/Chaos/TaperedCylinder.h
@@ -58,7 +58,7 @@ namespace Chaos
 		}
 		~TTaperedCylinder() {}
 
-		static EImplicitObjectType StaticType() { return ImplicitObjectType::TaperedCylinder; }
+		static constexpr EImplicitObjectType StaticType() { return ImplicitObjectType::TaperedCylinder; }
 
 		/**
 		 * Returns sample points centered about the origin.
diff --git a/Engine/Source/Runtime/Experimental/Chaos/Public/Chaos/Triangle.h b/Engine/Source/Runtime/Experimental/Chaos/Public/Chaos/Triangle.h
--- a/Engine/Source/Runtime/Experimental/Chaos/Public/Chaos/Triangle.h
+++ b/Engine/Source/Runtime/Experimental/Chaos/Public/Chaos/Triangle.h
@@ -47,11 +47,6 @@ namespace Chaos
 			return TVector<T, 3>::DotProduct((InSamplePoint - ClosestPoint), OutNormal);
 		}
 
-		FORCEINLINE TVector<T, 3> Support2(const TVector<T, 3>& Direction) const
-		{
-			return Support(Direction, 0);
-		}
-
 		FORCEINLINE TVector<T, 3> Support(const TVector<T, 3>& Direction, const T Thickness) const
 		{
 			const float DotA = TVector<T, 3>::DotProduct(A, Direction);
@@ -82,6 +77,26 @@ namespace Chaos
 			return C;
 		}
 
+		FORCEINLINE_DEBUGGABLE TVector<T, 3> Support2(const TVector<T, 3>& Direction) const
+		{
+			const float DotA = TVector<T, 3>::DotProduct(A, Direction);
+			const float DotB = TVector<T, 3>::DotProduct(B, Direction);
+			const float DotC = TVector<T, 3>::DotProduct(C, Direction);
+
+			if (DotA > DotB && DotA > DotC)
+			{
+				return A;
+			}
+			else if (DotB > DotA && DotB > DotC)
+			{
+				return B;
+			}
+
+			return C;
+		}
+
+		FORCEINLINE T GetMargin() const { return 0; }
+
 		FORCEINLINE bool Raycast(const TVector<T, 3>& StartPoint, const TVector<T, 3>& Dir, const T Length, const T Thickness, T& OutTime, TVector<T, 3>& OutPosition, TVector<T, 3>& OutNormal, int32& OutFaceIndex) const
 		{
 			// No face as this is only one triangle
@@ -165,7 +180,7 @@ namespace Chaos
 			return Tri.GetPlane();
 		}
 
-		static EImplicitObjectType StaticType()
+		static constexpr EImplicitObjectType StaticType()
 		{
 			return ImplicitObjectType::Triangle;
 		}
diff --git a/Engine/Source/Runtime/Experimental/Chaos/Public/Chaos/TriangleMeshImplicitObject.h b/Engine/Source/Runtime/Experimental/Chaos/Public/Chaos/TriangleMeshImplicitObject.h
--- a/Engine/Source/Runtime/Experimental/Chaos/Public/Chaos/TriangleMeshImplicitObject.h
+++ b/Engine/Source/Runtime/Experimental/Chaos/Public/Chaos/TriangleMeshImplicitObject.h
@@ -4,6 +4,8 @@
 #include "Chaos/Array.h"
 #include "Chaos/GeometryParticles.h"
 #include "Chaos/SegmentMesh.h"
+#include "Chaos/ImplicitObjectScaled.h"
+
 #include "ImplicitObject.h"
 #include "Box.h"
 #include "BoundingVolumeHierarchy.h"
@@ -14,6 +16,12 @@
 namespace Chaos
 {
 
+	template <typename T>
+	class TCapsule;
+
+	template <typename T, int d>
+	class TConvex;
+
 	template<typename T>
 	class CHAOS_API TTriangleMeshImplicitObject final : public FImplicitObject
 	{
@@ -34,8 +42,25 @@ namespace Chaos
 		virtual bool Raycast(const TVector<T, 3>& StartPoint, const TVector<T, 3>& Dir, const T Length, const T Thickness, T& OutTime, TVector<T,3>& OutPosition, TVector<T,3>& OutNormal, int32& OutFaceIndex) const override;
 		virtual bool Overlap(const TVector<T, 3>& Point, const T Thickness) const override;
 
-		bool OverlapGeom(const FImplicitObject& QueryGeom, const TRigidTransform<T, 3>& QueryTM, const T Thickness) const;
-		bool SweepGeom(const FImplicitObject& QueryGeom, const TRigidTransform<T, 3>& StartTM, const TVector<T,3>& Dir, const T Length, T& OutTime, TVector<T,3>& OutPosition, TVector<T,3>& OutNormal, int32& OutFaceIndex, const T Thickness = 0) const;
+		bool OverlapGeom(const TSphere<T,3>& QueryGeom, const TRigidTransform<T, 3>& QueryTM, const T Thickness) const;
+		bool OverlapGeom(const TBox<T, 3>& QueryGeom, const TRigidTransform<T, 3>& QueryTM, const T Thickness) const;
+		bool OverlapGeom(const TCapsule<T>& QueryGeom, const TRigidTransform<T, 3>& QueryTM, const T Thickness) const;
+		bool OverlapGeom(const TConvex<T, 3>& QueryGeom, const TRigidTransform<T, 3>& QueryTM, const T Thickness) const;
+		bool OverlapGeom(const TImplicitObjectScaled<TSphere<T, 3>>& QueryGeom, const TRigidTransform<T, 3>& QueryTM, const T Thickness) const;
+		bool OverlapGeom(const TImplicitObjectScaled<TBox<T, 3>>& QueryGeom, const TRigidTransform<T, 3>& QueryTM, const T Thickness) const;
+		bool OverlapGeom(const TImplicitObjectScaled<TCapsule<T>>& QueryGeom, const TRigidTransform<T, 3>& QueryTM, const T Thickness) const;
+		bool OverlapGeom(const TImplicitObjectScaled<TConvex<T, 3>>& QueryGeom, const TRigidTransform<T, 3>& QueryTM, const T Thickness) const;
+		bool OverlapGeom(const TImplicitObjectScaled<TImplicitObjectScaled<TConvex<T, 3>>>& QueryGeom, const TRigidTransform<T, 3>& QueryTM, const T Thickness) const;
+
+		bool SweepGeom(const TSphere<T,3>& QueryGeom, const TRigidTransform<T, 3>& StartTM, const TVector<T,3>& Dir, const T Length, T& OutTime, TVector<T,3>& OutPosition, TVector<T,3>& OutNormal, int32& OutFaceIndex, const T Thickness = 0) const;
+		bool SweepGeom(const TBox<T, 3>& QueryGeom, const TRigidTransform<T, 3>& StartTM, const TVector<T, 3>& Dir, const T Length, T& OutTime, TVector<T, 3>& OutPosition, TVector<T, 3>& OutNormal, int32& OutFaceIndex, const T Thickness = 0) const;
+		bool SweepGeom(const TCapsule<T>& QueryGeom, const TRigidTransform<T, 3>& StartTM, const TVector<T, 3>& Dir, const T Length, T& OutTime, TVector<T, 3>& OutPosition, TVector<T, 3>& OutNormal, int32& OutFaceIndex, const T Thickness = 0) const;
+		bool SweepGeom(const TConvex<T, 3>& QueryGeom, const TRigidTransform<T, 3>& StartTM, const TVector<T, 3>& Dir, const T Length, T& OutTime, TVector<T, 3>& OutPosition, TVector<T, 3>& OutNormal, int32& OutFaceIndex, const T Thickness = 0) const;
+		bool SweepGeom(const TImplicitObjectScaled<TSphere<T, 3>>& QueryGeom, const TRigidTransform<T, 3>& StartTM, const TVector<T, 3>& Dir, const T Length, T& OutTime, TVector<T, 3>& OutPosition, TVector<T, 3>& OutNormal, int32& OutFaceIndex, const T Thickness = 0) const;
+		bool SweepGeom(const TImplicitObjectScaled<TBox<T, 3>>& QueryGeom, const TRigidTransform<T, 3>& StartTM, const TVector<T, 3>& Dir, const T Length, T& OutTime, TVector<T, 3>& OutPosition, TVector<T, 3>& OutNormal, int32& OutFaceIndex, const T Thickness = 0) const;
+		bool SweepGeom(const TImplicitObjectScaled<TCapsule<T>>& QueryGeom, const TRigidTransform<T, 3>& StartTM, const TVector<T, 3>& Dir, const T Length, T& OutTime, TVector<T, 3>& OutPosition, TVector<T, 3>& OutNormal, int32& OutFaceIndex, const T Thickness = 0) const;
+		bool SweepGeom(const TImplicitObjectScaled<TConvex<T, 3>>& QueryGeom, const TRigidTransform<T, 3>& StartTM, const TVector<T, 3>& Dir, const T Length, T& OutTime, TVector<T, 3>& OutPosition, TVector<T, 3>& OutNormal, int32& OutFaceIndex, const T Thickness = 0) const;
+
 		virtual int32 FindMostOpposingFace(const TVector<T, 3>& Position, const TVector<T, 3>& UnitDir, int32 HintFaceIndex, T SearchDistance) const override;
 		virtual TVector<T, 3> FindGeometryOpposingNormal(const TVector<T, 3>& DenormDir, int32 FaceIndex, const TVector<T, 3>& OriginalNormal) const override;
 
@@ -44,7 +69,7 @@ namespace Chaos
 			return MLocalBoundingBox;
 		}
 
-		static EImplicitObjectType StaticType()
+		static constexpr EImplicitObjectType StaticType()
 		{
 			return ImplicitObjectType::TriangleMesh;
 		}
@@ -138,7 +163,7 @@ namespace Chaos
 
 		BVHType BVH;
 
-		template <typename R>
+		template <typename Geom, typename R>
 		friend struct TTriangleMeshSweepVisitor;
 
 		// Required by implicit object serialization, disabled for general use.
@@ -147,6 +172,12 @@ namespace Chaos
 		TTriangleMeshImplicitObject() 
 			: FImplicitObject(EImplicitObject::HasBoundingBox, ImplicitObjectType::TriangleMesh)
 		{};
+
+		template <typename QueryGeomType>
+		bool OverlapGeomImp(const QueryGeomType& QueryGeom, const TRigidTransform<T, 3>& QueryTM, const T Thickness) const;
+
+		template <typename QueryGeomType>
+		bool SweepGeomImp(const QueryGeomType& QueryGeom, const TRigidTransform<T, 3>& StartTM, const TVector<T, 3>& Dir, const T Length, T& OutTime, TVector<T, 3>& OutPosition, TVector<T, 3>& OutNormal, int32& OutFaceIndex, const T Thickness) const;
 	};
 }
 
diff --git a/Engine/Source/Runtime/PhysicsSQ/Private/SQAccelerator.cpp b/Engine/Source/Runtime/PhysicsSQ/Private/SQAccelerator.cpp
--- a/Engine/Source/Runtime/PhysicsSQ/Private/SQAccelerator.cpp
+++ b/Engine/Source/Runtime/PhysicsSQ/Private/SQAccelerator.cpp
@@ -80,7 +80,7 @@ void FillHitHelper(ChaosInterface::FOverlapHit& Hit, const float Distance, const
 {
 }
 
-template <typename TPayload, typename THitType>
+template <typename QueryGeometryType, typename TPayload, typename THitType>
 struct TSQVisitor : public Chaos::ISpatialVisitor<TPayload, float>
 {
 	TSQVisitor(const FVector& InStartPoint, const FVector& InDir, ChaosInterface::FSQHitBuffer<ChaosInterface::FRaycastHit>& InHitBuffer, EHitFlags InOutputFlags,
@@ -101,7 +101,7 @@ struct TSQVisitor : public Chaos::ISpatialVisitor<TPayload, float>
 	}
 
 	TSQVisitor(const FTransform& InStartTM, const FVector& InDir, ChaosInterface::FSQHitBuffer<ChaosInterface::FSweepHit>& InHitBuffer, EHitFlags InOutputFlags,
-		const FQueryFilterData& InQueryFilterData, ICollisionQueryFilterCallbackBase& InQueryCallback, const Chaos::FImplicitObject& InQueryGeom, const FQueryDebugParams& InDebugParams)
+		const FQueryFilterData& InQueryFilterData, ICollisionQueryFilterCallbackBase& InQueryCallback, const QueryGeometryType& InQueryGeom, const FQueryDebugParams& InDebugParams)
 		: StartTM(InStartTM)
 		, Dir(InDir)
 		, HitBuffer(InHitBuffer)
@@ -120,7 +120,7 @@ struct TSQVisitor : public Chaos::ISpatialVisitor<TPayload, float>
 	}
 
 	TSQVisitor(const FTransform& InWorldTM, ChaosInterface::FSQHitBuffer<ChaosInterface::FOverlapHit>& InHitBuffer,
-		const FQueryFilterData& InQueryFilterData, ICollisionQueryFilterCallbackBase& InQueryCallback, const Chaos::FImplicitObject& InQueryGeom, const FQueryDebugParams& InDebugParams)
+		const FQueryFilterData& InQueryFilterData, ICollisionQueryFilterCallbackBase& InQueryCallback, const QueryGeometryType& InQueryGeom, const FQueryDebugParams& InDebugParams)
 		: StartTM(InWorldTM)
 		, HitBuffer(InHitBuffer)
 		, QueryFilterData(InQueryFilterData)
@@ -221,11 +221,11 @@ private:
 				}
 				else if(SQ == ESQType::Sweep && CurLength > 0)
 				{
-					bHit = SweepQuery<float, 3>(*Geom, ActorTM, *QueryGeom, StartTM, Dir, CurLength, Distance, WorldPosition, WorldNormal, FaceIdx, 0.f, bComputeMTD);
+					bHit = SweepQuery(*Geom, ActorTM, *QueryGeom, StartTM, Dir, CurLength, Distance, WorldPosition, WorldNormal, FaceIdx, 0.f, bComputeMTD);
 				}
 				else if (SQ == ESQType::Overlap || (SQ == ESQType::Sweep && CurLength == 0))
 				{
-					bHit = OverlapQuery<float, 3>(*Geom, ActorTM, *QueryGeom, StartTM, /*Thickness=*/0);
+					bHit = OverlapQuery(*Geom, ActorTM, *QueryGeom, StartTM, /*Thickness=*/0);
 				}
 
 				if(bHit)
@@ -290,7 +290,7 @@ private:
 	const FQueryFilterData& QueryFilterData;
 	ICollisionQueryFilterCallbackBase& QueryCallback;
 	bool bAnyHit;
-	const Chaos::FImplicitObject* QueryGeom;
+	const QueryGeometryType* QueryGeom;
 	const FQueryDebugParams DebugParams;
 };
 
@@ -299,37 +299,50 @@ void FChaosSQAccelerator::Raycast(const FVector& Start, const FVector& Dir, cons
 	using namespace Chaos;
 	using namespace ChaosInterface;
 
-	TSQVisitor<TAccelerationStructureHandle<float, 3>, FRaycastHit> RaycastVisitor(Start, Dir, HitBuffer, OutputFlags, QueryFilterData, QueryCallback, DebugParams);
+	TSQVisitor<TSphere<float,3>, TAccelerationStructureHandle<float, 3>, FRaycastHit> RaycastVisitor(Start, Dir, HitBuffer, OutputFlags, QueryFilterData, QueryCallback, DebugParams);
 	HitBuffer.IncFlushCount();
 	SpatialAcceleration.Raycast(Start, Dir, DeltaMagnitude, RaycastVisitor);
 	HitBuffer.DecFlushCount();
 }
 
-void FChaosSQAccelerator::Sweep(const Chaos::FImplicitObject& QueryGeom, const FTransform& StartTM, const FVector& Dir, const float DeltaMagnitude, ChaosInterface::FSQHitBuffer<ChaosInterface::FSweepHit>& HitBuffer, EHitFlags OutputFlags, const FQueryFilterData& QueryFilterData, ICollisionQueryFilterCallbackBase& QueryCallback, const FQueryDebugParams& DebugParams) const
+template <typename QueryGeomType>
+void SweepHelper(const QueryGeomType& QueryGeom, const Chaos::ISpatialAcceleration<Chaos::TAccelerationStructureHandle<float, 3>, float, 3>& SpatialAcceleration, const FTransform& StartTM, const FVector& Dir, const float DeltaMagnitude, ChaosInterface::FSQHitBuffer<ChaosInterface::FSweepHit>& HitBuffer, EHitFlags OutputFlags, const FQueryFilterData& QueryFilterData, ICollisionQueryFilterCallbackBase& QueryCallback, const FQueryDebugParams& DebugParams)
 {
 	using namespace Chaos;
 	using namespace ChaosInterface;
 
 	const TBox<float, 3> Bounds = QueryGeom.BoundingBox().TransformedBox(StartTM);
 	const FVector HalfExtents = Bounds.Extents() * 0.5f;
-	TSQVisitor<TAccelerationStructureHandle<float, 3>, FSweepHit> SweepVisitor(StartTM, Dir, HitBuffer, OutputFlags, QueryFilterData, QueryCallback, QueryGeom, DebugParams);
+
+	TSQVisitor<QueryGeomType, TAccelerationStructureHandle<float, 3>, FSweepHit> SweepVisitor(StartTM, Dir, HitBuffer, OutputFlags, QueryFilterData, QueryCallback, QueryGeom, DebugParams);
 	HitBuffer.IncFlushCount();
 	SpatialAcceleration.Sweep(Bounds.GetCenter(), Dir, DeltaMagnitude, HalfExtents, SweepVisitor);
 	HitBuffer.DecFlushCount();
 }
 
-void FChaosSQAccelerator::Overlap(const Chaos::FImplicitObject& QueryGeom, const FTransform& GeomPose, ChaosInterface::FSQHitBuffer<ChaosInterface::FOverlapHit>& HitBuffer, const FQueryFilterData& QueryFilterData, ICollisionQueryFilterCallbackBase& QueryCallback, const FQueryDebugParams& DebugParams) const
+void FChaosSQAccelerator::Sweep(const Chaos::FImplicitObject& QueryGeom, const FTransform& StartTM, const FVector& Dir, const float DeltaMagnitude, ChaosInterface::FSQHitBuffer<ChaosInterface::FSweepHit>& HitBuffer, EHitFlags OutputFlags, const FQueryFilterData& QueryFilterData, ICollisionQueryFilterCallbackBase& QueryCallback, const FQueryDebugParams& DebugParams) const
+{
+	return Chaos::CastHelper(QueryGeom, [&](const auto& Downcast) { return SweepHelper(Downcast, SpatialAcceleration, StartTM, Dir, DeltaMagnitude, HitBuffer, OutputFlags, QueryFilterData, QueryCallback, DebugParams); });
+}
+
+template <typename QueryGeomType>
+void OverlapHelper(const QueryGeomType& QueryGeom, const Chaos::ISpatialAcceleration<Chaos::TAccelerationStructureHandle<float, 3>, float, 3>& SpatialAcceleration, const FTransform& GeomPose, ChaosInterface::FSQHitBuffer<ChaosInterface::FOverlapHit>& HitBuffer, const FQueryFilterData& QueryFilterData, ICollisionQueryFilterCallbackBase& QueryCallback, const FQueryDebugParams& DebugParams)
 {
 	using namespace Chaos;
 	using namespace ChaosInterface;
 
 	const TBox<float, 3> Bounds = QueryGeom.BoundingBox().TransformedBox(GeomPose);
-	TSQVisitor<TAccelerationStructureHandle<float, 3>, FOverlapHit> OverlapVisitor(GeomPose, HitBuffer, QueryFilterData, QueryCallback, QueryGeom, DebugParams);
+	TSQVisitor<QueryGeomType, TAccelerationStructureHandle<float, 3>, FOverlapHit> OverlapVisitor(GeomPose, HitBuffer, QueryFilterData, QueryCallback, QueryGeom, DebugParams);
 	HitBuffer.IncFlushCount();
 	SpatialAcceleration.Overlap(Bounds, OverlapVisitor);
 	HitBuffer.DecFlushCount();
 }
 
+void FChaosSQAccelerator::Overlap(const Chaos::FImplicitObject& QueryGeom, const FTransform& GeomPose, ChaosInterface::FSQHitBuffer<ChaosInterface::FOverlapHit>& HitBuffer, const FQueryFilterData& QueryFilterData, ICollisionQueryFilterCallbackBase& QueryCallback, const FQueryDebugParams& DebugParams) const
+{
+	return Chaos::CastHelper(QueryGeom, [&](const auto& Downcast) { return OverlapHelper(Downcast, SpatialAcceleration, GeomPose, HitBuffer, QueryFilterData, QueryCallback, DebugParams); });
+}
+
 #if WITH_PHYSX
 FChaosSQAcceleratorAdapter::FChaosSQAcceleratorAdapter(const Chaos::ISpatialAcceleration<Chaos::TAccelerationStructureHandle<float, 3>, float, 3>& InSpatialAcceleration)
 	: ChaosSQAccelerator(InSpatialAcceleration)
