Chaos constraints - push out iterations across all types.

diff --git a/Engine/Source/Runtime/Engine/Private/PhysicsEngine/ImmediatePhysics/ImmediatePhysicsChaos/ImmediatePhysicsSimulation_Chaos.cpp b/Engine/Source/Runtime/Engine/Private/PhysicsEngine/ImmediatePhysics/ImmediatePhysicsChaos/ImmediatePhysicsSimulation_Chaos.cpp
--- a/Engine/Source/Runtime/Engine/Private/PhysicsEngine/ImmediatePhysics/ImmediatePhysicsChaos/ImmediatePhysicsSimulation_Chaos.cpp
+++ b/Engine/Source/Runtime/Engine/Private/PhysicsEngine/ImmediatePhysics/ImmediatePhysicsChaos/ImmediatePhysicsSimulation_Chaos.cpp
@@ -22,18 +22,20 @@
 //
 float ChaosImmediate_Evolution_DeltaTime = 0.03f;
 int32 ChaosImmediate_Evolution_Iterations = 6;
+int32 ChaosImmediate_Evolution_PushOutIterations = 5;
 FAutoConsoleVariableRef CVarChaosImmPhysDeltaTime(TEXT("p.Chaos.ImmPhys.DeltaTime"), ChaosImmediate_Evolution_DeltaTime, TEXT("Override chaos immediate physics delta time if non-zero"));
 FAutoConsoleVariableRef CVarChaosImmPhysIterations(TEXT("p.Chaos.ImmPhys.Iterations"), ChaosImmediate_Evolution_Iterations, TEXT("Number of constraint solver loops in immediate physics"));
+FAutoConsoleVariableRef CVarChaosImmPhysPushOutIterations(TEXT("p.Chaos.ImmPhys.PushOutIterations"), ChaosImmediate_Evolution_PushOutIterations, TEXT("Set the ApplyPushOut() (position correction) iteration count"));
 
 int32 ChaosImmediate_Collision_Enabled = 1;
 int32 ChaosImmediate_Collision_ApplyEnabled = 1;
-int32 ChaosImmediate_Collision_PushOutIterations = 5;
 int32 ChaosImmediate_Collision_PushOutPairIterations = 2;
+int32 ChaosImmediate_Collision_Priority = 1;
 float ChaosImmediate_Collision_Thickness = 0;
 FAutoConsoleVariableRef CVarChaosImmPhysDisableCollisions(TEXT("p.Chaos.ImmPhys.EnableCollisions"), ChaosImmediate_Collision_Enabled, TEXT("Enable/Disable collisions in Immediate Physics."));
 FAutoConsoleVariableRef CVarChaosImmPhysCollisionDisableApply(TEXT("p.Chaos.ImmPhys.CollisionEnableApply"), ChaosImmediate_Collision_ApplyEnabled, TEXT("Enable/Disable the Apply() (velocity, friction, restitution) method for collisions"));
-FAutoConsoleVariableRef CVarChaosImmPhysCollisionPushOutIterations(TEXT("p.Chaos.ImmPhys.CollisionPushOutIterations"), ChaosImmediate_Collision_PushOutIterations, TEXT("Set the ApplyPushOut() (position correction) iteration count (0 to disable)"));
 FAutoConsoleVariableRef CVarChaosImmPhysCollisionPushOutPairIterations(TEXT("p.Chaos.ImmPhys.CollisionPushOutPairIterations"), ChaosImmediate_Collision_PushOutPairIterations, TEXT("Set the ApplyPushOut() internal pair interations (position correction) iteration count (0 to disable)"));
+FAutoConsoleVariableRef CVarChaosImmPhysCollisionPriority(TEXT("p.Chaos.ImmPhys.CollisionPriority"), ChaosImmediate_Collision_Priority, TEXT("Set the Collision constraint sort order (Joints have priority 0)"));
 FAutoConsoleVariableRef CVarChaosImmPhysThickness(TEXT("p.Chaos.ImmPhys.CollisionThickness"), ChaosImmediate_Collision_Thickness, TEXT("ChaosImmediateThickness"));
 
 float ChaosImmediate_Joint_SwingTwistAngleTolerance = 1.0e-6f;
@@ -419,7 +421,6 @@ namespace ImmediatePhysics_Chaos
 			JointsSettings.bEnableTwistLimits = ChaosImmediate_Joint_EnableTwistLimits != 0;
 			JointsSettings.bEnableSwingLimits = ChaosImmediate_Joint_EnableSwingLimits != 0;
 			JointsSettings.bEnableDrives = ChaosImmediate_Joint_EnableDrives != 0;
-			JointsSettings.PositionIterations = ChaosImmediate_Joint_PositionIterations;
 			JointsSettings.Projection = ChaosImmediate_Joint_Projection;
 			JointsSettings.Stiffness = ChaosImmediate_Joint_Stiffness;
 			JointsSettings.SoftLinearStiffness = ChaosImmediate_Joint_SoftLinearStiffness;
@@ -428,11 +429,12 @@ namespace ImmediatePhysics_Chaos
 			Joints->SetSettings(JointsSettings);
 
 			Evolution->SetNumIterations(ChaosImmediate_Evolution_Iterations);
+			Evolution->SetNumPushOutIterations(ChaosImmediate_Evolution_PushOutIterations);
 
 			Evolution->GetCollisionConstraints().SetThickness(ChaosImmediate_Collision_Thickness);
 			Evolution->GetCollisionConstraints().SetVelocitySolveEnabled(ChaosImmediate_Collision_ApplyEnabled != 0);
 			Evolution->GetCollisionConstraints().SetPushOutPairIterations(ChaosImmediate_Collision_PushOutPairIterations);
-			Evolution->GetCollisionConstraintsRule().SetPushOutIterations(ChaosImmediate_Collision_PushOutIterations);
+			Evolution->GetCollisionConstraintsRule().SetPriority(ChaosImmediate_Collision_Priority);
 		
 			// TEMP until we can remove constraints again, or I add broad-phase filtering. (FilterCollisionConstraints will crash since the persistent collision changes)
 			Evolution->GetCollisionConstraints().SetCollisionsEnabled(ChaosImmediate_Collision_Enabled != 0);
diff --git a/Engine/Source/Runtime/Experimental/Chaos/Private/Chaos/PBD6DJointConstraints.cpp b/Engine/Source/Runtime/Experimental/Chaos/Private/Chaos/PBD6DJointConstraints.cpp
--- a/Engine/Source/Runtime/Experimental/Chaos/Private/Chaos/PBD6DJointConstraints.cpp
+++ b/Engine/Source/Runtime/Experimental/Chaos/Private/Chaos/PBD6DJointConstraints.cpp
@@ -1372,12 +1372,17 @@ void TPBD6DJointConstraints<T, d>::Apply(const T Dt, const TArray<FConstraintHan
 }
 
 template<class T, int d>
-void TPBD6DJointConstraints<T, d>::ApplyPushOut(const T Dt, const TArray<FConstraintHandle*>& InConstraintHandles)
+bool TPBD6DJointConstraints<T, d>::ApplyPushOut(const T Dt, const TArray<FConstraintHandle*>& InConstraintHandles, const int32 It, const int32 NumIts)
 {
+	// @todo(ccaulfield): track whether we are sufficiently solved
+	bool bNeedsAnotherIteration = true;
+
 	for (FConstraintHandle* ConstraintHandle : InConstraintHandles)
 	{
 		ApplyPushOutSingle(Dt, ConstraintHandle->GetConstraintIndex());
 	}
+
+	return bNeedsAnotherIteration;
 }
 
 template<class T, int d>
diff --git a/Engine/Source/Runtime/Experimental/Chaos/Private/Chaos/PBDJointConstraints.cpp b/Engine/Source/Runtime/Experimental/Chaos/Private/Chaos/PBDJointConstraints.cpp
--- a/Engine/Source/Runtime/Experimental/Chaos/Private/Chaos/PBDJointConstraints.cpp
+++ b/Engine/Source/Runtime/Experimental/Chaos/Private/Chaos/PBDJointConstraints.cpp
@@ -137,7 +137,6 @@ namespace Chaos
 		, DriveStiffness((T)0)
 		, SoftLinearStiffness((T)0)
 		, SoftAngularStiffness((T)0)
-		, PositionIterations((T)0)
 	{
 	}
 
@@ -377,8 +376,11 @@ namespace Chaos
 	}
 
 	template<class T, int d>
-	void TPBDJointConstraints<T, d>::ApplyPushOut(const T Dt, const TArray<FConstraintHandle*>& InConstraintHandles)
+	bool TPBDJointConstraints<T, d>::ApplyPushOut(const T Dt, const TArray<FConstraintHandle*>& InConstraintHandles, const int32 It, const int32 NumIts)
 	{
+		// @todo(ccaulfield): track whether we are sufficiently solved
+		bool bNeedsAnotherIteration = true;
+
 		if (Settings.bEnableVelocitySolve)
 		{
 			TArray<FConstraintHandle*> SortedConstraintHandles = InConstraintHandles;
@@ -388,17 +390,19 @@ namespace Chaos
 					return L.GetConstraintLevel() < R.GetConstraintLevel();
 				});
 
-			for (int32 It = 0; It < Settings.PositionIterations; ++It)
+			for (FConstraintHandle* ConstraintHandle : SortedConstraintHandles)
 			{
-				for (FConstraintHandle* ConstraintHandle : SortedConstraintHandles)
-				{
-					SolvePosition(Dt, ConstraintHandle->GetConstraintIndex(), It, Settings.PositionIterations);
-				}
+				SolvePosition(Dt, ConstraintHandle->GetConstraintIndex(), It, NumIts);
 			}
 		}
 
-		// @todo(ccaulfield): should be called constraint rule
-		ApplyProjection(Dt, InConstraintHandles);
+		if (It == NumIts - 1)
+		{
+			// @todo(ccaulfield): should be called by constraint rule after PushOut
+			ApplyProjection(Dt, InConstraintHandles);
+		}
+
+		return bNeedsAnotherIteration;
 	}
 
 	template<class T, int d>
diff --git a/Engine/Source/Runtime/Experimental/Chaos/Private/Chaos/PBDRigidsEvolution.cpp b/Engine/Source/Runtime/Experimental/Chaos/Private/Chaos/PBDRigidsEvolution.cpp
--- a/Engine/Source/Runtime/Experimental/Chaos/Private/Chaos/PBDRigidsEvolution.cpp
+++ b/Engine/Source/Runtime/Experimental/Chaos/Private/Chaos/PBDRigidsEvolution.cpp
@@ -118,11 +118,12 @@ namespace Chaos
 	};
 
 	template<class FPBDRigidsEvolution, class FPBDCollisionConstraint, class T, int d>
-	TPBDRigidsEvolutionBase<FPBDRigidsEvolution, FPBDCollisionConstraint, T, d>::TPBDRigidsEvolutionBase(TPBDRigidsSOAs<T, d>& InParticles, int32 InNumIterations)
+	TPBDRigidsEvolutionBase<FPBDRigidsEvolution, FPBDCollisionConstraint, T, d>::TPBDRigidsEvolutionBase(TPBDRigidsSOAs<T, d>& InParticles, int32 InNumIterations, int32 InNumPushOutIterations)
 		: Particles(InParticles)
 		, bExternalReady(false)
 		, Clustering(static_cast<FPBDRigidsEvolution&>(*this), Particles.GetClusteredParticles())
 		, NumIterations(InNumIterations)
+		, NumPushOutIterations(InNumPushOutIterations)
 		, SpatialCollectionFactory(new TDefaultCollectionFactory<T, d>())
 	{
 		Particles.GetParticleHandles().AddArray(&PhysicsMaterials);
diff --git a/Engine/Source/Runtime/Experimental/Chaos/Public/Chaos/PBD6DJointConstraints.h b/Engine/Source/Runtime/Experimental/Chaos/Public/Chaos/PBD6DJointConstraints.h
--- a/Engine/Source/Runtime/Experimental/Chaos/Public/Chaos/PBD6DJointConstraints.h
+++ b/Engine/Source/Runtime/Experimental/Chaos/Public/Chaos/PBD6DJointConstraints.h
@@ -269,8 +269,7 @@ namespace Chaos
 
 		CHAOS_API void Apply(const T Dt, const TArray<FConstraintHandle*>& InConstraintHandles, const int32 It, const int32 NumIts);
 
-		// @todo(ccaulfield): remove  this
-		CHAOS_API void ApplyPushOut(const T Dt, const TArray<FConstraintHandle*>& InConstraintHandles);
+		CHAOS_API bool ApplyPushOut(const T Dt, const TArray<FConstraintHandle*>& InConstraintHandles, const int32 It, const int32 NumIts);
 
 	protected:
 		using Base::GetConstraintIndex;
diff --git a/Engine/Source/Runtime/Experimental/Chaos/Public/Chaos/PBDCollisionConstraintPGS.h b/Engine/Source/Runtime/Experimental/Chaos/Public/Chaos/PBDCollisionConstraintPGS.h
--- a/Engine/Source/Runtime/Experimental/Chaos/Public/Chaos/PBDCollisionConstraintPGS.h
+++ b/Engine/Source/Runtime/Experimental/Chaos/Public/Chaos/PBDCollisionConstraintPGS.h
@@ -34,7 +34,7 @@ class TPBDCollisionConstraintPGS
 	void ComputeConstraints(const TPBDRigidParticles<T, d>& InParticles, const TArray<int32>& InIndices, const T Dt);
 
 	void Apply(TPBDRigidParticles<T, d>& InParticles, const T Dt, const TArray<int32>& InConstraintIndices);
-	void ApplyPushOut(TPBDRigidParticles<T, d>& InParticles, const T Dt, const TArray<int32>& InConstraintIndices);
+	bool ApplyPushOut(TPBDRigidParticles<T, d>& InParticles, const T Dt, const TArray<int32>& InConstraintIndices);
 	template<class T_PARTICLES> void Solve(T_PARTICLES& InParticles, const T Dt, const TArray<int32>& InConstraintIndices);
 
 	void RemoveConstraints(const TSet<uint32>& RemovedParticles);
diff --git a/Engine/Source/Runtime/Experimental/Chaos/Public/Chaos/PBDConstraintRule.h b/Engine/Source/Runtime/Experimental/Chaos/Public/Chaos/PBDConstraintRule.h
--- a/Engine/Source/Runtime/Experimental/Chaos/Public/Chaos/PBDConstraintRule.h
+++ b/Engine/Source/Runtime/Experimental/Chaos/Public/Chaos/PBDConstraintRule.h
@@ -59,10 +59,10 @@ namespace Chaos
 		virtual void UpdatePositionBasedState(const T Dt) {}
 
 		/** Apply all corrections for constraints in the specified island */
-		virtual void ApplyConstraints(const T Dt, int32 Island, const int32 It, const int32 NumIts) {};
+		virtual void ApplyConstraints(const T Dt, int32 Island, const int32 It, const int32 NumIts) {}
 
 		/** Apply push out for constraints in the specified island */
-		virtual void ApplyPushOut(const T Dt, int32 Island) {};
+		virtual bool ApplyPushOut(const T Dt, int32 Island, const int32 It, const int32 NumIts) { return false; }
 
 		/** Add all constraints to the connectivity graph */
 		virtual void AddToGraph() {}
@@ -155,12 +155,13 @@ namespace Chaos
 			}
 		}
 
-		virtual void ApplyPushOut(const T Dt, int32 Island) override
+		virtual bool ApplyPushOut(const T Dt, int32 Island, const int32 It, const int32 NumIts) override
 		{
 			if (IslandConstraintLists[Island].Num())
 			{
-				Constraints.ApplyPushOut(Dt, GetIslandConstraints(Island));
+				return Constraints.ApplyPushOut(Dt, GetIslandConstraints(Island), It, NumIts);
 			}
+			return false;
 		}
 
 		virtual void InitializeAccelerationStructures() override
@@ -274,56 +275,56 @@ namespace Chaos
 			Constraints.RemoveConstraints(InConstraints);
 		}
 
-		virtual void ApplyPushOut(const T Dt, int32 Island) override
+		virtual bool ApplyPushOut(const T Dt, int32 Island, const int32 It, const int32 NumIts) override
 		{
 			const typename FConstraintColor::FLevelToColorToConstraintListMap& LevelToColorToConstraintListMap = GraphColor.GetIslandLevelToColorToConstraintListMap(Island);
 			int32 MaxColor = GraphColor.GetIslandMaxColor(Island);
 			int32 MaxLevel = GraphColor.GetIslandMaxLevel(Island);
 
-			TSet<TGeometryParticleHandle<T,d>*> IsTemporarilyStatic;	// Also needs to be per-constraint type
-			bool bNeedsAnotherIteration = true;
-			for (int32 Iteration = 0; bNeedsAnotherIteration && (Iteration < PushOutIterations); ++Iteration)
+			TSet<TGeometryParticleHandle<T, d>*> IsTemporarilyStatic;
+			bool bNeedsAnotherIteration = false;
+			for (int32 Level = 0; Level <= MaxLevel; ++Level)
 			{
-				bNeedsAnotherIteration = false;
-				for (int32 Level = 0; Level <= MaxLevel; ++Level)
+				for (int32 Color = 0; Color <= MaxColor; ++Color)
 				{
-					for (int32 Color = 0; Color <= MaxColor; ++Color)
+					if (LevelToColorToConstraintListMap[Level].Contains(Color) && LevelToColorToConstraintListMap[Level][Color].Num())
 					{
-						if (LevelToColorToConstraintListMap[Level].Contains(Color) && LevelToColorToConstraintListMap[Level][Color].Num())
+						const TArray<typename FConstraints::FConstraintHandle*>& ConstraintHandles = GetLevelColorConstraints(LevelToColorToConstraintListMap, Level, Color);
+						if (Constraints.ApplyPushOut(Dt, ConstraintHandles, IsTemporarilyStatic, It, NumIts))
 						{
-							const TArray<typename FConstraints::FConstraintHandle*>& ConstraintHandles = GetLevelColorConstraints(LevelToColorToConstraintListMap, Level, Color);
-							if (Constraints.ApplyPushOut(Dt, ConstraintHandles, IsTemporarilyStatic, Iteration, PushOutIterations))
-							{
-								bNeedsAnotherIteration = true;
-							}
+							bNeedsAnotherIteration = true;
 						}
 					}
+				}
+
+				// @todo(ccaulfield): Move shock propagation out of color rule
 #if USE_SHOCK_PROPOGATION
-					for (int32 Color = 0; Color <= MaxColor; ++Color)
+				for (int32 Color = 0; Color <= MaxColor; ++Color)
+				{
+					if (LevelToColorToConstraintListMap[Level].Contains(Color))
 					{
-						if (LevelToColorToConstraintListMap[Level].Contains(Color))
+						for (int32 Edge = 0; Edge < LevelToColorToConstraintListMap[Level][Color].Num(); ++Edge)
 						{
-							for (int32 Edge = 0; Edge < LevelToColorToConstraintListMap[Level][Color].Num(); ++Edge)
+							const int32 ConstraintIndex = LevelToColorToConstraintListMap[Level][Color][Edge]->GetConstraintIndex();
+							const TVector<TGeometryParticleHandle<T,d>*, 2> Particles = Constraints.ConstraintParticles(ConstraintIndex);
+							if (It == NumIts - 1)
 							{
-								const int32 ConstraintIndex = LevelToColorToConstraintListMap[Level][Color][Edge]->GetConstraintIndex();
-								const TVector<TGeometryParticleHandle<T,d>*, 2> Particles = Constraints.ConstraintParticles(ConstraintIndex);
-								if (Iteration == PushOutIterations - 1)
+								if (Particles[0]->AsDynamic() == nullptr || IsTemporarilyStatic.Contains(Particles[0]))
+								{
+									IsTemporarilyStatic.Add(Particles[1]);
+								}
+								else if (Particles[1]->AsDynamic() == nullptr || IsTemporarilyStatic.Contains(Particles[1]))
 								{
-									if (Particles[0]->AsDynamic() == nullptr || IsTemporarilyStatic.Contains(Particles[0]))
-									{
-										IsTemporarilyStatic.Add(Particles[1]);
-									}
-									else if (Particles[1]->AsDynamic() == nullptr || IsTemporarilyStatic.Contains(Particles[1]))
-									{
-										IsTemporarilyStatic.Add(Particles[0]);
-									}
+									IsTemporarilyStatic.Add(Particles[0]);
 								}
 							}
 						}
 					}
-#endif
 				}
+#endif
 			}
+
+			return bNeedsAnotherIteration;
 		}
 
 		virtual void InitializeAccelerationStructures() override
diff --git a/Engine/Source/Runtime/Experimental/Chaos/Public/Chaos/PBDJointConstraintTypes.h b/Engine/Source/Runtime/Experimental/Chaos/Public/Chaos/PBDJointConstraintTypes.h
--- a/Engine/Source/Runtime/Experimental/Chaos/Public/Chaos/PBDJointConstraintTypes.h
+++ b/Engine/Source/Runtime/Experimental/Chaos/Public/Chaos/PBDJointConstraintTypes.h
@@ -133,6 +133,5 @@ namespace Chaos
 		T DriveStiffness;
 		T SoftLinearStiffness;
 		T SoftAngularStiffness;
-		int32 PositionIterations;
 	};
 }
\ No newline at end of file
diff --git a/Engine/Source/Runtime/Experimental/Chaos/Public/Chaos/PBDJointConstraints.h b/Engine/Source/Runtime/Experimental/Chaos/Public/Chaos/PBDJointConstraints.h
--- a/Engine/Source/Runtime/Experimental/Chaos/Public/Chaos/PBDJointConstraints.h
+++ b/Engine/Source/Runtime/Experimental/Chaos/Public/Chaos/PBDJointConstraints.h
@@ -125,7 +125,7 @@ namespace Chaos
 
 		CHAOS_API void Apply(const T Dt, const TArray<FConstraintHandle*>& InConstraintHandles, const int32 It, const int32 NumIts);
 
-		CHAOS_API void ApplyPushOut(const T Dt, const TArray<FConstraintHandle*>& InConstraintHandles);
+		CHAOS_API bool ApplyPushOut(const T Dt, const TArray<FConstraintHandle*>& InConstraintHandles, const int32 It, const int32 NumIts);
 
 		CHAOS_API void ApplyProjection(const T Dt, const TArray<FConstraintHandle*>& InConstraintHandles);
 
diff --git a/Engine/Source/Runtime/Experimental/Chaos/Public/Chaos/PBDPositionConstraints.h b/Engine/Source/Runtime/Experimental/Chaos/Public/Chaos/PBDPositionConstraints.h
--- a/Engine/Source/Runtime/Experimental/Chaos/Public/Chaos/PBDPositionConstraints.h
+++ b/Engine/Source/Runtime/Experimental/Chaos/Public/Chaos/PBDPositionConstraints.h
@@ -164,8 +164,9 @@ namespace Chaos
 			}
 		}
 
-		void ApplyPushOut(const T Dt, const TArray<FConstraintHandle*>& InConstraintIndices) const
+		bool ApplyPushOut(const T Dt, const TArray<FConstraintHandle*>& InConstraintIndices, const int32 It, const int32 NumIts) const
 		{
+			return false;
 		}
 
 	protected:
diff --git a/Engine/Source/Runtime/Experimental/Chaos/Public/Chaos/PBDRigidDynamicSpringConstraints.h b/Engine/Source/Runtime/Experimental/Chaos/Public/Chaos/PBDRigidDynamicSpringConstraints.h
--- a/Engine/Source/Runtime/Experimental/Chaos/Public/Chaos/PBDRigidDynamicSpringConstraints.h
+++ b/Engine/Source/Runtime/Experimental/Chaos/Public/Chaos/PBDRigidDynamicSpringConstraints.h
@@ -170,9 +170,9 @@ namespace Chaos
 			}
 		}
 
-		// @todo(ccaulfield): remove  this
-		void ApplyPushOut(const T Dt, const TArray<FConstraintHandle*>& InConstraintHandles)
+		bool ApplyPushOut(const T Dt, const TArray<FConstraintHandle*>& InConstraintHandles, const int32 It, const int32 NumIts)
 		{
+			return false;
 		}
 
 	protected:
diff --git a/Engine/Source/Runtime/Experimental/Chaos/Public/Chaos/PBDRigidSpringConstraints.h b/Engine/Source/Runtime/Experimental/Chaos/Public/Chaos/PBDRigidSpringConstraints.h
--- a/Engine/Source/Runtime/Experimental/Chaos/Public/Chaos/PBDRigidSpringConstraints.h
+++ b/Engine/Source/Runtime/Experimental/Chaos/Public/Chaos/PBDRigidSpringConstraints.h
@@ -176,9 +176,9 @@ namespace Chaos
 			}
 		}
 
-		// @todo(ccaulfield): remove  this
-		void ApplyPushOut(const T Dt, const TArray<FConstraintHandle*>& InConstraintHandles)
+		bool ApplyPushOut(const T Dt, const TArray<FConstraintHandle*>& InConstraintHandles, const int32 It, const int32 NumIts)
 		{
+			return false;
 		}
 
 	protected:
diff --git a/Engine/Source/Runtime/Experimental/Chaos/Public/Chaos/PBDRigidsEvolution.h b/Engine/Source/Runtime/Experimental/Chaos/Public/Chaos/PBDRigidsEvolution.h
--- a/Engine/Source/Runtime/Experimental/Chaos/Public/Chaos/PBDRigidsEvolution.h
+++ b/Engine/Source/Runtime/Experimental/Chaos/Public/Chaos/PBDRigidsEvolution.h
@@ -219,7 +219,7 @@ class TPBDRigidsEvolutionBase
 	typedef TFunction<void(const TParticleView<TPBDRigidParticles<T, d>>&, const T)> FUpdatePositionRule;
 	typedef TFunction<void(TPBDRigidParticles<T, d>&, const T, const T, const int32)> FKinematicUpdateRule;
 
-	CHAOS_API TPBDRigidsEvolutionBase(TPBDRigidsSOAs<T, d>& InParticles, int32 InNumIterations = 1);
+	CHAOS_API TPBDRigidsEvolutionBase(TPBDRigidsSOAs<T, d>& InParticles, int32 InNumIterations = 1, int32 InNumPushOutIterations = 1);
 	CHAOS_API virtual ~TPBDRigidsEvolutionBase();
 
 	CHAOS_API TArray<TGeometryParticleHandle<T, d>*> CreateStaticParticles(int32 NumParticles, const TGeometryParticleParameters<T, d>& Params = TGeometryParticleParameters<T, d>())
@@ -287,6 +287,11 @@ class TPBDRigidsEvolutionBase
 		NumIterations = InNumIterations;
 	}
 
+	CHAOS_API void SetNumPushOutIterations(int32 InNumIterations)
+	{
+		NumPushOutIterations = InNumIterations;
+	}
+
 	CHAOS_API void EnableParticle(TGeometryParticleHandle<T,d>* Particle, const TGeometryParticleHandle<T, d>* ParentParticle)
 	{
 		DirtyParticle(*Particle);
@@ -411,6 +416,7 @@ class TPBDRigidsEvolutionBase
 			ConstraintRule->UpdateAccelerationStructures(Island);
 		}
 
+		// @todo(ccaulfield): track whether we are sufficiently solved and can early-out
 		for (int i = 0; i < NumIterations; ++i)
 		{
 			for (FConstraintRule* ConstraintRule : ConstraintRules)
@@ -571,9 +577,17 @@ protected:
 
 	void ApplyPushOut(const T Dt, int32 Island)
 	{
-		for (FConstraintRule* ConstraintRule : ConstraintRules)
+		bool bNeedsAnotherIteration = true;
+		for (int32 It = 0; bNeedsAnotherIteration && (It < NumPushOutIterations); ++It)
 		{
-			ConstraintRule->ApplyPushOut(Dt, Island);
+			bNeedsAnotherIteration = false;
+			for (FConstraintRule* ConstraintRule : ConstraintRules)
+			{
+				if (ConstraintRule->ApplyPushOut(Dt, Island, It, NumPushOutIterations))
+				{
+					bNeedsAnotherIteration = true;
+				}
+			}
 		}
 	}
 
@@ -681,6 +695,7 @@ protected:
 	FGraphEventRef AccelerationStructureTaskComplete;
 
 	int32 NumIterations;
+	int32 NumPushOutIterations;
 	TUniquePtr<ISpatialAccelerationCollectionFactory<T, d>> SpatialCollectionFactory;
 };
 
