Chaos : Convex Collision Manifold

Solves the manifold generation for initial collision contact. The manifold still needs to be updated for resting contact.The generation is based on GJKDistance, and will generate a disk (R2) about the collision point. The manifold can be temporarily tuned with pVars:
p.{ CollisionManifoldSamples, CollisionManifoldSampleScale, CollisionVelocityInflation}

 ... Convex-Convex collision manifold generation
... Renamed Create*Constraint to Construct*Constraints
... Double the velocity expansion of the narrow phase.
... ConstraintsMap added to summarize collision pairs that share multiple constraints.
... Removed core shape intersection.
... Added Implicit knowledge to the constraints.
... Added union-union traversal to the constraint creation.
... Updated Joint Unit tests to reflect changes from the velocity expansion.
... Removed Disable callback in ApplyConstraitns.
... Added pVar to control num collision iterations.

Todo:
... Snap the manifold to the collision surface, until source geometry information is accessible
... Thicken surface to keep GJKDistance valid
... Dynamically update manifold during resting contact.

diff --git a/Engine/Source/Runtime/Experimental/Chaos/Private/Chaos/PBDCollisionConstraint.cpp b/Engine/Source/Runtime/Experimental/Chaos/Private/Chaos/PBDCollisionConstraint.cpp
--- a/Engine/Source/Runtime/Experimental/Chaos/Private/Chaos/PBDCollisionConstraint.cpp
+++ b/Engine/Source/Runtime/Experimental/Chaos/Private/Chaos/PBDCollisionConstraint.cpp
@@ -27,19 +27,25 @@
 int32 CollisionParticlesBVHDepth = 4;
 FAutoConsoleVariableRef CVarCollisionParticlesBVHDepth(TEXT("p.CollisionParticlesBVHDepth"), CollisionParticlesBVHDepth, TEXT("The maximum depth for collision particles bvh"));
 
-
 int32 ConstraintBPBVHDepth = 2;
 FAutoConsoleVariableRef CVarConstraintBPBVHDepth(TEXT("p.ConstraintBPBVHDepth"), ConstraintBPBVHDepth, TEXT("The maximum depth for constraint bvh"));
 
 int32 BPTreeOfGrids = 1;
 FAutoConsoleVariableRef CVarBPTreeOfGrids(TEXT("p.BPTreeOfGrids"), BPTreeOfGrids, TEXT("Whether to use a seperate tree of grids for bp"));
 
-bool bSupportCollisionHistoryCVar = false;
-FAutoConsoleVariableRef CVarSupportCollisionHistory(TEXT("p.SupportCollisionHistory"), bSupportCollisionHistoryCVar, TEXT("Whether to support collision history caching."));
-
 int32 CollisionHistoryLifespanCVar = 1;
 FAutoConsoleVariableRef CVarCollisionHistoryLifespan(TEXT("p.CollisionHistoryLifespan"), CollisionHistoryLifespanCVar, TEXT("Number of iterations to cache a collision point.[def:5]"));
 
+int32 CollisionManifoldSamplesCVar = 5;
+FAutoConsoleVariableRef CVarCollisionManifoldSamples(TEXT("p.CollisionManifoldSamples"), CollisionManifoldSamplesCVar, TEXT("Number of sample points on the manifold.[def:6]"));
+
+float CollisionManifoldSampleScaleCVar = 0.7;
+FAutoConsoleVariableRef CVarCollisionManifoldSampleScale(TEXT("p.CollisionManifoldSampleScale"), CollisionManifoldSampleScaleCVar, TEXT("Scale factor for the collision manifold.[def:0.7]"));
+
+float CollisionVelocityInflationCVar = 2.0;
+FAutoConsoleVariableRef CVarCollisionVelocityInflation(TEXT("p.CollisionVelocityInflation"), CollisionVelocityInflationCVar, TEXT("Collision velocity inflation.[def:2.0]"));
+
+
 extern int32 UseLevelsetCollision;
 
 #if !UE_BUILD_SHIPPING
@@ -55,7 +61,10 @@ CHAOS_API int32 EnableCollisions = 1;
 FAutoConsoleVariableRef CVarEnableCollisions(TEXT("p.EnableCollisions"), EnableCollisions, TEXT("Enable/Disable collisions on the Chaos solver."));
 
 template<ECollisionUpdateType, typename T, int d>
-void UpdateConstraintImp2(const TRigidTransform<T, d>& ParticleTM, const FImplicitObject& LevelsetObject, const TRigidTransform<T, d>& LevelsetTM, const T Thickness, TRigidBodyContactConstraint<T, d>& Constraint);
+void UpdateConstraintImp2(const TRigidTransform<T, d>& ParticleTM, const FImplicitObject& LevelsetObject, const TRigidTransform<T, d>& LevelsetTM, const T Thickness, TRigidBodyContactConstraint<T, d>& Constraint, int32 NumConstraints, TPBDCollisionConstraintHistory<T, d>* History);
+
+template<typename T, int d>
+void ConstructConstraintsImpl(TGeometryParticleHandle<T, d>* Particle0, TGeometryParticleHandle<T, d>* Particle1, const FImplicitObject* Implicit0, const FImplicitObject* Implicit1, const T Thickness, TArray<typename TPBDCollisionConstraint<T, d>::FRigidBodyContactConstraint>& ConstraintBuffer, TPBDCollisionConstraintHistory<T, d>* History);
 
 //
 // Collision Constraint Handle
@@ -72,6 +81,12 @@ TPBDCollisionConstraintHandle<T, d>::TPBDCollisionConstraintHandle(FConstraintCo
 {
 }
 
+template<typename T, int d>
+TRigidBodyContactConstraint<T, d>& TPBDCollisionConstraintHandle<T, d>::GetContact()
+{
+	return ConstraintContainer->Constraints[ConstraintIndex];
+}
+
 template<typename T, int d>
 const TRigidBodyContactConstraint<T, d>& TPBDCollisionConstraintHandle<T, d>::GetContact() const
 {
@@ -96,10 +111,11 @@ TPBDCollisionConstraint<T, d>::TPBDCollisionConstraint(const TPBDRigidsSOAs<T,d>
 	, bUseCCD(false)
 	, bEnableCollisions(true)
 	, LifespanCounter(0)
+	, CollisionHistoryLifespan(CollisionHistoryLifespanCVar)
+	, CollisionVelocityInflation(CollisionVelocityInflationCVar)
 	, PostComputeCallback(nullptr)
 	, PostApplyCallback(nullptr)
 	, PostApplyPushOutCallback(nullptr)
-	, bSupportHistory(bSupportCollisionHistoryCVar)
 {
 }
 
@@ -110,15 +126,41 @@ void TPBDCollisionConstraint<T, d>::Reset(/*const TPBDRigidParticles<T, d>& InPa
 {
 	SCOPE_CYCLE_COUNTER(STAT_CollisionConstraintsReset);
 
-	int32 Threshold = LifespanCounter - CollisionHistoryLifespanCVar; // Maybe this should be solver time?
 	for (int32 Idx = Constraints.Num() - 1; Idx >= 0; Idx--)
 	{
-		if ((Constraints[Idx].Lifespan < Threshold) || !bEnableCollisions)
+		if (!bEnableCollisions)
 		{
 			RemoveConstraint(Idx);
 		}
 	}
 
+	TArray<FConstraintHandleID> HistoryKeys;
+	History.GetKeys(HistoryKeys);
+	for (FConstraintHandleID Key : HistoryKeys)
+	{
+		FConstraintHistory* Elem = History[Key];
+		if (LifespanCounter > Elem->GetTimestamp())
+		{
+			bool bInContact = false;
+			for (FConstraintHandle* Handle : Elem->GetHandles())
+			{
+				if (Handle->GetContact().Phi <= 0)
+				{
+					bInContact = true;
+				}
+			}
+
+			if (!bInContact)
+			{
+				const TArray<FConstraintHandle*> HandleArray = Elem->GetHandles();
+				for (FConstraintHandle* Handle : HandleArray)
+				{
+					RemoveConstraint(Handle->GetConstraintIndex());
+				}
+			}
+		}
+	}
+
 	MAngularFriction = 0;
 	bUseCCD = false;
 }
@@ -126,13 +168,17 @@ void TPBDCollisionConstraint<T, d>::Reset(/*const TPBDRigidParticles<T, d>& InPa
 template<typename T, int d>
 void TPBDCollisionConstraint<T, d>::RemoveConstraint(int32 Idx)
 {
-	if (bSupportHistory)
+	// History
+	FConstraintHandleID HandleID = GetConstraintHandleID(Constraints[Idx]);
+	if (FConstraintHistory** Ptr = History.Find(HandleID))
 	{
-		FConstraintHandleID HandleID = GetConstraintHandleID(Constraints[Idx]);
-		History[HandleID]->RemoveHandle(Handles[Idx]);
-		if (!History[HandleID]->GetHandles().Num())
+		if (FConstraintHistory* HistoryEntry = *Ptr)
 		{
-			History.Remove(HandleID);
+			HistoryEntry->RemoveHandle(Handles[Idx]);
+			if (!HistoryEntry->GetHandles().Num())
+			{
+				History.Remove(HandleID);
+			}
 		}
 	}
 
@@ -550,7 +596,7 @@ void TPBDCollisionConstraint<T, d>::Apply(const T Dt, FRigidBodyContactConstrain
 		ensure(!PBDRigid0 || PBDRigid0->Sleeping());
 		return;
 	}
-	UpdateConstraint<ECollisionUpdateType::Deepest>(MThickness, Constraint);
+	UpdateConstraint<ECollisionUpdateType::Deepest>(MThickness, &Constraint, 1);
 	if (Constraint.Phi >= MThickness)
 	{
 		return;
@@ -800,7 +846,7 @@ void TPBDCollisionConstraint<T, d>::ApplyPushOut(const T Dt, FRigidBodyContactCo
 
 	for (int32 PairIteration = 0; PairIteration < MPairIterations; ++PairIteration)
 	{
-		UpdateConstraint<ECollisionUpdateType::Deepest>(MThickness, Constraint);
+		UpdateConstraint<ECollisionUpdateType::Deepest>(MThickness, &Constraint, 1);
 		if (Constraint.Phi >= MThickness)
 		{
 			break;
@@ -1610,12 +1656,13 @@ TArray<Pair<const FImplicitObject*, TRigidTransform<T, d>>> FindRelevantShapes2(
 
 DECLARE_CYCLE_STAT(TEXT("TPBDCollisionConstraint::UpdateUnionUnionConstraint"), STAT_UpdateUnionUnionConstraint, STATGROUP_ChaosWide);
 template<ECollisionUpdateType UpdateType, typename T, int d>
-void UpdateUnionUnionConstraint(const T Thickness, TRigidBodyContactConstraint<T, d>& Constraint)
+void UpdateUnionUnionConstraint(const T Thickness, TRigidBodyContactConstraint<T, d>* Constraints, int32 NumConstraints, TPBDCollisionConstraintHistory<T, d>* History)
 {
+	checkSlow(Constraints != nullptr && NumConstraints >= 1);
 	SCOPE_CYCLE_COUNTER(STAT_UpdateUnionUnionConstraint);
 
-	TGenericParticleHandle<T, d> Particle0 = Constraint.Particle;
-	TGenericParticleHandle<T, d> Particle1 = Constraint.Levelset;
+	TGenericParticleHandle<T, d> Particle0 = Constraints[0].Particle;
+	TGenericParticleHandle<T, d> Particle1 = Constraints[0].Levelset;
 
 	TRigidTransform<T, d> ParticlesTM = TRigidTransform<T, d>(Particle0->P(), Particle0->Q());
 	TRigidTransform<T, d> LevelsetTM = TRigidTransform<T, d>(Particle1->P(), Particle1->Q());
@@ -1637,19 +1684,20 @@ void UpdateUnionUnionConstraint(const T Thickness, TRigidBodyContactConstraint<T
 		{
 			const FImplicitObject& ParticleInnerObj = *ParticlePair.First;
 			const TRigidTransform<T, d> ParticleInnerObjTM = ParticlePair.Second * ParticlesTM;
-			UpdateConstraintImp2<UpdateType>(ParticleInnerObj, ParticleInnerObjTM, LevelsetInnerObj, LevelsetInnerObjTM, Thickness, Constraint);
+			UpdateConstraintImp2<UpdateType>(ParticleInnerObj, ParticleInnerObjTM, LevelsetInnerObj, LevelsetInnerObjTM, Thickness, Constraints, NumConstraints, History);
 		}
 	}
 }
 
 DECLARE_CYCLE_STAT(TEXT("TPBDCollisionConstraint::UpdateSingleUnionConstraint"), STAT_UpdateSingleUnionConstraint, STATGROUP_ChaosWide);
 template<ECollisionUpdateType UpdateType, typename T, int d>
-void UpdateSingleUnionConstraint(const T Thickness, TRigidBodyContactConstraint<T, d>& Constraint)
+void UpdateSingleUnionConstraint(const T Thickness, TRigidBodyContactConstraint<T, d>* Constraints, int32 NumConstraints, TPBDCollisionConstraintHistory<T, d>* History)
 {
+	checkSlow(Constraints != nullptr && NumConstraints >= 1);
 	SCOPE_CYCLE_COUNTER(STAT_UpdateSingleUnionConstraint);
 
-	TGenericParticleHandle<T, d> Particle0 = Constraint.Particle;
-	TGenericParticleHandle<T, d> Particle1 = Constraint.Levelset;
+	TGenericParticleHandle<T, d> Particle0 = Constraints[0].Particle;
+	TGenericParticleHandle<T, d> Particle1 = Constraints[0].Levelset;
 
 	TRigidTransform<T, d> ParticlesTM = TRigidTransform<T, d>(Particle0->P(), Particle0->Q());
 	TRigidTransform<T, d> LevelsetTM = TRigidTransform<T, d>(Particle1->P(), Particle1->Q());
@@ -1662,7 +1710,7 @@ void UpdateSingleUnionConstraint(const T Thickness, TRigidBodyContactConstraint<
 	{
 		const FImplicitObject& LevelsetInnerObj = *LevelsetObjPair.First;
 		const TRigidTransform<T, d> LevelsetInnerObjTM = LevelsetObjPair.Second * LevelsetTM;
-		UpdateConstraintImp2<UpdateType>(*ParticleObj, ParticlesTM, LevelsetInnerObj, LevelsetInnerObjTM, Thickness, Constraint);
+		UpdateConstraintImp2<UpdateType>(*ParticleObj, ParticlesTM, LevelsetInnerObj, LevelsetInnerObjTM, Thickness, Constraints, NumConstraints, History);
 	}
 }
 
@@ -1696,20 +1744,6 @@ void TPBDCollisionConstraint<T, d>::UpdateLevelsetConstraint(const T Thickness,
 	}
 }
 
-DECLARE_CYCLE_STAT(TEXT("UpdateConvexConstraintsUsingCoreShapes"), UpdateConvexConstraintsUsingCoreShapes, STATGROUP_ChaosWide);
-template<ECollisionUpdateType UpdateType, typename T, int d>
-void UpdateConvexConstraintsUsingCoreShapes(const FImplicitObject & AObj, const TRigidTransform<T, d>& ATM, const FImplicitObject & BObj, const TRigidTransform<T, d>& BTM, const T Thickness, TRigidBodyContactConstraint<T, d>& Constraint)
-{
-	SCOPE_CYCLE_COUNTER(UpdateConvexConstraintsUsingCoreShapes);
-	CastHelper(AObj, [&](const auto& ADowncast)
-	{
-		CastHelper(BObj, [&](const auto& BDowncast)
-		{
-			GJKCoreShapeIntersection<T, 3>(ADowncast, ATM, BDowncast, BTM, Constraint.Location, Constraint.Phi, Constraint.Normal, Thickness);
-		});
-	});
-}
-
 DECLARE_CYCLE_STAT(TEXT("TPBDCollisionConstraint::UpdateUnionLevelsetConstraint"), STAT_UpdateUnionLevelsetConstraint, STATGROUP_ChaosWide);
 template<ECollisionUpdateType UpdateType, typename T, int d>
 void UpdateUnionLevelsetConstraint(const T Thickness, TRigidBodyContactConstraint<T, d>& Constraint)
@@ -1885,8 +1919,95 @@ void UpdateBoxConstraint(const TBox<T, d>& Box1, const TRigidTransform<T, d>& Bo
 	}
 }
 
+DECLARE_CYCLE_STAT(TEXT("UpdateConvexConvexConstraint"), UpdateConvexConvexConstraint, STATGROUP_ChaosWide);
+template<ECollisionUpdateType UpdateType, typename T, int d, typename TGeometryA, typename TGeometryB>
+void UpdateConvexConvexConstraint(const TGeometryA& A, TRigidTransform<T, d> ATM, const TGeometryB& B, TRigidTransform<T, d> BTM, T Thickness, TRigidBodyContactConstraint<T, d>* Constraints, int32 NumConstraints, TPBDCollisionConstraintHistory<T, d>* History)
+{
+	SCOPE_CYCLE_COUNTER(UpdateConvexConvexConstraint);
+
+	ensure(NumConstraints);
+
+	unsigned char AType = A.GetType();
+	unsigned char AStaticType = ImplicitObjectType::IsScaled;
+
+	if (ensure(IsScaled(A.GetType()) && IsScaled(B.GetType())))
+	{
+		const TImplicitObjectScaled<TConvex<T, d>>& AScaled = static_cast<const TImplicitObjectScaled<TConvex<T, d>>&>(A);
+		const TImplicitObjectScaled<TConvex<T, d>>& BScaled = static_cast<const TImplicitObjectScaled<TConvex<T, d>>&>(B);
+
+		TRigidTransform<T, d> BToATM = BTM.GetRelativeTransform(ATM);
+		TRigidTransform<T, d> AToBTM = ATM.GetRelativeTransform(BTM);
+
+		if (ensure(GetInnerType(A.GetType()) == ImplicitObjectType::Convex && GetInnerType(B.GetType()) == ImplicitObjectType::Convex))
+		{
+			const TConvex<T, d>* AObject = static_cast<const TConvex<T, d>*>(AScaled.Object().Get());
+			const TConvex<T, d>* BObject = static_cast<const TConvex<T, d>*>(BScaled.Object().Get());
+
+			const TParticles<T, d>& SurfaceParticles = AObject->GetSurfaceParticles();
+
+			T OutDistance;
+			TVector<T, d> StubNormal;
+			TVector<T, d> OutNearestA;
+			TVector<T, d> OutNearestB;
+			TVector<T, d> OutNormal;
+
+			if (History == nullptr)
+			{
+				if (GJKDistance<T>(*AObject, *BObject, BToATM, OutDistance, OutNearestA, OutNearestB))
+				{
+					TVector<T, d> Normal(0);
+					TVector<T, d> WorldA = ATM.TransformPosition(OutNearestA);
+					TVector<T, d> WorldB = BTM.TransformPosition(OutNearestB);
+					TVector<T, d> CenterLocation = WorldA;
+					T Phi = BObject->PhiWithNormal(BTM.InverseTransformPosition(CenterLocation), Normal);
+
+					T ManifoldSize = A.BoundingBox().Extents().Size() * CollisionManifoldSampleScaleCVar;
+					TVector<T, d> CrossVector = FMath::IsNearlyEqual(FMath::Abs(TVector<T, d>::DotProduct(TVector<T, d>(0, 1, 0), Normal)), T(1)) ?
+						TVector<T, d>(0, 0, 1) :
+						TVector<T, d>(0, 1, 0);
+
+					Constraints[0].LocalLocation = ATM.InverseTransformPosition(CenterLocation);
+					Constraints[0].Normal = Normal;
+
+					if (NumConstraints > 1)
+					{
+						TVector<T, d> DirVector = CrossVector;
+						T Angle = T(360) / T(NumConstraints - 1), LocalPhi;
+						for (int32 Idx = 1; Idx < NumConstraints; Idx++)
+						{
+							CrossVector = CrossVector.RotateAngleAxis(Angle,Normal);
+							TVector<T, d> WorldLocation = CenterLocation + ManifoldSize * CrossVector;
+							Constraints[Idx].LocalLocation = ATM.InverseTransformPosition(WorldLocation);
+
+							LocalPhi = AObject->PhiWithNormal(Constraints[Idx].LocalLocation, StubNormal);
+							Constraints[Idx].LocalLocation -= StubNormal * LocalPhi;
+
+							Constraints[Idx].Normal = Normal;
+						}
+					}
+				}
+			}
+
+			for (int32 Idx = 0; Idx < NumConstraints; Idx++)
+			{
+				Constraints[Idx].Location = ATM.TransformPosition(Constraints[Idx].LocalLocation);
+				Constraints[Idx].Phi = BObject->PhiWithNormal(BTM.InverseTransformPosition(Constraints[Idx].Location), StubNormal);
+
+				if (Constraints[Idx].Phi < 0.f && History)
+				{
+					History->SetTimestamp(INT_MAX);
+				}
+			}
+
+			return;
+		}
+	}
+
+	ensureMsgf(false, TEXT("Unsupported convex to convex constraint."));
+}
+
 template<typename T, int d>
-TRigidBodyContactConstraint<T, d> ComputeLevelsetConstraint(TGeometryParticleHandle<T, d>* Particle0, TGeometryParticleHandle<T, d>* Particle1, const T Thickness)
+void ConstructLevelsetConstraints(TGeometryParticleHandle<T, d>* Particle0, TGeometryParticleHandle<T, d>* Particle1, const FImplicitObject* Implicit0, const FImplicitObject* Implicit1, const T Thickness, TArray<typename TPBDCollisionConstraint<T, d>::FRigidBodyContactConstraint>& ConstraintBuffer, TPBDCollisionConstraintHistory<T, d>* History = nullptr)
 {
 	//todo(ocohen):if both have collision particles, use the one with fewer?
 	if (!Particle1->Geometry() || (Particle0->AsDynamic() && !Particle0->AsDynamic()->CollisionParticlesSize() && Particle0->Geometry() && !Particle0->Geometry()->IsUnderlyingUnion()))
@@ -1894,178 +2015,248 @@ TRigidBodyContactConstraint<T, d> ComputeLevelsetConstraint(TGeometryParticleHan
 		TRigidBodyContactConstraint<T, d> Constraint;
 		Constraint.Particle = Particle1;
 		Constraint.Levelset = Particle0;
-		return Constraint;
+		Constraint.Geometry[0] = Implicit1;
+		Constraint.Geometry[1] = Implicit0;
+		ConstraintBuffer.Add(Constraint);
 	}
 	else
 	{
 		TRigidBodyContactConstraint<T, d> Constraint;
 		Constraint.Particle = Particle0;
 		Constraint.Levelset = Particle1;
-		return Constraint;
+		Constraint.Geometry[0] = Implicit0;
+		Constraint.Geometry[1] = Implicit1;
+		ConstraintBuffer.Add(Constraint);
 	}
 }
 
 template<typename T, int d>
-TRigidBodyContactConstraint<T, d> ComputeLevelsetConstraintGJK(TGeometryParticleHandle<T, d>* Particle0, TGeometryParticleHandle<T, d>* Particle1, const T Thickness)
+void ConstructConvexConvexConstraints(TGeometryParticleHandle<T, d>* Particle0, TGeometryParticleHandle<T, d>* Particle1, const FImplicitObject* Implicit0, const FImplicitObject* Implicit1, const T Thickness, TArray<typename TPBDCollisionConstraint<T, d>::FRigidBodyContactConstraint>& ConstraintBuffer, TPBDCollisionConstraintHistory<T, d>* History = nullptr)
 {
-	TRigidBodyContactConstraint<T, d> Constraint;
-	Constraint.Particle = Particle0;
-	Constraint.Levelset = Particle1;
-	return Constraint;
+	if (History)
+	{
+		if (History->ContainsShapeConnection(Implicit0, Implicit1))
+		{
+			return;
+		}
+	}
+
+	for (int i = 0; i < CollisionManifoldSamplesCVar; i++)
+	{
+		TRigidBodyContactConstraint<T, d> Constraint;
+		Constraint.Particle = Particle0;
+		Constraint.Levelset = Particle1;
+		Constraint.Geometry[0] = Implicit0;
+		Constraint.Geometry[1] = Implicit1;
+		ConstraintBuffer.Add(Constraint);
+	}
 }
 
 template<typename T, int d>
-TRigidBodyContactConstraint<T, d> ComputeBoxConstraint(TGeometryParticleHandle<T, d>* Particle0, TGeometryParticleHandle<T, d>* Particle1, const T Thickness)
+void ConstructBoxConstraints(TGeometryParticleHandle<T, d>* Particle0, TGeometryParticleHandle<T, d>* Particle1, const FImplicitObject* Implicit0, const FImplicitObject* Implicit1, const T Thickness, TArray<typename TPBDCollisionConstraint<T, d>::FRigidBodyContactConstraint>& ConstraintBuffer, TPBDCollisionConstraintHistory<T, d>* History = nullptr)
 {
 	TRigidBodyContactConstraint<T, d> Constraint;
 	Constraint.Particle = Particle0;
 	Constraint.Levelset = Particle1;
-	return Constraint;
+	Constraint.Geometry[0] = Implicit0;
+	Constraint.Geometry[1] = Implicit1;
+	ConstraintBuffer.Add(Constraint);
 }
 
 template<typename T, int d>
-TRigidBodyContactConstraint<T, d> ComputeBoxPlaneConstraint(TGeometryParticleHandle<T, d>* Particle0, TGeometryParticleHandle<T, d>* Particle1, const T Thickness)
+void ConstructBoxPlaneConstraints(TGeometryParticleHandle<T, d>* Particle0, TGeometryParticleHandle<T, d>* Particle1, const FImplicitObject* Implicit0, const FImplicitObject* Implicit1, const T Thickness, TArray<typename TPBDCollisionConstraint<T, d>::FRigidBodyContactConstraint>& ConstraintBuffer, TPBDCollisionConstraintHistory<T, d>* History = nullptr)
 {
 	TRigidBodyContactConstraint<T, d> Constraint;
 	Constraint.Particle = Particle0;
 	Constraint.Levelset = Particle1;
-	return Constraint;
+	Constraint.Geometry[0] = Implicit0;
+	Constraint.Geometry[1] = Implicit1;
+	ConstraintBuffer.Add(Constraint);
 }
 
 template<typename T, int d>
-TRigidBodyContactConstraint<T, d> ComputeSphereConstraint(TGeometryParticleHandle<T, d>* Particle0, TGeometryParticleHandle<T, d>* Particle1, const T Thickness)
+void ConstructSphereConstraints(TGeometryParticleHandle<T, d>* Particle0, TGeometryParticleHandle<T, d>* Particle1, const FImplicitObject* Implicit0, const FImplicitObject* Implicit1, const T Thickness, TArray<typename TPBDCollisionConstraint<T, d>::FRigidBodyContactConstraint>& ConstraintBuffer, TPBDCollisionConstraintHistory<T, d>* History = nullptr)
 {
 	TRigidBodyContactConstraint<T, d> Constraint;
 	Constraint.Particle = Particle0;
 	Constraint.Levelset = Particle1;
-	return Constraint;
+	Constraint.Geometry[0] = Implicit0;
+	Constraint.Geometry[1] = Implicit1;
+	ConstraintBuffer.Add(Constraint);
 }
 
 template<typename T, int d>
-TRigidBodyContactConstraint<T, d> ComputeSpherePlaneConstraint(TGeometryParticleHandle<T, d>* Particle0, TGeometryParticleHandle<T, d>* Particle1, const T Thickness)
+void ConstructSpherePlaneConstraints(TGeometryParticleHandle<T, d>* Particle0, TGeometryParticleHandle<T, d>* Particle1, const FImplicitObject* Implicit0, const FImplicitObject* Implicit1, const T Thickness, TArray<typename TPBDCollisionConstraint<T, d>::FRigidBodyContactConstraint>& ConstraintBuffer, TPBDCollisionConstraintHistory<T, d>* History = nullptr)
 {
 	TRigidBodyContactConstraint<T, d> Constraint;
 	Constraint.Particle = Particle0;
 	Constraint.Levelset = Particle1;
-	return Constraint;
+	Constraint.Geometry[0] = Implicit0;
+	Constraint.Geometry[1] = Implicit1;
+	ConstraintBuffer.Add(Constraint);
 }
 
 template<typename T, int d>
-TRigidBodyContactConstraint<T, d> ComputeSphereBoxConstraint(TGeometryParticleHandle<T, d>* Particle0, TGeometryParticleHandle<T, d>* Particle1, const T Thickness)
+void ConstructSphereBoxConstraints(TGeometryParticleHandle<T, d>* Particle0, TGeometryParticleHandle<T, d>* Particle1, const FImplicitObject* Implicit0, const FImplicitObject* Implicit1, const T Thickness, TArray<typename TPBDCollisionConstraint<T, d>::FRigidBodyContactConstraint>& ConstraintBuffer, TPBDCollisionConstraintHistory<T, d>* History = nullptr)
 {
 	TRigidBodyContactConstraint<T, d> Constraint;
 	Constraint.Particle = Particle0;
 	Constraint.Levelset = Particle1;
-	return Constraint;
+	Constraint.Geometry[0] = Implicit0;
+	Constraint.Geometry[1] = Implicit1;
+	ConstraintBuffer.Add(Constraint);
 }
 
 template<typename T, int d>
-TRigidBodyContactConstraint<T, d> ComputeCapsuleCapsuleConstraint(TGeometryParticleHandle<T, d>* Particle0, TGeometryParticleHandle<T, d>* Particle1, const T Thickness)
+void ConstructCapsuleCapsuleConstraints(TGeometryParticleHandle<T, d>* Particle0, TGeometryParticleHandle<T, d>* Particle1, const FImplicitObject* Implicit0, const FImplicitObject* Implicit1, const T Thickness, TArray<typename TPBDCollisionConstraint<T, d>::FRigidBodyContactConstraint>& ConstraintBuffer, TPBDCollisionConstraintHistory<T, d>* History = nullptr)
 {
 	TRigidBodyContactConstraint<T, d> Constraint;
 	Constraint.Particle = Particle0;
 	Constraint.Levelset = Particle1;
-	return Constraint;
+	Constraint.Geometry[0] = Implicit0;
+	Constraint.Geometry[1] = Implicit1;
+	ConstraintBuffer.Add(Constraint);
 }
 
 template<typename T, int d>
-TRigidBodyContactConstraint<T, d> ComputeCapsuleBoxConstraint(TGeometryParticleHandle<T, d>* Particle0, TGeometryParticleHandle<T, d>* Particle1, const T Thickness)
+void ConstructCapsuleBoxConstraints(TGeometryParticleHandle<T, d>* Particle0, TGeometryParticleHandle<T, d>* Particle1, const FImplicitObject* Implicit0, const FImplicitObject* Implicit1, const T Thickness, TArray<typename TPBDCollisionConstraint<T, d>::FRigidBodyContactConstraint>& ConstraintBuffer, TPBDCollisionConstraintHistory<T, d>* History = nullptr)
 {
 	TRigidBodyContactConstraint<T, d> Constraint;
 	Constraint.Particle = Particle0;
 	Constraint.Levelset = Particle1;
-	return Constraint;
+	Constraint.Geometry[0] = Implicit0;
+	Constraint.Geometry[1] = Implicit1;
+	ConstraintBuffer.Add(Constraint);
 }
 
-template <typename T, int d>
-TRigidBodyContactConstraint<T, d> ComputeSingleUnionConstraint(TGeometryParticleHandle<T, d>* Particle0, TGeometryParticleHandle<T, d>* Particle1, const T Thickness)
+template<typename T, int d>
+void ConstructSingleUnionConstraints(TGeometryParticleHandle<T, d>* Particle0, TGeometryParticleHandle<T, d>* Particle1, const FImplicitObject* Implicit0, const FImplicitObject* Implicit1, const T Thickness, TArray<typename TPBDCollisionConstraint<T, d>::FRigidBodyContactConstraint>& ConstraintBuffer, TPBDCollisionConstraintHistory<T, d>* History = nullptr)
 {
 	TRigidBodyContactConstraint<T, d> Constraint;
 	Constraint.Particle = Particle0;
 	Constraint.Levelset = Particle1;
-	return Constraint;
+	Constraint.Geometry[0] = Implicit0;
+	Constraint.Geometry[1] = Implicit1;
+	ConstraintBuffer.Add(Constraint);
 }
 
-template <typename T, int d>
-TRigidBodyContactConstraint<T, d> ComputeUnionUnionConstraint(TGeometryParticleHandle<T, d>* Particle0, TGeometryParticleHandle<T, d>* Particle1, const T Thickness)
+template<typename T, int d>
+void ConstructUnionUnionConstraints(TGeometryParticleHandle<T, d>* Particle0, TGeometryParticleHandle<T, d>* Particle1, const FImplicitObject* Implicit0, const FImplicitObject* Implicit1, const T Thickness, TArray<typename TPBDCollisionConstraint<T, d>::FRigidBodyContactConstraint>& ConstraintBuffer, TPBDCollisionConstraintHistory<T, d>* History = nullptr)
 {
-	TRigidBodyContactConstraint<T, d> Constraint;
-	Constraint.Particle = Particle0;
-	Constraint.Levelset = Particle1;
-	//todo(ocohen): some heuristic for determining the order?
-	return Constraint;
+	SCOPE_CYCLE_COUNTER(STAT_UpdateUnionUnionConstraint);
+
+	const FImplicitObject* ParticleObj = Particle0->Geometry().Get();
+	const FImplicitObject* LevelsetObj = Particle1->Geometry().Get();
+
+	TRigidTransform<T, d> ParticlesTM = TRigidTransform<T, d>(Particle0->X(), Particle0->R());
+	TRigidTransform<T, d> LevelsetTM = TRigidTransform<T, d>(Particle1->X(), Particle1->R());
+
+	const TArray<Pair<const FImplicitObject*, TRigidTransform<T, d>>> LevelsetShapes = FindRelevantShapes2(ParticleObj, ParticlesTM, *LevelsetObj, LevelsetTM, Thickness);
+
+	for (const Pair<const FImplicitObject*, TRigidTransform<T, d>>& LevelsetObjPair : LevelsetShapes)
+	{
+		const FImplicitObject* LevelsetInnerObj = LevelsetObjPair.First;
+		const TRigidTransform<T, d>& LevelsetInnerObjTM = LevelsetObjPair.Second * LevelsetTM;
+
+		//now find all particle inner objects
+		const TArray<Pair<const FImplicitObject*, TRigidTransform<T, d>>> ParticleShapes = FindRelevantShapes2(LevelsetInnerObj, LevelsetInnerObjTM, *ParticleObj, ParticlesTM, Thickness);
+
+		//for each inner obj pair, update constraint
+		for (const Pair<const FImplicitObject*, TRigidTransform<T, d>>& ParticlePair : ParticleShapes)
+		{
+			const FImplicitObject* ParticleInnerObj = ParticlePair.First;
+			ConstructConstraintsImpl<T, d>(Particle0, Particle1, ParticleInnerObj, LevelsetInnerObj, Thickness, ConstraintBuffer, History);
+		}
+	}
 }
 
 template<typename T, int d>
-typename TPBDCollisionConstraint<T, d>::FRigidBodyContactConstraint TPBDCollisionConstraint<T, d>::ComputeConstraint(TGeometryParticleHandle<T, d>* Particle0, TGeometryParticleHandle<T, d>* Particle1, const T Thickness)
+typename void ConstructConstraintsImpl(TGeometryParticleHandle<T, d>* Particle0, TGeometryParticleHandle<T, d>* Particle1, const FImplicitObject* Implicit0, const FImplicitObject* Implicit1, const T Thickness, TArray<typename TPBDCollisionConstraint<T, d>::FRigidBodyContactConstraint>& ConstraintBuffer, TPBDCollisionConstraintHistory<T, d>* History)
 {
-	if (!Particle0->Geometry() || !Particle1->Geometry())
+	// TriangleMesh implicit's are for scene query only.
+	if (Implicit0 && GetInnerType(Implicit0->GetType()) == ImplicitObjectType::TriangleMesh) return;
+	if (Implicit1 && GetInnerType(Implicit1->GetType()) == ImplicitObjectType::TriangleMesh) return;
+
+	if (!Implicit0 || !Implicit1)
 	{
-		return ComputeLevelsetConstraint(Particle0, Particle1, Thickness);
+		ConstructLevelsetConstraints(Particle0, Particle1, Implicit0, Implicit1, Thickness, ConstraintBuffer, History);
 	}
-	if (Particle0->Geometry()->GetType() == TBox<T, d>::StaticType() && Particle1->Geometry()->GetType() == TBox<T, d>::StaticType())
+	if (Implicit0->GetType() == TBox<T, d>::StaticType() && Implicit1->GetType() == TBox<T, d>::StaticType())
 	{
-		return ComputeBoxConstraint(Particle0, Particle1, Thickness);
+		ConstructBoxConstraints(Particle0, Particle1, Implicit0, Implicit1, Thickness, ConstraintBuffer, History);
 	}
-	else if (Particle0->Geometry()->GetType() == TSphere<T, d>::StaticType() && Particle1->Geometry()->GetType() == TSphere<T, d>::StaticType())
+	else if (Implicit0->GetType() == TSphere<T, d>::StaticType() && Implicit1->GetType() == TSphere<T, d>::StaticType())
 	{
-		return ComputeSphereConstraint(Particle0, Particle1, Thickness);
+		ConstructSphereConstraints(Particle0, Particle1, Implicit0, Implicit1, Thickness, ConstraintBuffer, History);
 	}
-	else if (Particle0->Geometry()->GetType() == TBox<T, d>::StaticType() && Particle1->Geometry()->GetType() == TPlane<T, d>::StaticType())
+	else if (Implicit0->GetType() == TBox<T, d>::StaticType() && Implicit1->GetType() == TPlane<T, d>::StaticType())
 	{
-		return ComputeBoxPlaneConstraint(Particle0, Particle1, Thickness);
+		ConstructBoxPlaneConstraints(Particle0, Particle1, Implicit0, Implicit1, Thickness, ConstraintBuffer, History);
 	}
-	else if (Particle1->Geometry()->GetType() == TBox<T, d>::StaticType() && Particle0->Geometry()->GetType() == TPlane<T, d>::StaticType())
+	else if (Implicit1->GetType() == TBox<T, d>::StaticType() && Implicit0->GetType() == TPlane<T, d>::StaticType())
 	{
-		return ComputeBoxPlaneConstraint(Particle1, Particle0, Thickness);
+		ConstructBoxPlaneConstraints(Particle1, Particle0, Implicit0, Implicit1, Thickness, ConstraintBuffer, History);
 	}
-	else if (Particle0->Geometry()->GetType() == TSphere<T, d>::StaticType() && Particle1->Geometry()->GetType() == TPlane<T, d>::StaticType())
+	else if (Implicit0->GetType() == TSphere<T, d>::StaticType() && Implicit1->GetType() == TPlane<T, d>::StaticType())
 	{
-		return ComputeSpherePlaneConstraint(Particle0, Particle1, Thickness);
+		ConstructSpherePlaneConstraints(Particle0, Particle1, Implicit0, Implicit1, Thickness, ConstraintBuffer, History);
 	}
-	else if (Particle1->Geometry()->GetType() == TSphere<T, d>::StaticType() && Particle0->Geometry()->GetType() == TPlane<T, d>::StaticType())
+	else if (Implicit1->GetType() == TSphere<T, d>::StaticType() && Implicit0->GetType() == TPlane<T, d>::StaticType())
 	{
-		return ComputeSpherePlaneConstraint(Particle1, Particle0, Thickness);
+		ConstructSpherePlaneConstraints(Particle1, Particle0, Implicit0, Implicit1, Thickness, ConstraintBuffer, History);
 	}
-	else if (Particle0->Geometry()->GetType() == TSphere<T, d>::StaticType() && Particle1->Geometry()->GetType() == TBox<T, d>::StaticType())
+	else if (Implicit0->GetType() == TSphere<T, d>::StaticType() && Implicit1->GetType() == TBox<T, d>::StaticType())
 	{
-		return ComputeSphereBoxConstraint(Particle0, Particle1, Thickness);
+		ConstructSphereBoxConstraints(Particle0, Particle1, Implicit0, Implicit1, Thickness, ConstraintBuffer, History);
 	}
-	else if (Particle1->Geometry()->GetType() == TSphere<T, d>::StaticType() && Particle0->Geometry()->GetType() == TBox<T, d>::StaticType())
+	else if (Implicit1->GetType() == TSphere<T, d>::StaticType() && Implicit0->GetType() == TBox<T, d>::StaticType())
 	{
-		return ComputeSphereBoxConstraint(Particle1, Particle0, Thickness);
+		ConstructSphereBoxConstraints(Particle1, Particle0, Implicit0, Implicit1, Thickness, ConstraintBuffer, History);
 	}
-	else if (Particle0->Geometry()->GetType() == TCapsule<T>::StaticType() && Particle1->Geometry()->GetType() == TCapsule<T>::StaticType())
+	else if (Implicit0->GetType() == TCapsule<T>::StaticType() && Implicit1->GetType() == TCapsule<T>::StaticType())
 	{
-		return ComputeCapsuleCapsuleConstraint(Particle0, Particle1, Thickness);
+		ConstructCapsuleCapsuleConstraints(Particle0, Particle1, Implicit0, Implicit1, Thickness, ConstraintBuffer, History);
 	}
-	else if (Particle0->Geometry()->GetType() == TCapsule<T>::StaticType() && Particle1->Geometry()->GetType() == TBox<T, d>::StaticType())
+	else if (Implicit0->GetType() == TCapsule<T>::StaticType() && Implicit1->GetType() == TBox<T, d>::StaticType())
 	{
-		return ComputeCapsuleBoxConstraint(Particle0, Particle1, Thickness);
+		ConstructCapsuleBoxConstraints(Particle0, Particle1, Implicit0, Implicit1, Thickness, ConstraintBuffer, History);
 	}
-	else if (Particle1->Geometry()->GetType() == TCapsule<T>::StaticType() && Particle0->Geometry()->GetType() == TBox<T, d>::StaticType())
+	else if (Implicit1->GetType() == TCapsule<T>::StaticType() && Implicit0->GetType() == TBox<T, d>::StaticType())
 	{
-		return ComputeCapsuleBoxConstraint(Particle1, Particle0, Thickness);
+		ConstructCapsuleBoxConstraints(Particle1, Particle0, Implicit0, Implicit1, Thickness, ConstraintBuffer, History);
 	}
-	else if (Particle0->Geometry()->GetType() < TImplicitObjectUnion<T, d>::StaticType() && Particle1->Geometry()->GetType() == TImplicitObjectUnion<T, d>::StaticType())
+	else if (Implicit0->GetType() < TImplicitObjectUnion<T, d>::StaticType() && Implicit1->GetType() == TImplicitObjectUnion<T, d>::StaticType())
 	{
-		return ComputeSingleUnionConstraint(Particle0, Particle1, Thickness);
+		ConstructSingleUnionConstraints(Particle0, Particle1, Implicit0, Implicit1, Thickness, ConstraintBuffer, History);
 	}
-	else if (Particle0->Geometry()->GetType() == TImplicitObjectUnion<T, d>::StaticType() && Particle1->Geometry()->GetType() < TImplicitObjectUnion<T, d>::StaticType())
+	else if (Implicit0->GetType() == TImplicitObjectUnion<T, d>::StaticType() && Implicit1->GetType() < TImplicitObjectUnion<T, d>::StaticType())
 	{
-		return ComputeSingleUnionConstraint(Particle1, Particle0, Thickness);
+		ConstructSingleUnionConstraints(Particle1, Particle0, Implicit0, Implicit1, Thickness, ConstraintBuffer, History);
 	}
-	else if(Particle0->Geometry()->GetType() == TImplicitObjectUnion<T, d>::StaticType() && Particle1->Geometry()->GetType() == TImplicitObjectUnion<T, d>::StaticType())
+	else if (Implicit0->GetType() == TImplicitObjectUnion<T, d>::StaticType() && Implicit1->GetType() == TImplicitObjectUnion<T, d>::StaticType())
 	{
-		return ComputeUnionUnionConstraint(Particle0, Particle1, Thickness);
+		ConstructUnionUnionConstraints(Particle0, Particle1, Implicit0, Implicit1, Thickness, ConstraintBuffer, History);
 	}
-#if 0
-	else if (Particle0->Geometry()->IsConvex() && Particle1->Geometry()->IsConvex())
+	else if (Implicit0->IsConvex() && Implicit1->IsConvex())
 	{
-		return ComputeLevelsetConstraintGJK(Particle0, Particle1, Thickness);
+		ConstructConvexConvexConstraints(Particle0, Particle1, Implicit0, Implicit1, Thickness, ConstraintBuffer, History);
 	}
-#endif
-	return ComputeLevelsetConstraint(Particle0, Particle1, Thickness);
+	else
+	{
+		ConstructLevelsetConstraints(Particle0, Particle1, Implicit0, Implicit1, Thickness, ConstraintBuffer, History);
+	}
+}
+
+template<typename T, int d>
+typename void TPBDCollisionConstraint<T, d>::ConstructConstraints(TGeometryParticleHandle<T, d>* Particle0, TGeometryParticleHandle<T, d>* Particle1, const T Thickness, TArray<typename TPBDCollisionConstraint<T, d>::FRigidBodyContactConstraint>& ConstraintBuffer)
+{
+	ensure(Particle0 && Particle1);
+
+	FConstraintHistory* HistoryElement = nullptr;
+	if (FConstraintHistory** Ptr = History.Find(FConstraintHandleID(Particle0, Particle1)))
+		HistoryElement = *Ptr;
+
+	ConstructConstraintsImpl<T, d>(Particle0, Particle1, Particle0->Geometry().Get(), Particle1->Geometry().Get(), Thickness, ConstraintBuffer, HistoryElement);
 }
 
 // NOTE: UpdateLevelsetConstraintImp and its <float,3> specialization are here for the Linux build. It looks like
@@ -2083,8 +2274,11 @@ void UpdateLevelsetConstraintImp(const float Thickness, TRigidBodyContactConstra
 }
 
 template<ECollisionUpdateType UpdateType, typename T, int d>
-void UpdateConstraintImp2(const FImplicitObject& ParticleObject, const TRigidTransform<T, d>& ParticleTM, const FImplicitObject& LevelsetObject, const TRigidTransform<T, d>& LevelsetTM, const T Thickness, TRigidBodyContactConstraint<T, d>& Constraint)
+void UpdateConstraintImp2(const FImplicitObject& ParticleObject, const TRigidTransform<T, d>& ParticleTM, const FImplicitObject& LevelsetObject, const TRigidTransform<T, d>& LevelsetTM, const T Thickness,
+	TRigidBodyContactConstraint<T, d>* Constraints, int32 NumConstraints, TPBDCollisionConstraintHistory<T, d>* History)
 {
+	TRigidBodyContactConstraint<T, d>& Constraint = Constraints[0];
+
 	if (ParticleObject.GetType() == TBox<T, d>::StaticType() && LevelsetObject.GetType() == TBox<T, d>::StaticType())
 	{
 		UpdateBoxConstraint(*ParticleObject.template GetObject<TBox<T,d>>(), ParticleTM, *LevelsetObject.template GetObject<TBox<T, d>>(), LevelsetTM, Thickness, Constraint);
@@ -2155,7 +2349,7 @@ void UpdateConstraintImp2(const FImplicitObject& ParticleObject, const TRigidTra
 	}
 	else if (ParticleObject.GetType() < TImplicitObjectUnion<T, d>::StaticType() && LevelsetObject.GetType() == TImplicitObjectUnion<T, d>::StaticType())
 	{
-		return UpdateSingleUnionConstraint<UpdateType>(Thickness, Constraint);
+		return UpdateSingleUnionConstraint<UpdateType>(Thickness, Constraints, NumConstraints, History);
 	}
 	else if (ParticleObject.GetType() == TImplicitObjectUnion<T, d>::StaticType() && LevelsetObject.GetType() < TImplicitObjectUnion<T, d>::StaticType())
 	{
@@ -2163,11 +2357,11 @@ void UpdateConstraintImp2(const FImplicitObject& ParticleObject, const TRigidTra
 	}
 	else if (ParticleObject.GetType() == TImplicitObjectUnion<T, d>::StaticType() && LevelsetObject.GetType() == TImplicitObjectUnion<T, d>::StaticType())
 	{
-		return UpdateUnionUnionConstraint<UpdateType>(Thickness, Constraint);
+		return UpdateUnionUnionConstraint<UpdateType>(Thickness, Constraints, NumConstraints, History);
 	}
 	else if (ParticleObject.IsConvex() && LevelsetObject.IsConvex())
 	{
-		UpdateConvexConstraintsUsingCoreShapes<UpdateType>(ParticleObject, ParticleTM, LevelsetObject, LevelsetTM, Thickness, Constraint);
+		UpdateConvexConvexConstraint<UpdateType>(ParticleObject, ParticleTM, LevelsetObject, LevelsetTM, Thickness, Constraints, NumConstraints, History);
 	}
 	else if (LevelsetObject.IsUnderlyingUnion())
 	{
@@ -2187,12 +2381,21 @@ DECLARE_CYCLE_STAT(TEXT("TPBDCollisionConstraint::UpdateConstraint"), STAT_Updat
 
 template<typename T, int d>
 template<ECollisionUpdateType UpdateType>
-void TPBDCollisionConstraint<T, d>::UpdateConstraint(const T Thickness, FRigidBodyContactConstraint& Constraint)
+void TPBDCollisionConstraint<T, d>::UpdateConstraint(const T Thickness, FRigidBodyContactConstraint* Constraints, int32 NumConstraints)
 {
 	SCOPE_CYCLE_COUNTER(STAT_UpdateConstraint);
 
-	Constraint.Phi = Thickness;
-	
+	// Reset Phi and validate that the constraints operate on the same bodies.
+	check(NumConstraints >= 1);
+	check(Constraints != nullptr);
+	for (int32 i = 0; i < NumConstraints; i++)
+	{
+		Constraints[i].Phi = Thickness;
+		ensure(Constraints[i % NumConstraints].Particle == Constraints[(i + 1) % NumConstraints].Particle);
+		ensure(Constraints[i % NumConstraints].Levelset == Constraints[(i + 1) % NumConstraints].Levelset);
+	}
+
+	FRigidBodyContactConstraint& Constraint = Constraints[0];
 	const TRigidTransform<T, d> ParticleTM = GetTransform(Constraint.Particle);
 	const TRigidTransform<T, d> LevelsetTM = GetTransform(Constraint.Levelset);
 
@@ -2212,20 +2415,25 @@ void TPBDCollisionConstraint<T, d>::UpdateConstraint(const T Thickness, FRigidBo
 	}
 	else
 	{
-		UpdateConstraintImp2<UpdateType>(*Constraint.Particle->Geometry(), ParticleTM, *Constraint.Levelset->Geometry(), LevelsetTM, Thickness, Constraint);
+		TPBDCollisionConstraintHistory<T, d>* HistoryData = nullptr;
+		FConstraintHandleID HandleID = GetConstraintHandleID(Constraint);
+		if (History.Contains(HandleID))
+		{
+			HistoryData = History[HandleID];
+		}
+		UpdateConstraintImp2<UpdateType>(*Constraint.Particle->Geometry(), ParticleTM, *Constraint.Levelset->Geometry(), LevelsetTM, Thickness, Constraints, NumConstraints, HistoryData);
 	}
 }
 
 template class TPBDCollisionConstraintHandle<float, 3>;
 template class TAccelerationStructureHandle<float, 3>;
 template class CHAOS_API TPBDCollisionConstraint<float, 3>;
-template void TPBDCollisionConstraint<float, 3>::UpdateConstraint<ECollisionUpdateType::Any>(const float Thickness, FRigidBodyContactConstraint& Constraint);
-template void TPBDCollisionConstraint<float, 3>::UpdateConstraint<ECollisionUpdateType::Deepest>(const float Thickness, FRigidBodyContactConstraint& Constraint);
+template void TPBDCollisionConstraint<float, 3>::UpdateConstraint<ECollisionUpdateType::Any>(const float Thickness, FRigidBodyContactConstraint* Constraints, int32 NumConstraints);
+template void TPBDCollisionConstraint<float, 3>::UpdateConstraint<ECollisionUpdateType::Deepest>(const float Thickness, FRigidBodyContactConstraint* Constraints, int32 NumConstraints);
 template void TPBDCollisionConstraint<float, 3>::UpdateLevelsetConstraint<ECollisionUpdateType::Any>(const float Thickness, FRigidBodyContactConstraint& Constraint);
 template void TPBDCollisionConstraint<float, 3>::UpdateLevelsetConstraint<ECollisionUpdateType::Deepest>(const float Thickness, FRigidBodyContactConstraint& Constraint);
 template void TPBDCollisionConstraint<float, 3>::ComputeConstraints<false>(const TPBDCollisionConstraint<float,3>::FAccelerationStructure&, float Dt);
 template void TPBDCollisionConstraint<float, 3>::ComputeConstraints<true>(const TPBDCollisionConstraint<float, 3>::FAccelerationStructure&, float Dt);
-
-template void UpdateConstraintImp2<ECollisionUpdateType::Any, float, 3>(const FImplicitObject& ParticleObject, const TRigidTransform<float, 3>& ParticleTM, const FImplicitObject& LevelsetObject, const TRigidTransform<float, 3>& LevelsetTM, const float Thickness, TRigidBodyContactConstraint<float, 3>& Constraint);
-template void UpdateConstraintImp2<ECollisionUpdateType::Deepest, float, 3>(const FImplicitObject& ParticleObject, const TRigidTransform<float, 3>& ParticleTM, const FImplicitObject& LevelsetObject, const TRigidTransform<float, 3>& LevelsetTM, const float Thickness, TRigidBodyContactConstraint<float, 3>& Constraint);
+template void UpdateConstraintImp2<ECollisionUpdateType::Any, float, 3>(const FImplicitObject& ParticleObject, const TRigidTransform<float, 3>& ParticleTM, const FImplicitObject& LevelsetObject, const TRigidTransform<float, 3>& LevelsetTM, const float Thickness, TRigidBodyContactConstraint<float, 3>* Constraints, int32 NumConstraints, TPBDCollisionConstraintHistory<float, 3>* History);
+template void UpdateConstraintImp2<ECollisionUpdateType::Deepest, float, 3>(const FImplicitObject& ParticleObject, const TRigidTransform<float, 3>& ParticleTM, const FImplicitObject& LevelsetObject, const TRigidTransform<float, 3>& LevelsetTM, const float Thickness, TRigidBodyContactConstraint<float, 3>* Constraint, int32 NumConstraints, TPBDCollisionConstraintHistory<float, 3>* History);
 }
diff --git a/Engine/Source/Runtime/Experimental/Chaos/Public/Chaos/Convex.h b/Engine/Source/Runtime/Experimental/Chaos/Public/Chaos/Convex.h
--- a/Engine/Source/Runtime/Experimental/Chaos/Public/Chaos/Convex.h
+++ b/Engine/Source/Runtime/Experimental/Chaos/Public/Chaos/Convex.h
@@ -259,6 +259,11 @@ namespace Chaos
 			TConvexBuilder<T>::Simplify(Planes, SurfaceParticles, LocalBoundingBox);
 		}
 
+		TVector<T,d> GetCenter() const
+		{
+			return TVector<T, d>(0);
+		}
+
 	private:
 		TArray<TPlane<T, d>> Planes;
 		TParticles<T, d> SurfaceParticles;	//copy of the vertices that are just on the convex hull boundary
diff --git a/Engine/Source/Runtime/Experimental/Chaos/Public/Chaos/GJK.h b/Engine/Source/Runtime/Experimental/Chaos/Public/Chaos/GJK.h
--- a/Engine/Source/Runtime/Experimental/Chaos/Public/Chaos/GJK.h
+++ b/Engine/Source/Runtime/Experimental/Chaos/Public/Chaos/GJK.h
@@ -928,65 +928,4 @@ namespace Chaos
 		return false;
 	}
 
-	/** Approximate the collision point using Core Shape Reduction.
-	 @A The first geometry
-	 @ATM The transform of the first geometry
-	 @B The second geometry
-	 @BTM The transform of the second geometry
-	 @OutLocation The collision point.
-	 @OutPhi The depth along the normal. (FLT_MAX if failed)
-	 @OutNormal The impact normal (in A's world space)
-	 */
-	template<typename T, int d, typename TGeometryA, typename TGeometryB>
-	void GJKCoreShapeIntersection(const TGeometryA& A, const TRigidTransform<T, d>& ATM, const TGeometryB& B, const TRigidTransform<T, d>& BTM,
-		TVector<T, d> & OutLocation, T & OutPhi, TVector<T, d> & OutNormal, const T Thickness = 0, const TVector<T, 3>& InitialDir = TVector<T, 3>(1, 0, 0))
-	{
-		OutPhi = FLT_MAX;
-		check(A.IsConvex() && B.IsConvex());
-		TRigidTransform<T, d> BToATM = BTM.GetRelativeTransform(ATM);
-		TVector<T, d> ABVector = -BToATM.GetTranslation();
-
-		//
-		// Guess an initial thickness that will separate the bodies, then iterate until they are separated
-		//
-		T MinBoxMin = FMath::Min(A.BoundingBox().Extents().Min(), B.BoundingBox().Extents().Min());
-		T MinLocalThickness = -MinBoxMin;
-		T Delta = MinBoxMin / 20.; // reduce by delta at each step
-		T LocalThickness = 0.0;
-		while (GJKIntersection<T>(A, B, BToATM, LocalThickness, InitialDir, LocalThickness))
-		{
-			LocalThickness -= Delta;
-			if (LocalThickness < MinLocalThickness)
-			{
-				return;
-			}
-		}
-		LocalThickness *= 1.05; // corner case when equal.
-
-		//
-		// Measure the separation distance from a raycast, then remove the LocalThickness.
-		//
-		auto Calculate = [&](const TVector<T, d>& InRay, TVector<T, d> & Location, T & Phi, TVector<T, d> & Normal)
-		{
-			T RayTime = FLT_MAX;
-			TVector<T, d> RayPosition, RayNormal, InRayNormal = InRay.GetSafeNormal();
-			if (GJKRaycast(A, B, BToATM, InRayNormal, (T)InRay.Size(), RayTime, RayPosition, RayNormal, LocalThickness, InitialDir, LocalThickness))
-			{
-				// Note: Initial overlaps in GJKRaycast cannot set a contact position/normal. It sets RayTime = 0 in this case.
-				if (RayTime > 0)
-				{
-					Location = ATM.TransformPosition(RayPosition + RayNormal * (-LocalThickness));
-					Normal = -ATM.TransformVector(RayNormal);
-
-					// @todo: Remove PhiWithNormal call
-					TVector<T, d> TmpNormal;
-					Phi = B.PhiWithNormal(BTM.InverseTransformPosition(Location), TmpNormal);
-					return true;
-				}
-			}
-			return false;
-		};
-
-		Calculate(ABVector, OutLocation, OutPhi, OutNormal);
-	}
 }
diff --git a/Engine/Source/Runtime/Experimental/Chaos/Public/Chaos/PBDCollisionConstraint.h b/Engine/Source/Runtime/Experimental/Chaos/Public/Chaos/PBDCollisionConstraint.h
--- a/Engine/Source/Runtime/Experimental/Chaos/Public/Chaos/PBDCollisionConstraint.h
+++ b/Engine/Source/Runtime/Experimental/Chaos/Public/Chaos/PBDCollisionConstraint.h
@@ -62,6 +62,7 @@ public:
 	TPBDCollisionConstraintHandle();
 	TPBDCollisionConstraintHandle(FConstraintContainer* InConstraintContainer, int32 InConstraintIndex);
 
+	TRigidBodyContactConstraint<T, d>& GetContact();
 	const TRigidBodyContactConstraint<T, d>& GetContact() const;
 	void SetConstraintIndex(int32 IndexIn) { ConstraintIndex = IndexIn; }
 
@@ -73,29 +74,44 @@ protected:
 template<class T, int d>
 class CHAOS_API TPBDCollisionConstraintHistory
 {
+	typedef TPair<const FImplicitObject*, const FImplicitObject*> ImplicitPairs;
+
 public:
 	using FConstraintHandle = TPBDCollisionConstraintHandle<T, d>;
-
-	TPBDCollisionConstraintHistory(const TVector<T, d>& InLocation, const TRotation<T, d>& InRotation)
-		: Location(InLocation)
+	TPBDCollisionConstraintHistory(const TVector<T, d>& InLocation, const TRotation<T, d>& InRotation, int32 InTimestamp = -INT_MAX)
+		: Timestamp(InTimestamp)
+		, Location(InLocation)
 		, Rotation(InRotation)
 	{}
 
 	const TVector<T, d>& GetLocation() { return Location; }
 	const TRotation<T, d>& GetRotation() { return Rotation; }
 
-	void AddHandle(const FConstraintHandle* InHandle) {
+	void AddHandle(FConstraintHandle* InHandle)
+	{
+		Implicits.Add(ImplicitPairs(InHandle->GetContact().Geometry[0], InHandle->GetContact().Geometry[1]));
 		Handles.Push(InHandle);
 	}
-	void RemoveHandle(const FConstraintHandle* InHandle) {
+	void RemoveHandle(FConstraintHandle* InHandle)
+	{
 		Handles.RemoveSingleSwap(InHandle);
+		Implicits.Remove(ImplicitPairs(InHandle->GetContact().Geometry[0], InHandle->GetContact().Geometry[1]));
 	}
-	const TArray< const FConstraintHandle*>& GetHandles() { return Handles; }
+
+	TArray<FConstraintHandle*>& GetHandles() { return Handles; }
+	const TArray<FConstraintHandle*>& GetHandles() const { return Handles; }
+
+	bool ContainsShapeConnection(const FImplicitObject* Implicit0In, const FImplicitObject* Implicit1In) { return Implicits.Contains(ImplicitPairs(Implicit0In, Implicit1In)); }
+
+	int32 GetTimestamp() const { return Timestamp; }
+	void SetTimestamp(int32 InTimestamp) { Timestamp = InTimestamp; }
 
 private:
+	int32 Timestamp;
 	TVector<T, d> Location;
 	TRotation<T, d> Rotation;
-	TArray< const FConstraintHandle*> Handles;
+	TArray<FConstraintHandle*> Handles;
+	TSet<ImplicitPairs> Implicits;
 };
 
 template<typename T, int d>
@@ -142,6 +158,11 @@ public:
 		MThickness = InThickness;
 	}
 
+	void SetVelocityInflation(float ScaleFactor)
+	{
+		CollisionVelocityInflation = ScaleFactor;
+	}
+
 	void SetVelocitySolveEnabled(bool bInEnableVelocitySolve)
 	{
 		bEnableVelocitySolve = bInEnableVelocitySolve;
@@ -293,9 +314,9 @@ private:
 	void ComputeConstraints(const FAccelerationStructure& AccelerationStructure, T Dt);
 
 	template<ECollisionUpdateType>
-	void UpdateConstraint(const T Thickness, FRigidBodyContactConstraint& Constraint);
+	void UpdateConstraint(const T Thickness, FRigidBodyContactConstraint* Constraints, int32 NumConstraints);
 
-	FRigidBodyContactConstraint ComputeConstraint(TGeometryParticleHandle<T, d>* Particle0, TGeometryParticleHandle<T, d>* Particle1, const T Thickness);
+	void ConstructConstraints(TGeometryParticleHandle<T, d>* Particle0, TGeometryParticleHandle<T, d>* Particle1, const T Thickness, TArray<typename TPBDCollisionConstraint<T, d>::FRigidBodyContactConstraint>& ConstraintBuffer);
 
 
 	template<typename SPATIAL_ACCELERATION>
@@ -315,7 +336,10 @@ private:
 	T MAngularFriction;
 	bool bUseCCD;
 	bool bEnableCollisions;
+
 	int32 LifespanCounter;
+	int32 CollisionHistoryLifespan;
+	float CollisionVelocityInflation;
 
 	TRigidBodyContactConstraintsPostComputeCallback<T, d> PostComputeCallback;
 	TRigidBodyContactConstraintsPostApplyCallback<T, d> PostApplyCallback;
@@ -324,12 +348,11 @@ private:
 	TArray<FConstraintHandle*> Handles;
 	FConstraintHandleAllocator HandleAllocator;
 
-	bool bSupportHistory;
 	TMap<FConstraintHandleID, FConstraintHistory*> History;
 };
 
-extern template void TPBDCollisionConstraint<float, 3>::UpdateConstraint<ECollisionUpdateType::Any>(const float Thickness, FRigidBodyContactConstraint& Constraint);
-extern template void TPBDCollisionConstraint<float, 3>::UpdateConstraint<ECollisionUpdateType::Deepest>(const float Thickness, FRigidBodyContactConstraint& Constraint);
+extern template void TPBDCollisionConstraint<float, 3>::UpdateConstraint<ECollisionUpdateType::Any>(const float Thickness, FRigidBodyContactConstraint* Constraints, int32 NumConstraints);
+extern template void TPBDCollisionConstraint<float, 3>::UpdateConstraint<ECollisionUpdateType::Deepest>(const float Thickness, FRigidBodyContactConstraint* Constraints, int32 NumConstraints);
 extern template void TPBDCollisionConstraint<float, 3>::UpdateLevelsetConstraint<ECollisionUpdateType::Any>(const float Thickness, FRigidBodyContactConstraint& Constraint);
 extern template void TPBDCollisionConstraint<float, 3>::UpdateLevelsetConstraint<ECollisionUpdateType::Deepest>(const float Thickness, FRigidBodyContactConstraint& Constraint);
 extern template void TPBDCollisionConstraint<float, 3>::ComputeConstraints<false>(const TPBDCollisionConstraint<float, 3>::FAccelerationStructure&, float Dt);
diff --git a/Engine/Source/Runtime/Experimental/Chaos/Public/Chaos/PBDCollisionConstraintImp.h b/Engine/Source/Runtime/Experimental/Chaos/Public/Chaos/PBDCollisionConstraintImp.h
--- a/Engine/Source/Runtime/Experimental/Chaos/Public/Chaos/PBDCollisionConstraintImp.h
+++ b/Engine/Source/Runtime/Experimental/Chaos/Public/Chaos/PBDCollisionConstraintImp.h
@@ -166,7 +166,7 @@ namespace Chaos
 
 				const bool bBody1Bounded = HasBoundingBox(Particle1);
 
-				const T Box1Thickness = ComputeThickness(Particle1, Dt).Size();
+				const T Box1Thickness = ComputeThickness(Particle1, CollisionVelocityInflation * Dt).Size();
 				if (bBody1Bounded)
 				{
 #if CHAOS_PARTICLEHANDLE_TODO
@@ -203,6 +203,9 @@ namespace Chaos
 				int32 CountNP = 0;
 				int32 RejectedNP = 0;
 #endif
+				TArray<TRigidBodyContactConstraint<T, d>> TempConstraintBuffer;
+				TempConstraintBuffer.Reserve(4);
+
 				const int32 NumPotentials = PotentialIntersections.Num();
 				for (int32 i = 0; i < NumPotentials; ++i)
 				{
@@ -252,21 +255,22 @@ namespace Chaos
 					const TVector<T, d> Box2Thickness = Particle2.AsDynamic() ? ComputeThickness(*Particle2.AsDynamic(), Dt) : TVector<T, d>(0);
 					const T UseThickness = FMath::Max(Box1Thickness, Box2Thickness.Size());// + MThickness
 
-					auto Constraint = ComputeConstraint(Particle1.Handle(), Particle2.Handle(), UseThickness);
-
-					//if (true || !InParticles.Geometry(Body1Index)->HasBoundingBox() || !InParticles.Geometry(Body2Index)->HasBoundingBox())
-					{
-						//SCOPE_CYCLE_COUNTER(STAT_ComputeConstraintsNP3);
-						//use narrow phase to determine if constraint is needed. Without this we can't do shock propagation
+					TempConstraintBuffer.Reset();
+					ConstructConstraints(Particle1.Handle(), Particle2.Handle(), UseThickness, TempConstraintBuffer);
 
 #if !UE_BUILD_SHIPPING
-						if (bGatherStats)
-						{
-							++CountNP;
-						}
+					if (bGatherStats)
+					{
+						++CountNP;
+					}
 #endif
-						UpdateConstraint<ECollisionUpdateType::Any>(UseThickness, Constraint);
+					if (TempConstraintBuffer.Num())
+					{
+						UpdateConstraint<ECollisionUpdateType::Any>(UseThickness, &TempConstraintBuffer[0], TempConstraintBuffer.Num());
+					}
 
+					for (auto& Constraint : TempConstraintBuffer)
+					{
 						if (Constraint.Phi < UseThickness)
 						{
 							Queue.Enqueue(Constraint);
@@ -278,7 +282,6 @@ namespace Chaos
 #endif
 						}
 					}
-
 				}
 
 #if !UE_BUILD_SHIPPING
@@ -295,28 +298,24 @@ namespace Chaos
 				while (!Queue.IsEmpty())
 				{
 					int32 Idx = Constraints.AddUninitialized(1);
-					FConstraintHandle * Handle = HandleAllocator.AllocHandle(this, Idx);
+					FConstraintHandle* Handle = HandleAllocator.AllocHandle(this, Idx);
 
 					Handles.Add(Handle);
 					Queue.Dequeue(Constraints[Idx]);
-					Constraints[Idx].Lifespan = LifespanCounter;
 					ensure(Handles.Num() == Constraints.Num());
 
-					if (bSupportHistory)
+					FConstraintHandleID HandleID = GetConstraintHandleID(Constraints[Idx]);
+					FConstraintHistory* HistoryElement;
+					if (auto ValPtr = History.Find(HandleID))
 					{
-						FConstraintHandleID HandleID = GetConstraintHandleID(Constraints[Idx]);
-						FConstraintHistory* HistoryElement;
-						if (auto ValPtr = History.Find(HandleID))
-						{
-							HistoryElement = *ValPtr;
-						}
-						else
-						{
-							HistoryElement = new FConstraintHistory(Constraints[Idx].Particle->X(), Constraints[Idx].Particle->R());
-							History.Add(HandleID, HistoryElement);
-						}
-						HistoryElement->AddHandle(Handle);
+						HistoryElement = *ValPtr;
+					}
+					else
+					{
+						HistoryElement = new FConstraintHistory(Constraints[Idx].Particle->X(), Constraints[Idx].Particle->R(), LifespanCounter + CollisionHistoryLifespan);
+						History.Add(HandleID, HistoryElement);
 					}
+					HistoryElement->AddHandle(Handle);
 				}
 				LifespanCounter++;
 			}
diff --git a/Engine/Source/Runtime/Experimental/Chaos/Public/Chaos/PBDCollisionTypes.h b/Engine/Source/Runtime/Experimental/Chaos/Public/Chaos/PBDCollisionTypes.h
--- a/Engine/Source/Runtime/Experimental/Chaos/Public/Chaos/PBDCollisionTypes.h
+++ b/Engine/Source/Runtime/Experimental/Chaos/Public/Chaos/PBDCollisionTypes.h
@@ -13,17 +13,18 @@ namespace Chaos
 template<class T, int d>
 struct TRigidBodyContactConstraint
 {
-	TRigidBodyContactConstraint() : AccumulatedImpulse(0.f), bDisabled(false), Lifespan(0) {}
+	TRigidBodyContactConstraint() : bDisabled(false), Normal(0), AccumulatedImpulse(0.f) {}
 	TGeometryParticleHandle<T, d>* Particle;
 	TGeometryParticleHandle<T, d>* Levelset;
+	const FImplicitObject* Geometry[2];
+
+	bool bDisabled;
 	TVector<T, d> Normal;
 	TVector<T, d> Location;
+	TVector<T, d> LocalLocation;
 	T Phi;
 	TVector<T, d> AccumulatedImpulse;
 
-	bool bDisabled;
-	int32 Lifespan;
-
 	FString ToString() const
 	{
 		return FString::Printf(TEXT("Particle:%s, Levelset:%s, Normal:%s, Location:%s, Phi:%f, AccumulatedImpulse:%s"), *Particle->ToString(), *Levelset->ToString(), *Normal.ToString(), *Location.ToString(), Phi, *AccumulatedImpulse.ToString());
diff --git a/Engine/Source/Runtime/Experimental/Chaos/Public/Chaos/PBDConstraintRule.h b/Engine/Source/Runtime/Experimental/Chaos/Public/Chaos/PBDConstraintRule.h
--- a/Engine/Source/Runtime/Experimental/Chaos/Public/Chaos/PBDConstraintRule.h
+++ b/Engine/Source/Runtime/Experimental/Chaos/Public/Chaos/PBDConstraintRule.h
@@ -256,18 +256,6 @@ namespace Chaos
 					}
 				}
 			}
-
-			for (int32 Level = 0; Level <= MaxLevel; ++Level)
-			{
-				for (int32 Color = 0; Color <= MaxColor; ++Color)
-				{
-					if (LevelToColorToConstraintListMap[Level].Contains(Color) && LevelToColorToConstraintListMap[Level][Color].Num())
-					{
-						const TArray<typename FConstraints::FConstraintHandle*>& ConstraintHandles = GetLevelColorConstraints(LevelToColorToConstraintListMap, Level, Color);
-						Constraints.Disable( ConstraintHandles );
-					}
-				}
-			}
 		}
 
 		virtual void RemoveConstraints(const TSet<TGeometryParticleHandle<T, d>*>& InConstraints)
diff --git a/Engine/Source/Runtime/Experimental/ChaosSolvers/Private/PBDRigidsSolver.cpp b/Engine/Source/Runtime/Experimental/ChaosSolvers/Private/PBDRigidsSolver.cpp
--- a/Engine/Source/Runtime/Experimental/ChaosSolvers/Private/PBDRigidsSolver.cpp
+++ b/Engine/Source/Runtime/Experimental/ChaosSolvers/Private/PBDRigidsSolver.cpp
@@ -28,6 +28,10 @@ FAutoConsoleVariableRef CVarChaosSolverDrawCollisions(TEXT("p.Chaos.Solver.Debug
 FAutoConsoleVariableRef CVarChaosSolverDrawBPBounds(TEXT("p.Chaos.Solver.DrawBPBounds"), ChaosSolverDrawBPBounds, TEXT("Draw bounding volumes inside the broadphase (0 = never; 1 = end of frame)."));
 #endif
 
+float ChaosSolverCollisionDefaultIterationsCVar = 1;
+FAutoConsoleVariableRef CVarChaosSolverCollisionDefaultIterations(TEXT("p.ChaosSolverCollisionDefaultIterations"), ChaosSolverCollisionDefaultIterationsCVar, TEXT("Default collision iterations for the solver.[def:1]"));
+
+
 namespace Chaos
 {
 
@@ -123,7 +127,7 @@ namespace Chaos
 		, bHasFloor(true)
 		, bIsFloorAnalytic(false)
 		, FloorHeight(0.f)
-		, MEvolution(new FPBDRigidsEvolution(Particles))
+		, MEvolution(new FPBDRigidsEvolution(Particles, ChaosSolverCollisionDefaultIterationsCVar))
 		, MEventManager(new FEventManager(BufferingModeIn))
 		, MSolverEventFilters(new FSolverEventFilters())
 		, MActiveParticlesBuffer(new FActiveParticlesBuffer(BufferingModeIn))
@@ -346,7 +350,7 @@ namespace Chaos
 		CurrentFrame = 0;
 		MMaxDeltaTime = 1;
 		TimeStepMultiplier = 1;
-		MEvolution = TUniquePtr<FPBDRigidsEvolution>(new FPBDRigidsEvolution(Particles));
+		MEvolution = TUniquePtr<FPBDRigidsEvolution>(new FPBDRigidsEvolution(Particles, ChaosSolverCollisionDefaultIterationsCVar)); 
 
 		FEventDefaults::RegisterSystemEvents(*GetEventManager());
 	}
