Chaos :
Fix for CIS

diff --git a/Engine/Source/Runtime/Experimental/Chaos/Private/Chaos/PBDCollisionConstraint.cpp b/Engine/Source/Runtime/Experimental/Chaos/Private/Chaos/PBDCollisionConstraint.cpp
--- a/Engine/Source/Runtime/Experimental/Chaos/Private/Chaos/PBDCollisionConstraint.cpp
+++ b/Engine/Source/Runtime/Experimental/Chaos/Private/Chaos/PBDCollisionConstraint.cpp
@@ -64,7 +64,7 @@ template<ECollisionUpdateType, typename T, int d>
 void UpdateConstraintImp2(const TRigidTransform<T, d>& ParticleTM, const FImplicitObject& LevelsetObject, const TRigidTransform<T, d>& LevelsetTM, const T Thickness, TRigidBodyContactConstraint<T, d>& Constraint, int32 NumConstraints, TPBDCollisionConstraintHistory<T, d>* History);
 
 template<typename T, int d>
-void ConstructConstraintsImpl(TGeometryParticleHandle<T, d>* Particle0, TGeometryParticleHandle<T, d>* Particle1, const FImplicitObject* Implicit0, const FImplicitObject* Implicit1, const T Thickness, TArray<typename TPBDCollisionConstraint<T, d>::FRigidBodyContactConstraint>& ConstraintBuffer, TPBDCollisionConstraintHistory<T, d>* History);
+void ConstructConstraintsImpl(TGeometryParticleHandle<T, d>* Particle0, TGeometryParticleHandle<T, d>* Particle1, const FImplicitObject* Implicit0, const FImplicitObject* Implicit1, const T Thickness, TArray<TRigidBodyContactConstraint<T, d>>& ConstraintBuffer, TPBDCollisionConstraintHistory<T, d>* History);
 
 //
 // Collision Constraint Handle
@@ -2007,7 +2007,7 @@ void UpdateConvexConvexConstraint(const TGeometryA& A, TRigidTransform<T, d> ATM
 }
 
 template<typename T, int d>
-void ConstructLevelsetConstraints(TGeometryParticleHandle<T, d>* Particle0, TGeometryParticleHandle<T, d>* Particle1, const FImplicitObject* Implicit0, const FImplicitObject* Implicit1, const T Thickness, TArray<typename TPBDCollisionConstraint<T, d>::FRigidBodyContactConstraint>& ConstraintBuffer, TPBDCollisionConstraintHistory<T, d>* History = nullptr)
+void ConstructLevelsetConstraints(TGeometryParticleHandle<T, d>* Particle0, TGeometryParticleHandle<T, d>* Particle1, const FImplicitObject* Implicit0, const FImplicitObject* Implicit1, const T Thickness, TArray<TRigidBodyContactConstraint<T, d>>& ConstraintBuffer, TPBDCollisionConstraintHistory<T, d>* History = nullptr)
 {
 	//todo(ocohen):if both have collision particles, use the one with fewer?
 	if (!Particle1->Geometry() || (Particle0->AsDynamic() && !Particle0->AsDynamic()->CollisionParticlesSize() && Particle0->Geometry() && !Particle0->Geometry()->IsUnderlyingUnion()))
@@ -2031,7 +2031,7 @@ void ConstructLevelsetConstraints(TGeometryParticleHandle<T, d>* Particle0, TGeo
 }
 
 template<typename T, int d>
-void ConstructConvexConvexConstraints(TGeometryParticleHandle<T, d>* Particle0, TGeometryParticleHandle<T, d>* Particle1, const FImplicitObject* Implicit0, const FImplicitObject* Implicit1, const T Thickness, TArray<typename TPBDCollisionConstraint<T, d>::FRigidBodyContactConstraint>& ConstraintBuffer, TPBDCollisionConstraintHistory<T, d>* History = nullptr)
+void ConstructConvexConvexConstraints(TGeometryParticleHandle<T, d>* Particle0, TGeometryParticleHandle<T, d>* Particle1, const FImplicitObject* Implicit0, const FImplicitObject* Implicit1, const T Thickness, TArray<TRigidBodyContactConstraint<T, d>>& ConstraintBuffer, TPBDCollisionConstraintHistory<T, d>* History = nullptr)
 {
 	if (History)
 	{
@@ -2053,7 +2053,7 @@ void ConstructConvexConvexConstraints(TGeometryParticleHandle<T, d>* Particle0,
 }
 
 template<typename T, int d>
-void ConstructBoxConstraints(TGeometryParticleHandle<T, d>* Particle0, TGeometryParticleHandle<T, d>* Particle1, const FImplicitObject* Implicit0, const FImplicitObject* Implicit1, const T Thickness, TArray<typename TPBDCollisionConstraint<T, d>::FRigidBodyContactConstraint>& ConstraintBuffer, TPBDCollisionConstraintHistory<T, d>* History = nullptr)
+void ConstructBoxConstraints(TGeometryParticleHandle<T, d>* Particle0, TGeometryParticleHandle<T, d>* Particle1, const FImplicitObject* Implicit0, const FImplicitObject* Implicit1, const T Thickness, TArray<TRigidBodyContactConstraint<T, d>>& ConstraintBuffer, TPBDCollisionConstraintHistory<T, d>* History = nullptr)
 {
 	TRigidBodyContactConstraint<T, d> Constraint;
 	Constraint.Particle = Particle0;
@@ -2064,7 +2064,7 @@ void ConstructBoxConstraints(TGeometryParticleHandle<T, d>* Particle0, TGeometry
 }
 
 template<typename T, int d>
-void ConstructBoxPlaneConstraints(TGeometryParticleHandle<T, d>* Particle0, TGeometryParticleHandle<T, d>* Particle1, const FImplicitObject* Implicit0, const FImplicitObject* Implicit1, const T Thickness, TArray<typename TPBDCollisionConstraint<T, d>::FRigidBodyContactConstraint>& ConstraintBuffer, TPBDCollisionConstraintHistory<T, d>* History = nullptr)
+void ConstructBoxPlaneConstraints(TGeometryParticleHandle<T, d>* Particle0, TGeometryParticleHandle<T, d>* Particle1, const FImplicitObject* Implicit0, const FImplicitObject* Implicit1, const T Thickness, TArray<TRigidBodyContactConstraint<T, d>>& ConstraintBuffer, TPBDCollisionConstraintHistory<T, d>* History = nullptr)
 {
 	TRigidBodyContactConstraint<T, d> Constraint;
 	Constraint.Particle = Particle0;
@@ -2075,7 +2075,7 @@ void ConstructBoxPlaneConstraints(TGeometryParticleHandle<T, d>* Particle0, TGeo
 }
 
 template<typename T, int d>
-void ConstructSphereConstraints(TGeometryParticleHandle<T, d>* Particle0, TGeometryParticleHandle<T, d>* Particle1, const FImplicitObject* Implicit0, const FImplicitObject* Implicit1, const T Thickness, TArray<typename TPBDCollisionConstraint<T, d>::FRigidBodyContactConstraint>& ConstraintBuffer, TPBDCollisionConstraintHistory<T, d>* History = nullptr)
+void ConstructSphereConstraints(TGeometryParticleHandle<T, d>* Particle0, TGeometryParticleHandle<T, d>* Particle1, const FImplicitObject* Implicit0, const FImplicitObject* Implicit1, const T Thickness, TArray<TRigidBodyContactConstraint<T, d>>& ConstraintBuffer, TPBDCollisionConstraintHistory<T, d>* History = nullptr)
 {
 	TRigidBodyContactConstraint<T, d> Constraint;
 	Constraint.Particle = Particle0;
@@ -2086,7 +2086,7 @@ void ConstructSphereConstraints(TGeometryParticleHandle<T, d>* Particle0, TGeome
 }
 
 template<typename T, int d>
-void ConstructSpherePlaneConstraints(TGeometryParticleHandle<T, d>* Particle0, TGeometryParticleHandle<T, d>* Particle1, const FImplicitObject* Implicit0, const FImplicitObject* Implicit1, const T Thickness, TArray<typename TPBDCollisionConstraint<T, d>::FRigidBodyContactConstraint>& ConstraintBuffer, TPBDCollisionConstraintHistory<T, d>* History = nullptr)
+void ConstructSpherePlaneConstraints(TGeometryParticleHandle<T, d>* Particle0, TGeometryParticleHandle<T, d>* Particle1, const FImplicitObject* Implicit0, const FImplicitObject* Implicit1, const T Thickness, TArray<TRigidBodyContactConstraint<T, d>>& ConstraintBuffer, TPBDCollisionConstraintHistory<T, d>* History = nullptr)
 {
 	TRigidBodyContactConstraint<T, d> Constraint;
 	Constraint.Particle = Particle0;
@@ -2097,7 +2097,7 @@ void ConstructSpherePlaneConstraints(TGeometryParticleHandle<T, d>* Particle0, T
 }
 
 template<typename T, int d>
-void ConstructSphereBoxConstraints(TGeometryParticleHandle<T, d>* Particle0, TGeometryParticleHandle<T, d>* Particle1, const FImplicitObject* Implicit0, const FImplicitObject* Implicit1, const T Thickness, TArray<typename TPBDCollisionConstraint<T, d>::FRigidBodyContactConstraint>& ConstraintBuffer, TPBDCollisionConstraintHistory<T, d>* History = nullptr)
+void ConstructSphereBoxConstraints(TGeometryParticleHandle<T, d>* Particle0, TGeometryParticleHandle<T, d>* Particle1, const FImplicitObject* Implicit0, const FImplicitObject* Implicit1, const T Thickness, TArray<TRigidBodyContactConstraint<T, d>>& ConstraintBuffer, TPBDCollisionConstraintHistory<T, d>* History = nullptr)
 {
 	TRigidBodyContactConstraint<T, d> Constraint;
 	Constraint.Particle = Particle0;
@@ -2108,7 +2108,7 @@ void ConstructSphereBoxConstraints(TGeometryParticleHandle<T, d>* Particle0, TGe
 }
 
 template<typename T, int d>
-void ConstructCapsuleCapsuleConstraints(TGeometryParticleHandle<T, d>* Particle0, TGeometryParticleHandle<T, d>* Particle1, const FImplicitObject* Implicit0, const FImplicitObject* Implicit1, const T Thickness, TArray<typename TPBDCollisionConstraint<T, d>::FRigidBodyContactConstraint>& ConstraintBuffer, TPBDCollisionConstraintHistory<T, d>* History = nullptr)
+void ConstructCapsuleCapsuleConstraints(TGeometryParticleHandle<T, d>* Particle0, TGeometryParticleHandle<T, d>* Particle1, const FImplicitObject* Implicit0, const FImplicitObject* Implicit1, const T Thickness, TArray<TRigidBodyContactConstraint<T, d>>& ConstraintBuffer, TPBDCollisionConstraintHistory<T, d>* History = nullptr)
 {
 	TRigidBodyContactConstraint<T, d> Constraint;
 	Constraint.Particle = Particle0;
@@ -2119,7 +2119,7 @@ void ConstructCapsuleCapsuleConstraints(TGeometryParticleHandle<T, d>* Particle0
 }
 
 template<typename T, int d>
-void ConstructCapsuleBoxConstraints(TGeometryParticleHandle<T, d>* Particle0, TGeometryParticleHandle<T, d>* Particle1, const FImplicitObject* Implicit0, const FImplicitObject* Implicit1, const T Thickness, TArray<typename TPBDCollisionConstraint<T, d>::FRigidBodyContactConstraint>& ConstraintBuffer, TPBDCollisionConstraintHistory<T, d>* History = nullptr)
+void ConstructCapsuleBoxConstraints(TGeometryParticleHandle<T, d>* Particle0, TGeometryParticleHandle<T, d>* Particle1, const FImplicitObject* Implicit0, const FImplicitObject* Implicit1, const T Thickness, TArray<TRigidBodyContactConstraint<T, d>>& ConstraintBuffer, TPBDCollisionConstraintHistory<T, d>* History = nullptr)
 {
 	TRigidBodyContactConstraint<T, d> Constraint;
 	Constraint.Particle = Particle0;
@@ -2130,7 +2130,7 @@ void ConstructCapsuleBoxConstraints(TGeometryParticleHandle<T, d>* Particle0, TG
 }
 
 template<typename T, int d>
-void ConstructSingleUnionConstraints(TGeometryParticleHandle<T, d>* Particle0, TGeometryParticleHandle<T, d>* Particle1, const FImplicitObject* Implicit0, const FImplicitObject* Implicit1, const T Thickness, TArray<typename TPBDCollisionConstraint<T, d>::FRigidBodyContactConstraint>& ConstraintBuffer, TPBDCollisionConstraintHistory<T, d>* History = nullptr)
+void ConstructSingleUnionConstraints(TGeometryParticleHandle<T, d>* Particle0, TGeometryParticleHandle<T, d>* Particle1, const FImplicitObject* Implicit0, const FImplicitObject* Implicit1, const T Thickness, TArray<TRigidBodyContactConstraint<T, d>>& ConstraintBuffer, TPBDCollisionConstraintHistory<T, d>* History = nullptr)
 {
 	TRigidBodyContactConstraint<T, d> Constraint;
 	Constraint.Particle = Particle0;
@@ -2141,7 +2141,7 @@ void ConstructSingleUnionConstraints(TGeometryParticleHandle<T, d>* Particle0, T
 }
 
 template<typename T, int d>
-void ConstructUnionUnionConstraints(TGeometryParticleHandle<T, d>* Particle0, TGeometryParticleHandle<T, d>* Particle1, const FImplicitObject* Implicit0, const FImplicitObject* Implicit1, const T Thickness, TArray<typename TPBDCollisionConstraint<T, d>::FRigidBodyContactConstraint>& ConstraintBuffer, TPBDCollisionConstraintHistory<T, d>* History = nullptr)
+void ConstructUnionUnionConstraints(TGeometryParticleHandle<T, d>* Particle0, TGeometryParticleHandle<T, d>* Particle1, const FImplicitObject* Implicit0, const FImplicitObject* Implicit1, const T Thickness, TArray<TRigidBodyContactConstraint<T, d>>& ConstraintBuffer, TPBDCollisionConstraintHistory<T, d>* History = nullptr)
 {
 	SCOPE_CYCLE_COUNTER(STAT_UpdateUnionUnionConstraint);
 
@@ -2171,7 +2171,7 @@ void ConstructUnionUnionConstraints(TGeometryParticleHandle<T, d>* Particle0, TG
 }
 
 template<typename T, int d>
-typename void ConstructConstraintsImpl(TGeometryParticleHandle<T, d>* Particle0, TGeometryParticleHandle<T, d>* Particle1, const FImplicitObject* Implicit0, const FImplicitObject* Implicit1, const T Thickness, TArray<typename TPBDCollisionConstraint<T, d>::FRigidBodyContactConstraint>& ConstraintBuffer, TPBDCollisionConstraintHistory<T, d>* History)
+typename void ConstructConstraintsImpl(TGeometryParticleHandle<T, d>* Particle0, TGeometryParticleHandle<T, d>* Particle1, const FImplicitObject* Implicit0, const FImplicitObject* Implicit1, const T Thickness, TArray< TRigidBodyContactConstraint<T, d> >& ConstraintBuffer, TPBDCollisionConstraintHistory<T, d>* History)
 {
 	// TriangleMesh implicit's are for scene query only.
 	if (Implicit0 && GetInnerType(Implicit0->GetType()) == ImplicitObjectType::TriangleMesh) return;
@@ -2248,7 +2248,7 @@ typename void ConstructConstraintsImpl(TGeometryParticleHandle<T, d>* Particle0,
 }
 
 template<typename T, int d>
-typename void TPBDCollisionConstraint<T, d>::ConstructConstraints(TGeometryParticleHandle<T, d>* Particle0, TGeometryParticleHandle<T, d>* Particle1, const T Thickness, TArray<typename TPBDCollisionConstraint<T, d>::FRigidBodyContactConstraint>& ConstraintBuffer)
+typename void TPBDCollisionConstraint<T, d>::ConstructConstraints(TGeometryParticleHandle<T, d>* Particle0, TGeometryParticleHandle<T, d>* Particle1, const T Thickness, TArray<TRigidBodyContactConstraint<T, d>>& ConstraintBuffer)
 {
 	ensure(Particle0 && Particle1);
 
diff --git a/Engine/Source/Runtime/Experimental/Chaos/Public/Chaos/PBDCollisionConstraint.h b/Engine/Source/Runtime/Experimental/Chaos/Public/Chaos/PBDCollisionConstraint.h
--- a/Engine/Source/Runtime/Experimental/Chaos/Public/Chaos/PBDCollisionConstraint.h
+++ b/Engine/Source/Runtime/Experimental/Chaos/Public/Chaos/PBDCollisionConstraint.h
@@ -316,7 +316,7 @@ private:
 	template<ECollisionUpdateType>
 	void UpdateConstraint(const T Thickness, FRigidBodyContactConstraint* Constraints, int32 NumConstraints);
 
-	void ConstructConstraints(TGeometryParticleHandle<T, d>* Particle0, TGeometryParticleHandle<T, d>* Particle1, const T Thickness, TArray<typename TPBDCollisionConstraint<T, d>::FRigidBodyContactConstraint>& ConstraintBuffer);
+	void ConstructConstraints(TGeometryParticleHandle<T, d>* Particle0, TGeometryParticleHandle<T, d>* Particle1, const T Thickness, TArray< TRigidBodyContactConstraint<T,d> >& ConstraintBuffer);
 
 
 	template<typename SPATIAL_ACCELERATION>
