Make sure large heightfields don't end up in global array

diff --git a/Engine/Source/Runtime/Experimental/Chaos/Private/Chaos/PBDRigidsEvolution.cpp b/Engine/Source/Runtime/Experimental/Chaos/Private/Chaos/PBDRigidsEvolution.cpp
--- a/Engine/Source/Runtime/Experimental/Chaos/Private/Chaos/PBDRigidsEvolution.cpp
+++ b/Engine/Source/Runtime/Experimental/Chaos/Private/Chaos/PBDRigidsEvolution.cpp
@@ -28,7 +28,7 @@ namespace Chaos
 			MaxTreeDepth = 200;
 			AABBMaxChildrenInLeaf = 500;
 			AABBMaxTreeDepth = 200;
-			MaxPayloadSize = 20000;
+			MaxPayloadSize = 100000;
 			IterationsPerTimeSlice = 2000;
 		}
 	} ConfigSettings;
