Chaos - remove template params from constraint classes.
Also fix Linux/Strict CIS (PBDCollisionConstraints.cpp)

diff --git a/Engine/Source/Runtime/Engine/Private/PhysicsEngine/ImmediatePhysics/ImmediatePhysicsChaos/ImmediatePhysicsJointHandle_Chaos.cpp b/Engine/Source/Runtime/Engine/Private/PhysicsEngine/ImmediatePhysics/ImmediatePhysicsChaos/ImmediatePhysicsJointHandle_Chaos.cpp
--- a/Engine/Source/Runtime/Engine/Private/PhysicsEngine/ImmediatePhysics/ImmediatePhysicsChaos/ImmediatePhysicsJointHandle_Chaos.cpp
+++ b/Engine/Source/Runtime/Engine/Private/PhysicsEngine/ImmediatePhysics/ImmediatePhysicsChaos/ImmediatePhysicsJointHandle_Chaos.cpp
@@ -90,7 +90,7 @@ namespace ImmediatePhysics_Chaos
 		FTransform ConstraintFrame1 = ConstraintInstance->GetRefFrame(EConstraintFrame::Frame1);
 		FTransform ConstraintFrame2 = ConstraintInstance->GetRefFrame(EConstraintFrame::Frame2);
 
-		TPBDJointSettings<float, 3> ConstraintSettings;
+		FPBDJointSettings ConstraintSettings;
 		ConstraintSettings.ConstraintFrames = 
 		{ 
 			ConstraintFrame1.GetRelativeTransform(Actor1->GetLocalCoMTransform()),
diff --git a/Engine/Source/Runtime/Engine/Private/PhysicsEngine/ImmediatePhysics/ImmediatePhysicsChaos/ImmediatePhysicsSimulation_Chaos.cpp b/Engine/Source/Runtime/Engine/Private/PhysicsEngine/ImmediatePhysics/ImmediatePhysicsChaos/ImmediatePhysicsSimulation_Chaos.cpp
--- a/Engine/Source/Runtime/Engine/Private/PhysicsEngine/ImmediatePhysics/ImmediatePhysicsChaos/ImmediatePhysicsSimulation_Chaos.cpp
+++ b/Engine/Source/Runtime/Engine/Private/PhysicsEngine/ImmediatePhysics/ImmediatePhysicsChaos/ImmediatePhysicsSimulation_Chaos.cpp
@@ -27,7 +27,7 @@ FAutoConsoleVariableRef CVarChaosImmPhysDeltaTime(TEXT("p.Chaos.ImmPhys.DeltaTim
 FAutoConsoleVariableRef CVarChaosImmPhysIterations(TEXT("p.Chaos.ImmPhys.Iterations"), ChaosImmediate_Evolution_Iterations, TEXT("Number of constraint solver loops in immediate physics"));
 FAutoConsoleVariableRef CVarChaosImmPhysPushOutIterations(TEXT("p.Chaos.ImmPhys.PushOutIterations"), ChaosImmediate_Evolution_PushOutIterations, TEXT("Set the ApplyPushOut() (position correction) iteration count"));
 
-int32 ChaosImmediate_Collision_Enabled = 1;
+int32 ChaosImmediate_Collision_Enabled = 0;
 int32 ChaosImmediate_Collision_ApplyEnabled = 1;
 int32 ChaosImmediate_Collision_PushOutPairIterations = 2;
 int32 ChaosImmediate_Collision_Priority = 1;
@@ -126,8 +126,8 @@ namespace ImmediatePhysics_Chaos
 		using namespace Chaos;
 
 		Particles = MakeUnique<TPBDRigidsSOAs<FReal, Dimensions>>();
-		Joints = MakeUnique<TPBDJointConstraints<FReal, Dimensions>>(TPBDJointSolverSettings<FReal, Dimensions>());
-		JointsRule = MakeUnique<TPBDConstraintIslandRule<TPBDJointConstraints<FReal, Dimensions>, FReal, Dimensions>>(*Joints);
+		Joints = MakeUnique<FPBDJointConstraints>(FPBDJointSolverSettings());
+		JointsRule = MakeUnique<TPBDConstraintIslandRule<FPBDJointConstraints>>(*Joints);
 		Evolution = MakeUnique<TPBDRigidsEvolutionGBF<FReal, Dimensions>>(*Particles.Get(), 1);
 		TPBDCollisionConstraint<FReal, Dimensions>& Collisions = Evolution->GetCollisionConstraints();
 
@@ -192,7 +192,7 @@ namespace ImmediatePhysics_Chaos
 				}
 			});
 		Joints->SetPreApplyCallback(
-			[this](const float Dt, const TArray<TPBDJointConstraintHandle<float, 3>*>& InConstraintHandles)
+			[this](const float Dt, const TArray<FPBDJointConstraintHandle*>& InConstraintHandles)
 			{
 				if (ChaosImmediate_DebugDrawJoints == 4)
 				{
@@ -200,7 +200,7 @@ namespace ImmediatePhysics_Chaos
 				}
 			});
 		Joints->SetPostApplyCallback(
-			[this](const float Dt, const TArray<TPBDJointConstraintHandle<float, 3>*>& InConstraintHandles)
+			[this](const float Dt, const TArray<FPBDJointConstraintHandle*>& InConstraintHandles)
 			{
 				if (ChaosImmediate_DebugDrawJoints == 4)
 				{
@@ -412,7 +412,7 @@ namespace ImmediatePhysics_Chaos
 			}
 			UE_LOG(LogChaosJoint, Verbose, TEXT("Simulate dt = %f"), DeltaTime);
 
-			TPBDJointSolverSettings<float, 3> JointsSettings = Joints->GetSettings();
+			FPBDJointSolverSettings JointsSettings = Joints->GetSettings();
 			JointsSettings.SwingTwistAngleTolerance = ChaosImmediate_Joint_SwingTwistAngleTolerance;
 			JointsSettings.MinParentMassRatio = ChaosImmediate_Joint_MinParentMassRatio;
 			JointsSettings.MaxInertiaRatio = ChaosImmediate_Joint_MaxInertiaRatio;
@@ -527,7 +527,7 @@ namespace ImmediatePhysics_Chaos
 			if ((ChaosImmediate_DebugDrawJoints >= MinDebugLevel) && (ChaosImmediate_DebugDrawJoints <= MaxDebugLevel))
 			{
 				JointsRule->VisitIslandConstraints(Island,
-					[this, ColorScale](const TArray<TPBDJointConstraintHandle<float, 3>*>& ConstraintHandles)
+					[this, ColorScale](const TArray<FPBDJointConstraintHandle*>& ConstraintHandles)
 					{
 						DebugDraw::DrawJointConstraints(SimulationSpaceTransform, *Joints, ColorScale);
 					});
diff --git a/Engine/Source/Runtime/Engine/Public/Physics/ImmediatePhysics/ImmediatePhysicsChaos/ImmediatePhysicsCore_Chaos.h b/Engine/Source/Runtime/Engine/Public/Physics/ImmediatePhysics/ImmediatePhysicsChaos/ImmediatePhysicsCore_Chaos.h
--- a/Engine/Source/Runtime/Engine/Public/Physics/ImmediatePhysics/ImmediatePhysicsChaos/ImmediatePhysicsCore_Chaos.h
+++ b/Engine/Source/Runtime/Engine/Public/Physics/ImmediatePhysics/ImmediatePhysicsChaos/ImmediatePhysicsCore_Chaos.h
@@ -5,16 +5,17 @@
 #include "Physics/ImmediatePhysics/ImmediatePhysicsChaos/ImmediatePhysicsDeclares_Chaos.h"
 #include "Physics/ImmediatePhysics/ImmediatePhysicsShared/ImmediatePhysicsCore.h"
 
+#include "Chaos/Core.h"
 #include "Chaos/ParticleHandleFwd.h"
 
 namespace Chaos
 {
 	class FImplicitObject;
+	class FPBDJointConstraintHandle;
+	class FPBDJointConstraints;
 	template<typename T, int D> struct TKinematicGeometryParticleParameters;
 	template<typename T, int D> class TKinematicTarget;
-	template<typename P, typename T, int D> class TPBDConstraintIslandRule;
-	template<typename T, int D> class TPBDJointConstraintHandle;
-	template<typename T, int D> class TPBDJointConstraints;
+	template<typename P> class TPBDConstraintIslandRule;
 	template<typename T, int D> struct TPBDRigidParticleParameters;
 	template<typename T, int D> class TPBDRigidsEvolutionGBF;
 	template<typename T, int D> class TPBDRigidsSOAs;
@@ -23,7 +24,7 @@ namespace Chaos
 
 namespace ImmediatePhysics_Chaos
 {
-	using FReal = float;
+	using FReal = Chaos::FReal;
 	const int Dimensions = 3;
 
 	using EActorType = ImmediatePhysics_Shared::EActorType;
diff --git a/Engine/Source/Runtime/Engine/Public/Physics/ImmediatePhysics/ImmediatePhysicsChaos/ImmediatePhysicsJointHandle_Chaos.h b/Engine/Source/Runtime/Engine/Public/Physics/ImmediatePhysics/ImmediatePhysicsChaos/ImmediatePhysicsJointHandle_Chaos.h
--- a/Engine/Source/Runtime/Engine/Public/Physics/ImmediatePhysics/ImmediatePhysicsChaos/ImmediatePhysicsJointHandle_Chaos.h
+++ b/Engine/Source/Runtime/Engine/Public/Physics/ImmediatePhysics/ImmediatePhysicsChaos/ImmediatePhysicsJointHandle_Chaos.h
@@ -15,8 +15,8 @@ namespace ImmediatePhysics_Chaos
 	struct ENGINE_API FJointHandle
 	{
 	public:
-		using FChaosConstraintContainer = Chaos::TPBDJointConstraints<FReal, Dimensions>;
-		using FChaosConstraintHandle = typename Chaos::TPBDJointConstraintHandle<FReal, Dimensions>;
+		using FChaosConstraintContainer = Chaos::FPBDJointConstraints;
+		using FChaosConstraintHandle = typename Chaos::FPBDJointConstraintHandle;
 
 		FJointHandle(FChaosConstraintContainer* InConstraints, FConstraintInstance* ConstraintInstance, FActorHandle* InActor1, FActorHandle* InActor2);
 		~FJointHandle();
diff --git a/Engine/Source/Runtime/Engine/Public/Physics/ImmediatePhysics/ImmediatePhysicsChaos/ImmediatePhysicsSimulation_Chaos.h b/Engine/Source/Runtime/Engine/Public/Physics/ImmediatePhysics/ImmediatePhysicsChaos/ImmediatePhysicsSimulation_Chaos.h
--- a/Engine/Source/Runtime/Engine/Public/Physics/ImmediatePhysics/ImmediatePhysicsChaos/ImmediatePhysicsSimulation_Chaos.h
+++ b/Engine/Source/Runtime/Engine/Public/Physics/ImmediatePhysics/ImmediatePhysicsChaos/ImmediatePhysicsSimulation_Chaos.h
@@ -89,8 +89,8 @@ namespace ImmediatePhysics_Chaos
 		using FParticleHandle = Chaos::TGeometryParticleHandle<FReal, Dimensions>;
 		TUniquePtr<Chaos::TPBDRigidsEvolutionGBF<FReal, Dimensions>> Evolution;
 		TUniquePtr<Chaos::TPBDRigidsSOAs<FReal, Dimensions>> Particles;
-		TUniquePtr<Chaos::TPBDJointConstraints<FReal, Dimensions>> Joints;
-		TUniquePtr<Chaos::TPBDConstraintIslandRule<Chaos::TPBDJointConstraints<FReal, Dimensions>, FReal, Dimensions>> JointsRule;		// @todo(ccaulfield): Don't need islands for anim node physics...
+		TUniquePtr<Chaos::FPBDJointConstraints> Joints;
+		TUniquePtr<Chaos::TPBDConstraintIslandRule<Chaos::FPBDJointConstraints>> JointsRule;		// @todo(ccaulfield): Don't need islands for anim node physics...
 
 		/** Mapping from entity index to handle */
 		// @todo(ccaulfield): we now have handles pointing to handles which is inefficient - we can do better than this, but don't want to change API yet
diff --git a/Engine/Source/Runtime/Experimental/Chaos/Private/Chaos/ChaosDebugDraw.cpp b/Engine/Source/Runtime/Experimental/Chaos/Private/Chaos/ChaosDebugDraw.cpp
--- a/Engine/Source/Runtime/Experimental/Chaos/Private/Chaos/ChaosDebugDraw.cpp
+++ b/Engine/Source/Runtime/Experimental/Chaos/Private/Chaos/ChaosDebugDraw.cpp
@@ -45,49 +45,48 @@ namespace Chaos
 
 #if CHAOS_DEBUG_DRAW
 
-		template <typename T, int d>
-		void DrawShapesImpl(const TRigidTransform<float, 3>& ShapeTransform, const FImplicitObject* Shape, FColor Color)
+		void DrawShapesImpl(const FRigidTransform3& ShapeTransform, const FImplicitObject* Shape, FColor Color)
 		{
 			// @todo(ccaulfield): handle scale throughout
 			switch (Shape->GetType(false))
 			{
 			case ImplicitObjectType::Sphere:
 			{
-				const TSphere<T, d>* Sphere = Shape->template GetObject<TSphere<T, d>>();
-				const TVector<T, d> P = ShapeTransform.TransformPosition(Sphere->GetCenter());
+				const TSphere<FReal, 3>* Sphere = Shape->template GetObject<TSphere<FReal, 3>>();
+				const FVec3 P = ShapeTransform.TransformPosition(Sphere->GetCenter());
 				FDebugDrawQueue::GetInstance().DrawDebugSphere(P, Sphere->GetRadius(), 20, Color, false, KINDA_SMALL_NUMBER, DrawPriority, LineThickness);
 				break;
 			}
 			case ImplicitObjectType::Box:
 			{
-				const TBox<T, d>* Box = Shape->template GetObject<TBox<T, d>>();
-				const TVector<T, d> P = ShapeTransform.TransformPosition(Box->GetCenter());
-				FDebugDrawQueue::GetInstance().DrawDebugBox(P, (T)0.5 * Box->Extents(), ShapeTransform.GetRotation(), Color, false, KINDA_SMALL_NUMBER, DrawPriority, LineThickness);
+				const TBox<FReal, 3>* Box = Shape->template GetObject<TBox<FReal, 3>>();
+				const FVec3 P = ShapeTransform.TransformPosition(Box->GetCenter());
+				FDebugDrawQueue::GetInstance().DrawDebugBox(P, (FReal)0.5 * Box->Extents(), ShapeTransform.GetRotation(), Color, false, KINDA_SMALL_NUMBER, DrawPriority, LineThickness);
 				break;
 			}
 			case ImplicitObjectType::Plane:
 				break;
 			case ImplicitObjectType::Capsule:
 			{
-				const TCapsule<T>* Capsule = Shape->template GetObject<TCapsule<T>>();
-				const TVector<T, d> P = ShapeTransform.TransformPosition(Capsule->GetCenter());
-				const TRotation<T, d> Q = ShapeTransform.GetRotation() * FRotationMatrix::MakeFromZ(Capsule->GetAxis());
-				FDebugDrawQueue::GetInstance().DrawDebugCapsule(P, (T)0.5 * Capsule->GetHeight() + Capsule->GetRadius(), Capsule->GetRadius(), Q, Color, false, KINDA_SMALL_NUMBER, DrawPriority, LineThickness);
+				const TCapsule<FReal>* Capsule = Shape->template GetObject<TCapsule<FReal>>();
+				const FVec3 P = ShapeTransform.TransformPosition(Capsule->GetCenter());
+				const FRotation3 Q = ShapeTransform.GetRotation() * FRotationMatrix::MakeFromZ(Capsule->GetAxis());
+				FDebugDrawQueue::GetInstance().DrawDebugCapsule(P, (FReal)0.5 * Capsule->GetHeight() + Capsule->GetRadius(), Capsule->GetRadius(), Q, Color, false, KINDA_SMALL_NUMBER, DrawPriority, LineThickness);
 				break;
 			}
 			case ImplicitObjectType::Transformed:
 			{
-				const TImplicitObjectTransformed<T, d>* Transformed = Shape->template GetObject<TImplicitObjectTransformed<T, d>>();
-				TRigidTransform<float, 3> TransformedTransform = TRigidTransform<float, 3>(ShapeTransform.TransformPosition(Transformed->GetTransform().GetLocation()), ShapeTransform.GetRotation() * Transformed->GetTransform().GetRotation());
-				DrawShapesImpl<T, d>(TransformedTransform, Transformed->GetTransformedObject(), Color);
+				const TImplicitObjectTransformed<FReal, 3>* Transformed = Shape->template GetObject<TImplicitObjectTransformed<FReal, 3>>();
+				FRigidTransform3 TransformedTransform = FRigidTransform3(ShapeTransform.TransformPosition(Transformed->GetTransform().GetLocation()), ShapeTransform.GetRotation() * Transformed->GetTransform().GetRotation());
+				DrawShapesImpl(TransformedTransform, Transformed->GetTransformedObject(), Color);
 				break;
 			}
 			case ImplicitObjectType::Union:
 			{
-				const TImplicitObjectUnion<T, d>* Union = Shape->template GetObject<TImplicitObjectUnion<T, d>>();
+				const TImplicitObjectUnion<FReal, 3>* Union = Shape->template GetObject<TImplicitObjectUnion<FReal, 3>>();
 				for (auto& UnionShape : Union->GetObjects())
 				{
-					DrawShapesImpl<T, d>(ShapeTransform, UnionShape.Get(), Color);
+					DrawShapesImpl(ShapeTransform, UnionShape.Get(), Color);
 				}
 				break;
 			}
@@ -110,43 +109,40 @@ namespace Chaos
 			}
 		}
 
-		template <typename T, int d>
-		void DrawParticleShapesImpl(const TRigidTransform<float, 3>& SpaceTransform, const TGeometryParticleHandle<T, d>* Particle, T ColorScale)
+		void DrawParticleShapesImpl(const FRigidTransform3& SpaceTransform, const TGeometryParticleHandle<FReal, 3>* Particle, FReal ColorScale)
 		{
 			FColor ShapeColor = Particle->AsDynamic() ? FColor::Yellow : FColor::Orange;
-			FColor Color = (((T)0.5 * ColorScale) * ShapeColor).ToFColor(false);
-			TVector<T, d> P = SpaceTransform.TransformPosition(Particle->AsDynamic() ? Particle->AsDynamic()->P() : Particle->X());
-			TRotation<T, d> Q = SpaceTransform.GetRotation() * (Particle->AsDynamic() ? Particle->AsDynamic()->Q() : Particle->R());
+			FColor Color = (((FReal)0.5 * ColorScale) * ShapeColor).ToFColor(false);
+			FVec3 P = SpaceTransform.TransformPosition(Particle->AsDynamic() ? Particle->AsDynamic()->P() : Particle->X());
+			FRotation3 Q = SpaceTransform.GetRotation() * (Particle->AsDynamic() ? Particle->AsDynamic()->Q() : Particle->R());
 
-			DrawShapesImpl<T, d>(TRigidTransform<T, d>(P, Q), Particle->Geometry().Get(), Color);
+			DrawShapesImpl(FRigidTransform3(P, Q), Particle->Geometry().Get(), Color);
 		}
 
-		template <typename T, int d>
-		void DrawParticleTransformImpl(const TRigidTransform<float, 3>& SpaceTransform, const TGeometryParticleHandle<T, d>* Particle, T ColorScale)
+		void DrawParticleTransformImpl(const FRigidTransform3& SpaceTransform, const TGeometryParticleHandle<FReal, 3>* Particle, FReal ColorScale)
 		{
 			FColor R = (ColorScale * FColor::Red).ToFColor(false);
 			FColor G = (ColorScale * FColor::Green).ToFColor(false);
 			FColor B = (ColorScale * FColor::Blue).ToFColor(false);
 
-			TVector<T, d> P = SpaceTransform.TransformPosition((Particle->AsDynamic()) ? Particle->AsDynamic()->P() : Particle->X());
-			TRotation<T, d> Q = SpaceTransform.GetRotation() * ((Particle->AsDynamic()) ? Particle->AsDynamic()->Q() : Particle->R());
-			PMatrix<T, d, d> Qm = Q.ToMatrix();
+			FVec3 P = SpaceTransform.TransformPosition((Particle->AsDynamic()) ? Particle->AsDynamic()->P() : Particle->X());
+			FRotation3 Q = SpaceTransform.GetRotation() * ((Particle->AsDynamic()) ? Particle->AsDynamic()->Q() : Particle->R());
+			FMatrix33 Qm = Q.ToMatrix();
 			FDebugDrawQueue::GetInstance().DrawDebugDirectionalArrow(P, P + DrawScale * BodyAxisLen * Qm.GetAxis(0), DrawScale * ArrowSize, R, false, KINDA_SMALL_NUMBER, DrawPriority, LineThickness);
 			FDebugDrawQueue::GetInstance().DrawDebugDirectionalArrow(P, P + DrawScale * BodyAxisLen * Qm.GetAxis(1), DrawScale * ArrowSize, G, false, KINDA_SMALL_NUMBER, DrawPriority, LineThickness);
 			FDebugDrawQueue::GetInstance().DrawDebugDirectionalArrow(P, P + DrawScale * BodyAxisLen * Qm.GetAxis(2), DrawScale * ArrowSize, B, false, KINDA_SMALL_NUMBER, DrawPriority, LineThickness);
 		}
 
-		template <typename T, int d>
-		void DrawCollisionImpl(const TRigidTransform<float, 3>& SpaceTransform, const TPBDCollisionConstraintHandle<float, 3>* ConstraintHandle, float ColorScale)
+		void DrawCollisionImpl(const FRigidTransform3& SpaceTransform, const TPBDCollisionConstraintHandle<float, 3>* ConstraintHandle, float ColorScale)
 		{
-			const TRigidBodyContactConstraint<T, d>& Contact = ConstraintHandle->GetContact();
+			const TRigidBodyContactConstraint<FReal, 3>& Contact = ConstraintHandle->GetContact();
 			if (Contact.Phi > 0)
 			{
-				ColorScale = ColorScale * (T)0.1;
+				ColorScale = ColorScale * (FReal)0.1;
 			}
 
-			TVector<T, d> Location = SpaceTransform.TransformPosition(Contact.Location);
-			TVector<T, d> Normal = SpaceTransform.TransformVector(Contact.Normal);
+			FVec3 Location = SpaceTransform.TransformPosition(Contact.Location);
+			FVec3 Normal = SpaceTransform.TransformVector(Contact.Normal);
 
 			if (ContactWidth > 0)
 			{
@@ -167,8 +163,7 @@ namespace Chaos
 			}
 		}
 
-		template<typename T, int d>
-		void DrawJointConstraintImpl(const TRigidTransform<float, 3>& SpaceTransform, const TVector<T, d>& InXa, const PMatrix<T, d, d>& Ra, const TVector<T, d>& InXb, const PMatrix<T, d, d>& Rb, const TVector<T, d>& CR, int32 Level, T ColorScale)
+		void DrawJointConstraintImpl(const FRigidTransform3& SpaceTransform, const FVec3& InXa, const FMatrix33& Ra, const FVec3& InXb, const FMatrix33& Rb, const FVec3& CR, int32 Level, FReal ColorScale)
 		{
 			using namespace Chaos::DebugDraw;
 			FColor R = (ColorScale * FColor::Red).ToFColor(false);
@@ -177,8 +172,8 @@ namespace Chaos
 			FColor C = (ColorScale * FColor::Cyan).ToFColor(false);
 			FColor M = (ColorScale * FColor::Magenta).ToFColor(false);
 			FColor Y = (ColorScale * FColor::Yellow).ToFColor(false);
-			TVector<T, d> Xa = SpaceTransform.TransformPosition(InXa);
-			TVector<T, d> Xb = SpaceTransform.TransformPosition(InXb);
+			FVec3 Xa = SpaceTransform.TransformPosition(InXa);
+			FVec3 Xb = SpaceTransform.TransformPosition(InXb);
 			FDebugDrawQueue::GetInstance().DrawDebugDirectionalArrow(Xa, Xa + DrawScale * ConstraintAxisLen * SpaceTransform.TransformVector(Ra.GetAxis(0)), DrawScale * ArrowSize, R, false, KINDA_SMALL_NUMBER, DrawPriority, LineThickness);
 			FDebugDrawQueue::GetInstance().DrawDebugDirectionalArrow(Xa, Xa + DrawScale * ConstraintAxisLen * SpaceTransform.TransformVector(Ra.GetAxis(1)), DrawScale * ArrowSize, G, false, KINDA_SMALL_NUMBER, DrawPriority, LineThickness);
 			FDebugDrawQueue::GetInstance().DrawDebugDirectionalArrow(Xa, Xa + DrawScale * ConstraintAxisLen * SpaceTransform.TransformVector(Ra.GetAxis(2)), DrawScale * ArrowSize, B, false, KINDA_SMALL_NUMBER, DrawPriority, LineThickness);
@@ -187,33 +182,31 @@ namespace Chaos
 			FDebugDrawQueue::GetInstance().DrawDebugDirectionalArrow(Xb, Xb + DrawScale * ConstraintAxisLen * SpaceTransform.TransformVector(Rb.GetAxis(2)), DrawScale * ArrowSize, Y, false, KINDA_SMALL_NUMBER, DrawPriority, LineThickness);
 			if (Level >= 0)
 			{
-				FDebugDrawQueue::GetInstance().DrawDebugString(Xb + FontHeight * TVector<T, d>(0, 0, 1), FString::Format(TEXT("{0}"), { Level }), nullptr, FColor::Red, KINDA_SMALL_NUMBER, false, FontScale);
+				FDebugDrawQueue::GetInstance().DrawDebugString(Xb + FontHeight * FVec3(0, 0, 1), FString::Format(TEXT("{0}"), { Level }), nullptr, FColor::Red, KINDA_SMALL_NUMBER, false, FontScale);
 			}
-			//FDebugDrawQueue::GetInstance().DrawDebugString(Xb + 3 * FontHeight * TVector<T, d>(0, 0, 1), FString::Format(TEXT("T=  {0}"), { FMath::RadiansToDegrees(CR[(int32)E6DJointAngularConstraintIndex::Twist]) }), nullptr, FColor::Red, KINDA_SMALL_NUMBER, false, FontScale);
-			//FDebugDrawQueue::GetInstance().DrawDebugString(Xb + 2 * FontHeight * TVector<T, d>(0, 0, 1), FString::Format(TEXT("S1= {0}"), { FMath::RadiansToDegrees(CR[(int32)E6DJointAngularConstraintIndex::Swing1]) }), nullptr, FColor::Red, KINDA_SMALL_NUMBER, false, FontScale);
-			//FDebugDrawQueue::GetInstance().DrawDebugString(Xb + 1 * FontHeight * TVector<T, d>(0, 0, 1), FString::Format(TEXT("S2= {0}"), { FMath::RadiansToDegrees(CR[(int32)E6DJointAngularConstraintIndex::Swing2]) }), nullptr, FColor::Red, KINDA_SMALL_NUMBER, false, FontScale);
+			//FDebugDrawQueue::GetInstance().DrawDebugString(Xb + 3 * FontHeight * FVec3(0, 0, 1), FString::Format(TEXT("T=  {0}"), { FMath::RadiansToDegrees(CR[(int32)E6DJointAngularConstraintIndex::Twist]) }), nullptr, FColor::Red, KINDA_SMALL_NUMBER, false, FontScale);
+			//FDebugDrawQueue::GetInstance().DrawDebugString(Xb + 2 * FontHeight * FVec3(0, 0, 1), FString::Format(TEXT("S1= {0}"), { FMath::RadiansToDegrees(CR[(int32)E6DJointAngularConstraintIndex::Swing1]) }), nullptr, FColor::Red, KINDA_SMALL_NUMBER, false, FontScale);
+			//FDebugDrawQueue::GetInstance().DrawDebugString(Xb + 1 * FontHeight * FVec3(0, 0, 1), FString::Format(TEXT("S2= {0}"), { FMath::RadiansToDegrees(CR[(int32)E6DJointAngularConstraintIndex::Swing2]) }), nullptr, FColor::Red, KINDA_SMALL_NUMBER, false, FontScale);
 		}
 
-		template<typename T, int d>
-		void DrawJointConstraintImpl(const TRigidTransform<float, 3>& SpaceTransform, const TPBDJointConstraintHandle<float, 3>* ConstraintHandle, T ColorScale)
+		void DrawJointConstraintImpl(const FRigidTransform3& SpaceTransform, const FPBDJointConstraintHandle* ConstraintHandle, FReal ColorScale)
 		{
-			TVector<TGeometryParticleHandle<T, d>*, 2> ConstrainedParticles = ConstraintHandle->GetConstrainedParticles();
+			TVector<TGeometryParticleHandle<FReal, 3>*, 2> ConstrainedParticles = ConstraintHandle->GetConstrainedParticles();
 			if (ConstrainedParticles[0]->AsDynamic() || ConstrainedParticles[1]->AsDynamic())
 			{
-				TVector<T, d> Xa, Xb, CR;
-				PMatrix<T, d, d> Ra, Rb;
+				FVec3 Xa, Xb, CR;
+				FMatrix33 Ra, Rb;
 				ConstraintHandle->CalculateConstraintSpace(Xa, Ra, Xb, Rb, CR);
-				DrawJointConstraintImpl<T, d>(SpaceTransform, Xa, Ra, Xb, Rb, CR, ConstraintHandle->GetConstraintLevel(), ColorScale);
+				DrawJointConstraintImpl(SpaceTransform, Xa, Ra, Xb, Rb, CR, ConstraintHandle->GetConstraintLevel(), ColorScale);
 			}
 		}
 
-		template<typename T, int d>
-		void Draw6DofConstraintImpl(const TRigidTransform<float, 3>& SpaceTransform, const TPBD6DJointConstraintHandle<float, 3>* ConstraintHandle, T ColorScale)
+		void Draw6DofConstraintImpl(const FRigidTransform3& SpaceTransform, const FPBD6DJointConstraintHandle* ConstraintHandle, FReal ColorScale)
 		{
-			TVector<T, d> Xa, Xb, CR;
-			PMatrix<T, d, d> Ra, Rb;
+			FVec3 Xa, Xb, CR;
+			FMatrix33 Ra, Rb;
 			ConstraintHandle->CalculateConstraintSpace(Xa, Ra, Xb, Rb, CR);
-			DrawJointConstraintImpl<T, d>(SpaceTransform, Xa, Ra, Xb, Rb, CR, INDEX_NONE, ColorScale);
+			DrawJointConstraintImpl(SpaceTransform, Xa, Ra, Xb, Rb, CR, INDEX_NONE, ColorScale);
 		}
 
 #endif
@@ -222,7 +215,7 @@ namespace Chaos
 		//
 		//
 
-		void DrawParticleShapes(const TRigidTransform<float, 3>& SpaceTransform, const TParticleView<TGeometryParticles<float, 3>>& ParticlesView, float ColorScale, bool bDrawKinemtatic, bool bDrawDynamic)
+		void DrawParticleShapes(const FRigidTransform3& SpaceTransform, const TParticleView<TGeometryParticles<float, 3>>& ParticlesView, float ColorScale, bool bDrawKinemtatic, bool bDrawDynamic)
 		{
 #if CHAOS_DEBUG_DRAW
 			if (FDebugDrawQueue::IsDebugDrawingEnabled())
@@ -231,14 +224,14 @@ namespace Chaos
 				{
 					if ((bDrawKinemtatic && !Particle.AsDynamic()) || (bDrawDynamic && Particle.AsDynamic()))
 					{
-						DrawParticleShapesImpl<float, 3>(SpaceTransform, GetHandleHelper<float, 3>(&Particle), ColorScale);
+						DrawParticleShapesImpl(SpaceTransform, GetHandleHelper(&Particle), ColorScale);
 					}
 				}
 			}
 #endif
 		}
 
-		void DrawParticleShapes(const TRigidTransform<float, 3>& SpaceTransform, const TArray<TGeometryParticleHandle<float, 3>*>& Particles, float ColorScale, bool bDrawKinemtatic, bool bDrawDynamic)
+		void DrawParticleShapes(const FRigidTransform3& SpaceTransform, const TArray<TGeometryParticleHandle<float, 3>*>& Particles, float ColorScale, bool bDrawKinemtatic, bool bDrawDynamic)
 		{
 #if CHAOS_DEBUG_DRAW
 			if (FDebugDrawQueue::IsDebugDrawingEnabled())
@@ -247,14 +240,14 @@ namespace Chaos
 				{
 					if ((bDrawKinemtatic && !Particle->AsDynamic()) || (bDrawDynamic && Particle->AsDynamic()))
 					{
-						DrawParticleShapesImpl<float, 3>(SpaceTransform, Particle, ColorScale);
+						DrawParticleShapesImpl(SpaceTransform, Particle, ColorScale);
 					}
 				}
 			}
 #endif
 		}
 
-		void DrawParticleTransforms(const TRigidTransform<float, 3>& SpaceTransform, const TParticleView<TGeometryParticles<float, 3>>& ParticlesView, float ColorScale, bool bDrawKinemtatic, bool bDrawDynamic)
+		void DrawParticleTransforms(const FRigidTransform3& SpaceTransform, const TParticleView<TGeometryParticles<float, 3>>& ParticlesView, float ColorScale, bool bDrawKinemtatic, bool bDrawDynamic)
 		{
 #if CHAOS_DEBUG_DRAW
 			if (FDebugDrawQueue::IsDebugDrawingEnabled())
@@ -263,14 +256,14 @@ namespace Chaos
 				{
 					if ((bDrawKinemtatic && !Particle.AsDynamic()) || (bDrawDynamic && Particle.AsDynamic()))
 					{
-						DrawParticleTransformImpl<float, 3>(SpaceTransform, GetHandleHelper<float, 3>(&Particle), ColorScale);
+						DrawParticleTransformImpl(SpaceTransform, GetHandleHelper(&Particle), ColorScale);
 					}
 				}
 			}
 #endif
 		}
 
-		void DrawParticleTransforms(const TRigidTransform<float, 3>& SpaceTransform, const TArray<TGeometryParticleHandle<float, 3>*>& Particles, float ColorScale, bool bDrawKinemtatic, bool bDrawDynamic)
+		void DrawParticleTransforms(const FRigidTransform3& SpaceTransform, const TArray<TGeometryParticleHandle<float, 3>*>& Particles, float ColorScale, bool bDrawKinemtatic, bool bDrawDynamic)
 		{
 #if CHAOS_DEBUG_DRAW
 			if (FDebugDrawQueue::IsDebugDrawingEnabled())
@@ -279,14 +272,14 @@ namespace Chaos
 				{
 					if ((bDrawKinemtatic && !Particle->AsDynamic()) || (bDrawDynamic && Particle->AsDynamic()))
 					{
-						DrawParticleTransformImpl<float, 3>(SpaceTransform, Particle, ColorScale);
+						DrawParticleTransformImpl(SpaceTransform, Particle, ColorScale);
 					}
 				}
 			}
 #endif
 		}
 
-		void DrawParticleCollisions(const TRigidTransform<float, 3>& SpaceTransform, const TGeometryParticleHandle<float, 3>* Particle, const TPBDCollisionConstraint<float, 3>& Collisions)
+		void DrawParticleCollisions(const FRigidTransform3& SpaceTransform, const TGeometryParticleHandle<float, 3>* Particle, const TPBDCollisionConstraint<float, 3>& Collisions)
 		{
 #if CHAOS_DEBUG_DRAW
 			if (FDebugDrawQueue::IsDebugDrawingEnabled())
@@ -296,88 +289,88 @@ namespace Chaos
 					TVector<TGeometryParticleHandle<float, 3>*, 2> ConstrainedParticles = Collisions.GetConstrainedParticles(ConstraintIndex);
 					if ((ConstrainedParticles[0] == Particle) || (ConstrainedParticles[1] == Particle))
 					{
-						const typename TPBDCollisionConstraint<float, 3>::FConstraintHandle* ConstraintHandle = Collisions.GetConstraintHandle(ConstraintIndex);
-						DrawCollisionImpl<float, 3>(SpaceTransform, ConstraintHandle, 1.0f);
+						const typename TPBDCollisionConstraint<float, 3>::FConstraintContainerHandle* ConstraintHandle = Collisions.GetConstraintHandle(ConstraintIndex);
+						DrawCollisionImpl(SpaceTransform, ConstraintHandle, 1.0f);
 					}
 				}
 			}
 #endif
 		}
 
-		void DrawCollisions(const TRigidTransform<float, 3>& SpaceTransform, const TPBDCollisionConstraint<float, 3>& Collisions, float ColorScale)
+		void DrawCollisions(const FRigidTransform3& SpaceTransform, const TPBDCollisionConstraint<float, 3>& Collisions, float ColorScale)
 		{
 #if CHAOS_DEBUG_DRAW
 			if (FDebugDrawQueue::IsDebugDrawingEnabled())
 			{
 				for (int32 ConstraintIndex = 0; ConstraintIndex < Collisions.NumConstraints(); ++ConstraintIndex)
 				{
-					DrawCollisionImpl<float, 3>(SpaceTransform, Collisions.GetConstraintHandle(ConstraintIndex), ColorScale);
+					DrawCollisionImpl(SpaceTransform, Collisions.GetConstraintHandle(ConstraintIndex), ColorScale);
 				}
 			}
 #endif
 		}
 
-		void DrawCollisions(const TRigidTransform<float, 3>& SpaceTransform, const TArray<TPBDCollisionConstraintHandle<float, 3>*>& ConstraintHandles, float ColorScale)
+		void DrawCollisions(const FRigidTransform3& SpaceTransform, const TArray<TPBDCollisionConstraintHandle<float, 3>*>& ConstraintHandles, float ColorScale)
 		{
 #if CHAOS_DEBUG_DRAW
 			if (FDebugDrawQueue::IsDebugDrawingEnabled())
 			{
 				for (const TPBDCollisionConstraintHandle<float, 3>* ConstraintHandle : ConstraintHandles)
 				{
-					DrawCollisionImpl<float, 3>(SpaceTransform, ConstraintHandle, ColorScale);
+					DrawCollisionImpl(SpaceTransform, ConstraintHandle, ColorScale);
 				}
 			}
 #endif
 		}
 
 
-		void DrawJointConstraints(const TRigidTransform<float, 3>& SpaceTransform, const TArray<TPBDJointConstraintHandle<float, 3>*>& ConstraintHandles, float ColorScale)
+		void DrawJointConstraints(const FRigidTransform3& SpaceTransform, const TArray<FPBDJointConstraintHandle*>& ConstraintHandles, float ColorScale)
 		{
 #if CHAOS_DEBUG_DRAW
 			if (FDebugDrawQueue::IsDebugDrawingEnabled())
 			{
-				for (const TPBDJointConstraintHandle<float, 3>* ConstraintHandle : ConstraintHandles)
+				for (const FPBDJointConstraintHandle* ConstraintHandle : ConstraintHandles)
 				{
-					DrawJointConstraintImpl<float, 3>(SpaceTransform, ConstraintHandle, ColorScale);
+					DrawJointConstraintImpl(SpaceTransform, ConstraintHandle, ColorScale);
 				}
 			}
 #endif
 		}
 
-		void DrawJointConstraints(const TRigidTransform<float, 3>& SpaceTransform, const TPBDJointConstraints<float, 3>& Constraints, float ColorScale)
+		void DrawJointConstraints(const FRigidTransform3& SpaceTransform, const FPBDJointConstraints& Constraints, float ColorScale)
 		{
 #if CHAOS_DEBUG_DRAW
 			if (FDebugDrawQueue::IsDebugDrawingEnabled())
 			{
 				for (int32 ConstraintIndex = 0; ConstraintIndex < Constraints.NumConstraints(); ++ConstraintIndex)
 				{
-					DrawJointConstraintImpl<float, 3>(SpaceTransform, Constraints.GetConstraintHandle(ConstraintIndex), ColorScale);
+					DrawJointConstraintImpl(SpaceTransform, Constraints.GetConstraintHandle(ConstraintIndex), ColorScale);
 				}
 			}
 #endif
 		}
 
-		void Draw6DofConstraints(const TRigidTransform<float, 3>& SpaceTransform, const TArray<TPBD6DJointConstraintHandle<float, 3>*>& ConstraintHandles, float ColorScale)
+		void Draw6DofConstraints(const FRigidTransform3& SpaceTransform, const TArray<FPBD6DJointConstraintHandle*>& ConstraintHandles, float ColorScale)
 		{
 #if CHAOS_DEBUG_DRAW
 			if (FDebugDrawQueue::IsDebugDrawingEnabled())
 			{
-				for (const TPBD6DJointConstraintHandle<float, 3>* ConstraintHandle : ConstraintHandles)
+				for (const FPBD6DJointConstraintHandle* ConstraintHandle : ConstraintHandles)
 				{
-					Draw6DofConstraintImpl<float, 3>(SpaceTransform, ConstraintHandle, ColorScale);
+					Draw6DofConstraintImpl(SpaceTransform, ConstraintHandle, ColorScale);
 				}
 			}
 #endif
 		}
 
-		void Draw6DofConstraints(const TRigidTransform<float, 3>& SpaceTransform, const TPBD6DJointConstraints<float, 3>& Constraints, float ColorScale)
+		void Draw6DofConstraints(const FRigidTransform3& SpaceTransform, const FPBD6DJointConstraints& Constraints, float ColorScale)
 		{
 #if CHAOS_DEBUG_DRAW
 			if (FDebugDrawQueue::IsDebugDrawingEnabled())
 			{
 				for (int32 ConstraintIndex = 0; ConstraintIndex < Constraints.NumConstraints(); ++ConstraintIndex)
 				{
-					Draw6DofConstraintImpl<float, 3>(SpaceTransform, Constraints.GetConstraintHandle(ConstraintIndex), ColorScale);
+					Draw6DofConstraintImpl(SpaceTransform, Constraints.GetConstraintHandle(ConstraintIndex), ColorScale);
 				}
 			}
 #endif
