Chaos - remove template params from constraint classes.

diff --git a/Engine/Source/Runtime/Experimental/Chaos/Private/Chaos/PBD6DJointConstraints.cpp b/Engine/Source/Runtime/Experimental/Chaos/Private/Chaos/PBD6DJointConstraints.cpp
--- a/Engine/Source/Runtime/Experimental/Chaos/Private/Chaos/PBD6DJointConstraints.cpp
+++ b/Engine/Source/Runtime/Experimental/Chaos/Private/Chaos/PBD6DJointConstraints.cpp
@@ -10,1923 +10,1818 @@
 
 //#pragma optimize("", off)
 
-using namespace Chaos;
-
-//
-// Constraint Handle
-//
-
-template<typename T, int d>
-TPBD6DJointConstraintHandle<T, d>::TPBD6DJointConstraintHandle()
-{
-}
-
-template<typename T, int d>
-TPBD6DJointConstraintHandle<T, d>::TPBD6DJointConstraintHandle(FConstraintContainer* InConstraintContainer, int32 InConstraintIndex) : TContainerConstraintHandle<TPBD6DJointConstraints<T, d>>(InConstraintContainer, InConstraintIndex)
-{
-}
-
-template<typename T, int d>
-void TPBD6DJointConstraintHandle<T, d>::CalculateConstraintSpace(TVector<T, d>& OutXa, PMatrix<T, d, d>& OutRa, TVector<T, d>& OutXb, PMatrix<T, d, d>& OutRb, TVector<T, d>& OutCR) const
-{
-	ConstraintContainer->CalculateConstraintSpace(ConstraintIndex, OutXa, OutRa, OutXb, OutRb, OutCR);
-}
-
-template<typename T, int d>
-void TPBD6DJointConstraintHandle<T, d>::SetParticleLevels(const TVector<int32, 2>& ParticleLevels)
-{
-	ConstraintContainer->SetParticleLevels(ConstraintIndex, ParticleLevels);
-}
-
-template<typename T, int d>
-int32 TPBD6DJointConstraintHandle<T, d>::GetConstraintLevel() const
-{
-	return ConstraintContainer->GetConstraintLevel(ConstraintIndex);
-}
-
-//
-// Constraint Math
-//
-
-/**
-	* Invert a 2x2 block of square matrices using block-wise inversion
-	* https://en.wikipedia.org/wiki/Invertible_matrix#Blockwise_inversion
-	*
-	*	| AI		BI	|	=	| A		B |^-1
-	*	| CI		DI	|		| C		D |
-	*
-	*/
-template<typename T, int d>
-void TPBD6DJointConstraintUtilities<T, d>::BlockwiseInverse(const PMatrix<T, d, d>& A, const PMatrix<T, d, d>& B, const PMatrix<T, d, d>& C, const PMatrix<T, d, d>& D, PMatrix<T, d, d>& AI, PMatrix<T, d, d>& BI, PMatrix<T, d, d>& CI, PMatrix<T, d, d>& DI)
-{
-	PMatrix<T, d, d> AInv = A.Inverse();
-	PMatrix<T, d, d> ZInv = (D - Utilities::Multiply(C, Utilities::Multiply(AInv, B))).Inverse();
-	AI = AInv + Utilities::Multiply(AInv, Utilities::Multiply(B, Utilities::Multiply(ZInv, Utilities::Multiply(C, AInv))));
-	BI = -Utilities::Multiply(AInv, Utilities::Multiply(B, ZInv));
-	CI = -Utilities::Multiply(ZInv, Utilities::Multiply(C, AInv));
-	DI = ZInv;
-}
-
-template<typename T, int d>
-void TPBD6DJointConstraintUtilities<T, d>::BlockwiseInverse2(const PMatrix<T, d, d>& A, const PMatrix<T, d, d>& B, const PMatrix<T, d, d>& C, const PMatrix<T, d, d>& D, PMatrix<T, d, d>& AI, PMatrix<T, d, d>& BI, PMatrix<T, d, d>& CI, PMatrix<T, d, d>& DI)
-{
-	PMatrix<T, d, d> DInv = D.Inverse();
-	PMatrix<T, d, d> ZInv = (A - Utilities::Multiply(B, Utilities::Multiply(DInv, C))).Inverse();
-	AI = ZInv;
-	BI = -Utilities::Multiply(ZInv, Utilities::Multiply(B, DInv));
-	CI = -Utilities::Multiply(DInv, Utilities::Multiply(C, ZInv));
-	DI = DInv + Utilities::Multiply(DInv, Utilities::Multiply(C, Utilities::Multiply(ZInv, Utilities::Multiply(B, DInv))));
-}
-
-/**
- * This function returns:
- *
- *	F(X,R)	=	[J . M^-1 . J^t]
- *
- * for a single body, which is the (6x6) matrix that is the reciprocal component of the lambda matrix:
- *
- *	L(X,R)	=	(1 / [J . M^-1 . J^t]) . C	=	(1 / F) . C
- *
- * (although the reciprocal part is the sum of F for both bodies). Lambda is part of the constraint correction calculation:
- *
- *	D(X,R)	=	M^-1 . J . L
- *
- * F is returned as 4 3x3 matrices F00...F11 to avoid the need for a 6x6 matrix (plus, the problem splits nicely into
- * 3x3 matrices for position and rotation degrees of freedom, although it is wasteful if we are not using all 3 rotation constraints...)
- *
- * F	=	| F00	F01	|
- * 			| F10	F11	|
- *
- */
-template<typename T, int d>
-void TPBD6DJointConstraintUtilities<T, d>::ComputeJointFactorMatrix(const PMatrix<T, d, d>& XR, const PMatrix<T, d, d>& RR, float MInv, const PMatrix<T, d, d>& IInv, PMatrix<T, d, d>& M00, PMatrix<T, d, d>& M01, PMatrix<T, d, d>& M10, PMatrix<T, d, d>& M11)
+namespace Chaos
 {
-	// XR(3x3)	= dCX/DR: derivative of position constraint error with respect to rotation. 
-	//			= Cross-product matrix of X (the world-space offset of the constraint)
-	// RR(3x3)	= dCR/DR: derivative of world-space rotation constraint error with respect to rotation.
-	// 			=	|	(Swing1a x Swing2b)	|
-	//				|	(Swing1a x Twistb)	|
-	//				|	(Swing2a x Twistb)	|
-	//
-	// A^t = Transpose(A)
-	//
-	//	MInv	=	|	1(3x3).MInv		0(3x3)		|
-	//	(6x6)		|	0(3x3)			IInv		|
-	//
-	// Jacobian:
 	//
-	//	J(C(X,R)) =	|	dCX/dX			dCX/dR		|
-	//	(6x6)		|	dCR/dX			dCR/dR		|
+	// Constraint Handle
 	//
-	//			=	|	1(3x3)			-XR			|
-	//				|	0(3x3)			RR			|
-	//
-	// FactorMatrix:
-	//
-	//	F		=	J.MInv.J^t
-	//	(6x6)
-	//			=	|	1/Ma(3x3) - XR.IInv.XR		-XR.IInv.RR^t	|
-	//				|	(-XR.IInv.RR^t)^t			RR.IInv.RR^t	|
-	//
-	// @todo(ccaulfield): optimize
-	PMatrix<T, d, d> IXR = Utilities::Multiply(IInv, XR);
-	PMatrix<T, d, d> IRR = Utilities::Multiply(IInv, RR.GetTransposed());
-	M00 = PMatrix<T, d, d>(MInv, MInv, MInv) - Utilities::Multiply(XR, IXR);
-	M01 = -Utilities::Multiply(XR, IRR);
-	M10 = M01.GetTransposed();
-	M11 = Utilities::Multiply(RR, IRR);
-}
 
-template<typename T, int d>
-TVector<T, d> TPBD6DJointConstraintUtilities<T, d>::Calculate6dConstraintAngles(
-	const TPBD6DJointSolverSettings<T, d>& SolverSettings,
-	const TRotation<T, d>& Ra,
-	const TRotation<T, d>& Rb, 
-	const TPBD6DJointMotionSettings<T, d>& RSettings)
-{
-	// @todo(ccaulfield): optimize
-	TVector<T, d> CR;
-	PMatrix<T, d, d> RRa, RRb;
-	Calculate6dConstraintRotation(SolverSettings, Ra, Rb, RSettings, CR, RRa, RRb);
-	return CR;
-}
 
-template<typename T, int d>
-bool TPBD6DJointConstraintUtilities<T, d>::Calculate6dConstraintRotation_SwingFixed(
-	const TPBD6DJointSolverSettings<T, d>& SolverSettings,
-	const TRotation<T, d>& Ra,
-	const TRotation<T, d>& Rb, 
-	const TPBD6DJointMotionSettings<T, d>& MotionSettings,
-	TVector<T, d>& CR, 
-	PMatrix<T, d, d>& RRa, 
-	PMatrix<T, d, d>& RRb)
-{
-	// @todo(ccaulfield): optimize
-
-	// Get the transform from A to B, and use it to generate twist angles.
-	TRotation<T, d> Rab = Ra.Inverse() * Rb;
-	TRotation<T, d> RTwist, RSwing;
-	Rab.ToSwingTwist(T6DJointConstants<T, d>::TwistAxis(), RSwing, RTwist);
-	RSwing = RSwing.GetNormalized();
-	RTwist = RTwist.GetNormalized();
-
-	TVector<T, d> TwistAxisAB;
-	T TwistAngleAB;
-	RTwist.ToAxisAndAngleSafe(TwistAxisAB, TwistAngleAB, T6DJointConstants<T, d>::TwistAxis());
-	if (TwistAngleAB > PI)
+	FPBD6DJointConstraintHandle::FPBD6DJointConstraintHandle()
 	{
-		TwistAngleAB = TwistAngleAB - (T)2 * PI;
-		RTwist = TRotation<T, d>::FromAxisAngle(T6DJointConstants<T, d>::TwistAxis(), TwistAngleAB);
 	}
 
-	PMatrix<T, d, d> Axesa = Ra.ToMatrix();
-	PMatrix<T, d, d> Axesb = Rb.ToMatrix();
-
-	// Constraint-space in body A is just the constraint transform
-	TVector<T, d> Twista = Axesa.GetAxis((int32)E6DJointAngularAxisIndex::Twist);
-	TVector<T, d> Swing1a = Axesa.GetAxis((int32)E6DJointAngularAxisIndex::Swing1);
-	TVector<T, d> Swing2a = Axesa.GetAxis((int32)E6DJointAngularAxisIndex::Swing2);
-
-	// Remove Twist from body B's swing axes
-	TVector<T, d> Twistb = Axesb.GetAxis((int32)E6DJointAngularAxisIndex::Twist);
-	TVector<T, d> Swing1b = Rb * RTwist.Inverse() * T6DJointConstants<T, d>::Swing1Axis();
-	TVector<T, d> Swing2b = TVector<T, d>::CrossProduct(Swing1b, Twistb);
-
-	RRa.SetRow((int32)E6DJointAngularConstraintIndex::Twist, Twista);
-	RRa.SetRow((int32)E6DJointAngularConstraintIndex::Swing1, Swing1a);
-	RRa.SetRow((int32)E6DJointAngularConstraintIndex::Swing2, Swing2a);
-	RRb.SetRow((int32)E6DJointAngularConstraintIndex::Twist, Twistb);
-	RRb.SetRow((int32)E6DJointAngularConstraintIndex::Swing1, Swing1b);
-	RRb.SetRow((int32)E6DJointAngularConstraintIndex::Swing2, Swing2b);
-
-	T TwistAngle = TwistAngleAB * TVector<T, d>::DotProduct(TwistAxisAB, T6DJointConstants<T, d>::TwistAxis());
-	T Swing1Angle = FMath::Asin(FMath::Clamp(TVector<T, d>::DotProduct(TVector<T, d>::CrossProduct(Swing2a, Swing2b), Swing1a), (T)-1, (T)1));
-	T Swing2Angle = FMath::Asin(FMath::Clamp(TVector<T, d>::DotProduct(TVector<T, d>::CrossProduct(Swing1a, Swing1b), Swing2a), (T)-1, (T)1));
-
-	CR[(int32)E6DJointAngularConstraintIndex::Twist] = TwistAngle;
-	CR[(int32)E6DJointAngularConstraintIndex::Swing1] = Swing1Angle;
-	CR[(int32)E6DJointAngularConstraintIndex::Swing2] = Swing2Angle;
-
-	// If we're flipped 180 degrees about swing, just pretend the error is zero
-	T DotTT = TVector<T, d>::DotProduct(Twista, Twistb);
-	const T MinDotTT = (T)-1 + SolverSettings.InvertedAxisTolerance;
-	if (DotTT < MinDotTT)
+
+	FPBD6DJointConstraintHandle::FPBD6DJointConstraintHandle(FConstraintContainer* InConstraintContainer, int32 InConstraintIndex) : TContainerConstraintHandle<FPBD6DJointConstraints>(InConstraintContainer, InConstraintIndex)
 	{
-		return false;
 	}
 
-	return true;
-}
 
-template<typename T, int d>
-void TPBD6DJointConstraintUtilities<T, d>::Calculate6dConstraintRotationLimits_SwingFixed(
-	const TPBD6DJointSolverSettings<T, d>& SolverSettings,
-	const TRotation<T, d>& Ra,
-	const TRotation<T, d>& Rb, 
-	const TPBD6DJointMotionSettings<T, d>& MotionSettings, 
-	TVector<T, d>& SR,
-	TVector<T, d>& CR,
-	PMatrix<T, d, d>& RRa, 
-	PMatrix<T, d, d>& RRb, 
-	TVector<T, d>& LRMin, 
-	TVector<T, d>& LRMax)
-{
-	// Convert the target rotation into target angles
-	// @todo(ccaulfield): optimize (cache these values, or store them directly rather than the target rotation)
-	TRotation<T, d> DriveTwist, DriveSwing;
-	MotionSettings.AngularDriveTarget.ToSwingTwist(T6DJointConstants<T, d>::TwistAxis(), DriveSwing, DriveTwist);
-	
-	T DriveTwistAngle, DriveSwingAngle;
-	TVector<T, d> DriveTwistAxis, DriveSwingAxis;
-	DriveTwist.ToAxisAndAngleSafe(DriveTwistAxis, DriveTwistAngle, T6DJointConstants<T, d>::TwistAxis());
-	DriveSwing.ToAxisAndAngleSafe(DriveSwingAxis, DriveSwingAngle, T6DJointConstants<T, d>::Swing1Axis());
-	if (TVector<T, d>::DotProduct(DriveTwistAxis, T6DJointConstants<T, d>::TwistAxis()) < 0)
+	void FPBD6DJointConstraintHandle::CalculateConstraintSpace(FVec3& OutXa, FMatrix33& OutRa, FVec3& OutXb, FMatrix33& OutRb, FVec3& OutCR) const
 	{
-		DriveTwistAngle = -DriveTwistAngle;
+		ConstraintContainer->CalculateConstraintSpace(ConstraintIndex, OutXa, OutRa, OutXb, OutRb, OutCR);
 	}
-	if ((TVector<T, d>::DotProduct(DriveSwingAxis, T6DJointConstants<T, d>::Swing1Axis()) < -(T)0.9) || (TVector<T, d>::DotProduct(DriveSwingAxis, T6DJointConstants<T, d>::Swing2Axis()) < -(T)0.9))
+
+
+	void FPBD6DJointConstraintHandle::SetParticleLevels(const TVector<int32, 2>& ParticleLevels)
 	{
-		DriveSwingAngle = -DriveSwingAngle;
+		ConstraintContainer->SetParticleLevels(ConstraintIndex, ParticleLevels);
 	}
 
-	bool bDriveEnabled[3] = { false, false, false };
-	if (SolverSettings.bEnableDrives)
+
+	int32 FPBD6DJointConstraintHandle::GetConstraintLevel() const
 	{
-		bDriveEnabled[(int32)E6DJointAngularConstraintIndex::Twist] = MotionSettings.bAngularTwistDriveEnabled;
-		bDriveEnabled[(int32)E6DJointAngularConstraintIndex::Swing1] = MotionSettings.bAngularSwingDriveEnabled;
-		bDriveEnabled[(int32)E6DJointAngularConstraintIndex::Swing2] = MotionSettings.bAngularSwingDriveEnabled;
+		return ConstraintContainer->GetConstraintLevel(ConstraintIndex);
 	}
 
-	T DriveAngles[3] = { 0,0,0 };
-	DriveAngles[(int32)E6DJointAngularConstraintIndex::Twist] = DriveTwistAngle;
-	DriveAngles[(int32)E6DJointAngularConstraintIndex::Swing1] = DriveSwingAngle;
-	DriveAngles[(int32)E6DJointAngularConstraintIndex::Swing2] = DriveSwingAngle;
+	//
+	// Constraint Math
+	//
+
+	/**
+		* Invert a 2x2 block of square matrices using block-wise inversion
+		* https://en.wikipedia.org/wiki/Invertible_matrix#Blockwise_inversion
+		*
+		*	| AI		BI	|	=	| A		B |^-1
+		*	| CI		DI	|		| C		D |
+		*
+		*/
+
+	void FPBD6DJointConstraintUtilities::BlockwiseInverse(const FMatrix33& A, const FMatrix33& B, const FMatrix33& C, const FMatrix33& D, FMatrix33& AI, FMatrix33& BI, FMatrix33& CI, FMatrix33& DI)
+	{
+		FMatrix33 AInv = A.Inverse();
+		FMatrix33 ZInv = (D - Utilities::Multiply(C, Utilities::Multiply(AInv, B))).Inverse();
+		AI = AInv + Utilities::Multiply(AInv, Utilities::Multiply(B, Utilities::Multiply(ZInv, Utilities::Multiply(C, AInv))));
+		BI = -Utilities::Multiply(AInv, Utilities::Multiply(B, ZInv));
+		CI = -Utilities::Multiply(ZInv, Utilities::Multiply(C, AInv));
+		DI = ZInv;
+	}
+
+
+	void FPBD6DJointConstraintUtilities::BlockwiseInverse2(const FMatrix33& A, const FMatrix33& B, const FMatrix33& C, const FMatrix33& D, FMatrix33& AI, FMatrix33& BI, FMatrix33& CI, FMatrix33& DI)
+	{
+		FMatrix33 DInv = D.Inverse();
+		FMatrix33 ZInv = (A - Utilities::Multiply(B, Utilities::Multiply(DInv, C))).Inverse();
+		AI = ZInv;
+		BI = -Utilities::Multiply(ZInv, Utilities::Multiply(B, DInv));
+		CI = -Utilities::Multiply(DInv, Utilities::Multiply(C, ZInv));
+		DI = DInv + Utilities::Multiply(DInv, Utilities::Multiply(C, Utilities::Multiply(ZInv, Utilities::Multiply(B, DInv))));
+	}
+
+	/**
+	 * This function returns:
+	 *
+	 *	F(X,R)	=	[J . M^-1 . J^t]
+	 *
+	 * for a single body, which is the (6x6) matrix that is the reciprocal component of the lambda matrix:
+	 *
+	 *	L(X,R)	=	(1 / [J . M^-1 . J^t]) . C	=	(1 / F) . C
+	 *
+	 * (although the reciprocal part is the sum of F for both bodies). Lambda is part of the constraint correction calculation:
+	 *
+	 *	D(X,R)	=	M^-1 . J . L
+	 *
+	 * F is returned as 4 3x3 matrices F00...F11 to avoid the need for a 6x6 matrix (plus, the problem splits nicely into
+	 * 3x3 matrices for position and rotation degrees of freedom, although it is wasteful if we are not using all 3 rotation constraints...)
+	 *
+	 * F	=	| F00	F01	|
+	 * 			| F10	F11	|
+	 *
+	 */
+
+	void FPBD6DJointConstraintUtilities::ComputeJointFactorMatrix(const FMatrix33& XR, const FMatrix33& RR, float MInv, const FMatrix33& IInv, FMatrix33& M00, FMatrix33& M01, FMatrix33& M10, FMatrix33& M11)
+	{
+		// XR(3x3)	= dCX/DR: derivative of position constraint error with respect to rotation. 
+		//			= Cross-product matrix of X (the world-space offset of the constraint)
+		// RR(3x3)	= dCR/DR: derivative of world-space rotation constraint error with respect to rotation.
+		// 			=	|	(Swing1a x Swing2b)	|
+		//				|	(Swing1a x Twistb)	|
+		//				|	(Swing2a x Twistb)	|
+		//
+		// A^t = Transpose(A)
+		//
+		//	MInv	=	|	1(3x3).MInv		0(3x3)		|
+		//	(6x6)		|	0(3x3)			IInv		|
+		//
+		// Jacobian:
+		//
+		//	J(C(X,R)) =	|	dCX/dX			dCX/dR		|
+		//	(6x6)		|	dCR/dX			dCR/dR		|
+		//
+		//			=	|	1(3x3)			-XR			|
+		//				|	0(3x3)			RR			|
+		//
+		// FactorMatrix:
+		//
+		//	F		=	J.MInv.J^t
+		//	(6x6)
+		//			=	|	1/Ma(3x3) - XR.IInv.XR		-XR.IInv.RR^t	|
+		//				|	(-XR.IInv.RR^t)^t			RR.IInv.RR^t	|
+		//
+		// @todo(ccaulfield): optimize
+		FMatrix33 IXR = Utilities::Multiply(IInv, XR);
+		FMatrix33 IRR = Utilities::Multiply(IInv, RR.GetTransposed());
+		M00 = FMatrix33(MInv, MInv, MInv) - Utilities::Multiply(XR, IXR);
+		M01 = -Utilities::Multiply(XR, IRR);
+		M10 = M01.GetTransposed();
+		M11 = Utilities::Multiply(RR, IRR);
+	}
+
+
+	FVec3 FPBD6DJointConstraintUtilities::Calculate6dConstraintAngles(
+		const FPBD6DJointSolverSettings& SolverSettings,
+		const FRotation3& Ra,
+		const FRotation3& Rb, 
+		const FPBD6DJointMotionSettings& RSettings)
+	{
+		// @todo(ccaulfield): optimize
+		FVec3 CR;
+		FMatrix33 RRa, RRb;
+		Calculate6dConstraintRotation(SolverSettings, Ra, Rb, RSettings, CR, RRa, RRb);
+		return CR;
+	}
+
+
+	bool FPBD6DJointConstraintUtilities::Calculate6dConstraintRotation_SwingFixed(
+		const FPBD6DJointSolverSettings& SolverSettings,
+		const FRotation3& Ra,
+		const FRotation3& Rb, 
+		const FPBD6DJointMotionSettings& MotionSettings,
+		FVec3& CR, 
+		FMatrix33& RRa, 
+		FMatrix33& RRb)
+	{
+		// @todo(ccaulfield): optimize
+
+		// Get the transform from A to B, and use it to generate twist angles.
+		FRotation3 Rab = Ra.Inverse() * Rb;
+		FRotation3 RTwist, RSwing;
+		Rab.ToSwingTwist(F6DJointConstants::TwistAxis(), RSwing, RTwist);
+		RSwing = RSwing.GetNormalized();
+		RTwist = RTwist.GetNormalized();
+
+		FVec3 TwistAxisAB;
+		FReal TwistAngleAB;
+		RTwist.ToAxisAndAngleSafe(TwistAxisAB, TwistAngleAB, F6DJointConstants::TwistAxis());
+		if (TwistAngleAB > PI)
+		{
+			TwistAngleAB = TwistAngleAB - (FReal)2 * PI;
+			RTwist = FRotation3::FromAxisAngle(F6DJointConstants::TwistAxis(), TwistAngleAB);
+		}
 
-	T DriveStiffness = (SolverSettings.PBDDriveStiffness > 0) ? SolverSettings.PBDDriveStiffness : MotionSettings.AngularDriveStiffness;
+		FMatrix33 Axesa = Ra.ToMatrix();
+		FMatrix33 Axesb = Rb.ToMatrix();
+
+		// Constraint-space in body A is just the constraint transform
+		FVec3 Twista = Axesa.GetAxis((int32)E6DJointAngularAxisIndex::Twist);
+		FVec3 Swing1a = Axesa.GetAxis((int32)E6DJointAngularAxisIndex::Swing1);
+		FVec3 Swing2a = Axesa.GetAxis((int32)E6DJointAngularAxisIndex::Swing2);
+
+		// Remove Twist from body B's swing axes
+		FVec3 Twistb = Axesb.GetAxis((int32)E6DJointAngularAxisIndex::Twist);
+		FVec3 Swing1b = Rb * RTwist.Inverse() * F6DJointConstants::Swing1Axis();
+		FVec3 Swing2b = FVec3::CrossProduct(Swing1b, Twistb);
+
+		RRa.SetRow((int32)E6DJointAngularConstraintIndex::Twist, Twista);
+		RRa.SetRow((int32)E6DJointAngularConstraintIndex::Swing1, Swing1a);
+		RRa.SetRow((int32)E6DJointAngularConstraintIndex::Swing2, Swing2a);
+		RRb.SetRow((int32)E6DJointAngularConstraintIndex::Twist, Twistb);
+		RRb.SetRow((int32)E6DJointAngularConstraintIndex::Swing1, Swing1b);
+		RRb.SetRow((int32)E6DJointAngularConstraintIndex::Swing2, Swing2b);
+
+		FReal TwistAngle = TwistAngleAB * FVec3::DotProduct(TwistAxisAB, F6DJointConstants::TwistAxis());
+		FReal Swing1Angle = FMath::Asin(FMath::Clamp(FVec3::DotProduct(FVec3::CrossProduct(Swing2a, Swing2b), Swing1a), (FReal)-1, (FReal)1));
+		FReal Swing2Angle = FMath::Asin(FMath::Clamp(FVec3::DotProduct(FVec3::CrossProduct(Swing1a, Swing1b), Swing2a), (FReal)-1, (FReal)1));
+
+		CR[(int32)E6DJointAngularConstraintIndex::Twist] = TwistAngle;
+		CR[(int32)E6DJointAngularConstraintIndex::Swing1] = Swing1Angle;
+		CR[(int32)E6DJointAngularConstraintIndex::Swing2] = Swing2Angle;
+
+		// If we're flipped 180 degrees about swing, just pretend the error is zero
+		FReal DotTT = FVec3::DotProduct(Twista, Twistb);
+		const FReal MinDotTT = (FReal)-1 + SolverSettings.InvertedAxisTolerance;
+		if (DotTT < MinDotTT)
+		{
+			return false;
+		}
 
-	// Use constraint limits settings to specify valid range for constraint-space rotation corrections
-	for (int32 Axis = 0; Axis < d; ++Axis)
+		return true;
+	}
+
+
+	void FPBD6DJointConstraintUtilities::Calculate6dConstraintRotationLimits_SwingFixed(
+		const FPBD6DJointSolverSettings& SolverSettings,
+		const FRotation3& Ra,
+		const FRotation3& Rb, 
+		const FPBD6DJointMotionSettings& MotionSettings, 
+		FVec3& SR,
+		FVec3& CR,
+		FMatrix33& RRa, 
+		FMatrix33& RRb, 
+		FVec3& LRMin, 
+		FVec3& LRMax)
 	{
-		if (MotionSettings.AngularMotionTypes[Axis] == E6DJointMotionType::Free)
+		// Convert the target rotation into target angles
+		// @todo(ccaulfield): optimize (cache these values, or store them directly rather than the target rotation)
+		FRotation3 DriveTwist, DriveSwing;
+		MotionSettings.AngularDriveTarget.ToSwingTwist(F6DJointConstants::TwistAxis(), DriveSwing, DriveTwist);
+	
+		FReal DriveTwistAngle, DriveSwingAngle;
+		FVec3 DriveTwistAxis, DriveSwingAxis;
+		DriveTwist.ToAxisAndAngleSafe(DriveTwistAxis, DriveTwistAngle, F6DJointConstants::TwistAxis());
+		DriveSwing.ToAxisAndAngleSafe(DriveSwingAxis, DriveSwingAngle, F6DJointConstants::Swing1Axis());
+		if (FVec3::DotProduct(DriveTwistAxis, F6DJointConstants::TwistAxis()) < 0)
 		{
-			if (bDriveEnabled[Axis] && DriveStiffness > KINDA_SMALL_NUMBER)
-			{
-				CR[Axis] = CR[Axis] + DriveAngles[Axis];
-				SR[Axis] = DriveStiffness;
-			}
-			else
-			{
-				CR[Axis] = 0;
-				SR[Axis] = 0;
-				LRMin[Axis] = 0;
-				LRMax[Axis] = 0;
-			}
+			DriveTwistAngle = -DriveTwistAngle;
 		}
-		else if (MotionSettings.AngularMotionTypes[Axis] == E6DJointMotionType::Limited)
+		if ((FVec3::DotProduct(DriveSwingAxis, F6DJointConstants::Swing1Axis()) < -(FReal)0.9) || (FVec3::DotProduct(DriveSwingAxis, F6DJointConstants::Swing2Axis()) < -(FReal)0.9))
 		{
-			if (CR[Axis] >= MotionSettings.AngularLimits[Axis])
-			{
-				CR[Axis] -= MotionSettings.AngularLimits[Axis];
-				LRMin[Axis] = 0;
-			}
-			else if (CR[Axis] <= -MotionSettings.AngularLimits[Axis])
-			{
-				CR[Axis] -= -MotionSettings.AngularLimits[Axis];
-				LRMax[Axis] = 0;
-			}
-			if (bDriveEnabled[Axis] && DriveStiffness > KINDA_SMALL_NUMBER)
+			DriveSwingAngle = -DriveSwingAngle;
+		}
+
+		bool bDriveEnabled[3] = { false, false, false };
+		if (SolverSettings.bEnableDrives)
+		{
+			bDriveEnabled[(int32)E6DJointAngularConstraintIndex::Twist] = MotionSettings.bAngularTwistDriveEnabled;
+			bDriveEnabled[(int32)E6DJointAngularConstraintIndex::Swing1] = MotionSettings.bAngularSwingDriveEnabled;
+			bDriveEnabled[(int32)E6DJointAngularConstraintIndex::Swing2] = MotionSettings.bAngularSwingDriveEnabled;
+		}
+
+		FReal DriveAngles[3] = { 0,0,0 };
+		DriveAngles[(int32)E6DJointAngularConstraintIndex::Twist] = DriveTwistAngle;
+		DriveAngles[(int32)E6DJointAngularConstraintIndex::Swing1] = DriveSwingAngle;
+		DriveAngles[(int32)E6DJointAngularConstraintIndex::Swing2] = DriveSwingAngle;
+
+		FReal DriveStiffness = (SolverSettings.PBDDriveStiffness > 0) ? SolverSettings.PBDDriveStiffness : MotionSettings.AngularDriveStiffness;
+
+		// Use constraint limits settings to specify valid range for constraint-space rotation corrections
+		for (int32 Axis = 0; Axis < 3; ++Axis)
+		{
+			if (MotionSettings.AngularMotionTypes[Axis] == E6DJointMotionType::Free)
 			{
-				CR[Axis] = CR[Axis] + DriveAngles[Axis];
-				SR[Axis] = DriveStiffness;
+				if (bDriveEnabled[Axis] && DriveStiffness > KINDA_SMALL_NUMBER)
+				{
+					CR[Axis] = CR[Axis] + DriveAngles[Axis];
+					SR[Axis] = DriveStiffness;
+				}
+				else
+				{
+					CR[Axis] = 0;
+					SR[Axis] = 0;
+					LRMin[Axis] = 0;
+					LRMax[Axis] = 0;
+				}
 			}
-			else
+			else if (MotionSettings.AngularMotionTypes[Axis] == E6DJointMotionType::Limited)
 			{
-				CR[Axis] = 0;
-				SR[Axis] = 0;
-				LRMin[Axis] = 0;
-				LRMax[Axis] = 0;
+				if (CR[Axis] >= MotionSettings.AngularLimits[Axis])
+				{
+					CR[Axis] -= MotionSettings.AngularLimits[Axis];
+					LRMin[Axis] = 0;
+				}
+				else if (CR[Axis] <= -MotionSettings.AngularLimits[Axis])
+				{
+					CR[Axis] -= -MotionSettings.AngularLimits[Axis];
+					LRMax[Axis] = 0;
+				}
+				if (bDriveEnabled[Axis] && DriveStiffness > KINDA_SMALL_NUMBER)
+				{
+					CR[Axis] = CR[Axis] + DriveAngles[Axis];
+					SR[Axis] = DriveStiffness;
+				}
+				else
+				{
+					CR[Axis] = 0;
+					SR[Axis] = 0;
+					LRMin[Axis] = 0;
+					LRMax[Axis] = 0;
+				}
 			}
 		}
 	}
-}
 
-template<typename T, int d>
-bool TPBD6DJointConstraintUtilities<T, d>::Calculate6dConstraintRotation_SwingCone(
-	const TPBD6DJointSolverSettings<T, d>& SolverSettings,
-	const TRotation<T, d>& Ra, 
-	const TRotation<T, d>& Rb, 
-	const TPBD6DJointMotionSettings<T, d>& MotionSettings,
-	TVector<T, d>& CR, 
-	PMatrix<T, d, d>& RRa, 
-	PMatrix<T, d, d>& RRb)
-{
-	// @todo(ccaulfield): optimize
-
-	// Get the transform from A to B, and use it to generate twist angles.
-	TRotation<T, d> Rab = Ra.Inverse() * Rb;
-	TRotation<T, d> RTwist, RSwing;
-	Rab.ToSwingTwist(T6DJointConstants<T, d>::TwistAxis(), RSwing, RTwist);
-	RSwing = RSwing.GetNormalized();
-	RTwist = RTwist.GetNormalized();
-
-	TVector<T, d> TwistAxisAB, SwingAxisAB;
-	T TwistAngleAB, SwingAngleAB;
-	RTwist.ToAxisAndAngleSafe(TwistAxisAB, TwistAngleAB, T6DJointConstants<T, d>::TwistAxis(), SolverSettings.SwingTwistAngleTolerance);
-	RSwing.ToAxisAndAngleSafe(SwingAxisAB, SwingAngleAB, T6DJointConstants<T, d>::Swing1Axis(), SolverSettings.SwingTwistAngleTolerance);
-	if (TwistAngleAB > PI)
-	{
-		TwistAngleAB = TwistAngleAB - (T)2 * PI;
-	}
-	if (SwingAngleAB > PI)
-	{
-		SwingAngleAB = SwingAngleAB - (T)2 * PI;
-	}
-	if (TVector<T, d>::DotProduct(TwistAxisAB, T6DJointConstants<T, d>::TwistAxis()) < 0)
-	{
-		TwistAxisAB = -TwistAxisAB;
-		TwistAngleAB = -TwistAngleAB;
-	}
 
-	PMatrix<T, d, d> Axesa = Ra.ToMatrix();
-	PMatrix<T, d, d> Axesb = Rb.ToMatrix();
-
-	// Calculate constraint space axes for each body. Swing axes are generated as if twist rotation was removed from body B
-	TVector<T, d> Twista = Axesa.GetAxis((int32)E6DJointAngularAxisIndex::Twist);
-	TVector<T, d> Swing1a = Ra * SwingAxisAB;
-	TVector<T, d> Swing2a = TVector<T, d>::CrossProduct(Twista, Swing1a);
-	TVector<T, d> Twistb = Axesb.GetAxis((int32)E6DJointAngularAxisIndex::Twist);
-	TVector<T, d> Swing1b = Rb * RTwist.Inverse() * SwingAxisAB;
-	TVector<T, d> Swing2b = TVector<T, d>::CrossProduct(Twistb, Swing1b);
-
-	RRa.SetRow((int32)E6DJointAngularConstraintIndex::Twist, Twista);
-	RRa.SetRow((int32)E6DJointAngularConstraintIndex::Swing1, Swing1a);
-	RRa.SetRow((int32)E6DJointAngularConstraintIndex::Swing2, Swing2a);
-
-	RRb.SetRow((int32)E6DJointAngularConstraintIndex::Twist, Twistb);
-	RRb.SetRow((int32)E6DJointAngularConstraintIndex::Swing1, Swing1b);
-	RRb.SetRow((int32)E6DJointAngularConstraintIndex::Swing2, Swing2b);
-
-	CR[(int32)E6DJointAngularConstraintIndex::Twist] = TwistAngleAB;
-	CR[(int32)E6DJointAngularConstraintIndex::Swing1] = SwingAngleAB;
-	CR[(int32)E6DJointAngularConstraintIndex::Swing2] = 0;
-
-	// If we're flipped 180 degrees about swing, just pretend the error is zero
-	T DotTT = TVector<T, d>::DotProduct(Twista, Twistb);
-	const T MinDotTT = (T)-1 + SolverSettings.InvertedAxisTolerance;
-	if (DotTT < MinDotTT)
+	bool FPBD6DJointConstraintUtilities::Calculate6dConstraintRotation_SwingCone(
+		const FPBD6DJointSolverSettings& SolverSettings,
+		const FRotation3& Ra, 
+		const FRotation3& Rb, 
+		const FPBD6DJointMotionSettings& MotionSettings,
+		FVec3& CR, 
+		FMatrix33& RRa, 
+		FMatrix33& RRb)
 	{
-		return false;
-	}
+		// @todo(ccaulfield): optimize
 
-	return true;
-}
+		// Get the transform from A to B, and use it to generate twist angles.
+		FRotation3 Rab = Ra.Inverse() * Rb;
+		FRotation3 RTwist, RSwing;
+		Rab.ToSwingTwist(F6DJointConstants::TwistAxis(), RSwing, RTwist);
+		RSwing = RSwing.GetNormalized();
+		RTwist = RTwist.GetNormalized();
 
-template<typename T, int d>
-void TPBD6DJointConstraintUtilities<T, d>::Calculate6dConstraintRotationLimits_SwingCone(
-	const TPBD6DJointSolverSettings<T, d>& SolverSettings,
-	const TRotation<T, d>& Ra,
-	const TRotation<T, d>& Rb, 
-	const TPBD6DJointMotionSettings<T, d>& MotionSettings,
-	TVector<T, d>& SR,
-	TVector<T, d>& CR,
-	PMatrix<T, d, d>& RRa, 
-	PMatrix<T, d, d>& RRb, 
-	TVector<T, d>& LRMin, 
-	TVector<T, d>& LRMax)
-{
-	// @todo(ccaulfield): Fix low angle problems and use fixed axes when joint drives are being used
-
-	// Get the transform from A to B, and use it to generate twist angles.
-	TRotation<T, d> Rab = Ra.Inverse() * Rb;
-	TRotation<T, d> RTwist, RSwing;
-	Rab.ToSwingTwist(T6DJointConstants<T, d>::TwistAxis(), RSwing, RTwist);
-	RSwing = RSwing.GetNormalized();
-	RTwist = RTwist.GetNormalized();
-
-	TVector<T, d> TwistAxisAB, SwingAxisAB;
-	T TwistAngleAB, SwingAngleAB;
-	RTwist.ToAxisAndAngleSafe(TwistAxisAB, TwistAngleAB, T6DJointConstants<T, d>::TwistAxis(), SolverSettings.SwingTwistAngleTolerance);
-	RSwing.ToAxisAndAngleSafe(SwingAxisAB, SwingAngleAB, T6DJointConstants<T, d>::Swing1Axis(), SolverSettings.SwingTwistAngleTolerance);
-	if (TwistAngleAB > PI)
-	{
-		TwistAngleAB = TwistAngleAB - (T)2 * PI;
-	}
-	if (SwingAngleAB > PI)
-	{
-		SwingAngleAB = SwingAngleAB - (T)2 * PI;
+		FVec3 TwistAxisAB, SwingAxisAB;
+		FReal TwistAngleAB, SwingAngleAB;
+		RTwist.ToAxisAndAngleSafe(TwistAxisAB, TwistAngleAB, F6DJointConstants::TwistAxis(), SolverSettings.SwingTwistAngleTolerance);
+		RSwing.ToAxisAndAngleSafe(SwingAxisAB, SwingAngleAB, F6DJointConstants::Swing1Axis(), SolverSettings.SwingTwistAngleTolerance);
+		if (TwistAngleAB > PI)
+		{
+			TwistAngleAB = TwistAngleAB - (FReal)2 * PI;
+		}
+		if (SwingAngleAB > PI)
+		{
+			SwingAngleAB = SwingAngleAB - (FReal)2 * PI;
+		}
+		if (FVec3::DotProduct(TwistAxisAB, F6DJointConstants::TwistAxis()) < 0)
+		{
+			TwistAxisAB = -TwistAxisAB;
+			TwistAngleAB = -TwistAngleAB;
+		}
+
+		FMatrix33 Axesa = Ra.ToMatrix();
+		FMatrix33 Axesb = Rb.ToMatrix();
+
+		// Calculate constraint space axes for each body. Swing axes are generated as if twist rotation was removed from body B
+		FVec3 Twista = Axesa.GetAxis((int32)E6DJointAngularAxisIndex::Twist);
+		FVec3 Swing1a = Ra * SwingAxisAB;
+		FVec3 Swing2a = FVec3::CrossProduct(Twista, Swing1a);
+		FVec3 Twistb = Axesb.GetAxis((int32)E6DJointAngularAxisIndex::Twist);
+		FVec3 Swing1b = Rb * RTwist.Inverse() * SwingAxisAB;
+		FVec3 Swing2b = FVec3::CrossProduct(Twistb, Swing1b);
+
+		RRa.SetRow((int32)E6DJointAngularConstraintIndex::Twist, Twista);
+		RRa.SetRow((int32)E6DJointAngularConstraintIndex::Swing1, Swing1a);
+		RRa.SetRow((int32)E6DJointAngularConstraintIndex::Swing2, Swing2a);
+
+		RRb.SetRow((int32)E6DJointAngularConstraintIndex::Twist, Twistb);
+		RRb.SetRow((int32)E6DJointAngularConstraintIndex::Swing1, Swing1b);
+		RRb.SetRow((int32)E6DJointAngularConstraintIndex::Swing2, Swing2b);
+
+		CR[(int32)E6DJointAngularConstraintIndex::Twist] = TwistAngleAB;
+		CR[(int32)E6DJointAngularConstraintIndex::Swing1] = SwingAngleAB;
+		CR[(int32)E6DJointAngularConstraintIndex::Swing2] = 0;
+
+		// If we're flipped 180 degrees about swing, just pretend the error is zero
+		FReal DotTT = FVec3::DotProduct(Twista, Twistb);
+		const FReal MinDotTT = (FReal)-1 + SolverSettings.InvertedAxisTolerance;
+		if (DotTT < MinDotTT)
+		{
+			return false;
+		}
+
+		return true;
 	}
-	if (TVector<T, d>::DotProduct(TwistAxisAB, T6DJointConstants<T, d>::TwistAxis()) < 0)
+
+
+	void FPBD6DJointConstraintUtilities::Calculate6dConstraintRotationLimits_SwingCone(
+		const FPBD6DJointSolverSettings& SolverSettings,
+		const FRotation3& Ra,
+		const FRotation3& Rb, 
+		const FPBD6DJointMotionSettings& MotionSettings,
+		FVec3& SR,
+		FVec3& CR,
+		FMatrix33& RRa, 
+		FMatrix33& RRb, 
+		FVec3& LRMin, 
+		FVec3& LRMax)
 	{
-		TwistAxisAB = -TwistAxisAB;
-		TwistAngleAB = -TwistAngleAB;
+		// @todo(ccaulfield): Fix low angle problems and use fixed axes when joint drives are being used
+
+		// Get the transform from A to B, and use it to generate twist angles.
+		FRotation3 Rab = Ra.Inverse() * Rb;
+		FRotation3 RTwist, RSwing;
+		Rab.ToSwingTwist(F6DJointConstants::TwistAxis(), RSwing, RTwist);
+		RSwing = RSwing.GetNormalized();
+		RTwist = RTwist.GetNormalized();
+
+		FVec3 TwistAxisAB, SwingAxisAB;
+		FReal TwistAngleAB, SwingAngleAB;
+		RTwist.ToAxisAndAngleSafe(TwistAxisAB, TwistAngleAB, F6DJointConstants::TwistAxis(), SolverSettings.SwingTwistAngleTolerance);
+		RSwing.ToAxisAndAngleSafe(SwingAxisAB, SwingAngleAB, F6DJointConstants::Swing1Axis(), SolverSettings.SwingTwistAngleTolerance);
+		if (TwistAngleAB > PI)
+		{
+			TwistAngleAB = TwistAngleAB - (FReal)2 * PI;
+		}
+		if (SwingAngleAB > PI)
+		{
+			SwingAngleAB = SwingAngleAB - (FReal)2 * PI;
+		}
+		if (FVec3::DotProduct(TwistAxisAB, F6DJointConstants::TwistAxis()) < 0)
+		{
+			TwistAxisAB = -TwistAxisAB;
+			TwistAngleAB = -TwistAngleAB;
+		}
+
+		// Calculate angular limits in new constraint space (our cone constraint axes do not map directly onto settings' constraint axes)
+		FReal TwistLimit = MotionSettings.AngularLimits[(int32)E6DJointAngularConstraintIndex::Twist];
+		FReal DotSwing1 = FMath::Abs(FVec3::DotProduct(SwingAxisAB, F6DJointConstants::Swing1Axis()));
+		FReal DotSwing2 = FMath::Abs(FVec3::DotProduct(SwingAxisAB, F6DJointConstants::Swing2Axis()));
+		FReal Swing1Limit = MotionSettings.AngularLimits[(int32)E6DJointAngularConstraintIndex::Swing1];
+		FReal Swing2Limit = MotionSettings.AngularLimits[(int32)E6DJointAngularConstraintIndex::Swing2];
+		FReal SwingLimit = FMath::Sqrt(Swing1Limit * DotSwing1 * Swing1Limit * DotSwing1 + Swing2Limit * DotSwing2 * Swing2Limit * DotSwing2 );
+
+		FVec3 AngularLimits;
+		AngularLimits[(int32)E6DJointAngularConstraintIndex::Twist] = TwistLimit;
+		AngularLimits[(int32)E6DJointAngularConstraintIndex::Swing1] = SwingLimit;
+		AngularLimits[(int32)E6DJointAngularConstraintIndex::Swing2] = FLT_MAX;
+
+		// Convert the target rotation into target angles
+		// @todo(ccaulfield): optimize (cache these values, or store them directly rather than the target rotation)
+		FRotation3 DriveTwist, DriveSwing;
+		MotionSettings.AngularDriveTarget.ToSwingTwist(F6DJointConstants::TwistAxis(), DriveSwing, DriveTwist);
+		FReal DriveTwistAngle, DriveSwingAngle;
+		FVec3 DriveTwistAxis, DriveSwingAxis;
+		DriveTwist.ToAxisAndAngleSafe(DriveTwistAxis, DriveTwistAngle, F6DJointConstants::TwistAxis());
+		DriveSwing.ToAxisAndAngleSafe(DriveSwingAxis, DriveSwingAngle, F6DJointConstants::Swing1Axis());
+		if (FVec3::DotProduct(DriveTwistAxis, TwistAxisAB) < 0)
+		{
+			DriveTwistAngle = -DriveTwistAngle;
+		}
+
+		FVec3 SwingAxis2AB = FVec3::CrossProduct(TwistAxisAB, SwingAxisAB);
+		FReal DriveDotSwing1 = FVec3::DotProduct(DriveSwingAxis, SwingAxisAB);
+		FReal DriveDotSwing2 = FVec3::DotProduct(DriveSwingAxis, SwingAxis2AB);
+		FReal DriveSwing1Angle = DriveDotSwing1 * DriveSwingAngle;
+		FReal DriveSwing2Angle = DriveDotSwing2 * DriveSwingAngle;
+
+
+		bool bDriveEnabled[3] = { false, false, false };
+		if (SolverSettings.bEnableDrives)
+		{
+			bDriveEnabled[(int32)E6DJointAngularConstraintIndex::Twist] = MotionSettings.bAngularTwistDriveEnabled;
+			bDriveEnabled[(int32)E6DJointAngularConstraintIndex::Swing1] = MotionSettings.bAngularSwingDriveEnabled;
+			bDriveEnabled[(int32)E6DJointAngularConstraintIndex::Swing2] = MotionSettings.bAngularSwingDriveEnabled;
+		}
+
+		FReal DriveAngles[3] = { 0, 0, 0 };
+		DriveAngles[(int32)E6DJointAngularConstraintIndex::Twist] = DriveTwistAngle;
+		DriveAngles[(int32)E6DJointAngularConstraintIndex::Swing1] = DriveSwing1Angle;
+		DriveAngles[(int32)E6DJointAngularConstraintIndex::Swing2] = DriveSwing2Angle;
+
+		FReal DriveStiffness = (SolverSettings.PBDDriveStiffness > 0)? SolverSettings.PBDDriveStiffness : MotionSettings.AngularDriveStiffness;
+
+		// Use constraint limits settings to specify valid range for constraint-space rotation corrections
+		for (int32 Axis = 0; Axis < 3; ++Axis)
+		{
+			if (MotionSettings.AngularMotionTypes[Axis] == E6DJointMotionType::Free)
+			{
+				if (bDriveEnabled[Axis] && DriveStiffness > KINDA_SMALL_NUMBER)
+				{
+					CR[Axis] = CR[Axis] + DriveAngles[Axis];
+					SR[Axis] = DriveStiffness;
+				}
+				else
+				{
+					CR[Axis] = 0;
+					SR[Axis] = 0;
+					LRMin[Axis] = 0;
+					LRMax[Axis] = 0;
+				}
+			}
+			else if (MotionSettings.AngularMotionTypes[Axis] == E6DJointMotionType::Limited)
+			{
+				if (CR[Axis] >= AngularLimits[Axis])
+				{
+					CR[Axis] -= AngularLimits[Axis];
+					LRMin[Axis] = 0;
+				}
+				else if (CR[Axis] <= -AngularLimits[Axis])
+				{
+					CR[Axis] -= -AngularLimits[Axis];
+					LRMax[Axis] = 0;
+				}
+				else if (bDriveEnabled[Axis] && DriveStiffness > KINDA_SMALL_NUMBER)
+				{
+					CR[Axis] = CR[Axis] + DriveAngles[Axis];
+					SR[Axis] = DriveStiffness;
+				}
+				else
+				{
+					CR[Axis] = 0;
+					SR[Axis] = 0;
+					LRMin[Axis] = 0;
+					LRMax[Axis] = 0;
+				}
+			}
+		}
 	}
 
-	// Calculate angular limits in new constraint space (our cone constraint axes do not map directly onto settings' constraint axes)
-	T TwistLimit = MotionSettings.AngularLimits[(int32)E6DJointAngularConstraintIndex::Twist];
-	T DotSwing1 = FMath::Abs(TVector<T, d>::DotProduct(SwingAxisAB, T6DJointConstants<T, d>::Swing1Axis()));
-	T DotSwing2 = FMath::Abs(TVector<T, d>::DotProduct(SwingAxisAB, T6DJointConstants<T, d>::Swing2Axis()));
-	T Swing1Limit = MotionSettings.AngularLimits[(int32)E6DJointAngularConstraintIndex::Swing1];
-	T Swing2Limit = MotionSettings.AngularLimits[(int32)E6DJointAngularConstraintIndex::Swing2];
-	T SwingLimit = FMath::Sqrt(Swing1Limit * DotSwing1 * Swing1Limit * DotSwing1 + Swing2Limit * DotSwing2 * Swing2Limit * DotSwing2 );
-
-	TVector<T, d> AngularLimits;
-	AngularLimits[(int32)E6DJointAngularConstraintIndex::Twist] = TwistLimit;
-	AngularLimits[(int32)E6DJointAngularConstraintIndex::Swing1] = SwingLimit;
-	AngularLimits[(int32)E6DJointAngularConstraintIndex::Swing2] = FLT_MAX;
-
-	// Convert the target rotation into target angles
-	// @todo(ccaulfield): optimize (cache these values, or store them directly rather than the target rotation)
-	TRotation<T, d> DriveTwist, DriveSwing;
-	MotionSettings.AngularDriveTarget.ToSwingTwist(T6DJointConstants<T, d>::TwistAxis(), DriveSwing, DriveTwist);
-	T DriveTwistAngle, DriveSwingAngle;
-	TVector<T, d> DriveTwistAxis, DriveSwingAxis;
-	DriveTwist.ToAxisAndAngleSafe(DriveTwistAxis, DriveTwistAngle, T6DJointConstants<T, d>::TwistAxis());
-	DriveSwing.ToAxisAndAngleSafe(DriveSwingAxis, DriveSwingAngle, T6DJointConstants<T, d>::Swing1Axis());
-	if (TVector<T, d>::DotProduct(DriveTwistAxis, TwistAxisAB) < 0)
+
+	bool FPBD6DJointConstraintUtilities::Calculate6dConstraintRotation(
+		const FPBD6DJointSolverSettings& SolverSettings,
+		const FRotation3& Ra, 
+		const FRotation3& Rb, 
+		const FPBD6DJointMotionSettings& MotionSettings,
+		FVec3& CR, 
+		FMatrix33& RRa, 
+		FMatrix33& RRb)
 	{
-		DriveTwistAngle = -DriveTwistAngle;
-	}
+		RRa = FMatrix33(0, 0, 0);
+		RRb = FMatrix33(0, 0, 0);
+		CR = FVec3(0, 0, 0);
 
-	TVector<T, d> SwingAxis2AB = TVector<T, d>::CrossProduct(TwistAxisAB, SwingAxisAB);
-	T DriveDotSwing1 = TVector<T, d>::DotProduct(DriveSwingAxis, SwingAxisAB);
-	T DriveDotSwing2 = TVector<T, d>::DotProduct(DriveSwingAxis, SwingAxis2AB);
-	T DriveSwing1Angle = DriveDotSwing1 * DriveSwingAngle;
-	T DriveSwing2Angle = DriveDotSwing2 * DriveSwingAngle;
+		bool bUseSwingCone = true
+			&& (MotionSettings.AngularMotionTypes[(int32)E6DJointAngularConstraintIndex::Swing1] != E6DJointMotionType::Locked)
+			&& (MotionSettings.AngularMotionTypes[(int32)E6DJointAngularConstraintIndex::Swing2] != E6DJointMotionType::Locked);
+		if (bUseSwingCone)
+		{
+			return Calculate6dConstraintRotation_SwingCone(SolverSettings, Ra, Rb, MotionSettings, CR, RRa, RRb);
+		}
+		else
+		{
+			return Calculate6dConstraintRotation_SwingFixed(SolverSettings, Ra, Rb, MotionSettings, CR, RRa, RRb);
+		}
+	}
 
 
-	bool bDriveEnabled[3] = { false, false, false };
-	if (SolverSettings.bEnableDrives)
+	void FPBD6DJointConstraintUtilities::Calculate6dConstraintRotationLimits(
+		const FPBD6DJointSolverSettings& SolverSettings,
+		const FRotation3& Ra, 
+		const FRotation3& Rb, 
+		const FPBD6DJointMotionSettings& MotionSettings,
+		FVec3& SR,
+		FVec3& CR,
+		FMatrix33& RRa, 
+		FMatrix33& RRb, 
+		FVec3& LRMin, 
+		FVec3& LRMax)
 	{
-		bDriveEnabled[(int32)E6DJointAngularConstraintIndex::Twist] = MotionSettings.bAngularTwistDriveEnabled;
-		bDriveEnabled[(int32)E6DJointAngularConstraintIndex::Swing1] = MotionSettings.bAngularSwingDriveEnabled;
-		bDriveEnabled[(int32)E6DJointAngularConstraintIndex::Swing2] = MotionSettings.bAngularSwingDriveEnabled;
+		LRMin = FVec3(-FLT_MAX, -FLT_MAX, -FLT_MAX);
+		LRMax = FVec3(FLT_MAX, FLT_MAX, FLT_MAX);
+		SR = FVec3(1, 1, 1);
+
+		bool bUseSwingCone = true
+			&& (MotionSettings.AngularMotionTypes[(int32)E6DJointAngularConstraintIndex::Swing1] != E6DJointMotionType::Locked)
+			&& (MotionSettings.AngularMotionTypes[(int32)E6DJointAngularConstraintIndex::Swing2] != E6DJointMotionType::Locked);
+		if (bUseSwingCone)
+		{
+			Calculate6dConstraintRotationLimits_SwingCone(SolverSettings, Ra, Rb, MotionSettings, SR, CR, RRa, RRb, LRMin, LRMax);
+		}
+		else
+		{
+			Calculate6dConstraintRotationLimits_SwingFixed(SolverSettings, Ra, Rb, MotionSettings, SR, CR, RRa, RRb, LRMin, LRMax);
+		}
 	}
 
-	T DriveAngles[3] = { 0, 0, 0 };
-	DriveAngles[(int32)E6DJointAngularConstraintIndex::Twist] = DriveTwistAngle;
-	DriveAngles[(int32)E6DJointAngularConstraintIndex::Swing1] = DriveSwing1Angle;
-	DriveAngles[(int32)E6DJointAngularConstraintIndex::Swing2] = DriveSwing2Angle;
 
-	T DriveStiffness = (SolverSettings.PBDDriveStiffness > 0)? SolverSettings.PBDDriveStiffness : MotionSettings.AngularDriveStiffness;
+	bool FPBD6DJointConstraintUtilities::Calculate6dDelta(
+		const FPBD6DJointSolverSettings& SolverSettings,
+		const FReal Dt,
+		const FVec3& Pa,
+		const FRotation3& Qa,
+		float MaInv,
+		const FMatrix33& IaInv,
+		const FVec3& Pb,
+		const FRotation3& Qb,
+		float MbInv,
+		const FMatrix33& IbInv,
+		const FVec3& Xa,
+		const FRotation3& Ra,
+		const FVec3& Xb,
+		const FRotation3& Rb,
+		const FPBD6DJointMotionSettings& MotionSettings,
+		FPBD6DJointState& State,
+		FVec3& DPa,
+		FRotation3& DQa,
+		FVec3& DPb,
+		FRotation3& DQb)
+	{
+		// @todo(ccaulfield): optimize: should add the constraints 1 by one
+
+		DPa = FVec3(0, 0, 0);
+		DQa = FRotation3::FromElements(0, 0, 0, 0);
+		DPb = FVec3(0, 0, 0);
+		DQb = FRotation3::FromElements(0, 0, 0, 0);
+
+		FMatrix33 Axesa = Ra.ToMatrix();
+		FMatrix33 Axesb = Rb.ToMatrix();
+
+		// Constraint-space errors (B - A), each row represents a constraint.
+		// CX are the (x,y,z) position constraint error.
+		// CR are the (twist, swing1, swing2) rotation constraint error.
+		//
+		//	C	=	| CX |
+		//			| CR |
+		//
+		// Derivative of constraint error with respect to constraint parameters (Jacobian).
+		// Each row represents a constraint, and contains the derivative of the constraint error wrt each constraint variable.
+		// There is a J for each body: Ja and Jb.
+		//
+		//	J(C(X,R)) =	| dCX/dX^t	dCX/dR^t |
+		//	(6x6)		| dCR/dX^t	dCR/dR^t |
+		//
+		//	Ja	=	| 1(3x3)	-XRa(3x3) |
+		//			| 0(3x3)	RRa(3x3)  |
+		//
+		// Where XR is the cross-product matrix of the world-space constraint position relative to the body, and
+		// RR rows are the twist, swing1 and swing2 axes (about which we are calculating the required rotation to correct the error).
+		//
+		// J should be negated for body B, but it cancels out in lambda calculation
+		// and we reintroduce the sign in the final lambda multiple (where J gets used again)
+		//
+		FVec3 CX = (Xb - Xa);
+		FVec3 CR;
+		FMatrix33 RRa, RRb;
+		bool bRotationValid = Calculate6dConstraintRotation(SolverSettings, Ra, Rb, MotionSettings, CR, RRa, RRb);
+
+		// Set limits, apply drives
+		FVec3 SX = FVec3(1, 1, 1);
+		FVec3 LRMin, LRMax, SR;
+		Calculate6dConstraintRotationLimits(SolverSettings, Ra, Rb, MotionSettings, SR, CR, RRa, RRb, LRMin, LRMax);
+
+		CX = SX * CX;
+		CR = SR * CR;
+
+		if (!SolverSettings.bEnableTwistLimits || !bRotationValid)
+		{
+			CR[(int32)E6DJointAngularConstraintIndex::Twist] = 0;
+			LRMin[(int32)E6DJointAngularConstraintIndex::Twist] = 0;
+			LRMax[(int32)E6DJointAngularConstraintIndex::Twist] = 0;
+			RRa.SetRow((int32)E6DJointAngularConstraintIndex::Twist, FVec3(0, 0, 0));
+			RRb.SetRow((int32)E6DJointAngularConstraintIndex::Twist, FVec3(0, 0, 0));
+		}
+		if (!SolverSettings.bEnableSwingLimits || !bRotationValid)
+		{
+			CR[(int32)E6DJointAngularConstraintIndex::Swing1] = 0;
+			CR[(int32)E6DJointAngularConstraintIndex::Swing2] = 0;
+			LRMin[(int32)E6DJointAngularConstraintIndex::Swing1] = 0;
+			LRMax[(int32)E6DJointAngularConstraintIndex::Swing1] = 0;
+			LRMin[(int32)E6DJointAngularConstraintIndex::Swing2] = 0;
+			LRMax[(int32)E6DJointAngularConstraintIndex::Swing2] = 0;
+			RRa.SetRow((int32)E6DJointAngularConstraintIndex::Swing1, FVec3(0, 0, 0));
+			RRa.SetRow((int32)E6DJointAngularConstraintIndex::Swing2, FVec3(0, 0, 0));
+			RRb.SetRow((int32)E6DJointAngularConstraintIndex::Swing1, FVec3(0, 0, 0));
+			RRb.SetRow((int32)E6DJointAngularConstraintIndex::Swing2, FVec3(0, 0, 0));
+		}
 
-	// Use constraint limits settings to specify valid range for constraint-space rotation corrections
-	for (int32 Axis = 0; Axis < d; ++Axis)
-	{
-		if (MotionSettings.AngularMotionTypes[Axis] == E6DJointMotionType::Free)
+		FMatrix33 XRa = Utilities::CrossProductMatrix(Xa - Pa);
+		FMatrix33 XRb = Utilities::CrossProductMatrix(Xb - Pb);
+
+		// Lambda values for position and rotation (stored in 2x2 block of 3x3 matrices).
+		// Lambda holds the constraint-space corrections required to fix the constraint-space error C.
+		//
+		//	L   =   (1 / ([Ja . Ma^-1 . Ja^t] + [Jb . Mb^-1 . Jb^t])) . C	=	(1 / F) . C
+		//
+		//		=   | F00       F01	|^-1 . |CX|
+		//          | F10       F11	|      |CR|
+		//
+
+		// Fa = [Ja . Ma^-1 . Ja^t] = | Fa00 Fa01 |
+		//                            | Fa10 Da11 |
+		// Fb = [Jb . Mb^-1 . Jb^t]
+		// F = Fa + Fb
+		FMatrix33 Fa00, Fa01, Fa10, Fa11;
+		FMatrix33 Fb00, Fb01, Fb10, Fb11;
+		FPBD6DJointConstraintUtilities::ComputeJointFactorMatrix(XRa, RRa, MaInv, IaInv, Fa00, Fa01, Fa10, Fa11);
+		FPBD6DJointConstraintUtilities::ComputeJointFactorMatrix(XRb, RRb, MbInv, IbInv, Fb00, Fb01, Fb10, Fb11);
+
+		FMatrix33 F00 = Fa00 + Fb00;
+		FMatrix33 F01 = Fa01 + Fb01;
+		FMatrix33 F10 = Fa10 + Fb10;
+		FMatrix33 F11 = Fa11 + Fb11;
+		FMatrix33 FI00, FI01, FI10, FI11;
+
+		// Stiffness and damping Pt1 (XPBD denominator)
+		// (Also support PBD stiffness if XPBD stiffness is 0)
+		// Alpha = Inverse Stiffness, Beta = Damping (not inverse)
+		FReal Stiffness = (SolverSettings.PBDStiffness > 0) ? SolverSettings.PBDStiffness : MotionSettings.Stiffness;
+		FReal AlphaX = SolverSettings.XPBDAlphaX / (Dt * Dt);
+		FReal AlphaR = SolverSettings.XPBDAlphaR / (Dt * Dt);
+		FReal GammaX = SolverSettings.XPBDAlphaX * SolverSettings.XPBDBetaX / Dt;
+		FReal GammaR = SolverSettings.XPBDAlphaR * SolverSettings.XPBDBetaR / Dt;
+		if ((SolverSettings.XPBDAlphaX > 0) && (SolverSettings.XPBDAlphaR > 0))
 		{
-			if (bDriveEnabled[Axis] && DriveStiffness > KINDA_SMALL_NUMBER)
-			{
-				CR[Axis] = CR[Axis] + DriveAngles[Axis];
-				SR[Axis] = DriveStiffness;
-			}
-			else
+			F00.M[0][0] = ((FReal)1 + GammaX) * F00.M[0][0] + AlphaX;
+			F00.M[1][1] = ((FReal)1 + GammaX) * F00.M[1][1] + AlphaX;
+			F00.M[2][2] = ((FReal)1 + GammaX) * F00.M[2][2] + AlphaX;
+			F11.M[0][0] = ((FReal)1 + GammaR) * F11.M[0][0] + AlphaR;
+			F11.M[1][1] = ((FReal)1 + GammaR) * F11.M[1][1] + AlphaR;
+			F11.M[2][2] = ((FReal)1 + GammaR) * F11.M[2][2] + AlphaR;
+			SX = FVec3(1, 1, 1);
+			SR = FVec3(1, 1, 1);
+		}
+
+		// If we have no error for a constraint we remove its entry from F
+		// @todo(ccaulfield): this will not be necessary when constraints are built up correctly as opposed to always added.
+		for (int32 Axis = 0; Axis < 3; ++Axis)
+		{
+			if (SR[Axis] == 0)
 			{
-				CR[Axis] = 0;
-				SR[Axis] = 0;
-				LRMin[Axis] = 0;
-				LRMax[Axis] = 0;
+				F01.SetColumn(Axis, { 0,0,0 });
+				F11.SetColumn(Axis, { 0,0,0 });
+				F10.SetRow(Axis, { 0,0,0 });
+				F11.SetRow(Axis, { 0,0,0 });
+				F11.M[Axis][Axis] = 1;
 			}
 		}
-		else if (MotionSettings.AngularMotionTypes[Axis] == E6DJointMotionType::Limited)
+
+		FVec3 VX = CX - State.PrevTickCX;
+		FVec3 VR = CR - State.PrevTickCR;
+		FVec3 CXa = CX - AlphaX * State.LambdaXa - GammaX * (VX - Utilities::Multiply(XRa, VR));
+		FVec3 CRa = CR - AlphaR * State.LambdaRa - GammaR * Utilities::Multiply(RRa, VR);
+		FVec3 CXb = -CX - AlphaX * State.LambdaXb - GammaX * (VX - Utilities::Multiply(XRb, VR));
+		FVec3 CRb = -CR - AlphaR * State.LambdaRb - GammaR * Utilities::Multiply(RRb, VR);
+
+		// FI = 1 / F = | FI00 FI01 |
+		//              | FI10 FI11 |
+		//
+		FPBD6DJointConstraintUtilities::BlockwiseInverse(F00, F01, F10, F11, FI00, FI01, FI10, FI11);
+
+		// L = FI . C = | LX |
+		//              | LR |
+		//
+		FVec3 LXXa = Utilities::Multiply(FI00, CXa);
+		FVec3 LXRa = Utilities::Multiply(FI01, CRa);
+		FVec3 LRXa = Utilities::Multiply(FI10, CXa);
+		FVec3 LRRa = Utilities::Multiply(FI11, CRa);
+		FVec3 LXa = (LXXa + LXRa);
+		FVec3 LRa = (LRXa + LRRa);
+
+		FVec3 LXXb = Utilities::Multiply(FI00, CXb);
+		FVec3 LXRb = Utilities::Multiply(FI01, CRb);
+		FVec3 LRXb = Utilities::Multiply(FI10, CXb);
+		FVec3 LRRb = Utilities::Multiply(FI11, CRb);
+		FVec3 LXb = (LXXb + LXRb);
+		FVec3 LRb = (LRXb + LRRb);
+
+		// Apply joint limits (which are either 0 or -/+infinity)
+		for (int32 Axis = 0; Axis < 3; ++Axis)
+		{
+			LRa[Axis] = FMath::Clamp(LRa[Axis], LRMin[Axis], LRMax[Axis]);
+			LRb[Axis] = FMath::Clamp(LRb[Axis], -LRMax[Axis], -LRMin[Axis]);
+		}
+
+
+		// Final corrections
+		//
+		//	Da	=	Ma^-1 . Ja^t . L
+		//
+		//	Da	= | DPa | =	 Ma^-1 . |         LX           |   =    Ma^-1 . |      LX           |
+		//	      | DRa |            | -XRa^t.LX + RRa^t.LR |		         | XRa.LX + RRa^t.LR |
+		//
+		FVec3 DRaX, DRaR, DRa;
+		FVec3 DRbX, DRbR, DRb;
+		FRotation3 DQaq, DQbq;
+
+		// Reduce solver stiffness until the rotation correction falls below some threshold
+		for (int32 StiffnesIt = 0; StiffnesIt < 2; ++StiffnesIt)
 		{
-			if (CR[Axis] >= AngularLimits[Axis])
+			DRaX = Utilities::Multiply(XRa, LXa);
+			DRaR = Utilities::Multiply(RRa.GetTransposed(), LRa);
+			DRbX = Utilities::Multiply(XRb, LXb);
+			DRbR = Utilities::Multiply(RRb.GetTransposed(), LRb);
+
+			DPa = Stiffness * MaInv * LXa;
+			DPb = Stiffness * MbInv * LXb;
+			DRa = Stiffness * Utilities::Multiply(IaInv, (DRaX + DRaR));
+			DRb = Stiffness * Utilities::Multiply(IbInv, (DRbX + DRbR));
+
+			DQaq = FRotation3::FromElements(DRa, 0);
+			DQbq = FRotation3::FromElements(DRb, 0);
+
+			FReal MaxAbsDQaq = FMath::Max(FMath::Max(FMath::Abs(DQaq.X), FMath::Abs(DQaq.Y)), FMath::Max(FMath::Abs(DQaq.Z), FMath::Abs(DQaq.W)));
+			FReal MaxAbsDQbq = FMath::Max(FMath::Max(FMath::Abs(DQbq.X), FMath::Abs(DQbq.Y)), FMath::Max(FMath::Abs(DQbq.Z), FMath::Abs(DQbq.W)));
+			FReal MaxDQq = FMath::Max(MaxAbsDQaq, MaxAbsDQbq);
+			if ((MaxDQq < SolverSettings.MaxRotComponent) || !SolverSettings.bEnableAutoStiffness || (SolverSettings.MaxRotComponent == 0))
 			{
-				CR[Axis] -= AngularLimits[Axis];
-				LRMin[Axis] = 0;
+				break;
 			}
-			else if (CR[Axis] <= -AngularLimits[Axis])
+			Stiffness = Stiffness * (SolverSettings.MaxRotComponent / MaxDQq);
+		}
+
+		DQa = (DQaq * Qa) * FReal(0.5);
+		DQb = (DQbq * Qb) * FReal(0.5);
+
+		// Keep track of current constraint-space corrections for XPBD timestep dependence fix
+		State.LambdaXa = State.LambdaXa + LXa;
+		State.LambdaRa = State.LambdaRa + LRa;
+		State.LambdaXb = State.LambdaXb + LXb;
+		State.LambdaRb = State.LambdaRb + LRb;
+		State.PrevItCX = CX;
+		State.PrevItCR = CR;
+
+		UE_LOG(LogChaos6DJoint, VeryVerbose, TEXT("  C(X, R)   = (%f, %f, %f) (%f, %f, %f)"), CX.X, CX.Y, CX.Z, CR.X, CR.Y, CR.Z);
+		UE_LOG(LogChaos6DJoint, VeryVerbose, TEXT("  X(A, B)   = (%f, %f, %f) (%f, %f, %f)"), (Xa - Pa).X, (Xa - Pa).Y, (Xa - Pa).Z, (Xb - Pb).X, (Xb - Pb).Y, (Xb - Pb).Z);
+		UE_LOG(LogChaos6DJoint, VeryVerbose, TEXT("  RRX(A, B) = (%f, %f, %f) [%f] (%f, %f, %f) [%f]"), RRa.GetRow(0).X, RRa.GetRow(0).Y, RRa.GetRow(0).Z, RRa.GetRow(0).Size(), RRb.GetRow(0).X, RRb.GetRow(0).Y, RRb.GetRow(0).Z, RRa.GetRow(0).Size());
+		UE_LOG(LogChaos6DJoint, VeryVerbose, TEXT("  RRY(A, B) = (%f, %f, %f) [%f] (%f, %f, %f) [%f]"), RRa.GetRow(1).X, RRa.GetRow(1).Y, RRa.GetRow(1).Z, RRa.GetRow(1).Size(), RRb.GetRow(1).X, RRb.GetRow(1).Y, RRb.GetRow(1).Z, RRa.GetRow(1).Size());
+		UE_LOG(LogChaos6DJoint, VeryVerbose, TEXT("  RRZ(A, B) = (%f, %f, %f) [%f] (%f, %f, %f) [%f]"), RRa.GetRow(2).X, RRa.GetRow(2).Y, RRa.GetRow(2).Z, RRa.GetRow(2).Size(), RRb.GetRow(2).X, RRb.GetRow(2).Y, RRb.GetRow(2).Z, RRa.GetRow(2).Size());
+		UE_LOG(LogChaos6DJoint, VeryVerbose, TEXT("  LXA(X, R) = (%f, %f, %f) (%f, %f, %f)"), LXXa.X, LXXa.Y, LXXa.Z, LXRa.X, LXRa.Y, LXRa.Z);
+		UE_LOG(LogChaos6DJoint, VeryVerbose, TEXT("  LRA(X, R) = (%f, %f, %f) (%f, %f, %f)"), LRXa.X, LRXa.Y, LRXa.Z, LRRa.X, LRRa.Y, LRRa.Z);
+		UE_LOG(LogChaos6DJoint, VeryVerbose, TEXT("  LA(X, R)  = (%f, %f, %f) (%f, %f, %f)"), LXa.X, LXa.Y, LXa.Z, LRa.X, LRa.Y, LRa.Z);
+		UE_LOG(LogChaos6DJoint, VeryVerbose, TEXT("  LXB(X, R) = (%f, %f, %f) (%f, %f, %f)"), LXXb.X, LXXb.Y, LXXb.Z, LXRb.X, LXRb.Y, LXRb.Z);
+		UE_LOG(LogChaos6DJoint, VeryVerbose, TEXT("  LRB(X, R) = (%f, %f, %f) (%f, %f, %f)"), LRXb.X, LRXb.Y, LRXb.Z, LRRb.X, LRRb.Y, LRRb.Z);
+		UE_LOG(LogChaos6DJoint, VeryVerbose, TEXT("  LB(X, R)  = (%f, %f, %f) (%f, %f, %f)"), LXb.X, LXb.Y, LXb.Z, LRb.X, LRb.Y, LRb.Z);
+		UE_LOG(LogChaos6DJoint, VeryVerbose, TEXT("  DP(A, B)  = (%f, %f, %f) (%f, %f, %f)"), DPa.X, DPa.Y, DPa.Z, DPb.X, DPb.Y, DPb.Z);
+		UE_LOG(LogChaos6DJoint, VeryVerbose, TEXT("  DRA(X, R) = (%f, %f, %f) (%f, %f, %f)"), DRaX.X, DRaX.Y, DRaX.Z, DRaR.X, DRaR.Y, DRaR.Z);
+		UE_LOG(LogChaos6DJoint, VeryVerbose, TEXT("  DRB(X, R) = (%f, %f, %f) (%f, %f, %f)"), DRbX.X, DRbX.Y, DRbX.Z, DRbR.X, DRbR.Y, DRbR.Z);
+		UE_LOG(LogChaos6DJoint, VeryVerbose, TEXT("  DQ(A, B)  = (%f, %f, %f) (%f, %f, %f)"), DQaq.X, DQaq.Y, DQaq.Z, DQbq.X, DQbq.Y, DQbq.Z);
+
+		return true;
+	}
+
+
+	int FPBD6DJointConstraintUtilities::Solve6dConstraint(
+		const FPBD6DJointSolverSettings& SolverSettings,
+		const FReal Dt,
+		FVec3& PaInOut,
+		FRotation3& QaInOut, 
+		FReal MaInv,
+		const FMatrix33& ILaInv, 
+		const FVec3& XLa,
+		const FRotation3& RLa,
+		FVec3& PbInOut,
+		FRotation3& QbInOut, 
+		FReal MbInv,
+		const FMatrix33& ILbInv, 
+		const FVec3& XLb,
+		const FRotation3& RLb,
+		const FPBD6DJointMotionSettings& MotionSettings,
+		FPBD6DJointState& State)
+	{
+		FVec3& Pa = PaInOut;
+		FRotation3& Qa = QaInOut;
+		FVec3& Pb = PbInOut;
+		FRotation3& Qb = QbInOut;
+
+		// World-space mass state
+		FMatrix33 IaInv = Utilities::Multiply(Qa.ToMatrix(), Utilities::Multiply(ILaInv, Qa.ToMatrix().GetTransposed()));
+		FMatrix33 IbInv = Utilities::Multiply(Qb.ToMatrix(), Utilities::Multiply(ILbInv, Qb.ToMatrix().GetTransposed()));
+
+		bool bFlipQ = ((Qa | Qb) < 0);
+		if (bFlipQ)
+		{
+			Qa = FRotation3::Negate(Qa);
+		}
+
+		FReal CurrentError = FLT_MAX;
+
+		int NumLoops = 0;
+		for (int LoopIndex = 0; (LoopIndex < SolverSettings.MaxIterations) && (CurrentError > SolverSettings.SolveTolerance); ++LoopIndex)
+		{
+			// World-space constraint state
+			FVec3 Xa = Pa + Qa * XLa;
+			FVec3 Xb = Pb + Qb * XLb;
+			FRotation3 Ra = Qa * RLa;
+			FRotation3 Rb = Qb * RLb;
+
+	#if !NO_LOGGING
+			if (!LogChaos6DJoint.IsSuppressed(ELogVerbosity::Verbose))
 			{
-				CR[Axis] -= -AngularLimits[Axis];
-				LRMax[Axis] = 0;
+				FVec3 ConstraintAngles = Calculate6dConstraintAngles(SolverSettings, Ra, Rb, MotionSettings);
+				ConstraintAngles = FVec3(FMath::RadiansToDegrees(ConstraintAngles.X), FMath::RadiansToDegrees(ConstraintAngles.Y), FMath::RadiansToDegrees(ConstraintAngles.Z));
+				UE_LOG(LogChaos6DJoint, Verbose, TEXT("Pre Loop %d: Pos = %f (%f, %f, %f) Angle = (%f, %f, %f)"),
+					LoopIndex, (Xb - Xa).Size(), (Xb - Xa).X, (Xb - Xa).Y, (Xb - Xa).Z,
+					ConstraintAngles.X, ConstraintAngles.Y, ConstraintAngles.Z);
 			}
-			else if (bDriveEnabled[Axis] && DriveStiffness > KINDA_SMALL_NUMBER)
+	#endif
+
+			// Get deltas to apply to position and rotation to correct constraint error
+			FVec3 DPa, DPb;
+			FRotation3 DQa, DQb;
+			bool bSolveOk = Calculate6dDelta(SolverSettings, Dt, Pa, Qa, MaInv, IaInv, Pb, Qb, MbInv, IbInv, Xa, Ra, Xb, Rb, MotionSettings, State, DPa, DQa, DPb, DQb);
+			if (!bSolveOk)
 			{
-				CR[Axis] = CR[Axis] + DriveAngles[Axis];
-				SR[Axis] = DriveStiffness;
+				break;
 			}
-			else
+
+			// New world-space body state
+			FVec3 Pa2 = Pa + DPa;
+			FVec3 Pb2 = Pb + DPb;
+			FRotation3 Qa2 = (Qa + DQa).GetNormalized();
+			FRotation3 Qb2 = (Qb + DQb).GetNormalized();
+			Qb2.EnforceShortestArcWith(Qb);
+			Qa2.EnforceShortestArcWith(Qb2);
+
+			// Update body state
+			Pa = Pa2;
+			Pb = Pb2;
+			Qa = Qa2;
+			Qb = Qb2;
+			IaInv = Qa.ToMatrix() * ILaInv * Qa.ToMatrix().GetTransposed();
+			IbInv = Qb.ToMatrix() * ILbInv * Qb.ToMatrix().GetTransposed();
+
+			// New world-space constraint state
+			FVec3 Xa2 = Pa2 + Qa2 * XLa;
+			FVec3 Xb2 = Pb2 + Qb2 * XLb;
+			FRotation3 Ra2 = Qa2 * RLa;
+			FRotation3 Rb2 = Qb2 * RLb;
+
+	#if !NO_LOGGING
+			if (!LogChaos6DJoint.IsSuppressed(ELogVerbosity::Verbose))
 			{
-				CR[Axis] = 0;
-				SR[Axis] = 0;
-				LRMin[Axis] = 0;
-				LRMax[Axis] = 0;
+				FVec3 ConstraintAngles2 = Calculate6dConstraintAngles(SolverSettings, Ra2, Rb2, MotionSettings);
+				ConstraintAngles2 = FVec3(FMath::RadiansToDegrees(ConstraintAngles2.X), FMath::RadiansToDegrees(ConstraintAngles2.Y), FMath::RadiansToDegrees(ConstraintAngles2.Z));
+				UE_LOG(LogChaos6DJoint, Verbose, TEXT("Post Loop %d: Pos = %f (%f, %f, %f) Angle = (%f, %f, %f)"),
+					LoopIndex, (Xb2 - Xa2).Size(), (Xb2 - Xa2).X, (Xb2 - Xa2).Y, (Xb2 - Xa2).Z,
+					ConstraintAngles2.X, ConstraintAngles2.Y, ConstraintAngles2.Z);
 			}
+	#endif
+
+			// @todo(ccaulfield): this isn't really a good error calculation - the magnitudes of positions and rotations are too different and its very expensive. It'll do for now though.
+			FVec3 CX = Xb2 - Xa2;
+			FVec3 CR, SR;
+			FMatrix33 RRa, RRb;
+			FVec3 LRMin, LRMax;
+			Calculate6dConstraintRotation(SolverSettings, Ra2, Rb2, MotionSettings, CR, RRa, RRb);
+			Calculate6dConstraintRotationLimits(SolverSettings, Ra2, Rb2, MotionSettings, SR, CR, RRa, RRb, LRMin, LRMax);
+			CurrentError = FMath::Sqrt(CX.SizeSquared() + CR.SizeSquared());
+			++NumLoops;
+		}
+
+		if (bFlipQ)
+		{
+			Qa = FRotation3::Negate(Qa);
 		}
+
+		return NumLoops;
 	}
-}
 
-template<typename T, int d>
-bool TPBD6DJointConstraintUtilities<T, d>::Calculate6dConstraintRotation(
-	const TPBD6DJointSolverSettings<T, d>& SolverSettings,
-	const TRotation<T, d>& Ra, 
-	const TRotation<T, d>& Rb, 
-	const TPBD6DJointMotionSettings<T, d>& MotionSettings,
-	TVector<T, d>& CR, 
-	PMatrix<T, d, d>& RRa, 
-	PMatrix<T, d, d>& RRb)
-{
-	RRa = PMatrix<T, d, d>(0, 0, 0);
-	RRb = PMatrix<T, d, d>(0, 0, 0);
-	CR = TVector<T, d>(0, 0, 0);
-
-	bool bUseSwingCone = true
-		&& (MotionSettings.AngularMotionTypes[(int32)E6DJointAngularConstraintIndex::Swing1] != E6DJointMotionType::Locked)
-		&& (MotionSettings.AngularMotionTypes[(int32)E6DJointAngularConstraintIndex::Swing2] != E6DJointMotionType::Locked);
-	if (bUseSwingCone)
+
+	void FPBD6DJointConstraintUtilities::Calculate3dDelta(
+		const FPBD6DJointSolverSettings& SolverSettings,
+		const FVec3& Pa,
+		const FRotation3& Qa,
+		float MaInv,
+		const FMatrix33& IaInv,
+		const FVec3& Pb,
+		const FRotation3& Qb,
+		float MbInv,
+		const FMatrix33& IbInv,
+		const FVec3& Xa,
+		const FVec3& Xb,
+		const FPBD6DJointMotionSettings& XSettings,
+		FVec3& DPa,
+		FRotation3& DQa,
+		FVec3& DPb,
+		FRotation3& DQb)
 	{
-		return Calculate6dConstraintRotation_SwingCone(SolverSettings, Ra, Rb, MotionSettings, CR, RRa, RRb);
+		FVec3 CX = Xb - Xa;
+
+		FReal Stiffness = (SolverSettings.PBDStiffness > 0) ? SolverSettings.PBDStiffness : XSettings.Stiffness;
+		CX = Stiffness * CX;
+
+		FMatrix33 Ma00 = FMatrix33(0, 0, 0);
+		FMatrix33 Mb00 = FMatrix33(0, 0, 0);
+		if (MaInv > 0)
+		{
+			Ma00 = Utilities::ComputeJointFactorMatrix(Xa - Pa, IaInv, MaInv);
+		}
+		if (MbInv > 0)
+		{
+			Mb00 = Utilities::ComputeJointFactorMatrix(Xb - Pb, IbInv, MbInv);
+		}
+		FMatrix33 MI00 = (Ma00 + Mb00).Inverse();
+		FVec3 DX = Utilities::Multiply(MI00, CX);
+
+		// Divide position and rotation error between bodies based on mass distribution
+		DPa = MaInv * DX;
+		DPb = -(MbInv * DX);
+		FVec3 DQav = Utilities::Multiply(IaInv, FVec3::CrossProduct(Xa - Pa, DX));
+		FVec3 DQbv = Utilities::Multiply(IbInv, FVec3::CrossProduct(Xb - Pb, -DX));
+		FRotation3 DQaq = FRotation3::FromElements(DQav, 0);
+		FRotation3 DQbq = FRotation3::FromElements(DQbv, 0);
+		DQa = (DQaq * Qa) * 0.5f;
+		DQb = (DQbq * Qb) * 0.5f;
+	}
+
+
+	void FPBD6DJointConstraintUtilities::Solve3dConstraint(
+		const FPBD6DJointSolverSettings& SolverSettings,
+		FVec3& P0,
+		FRotation3& Q0,
+		FReal InvM0,
+		const FMatrix33& InvIL0,
+		const FVec3& XL0,
+		const FRotation3& RL0,
+		FVec3& P1,
+		FRotation3& Q1,
+		FReal InvM1,
+		const FMatrix33& InvIL1,
+		const FVec3& XL1,
+		const FRotation3& RL1,
+		const FPBD6DJointMotionSettings& MotionSettings)
+	{
+		FVec3 X0 = P0 + Q0.RotateVector(XL0);
+		FVec3 X1 = P1 + Q1.RotateVector(XL1);
+		FMatrix33 InvI0 = Utilities::Multiply(Q0.ToMatrix(), Utilities::Multiply(InvIL0, Q0.ToMatrix().GetTransposed()));
+		FMatrix33 InvI1 = Utilities::Multiply(Q1.ToMatrix(), Utilities::Multiply(InvIL1, Q1.ToMatrix().GetTransposed()));
+		FVec3 DP0, DP1;
+		FRotation3 DQ0, DQ1;
+
+		FReal CurrentError = FLT_MAX;
+		int NumLoops = 0;
+		for (int LoopIndex = 0; (LoopIndex < SolverSettings.MaxPreIterations) && (CurrentError > SolverSettings.SolveTolerance); ++LoopIndex)
+		{
+			// Calculate position and rotation corrections
+			Calculate3dDelta(SolverSettings, P0, Q0, InvM0, InvI0, P1, Q1, InvM1, InvI1, X0, X1, MotionSettings, DP0, DQ0, DP1, DQ1);
+
+			// Apply corrections to body state
+			P0 = P0 + DP0;
+			Q0 = (Q0 + DQ0).GetNormalized();
+			P1 = P1 + DP1;
+			Q1 = (Q1 + DQ1).GetNormalized();
+
+			// Recalculate constraint state
+			X0 = P0 + Q0 * XL0;
+			X1 = P1 + Q1 * XL1;
+			InvI0 = Utilities::Multiply(Q0.ToMatrix(), Utilities::Multiply(InvIL0, Q0.ToMatrix().GetTransposed()));
+			InvI1 = Utilities::Multiply(Q1.ToMatrix(), Utilities::Multiply(InvIL1, Q1.ToMatrix().GetTransposed()));
+		
+			// Measure error
+			FVec3 CX = X1 - X0;
+			CurrentError = CX.Size();
+			++NumLoops;
+		}
 	}
-	else
+
+	//
+	// Constraint JointSettings
+	//
+
+	
+	FPBD6DJointMotionSettings::FPBD6DJointMotionSettings()
+		: Stiffness((FReal)1)
+		, LinearMotionTypes({ E6DJointMotionType::Locked, E6DJointMotionType::Locked, E6DJointMotionType::Locked })
+		, LinearLimits(FVec3(FLT_MAX, FLT_MAX, FLT_MAX))
+		, AngularMotionTypes({ E6DJointMotionType::Free, E6DJointMotionType::Free, E6DJointMotionType::Free })
+		, AngularLimits(FVec3(FLT_MAX, FLT_MAX, FLT_MAX))
+		, AngularDriveTarget(FRotation3::FromIdentity())
+		, AngularDriveTargetAngles(FVec3(0, 0, 0))
+		, bAngularSLerpDriveEnabled(false)
+		, bAngularTwistDriveEnabled(false)
+		, bAngularSwingDriveEnabled(false)
+		, AngularDriveStiffness(0)
+		, AngularDriveDamping(0)
 	{
-		return Calculate6dConstraintRotation_SwingFixed(SolverSettings, Ra, Rb, MotionSettings, CR, RRa, RRb);
 	}
-}
 
-template<typename T, int d>
-void TPBD6DJointConstraintUtilities<T, d>::Calculate6dConstraintRotationLimits(
-	const TPBD6DJointSolverSettings<T, d>& SolverSettings,
-	const TRotation<T, d>& Ra, 
-	const TRotation<T, d>& Rb, 
-	const TPBD6DJointMotionSettings<T, d>& MotionSettings,
-	TVector<T, d>& SR,
-	TVector<T, d>& CR,
-	PMatrix<T, d, d>& RRa, 
-	PMatrix<T, d, d>& RRb, 
-	TVector<T, d>& LRMin, 
-	TVector<T, d>& LRMax)
-{
-	LRMin = TVector<T, d>(-FLT_MAX, -FLT_MAX, -FLT_MAX);
-	LRMax = TVector<T, d>(FLT_MAX, FLT_MAX, FLT_MAX);
-	SR = TVector<T, d>(1, 1, 1);
-
-	bool bUseSwingCone = true
-		&& (MotionSettings.AngularMotionTypes[(int32)E6DJointAngularConstraintIndex::Swing1] != E6DJointMotionType::Locked)
-		&& (MotionSettings.AngularMotionTypes[(int32)E6DJointAngularConstraintIndex::Swing2] != E6DJointMotionType::Locked);
-	if (bUseSwingCone)
+	
+	FPBD6DJointMotionSettings::FPBD6DJointMotionSettings(const TVector<E6DJointMotionType, 3>& InLinearMotionTypes, const TVector<E6DJointMotionType, 3>& InAngularMotionTypes)
+		: Stiffness((FReal)1)
+		, LinearMotionTypes(InLinearMotionTypes)
+		, LinearLimits(FVec3(FLT_MAX, FLT_MAX, FLT_MAX))
+		, AngularMotionTypes({ E6DJointMotionType::Free, E6DJointMotionType::Free, E6DJointMotionType::Free })
+		, AngularLimits(FVec3(FLT_MAX, FLT_MAX, FLT_MAX))
+		, AngularDriveTarget(FRotation3::FromIdentity())
+		, AngularDriveTargetAngles(FVec3(0, 0, 0))
+		, bAngularSLerpDriveEnabled(false)
+		, bAngularTwistDriveEnabled(false)
+		, bAngularSwingDriveEnabled(false)
+		, AngularDriveStiffness(0)
+		, AngularDriveDamping(0)
 	{
-		Calculate6dConstraintRotationLimits_SwingCone(SolverSettings, Ra, Rb, MotionSettings, SR, CR, RRa, RRb, LRMin, LRMax);
 	}
-	else
+
+	
+	FPBD6DJointSettings::FPBD6DJointSettings()
+		: ConstraintFrames({ FTransform::Identity, FTransform::Identity })
 	{
-		Calculate6dConstraintRotationLimits_SwingFixed(SolverSettings, Ra, Rb, MotionSettings, SR, CR, RRa, RRb, LRMin, LRMax);
 	}
-}
-
-template<typename T, int d>
-bool TPBD6DJointConstraintUtilities<T, d>::Calculate6dDelta(
-	const TPBD6DJointSolverSettings<T, d>& SolverSettings,
-	const T Dt,
-	const TVector<T, d>& Pa,
-	const TRotation<T, d>& Qa,
-	float MaInv,
-	const PMatrix<T, d, d>& IaInv,
-	const TVector<T, d>& Pb,
-	const TRotation<T, d>& Qb,
-	float MbInv,
-	const PMatrix<T, d, d>& IbInv,
-	const TVector<T, d>& Xa,
-	const TRotation<T, d>& Ra,
-	const TVector<T, d>& Xb,
-	const TRotation<T, d>& Rb,
-	const TPBD6DJointMotionSettings<T, d>& MotionSettings,
-	TPBD6DJointState<T, d>& State,
-	TVector<T, d>& DPa,
-	TRotation<T, d>& DQa,
-	TVector<T, d>& DPb,
-	TRotation<T, d>& DQb)
-{
-	// @todo(ccaulfield): optimize: should add the constraints 1 by one
 
-	DPa = TVector<T, d>(0, 0, 0);
-	DQa = TRotation<T, d>::FromElements(0, 0, 0, 0);
-	DPb = TVector<T, d>(0, 0, 0);
-	DQb = TRotation<T, d>::FromElements(0, 0, 0, 0);
-
-	PMatrix<T, d, d> Axesa = Ra.ToMatrix();
-	PMatrix<T, d, d> Axesb = Rb.ToMatrix();
+	
+	FPBD6DJointState::FPBD6DJointState()
+		: LambdaXa(FVec3(0, 0, 0))
+		, LambdaRa(FVec3(0, 0, 0))
+		, LambdaXb(FVec3(0, 0, 0))
+		, LambdaRb(FVec3(0, 0, 0))
+		, PrevTickCX(FVec3(0, 0, 0))
+		, PrevTickCR(FVec3(0, 0, 0))
+		, PrevItCX(FVec3(0, 0, 0))
+		, PrevItCR(FVec3(0, 0, 0))
+		, Level(INDEX_NONE)
+		, ParticleLevels({ INDEX_NONE, INDEX_NONE })
+	{
+	}
 
-	// Constraint-space errors (B - A), each row represents a constraint.
-	// CX are the (x,y,z) position constraint error.
-	// CR are the (twist, swing1, swing2) rotation constraint error.
-	//
-	//	C	=	| CX |
-	//			| CR |
-	//
-	// Derivative of constraint error with respect to constraint parameters (Jacobian).
-	// Each row represents a constraint, and contains the derivative of the constraint error wrt each constraint variable.
-	// There is a J for each body: Ja and Jb.
-	//
-	//	J(C(X,R)) =	| dCX/dX^t	dCX/dR^t |
-	//	(6x6)		| dCR/dX^t	dCR/dR^t |
 	//
-	//	Ja	=	| 1(3x3)	-XRa(3x3) |
-	//			| 0(3x3)	RRa(3x3)  |
+	// Container JointSettings
 	//
-	// Where XR is the cross-product matrix of the world-space constraint position relative to the body, and
-	// RR rows are the twist, swing1 and swing2 axes (about which we are calculating the required rotation to correct the error).
-	//
-	// J should be negated for body B, but it cancels out in lambda calculation
-	// and we reintroduce the sign in the final lambda multiple (where J gets used again)
-	//
-	TVector<T, d> CX = (Xb - Xa);
-	TVector<T, d> CR;
-	PMatrix<T, d, d> RRa, RRb;
-	bool bRotationValid = Calculate6dConstraintRotation(SolverSettings, Ra, Rb, MotionSettings, CR, RRa, RRb);
-
-	// Set limits, apply drives
-	TVector<T, d> SX = TVector<T, d>(1, 1, 1);
-	TVector<T, d> LRMin, LRMax, SR;
-	Calculate6dConstraintRotationLimits(SolverSettings, Ra, Rb, MotionSettings, SR, CR, RRa, RRb, LRMin, LRMax);
 
-	CX = SX * CX;
-	CR = SR * CR;
-
-	if (!SolverSettings.bEnableTwistLimits || !bRotationValid)
-	{
-		CR[(int32)E6DJointAngularConstraintIndex::Twist] = 0;
-		LRMin[(int32)E6DJointAngularConstraintIndex::Twist] = 0;
-		LRMax[(int32)E6DJointAngularConstraintIndex::Twist] = 0;
-		RRa.SetRow((int32)E6DJointAngularConstraintIndex::Twist, TVector<T, d>(0, 0, 0));
-		RRb.SetRow((int32)E6DJointAngularConstraintIndex::Twist, TVector<T, d>(0, 0, 0));
-	}
-	if (!SolverSettings.bEnableSwingLimits || !bRotationValid)
+	
+	FPBD6DJointSolverSettings::FPBD6DJointSolverSettings()
+		: SolveTolerance(KINDA_SMALL_NUMBER)
+		, InvertedAxisTolerance( 0.001f)	// 1 - Cos(97.5deg)
+		, SwingTwistAngleTolerance(1.0e-6f)
+		, bApplyProjection(false)
+		, MaxIterations(10)
+		, MaxPreIterations(0)
+		, MaxDriveIterations(1)
+		, MaxRotComponent(0.0f)
+		, PBDMinParentMassRatio(0.5f)
+		, PBDMaxInertiaRatio(5.0f)
+		, FreezeIterations(0)
+		, FrozenIterations(0)
+		, bEnableAutoStiffness(true)
+		, bEnableTwistLimits(1)
+		, bEnableSwingLimits(1)
+		, bEnableDrives(1)
+		, XPBDAlphaX(0)
+		, XPBDAlphaR(0)
+		, XPBDBetaX(0)
+		, XPBDBetaR(0)
+		, PBDStiffness(0.0f)
+		, PBDDriveStiffness(0.0f)
+		, bFastSolve(false)
 	{
-		CR[(int32)E6DJointAngularConstraintIndex::Swing1] = 0;
-		CR[(int32)E6DJointAngularConstraintIndex::Swing2] = 0;
-		LRMin[(int32)E6DJointAngularConstraintIndex::Swing1] = 0;
-		LRMax[(int32)E6DJointAngularConstraintIndex::Swing1] = 0;
-		LRMin[(int32)E6DJointAngularConstraintIndex::Swing2] = 0;
-		LRMax[(int32)E6DJointAngularConstraintIndex::Swing2] = 0;
-		RRa.SetRow((int32)E6DJointAngularConstraintIndex::Swing1, TVector<T, d>(0, 0, 0));
-		RRa.SetRow((int32)E6DJointAngularConstraintIndex::Swing2, TVector<T, d>(0, 0, 0));
-		RRb.SetRow((int32)E6DJointAngularConstraintIndex::Swing1, TVector<T, d>(0, 0, 0));
-		RRb.SetRow((int32)E6DJointAngularConstraintIndex::Swing2, TVector<T, d>(0, 0, 0));
 	}
 
-	PMatrix<T, d, d> XRa = Utilities::CrossProductMatrix(Xa - Pa);
-	PMatrix<T, d, d> XRb = Utilities::CrossProductMatrix(Xb - Pb);
-
-	// Lambda values for position and rotation (stored in 2x2 block of 3x3 matrices).
-	// Lambda holds the constraint-space corrections required to fix the constraint-space error C.
-	//
-	//	L   =   (1 / ([Ja . Ma^-1 . Ja^t] + [Jb . Mb^-1 . Jb^t])) . C	=	(1 / F) . C
 	//
-	//		=   | F00       F01	|^-1 . |CX|
-	//          | F10       F11	|      |CR|
+	// Constraint Container
 	//
 
-	// Fa = [Ja . Ma^-1 . Ja^t] = | Fa00 Fa01 |
-	//                            | Fa10 Da11 |
-	// Fb = [Jb . Mb^-1 . Jb^t]
-	// F = Fa + Fb
-	PMatrix<T, d, d> Fa00, Fa01, Fa10, Fa11;
-	PMatrix<T, d, d> Fb00, Fb01, Fb10, Fb11;
-	TPBD6DJointConstraintUtilities<T, d>::ComputeJointFactorMatrix(XRa, RRa, MaInv, IaInv, Fa00, Fa01, Fa10, Fa11);
-	TPBD6DJointConstraintUtilities<T, d>::ComputeJointFactorMatrix(XRb, RRb, MbInv, IbInv, Fb00, Fb01, Fb10, Fb11);
-
-	PMatrix<T, d, d> F00 = Fa00 + Fb00;
-	PMatrix<T, d, d> F01 = Fa01 + Fb01;
-	PMatrix<T, d, d> F10 = Fa10 + Fb10;
-	PMatrix<T, d, d> F11 = Fa11 + Fb11;
-	PMatrix<T, d, d> FI00, FI01, FI10, FI11;
-
-	// Stiffness and damping Pt1 (XPBD denominator)
-	// (Also support PBD stiffness if XPBD stiffness is 0)
-	// Alpha = Inverse Stiffness, Beta = Damping (not inverse)
-	T Stiffness = (SolverSettings.PBDStiffness > 0) ? SolverSettings.PBDStiffness : MotionSettings.Stiffness;
-	T AlphaX = SolverSettings.XPBDAlphaX / (Dt * Dt);
-	T AlphaR = SolverSettings.XPBDAlphaR / (Dt * Dt);
-	T GammaX = SolverSettings.XPBDAlphaX * SolverSettings.XPBDBetaX / Dt;
-	T GammaR = SolverSettings.XPBDAlphaR * SolverSettings.XPBDBetaR / Dt;
-	if ((SolverSettings.XPBDAlphaX > 0) && (SolverSettings.XPBDAlphaR > 0))
+	FPBD6DJointConstraints::FPBD6DJointConstraints(const FPBD6DJointSolverSettings& InSettings)
+		: Settings(InSettings)
+		, PreApplyCallback(nullptr)
+		, PostApplyCallback(nullptr)
 	{
-		F00.M[0][0] = ((T)1 + GammaX) * F00.M[0][0] + AlphaX;
-		F00.M[1][1] = ((T)1 + GammaX) * F00.M[1][1] + AlphaX;
-		F00.M[2][2] = ((T)1 + GammaX) * F00.M[2][2] + AlphaX;
-		F11.M[0][0] = ((T)1 + GammaR) * F11.M[0][0] + AlphaR;
-		F11.M[1][1] = ((T)1 + GammaR) * F11.M[1][1] + AlphaR;
-		F11.M[2][2] = ((T)1 + GammaR) * F11.M[2][2] + AlphaR;
-		SX = TVector<T, d>(1, 1, 1);
-		SR = TVector<T, d>(1, 1, 1);
 	}
 
-	// If we have no error for a constraint we remove its entry from F
-	// @todo(ccaulfield): this will not be necessary when constraints are built up correctly as opposed to always added.
-	for (int32 Axis = 0; Axis < d; ++Axis)
+	FPBD6DJointConstraints::~FPBD6DJointConstraints()
 	{
-		if (SR[Axis] == 0)
-		{
-			F01.SetColumn(Axis, { 0,0,0 });
-			F11.SetColumn(Axis, { 0,0,0 });
-			F10.SetRow(Axis, { 0,0,0 });
-			F11.SetRow(Axis, { 0,0,0 });
-			F11.M[Axis][Axis] = 1;
-		}
 	}
 
-	TVector<T, d> VX = CX - State.PrevTickCX;
-	TVector<T, d> VR = CR - State.PrevTickCR;
-	TVector<T, d> CXa = CX - AlphaX * State.LambdaXa - GammaX * (VX - Utilities::Multiply(XRa, VR));
-	TVector<T, d> CRa = CR - AlphaR * State.LambdaRa - GammaR * Utilities::Multiply(RRa, VR);
-	TVector<T, d> CXb = -CX - AlphaX * State.LambdaXb - GammaX * (VX - Utilities::Multiply(XRb, VR));
-	TVector<T, d> CRb = -CR - AlphaR * State.LambdaRb - GammaR * Utilities::Multiply(RRb, VR);
 
-	// FI = 1 / F = | FI00 FI01 |
-	//              | FI10 FI11 |
-	//
-	TPBD6DJointConstraintUtilities<T, d>::BlockwiseInverse(F00, F01, F10, F11, FI00, FI01, FI10, FI11);
+	const FPBD6DJointSolverSettings& FPBD6DJointConstraints::GetSettings() const
+	{
+		return Settings;
+	}
 
-	// L = FI . C = | LX |
-	//              | LR |
-	//
-	TVector<T, d> LXXa = Utilities::Multiply(FI00, CXa);
-	TVector<T, d> LXRa = Utilities::Multiply(FI01, CRa);
-	TVector<T, d> LRXa = Utilities::Multiply(FI10, CXa);
-	TVector<T, d> LRRa = Utilities::Multiply(FI11, CRa);
-	TVector<T, d> LXa = (LXXa + LXRa);
-	TVector<T, d> LRa = (LRXa + LRRa);
-
-	TVector<T, d> LXXb = Utilities::Multiply(FI00, CXb);
-	TVector<T, d> LXRb = Utilities::Multiply(FI01, CRb);
-	TVector<T, d> LRXb = Utilities::Multiply(FI10, CXb);
-	TVector<T, d> LRRb = Utilities::Multiply(FI11, CRb);
-	TVector<T, d> LXb = (LXXb + LXRb);
-	TVector<T, d> LRb = (LRXb + LRRb);
-
-	// Apply joint limits (which are either 0 or -/+infinity)
-	for (int32 Axis = 0; Axis < d; ++Axis)
+	void FPBD6DJointConstraints::SetSettings(const FPBD6DJointSolverSettings& InSettings)
 	{
-		LRa[Axis] = FMath::Clamp(LRa[Axis], LRMin[Axis], LRMax[Axis]);
-		LRb[Axis] = FMath::Clamp(LRb[Axis], -LRMax[Axis], -LRMin[Axis]);
+		Settings = InSettings;
 	}
 
+	int32 FPBD6DJointConstraints::NumConstraints() const
+	{
+		return ConstraintParticles.Num();
+	}
 
-	// Final corrections
-	//
-	//	Da	=	Ma^-1 . Ja^t . L
-	//
-	//	Da	= | DPa | =	 Ma^-1 . |         LX           |   =    Ma^-1 . |      LX           |
-	//	      | DRa |            | -XRa^t.LX + RRa^t.LR |		         | XRa.LX + RRa^t.LR |
-	//
-	TVector<T, d> DRaX, DRaR, DRa;
-	TVector<T, d> DRbX, DRbR, DRb;
-	TRotation<T, d> DQaq, DQbq;
+	typename FPBD6DJointConstraints::FConstraintContainerHandle* FPBD6DJointConstraints::AddConstraint(const FParticlePair& InConstrainedParticles, const FTransformPair& ConstraintFrames)
+	{
+		int ConstraintIndex = Handles.Num();
+		Handles.Add(HandleAllocator.AllocHandle(this, ConstraintIndex));
+		ConstraintParticles.Add(InConstrainedParticles);
+		ConstraintSettings.Add(FPBD6DJointSettings());
+		ConstraintSettings[ConstraintIndex].ConstraintFrames = ConstraintFrames;
+		ConstraintStates.Add(FPBD6DJointState());
+		return Handles.Last();
+	}
 
-	// Reduce solver stiffness until the rotation correction falls below some threshold
-	for (int32 StiffnesIt = 0; StiffnesIt < 2; ++StiffnesIt)
+	typename FPBD6DJointConstraints::FConstraintContainerHandle* FPBD6DJointConstraints::AddConstraint(const FParticlePair& InConstrainedParticles, const FPBD6DJointSettings& InConstraintSettings)
 	{
-		DRaX = Utilities::Multiply(XRa, LXa);
-		DRaR = Utilities::Multiply(RRa.GetTransposed(), LRa);
-		DRbX = Utilities::Multiply(XRb, LXb);
-		DRbR = Utilities::Multiply(RRb.GetTransposed(), LRb);
+		int ConstraintIndex = Handles.Num();
+		Handles.Add(HandleAllocator.AllocHandle(this, ConstraintIndex));
+		ConstraintParticles.Add(InConstrainedParticles);
+		ConstraintSettings.Add(InConstraintSettings);
+		ConstraintStates.Add(FPBD6DJointState());
+		return Handles.Last();
+	}
 
-		DPa = Stiffness * MaInv * LXa;
-		DPb = Stiffness * MbInv * LXb;
-		DRa = Stiffness * Utilities::Multiply(IaInv, (DRaX + DRaR));
-		DRb = Stiffness * Utilities::Multiply(IbInv, (DRbX + DRbR));
+	void FPBD6DJointConstraints::RemoveConstraint(int ConstraintIndex)
+	{
+		FConstraintContainerHandle* ConstraintHandle = Handles[ConstraintIndex];
+		if (ConstraintHandle != nullptr)
+		{
+			// Release the handle for the freed constraint
+			HandleAllocator.FreeHandle(ConstraintHandle);
+			Handles[ConstraintIndex] = nullptr;
+		}
 
-		DQaq = TRotation<T, d>::FromElements(DRa, 0);
-		DQbq = TRotation<T, d>::FromElements(DRb, 0);
+		// Swap the last constraint into the gap to keep the array packed
+		ConstraintParticles.RemoveAtSwap(ConstraintIndex);
+		ConstraintSettings.RemoveAtSwap(ConstraintIndex);
+		ConstraintStates.RemoveAtSwap(ConstraintIndex);
+		Handles.RemoveAtSwap(ConstraintIndex);
 
-		T MaxAbsDQaq = FMath::Max(FMath::Max(FMath::Abs(DQaq.X), FMath::Abs(DQaq.Y)), FMath::Max(FMath::Abs(DQaq.Z), FMath::Abs(DQaq.W)));
-		T MaxAbsDQbq = FMath::Max(FMath::Max(FMath::Abs(DQbq.X), FMath::Abs(DQbq.Y)), FMath::Max(FMath::Abs(DQbq.Z), FMath::Abs(DQbq.W)));
-		T MaxDQq = FMath::Max(MaxAbsDQaq, MaxAbsDQbq);
-		if ((MaxDQq < SolverSettings.MaxRotComponent) || !SolverSettings.bEnableAutoStiffness || (SolverSettings.MaxRotComponent == 0))
+		// Update the handle for the constraint that was moved
+		if (ConstraintIndex < Handles.Num())
 		{
-			break;
+			SetConstraintIndex(Handles[ConstraintIndex], ConstraintIndex);
 		}
-		Stiffness = Stiffness * (SolverSettings.MaxRotComponent / MaxDQq);
 	}
 
-	DQa = (DQaq * Qa) * T(0.5);
-	DQb = (DQbq * Qb) * T(0.5);
-
-	// Keep track of current constraint-space corrections for XPBD timestep dependence fix
-	State.LambdaXa = State.LambdaXa + LXa;
-	State.LambdaRa = State.LambdaRa + LRa;
-	State.LambdaXb = State.LambdaXb + LXb;
-	State.LambdaRb = State.LambdaRb + LRb;
-	State.PrevItCX = CX;
-	State.PrevItCR = CR;
-
-	UE_LOG(LogChaos6DJoint, VeryVerbose, TEXT("  C(X, R)   = (%f, %f, %f) (%f, %f, %f)"), CX.X, CX.Y, CX.Z, CR.X, CR.Y, CR.Z);
-	UE_LOG(LogChaos6DJoint, VeryVerbose, TEXT("  X(A, B)   = (%f, %f, %f) (%f, %f, %f)"), (Xa - Pa).X, (Xa - Pa).Y, (Xa - Pa).Z, (Xb - Pb).X, (Xb - Pb).Y, (Xb - Pb).Z);
-	UE_LOG(LogChaos6DJoint, VeryVerbose, TEXT("  RRX(A, B) = (%f, %f, %f) [%f] (%f, %f, %f) [%f]"), RRa.GetRow(0).X, RRa.GetRow(0).Y, RRa.GetRow(0).Z, RRa.GetRow(0).Size(), RRb.GetRow(0).X, RRb.GetRow(0).Y, RRb.GetRow(0).Z, RRa.GetRow(0).Size());
-	UE_LOG(LogChaos6DJoint, VeryVerbose, TEXT("  RRY(A, B) = (%f, %f, %f) [%f] (%f, %f, %f) [%f]"), RRa.GetRow(1).X, RRa.GetRow(1).Y, RRa.GetRow(1).Z, RRa.GetRow(1).Size(), RRb.GetRow(1).X, RRb.GetRow(1).Y, RRb.GetRow(1).Z, RRa.GetRow(1).Size());
-	UE_LOG(LogChaos6DJoint, VeryVerbose, TEXT("  RRZ(A, B) = (%f, %f, %f) [%f] (%f, %f, %f) [%f]"), RRa.GetRow(2).X, RRa.GetRow(2).Y, RRa.GetRow(2).Z, RRa.GetRow(2).Size(), RRb.GetRow(2).X, RRb.GetRow(2).Y, RRb.GetRow(2).Z, RRa.GetRow(2).Size());
-	UE_LOG(LogChaos6DJoint, VeryVerbose, TEXT("  LXA(X, R) = (%f, %f, %f) (%f, %f, %f)"), LXXa.X, LXXa.Y, LXXa.Z, LXRa.X, LXRa.Y, LXRa.Z);
-	UE_LOG(LogChaos6DJoint, VeryVerbose, TEXT("  LRA(X, R) = (%f, %f, %f) (%f, %f, %f)"), LRXa.X, LRXa.Y, LRXa.Z, LRRa.X, LRRa.Y, LRRa.Z);
-	UE_LOG(LogChaos6DJoint, VeryVerbose, TEXT("  LA(X, R)  = (%f, %f, %f) (%f, %f, %f)"), LXa.X, LXa.Y, LXa.Z, LRa.X, LRa.Y, LRa.Z);
-	UE_LOG(LogChaos6DJoint, VeryVerbose, TEXT("  LXB(X, R) = (%f, %f, %f) (%f, %f, %f)"), LXXb.X, LXXb.Y, LXXb.Z, LXRb.X, LXRb.Y, LXRb.Z);
-	UE_LOG(LogChaos6DJoint, VeryVerbose, TEXT("  LRB(X, R) = (%f, %f, %f) (%f, %f, %f)"), LRXb.X, LRXb.Y, LRXb.Z, LRRb.X, LRRb.Y, LRRb.Z);
-	UE_LOG(LogChaos6DJoint, VeryVerbose, TEXT("  LB(X, R)  = (%f, %f, %f) (%f, %f, %f)"), LXb.X, LXb.Y, LXb.Z, LRb.X, LRb.Y, LRb.Z);
-	UE_LOG(LogChaos6DJoint, VeryVerbose, TEXT("  DP(A, B)  = (%f, %f, %f) (%f, %f, %f)"), DPa.X, DPa.Y, DPa.Z, DPb.X, DPb.Y, DPb.Z);
-	UE_LOG(LogChaos6DJoint, VeryVerbose, TEXT("  DRA(X, R) = (%f, %f, %f) (%f, %f, %f)"), DRaX.X, DRaX.Y, DRaX.Z, DRaR.X, DRaR.Y, DRaR.Z);
-	UE_LOG(LogChaos6DJoint, VeryVerbose, TEXT("  DRB(X, R) = (%f, %f, %f) (%f, %f, %f)"), DRbX.X, DRbX.Y, DRbX.Z, DRbR.X, DRbR.Y, DRbR.Z);
-	UE_LOG(LogChaos6DJoint, VeryVerbose, TEXT("  DQ(A, B)  = (%f, %f, %f) (%f, %f, %f)"), DQaq.X, DQaq.Y, DQaq.Z, DQbq.X, DQbq.Y, DQbq.Z);
-
-	return true;
-}
-
-template<typename T, int d>
-int TPBD6DJointConstraintUtilities<T, d>::Solve6dConstraint(
-	const TPBD6DJointSolverSettings<T, d>& SolverSettings,
-	const T Dt,
-	TVector<T, d>& PaInOut,
-	TRotation<T, d>& QaInOut, 
-	T MaInv, 
-	const PMatrix<T, d, d>& ILaInv, 
-	const TVector<T, d>& XLa,
-	const TRotation<T, d>& RLa,
-	TVector<T, d>& PbInOut,
-	TRotation<T, d>& QbInOut, 
-	T MbInv, 
-	const PMatrix<T, d, d>& ILbInv, 
-	const TVector<T, d>& XLb,
-	const TRotation<T, d>& RLb,
-	const TPBD6DJointMotionSettings<T, d>& MotionSettings,
-	TPBD6DJointState<T, d>& State)
-{
-	TVector<T, d>& Pa = PaInOut;
-	TRotation<T, d>& Qa = QaInOut;
-	TVector<T, d>& Pb = PbInOut;
-	TRotation<T, d>& Qb = QbInOut;
-
-	// World-space mass state
-	PMatrix<T, d, d> IaInv = Utilities::Multiply(Qa.ToMatrix(), Utilities::Multiply(ILaInv, Qa.ToMatrix().GetTransposed()));
-	PMatrix<T, d, d> IbInv = Utilities::Multiply(Qb.ToMatrix(), Utilities::Multiply(ILbInv, Qb.ToMatrix().GetTransposed()));
-
-	bool bFlipQ = ((Qa | Qb) < 0);
-	if (bFlipQ)
+	void FPBD6DJointConstraints::RemoveConstraints(const TSet<TGeometryParticleHandle<FReal, 3>*>& RemovedParticles)
 	{
-		Qa = TRotation<T, d>::Negate(Qa);
 	}
 
-	T CurrentError = FLT_MAX;
 
-	int NumLoops = 0;
-	for (int LoopIndex = 0; (LoopIndex < SolverSettings.MaxIterations) && (CurrentError > SolverSettings.SolveTolerance); ++LoopIndex)
+	void FPBD6DJointConstraints::SetPreApplyCallback(const FD6JointPreApplyCallback& Callback)
 	{
-		// World-space constraint state
-		TVector<T, d> Xa = Pa + Qa * XLa;
-		TVector<T, d> Xb = Pb + Qb * XLb;
-		TRotation<T, d> Ra = Qa * RLa;
-		TRotation<T, d> Rb = Qb * RLb;
-
-#if !NO_LOGGING
-		if (!LogChaos6DJoint.IsSuppressed(ELogVerbosity::Verbose))
-		{
-			TVector<T, d> ConstraintAngles = Calculate6dConstraintAngles(SolverSettings, Ra, Rb, MotionSettings);
-			ConstraintAngles = TVector<T, d>(FMath::RadiansToDegrees(ConstraintAngles.X), FMath::RadiansToDegrees(ConstraintAngles.Y), FMath::RadiansToDegrees(ConstraintAngles.Z));
-			UE_LOG(LogChaos6DJoint, Verbose, TEXT("Pre Loop %d: Pos = %f (%f, %f, %f) Angle = (%f, %f, %f)"),
-				LoopIndex, (Xb - Xa).Size(), (Xb - Xa).X, (Xb - Xa).Y, (Xb - Xa).Z,
-				ConstraintAngles.X, ConstraintAngles.Y, ConstraintAngles.Z);
-		}
-#endif
-
-		// Get deltas to apply to position and rotation to correct constraint error
-		TVector<T, d> DPa, DPb;
-		TRotation<T, d> DQa, DQb;
-		bool bSolveOk = Calculate6dDelta(SolverSettings, Dt, Pa, Qa, MaInv, IaInv, Pb, Qb, MbInv, IbInv, Xa, Ra, Xb, Rb, MotionSettings, State, DPa, DQa, DPb, DQb);
-		if (!bSolveOk)
-		{
-			break;
-		}
-
-		// New world-space body state
-		TVector<T, d> Pa2 = Pa + DPa;
-		TVector<T, d> Pb2 = Pb + DPb;
-		TRotation<T, d> Qa2 = (Qa + DQa).GetNormalized();
-		TRotation<T, d> Qb2 = (Qb + DQb).GetNormalized();
-		Qb2.EnforceShortestArcWith(Qb);
-		Qa2.EnforceShortestArcWith(Qb2);
-
-		// Update body state
-		Pa = Pa2;
-		Pb = Pb2;
-		Qa = Qa2;
-		Qb = Qb2;
-		IaInv = Qa.ToMatrix() * ILaInv * Qa.ToMatrix().GetTransposed();
-		IbInv = Qb.ToMatrix() * ILbInv * Qb.ToMatrix().GetTransposed();
-
-		// New world-space constraint state
-		TVector<T, d> Xa2 = Pa2 + Qa2 * XLa;
-		TVector<T, d> Xb2 = Pb2 + Qb2 * XLb;
-		TRotation<T, d> Ra2 = Qa2 * RLa;
-		TRotation<T, d> Rb2 = Qb2 * RLb;
-
-#if !NO_LOGGING
-		if (!LogChaos6DJoint.IsSuppressed(ELogVerbosity::Verbose))
-		{
-			TVector<T, d> ConstraintAngles2 = Calculate6dConstraintAngles(SolverSettings, Ra2, Rb2, MotionSettings);
-			ConstraintAngles2 = TVector<T, d>(FMath::RadiansToDegrees(ConstraintAngles2.X), FMath::RadiansToDegrees(ConstraintAngles2.Y), FMath::RadiansToDegrees(ConstraintAngles2.Z));
-			UE_LOG(LogChaos6DJoint, Verbose, TEXT("Post Loop %d: Pos = %f (%f, %f, %f) Angle = (%f, %f, %f)"),
-				LoopIndex, (Xb2 - Xa2).Size(), (Xb2 - Xa2).X, (Xb2 - Xa2).Y, (Xb2 - Xa2).Z,
-				ConstraintAngles2.X, ConstraintAngles2.Y, ConstraintAngles2.Z);
-		}
-#endif
-
-		// @todo(ccaulfield): this isn't really a good error calculation - the magnitudes of positions and rotations are too different and its very expensive. It'll do for now though.
-		TVector<T, d> CX = Xb2 - Xa2;
-		TVector<T, d> CR, SR;
-		PMatrix<T, d, d> RRa, RRb;
-		TVector<T, d> LRMin, LRMax;
-		Calculate6dConstraintRotation(SolverSettings, Ra2, Rb2, MotionSettings, CR, RRa, RRb);
-		Calculate6dConstraintRotationLimits(SolverSettings, Ra2, Rb2, MotionSettings, SR, CR, RRa, RRb, LRMin, LRMax);
-		CurrentError = FMath::Sqrt(CX.SizeSquared() + CR.SizeSquared());
-		++NumLoops;
+		PreApplyCallback = Callback;
 	}
 
-	if (bFlipQ)
+	void FPBD6DJointConstraints::ClearPreApplyCallback()
 	{
-		Qa = TRotation<T, d>::Negate(Qa);
+		PreApplyCallback = nullptr;
 	}
 
-	return NumLoops;
-}
-
-template<typename T, int d>
-void TPBD6DJointConstraintUtilities<T, d>::Calculate3dDelta(
-	const TPBD6DJointSolverSettings<T, d>& SolverSettings,
-	const TVector<T, d>& Pa,
-	const TRotation<T, d>& Qa,
-	float MaInv,
-	const PMatrix<T, d, d>& IaInv,
-	const TVector<T, d>& Pb,
-	const TRotation<T, d>& Qb,
-	float MbInv,
-	const PMatrix<T, d, d>& IbInv,
-	const TVector<T, d>& Xa,
-	const TVector<T, d>& Xb,
-	const TPBD6DJointMotionSettings<T, d>& XSettings,
-	TVector<T, d>& DPa,
-	TRotation<T, d>& DQa,
-	TVector<T, d>& DPb,
-	TRotation<T, d>& DQb)
-{
-	TVector<T, d> CX = Xb - Xa;
-
-	T Stiffness = (SolverSettings.PBDStiffness > 0) ? SolverSettings.PBDStiffness : XSettings.Stiffness;
-	CX = Stiffness * CX;
-
-	PMatrix<T, d, d> Ma00 = PMatrix<T, d, d>(0, 0, 0);
-	PMatrix<T, d, d> Mb00 = PMatrix<T, d, d>(0, 0, 0);
-	if (MaInv > 0)
+	void FPBD6DJointConstraints::SetPostApplyCallback(const FD6JointPostApplyCallback& Callback)
 	{
-		Ma00 = Utilities::ComputeJointFactorMatrix(Xa - Pa, IaInv, MaInv);
+		PostApplyCallback = Callback;
 	}
-	if (MbInv > 0)
+
+	void FPBD6DJointConstraints::ClearPostApplyCallback()
 	{
-		Mb00 = Utilities::ComputeJointFactorMatrix(Xb - Pb, IbInv, MbInv);
+		PostApplyCallback = nullptr;
 	}
-	PMatrix<T, d, d> MI00 = (Ma00 + Mb00).Inverse();
-	TVector<T, d> DX = Utilities::Multiply(MI00, CX);
-
-	// Divide position and rotation error between bodies based on mass distribution
-	DPa = MaInv * DX;
-	DPb = -(MbInv * DX);
-	TVector<T, d> DQav = Utilities::Multiply(IaInv, TVector<T, d>::CrossProduct(Xa - Pa, DX));
-	TVector<T, d> DQbv = Utilities::Multiply(IbInv, TVector<T, d>::CrossProduct(Xb - Pb, -DX));
-	TRotation<T, d> DQaq = TRotation<T, d>::FromElements(DQav, 0);
-	TRotation<T, d> DQbq = TRotation<T, d>::FromElements(DQbv, 0);
-	DQa = (DQaq * Qa) * 0.5f;
-	DQb = (DQbq * Qb) * 0.5f;
-}
 
-template<typename T, int d>
-void TPBD6DJointConstraintUtilities<T, d>::Solve3dConstraint(
-	const TPBD6DJointSolverSettings<T, d>& SolverSettings,
-	TVector<T, d>& P0,
-	TRotation<T, d>& Q0,
-	T InvM0,
-	const PMatrix<T, d, d>& InvIL0,
-	const TVector<T, d>& XL0,
-	const TRotation<T, d>& RL0,
-	TVector<T, d>& P1,
-	TRotation<T, d>& Q1,
-	T InvM1,
-	const PMatrix<T, d, d>& InvIL1,
-	const TVector<T, d>& XL1,
-	const TRotation<T, d>& RL1,
-	const TPBD6DJointMotionSettings<T, d>& MotionSettings)
-{
-	TVector<T, d> X0 = P0 + Q0.RotateVector(XL0);
-	TVector<T, d> X1 = P1 + Q1.RotateVector(XL1);
-	PMatrix<T, d, d> InvI0 = Utilities::Multiply(Q0.ToMatrix(), Utilities::Multiply(InvIL0, Q0.ToMatrix().GetTransposed()));
-	PMatrix<T, d, d> InvI1 = Utilities::Multiply(Q1.ToMatrix(), Utilities::Multiply(InvIL1, Q1.ToMatrix().GetTransposed()));
-	TVector<T, d> DP0, DP1;
-	TRotation<T, d> DQ0, DQ1;
-
-	T CurrentError = FLT_MAX;
-	int NumLoops = 0;
-	for (int LoopIndex = 0; (LoopIndex < SolverSettings.MaxPreIterations) && (CurrentError > SolverSettings.SolveTolerance); ++LoopIndex)
+	const typename FPBD6DJointConstraints::FConstraintContainerHandle* FPBD6DJointConstraints::GetConstraintHandle(int32 ConstraintIndex) const
 	{
-		// Calculate position and rotation corrections
-		Calculate3dDelta(SolverSettings, P0, Q0, InvM0, InvI0, P1, Q1, InvM1, InvI1, X0, X1, MotionSettings, DP0, DQ0, DP1, DQ1);
-
-		// Apply corrections to body state
-		P0 = P0 + DP0;
-		Q0 = (Q0 + DQ0).GetNormalized();
-		P1 = P1 + DP1;
-		Q1 = (Q1 + DQ1).GetNormalized();
-
-		// Recalculate constraint state
-		X0 = P0 + Q0 * XL0;
-		X1 = P1 + Q1 * XL1;
-		InvI0 = Utilities::Multiply(Q0.ToMatrix(), Utilities::Multiply(InvIL0, Q0.ToMatrix().GetTransposed()));
-		InvI1 = Utilities::Multiply(Q1.ToMatrix(), Utilities::Multiply(InvIL1, Q1.ToMatrix().GetTransposed()));
-		
-		// Measure error
-		TVector<T, d> CX = X1 - X0;
-		CurrentError = CX.Size();
-		++NumLoops;
+		return Handles[ConstraintIndex];
 	}
-}
-
-//
-// Constraint JointSettings
-//
-
-template<class T, int d>
-TPBD6DJointMotionSettings<T, d>::TPBD6DJointMotionSettings()
-	: Stiffness((T)1) 
-	, LinearMotionTypes({ E6DJointMotionType::Locked, E6DJointMotionType::Locked, E6DJointMotionType::Locked })
-	, LinearLimits(TVector<T, d>(FLT_MAX, FLT_MAX, FLT_MAX))
-	, AngularMotionTypes({ E6DJointMotionType::Free, E6DJointMotionType::Free, E6DJointMotionType::Free })
-	, AngularLimits(TVector<T, d>(FLT_MAX, FLT_MAX, FLT_MAX))
-	, AngularDriveTarget(TRotation<T, d>::FromIdentity())
-	, AngularDriveTargetAngles(TVector<T, d>(0, 0, 0))
-	, bAngularSLerpDriveEnabled(false)
-	, bAngularTwistDriveEnabled(false)
-	, bAngularSwingDriveEnabled(false)
-	, AngularDriveStiffness(0)
-	, AngularDriveDamping(0)
-{
-}
-
-template<class T, int d>
-TPBD6DJointMotionSettings<T, d>::TPBD6DJointMotionSettings(const TVector<E6DJointMotionType, d>& InLinearMotionTypes, const TVector<E6DJointMotionType, d>& InAngularMotionTypes)
-	: Stiffness((T)1)
-	, LinearMotionTypes(InLinearMotionTypes)
-	, LinearLimits(TVector<T, d>(FLT_MAX, FLT_MAX, FLT_MAX))
-	, AngularMotionTypes({ E6DJointMotionType::Free, E6DJointMotionType::Free, E6DJointMotionType::Free })
-	, AngularLimits(TVector<T, d>(FLT_MAX, FLT_MAX, FLT_MAX))
-	, AngularDriveTarget(TRotation<T, d>::FromIdentity())
-	, AngularDriveTargetAngles(TVector<T, d>(0, 0, 0))
-	, bAngularSLerpDriveEnabled(false)
-	, bAngularTwistDriveEnabled(false)
-	, bAngularSwingDriveEnabled(false)
-	, AngularDriveStiffness(0)
-	, AngularDriveDamping(0)
-{
-}
-
-template<class T, int d>
-TPBD6DJointSettings<T, d>::TPBD6DJointSettings()
-	: ConstraintFrames({ FTransform::Identity, FTransform::Identity })
-{
-}
-
-template<class T, int d>
-TPBD6DJointState<T, d>::TPBD6DJointState()
-	: LambdaXa(TVector<T, d>(0, 0, 0))
-	, LambdaRa(TVector<T, d>(0, 0, 0))
-	, LambdaXb(TVector<T, d>(0, 0, 0))
-	, LambdaRb(TVector<T, d>(0, 0, 0))
-	, PrevTickCX(TVector<T, d>(0, 0, 0))
-	, PrevTickCR(TVector<T, d>(0, 0, 0))
-	, PrevItCX(TVector<T, d>(0, 0, 0))
-	, PrevItCR(TVector<T, d>(0, 0, 0))
-	, Level(INDEX_NONE)
-	, ParticleLevels({ INDEX_NONE, INDEX_NONE })
-{
-}
-
-//
-// Container JointSettings
-//
-
-template<class T, int d>
-TPBD6DJointSolverSettings<T, d>::TPBD6DJointSolverSettings()
-	: SolveTolerance(KINDA_SMALL_NUMBER)
-	, InvertedAxisTolerance( 0.001f)	// 1 - Cos(97.5deg)
-	, SwingTwistAngleTolerance(1.0e-6f)
-	, bApplyProjection(false)
-	, MaxIterations(10)
-	, MaxPreIterations(0)
-	, MaxDriveIterations(1)
-	, MaxRotComponent(0.0f)
-	, PBDMinParentMassRatio(0.5f)
-	, PBDMaxInertiaRatio(5.0f)
-	, FreezeIterations(0)
-	, FrozenIterations(0)
-	, bEnableAutoStiffness(true)
-	, bEnableTwistLimits(1)
-	, bEnableSwingLimits(1)
-	, bEnableDrives(1)
-	, XPBDAlphaX(0)
-	, XPBDAlphaR(0)
-	, XPBDBetaX(0)
-	, XPBDBetaR(0)
-	, PBDStiffness(0.0f)
-	, PBDDriveStiffness(0.0f)
-	, bFastSolve(false)
-{
-}
-
-//
-// Constraint Container
-//
-
-template<class T, int d>
-TPBD6DJointConstraints<T, d>::TPBD6DJointConstraints(const TPBD6DJointSolverSettings<T, d>& InSettings)
-	: Settings(InSettings)
-	, PreApplyCallback(nullptr)
-	, PostApplyCallback(nullptr)
-{
-}
 
-template<class T, int d>
-TPBD6DJointConstraints<T, d>::~TPBD6DJointConstraints() {}
-
-
-template<class T, int d>
-const TPBD6DJointSolverSettings<T, d>& TPBD6DJointConstraints<T, d>::GetSettings() const
-{
-	return Settings;
-}
-
-template<class T, int d>
-void TPBD6DJointConstraints<T, d>::SetSettings(const TPBD6DJointSolverSettings<T, d>& InSettings)
-{
-	Settings = InSettings;
-}
-
-template<class T, int d>
-int32 TPBD6DJointConstraints<T, d>::NumConstraints() const
-{
-	return ConstraintParticles.Num();
-}
-
-template<class T, int d>
-typename TPBD6DJointConstraints<T, d>::FConstraintHandle* TPBD6DJointConstraints<T, d>::AddConstraint(const FParticlePair& InConstrainedParticles, const FTransformPair& ConstraintFrames)
-{
-	int ConstraintIndex = Handles.Num();
-	Handles.Add(HandleAllocator.AllocHandle(this, ConstraintIndex));
-	ConstraintParticles.Add(InConstrainedParticles);
-	ConstraintSettings.Add(FJointSettings());
-	ConstraintSettings[ConstraintIndex].ConstraintFrames = ConstraintFrames;
-	ConstraintStates.Add(FJointState());
-	return Handles.Last();
-}
-
-template<class T, int d>
-typename TPBD6DJointConstraints<T, d>::FConstraintHandle* TPBD6DJointConstraints<T, d>::AddConstraint(const FParticlePair& InConstrainedParticles, const TPBD6DJointSettings<T, d>& InConstraintSettings)
-{
-	int ConstraintIndex = Handles.Num();
-	Handles.Add(HandleAllocator.AllocHandle(this, ConstraintIndex));
-	ConstraintParticles.Add(InConstrainedParticles);
-	ConstraintSettings.Add(InConstraintSettings);
-	ConstraintStates.Add(FJointState());
-	return Handles.Last();
-}
-
-template<class T, int d>
-void TPBD6DJointConstraints<T, d>::RemoveConstraint(int ConstraintIndex)
-{
-	FConstraintHandle* ConstraintHandle = Handles[ConstraintIndex];
-	if (ConstraintHandle != nullptr)
+	typename FPBD6DJointConstraints::FConstraintContainerHandle* FPBD6DJointConstraints::GetConstraintHandle(int32 ConstraintIndex)
 	{
-		// Release the handle for the freed constraint
-		HandleAllocator.FreeHandle(ConstraintHandle);
-		Handles[ConstraintIndex] = nullptr;
+		return Handles[ConstraintIndex];
 	}
 
-	// Swap the last constraint into the gap to keep the array packed
-	ConstraintParticles.RemoveAtSwap(ConstraintIndex);
-	ConstraintSettings.RemoveAtSwap(ConstraintIndex);
-	ConstraintStates.RemoveAtSwap(ConstraintIndex);
-	Handles.RemoveAtSwap(ConstraintIndex);
-
-	// Update the handle for the constraint that was moved
-	if (ConstraintIndex < Handles.Num())
+	const typename FPBD6DJointConstraints::FParticlePair& FPBD6DJointConstraints::GetConstrainedParticles(int32 ConstraintIndex) const
 	{
-		SetConstraintIndex(Handles[ConstraintIndex], ConstraintIndex);
+		return ConstraintParticles[ConstraintIndex];
 	}
-}
 
-template<class T, int d>
-void TPBD6DJointConstraints<T, d>::RemoveConstraints(const TSet<TGeometryParticleHandle<T, d>*>& RemovedParticles)
-{
-}
-
-
-template<class T, int d>
-void TPBD6DJointConstraints<T, d>::SetPreApplyCallback(const TD6JointPreApplyCallback<T, d>& Callback)
-{
-	PreApplyCallback = Callback;
-}
-
-template<class T, int d>
-void TPBD6DJointConstraints<T, d>::ClearPreApplyCallback()
-{
-	PreApplyCallback = nullptr;
-}
-
-template<class T, int d>
-void TPBD6DJointConstraints<T, d>::SetPostApplyCallback(const TD6JointPostApplyCallback<T, d>& Callback)
-{
-	PostApplyCallback = Callback;
-}
-
-template<class T, int d>
-void TPBD6DJointConstraints<T, d>::ClearPostApplyCallback()
-{
-	PostApplyCallback = nullptr;
-}
-
-template<class T, int d>
-const typename TPBD6DJointConstraints<T, d>::FConstraintHandle* TPBD6DJointConstraints<T, d>::GetConstraintHandle(int32 ConstraintIndex) const
-{
-	return Handles[ConstraintIndex];
-}
-
-template<class T, int d>
-typename TPBD6DJointConstraints<T, d>::FConstraintHandle* TPBD6DJointConstraints<T, d>::GetConstraintHandle(int32 ConstraintIndex)
-{
-	return Handles[ConstraintIndex];
-}
-
-template<class T, int d>
-const typename TPBD6DJointConstraints<T, d>::FParticlePair& TPBD6DJointConstraints<T, d>::GetConstrainedParticles(int32 ConstraintIndex) const
-{
-	return ConstraintParticles[ConstraintIndex];
-}
-
-template<class T, int d>
-int32 TPBD6DJointConstraints<T, d>::GetConstraintLevel(int32 ConstraintIndex) const
-{
-	return ConstraintStates[ConstraintIndex].Level;
-}
+	int32 FPBD6DJointConstraints::GetConstraintLevel(int32 ConstraintIndex) const
+	{
+		return ConstraintStates[ConstraintIndex].Level;
+	}
 
-template<class T, int d>
-void TPBD6DJointConstraints<T, d>::SetParticleLevels(int32 ConstraintIndex, const TVector<int32, 2>& ParticleLevels)
-{
-	ConstraintStates[ConstraintIndex].Level = FMath::Min(ParticleLevels[0], ParticleLevels[1]);
-	ConstraintStates[ConstraintIndex].ParticleLevels = ParticleLevels;
-}
+	void FPBD6DJointConstraints::SetParticleLevels(int32 ConstraintIndex, const TVector<int32, 2>& ParticleLevels)
+	{
+		ConstraintStates[ConstraintIndex].Level = FMath::Min(ParticleLevels[0], ParticleLevels[1]);
+		ConstraintStates[ConstraintIndex].ParticleLevels = ParticleLevels;
+	}
 
-template<class T, int d>
-void TPBD6DJointConstraints<T, d>::UpdatePositionBasedState(const T Dt)
-{
-	// @todo(ccaulfield): repurposing this since it is called before Apply, but maybe we need to rename the callback
-	for (FJointState& State : ConstraintStates)
+	void FPBD6DJointConstraints::UpdatePositionBasedState(const FReal Dt)
 	{
-		// @todo(ccaulfield): we should reinitialize PrevCX and PrevCR when initialized and teleported, etc
-		State.LambdaXa = TVector<T, d>(0, 0, 0);
-		State.LambdaRa = TVector<T, d>(0, 0, 0);
-		State.LambdaXb = TVector<T, d>(0, 0, 0);
-		State.LambdaRb = TVector<T, d>(0, 0, 0);
-		State.PrevTickCX = State.PrevItCX;
-		State.PrevTickCR = State.PrevItCR;
+		// @todo(ccaulfield): re-purposing this since it is called before Apply, but maybe we need to rename the callback
+		for (FPBD6DJointState& State : ConstraintStates)
+		{
+			// @todo(ccaulfield): we should reinitialize PrevCX and PrevCR when initialized and teleported, etc
+			State.LambdaXa = FVec3(0, 0, 0);
+			State.LambdaRa = FVec3(0, 0, 0);
+			State.LambdaXb = FVec3(0, 0, 0);
+			State.LambdaRb = FVec3(0, 0, 0);
+			State.PrevTickCX = State.PrevItCX;
+			State.PrevTickCR = State.PrevItCR;
+		}
 	}
-}
 
-template<class T, int d>
-void TPBD6DJointConstraints<T, d>::CalculateConstraintSpace(int32 ConstraintIndex, TVector<T, d>& OutXa, PMatrix<T, d, d>& OutRa, TVector<T, d>& OutXb, PMatrix<T, d, d>& OutRb, TVector<T, d>& OutCR) const
-{
-	const TVector<TGeometryParticleHandle<T, d>*, 2>& ConstrainedParticles = ConstraintParticles[ConstraintIndex];
-	const TPBD6DJointSettings<T, d>& JointSettings = ConstraintSettings[ConstraintIndex];
-
-	const TVector<T, d>& P0 = ConstrainedParticles[0]->AsDynamic() ? ConstrainedParticles[0]->AsDynamic()->P() : ConstrainedParticles[0]->X();
-	const TVector<T, d>& P1 = ConstrainedParticles[1]->AsDynamic() ? ConstrainedParticles[1]->AsDynamic()->P() : ConstrainedParticles[1]->X();
-	const TRotation<T, d>& Q0 = ConstrainedParticles[0]->AsDynamic() ? ConstrainedParticles[0]->AsDynamic()->Q() : ConstrainedParticles[0]->R();
-	const TRotation<T, d>& Q1 = ConstrainedParticles[1]->AsDynamic() ? ConstrainedParticles[1]->AsDynamic()->Q() : ConstrainedParticles[1]->R();
-	const TVector<T, d>& XL0 = JointSettings.ConstraintFrames[0].GetTranslation();
-	const TRotation<T, d>& RL0 = JointSettings.ConstraintFrames[0].GetRotation();
-	const TVector<T, d>& XL1 = JointSettings.ConstraintFrames[1].GetTranslation();
-	const TRotation<T, d>& RL1 = JointSettings.ConstraintFrames[1].GetRotation();
-	const TVector<T, d> X0 = P0 + Q0 * XL0;
-	const TVector<T, d> X1 = P1 + Q1 * XL1;
-	const TRotation<T, d> R0 = Q0 * RL0;
-	const TRotation<T, d> R1 = Q1 * RL1;
-
-	TPBD6DJointConstraintUtilities<T, d>::Calculate6dConstraintRotation(Settings, R0, R1, JointSettings.Motion, OutCR, OutRa, OutRb);
-	OutXa = X0;
-	OutXb = X1;
-	OutRa = OutRa.GetTransposed();
-	OutRb = OutRb.GetTransposed();
-}
+	void FPBD6DJointConstraints::CalculateConstraintSpace(int32 ConstraintIndex, FVec3& OutXa, FMatrix33& OutRa, FVec3& OutXb, FMatrix33& OutRb, FVec3& OutCR) const
+	{
+		const TVector<TGeometryParticleHandle<FReal, 3>*, 2>& ConstrainedParticles = ConstraintParticles[ConstraintIndex];
+		const FPBD6DJointSettings& JointSettings = ConstraintSettings[ConstraintIndex];
 
-template<class T, int d>
-void TPBD6DJointConstraints<T, d>::Apply(const T Dt, const TArray<FConstraintHandle*>& InConstraintHandles, const int32 It, const int32 NumIts)
-{
-	// @todo(ccaulfield): handles should be sorted by level by the constraint rule/graph
-	// @todo(ccaulfield): the best sort order depends on whether we are freezing.
-	// If we are freezing, we want the root-most bodies solved first, otherwise we want them last
-	TArray<FConstraintHandle*> SortedConstraintHandles = InConstraintHandles;
-	SortedConstraintHandles.Sort([](const FConstraintHandle& L, const FConstraintHandle& R)
-		{
-			return L.GetConstraintLevel() > R.GetConstraintLevel();
-		});
+		const FVec3& P0 = ConstrainedParticles[0]->AsDynamic() ? ConstrainedParticles[0]->AsDynamic()->P() : ConstrainedParticles[0]->X();
+		const FVec3& P1 = ConstrainedParticles[1]->AsDynamic() ? ConstrainedParticles[1]->AsDynamic()->P() : ConstrainedParticles[1]->X();
+		const FRotation3& Q0 = ConstrainedParticles[0]->AsDynamic() ? ConstrainedParticles[0]->AsDynamic()->Q() : ConstrainedParticles[0]->R();
+		const FRotation3& Q1 = ConstrainedParticles[1]->AsDynamic() ? ConstrainedParticles[1]->AsDynamic()->Q() : ConstrainedParticles[1]->R();
+		const FVec3& XL0 = JointSettings.ConstraintFrames[0].GetTranslation();
+		const FRotation3& RL0 = JointSettings.ConstraintFrames[0].GetRotation();
+		const FVec3& XL1 = JointSettings.ConstraintFrames[1].GetTranslation();
+		const FRotation3& RL1 = JointSettings.ConstraintFrames[1].GetRotation();
+		const FVec3 X0 = P0 + Q0 * XL0;
+		const FVec3 X1 = P1 + Q1 * XL1;
+		const FRotation3 R0 = Q0 * RL0;
+		const FRotation3 R1 = Q1 * RL1;
 
-	if (PreApplyCallback != nullptr)
-	{
-		PreApplyCallback(Dt, SortedConstraintHandles);
+		FPBD6DJointConstraintUtilities::Calculate6dConstraintRotation(Settings, R0, R1, JointSettings.Motion, OutCR, OutRa, OutRb);
+		OutXa = X0;
+		OutXb = X1;
+		OutRa = OutRa.GetTransposed();
+		OutRb = OutRb.GetTransposed();
 	}
 
-	if (Settings.bFastSolve)
+	void FPBD6DJointConstraints::Apply(const FReal Dt, const TArray<FConstraintContainerHandle*>& InConstraintHandles, const int32 It, const int32 NumIts)
 	{
-		for (FConstraintHandle* ConstraintHandle : SortedConstraintHandles)
+		// @todo(ccaulfield): handles should be sorted by level by the constraint rule/graph
+		// @todo(ccaulfield): the best sort order depends on whether we are freezing.
+		// If we are freezing, we want the root-most bodies solved first, otherwise we want them last
+		TArray<FConstraintContainerHandle*> SortedConstraintHandles = InConstraintHandles;
+		SortedConstraintHandles.Sort([](const FConstraintContainerHandle& L, const FConstraintContainerHandle& R)
+			{
+				return L.GetConstraintLevel() > R.GetConstraintLevel();
+			});
+
+		if (PreApplyCallback != nullptr)
 		{
-			ApplySingleFast(Dt, ConstraintHandle->GetConstraintIndex(), It, NumIts);
+			PreApplyCallback(Dt, SortedConstraintHandles);
 		}
-	}
-	else
-	{
-		T FreezeScale = (T)1;
-		if ((Settings.FreezeIterations + Settings.FrozenIterations) > 0)
+
+		if (Settings.bFastSolve)
+		{
+			for (FConstraintContainerHandle* ConstraintHandle : SortedConstraintHandles)
+			{
+				ApplySingleFast(Dt, ConstraintHandle->GetConstraintIndex(), It, NumIts);
+			}
+		}
+		else
 		{
-			const int32 BeginFreezingAt = NumIts - (Settings.FreezeIterations + Settings.FrozenIterations);
-			const int32 BeginFrozenAt = NumIts - Settings.FrozenIterations;
-			if (It >= BeginFrozenAt)
+			FReal FreezeScale = (FReal)1;
+			if ((Settings.FreezeIterations + Settings.FrozenIterations) > 0)
 			{
-				FreezeScale = (T)0;
+				const int32 BeginFreezingAt = NumIts - (Settings.FreezeIterations + Settings.FrozenIterations);
+				const int32 BeginFrozenAt = NumIts - Settings.FrozenIterations;
+				if (It >= BeginFrozenAt)
+				{
+					FreezeScale = (FReal)0;
+				}
+				else if (It >= BeginFreezingAt)
+				{
+					FreezeScale = (FReal)1 - (FReal)(It - BeginFreezingAt + 1) / (FReal)(BeginFrozenAt - BeginFreezingAt);
+				}
 			}
-			else if (It >= BeginFreezingAt)
+
+			for (FConstraintContainerHandle* ConstraintHandle : SortedConstraintHandles)
 			{
-				FreezeScale = (T)1 - (T)(It - BeginFreezingAt + 1) / (T)(BeginFrozenAt - BeginFreezingAt);
+				ApplySingle(Dt, ConstraintHandle->GetConstraintIndex(), FreezeScale);
 			}
 		}
 
-		for (FConstraintHandle* ConstraintHandle : SortedConstraintHandles)
+		if (PostApplyCallback != nullptr)
 		{
-			ApplySingle(Dt, ConstraintHandle->GetConstraintIndex(), FreezeScale);
+			PostApplyCallback(Dt, SortedConstraintHandles);
 		}
 	}
 
-	if (PostApplyCallback != nullptr)
+	bool FPBD6DJointConstraints::ApplyPushOut(const FReal Dt, const TArray<FConstraintContainerHandle*>& InConstraintHandles, const int32 It, const int32 NumIts)
 	{
-		PostApplyCallback(Dt, SortedConstraintHandles);
-	}
-}
+		// @todo(ccaulfield): track whether we are sufficiently solved
+		bool bNeedsAnotherIteration = true;
 
-template<class T, int d>
-bool TPBD6DJointConstraints<T, d>::ApplyPushOut(const T Dt, const TArray<FConstraintHandle*>& InConstraintHandles, const int32 It, const int32 NumIts)
-{
-	// @todo(ccaulfield): track whether we are sufficiently solved
-	bool bNeedsAnotherIteration = true;
+		for (FConstraintContainerHandle* ConstraintHandle : InConstraintHandles)
+		{
+			ApplyPushOutSingle(Dt, ConstraintHandle->GetConstraintIndex());
+		}
 
-	for (FConstraintHandle* ConstraintHandle : InConstraintHandles)
-	{
-		ApplyPushOutSingle(Dt, ConstraintHandle->GetConstraintIndex());
+		return bNeedsAnotherIteration;
 	}
 
-	return bNeedsAnotherIteration;
-}
-
-template<class T, int d>
-void TPBD6DJointConstraints<T, d>::ApplySingle(const T Dt, const int32 ConstraintIndex, const T FreezeScale)
-{
-	const TVector<TGeometryParticleHandle<T, d>*, 2>& Constraint = ConstraintParticles[ConstraintIndex];
-	UE_LOG(LogChaos6DJoint, Verbose, TEXT("6DoF Solve Constraint %d %s %s (dt = %f; freeze = %f)"), ConstraintIndex, *Constraint[0]->ToString(), *Constraint[1]->ToString(), Dt, FreezeScale);
-
-	if (Constraint[0]->AsDynamic() && Constraint[1]->AsDynamic())
-	{
-		ApplyDynamicDynamic(Dt, ConstraintIndex, 0, 1, FreezeScale);
-	}
-	else if (Constraint[0]->AsDynamic())
+	void FPBD6DJointConstraints::ApplySingle(const FReal Dt, const int32 ConstraintIndex, const FReal FreezeScale)
 	{
-		ApplyDynamicStatic(Dt, ConstraintIndex, 0, 1);
-	}
-	else
-	{
-		ApplyDynamicStatic(Dt, ConstraintIndex, 1, 0);
+		const TVector<TGeometryParticleHandle<FReal, 3>*, 2>& Constraint = ConstraintParticles[ConstraintIndex];
+		UE_LOG(LogChaos6DJoint, Verbose, TEXT("6DoF Solve Constraint %d %s %s (dt = %f; freeze = %f)"), ConstraintIndex, *Constraint[0]->ToString(), *Constraint[1]->ToString(), Dt, FreezeScale);
+
+		if (Constraint[0]->AsDynamic() && Constraint[1]->AsDynamic())
+		{
+			ApplyDynamicDynamic(Dt, ConstraintIndex, 0, 1, FreezeScale);
+		}
+		else if (Constraint[0]->AsDynamic())
+		{
+			ApplyDynamicStatic(Dt, ConstraintIndex, 0, 1);
+		}
+		else
+		{
+			ApplyDynamicStatic(Dt, ConstraintIndex, 1, 0);
+		}
 	}
-}
 
-/**
- * Increase the lower inertia components to ensure that the maximum ratio between any pair of elements is MaxRatio.
- *
- * @param InI The input inertia.
- * @return An altered inertia so that the minimum element is at least MaxElement/MaxRatio.
- */
-template<class T, int d>
-TVector<T, d> ConditionInertia(const TVector<T, d>& InI, const T MaxRatio)
-{
-	// @todo(ccaulfield): simd
-	if (MaxRatio > 0)
+	/**
+	 * Increase the lower inertia components to ensure that the maximum ratio between any pair of elements is MaxRatio.
+	 *
+	 * @param InI The input inertia.
+	 * @return An altered inertia so that the minimum element is at least MaxElement/MaxRatio.
+	 */
+	FVec3 ConditionInertia(const FVec3& InI, const FReal MaxRatio)
 	{
-		T IMin = InI.Min();
-		T IMax = InI.Max();
-		T Ratio = IMax / IMin;
-		if (Ratio > MaxRatio)
+		// @todo(ccaulfield): simd
+		if (MaxRatio > 0)
 		{
-			T MinIMin = IMax / MaxRatio;
-			return TVector<T, d>(
-				FMath::Lerp(MinIMin, IMax, (InI.X - IMin) / (IMax - IMin)),
-				FMath::Lerp(MinIMin, IMax, (InI.Y - IMin) / (IMax - IMin)),
-				FMath::Lerp(MinIMin, IMax, (InI.Z - IMin) / (IMax - IMin)));
+			FReal IMin = InI.Min();
+			FReal IMax = InI.Max();
+			FReal Ratio = IMax / IMin;
+			if (Ratio > MaxRatio)
+			{
+				FReal MinIMin = IMax / MaxRatio;
+				return FVec3(
+					FMath::Lerp(MinIMin, IMax, (InI.X - IMin) / (IMax - IMin)),
+					FMath::Lerp(MinIMin, IMax, (InI.Y - IMin) / (IMax - IMin)),
+					FMath::Lerp(MinIMin, IMax, (InI.Z - IMin) / (IMax - IMin)));
+			}
 		}
+		return InI;
 	}
-	return InI;
-}
 
 
-/**
- * Increase the IParent inertia so that its largest component is at least MinRatio times the largest IChild component.
- * This is used to condition joint chains for more robust solving with low iteration counts or larger time steps.
- *
- * @param IParent The input inertia.
- * @param IChild The input inertia.
- * @param OutIParent The output inertia.
- * @param MinRatio Parent inertia will be at least this multiple of child inertia
- * @return The max/min ratio of InI elements.
- */
-template<class T, int d>
-TVector<T, d> ConditionParentInertia(const TVector<T, d>& IParent, const TVector<T, d>& IChild, const T MinRatio)
-{
-	// @todo(ccaulfield): simd
-	if (MinRatio > 0)
+	/**
+	 * Increase the IParent inertia so that its largest component is at least MinRatio times the largest IChild component.
+	 * This is used to condition joint chains for more robust solving with low iteration counts or larger time steps.
+	 *
+	 * @param IParent The input inertia.
+	 * @param IChild The input inertia.
+	 * @param OutIParent The output inertia.
+	 * @param MinRatio Parent inertia will be at least this multiple of child inertia
+	 * @return The max/min ratio of InI elements.
+	 */
+	FVec3 ConditionParentInertia(const FVec3& IParent, const FVec3& IChild, const FReal MinRatio)
 	{
-		T IParentMax = IParent.Max();
-		T IChildMax = IChild.Max();
-		T Ratio = IParentMax / IChildMax;
-		if (Ratio < MinRatio)
+		// @todo(ccaulfield): simd
+		if (MinRatio > 0)
 		{
-			T Multiplier = MinRatio / Ratio;
-			return IParent * Multiplier;
+			FReal IParentMax = IParent.Max();
+			FReal IChildMax = IChild.Max();
+			FReal Ratio = IParentMax / IChildMax;
+			if (Ratio < MinRatio)
+			{
+				FReal Multiplier = MinRatio / Ratio;
+				return IParent * Multiplier;
+			}
 		}
+		return IParent;
 	}
-	return IParent;
-}
 
-template<class T, int d>
-T ConditionParentMass(const T MParent, const T MChild, const T MinRatio)
-{
-	if (MinRatio > 0)
+	FReal ConditionParentMass(const FReal MParent, const FReal MChild, const FReal MinRatio)
 	{
-		T Ratio = MParent / MChild;
-		if (Ratio < MinRatio)
+		if (MinRatio > 0)
 		{
-			T Multiplier = MinRatio / Ratio;
-			return MParent * Multiplier;
+			FReal Ratio = MParent / MChild;
+			if (Ratio < MinRatio)
+			{
+				FReal Multiplier = MinRatio / Ratio;
+				return MParent * Multiplier;
+			}
 		}
+		return MParent;
 	}
-	return MParent;
-}
-
-template<class T, int d>
-void GetConditionedInverseMass(const TPBDRigidParticleHandle<T, d>* PParent, const TPBDRigidParticleHandle<T, d>* PChild, T& OutInvMParent, T& OutInvMChild, PMatrix<T, d, d>& OutInvIParent, PMatrix<T, d, d>& OutInvIChild, const T MinParentMassRatio, const T MaxInertiaRatio)
-{
-	T MParent = PParent->M();
-	T MChild = PChild->M();
-	MParent = ConditionParentMass<T, d>(MParent, MChild, MinParentMassRatio);
-
-	TVector<T, d> IParent = ConditionInertia<T, d>(PParent->I().GetDiagonal(), MaxInertiaRatio);
-	TVector<T, d> IChild = ConditionInertia<T, d>(PChild->I().GetDiagonal(), MaxInertiaRatio);
-	IParent = ConditionParentInertia(IParent, IChild, MinParentMassRatio);
-
-	OutInvMParent = (T)1 / MParent;
-	OutInvMChild = (T)1 / MChild;
-	OutInvIParent = PMatrix<T, d, d>((T)1 / IParent.X, (T)1 / IParent.Y, (T)1 / IParent.Z);
-	OutInvIChild = PMatrix<T, d, d>((T)1 / IChild.X, (T)1 / IChild.Y, (T)1 / IChild.Z);
-}
-
-
-template<class T, int d>
-void GetConditionedInverseMass(const TPBDRigidParticleHandle<T, d>* P0, T& OutInvM0, PMatrix<T, d, d>& OutInvI0, const T MaxInertiaRatio)
-{
-	TVector<T, d> I0 = ConditionInertia<T, d>(P0->I().GetDiagonal(), MaxInertiaRatio);
-
-	OutInvM0 = P0->InvM();
-	OutInvI0 = PMatrix<T, d, d>((T)1 / I0.X, (T)1 / I0.Y, (T)1 / I0.Z);
-}
-
 
-template<class T, int d>
-void TPBD6DJointConstraints<T, d>::ApplyDynamicDynamic(const T Dt, const int32 ConstraintIndex, const int32 PBDRigid0Index, const int32 PBDRigid1Index, const T FreezeScale)
-{
-	check((PBDRigid0Index == 0) || (PBDRigid0Index == 1));
-	check((PBDRigid1Index == 0) || (PBDRigid1Index == 1));
-	check(PBDRigid0Index != PBDRigid1Index);
-
-	TPBDRigidParticleHandle<T, d>* PBDRigid0 = ConstraintParticles[ConstraintIndex][PBDRigid0Index]->AsDynamic();
-	TPBDRigidParticleHandle<T, d>* PBDRigid1 = ConstraintParticles[ConstraintIndex][PBDRigid1Index]->AsDynamic();
-	check(PBDRigid0 && PBDRigid1 && (PBDRigid0->Island() == PBDRigid1->Island()));
-
-	TRotation<T, d> Q0 = PBDRigid0->Q();
-	TVector<T, d> P0 = PBDRigid0->P();
-	const TVector<T, d>& XL0 = ConstraintSettings[ConstraintIndex].ConstraintFrames[PBDRigid0Index].GetTranslation();
-	const TRotation<T, d>& RL0 = ConstraintSettings[ConstraintIndex].ConstraintFrames[PBDRigid0Index].GetRotation();
-	TRotation<T, d> Q1 = PBDRigid1->Q();
-	TVector<T, d> P1 = PBDRigid1->P();
-	const TVector<T, d>& XL1 = ConstraintSettings[ConstraintIndex].ConstraintFrames[PBDRigid1Index].GetTranslation();
-	const TRotation<T, d>& RL1 = ConstraintSettings[ConstraintIndex].ConstraintFrames[PBDRigid1Index].GetRotation();
-	float InvM0, InvM1, InvM0F, InvM1F;
-	PMatrix<T, d, d> InvIL0, InvIL1, InvIL0F, InvIL1F;
-
-	// Adjust mass and freeze particle up the chain (nearest to a connection to a non-dynamic particle).
-	if (ConstraintStates[ConstraintIndex].ParticleLevels[0] < ConstraintStates[ConstraintIndex].ParticleLevels[1])
-	{
-		GetConditionedInverseMass(PBDRigid0, PBDRigid1, InvM0, InvM1, InvIL0, InvIL1, Settings.PBDMinParentMassRatio, Settings.PBDMaxInertiaRatio);
-		InvM0F = InvM0 * FreezeScale;
-		InvM1F = InvM1;
-		InvIL0F = InvIL0 * FreezeScale;
-		InvIL1F = InvIL1;
-	}
-	else if (ConstraintStates[ConstraintIndex].ParticleLevels[0] > ConstraintStates[ConstraintIndex].ParticleLevels[1])
-	{
-		GetConditionedInverseMass(PBDRigid1, PBDRigid0, InvM1, InvM0, InvIL1, InvIL0, Settings.PBDMinParentMassRatio, Settings.PBDMaxInertiaRatio);
-		InvM0F = InvM0;
-		InvM1F = InvM1 * FreezeScale;
-		InvIL0F = InvIL0;
-		InvIL1F = InvIL1 * FreezeScale;
-	}
-	else
+	void GetConditionedInverseMass(const TPBDRigidParticleHandle<FReal, 3>* PParent, const TPBDRigidParticleHandle<FReal, 3>* PChild, FReal& OutInvMParent, FReal& OutInvMChild, FMatrix33& OutInvIParent, FMatrix33& OutInvIChild, const FReal MinParentMassRatio, const FReal MaxInertiaRatio)
 	{
-		GetConditionedInverseMass(PBDRigid1, PBDRigid0, InvM1, InvM0, InvIL1, InvIL0, (T)0, Settings.PBDMaxInertiaRatio);
-		InvM0F = InvM0;
-		InvM1F = InvM1;
-		InvIL0F = InvIL0;
-		InvIL1F = InvIL1;
-	}
+		FReal MParent = PParent->M();
+		FReal MChild = PChild->M();
+		MParent = ConditionParentMass(MParent, MChild, MinParentMassRatio);
 
-	TVector<E6DJointMotionType, d> AngularMotionTypes = ConstraintSettings[ConstraintIndex].Motion.AngularMotionTypes;
-	bool bTwistDriveEnabled = ConstraintSettings[ConstraintIndex].Motion.bAngularTwistDriveEnabled;
-	bool bSwingDriveEnabled = ConstraintSettings[ConstraintIndex].Motion.bAngularSwingDriveEnabled;
-	bool bAutoStiffness = Settings.bEnableAutoStiffness;
+		FVec3 IParent = ConditionInertia(PParent->I().GetDiagonal(), MaxInertiaRatio);
+		FVec3 IChild = ConditionInertia(PChild->I().GetDiagonal(), MaxInertiaRatio);
+		IParent = ConditionParentInertia(IParent, IChild, MinParentMassRatio);
 
-	if (Settings.MaxPreIterations > 0)
-	{
-		TPBD6DJointConstraintUtilities<T, d>::Solve3dConstraint(Settings, P0, Q0, InvM0F, InvIL0F, XL0, RL0, P1, Q1, InvM1F, InvIL1F, XL1, RL1, ConstraintSettings[ConstraintIndex].Motion);
+		OutInvMParent = (FReal)1 / MParent;
+		OutInvMChild = (FReal)1 / MChild;
+		OutInvIParent = FMatrix33((FReal)1 / IParent.X, (FReal)1 / IParent.Y, (FReal)1 / IParent.Z);
+		OutInvIChild = FMatrix33((FReal)1 / IChild.X, (FReal)1 / IChild.Y, (FReal)1 / IChild.Z);
 	}
-	if (Settings.MaxIterations > 0)
-	{
-		ConstraintSettings[ConstraintIndex].Motion.bAngularTwistDriveEnabled = false;
-		ConstraintSettings[ConstraintIndex].Motion.bAngularSwingDriveEnabled = false;
 
-		//TPBD6DJointConstraintUtilities<T, d>::Solve6dConstraint(Settings, Dt, P0, Q0, InvM0F, InvIL0F, XL0, RL0, P1, Q1, InvM1F, InvIL1F, XL1, RL1, ConstraintSettings[ConstraintIndex].Motion, ConstraintStates[ConstraintIndex]);
-		TPBD6DJointConstraintUtilities<T, d>::Solve6dConstraint(Settings, Dt, P1, Q1, InvM1F, InvIL1F, XL1, RL1, P0, Q0, InvM0F, InvIL0F, XL0, RL0, ConstraintSettings[ConstraintIndex].Motion, ConstraintStates[ConstraintIndex]);
 
-		ConstraintSettings[ConstraintIndex].Motion.bAngularTwistDriveEnabled = bTwistDriveEnabled;
-		ConstraintSettings[ConstraintIndex].Motion.bAngularSwingDriveEnabled = bSwingDriveEnabled;
-	}
-	if ((bTwistDriveEnabled || bSwingDriveEnabled) && (Settings.MaxDriveIterations > 0))
+	void GetConditionedInverseMass(const TPBDRigidParticleHandle<FReal, 3>* P0, FReal& OutInvM0, FMatrix33& OutInvI0, const FReal MaxInertiaRatio)
 	{
-		Settings.bEnableAutoStiffness = false;
-		ConstraintSettings[ConstraintIndex].Motion.AngularMotionTypes = { E6DJointMotionType::Free, E6DJointMotionType::Free , E6DJointMotionType::Free };
-
-		TPBD6DJointConstraintUtilities<T, d>::Solve6dConstraint(Settings, Dt, P1, Q1, InvM1, InvIL1, XL1, RL1, P0, Q0, InvM0, InvIL0, XL0, RL0, ConstraintSettings[ConstraintIndex].Motion, ConstraintStates[ConstraintIndex]);
+		FVec3 I0 = ConditionInertia(P0->I().GetDiagonal(), MaxInertiaRatio);
 
-		Settings.bEnableAutoStiffness = bAutoStiffness;
-		ConstraintSettings[ConstraintIndex].Motion.AngularMotionTypes = AngularMotionTypes;
+		OutInvM0 = P0->InvM();
+		OutInvI0 = FMatrix33((FReal)1 / I0.X, (FReal)1 / I0.Y, (FReal)1 / I0.Z);
 	}
 
-	PBDRigid0->SetQ(Q0);
-	PBDRigid0->SetP(P0);
-	PBDRigid1->SetQ(Q1);
-	PBDRigid1->SetP(P1);
-}
 
-template<class T, int d>
-void TPBD6DJointConstraints<T, d>::ApplyDynamicStatic(const T Dt, const int32 ConstraintIndex, const int32 PBDRigid0Index, const int32 Static1Index)
-{
-	check((PBDRigid0Index == 0) || (PBDRigid0Index == 1));
-	check((Static1Index == 0) || (Static1Index == 1));
-	check(PBDRigid0Index != Static1Index);
-
-	TPBDRigidParticleHandle<T, d>* PBDRigid0 = ConstraintParticles[ConstraintIndex][PBDRigid0Index]->AsDynamic();
-	TGeometryParticleHandle<T, d>* Static1 = ConstraintParticles[ConstraintIndex][Static1Index];
-	check(PBDRigid0 && Static1 && !Static1->AsDynamic());
-
-	TRotation<T, d> Q0 = PBDRigid0->Q();
-	TVector<T, d> P0 = PBDRigid0->P();
-	const TVector<T, d>& XL0 = ConstraintSettings[ConstraintIndex].ConstraintFrames[PBDRigid0Index].GetTranslation();
-	const TRotation<T, d>& RL0 = ConstraintSettings[ConstraintIndex].ConstraintFrames[PBDRigid0Index].GetRotation();
-	TRotation<T, d> Q1 = Static1->R();
-	TVector<T, d> P1 = Static1->X();
-	const TVector<T, d>& XL1 = ConstraintSettings[ConstraintIndex].ConstraintFrames[Static1Index].GetTranslation();
-	const TRotation<T, d>& RL1 = ConstraintSettings[ConstraintIndex].ConstraintFrames[Static1Index].GetRotation();
-	float InvM0;
-	PMatrix<T, d, d> InvIL0;
-	GetConditionedInverseMass(PBDRigid0, InvM0, InvIL0, Settings.PBDMaxInertiaRatio);
-	const float InvM1 = 0;
-	const PMatrix<T, d, d> InvIL1 = PMatrix<T, d, d>(0, 0, 0);
-
-	TVector<E6DJointMotionType, d> AngularMotionTypes = ConstraintSettings[ConstraintIndex].Motion.AngularMotionTypes;
-	bool bTwistDriveEnabled = ConstraintSettings[ConstraintIndex].Motion.bAngularTwistDriveEnabled;
-	bool bSwingDriveEnabled = ConstraintSettings[ConstraintIndex].Motion.bAngularSwingDriveEnabled;
-	bool bAutoStiffness = Settings.bEnableAutoStiffness;
-
-	// NOTE: We put the static body first in the solver - swing axes are calculated relative to this
-	if (Settings.MaxPreIterations > 0)
+	void FPBD6DJointConstraints::ApplyDynamicDynamic(const FReal Dt, const int32 ConstraintIndex, const int32 PBDRigid0Index, const int32 PBDRigid1Index, const FReal FreezeScale)
 	{
-		TPBD6DJointConstraintUtilities<T, d>::Solve3dConstraint(Settings, P1, Q1, InvM1, InvIL1, XL1, RL1, P0, Q0, InvM0, InvIL0, XL0, RL0, ConstraintSettings[ConstraintIndex].Motion);
-	}
-	if (Settings.MaxIterations > 0)
-	{
-		ConstraintSettings[ConstraintIndex].Motion.bAngularTwistDriveEnabled = false;
-		ConstraintSettings[ConstraintIndex].Motion.bAngularSwingDriveEnabled = false;
+		check((PBDRigid0Index == 0) || (PBDRigid0Index == 1));
+		check((PBDRigid1Index == 0) || (PBDRigid1Index == 1));
+		check(PBDRigid0Index != PBDRigid1Index);
 
-		TPBD6DJointConstraintUtilities<T, d>::Solve6dConstraint(Settings, Dt, P1, Q1, InvM1, InvIL1, XL1, RL1, P0, Q0, InvM0, InvIL0, XL0, RL0, ConstraintSettings[ConstraintIndex].Motion, ConstraintStates[ConstraintIndex]);
+		TPBDRigidParticleHandle<FReal, 3>* PBDRigid0 = ConstraintParticles[ConstraintIndex][PBDRigid0Index]->AsDynamic();
+		TPBDRigidParticleHandle<FReal, 3>* PBDRigid1 = ConstraintParticles[ConstraintIndex][PBDRigid1Index]->AsDynamic();
+		check(PBDRigid0 && PBDRigid1 && (PBDRigid0->Island() == PBDRigid1->Island()));
 
-		ConstraintSettings[ConstraintIndex].Motion.bAngularTwistDriveEnabled = bTwistDriveEnabled;
-		ConstraintSettings[ConstraintIndex].Motion.bAngularSwingDriveEnabled = bSwingDriveEnabled;
-	}
-	if ((bTwistDriveEnabled || bSwingDriveEnabled) && (Settings.MaxDriveIterations > 0))
-	{
-		Settings.bEnableAutoStiffness = false;
-		ConstraintSettings[ConstraintIndex].Motion.AngularMotionTypes = { E6DJointMotionType::Free, E6DJointMotionType::Free , E6DJointMotionType::Free };
+		FRotation3 Q0 = PBDRigid0->Q();
+		FVec3 P0 = PBDRigid0->P();
+		const FVec3& XL0 = ConstraintSettings[ConstraintIndex].ConstraintFrames[PBDRigid0Index].GetTranslation();
+		const FRotation3& RL0 = ConstraintSettings[ConstraintIndex].ConstraintFrames[PBDRigid0Index].GetRotation();
+		FRotation3 Q1 = PBDRigid1->Q();
+		FVec3 P1 = PBDRigid1->P();
+		const FVec3& XL1 = ConstraintSettings[ConstraintIndex].ConstraintFrames[PBDRigid1Index].GetTranslation();
+		const FRotation3& RL1 = ConstraintSettings[ConstraintIndex].ConstraintFrames[PBDRigid1Index].GetRotation();
+		float InvM0, InvM1, InvM0F, InvM1F;
+		FMatrix33 InvIL0, InvIL1, InvIL0F, InvIL1F;
 
-		TPBD6DJointConstraintUtilities<T, d>::Solve6dConstraint(Settings, Dt, P1, Q1, InvM1, InvIL1, XL1, RL1, P0, Q0, InvM0, InvIL0, XL0, RL0, ConstraintSettings[ConstraintIndex].Motion, ConstraintStates[ConstraintIndex]);
+		// Adjust mass and freeze particle up the chain (nearest to a connection to a non-dynamic particle).
+		if (ConstraintStates[ConstraintIndex].ParticleLevels[0] < ConstraintStates[ConstraintIndex].ParticleLevels[1])
+		{
+			GetConditionedInverseMass(PBDRigid0, PBDRigid1, InvM0, InvM1, InvIL0, InvIL1, Settings.PBDMinParentMassRatio, Settings.PBDMaxInertiaRatio);
+			InvM0F = InvM0 * FreezeScale;
+			InvM1F = InvM1;
+			InvIL0F = InvIL0 * FreezeScale;
+			InvIL1F = InvIL1;
+		}
+		else if (ConstraintStates[ConstraintIndex].ParticleLevels[0] > ConstraintStates[ConstraintIndex].ParticleLevels[1])
+		{
+			GetConditionedInverseMass(PBDRigid1, PBDRigid0, InvM1, InvM0, InvIL1, InvIL0, Settings.PBDMinParentMassRatio, Settings.PBDMaxInertiaRatio);
+			InvM0F = InvM0;
+			InvM1F = InvM1 * FreezeScale;
+			InvIL0F = InvIL0;
+			InvIL1F = InvIL1 * FreezeScale;
+		}
+		else
+		{
+			GetConditionedInverseMass(PBDRigid1, PBDRigid0, InvM1, InvM0, InvIL1, InvIL0, (FReal)0, Settings.PBDMaxInertiaRatio);
+			InvM0F = InvM0;
+			InvM1F = InvM1;
+			InvIL0F = InvIL0;
+			InvIL1F = InvIL1;
+		}
 
-		Settings.bEnableAutoStiffness = bAutoStiffness;
-		ConstraintSettings[ConstraintIndex].Motion.AngularMotionTypes = AngularMotionTypes;
-	}
+		TVector<E6DJointMotionType, 3> AngularMotionTypes = ConstraintSettings[ConstraintIndex].Motion.AngularMotionTypes;
+		bool bTwistDriveEnabled = ConstraintSettings[ConstraintIndex].Motion.bAngularTwistDriveEnabled;
+		bool bSwingDriveEnabled = ConstraintSettings[ConstraintIndex].Motion.bAngularSwingDriveEnabled;
+		bool bAutoStiffness = Settings.bEnableAutoStiffness;
 
-	PBDRigid0->SetQ(Q0);
-	PBDRigid0->SetP(P0);
-}
+		if (Settings.MaxPreIterations > 0)
+		{
+			FPBD6DJointConstraintUtilities::Solve3dConstraint(Settings, P0, Q0, InvM0F, InvIL0F, XL0, RL0, P1, Q1, InvM1F, InvIL1F, XL1, RL1, ConstraintSettings[ConstraintIndex].Motion);
+		}
+		if (Settings.MaxIterations > 0)
+		{
+			ConstraintSettings[ConstraintIndex].Motion.bAngularTwistDriveEnabled = false;
+			ConstraintSettings[ConstraintIndex].Motion.bAngularSwingDriveEnabled = false;
 
-template<class T, int d>
-void TPBD6DJointConstraints<T, d>::ApplyPushOutSingle(const T Dt, const int32 ConstraintIndex)
-{
-	// Correct any remaining error by translating
-	if (Settings.bApplyProjection)
-	{
-		const TVector<TGeometryParticleHandle<T, d>*, 2>& ConstrainedParticles = ConstraintParticles[ConstraintIndex];
-		const TPBD6DJointSettings<T, d>& JointSettings = ConstraintSettings[ConstraintIndex];
+			//FPBD6DJointConstraintUtilities::Solve6dConstraint(Settings, Dt, P0, Q0, InvM0F, InvIL0F, XL0, RL0, P1, Q1, InvM1F, InvIL1F, XL1, RL1, ConstraintSettings[ConstraintIndex].Motion, ConstraintStates[ConstraintIndex]);
+			FPBD6DJointConstraintUtilities::Solve6dConstraint(Settings, Dt, P1, Q1, InvM1F, InvIL1F, XL1, RL1, P0, Q0, InvM0F, InvIL0F, XL0, RL0, ConstraintSettings[ConstraintIndex].Motion, ConstraintStates[ConstraintIndex]);
 
-		TVector<T, d> P0 = ConstrainedParticles[0]->AsDynamic() ? ConstrainedParticles[0]->AsDynamic()->P() : ConstrainedParticles[0]->X();
-		TVector<T, d> P1 = ConstrainedParticles[1]->AsDynamic() ? ConstrainedParticles[1]->AsDynamic()->P() : ConstrainedParticles[1]->X();
-		TRotation<T, d> Q0 = ConstrainedParticles[0]->AsDynamic() ? ConstrainedParticles[0]->AsDynamic()->Q() : ConstrainedParticles[0]->R();
-		TRotation<T, d> Q1 = ConstrainedParticles[1]->AsDynamic() ? ConstrainedParticles[1]->AsDynamic()->Q() : ConstrainedParticles[1]->R();
-		const TVector<T, d>& XL0 = JointSettings.ConstraintFrames[0].GetTranslation();
-		const TRotation<T, d>& RL0 = JointSettings.ConstraintFrames[0].GetRotation();
-		const TVector<T, d>& XL1 = JointSettings.ConstraintFrames[1].GetTranslation();
-		const TRotation<T, d>& RL1 = JointSettings.ConstraintFrames[1].GetRotation();
-		TVector<T, d> X0 = P0 + Q0 * XL0;
-		TVector<T, d> X1 = P1 + Q1 * XL1;
-		const T InvM0 = ConstrainedParticles[0]->AsDynamic()? ConstrainedParticles[0]->AsDynamic()->InvM() : (T)0;
-		const T InvM1 = ConstrainedParticles[1]->AsDynamic() ? ConstrainedParticles[1]->AsDynamic()->InvM() : (T)0;
+			ConstraintSettings[ConstraintIndex].Motion.bAngularTwistDriveEnabled = bTwistDriveEnabled;
+			ConstraintSettings[ConstraintIndex].Motion.bAngularSwingDriveEnabled = bSwingDriveEnabled;
+		}
+		if ((bTwistDriveEnabled || bSwingDriveEnabled) && (Settings.MaxDriveIterations > 0))
+		{
+			Settings.bEnableAutoStiffness = false;
+			ConstraintSettings[ConstraintIndex].Motion.AngularMotionTypes = { E6DJointMotionType::Free, E6DJointMotionType::Free , E6DJointMotionType::Free };
 
+			FPBD6DJointConstraintUtilities::Solve6dConstraint(Settings, Dt, P1, Q1, InvM1, InvIL1, XL1, RL1, P0, Q0, InvM0, InvIL0, XL0, RL0, ConstraintSettings[ConstraintIndex].Motion, ConstraintStates[ConstraintIndex]);
 
-		const TVector<T, d> DeltaProj = (X1 - X0) / (InvM0 + InvM1);
-		P0 += InvM0 * DeltaProj;
-		P1 -= InvM1 * DeltaProj;
+			Settings.bEnableAutoStiffness = bAutoStiffness;
+			ConstraintSettings[ConstraintIndex].Motion.AngularMotionTypes = AngularMotionTypes;
+		}
 
-		if (ConstrainedParticles[0]->AsDynamic())
+		PBDRigid0->SetQ(Q0);
+		PBDRigid0->SetP(P0);
+		PBDRigid1->SetQ(Q1);
+		PBDRigid1->SetP(P1);
+	}
+
+	void FPBD6DJointConstraints::ApplyDynamicStatic(const FReal Dt, const int32 ConstraintIndex, const int32 PBDRigid0Index, const int32 Static1Index)
+	{
+		check((PBDRigid0Index == 0) || (PBDRigid0Index == 1));
+		check((Static1Index == 0) || (Static1Index == 1));
+		check(PBDRigid0Index != Static1Index);
+
+		TPBDRigidParticleHandle<FReal, 3>* PBDRigid0 = ConstraintParticles[ConstraintIndex][PBDRigid0Index]->AsDynamic();
+		TGeometryParticleHandle<FReal, 3>* Static1 = ConstraintParticles[ConstraintIndex][Static1Index];
+		check(PBDRigid0 && Static1 && !Static1->AsDynamic());
+
+		FRotation3 Q0 = PBDRigid0->Q();
+		FVec3 P0 = PBDRigid0->P();
+		const FVec3& XL0 = ConstraintSettings[ConstraintIndex].ConstraintFrames[PBDRigid0Index].GetTranslation();
+		const FRotation3& RL0 = ConstraintSettings[ConstraintIndex].ConstraintFrames[PBDRigid0Index].GetRotation();
+		FRotation3 Q1 = Static1->R();
+		FVec3 P1 = Static1->X();
+		const FVec3& XL1 = ConstraintSettings[ConstraintIndex].ConstraintFrames[Static1Index].GetTranslation();
+		const FRotation3& RL1 = ConstraintSettings[ConstraintIndex].ConstraintFrames[Static1Index].GetRotation();
+		float InvM0;
+		FMatrix33 InvIL0;
+		GetConditionedInverseMass(PBDRigid0, InvM0, InvIL0, Settings.PBDMaxInertiaRatio);
+		const float InvM1 = 0;
+		const FMatrix33 InvIL1 = FMatrix33(0, 0, 0);
+
+		TVector<E6DJointMotionType, 3> AngularMotionTypes = ConstraintSettings[ConstraintIndex].Motion.AngularMotionTypes;
+		bool bTwistDriveEnabled = ConstraintSettings[ConstraintIndex].Motion.bAngularTwistDriveEnabled;
+		bool bSwingDriveEnabled = ConstraintSettings[ConstraintIndex].Motion.bAngularSwingDriveEnabled;
+		bool bAutoStiffness = Settings.bEnableAutoStiffness;
+
+		// NOTE: We put the static body first in the solver - swing axes are calculated relative to this
+		if (Settings.MaxPreIterations > 0)
 		{
-			ConstrainedParticles[0]->AsDynamic()->P() = P0;
+			FPBD6DJointConstraintUtilities::Solve3dConstraint(Settings, P1, Q1, InvM1, InvIL1, XL1, RL1, P0, Q0, InvM0, InvIL0, XL0, RL0, ConstraintSettings[ConstraintIndex].Motion);
 		}
-		if (ConstrainedParticles[1]->AsDynamic())
+		if (Settings.MaxIterations > 0)
 		{
-			ConstrainedParticles[1]->AsDynamic()->P() = P1;
-		}
-	}
-}
+			ConstraintSettings[ConstraintIndex].Motion.bAngularTwistDriveEnabled = false;
+			ConstraintSettings[ConstraintIndex].Motion.bAngularSwingDriveEnabled = false;
 
-template<class T, int d>
-void ApplyJointPositionConstraint(
-	const T Dt,
-	const TPBD6DJointSolverSettings<T, d>& SolverSettings,
-	const TPBD6DJointSettings<T, d>& JointSettings,
-	const int32 Index0,
-	const int32 Index1,
-	TVector<T, d>& P0,
-	TRotation<T, d>& Q0,
-	TVector<T, d>& P1,
-	TRotation<T, d>& Q1,
-	float InvM0,
-	const PMatrix<T, d, d>& InvIL0,
-	float InvM1,
-	const PMatrix<T, d, d>& InvIL1)
-{
-	const TRigidTransform<T, d>& XL0 = JointSettings.ConstraintFrames[Index0];
-	const TRigidTransform<T, d>& XL1 = JointSettings.ConstraintFrames[Index1];
-	const TVector<T, d> X0 = P0 + Q0 * XL0.GetTranslation();
-	const TVector<T, d> X1 = P1 + Q1 * XL1.GetTranslation();
-	PMatrix<T, d, d> InvI0 = Utilities::Multiply(Q0.ToMatrix(), Utilities::Multiply(InvIL0, Q0.ToMatrix().GetTransposed()));
-	PMatrix<T, d, d> InvI1 = Utilities::Multiply(Q1.ToMatrix(), Utilities::Multiply(InvIL1, Q1.ToMatrix().GetTransposed()));
-
-	TVector<T, d> DP0, DP1;
-	TRotation<T, d> DQ0, DQ1;
-	TPBD6DJointConstraintUtilities<T, d>::Calculate3dDelta(SolverSettings, P0, Q0, InvM0, InvI0, P1, Q1, InvM1, InvI1, X0, X1, JointSettings.Motion, DP0, DQ0, DP1, DQ1);
-
-	P0 = P0 + DP0;
-	P1 = P1 + DP1;
-	Q0 = (Q0 + DQ0).GetNormalized();
-	Q1 = (Q1 + DQ1).GetNormalized();
-	Q1.EnforceShortestArcWith(Q0);
-}
+			FPBD6DJointConstraintUtilities::Solve6dConstraint(Settings, Dt, P1, Q1, InvM1, InvIL1, XL1, RL1, P0, Q0, InvM0, InvIL0, XL0, RL0, ConstraintSettings[ConstraintIndex].Motion, ConstraintStates[ConstraintIndex]);
 
-template<class T, int d>
-void ApplyJointTwistConstraint(
-	const T Dt,
-	const TPBD6DJointSolverSettings<T, d>& SolverSettings,
-	const TPBD6DJointSettings<T, d>& JointSettings,
-	const int32 Index0,
-	const int32 Index1,
-	TVector<T, d>& P0,
-	TRotation<T, d>& Q0,
-	TVector<T, d>& P1,
-	TRotation<T, d>& Q1,
-	float InvM0,
-	const PMatrix<T, d, d>& InvIL0,
-	float InvM1,
-	const PMatrix<T, d, d>& InvIL1)
-{
-	const TRigidTransform<T, d>& XL0 = JointSettings.ConstraintFrames[Index0];
-	const TRigidTransform<T, d>& XL1 = JointSettings.ConstraintFrames[Index1];
-	const TRotation<T, d> R0 = Q0 * XL0.GetRotation();
-	const TRotation<T, d> R1 = Q1 * XL1.GetRotation();
-
-	TRotation<T, d> R01 = R0.Inverse() * R1;
-	TRotation<T, d> R01Twist, R01Swing;
-	R01.ToSwingTwist(T6DJointConstants<T, d>::TwistAxis(), R01Swing, R01Twist);
-	R01Swing = R01Swing.GetNormalized();
-	R01Twist = R01Twist.GetNormalized();
-
-	TVector<T, d> TwistAxis01;
-	T TwistAngle;
-	R01Twist.ToAxisAndAngleSafe(TwistAxis01, TwistAngle, T6DJointConstants<T, d>::TwistAxis(), SolverSettings.SwingTwistAngleTolerance);
-	if (TwistAngle > PI)
-	{
-		TwistAngle = TwistAngle - (T)2 * PI;
-	}
-	if (TVector<T, d>::DotProduct(TwistAxis01, T6DJointConstants<T, d>::TwistAxis()) < 0)
-	{
-		TwistAxis01 = -TwistAxis01;
-		TwistAngle = -TwistAngle;
-	}
+			ConstraintSettings[ConstraintIndex].Motion.bAngularTwistDriveEnabled = bTwistDriveEnabled;
+			ConstraintSettings[ConstraintIndex].Motion.bAngularSwingDriveEnabled = bSwingDriveEnabled;
+		}
+		if ((bTwistDriveEnabled || bSwingDriveEnabled) && (Settings.MaxDriveIterations > 0))
+		{
+			Settings.bEnableAutoStiffness = false;
+			ConstraintSettings[ConstraintIndex].Motion.AngularMotionTypes = { E6DJointMotionType::Free, E6DJointMotionType::Free , E6DJointMotionType::Free };
 
-	TVector<T, d> TwistAxis0 = R0 * TwistAxis01;
-	TVector<T, d> TwistAxis1 = R1 * TwistAxis01;
-	T TwistAngleMax = FLT_MAX;
-	if (JointSettings.Motion.AngularMotionTypes[(int32)E6DJointAngularConstraintIndex::Twist] == E6DJointMotionType::Limited)
-	{
-		TwistAngleMax = JointSettings.Motion.AngularLimits[(int32)E6DJointAngularConstraintIndex::Twist];
-	}
-	else if (JointSettings.Motion.AngularMotionTypes[(int32)E6DJointAngularConstraintIndex::Twist] == E6DJointMotionType::Locked)
-	{
-		TwistAngleMax = 0;
-	}
+			FPBD6DJointConstraintUtilities::Solve6dConstraint(Settings, Dt, P1, Q1, InvM1, InvIL1, XL1, RL1, P0, Q0, InvM0, InvIL0, XL0, RL0, ConstraintSettings[ConstraintIndex].Motion, ConstraintStates[ConstraintIndex]);
 
-	T DTwistAngle = 0;
-	if (TwistAngle > TwistAngleMax)
-	{
-		DTwistAngle = TwistAngle - TwistAngleMax;
-	}
-	else if (TwistAngle < -TwistAngleMax)
-	{
-		DTwistAngle = TwistAngle + TwistAngleMax;
+			Settings.bEnableAutoStiffness = bAutoStiffness;
+			ConstraintSettings[ConstraintIndex].Motion.AngularMotionTypes = AngularMotionTypes;
+		}
+
+		PBDRigid0->SetQ(Q0);
+		PBDRigid0->SetP(P0);
 	}
-	T DTwistAngle0 = DTwistAngle * InvM0 / (InvM0 + InvM1);
-	T DTwistAngle1 = -DTwistAngle * InvM1 / (InvM0 + InvM1);
-
-	TVector<T, d> W0 = TwistAxis0 * DTwistAngle0;
-	TVector<T, d> W1 = TwistAxis1 * DTwistAngle1;
-	TRotation<T, d> DQ0 = (TRotation<T, d>::FromElements(W0, (T)0.0) * Q0) * (T)0.5;
-	TRotation<T, d> DQ1 = (TRotation<T, d>::FromElements(W1, (T)0.0) * Q1) * (T)0.5;
-	Q0 = (Q0 + DQ0).GetNormalized();
-	Q1 = (Q1 + DQ1).GetNormalized();
-	Q1.EnforceShortestArcWith(Q0);
-}
 
-template<class T, int d>
-void ApplyJointConeConstraint(
-	const T Dt,
-	const TPBD6DJointSolverSettings<T, d>& SolverSettings,
-	const TPBD6DJointSettings<T, d>& JointSettings,
-	const int32 Index0,
-	const int32 Index1,
-	TVector<T, d>& P0,
-	TRotation<T, d>& Q0,
-	TVector<T, d>& P1,
-	TRotation<T, d>& Q1,
-	float InvM0,
-	const PMatrix<T, d, d>& InvIL0,
-	float InvM1,
-	const PMatrix<T, d, d>& InvIL1)
-{
-	const TRigidTransform<T, d>& XL0 = JointSettings.ConstraintFrames[Index0];
-	const TRigidTransform<T, d>& XL1 = JointSettings.ConstraintFrames[Index1];
-	const TRotation<T, d> R0 = Q0 * XL0.GetRotation();
-	const TRotation<T, d> R1 = Q1 * XL1.GetRotation();
-
-	// Calculate Swing axis for each body
-	TRotation<T, d> R01 = R0.Inverse() * R1;
-	TRotation<T, d> R01Twist, R01Swing;
-	R01.ToSwingTwist(T6DJointConstants<T, d>::TwistAxis(), R01Swing, R01Twist);
-	R01Swing = R01Swing.GetNormalized();
-	R01Twist = R01Twist.GetNormalized();
-
-	TVector<T, d> SwingAxis01;
-	T SwingAngle;
-	R01Swing.ToAxisAndAngleSafe(SwingAxis01, SwingAngle, T6DJointConstants<T, d>::Swing1Axis(), SolverSettings.SwingTwistAngleTolerance);
-	if (SwingAngle > PI)
+	void FPBD6DJointConstraints::ApplyPushOutSingle(const FReal Dt, const int32 ConstraintIndex)
 	{
-		SwingAngle = SwingAngle - (T)2 * PI;
+		// Correct any remaining error by translating
+		if (Settings.bApplyProjection)
+		{
+			const TVector<TGeometryParticleHandle<FReal, 3>*, 2>& ConstrainedParticles = ConstraintParticles[ConstraintIndex];
+			const FPBD6DJointSettings& JointSettings = ConstraintSettings[ConstraintIndex];
+
+			FVec3 P0 = ConstrainedParticles[0]->AsDynamic() ? ConstrainedParticles[0]->AsDynamic()->P() : ConstrainedParticles[0]->X();
+			FVec3 P1 = ConstrainedParticles[1]->AsDynamic() ? ConstrainedParticles[1]->AsDynamic()->P() : ConstrainedParticles[1]->X();
+			FRotation3 Q0 = ConstrainedParticles[0]->AsDynamic() ? ConstrainedParticles[0]->AsDynamic()->Q() : ConstrainedParticles[0]->R();
+			FRotation3 Q1 = ConstrainedParticles[1]->AsDynamic() ? ConstrainedParticles[1]->AsDynamic()->Q() : ConstrainedParticles[1]->R();
+			const FVec3& XL0 = JointSettings.ConstraintFrames[0].GetTranslation();
+			const FRotation3& RL0 = JointSettings.ConstraintFrames[0].GetRotation();
+			const FVec3& XL1 = JointSettings.ConstraintFrames[1].GetTranslation();
+			const FRotation3& RL1 = JointSettings.ConstraintFrames[1].GetRotation();
+			FVec3 X0 = P0 + Q0 * XL0;
+			FVec3 X1 = P1 + Q1 * XL1;
+			const FReal InvM0 = ConstrainedParticles[0]->AsDynamic()? ConstrainedParticles[0]->AsDynamic()->InvM() : (FReal)0;
+			const FReal InvM1 = ConstrainedParticles[1]->AsDynamic() ? ConstrainedParticles[1]->AsDynamic()->InvM() : (FReal)0;
+
+
+			const FVec3 DeltaProj = (X1 - X0) / (InvM0 + InvM1);
+			P0 += InvM0 * DeltaProj;
+			P1 -= InvM1 * DeltaProj;
+
+			if (ConstrainedParticles[0]->AsDynamic())
+			{
+				ConstrainedParticles[0]->AsDynamic()->P() = P0;
+			}
+			if (ConstrainedParticles[1]->AsDynamic())
+			{
+				ConstrainedParticles[1]->AsDynamic()->P() = P1;
+			}
+		}
 	}
 
-	TVector<T, d> SwingAxis0 = R0 * SwingAxis01;
-	TVector<T, d> SwingAxis1 = SwingAxis0;
-
-	// Calculate swing limit for the current swing axis
-	T SwingAngleMax = FLT_MAX;
-	T Swing1Limit = JointSettings.Motion.AngularLimits[(int32)E6DJointAngularConstraintIndex::Swing1];
-	T Swing2Limit = JointSettings.Motion.AngularLimits[(int32)E6DJointAngularConstraintIndex::Swing2];
-
-	// Circular swing limit
-	SwingAngleMax = Swing1Limit;
+	void ApplyJointPositionConstraint(
+		const FReal Dt,
+		const FPBD6DJointSolverSettings& SolverSettings,
+		const FPBD6DJointSettings& JointSettings,
+		const int32 Index0,
+		const int32 Index1,
+		FVec3& P0,
+		FRotation3& Q0,
+		FVec3& P1,
+		FRotation3& Q1,
+		float InvM0,
+		const FMatrix33& InvIL0,
+		float InvM1,
+		const FMatrix33& InvIL1)
+	{
+		const FRigidTransform3& XL0 = JointSettings.ConstraintFrames[Index0];
+		const FRigidTransform3& XL1 = JointSettings.ConstraintFrames[Index1];
+		const FVec3 X0 = P0 + Q0 * XL0.GetTranslation();
+		const FVec3 X1 = P1 + Q1 * XL1.GetTranslation();
+		FMatrix33 InvI0 = Utilities::Multiply(Q0.ToMatrix(), Utilities::Multiply(InvIL0, Q0.ToMatrix().GetTransposed()));
+		FMatrix33 InvI1 = Utilities::Multiply(Q1.ToMatrix(), Utilities::Multiply(InvIL1, Q1.ToMatrix().GetTransposed()));
+
+		FVec3 DP0, DP1;
+		FRotation3 DQ0, DQ1;
+		FPBD6DJointConstraintUtilities::Calculate3dDelta(SolverSettings, P0, Q0, InvM0, InvI0, P1, Q1, InvM1, InvI1, X0, X1, JointSettings.Motion, DP0, DQ0, DP1, DQ1);
 
-	// Elliptical swing limit
-	if (!FMath::IsNearlyEqual(Swing1Limit, Swing2Limit, KINDA_SMALL_NUMBER))
-	{
-		SwingAngleMax = (T)0.5 * (Swing1Limit + Swing2Limit);
-
-		// Map swing axis to ellipse and calculate limit for this swing axis
-		//T DotSwing1 = FMath::Abs(TVector<T, d>::DotProduct(SwingAxis01, T6DJointConstants<T, d>::Swing1Axis()));
-		//T DotSwing2 = FMath::Abs(TVector<T, d>::DotProduct(SwingAxis01, T6DJointConstants<T, d>::Swing2Axis()));
-		//SwingAngleMax = FMath::Sqrt(Swing1Limit * DotSwing1 * Swing1Limit * DotSwing1 + Swing2Limit * DotSwing2 * Swing2Limit * DotSwing2);
+		P0 = P0 + DP0;
+		P1 = P1 + DP1;
+		Q0 = (Q0 + DQ0).GetNormalized();
+		Q1 = (Q1 + DQ1).GetNormalized();
+		Q1.EnforceShortestArcWith(Q0);
 	}
 
-	// Calculate swing error we need to correct
-	T DSwingAngle = 0;
-	if (SwingAngle > SwingAngleMax)
-	{
-		DSwingAngle = SwingAngle - SwingAngleMax;
-	}
-	else if (SwingAngle < -SwingAngleMax)
-	{
-		DSwingAngle = SwingAngle + SwingAngleMax;
-	}
-	T DSwingAngle0 = DSwingAngle * InvM0 / (InvM0 + InvM1);
-	T DSwingAngle1 = -DSwingAngle * InvM1 / (InvM0 + InvM1);
-
-	// Apply swing correction
-	TVector<T, d> W0 = SwingAxis0 * DSwingAngle0;
-	TVector<T, d> W1 = SwingAxis1 * DSwingAngle1;
-	TRotation<T, d> DQ0 = (TRotation<T, d>::FromElements(W0, (T)0.0) * Q0) * (T)0.5;
-	TRotation<T, d> DQ1 = (TRotation<T, d>::FromElements(W1, (T)0.0) * Q1) * (T)0.5;
-	Q0 = (Q0 + DQ0).GetNormalized();
-	Q1 = (Q1 + DQ1).GetNormalized();
-	Q1.EnforceShortestArcWith(Q0);
-}
+	void ApplyJointTwistConstraint(
+		const FReal Dt,
+		const FPBD6DJointSolverSettings& SolverSettings,
+		const FPBD6DJointSettings& JointSettings,
+		const int32 Index0,
+		const int32 Index1,
+		FVec3& P0,
+		FRotation3& Q0,
+		FVec3& P1,
+		FRotation3& Q1,
+		FReal InvM0,
+		const FMatrix33& InvIL0,
+		FReal InvM1,
+		const FMatrix33& InvIL1)
+	{
+		const FRigidTransform3& XL0 = JointSettings.ConstraintFrames[Index0];
+		const FRigidTransform3& XL1 = JointSettings.ConstraintFrames[Index1];
+		const FRotation3 R0 = Q0 * XL0.GetRotation();
+		const FRotation3 R1 = Q1 * XL1.GetRotation();
+
+		FRotation3 R01 = R0.Inverse() * R1;
+		FRotation3 R01Twist, R01Swing;
+		R01.ToSwingTwist(F6DJointConstants::TwistAxis(), R01Swing, R01Twist);
+		R01Swing = R01Swing.GetNormalized();
+		R01Twist = R01Twist.GetNormalized();
+
+		FVec3 TwistAxis01;
+		FReal TwistAngle;
+		R01Twist.ToAxisAndAngleSafe(TwistAxis01, TwistAngle, F6DJointConstants::TwistAxis(), SolverSettings.SwingTwistAngleTolerance);
+		if (TwistAngle > PI)
+		{
+			TwistAngle = TwistAngle - (FReal)2 * PI;
+		}
+		if (FVec3::DotProduct(TwistAxis01, F6DJointConstants::TwistAxis()) < 0)
+		{
+			TwistAxis01 = -TwistAxis01;
+			TwistAngle = -TwistAngle;
+		}
 
-template<class T, int d>
-void ApplyJointSwingConstraint(
-	const T Dt,
-	const TPBD6DJointSolverSettings<T, d>& SolverSettings,
-	const TPBD6DJointSettings<T, d>& JointSettings,
-	const int32 Index0,
-	const int32 Index1,
-	const E6DJointAngularConstraintIndex SwingConstraint,
-	TVector<T, d>& P0,
-	TRotation<T, d>& Q0,
-	TVector<T, d>& P1,
-	TRotation<T, d>& Q1,
-	float InvM0,
-	const PMatrix<T, d, d>& InvIL0,
-	float InvM1,
-	const PMatrix<T, d, d>& InvIL1)
-{
-	const TRigidTransform<T, d>& XL0 = JointSettings.ConstraintFrames[Index0];
-	const TRigidTransform<T, d>& XL1 = JointSettings.ConstraintFrames[Index1];
-	const TRotation<T, d> R0 = Q0 * XL0.GetRotation();
-	const TRotation<T, d> R1 = Q1 * XL1.GetRotation();
-
-	const PMatrix<T, d, d> Axes0 = R0.ToMatrix();
-	const PMatrix<T, d, d> Axes1 = R1.ToMatrix();
-	const TVector<T, d> SwingCross = TVector<T, d>::CrossProduct(Axes0.GetAxis((int32)SwingConstraint), Axes1.GetAxis((int32)SwingConstraint));
-	const T SwingCrossLen = SwingCross.Size();
-	if (SwingCrossLen > KINDA_SMALL_NUMBER)
-	{
-		const TVector<T, d> SwingAxis = SwingCross / SwingCrossLen;
-		TVector<T, d> SwingAxis0 = SwingAxis;
-		TVector<T, d> SwingAxis1 = SwingAxis;
+		FVec3 TwistAxis0 = R0 * TwistAxis01;
+		FVec3 TwistAxis1 = R1 * TwistAxis01;
+		FReal TwistAngleMax = FLT_MAX;
+		if (JointSettings.Motion.AngularMotionTypes[(int32)E6DJointAngularConstraintIndex::Twist] == E6DJointMotionType::Limited)
+		{
+			TwistAngleMax = JointSettings.Motion.AngularLimits[(int32)E6DJointAngularConstraintIndex::Twist];
+		}
+		else if (JointSettings.Motion.AngularMotionTypes[(int32)E6DJointAngularConstraintIndex::Twist] == E6DJointMotionType::Locked)
+		{
+			TwistAngleMax = 0;
+		}
 
-		T SwingAngle = FMath::Asin(FMath::Clamp(SwingCrossLen, (T)0, (T)1));
-		const T SwingDot = TVector<T, d>::DotProduct(Axes0.GetAxis((int32)SwingConstraint), Axes1.GetAxis((int32)SwingConstraint));
-		if (SwingDot < (T)0)
+		FReal DTwistAngle = 0;
+		if (TwistAngle > TwistAngleMax)
+		{
+			DTwistAngle = TwistAngle - TwistAngleMax;
+		}
+		else if (TwistAngle < -TwistAngleMax)
 		{
-			SwingAngle = (T)PI - SwingAngle;
+			DTwistAngle = TwistAngle + TwistAngleMax;
 		}
+		FReal DTwistAngle0 = DTwistAngle * InvM0 / (InvM0 + InvM1);
+		FReal DTwistAngle1 = -DTwistAngle * InvM1 / (InvM0 + InvM1);
 
-		T SwingAngleMax = FLT_MAX;
-		if (JointSettings.Motion.AngularMotionTypes[(int32)SwingConstraint] == E6DJointMotionType::Limited)
+		FVec3 W0 = TwistAxis0 * DTwistAngle0;
+		FVec3 W1 = TwistAxis1 * DTwistAngle1;
+		FRotation3 DQ0 = (FRotation3::FromElements(W0, (FReal)0.0) * Q0) * (FReal)0.5;
+		FRotation3 DQ1 = (FRotation3::FromElements(W1, (FReal)0.0) * Q1) * (FReal)0.5;
+		Q0 = (Q0 + DQ0).GetNormalized();
+		Q1 = (Q1 + DQ1).GetNormalized();
+		Q1.EnforceShortestArcWith(Q0);
+	}
+
+	void ApplyJointConeConstraint(
+		const FReal Dt,
+		const FPBD6DJointSolverSettings& SolverSettings,
+		const FPBD6DJointSettings& JointSettings,
+		const int32 Index0,
+		const int32 Index1,
+		FVec3& P0,
+		FRotation3& Q0,
+		FVec3& P1,
+		FRotation3& Q1,
+		FReal InvM0,
+		const FMatrix33& InvIL0,
+		FReal InvM1,
+		const FMatrix33& InvIL1)
+	{
+		const FRigidTransform3& XL0 = JointSettings.ConstraintFrames[Index0];
+		const FRigidTransform3& XL1 = JointSettings.ConstraintFrames[Index1];
+		const FRotation3 R0 = Q0 * XL0.GetRotation();
+		const FRotation3 R1 = Q1 * XL1.GetRotation();
+
+		// Calculate Swing axis for each body
+		FRotation3 R01 = R0.Inverse() * R1;
+		FRotation3 R01Twist, R01Swing;
+		R01.ToSwingTwist(F6DJointConstants::TwistAxis(), R01Swing, R01Twist);
+		R01Swing = R01Swing.GetNormalized();
+		R01Twist = R01Twist.GetNormalized();
+
+		FVec3 SwingAxis01;
+		FReal SwingAngle;
+		R01Swing.ToAxisAndAngleSafe(SwingAxis01, SwingAngle, F6DJointConstants::Swing1Axis(), SolverSettings.SwingTwistAngleTolerance);
+		if (SwingAngle > PI)
 		{
-			T Swing1Limit = JointSettings.Motion.AngularLimits[(int32)SwingConstraint];
-			SwingAngleMax = Swing1Limit;
+			SwingAngle = SwingAngle - (FReal)2 * PI;
 		}
-		else if (JointSettings.Motion.AngularMotionTypes[(int32)SwingConstraint] == E6DJointMotionType::Locked)
+
+		FVec3 SwingAxis0 = R0 * SwingAxis01;
+		FVec3 SwingAxis1 = SwingAxis0;
+
+		// Calculate swing limit for the current swing axis
+		FReal SwingAngleMax = FLT_MAX;
+		FReal Swing1Limit = JointSettings.Motion.AngularLimits[(int32)E6DJointAngularConstraintIndex::Swing1];
+		FReal Swing2Limit = JointSettings.Motion.AngularLimits[(int32)E6DJointAngularConstraintIndex::Swing2];
+
+		// Circular swing limit
+		SwingAngleMax = Swing1Limit;
+
+		// Elliptical swing limit
+		if (!FMath::IsNearlyEqual(Swing1Limit, Swing2Limit, KINDA_SMALL_NUMBER))
 		{
-			SwingAngleMax = 0;
+			SwingAngleMax = (FReal)0.5 * (Swing1Limit + Swing2Limit);
+
+			// Map swing axis to ellipse and calculate limit for this swing axis
+			//T DotSwing1 = FMath::Abs(FVec3::DotProduct(SwingAxis01, F6DJointConstants::Swing1Axis()));
+			//T DotSwing2 = FMath::Abs(FVec3::DotProduct(SwingAxis01, F6DJointConstants::Swing2Axis()));
+			//SwingAngleMax = FMath::Sqrt(Swing1Limit * DotSwing1 * Swing1Limit * DotSwing1 + Swing2Limit * DotSwing2 * Swing2Limit * DotSwing2);
 		}
 
 		// Calculate swing error we need to correct
-		T DSwingAngle = 0;
+		FReal DSwingAngle = 0;
 		if (SwingAngle > SwingAngleMax)
 		{
 			DSwingAngle = SwingAngle - SwingAngleMax;
@@ -1935,310 +1830,370 @@ void ApplyJointSwingConstraint(
 		{
 			DSwingAngle = SwingAngle + SwingAngleMax;
 		}
-		T DSwingAngle0 = DSwingAngle * InvM0 / (InvM0 + InvM1);
-		T DSwingAngle1 = -DSwingAngle * InvM1 / (InvM0 + InvM1);
+		FReal DSwingAngle0 = DSwingAngle * InvM0 / (InvM0 + InvM1);
+		FReal DSwingAngle1 = -DSwingAngle * InvM1 / (InvM0 + InvM1);
 
 		// Apply swing correction
-		TVector<T, d> W0 = SwingAxis0 * DSwingAngle0;
-		TVector<T, d> W1 = SwingAxis1 * DSwingAngle1;
-		TRotation<T, d> DQ0 = (TRotation<T, d>::FromElements(W0, (T)0.0) * Q0) * (T)0.5;
-		TRotation<T, d> DQ1 = (TRotation<T, d>::FromElements(W1, (T)0.0) * Q1) * (T)0.5;
+		FVec3 W0 = SwingAxis0 * DSwingAngle0;
+		FVec3 W1 = SwingAxis1 * DSwingAngle1;
+		FRotation3 DQ0 = (FRotation3::FromElements(W0, (FReal)0.0) * Q0) * (FReal)0.5;
+		FRotation3 DQ1 = (FRotation3::FromElements(W1, (FReal)0.0) * Q1) * (FReal)0.5;
 		Q0 = (Q0 + DQ0).GetNormalized();
 		Q1 = (Q1 + DQ1).GetNormalized();
 		Q1.EnforceShortestArcWith(Q0);
 	}
-}
-
-template<class T, int d>
-void ApplyJointTwistDrive(
-	const T Dt,
-	const TPBD6DJointSolverSettings<T, d>& SolverSettings,
-	const TPBD6DJointSettings<T, d>& JointSettings,
-	const int32 Index0,
-	const int32 Index1,
-	TVector<T, d>& P0,
-	TRotation<T, d>& Q0,
-	TVector<T, d>& P1,
-	TRotation<T, d>& Q1,
-	float InvM0,
-	const PMatrix<T, d, d>& InvIL0,
-	float InvM1,
-	const PMatrix<T, d, d>& InvIL1)
-{
-	const TRigidTransform<T, d>& XL0 = JointSettings.ConstraintFrames[Index0];
-	const TRigidTransform<T, d>& XL1 = JointSettings.ConstraintFrames[Index1];
-	const TRotation<T, d> R0 = Q0 * XL0.GetRotation();
-	const TRotation<T, d> R1 = Q1 * XL1.GetRotation();
-
-	TRotation<T, d> R01 = R0.Inverse() * R1;
-	TRotation<T, d> R01Twist, R01Swing;
-	R01.ToSwingTwist(T6DJointConstants<T, d>::TwistAxis(), R01Swing, R01Twist);
-	R01Swing = R01Swing.GetNormalized();
-	R01Twist = R01Twist.GetNormalized();
-
-	TVector<T, d> TwistAxis01;
-	T TwistAngle;
-	R01Twist.ToAxisAndAngleSafe(TwistAxis01, TwistAngle, T6DJointConstants<T, d>::TwistAxis(), SolverSettings.SwingTwistAngleTolerance);
-	if (TwistAngle > PI)
-	{
-		TwistAngle = TwistAngle - (T)2 * PI;
-	}
-	if (TVector<T, d>::DotProduct(TwistAxis01, T6DJointConstants<T, d>::TwistAxis()) < 0)
-	{
-		TwistAxis01 = -TwistAxis01;
-		TwistAngle = -TwistAngle;
-	}
 
-	TVector<T, d> TwistAxis0 = R0 * TwistAxis01;
-	TVector<T, d> TwistAxis1 = R1 * TwistAxis01;
-	T TwistAngleTarget = JointSettings.Motion.AngularDriveTargetAngles[(int32)E6DJointAngularConstraintIndex::Twist];
-	T Stiffness = (SolverSettings.PBDDriveStiffness > 0) ? SolverSettings.PBDDriveStiffness : JointSettings.Motion.AngularDriveStiffness;
-	T DriveStiffness = FMath::Clamp(Stiffness, (T)0, (T)1);
-	T DTwistAngle = TwistAngle - TwistAngleTarget;
-	T DTwistAngle0 = DriveStiffness * DTwistAngle * InvM0 / (InvM0 + InvM1);
-	T DTwistAngle1 = -DriveStiffness * DTwistAngle * InvM1 / (InvM0 + InvM1);
-
-	TVector<T, d> W0 = TwistAxis0 * DTwistAngle0;
-	TVector<T, d> W1 = TwistAxis1 * DTwistAngle1;
-	TRotation<T, d> DQ0 = (TRotation<T, d>::FromElements(W0, (T)0.0) * Q0) * (T)0.5;
-	TRotation<T, d> DQ1 = (TRotation<T, d>::FromElements(W1, (T)0.0) * Q1) * (T)0.5;
-	Q0 = (Q0 + DQ0).GetNormalized();
-	Q1 = (Q1 + DQ1).GetNormalized();
-	Q1.EnforceShortestArcWith(Q0);
-}
+	void ApplyJointSwingConstraint(
+		const FReal Dt,
+		const FPBD6DJointSolverSettings& SolverSettings,
+		const FPBD6DJointSettings& JointSettings,
+		const int32 Index0,
+		const int32 Index1,
+		const E6DJointAngularConstraintIndex SwingConstraint,
+		FVec3& P0,
+		FRotation3& Q0,
+		FVec3& P1,
+		FRotation3& Q1,
+		FReal InvM0,
+		const FMatrix33& InvIL0,
+		FReal InvM1,
+		const FMatrix33& InvIL1)
+	{
+		const FRigidTransform3& XL0 = JointSettings.ConstraintFrames[Index0];
+		const FRigidTransform3& XL1 = JointSettings.ConstraintFrames[Index1];
+		const FRotation3 R0 = Q0 * XL0.GetRotation();
+		const FRotation3 R1 = Q1 * XL1.GetRotation();
+
+		const FMatrix33 Axes0 = R0.ToMatrix();
+		const FMatrix33 Axes1 = R1.ToMatrix();
+		const FVec3 SwingCross = FVec3::CrossProduct(Axes0.GetAxis((int32)SwingConstraint), Axes1.GetAxis((int32)SwingConstraint));
+		const FReal SwingCrossLen = SwingCross.Size();
+		if (SwingCrossLen > KINDA_SMALL_NUMBER)
+		{
+			const FVec3 SwingAxis = SwingCross / SwingCrossLen;
+			FVec3 SwingAxis0 = SwingAxis;
+			FVec3 SwingAxis1 = SwingAxis;
 
-template<class T, int d>
-void ApplyJointConeDrive(
-	const T Dt,
-	const TPBD6DJointSolverSettings<T, d>& SolverSettings,
-	const TPBD6DJointSettings<T, d>& JointSettings,
-	const int32 Index0,
-	const int32 Index1,
-	TVector<T, d>& P0,
-	TRotation<T, d>& Q0,
-	TVector<T, d>& P1,
-	TRotation<T, d>& Q1,
-	float InvM0,
-	const PMatrix<T, d, d>& InvIL0,
-	float InvM1,
-	const PMatrix<T, d, d>& InvIL1)
-{
-	const TRigidTransform<T, d>& XL0 = JointSettings.ConstraintFrames[Index0];
-	const TRigidTransform<T, d>& XL1 = JointSettings.ConstraintFrames[Index1];
-	const TRotation<T, d> R0 = Q0 * XL0.GetRotation();
-	const TRotation<T, d> R1 = Q1 * XL1.GetRotation();
-
-	// Calculate Swing axis for each body
-	TRotation<T, d> R01 = R0.Inverse() * R1;
-	TRotation<T, d> R01Twist, R01Swing;
-	R01.ToSwingTwist(T6DJointConstants<T, d>::TwistAxis(), R01Swing, R01Twist);
-	R01Swing = R01Swing.GetNormalized();
-	R01Twist = R01Twist.GetNormalized();
-
-	TVector<T, d> SwingAxis01;
-	T SwingAngle;
-	R01Swing.ToAxisAndAngleSafe(SwingAxis01, SwingAngle, T6DJointConstants<T, d>::Swing1Axis(), SolverSettings.SwingTwistAngleTolerance);
-	if (SwingAngle > PI)
-	{
-		SwingAngle = SwingAngle - (T)2 * PI;
-	}
+			FReal SwingAngle = FMath::Asin(FMath::Clamp(SwingCrossLen, (FReal)0, (FReal)1));
+			const FReal SwingDot = FVec3::DotProduct(Axes0.GetAxis((int32)SwingConstraint), Axes1.GetAxis((int32)SwingConstraint));
+			if (SwingDot < (FReal)0)
+			{
+				SwingAngle = (FReal)PI - SwingAngle;
+			}
 
-	TVector<T, d> SwingAxis0 = R0 * SwingAxis01;
-	TVector<T, d> SwingAxis1 = SwingAxis0;
-
-	// Circular swing target (max of Swing1, Swing2 targets)
-	T Swing1Target = JointSettings.Motion.AngularDriveTargetAngles[(int32)E6DJointAngularConstraintIndex::Swing1];
-	T Swing2Target = JointSettings.Motion.AngularDriveTargetAngles[(int32)E6DJointAngularConstraintIndex::Swing2];
-	T SwingAngleTarget = FMath::Max(Swing1Target, Swing2Target);
-
-	T Stiffness = (SolverSettings.PBDDriveStiffness > 0) ? SolverSettings.PBDDriveStiffness : JointSettings.Motion.AngularDriveStiffness;
-	T DriveStiffness = FMath::Clamp(Stiffness, (T)0, (T)1);
-	T DSwingAngle = SwingAngle - SwingAngleTarget;
-	T DSwingAngle0 = DriveStiffness * DSwingAngle * InvM0 / (InvM0 + InvM1);
-	T DSwingAngle1 = -DriveStiffness * DSwingAngle * InvM1 / (InvM0 + InvM1);
-
-	// Apply swing correction
-	TVector<T, d> W0 = SwingAxis0 * DSwingAngle0;
-	TVector<T, d> W1 = SwingAxis1 * DSwingAngle1;
-	TRotation<T, d> DQ0 = (TRotation<T, d>::FromElements(W0, (T)0.0) * Q0) * (T)0.5;
-	TRotation<T, d> DQ1 = (TRotation<T, d>::FromElements(W1, (T)0.0) * Q1) * (T)0.5;
-	Q0 = (Q0 + DQ0).GetNormalized();
-	Q1 = (Q1 + DQ1).GetNormalized();
-	Q1.EnforceShortestArcWith(Q0);
-}
+			FReal SwingAngleMax = FLT_MAX;
+			if (JointSettings.Motion.AngularMotionTypes[(int32)SwingConstraint] == E6DJointMotionType::Limited)
+			{
+				FReal Swing1Limit = JointSettings.Motion.AngularLimits[(int32)SwingConstraint];
+				SwingAngleMax = Swing1Limit;
+			}
+			else if (JointSettings.Motion.AngularMotionTypes[(int32)SwingConstraint] == E6DJointMotionType::Locked)
+			{
+				SwingAngleMax = 0;
+			}
 
-template<class T, int d>
-void ApplyJointSLerpDrive(
-	const T Dt,
-	const TPBD6DJointSolverSettings<T, d>& SolverSettings,
-	const TPBD6DJointSettings<T, d>& JointSettings,
-	const int32 Index0,
-	const int32 Index1,
-	TVector<T, d>& P0,
-	TRotation<T, d>& Q0,
-	TVector<T, d>& P1,
-	TRotation<T, d>& Q1,
-	float InvM0,
-	const PMatrix<T, d, d>& InvIL0,
-	float InvM1,
-	const PMatrix<T, d, d>& InvIL1)
-{
-	const TRigidTransform<T, d>& XL0 = JointSettings.ConstraintFrames[Index0];
-	const TRigidTransform<T, d>& XL1 = JointSettings.ConstraintFrames[Index1];
-	const TRotation<T, d> R0 = Q0 * XL0.GetRotation();
-	const TRotation<T, d> R1 = Q1 * XL1.GetRotation();
-
-	const TRotation<T, d> TargetR1 = R0 * JointSettings.Motion.AngularDriveTarget;
-	const TRotation<T, d> DR1 = TargetR1 * R1.Inverse();
-	const TRotation<T, d> TargetQ0 = DR1.Inverse() * Q0;
-	const TRotation<T, d> TargetQ1 = DR1 * Q1;
-
-	T Stiffness = (SolverSettings.PBDDriveStiffness > 0) ? SolverSettings.PBDDriveStiffness : JointSettings.Motion.AngularDriveStiffness;
-	T DriveStiffness = FMath::Clamp(Stiffness, (T)0, (T)1);
-	const T F0 = DriveStiffness * InvM0 / (InvM0 + InvM1);
-	const T F1 = DriveStiffness * InvM1 / (InvM0 + InvM1);
-
-	Q0 = TRotation<T, d>::Slerp(Q0, TargetQ0, F0);
-	Q1 = TRotation<T, d>::Slerp(Q1, TargetQ1, F1);
-	Q1.EnforceShortestArcWith(Q0);
-}
+			// Calculate swing error we need to correct
+			FReal DSwingAngle = 0;
+			if (SwingAngle > SwingAngleMax)
+			{
+				DSwingAngle = SwingAngle - SwingAngleMax;
+			}
+			else if (SwingAngle < -SwingAngleMax)
+			{
+				DSwingAngle = SwingAngle + SwingAngleMax;
+			}
+			FReal DSwingAngle0 = DSwingAngle * InvM0 / (InvM0 + InvM1);
+			FReal DSwingAngle1 = -DSwingAngle * InvM1 / (InvM0 + InvM1);
+
+			// Apply swing correction
+			FVec3 W0 = SwingAxis0 * DSwingAngle0;
+			FVec3 W1 = SwingAxis1 * DSwingAngle1;
+			FRotation3 DQ0 = (FRotation3::FromElements(W0, (FReal)0.0) * Q0) * (FReal)0.5;
+			FRotation3 DQ1 = (FRotation3::FromElements(W1, (FReal)0.0) * Q1) * (FReal)0.5;
+			Q0 = (Q0 + DQ0).GetNormalized();
+			Q1 = (Q1 + DQ1).GetNormalized();
+			Q1.EnforceShortestArcWith(Q0);
+		}
+	}
 
-template<class T, int d>
-void TPBD6DJointConstraints<T, d>::ApplySingleFast(const T Dt, const int32 ConstraintIndex, const int32 It, const int32 NumIts)
-{
-	const TVector<TGeometryParticleHandle<T, d>*, 2>& Constraint = ConstraintParticles[ConstraintIndex];
-	UE_LOG(LogChaos6DJoint, Verbose, TEXT("6DoF FastSolve Constraint %d %s %s (dt = %f; it = %d / %d)"), ConstraintIndex, *Constraint[0]->ToString(), *Constraint[1]->ToString(), Dt, It, NumIts);
-
-	const FJointSettings& JointSettings = ConstraintSettings[ConstraintIndex];
-
-	// Switch particles - internally we assume the first body is the parent (i.e., the space in which constraint limits are specified)
-	const int32 Index0 = 1;
-	const int32 Index1 = 0;
-	TGenericParticleHandle<T, d> Particle0 = TGenericParticleHandle<T, d>(ConstraintParticles[ConstraintIndex][Index0]);
-	TGenericParticleHandle<T, d> Particle1 = TGenericParticleHandle<T, d>(ConstraintParticles[ConstraintIndex][Index1]);
-	TPBDRigidParticleHandle<T, d>* Rigid0 = ConstraintParticles[ConstraintIndex][Index0]->AsDynamic();
-	TPBDRigidParticleHandle<T, d>* Rigid1 = ConstraintParticles[ConstraintIndex][Index1]->AsDynamic();
-
-	TVector<T, d> P0 = Particle0->P();
-	TRotation<T, d> Q0 = Particle0->Q();
-	TVector<T, d> P1 = Particle1->P();
-	TRotation<T, d> Q1 = Particle1->Q();
-	float InvM0 = Particle0->InvM();
-	float InvM1 = Particle1->InvM();
-	PMatrix<T, d, d> InvIL0 = Particle0->InvI();
-	PMatrix<T, d, d> InvIL1 = Particle1->InvI();
-
-	Q1.EnforceShortestArcWith(Q0);
-
-	// Adjust mass for stability
-	if (Rigid0 && Rigid1)
-	{
-		if (ConstraintStates[ConstraintIndex].ParticleLevels[Index0] < ConstraintStates[ConstraintIndex].ParticleLevels[Index1])
+	void ApplyJointTwistDrive(
+		const FReal Dt,
+		const FPBD6DJointSolverSettings& SolverSettings,
+		const FPBD6DJointSettings& JointSettings,
+		const int32 Index0,
+		const int32 Index1,
+		FVec3& P0,
+		FRotation3& Q0,
+		FVec3& P1,
+		FRotation3& Q1,
+		FReal InvM0,
+		const FMatrix33& InvIL0,
+		FReal InvM1,
+		const FMatrix33& InvIL1)
+	{
+		const FRigidTransform3& XL0 = JointSettings.ConstraintFrames[Index0];
+		const FRigidTransform3& XL1 = JointSettings.ConstraintFrames[Index1];
+		const FRotation3 R0 = Q0 * XL0.GetRotation();
+		const FRotation3 R1 = Q1 * XL1.GetRotation();
+
+		FRotation3 R01 = R0.Inverse() * R1;
+		FRotation3 R01Twist, R01Swing;
+		R01.ToSwingTwist(F6DJointConstants::TwistAxis(), R01Swing, R01Twist);
+		R01Swing = R01Swing.GetNormalized();
+		R01Twist = R01Twist.GetNormalized();
+
+		FVec3 TwistAxis01;
+		FReal TwistAngle;
+		R01Twist.ToAxisAndAngleSafe(TwistAxis01, TwistAngle, F6DJointConstants::TwistAxis(), SolverSettings.SwingTwistAngleTolerance);
+		if (TwistAngle > PI)
 		{
-			GetConditionedInverseMass(Rigid0, Rigid1, InvM0, InvM1, InvIL0, InvIL1, Settings.PBDMinParentMassRatio, Settings.PBDMaxInertiaRatio);
+			TwistAngle = TwistAngle - (FReal)2 * PI;
 		}
-		else if (ConstraintStates[ConstraintIndex].ParticleLevels[Index0] > ConstraintStates[ConstraintIndex].ParticleLevels[Index1])
+		if (FVec3::DotProduct(TwistAxis01, F6DJointConstants::TwistAxis()) < 0)
 		{
-			GetConditionedInverseMass(Rigid1, Rigid0, InvM1, InvM0, InvIL1, InvIL0, Settings.PBDMinParentMassRatio, Settings.PBDMaxInertiaRatio);
+			TwistAxis01 = -TwistAxis01;
+			TwistAngle = -TwistAngle;
 		}
-		else
+
+		FVec3 TwistAxis0 = R0 * TwistAxis01;
+		FVec3 TwistAxis1 = R1 * TwistAxis01;
+		FReal TwistAngleTarget = JointSettings.Motion.AngularDriveTargetAngles[(int32)E6DJointAngularConstraintIndex::Twist];
+		FReal Stiffness = (SolverSettings.PBDDriveStiffness > 0) ? SolverSettings.PBDDriveStiffness : JointSettings.Motion.AngularDriveStiffness;
+		FReal DriveStiffness = FMath::Clamp(Stiffness, (FReal)0, (FReal)1);
+		FReal DTwistAngle = TwistAngle - TwistAngleTarget;
+		FReal DTwistAngle0 = DriveStiffness * DTwistAngle * InvM0 / (InvM0 + InvM1);
+		FReal DTwistAngle1 = -DriveStiffness * DTwistAngle * InvM1 / (InvM0 + InvM1);
+
+		FVec3 W0 = TwistAxis0 * DTwistAngle0;
+		FVec3 W1 = TwistAxis1 * DTwistAngle1;
+		FRotation3 DQ0 = (FRotation3::FromElements(W0, (FReal)0.0) * Q0) * (FReal)0.5;
+		FRotation3 DQ1 = (FRotation3::FromElements(W1, (FReal)0.0) * Q1) * (FReal)0.5;
+		Q0 = (Q0 + DQ0).GetNormalized();
+		Q1 = (Q1 + DQ1).GetNormalized();
+		Q1.EnforceShortestArcWith(Q0);
+	}
+
+	void ApplyJointConeDrive(
+		const FReal Dt,
+		const FPBD6DJointSolverSettings& SolverSettings,
+		const FPBD6DJointSettings& JointSettings,
+		const int32 Index0,
+		const int32 Index1,
+		FVec3& P0,
+		FRotation3& Q0,
+		FVec3& P1,
+		FRotation3& Q1,
+		FReal InvM0,
+		const FMatrix33& InvIL0,
+		FReal InvM1,
+		const FMatrix33& InvIL1)
+	{
+		const FRigidTransform3& XL0 = JointSettings.ConstraintFrames[Index0];
+		const FRigidTransform3& XL1 = JointSettings.ConstraintFrames[Index1];
+		const FRotation3 R0 = Q0 * XL0.GetRotation();
+		const FRotation3 R1 = Q1 * XL1.GetRotation();
+
+		// Calculate Swing axis for each body
+		FRotation3 R01 = R0.Inverse() * R1;
+		FRotation3 R01Twist, R01Swing;
+		R01.ToSwingTwist(F6DJointConstants::TwistAxis(), R01Swing, R01Twist);
+		R01Swing = R01Swing.GetNormalized();
+		R01Twist = R01Twist.GetNormalized();
+
+		FVec3 SwingAxis01;
+		FReal SwingAngle;
+		R01Swing.ToAxisAndAngleSafe(SwingAxis01, SwingAngle, F6DJointConstants::Swing1Axis(), SolverSettings.SwingTwistAngleTolerance);
+		if (SwingAngle > PI)
 		{
-			GetConditionedInverseMass(Rigid1, Rigid0, InvM1, InvM0, InvIL1, InvIL0, (T)0, Settings.PBDMaxInertiaRatio);
+			SwingAngle = SwingAngle - (FReal)2 * PI;
 		}
+
+		FVec3 SwingAxis0 = R0 * SwingAxis01;
+		FVec3 SwingAxis1 = SwingAxis0;
+
+		// Circular swing target (max of Swing1, Swing2 targets)
+		FReal Swing1Target = JointSettings.Motion.AngularDriveTargetAngles[(int32)E6DJointAngularConstraintIndex::Swing1];
+		FReal Swing2Target = JointSettings.Motion.AngularDriveTargetAngles[(int32)E6DJointAngularConstraintIndex::Swing2];
+		FReal SwingAngleTarget = FMath::Max(Swing1Target, Swing2Target);
+
+		FReal Stiffness = (SolverSettings.PBDDriveStiffness > 0) ? SolverSettings.PBDDriveStiffness : JointSettings.Motion.AngularDriveStiffness;
+		FReal DriveStiffness = FMath::Clamp(Stiffness, (FReal)0, (FReal)1);
+		FReal DSwingAngle = SwingAngle - SwingAngleTarget;
+		FReal DSwingAngle0 = DriveStiffness * DSwingAngle * InvM0 / (InvM0 + InvM1);
+		FReal DSwingAngle1 = -DriveStiffness * DSwingAngle * InvM1 / (InvM0 + InvM1);
+
+		// Apply swing correction
+		FVec3 W0 = SwingAxis0 * DSwingAngle0;
+		FVec3 W1 = SwingAxis1 * DSwingAngle1;
+		FRotation3 DQ0 = (FRotation3::FromElements(W0, (FReal)0.0) * Q0) * (FReal)0.5;
+		FRotation3 DQ1 = (FRotation3::FromElements(W1, (FReal)0.0) * Q1) * (FReal)0.5;
+		Q0 = (Q0 + DQ0).GetNormalized();
+		Q1 = (Q1 + DQ1).GetNormalized();
+		Q1.EnforceShortestArcWith(Q0);
 	}
 
-	E6DJointMotionType TwistMotion = JointSettings.Motion.AngularMotionTypes[(int32)E6DJointAngularConstraintIndex::Twist];
-	E6DJointMotionType Swing1Motion = JointSettings.Motion.AngularMotionTypes[(int32)E6DJointAngularConstraintIndex::Swing1];
-	E6DJointMotionType Swing2Motion = JointSettings.Motion.AngularMotionTypes[(int32)E6DJointAngularConstraintIndex::Swing2];
+	void ApplyJointSLerpDrive(
+		const FReal Dt,
+		const FPBD6DJointSolverSettings& SolverSettings,
+		const FPBD6DJointSettings& JointSettings,
+		const int32 Index0,
+		const int32 Index1,
+		FVec3& P0,
+		FRotation3& Q0,
+		FVec3& P1,
+		FRotation3& Q1,
+		FReal InvM0,
+		const FMatrix33& InvIL0,
+		FReal InvM1,
+		const FMatrix33& InvIL1)
+	{
+		const FRigidTransform3& XL0 = JointSettings.ConstraintFrames[Index0];
+		const FRigidTransform3& XL1 = JointSettings.ConstraintFrames[Index1];
+		const FRotation3 R0 = Q0 * XL0.GetRotation();
+		const FRotation3 R1 = Q1 * XL1.GetRotation();
+
+		const FRotation3 TargetR1 = R0 * JointSettings.Motion.AngularDriveTarget;
+		const FRotation3 DR1 = TargetR1 * R1.Inverse();
+		const FRotation3 TargetQ0 = DR1.Inverse() * Q0;
+		const FRotation3 TargetQ1 = DR1 * Q1;
+
+		FReal Stiffness = (SolverSettings.PBDDriveStiffness > 0) ? SolverSettings.PBDDriveStiffness : JointSettings.Motion.AngularDriveStiffness;
+		FReal DriveStiffness = FMath::Clamp(Stiffness, (FReal)0, (FReal)1);
+		const FReal F0 = DriveStiffness * InvM0 / (InvM0 + InvM1);
+		const FReal F1 = DriveStiffness * InvM1 / (InvM0 + InvM1);
+
+		Q0 = FRotation3::Slerp(Q0, TargetQ0, F0);
+		Q1 = FRotation3::Slerp(Q1, TargetQ1, F1);
+		Q1.EnforceShortestArcWith(Q0);
+	}
 
-	// Apply angular drives
-	if (Settings.bEnableDrives)
+	void FPBD6DJointConstraints::ApplySingleFast(const FReal Dt, const int32 ConstraintIndex, const int32 It, const int32 NumIts)
 	{
-		bool bTwistLocked = TwistMotion == E6DJointMotionType::Locked;
-		bool bSwing1Locked = Swing1Motion == E6DJointMotionType::Locked;
-		bool bSwing2Locked = Swing2Motion == E6DJointMotionType::Locked;
+		const TVector<TGeometryParticleHandle<FReal, 3>*, 2>& Constraint = ConstraintParticles[ConstraintIndex];
+		UE_LOG(LogChaos6DJoint, Verbose, TEXT("6DoF FastSolve Constraint %d %s %s (dt = %f; it = %d / %d)"), ConstraintIndex, *Constraint[0]->ToString(), *Constraint[1]->ToString(), Dt, It, NumIts);
 
-		// No SLerp drive if we have a locked rotation (it will be grayed out in the editor in this case, but could still have been set before the rotation was locked)
-		if (JointSettings.Motion.bAngularSLerpDriveEnabled && !bTwistLocked && !bSwing1Locked && !bSwing2Locked)
-		{
-			ApplyJointSLerpDrive(Dt, Settings, JointSettings, Index0, Index1, P0, Q0, P1, Q1, InvM0, InvIL0, InvM1, InvIL1);
-		}
+		const FPBD6DJointSettings& JointSettings = ConstraintSettings[ConstraintIndex];
 
-		if (JointSettings.Motion.bAngularTwistDriveEnabled && !bTwistLocked)
-		{
-			ApplyJointTwistDrive(Dt, Settings, JointSettings, Index0, Index1, P0, Q0, P1, Q1, InvM0, InvIL0, InvM1, InvIL1);
-		}
+		// Switch particles - internally we assume the first body is the parent (i.e., the space in which constraint limits are specified)
+		const int32 Index0 = 1;
+		const int32 Index1 = 0;
+		TGenericParticleHandle<FReal, 3> Particle0 = TGenericParticleHandle<FReal, 3>(ConstraintParticles[ConstraintIndex][Index0]);
+		TGenericParticleHandle<FReal, 3> Particle1 = TGenericParticleHandle<FReal, 3>(ConstraintParticles[ConstraintIndex][Index1]);
+		TPBDRigidParticleHandle<FReal, 3>* Rigid0 = ConstraintParticles[ConstraintIndex][Index0]->AsDynamic();
+		TPBDRigidParticleHandle<FReal, 3>* Rigid1 = ConstraintParticles[ConstraintIndex][Index1]->AsDynamic();
 
-		if (JointSettings.Motion.bAngularSwingDriveEnabled && !bSwing1Locked && !bSwing2Locked)
-		{
-			ApplyJointConeDrive(Dt, Settings, JointSettings, Index0, Index1, P0, Q0, P1, Q1, InvM0, InvIL0, InvM1, InvIL1);
-		}
-		else if (JointSettings.Motion.bAngularSwingDriveEnabled && !bSwing1Locked)
-		{
-			//ApplyJointSwingDrive(Dt, Settings, JointSettings, Index0, Index1, E6DJointAngularConstraintIndex::Swing1, P0, Q0, P1, Q1, InvM0, InvIL0, InvM1, InvIL1);
-		}
-		else if (JointSettings.Motion.bAngularSwingDriveEnabled && !bSwing2Locked)
+		FVec3 P0 = Particle0->P();
+		FRotation3 Q0 = Particle0->Q();
+		FVec3 P1 = Particle1->P();
+		FRotation3 Q1 = Particle1->Q();
+		float InvM0 = Particle0->InvM();
+		float InvM1 = Particle1->InvM();
+		FMatrix33 InvIL0 = Particle0->InvI();
+		FMatrix33 InvIL1 = Particle1->InvI();
+
+		Q1.EnforceShortestArcWith(Q0);
+
+		// Adjust mass for stability
+		if (Rigid0 && Rigid1)
 		{
-			//ApplyJointSwingDrive(Dt, Settings, JointSettings, Index0, Index1, E6DJointAngularConstraintIndex::Swing2, P0, Q0, P1, Q1, InvM0, InvIL0, InvM1, InvIL1);
+			if (ConstraintStates[ConstraintIndex].ParticleLevels[Index0] < ConstraintStates[ConstraintIndex].ParticleLevels[Index1])
+			{
+				GetConditionedInverseMass(Rigid0, Rigid1, InvM0, InvM1, InvIL0, InvIL1, Settings.PBDMinParentMassRatio, Settings.PBDMaxInertiaRatio);
+			}
+			else if (ConstraintStates[ConstraintIndex].ParticleLevels[Index0] > ConstraintStates[ConstraintIndex].ParticleLevels[Index1])
+			{
+				GetConditionedInverseMass(Rigid1, Rigid0, InvM1, InvM0, InvIL1, InvIL0, Settings.PBDMinParentMassRatio, Settings.PBDMaxInertiaRatio);
+			}
+			else
+			{
+				GetConditionedInverseMass(Rigid1, Rigid0, InvM1, InvM0, InvIL1, InvIL0, (FReal)0, Settings.PBDMaxInertiaRatio);
+			}
 		}
-	}
 
-	// Apply twist constraint
-	if (Settings.bEnableTwistLimits)
-	{
-		if (TwistMotion != E6DJointMotionType::Free)
+		E6DJointMotionType TwistMotion = JointSettings.Motion.AngularMotionTypes[(int32)E6DJointAngularConstraintIndex::Twist];
+		E6DJointMotionType Swing1Motion = JointSettings.Motion.AngularMotionTypes[(int32)E6DJointAngularConstraintIndex::Swing1];
+		E6DJointMotionType Swing2Motion = JointSettings.Motion.AngularMotionTypes[(int32)E6DJointAngularConstraintIndex::Swing2];
+
+		// Apply angular drives
+		if (Settings.bEnableDrives)
 		{
-			ApplyJointTwistConstraint(Dt, Settings, JointSettings, Index0, Index1, P0, Q0, P1, Q1, InvM0, InvIL0, InvM1, InvIL1);
+			bool bTwistLocked = TwistMotion == E6DJointMotionType::Locked;
+			bool bSwing1Locked = Swing1Motion == E6DJointMotionType::Locked;
+			bool bSwing2Locked = Swing2Motion == E6DJointMotionType::Locked;
+
+			// No SLerp drive if we have a locked rotation (it will be grayed out in the editor in this case, but could still have been set before the rotation was locked)
+			if (JointSettings.Motion.bAngularSLerpDriveEnabled && !bTwistLocked && !bSwing1Locked && !bSwing2Locked)
+			{
+				ApplyJointSLerpDrive(Dt, Settings, JointSettings, Index0, Index1, P0, Q0, P1, Q1, InvM0, InvIL0, InvM1, InvIL1);
+			}
+
+			if (JointSettings.Motion.bAngularTwistDriveEnabled && !bTwistLocked)
+			{
+				ApplyJointTwistDrive(Dt, Settings, JointSettings, Index0, Index1, P0, Q0, P1, Q1, InvM0, InvIL0, InvM1, InvIL1);
+			}
+
+			if (JointSettings.Motion.bAngularSwingDriveEnabled && !bSwing1Locked && !bSwing2Locked)
+			{
+				ApplyJointConeDrive(Dt, Settings, JointSettings, Index0, Index1, P0, Q0, P1, Q1, InvM0, InvIL0, InvM1, InvIL1);
+			}
+			else if (JointSettings.Motion.bAngularSwingDriveEnabled && !bSwing1Locked)
+			{
+				//ApplyJointSwingDrive(Dt, Settings, JointSettings, Index0, Index1, E6DJointAngularConstraintIndex::Swing1, P0, Q0, P1, Q1, InvM0, InvIL0, InvM1, InvIL1);
+			}
+			else if (JointSettings.Motion.bAngularSwingDriveEnabled && !bSwing2Locked)
+			{
+				//ApplyJointSwingDrive(Dt, Settings, JointSettings, Index0, Index1, E6DJointAngularConstraintIndex::Swing2, P0, Q0, P1, Q1, InvM0, InvIL0, InvM1, InvIL1);
+			}
 		}
-	}
 
-	// Apply swing constraints
-	if (Settings.bEnableSwingLimits)
-	{
-		if ((Swing1Motion == E6DJointMotionType::Limited) && (Swing2Motion == E6DJointMotionType::Limited))
+		// Apply twist constraint
+		if (Settings.bEnableTwistLimits)
 		{
-			// Swing Cone
-			ApplyJointConeConstraint(Dt, Settings, JointSettings, Index0, Index1, P0, Q0, P1, Q1, InvM0, InvIL0, InvM1, InvIL1);
+			if (TwistMotion != E6DJointMotionType::Free)
+			{
+				ApplyJointTwistConstraint(Dt, Settings, JointSettings, Index0, Index1, P0, Q0, P1, Q1, InvM0, InvIL0, InvM1, InvIL1);
+			}
 		}
-		else
+
+		// Apply swing constraints
+		if (Settings.bEnableSwingLimits)
 		{
-			if (Swing1Motion != E6DJointMotionType::Free)
+			if ((Swing1Motion == E6DJointMotionType::Limited) && (Swing2Motion == E6DJointMotionType::Limited))
 			{
-				// Swing Arc/Lock
-				ApplyJointSwingConstraint(Dt, Settings, JointSettings, Index0, Index1, E6DJointAngularConstraintIndex::Swing1, P0, Q0, P1, Q1, InvM0, InvIL0, InvM1, InvIL1);
+				// Swing Cone
+				ApplyJointConeConstraint(Dt, Settings, JointSettings, Index0, Index1, P0, Q0, P1, Q1, InvM0, InvIL0, InvM1, InvIL1);
 			}
-			if (Swing2Motion != E6DJointMotionType::Free)
+			else
 			{
-				// Swing Arc/Lock
-				ApplyJointSwingConstraint(Dt, Settings, JointSettings, Index0, Index1, E6DJointAngularConstraintIndex::Swing2, P0, Q0, P1, Q1, InvM0, InvIL0, InvM1, InvIL1);
+				if (Swing1Motion != E6DJointMotionType::Free)
+				{
+					// Swing Arc/Lock
+					ApplyJointSwingConstraint(Dt, Settings, JointSettings, Index0, Index1, E6DJointAngularConstraintIndex::Swing1, P0, Q0, P1, Q1, InvM0, InvIL0, InvM1, InvIL1);
+				}
+				if (Swing2Motion != E6DJointMotionType::Free)
+				{
+					// Swing Arc/Lock
+					ApplyJointSwingConstraint(Dt, Settings, JointSettings, Index0, Index1, E6DJointAngularConstraintIndex::Swing2, P0, Q0, P1, Q1, InvM0, InvIL0, InvM1, InvIL1);
+				}
 			}
 		}
-	}
 
-	// Apply linear constraints
-	{
-		ApplyJointPositionConstraint(Dt, Settings, JointSettings, Index0, Index1, P0, Q0, P1, Q1, InvM0, InvIL0, InvM1, InvIL1);
-	}
+		// Apply linear constraints
+		{
+			ApplyJointPositionConstraint(Dt, Settings, JointSettings, Index0, Index1, P0, Q0, P1, Q1, InvM0, InvIL0, InvM1, InvIL1);
+		}
 
-	// Update the particles
-	if (Rigid0)
-	{
-		Rigid0->SetP(P0);
-		Rigid0->SetQ(Q0);
-	}
-	if (Rigid1)
-	{
-		Rigid1->SetP(P1);
-		Rigid1->SetQ(Q1);
+		// Update the particles
+		if (Rigid0)
+		{
+			Rigid0->SetP(P0);
+			Rigid0->SetQ(Q0);
+		}
+		if (Rigid1)
+		{
+			Rigid1->SetP(P1);
+			Rigid1->SetQ(Q1);
+		}
 	}
-}
 
-namespace Chaos
-{
-	template class TPBD6DJointConstraintUtilities<float, 3>;
-	template class TPBD6DJointSettings<float, 3>;
-	template class TPBD6DJointSolverSettings<float, 3>;
-	template class TPBD6DJointConstraintHandle<float, 3>;
-	template class TContainerConstraintHandle<TPBD6DJointConstraints<float, 3>>;
-	template class TPBD6DJointConstraints<float, 3>;
 }
