Chaos - remove template params from constraint classes.

diff --git a/Engine/Source/Runtime/Experimental/Chaos/Private/Chaos/PBDCollisionConstraint.cpp b/Engine/Source/Runtime/Experimental/Chaos/Private/Chaos/PBDCollisionConstraint.cpp
--- a/Engine/Source/Runtime/Experimental/Chaos/Private/Chaos/PBDCollisionConstraint.cpp
+++ b/Engine/Source/Runtime/Experimental/Chaos/Private/Chaos/PBDCollisionConstraint.cpp
@@ -142,7 +142,7 @@ void TPBDCollisionConstraint<T, d>::Reset(/*const TPBDRigidParticles<T, d>& InPa
 		if (LifespanCounter > Elem->GetTimestamp())
 		{
 			bool bInContact = false;
-			for (FConstraintHandle* Handle : Elem->GetHandles())
+			for (FConstraintContainerHandle* Handle : Elem->GetHandles())
 			{
 				if (Handle->GetContact().Phi <= 0)
 				{
@@ -152,8 +152,8 @@ void TPBDCollisionConstraint<T, d>::Reset(/*const TPBDRigidParticles<T, d>& InPa
 
 			if (!bInContact)
 			{
-				const TArray<FConstraintHandle*> HandleArray = Elem->GetHandles();
-				for (FConstraintHandle* Handle : HandleArray)
+				const TArray<FConstraintContainerHandle*> HandleArray = Elem->GetHandles();
+				for (FConstraintContainerHandle* Handle : HandleArray)
 				{
 					RemoveConstraint(Handle->GetConstraintIndex());
 				}
@@ -461,7 +461,7 @@ void TPBDCollisionConstraint<T, d>::UpdateConstraintsHelper(/*const TPBDRigidPar
 		FConstraintHandleID HandleID = GetConstraintHandleID(*Constraint);
 		if (Handles.Contains(HandleID))
 		{
-			FConstraintHandle* Handle = Handles[HandleID];
+			FConstraintContainerHandle* Handle = Handles[HandleID];
 			int32 Idx = Handle->GetConstraintIndex();
 			Queue.Dequeue(Constraints[Idx]);
 			Constraints[Idx].Lifespan = LifespanCounter;
@@ -782,13 +782,13 @@ void TPBDCollisionConstraint<T, d>::Apply(const T Dt, FRigidBodyContactConstrain
 
 DECLARE_CYCLE_STAT(TEXT("TPBDCollisionConstraint::Apply"), STAT_Apply, STATGROUP_Chaos);
 template<typename T, int d>
-void TPBDCollisionConstraint<T, d>::Apply(const T Dt, const TArray<FConstraintHandle*>& InConstraintHandles, const int32 It, const int32 NumIts)
+void TPBDCollisionConstraint<T, d>::Apply(const T Dt, const TArray<FConstraintContainerHandle*>& InConstraintHandles, const int32 It, const int32 NumIts)
 {
 	SCOPE_CYCLE_COUNTER(STAT_Apply);
 	if (bEnableVelocitySolve)
 	{
 		PhysicsParallelFor(InConstraintHandles.Num(), [&](int32 ConstraintHandleIndex) {
-			FConstraintHandle* ConstraintHandle = InConstraintHandles[ConstraintHandleIndex];
+			FConstraintContainerHandle* ConstraintHandle = InConstraintHandles[ConstraintHandleIndex];
 			check(ConstraintHandle != nullptr);
 			Apply(Dt, Constraints[ConstraintHandle->GetConstraintIndex()]);
 			}, bDisableCollisionParallelFor);
@@ -801,7 +801,7 @@ void TPBDCollisionConstraint<T, d>::Apply(const T Dt, const TArray<FConstraintHa
 }
 
 template<typename T, int d>
-void TPBDCollisionConstraint<T, d>::Disable(const TArray<FConstraintHandle*>& InConstraintHandles)
+void TPBDCollisionConstraint<T, d>::Disable(const TArray<FConstraintContainerHandle*>& InConstraintHandles)
 {
 	for(int32 Idx = InConstraintHandles.Num()-1; Idx>=0; Idx--) 
 	{
@@ -927,7 +927,7 @@ void TPBDCollisionConstraint<T, d>::ApplyPushOut(const T Dt, FRigidBodyContactCo
 
 DECLARE_CYCLE_STAT(TEXT("TPBDCollisionConstraint::ApplyPushOut"), STAT_ApplyPushOut, STATGROUP_Chaos);
 template<typename T, int d>
-bool TPBDCollisionConstraint<T, d>::ApplyPushOut(const T Dt, const TArray<FConstraintHandle*>& InConstraintHandles, const TSet<TGeometryParticleHandle<T,d>*>& IsTemporarilyStatic, int32 Iteration, int32 NumIterations)
+bool TPBDCollisionConstraint<T, d>::ApplyPushOut(const T Dt, const TArray<FConstraintContainerHandle*>& InConstraintHandles, const TSet<TGeometryParticleHandle<T,d>*>& IsTemporarilyStatic, int32 Iteration, int32 NumIterations)
 {
 	SCOPE_CYCLE_COUNTER(STAT_ApplyPushOut);
 
@@ -936,7 +936,7 @@ bool TPBDCollisionConstraint<T, d>::ApplyPushOut(const T Dt, const TArray<FConst
 	if (MPairIterations > 0)
 	{
 		PhysicsParallelFor(InConstraintHandles.Num(), [&](int32 ConstraintHandleIndex) {
-			FConstraintHandle* ConstraintHandle = InConstraintHandles[ConstraintHandleIndex];
+			FConstraintContainerHandle* ConstraintHandle = InConstraintHandles[ConstraintHandleIndex];
 			check(ConstraintHandle != nullptr);
 			ApplyPushOut(Dt, Constraints[ConstraintHandle->GetConstraintIndex()], IsTemporarilyStatic, Iteration, NumIterations, NeedsAnotherIteration);
 			}, bDisableCollisionParallelFor);
@@ -2171,7 +2171,7 @@ void ConstructUnionUnionConstraints(TGeometryParticleHandle<T, d>* Particle0, TG
 }
 
 template<typename T, int d>
-typename void ConstructConstraintsImpl(TGeometryParticleHandle<T, d>* Particle0, TGeometryParticleHandle<T, d>* Particle1, const FImplicitObject* Implicit0, const FImplicitObject* Implicit1, const T Thickness, TArray< TRigidBodyContactConstraint<T, d> >& ConstraintBuffer, TPBDCollisionConstraintHistory<T, d>* History)
+void ConstructConstraintsImpl(TGeometryParticleHandle<T, d>* Particle0, TGeometryParticleHandle<T, d>* Particle1, const FImplicitObject* Implicit0, const FImplicitObject* Implicit1, const T Thickness, TArray< TRigidBodyContactConstraint<T, d> >& ConstraintBuffer, TPBDCollisionConstraintHistory<T, d>* History)
 {
 	// TriangleMesh implicit's are for scene query only.
 	if (Implicit0 && GetInnerType(Implicit0->GetType()) == ImplicitObjectType::TriangleMesh) return;
@@ -2248,7 +2248,7 @@ typename void ConstructConstraintsImpl(TGeometryParticleHandle<T, d>* Particle0,
 }
 
 template<typename T, int d>
-typename void TPBDCollisionConstraint<T, d>::ConstructConstraints(TGeometryParticleHandle<T, d>* Particle0, TGeometryParticleHandle<T, d>* Particle1, const T Thickness, TArray<TRigidBodyContactConstraint<T, d>>& ConstraintBuffer)
+void TPBDCollisionConstraint<T, d>::ConstructConstraints(TGeometryParticleHandle<T, d>* Particle0, TGeometryParticleHandle<T, d>* Particle1, const T Thickness, TArray<TRigidBodyContactConstraint<T, d>>& ConstraintBuffer)
 {
 	ensure(Particle0 && Particle1);
 
@@ -2381,21 +2381,21 @@ DECLARE_CYCLE_STAT(TEXT("TPBDCollisionConstraint::UpdateConstraint"), STAT_Updat
 
 template<typename T, int d>
 template<ECollisionUpdateType UpdateType>
-void TPBDCollisionConstraint<T, d>::UpdateConstraint(const T Thickness, FRigidBodyContactConstraint* Constraints, int32 NumConstraints)
+void TPBDCollisionConstraint<T, d>::UpdateConstraint(const T Thickness, FRigidBodyContactConstraint* InConstraints, int32 NumConstraints)
 {
 	SCOPE_CYCLE_COUNTER(STAT_UpdateConstraint);
 
 	// Reset Phi and validate that the constraints operate on the same bodies.
 	check(NumConstraints >= 1);
-	check(Constraints != nullptr);
+	check(InConstraints != nullptr);
 	for (int32 i = 0; i < NumConstraints; i++)
 	{
-		Constraints[i].Phi = Thickness;
-		ensure(Constraints[i % NumConstraints].Particle == Constraints[(i + 1) % NumConstraints].Particle);
-		ensure(Constraints[i % NumConstraints].Levelset == Constraints[(i + 1) % NumConstraints].Levelset);
+		InConstraints[i].Phi = Thickness;
+		ensure(InConstraints[i % NumConstraints].Particle == InConstraints[(i + 1) % NumConstraints].Particle);
+		ensure(InConstraints[i % NumConstraints].Levelset == InConstraints[(i + 1) % NumConstraints].Levelset);
 	}
 
-	FRigidBodyContactConstraint& Constraint = Constraints[0];
+	FRigidBodyContactConstraint& Constraint = InConstraints[0];
 	const TRigidTransform<T, d> ParticleTM = GetTransform(Constraint.Particle);
 	const TRigidTransform<T, d> LevelsetTM = GetTransform(Constraint.Levelset);
 
@@ -2421,7 +2421,7 @@ void TPBDCollisionConstraint<T, d>::UpdateConstraint(const T Thickness, FRigidBo
 		{
 			HistoryData = History[HandleID];
 		}
-		UpdateConstraintImp2<UpdateType>(*Constraint.Particle->Geometry(), ParticleTM, *Constraint.Levelset->Geometry(), LevelsetTM, Thickness, Constraints, NumConstraints, HistoryData);
+		UpdateConstraintImp2<UpdateType>(*Constraint.Particle->Geometry(), ParticleTM, *Constraint.Levelset->Geometry(), LevelsetTM, Thickness, InConstraints, NumConstraints, HistoryData);
 	}
 }
 
diff --git a/Engine/Source/Runtime/Experimental/Chaos/Private/Chaos/PBDConstraintColor.cpp b/Engine/Source/Runtime/Experimental/Chaos/Private/Chaos/PBDConstraintColor.cpp
--- a/Engine/Source/Runtime/Experimental/Chaos/Private/Chaos/PBDConstraintColor.cpp
+++ b/Engine/Source/Runtime/Experimental/Chaos/Private/Chaos/PBDConstraintColor.cpp
@@ -18,10 +18,10 @@
 
 using namespace Chaos;
 
-template<typename T, int d>
-void TPBDConstraintColor<T, d>::ComputeIslandColoring(const int32 Island, const FConstraintGraph& ConstraintGraph, uint32 ContainerId)
+
+void FPBDConstraintColor::ComputeIslandColoring(const int32 Island, const FPBDConstraintGraph& ConstraintGraph, uint32 ContainerId)
 {
-	const TArray<TGeometryParticleHandle<T,d>*>& IslandParticles = ConstraintGraph.GetIslandParticles(Island);
+	const TArray<TGeometryParticleHandle<FReal, 3>*>& IslandParticles = ConstraintGraph.GetIslandParticles(Island);
 	FLevelToColorToConstraintListMap& LevelToColorToConstraintListMap = IslandData[Island].LevelToColorToConstraintListMap;
 	int32& MaxColor = IslandData[Island].MaxColor;
 
@@ -38,7 +38,7 @@ void TPBDConstraintColor<T, d>::ComputeIslandColoring(const int32 Island, const
 	TSet<int32> ProcessedNodes;
 	TArray<int32> NodesToProcess;
 
-	for (const TGeometryParticleHandle<T,d>* Particle :IslandParticles)
+	for (const TGeometryParticleHandle<FReal, 3>* Particle :IslandParticles)
 	{
 		if (!ConstraintGraph.ParticleToNodeIndex.Find(Particle))
 		{
@@ -57,7 +57,7 @@ void TPBDConstraintColor<T, d>::ComputeIslandColoring(const int32 Island, const
 		while (NodesToProcess.Num())
 		{
 			const int32 NodeIndex = NodesToProcess.Last();
-			const typename FConstraintGraph::FGraphNode& GraphNode = ConstraintGraph.Nodes[NodeIndex];
+			const typename FPBDConstraintGraph::FGraphNode& GraphNode = ConstraintGraph.Nodes[NodeIndex];
 			FGraphNodeColor& ColorNode = Nodes[NodeIndex];
 
 			NodesToProcess.SetNum(NodesToProcess.Num() - 1, /*bAllowShrinking=*/false);
@@ -65,11 +65,11 @@ void TPBDConstraintColor<T, d>::ComputeIslandColoring(const int32 Island, const
 
 			for (const int32 EdgeIndex : GraphNode.Edges)
 			{
-				const typename FConstraintGraph::FGraphEdge& GraphEdge = ConstraintGraph.Edges[EdgeIndex];
+				const typename FPBDConstraintGraph::FGraphEdge& GraphEdge = ConstraintGraph.Edges[EdgeIndex];
 				FGraphEdgeColor& ColorEdge = Edges[EdgeIndex];
 
 				// If this is not from our rule, ignore it
-				if (GraphEdge.Data.ContainerId != ContainerId)
+				if (GraphEdge.Data.GetContainerId() != ContainerId)
 				{
 					continue;
 				}
@@ -139,11 +139,11 @@ void TPBDConstraintColor<T, d>::ComputeIslandColoring(const int32 Island, const
 					LevelToColorToConstraintListMap[Level].Add(ColorEdge.Color, {});
 				}
 
-				LevelToColorToConstraintListMap[Level][ColorEdge.Color].Add(GraphEdge.Data.ConstraintHandle);
+				LevelToColorToConstraintListMap[Level][ColorEdge.Color].Add(GraphEdge.Data.GetConstraintHandle());
 
 				if (OtherNodeIndex != INDEX_NONE)
 				{
-					const typename FConstraintGraph::FGraphNode& OtherGraphNode = ConstraintGraph.Nodes[OtherNodeIndex];
+					const typename FPBDConstraintGraph::FGraphNode& OtherGraphNode = ConstraintGraph.Nodes[OtherNodeIndex];
 					if (OtherGraphNode.Particle->AsDynamic() != nullptr)
 					{
 						FGraphNodeColor& OtherColorNode = Nodes[OtherNodeIndex];
@@ -169,15 +169,15 @@ void TPBDConstraintColor<T, d>::ComputeIslandColoring(const int32 Island, const
 }
 
 #ifdef USE_CONTACT_LEVELS
-template<typename T, int d>
-void TPBDConstraintColor<T, d>::ComputeContactGraph(const int32 Island, const FConstraintGraph& ConstraintGraph, uint32 ContainerId)
+
+void FPBDConstraintColor::ComputeContactGraph(const int32 Island, const FPBDConstraintGraph& ConstraintGraph, uint32 ContainerId)
 {
 	const TArray<int32>& ConstraintDataIndices = ConstraintGraph.GetIslandConstraintData(Island);
 
 	IslandData[Island].MaxLevel = ConstraintDataIndices.Num() ? 0 : -1;
 
 	std::queue<std::pair<int32, int32>> QueueToProcess;
-	for (const TGeometryParticleHandle<T,d>* Particle : ConstraintGraph.GetIslandParticles(Island))
+	for (const TGeometryParticleHandle<FReal, 3>* Particle : ConstraintGraph.GetIslandParticles(Island))
 	{
 		const int32* NodeIndexPtr = ConstraintGraph.ParticleToNodeIndex.Find(Particle);
 		if (Particle->AsDynamic() == nullptr && NodeIndexPtr)
@@ -194,15 +194,15 @@ void TPBDConstraintColor<T, d>::ComputeContactGraph(const int32 Island, const FC
 
 		int32 Level = Elem.first;
 		int32 NodeIndex = Elem.second;
-		const typename FConstraintGraph::FGraphNode& GraphNode = ConstraintGraph.Nodes[NodeIndex];
+		const typename FPBDConstraintGraph::FGraphNode& GraphNode = ConstraintGraph.Nodes[NodeIndex];
 
 		for (int32 EdgeIndex : GraphNode.Edges)
 		{
-			const typename FConstraintGraph::FGraphEdge& GraphEdge = ConstraintGraph.Edges[EdgeIndex];
+			const typename FPBDConstraintGraph::FGraphEdge& GraphEdge = ConstraintGraph.Edges[EdgeIndex];
 			FGraphEdgeColor& ColorEdge = Edges[EdgeIndex];
 
 			// If this is not from our rule, ignore it
-			if (GraphEdge.Data.ContainerId != ContainerId)
+			if (GraphEdge.Data.GetContainerId() != ContainerId)
 			{
 				continue;
 			}
@@ -257,8 +257,8 @@ void TPBDConstraintColor<T, d>::ComputeContactGraph(const int32 Island, const FC
 }
 #endif
 
-template<typename T, int d>
-void TPBDConstraintColor<T, d>::InitializeColor(const FConstraintGraph& ConstraintGraph)
+
+void FPBDConstraintColor::InitializeColor(const FPBDConstraintGraph& ConstraintGraph)
 {
 	// Reset the existing Nodes - so colors are all reset to zero
 	ParallelFor(Nodes.Num(), [&](int32 Index)
@@ -282,8 +282,8 @@ void TPBDConstraintColor<T, d>::InitializeColor(const FConstraintGraph& Constrai
 	IslandData.SetNum(ConstraintGraph.IslandToData.Num());
 }
 
-template<typename T, int d>
-void TPBDConstraintColor<T, d>::ComputeColor(const int32 Island, const FConstraintGraph& ConstraintGraph, uint32 ContainerId)
+
+void FPBDConstraintColor::ComputeColor(const int32 Island, const FPBDConstraintGraph& ConstraintGraph, uint32 ContainerId)
 {
 #ifdef USE_CONTACT_LEVELS
 	ComputeContactGraph(Island, ConstraintGraph, ContainerId);
@@ -291,8 +291,8 @@ void TPBDConstraintColor<T, d>::ComputeColor(const int32 Island, const FConstrai
 	ComputeIslandColoring(Island, ConstraintGraph, ContainerId);
 }
 
-template<typename T, int d>
-const typename TPBDConstraintColor<T, d>::FLevelToColorToConstraintListMap& TPBDConstraintColor<T, d>::GetIslandLevelToColorToConstraintListMap(int32 Island) const
+
+const typename FPBDConstraintColor::FLevelToColorToConstraintListMap& FPBDConstraintColor::GetIslandLevelToColorToConstraintListMap(int32 Island) const
 {
 	if (Island < IslandData.Num())
 	{
@@ -301,8 +301,8 @@ const typename TPBDConstraintColor<T, d>::FLevelToColorToConstraintListMap& TPBD
 	return EmptyLevelToColorToConstraintListMap;
 }
 
-template<typename T, int d>
-int TPBDConstraintColor<T, d>::GetIslandMaxColor(int32 Island) const
+
+int FPBDConstraintColor::GetIslandMaxColor(int32 Island) const
 {
 	if (Island < IslandData.Num())
 	{
@@ -311,8 +311,8 @@ int TPBDConstraintColor<T, d>::GetIslandMaxColor(int32 Island) const
 	return -1;
 }
 
-template<typename T, int d>
-int TPBDConstraintColor<T, d>::GetIslandMaxLevel(int32 Island) const
+
+int FPBDConstraintColor::GetIslandMaxLevel(int32 Island) const
 {
 	if (Island < IslandData.Num())
 	{
@@ -324,5 +324,3 @@ int TPBDConstraintColor<T, d>::GetIslandMaxLevel(int32 Island) const
 	}
 	return -1;
 }
-
-template class Chaos::TPBDConstraintColor<float, 3>;
diff --git a/Engine/Source/Runtime/Experimental/Chaos/Private/Chaos/PBDConstraintContainer.cpp b/Engine/Source/Runtime/Experimental/Chaos/Private/Chaos/PBDConstraintContainer.cpp
--- a/Engine/Source/Runtime/Experimental/Chaos/Private/Chaos/PBDConstraintContainer.cpp
+++ b/Engine/Source/Runtime/Experimental/Chaos/Private/Chaos/PBDConstraintContainer.cpp
@@ -3,30 +3,21 @@
 
 namespace Chaos
 {
-	template<typename T, int d>
-	TPBDConstraintContainer<T, d>::TPBDConstraintContainer()
+	FPBDConstraintContainer::FPBDConstraintContainer()
 	{
 	}
 
-	template<typename T, int d>
-	TPBDConstraintContainer<T, d>::~TPBDConstraintContainer()
+	FPBDConstraintContainer::~FPBDConstraintContainer()
 	{
 	}
 
-	template<typename T, int d>
-	int32 TPBDConstraintContainer<T, d>::GetConstraintIndex(const TConstraintHandle<T, d>* ConstraintHandle) const
+	int32 FPBDConstraintContainer::GetConstraintIndex(const FConstraintHandle* ConstraintHandle) const
 	{
-		return ConstraintHandle->ConstraintIndex;
+		return ConstraintHandle->GetConstraintIndex();
 	}
 
-	template<typename T, int d>
-	void TPBDConstraintContainer<T, d>::SetConstraintIndex(TConstraintHandle<T, d>* ConstraintHandle, int32 ConstraintIndex) const
+	void FPBDConstraintContainer::SetConstraintIndex(FConstraintHandle* ConstraintHandle, int32 ConstraintIndex) const
 	{
 		ConstraintHandle->ConstraintIndex = ConstraintIndex;
 	}
 }
-
-namespace Chaos
-{
-	template class TPBDConstraintContainer<float, 3>;
-}
\ No newline at end of file
diff --git a/Engine/Source/Runtime/Experimental/Chaos/Private/Chaos/PBDConstraintGraph.cpp b/Engine/Source/Runtime/Experimental/Chaos/Private/Chaos/PBDConstraintGraph.cpp
--- a/Engine/Source/Runtime/Experimental/Chaos/Private/Chaos/PBDConstraintGraph.cpp
+++ b/Engine/Source/Runtime/Experimental/Chaos/Private/Chaos/PBDConstraintGraph.cpp
@@ -12,13 +12,13 @@
 
 using namespace Chaos;
 
-template<typename T, int d>
-TPBDConstraintGraph<T, d>::TPBDConstraintGraph() : VisitToken(0)
+
+FPBDConstraintGraph::FPBDConstraintGraph() : VisitToken(0)
 {
 }
 
-template<typename T, int d>
-TPBDConstraintGraph<T, d>::TPBDConstraintGraph(const TParticleView<TGeometryParticles<T, d>>& Particles) : VisitToken(0)
+
+FPBDConstraintGraph::FPBDConstraintGraph(const TParticleView<TGeometryParticles<FReal, 3>>& Particles) : VisitToken(0)
 {
 	InitializeGraph(Particles);
 }
@@ -27,8 +27,8 @@ TPBDConstraintGraph<T, d>::TPBDConstraintGraph(const TParticleView<TGeometryPart
  * Bill added this.
  * Adds new Node to Nodes array when a new particle is created
  */
-template<typename T, int d>
-void TPBDConstraintGraph<T, d>::ParticleAdd(TGeometryParticleHandle<T, d>* AddedParticle)
+
+void FPBDConstraintGraph::ParticleAdd(TGeometryParticleHandle<FReal, 3>* AddedParticle)
 {
 	int32 NewNodeIndex = GetNextNodeIndex();
 
@@ -46,8 +46,8 @@ void TPBDConstraintGraph<T, d>::ParticleAdd(TGeometryParticleHandle<T, d>* Added
  * Bill added this
  * Removes Node from Nodes array - marking it an unused, also clears ParticleToNodeIndex
  */
-template<typename T, int d>
-void TPBDConstraintGraph<T, d>::ParticleRemove(TGeometryParticleHandle<T, d>* RemovedParticle)
+
+void FPBDConstraintGraph::ParticleRemove(TGeometryParticleHandle<FReal, 3>* RemovedParticle)
 {
 	if (ParticleToNodeIndex.Contains(RemovedParticle))
 	{
@@ -65,8 +65,8 @@ void TPBDConstraintGraph<T, d>::ParticleRemove(TGeometryParticleHandle<T, d>* Re
 }
 
 
-template<typename T, int d>
-int32 TPBDConstraintGraph<T, d>::GetNextNodeIndex()
+
+int32 FPBDConstraintGraph::GetNextNodeIndex()
 {
 	int32 NewNodeIndex = Nodes.Num();
 	if (FreeIndexList.Num() > 0)
@@ -89,8 +89,8 @@ int32 TPBDConstraintGraph<T, d>::GetNextNodeIndex()
  * Now clears Edges and attempts to retain Nodes and ParticleToNodeIndex
  * It is still setting up nodes that don't have a constraint so wasted effort iterating over all nodes, better to iterate over constraints or don't fill out Nodes if they don't have a constraint
  */
-template<typename T, int d>
-void TPBDConstraintGraph<T, d>::InitializeGraph(const TParticleView<TGeometryParticles<T, d>>& Particles)
+
+void FPBDConstraintGraph::InitializeGraph(const TParticleView<TGeometryParticles<FReal, 3>>& Particles)
 {
 	const int32 NumNonDisabledParticles = Particles.Num();
 
@@ -145,8 +145,8 @@ void TPBDConstraintGraph<T, d>::InitializeGraph(const TParticleView<TGeometryPar
 	}
 }
 
-template<typename T, int d>
-void TPBDConstraintGraph<T, d>::ResetIslands(const TParticleView<TPBDRigidParticles<T, d>>& PBDRigids)
+
+void FPBDConstraintGraph::ResetIslands(const TParticleView<TPBDRigidParticles<FReal, 3>>& PBDRigids)
 {
 	//@todo(ocohen): Should we reset more than just the edges? What about bIsIslandPersistant?
 	for (TArray<int32>& IslandConstraintList : IslandToConstraints)
@@ -169,14 +169,14 @@ void TPBDConstraintGraph<T, d>::ResetIslands(const TParticleView<TPBDRigidPartic
 	}
 }
 
-template<typename T, int d>
-void TPBDConstraintGraph<T, d>::ReserveConstraints(const int32 NumConstraints)
+
+void FPBDConstraintGraph::ReserveConstraints(const int32 NumConstraints)
 {
 	Edges.Reserve(Edges.Num() + NumConstraints);
 }
 
-template<typename T, int d>
-void TPBDConstraintGraph<T, d>::AddConstraint(const uint32 InContainerId, TConstraintHandle<T, d>* InConstraintHandle, const TVector<TGeometryParticleHandle<T,d>*, 2>& ConstrainedParticles)
+
+void FPBDConstraintGraph::AddConstraint(const uint32 InContainerId, FConstraintHandle* InConstraintHandle, const TVector<TGeometryParticleHandle<FReal, 3>*, 2>& ConstrainedParticles)
 {
 	// Must have at least one constrained particle
 	check((ConstrainedParticles[0]) || (ConstrainedParticles[1]));
@@ -208,14 +208,14 @@ void TPBDConstraintGraph<T, d>::AddConstraint(const uint32 InContainerId, TConst
 	}
 }
 
-template<typename T, int d>
-const typename TPBDConstraintGraph<T, d>::FConstraintData& TPBDConstraintGraph<T, d>::GetConstraintData(int32 ConstraintDataIndex) const
+
+const typename FPBDConstraintGraph::FConstraintData& FPBDConstraintGraph::GetConstraintData(int32 ConstraintDataIndex) const
 {
 	return Edges[ConstraintDataIndex].Data;
 }
 
-template<typename T, int d>
-void TPBDConstraintGraph<T, d>::UpdateIslands(const TParticleView<TPBDRigidParticles<T,d>>& PBDRigids, TPBDRigidsSOAs<T,d>& Particles)
+
+void FPBDConstraintGraph::UpdateIslands(const TParticleView<TPBDRigidParticles<FReal, 3>>& PBDRigids, TPBDRigidsSOAs<FReal, 3>& Particles)
 {
 	// Maybe expose a memset style function for this instead of iterating
 	for (auto& PBDRigid : PBDRigids)
@@ -228,13 +228,13 @@ void TPBDConstraintGraph<T, d>::UpdateIslands(const TParticleView<TPBDRigidParti
 
 DECLARE_CYCLE_STAT(TEXT("IslandGeneration2"), STAT_IslandGeneration2, STATGROUP_Chaos);
 
-template<typename T, int d>
-void TPBDConstraintGraph<T, d>::ComputeIslands(const TParticleView<TPBDRigidParticles<T,d>>& PBDRigids, TPBDRigidsSOAs<T,d>& Particles)
+
+void FPBDConstraintGraph::ComputeIslands(const TParticleView<TPBDRigidParticles<FReal, 3>>& PBDRigids, TPBDRigidsSOAs<FReal, 3>& Particles)
 {
 	SCOPE_CYCLE_COUNTER(STAT_IslandGeneration2);
 
 	int32 NextIsland = 0;
-	TArray<TSet<TGeometryParticleHandle<T, d>*>> NewIslandParticles;
+	TArray<TSet<TGeometryParticleHandle<FReal, 3>*>> NewIslandParticles;
 	TArray<int32> NewIslandToSleepCount;
 
 	VisitToken++;
@@ -270,11 +270,11 @@ void TPBDConstraintGraph<T, d>::ComputeIslands(const TParticleView<TPBDRigidPart
 			continue;
 		}
 
-		TSet<TGeometryParticleHandle<T, d>*> SingleIslandParticles;
-		TSet<TGeometryParticleHandle<T, d>*> SingleIslandStaticParticles;
+		TSet<TGeometryParticleHandle<FReal, 3>*> SingleIslandParticles;
+		TSet<TGeometryParticleHandle<FReal, 3>*> SingleIslandStaticParticles;
 		ComputeIsland(Idx, NextIsland, SingleIslandParticles, SingleIslandStaticParticles);
 
-		for (TGeometryParticleHandle<T, d>* StaticParticle : SingleIslandStaticParticles)
+		for (TGeometryParticleHandle<FReal, 3>* StaticParticle : SingleIslandStaticParticles)
 		{
 			SingleIslandParticles.Add(StaticParticle);
 		}
@@ -314,9 +314,9 @@ void TPBDConstraintGraph<T, d>::ComputeIslands(const TParticleView<TPBDRigidPart
 		{
 			NewIslandToSleepCount[Island] = 0;
 
-			for (TGeometryParticleHandle<T, d>* Particle : NewIslandParticles[Island])
+			for (TGeometryParticleHandle<FReal, 3>* Particle : NewIslandParticles[Island])
 			{
-				if (TPBDRigidParticleHandle<T,d>* PBDRigid = Particle->AsDynamic())
+				if (TPBDRigidParticleHandle<FReal, 3>* PBDRigid = Particle->AsDynamic())
 				{
 					PBDRigid->Island() = Island;
 				}
@@ -330,7 +330,7 @@ void TPBDConstraintGraph<T, d>::ComputeIslands(const TParticleView<TPBDRigidPart
 				IslandToData[Island].bIsIslandPersistant = true;
 				bool bSleepState = true;
 
-				for (TGeometryParticleHandle<T, d>* Particle : NewIslandParticles[Island])
+				for (TGeometryParticleHandle<FReal, 3>* Particle : NewIslandParticles[Island])
 				{
 					if (!Particle->Sleeping())
 					{
@@ -339,7 +339,7 @@ void TPBDConstraintGraph<T, d>::ComputeIslands(const TParticleView<TPBDRigidPart
 					}
 				}
 
-				for (TGeometryParticleHandle<T, d>* Particle : NewIslandParticles[Island])
+				for (TGeometryParticleHandle<FReal, 3>* Particle : NewIslandParticles[Island])
 				{
 					//@todo(DEMO_HACK) : Need to fix, remove the !InParticles.Disabled(Index)
 					if (Particle->Sleeping() && !bSleepState/* && !Particle->Disabled()*/)
@@ -347,13 +347,13 @@ void TPBDConstraintGraph<T, d>::ComputeIslands(const TParticleView<TPBDRigidPart
 						Particles.ActivateParticle(Particle); 	//todo: record state change for array reorder
 					}
 
-					if (TPBDRigidParticleHandle<T, d>* PBDRigid = Particle->AsDynamic())
+					if (TPBDRigidParticleHandle<FReal, 3>* PBDRigid = Particle->AsDynamic())
 					{
 						if (!Particle->Sleeping() && bSleepState)
 						{
 							Particles.DeactivateParticle(Particle); 	//todo: record state change for array reorder
-							PBDRigid->V() = TVector<T, d>(0);
-							PBDRigid->W() = TVector<T, d>(0);
+							PBDRigid->V() = TVector<FReal, 3>(0);
+							PBDRigid->W() = TVector<FReal, 3>(0);
 						}
 
 						PBDRigid->SetSleeping(bSleepState);
@@ -374,9 +374,9 @@ void TPBDConstraintGraph<T, d>::ComputeIslands(const TParticleView<TPBDRigidPart
 			// Objects were removed from the island
 			int32 OtherIsland = -1;
 
-			for (TGeometryParticleHandle<T, d>* Particle : IslandToParticles[Island])
+			for (TGeometryParticleHandle<FReal, 3>* Particle : IslandToParticles[Island])
 			{
-				TPBDRigidParticleHandle<T, d>* PBDRigid = Particle->AsDynamic();
+				TPBDRigidParticleHandle<FReal, 3>* PBDRigid = Particle->AsDynamic();
 				int32 TmpIsland = PBDRigid ? PBDRigid->Island() : INDEX_NONE; //question: should we even store non dynamics in this array?
 
 				if (OtherIsland == INDEX_NONE && TmpIsland >= 0)
@@ -406,9 +406,9 @@ void TPBDConstraintGraph<T, d>::ComputeIslands(const TParticleView<TPBDRigidPart
 			}
 			else
 			{
-				for (TGeometryParticleHandle<T, d>* Particle : IslandToParticles[Island])
+				for (TGeometryParticleHandle<FReal, 3>* Particle : IslandToParticles[Island])
 				{
-					if (TPBDRigidParticleHandle<T, d>* PBDRigid = Particle->AsDynamic())
+					if (TPBDRigidParticleHandle<FReal, 3>* PBDRigid = Particle->AsDynamic())
 					{
 						if (!PBDRigid->Disabled())	// todo: why is this needed? [we aren't handling enable/disable state changes properly so disabled particles end up in the graph.]
 						{
@@ -446,9 +446,9 @@ void TPBDConstraintGraph<T, d>::ComputeIslands(const TParticleView<TPBDRigidPart
 	//checkSlow(CheckIslands(InParticles, ActiveIndices));
 }
 
-template<typename T, int d>
-void TPBDConstraintGraph<T, d>::ComputeIsland(const int32 InNode, const int32 Island, TSet<TGeometryParticleHandle<T, d> *>& DynamicParticlesInIsland,
-	TSet<TGeometryParticleHandle<T, d> *>& StaticParticlesInIsland)
+
+void FPBDConstraintGraph::ComputeIsland(const int32 InNode, const int32 Island, TSet<TGeometryParticleHandle<FReal, 3> *>& DynamicParticlesInIsland,
+	TSet<TGeometryParticleHandle<FReal, 3> *>& StaticParticlesInIsland)
 {
 	TQueue<int32> NodeQueue;
 	NodeQueue.Enqueue(InNode);
@@ -464,7 +464,7 @@ void TPBDConstraintGraph<T, d>::ComputeIsland(const int32 InNode, const int32 Is
 			continue;
 		}
 
-		TPBDRigidParticleHandle<T, d>* RigidHandle = Node.Particle->AsDynamic();
+		TPBDRigidParticleHandle<FReal, 3>* RigidHandle = Node.Particle->AsDynamic();
 
 		if (RigidHandle == nullptr)
 		{
@@ -505,12 +505,12 @@ void TPBDConstraintGraph<T, d>::ComputeIsland(const int32 InNode, const int32 Is
 	}
 }
 
-template<typename T, int d>
-bool TPBDConstraintGraph<T, d>::SleepInactive(const int32 Island, const TArrayCollectionArray<TSerializablePtr<TChaosPhysicsMaterial<T>>>& PerParticleMaterialAttributes)
+
+bool FPBDConstraintGraph::SleepInactive(const int32 Island, const TArrayCollectionArray<TSerializablePtr<TChaosPhysicsMaterial<FReal>>>& PerParticleMaterialAttributes)
 {
 	// @todo(ccaulfield): should be able to eliminate this when island is already sleeping
 
-	const TArray<TGeometryParticleHandle<T, d>*>& IslandParticles = GetIslandParticles(Island);
+	const TArray<TGeometryParticleHandle<FReal, 3>*>& IslandParticles = GetIslandParticles(Island);
 	check(IslandParticles.Num());
 
 	if (!IslandToData[Island].bIsIslandPersistant)
@@ -520,30 +520,30 @@ bool TPBDConstraintGraph<T, d>::SleepInactive(const int32 Island, const TArrayCo
 
 	int32& IslandSleepCount = IslandToSleepCount[Island];
 
-	TVector<T, d> X(0);
-	TVector<T, d> V(0);
-	TVector<T, d> W(0);
-	T M = 0;
-	T LinearSleepingThreshold = FLT_MAX;
-	T AngularSleepingThreshold = FLT_MAX;
+	TVector<FReal, 3> X(0);
+	TVector<FReal, 3> V(0);
+	TVector<FReal, 3> W(0);
+	FReal M = 0;
+	FReal LinearSleepingThreshold = FLT_MAX;
+	FReal AngularSleepingThreshold = FLT_MAX;
 
-	for (const TGeometryParticleHandle<T, d>* Particle : IslandToParticles[Island])
+	for (const TGeometryParticleHandle<FReal, 3>* Particle : IslandToParticles[Island])
 	{
-		if (const TPBDRigidParticleHandle<T, d>* PBDRigid = Particle->AsDynamic())
+		if (const TPBDRigidParticleHandle<FReal, 3>* PBDRigid = Particle->AsDynamic())
 		{
 			X += PBDRigid->X() * PBDRigid->M();
 			M += PBDRigid->M();
 			V += PBDRigid->V() * PBDRigid->M();
 
-			if (TSerializablePtr<TChaosPhysicsMaterial<T>> PhysicsMaterial = Particle->AuxilaryValue(PerParticleMaterialAttributes))
+			if (TSerializablePtr<TChaosPhysicsMaterial<FReal>> PhysicsMaterial = Particle->AuxilaryValue(PerParticleMaterialAttributes))
 			{
 				LinearSleepingThreshold = FMath::Min(LinearSleepingThreshold, PhysicsMaterial->SleepingLinearThreshold);
 				AngularSleepingThreshold = FMath::Min(AngularSleepingThreshold, PhysicsMaterial->SleepingAngularThreshold);
 			}
 			else
 			{
-				LinearSleepingThreshold = (T)0;
-				LinearSleepingThreshold = (T)0;
+				LinearSleepingThreshold = (FReal)0;
+				LinearSleepingThreshold = (FReal)0;
 			}
 		}
 	}
@@ -551,30 +551,30 @@ bool TPBDConstraintGraph<T, d>::SleepInactive(const int32 Island, const TArrayCo
 	X /= M;
 	V /= M;
 
-	for (const TGeometryParticleHandle<T, d>* Particle: IslandParticles)
+	for (const TGeometryParticleHandle<FReal, 3>* Particle: IslandParticles)
 	{
-		if (const TPBDRigidParticleHandle<T,d>* PBDRigid = Particle->AsDynamic())
+		if (const TPBDRigidParticleHandle<FReal, 3>* PBDRigid = Particle->AsDynamic())
 		{
-			W += /*TVector<T, d>::CrossProduct(PBDRigid->X() - X, PBDRigid->M() * PBDRigid->V()/ +*/ PBDRigid->W() * PBDRigid->M();
+			W += /*TVector<FReal, 3>::CrossProduct(PBDRigid->X() - X, PBDRigid->M() * PBDRigid->V()/ +*/ PBDRigid->W() * PBDRigid->M();
 		}
 	}
 
 	W /= M;
 
-	const T VSize = V.SizeSquared();
-	const T WSize = W.SizeSquared();
+	const FReal VSize = V.SizeSquared();
+	const FReal WSize = W.SizeSquared();
 
 	if (VSize < LinearSleepingThreshold && WSize < AngularSleepingThreshold)
 	{
 		if (IslandSleepCount > SleepCountThreshold)
 		{
-			for (TGeometryParticleHandle<T, d>* Particle : IslandParticles)
+			for (TGeometryParticleHandle<FReal, 3>* Particle : IslandParticles)
 			{
-				if (TPBDRigidParticleHandle<T, d>* PBDRigid = Particle->AsDynamic())
+				if (TPBDRigidParticleHandle<FReal, 3>* PBDRigid = Particle->AsDynamic())
 				{
 					PBDRigid->SetSleeping(true);
-					PBDRigid->V() = TVector<T, d>(0);
-					PBDRigid->W() = TVector<T, d>(0);
+					PBDRigid->V() = TVector<FReal, 3>(0);
+					PBDRigid->W() = TVector<FReal, 3>(0);
 				}
 				
 			}
@@ -589,12 +589,12 @@ bool TPBDConstraintGraph<T, d>::SleepInactive(const int32 Island, const TArrayCo
 	return false;
 }
 
-template<typename T, int d>
-void TPBDConstraintGraph<T, d>::WakeIsland(const int32 Island)
+
+void FPBDConstraintGraph::WakeIsland(const int32 Island)
 {
-	for (TGeometryParticleHandle<T, d>* Particle : IslandToParticles[Island])
+	for (TGeometryParticleHandle<FReal, 3>* Particle : IslandToParticles[Island])
 	{
-		if (TPBDRigidParticleHandle<T, d>* PBDRigid = Particle->AsDynamic())
+		if (TPBDRigidParticleHandle<FReal, 3>* PBDRigid = Particle->AsDynamic())
 		{
 			if (PBDRigid->Sleeping())
 			{
@@ -606,14 +606,14 @@ void TPBDConstraintGraph<T, d>::WakeIsland(const int32 Island)
 }
 
 
-template<typename T, int d>
-void TPBDConstraintGraph<T, d>::ReconcileIslands()
+
+void FPBDConstraintGraph::ReconcileIslands()
 {
 	for (int32 Island = 0; Island < IslandToParticles.Num(); ++Island)
 	{
 		bool IsSleeping = true;
 		bool IsSet = false;
-		for (TGeometryParticleHandle<T, d>* Particle : IslandToParticles[Island])
+		for (TGeometryParticleHandle<FReal, 3>* Particle : IslandToParticles[Island])
 		{
 			if (Particle->ObjectState() == EObjectStateType::Static)
 			{
@@ -633,16 +633,16 @@ void TPBDConstraintGraph<T, d>::ReconcileIslands()
 	}
 }
 
-template<typename T, int d>
-void TPBDConstraintGraph<T, d>::EnableParticle(TGeometryParticleHandle<T, d>* Particle, const TGeometryParticleHandle<T, d>* ParentParticle)
+
+void FPBDConstraintGraph::EnableParticle(TGeometryParticleHandle<FReal, 3>* Particle, const TGeometryParticleHandle<FReal, 3>* ParentParticle)
 {
 	if (ParentParticle)
 	{
-		if (const TPBDRigidParticleHandle<T, d>* ParentPBDRigid = ParentParticle->AsDynamic())
+		if (const TPBDRigidParticleHandle<FReal, 3>* ParentPBDRigid = ParentParticle->AsDynamic())
 		{
 			ParticleAdd(Particle);
 
-			if (TPBDRigidParticleHandle<T, d>* ChildPBDRigid = Particle->AsDynamic())
+			if (TPBDRigidParticleHandle<FReal, 3>* ChildPBDRigid = Particle->AsDynamic())
 			{
 				const int32 Island = ParentPBDRigid->Island();
 				ChildPBDRigid->Island() = Island;
@@ -662,10 +662,10 @@ void TPBDConstraintGraph<T, d>::EnableParticle(TGeometryParticleHandle<T, d>* Pa
 	}
 }
 
-template<typename T, int d>
-void TPBDConstraintGraph<T, d>::DisableParticle(TGeometryParticleHandle<T, d>* Particle)
+
+void FPBDConstraintGraph::DisableParticle(TGeometryParticleHandle<FReal, 3>* Particle)
 {
-	if (TPBDRigidParticleHandle<T, d>* PBDRigid = Particle->AsDynamic())
+	if (TPBDRigidParticleHandle<FReal, 3>* PBDRigid = Particle->AsDynamic())
 	{
 		const int32 Island = PBDRigid->Island();
 		if (Island != INDEX_NONE)
@@ -685,33 +685,33 @@ void TPBDConstraintGraph<T, d>::DisableParticle(TGeometryParticleHandle<T, d>* P
 	}
 }
 
-template<typename T, int d>
-void TPBDConstraintGraph<T, d>::DisableParticles(const TSet<TGeometryParticleHandle<T, d> *>& Particles)
+
+void FPBDConstraintGraph::DisableParticles(const TSet<TGeometryParticleHandle<FReal, 3> *>& Particles)
 {
 	// @todo(ccaulfield): optimize
-	for (TGeometryParticleHandle<T,d>* Particle : Particles)
+	for (TGeometryParticleHandle<FReal, 3>* Particle : Particles)
 	{
 		DisableParticle(Particle);
 	}
 }
 
-template<typename T, int d>
-bool TPBDConstraintGraph<T, d>::CheckIslands(const TArray<TGeometryParticleHandle<T, d> *>& Particles)
+
+bool FPBDConstraintGraph::CheckIslands(const TArray<TGeometryParticleHandle<FReal, 3> *>& Particles)
 {
 	bool bIsValid = true;
 
 	// Check that no particles are in multiple islands
-	TSet<TGeometryParticleHandle<T, d>*> IslandParticlesUnionSet;
+	TSet<TGeometryParticleHandle<FReal, 3>*> IslandParticlesUnionSet;
 	IslandParticlesUnionSet.Reserve(Particles.Num());
 	for (int32 Island = 0; Island < IslandToParticles.Num(); ++Island)
 	{
-		TSet<TGeometryParticleHandle<T, d>*> IslandParticlesSet = TSet<TGeometryParticleHandle<T,d>*>(IslandToParticles[Island]);
-		TSet<TGeometryParticleHandle<T, d>*> IslandParticlesIntersectSet = IslandParticlesUnionSet.Intersect(IslandParticlesSet);
+		TSet<TGeometryParticleHandle<FReal, 3>*> IslandParticlesSet = TSet<TGeometryParticleHandle<FReal, 3>*>(IslandToParticles[Island]);
+		TSet<TGeometryParticleHandle<FReal, 3>*> IslandParticlesIntersectSet = IslandParticlesUnionSet.Intersect(IslandParticlesSet);
 		if (IslandParticlesIntersectSet.Num() > 0)
 		{
 			// This islands contains particles that were in a previous island.
 			// This is ok only if those particles are static
-			for (TGeometryParticleHandle<T,d>* Particle : IslandParticlesIntersectSet)
+			for (TGeometryParticleHandle<FReal, 3>* Particle : IslandParticlesIntersectSet)
 			{
 				if (Particle->AsDynamic())
 				{
@@ -741,6 +741,3 @@ bool TPBDConstraintGraph<T, d>::CheckIslands(const TArray<TGeometryParticleHandl
 
 	return bIsValid;
 }
-
-
-template class Chaos::TPBDConstraintGraph<float, 3>;
diff --git a/Engine/Source/Runtime/Experimental/Chaos/Private/Chaos/PBDConstraintRule.cpp b/Engine/Source/Runtime/Experimental/Chaos/Private/Chaos/PBDConstraintRule.cpp
--- a/Engine/Source/Runtime/Experimental/Chaos/Private/Chaos/PBDConstraintRule.cpp
+++ b/Engine/Source/Runtime/Experimental/Chaos/Private/Chaos/PBDConstraintRule.cpp
@@ -11,9 +11,9 @@
 
 namespace Chaos
 {
-	template class TPBDConstraintColorRule<TPBDCollisionConstraint<float, 3>, float, 3>;
-	template class TPBDConstraintIslandRule<TPBD6DJointConstraints<float, 3>, float, 3>;
-	template class TPBDConstraintIslandRule<TPBDJointConstraints<float, 3>, float, 3>;
-	template class TPBDConstraintIslandRule<TPBDPositionConstraints<float, 3>, float, 3>;
-	template class TPBDConstraintIslandRule<TPBDRigidDynamicSpringConstraints<float, 3>, float, 3>;
+	template class TPBDConstraintColorRule<TPBDCollisionConstraint<float, 3>>;
+	template class TPBDConstraintIslandRule<FPBD6DJointConstraints>;
+	template class TPBDConstraintIslandRule<FPBDJointConstraints>;
+	template class TPBDConstraintIslandRule<TPBDPositionConstraints<float, 3>>;
+	template class TPBDConstraintIslandRule<TPBDRigidDynamicSpringConstraints<float, 3>>;
 }
\ No newline at end of file
