Chaos - remove template params from constraint classes.

diff --git a/Engine/Source/Runtime/Experimental/Chaos/Private/Chaos/PBDJointConstraintUtilities.cpp b/Engine/Source/Runtime/Experimental/Chaos/Private/Chaos/PBDJointConstraintUtilities.cpp
--- a/Engine/Source/Runtime/Experimental/Chaos/Private/Chaos/PBDJointConstraintUtilities.cpp
+++ b/Engine/Source/Runtime/Experimental/Chaos/Private/Chaos/PBDJointConstraintUtilities.cpp
@@ -7,18 +7,18 @@
 
 namespace Chaos
 {
-	template<typename T, int d>
-	TVector<T, d> TPBDJointUtilities<T, d>::ConditionInertia(const TVector<T, d>& InI, const T MaxRatio)
+	
+	FVec3 FPBDJointUtilities::ConditionInertia(const FVec3& InI, const FReal MaxRatio)
 	{
-		T IMin = InI.Min();
-		T IMax = InI.Max();
+		FReal IMin = InI.Min();
+		FReal IMax = InI.Max();
 		if ((MaxRatio > 0) && (IMin > 0))
 		{
-			T Ratio = IMax / IMin;
+			FReal Ratio = IMax / IMin;
 			if (Ratio > MaxRatio)
 			{
-				T MinIMin = IMax / MaxRatio;
-				return TVector<T, d>(
+				FReal MinIMin = IMax / MaxRatio;
+				return FVec3(
 					FMath::Lerp(MinIMin, IMax, (InI.X - IMin) / (IMax - IMin)),
 					FMath::Lerp(MinIMin, IMax, (InI.Y - IMin) / (IMax - IMin)),
 					FMath::Lerp(MinIMin, IMax, (InI.Z - IMin) / (IMax - IMin)));
@@ -27,19 +27,19 @@ namespace Chaos
 		return InI;
 	}
 
-	template<typename T, int d>
-	TVector<T, d> TPBDJointUtilities<T, d>::ConditionParentInertia(const TVector<T, d>& IParent, const TVector<T, d>& IChild, const T MinRatio)
+	
+	FVec3 FPBDJointUtilities::ConditionParentInertia(const FVec3& IParent, const FVec3& IChild, const FReal MinRatio)
 	{
 		if (MinRatio > 0)
 		{
-			T IParentMax = IParent.Max();
-			T IChildMax = IChild.Max();
+			FReal IParentMax = IParent.Max();
+			FReal IChildMax = IChild.Max();
 			if ((IParentMax > 0) && (IChildMax > 0))
 			{
-				T Ratio = IParentMax / IChildMax;
+				FReal Ratio = IParentMax / IChildMax;
 				if (Ratio < MinRatio)
 				{
-					T Multiplier = MinRatio / Ratio;
+					FReal Multiplier = MinRatio / Ratio;
 					return IParent * Multiplier;
 				}
 			}
@@ -47,164 +47,164 @@ namespace Chaos
 		return IParent;
 	}
 
-	template<typename T, int d>
-	T TPBDJointUtilities<T, d>::ConditionParentMass(const T MParent, const T MChild, const T MinRatio)
+	
+	FReal FPBDJointUtilities::ConditionParentMass(const FReal MParent, const FReal MChild, const FReal MinRatio)
 	{
 		if ((MinRatio > 0) && (MParent > 0) && (MChild > 0))
 		{
-			T Ratio = MParent / MChild;
+			FReal Ratio = MParent / MChild;
 			if (Ratio < MinRatio)
 			{
-				T Multiplier = MinRatio / Ratio;
+				FReal Multiplier = MinRatio / Ratio;
 				return MParent * Multiplier;
 			}
 		}
 		return MParent;
 	}
 
-	template<typename T, int d>
-	void TPBDJointUtilities<T, d>::GetConditionedInverseMass(
-		const float InMParent, 
-		const TVector<T, d> InIParent, 
-		const float InMChild, 
-		const TVector<T, d> InIChild, 
-		T& OutInvMParent, 
-		T& OutInvMChild, 
-		PMatrix<T, d, d>& OutInvIParent, 
-		PMatrix<T, d, d>& OutInvIChild, 
-		const T MinParentMassRatio, 
-		const T MaxInertiaRatio)
+	
+	void FPBDJointUtilities::GetConditionedInverseMass(
+		const FReal InMParent, 
+		const FVec3 InIParent, 
+		const FReal InMChild, 
+		const FVec3 InIChild, 
+		FReal& OutInvMParent, 
+		FReal& OutInvMChild, 
+		FMatrix33& OutInvIParent, 
+		FMatrix33& OutInvIChild, 
+		const FReal MinParentMassRatio, 
+		const FReal MaxInertiaRatio)
 	{
-		T MParent = ConditionParentMass(InMParent, InMChild, MinParentMassRatio);
-		T MChild = InMChild;
+		FReal MParent = ConditionParentMass(InMParent, InMChild, MinParentMassRatio);
+		FReal MChild = InMChild;
 
-		TVector<T, d> IParent = ConditionInertia(InIParent, MaxInertiaRatio);
-		TVector<T, d> IChild = ConditionInertia(InIChild, MaxInertiaRatio);
+		FVec3 IParent = ConditionInertia(InIParent, MaxInertiaRatio);
+		FVec3 IChild = ConditionInertia(InIChild, MaxInertiaRatio);
 		IParent = ConditionParentInertia(IParent, IChild, MinParentMassRatio);
 
 		OutInvMParent = 0;
-		OutInvIParent = PMatrix<T, d, d>(0, 0, 0);
+		OutInvIParent = FMatrix33(0, 0, 0);
 		if (MParent > 0)
 		{
-			OutInvMParent = (T)1 / MParent;
-			OutInvIParent = PMatrix<T, d, d>((T)1 / IParent.X, (T)1 / IParent.Y, (T)1 / IParent.Z);
+			OutInvMParent = (FReal)1 / MParent;
+			OutInvIParent = FMatrix33((FReal)1 / IParent.X, (FReal)1 / IParent.Y, (FReal)1 / IParent.Z);
 		}
 
 		OutInvMChild = 0;
-		OutInvIChild = PMatrix<T, d, d>(0, 0, 0);
+		OutInvIChild = FMatrix33(0, 0, 0);
 		if (MChild > 0)
 		{
-			OutInvMChild = (T)1 / MChild;
-			OutInvIChild = PMatrix<T, d, d>((T)1 / IChild.X, (T)1 / IChild.Y, (T)1 / IChild.Z);
+			OutInvMChild = (FReal)1 / MChild;
+			OutInvIChild = FMatrix33((FReal)1 / IChild.X, (FReal)1 / IChild.Y, (FReal)1 / IChild.Z);
 		}
 	}
 
 
-	template<typename T, int d>
-	void TPBDJointUtilities<T, d>::GetConditionedInverseMass(
-		const float InM0,
-		const TVector<T, d> InI0,
-		T& OutInvM0, 
-		PMatrix<T, d, d>& OutInvI0, 
-		const T MaxInertiaRatio)
+	
+	void FPBDJointUtilities::GetConditionedInverseMass(
+		const FReal InM0,
+		const FVec3 InI0,
+		FReal& OutInvM0, 
+		FMatrix33& OutInvI0, 
+		const FReal MaxInertiaRatio)
 	{
 		OutInvM0 = 0;
-		OutInvI0 = PMatrix<T, d, d>(0, 0, 0);
+		OutInvI0 = FMatrix33(0, 0, 0);
 		if (InM0 > 0)
 		{
-			TVector<T, d> I0 = ConditionInertia(InI0, MaxInertiaRatio);
-			OutInvM0 = (T)1 / InM0;
-			OutInvI0 = PMatrix<T, d, d>((T)1 / I0.X, (T)1 / I0.Y, (T)1 / I0.Z);
+			FVec3 I0 = ConditionInertia(InI0, MaxInertiaRatio);
+			OutInvM0 = (FReal)1 / InM0;
+			OutInvI0 = FMatrix33((FReal)1 / I0.X, (FReal)1 / I0.Y, (FReal)1 / I0.Z);
 		}
 	}
 
 
-	template<typename T, int d>
-	TVector<T, d> GetSphereLimitedPositionError(const TVector<T, d>& CX, const T Radius)
+	
+	FVec3 GetSphereLimitedPositionError(const FVec3& CX, const FReal Radius)
 	{
-		T CXLen = CX.Size();
+		FReal CXLen = CX.Size();
 		if (CXLen < Radius)
 		{
-			return TVector<T, d>(0, 0, 0);
+			return FVec3(0, 0, 0);
 		}
 		else if (CXLen > SMALL_NUMBER)
 		{
-			TVector<T, d> Dir = CX / CXLen;
+			FVec3 Dir = CX / CXLen;
 			return CX - Radius * Dir;
 		}
 		return CX;
 	}
 
-	template<typename T, int d>
-	TVector<T, d> GetSphereLimitedVelocityError(const TVector<T, d>& CX, const T Radius, const TVector<T, d>& CV)
+	
+	FVec3 GetSphereLimitedVelocityError(const FVec3& CX, const FReal Radius, const FVec3& CV)
 	{
-		T CXLen = CX.Size();
+		FReal CXLen = CX.Size();
 		if (CXLen < Radius)
 		{
-			return TVector<T, d>(0, 0, 0);
+			return FVec3(0, 0, 0);
 		}
 		else if (CXLen > SMALL_NUMBER)
 		{
-			TVector<T, d> Dir = CX / CXLen;
-			T CVDir = TVector<T, d>::DotProduct(CV, Dir);
-			return FMath::Max((T)0, CVDir) * Dir;
+			FVec3 Dir = CX / CXLen;
+			FReal CVDir = FVec3::DotProduct(CV, Dir);
+			return FMath::Max((FReal)0, CVDir) * Dir;
 		}
 		return CV;
 	}
 
-	template<typename T, int d>
-	TVector<T, d> GetCylinderLimitedPositionError(const TVector<T, d>& InCX, const TVector<T, d>& Axis, const T Limit, const EJointMotionType AxisMotion)
+	
+	FVec3 GetCylinderLimitedPositionError(const FVec3& InCX, const FVec3& Axis, const FReal Limit, const EJointMotionType AxisMotion)
 	{
-		TVector<T, d> CXAxis = TVector<T, d>::DotProduct(InCX, Axis) * Axis;
-		TVector<T, d> CXPlane = InCX - CXAxis;
-		T CXPlaneLen = CXPlane.Size();
+		FVec3 CXAxis = FVec3::DotProduct(InCX, Axis) * Axis;
+		FVec3 CXPlane = InCX - CXAxis;
+		FReal CXPlaneLen = CXPlane.Size();
 		if (AxisMotion == EJointMotionType::Free)
 		{
-			CXAxis = TVector<T, d>(0, 0, 0);
+			CXAxis = FVec3(0, 0, 0);
 		}
 		if (CXPlaneLen < Limit)
 		{
-			CXPlane = TVector<T, d>(0, 0, 0);
+			CXPlane = FVec3(0, 0, 0);
 		}
 		else if (CXPlaneLen > KINDA_SMALL_NUMBER)
 		{
-			TVector<T, d> Dir = CXPlane / CXPlaneLen;
+			FVec3 Dir = CXPlane / CXPlaneLen;
 			CXPlane = CXPlane - Limit * Dir;
 		}
 		return CXAxis + CXPlane;
 	}
 
-	template<typename T, int d>
-	TVector<T, d> GetCylinderLimitedVelocityError(const TVector<T, d>& InCX, const TVector<T, d>& Axis, const T Limit, const EJointMotionType AxisMotion, const TVector<T, d>& CV)
+	
+	FVec3 GetCylinderLimitedVelocityError(const FVec3& InCX, const FVec3& Axis, const FReal Limit, const EJointMotionType AxisMotion, const FVec3& CV)
 	{
-		TVector<T, d> CXAxis = TVector<T, d>::DotProduct(InCX, Axis) * Axis;
-		TVector<T, d> CXPlane = InCX - CXAxis;
-		T CXPlaneLen = CXPlane.Size();
+		FVec3 CXAxis = FVec3::DotProduct(InCX, Axis) * Axis;
+		FVec3 CXPlane = InCX - CXAxis;
+		FReal CXPlaneLen = CXPlane.Size();
 
-		TVector<T, d> CVAxis = TVector<T, d>::DotProduct(CV, Axis) * Axis;;
-		TVector<T, d> CVPlane = CV - CVAxis;
+		FVec3 CVAxis = FVec3::DotProduct(CV, Axis) * Axis;;
+		FVec3 CVPlane = CV - CVAxis;
 
 		if (AxisMotion == EJointMotionType::Free)
 		{
-			CVAxis = TVector<T, d>(0, 0, 0);
+			CVAxis = FVec3(0, 0, 0);
 		}
 		if (CXPlaneLen < Limit)
 		{
-			CVPlane = TVector<T, d>(0, 0, 0);
+			CVPlane = FVec3(0, 0, 0);
 		}
 		else if (CXPlaneLen > KINDA_SMALL_NUMBER)
 		{
-			TVector<T, d> Dir = CXPlane / CXPlaneLen;
-			T CVDir = TVector<T, d>::DotProduct(CV, Dir);
-			CVPlane = FMath::Max((T)0, CVDir) * Dir;
+			FVec3 Dir = CXPlane / CXPlaneLen;
+			FReal CVDir = FVec3::DotProduct(CV, Dir);
+			CVPlane = FMath::Max((FReal)0, CVDir) * Dir;
 		}
 		return CVAxis + CVPlane;
 	}
 
-	template<typename T, int d>
-	TVector<T, d> GetLineLimitedPositionError(const TVector<T, d>& CX, const TVector<T, d>& Axis, const T Limit, const EJointMotionType AxisMotion)
+	
+	FVec3 GetLineLimitedPositionError(const FVec3& CX, const FVec3& Axis, const FReal Limit, const EJointMotionType AxisMotion)
 	{
-		T CXDist = TVector<T, d>::DotProduct(CX, Axis);
+		FReal CXDist = FVec3::DotProduct(CX, Axis);
 		if ((AxisMotion == EJointMotionType::Free) || (FMath::Abs(CXDist) < Limit))
 		{
 			return CX - CXDist * Axis;
@@ -219,30 +219,30 @@ namespace Chaos
 		}
 	}
 
-	template<typename T, int d>
-	TVector<T, d> GetLineLimitedVelocityError(const TVector<T, d>& CX, const TVector<T, d>& Axis, const T Limit, const EJointMotionType AxisMotion, const TVector<T, d>& CV)
+	
+	FVec3 GetLineLimitedVelocityError(const FVec3& CX, const FVec3& Axis, const FReal Limit, const EJointMotionType AxisMotion, const FVec3& CV)
 	{
-		T CXDist = TVector<T, d>::DotProduct(CX, Axis);
-		T CVAxis = TVector<T, d>::DotProduct(CV, Axis);
+		FReal CXDist = FVec3::DotProduct(CX, Axis);
+		FReal CVAxis = FVec3::DotProduct(CV, Axis);
 		if ((AxisMotion == EJointMotionType::Free) || (FMath::Abs(CXDist) < Limit))
 		{
 			return CV - CVAxis * Axis;
 		}
 		else if (CXDist >= Limit)
 		{
-			return CV - FMath::Min((T)0, CVAxis) * Axis;
+			return CV - FMath::Min((FReal)0, CVAxis) * Axis;
 		}
 		else
 		{
-			return CV - FMath::Max((T)0, CVAxis) * Axis;
+			return CV - FMath::Max((FReal)0, CVAxis) * Axis;
 		}
 	}
 
 
-	template<typename T, int d>
-	TVector<T, d> GetLimitedPositionError(const TPBDJointSettings<T, d>& JointSettings, const TRotation<T, d>& R0, const TVector<T, d>& InCX)
+	
+	FVec3 GetLimitedPositionError(const FPBDJointSettings& JointSettings, const FRotation3& R0, const FVec3& InCX)
 	{
-		const TVector<EJointMotionType, d>& Motion = JointSettings.Motion.LinearMotionTypes;
+		const TVector<EJointMotionType, 3>& Motion = JointSettings.Motion.LinearMotionTypes;
 		if ((Motion[0] == EJointMotionType::Locked) && (Motion[1] == EJointMotionType::Locked) && (Motion[2] == EJointMotionType::Locked))
 		{
 			return InCX;
@@ -255,48 +255,48 @@ namespace Chaos
 		else if ((Motion[1] == EJointMotionType::Limited) && (Motion[2] == EJointMotionType::Limited))
 		{
 			// Circular Limit (X Axis)
-			TVector<T, d> Axis = R0 * TVector<T, d>(1, 0, 0);
+			FVec3 Axis = R0 * FVec3(1, 0, 0);
 			return GetCylinderLimitedPositionError(InCX, Axis, JointSettings.Motion.LinearLimit, Motion[0]);
 		}
 		else if ((Motion[0] == EJointMotionType::Limited) && (Motion[2] == EJointMotionType::Limited))
 		{
 			// Circular Limit (Y Axis)
-			TVector<T, d> Axis = R0 * TVector<T, d>(0, 1, 0);
+			FVec3 Axis = R0 * FVec3(0, 1, 0);
 			return GetCylinderLimitedPositionError(InCX, Axis, JointSettings.Motion.LinearLimit, Motion[1]);
 		}
 		else if ((Motion[0] == EJointMotionType::Limited) && (Motion[1] == EJointMotionType::Limited))
 		{
 			// Circular Limit (Z Axis)
-			TVector<T, d> Axis = R0 * TVector<T, d>(0, 0, 1);
+			FVec3 Axis = R0 * FVec3(0, 0, 1);
 			return GetCylinderLimitedPositionError(InCX, Axis, JointSettings.Motion.LinearLimit, Motion[2]);
 		}
 		else
 		{
 			// Line/Square/Cube Limits (no way to author square or cube limits, but would work if we wanted it)
-			TVector<T, d> CX = InCX;
+			FVec3 CX = InCX;
 			if (Motion[0] != EJointMotionType::Locked)
 			{
-				TVector<T, d> Axis = R0 * TVector<T, d>(1, 0, 0);
+				FVec3 Axis = R0 * FVec3(1, 0, 0);
 				CX = GetLineLimitedPositionError(CX, Axis, JointSettings.Motion.LinearLimit, Motion[0]);
 			}
 			if (Motion[1] != EJointMotionType::Locked)
 			{
-				TVector<T, d> Axis = R0 * TVector<T, d>(0, 1, 0);
+				FVec3 Axis = R0 * FVec3(0, 1, 0);
 				CX = GetLineLimitedPositionError(CX, Axis, JointSettings.Motion.LinearLimit, Motion[1]);
 			}
 			if (Motion[2] != EJointMotionType::Locked)
 			{
-				TVector<T, d> Axis = R0 * TVector<T, d>(0, 0, 1);
+				FVec3 Axis = R0 * FVec3(0, 0, 1);
 				CX = GetLineLimitedPositionError(CX, Axis, JointSettings.Motion.LinearLimit, Motion[2]);
 			}
 			return CX;
 		}
 	}
 
-	template<typename T, int d>
-	TVector<T, d> GetLimitedVelocityError(const TPBDJointSettings<T, d>& JointSettings, const TRotation<T, d>& R0, const TVector<T, d>& InCX, const TVector<T, d>& InCV)
+	
+	FVec3 GetLimitedVelocityError(const FPBDJointSettings& JointSettings, const FRotation3& R0, const FVec3& InCX, const FVec3& InCV)
 	{
-		const TVector<EJointMotionType, d>& Motion = JointSettings.Motion.LinearMotionTypes;
+		const TVector<EJointMotionType, 3>& Motion = JointSettings.Motion.LinearMotionTypes;
 		if ((Motion[0] == EJointMotionType::Locked) && (Motion[1] == EJointMotionType::Locked) && (Motion[2] == EJointMotionType::Locked))
 		{
 			return InCV;
@@ -309,111 +309,111 @@ namespace Chaos
 		else if ((Motion[1] == EJointMotionType::Limited) && (Motion[2] == EJointMotionType::Limited))
 		{
 			// Circular Limit (X Axis)
-			TVector<T, d> Axis = R0 * TVector<T, d>(1, 0, 0);
+			FVec3 Axis = R0 * FVec3(1, 0, 0);
 			return GetCylinderLimitedVelocityError(InCX, Axis, JointSettings.Motion.LinearLimit, Motion[0], InCV);
 		}
 		else if ((Motion[0] == EJointMotionType::Limited) && (Motion[2] == EJointMotionType::Limited))
 		{
 			// Circular Limit (Y Axis)
-			TVector<T, d> Axis = R0 * TVector<T, d>(0, 1, 0);
+			FVec3 Axis = R0 * FVec3(0, 1, 0);
 			return GetCylinderLimitedVelocityError(InCX, Axis, JointSettings.Motion.LinearLimit, Motion[1], InCV);
 		}
 		else if ((Motion[0] == EJointMotionType::Limited) && (Motion[1] == EJointMotionType::Limited))
 		{
 			// Circular Limit (Z Axis)
-			TVector<T, d> Axis = R0 * TVector<T, d>(0, 0, 1);
+			FVec3 Axis = R0 * FVec3(0, 0, 1);
 			return GetCylinderLimitedVelocityError(InCX, Axis, JointSettings.Motion.LinearLimit, Motion[2], InCV);
 		}
 		else
 		{
 			// Line/Square/Cube Limits (no way to author square or cube limits, but would work if we wanted it)
-			TVector<T, d> CV = InCV;
+			FVec3 CV = InCV;
 			if (Motion[0] != EJointMotionType::Locked)
 			{
-				TVector<T, d> Axis = R0 * TVector<T, d>(1, 0, 0);
+				FVec3 Axis = R0 * FVec3(1, 0, 0);
 				CV = GetLineLimitedVelocityError(InCX, Axis, JointSettings.Motion.LinearLimit, Motion[0], CV);
 			}
 			if (Motion[1] != EJointMotionType::Locked)
 			{
-				TVector<T, d> Axis = R0 * TVector<T, d>(0, 1, 0);
+				FVec3 Axis = R0 * FVec3(0, 1, 0);
 				CV = GetLineLimitedVelocityError(InCX, Axis, JointSettings.Motion.LinearLimit, Motion[1], CV);
 			}
 			if (Motion[2] != EJointMotionType::Locked)
 			{
-				TVector<T, d> Axis = R0 * TVector<T, d>(0, 0, 1);
+				FVec3 Axis = R0 * FVec3(0, 0, 1);
 				CV = GetLineLimitedVelocityError(InCX, Axis, JointSettings.Motion.LinearLimit, Motion[2], CV);
 			}
 			return CV;
 		}
 	}
 
-	template<typename T, int d>
-	T GetLinearStiffness(
-		const TPBDJointSolverSettings<T, d>& SolverSettings,
-		const TPBDJointSettings<T, d>& JointSettings)
+	
+	FReal GetLinearStiffness(
+		const FPBDJointSolverSettings& SolverSettings,
+		const FPBDJointSettings& JointSettings)
 	{
-		const float SolverStiffness = (SolverSettings.Stiffness > (T)0) ? SolverSettings.Stiffness : JointSettings.Motion.Stiffness;
-		const float SoftSolverStiffness = (SolverSettings.SoftLinearStiffness > (T)0) ? SolverSettings.SoftLinearStiffness : JointSettings.Motion.SoftLinearStiffness;
+		const FReal SolverStiffness = (SolverSettings.Stiffness > (FReal)0) ? SolverSettings.Stiffness : JointSettings.Motion.Stiffness;
+		const FReal SoftSolverStiffness = (SolverSettings.SoftLinearStiffness > (FReal)0) ? SolverSettings.SoftLinearStiffness : JointSettings.Motion.SoftLinearStiffness;
 		const bool bIsSoft = JointSettings.Motion.bSoftLinearLimitsEnabled && ((JointSettings.Motion.LinearMotionTypes[0] == EJointMotionType::Limited) || (JointSettings.Motion.LinearMotionTypes[1] == EJointMotionType::Limited) || (JointSettings.Motion.LinearMotionTypes[2] == EJointMotionType::Limited));
-		const float Stiffness = bIsSoft ? SolverStiffness * SoftSolverStiffness : SolverStiffness;
+		const FReal Stiffness = bIsSoft ? SolverStiffness * SoftSolverStiffness : SolverStiffness;
 		return Stiffness;
 	}
 
-	template<typename T, int d>
-	T GetTwistStiffness(
-		const TPBDJointSolverSettings<T, d>& SolverSettings,
-		const TPBDJointSettings<T, d>& JointSettings)
+	
+	FReal GetTwistStiffness(
+		const FPBDJointSolverSettings& SolverSettings,
+		const FPBDJointSettings& JointSettings)
 	{
-		const float SolverStiffness = (SolverSettings.Stiffness > (T)0) ? SolverSettings.Stiffness : JointSettings.Motion.Stiffness;
-		const float SoftSolverStiffness = (SolverSettings.SoftAngularStiffness > (T)0) ? SolverSettings.SoftAngularStiffness : JointSettings.Motion.SoftTwistStiffness;
+		const FReal SolverStiffness = (SolverSettings.Stiffness > (FReal)0) ? SolverSettings.Stiffness : JointSettings.Motion.Stiffness;
+		const FReal SoftSolverStiffness = (SolverSettings.SoftAngularStiffness > (FReal)0) ? SolverSettings.SoftAngularStiffness : JointSettings.Motion.SoftTwistStiffness;
 		const bool bIsSoft = JointSettings.Motion.bSoftTwistLimitsEnabled && (JointSettings.Motion.AngularMotionTypes[(int32)EJointAngularConstraintIndex::Twist] == EJointMotionType::Limited);
-		const float Stiffness = bIsSoft ? SolverStiffness * SoftSolverStiffness : SolverStiffness;
+		const FReal Stiffness = bIsSoft ? SolverStiffness * SoftSolverStiffness : SolverStiffness;
 		return Stiffness;
 	}
 
-	template<typename T, int d>
-	T GetSwingStiffness(
-		const TPBDJointSolverSettings<T, d>& SolverSettings,
-		const TPBDJointSettings<T, d>& JointSettings)
+	
+	FReal GetSwingStiffness(
+		const FPBDJointSolverSettings& SolverSettings,
+		const FPBDJointSettings& JointSettings)
 	{
-		const float SolverStiffness = (SolverSettings.Stiffness > (T)0) ? SolverSettings.Stiffness : JointSettings.Motion.Stiffness;
-		const float SoftSolverStiffness = (SolverSettings.SoftAngularStiffness > (T)0) ? SolverSettings.SoftAngularStiffness : JointSettings.Motion.SoftSwingStiffness;
+		const FReal SolverStiffness = (SolverSettings.Stiffness > (FReal)0) ? SolverSettings.Stiffness : JointSettings.Motion.Stiffness;
+		const FReal SoftSolverStiffness = (SolverSettings.SoftAngularStiffness > (FReal)0) ? SolverSettings.SoftAngularStiffness : JointSettings.Motion.SoftSwingStiffness;
 		const bool bIsSoft = JointSettings.Motion.bSoftSwingLimitsEnabled && ((JointSettings.Motion.AngularMotionTypes[(int32)EJointAngularConstraintIndex::Swing1] == EJointMotionType::Limited) || (JointSettings.Motion.AngularMotionTypes[(int32)EJointAngularConstraintIndex::Swing2] == EJointMotionType::Limited));
-		const float Stiffness = bIsSoft ? SolverStiffness * SoftSolverStiffness : SolverStiffness;
+		const FReal Stiffness = bIsSoft ? SolverStiffness * SoftSolverStiffness : SolverStiffness;
 		return Stiffness;
 	}
 
-	template<typename T, int d>
+	
 	void ApplyPositionDelta(
-		const T Dt,
-		const TPBDJointSolverSettings<T, d>& SolverSettings,
-		const TPBDJointSettings<T, d>& JointSettings,
+		const FReal Dt,
+		const FPBDJointSolverSettings& SolverSettings,
+		const FPBDJointSettings& JointSettings,
 		const int32 Index0,
 		const int32 Index1,
-		TVector<T, d>& P0,
-		TVector<T, d>& P1,
-		const TVector<T, d>& DP0,
-		const TVector<T, d>& DP1,
-		const T Stiffness)
+		FVec3& P0,
+		FVec3& P1,
+		const FVec3& DP0,
+		const FVec3& DP1,
+		const FReal Stiffness)
 	{
 		P0 = P0 + Stiffness * DP0;
 		P1 = P1 + Stiffness * DP1;
 	}
 
-	template<typename T, int d>
+	
 	void ApplyVelocityDelta(
-		const T Dt,
-		const TPBDJointSolverSettings<T, d>& SolverSettings,
-		const TPBDJointSettings<T, d>& JointSettings,
+		const FReal Dt,
+		const FPBDJointSolverSettings& SolverSettings,
+		const FPBDJointSettings& JointSettings,
 		const int32 Index0,
 		const int32 Index1,
-		TVector<T, d>& P0,
-		TVector<T, d>& V0,
-		TVector<T, d>& P1,
-		TVector<T, d>& V1,
-		const TVector<T, d>& DV0,
-		const TVector<T, d>& DV1, 
-		const T Stiffness)
+		FVec3& P0,
+		FVec3& V0,
+		FVec3& P1,
+		FVec3& V1,
+		const FVec3& DV0,
+		const FVec3& DV1, 
+		const FReal Stiffness)
 	{
 		V0 = V0 + Stiffness * DV0;
 		V1 = V1 + Stiffness * DV1;
@@ -421,41 +421,41 @@ namespace Chaos
 		ApplyPositionDelta(Dt, SolverSettings, JointSettings, Index0, Index1, P0, P1, DV0 * Dt, DV1 * Dt, Stiffness);
 	}
 
-	template<typename T, int d>
+	
 	void ApplyRotationDelta(
-		const T Dt,
-		const TPBDJointSolverSettings<T, d>& SolverSettings,
-		const TPBDJointSettings<T, d>& JointSettings,
+		const FReal Dt,
+		const FPBDJointSolverSettings& SolverSettings,
+		const FPBDJointSettings& JointSettings,
 		const int32 Index0,
 		const int32 Index1,
-		TRotation<T, d>& Q0,
-		TRotation<T, d>& Q1,
-		const TVector<T, d>& DR0,
-		const TVector<T, d>& DR1,
-		const T Stiffness)
+		FRotation3& Q0,
+		FRotation3& Q1,
+		const FVec3& DR0,
+		const FVec3& DR1,
+		const FReal Stiffness)
 	{
-		const TRotation<T, d> DQ0 = (TRotation<T, d>::FromElements(Stiffness * DR0, 0) * Q0) * (T)0.5;
-		const TRotation<T, d> DQ1 = (TRotation<T, d>::FromElements(Stiffness * DR1, 0) * Q1) * (T)0.5;
+		const FRotation3 DQ0 = (FRotation3::FromElements(Stiffness * DR0, 0) * Q0) * (FReal)0.5;
+		const FRotation3 DQ1 = (FRotation3::FromElements(Stiffness * DR1, 0) * Q1) * (FReal)0.5;
 		Q0 = (Q0 + DQ0).GetNormalized();
 		Q1 = (Q1 + DQ1).GetNormalized();
 		Q1.EnforceShortestArcWith(Q0);
 	}
 
 
-	template<typename T, int d>
+	
 	void ApplyRotationVelocityDelta(
-		const T Dt,
-		const TPBDJointSolverSettings<T, d>& SolverSettings,
-		const TPBDJointSettings<T, d>& JointSettings,
+		const FReal Dt,
+		const FPBDJointSolverSettings& SolverSettings,
+		const FPBDJointSettings& JointSettings,
 		const int32 Index0,
 		const int32 Index1,
-		TRotation<T, d>& Q0,
-		TVector<T, d>& W0,
-		TRotation<T, d>& Q1,
-		TVector<T, d>& W1,
-		const TVector<T, d>& DW0,
-		const TVector<T, d>& DW1,
-		const T Stiffness)
+		FRotation3& Q0,
+		FVec3& W0,
+		FRotation3& Q1,
+		FVec3& W1,
+		const FVec3& DW0,
+		const FVec3& DW1,
+		const FReal Stiffness)
 	{
 		W0 = W0 + Stiffness * DW0;
 		W1 = W1 + Stiffness * DW1;
@@ -464,142 +464,142 @@ namespace Chaos
 	}
 
 
-	template<typename T, int d>
+	
 	void ApplyRotationVelocityDelta(
-		const T Dt,
-		const TPBDJointSolverSettings<T, d>& SolverSettings,
-		const TPBDJointSettings<T, d>& JointSettings,
+		const FReal Dt,
+		const FPBDJointSolverSettings& SolverSettings,
+		const FPBDJointSettings& JointSettings,
 		const int32 Index0,
 		const int32 Index1,
-		TRotation<T, d>& Q0,
-		TVector<T, d>& W0,
-		TRotation<T, d>& Q1,
-		TVector<T, d>& W1,
-		const T InvM0,
-		const PMatrix<T, d, d>& InvIL0,
-		const T InvM1,
-		const PMatrix<T, d, d>& InvIL1,
-		const TVector<T, d>& Axis0,
-		const TVector<T, d>& Axis1,
-		const T WC,
-		const T Stiffness)
+		FRotation3& Q0,
+		FVec3& W0,
+		FRotation3& Q1,
+		FVec3& W1,
+		const FReal InvM0,
+		const FMatrix33& InvIL0,
+		const FReal InvM1,
+		const FMatrix33& InvIL1,
+		const FVec3& Axis0,
+		const FVec3& Axis1,
+		const FReal WC,
+		const FReal Stiffness)
 	{
-		const PMatrix<T, d, d> InvI0 = Utilities::Multiply(Q0.ToMatrix(), Utilities::Multiply(InvIL0, Q0.ToMatrix().GetTransposed()));
-		const PMatrix<T, d, d> InvI1 = Utilities::Multiply(Q1.ToMatrix(), Utilities::Multiply(InvIL1, Q1.ToMatrix().GetTransposed()));
-		//const T I0 = TVector<T, d>::DotProduct(Axis0, Utilities::Multiply(InvI0, Axis0));
-		//const T I1 = TVector<T, d>::DotProduct(Axis1, Utilities::Multiply(InvI1, Axis1));
-		//const TVector<T, d> DW0 = Axis0 * WC * I0 / (I0 + I1);
-		//const TVector<T, d> DW1 = -Axis1 * WC * I1 / (I0 + I1);
-		const T L = (T)1 / (TVector<T, d>::DotProduct(Axis0, Utilities::Multiply(InvI0, Axis0)) + TVector<T, d>::DotProduct(Axis1, Utilities::Multiply(InvI1, Axis1)));
-		const TVector<T, d> DW0 = Utilities::Multiply(InvI0, Axis0) * L * WC;
-		const TVector<T, d> DW1 = -Utilities::Multiply(InvI1, Axis1) * L * WC;
+		const FMatrix33 InvI0 = Utilities::Multiply(Q0.ToMatrix(), Utilities::Multiply(InvIL0, Q0.ToMatrix().GetTransposed()));
+		const FMatrix33 InvI1 = Utilities::Multiply(Q1.ToMatrix(), Utilities::Multiply(InvIL1, Q1.ToMatrix().GetTransposed()));
+		//const FReal I0 = FVec3::DotProduct(Axis0, Utilities::Multiply(InvI0, Axis0));
+		//const FReal I1 = FVec3::DotProduct(Axis1, Utilities::Multiply(InvI1, Axis1));
+		//const FVec3 DW0 = Axis0 * WC * I0 / (I0 + I1);
+		//const FVec3 DW1 = -Axis1 * WC * I1 / (I0 + I1);
+		const FReal L = (FReal)1 / (FVec3::DotProduct(Axis0, Utilities::Multiply(InvI0, Axis0)) + FVec3::DotProduct(Axis1, Utilities::Multiply(InvI1, Axis1)));
+		const FVec3 DW0 = Utilities::Multiply(InvI0, Axis0) * L * WC;
+		const FVec3 DW1 = -Utilities::Multiply(InvI1, Axis1) * L * WC;
 
 		ApplyRotationVelocityDelta(Dt, SolverSettings, JointSettings, Index0, Index1, Q0, W0, Q1, W1, DW0, DW1, Stiffness);
 	}
 
-	template<typename T, int d>
+	
 	void ApplyRotationDelta(
-		const T Dt,
-		const TPBDJointSolverSettings<T, d>& SolverSettings,
-		const TPBDJointSettings<T, d>& JointSettings,
+		const FReal Dt,
+		const FPBDJointSolverSettings& SolverSettings,
+		const FPBDJointSettings& JointSettings,
 		const int32 Index0,
 		const int32 Index1,
-		TVector<T, d>& P0,
-		TRotation<T, d>& Q0,
-		TVector<T, d>& P1,
-		TRotation<T, d>& Q1,
-		const T InvM0,
-		const PMatrix<T, d, d>& InvIL0,
-		const T InvM1,
-		const PMatrix<T, d, d>& InvIL1,
-		const TVector<T, d>& Axis0,
-		const T Angle0,
-		const TVector<T, d>& Axis1,
-		const T Angle1, 
-		const T Stiffness)
+		FVec3& P0,
+		FRotation3& Q0,
+		FVec3& P1,
+		FRotation3& Q1,
+		const FReal InvM0,
+		const FMatrix33& InvIL0,
+		const FReal InvM1,
+		const FMatrix33& InvIL1,
+		const FVec3& Axis0,
+		const FReal Angle0,
+		const FVec3& Axis1,
+		const FReal Angle1, 
+		const FReal Stiffness)
 	{
-		const PMatrix<T, d, d> InvI0 = Utilities::Multiply(Q0.ToMatrix(), Utilities::Multiply(InvIL0, Q0.ToMatrix().GetTransposed()));
-		const PMatrix<T, d, d> InvI1 = Utilities::Multiply(Q1.ToMatrix(), Utilities::Multiply(InvIL1, Q1.ToMatrix().GetTransposed()));
+		const FMatrix33 InvI0 = Utilities::Multiply(Q0.ToMatrix(), Utilities::Multiply(InvIL0, Q0.ToMatrix().GetTransposed()));
+		const FMatrix33 InvI1 = Utilities::Multiply(Q1.ToMatrix(), Utilities::Multiply(InvIL1, Q1.ToMatrix().GetTransposed()));
 
-		//const T I0 = TVector<T, d>::DotProduct(Axis0, Utilities::Multiply(InvI0, Axis0));
-		//const T I1 = TVector<T, d>::DotProduct(Axis1, Utilities::Multiply(InvI1, Axis1));
-		//const TVector<T, d> DR0 = Axis0 * Angle0 * I0 / (I0 + I1);
-		//const TVector<T, d> DR1 = Axis1 * Angle1 * I1 / (I0 + I1);
-		const T L = (T)1 / (TVector<T, d>::DotProduct(Axis0, Utilities::Multiply(InvI0, Axis0)) + TVector<T, d>::DotProduct(Axis1, Utilities::Multiply(InvI1, Axis1)));
-		const TVector<T, d> DR0 = Utilities::Multiply(InvI0, Axis0) * L * Angle0;
-		const TVector<T, d> DR1 = Utilities::Multiply(InvI1, Axis1) * L * Angle1;
+		//const FReal I0 = FVec3::DotProduct(Axis0, Utilities::Multiply(InvI0, Axis0));
+		//const FReal I1 = FVec3::DotProduct(Axis1, Utilities::Multiply(InvI1, Axis1));
+		//const FVec3 DR0 = Axis0 * Angle0 * I0 / (I0 + I1);
+		//const FVec3 DR1 = Axis1 * Angle1 * I1 / (I0 + I1);
+		const FReal L = (FReal)1 / (FVec3::DotProduct(Axis0, Utilities::Multiply(InvI0, Axis0)) + FVec3::DotProduct(Axis1, Utilities::Multiply(InvI1, Axis1)));
+		const FVec3 DR0 = Utilities::Multiply(InvI0, Axis0) * L * Angle0;
+		const FVec3 DR1 = Utilities::Multiply(InvI1, Axis1) * L * Angle1;
 
 		ApplyRotationDelta(Dt, SolverSettings, JointSettings, Index0, Index1, Q0, Q1, DR0, DR1, Stiffness);
 	}
 
-	template<class T, int d>
-	void TPBDJointUtilities<T, d>::CalculateSwingConstraintSpace(
-		const TPBDJointSolverSettings<T, d>& SolverSettings,
-		const TPBDJointSettings<T, d>& JointSettings,
+	
+	void FPBDJointUtilities::CalculateSwingConstraintSpace(
+		const FPBDJointSolverSettings& SolverSettings,
+		const FPBDJointSettings& JointSettings,
 		const int32 Index0,
 		const int32 Index1,
-		TVector<T, d>& P0,
-		TRotation<T, d>& Q0,
-		TVector<T, d>& P1,
-		TRotation<T, d>& Q1,
-		TVector<T, d>& OutX0,
-		PMatrix<T, d, d>& OutR0,
-		TVector<T, d>& OutX1,
-		PMatrix<T, d, d>& OutR1,
-		TVector<T, d>& OutCR)
+		FVec3& P0,
+		FRotation3& Q0,
+		FVec3& P1,
+		FRotation3& Q1,
+		FVec3& OutX0,
+		FMatrix33& OutR0,
+		FVec3& OutX1,
+		FMatrix33& OutR1,
+		FVec3& OutCR)
 	{
-		const TRigidTransform<T, d>& XL0 = JointSettings.ConstraintFrames[Index0];
-		const TRigidTransform<T, d>& XL1 = JointSettings.ConstraintFrames[Index1];
-		const TVector<T, d> X0 = P0 + Q0 * XL0.GetTranslation();
-		const TVector<T, d> X1 = P1 + Q1 * XL1.GetTranslation();
-		const TRotation<T, d> R0 = Q0 * XL0.GetRotation();
-		const TRotation<T, d> R1 = Q1 * XL1.GetRotation();
-
-		TRotation<T, d> R01 = R0.Inverse() * R1;
-		TRotation<T, d> R01Twist, R01Swing;
-		R01.ToSwingTwist(TJointConstants<T, d>::TwistAxis(), R01Swing, R01Twist);
+		const FRigidTransform3& XL0 = JointSettings.ConstraintFrames[Index0];
+		const FRigidTransform3& XL1 = JointSettings.ConstraintFrames[Index1];
+		const FVec3 X0 = P0 + Q0 * XL0.GetTranslation();
+		const FVec3 X1 = P1 + Q1 * XL1.GetTranslation();
+		const FRotation3 R0 = Q0 * XL0.GetRotation();
+		const FRotation3 R1 = Q1 * XL1.GetRotation();
+
+		FRotation3 R01 = R0.Inverse() * R1;
+		FRotation3 R01Twist, R01Swing;
+		R01.ToSwingTwist(FJointConstants::TwistAxis(), R01Swing, R01Twist);
 		R01Swing = R01Swing.GetNormalized();
 		R01Twist = R01Twist.GetNormalized();
 
-		TVector<T, d> TwistAxis01;
-		T TwistAngle = (T)0;
-		R01Twist.ToAxisAndAngleSafe(TwistAxis01, TwistAngle, TJointConstants<T, d>::TwistAxis(), SolverSettings.SwingTwistAngleTolerance);
+		FVec3 TwistAxis01;
+		FReal TwistAngle = (FReal)0;
+		R01Twist.ToAxisAndAngleSafe(TwistAxis01, TwistAngle, FJointConstants::TwistAxis(), SolverSettings.SwingTwistAngleTolerance);
 		if (TwistAngle > PI)
 		{
-			TwistAngle = TwistAngle - (T)2 * PI;
+			TwistAngle = TwistAngle - (FReal)2 * PI;
 		}
-		if (TVector<T, d>::DotProduct(TwistAxis01, TJointConstants<T, d>::TwistAxis()) < 0)
+		if (FVec3::DotProduct(TwistAxis01, FJointConstants::TwistAxis()) < 0)
 		{
 			TwistAngle = -TwistAngle;
 		}
 
-		const PMatrix<T, d, d> Axes0 = R0.ToMatrix();
-		const PMatrix<T, d, d> Axes1 = R1.ToMatrix();
+		const FMatrix33 Axes0 = R0.ToMatrix();
+		const FMatrix33 Axes1 = R1.ToMatrix();
 
-		T Swing1Angle = (T)0;
-		const TVector<T, d> SwingCross1 = TVector<T, d>::CrossProduct(Axes0.GetAxis((int32)EJointAngularConstraintIndex::Swing1), Axes1.GetAxis((int32)EJointAngularConstraintIndex::Swing1));
-		const T SwingCross1Len = SwingCross1.Size();
+		FReal Swing1Angle = (FReal)0;
+		const FVec3 SwingCross1 = FVec3::CrossProduct(Axes0.GetAxis((int32)EJointAngularConstraintIndex::Swing1), Axes1.GetAxis((int32)EJointAngularConstraintIndex::Swing1));
+		const FReal SwingCross1Len = SwingCross1.Size();
 		if (SwingCross1Len > KINDA_SMALL_NUMBER)
 		{
-			Swing1Angle = FMath::Asin(FMath::Clamp(SwingCross1Len, (T)0, (T)1));
+			Swing1Angle = FMath::Asin(FMath::Clamp(SwingCross1Len, (FReal)0, (FReal)1));
 		}
-		const T Swing1Dot = TVector<T, d>::DotProduct(Axes0.GetAxis((int32)EJointAngularConstraintIndex::Swing1), Axes1.GetAxis((int32)EJointAngularConstraintIndex::Swing1));
-		if (Swing1Dot < (T)0)
+		const FReal Swing1Dot = FVec3::DotProduct(Axes0.GetAxis((int32)EJointAngularConstraintIndex::Swing1), Axes1.GetAxis((int32)EJointAngularConstraintIndex::Swing1));
+		if (Swing1Dot < (FReal)0)
 		{
-			Swing1Angle = (T)PI - Swing1Angle;
+			Swing1Angle = (FReal)PI - Swing1Angle;
 		}
 
-		T Swing2Angle = (T)0;
-		const TVector<T, d> SwingCross2 = TVector<T, d>::CrossProduct(Axes0.GetAxis((int32)EJointAngularConstraintIndex::Swing2), Axes1.GetAxis((int32)EJointAngularConstraintIndex::Swing2));
-		const T SwingCross2Len = SwingCross2.Size();
+		FReal Swing2Angle = (FReal)0;
+		const FVec3 SwingCross2 = FVec3::CrossProduct(Axes0.GetAxis((int32)EJointAngularConstraintIndex::Swing2), Axes1.GetAxis((int32)EJointAngularConstraintIndex::Swing2));
+		const FReal SwingCross2Len = SwingCross2.Size();
 		if (SwingCross2Len > KINDA_SMALL_NUMBER)
 		{
-			Swing2Angle = FMath::Asin(FMath::Clamp(SwingCross2Len, (T)0, (T)1));
+			Swing2Angle = FMath::Asin(FMath::Clamp(SwingCross2Len, (FReal)0, (FReal)1));
 		}
-		const T Swing2Dot = TVector<T, d>::DotProduct(Axes0.GetAxis((int32)EJointAngularConstraintIndex::Swing2), Axes1.GetAxis((int32)EJointAngularConstraintIndex::Swing2));
-		if (Swing2Dot < (T)0)
+		const FReal Swing2Dot = FVec3::DotProduct(Axes0.GetAxis((int32)EJointAngularConstraintIndex::Swing2), Axes1.GetAxis((int32)EJointAngularConstraintIndex::Swing2));
+		if (Swing2Dot < (FReal)0)
 		{
-			Swing2Angle = (T)PI - Swing2Angle;
+			Swing2Angle = (FReal)PI - Swing2Angle;
 		}
 
 		OutX0 = X0;
@@ -611,103 +611,103 @@ namespace Chaos
 		OutCR[(int32)EJointAngularAxisIndex::Swing2] = Swing2Angle;
 	}
 
-	template<class T, int d>
-	void TPBDJointUtilities<T, d>::CalculateConeConstraintSpace(
-		const TPBDJointSolverSettings<T, d>& SolverSettings,
-		const TPBDJointSettings<T, d>& JointSettings,
+	
+	void FPBDJointUtilities::CalculateConeConstraintSpace(
+		const FPBDJointSolverSettings& SolverSettings,
+		const FPBDJointSettings& JointSettings,
 		const int32 Index0,
 		const int32 Index1,
-		TVector<T, d>& P0,
-		TRotation<T, d>& Q0,
-		TVector<T, d>& P1,
-		TRotation<T, d>& Q1,
-		TVector<T, d>& OutX0,
-		PMatrix<T, d, d>& OutR0, 
-		TVector<T, d>& OutX1, 
-		PMatrix<T, d, d>& OutR1, 
-		TVector<T, d>& OutCR)
+		FVec3& P0,
+		FRotation3& Q0,
+		FVec3& P1,
+		FRotation3& Q1,
+		FVec3& OutX0,
+		FMatrix33& OutR0, 
+		FVec3& OutX1, 
+		FMatrix33& OutR1, 
+		FVec3& OutCR)
 	{
-		const TRigidTransform<T, d>& XL0 = JointSettings.ConstraintFrames[Index0];
-		const TRigidTransform<T, d>& XL1 = JointSettings.ConstraintFrames[Index1];
-		const TVector<T, d> X0 = P0 + Q0 * XL0.GetTranslation();
-		const TVector<T, d> X1 = P1 + Q1 * XL1.GetTranslation();
-		const TRotation<T, d> R0 = Q0 * XL0.GetRotation();
-		const TRotation<T, d> R1 = Q1 * XL1.GetRotation();
-
-		TRotation<T, d> R01 = R0.Inverse() * R1;
-		TRotation<T, d> R01Twist, R01Swing;
-		R01.ToSwingTwist(TJointConstants<T, d>::TwistAxis(), R01Swing, R01Twist);
+		const FRigidTransform3& XL0 = JointSettings.ConstraintFrames[Index0];
+		const FRigidTransform3& XL1 = JointSettings.ConstraintFrames[Index1];
+		const FVec3 X0 = P0 + Q0 * XL0.GetTranslation();
+		const FVec3 X1 = P1 + Q1 * XL1.GetTranslation();
+		const FRotation3 R0 = Q0 * XL0.GetRotation();
+		const FRotation3 R1 = Q1 * XL1.GetRotation();
+
+		FRotation3 R01 = R0.Inverse() * R1;
+		FRotation3 R01Twist, R01Swing;
+		R01.ToSwingTwist(FJointConstants::TwistAxis(), R01Swing, R01Twist);
 		R01Swing = R01Swing.GetNormalized();
 		R01Twist = R01Twist.GetNormalized();
 
-		TVector<T, d> TwistAxis01;
-		T TwistAngle;
-		R01Twist.ToAxisAndAngleSafe(TwistAxis01, TwistAngle, TJointConstants<T, d>::TwistAxis(), SolverSettings.SwingTwistAngleTolerance);
+		FVec3 TwistAxis01;
+		FReal TwistAngle;
+		R01Twist.ToAxisAndAngleSafe(TwistAxis01, TwistAngle, FJointConstants::TwistAxis(), SolverSettings.SwingTwistAngleTolerance);
 		if (TwistAngle > PI)
 		{
-			TwistAngle = TwistAngle - (T)2 * PI;
+			TwistAngle = TwistAngle - (FReal)2 * PI;
 		}
-		if (TVector<T, d>::DotProduct(TwistAxis01, TJointConstants<T, d>::TwistAxis()) < 0)
+		if (FVec3::DotProduct(TwistAxis01, FJointConstants::TwistAxis()) < 0)
 		{
 			TwistAxis01 = -TwistAxis01;
 			TwistAngle = -TwistAngle;
 		}
-		TVector<T, d> TwistAxis0 = R0 * TwistAxis01;
-		TVector<T, d> TwistAxis1 = R1 * TwistAxis01;
+		FVec3 TwistAxis0 = R0 * TwistAxis01;
+		FVec3 TwistAxis1 = R1 * TwistAxis01;
 
-		TVector<T, d> SwingAxis01;
-		T SwingAngle;
-		R01Swing.ToAxisAndAngleSafe(SwingAxis01, SwingAngle, TJointConstants<T, d>::Swing1Axis(), SolverSettings.SwingTwistAngleTolerance);
+		FVec3 SwingAxis01;
+		FReal SwingAngle;
+		R01Swing.ToAxisAndAngleSafe(SwingAxis01, SwingAngle, FJointConstants::Swing1Axis(), SolverSettings.SwingTwistAngleTolerance);
 		if (SwingAngle > PI)
 		{
-			SwingAngle = SwingAngle - (T)2 * PI;
+			SwingAngle = SwingAngle - (FReal)2 * PI;
 		}
-		TVector<T, d> SwingAxis0 = R0 * SwingAxis01;
-		TVector<T, d> SwingAxis1 = SwingAxis0;
+		FVec3 SwingAxis0 = R0 * SwingAxis01;
+		FVec3 SwingAxis1 = SwingAxis0;
 
 		OutX0 = X0;
 		OutX1 = X1;
 		OutR0.SetAxis((int32)EJointAngularAxisIndex::Twist, TwistAxis0);
 		OutR0.SetAxis((int32)EJointAngularAxisIndex::Swing1, SwingAxis0);
-		OutR0.SetAxis((int32)EJointAngularAxisIndex::Swing2, TVector<T, d>::CrossProduct(SwingAxis0, TwistAxis0));
+		OutR0.SetAxis((int32)EJointAngularAxisIndex::Swing2, FVec3::CrossProduct(SwingAxis0, TwistAxis0));
 		OutR1.SetAxis((int32)EJointAngularAxisIndex::Twist, TwistAxis1);
 		OutR1.SetAxis((int32)EJointAngularAxisIndex::Swing1, SwingAxis1);
-		OutR1.SetAxis((int32)EJointAngularAxisIndex::Swing2, TVector<T, d>::CrossProduct(SwingAxis1, TwistAxis1));
+		OutR1.SetAxis((int32)EJointAngularAxisIndex::Swing2, FVec3::CrossProduct(SwingAxis1, TwistAxis1));
 		OutCR[(int32)EJointAngularAxisIndex::Twist] = TwistAngle;
 		OutCR[(int32)EJointAngularAxisIndex::Swing1] = SwingAngle;
-		OutCR[(int32)EJointAngularAxisIndex::Swing2] = (T)0;
+		OutCR[(int32)EJointAngularAxisIndex::Swing2] = (FReal)0;
 	}
 
-	template<typename T, int d>
-	void TPBDJointUtilities<T, d>::ApplyJointPositionConstraint(
-		const T Dt,
-		const TPBDJointSolverSettings<T, d>& SolverSettings,
-		const TPBDJointSettings<T, d>& JointSettings,
+	
+	void FPBDJointUtilities::ApplyJointPositionConstraint(
+		const FReal Dt,
+		const FPBDJointSolverSettings& SolverSettings,
+		const FPBDJointSettings& JointSettings,
 		const int32 Index0,
 		const int32 Index1,
-		TVector<T, d>& P0,
-		TRotation<T, d>& Q0,
-		TVector<T, d>& P1,
-		TRotation<T, d>& Q1,
-		float InvM0,
-		const PMatrix<T, d, d>& InvIL0,
-		float InvM1,
-		const PMatrix<T, d, d>& InvIL1)
+		FVec3& P0,
+		FRotation3& Q0,
+		FVec3& P1,
+		FRotation3& Q1,
+		FReal InvM0,
+		const FMatrix33& InvIL0,
+		FReal InvM1,
+		const FMatrix33& InvIL1)
 	{
-		const TRigidTransform<T, d>& XL0 = JointSettings.ConstraintFrames[Index0];
-		const TRigidTransform<T, d>& XL1 = JointSettings.ConstraintFrames[Index1];
-		const TVector<T, d> X0 = P0 + Q0 * XL0.GetTranslation();
-		const TVector<T, d> X1 = P1 + Q1 * XL1.GetTranslation();
-		const TRotation<T, d> R0 = Q0 * XL0.GetRotation();
-		const PMatrix<T, d, d> InvI0 = Utilities::Multiply(Q0.ToMatrix(), Utilities::Multiply(InvIL0, Q0.ToMatrix().GetTransposed()));
-		const PMatrix<T, d, d> InvI1 = Utilities::Multiply(Q1.ToMatrix(), Utilities::Multiply(InvIL1, Q1.ToMatrix().GetTransposed()));
+		const FRigidTransform3& XL0 = JointSettings.ConstraintFrames[Index0];
+		const FRigidTransform3& XL1 = JointSettings.ConstraintFrames[Index1];
+		const FVec3 X0 = P0 + Q0 * XL0.GetTranslation();
+		const FVec3 X1 = P1 + Q1 * XL1.GetTranslation();
+		const FRotation3 R0 = Q0 * XL0.GetRotation();
+		const FMatrix33 InvI0 = Utilities::Multiply(Q0.ToMatrix(), Utilities::Multiply(InvIL0, Q0.ToMatrix().GetTransposed()));
+		const FMatrix33 InvI1 = Utilities::Multiply(Q1.ToMatrix(), Utilities::Multiply(InvIL1, Q1.ToMatrix().GetTransposed()));
 
 		// Calculate constraint error
-		const TVector<T, d> CX = GetLimitedPositionError(JointSettings, R0, X1 - X0);
+		const FVec3 CX = GetLimitedPositionError(JointSettings, R0, X1 - X0);
 
 		// Calculate constraint correction
-		PMatrix<T, d, d> M0 = PMatrix<T, d, d>(0, 0, 0);
-		PMatrix<T, d, d> M1 = PMatrix<T, d, d>(0, 0, 0);
+		FMatrix33 M0 = FMatrix33(0, 0, 0);
+		FMatrix33 M1 = FMatrix33(0, 0, 0);
 		if (InvM0 > 0)
 		{
 			M0 = Utilities::ComputeJointFactorMatrix(X0 - P0, InvI0, InvM0);
@@ -716,53 +716,53 @@ namespace Chaos
 		{
 			M1 = Utilities::ComputeJointFactorMatrix(X1 - P1, InvI1, InvM1);
 		}
-		const PMatrix<T, d, d> MI = (M0 + M1).Inverse();
-		const TVector<T, d> DX = Utilities::Multiply(MI, CX);
+		const FMatrix33 MI = (M0 + M1).Inverse();
+		const FVec3 DX = Utilities::Multiply(MI, CX);
 
 		// Apply constraint correction
-		const TVector<T, d> DP0 = InvM0 * DX;
-		const TVector<T, d> DP1 = -InvM1 * DX;
-		const TVector<T, d> DR0 = Utilities::Multiply(InvI0, TVector<T, d>::CrossProduct(X0 - P0, DX));
-		const TVector<T, d> DR1 = Utilities::Multiply(InvI1, TVector<T, d>::CrossProduct(X1 - P1, -DX));
+		const FVec3 DP0 = InvM0 * DX;
+		const FVec3 DP1 = -InvM1 * DX;
+		const FVec3 DR0 = Utilities::Multiply(InvI0, FVec3::CrossProduct(X0 - P0, DX));
+		const FVec3 DR1 = Utilities::Multiply(InvI1, FVec3::CrossProduct(X1 - P1, -DX));
 		ApplyPositionDelta(Dt, SolverSettings, JointSettings, Index0, Index1, P0, P1, DP0, DP1, GetLinearStiffness(SolverSettings, JointSettings));
 		ApplyRotationDelta(Dt, SolverSettings, JointSettings, Index0, Index1, Q0, Q1, DR0, DR1, GetLinearStiffness(SolverSettings, JointSettings));
 	}
 
-	template<typename T, int d>
-	void TPBDJointUtilities<T, d>::ApplyJointVelocityConstraint(
-		const T Dt,
-		const TPBDJointSolverSettings<T, d>& SolverSettings,
-		const TPBDJointSettings<T, d>& JointSettings,
+	
+	void FPBDJointUtilities::ApplyJointVelocityConstraint(
+		const FReal Dt,
+		const FPBDJointSolverSettings& SolverSettings,
+		const FPBDJointSettings& JointSettings,
 		const int32 Index0,
 		const int32 Index1,
-		TVector<T, d>& P0,
-		TRotation<T, d>& Q0,
-		TVector<T, d>& V0,
-		TVector<T, d>& W0,
-		TVector<T, d>& P1,
-		TRotation<T, d>& Q1,
-		TVector<T, d>& V1,
-		TVector<T, d>& W1,
-		float InvM0,
-		const PMatrix<T, d, d>& InvIL0,
-		float InvM1,
-		const PMatrix<T, d, d>& InvIL1)
+		FVec3& P0,
+		FRotation3& Q0,
+		FVec3& V0,
+		FVec3& W0,
+		FVec3& P1,
+		FRotation3& Q1,
+		FVec3& V1,
+		FVec3& W1,
+		FReal InvM0,
+		const FMatrix33& InvIL0,
+		FReal InvM1,
+		const FMatrix33& InvIL1)
 	{
-		const TRigidTransform<T, d>& XL0 = JointSettings.ConstraintFrames[Index0];
-		const TRigidTransform<T, d>& XL1 = JointSettings.ConstraintFrames[Index1];
-		const TVector<T, d> XC0 = Q0 * XL0.GetTranslation();
-		const TVector<T, d> XC1 = Q1 * XL1.GetTranslation();
-		const TRotation<T, d> R0 = Q0 * XL0.GetRotation();
-		const PMatrix<T, d, d> InvI0 = Utilities::Multiply(Q0.ToMatrix(), Utilities::Multiply(InvIL0, Q0.ToMatrix().GetTransposed()));
-		const PMatrix<T, d, d> InvI1 = Utilities::Multiply(Q1.ToMatrix(), Utilities::Multiply(InvIL1, Q1.ToMatrix().GetTransposed()));
-
-		const TVector<T, d> VC0 = V0 + TVector<T, d>::CrossProduct(W0, XC0);
-		const TVector<T, d> VC1 = V1 + TVector<T, d>::CrossProduct(W1, XC1);
-		const TVector<T, d> VC = GetLimitedVelocityError(JointSettings, R0, (P1 + XC1) - (P0 + XC0), VC1 - VC0);
+		const FRigidTransform3& XL0 = JointSettings.ConstraintFrames[Index0];
+		const FRigidTransform3& XL1 = JointSettings.ConstraintFrames[Index1];
+		const FVec3 XC0 = Q0 * XL0.GetTranslation();
+		const FVec3 XC1 = Q1 * XL1.GetTranslation();
+		const FRotation3 R0 = Q0 * XL0.GetRotation();
+		const FMatrix33 InvI0 = Utilities::Multiply(Q0.ToMatrix(), Utilities::Multiply(InvIL0, Q0.ToMatrix().GetTransposed()));
+		const FMatrix33 InvI1 = Utilities::Multiply(Q1.ToMatrix(), Utilities::Multiply(InvIL1, Q1.ToMatrix().GetTransposed()));
+
+		const FVec3 VC0 = V0 + FVec3::CrossProduct(W0, XC0);
+		const FVec3 VC1 = V1 + FVec3::CrossProduct(W1, XC1);
+		const FVec3 VC = GetLimitedVelocityError(JointSettings, R0, (P1 + XC1) - (P0 + XC0), VC1 - VC0);
 
 		// Calculate constraint correction
-		PMatrix<T, d, d> M0 = PMatrix<T, d, d>(0, 0, 0);
-		PMatrix<T, d, d> M1 = PMatrix<T, d, d>(0, 0, 0);
+		FMatrix33 M0 = FMatrix33(0, 0, 0);
+		FMatrix33 M1 = FMatrix33(0, 0, 0);
 		if (InvM0 > 0)
 		{
 			M0 = Utilities::ComputeJointFactorMatrix(XC0, InvI0, InvM0);
@@ -771,67 +771,67 @@ namespace Chaos
 		{
 			M1 = Utilities::ComputeJointFactorMatrix(XC1, InvI1, InvM1);
 		}
-		const PMatrix<T, d, d> MI = (M0 + M1).Inverse();
-		const TVector<T, d> DL = Utilities::Multiply(MI, VC);
+		const FMatrix33 MI = (M0 + M1).Inverse();
+		const FVec3 DL = Utilities::Multiply(MI, VC);
 
 		// Apply constraint correction
-		const TVector<T, d> DV0 = InvM0 * DL;
-		const TVector<T, d> DV1 = -InvM1 * DL;
-		const TVector<T, d> DW0 = Utilities::Multiply(InvI0, TVector<T, d>::CrossProduct(XC0, DL));
-		const TVector<T, d> DW1 = -Utilities::Multiply(InvI1, TVector<T, d>::CrossProduct(XC1, DL));
+		const FVec3 DV0 = InvM0 * DL;
+		const FVec3 DV1 = -InvM1 * DL;
+		const FVec3 DW0 = Utilities::Multiply(InvI0, FVec3::CrossProduct(XC0, DL));
+		const FVec3 DW1 = -Utilities::Multiply(InvI1, FVec3::CrossProduct(XC1, DL));
 
 		ApplyVelocityDelta(Dt, SolverSettings, JointSettings, Index0, Index1, P0, V0, P1, V1, DV0, DV1, GetLinearStiffness(SolverSettings, JointSettings));
 		ApplyRotationVelocityDelta(Dt, SolverSettings, JointSettings, Index0, Index1, Q0, W0, Q1, W1, DW0, DW1, GetLinearStiffness(SolverSettings, JointSettings));
 	}
 
 
-	template<typename T, int d>
-	void TPBDJointUtilities<T, d>::ApplyJointTwistConstraint(
-		const T Dt,
-		const TPBDJointSolverSettings<T, d>& SolverSettings,
-		const TPBDJointSettings<T, d>& JointSettings,
+	
+	void FPBDJointUtilities::ApplyJointTwistConstraint(
+		const FReal Dt,
+		const FPBDJointSolverSettings& SolverSettings,
+		const FPBDJointSettings& JointSettings,
 		const int32 Index0,
 		const int32 Index1,
-		TVector<T, d>& P0,
-		TRotation<T, d>& Q0,
-		TVector<T, d>& P1,
-		TRotation<T, d>& Q1,
-		float InvM0,
-		const PMatrix<T, d, d>& InvIL0,
-		float InvM1,
-		const PMatrix<T, d, d>& InvIL1)
+		FVec3& P0,
+		FRotation3& Q0,
+		FVec3& P1,
+		FRotation3& Q1,
+		FReal InvM0,
+		const FMatrix33& InvIL0,
+		FReal InvM1,
+		const FMatrix33& InvIL1)
 	{
-		const TRigidTransform<T, d>& XL0 = JointSettings.ConstraintFrames[Index0];
-		const TRigidTransform<T, d>& XL1 = JointSettings.ConstraintFrames[Index1];
-		const TVector<T, d> X0 = P0 + Q0 * XL0.GetTranslation();
-		const TVector<T, d> X1 = P1 + Q1 * XL1.GetTranslation();
-		const TRotation<T, d> R0 = Q0 * XL0.GetRotation();
-		const TRotation<T, d> R1 = Q1 * XL1.GetRotation();
-		const TVector<T, d> CX = GetLimitedPositionError(JointSettings, R0, X1 - X0);
+		const FRigidTransform3& XL0 = JointSettings.ConstraintFrames[Index0];
+		const FRigidTransform3& XL1 = JointSettings.ConstraintFrames[Index1];
+		const FVec3 X0 = P0 + Q0 * XL0.GetTranslation();
+		const FVec3 X1 = P1 + Q1 * XL1.GetTranslation();
+		const FRotation3 R0 = Q0 * XL0.GetRotation();
+		const FRotation3 R1 = Q1 * XL1.GetRotation();
+		const FVec3 CX = GetLimitedPositionError(JointSettings, R0, X1 - X0);
 
 		// Calculate the Twist Axis and Angle for each body
-		const TRotation<T, d> R01 = R0.Inverse() * R1;
-		TRotation<T, d> R01Twist, R01Swing;
-		R01.ToSwingTwist(TJointConstants<T, d>::TwistAxis(), R01Swing, R01Twist);
+		const FRotation3 R01 = R0.Inverse() * R1;
+		FRotation3 R01Twist, R01Swing;
+		R01.ToSwingTwist(FJointConstants::TwistAxis(), R01Swing, R01Twist);
 		R01Swing = R01Swing.GetNormalized();
 		R01Twist = R01Twist.GetNormalized();
 
-		TVector<T, d> TwistAxis01;
-		T TwistAngle;
-		R01Twist.ToAxisAndAngleSafe(TwistAxis01, TwistAngle, TJointConstants<T, d>::TwistAxis(), SolverSettings.SwingTwistAngleTolerance);
+		FVec3 TwistAxis01;
+		FReal TwistAngle;
+		R01Twist.ToAxisAndAngleSafe(TwistAxis01, TwistAngle, FJointConstants::TwistAxis(), SolverSettings.SwingTwistAngleTolerance);
 		if (TwistAngle > PI)
 		{
-			TwistAngle = TwistAngle - (T)2 * PI;
+			TwistAngle = TwistAngle - (FReal)2 * PI;
 		}
-		if (TVector<T, d>::DotProduct(TwistAxis01, TJointConstants<T, d>::TwistAxis()) < 0)
+		if (FVec3::DotProduct(TwistAxis01, FJointConstants::TwistAxis()) < 0)
 		{
 			TwistAxis01 = -TwistAxis01;
 			TwistAngle = -TwistAngle;
 		}
 
-		const TVector<T, d> TwistAxis0 = R0 * TwistAxis01;
-		const TVector<T, d> TwistAxis1 = R1 * TwistAxis01;
-		T TwistAngleMax = FLT_MAX;
+		const FVec3 TwistAxis0 = R0 * TwistAxis01;
+		const FVec3 TwistAxis1 = R1 * TwistAxis01;
+		FReal TwistAngleMax = FLT_MAX;
 		if (JointSettings.Motion.AngularMotionTypes[(int32)EJointAngularConstraintIndex::Twist] == EJointMotionType::Limited)
 		{
 			TwistAngleMax = JointSettings.Motion.AngularLimits[(int32)EJointAngularConstraintIndex::Twist];
@@ -842,7 +842,7 @@ namespace Chaos
 		}
 
 		// Calculate the twist correction to apply to each body
-		T DTwistAngle = 0;
+		FReal DTwistAngle = 0;
 		if (TwistAngle > TwistAngleMax)
 		{
 			DTwistAngle = TwistAngle - TwistAngleMax;
@@ -851,65 +851,65 @@ namespace Chaos
 		{
 			DTwistAngle = TwistAngle + TwistAngleMax;
 		}
-		const T DTwistAngle0 = DTwistAngle;
-		const T DTwistAngle1 = -DTwistAngle;
+		const FReal DTwistAngle0 = DTwistAngle;
+		const FReal DTwistAngle1 = -DTwistAngle;
 
 		// Apply twist correction
 		ApplyRotationDelta(Dt, SolverSettings, JointSettings, Index0, Index1, P0, Q0, P1, Q1, InvM0, InvIL0, InvM1, InvIL1, TwistAxis0, DTwistAngle0, TwistAxis1, DTwistAngle1, GetTwistStiffness(SolverSettings, JointSettings));
 	}
 
-	template<typename T, int d>
-	void TPBDJointUtilities<T, d>::ApplyJointTwistVelocityConstraint(
-		const T Dt,
-		const TPBDJointSolverSettings<T, d>& SolverSettings,
-		const TPBDJointSettings<T, d>& JointSettings,
+	
+	void FPBDJointUtilities::ApplyJointTwistVelocityConstraint(
+		const FReal Dt,
+		const FPBDJointSolverSettings& SolverSettings,
+		const FPBDJointSettings& JointSettings,
 		const int32 Index0,
 		const int32 Index1,
-		TVector<T, d>& P0,
-		TRotation<T, d>& Q0,
-		TVector<T, d>& V0,
-		TVector<T, d>& W0,
-		TVector<T, d>& P1,
-		TRotation<T, d>& Q1,
-		TVector<T, d>& V1,
-		TVector<T, d>& W1,
-		float InvM0,
-		const PMatrix<T, d, d>& InvIL0,
-		float InvM1,
-		const PMatrix<T, d, d>& InvIL1)
+		FVec3& P0,
+		FRotation3& Q0,
+		FVec3& V0,
+		FVec3& W0,
+		FVec3& P1,
+		FRotation3& Q1,
+		FVec3& V1,
+		FVec3& W1,
+		FReal InvM0,
+		const FMatrix33& InvIL0,
+		FReal InvM1,
+		const FMatrix33& InvIL1)
 	{
-		const TRigidTransform<T, d>& XL0 = JointSettings.ConstraintFrames[Index0];
-		const TRigidTransform<T, d>& XL1 = JointSettings.ConstraintFrames[Index1];
-		const TVector<T, d> X0 = P0 + Q0 * XL0.GetTranslation();
-		const TVector<T, d> X1 = P1 + Q1 * XL1.GetTranslation();
-		const TRotation<T, d> R0 = Q0 * XL0.GetRotation();
-		const TRotation<T, d> R1 = Q1 * XL1.GetRotation();
-		const TVector<T, d> CX = GetLimitedPositionError(JointSettings, R0, X1 - X0);
+		const FRigidTransform3& XL0 = JointSettings.ConstraintFrames[Index0];
+		const FRigidTransform3& XL1 = JointSettings.ConstraintFrames[Index1];
+		const FVec3 X0 = P0 + Q0 * XL0.GetTranslation();
+		const FVec3 X1 = P1 + Q1 * XL1.GetTranslation();
+		const FRotation3 R0 = Q0 * XL0.GetRotation();
+		const FRotation3 R1 = Q1 * XL1.GetRotation();
+		const FVec3 CX = GetLimitedPositionError(JointSettings, R0, X1 - X0);
 
 		// Calculate the Twist Axis and Angle for each body
-		const TRotation<T, d> R01 = R0.Inverse() * R1;
-		TRotation<T, d> R01Twist, R01Swing;
-		R01.ToSwingTwist(TJointConstants<T, d>::TwistAxis(), R01Swing, R01Twist);
+		const FRotation3 R01 = R0.Inverse() * R1;
+		FRotation3 R01Twist, R01Swing;
+		R01.ToSwingTwist(FJointConstants::TwistAxis(), R01Swing, R01Twist);
 		R01Swing = R01Swing.GetNormalized();
 		R01Twist = R01Twist.GetNormalized();
 
-		TVector<T, d> TwistAxis01;
-		T TwistAngle;
-		R01Twist.ToAxisAndAngleSafe(TwistAxis01, TwistAngle, TJointConstants<T, d>::TwistAxis(), SolverSettings.SwingTwistAngleTolerance);
+		FVec3 TwistAxis01;
+		FReal TwistAngle;
+		R01Twist.ToAxisAndAngleSafe(TwistAxis01, TwistAngle, FJointConstants::TwistAxis(), SolverSettings.SwingTwistAngleTolerance);
 		if (TwistAngle > PI)
 		{
-			TwistAngle = TwistAngle - (T)2 * PI;
+			TwistAngle = TwistAngle - (FReal)2 * PI;
 		}
-		if (TVector<T, d>::DotProduct(TwistAxis01, TJointConstants<T, d>::TwistAxis()) < 0)
+		if (FVec3::DotProduct(TwistAxis01, FJointConstants::TwistAxis()) < 0)
 		{
 			TwistAxis01 = -TwistAxis01;
 			TwistAngle = -TwistAngle;
 		}
 
-		const TVector<T, d> TwistAxis0 = R0 * TwistAxis01;
-		const TVector<T, d> TwistAxis1 = R1 * TwistAxis01;
+		const FVec3 TwistAxis0 = R0 * TwistAxis01;
+		const FVec3 TwistAxis1 = R1 * TwistAxis01;
 
-		T TwistAngleMax = FLT_MAX;
+		FReal TwistAngleMax = FLT_MAX;
 		if (JointSettings.Motion.AngularMotionTypes[(int32)EJointAngularConstraintIndex::Twist] == EJointMotionType::Limited)
 		{
 			TwistAngleMax = JointSettings.Motion.AngularLimits[(int32)EJointAngularConstraintIndex::Twist];
@@ -919,68 +919,68 @@ namespace Chaos
 			TwistAngleMax = 0;
 		}
 
-		const T WC0 = TVector<T, d>::DotProduct(W0, TwistAxis0);
-		const T WC1 = TVector<T, d>::DotProduct(W1, TwistAxis1);
-		T DW = (T)0;
+		const FReal WC0 = FVec3::DotProduct(W0, TwistAxis0);
+		const FReal WC1 = FVec3::DotProduct(W1, TwistAxis1);
+		FReal DW = (FReal)0;
 		if (TwistAngle > TwistAngleMax)
 		{
-			DW = FMath::Max((T)0, WC1 - WC0);
+			DW = FMath::Max((FReal)0, WC1 - WC0);
 		}
 		else if (TwistAngle < -TwistAngleMax)
 		{
-			DW = FMath::Min((T)0, WC1 - WC0);
+			DW = FMath::Min((FReal)0, WC1 - WC0);
 		}
 
 		ApplyRotationVelocityDelta(Dt, SolverSettings, JointSettings, Index0, Index1, Q0, W0, Q1, W1, InvM0, InvIL0, InvM1, InvIL1, TwistAxis0, TwistAxis1, DW, GetTwistStiffness(SolverSettings, JointSettings));
 	}
 
 
-	template<typename T, int d>
-	void TPBDJointUtilities<T, d>::ApplyJointConeConstraint(
-		const T Dt,
-		const TPBDJointSolverSettings<T, d>& SolverSettings,
-		const TPBDJointSettings<T, d>& JointSettings,
+	
+	void FPBDJointUtilities::ApplyJointConeConstraint(
+		const FReal Dt,
+		const FPBDJointSolverSettings& SolverSettings,
+		const FPBDJointSettings& JointSettings,
 		const int32 Index0,
 		const int32 Index1,
-		TVector<T, d>& P0,
-		TRotation<T, d>& Q0,
-		TVector<T, d>& P1,
-		TRotation<T, d>& Q1,
-		float InvM0,
-		const PMatrix<T, d, d>& InvIL0,
-		float InvM1,
-		const PMatrix<T, d, d>& InvIL1)
+		FVec3& P0,
+		FRotation3& Q0,
+		FVec3& P1,
+		FRotation3& Q1,
+		FReal InvM0,
+		const FMatrix33& InvIL0,
+		FReal InvM1,
+		const FMatrix33& InvIL1)
 	{
-		const TRigidTransform<T, d>& XL0 = JointSettings.ConstraintFrames[Index0];
-		const TRigidTransform<T, d>& XL1 = JointSettings.ConstraintFrames[Index1];
-		const TVector<T, d> X0 = P0 + Q0 * XL0.GetTranslation();
-		const TVector<T, d> X1 = P1 + Q1 * XL1.GetTranslation();
-		const TRotation<T, d> R0 = Q0 * XL0.GetRotation();
-		const TRotation<T, d> R1 = Q1 * XL1.GetRotation();
-		const TVector<T, d> CX = GetLimitedPositionError(JointSettings, R0, X1 - X0);
+		const FRigidTransform3& XL0 = JointSettings.ConstraintFrames[Index0];
+		const FRigidTransform3& XL1 = JointSettings.ConstraintFrames[Index1];
+		const FVec3 X0 = P0 + Q0 * XL0.GetTranslation();
+		const FVec3 X1 = P1 + Q1 * XL1.GetTranslation();
+		const FRotation3 R0 = Q0 * XL0.GetRotation();
+		const FRotation3 R1 = Q1 * XL1.GetRotation();
+		const FVec3 CX = GetLimitedPositionError(JointSettings, R0, X1 - X0);
 
 		// Calculate Swing axis for each body
-		const TRotation<T, d> R01 = R0.Inverse() * R1;
-		TRotation<T, d> R01Twist, R01Swing;
-		R01.ToSwingTwist(TJointConstants<T, d>::TwistAxis(), R01Swing, R01Twist);
+		const FRotation3 R01 = R0.Inverse() * R1;
+		FRotation3 R01Twist, R01Swing;
+		R01.ToSwingTwist(FJointConstants::TwistAxis(), R01Swing, R01Twist);
 		R01Swing = R01Swing.GetNormalized();
 		R01Twist = R01Twist.GetNormalized();
 
-		TVector<T, d> SwingAxis01;
-		T SwingAngle;
-		R01Swing.ToAxisAndAngleSafe(SwingAxis01, SwingAngle, TJointConstants<T, d>::Swing1Axis(), SolverSettings.SwingTwistAngleTolerance);
+		FVec3 SwingAxis01;
+		FReal SwingAngle;
+		R01Swing.ToAxisAndAngleSafe(SwingAxis01, SwingAngle, FJointConstants::Swing1Axis(), SolverSettings.SwingTwistAngleTolerance);
 		if (SwingAngle > PI)
 		{
-			SwingAngle = SwingAngle - (T)2 * PI;
+			SwingAngle = SwingAngle - (FReal)2 * PI;
 		}
 
-		const TVector<T, d> SwingAxis0 = R0 * SwingAxis01;
-		const TVector<T, d> SwingAxis1 = SwingAxis0;
+		const FVec3 SwingAxis0 = R0 * SwingAxis01;
+		const FVec3 SwingAxis1 = SwingAxis0;
 
 		// Calculate swing limit for the current swing axis
-		T SwingAngleMax = FLT_MAX;
-		const T Swing1Limit = JointSettings.Motion.AngularLimits[(int32)EJointAngularConstraintIndex::Swing1];
-		const T Swing2Limit = JointSettings.Motion.AngularLimits[(int32)EJointAngularConstraintIndex::Swing2];
+		FReal SwingAngleMax = FLT_MAX;
+		const FReal Swing1Limit = JointSettings.Motion.AngularLimits[(int32)EJointAngularConstraintIndex::Swing1];
+		const FReal Swing2Limit = JointSettings.Motion.AngularLimits[(int32)EJointAngularConstraintIndex::Swing2];
 
 		// Circular swing limit
 		SwingAngleMax = Swing1Limit;
@@ -989,13 +989,13 @@ namespace Chaos
 		if (!FMath::IsNearlyEqual(Swing1Limit, Swing2Limit, KINDA_SMALL_NUMBER))
 		{
 			// Map swing axis to ellipse and calculate limit for this swing axis
-			const T DotSwing1 = FMath::Abs(TVector<T, d>::DotProduct(SwingAxis01, TJointConstants<T, d>::Swing1Axis()));
-			const T DotSwing2 = FMath::Abs(TVector<T, d>::DotProduct(SwingAxis01, TJointConstants<T, d>::Swing2Axis()));
+			const FReal DotSwing1 = FMath::Abs(FVec3::DotProduct(SwingAxis01, FJointConstants::Swing1Axis()));
+			const FReal DotSwing2 = FMath::Abs(FVec3::DotProduct(SwingAxis01, FJointConstants::Swing2Axis()));
 			SwingAngleMax = FMath::Sqrt(Swing1Limit * DotSwing2 * Swing1Limit * DotSwing2 + Swing2Limit * DotSwing1 * Swing2Limit * DotSwing1);
 		}
 
 		// Calculate swing error we need to correct
-		T DSwingAngle = 0;
+		FReal DSwingAngle = 0;
 		if (SwingAngle > SwingAngleMax)
 		{
 			DSwingAngle = SwingAngle - SwingAngleMax;
@@ -1004,62 +1004,62 @@ namespace Chaos
 		{
 			DSwingAngle = SwingAngle + SwingAngleMax;
 		}
-		T DSwingAngle0 = DSwingAngle;
-		T DSwingAngle1 = -DSwingAngle;
+		FReal DSwingAngle0 = DSwingAngle;
+		FReal DSwingAngle1 = -DSwingAngle;
 
 		// Apply swing correction
 		ApplyRotationDelta(Dt, SolverSettings, JointSettings, Index0, Index1, P0, Q0, P1, Q1, InvM0, InvIL0, InvM1, InvIL1, SwingAxis0, DSwingAngle0, SwingAxis1, DSwingAngle1, GetSwingStiffness(SolverSettings, JointSettings));
 	}
 
-	template<typename T, int d>
-	void TPBDJointUtilities<T, d>::ApplyJointConeVelocityConstraint(
-		const T Dt,
-		const TPBDJointSolverSettings<T, d>& SolverSettings,
-		const TPBDJointSettings<T, d>& JointSettings,
+	
+	void FPBDJointUtilities::ApplyJointConeVelocityConstraint(
+		const FReal Dt,
+		const FPBDJointSolverSettings& SolverSettings,
+		const FPBDJointSettings& JointSettings,
 		const int32 Index0,
 		const int32 Index1,
-		TVector<T, d>& P0,
-		TRotation<T, d>& Q0,
-		TVector<T, d>& V0,
-		TVector<T, d>& W0,
-		TVector<T, d>& P1,
-		TRotation<T, d>& Q1,
-		TVector<T, d>& V1,
-		TVector<T, d>& W1,
-		float InvM0,
-		const PMatrix<T, d, d>& InvIL0,
-		float InvM1,
-		const PMatrix<T, d, d>& InvIL1)
+		FVec3& P0,
+		FRotation3& Q0,
+		FVec3& V0,
+		FVec3& W0,
+		FVec3& P1,
+		FRotation3& Q1,
+		FVec3& V1,
+		FVec3& W1,
+		FReal InvM0,
+		const FMatrix33& InvIL0,
+		FReal InvM1,
+		const FMatrix33& InvIL1)
 	{
-		const TRigidTransform<T, d>& XL0 = JointSettings.ConstraintFrames[Index0];
-		const TRigidTransform<T, d>& XL1 = JointSettings.ConstraintFrames[Index1];
-		const TVector<T, d> X0 = P0 + Q0 * XL0.GetTranslation();
-		const TVector<T, d> X1 = P1 + Q1 * XL1.GetTranslation();
-		const TRotation<T, d> R0 = Q0 * XL0.GetRotation();
-		const TRotation<T, d> R1 = Q1 * XL1.GetRotation();
+		const FRigidTransform3& XL0 = JointSettings.ConstraintFrames[Index0];
+		const FRigidTransform3& XL1 = JointSettings.ConstraintFrames[Index1];
+		const FVec3 X0 = P0 + Q0 * XL0.GetTranslation();
+		const FVec3 X1 = P1 + Q1 * XL1.GetTranslation();
+		const FRotation3 R0 = Q0 * XL0.GetRotation();
+		const FRotation3 R1 = Q1 * XL1.GetRotation();
 
 		// Calculate Swing axis for each body
-		const TRotation<T, d> R01 = R0.Inverse() * R1;
-		TRotation<T, d> R01Twist, R01Swing;
-		R01.ToSwingTwist(TJointConstants<T, d>::TwistAxis(), R01Swing, R01Twist);
+		const FRotation3 R01 = R0.Inverse() * R1;
+		FRotation3 R01Twist, R01Swing;
+		R01.ToSwingTwist(FJointConstants::TwistAxis(), R01Swing, R01Twist);
 		R01Swing = R01Swing.GetNormalized();
 		R01Twist = R01Twist.GetNormalized();
 
-		TVector<T, d> SwingAxis01;
-		T SwingAngle;
-		R01Swing.ToAxisAndAngleSafe(SwingAxis01, SwingAngle, TJointConstants<T, d>::Swing1Axis(), SolverSettings.SwingTwistAngleTolerance);
+		FVec3 SwingAxis01;
+		FReal SwingAngle;
+		R01Swing.ToAxisAndAngleSafe(SwingAxis01, SwingAngle, FJointConstants::Swing1Axis(), SolverSettings.SwingTwistAngleTolerance);
 		if (SwingAngle > PI)
 		{
-			SwingAngle = SwingAngle - (T)2 * PI;
+			SwingAngle = SwingAngle - (FReal)2 * PI;
 		}
 
-		const TVector<T, d> SwingAxis0 = R0 * SwingAxis01;
-		const TVector<T, d> SwingAxis1 = SwingAxis0;
+		const FVec3 SwingAxis0 = R0 * SwingAxis01;
+		const FVec3 SwingAxis1 = SwingAxis0;
 
 		// Calculate swing limit for the current swing axis
-		T SwingAngleMax = FLT_MAX;
-		const T Swing1Limit = JointSettings.Motion.AngularLimits[(int32)EJointAngularConstraintIndex::Swing1];
-		const T Swing2Limit = JointSettings.Motion.AngularLimits[(int32)EJointAngularConstraintIndex::Swing2];
+		FReal SwingAngleMax = FLT_MAX;
+		const FReal Swing1Limit = JointSettings.Motion.AngularLimits[(int32)EJointAngularConstraintIndex::Swing1];
+		const FReal Swing2Limit = JointSettings.Motion.AngularLimits[(int32)EJointAngularConstraintIndex::Swing2];
 
 		// Circular swing limit
 		SwingAngleMax = Swing1Limit;
@@ -1068,98 +1068,98 @@ namespace Chaos
 		if (!FMath::IsNearlyEqual(Swing1Limit, Swing2Limit, KINDA_SMALL_NUMBER))
 		{
 			// Map swing axis to ellipse and calculate limit for this swing axis
-			const T DotSwing1 = FMath::Abs(TVector<T, d>::DotProduct(SwingAxis01, TJointConstants<T, d>::Swing1Axis()));
-			const T DotSwing2 = FMath::Abs(TVector<T, d>::DotProduct(SwingAxis01, TJointConstants<T, d>::Swing2Axis()));
+			const FReal DotSwing1 = FMath::Abs(FVec3::DotProduct(SwingAxis01, FJointConstants::Swing1Axis()));
+			const FReal DotSwing2 = FMath::Abs(FVec3::DotProduct(SwingAxis01, FJointConstants::Swing2Axis()));
 			SwingAngleMax = FMath::Sqrt(Swing1Limit * DotSwing2 * Swing1Limit * DotSwing2 + Swing2Limit * DotSwing1 * Swing2Limit * DotSwing1);
 		}
 
 		// Only clamp veloicity if we are outside the limits and moving to increase the error
-		const T WC0 = TVector<T, d>::DotProduct(W0, SwingAxis0);
-		const T WC1 = TVector<T, d>::DotProduct(W1, SwingAxis1);
-		T DW = 0;
+		const FReal WC0 = FVec3::DotProduct(W0, SwingAxis0);
+		const FReal WC1 = FVec3::DotProduct(W1, SwingAxis1);
+		FReal DW = 0;
 		if (SwingAngle > SwingAngleMax)
 		{
-			DW = FMath::Max((T)0, WC1 - WC0);
+			DW = FMath::Max((FReal)0, WC1 - WC0);
 		}
 		else if (SwingAngle < -SwingAngleMax)
 		{
-			DW = FMath::Min((T)0, WC1 - WC0);
+			DW = FMath::Min((FReal)0, WC1 - WC0);
 		}
 
 		ApplyRotationVelocityDelta(Dt, SolverSettings, JointSettings, Index0, Index1, Q0, W0, Q1, W1, InvM0, InvIL0, InvM1, InvIL1, SwingAxis0, SwingAxis1, DW, GetSwingStiffness(SolverSettings, JointSettings));
 	}
 
 
-	template<typename T, int d>
-	void TPBDJointUtilities<T, d>::ApplyJointSwingConstraint(
-		const T Dt,
-		const TPBDJointSolverSettings<T, d>& SolverSettings,
-		const TPBDJointSettings<T, d>& JointSettings,
+	
+	void FPBDJointUtilities::ApplyJointSwingConstraint(
+		const FReal Dt,
+		const FPBDJointSolverSettings& SolverSettings,
+		const FPBDJointSettings& JointSettings,
 		const int32 Index0,
 		const int32 Index1,
 		const EJointAngularConstraintIndex SwingConstraintIndex,
 		const EJointAngularAxisIndex SwingAxisIndex,
-		TVector<T, d>& P0,
-		TRotation<T, d>& Q0,
-		TVector<T, d>& P1,
-		TRotation<T, d>& Q1,
-		float InvM0,
-		const PMatrix<T, d, d>& InvIL0,
-		float InvM1,
-		const PMatrix<T, d, d>& InvIL1)
+		FVec3& P0,
+		FRotation3& Q0,
+		FVec3& P1,
+		FRotation3& Q1,
+		FReal InvM0,
+		const FMatrix33& InvIL0,
+		FReal InvM1,
+		const FMatrix33& InvIL1)
 	{
-		const TRigidTransform<T, d>& XL0 = JointSettings.ConstraintFrames[Index0];
-		const TRigidTransform<T, d>& XL1 = JointSettings.ConstraintFrames[Index1];
-		const TVector<T, d> X0 = P0 + Q0 * XL0.GetTranslation();
-		const TVector<T, d> X1 = P1 + Q1 * XL1.GetTranslation();
-		const TRotation<T, d> R0 = Q0 * XL0.GetRotation();
-		const TRotation<T, d> R1 = Q1 * XL1.GetRotation();
-		const TVector<T, d> CX = GetLimitedPositionError(JointSettings, R0, X1 - X0);
+		const FRigidTransform3& XL0 = JointSettings.ConstraintFrames[Index0];
+		const FRigidTransform3& XL1 = JointSettings.ConstraintFrames[Index1];
+		const FVec3 X0 = P0 + Q0 * XL0.GetTranslation();
+		const FVec3 X1 = P1 + Q1 * XL1.GetTranslation();
+		const FRotation3 R0 = Q0 * XL0.GetRotation();
+		const FRotation3 R1 = Q1 * XL1.GetRotation();
+		const FVec3 CX = GetLimitedPositionError(JointSettings, R0, X1 - X0);
 
 		// Calculate the swing axis for each body
-		const TRotation<T, d> R01 = R0.Inverse() * R1;
-		TRotation<T, d> R01Twist, R01Swing;
-		R01.ToSwingTwist(TJointConstants<T, d>::TwistAxis(), R01Swing, R01Twist);
+		const FRotation3 R01 = R0.Inverse() * R1;
+		FRotation3 R01Twist, R01Swing;
+		R01.ToSwingTwist(FJointConstants::TwistAxis(), R01Swing, R01Twist);
 		R01Swing = R01Swing.GetNormalized();
 		R01Twist = R01Twist.GetNormalized();
 
-		TVector<T, d> TwistAxis01;
-		T TwistAngle;
-		R01Twist.ToAxisAndAngleSafe(TwistAxis01, TwistAngle, TJointConstants<T, d>::TwistAxis(), SolverSettings.SwingTwistAngleTolerance);
+		FVec3 TwistAxis01;
+		FReal TwistAngle;
+		R01Twist.ToAxisAndAngleSafe(TwistAxis01, TwistAngle, FJointConstants::TwistAxis(), SolverSettings.SwingTwistAngleTolerance);
 		if (TwistAngle > PI)
 		{
-			TwistAngle = TwistAngle - (T)2 * PI;
+			TwistAngle = TwistAngle - (FReal)2 * PI;
 		}
-		if (TVector<T, d>::DotProduct(TwistAxis01, TJointConstants<T, d>::TwistAxis()) < 0)
+		if (FVec3::DotProduct(TwistAxis01, FJointConstants::TwistAxis()) < 0)
 		{
 			TwistAxis01 = -TwistAxis01;
 			TwistAngle = -TwistAngle;
 		}
-		const TVector<T, d> TwistAxis = R0 * TwistAxis01;
+		const FVec3 TwistAxis = R0 * TwistAxis01;
 
-		const TRotation<T, d> R1NoTwist = R1 * R01Twist.Inverse();
-		const PMatrix<T, d, d> Axes0 = R0.ToMatrix();
-		const PMatrix<T, d, d> Axes1 = R1NoTwist.ToMatrix();
-		TVector<T, d> SwingCross = TVector<T, d>::CrossProduct(Axes0.GetAxis((int32)SwingAxisIndex), Axes1.GetAxis((int32)SwingAxisIndex));
-		SwingCross = SwingCross - TVector<T, d>::DotProduct(TwistAxis, SwingCross) * TwistAxis;
-		const T SwingCrossLen = SwingCross.Size();
+		const FRotation3 R1NoTwist = R1 * R01Twist.Inverse();
+		const FMatrix33 Axes0 = R0.ToMatrix();
+		const FMatrix33 Axes1 = R1NoTwist.ToMatrix();
+		FVec3 SwingCross = FVec3::CrossProduct(Axes0.GetAxis((int32)SwingAxisIndex), Axes1.GetAxis((int32)SwingAxisIndex));
+		SwingCross = SwingCross - FVec3::DotProduct(TwistAxis, SwingCross) * TwistAxis;
+		const FReal SwingCrossLen = SwingCross.Size();
 		if (SwingCrossLen > KINDA_SMALL_NUMBER)
 		{
-			const TVector<T, d> SwingAxis = SwingCross / SwingCrossLen;
-			TVector<T, d> SwingAxis0 = SwingAxis;
-			TVector<T, d> SwingAxis1 = SwingAxis;
+			const FVec3 SwingAxis = SwingCross / SwingCrossLen;
+			FVec3 SwingAxis0 = SwingAxis;
+			FVec3 SwingAxis1 = SwingAxis;
 
-			T SwingAngle = FMath::Asin(FMath::Clamp(SwingCrossLen, (T)0, (T)1));
-			const T SwingDot = TVector<T, d>::DotProduct(Axes0.GetAxis((int32)SwingAxisIndex), Axes1.GetAxis((int32)SwingAxisIndex));
-			if (SwingDot < (T)0)
+			FReal SwingAngle = FMath::Asin(FMath::Clamp(SwingCrossLen, (FReal)0, (FReal)1));
+			const FReal SwingDot = FVec3::DotProduct(Axes0.GetAxis((int32)SwingAxisIndex), Axes1.GetAxis((int32)SwingAxisIndex));
+			if (SwingDot < (FReal)0)
 			{
-				SwingAngle = (T)PI - SwingAngle;
+				SwingAngle = (FReal)PI - SwingAngle;
 			}
 
-			T SwingAngleMax = FLT_MAX;
+			FReal SwingAngleMax = FLT_MAX;
 			if (JointSettings.Motion.AngularMotionTypes[(int32)SwingConstraintIndex] == EJointMotionType::Limited)
 			{
-				T SwingLimit = JointSettings.Motion.AngularLimits[(int32)SwingConstraintIndex];
+				FReal SwingLimit = JointSettings.Motion.AngularLimits[(int32)SwingConstraintIndex];
 				SwingAngleMax = SwingLimit;
 			}
 			else if (JointSettings.Motion.AngularMotionTypes[(int32)SwingConstraintIndex] == EJointMotionType::Locked)
@@ -1168,7 +1168,7 @@ namespace Chaos
 			}
 
 			// Calculate swing error we need to correct
-			T DSwingAngle = 0;
+			FReal DSwingAngle = 0;
 			if (SwingAngle > SwingAngleMax)
 			{
 				DSwingAngle = SwingAngle - SwingAngleMax;
@@ -1177,87 +1177,87 @@ namespace Chaos
 			{
 				DSwingAngle = SwingAngle + SwingAngleMax;
 			}
-			T DSwingAngle0 = DSwingAngle;
-			T DSwingAngle1 = -DSwingAngle;
+			FReal DSwingAngle0 = DSwingAngle;
+			FReal DSwingAngle1 = -DSwingAngle;
 
 			// Apply swing correction
 			ApplyRotationDelta(Dt, SolverSettings, JointSettings, Index0, Index1, P0, Q0, P1, Q1, InvM0, InvIL0, InvM1, InvIL1, SwingAxis0, DSwingAngle0, SwingAxis1, DSwingAngle1, GetSwingStiffness(SolverSettings, JointSettings));
 		}
 	}
 
-	template<typename T, int d>
-	void TPBDJointUtilities<T, d>::ApplyJointSwingVelocityConstraint(
-		const T Dt,
-		const TPBDJointSolverSettings<T, d>& SolverSettings,
-		const TPBDJointSettings<T, d>& JointSettings,
+	
+	void FPBDJointUtilities::ApplyJointSwingVelocityConstraint(
+		const FReal Dt,
+		const FPBDJointSolverSettings& SolverSettings,
+		const FPBDJointSettings& JointSettings,
 		const int32 Index0,
 		const int32 Index1,
 		const EJointAngularConstraintIndex SwingConstraintIndex,
 		const EJointAngularAxisIndex SwingAxisIndex,
-		TVector<T, d>& P0,
-		TRotation<T, d>& Q0,
-		TVector<T, d>& V0,
-		TVector<T, d>& W0,
-		TVector<T, d>& P1,
-		TRotation<T, d>& Q1,
-		TVector<T, d>& V1,
-		TVector<T, d>& W1,
-		float InvM0,
-		const PMatrix<T, d, d>& InvIL0,
-		float InvM1,
-		const PMatrix<T, d, d>& InvIL1)
+		FVec3& P0,
+		FRotation3& Q0,
+		FVec3& V0,
+		FVec3& W0,
+		FVec3& P1,
+		FRotation3& Q1,
+		FVec3& V1,
+		FVec3& W1,
+		FReal InvM0,
+		const FMatrix33& InvIL0,
+		FReal InvM1,
+		const FMatrix33& InvIL1)
 	{
-		const TRigidTransform<T, d>& XL0 = JointSettings.ConstraintFrames[Index0];
-		const TRigidTransform<T, d>& XL1 = JointSettings.ConstraintFrames[Index1];
-		const TVector<T, d> X0 = P0 + Q0 * XL0.GetTranslation();
-		const TVector<T, d> X1 = P1 + Q1 * XL1.GetTranslation();
-		const TRotation<T, d> R0 = Q0 * XL0.GetRotation();
-		const TRotation<T, d> R1 = Q1 * XL1.GetRotation();
+		const FRigidTransform3& XL0 = JointSettings.ConstraintFrames[Index0];
+		const FRigidTransform3& XL1 = JointSettings.ConstraintFrames[Index1];
+		const FVec3 X0 = P0 + Q0 * XL0.GetTranslation();
+		const FVec3 X1 = P1 + Q1 * XL1.GetTranslation();
+		const FRotation3 R0 = Q0 * XL0.GetRotation();
+		const FRotation3 R1 = Q1 * XL1.GetRotation();
 
 		// Calculate the swing axis for each body
-		const TRotation<T, d> R01 = R0.Inverse() * R1;
-		TRotation<T, d> R01Twist, R01Swing;
-		R01.ToSwingTwist(TJointConstants<T, d>::TwistAxis(), R01Swing, R01Twist);
+		const FRotation3 R01 = R0.Inverse() * R1;
+		FRotation3 R01Twist, R01Swing;
+		R01.ToSwingTwist(FJointConstants::TwistAxis(), R01Swing, R01Twist);
 		R01Swing = R01Swing.GetNormalized();
 		R01Twist = R01Twist.GetNormalized();
 
-		TVector<T, d> TwistAxis01;
-		T TwistAngle;
-		R01Twist.ToAxisAndAngleSafe(TwistAxis01, TwistAngle, TJointConstants<T, d>::TwistAxis(), SolverSettings.SwingTwistAngleTolerance);
+		FVec3 TwistAxis01;
+		FReal TwistAngle;
+		R01Twist.ToAxisAndAngleSafe(TwistAxis01, TwistAngle, FJointConstants::TwistAxis(), SolverSettings.SwingTwistAngleTolerance);
 		if (TwistAngle > PI)
 		{
-			TwistAngle = TwistAngle - (T)2 * PI;
+			TwistAngle = TwistAngle - (FReal)2 * PI;
 		}
-		if (TVector<T, d>::DotProduct(TwistAxis01, TJointConstants<T, d>::TwistAxis()) < 0)
+		if (FVec3::DotProduct(TwistAxis01, FJointConstants::TwistAxis()) < 0)
 		{
 			TwistAxis01 = -TwistAxis01;
 			TwistAngle = -TwistAngle;
 		}
-		const TVector<T, d> TwistAxis = R0 * TwistAxis01;
+		const FVec3 TwistAxis = R0 * TwistAxis01;
 
-		const TRotation<T, d> R1NoTwist = R1 * R01Twist.Inverse();
-		const PMatrix<T, d, d> Axes0 = R0.ToMatrix();
-		const PMatrix<T, d, d> Axes1 = R1NoTwist.ToMatrix();
-		TVector<T, d> SwingCross = TVector<T, d>::CrossProduct(Axes0.GetAxis((int32)SwingAxisIndex), Axes1.GetAxis((int32)SwingAxisIndex));
-		SwingCross = SwingCross - TVector<T, d>::DotProduct(TwistAxis, SwingCross) * TwistAxis;
-		const T SwingCrossLen = SwingCross.Size();
+		const FRotation3 R1NoTwist = R1 * R01Twist.Inverse();
+		const FMatrix33 Axes0 = R0.ToMatrix();
+		const FMatrix33 Axes1 = R1NoTwist.ToMatrix();
+		FVec3 SwingCross = FVec3::CrossProduct(Axes0.GetAxis((int32)SwingAxisIndex), Axes1.GetAxis((int32)SwingAxisIndex));
+		SwingCross = SwingCross - FVec3::DotProduct(TwistAxis, SwingCross) * TwistAxis;
+		const FReal SwingCrossLen = SwingCross.Size();
 		if (SwingCrossLen > KINDA_SMALL_NUMBER)
 		{
-			const TVector<T, d> SwingAxis = SwingCross / SwingCrossLen;
-			TVector<T, d> SwingAxis0 = SwingAxis;
-			TVector<T, d> SwingAxis1 = SwingAxis;
+			const FVec3 SwingAxis = SwingCross / SwingCrossLen;
+			FVec3 SwingAxis0 = SwingAxis;
+			FVec3 SwingAxis1 = SwingAxis;
 
-			T SwingAngle = FMath::Asin(FMath::Clamp(SwingCrossLen, (T)0, (T)1));
-			const T SwingDot = TVector<T, d>::DotProduct(Axes0.GetAxis((int32)SwingAxisIndex), Axes1.GetAxis((int32)SwingAxisIndex));
-			if (SwingDot < (T)0)
+			FReal SwingAngle = FMath::Asin(FMath::Clamp(SwingCrossLen, (FReal)0, (FReal)1));
+			const FReal SwingDot = FVec3::DotProduct(Axes0.GetAxis((int32)SwingAxisIndex), Axes1.GetAxis((int32)SwingAxisIndex));
+			if (SwingDot < (FReal)0)
 			{
-				SwingAngle = (T)PI - SwingAngle;
+				SwingAngle = (FReal)PI - SwingAngle;
 			}
 
-			T SwingAngleMax = FLT_MAX;
+			FReal SwingAngleMax = FLT_MAX;
 			if (JointSettings.Motion.AngularMotionTypes[(int32)SwingConstraintIndex] == EJointMotionType::Limited)
 			{
-				T Swing1Limit = JointSettings.Motion.AngularLimits[(int32)SwingConstraintIndex];
+				FReal Swing1Limit = JointSettings.Motion.AngularLimits[(int32)SwingConstraintIndex];
 				SwingAngleMax = Swing1Limit;
 			}
 			else if (JointSettings.Motion.AngularMotionTypes[(int32)SwingConstraintIndex] == EJointMotionType::Locked)
@@ -1266,278 +1266,273 @@ namespace Chaos
 			}
 
 			// Only clamp veloicity if we are outside the limits and moving to increase the error
-			const T WC0 = TVector<T, d>::DotProduct(W0, SwingAxis0);
-			const T WC1 = TVector<T, d>::DotProduct(W1, SwingAxis1);
-			T DW = 0;
+			const FReal WC0 = FVec3::DotProduct(W0, SwingAxis0);
+			const FReal WC1 = FVec3::DotProduct(W1, SwingAxis1);
+			FReal DW = 0;
 			if (SwingAngle > SwingAngleMax)
 			{
-				DW = FMath::Max((T)0, WC1 - WC0);
+				DW = FMath::Max((FReal)0, WC1 - WC0);
 			}
 			else if (SwingAngle < -SwingAngleMax)
 			{
-				DW = FMath::Min((T)0, WC1 - WC0);
+				DW = FMath::Min((FReal)0, WC1 - WC0);
 			}
 
 			ApplyRotationVelocityDelta(Dt, SolverSettings, JointSettings, Index0, Index1, Q0, W0, Q1, W1, InvM0, InvIL0, InvM1, InvIL1, SwingAxis0, SwingAxis1, DW, GetSwingStiffness(SolverSettings, JointSettings));
 		}
 	}
 
-	template<typename T, int d>
-	void TPBDJointUtilities<T, d>::ApplyJointTwistDrive(
-		const T Dt,
-		const TPBDJointSolverSettings<T, d>& SolverSettings,
-		const TPBDJointSettings<T, d>& JointSettings,
+	
+	void FPBDJointUtilities::ApplyJointTwistDrive(
+		const FReal Dt,
+		const FPBDJointSolverSettings& SolverSettings,
+		const FPBDJointSettings& JointSettings,
 		const int32 Index0,
 		const int32 Index1,
-		TVector<T, d>& P0,
-		TRotation<T, d>& Q0,
-		TVector<T, d>& P1,
-		TRotation<T, d>& Q1,
-		float InvM0,
-		const PMatrix<T, d, d>& InvIL0,
-		float InvM1,
-		const PMatrix<T, d, d>& InvIL1)
+		FVec3& P0,
+		FRotation3& Q0,
+		FVec3& P1,
+		FRotation3& Q1,
+		FReal InvM0,
+		const FMatrix33& InvIL0,
+		FReal InvM1,
+		const FMatrix33& InvIL1)
 	{
-		const TRigidTransform<T, d>& XL0 = JointSettings.ConstraintFrames[Index0];
-		const TRigidTransform<T, d>& XL1 = JointSettings.ConstraintFrames[Index1];
-		const TVector<T, d> X0 = P0 + Q0 * XL0.GetTranslation();
-		const TVector<T, d> X1 = P1 + Q1 * XL1.GetTranslation();
-		const TRotation<T, d> R0 = Q0 * XL0.GetRotation();
-		const TRotation<T, d> R1 = Q1 * XL1.GetRotation();
-		const TVector<T, d> CX = GetLimitedPositionError(JointSettings, R0, X1 - X0);
-
-		const TRotation<T, d> R01 = R0.Inverse() * R1;
-		TRotation<T, d> R01Twist, R01Swing;
-		R01.ToSwingTwist(TJointConstants<T, d>::TwistAxis(), R01Swing, R01Twist);
+		const FRigidTransform3& XL0 = JointSettings.ConstraintFrames[Index0];
+		const FRigidTransform3& XL1 = JointSettings.ConstraintFrames[Index1];
+		const FVec3 X0 = P0 + Q0 * XL0.GetTranslation();
+		const FVec3 X1 = P1 + Q1 * XL1.GetTranslation();
+		const FRotation3 R0 = Q0 * XL0.GetRotation();
+		const FRotation3 R1 = Q1 * XL1.GetRotation();
+		const FVec3 CX = GetLimitedPositionError(JointSettings, R0, X1 - X0);
+
+		const FRotation3 R01 = R0.Inverse() * R1;
+		FRotation3 R01Twist, R01Swing;
+		R01.ToSwingTwist(FJointConstants::TwistAxis(), R01Swing, R01Twist);
 		R01Swing = R01Swing.GetNormalized();
 		R01Twist = R01Twist.GetNormalized();
 
-		TVector<T, d> TwistAxis01;
-		T TwistAngle;
-		R01Twist.ToAxisAndAngleSafe(TwistAxis01, TwistAngle, TJointConstants<T, d>::TwistAxis(), SolverSettings.SwingTwistAngleTolerance);
+		FVec3 TwistAxis01;
+		FReal TwistAngle;
+		R01Twist.ToAxisAndAngleSafe(TwistAxis01, TwistAngle, FJointConstants::TwistAxis(), SolverSettings.SwingTwistAngleTolerance);
 		if (TwistAngle > PI)
 		{
-			TwistAngle = TwistAngle - (T)2 * PI;
+			TwistAngle = TwistAngle - (FReal)2 * PI;
 		}
-		if (TVector<T, d>::DotProduct(TwistAxis01, TJointConstants<T, d>::TwistAxis()) < 0)
+		if (FVec3::DotProduct(TwistAxis01, FJointConstants::TwistAxis()) < 0)
 		{
 			TwistAxis01 = -TwistAxis01;
 			TwistAngle = -TwistAngle;
 		}
 
-		const TVector<T, d> TwistAxis0 = R0 * TwistAxis01;
-		const TVector<T, d> TwistAxis1 = R1 * TwistAxis01;
-		const T TwistAngleTarget = JointSettings.Motion.AngularDriveTargetAngles[(int32)EJointAngularConstraintIndex::Twist];
-		const T DriveStiffnessUnclamped = (SolverSettings.DriveStiffness > 0) ? SolverSettings.DriveStiffness : JointSettings.Motion.AngularDriveStiffness;
-		const T DriveStiffness = FMath::Clamp(DriveStiffnessUnclamped, (T)0, (T)1);
-		const T DTwistAngle = TwistAngle - TwistAngleTarget;
-		const T DTwistAngle0 = DTwistAngle;
-		const T DTwistAngle1 = -DTwistAngle;
+		const FVec3 TwistAxis0 = R0 * TwistAxis01;
+		const FVec3 TwistAxis1 = R1 * TwistAxis01;
+		const FReal TwistAngleTarget = JointSettings.Motion.AngularDriveTargetAngles[(int32)EJointAngularConstraintIndex::Twist];
+		const FReal DriveStiffnessUnclamped = (SolverSettings.DriveStiffness > 0) ? SolverSettings.DriveStiffness : JointSettings.Motion.AngularDriveStiffness;
+		const FReal DriveStiffness = FMath::Clamp(DriveStiffnessUnclamped, (FReal)0, (FReal)1);
+		const FReal DTwistAngle = TwistAngle - TwistAngleTarget;
+		const FReal DTwistAngle0 = DTwistAngle;
+		const FReal DTwistAngle1 = -DTwistAngle;
 
 		ApplyRotationDelta(Dt, SolverSettings, JointSettings, Index0, Index1, P0, Q0, P1, Q1, InvM0, InvIL0, InvM1, InvIL1, TwistAxis0, DTwistAngle0, TwistAxis1, DTwistAngle1, DriveStiffness);
 	}
 
-	template<typename T, int d>
-	void TPBDJointUtilities<T, d>::ApplyJointConeDrive(
-		const T Dt,
-		const TPBDJointSolverSettings<T, d>& SolverSettings,
-		const TPBDJointSettings<T, d>& JointSettings,
+	
+	void FPBDJointUtilities::ApplyJointConeDrive(
+		const FReal Dt,
+		const FPBDJointSolverSettings& SolverSettings,
+		const FPBDJointSettings& JointSettings,
 		const int32 Index0,
 		const int32 Index1,
-		TVector<T, d>& P0,
-		TRotation<T, d>& Q0,
-		TVector<T, d>& P1,
-		TRotation<T, d>& Q1,
-		float InvM0,
-		const PMatrix<T, d, d>& InvIL0,
-		float InvM1,
-		const PMatrix<T, d, d>& InvIL1)
+		FVec3& P0,
+		FRotation3& Q0,
+		FVec3& P1,
+		FRotation3& Q1,
+		FReal InvM0,
+		const FMatrix33& InvIL0,
+		FReal InvM1,
+		const FMatrix33& InvIL1)
 	{
-		const TRigidTransform<T, d>& XL0 = JointSettings.ConstraintFrames[Index0];
-		const TRigidTransform<T, d>& XL1 = JointSettings.ConstraintFrames[Index1];
-		const TVector<T, d> X0 = P0 + Q0 * XL0.GetTranslation();
-		const TVector<T, d> X1 = P1 + Q1 * XL1.GetTranslation();
-		const TRotation<T, d> R0 = Q0 * XL0.GetRotation();
-		const TRotation<T, d> R1 = Q1 * XL1.GetRotation();
-		const TVector<T, d> CX = GetLimitedPositionError(JointSettings, R0, X1 - X0);
+		const FRigidTransform3& XL0 = JointSettings.ConstraintFrames[Index0];
+		const FRigidTransform3& XL1 = JointSettings.ConstraintFrames[Index1];
+		const FVec3 X0 = P0 + Q0 * XL0.GetTranslation();
+		const FVec3 X1 = P1 + Q1 * XL1.GetTranslation();
+		const FRotation3 R0 = Q0 * XL0.GetRotation();
+		const FRotation3 R1 = Q1 * XL1.GetRotation();
+		const FVec3 CX = GetLimitedPositionError(JointSettings, R0, X1 - X0);
 
 		// Calculate Swing axis for each body
-		const TRotation<T, d> R01 = R0.Inverse() * R1;
-		TRotation<T, d> R01Twist, R01Swing;
-		R01.ToSwingTwist(TJointConstants<T, d>::TwistAxis(), R01Swing, R01Twist);
+		const FRotation3 R01 = R0.Inverse() * R1;
+		FRotation3 R01Twist, R01Swing;
+		R01.ToSwingTwist(FJointConstants::TwistAxis(), R01Swing, R01Twist);
 		R01Swing = R01Swing.GetNormalized();
 		R01Twist = R01Twist.GetNormalized();
 
-		TVector<T, d> SwingAxis01;
-		T SwingAngle;
-		R01Swing.ToAxisAndAngleSafe(SwingAxis01, SwingAngle, TJointConstants<T, d>::Swing1Axis(), SolverSettings.SwingTwistAngleTolerance);
+		FVec3 SwingAxis01;
+		FReal SwingAngle;
+		R01Swing.ToAxisAndAngleSafe(SwingAxis01, SwingAngle, FJointConstants::Swing1Axis(), SolverSettings.SwingTwistAngleTolerance);
 		if (SwingAngle > PI)
 		{
-			SwingAngle = SwingAngle - (T)2 * PI;
+			SwingAngle = SwingAngle - (FReal)2 * PI;
 		}
 
-		const TVector<T, d> SwingAxis0 = R0 * SwingAxis01;
-		const TVector<T, d> SwingAxis1 = SwingAxis0;
+		const FVec3 SwingAxis0 = R0 * SwingAxis01;
+		const FVec3 SwingAxis1 = SwingAxis0;
 
 		// Circular swing target (max of Swing1, Swing2 targets)
-		T Swing1Target = JointSettings.Motion.AngularDriveTargetAngles[(int32)EJointAngularConstraintIndex::Swing1];
-		T Swing2Target = JointSettings.Motion.AngularDriveTargetAngles[(int32)EJointAngularConstraintIndex::Swing2];
-		T SwingAngleTarget = FMath::Max(Swing1Target, Swing2Target);
+		FReal Swing1Target = JointSettings.Motion.AngularDriveTargetAngles[(int32)EJointAngularConstraintIndex::Swing1];
+		FReal Swing2Target = JointSettings.Motion.AngularDriveTargetAngles[(int32)EJointAngularConstraintIndex::Swing2];
+		FReal SwingAngleTarget = FMath::Max(Swing1Target, Swing2Target);
 
-		T DriveStiffnessUnclamped = (SolverSettings.DriveStiffness > 0) ? SolverSettings.DriveStiffness : JointSettings.Motion.AngularDriveStiffness;
-		T DriveStiffness = FMath::Clamp(DriveStiffnessUnclamped, (T)0, (T)1);
-		T DSwingAngle = SwingAngle - SwingAngleTarget;
-		T DSwingAngle0 = DSwingAngle;
-		T DSwingAngle1 = -DSwingAngle;
+		FReal DriveStiffnessUnclamped = (SolverSettings.DriveStiffness > 0) ? SolverSettings.DriveStiffness : JointSettings.Motion.AngularDriveStiffness;
+		FReal DriveStiffness = FMath::Clamp(DriveStiffnessUnclamped, (FReal)0, (FReal)1);
+		FReal DSwingAngle = SwingAngle - SwingAngleTarget;
+		FReal DSwingAngle0 = DSwingAngle;
+		FReal DSwingAngle1 = -DSwingAngle;
 
 		// Apply swing correction
 		ApplyRotationDelta(Dt, SolverSettings, JointSettings, Index0, Index1, P0, Q0, P1, Q1, InvM0, InvIL0, InvM1, InvIL1, SwingAxis0, DSwingAngle0, SwingAxis1, DSwingAngle1, DriveStiffness);
 	}
 
-	template<typename T, int d>
-	void TPBDJointUtilities<T, d>::ApplyJointSLerpDrive(
-		const T Dt,
-		const TPBDJointSolverSettings<T, d>& SolverSettings,
-		const TPBDJointSettings<T, d>& JointSettings,
+	
+	void FPBDJointUtilities::ApplyJointSLerpDrive(
+		const FReal Dt,
+		const FPBDJointSolverSettings& SolverSettings,
+		const FPBDJointSettings& JointSettings,
 		const int32 Index0,
 		const int32 Index1,
-		TVector<T, d>& P0,
-		TRotation<T, d>& Q0,
-		TVector<T, d>& P1,
-		TRotation<T, d>& Q1,
-		float InvM0,
-		const PMatrix<T, d, d>& InvIL0,
-		float InvM1,
-		const PMatrix<T, d, d>& InvIL1)
+		FVec3& P0,
+		FRotation3& Q0,
+		FVec3& P1,
+		FRotation3& Q1,
+		FReal InvM0,
+		const FMatrix33& InvIL0,
+		FReal InvM1,
+		const FMatrix33& InvIL1)
 	{
-		const TRigidTransform<T, d>& XL0 = JointSettings.ConstraintFrames[Index0];
-		const TRigidTransform<T, d>& XL1 = JointSettings.ConstraintFrames[Index1];
-		const TVector<T, d> X0 = P0 + Q0 * XL0.GetTranslation();
-		const TVector<T, d> X1 = P1 + Q1 * XL1.GetTranslation();
-		const TRotation<T, d> R0 = Q0 * XL0.GetRotation();
-		const TRotation<T, d> R1 = Q1 * XL1.GetRotation();
-		const TVector<T, d> CX = GetLimitedPositionError(JointSettings, R0, X1 - X0);
+		const FRigidTransform3& XL0 = JointSettings.ConstraintFrames[Index0];
+		const FRigidTransform3& XL1 = JointSettings.ConstraintFrames[Index1];
+		const FVec3 X0 = P0 + Q0 * XL0.GetTranslation();
+		const FVec3 X1 = P1 + Q1 * XL1.GetTranslation();
+		const FRotation3 R0 = Q0 * XL0.GetRotation();
+		const FRotation3 R1 = Q1 * XL1.GetRotation();
+		const FVec3 CX = GetLimitedPositionError(JointSettings, R0, X1 - X0);
 
 		// Calculate the rotation we need to apply to resolve the rotation delta
-		const TRotation<T, d> TargetR1 = R0 * JointSettings.Motion.AngularDriveTarget;
-		const TRotation<T, d> DR1 = TargetR1 * R1.Inverse();
-		const TRotation<T, d> TargetQ0 = DR1.Inverse() * Q0;
-		const TRotation<T, d> TargetQ1 = DR1 * Q1;
+		const FRotation3 TargetR1 = R0 * JointSettings.Motion.AngularDriveTarget;
+		const FRotation3 DR1 = TargetR1 * R1.Inverse();
+		const FRotation3 TargetQ0 = DR1.Inverse() * Q0;
+		const FRotation3 TargetQ1 = DR1 * Q1;
 
-		T DriveStiffnessUnclamped = (SolverSettings.DriveStiffness > 0) ? SolverSettings.DriveStiffness : JointSettings.Motion.AngularDriveStiffness;
-		T DriveStiffness = FMath::Clamp(DriveStiffnessUnclamped, (T)0, (T)1);
+		FReal DriveStiffnessUnclamped = (SolverSettings.DriveStiffness > 0) ? SolverSettings.DriveStiffness : JointSettings.Motion.AngularDriveStiffness;
+		FReal DriveStiffness = FMath::Clamp(DriveStiffnessUnclamped, (FReal)0, (FReal)1);
 
-		TVector<T, d> SLerpAxis;
-		T SLerpAngle;
-		if (DR1.ToAxisAndAngleSafe(SLerpAxis, SLerpAngle, TVector<T, d>(1, 0, 0)))
+		FVec3 SLerpAxis;
+		FReal SLerpAngle;
+		if (DR1.ToAxisAndAngleSafe(SLerpAxis, SLerpAngle, FVec3(1, 0, 0)))
 		{
-			const PMatrix<T, d, d> InvI0 = Utilities::Multiply(Q0.ToMatrix(), Utilities::Multiply(InvIL0, Q0.ToMatrix().GetTransposed()));
-			const PMatrix<T, d, d> InvI1 = Utilities::Multiply(Q1.ToMatrix(), Utilities::Multiply(InvIL1, Q1.ToMatrix().GetTransposed()));
-			const T I0 = TVector<T, d>::DotProduct(SLerpAxis, Utilities::Multiply(InvI0, SLerpAxis));
-			const T I1 = TVector<T, d>::DotProduct(SLerpAxis, Utilities::Multiply(InvI1, SLerpAxis));
-			const T F0 = DriveStiffness * I0 / (I0 + I1);
-			const T F1 = DriveStiffness * I1 / (I0 + I1);
+			const FMatrix33 InvI0 = Utilities::Multiply(Q0.ToMatrix(), Utilities::Multiply(InvIL0, Q0.ToMatrix().GetTransposed()));
+			const FMatrix33 InvI1 = Utilities::Multiply(Q1.ToMatrix(), Utilities::Multiply(InvIL1, Q1.ToMatrix().GetTransposed()));
+			const FReal I0 = FVec3::DotProduct(SLerpAxis, Utilities::Multiply(InvI0, SLerpAxis));
+			const FReal I1 = FVec3::DotProduct(SLerpAxis, Utilities::Multiply(InvI1, SLerpAxis));
+			const FReal F0 = DriveStiffness * I0 / (I0 + I1);
+			const FReal F1 = DriveStiffness * I1 / (I0 + I1);
 
 			// Apply the rotation delta
-			Q0 = TRotation<T, d>::Slerp(Q0, TargetQ0, F0);
-			Q1 = TRotation<T, d>::Slerp(Q1, TargetQ1, F1);
+			Q0 = FRotation3::Slerp(Q0, TargetQ0, F0);
+			Q1 = FRotation3::Slerp(Q1, TargetQ1, F1);
 			Q1.EnforceShortestArcWith(Q0);
 		}
 	}
 
-	template<typename T, int d>
-	void TPBDJointUtilities<T, d>::ApplyJointPositionProjection(
-		const T Dt,
-		const TPBDJointSolverSettings<T, d>& SolverSettings,
-		const TPBDJointSettings<T, d>& JointSettings,
+	
+	void FPBDJointUtilities::ApplyJointPositionProjection(
+		const FReal Dt,
+		const FPBDJointSolverSettings& SolverSettings,
+		const FPBDJointSettings& JointSettings,
 		const int32 Index0,
 		const int32 Index1,
-		TVector<T, d>& P0,
-		TRotation<T, d>& Q0,
-		TVector<T, d>& P1,
-		TRotation<T, d>& Q1,
-		float InvM0,
-		const PMatrix<T, d, d>& InvIL0,
-		float InvM1,
-		const PMatrix<T, d, d>& InvIL1,
-		const T ProjectionFactor)
+		FVec3& P0,
+		FRotation3& Q0,
+		FVec3& P1,
+		FRotation3& Q1,
+		FReal InvM0,
+		const FMatrix33& InvIL0,
+		FReal InvM1,
+		const FMatrix33& InvIL1,
+		const FReal ProjectionFactor)
 	{
-		const TRigidTransform<T, d>& XL0 = JointSettings.ConstraintFrames[Index0];
-		const TRigidTransform<T, d>& XL1 = JointSettings.ConstraintFrames[Index1];
-		const TVector<T, d> X0 = P0 + Q0 * XL0.GetTranslation();
-		const TVector<T, d> X1 = P1 + Q1 * XL1.GetTranslation();
-		const TRotation<T, d> R0 = Q0 * XL0.GetRotation();
-		const TVector<T, d> CX = GetLimitedPositionError(JointSettings, R0, X1 - X0);
-
-		const TVector<T, d> DP0 = CX * (ProjectionFactor * InvM0 / (InvM0 + InvM1));
-		const TVector<T, d> DP1 = -CX * (ProjectionFactor * InvM1 / (InvM0 + InvM1));
+		const FRigidTransform3& XL0 = JointSettings.ConstraintFrames[Index0];
+		const FRigidTransform3& XL1 = JointSettings.ConstraintFrames[Index1];
+		const FVec3 X0 = P0 + Q0 * XL0.GetTranslation();
+		const FVec3 X1 = P1 + Q1 * XL1.GetTranslation();
+		const FRotation3 R0 = Q0 * XL0.GetRotation();
+		const FVec3 CX = GetLimitedPositionError(JointSettings, R0, X1 - X0);
+
+		const FVec3 DP0 = CX * (ProjectionFactor * InvM0 / (InvM0 + InvM1));
+		const FVec3 DP1 = -CX * (ProjectionFactor * InvM1 / (InvM0 + InvM1));
 		P0 = P0 + DP0;
 		P1 = P1 + DP1;
 	}
 
-	template<typename T, int d>
-	void TPBDJointUtilities<T, d>::ApplyJointTwistProjection(
-		const T Dt,
-		const TPBDJointSolverSettings<T, d>& SolverSettings,
-		const TPBDJointSettings<T, d>& JointSettings,
+	
+	void FPBDJointUtilities::ApplyJointTwistProjection(
+		const FReal Dt,
+		const FPBDJointSolverSettings& SolverSettings,
+		const FPBDJointSettings& JointSettings,
 		const int32 Index0,
 		const int32 Index1,
-		TVector<T, d>& P0,
-		TRotation<T, d>& Q0,
-		TVector<T, d>& P1,
-		TRotation<T, d>& Q1,
-		float InvM0,
-		const PMatrix<T, d, d>& InvIL0,
-		float InvM1,
-		const PMatrix<T, d, d>& InvIL1,
-		const T ProjectionFactor)
+		FVec3& P0,
+		FRotation3& Q0,
+		FVec3& P1,
+		FRotation3& Q1,
+		FReal InvM0,
+		const FMatrix33& InvIL0,
+		FReal InvM1,
+		const FMatrix33& InvIL1,
+		const FReal ProjectionFactor)
 	{
 
 	}
 
-	template<typename T, int d>
-	void TPBDJointUtilities<T, d>::ApplyJointConeProjection(
-		const T Dt,
-		const TPBDJointSolverSettings<T, d>& SolverSettings,
-		const TPBDJointSettings<T, d>& JointSettings,
+	
+	void FPBDJointUtilities::ApplyJointConeProjection(
+		const FReal Dt,
+		const FPBDJointSolverSettings& SolverSettings,
+		const FPBDJointSettings& JointSettings,
 		const int32 Index0,
 		const int32 Index1,
-		TVector<T, d>& P0,
-		TRotation<T, d>& Q0,
-		TVector<T, d>& P1,
-		TRotation<T, d>& Q1,
-		float InvM0,
-		const PMatrix<T, d, d>& InvIL0,
-		float InvM1,
-		const PMatrix<T, d, d>& InvIL1,
-		const T ProjectionFactor)
+		FVec3& P0,
+		FRotation3& Q0,
+		FVec3& P1,
+		FRotation3& Q1,
+		FReal InvM0,
+		const FMatrix33& InvIL0,
+		FReal InvM1,
+		const FMatrix33& InvIL1,
+		const FReal ProjectionFactor)
 	{
 
 	}
 
-	template<typename T, int d>
-	void TPBDJointUtilities<T, d>::ApplyJointSwingProjection(
-		const T Dt,
-		const TPBDJointSolverSettings<T, d>& SolverSettings,
-		const TPBDJointSettings<T, d>& JointSettings,
+	
+	void FPBDJointUtilities::ApplyJointSwingProjection(
+		const FReal Dt,
+		const FPBDJointSolverSettings& SolverSettings,
+		const FPBDJointSettings& JointSettings,
 		const int32 Index0,
 		const int32 Index1,
 		const EJointAngularConstraintIndex SwingConstraint,
-		TVector<T, d>& P0,
-		TRotation<T, d>& Q0,
-		TVector<T, d>& P1,
-		TRotation<T, d>& Q1,
-		float InvM0,
-		const PMatrix<T, d, d>& InvIL0,
-		float InvM1,
-		const PMatrix<T, d, d>& InvIL1,
-		const T ProjectionFactor)
+		FVec3& P0,
+		FRotation3& Q0,
+		FVec3& P1,
+		FRotation3& Q1,
+		FReal InvM0,
+		const FMatrix33& InvIL0,
+		FReal InvM1,
+		const FMatrix33& InvIL1,
+		const FReal ProjectionFactor)
 	{
 
 	}
 }
-
-namespace Chaos
-{
-	template class TPBDJointUtilities<float, 3>;
-}
\ No newline at end of file
