Chaos - remove template params from constraint classes.

diff --git a/Engine/Source/Runtime/Experimental/Chaos/Private/Chaos/PBDJointConstraints.cpp b/Engine/Source/Runtime/Experimental/Chaos/Private/Chaos/PBDJointConstraints.cpp
--- a/Engine/Source/Runtime/Experimental/Chaos/Private/Chaos/PBDJointConstraints.cpp
+++ b/Engine/Source/Runtime/Experimental/Chaos/Private/Chaos/PBDJointConstraints.cpp
@@ -21,36 +21,36 @@ namespace Chaos
 	// Constraint Handle
 	//
 
-	template<typename T, int d>
-	TPBDJointConstraintHandle<T, d>::TPBDJointConstraintHandle()
+	
+	FPBDJointConstraintHandle::FPBDJointConstraintHandle()
 	{
 	}
 
-	template<typename T, int d>
-	TPBDJointConstraintHandle<T, d>::TPBDJointConstraintHandle(FConstraintContainer* InConstraintContainer, int32 InConstraintIndex) : TContainerConstraintHandle<TPBDJointConstraints<T, d>>(InConstraintContainer, InConstraintIndex)
+	
+	FPBDJointConstraintHandle::FPBDJointConstraintHandle(FConstraintContainer* InConstraintContainer, int32 InConstraintIndex) : TContainerConstraintHandle<FPBDJointConstraints>(InConstraintContainer, InConstraintIndex)
 	{
 	}
 
-	template<typename T, int d>
-	void TPBDJointConstraintHandle<T, d>::CalculateConstraintSpace(TVector<T, d>& OutXa, PMatrix<T, d, d>& OutRa, TVector<T, d>& OutXb, PMatrix<T, d, d>& OutRb, TVector<T, d>& OutCR) const
+	
+	void FPBDJointConstraintHandle::CalculateConstraintSpace(FVec3& OutXa, FMatrix33& OutRa, FVec3& OutXb, FMatrix33& OutRb, FVec3& OutCR) const
 	{
 		ConstraintContainer->CalculateConstraintSpace(ConstraintIndex, OutXa, OutRa, OutXb, OutRb, OutCR);
 	}
 
-	template<typename T, int d>
-	void TPBDJointConstraintHandle<T, d>::SetParticleLevels(const TVector<int32, 2>& ParticleLevels)
+	
+	void FPBDJointConstraintHandle::SetParticleLevels(const TVector<int32, 2>& ParticleLevels)
 	{
 		ConstraintContainer->SetParticleLevels(ConstraintIndex, ParticleLevels);
 	}
 
-	template<typename T, int d>
-	int32 TPBDJointConstraintHandle<T, d>::GetConstraintLevel() const
+	
+	int32 FPBDJointConstraintHandle::GetConstraintLevel() const
 	{
 		return ConstraintContainer->GetConstraintLevel(ConstraintIndex);
 	}
 
-	template<typename T, int d>
-	const TPBDJointSettings<T, d>& TPBDJointConstraintHandle<T, d>::GetSettings() const
+	
+	const FPBDJointSettings& FPBDJointConstraintHandle::GetSettings() const
 	{
 		return ConstraintContainer->GetConstraintSettings(ConstraintIndex);
 	}
@@ -59,22 +59,22 @@ namespace Chaos
 	// Constraint Settings
 	//
 
-	template<class T, int d>
-	TPBDJointMotionSettings<T, d>::TPBDJointMotionSettings()
-		: Stiffness((T)1)
-		, Projection((T)0)
+	
+	FPBDJointMotionSettings::FPBDJointMotionSettings()
+		: Stiffness((FReal)1)
+		, Projection((FReal)0)
 		, LinearMotionTypes({ EJointMotionType::Locked, EJointMotionType::Locked, EJointMotionType::Locked })
 		, LinearLimit(FLT_MAX)
 		, AngularMotionTypes({ EJointMotionType::Free, EJointMotionType::Free, EJointMotionType::Free })
-		, AngularLimits(TVector<T, d>(FLT_MAX, FLT_MAX, FLT_MAX))
+		, AngularLimits(FVec3(FLT_MAX, FLT_MAX, FLT_MAX))
 		, bSoftLinearLimitsEnabled(false)
 		, bSoftTwistLimitsEnabled(false)
 		, bSoftSwingLimitsEnabled(false)
 		, SoftLinearStiffness(0)
 		, SoftTwistStiffness(0)
 		, SoftSwingStiffness(0)
-		, AngularDriveTarget(TRotation<T, d>::FromIdentity())
-		, AngularDriveTargetAngles(TVector<T, d>(0, 0, 0))
+		, AngularDriveTarget(FRotation3::FromIdentity())
+		, AngularDriveTargetAngles(FVec3(0, 0, 0))
 		, bAngularSLerpDriveEnabled(false)
 		, bAngularTwistDriveEnabled(false)
 		, bAngularSwingDriveEnabled(false)
@@ -82,22 +82,22 @@ namespace Chaos
 	{
 	}
 
-	template<class T, int d>
-	TPBDJointMotionSettings<T, d>::TPBDJointMotionSettings(const TVector<EJointMotionType, d>& InLinearMotionTypes, const TVector<EJointMotionType, d>& InAngularMotionTypes)
-		: Stiffness((T)1)
-		, Projection((T)0)
+	
+	FPBDJointMotionSettings::FPBDJointMotionSettings(const TVector<EJointMotionType, 3>& InLinearMotionTypes, const TVector<EJointMotionType, 3>& InAngularMotionTypes)
+		: Stiffness((FReal)1)
+		, Projection((FReal)0)
 		, LinearMotionTypes(InLinearMotionTypes)
 		, LinearLimit(FLT_MAX)
 		, AngularMotionTypes({ EJointMotionType::Free, EJointMotionType::Free, EJointMotionType::Free })
-		, AngularLimits(TVector<T, d>(FLT_MAX, FLT_MAX, FLT_MAX))
+		, AngularLimits(FVec3(FLT_MAX, FLT_MAX, FLT_MAX))
 		, bSoftLinearLimitsEnabled(false)
 		, bSoftTwistLimitsEnabled(false)
 		, bSoftSwingLimitsEnabled(false)
 		, SoftLinearStiffness(0)
 		, SoftTwistStiffness(0)
 		, SoftSwingStiffness(0)
-		, AngularDriveTarget(TRotation<T, d>::FromIdentity())
-		, AngularDriveTargetAngles(TVector<T, d>(0, 0, 0))
+		, AngularDriveTarget(FRotation3::FromIdentity())
+		, AngularDriveTargetAngles(FVec3(0, 0, 0))
 		, bAngularSLerpDriveEnabled(false)
 		, bAngularTwistDriveEnabled(false)
 		, bAngularSwingDriveEnabled(false)
@@ -105,14 +105,14 @@ namespace Chaos
 	{
 	}
 
-	template<class T, int d>
-	TPBDJointSettings<T, d>::TPBDJointSettings()
+	
+	FPBDJointSettings::FPBDJointSettings()
 		: ConstraintFrames({ FTransform::Identity, FTransform::Identity })
 	{
 	}
 
-	template<class T, int d>
-	TPBDJointState<T, d>::TPBDJointState()
+	
+	FPBDJointState::FPBDJointState()
 		: Level(INDEX_NONE)
 		, ParticleLevels({ INDEX_NONE, INDEX_NONE })
 	{
@@ -122,9 +122,9 @@ namespace Chaos
 	// Solver Settings
 	//
 
-	template<class T, int d>
-	TPBDJointSolverSettings<T, d>::TPBDJointSolverSettings()
-		: SwingTwistAngleTolerance((T)1.0e-6)
+	
+	FPBDJointSolverSettings::FPBDJointSolverSettings()
+		: SwingTwistAngleTolerance((FReal)1.0e-6)
 		, MinParentMassRatio(0)
 		, MaxInertiaRatio(0)
 		, bEnableVelocitySolve(false)
@@ -132,11 +132,11 @@ namespace Chaos
 		, bEnableTwistLimits(true)
 		, bEnableSwingLimits(true)
 		, bEnableDrives(true)
-		, Projection((T)0)
-		, Stiffness((T)0)
-		, DriveStiffness((T)0)
-		, SoftLinearStiffness((T)0)
-		, SoftAngularStiffness((T)0)
+		, Projection((FReal)0)
+		, Stiffness((FReal)0)
+		, DriveStiffness((FReal)0)
+		, SoftLinearStiffness((FReal)0)
+		, SoftAngularStiffness((FReal)0)
 	{
 	}
 
@@ -144,79 +144,79 @@ namespace Chaos
 	// Constraint Container
 	//
 
-	template<class T, int d>
-	TPBDJointConstraints<T, d>::TPBDJointConstraints(const TPBDJointSolverSettings<T, d>& InSettings)
+	
+	FPBDJointConstraints::FPBDJointConstraints(const FPBDJointSolverSettings& InSettings)
 		: Settings(InSettings)
 		, PreApplyCallback(nullptr)
 		, PostApplyCallback(nullptr)
 	{
 	}
 
-	template<class T, int d>
-	TPBDJointConstraints<T, d>::~TPBDJointConstraints()
+	
+	FPBDJointConstraints::~FPBDJointConstraints()
 	{
 	}
 
-	template<class T, int d>
-	const TPBDJointSolverSettings<T, d>& TPBDJointConstraints<T, d>::GetSettings() const
+	
+	const FPBDJointSolverSettings& FPBDJointConstraints::GetSettings() const
 	{
 		return Settings;
 	}
 
-	template<class T, int d>
-	void TPBDJointConstraints<T, d>::SetSettings(const TPBDJointSolverSettings<T, d>& InSettings)
+	
+	void FPBDJointConstraints::SetSettings(const FPBDJointSolverSettings& InSettings)
 	{
 		Settings = InSettings;
 	}
 
-	template<class T, int d>
-	int32 TPBDJointConstraints<T, d>::NumConstraints() const
+	
+	int32 FPBDJointConstraints::NumConstraints() const
 	{
 		return ConstraintParticles.Num();
 	}
 
-	template<class T, int d>
-	typename TPBDJointConstraints<T, d>::FConstraintHandle* TPBDJointConstraints<T, d>::AddConstraint(const FParticlePair& InConstrainedParticles, const TRigidTransform<FReal, Dimensions>& WorldConstraintFrame)
+	
+	typename FPBDJointConstraints::FConstraintContainerHandle* FPBDJointConstraints::AddConstraint(const FParticlePair& InConstrainedParticles, const FRigidTransform3& WorldConstraintFrame)
 	{
 		FTransformPair ConstraintFrames;
-		ConstraintFrames[0] = TRigidTransform<FReal, Dimensions>(
+		ConstraintFrames[0] = FRigidTransform3(
 			WorldConstraintFrame.GetTranslation() - InConstrainedParticles[0]->X(),
 			WorldConstraintFrame.GetRotation() * InConstrainedParticles[0]->R().Inverse()
 			);
-		ConstraintFrames[1] = TRigidTransform<FReal, Dimensions>(
+		ConstraintFrames[1] = FRigidTransform3(
 			WorldConstraintFrame.GetTranslation() - InConstrainedParticles[1]->X(),
 			WorldConstraintFrame.GetRotation() * InConstrainedParticles[1]->R().Inverse()
 			);
 		return AddConstraint(InConstrainedParticles, ConstraintFrames);
 	}
 
-	template<class T, int d>
-	typename TPBDJointConstraints<T, d>::FConstraintHandle* TPBDJointConstraints<T, d>::AddConstraint(const FParticlePair& InConstrainedParticles, const FTransformPair& ConstraintFrames)
+	
+	typename FPBDJointConstraints::FConstraintContainerHandle* FPBDJointConstraints::AddConstraint(const FParticlePair& InConstrainedParticles, const FTransformPair& ConstraintFrames)
 	{
 		int ConstraintIndex = Handles.Num();
 		Handles.Add(HandleAllocator.AllocHandle(this, ConstraintIndex));
 		ConstraintParticles.Add(InConstrainedParticles);
-		ConstraintSettings.Add(FJointSettings());
+		ConstraintSettings.Add(FPBDJointSettings());
 		ConstraintSettings[ConstraintIndex].ConstraintFrames = ConstraintFrames;
-		ConstraintStates.Add(FJointState());
+		ConstraintStates.Add(FPBDJointState());
 		return Handles.Last();
 	}
 
-	template<class T, int d>
-	typename TPBDJointConstraints<T, d>::FConstraintHandle* TPBDJointConstraints<T, d>::AddConstraint(const FParticlePair& InConstrainedParticles, const TPBDJointSettings<T, d>& InConstraintSettings)
+	
+	typename FPBDJointConstraints::FConstraintContainerHandle* FPBDJointConstraints::AddConstraint(const FParticlePair& InConstrainedParticles, const FPBDJointSettings& InConstraintSettings)
 	{
 		int ConstraintIndex = Handles.Num();
 		Handles.Add(HandleAllocator.AllocHandle(this, ConstraintIndex));
 		ConstraintParticles.Add(InConstrainedParticles);
 		ConstraintSettings.Add(InConstraintSettings);
-		ConstraintStates.Add(FJointState());
+		ConstraintStates.Add(FPBDJointState());
 		return Handles.Last();
 	}
 
-	template<class T, int d>
-	void TPBDJointConstraints<T, d>::RemoveConstraint(int ConstraintIndex)
+	
+	void FPBDJointConstraints::RemoveConstraint(int ConstraintIndex)
 	{
-		FConstraintHandle* ConstraintHandle = Handles[ConstraintIndex];
+		FConstraintContainerHandle* ConstraintHandle = Handles[ConstraintIndex];
 		if (ConstraintHandle != nullptr)
 		{
 			// Release the handle for the freed constraint
@@ -237,105 +237,105 @@ namespace Chaos
 		}
 	}
 
-	template<class T, int d>
-	void TPBDJointConstraints<T, d>::RemoveConstraints(const TSet<TGeometryParticleHandle<T, d>*>& RemovedParticles)
+	
+	void FPBDJointConstraints::RemoveConstraints(const TSet<TGeometryParticleHandle<FReal, 3>*>& RemovedParticles)
 	{
 	}
 
 
-	template<class T, int d>
-	void TPBDJointConstraints<T, d>::SetPreApplyCallback(const TJointPreApplyCallback<T, d>& Callback)
+	
+	void FPBDJointConstraints::SetPreApplyCallback(const FJointPreApplyCallback& Callback)
 	{
 		PreApplyCallback = Callback;
 	}
 
-	template<class T, int d>
-	void TPBDJointConstraints<T, d>::ClearPreApplyCallback()
+	
+	void FPBDJointConstraints::ClearPreApplyCallback()
 	{
 		PreApplyCallback = nullptr;
 	}
 
-	template<class T, int d>
-	void TPBDJointConstraints<T, d>::SetPostApplyCallback(const TJointPostApplyCallback<T, d>& Callback)
+	
+	void FPBDJointConstraints::SetPostApplyCallback(const FJointPostApplyCallback& Callback)
 	{
 		PostApplyCallback = Callback;
 	}
 
-	template<class T, int d>
-	void TPBDJointConstraints<T, d>::ClearPostApplyCallback()
+	
+	void FPBDJointConstraints::ClearPostApplyCallback()
 	{
 		PostApplyCallback = nullptr;
 	}
 
-	template<class T, int d>
-	const typename TPBDJointConstraints<T, d>::FConstraintHandle* TPBDJointConstraints<T, d>::GetConstraintHandle(int32 ConstraintIndex) const
+	
+	const typename FPBDJointConstraints::FConstraintContainerHandle* FPBDJointConstraints::GetConstraintHandle(int32 ConstraintIndex) const
 	{
 		return Handles[ConstraintIndex];
 	}
 
-	template<class T, int d>
-	typename TPBDJointConstraints<T, d>::FConstraintHandle* TPBDJointConstraints<T, d>::GetConstraintHandle(int32 ConstraintIndex)
+	
+	typename FPBDJointConstraints::FConstraintContainerHandle* FPBDJointConstraints::GetConstraintHandle(int32 ConstraintIndex)
 	{
 		return Handles[ConstraintIndex];
 	}
 
-	template<class T, int d>
-	const typename TPBDJointConstraints<T, d>::FParticlePair& TPBDJointConstraints<T, d>::GetConstrainedParticles(int32 ConstraintIndex) const
+	
+	const typename FPBDJointConstraints::FParticlePair& FPBDJointConstraints::GetConstrainedParticles(int32 ConstraintIndex) const
 	{
 		return ConstraintParticles[ConstraintIndex];
 	}
 
-	template<class T, int d>
-	const TPBDJointSettings<T, d>& TPBDJointConstraints<T, d>::GetConstraintSettings(int32 ConstraintIndex) const
+	
+	const FPBDJointSettings& FPBDJointConstraints::GetConstraintSettings(int32 ConstraintIndex) const
 	{
 		return ConstraintSettings[ConstraintIndex];
 	}
 
-	template<class T, int d>
-	int32 TPBDJointConstraints<T, d>::GetConstraintLevel(int32 ConstraintIndex) const
+	
+	int32 FPBDJointConstraints::GetConstraintLevel(int32 ConstraintIndex) const
 	{
 		return ConstraintStates[ConstraintIndex].Level;
 	}
 
-	template<class T, int d>
-	void TPBDJointConstraints<T, d>::SetParticleLevels(int32 ConstraintIndex, const TVector<int32, 2>& ParticleLevels)
+	
+	void FPBDJointConstraints::SetParticleLevels(int32 ConstraintIndex, const TVector<int32, 2>& ParticleLevels)
 	{
 		ConstraintStates[ConstraintIndex].Level = FMath::Min(ParticleLevels[0], ParticleLevels[1]);
 		ConstraintStates[ConstraintIndex].ParticleLevels = ParticleLevels;
 	}
 
-	template<class T, int d>
-	void TPBDJointConstraints<T, d>::UpdatePositionBasedState(const T Dt)
+	
+	void FPBDJointConstraints::UpdatePositionBasedState(const FReal Dt)
 	{
 	}
 
-	template<class T, int d>
-	void TPBDJointConstraints<T, d>::CalculateConstraintSpace(int32 ConstraintIndex, TVector<T, d>& OutX0, PMatrix<T, d, d>& OutR0, TVector<T, d>& OutX1, PMatrix<T, d, d>& OutR1, TVector<T, d>& OutCR) const
+	
+	void FPBDJointConstraints::CalculateConstraintSpace(int32 ConstraintIndex, FVec3& OutX0, FMatrix33& OutR0, FVec3& OutX1, FMatrix33& OutR1, FVec3& OutCR) const
 	{
 		const int32 Index0 = 1;
 		const int32 Index1 = 0;
-		TGenericParticleHandle<T, d> Particle0 = TGenericParticleHandle<T, d>(ConstraintParticles[ConstraintIndex][Index0]);
-		TGenericParticleHandle<T, d> Particle1 = TGenericParticleHandle<T, d>(ConstraintParticles[ConstraintIndex][Index1]);
-		TVector<T, d> P0 = Particle0->P();
-		TRotation<T, d> Q0 = Particle0->Q();
-		TVector<T, d> P1 = Particle1->P();
-		TRotation<T, d> Q1 = Particle1->Q();
-
-		const FJointSettings& JointSettings = ConstraintSettings[ConstraintIndex];
+		TGenericParticleHandle<FReal, 3> Particle0 = TGenericParticleHandle<FReal, 3>(ConstraintParticles[ConstraintIndex][Index0]);
+		TGenericParticleHandle<FReal, 3> Particle1 = TGenericParticleHandle<FReal, 3>(ConstraintParticles[ConstraintIndex][Index1]);
+		FVec3 P0 = Particle0->P();
+		FRotation3 Q0 = Particle0->Q();
+		FVec3 P1 = Particle1->P();
+		FRotation3 Q1 = Particle1->Q();
+
+		const FPBDJointSettings& JointSettings = ConstraintSettings[ConstraintIndex];
 		EJointMotionType Swing1Motion = JointSettings.Motion.AngularMotionTypes[(int32)EJointAngularConstraintIndex::Swing1];
 		EJointMotionType Swing2Motion = JointSettings.Motion.AngularMotionTypes[(int32)EJointAngularConstraintIndex::Swing2];
 		if ((Swing1Motion == EJointMotionType::Limited) && (Swing2Motion == EJointMotionType::Limited))
 		{
-			TPBDJointUtilities<T, d>::CalculateConeConstraintSpace(Settings, ConstraintSettings[ConstraintIndex], Index0, Index1, P0, Q0, P1, Q1, OutX0, OutR0, OutX1, OutR1, OutCR);
+			FPBDJointUtilities::CalculateConeConstraintSpace(Settings, ConstraintSettings[ConstraintIndex], Index0, Index1, P0, Q0, P1, Q1, OutX0, OutR0, OutX1, OutR1, OutCR);
 		}
 		else
 		{
-			TPBDJointUtilities<T, d>::CalculateSwingConstraintSpace(Settings, ConstraintSettings[ConstraintIndex], Index0, Index1, P0, Q0, P1, Q1, OutX0, OutR0, OutX1, OutR1, OutCR);
+			FPBDJointUtilities::CalculateSwingConstraintSpace(Settings, ConstraintSettings[ConstraintIndex], Index0, Index1, P0, Q0, P1, Q1, OutX0, OutR0, OutX1, OutR1, OutCR);
 		}
 	}
 
-	template<class T, int d>
-	void TPBDJointConstraints<T, d>::Apply(const T Dt, const TArray<FConstraintHandle*>& InConstraintHandles, const int32 It, const int32 NumIts)
+	
+	void FPBDJointConstraints::Apply(const FReal Dt, const TArray<FConstraintContainerHandle*>& InConstraintHandles, const int32 It, const int32 NumIts)
 	{
 		SCOPE_CYCLE_COUNTER(STAT_ApplyJointConstraints);
 
@@ -345,8 +345,8 @@ namespace Chaos
 		// If we are freezing we want the root-most (nearest to kinematic) bodies solved first.
 		// For normal update we want the root body last, otherwise it gets dragged away from the root by the other bodies
 
-		TArray<FConstraintHandle*> SortedConstraintHandles = InConstraintHandles;
-		SortedConstraintHandles.Sort([](const FConstraintHandle& L, const FConstraintHandle& R)
+		TArray<FConstraintContainerHandle*> SortedConstraintHandles = InConstraintHandles;
+		SortedConstraintHandles.Sort([](const FConstraintContainerHandle& L, const FConstraintContainerHandle& R)
 			{
 				// Sort bodies from leaf to root
 				return L.GetConstraintLevel() > R.GetConstraintLevel();
@@ -357,7 +357,7 @@ namespace Chaos
 			PreApplyCallback(Dt, SortedConstraintHandles);
 		}
 
-		for (FConstraintHandle* ConstraintHandle : SortedConstraintHandles)
+		for (FConstraintContainerHandle* ConstraintHandle : SortedConstraintHandles)
 		{
 			if (Settings.bEnableVelocitySolve)
 			{
@@ -375,22 +375,22 @@ namespace Chaos
 		}
 	}
 
-	template<class T, int d>
-	bool TPBDJointConstraints<T, d>::ApplyPushOut(const T Dt, const TArray<FConstraintHandle*>& InConstraintHandles, const int32 It, const int32 NumIts)
+	
+	bool FPBDJointConstraints::ApplyPushOut(const FReal Dt, const TArray<FConstraintContainerHandle*>& InConstraintHandles, const int32 It, const int32 NumIts)
 	{
 		// @todo(ccaulfield): track whether we are sufficiently solved
 		bool bNeedsAnotherIteration = true;
 
 		if (Settings.bEnableVelocitySolve)
 		{
-			TArray<FConstraintHandle*> SortedConstraintHandles = InConstraintHandles;
-			SortedConstraintHandles.Sort([](const FConstraintHandle& L, const FConstraintHandle& R)
+			TArray<FConstraintContainerHandle*> SortedConstraintHandles = InConstraintHandles;
+			SortedConstraintHandles.Sort([](const FConstraintContainerHandle& L, const FConstraintContainerHandle& R)
 				{
 					// Sort bodies from root to leaf
 					return L.GetConstraintLevel() < R.GetConstraintLevel();
 				});
 
-			for (FConstraintHandle* ConstraintHandle : SortedConstraintHandles)
+			for (FConstraintContainerHandle* ConstraintHandle : SortedConstraintHandles)
 			{
 				SolvePosition(Dt, ConstraintHandle->GetConstraintIndex(), It, NumIts);
 			}
@@ -405,53 +405,53 @@ namespace Chaos
 		return bNeedsAnotherIteration;
 	}
 
-	template<class T, int d>
-	void TPBDJointConstraints<T, d>::ApplyProjection(const T Dt, const TArray<FConstraintHandle*>& InConstraintHandles)
+	
+	void FPBDJointConstraints::ApplyProjection(const FReal Dt, const TArray<FConstraintContainerHandle*>& InConstraintHandles)
 	{
 		if (Settings.Projection > 0)
 		{
-			TArray<FConstraintHandle*> SortedConstraintHandles = InConstraintHandles;
-			SortedConstraintHandles.Sort([](const FConstraintHandle& L, const FConstraintHandle& R)
+			TArray<FConstraintContainerHandle*> SortedConstraintHandles = InConstraintHandles;
+			SortedConstraintHandles.Sort([](const FConstraintContainerHandle& L, const FConstraintContainerHandle& R)
 				{
 					// Sort bodies from root to leaf
 					return L.GetConstraintLevel() < R.GetConstraintLevel();
 				});
 
-			for (FConstraintHandle* ConstraintHandle : SortedConstraintHandles)
+			for (FConstraintContainerHandle* ConstraintHandle : SortedConstraintHandles)
 			{
 				ProjectPosition(Dt, ConstraintHandle->GetConstraintIndex(), 0, 1);
 			}
 		}
 	}
 
-	template<class T, int d>
-	void TPBDJointConstraints<T, d>::SolveVelocity(const T Dt, const int32 ConstraintIndex, const int32 It, const int32 NumIts)
+	
+	void FPBDJointConstraints::SolveVelocity(const FReal Dt, const int32 ConstraintIndex, const int32 It, const int32 NumIts)
 	{
-		const TVector<TGeometryParticleHandle<T, d>*, 2>& Constraint = ConstraintParticles[ConstraintIndex];
+		const TVector<TGeometryParticleHandle<FReal, 3>*, 2>& Constraint = ConstraintParticles[ConstraintIndex];
 		UE_LOG(LogChaosJoint, Verbose, TEXT("Solve Joint Constraint %d %s %s (dt = %f; it = %d / %d)"), ConstraintIndex, *Constraint[0]->ToString(), *Constraint[1]->ToString(), Dt, It, NumIts);
 
-		const FJointSettings& JointSettings = ConstraintSettings[ConstraintIndex];
+		const FPBDJointSettings& JointSettings = ConstraintSettings[ConstraintIndex];
 
 		// Switch particles - internally we assume the first body is the parent (i.e., the space in which constraint limits are specified)
 		const int32 Index0 = 1;
 		const int32 Index1 = 0;
-		TGenericParticleHandle<T, d> Particle0 = TGenericParticleHandle<T, d>(ConstraintParticles[ConstraintIndex][Index0]);
-		TGenericParticleHandle<T, d> Particle1 = TGenericParticleHandle<T, d>(ConstraintParticles[ConstraintIndex][Index1]);
-		TPBDRigidParticleHandle<T, d>* Rigid0 = ConstraintParticles[ConstraintIndex][Index0]->AsDynamic();
-		TPBDRigidParticleHandle<T, d>* Rigid1 = ConstraintParticles[ConstraintIndex][Index1]->AsDynamic();
-
-		TVector<T, d> P0 = Particle0->P();
-		TRotation<T, d> Q0 = Particle0->Q();
-		TVector<T, d> V0 = Particle0->V();
-		TVector<T, d> W0 = Particle0->W();
-		TVector<T, d> P1 = Particle1->P();
-		TRotation<T, d> Q1 = Particle1->Q();
-		TVector<T, d> V1 = Particle1->V();
-		TVector<T, d> W1 = Particle1->W();
+		TGenericParticleHandle<FReal, 3> Particle0 = TGenericParticleHandle<FReal, 3>(ConstraintParticles[ConstraintIndex][Index0]);
+		TGenericParticleHandle<FReal, 3> Particle1 = TGenericParticleHandle<FReal, 3>(ConstraintParticles[ConstraintIndex][Index1]);
+		TPBDRigidParticleHandle<FReal, 3>* Rigid0 = ConstraintParticles[ConstraintIndex][Index0]->AsDynamic();
+		TPBDRigidParticleHandle<FReal, 3>* Rigid1 = ConstraintParticles[ConstraintIndex][Index1]->AsDynamic();
+
+		FVec3 P0 = Particle0->P();
+		FRotation3 Q0 = Particle0->Q();
+		FVec3 V0 = Particle0->V();
+		FVec3 W0 = Particle0->W();
+		FVec3 P1 = Particle1->P();
+		FRotation3 Q1 = Particle1->Q();
+		FVec3 V1 = Particle1->V();
+		FVec3 W1 = Particle1->W();
 		float InvM0 = Particle0->InvM();
 		float InvM1 = Particle1->InvM();
-		PMatrix<T, d, d> InvIL0 = Particle0->InvI();
-		PMatrix<T, d, d> InvIL1 = Particle1->InvI();
+		FMatrix33 InvIL0 = Particle0->InvI();
+		FMatrix33 InvIL1 = Particle1->InvI();
 
 		Q1.EnforceShortestArcWith(Q0);
 
@@ -460,18 +460,18 @@ namespace Chaos
 		const int32 Level1 = ConstraintStates[ConstraintIndex].ParticleLevels[Index1];
 		if (Level0 < Level1)
 		{
-			TPBDJointUtilities<T, d>::GetConditionedInverseMass(Particle0->M(), Particle0->I().GetDiagonal(), Particle1->M(), Particle1->I().GetDiagonal(), InvM0, InvM1, InvIL0, InvIL1, Settings.MinParentMassRatio, Settings.MaxInertiaRatio);
+			FPBDJointUtilities::GetConditionedInverseMass(Particle0->M(), Particle0->I().GetDiagonal(), Particle1->M(), Particle1->I().GetDiagonal(), InvM0, InvM1, InvIL0, InvIL1, Settings.MinParentMassRatio, Settings.MaxInertiaRatio);
 		}
 		else if (Level0 > Level1)
 		{
-			TPBDJointUtilities<T, d>::GetConditionedInverseMass(Particle1->M(), Particle1->I().GetDiagonal(), Particle0->M(), Particle0->I().GetDiagonal(), InvM1, InvM0, InvIL1, InvIL0, Settings.MinParentMassRatio, Settings.MaxInertiaRatio);
+			FPBDJointUtilities::GetConditionedInverseMass(Particle1->M(), Particle1->I().GetDiagonal(), Particle0->M(), Particle0->I().GetDiagonal(), InvM1, InvM0, InvIL1, InvIL0, Settings.MinParentMassRatio, Settings.MaxInertiaRatio);
 		}
 		else
 		{
-			TPBDJointUtilities<T, d>::GetConditionedInverseMass(Particle0->M(), Particle0->I().GetDiagonal(), Particle1->M(), Particle1->I().GetDiagonal(), InvM0, InvM1, InvIL0, InvIL1, (T)0, Settings.MaxInertiaRatio);
+			FPBDJointUtilities::GetConditionedInverseMass(Particle0->M(), Particle0->I().GetDiagonal(), Particle1->M(), Particle1->I().GetDiagonal(), InvM0, InvM1, InvIL0, InvIL1, (FReal)0, Settings.MaxInertiaRatio);
 		}
 
-		const TVector<EJointMotionType, d>& LinearMotion = JointSettings.Motion.LinearMotionTypes;
+		const TVector<EJointMotionType, 3>& LinearMotion = JointSettings.Motion.LinearMotionTypes;
 		EJointMotionType TwistMotion = JointSettings.Motion.AngularMotionTypes[(int32)EJointAngularConstraintIndex::Twist];
 		EJointMotionType Swing1Motion = JointSettings.Motion.AngularMotionTypes[(int32)EJointAngularConstraintIndex::Swing1];
 		EJointMotionType Swing2Motion = JointSettings.Motion.AngularMotionTypes[(int32)EJointAngularConstraintIndex::Swing2];
@@ -486,25 +486,25 @@ namespace Chaos
 			// No SLerp drive if we have a locked rotation (it will be grayed out in the editor in this case, but could still have been set before the rotation was locked)
 			if (JointSettings.Motion.bAngularSLerpDriveEnabled && !bTwistLocked && !bSwing1Locked && !bSwing2Locked)
 			{
-				TPBDJointUtilities<T, d>::ApplyJointSLerpDrive(Dt, Settings, JointSettings, Index0, Index1, P0, Q0, P1, Q1, InvM0, InvIL0, InvM1, InvIL1);
+				FPBDJointUtilities::ApplyJointSLerpDrive(Dt, Settings, JointSettings, Index0, Index1, P0, Q0, P1, Q1, InvM0, InvIL0, InvM1, InvIL1);
 			}
 
 			if (JointSettings.Motion.bAngularTwistDriveEnabled && !bTwistLocked)
 			{
-				TPBDJointUtilities<T, d>::ApplyJointTwistDrive(Dt, Settings, JointSettings, Index0, Index1, P0, Q0, P1, Q1, InvM0, InvIL0, InvM1, InvIL1);
+				FPBDJointUtilities::ApplyJointTwistDrive(Dt, Settings, JointSettings, Index0, Index1, P0, Q0, P1, Q1, InvM0, InvIL0, InvM1, InvIL1);
 			}
 
 			if (JointSettings.Motion.bAngularSwingDriveEnabled && !bSwing1Locked && !bSwing2Locked)
 			{
-				TPBDJointUtilities<T, d>::ApplyJointConeDrive(Dt, Settings, JointSettings, Index0, Index1, P0, Q0, P1, Q1, InvM0, InvIL0, InvM1, InvIL1);
+				FPBDJointUtilities::ApplyJointConeDrive(Dt, Settings, JointSettings, Index0, Index1, P0, Q0, P1, Q1, InvM0, InvIL0, InvM1, InvIL1);
 			}
 			else if (JointSettings.Motion.bAngularSwingDriveEnabled && !bSwing1Locked)
 			{
-				//TPBDJointUtilities<T, d>::ApplyJointSwingDrive(Dt, Settings, JointSettings, Index0, Index1, EJointAngularConstraintIndex::Swing1, P0, Q0, P1, Q1, InvM0, InvIL0, InvM1, InvIL1);
+				//FPBDJointUtilities::ApplyJointSwingDrive(Dt, Settings, JointSettings, Index0, Index1, EJointAngularConstraintIndex::Swing1, P0, Q0, P1, Q1, InvM0, InvIL0, InvM1, InvIL1);
 			}
 			else if (JointSettings.Motion.bAngularSwingDriveEnabled && !bSwing2Locked)
 			{
-				//TPBDJointUtilities<T, d>::ApplyJointSwingDrive(Dt, Settings, JointSettings, Index0, Index1, EJointAngularConstraintIndex::Swing2, P0, Q0, P1, Q1, InvM0, InvIL0, InvM1, InvIL1);
+				//FPBDJointUtilities::ApplyJointSwingDrive(Dt, Settings, JointSettings, Index0, Index1, EJointAngularConstraintIndex::Swing2, P0, Q0, P1, Q1, InvM0, InvIL0, InvM1, InvIL1);
 			}
 		}
 
@@ -513,7 +513,7 @@ namespace Chaos
 		{
 			if (TwistMotion != EJointMotionType::Free)
 			{
-				TPBDJointUtilities<T, d>::ApplyJointTwistVelocityConstraint(Dt, Settings, JointSettings, Index0, Index1, P0, Q0, V0, W0, P1, Q1, V1, W1, InvM0, InvIL0, InvM1, InvIL1);
+				FPBDJointUtilities::ApplyJointTwistVelocityConstraint(Dt, Settings, JointSettings, Index0, Index1, P0, Q0, V0, W0, P1, Q1, V1, W1, InvM0, InvIL0, InvM1, InvIL1);
 			}
 		}
 
@@ -523,19 +523,19 @@ namespace Chaos
 			if ((Swing1Motion == EJointMotionType::Limited) && (Swing2Motion == EJointMotionType::Limited))
 			{
 				// Swing Cone
-				TPBDJointUtilities<T, d>::ApplyJointConeVelocityConstraint(Dt, Settings, JointSettings, Index0, Index1, P0, Q0, V0, W0, P1, Q1, V1, W1, InvM0, InvIL0, InvM1, InvIL1);
+				FPBDJointUtilities::ApplyJointConeVelocityConstraint(Dt, Settings, JointSettings, Index0, Index1, P0, Q0, V0, W0, P1, Q1, V1, W1, InvM0, InvIL0, InvM1, InvIL1);
 			}
 			else
 			{
 				if (Swing1Motion != EJointMotionType::Free)
 				{
 					// Swing Arc/Lock
-					TPBDJointUtilities<T, d>::ApplyJointSwingVelocityConstraint(Dt, Settings, JointSettings, Index0, Index1, EJointAngularConstraintIndex::Swing1, EJointAngularAxisIndex::Swing1, P0, Q0, V0, W0, P1, Q1, V1, W1, InvM0, InvIL0, InvM1, InvIL1);
+					FPBDJointUtilities::ApplyJointSwingVelocityConstraint(Dt, Settings, JointSettings, Index0, Index1, EJointAngularConstraintIndex::Swing1, EJointAngularAxisIndex::Swing1, P0, Q0, V0, W0, P1, Q1, V1, W1, InvM0, InvIL0, InvM1, InvIL1);
 				}
 				if (Swing2Motion != EJointMotionType::Free)
 				{
 					// Swing Arc/Lock
-					TPBDJointUtilities<T, d>::ApplyJointSwingVelocityConstraint(Dt, Settings, JointSettings, Index0, Index1, EJointAngularConstraintIndex::Swing2, EJointAngularAxisIndex::Swing2, P0, Q0, V0, W0, P1, Q1, V1, W1, InvM0, InvIL0, InvM1, InvIL1);
+					FPBDJointUtilities::ApplyJointSwingVelocityConstraint(Dt, Settings, JointSettings, Index0, Index1, EJointAngularConstraintIndex::Swing2, EJointAngularAxisIndex::Swing2, P0, Q0, V0, W0, P1, Q1, V1, W1, InvM0, InvIL0, InvM1, InvIL1);
 				}
 			}
 		}
@@ -543,7 +543,7 @@ namespace Chaos
 		// Apply linear velocity  constraints
 		if ((LinearMotion[0] != EJointMotionType::Free) || (LinearMotion[1] != EJointMotionType::Free) || (LinearMotion[2] != EJointMotionType::Free))
 		{
-			TPBDJointUtilities<T, d>::ApplyJointVelocityConstraint(Dt, Settings, JointSettings, Index0, Index1, P0, Q0, V0, W0, P1, Q1, V1, W1, InvM0, InvIL0, InvM1, InvIL1);
+			FPBDJointUtilities::ApplyJointVelocityConstraint(Dt, Settings, JointSettings, Index0, Index1, P0, Q0, V0, W0, P1, Q1, V1, W1, InvM0, InvIL0, InvM1, InvIL1);
 		}
 
 		// Update the particles
@@ -563,30 +563,30 @@ namespace Chaos
 		}
 	}
 
-	template<class T, int d>
-	void TPBDJointConstraints<T, d>::SolvePosition(const T Dt, const int32 ConstraintIndex, const int32 It, const int32 NumIts)
+	
+	void FPBDJointConstraints::SolvePosition(const FReal Dt, const int32 ConstraintIndex, const int32 It, const int32 NumIts)
 	{
-		const TVector<TGeometryParticleHandle<T, d>*, 2>& Constraint = ConstraintParticles[ConstraintIndex];
+		const TVector<TGeometryParticleHandle<FReal, 3>*, 2>& Constraint = ConstraintParticles[ConstraintIndex];
 		UE_LOG(LogChaosJoint, Verbose, TEXT("Solve Joint Constraint %d %s %s (dt = %f; it = %d / %d)"), ConstraintIndex, *Constraint[0]->ToString(), *Constraint[1]->ToString(), Dt, It, NumIts);
 
-		const FJointSettings& JointSettings = ConstraintSettings[ConstraintIndex];
+		const FPBDJointSettings& JointSettings = ConstraintSettings[ConstraintIndex];
 
 		// Switch particles - internally we assume the first body is the parent (i.e., the space in which constraint limits are specified)
 		const int32 Index0 = 1;
 		const int32 Index1 = 0;
-		TGenericParticleHandle<T, d> Particle0 = TGenericParticleHandle<T, d>(ConstraintParticles[ConstraintIndex][Index0]);
-		TGenericParticleHandle<T, d> Particle1 = TGenericParticleHandle<T, d>(ConstraintParticles[ConstraintIndex][Index1]);
-		TPBDRigidParticleHandle<T, d>* Rigid0 = ConstraintParticles[ConstraintIndex][Index0]->AsDynamic();
-		TPBDRigidParticleHandle<T, d>* Rigid1 = ConstraintParticles[ConstraintIndex][Index1]->AsDynamic();
-
-		TVector<T, d> P0 = Particle0->P();
-		TRotation<T, d> Q0 = Particle0->Q();
-		TVector<T, d> P1 = Particle1->P();
-		TRotation<T, d> Q1 = Particle1->Q();
+		TGenericParticleHandle<FReal, 3> Particle0 = TGenericParticleHandle<FReal, 3>(ConstraintParticles[ConstraintIndex][Index0]);
+		TGenericParticleHandle<FReal, 3> Particle1 = TGenericParticleHandle<FReal, 3>(ConstraintParticles[ConstraintIndex][Index1]);
+		TPBDRigidParticleHandle<FReal, 3>* Rigid0 = ConstraintParticles[ConstraintIndex][Index0]->AsDynamic();
+		TPBDRigidParticleHandle<FReal, 3>* Rigid1 = ConstraintParticles[ConstraintIndex][Index1]->AsDynamic();
+
+		FVec3 P0 = Particle0->P();
+		FRotation3 Q0 = Particle0->Q();
+		FVec3 P1 = Particle1->P();
+		FRotation3 Q1 = Particle1->Q();
 		float InvM0 = Particle0->InvM();
 		float InvM1 = Particle1->InvM();
-		PMatrix<T, d, d> InvIL0 = Particle0->InvI();
-		PMatrix<T, d, d> InvIL1 = Particle1->InvI();
+		FMatrix33 InvIL0 = Particle0->InvI();
+		FMatrix33 InvIL1 = Particle1->InvI();
 
 		Q1.EnforceShortestArcWith(Q0);
 
@@ -595,21 +595,21 @@ namespace Chaos
 		const int32 Level1 = ConstraintStates[ConstraintIndex].ParticleLevels[Index1];
 		if (Level0 < Level1)
 		{
-			TPBDJointUtilities<T, d>::GetConditionedInverseMass(Particle0->M(), Particle0->I().GetDiagonal(), Particle1->M(), Particle1->I().GetDiagonal(), InvM0, InvM1, InvIL0, InvIL1, Settings.MinParentMassRatio, Settings.MaxInertiaRatio);
+			FPBDJointUtilities::GetConditionedInverseMass(Particle0->M(), Particle0->I().GetDiagonal(), Particle1->M(), Particle1->I().GetDiagonal(), InvM0, InvM1, InvIL0, InvIL1, Settings.MinParentMassRatio, Settings.MaxInertiaRatio);
 		}
 		else if (Level0 > Level1)
 		{
-			TPBDJointUtilities<T, d>::GetConditionedInverseMass(Particle1->M(), Particle1->I().GetDiagonal(), Particle0->M(), Particle0->I().GetDiagonal(), InvM1, InvM0, InvIL1, InvIL0, Settings.MinParentMassRatio, Settings.MaxInertiaRatio);
+			FPBDJointUtilities::GetConditionedInverseMass(Particle1->M(), Particle1->I().GetDiagonal(), Particle0->M(), Particle0->I().GetDiagonal(), InvM1, InvM0, InvIL1, InvIL0, Settings.MinParentMassRatio, Settings.MaxInertiaRatio);
 		}
 		else
 		{
-			TPBDJointUtilities<T, d>::GetConditionedInverseMass(Particle0->M(), Particle0->I().GetDiagonal(), Particle1->M(), Particle1->I().GetDiagonal(), InvM0, InvM1, InvIL0, InvIL1, (T)0, Settings.MaxInertiaRatio);
+			FPBDJointUtilities::GetConditionedInverseMass(Particle0->M(), Particle0->I().GetDiagonal(), Particle1->M(), Particle1->I().GetDiagonal(), InvM0, InvM1, InvIL0, InvIL1, (FReal)0, Settings.MaxInertiaRatio);
 		}
 
 		// Freeze the closest to kinematic connection (if one is closer than the other)
 		if (Settings.bEnableVelocitySolve && (Level0 != Level1))
 		{
-			const T FreezeFactor = (T)(NumIts - (It + 1)) / (T)NumIts;
+			const FReal FreezeFactor = (FReal)(NumIts - (It + 1)) / (FReal)NumIts;
 			if (Level0 < Level1)
 			{
 				InvM0 = InvM0 * FreezeFactor * FreezeFactor;
@@ -622,7 +622,7 @@ namespace Chaos
 			}
 		}
 
-		const TVector<EJointMotionType, d>& LinearMotion = JointSettings.Motion.LinearMotionTypes;
+		const TVector<EJointMotionType, 3>& LinearMotion = JointSettings.Motion.LinearMotionTypes;
 		EJointMotionType TwistMotion = JointSettings.Motion.AngularMotionTypes[(int32)EJointAngularConstraintIndex::Twist];
 		EJointMotionType Swing1Motion = JointSettings.Motion.AngularMotionTypes[(int32)EJointAngularConstraintIndex::Swing1];
 		EJointMotionType Swing2Motion = JointSettings.Motion.AngularMotionTypes[(int32)EJointAngularConstraintIndex::Swing2];
@@ -643,25 +643,25 @@ namespace Chaos
 			// No SLerp drive if we have a locked rotation (it will be grayed out in the editor in this case, but could still have been set before the rotation was locked)
 			if (JointSettings.Motion.bAngularSLerpDriveEnabled && !bTwistLocked && !bSwing1Locked && !bSwing2Locked)
 			{
-				TPBDJointUtilities<T, d>::ApplyJointSLerpDrive(Dt, Settings, JointSettings, Index0, Index1, P0, Q0, P1, Q1, InvM0, InvIL0, InvM1, InvIL1);
+				FPBDJointUtilities::ApplyJointSLerpDrive(Dt, Settings, JointSettings, Index0, Index1, P0, Q0, P1, Q1, InvM0, InvIL0, InvM1, InvIL1);
 			}
 
 			if (JointSettings.Motion.bAngularTwistDriveEnabled && !bTwistLocked)
 			{
-				TPBDJointUtilities<T, d>::ApplyJointTwistDrive(Dt, Settings, JointSettings, Index0, Index1, P0, Q0, P1, Q1, InvM0, InvIL0, InvM1, InvIL1);
+				FPBDJointUtilities::ApplyJointTwistDrive(Dt, Settings, JointSettings, Index0, Index1, P0, Q0, P1, Q1, InvM0, InvIL0, InvM1, InvIL1);
 			}
 
 			if (JointSettings.Motion.bAngularSwingDriveEnabled && !bSwing1Locked && !bSwing2Locked)
 			{
-				TPBDJointUtilities<T, d>::ApplyJointConeDrive(Dt, Settings, JointSettings, Index0, Index1, P0, Q0, P1, Q1, InvM0, InvIL0, InvM1, InvIL1);
+				FPBDJointUtilities::ApplyJointConeDrive(Dt, Settings, JointSettings, Index0, Index1, P0, Q0, P1, Q1, InvM0, InvIL0, InvM1, InvIL1);
 			}
 			else if (JointSettings.Motion.bAngularSwingDriveEnabled && !bSwing1Locked)
 			{
-				//TPBDJointUtilities<T, d>::ApplyJointSwingDrive(Dt, Settings, JointSettings, Index0, Index1, EJointAngularConstraintIndex::Swing1, P0, Q0, P1, Q1, InvM0, InvIL0, InvM1, InvIL1);
+				//FPBDJointUtilities::ApplyJointSwingDrive(Dt, Settings, JointSettings, Index0, Index1, EJointAngularConstraintIndex::Swing1, P0, Q0, P1, Q1, InvM0, InvIL0, InvM1, InvIL1);
 			}
 			else if (JointSettings.Motion.bAngularSwingDriveEnabled && !bSwing2Locked)
 			{
-				//TPBDJointUtilities<T, d>::ApplyJointSwingDrive(Dt, Settings, JointSettings, Index0, Index1, EJointAngularConstraintIndex::Swing2, P0, Q0, P1, Q1, InvM0, InvIL0, InvM1, InvIL1);
+				//FPBDJointUtilities::ApplyJointSwingDrive(Dt, Settings, JointSettings, Index0, Index1, EJointAngularConstraintIndex::Swing2, P0, Q0, P1, Q1, InvM0, InvIL0, InvM1, InvIL1);
 			}
 		}
 
@@ -670,7 +670,7 @@ namespace Chaos
 		{
 			if (TwistMotion != EJointMotionType::Free)
 			{
-				TPBDJointUtilities<T, d>::ApplyJointTwistConstraint(Dt, Settings, JointSettings, Index0, Index1, P0, Q0, P1, Q1, InvM0, InvIL0, InvM1, InvIL1);
+				FPBDJointUtilities::ApplyJointTwistConstraint(Dt, Settings, JointSettings, Index0, Index1, P0, Q0, P1, Q1, InvM0, InvIL0, InvM1, InvIL1);
 			}
 		}
 
@@ -680,19 +680,19 @@ namespace Chaos
 			if ((Swing1Motion == EJointMotionType::Limited) && (Swing2Motion == EJointMotionType::Limited))
 			{
 				// Swing Cone
-				TPBDJointUtilities<T, d>::ApplyJointConeConstraint(Dt, Settings, JointSettings, Index0, Index1, P0, Q0, P1, Q1, InvM0, InvIL0, InvM1, InvIL1);
+				FPBDJointUtilities::ApplyJointConeConstraint(Dt, Settings, JointSettings, Index0, Index1, P0, Q0, P1, Q1, InvM0, InvIL0, InvM1, InvIL1);
 			}
 			else
 			{
 				if (Swing1Motion != EJointMotionType::Free)
 				{
 					// Swing Arc/Lock
-					TPBDJointUtilities<T, d>::ApplyJointSwingConstraint(Dt, Settings, JointSettings, Index0, Index1, EJointAngularConstraintIndex::Swing1, EJointAngularAxisIndex::Swing1, P0, Q0, P1, Q1, InvM0, InvIL0, InvM1, InvIL1);
+					FPBDJointUtilities::ApplyJointSwingConstraint(Dt, Settings, JointSettings, Index0, Index1, EJointAngularConstraintIndex::Swing1, EJointAngularAxisIndex::Swing1, P0, Q0, P1, Q1, InvM0, InvIL0, InvM1, InvIL1);
 				}
 				if (Swing2Motion != EJointMotionType::Free)
 				{
 					// Swing Arc/Lock
-					TPBDJointUtilities<T, d>::ApplyJointSwingConstraint(Dt, Settings, JointSettings, Index0, Index1, EJointAngularConstraintIndex::Swing2, EJointAngularAxisIndex::Swing2, P0, Q0, P1, Q1, InvM0, InvIL0, InvM1, InvIL1);
+					FPBDJointUtilities::ApplyJointSwingConstraint(Dt, Settings, JointSettings, Index0, Index1, EJointAngularConstraintIndex::Swing2, EJointAngularAxisIndex::Swing2, P0, Q0, P1, Q1, InvM0, InvIL0, InvM1, InvIL1);
 				}
 			}
 		}
@@ -700,7 +700,7 @@ namespace Chaos
 		// Apply linear constraints
 		if ((LinearMotion[0] != EJointMotionType::Free) || (LinearMotion[1] != EJointMotionType::Free) || (LinearMotion[2] != EJointMotionType::Free))
 		{
-			TPBDJointUtilities<T, d>::ApplyJointPositionConstraint(Dt, Settings, JointSettings, Index0, Index1, P0, Q0, P1, Q1, InvM0, InvIL0, InvM1, InvIL1);
+			FPBDJointUtilities::ApplyJointPositionConstraint(Dt, Settings, JointSettings, Index0, Index1, P0, Q0, P1, Q1, InvM0, InvIL0, InvM1, InvIL1);
 		}
 
 		// Update the particles
@@ -716,35 +716,35 @@ namespace Chaos
 		}
 	}
 
-	template<class T, int d>
-	void TPBDJointConstraints<T, d>::ProjectPosition(const T Dt, const int32 ConstraintIndex, const int32 It, const int32 NumIts)
+	
+	void FPBDJointConstraints::ProjectPosition(const FReal Dt, const int32 ConstraintIndex, const int32 It, const int32 NumIts)
 	{
-		const FJointSettings& JointSettings = ConstraintSettings[ConstraintIndex];
+		const FPBDJointSettings& JointSettings = ConstraintSettings[ConstraintIndex];
 
 		// Scale projection up to ProjectionSetting over NumProjectionIts
-		const T ProjectionFactor = (Settings.Projection > 0) ? Settings.Projection : JointSettings.Motion.Projection;
-		if (ProjectionFactor == (T)0)
+		const FReal ProjectionFactor = (Settings.Projection > 0) ? Settings.Projection : JointSettings.Motion.Projection;
+		if (ProjectionFactor == (FReal)0)
 		{
 			return;
 		}
 
-		const TVector<TGeometryParticleHandle<T, d>*, 2>& Constraint = ConstraintParticles[ConstraintIndex];
+		const TVector<TGeometryParticleHandle<FReal, 3>*, 2>& Constraint = ConstraintParticles[ConstraintIndex];
 		UE_LOG(LogChaosJoint, Verbose, TEXT("Project Joint Constraint %d %f (it = %d / %d)"), ConstraintIndex, ProjectionFactor, It, NumIts);
 
 		// Switch particles - internally we assume the first body is the parent (i.e., the space in which constraint limits are specified)
 		const int32 Index0 = 1;
 		const int32 Index1 = 0;
-		TGenericParticleHandle<T, d> Particle0 = TGenericParticleHandle<T, d>(ConstraintParticles[ConstraintIndex][Index0]);
-		TGenericParticleHandle<T, d> Particle1 = TGenericParticleHandle<T, d>(ConstraintParticles[ConstraintIndex][Index1]);
+		TGenericParticleHandle<FReal, 3> Particle0 = TGenericParticleHandle<FReal, 3>(ConstraintParticles[ConstraintIndex][Index0]);
+		TGenericParticleHandle<FReal, 3> Particle1 = TGenericParticleHandle<FReal, 3>(ConstraintParticles[ConstraintIndex][Index1]);
 
-		TVector<T, d> P0 = Particle0->P();
-		TRotation<T, d> Q0 = Particle0->Q();
-		TVector<T, d> P1 = Particle1->P();
-		TRotation<T, d> Q1 = Particle1->Q();
+		FVec3 P0 = Particle0->P();
+		FRotation3 Q0 = Particle0->Q();
+		FVec3 P1 = Particle1->P();
+		FRotation3 Q1 = Particle1->Q();
 		float InvM0 = Particle0->InvM();
 		float InvM1 = Particle1->InvM();
-		PMatrix<T, d, d> InvIL0 = Particle0->InvI();
-		PMatrix<T, d, d> InvIL1 = Particle1->InvI();
+		FMatrix33 InvIL0 = Particle0->InvI();
+		FMatrix33 InvIL1 = Particle1->InvI();
 
 		// Freeze the closest to kinematic connection if there is a difference
 		const int32 Level0 = ConstraintStates[ConstraintIndex].ParticleLevels[Index0];
@@ -752,20 +752,20 @@ namespace Chaos
 		if (Level0 < Level1)
 		{
 			InvM0 = 0;
-			InvIL0 = PMatrix<T, d, d>(0, 0, 0);
+			InvIL0 = FMatrix33(0, 0, 0);
 		}
 		else if (Level1 < Level0)
 		{
 			InvM1 = 0;
-			InvIL1 = PMatrix<T, d, d>(0, 0, 0);
+			InvIL1 = FMatrix33(0, 0, 0);
 		}
 
 		// Project position error
-		TPBDJointUtilities<T, d>::ApplyJointPositionProjection(Dt, Settings, JointSettings, Index0, Index1, P0, Q0, P1, Q1, InvM0, InvIL0, InvM1, InvIL1, ProjectionFactor);
+		FPBDJointUtilities::ApplyJointPositionProjection(Dt, Settings, JointSettings, Index0, Index1, P0, Q0, P1, Q1, InvM0, InvIL0, InvM1, InvIL1, ProjectionFactor);
 
 		// Update the particles
-		TPBDRigidParticleHandle<T, d>* Rigid0 = ConstraintParticles[ConstraintIndex][Index0]->AsDynamic();
-		TPBDRigidParticleHandle<T, d>* Rigid1 = ConstraintParticles[ConstraintIndex][Index1]->AsDynamic();
+		TPBDRigidParticleHandle<FReal, 3>* Rigid0 = ConstraintParticles[ConstraintIndex][Index0]->AsDynamic();
+		TPBDRigidParticleHandle<FReal, 3>* Rigid1 = ConstraintParticles[ConstraintIndex][Index1]->AsDynamic();
 		if (Rigid0)
 		{
 			Rigid0->SetP(P0);
@@ -781,9 +781,5 @@ namespace Chaos
 
 namespace Chaos
 {
-	template class TPBDJointSettings<float, 3>;
-	template class TPBDJointSolverSettings<float, 3>;
-	template class TPBDJointConstraintHandle<float, 3>;
-	template class TContainerConstraintHandle<TPBDJointConstraints<float, 3>>;
-	template class TPBDJointConstraints<float, 3>;
+	template class TContainerConstraintHandle<FPBDJointConstraints>;
 }
diff --git a/Engine/Source/Runtime/Experimental/Chaos/Public/Chaos/ChaosDebugDraw.h b/Engine/Source/Runtime/Experimental/Chaos/Public/Chaos/ChaosDebugDraw.h
--- a/Engine/Source/Runtime/Experimental/Chaos/Public/Chaos/ChaosDebugDraw.h
+++ b/Engine/Source/Runtime/Experimental/Chaos/Public/Chaos/ChaosDebugDraw.h
@@ -10,17 +10,17 @@ namespace Chaos
 	namespace DebugDraw
 	{
 #if CHAOS_DEBUG_DRAW
-		CHAOS_API void DrawParticleShapes(const TRigidTransform<float, 3>& SpaceTransform, const TParticleView<TGeometryParticles<float, 3>>& ParticlesView, float ColorScale, bool bDrawKinematic = true, bool bDrawDynamic = true);
-		CHAOS_API void DrawParticleShapes(const TRigidTransform<float, 3>& SpaceTransform, const TArray<TGeometryParticleHandle<float, 3>*>& Particles, float ColorScale, bool bDrawKinematic = true, bool bDrawDynamic = true);
-		CHAOS_API void DrawParticleTransforms(const TRigidTransform<float, 3>& SpaceTransform, const TParticleView<TGeometryParticles<float, 3>>& ParticlesView, float ColorScale, bool bDrawKinematic = true, bool bDrawDynamic = true);
-		CHAOS_API void DrawParticleTransforms(const TRigidTransform<float, 3>& SpaceTransform, const TArray<TGeometryParticleHandle<float, 3>*>& Particles, float ColorScale, bool bDrawKinematic = true, bool bDrawDynamic = true);
-		CHAOS_API void DrawParticleCollisions(const TRigidTransform<float, 3>& SpaceTransform, const TGeometryParticleHandle<float, 3>* Particle, const TPBDCollisionConstraint<float, 3>& Collisions);
-		CHAOS_API void DrawCollisions(const TRigidTransform<float, 3>& SpaceTransform, const TPBDCollisionConstraint<float, 3>& Collisions, float ColorScale);
-		CHAOS_API void DrawCollisions(const TRigidTransform<float, 3>& SpaceTransform, const TArray<TPBDCollisionConstraintHandle<float, 3>*>& ConstraintHandles, float ColorScale);
-		CHAOS_API void DrawJointConstraints(const TRigidTransform<float, 3>& SpaceTransform, const TArray<TPBDJointConstraintHandle<float, 3>*>& ConstraintHandles, float ColorScale);
-		CHAOS_API void DrawJointConstraints(const TRigidTransform<float, 3>& SpaceTransform, const TPBDJointConstraints<float, 3>& Constraints, float ColorScale);
-		CHAOS_API void Draw6DofConstraints(const TRigidTransform<float, 3>& SpaceTransform, const TArray<TPBD6DJointConstraintHandle<float, 3>*>& ConstraintHandles, float ColorScale);
-		CHAOS_API void Draw6DofConstraints(const TRigidTransform<float, 3>& SpaceTransform, const TPBD6DJointConstraints<float, 3>& Constraints, float ColorScale);
+		CHAOS_API void DrawParticleShapes(const FRigidTransform3& SpaceTransform, const TParticleView<TGeometryParticles<float, 3>>& ParticlesView, float ColorScale, bool bDrawKinematic = true, bool bDrawDynamic = true);
+		CHAOS_API void DrawParticleShapes(const FRigidTransform3& SpaceTransform, const TArray<TGeometryParticleHandle<float, 3>*>& Particles, float ColorScale, bool bDrawKinematic = true, bool bDrawDynamic = true);
+		CHAOS_API void DrawParticleTransforms(const FRigidTransform3& SpaceTransform, const TParticleView<TGeometryParticles<float, 3>>& ParticlesView, float ColorScale, bool bDrawKinematic = true, bool bDrawDynamic = true);
+		CHAOS_API void DrawParticleTransforms(const FRigidTransform3& SpaceTransform, const TArray<TGeometryParticleHandle<float, 3>*>& Particles, float ColorScale, bool bDrawKinematic = true, bool bDrawDynamic = true);
+		CHAOS_API void DrawParticleCollisions(const FRigidTransform3& SpaceTransform, const TGeometryParticleHandle<float, 3>* Particle, const TPBDCollisionConstraint<float, 3>& Collisions);
+		CHAOS_API void DrawCollisions(const FRigidTransform3& SpaceTransform, const TPBDCollisionConstraint<float, 3>& Collisions, float ColorScale);
+		CHAOS_API void DrawCollisions(const FRigidTransform3& SpaceTransform, const TArray<TPBDCollisionConstraintHandle<float, 3>*>& ConstraintHandles, float ColorScale);
+		CHAOS_API void DrawJointConstraints(const FRigidTransform3& SpaceTransform, const TArray<FPBDJointConstraintHandle*>& ConstraintHandles, float ColorScale);
+		CHAOS_API void DrawJointConstraints(const FRigidTransform3& SpaceTransform, const FPBDJointConstraints& Constraints, float ColorScale);
+		CHAOS_API void Draw6DofConstraints(const FRigidTransform3& SpaceTransform, const TArray<FPBD6DJointConstraintHandle*>& ConstraintHandles, float ColorScale);
+		CHAOS_API void Draw6DofConstraints(const FRigidTransform3& SpaceTransform, const FPBD6DJointConstraints& Constraints, float ColorScale);
 
 		extern CHAOS_API float ConstraintAxisLen;
 		extern CHAOS_API float BodyAxisLen;
diff --git a/Engine/Source/Runtime/Experimental/Chaos/Public/Chaos/ChaosDebugDrawDeclares.h b/Engine/Source/Runtime/Experimental/Chaos/Public/Chaos/ChaosDebugDrawDeclares.h
--- a/Engine/Source/Runtime/Experimental/Chaos/Public/Chaos/ChaosDebugDrawDeclares.h
+++ b/Engine/Source/Runtime/Experimental/Chaos/Public/Chaos/ChaosDebugDrawDeclares.h
@@ -1,6 +1,7 @@
 // Copyright 1998-2019 Epic Games, Inc. All Rights Reserved.
 #pragma once
 
+#include "Chaos/Core.h"
 #include "Chaos/Declares.h"
 #include "Chaos/GeometryParticlesfwd.h"
 #include "Misc/Build.h"
@@ -11,14 +12,11 @@
 
 namespace Chaos
 {
-	template<typename T, int d>
-	class TConstraintHandle;
+	class FConstraintHandle;
 
-	template<typename T, int d>
-	class TPBD6DJointConstraintHandle;
+	class FPBD6DJointConstraintHandle;
 
-	template<typename T, int d>
-	class TPBD6DJointConstraints;
+	class FPBD6DJointConstraints;
 
 	template<typename T, int d>
 	class TPBDCollisionConstraint;
@@ -26,23 +24,7 @@ namespace Chaos
 	template<typename T, int d>
 	class TPBDCollisionConstraintHandle;
 
-	template<typename T, int d>
-	class TPBDJointConstraintHandle;
-
-	template<typename T, int d>
-	class TPBDJointConstraints;
-
-	template<typename T, int d>
-	class TRigidTransform;
-
-
-
-	template<typename T, int m, int n>
-	class PMatrix;
+	class FPBDJointConstraintHandle;
 
-	template<typename T, int d>
-	class TRotation;
-
-	template<typename T, int d>
-	class TVector;
+	class FPBDJointConstraints;
 }
diff --git a/Engine/Source/Runtime/Experimental/Chaos/Public/Chaos/ConstraintHandle.h b/Engine/Source/Runtime/Experimental/Chaos/Public/Chaos/ConstraintHandle.h
--- a/Engine/Source/Runtime/Experimental/Chaos/Public/Chaos/ConstraintHandle.h
+++ b/Engine/Source/Runtime/Experimental/Chaos/Public/Chaos/ConstraintHandle.h
@@ -1,29 +1,25 @@
 // Copyright 1998-2019 Epic Games, Inc. All Rights Reserved.
 #pragma once
 
+#include "Chaos/Core.h"
 #include "Chaos/Declares.h"
 #include "Chaos/Vector.h"
 
 namespace Chaos
 {
-	template<class T, int d>
-	class TPBDConstraintContainer;
+	class FPBDConstraintContainer;
 
 	/**
 	 * Base class for constraint handles.
-	 * @todo(ccaulfield): Add checked down-casting to specific constraint-type handle
 	 */
-	template<class T, int d>
-	class CHAOS_API TConstraintHandle
+	class CHAOS_API FConstraintHandle
 	{
 	public:
-		using FReal = T;
-		static const int Dimensions = d;
-		using FGeometryParticleHandle = TGeometryParticleHandle<FReal, Dimensions>;
+		using FGeometryParticleHandle = TGeometryParticleHandle<FReal, 3>;
 
-		TConstraintHandle() : ConstraintIndex(INDEX_NONE) {}
-		TConstraintHandle(int32 InConstraintIndex) : ConstraintIndex(InConstraintIndex) {}
-		~TConstraintHandle() {}
+		FConstraintHandle() : ConstraintIndex(INDEX_NONE) {}
+		FConstraintHandle(int32 InConstraintIndex) : ConstraintIndex(InConstraintIndex) {}
+		~FConstraintHandle() {}
 
 		bool IsValid() const
 		{
@@ -35,8 +31,12 @@ namespace Chaos
 			return ConstraintIndex;
 		}
 
+		// @todo(ccaulfield): Add checked down - casting to specific constraint - type handle
+		template<typename T> 
+		T* As() { return static_cast<T*>(this); }
+
 	protected:
-		friend class TPBDConstraintContainer<T, d>;
+		friend class FPBDConstraintContainer;
 
 		int32 ConstraintIndex;
 	};
@@ -46,17 +46,15 @@ namespace Chaos
 	 * Utility base class for ConstraintHandles. Provides basic functionality common to most constraint containers.
 	 */
 	template<typename T_CONTAINER>
-	class CHAOS_API TContainerConstraintHandle : public TConstraintHandle<typename T_CONTAINER::FReal, T_CONTAINER::Dimensions>
+	class CHAOS_API TContainerConstraintHandle : public FConstraintHandle
 	{
 	public:
-		using Base = TConstraintHandle<typename T_CONTAINER::FReal, T_CONTAINER::Dimensions>;
-		using FReal = typename Base::FReal;
-		static const int Dimensions = Base::Dimensions;
+		using Base = FConstraintHandle;
 		using FGeometryParticleHandle = typename Base::FGeometryParticleHandle;
 		using FConstraintContainer = T_CONTAINER;
 
 		TContainerConstraintHandle() : ConstraintContainer(nullptr) {}
-		TContainerConstraintHandle(FConstraintContainer* InConstraintContainer, int32 InConstraintIndex) : TConstraintHandle<FReal, Dimensions>(InConstraintIndex), ConstraintContainer(InConstraintContainer) {}
+		TContainerConstraintHandle(FConstraintContainer* InConstraintContainer, int32 InConstraintIndex) : FConstraintHandle(InConstraintIndex), ConstraintContainer(InConstraintContainer) {}
 
 		void RemoveConstraint() { ConstraintContainer->RemoveConstraint(ConstraintIndex); }
 
@@ -78,9 +76,9 @@ namespace Chaos
 	{
 	public:
 		using FConstraintContainer = T_CONTAINER;
-		using FConstraintHandle = typename FConstraintContainer::FConstraintHandle;
+		using FConstraintContainerHandle = typename FConstraintContainer::FConstraintContainerHandle;
 
-		FConstraintHandle* AllocHandle(FConstraintContainer* ConstraintContainer, int32 ConstraintIndex) { return new FConstraintHandle(ConstraintContainer, ConstraintIndex); }
-		void FreeHandle(FConstraintHandle* Handle) { delete Handle; }
+		FConstraintContainerHandle* AllocHandle(FConstraintContainer* ConstraintContainer, int32 ConstraintIndex) { return new FConstraintContainerHandle(ConstraintContainer, ConstraintIndex); }
+		void FreeHandle(FConstraintContainerHandle* Handle) { delete Handle; }
 	};
 }
diff --git a/Engine/Source/Runtime/Experimental/Chaos/Public/Chaos/PBD6DJointConstraints.h b/Engine/Source/Runtime/Experimental/Chaos/Public/Chaos/PBD6DJointConstraints.h
--- a/Engine/Source/Runtime/Experimental/Chaos/Public/Chaos/PBD6DJointConstraints.h
+++ b/Engine/Source/Runtime/Experimental/Chaos/Public/Chaos/PBD6DJointConstraints.h
@@ -3,6 +3,7 @@
 
 #include "CoreMinimal.h"
 
+#include "Chaos/Core.h"
 #include "Chaos/Array.h"
 #include "Chaos/Transform.h"
 #include "Chaos/Vector.h"
@@ -13,33 +14,29 @@
 
 namespace Chaos
 {
-	template<class T, int d>
-	class TPBD6DJointConstraints;
+	class FPBD6DJointConstraints;
 
-	template<class T, int d>
-	class TPBD6DJointConstraintHandle : public TContainerConstraintHandle<TPBD6DJointConstraints<T, d>>
+	class CHAOS_API FPBD6DJointConstraintHandle : public TContainerConstraintHandle<FPBD6DJointConstraints>
 	{
 	public:
-		using Base = TContainerConstraintHandle<TPBD6DJointConstraints<T, d>>;
-		using FConstraintContainer = TPBD6DJointConstraints<T, d>;
+		using Base = TContainerConstraintHandle<FPBD6DJointConstraints>;
+		using FConstraintContainer = FPBD6DJointConstraints;
 
-		CHAOS_API TPBD6DJointConstraintHandle();
-		CHAOS_API TPBD6DJointConstraintHandle(FConstraintContainer* InConstraintContainer, int32 InConstraintIndex);
+		FPBD6DJointConstraintHandle();
+		FPBD6DJointConstraintHandle(FConstraintContainer* InConstraintContainer, int32 InConstraintIndex);
 
-		CHAOS_API void CalculateConstraintSpace(TVector<T, d>& OutXa, PMatrix<T, d, d>& OutRa, TVector<T, d>& OutXb, PMatrix<T, d, d>& OutRb, TVector<T, d>& OutCR) const;
-		CHAOS_API void SetParticleLevels(const TVector<int32, 2>& ParticleLevels);
-		CHAOS_API int32 GetConstraintLevel() const;
+		void CalculateConstraintSpace(FVec3& OutXa, FMatrix33& OutRa, FVec3& OutXb, FMatrix33& OutRb, FVec3& OutCR) const;
+		void SetParticleLevels(const TVector<int32, 2>& ParticleLevels);
+		int32 GetConstraintLevel() const;
 
 	protected:
 		using Base::ConstraintIndex;
 		using Base::ConstraintContainer;
 	};
 
-	template<typename T, int d>
-	using TD6JointPostApplyCallback = TFunction<void(const T Dt, const TArray<TPBD6DJointConstraintHandle<T, d>*>& InConstraintHandles)>;
+	using FD6JointPostApplyCallback = TFunction<void(const FReal Dt, const TArray<FPBD6DJointConstraintHandle*>& InConstraintHandles)>;
 
-	template<typename T, int d>
-	using TD6JointPreApplyCallback = TFunction<void(const T Dt, const TArray<TPBD6DJointConstraintHandle<T, d>*>& InConstraintHandles)>;
+	using FD6JointPreApplyCallback = TFunction<void(const FReal Dt, const TArray<FPBD6DJointConstraintHandle*>& InConstraintHandles)>;
 
 	enum class E6DJointMotionType : int32
 	{
@@ -76,104 +73,99 @@ namespace Chaos
 	};
 
 
-	template<class T, int d>
-	struct T6DJointConstants
+	struct F6DJointConstants
 	{
 		/** The constraint-space twist axis (X Axis) */
-		static const TVector<T, d> TwistAxis() { return TVector<T, d>(1, 0, 0); }
+		static const FVec3 TwistAxis() { return FVec3(1, 0, 0); }
 
 		/** The constraint-space Swing1 axis (Z Axis) */
-		static const TVector<T, d> Swing1Axis() { return TVector<T, d>(0, 0, 1); }
+		static const FVec3 Swing1Axis() { return FVec3(0, 0, 1); }
 
 		/** The constraint-space Swing2 axis (Y Axis) */
-		static const TVector<T, d> Swing2Axis() { return TVector<T, d>(0, 1, 0); }
+		static const FVec3 Swing2Axis() { return FVec3(0, 1, 0); }
 	};
 
-	template<class T, int d>
-	class CHAOS_API TPBD6DJointMotionSettings
+	class CHAOS_API FPBD6DJointMotionSettings
 	{
 	public:
-		TPBD6DJointMotionSettings();
-		TPBD6DJointMotionSettings(const TVector<E6DJointMotionType, d>& InLinearMotionTypes, const TVector<E6DJointMotionType, d>& InAngularMotionTypes);
+		FPBD6DJointMotionSettings();
+		FPBD6DJointMotionSettings(const TVector<E6DJointMotionType, 3>& InLinearMotionTypes, const TVector<E6DJointMotionType, 3>& InAngularMotionTypes);
 
-		T Stiffness;
+		FReal Stiffness;
 
-		TVector<E6DJointMotionType, d> LinearMotionTypes;
-		TVector<T, d> LinearLimits;
+		TVector<E6DJointMotionType, 3> LinearMotionTypes;
+		FVec3 LinearLimits;
 
-		TVector<E6DJointMotionType, d> AngularMotionTypes;
-		TVector<T, d> AngularLimits;
+		TVector<E6DJointMotionType, 3> AngularMotionTypes;
+		FVec3 AngularLimits;
 
 		// @todo(ccaulfield): remove one of these
-		TRotation<T, d> AngularDriveTarget;
-		TVector<T, d> AngularDriveTargetAngles;
+		FRotation3 AngularDriveTarget;
+		FVec3 AngularDriveTargetAngles;
 
 		bool bAngularSLerpDriveEnabled;
 		bool bAngularTwistDriveEnabled;
 		bool bAngularSwingDriveEnabled;
 
-		T AngularDriveStiffness;
-		T AngularDriveDamping;
+		FReal AngularDriveStiffness;
+		FReal AngularDriveDamping;
 	};
 
 
-	template<class T, int d>
-	class CHAOS_API TPBD6DJointSettings
+	class CHAOS_API FPBD6DJointSettings
 	{
 	public:
-		using FTransformPair = TVector<TRigidTransform<T, d>, 2>;
+		using FTransformPair = TVector<FRigidTransform3, 2>;
 
-		TPBD6DJointSettings();
+		FPBD6DJointSettings();
 
 		// Particle-relative joint axes and positions
 		FTransformPair ConstraintFrames;
 
 		// How the constraint is allowed to move
-		TPBD6DJointMotionSettings<T, d> Motion;
+		FPBD6DJointMotionSettings Motion;
 	};
 
-	template<class T, int d>
-	class CHAOS_API TPBD6DJointState
+	class CHAOS_API FPBD6DJointState
 	{
 	public:
-		TPBD6DJointState();
+		FPBD6DJointState();
 
 		// XPBD state lambda (See XPBD paper for info)
 		// This needs to be initialized each frame before we start iterating, but we only know
 		// some of the values during the iteration, hence the initialize flag.
-		TVector<T, d> LambdaXa;
-		TVector<T, d> LambdaRa;
-		TVector<T, d> LambdaXb;
-		TVector<T, d> LambdaRb;
-		TVector<T, d> PrevTickCX;
-		TVector<T, d> PrevTickCR;
-		TVector<T, d> PrevItCX;
-		TVector<T, d> PrevItCR;
+		FVec3 LambdaXa;
+		FVec3 LambdaRa;
+		FVec3 LambdaXb;
+		FVec3 LambdaRb;
+		FVec3 PrevTickCX;
+		FVec3 PrevTickCR;
+		FVec3 PrevItCX;
+		FVec3 PrevItCR;
 
 		// Priorities used for ordering, mass conditioning, projection, and freezing
 		int32 Level;
 		TVector<int32, 2> ParticleLevels;
 	};
 
-	template<class T, int d>
-	class CHAOS_API TPBD6DJointSolverSettings
+	class CHAOS_API FPBD6DJointSolverSettings
 	{
 	public:
-		TPBD6DJointSolverSettings();
+		FPBD6DJointSolverSettings();
 
 		// Tolerances
-		T SolveTolerance;
-		T InvertedAxisTolerance;
-		T SwingTwistAngleTolerance;
+		FReal SolveTolerance;
+		FReal InvertedAxisTolerance;
+		FReal SwingTwistAngleTolerance;
 
 		// Stability control
 		bool bApplyProjection;
 		int32 MaxIterations;
 		int32 MaxPreIterations;
 		int32 MaxDriveIterations;
-		T MaxRotComponent;
-		T PBDMinParentMassRatio;
-		T PBDMaxInertiaRatio;
+		FReal MaxRotComponent;
+		FReal PBDMinParentMassRatio;
+		FReal PBDMaxInertiaRatio;
 		int32 FreezeIterations;
 		int32 FrozenIterations;
 
@@ -182,12 +174,12 @@ namespace Chaos
 		bool bEnableTwistLimits;
 		bool bEnableSwingLimits;
 		bool bEnableDrives;
-		T XPBDAlphaX;
-		T XPBDAlphaR;
-		T XPBDBetaX;
-		T XPBDBetaR;
-		T PBDStiffness;
-		T PBDDriveStiffness;
+		FReal XPBDAlphaX;
+		FReal XPBDAlphaR;
+		FReal XPBDBetaX;
+		FReal XPBDBetaR;
+		FReal PBDStiffness;
+		FReal PBDDriveStiffness;
 
 		bool bFastSolve;
 	};
@@ -195,27 +187,22 @@ namespace Chaos
 	/**
 	 * A joint restricting up to 6 degrees of freedom, with linear and angular limits.
 	 */
-	template<class T, int d>
-	class TPBD6DJointConstraints : public TPBDConstraintContainer<T, d>
+	class CHAOS_API FPBD6DJointConstraints : public FPBDConstraintContainer
 	{
 	public:
-		using Base = TPBDConstraintContainer<T, d>;
-		using FReal = T;
-		static const int Dimensions = d;
-		using FConstraintHandle = TPBD6DJointConstraintHandle<FReal, Dimensions>;
-		using FConstraintHandleAllocator = TConstraintHandleAllocator<TPBD6DJointConstraints<FReal, Dimensions>>;
-		using FParticlePair = TVector<TGeometryParticleHandle<T, d>*, 2>;
-		using FVectorPair = TVector<TVector<T, d>, 2>;
-		using FTransformPair = TVector<TRigidTransform<T, d>, 2>;
-		using FJointSettings = TPBD6DJointSettings<T, d>;
-		using FJointState = TPBD6DJointState<T, d>;
+		using Base = FPBDConstraintContainer;
+		using FConstraintContainerHandle = FPBD6DJointConstraintHandle;
+		using FConstraintHandleAllocator = TConstraintHandleAllocator<FPBD6DJointConstraints>;
+		using FParticlePair = TVector<TGeometryParticleHandle<FReal, 3>*, 2>;
+		using FVectorPair = TVector<FVec3, 2>;
+		using FTransformPair = TVector<FRigidTransform3, 2>;
 
-		CHAOS_API TPBD6DJointConstraints(const TPBD6DJointSolverSettings<T, d>& InSettings);
+		FPBD6DJointConstraints(const FPBD6DJointSolverSettings& InSettings);
 
-		virtual ~TPBD6DJointConstraints();
+		virtual ~FPBD6DJointConstraints();
 
-		CHAOS_API const TPBD6DJointSolverSettings<T, d>& GetSettings() const;
-		CHAOS_API void SetSettings(const TPBD6DJointSolverSettings<T, d>& InSettings);
+		const FPBD6DJointSolverSettings& GetSettings() const;
+		void SetSettings(const FPBD6DJointSolverSettings& InSettings);
 
 		//
 		// Constraint Container API
@@ -224,39 +211,39 @@ namespace Chaos
 		/**
 		 * Get the number of constraints.
 		 */
-		CHAOS_API int32 NumConstraints() const;
+		int32 NumConstraints() const;
 
 		/** 
 		 * Add a constraint with particle-space constraint offsets. 
 		 */
-		CHAOS_API FConstraintHandle* AddConstraint(const FParticlePair& InConstrainedParticles, const FTransformPair& ConstraintFrames);
-		CHAOS_API FConstraintHandle* AddConstraint(const FParticlePair& InConstrainedParticles, const TPBD6DJointSettings<T, d>& InConstraintSettings);
+		FConstraintContainerHandle* AddConstraint(const FParticlePair& InConstrainedParticles, const FTransformPair& ConstraintFrames);
+		FConstraintContainerHandle* AddConstraint(const FParticlePair& InConstrainedParticles, const FPBD6DJointSettings& InConstraintSettings);
 
 		/**
 		 * Remove the specified constraint.
 		 */
-		CHAOS_API void RemoveConstraint(int ConstraintIndex);
+		void RemoveConstraint(int ConstraintIndex);
 
 		// @todo(ccaulfield): rename/remove  this
-		CHAOS_API void RemoveConstraints(const TSet<TGeometryParticleHandle<T, d>*>& RemovedParticles);
+		void RemoveConstraints(const TSet<TGeometryParticleHandle<FReal, 3>*>& RemovedParticles);
 
-		CHAOS_API void SetPreApplyCallback(const TD6JointPostApplyCallback<T, d>& Callback);
-		CHAOS_API void ClearPreApplyCallback();
+		void SetPreApplyCallback(const FD6JointPostApplyCallback& Callback);
+		void ClearPreApplyCallback();
 
-		CHAOS_API void SetPostApplyCallback(const TD6JointPostApplyCallback<T, d>& Callback);
-		CHAOS_API void ClearPostApplyCallback();
+		void SetPostApplyCallback(const FD6JointPostApplyCallback& Callback);
+		void ClearPostApplyCallback();
 
 		//
 		// Constraint API
 		//
 
-		CHAOS_API const FConstraintHandle* GetConstraintHandle(int32 ConstraintIndex) const;
-		CHAOS_API FConstraintHandle* GetConstraintHandle(int32 ConstraintIndex);
+		const FConstraintContainerHandle* GetConstraintHandle(int32 ConstraintIndex) const;
+		FConstraintContainerHandle* GetConstraintHandle(int32 ConstraintIndex);
 
 		/**
 		 * Get the particles that are affected by the specified constraint.
 		 */
-		CHAOS_API const FParticlePair& GetConstrainedParticles(int32 ConstraintIndex) const;
+		const FParticlePair& GetConstrainedParticles(int32 ConstraintIndex) const;
 
 		int32 GetConstraintLevel(int32 ConstraintIndex) const;
 		void SetParticleLevels(int32 ConstraintIndex, const TVector<int32, 2>& ParticleLevels);
@@ -265,194 +252,193 @@ namespace Chaos
 		// Island Rule API
 		//
 
-		CHAOS_API void UpdatePositionBasedState(const T Dt);
+		void UpdatePositionBasedState(const FReal Dt);
 
-		CHAOS_API void Apply(const T Dt, const TArray<FConstraintHandle*>& InConstraintHandles, const int32 It, const int32 NumIts);
+		void Apply(const FReal Dt, const TArray<FConstraintContainerHandle*>& InConstraintHandles, const int32 It, const int32 NumIts);
 
-		CHAOS_API bool ApplyPushOut(const T Dt, const TArray<FConstraintHandle*>& InConstraintHandles, const int32 It, const int32 NumIts);
+		bool ApplyPushOut(const FReal Dt, const TArray<FConstraintContainerHandle*>& InConstraintHandles, const int32 It, const int32 NumIts);
 
 	protected:
 		using Base::GetConstraintIndex;
 		using Base::SetConstraintIndex;
 
 	private:
-		friend class TPBD6DJointConstraintHandle<T, d>;
+		friend class FPBD6DJointConstraintHandle;
 
-		CHAOS_API void CalculateConstraintSpace(int32 ConstraintIndex, TVector<T, d>& OutXa, PMatrix<T, d, d>& OutRa, TVector<T, d>& OutXb, PMatrix<T, d, d>& OutRb, TVector<T, d>& OutCR) const;
+		void CalculateConstraintSpace(int32 ConstraintIndex, FVec3& OutXa, FMatrix33& OutRa, FVec3& OutXb, FMatrix33& OutRb, FVec3& OutCR) const;
 
-		CHAOS_API void ApplySingle(const T Dt, const int32 ConstraintIndex, const T FreezeScale);
-		CHAOS_API void ApplyDynamicStatic(const T Dt, const int32 ConstraintIndex, const int32 PBDRigid0Index, const int32 Static1Index);
-		CHAOS_API void ApplyDynamicDynamic(const T Dt, const int32 ConstraintIndex, const int32 PBDRigid0Index, const int32 PBDRigid1Index, const T FreezeScale);
+		void ApplySingle(const FReal Dt, const int32 ConstraintIndex, const FReal FreezeScale);
+		void ApplyDynamicStatic(const FReal Dt, const int32 ConstraintIndex, const int32 PBDRigid0Index, const int32 Static1Index);
+		void ApplyDynamicDynamic(const FReal Dt, const int32 ConstraintIndex, const int32 PBDRigid0Index, const int32 PBDRigid1Index, const FReal FreezeScale);
 
-		CHAOS_API void ApplyPushOutSingle(const T Dt, const int32 ConstraintIndex);
+		void ApplyPushOutSingle(const FReal Dt, const int32 ConstraintIndex);
 
-		CHAOS_API void ApplySingleFast(const T Dt, const int32 ConstraintIndex, const int32 It, const int32 NumIts);
+		void ApplySingleFast(const FReal Dt, const int32 ConstraintIndex, const int32 It, const int32 NumIts);
 
-		TPBD6DJointSolverSettings<T, d> Settings;
+		FPBD6DJointSolverSettings Settings;
 
-		TArray<FJointSettings> ConstraintSettings;
+		TArray<FPBD6DJointSettings> ConstraintSettings;
 		TArray<FParticlePair> ConstraintParticles;
-		TArray<FJointState> ConstraintStates;
+		TArray<FPBD6DJointState> ConstraintStates;
 
-		TArray<FConstraintHandle*> Handles;
+		TArray<FConstraintContainerHandle*> Handles;
 		FConstraintHandleAllocator HandleAllocator;
 
-		TD6JointPreApplyCallback<T, d> PreApplyCallback;
-		TD6JointPostApplyCallback<T, d> PostApplyCallback;
+		FD6JointPreApplyCallback PreApplyCallback;
+		FD6JointPostApplyCallback PostApplyCallback;
 	};
 
-	template<class T, int d>
-	class CHAOS_API TPBD6DJointConstraintUtilities
+	class CHAOS_API FPBD6DJointConstraintUtilities
 	{
 	public:
-		static void BlockwiseInverse(const PMatrix<T, d, d>& A, const PMatrix<T, d, d>& B, const PMatrix<T, d, d>& C, const PMatrix<T, d, d>& D, PMatrix<T, d, d>& AI, PMatrix<T, d, d>& BI, PMatrix<T, d, d>& CI, PMatrix<T, d, d>& DI);
-		static void BlockwiseInverse2(const PMatrix<T, d, d>& A, const PMatrix<T, d, d>& B, const PMatrix<T, d, d>& C, const PMatrix<T, d, d>& D, PMatrix<T, d, d>& AI, PMatrix<T, d, d>& BI, PMatrix<T, d, d>& CI, PMatrix<T, d, d>& DI);
+		static void BlockwiseInverse(const FMatrix33& A, const FMatrix33& B, const FMatrix33& C, const FMatrix33& D, FMatrix33& AI, FMatrix33& BI, FMatrix33& CI, FMatrix33& DI);
+		static void BlockwiseInverse2(const FMatrix33& A, const FMatrix33& B, const FMatrix33& C, const FMatrix33& D, FMatrix33& AI, FMatrix33& BI, FMatrix33& CI, FMatrix33& DI);
 		
-		static void ComputeJointFactorMatrix(const PMatrix<T, d, d>& XR, const PMatrix<T, d, d>& RR, float MInv, const PMatrix<T, d, d>& IInv, PMatrix<T, d, d>& M00, PMatrix<T, d, d>& M01, PMatrix<T, d, d>& M10, PMatrix<T, d, d>& M11);
+		static void ComputeJointFactorMatrix(const FMatrix33& XR, const FMatrix33& RR, float MInv, const FMatrix33& IInv, FMatrix33& M00, FMatrix33& M01, FMatrix33& M10, FMatrix33& M11);
 
-		static TVector<T, d> Calculate6dConstraintAngles(
-			const TPBD6DJointSolverSettings<T, d>& SolverSettings,
-			const TRotation<T, d>& Ra,
-			const TRotation<T, d>& Rb, 
-			const TPBD6DJointMotionSettings<T, d>& MotionSettings);
+		static FVec3 Calculate6dConstraintAngles(
+			const FPBD6DJointSolverSettings& SolverSettings,
+			const FRotation3& Ra,
+			const FRotation3& Rb, 
+			const FPBD6DJointMotionSettings& MotionSettings);
 		
 		static bool Calculate6dConstraintRotation(
-			const TPBD6DJointSolverSettings<T, d>& SolverSettings,
-			const TRotation<T, d>& Ra,
-			const TRotation<T, d>& Rb, 
-			const TPBD6DJointMotionSettings<T, d>& MotionSettings,
-			TVector<T, d>& CR, 
-			PMatrix<T, d, d>& RRa, 
-			PMatrix<T, d, d>& RRb);
+			const FPBD6DJointSolverSettings& SolverSettings,
+			const FRotation3& Ra,
+			const FRotation3& Rb, 
+			const FPBD6DJointMotionSettings& MotionSettings,
+			FVec3& CR, 
+			FMatrix33& RRa, 
+			FMatrix33& RRb);
 		
 		static void Calculate6dConstraintRotationLimits(
-			const TPBD6DJointSolverSettings<T, d>& SolverSettings,
-			const TRotation<T, d>& Ra,
-			const TRotation<T, d>& Rb, 
-			const TPBD6DJointMotionSettings<T, d>& MotionSettings,
-			TVector<T, d>& SR,
-			TVector<T, d>& CR,
-			PMatrix<T, d, d>& RRa, 
-			PMatrix<T, d, d>& RRb, 
-			TVector<T, d>& LRMin, 
-			TVector<T, d>& LRMax);
+			const FPBD6DJointSolverSettings& SolverSettings,
+			const FRotation3& Ra,
+			const FRotation3& Rb, 
+			const FPBD6DJointMotionSettings& MotionSettings,
+			FVec3& SR,
+			FVec3& CR,
+			FMatrix33& RRa, 
+			FMatrix33& RRb, 
+			FVec3& LRMin, 
+			FVec3& LRMax);
 
 		static bool Calculate6dConstraintRotation_SwingCone(
-			const TPBD6DJointSolverSettings<T, d>& SolverSettings,
-			const TRotation<T, d>& Ra,
-			const TRotation<T, d>& Rb, 
-			const TPBD6DJointMotionSettings<T, d>& MotionSettings,
-			TVector<T, d>& CR, 
-			PMatrix<T, d, d>& RRa, 
-			PMatrix<T, d, d>& RRb);
+			const FPBD6DJointSolverSettings& SolverSettings,
+			const FRotation3& Ra,
+			const FRotation3& Rb, 
+			const FPBD6DJointMotionSettings& MotionSettings,
+			FVec3& CR, 
+			FMatrix33& RRa, 
+			FMatrix33& RRb);
 
 		static void Calculate6dConstraintRotationLimits_SwingCone(
-			const TPBD6DJointSolverSettings<T, d>& SolverSettings,
-			const TRotation<T, d>& Ra,
-			const TRotation<T, d>& Rb, 
-			const TPBD6DJointMotionSettings<T, d>& MotionSettings,
-			TVector<T, d>& SR,
-			TVector<T, d>& CR,
-			PMatrix<T, d, d>& RRa,
-			PMatrix<T, d, d>& RRb, 
-			TVector<T, d>& LRMin, 
-			TVector<T, d>& LRMax);
+			const FPBD6DJointSolverSettings& SolverSettings,
+			const FRotation3& Ra,
+			const FRotation3& Rb, 
+			const FPBD6DJointMotionSettings& MotionSettings,
+			FVec3& SR,
+			FVec3& CR,
+			FMatrix33& RRa,
+			FMatrix33& RRb, 
+			FVec3& LRMin, 
+			FVec3& LRMax);
 
 		static bool Calculate6dConstraintRotation_SwingFixed(
-			const TPBD6DJointSolverSettings<T, d>& SolverSettings,
-			const TRotation<T, d>& Ra,
-			const TRotation<T, d>& Rb, 
-			const TPBD6DJointMotionSettings<T, d>& MotionSettings,
-			TVector<T, d>& CR, 
-			PMatrix<T, d, d>& RRa, 
-			PMatrix<T, d, d>& RRb);
+			const FPBD6DJointSolverSettings& SolverSettings,
+			const FRotation3& Ra,
+			const FRotation3& Rb, 
+			const FPBD6DJointMotionSettings& MotionSettings,
+			FVec3& CR, 
+			FMatrix33& RRa, 
+			FMatrix33& RRb);
 
 		static void Calculate6dConstraintRotationLimits_SwingFixed(
-			const TPBD6DJointSolverSettings<T, d>& SolverSettings,
-			const TRotation<T, d>& Ra,
-			const TRotation<T, d>& Rb, 
-			const TPBD6DJointMotionSettings<T, d>& MotionSettings,
-			TVector<T, d>& SR,
-			TVector<T, d>& CR,
-			PMatrix<T, d, d>& RRa, 
-			PMatrix<T, d, d>& RRb, 
-			TVector<T, d>& LRMin, 
-			TVector<T, d>& LRMax);
+			const FPBD6DJointSolverSettings& SolverSettings,
+			const FRotation3& Ra,
+			const FRotation3& Rb, 
+			const FPBD6DJointMotionSettings& MotionSettings,
+			FVec3& SR,
+			FVec3& CR,
+			FMatrix33& RRa, 
+			FMatrix33& RRb, 
+			FVec3& LRMin, 
+			FVec3& LRMax);
 
 		static bool Calculate6dDelta(
-			const TPBD6DJointSolverSettings<T, d>& SolverSettings,
-			const T Dt,
-			const TVector<T, d>& Pa,
-			const TRotation<T, d>& Qa,
+			const FPBD6DJointSolverSettings& SolverSettings,
+			const FReal Dt,
+			const FVec3& Pa,
+			const FRotation3& Qa,
 			float MaInv,
-			const PMatrix<T, d, d>& IaInv,
-			const TVector<T, d>& Pb,
-			const TRotation<T, d>& Qb,
+			const FMatrix33& IaInv,
+			const FVec3& Pb,
+			const FRotation3& Qb,
 			float MbInv,
-			const PMatrix<T, d, d>& IbInv,
-			const TVector<T, d>& Xa,
-			const TRotation<T, d>& Ra,
-			const TVector<T, d>& Xb,
-			const TRotation<T, d>& Rb,
-			const TPBD6DJointMotionSettings<T, d>& MotionSettings,
-			TPBD6DJointState<T, d>& State,
-			TVector<T, d>& DPa,
-			TRotation<T, d>& DQa,
-			TVector<T, d>& DPb,
-			TRotation<T, d>& DQb);
+			const FMatrix33& IbInv,
+			const FVec3& Xa,
+			const FRotation3& Ra,
+			const FVec3& Xb,
+			const FRotation3& Rb,
+			const FPBD6DJointMotionSettings& MotionSettings,
+			FPBD6DJointState& State,
+			FVec3& DPa,
+			FRotation3& DQa,
+			FVec3& DPb,
+			FRotation3& DQb);
 
 		static int Solve6dConstraint(
-			const TPBD6DJointSolverSettings<T, d>& SolverSettings,
-			const T Dt,
-			TVector<T, d>& Pa,
-			TRotation<T, d>& Qa,
-			T MaInv,
-			const PMatrix<T, d, d>& ILaInv,
-			const TVector<T, d>& XLa,
-			const TRotation<T, d>& RLa,
-			TVector<T, d>& Pb,
-			TRotation<T, d>& Qb,
-			T MbInv,
-			const PMatrix<T, d, d>& ILbInv,
-			const TVector<T, d>& XLb,
-			const TRotation<T, d>& RLb,
-			const TPBD6DJointMotionSettings<T, d>& MotionSettings,
-			TPBD6DJointState<T, d>& State);
+			const FPBD6DJointSolverSettings& SolverSettings,
+			const FReal Dt,
+			FVec3& Pa,
+			FRotation3& Qa,
+			FReal MaInv,
+			const FMatrix33& ILaInv,
+			const FVec3& XLa,
+			const FRotation3& RLa,
+			FVec3& Pb,
+			FRotation3& Qb,
+			FReal MbInv,
+			const FMatrix33& ILbInv,
+			const FVec3& XLb,
+			const FRotation3& RLb,
+			const FPBD6DJointMotionSettings& MotionSettings,
+			FPBD6DJointState& State);
 
 
 		static void Calculate3dDelta(
-			const TPBD6DJointSolverSettings<T, d>& SolverSettings,
-			const TVector<T, d>& Pa,
-			const TRotation<T, d>& Qa,
-			float MaInv,
-			const PMatrix<T, d, d>& IaInv,
-			const TVector<T, d>& Pb,
-			const TRotation<T, d>& Qb,
-			float MbInv,
-			const PMatrix<T, d, d>& IbInv,
-			const TVector<T, d>& Xa,
-			const TVector<T, d>& Xb,
-			const TPBD6DJointMotionSettings<T, d>& XSettings,
-			TVector<T, d>& DPa,
-			TRotation<T, d>& DQa,
-			TVector<T, d>& DPb,
-			TRotation<T, d>& DQb);
+			const FPBD6DJointSolverSettings& SolverSettings,
+			const FVec3& Pa,
+			const FRotation3& Qa,
+			FReal MaInv,
+			const FMatrix33& IaInv,
+			const FVec3& Pb,
+			const FRotation3& Qb,
+			FReal MbInv,
+			const FMatrix33& IbInv,
+			const FVec3& Xa,
+			const FVec3& Xb,
+			const FPBD6DJointMotionSettings& XSettings,
+			FVec3& DPa,
+			FRotation3& DQa,
+			FVec3& DPb,
+			FRotation3& DQb);
 
 		static void Solve3dConstraint(
-			const TPBD6DJointSolverSettings<T, d>& SolverSettings,
-			TVector<T, d>& Pa,
-			TRotation<T, d>& Qa,
-			T MaInv,
-			const PMatrix<T, d, d>& ILaInv,
-			const TVector<T, d>& XLa,
-			const TRotation<T, d>& RLa,
-			TVector<T, d>& Pb,
-			TRotation<T, d>& Qb,
-			T MbInv,
-			const PMatrix<T, d, d>& ILbInv,
-			const TVector<T, d>& XLb,
-			const TRotation<T, d>& RLb,
-			const TPBD6DJointMotionSettings<T, d>& MotionSettings);
+			const FPBD6DJointSolverSettings& SolverSettings,
+			FVec3& Pa,
+			FRotation3& Qa,
+			FReal MaInv,
+			const FMatrix33& ILaInv,
+			const FVec3& XLa,
+			const FRotation3& RLa,
+			FVec3& Pb,
+			FRotation3& Qb,
+			FReal MbInv,
+			const FMatrix33& ILbInv,
+			const FVec3& XLb,
+			const FRotation3& RLb,
+			const FPBD6DJointMotionSettings& MotionSettings);
 	};
 
 }
diff --git a/Engine/Source/Runtime/Experimental/Chaos/Public/Chaos/PBDChainConstraints.h b/Engine/Source/Runtime/Experimental/Chaos/Public/Chaos/PBDChainConstraints.h
--- a/Engine/Source/Runtime/Experimental/Chaos/Public/Chaos/PBDChainConstraints.h
+++ b/Engine/Source/Runtime/Experimental/Chaos/Public/Chaos/PBDChainConstraints.h
@@ -8,7 +8,7 @@
 namespace Chaos
 {
 template<class T, int d>
-class TPBDChainConstraints : public TPBDConstraintContainer<T, d>
+class TPBDChainConstraints : public FPBDConstraintContainer
 {
 public:
 	TPBDChainConstraints(const TDynamicParticles<T, d>& InParticles, TArray<TArray<int32>>&& Constraints, const T Coefficient = (T)1)
diff --git a/Engine/Source/Runtime/Experimental/Chaos/Public/Chaos/PBDCollisionConstraint.h b/Engine/Source/Runtime/Experimental/Chaos/Public/Chaos/PBDCollisionConstraint.h
--- a/Engine/Source/Runtime/Experimental/Chaos/Public/Chaos/PBDCollisionConstraint.h
+++ b/Engine/Source/Runtime/Experimental/Chaos/Public/Chaos/PBDCollisionConstraint.h
@@ -77,7 +77,7 @@ class CHAOS_API TPBDCollisionConstraintHistory
 	typedef TPair<const FImplicitObject*, const FImplicitObject*> ImplicitPairs;
 
 public:
-	using FConstraintHandle = TPBDCollisionConstraintHandle<T, d>;
+	using FConstraintContainerHandle = TPBDCollisionConstraintHandle<T, d>;
 	TPBDCollisionConstraintHistory(const TVector<T, d>& InLocation, const TRotation<T, d>& InRotation, int32 InTimestamp = -INT_MAX)
 		: Timestamp(InTimestamp)
 		, Location(InLocation)
@@ -87,19 +87,19 @@ public:
 	const TVector<T, d>& GetLocation() { return Location; }
 	const TRotation<T, d>& GetRotation() { return Rotation; }
 
-	void AddHandle(FConstraintHandle* InHandle)
+	void AddHandle(FConstraintContainerHandle* InHandle)
 	{
 		Implicits.Add(ImplicitPairs(InHandle->GetContact().Geometry[0], InHandle->GetContact().Geometry[1]));
 		Handles.Push(InHandle);
 	}
-	void RemoveHandle(FConstraintHandle* InHandle)
+	void RemoveHandle(FConstraintContainerHandle* InHandle)
 	{
 		Handles.RemoveSingleSwap(InHandle);
 		Implicits.Remove(ImplicitPairs(InHandle->GetContact().Geometry[0], InHandle->GetContact().Geometry[1]));
 	}
 
-	TArray<FConstraintHandle*>& GetHandles() { return Handles; }
-	const TArray<FConstraintHandle*>& GetHandles() const { return Handles; }
+	TArray<FConstraintContainerHandle*>& GetHandles() { return Handles; }
+	const TArray<FConstraintContainerHandle*>& GetHandles() const { return Handles; }
 
 	bool ContainsShapeConnection(const FImplicitObject* Implicit0In, const FImplicitObject* Implicit1In) { return Implicits.Contains(ImplicitPairs(Implicit0In, Implicit1In)); }
 
@@ -110,7 +110,7 @@ private:
 	int32 Timestamp;
 	TVector<T, d> Location;
 	TRotation<T, d> Rotation;
-	TArray<FConstraintHandle*> Handles;
+	TArray<FConstraintContainerHandle*> Handles;
 	TSet<ImplicitPairs> Implicits;
 };
 
@@ -132,15 +132,15 @@ using TRigidBodyContactConstraintsPostApplyPushOutCallback = TFunction<void(cons
  * @todo(ccaulfield): separate collision detection from constraint container
  */
 template<typename T, int d>
-class CHAOS_API TPBDCollisionConstraint : public TPBDConstraintContainer<T, d>
+class CHAOS_API TPBDCollisionConstraint : public FPBDConstraintContainer
 {
 public:
-	using Base = TPBDConstraintContainer<T, d>;
+	using Base = FPBDConstraintContainer;
 	friend class TPBDCollisionConstraintAccessor<T, d>;
 	friend class TPBDCollisionConstraintHandle<T, d>;
 	using FReal = T;
 	static const int Dimensions = d;
-	using FConstraintHandle = TPBDCollisionConstraintHandle<T, d>;
+	using FConstraintContainerHandle = TPBDCollisionConstraintHandle<T, d>;
 	using FConstraintHistory = TPBDCollisionConstraintHistory<T, d>;
 	using FConstraintHandleAllocator = TConstraintHandleAllocator<TPBDCollisionConstraint<T, d>>;
 	using FRigidBodyContactConstraint = TRigidBodyContactConstraint<T, d>;
@@ -196,12 +196,12 @@ public:
 		return  FConstraintHandleID( Constraints[ConstraintIndex].Particle, Constraints[ConstraintIndex].Levelset );
 	}
 		
-	const FConstraintHandle* GetConstraintHandle(int32 ConstraintIndex) const
+	const FConstraintContainerHandle* GetConstraintHandle(int32 ConstraintIndex) const
 	{
 		return Handles[ConstraintIndex];
 	}
 
-	FConstraintHandle* GetConstraintHandle(int32 ConstraintIndex)
+	FConstraintContainerHandle* GetConstraintHandle(int32 ConstraintIndex)
 	{
 		return Handles[ ConstraintIndex ];
 	}
@@ -258,12 +258,12 @@ public:
 	 * Apply corrections for the specified list of constraints.
 	 */
 	 // @todo(ccaulfield): this runs wide. The serial/parallel decision should be in the ConstraintRule
-	void Apply(const T Dt, const TArray<FConstraintHandle*>& InConstraintHandles, const int32 It, const int32 NumIts);
+	void Apply(const T Dt, const TArray<FConstraintContainerHandle*>& InConstraintHandles, const int32 It, const int32 NumIts);
 
 	/**
 	 * Set disabled attribute on all input handles
 	 */
-	void Disable(const TArray<FConstraintHandle*>& InConstraintHandles);
+	void Disable(const TArray<FConstraintContainerHandle*>& InConstraintHandles);
 		
 	/**
 	 * Generate all contact constraints.
@@ -275,7 +275,7 @@ public:
 	 * Return true if we need another iteration 
 	 */
 	 // @todo(ccaulfield): this runs wide. The serial/parallel decision should be in the ConstraintRule
-	bool ApplyPushOut(const T Dt, const TArray<FConstraintHandle*>& InConstraintHandles, const TSet<TGeometryParticleHandle<T,d>*>& IsTemporarilyStatic, int32 Iteration, int32 NumIterations);
+	bool ApplyPushOut(const T Dt, const TArray<FConstraintContainerHandle*>& InConstraintHandles, const TSet<TGeometryParticleHandle<T,d>*>& IsTemporarilyStatic, int32 Iteration, int32 NumIterations);
 
 	void UpdateConstraints(/*const TPBDRigidParticles<T, d>& InParticles, const TArray<int32>& InIndices,*/ T Dt, const TSet<TGeometryParticleHandle<T, d>*>& AddedParticles);
 
@@ -345,7 +345,7 @@ private:
 	TRigidBodyContactConstraintsPostApplyCallback<T, d> PostApplyCallback;
 	TRigidBodyContactConstraintsPostApplyPushOutCallback<T, d> PostApplyPushOutCallback;
 
-	TArray<FConstraintHandle*> Handles;
+	TArray<FConstraintContainerHandle*> Handles;
 	FConstraintHandleAllocator HandleAllocator;
 
 	TMap<FConstraintHandleID, FConstraintHistory*> History;
diff --git a/Engine/Source/Runtime/Experimental/Chaos/Public/Chaos/PBDCollisionConstraintImp.h b/Engine/Source/Runtime/Experimental/Chaos/Public/Chaos/PBDCollisionConstraintImp.h
--- a/Engine/Source/Runtime/Experimental/Chaos/Public/Chaos/PBDCollisionConstraintImp.h
+++ b/Engine/Source/Runtime/Experimental/Chaos/Public/Chaos/PBDCollisionConstraintImp.h
@@ -298,7 +298,7 @@ namespace Chaos
 				while (!Queue.IsEmpty())
 				{
 					int32 Idx = Constraints.AddUninitialized(1);
-					FConstraintHandle* Handle = HandleAllocator.AllocHandle(this, Idx);
+					FConstraintContainerHandle* Handle = HandleAllocator.AllocHandle(this, Idx);
 
 					Handles.Add(Handle);
 					Queue.Dequeue(Constraints[Idx]);
diff --git a/Engine/Source/Runtime/Experimental/Chaos/Public/Chaos/PBDCollisionSpringConstraints.h b/Engine/Source/Runtime/Experimental/Chaos/Public/Chaos/PBDCollisionSpringConstraints.h
--- a/Engine/Source/Runtime/Experimental/Chaos/Public/Chaos/PBDCollisionSpringConstraints.h
+++ b/Engine/Source/Runtime/Experimental/Chaos/Public/Chaos/PBDCollisionSpringConstraints.h
@@ -12,7 +12,7 @@
 namespace Chaos
 {
 template<class T, int d>
-class TPBDCollisionSpringConstraints : public PBDCollisionSpringConstraintsBase<T, d>, public TPBDConstraintContainer<T, d>
+class TPBDCollisionSpringConstraints : public PBDCollisionSpringConstraintsBase<T, d>, public FPBDConstraintContainer
 {
 	typedef PBDCollisionSpringConstraintsBase<T, d> Base;
 	using Base::MBarys;
diff --git a/Engine/Source/Runtime/Experimental/Chaos/Public/Chaos/PBDConstraintColor.h b/Engine/Source/Runtime/Experimental/Chaos/Public/Chaos/PBDConstraintColor.h
--- a/Engine/Source/Runtime/Experimental/Chaos/Public/Chaos/PBDConstraintColor.h
+++ b/Engine/Source/Runtime/Experimental/Chaos/Public/Chaos/PBDConstraintColor.h
@@ -3,40 +3,35 @@
 
 #include "Chaos/ParticleHandle.h"
 
-// @todo(ccaulfield): can we get rid of this now?
+// @todo(ccaulfield): move level/shock system out of coloring
 #define USE_CONTACT_LEVELS 1
 
 namespace Chaos
 {
-	template<typename T, int d>
-	class TPBDConstraintGraph;
-
-	template<typename T, int d>
-	class TConstraintHandle;
+	class FPBDConstraintGraph;
+	class FConstraintHandle;
 
 	/**
 	 * Generates color information for a single constraint rule in a connection graph.
 	 * Edges with the same color are non-interacting and can safely be processed in parallel.
 	 */
-	template<typename T, int d>
-	class CHAOS_API TPBDConstraintColor
+	class CHAOS_API FPBDConstraintColor
 	{
 	public:
-		typedef TPBDConstraintGraph<T, d> FConstraintGraph;
 		typedef TSet<int32> FColorSet;
-		typedef TArray<TConstraintHandle<T, d>*> FConstraintList;
+		typedef TArray<FConstraintHandle*> FConstraintList;
 		typedef TMap<int32, FConstraintList> FColorToConstraintListMap;
 		typedef TArray<FColorToConstraintListMap> FLevelToColorToConstraintListMap;
 
 		/**
 		 * Initialize the color structures based on the connectivity graph (i.e., reset all color-related node, edge and island data).
 		 */
-		void InitializeColor(const FConstraintGraph& ConstraintGraph);
+		void InitializeColor(const FPBDConstraintGraph& ConstraintGraph);
 
 		/**
 		 * Calculate the color information for the specified island.
 		 */
-		void ComputeColor(const int32 Island, const FConstraintGraph& ConstraintGraph, uint32 ContainerId);
+		void ComputeColor(const int32 Island, const FPBDConstraintGraph& ConstraintGraph, uint32 ContainerId);
 
 		/**
 		 * Get the Level-Color-ConstraintList map for the specified island.
@@ -54,8 +49,8 @@ namespace Chaos
 		int GetIslandMaxLevel(int32 Island) const;
 
 	private:
-		void ComputeContactGraph(const int32 Island, const FConstraintGraph& ConstraintGraph, uint32 ContainerId);
-		void ComputeIslandColoring(const int32 Island, const FConstraintGraph& ConstraintGraph, uint32 ContainerId);
+		void ComputeContactGraph(const int32 Island, const FPBDConstraintGraph& ConstraintGraph, uint32 ContainerId);
+		void ComputeIslandColoring(const int32 Island, const FPBDConstraintGraph& ConstraintGraph, uint32 ContainerId);
 
 		struct FGraphNodeColor
 		{
diff --git a/Engine/Source/Runtime/Experimental/Chaos/Public/Chaos/PBDConstraintContainer.h b/Engine/Source/Runtime/Experimental/Chaos/Public/Chaos/PBDConstraintContainer.h
--- a/Engine/Source/Runtime/Experimental/Chaos/Public/Chaos/PBDConstraintContainer.h
+++ b/Engine/Source/Runtime/Experimental/Chaos/Public/Chaos/PBDConstraintContainer.h
@@ -11,17 +11,16 @@ namespace Chaos
 	 * A Constraint Container holds an array of constraints and provides methods to allocate and deallocate constraints
 	 *as well as the API required to plug into Constraint Rules.
 	 */
-	template<class T, int d>
-	class CHAOS_API TPBDConstraintContainer
+	class CHAOS_API FPBDConstraintContainer
 	{
 	public:
-		TPBDConstraintContainer();
+		FPBDConstraintContainer();
 
-		virtual ~TPBDConstraintContainer();
+		virtual ~FPBDConstraintContainer();
 
 	protected:
 		// friend access to the Constraint Handle's container API
-		int32 GetConstraintIndex(const TConstraintHandle<T, d>* ConstraintHandle) const;
-		void SetConstraintIndex(TConstraintHandle<T, d>* ConstraintHandle, int32 ConstraintIndex) const;
+		int32 GetConstraintIndex(const FConstraintHandle* ConstraintHandle) const;
+		void SetConstraintIndex(FConstraintHandle* ConstraintHandle, int32 ConstraintIndex) const;
 	};
 }
\ No newline at end of file
diff --git a/Engine/Source/Runtime/Experimental/Chaos/Public/Chaos/PBDConstraintGraph.h b/Engine/Source/Runtime/Experimental/Chaos/Public/Chaos/PBDConstraintGraph.h
--- a/Engine/Source/Runtime/Experimental/Chaos/Public/Chaos/PBDConstraintGraph.h
+++ b/Engine/Source/Runtime/Experimental/Chaos/Public/Chaos/PBDConstraintGraph.h
@@ -15,8 +15,7 @@ namespace Chaos
 	template<typename T, int d>
 	class TPBDConstraintColor;
 
-	template<typename T, int d>
-	class TConstraintHandle;
+	class FConstraintHandle;
 
 	template<class T>
 	class TChaosPhysicsMaterial;
@@ -38,12 +37,10 @@ namespace Chaos
 	 */
 	 // @todo(ccaulfield): Break out the Island and Sleep stuff from the particle graph
 	 // @todo(ccaulfield): The level system should be part of the Island system, and not embedded with the color system
-	template<typename T, int d>
-	class CHAOS_API TPBDConstraintGraph
+	class CHAOS_API FPBDConstraintGraph
 	{
 	public:
-		template <typename T2, int d2>
-		friend class TPBDConstraintColor;
+		friend class FPBDConstraintColor;
 
 		/** Number of sleep threshold passes before sleep mode is enabled */
 		static const int32 SleepCountThreshold = 5;
@@ -51,19 +48,27 @@ namespace Chaos
 		/** Information required to map a graph edge back to its constraint */
 		struct FConstraintData
 		{
+		public:
+			FConstraintData() : ContainerId(0), ConstraintHandle(nullptr) {}
+			FConstraintData(uint32 InContainerId, FConstraintHandle* InConstraintHandle) : ContainerId(InContainerId), ConstraintHandle(InConstraintHandle) {}
+
+			uint32 GetContainerId() const { return ContainerId; }
+			FConstraintHandle* GetConstraintHandle() const { return ConstraintHandle; }
+
+		private:
 			uint32 ContainerId;
-			TConstraintHandle<T, d>* ConstraintHandle;
+			FConstraintHandle* ConstraintHandle;
 		};
 
-		TPBDConstraintGraph();
-		TPBDConstraintGraph(const TParticleView<TGeometryParticles<T,d>>& Particles);
-		virtual ~TPBDConstraintGraph() {}
+		FPBDConstraintGraph();
+		FPBDConstraintGraph(const TParticleView<TGeometryParticles<FReal,3>>& Particles);
+		virtual ~FPBDConstraintGraph() {}
 
 		/**
 		 * Clear the graph and set up the particle-to-graph-node mapping for the specified particles
 		 * Should be called before AddConstraint.
 		 */
-		void InitializeGraph(const TParticleView<TGeometryParticles<T, d>>& Particles);
+		void InitializeGraph(const TParticleView<TGeometryParticles<FReal, 3>>& Particles);
 
 		/**
 		 * Reserve space in the graph for NumConstraints additional constraints.
@@ -73,23 +78,23 @@ namespace Chaos
 		/**
 		 * Add a constraint to the graph for each constraint in the container.
 		 */
-		void AddConstraint(const uint32 InContainerId, TConstraintHandle<T, d>* InConstraintHandle, const TVector<TGeometryParticleHandle<T,d>*, 2>& InConstrainedParticles);
+		void AddConstraint(const uint32 InContainerId, FConstraintHandle* InConstraintHandle, const TVector<TGeometryParticleHandle<FReal, 3>*, 2>& InConstrainedParticles);
 
 		/**
 		 * Add particles/constraints to their particle's already-assigned islands (if applicable).
 		 */
 		 // @todo(ccaulfield): InitializeIslands and ResetIslands are a bit confusing. Try to come up with better names.
-		void ResetIslands(const TParticleView<TPBDRigidParticles<T, d>>& PBDRigids);
+		void ResetIslands(const TParticleView<TPBDRigidParticles<FReal, 3>>& PBDRigids);
 
 		/**
 		 * Generate the simulation islands of connected particles (AddConstraints must have already been called). There are no constraints connecting particles in different islands.
 		 */
-		void UpdateIslands(const TParticleView<TPBDRigidParticles<T, d>>& confusing, TPBDRigidsSOAs<T, d>& Particles);
+		void UpdateIslands(const TParticleView<TPBDRigidParticles<FReal, 3>>& confusing, TPBDRigidsSOAs<FReal, 3>& Particles);
 
 		/**
 		 * Put particles in inactive islands to sleep.
 		 */
-		bool SleepInactive(const int32 Island, const TArrayCollectionArray<TSerializablePtr<TChaosPhysicsMaterial<T>>>& PerParticleMaterialAttributes);
+		bool SleepInactive(const int32 Island, const TArrayCollectionArray<TSerializablePtr<TChaosPhysicsMaterial<FReal>>>& PerParticleMaterialAttributes);
 
 		/**
 		 * Wake all particles in an Island.
@@ -115,7 +120,7 @@ namespace Chaos
 		/**
 		 * Get the list of Particle indices associated with the specified island.
 		 */
-		const TArray<TGeometryParticleHandle<T,d>*>& GetIslandParticles(int32 Island) const
+		const TArray<TGeometryParticleHandle<FReal, 3>*>& GetIslandParticles(int32 Island) const
 		{
 			return IslandToParticles[Island];
 		}
@@ -144,21 +149,21 @@ namespace Chaos
 		/**
 		 * Enable a particle after initialization. Adds the particle to the island of the parent particle if supplied.
 		 */
-		void EnableParticle(TGeometryParticleHandle<T,d>* Particle, const TGeometryParticleHandle<T, d>* ParentParticle);
+		void EnableParticle(TGeometryParticleHandle<FReal, 3>* Particle, const TGeometryParticleHandle<FReal, 3>* ParentParticle);
 
 		/**
 		 * Disable a particle and remove it from its island.
 		 * @note: this does not remove any attached constraints - constraints need to be re-added and islands will need to be rebuilt.
 		 */
-		void DisableParticle(TGeometryParticleHandle<T, d>* Particle);
+		void DisableParticle(TGeometryParticleHandle<FReal, 3>* Particle);
 
-		void AddParticle(TGeometryParticleHandle<T, d>* AddedParticle)
+		void AddParticle(TGeometryParticleHandle<FReal, 3>* AddedParticle)
 		{
 			ParticleAdd(AddedParticle);
 		}
 
 		//Remove particle from constraint, maybe rethink some of these names
-		void RemoveParticle(TGeometryParticleHandle<T, d>* Particle)
+		void RemoveParticle(TGeometryParticleHandle<FReal, 3>* Particle)
 		{
 			ParticleRemove(Particle);
 
@@ -170,7 +175,7 @@ namespace Chaos
 		 * Disable a set of particles and remove them from their island.
 		 * @note: this does not remove any attached constraints - constraints need to be re-added and islands will need to be rebuilt.
 		 */
-		void DisableParticles(const TSet<TGeometryParticleHandle<T,d>*>& Particles);
+		void DisableParticles(const TSet<TGeometryParticleHandle<FReal, 3>*>& Particles);
 
 	private:
 		struct FGraphNode
@@ -182,7 +187,7 @@ namespace Chaos
 			}
 
 			TArray<int32> Edges;
-			TGeometryParticleHandle<T,d>* Particle;
+			TGeometryParticleHandle<FReal, 3>* Particle;
 			int32 Island;
 		};
 
@@ -209,21 +214,21 @@ namespace Chaos
 			bool bIsIslandPersistant;
 		};
 
-		void ComputeIslands(const TParticleView<TPBDRigidParticles<T,d>>& PBDRigids, TPBDRigidsSOAs<T,d>& Particles);
+		void ComputeIslands(const TParticleView<TPBDRigidParticles<FReal, 3>>& PBDRigids, TPBDRigidsSOAs<FReal, 3>& Particles);
 		void ComputeIsland(const int32 Node, const int32 Island,
-			TSet<TGeometryParticleHandle<T, d>*>& DynamicParticlesInIsland, TSet<TGeometryParticleHandle<T, d>*>& StaticParticlesInIsland);
-		bool CheckIslands(const TArray<TGeometryParticleHandle<T, d>*>& Particles);
+			TSet<TGeometryParticleHandle<FReal, 3>*>& DynamicParticlesInIsland, TSet<TGeometryParticleHandle<FReal, 3>*>& StaticParticlesInIsland);
+		bool CheckIslands(const TArray<TGeometryParticleHandle<FReal, 3>*>& Particles);
 		
-		void ParticleAdd(TGeometryParticleHandle<T, d>* AddedParticle);
-		void ParticleRemove(TGeometryParticleHandle<T, d>* RemovedParticle);
+		void ParticleAdd(TGeometryParticleHandle<FReal, 3>* AddedParticle);
+		void ParticleRemove(TGeometryParticleHandle<FReal, 3>* RemovedParticle);
 		int32 GetNextNodeIndex();
 
 		TArray<FGraphNode> Nodes;
 		TArray<FGraphEdge> Edges;
 		TArray<FIslandData> IslandToData;
-		TMap<TGeometryParticleHandle<T,d>*, int32> ParticleToNodeIndex;
+		TMap<TGeometryParticleHandle<FReal, 3>*, int32> ParticleToNodeIndex;
 
-		TArray<TArray<TGeometryParticleHandle<T,d>*>> IslandToParticles;
+		TArray<TArray<TGeometryParticleHandle<FReal, 3>*>> IslandToParticles;
 		TArray<TArray<int32>> IslandToConstraints;
 		TArray<int32> IslandToSleepCount;
 
@@ -234,8 +239,3 @@ namespace Chaos
 	};
 
 }
-
-#if PLATFORM_MAC || PLATFORM_LINUX
-extern template class CHAOS_API Chaos::TPBDConstraintGraph<float, 3>;
-extern template class CHAOS_API Chaos::TPBDConstraintGraph<float, 3>;
-#endif
