Chaos - remove template params from constraint classes.

diff --git a/Engine/Source/Runtime/Experimental/Chaos/Public/Chaos/PBDConstraintRule.h b/Engine/Source/Runtime/Experimental/Chaos/Public/Chaos/PBDConstraintRule.h
--- a/Engine/Source/Runtime/Experimental/Chaos/Public/Chaos/PBDConstraintRule.h
+++ b/Engine/Source/Runtime/Experimental/Chaos/Public/Chaos/PBDConstraintRule.h
@@ -2,17 +2,15 @@
 #pragma once
 
 #include "Chaos/Declares.h"
+#include "Chaos/ConstraintHandle.h"
 #include "Chaos/PBDCollisionTypes.h"
-#include "Chaos/PBDConstraintGraph.h"
 #include "Chaos/PBDConstraintColor.h"
+#include "Chaos/PBDConstraintGraph.h"
 
 #define USE_SHOCK_PROPOGATION 1
 
 namespace Chaos
 {
-	template<typename T, int d>
-	class TConstraintHandle;
-
 	/**
 	 * Constraint Rules bind constraint collections to the evolution and provide their update algorithm.
 	 */
@@ -44,25 +42,22 @@ namespace Chaos
 	 * each rule is assigned an ID which is stored alongside the constraint index in the graph.
 	 * @see TPBDConstraintGraphRuleImpl
 	 */
-	template<typename T, int d>
-	class CHAOS_API TPBDConstraintGraphRule : public FConstraintRule
+	class CHAOS_API FPBDConstraintGraphRule : public FConstraintRule
 	{
 	public:
-		typedef TPBDConstraintGraph<T, d> FConstraintGraph;
-
-		TPBDConstraintGraphRule(int32 InPriority) : FConstraintRule(InPriority) {}
-		virtual ~TPBDConstraintGraphRule() {}
+		FPBDConstraintGraphRule(int32 InPriority) : FConstraintRule(InPriority) {}
+		virtual ~FPBDConstraintGraphRule() {}
 
-		virtual void BindToGraph(FConstraintGraph& InContactGraph, uint32 InContainerId) {}
+		virtual void BindToGraph(FPBDConstraintGraph& InContactGraph, uint32 InContainerId) {}
 
 		/** Called once per tick to allow constraint containers to create/alter their constraints based on particle position */
-		virtual void UpdatePositionBasedState(const T Dt) {}
+		virtual void UpdatePositionBasedState(const FReal Dt) {}
 
 		/** Apply all corrections for constraints in the specified island */
-		virtual void ApplyConstraints(const T Dt, int32 Island, const int32 It, const int32 NumIts) {}
+		virtual void ApplyConstraints(const FReal Dt, int32 Island, const int32 It, const int32 NumIts) {}
 
 		/** Apply push out for constraints in the specified island */
-		virtual bool ApplyPushOut(const T Dt, int32 Island, const int32 It, const int32 NumIts) { return false; }
+		virtual bool ApplyPushOut(const FReal Dt, int32 Island, const int32 It, const int32 NumIts) { return false; }
 
 		/** Add all constraints to the connectivity graph */
 		virtual void AddToGraph() {}
@@ -75,7 +70,7 @@ namespace Chaos
 
 		/** Remove all constraints associated with the specified particles */
 		// @todo(ccaulfield): remove uint version
-		virtual void RemoveConstraints(const TSet<TGeometryParticleHandle<T,d>*>& RemovedParticles) { }
+		virtual void RemoveConstraints(const TSet<TGeometryParticleHandle<FReal, 3>*>& RemovedParticles) { }
 
 		/** The number of constraints in the collection */
 		virtual int32 NumConstraints() const { return 0; }
@@ -85,28 +80,26 @@ namespace Chaos
 	/**
 	 * ConstraintGraphRule helper base class - templatized on Constraint Container.
 	 */
-	template<typename T_CONSTRAINTS, typename T, int d>
-	class CHAOS_API TPBDConstraintGraphRuleImpl : public TPBDConstraintGraphRule<T, d>
+	template<typename T_CONSTRAINTS>
+	class CHAOS_API TPBDConstraintGraphRuleImpl : public FPBDConstraintGraphRule
 	{
-		typedef TPBDConstraintGraphRule<T, d> Base;
 	public:
 		typedef T_CONSTRAINTS FConstraints;
-		typedef TPBDConstraintGraph<T, d> FConstraintGraph;
 
 		TPBDConstraintGraphRuleImpl(FConstraints& InConstraints, int32 InPriority)
-			: TPBDConstraintGraphRule<T, d>(InPriority)
+			: FPBDConstraintGraphRule(InPriority)
 			, Constraints(InConstraints)
 			, ConstraintGraph(nullptr)
 		{
 		}
 
-		virtual void BindToGraph(FConstraintGraph& InContactGraph, uint32 InContainerId) override
+		virtual void BindToGraph(FPBDConstraintGraph& InContactGraph, uint32 InContainerId) override
 		{
 			ConstraintGraph = &InContactGraph;
 			ContainerId = InContainerId;
 		}
 
-		virtual void UpdatePositionBasedState(const T Dt) override
+		virtual void UpdatePositionBasedState(const FReal Dt) override
 		{
 			Constraints.UpdatePositionBasedState(Dt);
 		}
@@ -124,30 +117,29 @@ namespace Chaos
 
 	protected:
 		FConstraints& Constraints;
-		FConstraintGraph* ConstraintGraph;
+		FPBDConstraintGraph* ConstraintGraph;
 		uint32 ContainerId;
 	};
 
 	/**
 	 * Island-based constraint rule. All constraints in an island are updated in single-threaded a loop. Islands may be updated in parallel.
 	 */
-	template<typename T_CONSTRAINTS, typename T, int d>
-	class CHAOS_API TPBDConstraintIslandRule : public TPBDConstraintGraphRuleImpl<T_CONSTRAINTS, T, d>
+	template<typename T_CONSTRAINTS>
+	class CHAOS_API TPBDConstraintIslandRule : public TPBDConstraintGraphRuleImpl<T_CONSTRAINTS>
 	{
-		typedef TPBDConstraintGraphRuleImpl<T_CONSTRAINTS, T, d> Base;
+		typedef TPBDConstraintGraphRuleImpl<T_CONSTRAINTS> Base;
 
 	public:
-		using FConstraints = typename Base::FConstraints;
-		using FConstraintHandle = typename FConstraints::FConstraintHandle;
-		using FConstraintList = TArray<TConstraintHandle<T, d>*>;
-		using FConstraintGraph = typename Base::FConstraintGraph;
+		using FConstraints = T_CONSTRAINTS;
+		using FConstraintContainerHandle = typename FConstraints::FConstraintContainerHandle;
+		using FConstraintList = TArray<FConstraintContainerHandle*>;
 
 		TPBDConstraintIslandRule(FConstraints& InConstraints, int32 InPriority = 0)
-			: TPBDConstraintGraphRuleImpl<T_CONSTRAINTS, T, d>(InConstraints, InPriority)
+			: TPBDConstraintGraphRuleImpl<T_CONSTRAINTS>(InConstraints, InPriority)
 		{
 		}
 
-		virtual void ApplyConstraints(const T Dt, int32 Island, const int32 It, const int32 NumIts) override
+		virtual void ApplyConstraints(const FReal Dt, int32 Island, const int32 It, const int32 NumIts) override
 		{
 			if (IslandConstraintLists[Island].Num())
 			{
@@ -155,7 +147,7 @@ namespace Chaos
 			}
 		}
 
-		virtual bool ApplyPushOut(const T Dt, int32 Island, const int32 It, const int32 NumIts) override
+		virtual bool ApplyPushOut(const FReal Dt, int32 Island, const int32 It, const int32 NumIts) override
 		{
 			if (IslandConstraintLists[Island].Num())
 			{
@@ -181,10 +173,10 @@ namespace Chaos
 			IslandConstraintList.Reserve(ConstraintDataIndices.Num());
 			for (int32 ConstraintDataIndex : ConstraintDataIndices)
 			{
-				const typename FConstraintGraph::FConstraintData& ConstraintData = ConstraintGraph->GetConstraintData(ConstraintDataIndex);
-				if (ConstraintData.ContainerId == ContainerId)
+				const typename FPBDConstraintGraph::FConstraintData& ConstraintData = ConstraintGraph->GetConstraintData(ConstraintDataIndex);
+				if (ConstraintData.GetContainerId() == ContainerId)
 				{
-					IslandConstraintList.Add(ConstraintData.ConstraintHandle);
+					IslandConstraintList.Add(ConstraintData.GetConstraintHandle()->As<FConstraintContainerHandle>());
 				}
 			}
 		}
@@ -200,11 +192,11 @@ namespace Chaos
 		using Base::ConstraintGraph;
 		using Base::ContainerId;
 
-		const TArray<FConstraintHandle*>& GetIslandConstraints(int32 Island) const
+		const TArray<FConstraintContainerHandle*>& GetIslandConstraints(int32 Island) const
 		{
-			// Constraint rules are bound to a single type, but the FConstraintGraph works with many types. We have
+			// Constraint rules are bound to a single type, but the FPBDConstraintGraph works with many types. We have
 			// already pre-filtered the constraint lists based on type, so this case is safe.
-			return reinterpret_cast<const TArray<FConstraintHandle*>&>(IslandConstraintLists[Island]);
+			return reinterpret_cast<const TArray<FConstraintContainerHandle*>&>(IslandConstraintLists[Island]);
 		}
 
 		// @todo(ccaulfield): optimize: eliminate the need for this index list - it is a subset of EdgeData. 
@@ -218,31 +210,29 @@ namespace Chaos
 	 * The level is used to implement shock propagation: constraints of lower levels are frozen in 
 	 * place as far as higher-level constraints are concerned.
 	 */
-	template<typename T_CONSTRAINTS, typename T, int d>
-	class CHAOS_API TPBDConstraintColorRule : public TPBDConstraintGraphRuleImpl<T_CONSTRAINTS, T, d>
+	template<typename T_CONSTRAINTS>
+	class CHAOS_API TPBDConstraintColorRule : public TPBDConstraintGraphRuleImpl<T_CONSTRAINTS>
 	{
-		typedef TPBDConstraintGraphRuleImpl<T_CONSTRAINTS, T, d> Base;
+		typedef TPBDConstraintGraphRuleImpl<T_CONSTRAINTS> Base;
 
 	public:
-		typedef typename Base::FConstraints FConstraints;
-		typedef typename Base::FConstraintGraph FConstraintGraph;
-		typedef TPBDConstraintColor<T, d> FConstraintColor;
+		using FConstraints = T_CONSTRAINTS;
 
 		TPBDConstraintColorRule(FConstraints& InConstraints, const int32 InPushOutIterations, int32 InPriority = 0)
-			: TPBDConstraintGraphRuleImpl<T_CONSTRAINTS, T, d>(InConstraints, InPriority)
+			: TPBDConstraintGraphRuleImpl<T_CONSTRAINTS>(InConstraints, InPriority)
 			, PushOutIterations(InPushOutIterations)
 		{
 		}
 
 
-		virtual void UpdatePositionBasedState(const T Dt) override
+		virtual void UpdatePositionBasedState(const FReal Dt) override
 		{
 			Constraints.UpdatePositionBasedState(Dt);
 		}
 
-		virtual void ApplyConstraints(const T Dt, int32 Island, const int32 It, const int32 NumIts) override
+		virtual void ApplyConstraints(const FReal Dt, int32 Island, const int32 It, const int32 NumIts) override
 		{
-			const typename FConstraintColor::FLevelToColorToConstraintListMap& LevelToColorToConstraintListMap = GraphColor.GetIslandLevelToColorToConstraintListMap(Island);
+			const typename FPBDConstraintColor::FLevelToColorToConstraintListMap& LevelToColorToConstraintListMap = GraphColor.GetIslandLevelToColorToConstraintListMap(Island);
 			int32 MaxColor = GraphColor.GetIslandMaxColor(Island);
 			int32 MaxLevel = GraphColor.GetIslandMaxLevel(Island);
 			for (int32 Level = 0; Level <= MaxLevel; ++Level)
@@ -251,25 +241,25 @@ namespace Chaos
 				{
 					if (LevelToColorToConstraintListMap[Level].Contains(Color) && LevelToColorToConstraintListMap[Level][Color].Num())
 					{
-						const TArray<typename FConstraints::FConstraintHandle*>& ConstraintHandles = GetLevelColorConstraints(LevelToColorToConstraintListMap, Level, Color);
+						const TArray<typename FConstraints::FConstraintContainerHandle*>& ConstraintHandles = GetLevelColorConstraints(LevelToColorToConstraintListMap, Level, Color);
 						Constraints.Apply(Dt, ConstraintHandles, It, NumIts);
 					}
 				}
 			}
 		}
 
-		virtual void RemoveConstraints(const TSet<TGeometryParticleHandle<T, d>*>& InConstraints)
+		virtual void RemoveConstraints(const TSet<TGeometryParticleHandle<FReal, 3>*>& InConstraints)
 		{
 			Constraints.RemoveConstraints(InConstraints);
 		}
 
-		virtual bool ApplyPushOut(const T Dt, int32 Island, const int32 It, const int32 NumIts) override
+		virtual bool ApplyPushOut(const FReal Dt, int32 Island, const int32 It, const int32 NumIts) override
 		{
-			const typename FConstraintColor::FLevelToColorToConstraintListMap& LevelToColorToConstraintListMap = GraphColor.GetIslandLevelToColorToConstraintListMap(Island);
+			const typename FPBDConstraintColor::FLevelToColorToConstraintListMap& LevelToColorToConstraintListMap = GraphColor.GetIslandLevelToColorToConstraintListMap(Island);
 			int32 MaxColor = GraphColor.GetIslandMaxColor(Island);
 			int32 MaxLevel = GraphColor.GetIslandMaxLevel(Island);
 
-			TSet<TGeometryParticleHandle<T, d>*> IsTemporarilyStatic;
+			TSet<TGeometryParticleHandle<FReal, 3>*> IsTemporarilyStatic;
 			bool bNeedsAnotherIteration = false;
 			for (int32 Level = 0; Level <= MaxLevel; ++Level)
 			{
@@ -277,7 +267,7 @@ namespace Chaos
 				{
 					if (LevelToColorToConstraintListMap[Level].Contains(Color) && LevelToColorToConstraintListMap[Level][Color].Num())
 					{
-						const TArray<typename FConstraints::FConstraintHandle*>& ConstraintHandles = GetLevelColorConstraints(LevelToColorToConstraintListMap, Level, Color);
+						const TArray<typename FConstraints::FConstraintContainerHandle*>& ConstraintHandles = GetLevelColorConstraints(LevelToColorToConstraintListMap, Level, Color);
 						if (Constraints.ApplyPushOut(Dt, ConstraintHandles, IsTemporarilyStatic, It, NumIts))
 						{
 							bNeedsAnotherIteration = true;
@@ -294,7 +284,7 @@ namespace Chaos
 						for (int32 Edge = 0; Edge < LevelToColorToConstraintListMap[Level][Color].Num(); ++Edge)
 						{
 							const int32 ConstraintIndex = LevelToColorToConstraintListMap[Level][Color][Edge]->GetConstraintIndex();
-							const TVector<TGeometryParticleHandle<T,d>*, 2> Particles = Constraints.ConstraintParticles(ConstraintIndex);
+							const TVector<TGeometryParticleHandle<FReal, 3>*, 2> Particles = Constraints.ConstraintParticles(ConstraintIndex);
 							if (It == NumIts - 1)
 							{
 								if (Particles[0]->AsDynamic() == nullptr || IsTemporarilyStatic.Contains(Particles[0]))
@@ -333,7 +323,7 @@ namespace Chaos
 		template<typename TVisitor>
 		void VisitIslandConstraints(const int32 Island, const TVisitor& Visitor) const
 		{
-			const typename FConstraintColor::FLevelToColorToConstraintListMap& LevelToColorToConstraintListMap = GraphColor.GetIslandLevelToColorToConstraintListMap(Island);
+			const typename FPBDConstraintColor::FLevelToColorToConstraintListMap& LevelToColorToConstraintListMap = GraphColor.GetIslandLevelToColorToConstraintListMap(Island);
 			int32 MaxColor = GraphColor.GetIslandMaxColor(Island);
 			int32 MaxLevel = GraphColor.GetIslandMaxLevel(Island);
 			for (int32 Level = 0; Level <= MaxLevel; ++Level)
@@ -342,7 +332,7 @@ namespace Chaos
 				{
 					if (LevelToColorToConstraintListMap[Level].Contains(Color) && LevelToColorToConstraintListMap[Level][Color].Num())
 					{
-						const TArray<typename FConstraints::FConstraintHandle*>& ConstraintHandles = GetLevelColorConstraints(LevelToColorToConstraintListMap, Level, Color);
+						const TArray<typename FConstraints::FConstraintContainerHandle*>& ConstraintHandles = GetLevelColorConstraints(LevelToColorToConstraintListMap, Level, Color);
 						Visitor(ConstraintHandles);
 					}
 				}
@@ -354,47 +344,14 @@ namespace Chaos
 		using Base::ConstraintGraph;
 		using Base::ContainerId;
 
-		const TArray<typename FConstraints::FConstraintHandle*>& GetLevelColorConstraints(const typename FConstraintColor::FLevelToColorToConstraintListMap& LevelToColorToConstraintListMap, int32 Level, int32 Color) const
+		const TArray<typename FConstraints::FConstraintContainerHandle*>& GetLevelColorConstraints(const typename FPBDConstraintColor::FLevelToColorToConstraintListMap& LevelToColorToConstraintListMap, int32 Level, int32 Color) const
 		{
-			// FConstraintColor works with any constraint type (in principle - currently only used with Collisions), but the rule is bound to a single type and so this cast is ok
-			return reinterpret_cast<const TArray<typename FConstraints::FConstraintHandle*>&>(LevelToColorToConstraintListMap[Level][Color]);
+			// FPBDConstraintColor works with any constraint type (in principle - currently only used with Collisions), but the rule is bound to a single type and so this cast is ok
+			return reinterpret_cast<const TArray<typename FConstraints::FConstraintContainerHandle*>&>(LevelToColorToConstraintListMap[Level][Color]);
 		}
 
-		FConstraintColor GraphColor;
+		FPBDConstraintColor GraphColor;
 		int32 PushOutIterations;
 	};
 
-	/**
-	 * Simplify creation of constraint rules thanks to template parameter deduction.
-	 * 
-	 * E.g., it allows you to write code like this:
-	 *		TPBDCollisionConstraints<float, 3> Constraints;
-	 *		auto Rule = FConstraintRuleFactory::CreateIslandRule(Constraints);
-	 *
-	 * as opposed to 
-	 *		TPBDCollisionConstraints<float, 3> Constraints;
-	 *		auto Rule = TPBDConstraintIslandRule<TPBDCollisionConstraints<float, 3>, float, 3>(Constraints);
-	 *
-	 * @todo(ccaulfield): The evolution classes should be factories of Constraint Rules since they own the 
-	 * update loop and therefore know what algorithms should be applied.
-	 */
-	template<typename T, int d>
-	struct CHAOS_API TConstraintRuleFactory
-	{
-	public:
-		typedef TPBDConstraintGraph<T, d> FConstraintGraph;
-
-		template<typename T_CONSTRAINTS>
-		static TPBDConstraintIslandRule<T_CONSTRAINTS, T, d> CreateIslandRule(T_CONSTRAINTS& Constraints, const int32 Priority = 0)
-		{
-			return TPBDConstraintIslandRule<T_CONSTRAINTS, T, d>(Constraints, Priority);
-		}
-
-		template<typename T_CONSTRAINTS>
-		static TPBDConstraintColorRule<T_CONSTRAINTS, T, d> CreateColorRule(T_CONSTRAINTS& Constraints, const int32 InPushOutIterations, const int32 Priority = 0)
-		{
-			return TPBDConstraintColorRule<T_CONSTRAINTS, T, d>(Constraints, InPushOutIterations, Priority);
-		}
-	};
-
 }
\ No newline at end of file
diff --git a/Engine/Source/Runtime/Experimental/Chaos/Public/Chaos/PBDJointConstraintTypes.h b/Engine/Source/Runtime/Experimental/Chaos/Public/Chaos/PBDJointConstraintTypes.h
--- a/Engine/Source/Runtime/Experimental/Chaos/Public/Chaos/PBDJointConstraintTypes.h
+++ b/Engine/Source/Runtime/Experimental/Chaos/Public/Chaos/PBDJointConstraintTypes.h
@@ -3,22 +3,19 @@
 
 #include "CoreMinimal.h"
 
+#include "Chaos/Core.h"
 #include "Chaos/Transform.h"
 #include "Chaos/Vector.h"
 
 namespace Chaos
 {
-	template<class T, int d>
-	class TPBDJointConstraints;
+	class FPBDJointConstraints;
 
-	template<class T, int d>
-	class TPBDJointConstraintHandle;
+	class FPBDJointConstraintHandle;
 
-	template<typename T, int d>
-	using TJointPreApplyCallback = TFunction<void(const T Dt, const TArray<TPBDJointConstraintHandle<T, d>*>& InConstraintHandles)>;
+	using FJointPreApplyCallback = TFunction<void(const FReal Dt, const TArray<FPBDJointConstraintHandle*>& InConstraintHandles)>;
 
-	template<typename T, int d>
-	using TJointPostApplyCallback = TFunction<void(const T Dt, const TArray<TPBDJointConstraintHandle<T, d>*>& InConstraintHandles)>;
+	using FJointPostApplyCallback = TFunction<void(const FReal Dt, const TArray<FPBDJointConstraintHandle*>& InConstraintHandles)>;
 
 	enum class EJointMotionType : int32
 	{
@@ -47,80 +44,76 @@ namespace Chaos
 		Swing1 = 2,		// Swing1 Axis = Z
 	};
 
-	template<class T, int d>
-	struct TJointConstants
+	struct FJointConstants
 	{
 		/** The constraint-space twist axis (X Axis) */
-		static const TVector<T, d> TwistAxis() { return TVector<T, d>(1, 0, 0); }
+		static const FVec3 TwistAxis() { return FVec3(1, 0, 0); }
 
 		/** The constraint-space Swing1 axis (Z Axis) */
-		static const TVector<T, d> Swing1Axis() { return TVector<T, d>(0, 0, 1); }
+		static const FVec3 Swing1Axis() { return FVec3(0, 0, 1); }
 
 		/** The constraint-space Swing2 axis (Y Axis) */
-		static const TVector<T, d> Swing2Axis() { return TVector<T, d>(0, 1, 0); }
+		static const FVec3 Swing2Axis() { return FVec3(0, 1, 0); }
 	};
 
-	template<class T, int d>
-	class CHAOS_API TPBDJointMotionSettings
+	class CHAOS_API FPBDJointMotionSettings
 	{
 	public:
-		TPBDJointMotionSettings();
-		TPBDJointMotionSettings(const TVector<EJointMotionType, d>& InLinearMotionTypes, const TVector<EJointMotionType, d>& InAngularMotionTypes);
+		FPBDJointMotionSettings();
+		FPBDJointMotionSettings(const TVector<EJointMotionType, 3>& InLinearMotionTypes, const TVector<EJointMotionType, 3>& InAngularMotionTypes);
 
-		T Stiffness;
-		T Projection;
+		FReal Stiffness;
+		FReal Projection;
 
-		TVector<EJointMotionType, d> LinearMotionTypes;
-		T LinearLimit;
+		TVector<EJointMotionType, 3> LinearMotionTypes;
+		FReal LinearLimit;
 
-		TVector<EJointMotionType, d> AngularMotionTypes;
-		TVector<T, d> AngularLimits;
+		TVector<EJointMotionType, 3> AngularMotionTypes;
+		FVec3 AngularLimits;
 
 		bool bSoftLinearLimitsEnabled;
 		bool bSoftTwistLimitsEnabled;
 		bool bSoftSwingLimitsEnabled;
-		T SoftLinearStiffness;
-		T SoftTwistStiffness;
-		T SoftSwingStiffness;
+		FReal SoftLinearStiffness;
+		FReal SoftTwistStiffness;
+		FReal SoftSwingStiffness;
 
 		// @todo(ccaulfield): remove one of these
-		TRotation<T, d> AngularDriveTarget;
-		TVector<T, d> AngularDriveTargetAngles;
+		FRotation3 AngularDriveTarget;
+		FVec3 AngularDriveTargetAngles;
 
 		bool bAngularSLerpDriveEnabled;
 		bool bAngularTwistDriveEnabled;
 		bool bAngularSwingDriveEnabled;
-		T AngularDriveStiffness;
+		FReal AngularDriveStiffness;
 	};
 
 
-	template<class T, int d>
-	class CHAOS_API TPBDJointSettings
+	class CHAOS_API FPBDJointSettings
 	{
 	public:
-		using FTransformPair = TVector<TRigidTransform<T, d>, 2>;
+		using FTransformPair = TVector<FRigidTransform3, 2>;
 
-		TPBDJointSettings();
+		FPBDJointSettings();
 
 		// Particle-relative joint axes and positions
 		FTransformPair ConstraintFrames;
 
 		// How the constraint is allowed to move
-		TPBDJointMotionSettings<T, d> Motion;
+		FPBDJointMotionSettings Motion;
 	};
 
-	template<class T, int d>
-	class CHAOS_API TPBDJointSolverSettings
+	class CHAOS_API FPBDJointSolverSettings
 	{
 	public:
-		TPBDJointSolverSettings();
+		FPBDJointSolverSettings();
 
 		// Tolerances
-		T SwingTwistAngleTolerance;
+		FReal SwingTwistAngleTolerance;
 
 		// Stability control
-		T MinParentMassRatio;
-		T MaxInertiaRatio;
+		FReal MinParentMassRatio;
+		FReal MaxInertiaRatio;
 
 		// @todo(ccaulfield): remove these TEMP overrides for testing
 		bool bEnableVelocitySolve;
@@ -128,10 +121,10 @@ namespace Chaos
 		bool bEnableTwistLimits;
 		bool bEnableSwingLimits;
 		bool bEnableDrives;
-		T Projection;
-		T Stiffness;
-		T DriveStiffness;
-		T SoftLinearStiffness;
-		T SoftAngularStiffness;
+		FReal Projection;
+		FReal Stiffness;
+		FReal DriveStiffness;
+		FReal SoftLinearStiffness;
+		FReal SoftAngularStiffness;
 	};
 }
\ No newline at end of file
diff --git a/Engine/Source/Runtime/Experimental/Chaos/Public/Chaos/PBDJointConstraintUtilities.h b/Engine/Source/Runtime/Experimental/Chaos/Public/Chaos/PBDJointConstraintUtilities.h
--- a/Engine/Source/Runtime/Experimental/Chaos/Public/Chaos/PBDJointConstraintUtilities.h
+++ b/Engine/Source/Runtime/Experimental/Chaos/Public/Chaos/PBDJointConstraintUtilities.h
@@ -9,8 +9,7 @@
 
 namespace Chaos
 {
-	template<typename T, int d>
-	class TPBDJointUtilities
+	class FPBDJointUtilities
 	{
 	public:
 		/**
@@ -19,9 +18,9 @@ namespace Chaos
 		 * @param InI The input inertia.
 		 * @return An altered inertia so that the minimum element is at least MaxElement/MaxRatio.
 		 */
-		static CHAOS_API TVector<T, d> ConditionInertia(
-			const TVector<T, d>& InI, 
-			const T MaxRatio);
+		static CHAOS_API FVec3 ConditionInertia(
+			const FVec3& InI, 
+			const FReal MaxRatio);
 
 		/**
 		 * Increase the IParent inertia so that its largest component is at least MinRatio times the largest IChild component.
@@ -33,314 +32,314 @@ namespace Chaos
 		 * @param MinRatio Parent inertia will be at least this multiple of child inertia
 		 * @return The max/min ratio of InI elements.
 		 */
-		static CHAOS_API TVector<T, d> ConditionParentInertia(
-			const TVector<T, d>& IParent, 
-			const TVector<T, d>& IChild, 
-			const T MinRatio);
+		static CHAOS_API FVec3 ConditionParentInertia(
+			const FVec3& IParent, 
+			const FVec3& IChild, 
+			const FReal MinRatio);
 
-		static CHAOS_API T ConditionParentMass(
-			const T MParent, 
-			const T MChild, 
-			const T MinRatio);
+		static CHAOS_API FReal ConditionParentMass(
+			const FReal MParent, 
+			const FReal MChild, 
+			const FReal MinRatio);
 
 		static CHAOS_API void GetConditionedInverseMass(
 			const float MParent,
-			const TVector<T, d> IParent,
+			const FVec3 IParent,
 			const float MChild,
-			const TVector<T, d> IChild,
-			T& OutInvMParent,
-			T& OutInvMChild,
-			PMatrix<T, d, d>& OutInvIParent,
-			PMatrix<T, d, d>& OutInvIChild,
-			const T MinParentMassRatio,
-			const T MaxInertiaRatio);
+			const FVec3 IChild,
+			FReal& OutInvMParent,
+			FReal& OutInvMChild,
+			FMatrix33& OutInvIParent,
+			FMatrix33& OutInvIChild,
+			const FReal MinParentMassRatio,
+			const FReal MaxInertiaRatio);
 
 		static CHAOS_API void GetConditionedInverseMass(
 			const float M,
-			const TVector<T, d> I,
-			T& OutInvM0,
-			PMatrix<T, d, d>& OutInvI0, 
-			const T MaxInertiaRatio);
+			const FVec3 I,
+			FReal& OutInvM0,
+			FMatrix33& OutInvI0, 
+			const FReal MaxInertiaRatio);
 
 		static CHAOS_API void CalculateSwingConstraintSpace(
-			const TPBDJointSolverSettings<T, d>& SolverSettings,
-			const TPBDJointSettings<T, d>& JointSettings,
+			const FPBDJointSolverSettings& SolverSettings,
+			const FPBDJointSettings& JointSettings,
 			const int32 Index0,
 			const int32 Index1,
-			TVector<T, d>& P0,
-			TRotation<T, d>& Q0,
-			TVector<T, d>& P1,
-			TRotation<T, d>& Q1,
-			TVector<T, d>& OutX0,
-			PMatrix<T, d, d>& OutR0,
-			TVector<T, d>& OutX1,
-			PMatrix<T, d, d>& OutR1,
-			TVector<T, d>& OutCR);
+			FVec3& P0,
+			FRotation3& Q0,
+			FVec3& P1,
+			FRotation3& Q1,
+			FVec3& OutX0,
+			FMatrix33& OutR0,
+			FVec3& OutX1,
+			FMatrix33& OutR1,
+			FVec3& OutCR);
 
 		static CHAOS_API void CalculateConeConstraintSpace(
-			const TPBDJointSolverSettings<T, d>& SolverSettings,
-			const TPBDJointSettings<T, d>& JointSettings,
+			const FPBDJointSolverSettings& SolverSettings,
+			const FPBDJointSettings& JointSettings,
 			const int32 Index0,
 			const int32 Index1,
-			TVector<T, d>& P0,
-			TRotation<T, d>& Q0,
-			TVector<T, d>& P1,
-			TRotation<T, d>& Q1,
-			TVector<T, d>& OutX0,
-			PMatrix<T, d, d>& OutR0,
-			TVector<T, d>& OutX1,
-			PMatrix<T, d, d>& OutR1,
-			TVector<T, d>& OutCR);
+			FVec3& P0,
+			FRotation3& Q0,
+			FVec3& P1,
+			FRotation3& Q1,
+			FVec3& OutX0,
+			FMatrix33& OutR0,
+			FVec3& OutX1,
+			FMatrix33& OutR1,
+			FVec3& OutCR);
 
 		static CHAOS_API void ApplyJointPositionConstraint(
-			const T Dt,
-			const TPBDJointSolverSettings<T, d>& SolverSettings,
-			const TPBDJointSettings<T, d>& JointSettings,
+			const FReal Dt,
+			const FPBDJointSolverSettings& SolverSettings,
+			const FPBDJointSettings& JointSettings,
 			const int32 Index0,
 			const int32 Index1,
-			TVector<T, d>& P0,
-			TRotation<T, d>& Q0,
-			TVector<T, d>& P1,
-			TRotation<T, d>& Q1,
+			FVec3& P0,
+			FRotation3& Q0,
+			FVec3& P1,
+			FRotation3& Q1,
 			float InvM0,
-			const PMatrix<T, d, d>& InvIL0,
+			const FMatrix33& InvIL0,
 			float InvM1,
-			const PMatrix<T, d, d>& InvIL1);
+			const FMatrix33& InvIL1);
 
 		static CHAOS_API void ApplyJointVelocityConstraint(
-			const T Dt,
-			const TPBDJointSolverSettings<T, d>& SolverSettings,
-			const TPBDJointSettings<T, d>& JointSettings,
+			const FReal Dt,
+			const FPBDJointSolverSettings& SolverSettings,
+			const FPBDJointSettings& JointSettings,
 			const int32 Index0,
 			const int32 Index1,
-			TVector<T, d>& P0,
-			TRotation<T, d>& Q0,
-			TVector<T, d>& V0,
-			TVector<T, d>& W0,
-			TVector<T, d>& P1,
-			TRotation<T, d>& Q1,
-			TVector<T, d>& V1,
-			TVector<T, d>& W1,
+			FVec3& P0,
+			FRotation3& Q0,
+			FVec3& V0,
+			FVec3& W0,
+			FVec3& P1,
+			FRotation3& Q1,
+			FVec3& V1,
+			FVec3& W1,
 			float InvM0,
-			const PMatrix<T, d, d>& InvIL0,
+			const FMatrix33& InvIL0,
 			float InvM1,
-			const PMatrix<T, d, d>& InvIL1);
+			const FMatrix33& InvIL1);
 
 		static CHAOS_API void ApplyJointTwistConstraint(
-			const T Dt,
-			const TPBDJointSolverSettings<T, d>& SolverSettings,
-			const TPBDJointSettings<T, d>& JointSettings,
+			const FReal Dt,
+			const FPBDJointSolverSettings& SolverSettings,
+			const FPBDJointSettings& JointSettings,
 			const int32 Index0,
 			const int32 Index1,
-			TVector<T, d>& P0,
-			TRotation<T, d>& Q0,
-			TVector<T, d>& P1,
-			TRotation<T, d>& Q1,
+			FVec3& P0,
+			FRotation3& Q0,
+			FVec3& P1,
+			FRotation3& Q1,
 			float InvM0,
-			const PMatrix<T, d, d>& InvIL0,
+			const FMatrix33& InvIL0,
 			float InvM1,
-			const PMatrix<T, d, d>& InvIL1);
+			const FMatrix33& InvIL1);
 
 		static CHAOS_API void ApplyJointTwistVelocityConstraint(
-			const T Dt,
-			const TPBDJointSolverSettings<T, d>& SolverSettings,
-			const TPBDJointSettings<T, d>& JointSettings,
+			const FReal Dt,
+			const FPBDJointSolverSettings& SolverSettings,
+			const FPBDJointSettings& JointSettings,
 			const int32 Index0,
 			const int32 Index1,
-			TVector<T, d>& P0,
-			TRotation<T, d>& Q0,
-			TVector<T, d>& V0,
-			TVector<T, d>& W0,
-			TVector<T, d>& P1,
-			TRotation<T, d>& Q1,
-			TVector<T, d>& V1,
-			TVector<T, d>& W1,
+			FVec3& P0,
+			FRotation3& Q0,
+			FVec3& V0,
+			FVec3& W0,
+			FVec3& P1,
+			FRotation3& Q1,
+			FVec3& V1,
+			FVec3& W1,
 			float InvM0,
-			const PMatrix<T, d, d>& InvIL0,
+			const FMatrix33& InvIL0,
 			float InvM1,
-			const PMatrix<T, d, d>& InvIL1);
+			const FMatrix33& InvIL1);
 
 		static CHAOS_API void ApplyJointConeConstraint(
-			const T Dt,
-			const TPBDJointSolverSettings<T, d>& SolverSettings,
-			const TPBDJointSettings<T, d>& JointSettings,
+			const FReal Dt,
+			const FPBDJointSolverSettings& SolverSettings,
+			const FPBDJointSettings& JointSettings,
 			const int32 Index0,
 			const int32 Index1,
-			TVector<T, d>& P0,
-			TRotation<T, d>& Q0,
-			TVector<T, d>& P1,
-			TRotation<T, d>& Q1,
+			FVec3& P0,
+			FRotation3& Q0,
+			FVec3& P1,
+			FRotation3& Q1,
 			float InvM0,
-			const PMatrix<T, d, d>& InvIL0,
+			const FMatrix33& InvIL0,
 			float InvM1,
-			const PMatrix<T, d, d>& InvIL1);
+			const FMatrix33& InvIL1);
 
 		static CHAOS_API void ApplyJointConeVelocityConstraint(
-			const T Dt,
-			const TPBDJointSolverSettings<T, d>& SolverSettings,
-			const TPBDJointSettings<T, d>& JointSettings,
+			const FReal Dt,
+			const FPBDJointSolverSettings& SolverSettings,
+			const FPBDJointSettings& JointSettings,
 			const int32 Index0,
 			const int32 Index1,
-			TVector<T, d>& P0,
-			TRotation<T, d>& Q0,
-			TVector<T, d>& V0,
-			TVector<T, d>& W0,
-			TVector<T, d>& P1,
-			TRotation<T, d>& Q1,
-			TVector<T, d>& V1,
-			TVector<T, d>& W1,
+			FVec3& P0,
+			FRotation3& Q0,
+			FVec3& V0,
+			FVec3& W0,
+			FVec3& P1,
+			FRotation3& Q1,
+			FVec3& V1,
+			FVec3& W1,
 			float InvM0,
-			const PMatrix<T, d, d>& InvIL0,
+			const FMatrix33& InvIL0,
 			float InvM1,
-			const PMatrix<T, d, d>& InvIL1);
+			const FMatrix33& InvIL1);
 
 		static CHAOS_API void ApplyJointSwingConstraint(
-			const T Dt,
-			const TPBDJointSolverSettings<T, d>& SolverSettings,
-			const TPBDJointSettings<T, d>& JointSettings,
+			const FReal Dt,
+			const FPBDJointSolverSettings& SolverSettings,
+			const FPBDJointSettings& JointSettings,
 			const int32 Index0,
 			const int32 Index1,
 			const EJointAngularConstraintIndex SwingConstraintIndex,
 			const EJointAngularAxisIndex SwingAxisIndex,
-			TVector<T, d>& P0,
-			TRotation<T, d>& Q0,
-			TVector<T, d>& P1,
-			TRotation<T, d>& Q1,
+			FVec3& P0,
+			FRotation3& Q0,
+			FVec3& P1,
+			FRotation3& Q1,
 			float InvM0,
-			const PMatrix<T, d, d>& InvIL0,
+			const FMatrix33& InvIL0,
 			float InvM1,
-			const PMatrix<T, d, d>& InvIL1);
+			const FMatrix33& InvIL1);
 
 		static CHAOS_API void ApplyJointSwingVelocityConstraint(
-			const T Dt,
-			const TPBDJointSolverSettings<T, d>& SolverSettings,
-			const TPBDJointSettings<T, d>& JointSettings,
+			const FReal Dt,
+			const FPBDJointSolverSettings& SolverSettings,
+			const FPBDJointSettings& JointSettings,
 			const int32 Index0,
 			const int32 Index1,
 			const EJointAngularConstraintIndex SwingConstraintIndex,
 			const EJointAngularAxisIndex SwingAxisIndex,
-			TVector<T, d>& P0,
-			TRotation<T, d>& Q0,
-			TVector<T, d>& V0,
-			TVector<T, d>& W0,
-			TVector<T, d>& P1,
-			TRotation<T, d>& Q1,
-			TVector<T, d>& V1,
-			TVector<T, d>& W1,
+			FVec3& P0,
+			FRotation3& Q0,
+			FVec3& V0,
+			FVec3& W0,
+			FVec3& P1,
+			FRotation3& Q1,
+			FVec3& V1,
+			FVec3& W1,
 			float InvM0,
-			const PMatrix<T, d, d>& InvIL0,
+			const FMatrix33& InvIL0,
 			float InvM1,
-			const PMatrix<T, d, d>& InvIL1);
+			const FMatrix33& InvIL1);
 
 		static CHAOS_API void ApplyJointTwistDrive(
-			const T Dt,
-			const TPBDJointSolverSettings<T, d>& SolverSettings,
-			const TPBDJointSettings<T, d>& JointSettings,
+			const FReal Dt,
+			const FPBDJointSolverSettings& SolverSettings,
+			const FPBDJointSettings& JointSettings,
 			const int32 Index0,
 			const int32 Index1,
-			TVector<T, d>& P0,
-			TRotation<T, d>& Q0,
-			TVector<T, d>& P1,
-			TRotation<T, d>& Q1,
+			FVec3& P0,
+			FRotation3& Q0,
+			FVec3& P1,
+			FRotation3& Q1,
 			float InvM0,
-			const PMatrix<T, d, d>& InvIL0,
+			const FMatrix33& InvIL0,
 			float InvM1,
-			const PMatrix<T, d, d>& InvIL1);
+			const FMatrix33& InvIL1);
 
 		static CHAOS_API void ApplyJointConeDrive(
-			const T Dt,
-			const TPBDJointSolverSettings<T, d>& SolverSettings,
-			const TPBDJointSettings<T, d>& JointSettings,
+			const FReal Dt,
+			const FPBDJointSolverSettings& SolverSettings,
+			const FPBDJointSettings& JointSettings,
 			const int32 Index0,
 			const int32 Index1,
-			TVector<T, d>& P0,
-			TRotation<T, d>& Q0,
-			TVector<T, d>& P1,
-			TRotation<T, d>& Q1,
+			FVec3& P0,
+			FRotation3& Q0,
+			FVec3& P1,
+			FRotation3& Q1,
 			float InvM0,
-			const PMatrix<T, d, d>& InvIL0,
+			const FMatrix33& InvIL0,
 			float InvM1,
-			const PMatrix<T, d, d>& InvIL1);
+			const FMatrix33& InvIL1);
 
 		static CHAOS_API void ApplyJointSLerpDrive(
-			const T Dt,
-			const TPBDJointSolverSettings<T, d>& SolverSettings,
-			const TPBDJointSettings<T, d>& JointSettings,
+			const FReal Dt,
+			const FPBDJointSolverSettings& SolverSettings,
+			const FPBDJointSettings& JointSettings,
 			const int32 Index0,
 			const int32 Index1,
-			TVector<T, d>& P0,
-			TRotation<T, d>& Q0,
-			TVector<T, d>& P1,
-			TRotation<T, d>& Q1,
+			FVec3& P0,
+			FRotation3& Q0,
+			FVec3& P1,
+			FRotation3& Q1,
 			float InvM0,
-			const PMatrix<T, d, d>& InvIL0,
+			const FMatrix33& InvIL0,
 			float InvM1,
-			const PMatrix<T, d, d>& InvIL1);
+			const FMatrix33& InvIL1);
 
 		static CHAOS_API void ApplyJointPositionProjection(
-			const T Dt,
-			const TPBDJointSolverSettings<T, d>& SolverSettings,
-			const TPBDJointSettings<T, d>& JointSettings,
+			const FReal Dt,
+			const FPBDJointSolverSettings& SolverSettings,
+			const FPBDJointSettings& JointSettings,
 			const int32 Index0,
 			const int32 Index1,
-			TVector<T, d>& P0,
-			TRotation<T, d>& Q0,
-			TVector<T, d>& P1,
-			TRotation<T, d>& Q1,
+			FVec3& P0,
+			FRotation3& Q0,
+			FVec3& P1,
+			FRotation3& Q1,
 			float InvM0,
-			const PMatrix<T, d, d>& InvIL0,
+			const FMatrix33& InvIL0,
 			float InvM1,
-			const PMatrix<T, d, d>& InvIL1,
-			const T ProjectionFactor);
+			const FMatrix33& InvIL1,
+			const FReal ProjectionFactor);
 
 		static CHAOS_API void ApplyJointTwistProjection(
-			const T Dt,
-			const TPBDJointSolverSettings<T, d>& SolverSettings,
-			const TPBDJointSettings<T, d>& JointSettings,
+			const FReal Dt,
+			const FPBDJointSolverSettings& SolverSettings,
+			const FPBDJointSettings& JointSettings,
 			const int32 Index0,
 			const int32 Index1,
-			TVector<T, d>& P0,
-			TRotation<T, d>& Q0,
-			TVector<T, d>& P1,
-			TRotation<T, d>& Q1,
+			FVec3& P0,
+			FRotation3& Q0,
+			FVec3& P1,
+			FRotation3& Q1,
 			float InvM0,
-			const PMatrix<T, d, d>& InvIL0,
+			const FMatrix33& InvIL0,
 			float InvM1,
-			const PMatrix<T, d, d>& InvIL1,
-			const T ProjectionFactor);
+			const FMatrix33& InvIL1,
+			const FReal ProjectionFactor);
 
 		static CHAOS_API void ApplyJointConeProjection(
-			const T Dt,
-			const TPBDJointSolverSettings<T, d>& SolverSettings,
-			const TPBDJointSettings<T, d>& JointSettings,
+			const FReal Dt,
+			const FPBDJointSolverSettings& SolverSettings,
+			const FPBDJointSettings& JointSettings,
 			const int32 Index0,
 			const int32 Index1,
-			TVector<T, d>& P0,
-			TRotation<T, d>& Q0,
-			TVector<T, d>& P1,
-			TRotation<T, d>& Q1,
+			FVec3& P0,
+			FRotation3& Q0,
+			FVec3& P1,
+			FRotation3& Q1,
 			float InvM0,
-			const PMatrix<T, d, d>& InvIL0,
+			const FMatrix33& InvIL0,
 			float InvM1,
-			const PMatrix<T, d, d>& InvIL1,
-			const T ProjectionFactor);
+			const FMatrix33& InvIL1,
+			const FReal ProjectionFactor);
 
 		static CHAOS_API void ApplyJointSwingProjection(
-			const T Dt,
-			const TPBDJointSolverSettings<T, d>& SolverSettings,
-			const TPBDJointSettings<T, d>& JointSettings,
+			const FReal Dt,
+			const FPBDJointSolverSettings& SolverSettings,
+			const FPBDJointSettings& JointSettings,
 			const int32 Index0,
 			const int32 Index1,
 			const EJointAngularConstraintIndex SwingConstraint,
-			TVector<T, d>& P0,
-			TRotation<T, d>& Q0,
-			TVector<T, d>& P1,
-			TRotation<T, d>& Q1,
+			FVec3& P0,
+			FRotation3& Q0,
+			FVec3& P1,
+			FRotation3& Q1,
 			float InvM0,
-			const PMatrix<T, d, d>& InvIL0,
+			const FMatrix33& InvIL0,
 			float InvM1,
-			const PMatrix<T, d, d>& InvIL1,
-			const T ProjectionFactor);
+			const FMatrix33& InvIL1,
+			const FReal ProjectionFactor);
 
 	};
 }
\ No newline at end of file
diff --git a/Engine/Source/Runtime/Experimental/Chaos/Public/Chaos/PBDJointConstraints.h b/Engine/Source/Runtime/Experimental/Chaos/Public/Chaos/PBDJointConstraints.h
--- a/Engine/Source/Runtime/Experimental/Chaos/Public/Chaos/PBDJointConstraints.h
+++ b/Engine/Source/Runtime/Experimental/Chaos/Public/Chaos/PBDJointConstraints.h
@@ -4,6 +4,7 @@
 #include "CoreMinimal.h"
 
 #include "Chaos/Array.h"
+#include "Chaos/Core.h"
 #include "Chaos/Transform.h"
 #include "Chaos/Vector.h"
 
@@ -14,30 +15,28 @@
 
 namespace Chaos
 {
-	template<class T, int d>
-	class TPBDJointConstraintHandle : public TContainerConstraintHandle<TPBDJointConstraints<T, d>>
+	class FPBDJointConstraintHandle : public TContainerConstraintHandle<FPBDJointConstraints>
 	{
 	public:
-		using Base = TContainerConstraintHandle<TPBDJointConstraints<T, d>>;
-		using FConstraintContainer = TPBDJointConstraints<T, d>;
+		using Base = TContainerConstraintHandle<FPBDJointConstraints>;
+		using FConstraintContainer = FPBDJointConstraints;
 
-		CHAOS_API TPBDJointConstraintHandle();
-		CHAOS_API TPBDJointConstraintHandle(FConstraintContainer* InConstraintContainer, int32 InConstraintIndex);
+		CHAOS_API FPBDJointConstraintHandle();
+		CHAOS_API FPBDJointConstraintHandle(FConstraintContainer* InConstraintContainer, int32 InConstraintIndex);
 
-		CHAOS_API void CalculateConstraintSpace(TVector<T, d>& OutXa, PMatrix<T, d, d>& OutRa, TVector<T, d>& OutXb, PMatrix<T, d, d>& OutRb, TVector<T, d>& OutCR) const;
+		CHAOS_API void CalculateConstraintSpace(FVec3& OutXa, FMatrix33& OutRa, FVec3& OutXb, FMatrix33& OutRb, FVec3& OutCR) const;
 		CHAOS_API void SetParticleLevels(const TVector<int32, 2>& ParticleLevels);
 		CHAOS_API int32 GetConstraintLevel() const;
-		CHAOS_API const TPBDJointSettings<T, d>& GetSettings() const;
+		CHAOS_API const FPBDJointSettings& GetSettings() const;
 	protected:
 		using Base::ConstraintIndex;
 		using Base::ConstraintContainer;
 	};
 
-	template<class T, int d>
-	class CHAOS_API TPBDJointState
+	class CHAOS_API FPBDJointState
 	{
 	public:
-		TPBDJointState();
+		FPBDJointState();
 
 		// Priorities used for ordering, mass conditioning, projection, and freezing
 		int32 Level;
@@ -47,28 +46,23 @@ namespace Chaos
 	/**
 	 * A joint restricting up to 6 degrees of freedom, with linear and angular limits.
 	 */
-	template<class T, int d>
-	class TPBDJointConstraints : public TPBDConstraintContainer<T, d>
+	class CHAOS_API FPBDJointConstraints : public FPBDConstraintContainer
 	{
 	public:
-		using Base = TPBDConstraintContainer<T, d>;
-		using FReal = T;
-		static const int Dimensions = d;
+		using Base = FPBDConstraintContainer;
 
-		using FConstraintHandle = TPBDJointConstraintHandle<FReal, Dimensions>;
-		using FConstraintHandleAllocator = TConstraintHandleAllocator<TPBDJointConstraints<FReal, Dimensions>>;
-		using FParticlePair = TVector<TGeometryParticleHandle<FReal, Dimensions>*, 2>;
-		using FVectorPair = TVector<TVector<FReal, Dimensions>, 2>;
-		using FTransformPair = TVector<TRigidTransform<FReal, Dimensions>, 2>;
-		using FJointSettings = TPBDJointSettings<FReal, Dimensions>;
-		using FJointState = TPBDJointState<FReal, Dimensions>;
+		using FConstraintContainerHandle = FPBDJointConstraintHandle;
+		using FConstraintHandleAllocator = TConstraintHandleAllocator<FPBDJointConstraints>;
+		using FParticlePair = TVector<TGeometryParticleHandle<FReal, 3>*, 2>;
+		using FVectorPair = TVector<FVec3, 2>;
+		using FTransformPair = TVector<FRigidTransform3, 2>;
 
-		CHAOS_API TPBDJointConstraints(const TPBDJointSolverSettings<T, d>& InSettings = TPBDJointSolverSettings<T, d>());
+		FPBDJointConstraints(const FPBDJointSolverSettings& InSettings = FPBDJointSolverSettings());
 
-		CHAOS_API virtual ~TPBDJointConstraints();
+		virtual ~FPBDJointConstraints();
 
-		CHAOS_API const TPBDJointSolverSettings<T, d>& GetSettings() const;
-		CHAOS_API void SetSettings(const TPBDJointSolverSettings<T, d>& InSettings);
+		const FPBDJointSolverSettings& GetSettings() const;
+		void SetSettings(const FPBDJointSolverSettings& InSettings);
 
 		//
 		// Constraint Container API
@@ -77,82 +71,82 @@ namespace Chaos
 		/**
 		 * Get the number of constraints.
 		 */
-		CHAOS_API int32 NumConstraints() const;
+		int32 NumConstraints() const;
 
 		/**
 		 * Add a constraint with particle-space constraint offsets.
 		 */
-		CHAOS_API FConstraintHandle* AddConstraint(const FParticlePair& InConstrainedParticles, const TRigidTransform<FReal, Dimensions>& WorldConstraintFrame);
-		CHAOS_API FConstraintHandle* AddConstraint(const FParticlePair& InConstrainedParticles, const FTransformPair& ConstraintFrames);
-		CHAOS_API FConstraintHandle* AddConstraint(const FParticlePair& InConstrainedParticles, const TPBDJointSettings<T, d>& InConstraintSettings);
+		FConstraintContainerHandle* AddConstraint(const FParticlePair& InConstrainedParticles, const FRigidTransform3& WorldConstraintFrame);
+		FConstraintContainerHandle* AddConstraint(const FParticlePair& InConstrainedParticles, const FTransformPair& ConstraintFrames);
+		FConstraintContainerHandle* AddConstraint(const FParticlePair& InConstrainedParticles, const FPBDJointSettings& InConstraintSettings);
 
 		/**
 		 * Remove the specified constraint.
 		 */
-		CHAOS_API void RemoveConstraint(int ConstraintIndex);
+		void RemoveConstraint(int ConstraintIndex);
 
 		// @todo(ccaulfield): rename/remove  this
-		CHAOS_API void RemoveConstraints(const TSet<TGeometryParticleHandle<T, d>*>& RemovedParticles);
+		void RemoveConstraints(const TSet<TGeometryParticleHandle<FReal, 3>*>& RemovedParticles);
 
-		CHAOS_API void SetPreApplyCallback(const TJointPostApplyCallback<T, d>& Callback);
-		CHAOS_API void ClearPreApplyCallback();
+		void SetPreApplyCallback(const FJointPostApplyCallback& Callback);
+		void ClearPreApplyCallback();
 
-		CHAOS_API void SetPostApplyCallback(const TJointPostApplyCallback<T, d>& Callback);
-		CHAOS_API void ClearPostApplyCallback();
+		void SetPostApplyCallback(const FJointPostApplyCallback& Callback);
+		void ClearPostApplyCallback();
 
 		//
 		// Constraint API
 		//
 
-		CHAOS_API const FConstraintHandle* GetConstraintHandle(int32 ConstraintIndex) const;
-		CHAOS_API FConstraintHandle* GetConstraintHandle(int32 ConstraintIndex);
+		const FConstraintContainerHandle* GetConstraintHandle(int32 ConstraintIndex) const;
+		FConstraintContainerHandle* GetConstraintHandle(int32 ConstraintIndex);
 
 		/**
 		 * Get the particles that are affected by the specified constraint.
 		 */
-		CHAOS_API const FParticlePair& GetConstrainedParticles(int32 ConstraintIndex) const;
+		const FParticlePair& GetConstrainedParticles(int32 ConstraintIndex) const;
 
-		CHAOS_API const TPBDJointSettings<T, d>& GetConstraintSettings(int32 ConstraintIndex) const;
+		const FPBDJointSettings& GetConstraintSettings(int32 ConstraintIndex) const;
 
-		CHAOS_API int32 GetConstraintLevel(int32 ConstraintIndex) const;
-		CHAOS_API void SetParticleLevels(int32 ConstraintIndex, const TVector<int32, 2>& ParticleLevels);
+		int32 GetConstraintLevel(int32 ConstraintIndex) const;
+		void SetParticleLevels(int32 ConstraintIndex, const TVector<int32, 2>& ParticleLevels);
 
 		//
 		// Island Rule API
 		//
 
-		CHAOS_API void UpdatePositionBasedState(const T Dt);
+		void UpdatePositionBasedState(const FReal Dt);
 
-		CHAOS_API void Apply(const T Dt, const TArray<FConstraintHandle*>& InConstraintHandles, const int32 It, const int32 NumIts);
+		void Apply(const FReal Dt, const TArray<FConstraintContainerHandle*>& InConstraintHandles, const int32 It, const int32 NumIts);
 
-		CHAOS_API bool ApplyPushOut(const T Dt, const TArray<FConstraintHandle*>& InConstraintHandles, const int32 It, const int32 NumIts);
+		bool ApplyPushOut(const FReal Dt, const TArray<FConstraintContainerHandle*>& InConstraintHandles, const int32 It, const int32 NumIts);
 
-		CHAOS_API void ApplyProjection(const T Dt, const TArray<FConstraintHandle*>& InConstraintHandles);
+		void ApplyProjection(const FReal Dt, const TArray<FConstraintContainerHandle*>& InConstraintHandles);
 
 	protected:
 		using Base::GetConstraintIndex;
 		using Base::SetConstraintIndex;
 
 	private:
-		friend class TPBDJointConstraintHandle<T, d>;
+		friend class FPBDJointConstraintHandle;
 
-		CHAOS_API void CalculateConstraintSpace(int32 ConstraintIndex, TVector<T, d>& OutX0, PMatrix<T, d, d>& OutR0, TVector<T, d>& OutX1, PMatrix<T, d, d>& OutR1, TVector<T, d>& OutAngles) const;
+		void CalculateConstraintSpace(int32 ConstraintIndex, FVec3& OutX0, FMatrix33& OutR0, FVec3& OutX1, FMatrix33& OutR1, FVec3& OutAngles) const;
 
-		CHAOS_API void SolveVelocity(const T Dt, const int32 ConstraintIndex, const int32 It, const int32 NumIts);
-		CHAOS_API void SolvePosition(const T Dt, const int32 ConstraintIndex, const int32 It, const int32 NumIts);
-		CHAOS_API void ProjectPosition(const T Dt, const int32 ConstraintIndex, const int32 It, const int32 NumIts);
+		void SolveVelocity(const FReal Dt, const int32 ConstraintIndex, const int32 It, const int32 NumIts);
+		void SolvePosition(const FReal Dt, const int32 ConstraintIndex, const int32 It, const int32 NumIts);
+		void ProjectPosition(const FReal Dt, const int32 ConstraintIndex, const int32 It, const int32 NumIts);
 
-		TPBDJointSolverSettings<T, d> Settings;
+		FPBDJointSolverSettings Settings;
 
-		TArray<FJointSettings> ConstraintSettings;
+		TArray<FPBDJointSettings> ConstraintSettings;
 		TArray<FParticlePair> ConstraintParticles;
-		TArray<FJointState> ConstraintStates;
+		TArray<FPBDJointState> ConstraintStates;
 
-		TArray<FConstraintHandle*> Handles;
+		TArray<FConstraintContainerHandle*> Handles;
 		FConstraintHandleAllocator HandleAllocator;
 
-		TJointPreApplyCallback<T, d> PreApplyCallback;
-		TJointPostApplyCallback<T, d> PostApplyCallback;
+		FJointPreApplyCallback PreApplyCallback;
+		FJointPostApplyCallback PostApplyCallback;
 	};
 
 }
diff --git a/Engine/Source/Runtime/Experimental/Chaos/Public/Chaos/PBDLongRangeConstraints.h b/Engine/Source/Runtime/Experimental/Chaos/Public/Chaos/PBDLongRangeConstraints.h
--- a/Engine/Source/Runtime/Experimental/Chaos/Public/Chaos/PBDLongRangeConstraints.h
+++ b/Engine/Source/Runtime/Experimental/Chaos/Public/Chaos/PBDLongRangeConstraints.h
@@ -8,7 +8,7 @@
 namespace Chaos
 {
 template<class T, int d>
-class CHAOS_API PBDLongRangeConstraints : public TPBDLongRangeConstraintsBase<T, d>, public TPBDConstraintContainer<T, d>
+class CHAOS_API PBDLongRangeConstraints : public TPBDLongRangeConstraintsBase<T, d>, public FPBDConstraintContainer
 {
 	typedef TPBDLongRangeConstraintsBase<T, d> Base;
 	using Base::MConstraints;
diff --git a/Engine/Source/Runtime/Experimental/Chaos/Public/Chaos/PBDPositionConstraints.h b/Engine/Source/Runtime/Experimental/Chaos/Public/Chaos/PBDPositionConstraints.h
--- a/Engine/Source/Runtime/Experimental/Chaos/Public/Chaos/PBDPositionConstraints.h
+++ b/Engine/Source/Runtime/Experimental/Chaos/Public/Chaos/PBDPositionConstraints.h
@@ -27,13 +27,13 @@ namespace Chaos
 	};
 
 	template<class T, int d>
-	class TPBDPositionConstraints : public TPBDConstraintContainer<T, d>
+	class TPBDPositionConstraints : public FPBDConstraintContainer
 	{
 	public:
-		using Base = TPBDConstraintContainer<T, d>;
+		using Base = FPBDConstraintContainer;
 		using FReal = T;
 		static const int Dimensions = d;
-		using FConstraintHandle = TPBDPositionConstraintHandle<FReal, Dimensions>;
+		using FConstraintContainerHandle = TPBDPositionConstraintHandle<FReal, Dimensions>;
 		using FConstraintHandleAllocator = TConstraintHandleAllocator<TPBDPositionConstraints<FReal, Dimensions>>;
 
 		TPBDPositionConstraints(const T InStiffness = (T)1)
@@ -71,7 +71,7 @@ namespace Chaos
 		/**
 		 * Add a constraint.
 		 */
-		FConstraintHandle* AddConstraint(TPBDRigidParticleHandle<T, d>* Particle, const TVector<T, d>& Position)
+		FConstraintContainerHandle* AddConstraint(TPBDRigidParticleHandle<T, d>* Particle, const TVector<T, d>& Position)
 		{
 			int32 NewIndex = Targets.Num();
 			Targets.Add(Position);
@@ -85,7 +85,7 @@ namespace Chaos
 		 */
 		void RemoveConstraint(int ConstraintIndex)
 		{
-			FConstraintHandle* ConstraintHandle = Handles[ConstraintIndex];
+			FConstraintContainerHandle* ConstraintHandle = Handles[ConstraintIndex];
 			if (ConstraintHandle != nullptr)
 			{
 				// Release the handle for the freed constraint
@@ -115,12 +115,12 @@ namespace Chaos
 		// Constraint API
 		//
 
-		const FConstraintHandle* GetConstraintHandle(int32 ConstraintIndex) const
+		const FConstraintContainerHandle* GetConstraintHandle(int32 ConstraintIndex) const
 		{
 			return Handles[ConstraintIndex];
 		}
 
-		FConstraintHandle* GetConstraintHandle(int32 ConstraintIndex)
+		FConstraintContainerHandle* GetConstraintHandle(int32 ConstraintIndex)
 		{
 			return Handles[ConstraintIndex];
 		}
@@ -156,15 +156,15 @@ namespace Chaos
 		{
 		}
 
-		void Apply(const T Dt, const TArray<FConstraintHandle*>& ConstraintHandles, const int32 It, const int32 NumIts) const
+		void Apply(const T Dt, const TArray<FConstraintContainerHandle*>& ConstraintHandles, const int32 It, const int32 NumIts) const
 		{
-			for (FConstraintHandle* ConstraintHandle : ConstraintHandles)
+			for (FConstraintContainerHandle* ConstraintHandle : ConstraintHandles)
 			{
 				ApplySingle(Dt, ConstraintHandle->GetConstraintIndex());
 			}
 		}
 
-		bool ApplyPushOut(const T Dt, const TArray<FConstraintHandle*>& InConstraintIndices, const int32 It, const int32 NumIts) const
+		bool ApplyPushOut(const T Dt, const TArray<FConstraintContainerHandle*>& InConstraintIndices, const int32 It, const int32 NumIts) const
 		{
 			return false;
 		}
@@ -189,7 +189,7 @@ namespace Chaos
 		TArray<TPBDRigidParticleHandle<T,d>*> ConstrainedParticles;
 		T Stiffness;
 
-		TArray<FConstraintHandle*> Handles;
+		TArray<FConstraintContainerHandle*> Handles;
 		FConstraintHandleAllocator HandleAllocator;
 	};
 }
diff --git a/Engine/Source/Runtime/Experimental/Chaos/Public/Chaos/PBDRigidDynamicSpringConstraints.h b/Engine/Source/Runtime/Experimental/Chaos/Public/Chaos/PBDRigidDynamicSpringConstraints.h
--- a/Engine/Source/Runtime/Experimental/Chaos/Public/Chaos/PBDRigidDynamicSpringConstraints.h
+++ b/Engine/Source/Runtime/Experimental/Chaos/Public/Chaos/PBDRigidDynamicSpringConstraints.h
@@ -27,14 +27,14 @@ namespace Chaos
 	};
 
 	template<class T, int d>
-	class CHAOS_API TPBDRigidDynamicSpringConstraints : public TPBDConstraintContainer<T, d>
+	class CHAOS_API TPBDRigidDynamicSpringConstraints : public FPBDConstraintContainer
 	{
 	public:
-		using Base = TPBDConstraintContainer<T, d>;
+		using Base = FPBDConstraintContainer;
 		using FConstrainedParticlePair = TVector<TGeometryParticleHandle<T, d>*, 2>;
 		using FReal = T;
 		static const int Dimensions = d;
-		using FConstraintHandle = TPBDRigidDynamicSpringConstraintHandle<FReal, Dimensions>;
+		using FConstraintContainerHandle = TPBDRigidDynamicSpringConstraintHandle<FReal, Dimensions>;
 		using FConstraintHandleAllocator = TConstraintHandleAllocator<TPBDRigidDynamicSpringConstraints<FReal, Dimensions>>;
 
 		TPBDRigidDynamicSpringConstraints(const T InStiffness = (T)1)
@@ -76,7 +76,7 @@ namespace Chaos
 		 * Add a constraint initialized from current world-space particle positions.
 		 * You would use this method when your objects are already positioned in the world.
 		 */
-		FConstraintHandle* AddConstraint(const FConstrainedParticlePair& InConstrainedParticles)
+		FConstraintContainerHandle* AddConstraint(const FConstrainedParticlePair& InConstrainedParticles)
 		{
 			Handles.Add(HandleAllocator.AllocHandle(this, Handles.Num()));
 			Constraints.Add(InConstrainedParticles);
@@ -90,7 +90,7 @@ namespace Chaos
 		 */
 		void RemoveConstraint(int ConstraintIndex)
 		{
-			FConstraintHandle* ConstraintHandle = Handles[ConstraintIndex];
+			FConstraintContainerHandle* ConstraintHandle = Handles[ConstraintIndex];
 			if (ConstraintHandle != nullptr)
 			{
 				// Release the handle for the freed constraint
@@ -137,12 +137,12 @@ namespace Chaos
 		// Constraint API
 		//
 
-		const FConstraintHandle* GetConstraintHandle(int32 ConstraintIndex) const
+		const FConstraintContainerHandle* GetConstraintHandle(int32 ConstraintIndex) const
 		{
 			return Handles[ConstraintIndex];
 		}
 
-		FConstraintHandle* GetConstraintHandle(int32 ConstraintIndex)
+		FConstraintContainerHandle* GetConstraintHandle(int32 ConstraintIndex)
 		{
 			return Handles[ConstraintIndex];
 		}
@@ -162,15 +162,15 @@ namespace Chaos
 
 		void UpdatePositionBasedState(const T Dt);
 
-		void Apply(const T Dt, const TArray<FConstraintHandle*>& InConstraintHandles, const int32 It, const int32 NumIts)
+		void Apply(const T Dt, const TArray<FConstraintContainerHandle*>& InConstraintHandles, const int32 It, const int32 NumIts)
 		{
-			for (FConstraintHandle* ConstraintHandle : InConstraintHandles)
+			for (FConstraintContainerHandle* ConstraintHandle : InConstraintHandles)
 			{
 				ApplySingle(Dt, ConstraintHandle->GetConstraintIndex());
 			}
 		}
 
-		bool ApplyPushOut(const T Dt, const TArray<FConstraintHandle*>& InConstraintHandles, const int32 It, const int32 NumIts)
+		bool ApplyPushOut(const T Dt, const TArray<FConstraintContainerHandle*>& InConstraintHandles, const int32 It, const int32 NumIts)
 		{
 			return false;
 		}
@@ -191,7 +191,7 @@ namespace Chaos
 		int32 MaxSprings;
 		T Stiffness;
 
-		TArray<FConstraintHandle*> Handles;
+		TArray<FConstraintContainerHandle*> Handles;
 		FConstraintHandleAllocator HandleAllocator;
 	};
 }
diff --git a/Engine/Source/Runtime/Experimental/Chaos/Public/Chaos/PBDRigidSpringConstraints.h b/Engine/Source/Runtime/Experimental/Chaos/Public/Chaos/PBDRigidSpringConstraints.h
--- a/Engine/Source/Runtime/Experimental/Chaos/Public/Chaos/PBDRigidSpringConstraints.h
+++ b/Engine/Source/Runtime/Experimental/Chaos/Public/Chaos/PBDRigidSpringConstraints.h
@@ -31,13 +31,13 @@ namespace Chaos
 
 
 	template<class T, int d>
-	class TPBDRigidSpringConstraints : public TPBDConstraintContainer<T, d>
+	class TPBDRigidSpringConstraints : public FPBDConstraintContainer
 	{
 	public:
-		using Base = TPBDConstraintContainer<T, d>;
+		using Base = FPBDConstraintContainer;
 		using FReal = T;
 		static const int Dimensions = d;
-		using FConstraintHandle = TPBDRigidSpringConstraintHandle<FReal, Dimensions>;
+		using FConstraintContainerHandle = TPBDRigidSpringConstraintHandle<FReal, Dimensions>;
 		using FConstraintHandleAllocator = TConstraintHandleAllocator<TPBDRigidSpringConstraints<FReal, Dimensions>>;
 		using FConstrainedParticlePair = TVector<TGeometryParticleHandle<T, d>*, 2>;
 
@@ -81,7 +81,7 @@ namespace Chaos
 		 * Add a constraint initialized from current world-space particle positions.
 		 * You would use this method when your objects are already positioned in the world.
 		 */
-		FConstraintHandle* AddConstraint(const FConstrainedParticlePair& InConstrainedParticles, const  TVector<TVector<T, 3>, 2>& InLocations)
+		FConstraintContainerHandle* AddConstraint(const FConstrainedParticlePair& InConstrainedParticles, const  TVector<TVector<T, 3>, 2>& InLocations)
 		{
 			Handles.Add(HandleAllocator.AllocHandle(this, Handles.Num()));
 			int32 ConstraintIndex = Constraints.Add(InConstrainedParticles);
@@ -96,7 +96,7 @@ namespace Chaos
 		 */
 		void RemoveConstraint(int ConstraintIndex)
 		{
-			FConstraintHandle* ConstraintHandle = Handles[ConstraintIndex];
+			FConstraintContainerHandle* ConstraintHandle = Handles[ConstraintIndex];
 			if (ConstraintHandle != nullptr)
 			{
 				// Release the handle for the freed constraint
@@ -125,12 +125,12 @@ namespace Chaos
 		// Constraint API
 		//
 
-		const FConstraintHandle* GetConstraintHandle(int32 ConstraintIndex) const
+		const FConstraintContainerHandle* GetConstraintHandle(int32 ConstraintIndex) const
 		{
 			return Handles[ConstraintIndex];
 		}
 
-		FConstraintHandle* GetConstraintHandle(int32 ConstraintIndex)
+		FConstraintContainerHandle* GetConstraintHandle(int32 ConstraintIndex)
 		{
 			return Handles[ConstraintIndex];
 		}
@@ -168,15 +168,15 @@ namespace Chaos
 		{
 		}
 
-		void Apply(const T Dt, const TArray<FConstraintHandle*>& InConstraintHandles, const int32 It, const int32 NumIts)
+		void Apply(const T Dt, const TArray<FConstraintContainerHandle*>& InConstraintHandles, const int32 It, const int32 NumIts)
 		{
-			for (FConstraintHandle* ConstraintHandle : InConstraintHandles)
+			for (FConstraintContainerHandle* ConstraintHandle : InConstraintHandles)
 			{
 				ApplySingle(Dt, ConstraintHandle->GetConstraintIndex());
 			}
 		}
 
-		bool ApplyPushOut(const T Dt, const TArray<FConstraintHandle*>& InConstraintHandles, const int32 It, const int32 NumIts)
+		bool ApplyPushOut(const T Dt, const TArray<FConstraintContainerHandle*>& InConstraintHandles, const int32 It, const int32 NumIts)
 		{
 			return false;
 		}
@@ -197,7 +197,7 @@ namespace Chaos
 		TArray<T> SpringDistances;
 		T Stiffness;
 
-		TArray<FConstraintHandle*> Handles;
+		TArray<FConstraintContainerHandle*> Handles;
 		FConstraintHandleAllocator HandleAllocator;
 	};
 }
diff --git a/Engine/Source/Runtime/Experimental/Chaos/Public/Chaos/PBDRigidsEvolution.h b/Engine/Source/Runtime/Experimental/Chaos/Public/Chaos/PBDRigidsEvolution.h
--- a/Engine/Source/Runtime/Experimental/Chaos/Public/Chaos/PBDRigidsEvolution.h
+++ b/Engine/Source/Runtime/Experimental/Chaos/Public/Chaos/PBDRigidsEvolution.h
@@ -272,10 +272,7 @@ class TPBDRigidsEvolutionBase
 	CHAOS_API TPBDRigidsSOAs<T,d>& GetParticles() { return Particles; }
 	CHAOS_API const TPBDRigidsSOAs<T, d>& GetParticles() const { return Particles; }
 
-	typedef TPBDConstraintGraph<T, d> FConstraintGraph;
-	typedef TPBDConstraintGraphRule<T, d> FConstraintRule;
-
-	CHAOS_API void AddConstraintRule(FConstraintRule* ConstraintRule)
+	CHAOS_API void AddConstraintRule(FPBDConstraintGraphRule* ConstraintRule)
 	{
 		uint32 ContainerId = (uint32)ConstraintRules.Num();
 		ConstraintRules.Add(ConstraintRule);
@@ -364,7 +361,7 @@ class TPBDRigidsEvolutionBase
 	// @todo(ccaulfield): Remove the uint version
 	CHAOS_API void RemoveConstraints(const TSet<TGeometryParticleHandle<T, d>*>& RemovedParticles)
 	{
-		for (FConstraintRule* ConstraintRule : ConstraintRules)
+		for (FPBDConstraintGraphRule* ConstraintRule : ConstraintRules)
 		{
 			ConstraintRule->RemoveConstraints(RemovedParticles);
 		}
@@ -388,14 +385,14 @@ class TPBDRigidsEvolutionBase
 	{
 		ConstraintGraph.InitializeGraph(Particles.GetNonDisabledView());
 
-		for (FConstraintRule* ConstraintRule : ConstraintRules)
+		for (FPBDConstraintGraphRule* ConstraintRule : ConstraintRules)
 		{
 			ConstraintRule->AddToGraph();
 		}
 
 		ConstraintGraph.ResetIslands(Particles.GetNonDisabledDynamicView());
 
-		for (FConstraintRule* ConstraintRule : ConstraintRules)
+		for (FPBDConstraintGraphRule* ConstraintRule : ConstraintRules)
 		{
 			ConstraintRule->InitializeAccelerationStructures();
 		}
@@ -403,7 +400,7 @@ class TPBDRigidsEvolutionBase
 
 	void UpdateAccelerationStructures(int32 Island)
 	{
-		for (FConstraintRule* ConstraintRule : ConstraintRules)
+		for (FPBDConstraintGraphRule* ConstraintRule : ConstraintRules)
 		{
 			ConstraintRule->UpdateAccelerationStructures(Island);
 		}
@@ -411,7 +408,7 @@ class TPBDRigidsEvolutionBase
 
 	void ApplyConstraints(const T Dt, int32 Island)
 	{
-		for (FConstraintRule* ConstraintRule : ConstraintRules)
+		for (FPBDConstraintGraphRule* ConstraintRule : ConstraintRules)
 		{
 			ConstraintRule->UpdateAccelerationStructures(Island);
 		}
@@ -419,7 +416,7 @@ class TPBDRigidsEvolutionBase
 		// @todo(ccaulfield): track whether we are sufficiently solved and can early-out
 		for (int i = 0; i < NumIterations; ++i)
 		{
-			for (FConstraintRule* ConstraintRule : ConstraintRules)
+			for (FPBDConstraintGraphRule* ConstraintRule : ConstraintRules)
 			{
 				ConstraintRule->ApplyConstraints(Dt, Island, i, NumIterations);
 			}
@@ -512,7 +509,7 @@ protected:
 	int32 NumConstraints() const
 	{
 		int32 NumConstraints = 0;
-		for (const FConstraintRule* ConstraintRule : ConstraintRules)
+		for (const FPBDConstraintGraphRule* ConstraintRule : ConstraintRules)
 		{
 			NumConstraints += ConstraintRule->NumConstraints();
 		}
@@ -544,7 +541,7 @@ protected:
 
 	void UpdateConstraintPositionBasedState(T Dt)
 	{
-		for (FConstraintRule* ConstraintRule : ConstraintRules)
+		for (FPBDConstraintGraphRule* ConstraintRule : ConstraintRules)
 		{
 			ConstraintRule->UpdatePositionBasedState(Dt);
 		}
@@ -554,7 +551,7 @@ protected:
 	{
 		ConstraintGraph.InitializeGraph(Particles.GetNonDisabledView());
 
-		for (FConstraintRule* ConstraintRule : ConstraintRules)
+		for (FPBDConstraintGraphRule* ConstraintRule : ConstraintRules)
 		{
 			ConstraintRule->AddToGraph();
 		}
@@ -564,7 +561,7 @@ protected:
 	{
 		ConstraintGraph.UpdateIslands(Particles.GetNonDisabledDynamicView(), Particles);
 
-		for (FConstraintRule* ConstraintRule : ConstraintRules)
+		for (FPBDConstraintGraphRule* ConstraintRule : ConstraintRules)
 		{
 			ConstraintRule->InitializeAccelerationStructures();
 		}
@@ -581,7 +578,7 @@ protected:
 		for (int32 It = 0; bNeedsAnotherIteration && (It < NumPushOutIterations); ++It)
 		{
 			bNeedsAnotherIteration = false;
-			for (FConstraintRule* ConstraintRule : ConstraintRules)
+			for (FPBDConstraintGraphRule* ConstraintRule : ConstraintRules)
 			{
 				if (ConstraintRule->ApplyPushOut(Dt, Island, It, NumPushOutIterations))
 				{
@@ -603,8 +600,8 @@ protected:
 	FUpdateVelocityRule ParticleUpdateVelocity;
 	FUpdatePositionRule ParticleUpdatePosition;
 	FKinematicUpdateRule KinematicUpdate;
-	TArray<FConstraintRule*> ConstraintRules;
-	FConstraintGraph ConstraintGraph;
+	TArray<FPBDConstraintGraphRule*> ConstraintRules;
+	FPBDConstraintGraph ConstraintGraph;
 	TArrayCollectionArray<TSerializablePtr<TChaosPhysicsMaterial<T>>> PhysicsMaterials;
 	TArrayCollectionArray<TUniquePtr<TChaosPhysicsMaterial<T>>> PerParticlePhysicsMaterials;
 	TArrayCollectionArray<int32> ParticleDisableCount;
diff --git a/Engine/Source/Runtime/Experimental/Chaos/Public/Chaos/PBDRigidsEvolutionGBF.h b/Engine/Source/Runtime/Experimental/Chaos/Public/Chaos/PBDRigidsEvolutionGBF.h
--- a/Engine/Source/Runtime/Experimental/Chaos/Public/Chaos/PBDRigidsEvolutionGBF.h
+++ b/Engine/Source/Runtime/Experimental/Chaos/Public/Chaos/PBDRigidsEvolutionGBF.h
@@ -49,7 +49,7 @@ public:
 	using FGravityForces = TPerParticleGravity<T, d>;
 	using FCollisionConstraints = TPBDCollisionConstraint<T, d>;
 	using FExternalForces = TPerParticleExternalForces<T, d>;
-	using FCollisionConstraintRule = TPBDConstraintColorRule<FCollisionConstraints, T, d>;
+	using FCollisionConstraintRule = TPBDConstraintColorRule<FCollisionConstraints>;
 
 	static constexpr int32 DefaultNumIterations = 1;
 	static constexpr int32 DefaultNumPushOutIterations = 5;
diff --git a/Engine/Source/Runtime/Experimental/Chaos/Public/Chaos/PBDSpringConstraints.h b/Engine/Source/Runtime/Experimental/Chaos/Public/Chaos/PBDSpringConstraints.h
--- a/Engine/Source/Runtime/Experimental/Chaos/Public/Chaos/PBDSpringConstraints.h
+++ b/Engine/Source/Runtime/Experimental/Chaos/Public/Chaos/PBDSpringConstraints.h
@@ -11,7 +11,7 @@
 namespace Chaos
 {
 template<class T, int32 d>
-class TPBDSpringConstraints : public TPBDSpringConstraintsBase<T, d>, public TPBDConstraintContainer<T, d>
+class TPBDSpringConstraints : public TPBDSpringConstraintsBase<T, d>, public FPBDConstraintContainer
 {
 	typedef TPBDSpringConstraintsBase<T, d> Base;
 	using Base::MConstraints;
diff --git a/Engine/Source/Runtime/Experimental/ChaosSolvers/Private/PhysicsProxy/SkeletalMeshPhysicsProxy.cpp b/Engine/Source/Runtime/Experimental/ChaosSolvers/Private/PhysicsProxy/SkeletalMeshPhysicsProxy.cpp
--- a/Engine/Source/Runtime/Experimental/ChaosSolvers/Private/PhysicsProxy/SkeletalMeshPhysicsProxy.cpp
+++ b/Engine/Source/Runtime/Experimental/ChaosSolvers/Private/PhysicsProxy/SkeletalMeshPhysicsProxy.cpp
@@ -7,7 +7,7 @@
 
 FSkeletalMeshPhysicsProxy::FSkeletalMeshPhysicsProxy(UObject* InOwner, const FInitFunc& InInitFunc)
 	: Base(InOwner)
-	, JointConstraints(Chaos::TPBDJointSolverSettings<float, 3>())
+	, JointConstraints(Chaos::FPBDJointSolverSettings())
 	, JointConstraintsRule(JointConstraints)
 	, NextInputProducerBuffer(nullptr)
 	, CurrentOutputConsumerBuffer(nullptr)
diff --git a/Engine/Source/Runtime/Experimental/ChaosSolvers/Public/PBDRigidsSolver.h b/Engine/Source/Runtime/Experimental/ChaosSolvers/Public/PBDRigidsSolver.h
--- a/Engine/Source/Runtime/Experimental/ChaosSolvers/Public/PBDRigidsSolver.h
+++ b/Engine/Source/Runtime/Experimental/ChaosSolvers/Public/PBDRigidsSolver.h
@@ -75,13 +75,12 @@ namespace Chaos
 		typedef Chaos::TPBDCollisionConstraint<float, 3> FPBDCollisionConstraints;
 
 		typedef FPBDCollisionConstraints FCollisionConstraints;
-		typedef TPBDJointConstraints<float, 3> FJointConstraints;
 		typedef TPBDRigidDynamicSpringConstraints<float, 3> FRigidDynamicSpringConstraints;
 		typedef TPBDPositionConstraints<float, 3> FPositionConstraints;
 
-		typedef TPBDConstraintIslandRule<FJointConstraints, float, 3> FJointConstraintsRule;
-		typedef TPBDConstraintIslandRule<FRigidDynamicSpringConstraints, float, 3> FRigidDynamicSpringConstraintsRule;
-		typedef TPBDConstraintIslandRule<FPositionConstraints, float, 3> FPositionConstraintsRule;
+		typedef TPBDConstraintIslandRule<FPBDJointConstraints> FJointConstraintsRule;
+		typedef TPBDConstraintIslandRule<FRigidDynamicSpringConstraints> FRigidDynamicSpringConstraintsRule;
+		typedef TPBDConstraintIslandRule<FPositionConstraints> FPositionConstraintsRule;
 
 		//
 		// Execution API
diff --git a/Engine/Source/Runtime/Experimental/ChaosSolvers/Public/PhysicsProxy/SkeletalMeshPhysicsProxy.h b/Engine/Source/Runtime/Experimental/ChaosSolvers/Public/PhysicsProxy/SkeletalMeshPhysicsProxy.h
--- a/Engine/Source/Runtime/Experimental/ChaosSolvers/Public/PhysicsProxy/SkeletalMeshPhysicsProxy.h
+++ b/Engine/Source/Runtime/Experimental/ChaosSolvers/Public/PhysicsProxy/SkeletalMeshPhysicsProxy.h
@@ -168,8 +168,8 @@ public:
 	const FBoneHierarchy& GetBoneHierarchy() const { return Parameters.BoneHierarchy; }
 
 private:
-	typedef Chaos::TPBDJointConstraints<float, 3> FJointConstraints;
-	typedef Chaos::TPBDConstraintIslandRule<FJointConstraints, float, 3> FJointConstraintsRule;
+	using FJointConstraints = Chaos::FPBDJointConstraints;
+	using FJointConstraintsRule = Chaos::TPBDConstraintIslandRule<FJointConstraints>;
 
 	FSkeletalMeshPhysicsProxyParams Parameters;
 	TArray<int32> RigidBodyIds;
