Speculative fix for strict CIS error (PBDRigidEvolution.cpp)

diff --git a/Engine/Source/Runtime/Experimental/Chaos/Private/Chaos/PBDRigidsEvolutionGBF.cpp b/Engine/Source/Runtime/Experimental/Chaos/Private/Chaos/PBDRigidsEvolutionGBF.cpp
--- a/Engine/Source/Runtime/Experimental/Chaos/Private/Chaos/PBDRigidsEvolutionGBF.cpp
+++ b/Engine/Source/Runtime/Experimental/Chaos/Private/Chaos/PBDRigidsEvolutionGBF.cpp
@@ -162,7 +162,7 @@ void TPBDRigidsEvolutionGBF<T, d>::AdvanceOneTimeStep(const T Dt, const T StepFr
 	if(Dt > 0)
 	{
 		SCOPE_CYCLE_COUNTER(STAT_Evo_ParallelSolve);
-		PhysicsParallelFor(ConstraintGraph.NumIslands(), [&](int32 Island) {
+		PhysicsParallelFor(ConstraintGraph.NumIslands(), [this, &SleepedIslands, &DisabledParticles, Dt](int32 Island) {
 			const TArray<TGeometryParticleHandle<T, d>*>& IslandParticles = ConstraintGraph.GetIslandParticles(Island);
 
 			{
