Next attempt to fix strict CIS (PBDRigidsEvolution.cpp)

diff --git a/Engine/Source/Runtime/Experimental/Chaos/Private/Chaos/PBDRigidsEvolutionGBF.cpp b/Engine/Source/Runtime/Experimental/Chaos/Private/Chaos/PBDRigidsEvolutionGBF.cpp
--- a/Engine/Source/Runtime/Experimental/Chaos/Private/Chaos/PBDRigidsEvolutionGBF.cpp
+++ b/Engine/Source/Runtime/Experimental/Chaos/Private/Chaos/PBDRigidsEvolutionGBF.cpp
@@ -155,15 +155,15 @@ void TPBDRigidsEvolutionGBF<T, d>::AdvanceOneTimeStep(const T Dt, const T StepFr
 	}
 
 	TArray<bool> SleepedIslands;
-	SleepedIslands.SetNum(ConstraintGraph.NumIslands());
+	SleepedIslands.SetNum(GetConstraintGraph().NumIslands());
 	TArray<TArray<TPBDRigidParticleHandle<T,d>*>> DisabledParticles;
-	DisabledParticles.SetNum(ConstraintGraph.NumIslands());
-	SleepedIslands.SetNum(ConstraintGraph.NumIslands());
+	DisabledParticles.SetNum(GetConstraintGraph().NumIslands());
+	SleepedIslands.SetNum(GetConstraintGraph().NumIslands());
 	if(Dt > 0)
 	{
 		SCOPE_CYCLE_COUNTER(STAT_Evo_ParallelSolve);
-		PhysicsParallelFor(ConstraintGraph.NumIslands(), [this, &SleepedIslands, &DisabledParticles, Dt](int32 Island) {
-			const TArray<TGeometryParticleHandle<T, d>*>& IslandParticles = ConstraintGraph.GetIslandParticles(Island);
+		PhysicsParallelFor(GetConstraintGraph().NumIslands(), [&](int32 Island) {
+			const TArray<TGeometryParticleHandle<T, d>*>& IslandParticles = GetConstraintGraph().GetIslandParticles(Island);
 
 			{
 				SCOPE_CYCLE_COUNTER(STAT_Evo_ApplyConstraints);
@@ -223,16 +223,16 @@ void TPBDRigidsEvolutionGBF<T, d>::AdvanceOneTimeStep(const T Dt, const T StepFr
 			}
 
 			// Turn off if not moving
-			SleepedIslands[Island] = ConstraintGraph.SleepInactive(Island, PhysicsMaterials);
+			SleepedIslands[Island] = GetConstraintGraph().SleepInactive(Island, PhysicsMaterials);
 		});
 	}
 	{
 		SCOPE_CYCLE_COUNTER(STAT_Evo_DeactivateSleep);
-		for (int32 Island = 0; Island < ConstraintGraph.NumIslands(); ++Island)
+		for (int32 Island = 0; Island < GetConstraintGraph().NumIslands(); ++Island)
 		{
 			if (SleepedIslands[Island])
 			{
-				Particles.DeactivateParticles(ConstraintGraph.GetIslandParticles(Island));
+				Particles.DeactivateParticles(GetConstraintGraph().GetIslandParticles(Island));
 			}
 			for (const auto Particle : DisabledParticles[Island])
 			{
diff --git a/Engine/Source/Runtime/Experimental/Chaos/Public/Chaos/PBDRigidsEvolution.h b/Engine/Source/Runtime/Experimental/Chaos/Public/Chaos/PBDRigidsEvolution.h
--- a/Engine/Source/Runtime/Experimental/Chaos/Public/Chaos/PBDRigidsEvolution.h
+++ b/Engine/Source/Runtime/Experimental/Chaos/Public/Chaos/PBDRigidsEvolution.h
@@ -503,6 +503,9 @@ class TPBDRigidsEvolutionBase
 	const auto& GetRigidClustering() const { return Clustering; }
 	auto& GetRigidClustering() { return Clustering; }
 
+	CHAOS_API const FPBDConstraintGraph& GetConstraintGraph() const { return ConstraintGraph; }
+	CHAOS_API FPBDConstraintGraph& GetConstraintGraph() { return ConstraintGraph; }
+
 	void Serialize(FChaosArchive& Ar);
 
 protected:
diff --git a/Engine/Source/Runtime/Experimental/Chaos/Public/Chaos/PBDRigidsEvolutionGBF.h b/Engine/Source/Runtime/Experimental/Chaos/Public/Chaos/PBDRigidsEvolutionGBF.h
--- a/Engine/Source/Runtime/Experimental/Chaos/Public/Chaos/PBDRigidsEvolutionGBF.h
+++ b/Engine/Source/Runtime/Experimental/Chaos/Public/Chaos/PBDRigidsEvolutionGBF.h
@@ -179,8 +179,8 @@ protected:
 	using Base::UpdateConstraintPositionBasedState;
 	using Base::CreateConstraintGraph;
 	using Base::CreateIslands;
-	using Base::ConstraintGraph;
 	using Base::ConstraintRules;
+	using Base::GetConstraintGraph;
 	using Base::UpdateVelocities;
 	using Base::PhysicsMaterials;
 	using Base::ParticleDisableCount;
