Make trimesh use an AABBTree instead of a grid

diff --git a/Engine/Source/Runtime/Engine/Private/PhysicsEngine/Experimental/ChaosDerivedData.cpp b/Engine/Source/Runtime/Engine/Private/PhysicsEngine/Experimental/ChaosDerivedData.cpp
--- a/Engine/Source/Runtime/Engine/Private/PhysicsEngine/Experimental/ChaosDerivedData.cpp
+++ b/Engine/Source/Runtime/Engine/Private/PhysicsEngine/Experimental/ChaosDerivedData.cpp
@@ -21,7 +21,7 @@ const TCHAR* FChaosDerivedDataCooker::GetPluginName() const
 
 const TCHAR* FChaosDerivedDataCooker::GetVersionString() const
 {
-	return TEXT("EF20258F42C54D658B5BEE64EB779C6A");
+	return TEXT("40C102F05B01482DB74F98FE0F6C4166");
 }
 
 FString FChaosDerivedDataCooker::GetPluginSpecificCacheKeySuffix() const
diff --git a/Engine/Source/Runtime/Experimental/Chaos/Private/Chaos/TriangleMeshImplicitObject.cpp b/Engine/Source/Runtime/Experimental/Chaos/Private/Chaos/TriangleMeshImplicitObject.cpp
--- a/Engine/Source/Runtime/Experimental/Chaos/Private/Chaos/TriangleMeshImplicitObject.cpp
+++ b/Engine/Source/Runtime/Experimental/Chaos/Private/Chaos/TriangleMeshImplicitObject.cpp
@@ -182,6 +182,12 @@ struct TTriangleMeshRaycastVisitor
 		return Visit<ERaycastType::Sweep>(TriIdx.Payload, CurLength);
 	}
 
+	bool VisitOverlap(TSpatialVisitorData<int32> TriIdx)
+	{
+		check(false);
+		return true;
+	}
+
 	const TParticles<T, 3>& Particles;
 	const TArray<TVector<int32, 3>>& Elements;
 	const TVector<T, 3>& StartPoint;
@@ -366,6 +372,18 @@ struct TTriangleMeshSweepVisitor
 	{
 	}
 
+	bool VisitOverlap(const TSpatialVisitorData<int32>& VisitData)
+	{
+		check(false);
+		return true;
+	}
+
+	bool VisitRaycast(const TSpatialVisitorData<int32>& VisitData, T& CurLength)
+	{
+		check(false);
+		return true;
+	}
+
 	bool VisitSweep(const TSpatialVisitorData<int32>& VisitData, T& CurLength)
 	{
 		const int32 TriIdx = VisitData.Payload;
@@ -378,7 +396,7 @@ struct TTriangleMeshSweepVisitor
 			TriMesh.MParticles.X(TriMesh.MElements[TriIdx][1]),
 			TriMesh.MParticles.X(TriMesh.MElements[TriIdx][2]));
 
-		if(GJKRaycast<T>(Tri, QueryGeom, StartTM, Dir, CurLength, Time, HitPosition, HitNormal, Thickness))
+		if(GJKRaycast2<T>(Tri, QueryGeom, StartTM, Dir, CurLength, Time, HitPosition, HitNormal, Thickness))
 		{
 			if(Time < OutTime)
 			{
@@ -420,7 +438,7 @@ bool TTriangleMeshImplicitObject<T>::SweepGeomImp(const QueryGeomType& QueryGeom
 	const TBox<T, 3> QueryBounds = QueryGeom.BoundingBox();
 	const TVector<T, 3> StartPoint = StartTM.TransformPositionNoScale(QueryBounds.Center());
 	const TVector<T, 3> Inflation = QueryBounds.Extents() * 0.5 + TVector<T, 3>(Thickness);
-	BVH.template Sweep<TTriangleMeshSweepVisitor<QueryGeomType, T>, false>(StartPoint, Dir, Length, Inflation, SQVisitor);
+	BVH.template Sweep<TTriangleMeshSweepVisitor<QueryGeomType, T>>(StartPoint, Dir, Length, Inflation, SQVisitor);
 
 	if (SQVisitor.OutTime <= Length)
 	{
@@ -566,6 +584,7 @@ void Chaos::TTriangleMeshImplicitObject<T>::RebuildBV()
 		BVEntries.Add({ this, Tri });
 	}
 	BVH.Reinitialize(BVEntries);
+
 }
 
 
diff --git a/Engine/Source/Runtime/Experimental/Chaos/Public/Chaos/AABBTree.h b/Engine/Source/Runtime/Experimental/Chaos/Public/Chaos/AABBTree.h
--- a/Engine/Source/Runtime/Experimental/Chaos/Public/Chaos/AABBTree.h
+++ b/Engine/Source/Runtime/Experimental/Chaos/Public/Chaos/AABBTree.h
@@ -226,7 +226,7 @@ public:
 	static constexpr T DefaultMaxPayloadBounds = 100000;
 	static constexpr int32 DefaultMaxChildrenInLeaf = 12;
 	static constexpr int32 DefaultMaxTreeDepth = 16;
-	static constexpr int32 DefaultMaxNumToProcess = 4000;
+	static constexpr int32 DefaultMaxNumToProcess = 4000000;
 	static constexpr ESpatialAcceleration StaticType = TIsSame<TAABBTreeLeafArray<TPayloadType, T>, TLeafType>::Value ? ESpatialAcceleration::AABBTree : 
 		(TIsSame<TBoundingVolume<TPayloadType, T, 3>, TLeafType>::Value ? ESpatialAcceleration::AABBTreeBV : ESpatialAcceleration::Unknown);
 	TAABBTree()
@@ -276,6 +276,14 @@ public:
 
 	}
 
+	template <typename ParticleView>
+	void Reinitialize(const ParticleView& Particles)
+	{
+		GenerateTree(Particles);
+	}
+
+	virtual TArray<TPayloadType> FindAllIntersections(const TBox<T, 3>& Box) const override { return FindAllIntersectionsImp(Box); }
+
 	bool DumpTreeToLog(int32 NodeIdx, int32 ParentNode, TBox<T,3>& Bounds)
 	{
 		if (Nodes[NodeIdx].bLeaf)
@@ -835,6 +843,36 @@ private:
 		return; // keep working
 	}
 
+	TArray<TPayloadType> FindAllIntersectionsImp(const TBox<T, 3>& Intersection) const
+	{
+		struct FSimpleVisitor
+		{
+			FSimpleVisitor(TArray<TPayloadType>& InResults) : CollectedResults(InResults) {}
+			bool VisitOverlap(const TSpatialVisitorData<TPayloadType>& Instance)
+			{
+				CollectedResults.Add(Instance.Payload);
+				return true;
+			}
+			bool VisitSweep(const TSpatialVisitorData<TPayloadType>& Instance, T& CurLength)
+			{
+				check(false);
+				return true;
+			}
+			bool VisitRaycast(const TSpatialVisitorData<TPayloadType>& Instance, T& CurLength)
+			{
+				check(false);
+				return true;
+			}
+			TArray<TPayloadType>& CollectedResults;
+		};
+
+		TArray<TPayloadType> Results;
+		FSimpleVisitor Collector(Results);
+		Overlap(Intersection, Collector);
+
+		return Results;
+	}
+
 
 	TAABBTree(const TAABBTree<TPayloadType, TLeafType, T>& Other)
 		: ISpatialAcceleration<TPayloadType, T, 3>(StaticType)
@@ -880,8 +918,7 @@ private:
 template<typename TPayloadType, typename TLeafType, class T>
 FArchive& operator<<(FChaosArchive& Ar, TAABBTree<TPayloadType, TLeafType, T>& AABBTree)
 {
-	check(false);
-	//AABBTree.Serialize(Ar);
+	AABBTree.Serialize(Ar);
 	return Ar;
 }
 
diff --git a/Engine/Source/Runtime/Experimental/Chaos/Public/Chaos/TriangleMeshImplicitObject.h b/Engine/Source/Runtime/Experimental/Chaos/Public/Chaos/TriangleMeshImplicitObject.h
--- a/Engine/Source/Runtime/Experimental/Chaos/Public/Chaos/TriangleMeshImplicitObject.h
+++ b/Engine/Source/Runtime/Experimental/Chaos/Public/Chaos/TriangleMeshImplicitObject.h
@@ -12,6 +12,7 @@
 #include "BoundingVolume.h"
 #include "ChaosArchive.h"
 #include "UObject/ExternalPhysicsCustomObjectVersion.h"
+#include "AABBTree.h"
 
 namespace Chaos
 {
@@ -131,7 +132,8 @@ namespace Chaos
 		TArray<TVector<int32, 3>> MElements;
 		TBox<T, 3> MLocalBoundingBox;
 
-		using BVHType = TBoundingVolume<int32, T, 3>;
+		//using BVHType = TBoundingVolume<int32, T, 3>;
+		using BVHType = TAABBTree<int32, TAABBTreeLeafArray<int32, T>, T>;
 
 		template<typename InStorageType, typename InRealType>
 		friend struct FBvEntry;
