Make sure existing sq captures keep working

diff --git a/Engine/Source/Runtime/Core/Public/UObject/ExternalPhysicsCustomObjectVersion.h b/Engine/Source/Runtime/Core/Public/UObject/ExternalPhysicsCustomObjectVersion.h
--- a/Engine/Source/Runtime/Core/Public/UObject/ExternalPhysicsCustomObjectVersion.h
+++ b/Engine/Source/Runtime/Core/Public/UObject/ExternalPhysicsCustomObjectVersion.h
@@ -48,6 +48,9 @@ struct CORE_API FExternalPhysicsCustomObjectVersion
 		// Allow scaled geometry to be a concrete type
 		ScaledGeometryIsConcrete,
 
+		// Trimeshes serialize AABBTree
+		TrimeshSerializesAABBTree,
+
 		// -----<new versions can be added above this line>-------------------------------------------------
 		VersionPlusOne,
 		LatestVersion = VersionPlusOne - 1
diff --git a/Engine/Source/Runtime/Experimental/Chaos/Public/Chaos/TriangleMeshImplicitObject.h b/Engine/Source/Runtime/Experimental/Chaos/Public/Chaos/TriangleMeshImplicitObject.h
--- a/Engine/Source/Runtime/Experimental/Chaos/Public/Chaos/TriangleMeshImplicitObject.h
+++ b/Engine/Source/Runtime/Experimental/Chaos/Public/Chaos/TriangleMeshImplicitObject.h
@@ -95,6 +95,12 @@ namespace Chaos
 				// Should now only hit when loading older trimeshes
 				RebuildBV();
 			}
+			else if(Ar.CustomVer(FExternalPhysicsCustomObjectVersion::GUID) < FExternalPhysicsCustomObjectVersion::TrimeshSerializesAABBTree)
+			{
+				TBoundingVolume<int32, T, 3> Dummy;
+				Ar << Dummy;
+				RebuildBV();
+			}
 			else
 			{
 				// Serialize acceleration
