Allow sphere sweeps to use inflated raycasts

diff --git a/Engine/Source/Runtime/Experimental/Chaos/Public/Chaos/GeometryQueries.h b/Engine/Source/Runtime/Experimental/Chaos/Public/Chaos/GeometryQueries.h
--- a/Engine/Source/Runtime/Experimental/Chaos/Public/Chaos/GeometryQueries.h
+++ b/Engine/Source/Runtime/Experimental/Chaos/Public/Chaos/GeometryQueries.h
@@ -119,7 +119,7 @@ namespace Chaos
 		const FRigidTransform3 BToATM = BTM.GetRelativeTransform(ATM);
 		const FVec3 LocalDir = ATM.InverseTransformVectorNoScale(Dir);
 
-		bool bSweepAsRaycast = false;//todo: fix BType == ImplicitObjectType::Sphere;
+		bool bSweepAsRaycast = BType == ImplicitObjectType::Sphere;
 		if (bSweepAsRaycast && IsScaled(AType))
 		{
 			const auto& Scaled = TImplicitObjectScaledGeneric<FReal, 3>::AsScaledChecked(A);
@@ -129,11 +129,10 @@ namespace Chaos
 
 		if (bSweepAsRaycast)
 		{
-			/*
-			const TSphere<T, d>& BSphere = static_cast<const TSphere<T, d>&>(B);
-			const TVector<T, d> Start = BToATM.TransformPositionNoScale(BSphere.GetCenter());
-			bResult = A.Raycast(Start, LocalDir, Length, Thickness + BSphere.GetRadius(), OutTime, LocalPosition, LocalNormal, OutFaceIndex);*/
-			ensure(false);	//todo: override helper
+			const FImplicitObject& BBase = B;
+			const TSphere<FReal, 3>& BSphere = BBase.GetObjectChecked<TSphere<FReal, 3>>();
+			const FVec3 Start = BToATM.TransformPositionNoScale(BSphere.GetCenter());
+			bResult = A.Raycast(Start, LocalDir, Length, Thickness + BSphere.GetRadius(), OutTime, LocalPosition, LocalNormal, OutFaceIndex);
 		}
 		//todo: handle case where A is a sphere
 		else if (A.IsConvex())
@@ -202,7 +201,7 @@ namespace Chaos
 			// Compute MTD in the case of an initial overlap
 			if (bResult && bComputeMTD && OutTime == 0.f)
 			{
-				ensure(false); // We don't support MTD for non-convex types yet!
+				//ensure(false); // We don't support MTD for non-convex types yet!
 				OutNormal = FVec3(0, 0, 1);
 			}
 		}
