Fix compiler warning about unused variable

diff --git a/Engine/Source/Runtime/Experimental/Chaos/Public/Chaos/GJK.h b/Engine/Source/Runtime/Experimental/Chaos/Public/Chaos/GJK.h
--- a/Engine/Source/Runtime/Experimental/Chaos/Public/Chaos/GJK.h
+++ b/Engine/Source/Runtime/Experimental/Chaos/Public/Chaos/GJK.h
@@ -376,7 +376,7 @@ namespace Chaos
 		TVector<T, 3> V = -InitialDir;
 		FSimplex SimplexIDs;
 		TVector<T, 3> Simplex[4];
-		T Barycentric[4];
+		T Barycentric[4] = { -1,-1,-1,-1 };	//not needed, but compiler warns
 		const TRotation<T, 3> AToBRotation = BToATM.GetRotation().Inverse();
 		bool bTerminate;
 		bool bNearZero = false;
@@ -450,7 +450,7 @@ namespace Chaos
 		TVector<T, 3> As[4] = { TVector<T,3>(0), TVector<T,3>(0), TVector<T,3>(0), TVector<T,3>(0) };
 		TVector<T, 3> Bs[4] = { TVector<T,3>(0), TVector<T,3>(0), TVector<T,3>(0), TVector<T,3>(0) };
 
-		T Barycentric[4];
+		T Barycentric[4] = { -1,-1,-1,-1 };	//not needed, but compiler warns
 
 		FSimplex SimplexIDs;
 		const TRotation<T, 3> BToARotation = StartTM.GetRotation();
@@ -591,7 +591,7 @@ namespace Chaos
 		TVector<T, 3> As[4] = { TVector<T,3>(0), TVector<T,3>(0), TVector<T,3>(0), TVector<T,3>(0) };
 		TVector<T, 3> Bs[4] = { TVector<T,3>(0), TVector<T,3>(0), TVector<T,3>(0), TVector<T,3>(0) };
 
-		T Barycentric[4];
+		T Barycentric[4] = { -1,-1,-1,-1 };	//not needed, but compiler warns
 		const T Inflation = ThicknessA + ThicknessB;
 		const T Inflation2 = Inflation*Inflation + 1e-6;
 
