Chaos : Convex Collisions
- Renamed History to Manifold
- Moved the convex-convex code out to a separate class with static accessors.

diff --git a/Engine/Source/Runtime/Experimental/Chaos/Private/Chaos/CollisionResolutionConvexConvex.cpp b/Engine/Source/Runtime/Experimental/Chaos/Private/Chaos/CollisionResolutionConvexConvex.cpp
new file mode 100644
--- /dev/null
+++ b/Engine/Source/Runtime/Experimental/Chaos/Private/Chaos/CollisionResolutionConvexConvex.cpp
@@ -0,0 +1,130 @@
+// Copyright 1998-2019 Epic Games, Inc. All Rights Reserved.
+
+#include "Chaos/CollisionResolutionConvexConvex.h"
+#include "Chaos/ImplicitObjectScaled.h"
+
+namespace Chaos
+{
+
+	template<class T, int d>
+	void CollisionResolutionConvexConvex<T, d>::ConstructConvexConvexConstraints(
+		FGeometryParticleHandle* Particle0, FGeometryParticleHandle* Particle1,
+		const FImplicitObject* Implicit0, const FImplicitObject* Implicit1,
+		const float Thickness,
+		TArray<FRigidBodyContactConstraint>& ConstraintBuffer,
+		FCollisionResolutionManifold* Manifold,
+		float ManifoldScale, int32 ManifoldSamples)
+	{
+		if (Manifold)
+		{
+			if (Manifold->ContainsShapeConnection(Implicit0, Implicit1))
+			{
+				return;
+			}
+		}
+
+		for (int i = 0; i < ManifoldSamples; i++)
+		{
+			TRigidBodyContactConstraint<T, d> Constraint;
+			Constraint.Particle = Particle0;
+			Constraint.Levelset = Particle1;
+			Constraint.Geometry[0] = Implicit0;
+			Constraint.Geometry[1] = Implicit1;
+			ConstraintBuffer.Add(Constraint);
+		}
+	}
+
+	template<class T, int d>
+	void CollisionResolutionConvexConvex<T, d>::UpdateConvexConvexConstraint(
+		const FImplicitObject& A, const FRigidTransform& ATM,
+		const FImplicitObject& B, FRigidTransform BTM,
+		float Thickness,
+		FRigidBodyContactConstraint* Constraints, int32 NumConstraints,
+		FCollisionResolutionManifold* Manifold,
+		float ManifoldScale, int32 ManifoldSamples)
+	{
+		ensure(NumConstraints);
+
+		unsigned char AType = A.GetType();
+		unsigned char AStaticType = ImplicitObjectType::IsScaled;
+
+		if (ensure(IsScaled(A.GetType()) && IsScaled(B.GetType())))
+		{
+			const TImplicitObjectScaled<TConvex<T, d>>& AScaled = static_cast<const TImplicitObjectScaled<TConvex<T, d>>&>(A);
+			const TImplicitObjectScaled<TConvex<T, d>>& BScaled = static_cast<const TImplicitObjectScaled<TConvex<T, d>>&>(B);
+
+			TRigidTransform<T, d> BToATM = BTM.GetRelativeTransform(ATM);
+			TRigidTransform<T, d> AToBTM = ATM.GetRelativeTransform(BTM);
+
+			if (ensure(GetInnerType(A.GetType()) == ImplicitObjectType::Convex && GetInnerType(B.GetType()) == ImplicitObjectType::Convex))
+			{
+				const TConvex<T, d>* AObject = static_cast<const TConvex<T, d>*>(AScaled.Object().Get());
+				const TConvex<T, d>* BObject = static_cast<const TConvex<T, d>*>(BScaled.Object().Get());
+
+				const TParticles<T, d>& SurfaceParticles = AObject->GetSurfaceParticles();
+
+				T OutDistance;
+				TVector<T, d> StubNormal;
+				TVector<T, d> OutNearestA;
+				TVector<T, d> OutNearestB;
+				TVector<T, d> OutNormal;
+
+				if (Manifold == nullptr)
+				{
+					if (GJKDistance<T>(*AObject, *BObject, BToATM, OutDistance, OutNearestA, OutNearestB))
+					{
+						TVector<T, d> Normal(0);
+						TVector<T, d> WorldA = ATM.TransformPosition(OutNearestA);
+						TVector<T, d> WorldB = BTM.TransformPosition(OutNearestB);
+						TVector<T, d> CenterLocation = WorldA;
+						T Phi = BObject->PhiWithNormal(BTM.InverseTransformPosition(CenterLocation), Normal);
+
+						T ManifoldSize = A.BoundingBox().Extents().Size() * ManifoldScale;
+						TVector<T, d> CrossVector = FMath::IsNearlyEqual(FMath::Abs(TVector<T, d>::DotProduct(TVector<T, d>(0, 1, 0), Normal)), T(1)) ?
+							TVector<T, d>(0, 0, 1) :
+							TVector<T, d>(0, 1, 0);
+
+						Constraints[0].LocalLocation = ATM.InverseTransformPosition(CenterLocation);
+						Constraints[0].Normal = Normal;
+
+						if (NumConstraints > 1)
+						{
+							TVector<T, d> DirVector = CrossVector;
+							T Angle = T(360) / T(NumConstraints - 1), LocalPhi;
+							for (int32 Idx = 1; Idx < NumConstraints; Idx++)
+							{
+								CrossVector = CrossVector.RotateAngleAxis(Angle, Normal);
+								TVector<T, d> WorldLocation = CenterLocation + ManifoldSize * CrossVector;
+								Constraints[Idx].LocalLocation = ATM.InverseTransformPosition(WorldLocation);
+
+								LocalPhi = AObject->PhiWithNormal(Constraints[Idx].LocalLocation, StubNormal);
+								Constraints[Idx].LocalLocation -= StubNormal * LocalPhi;
+
+								Constraints[Idx].Normal = Normal;
+							}
+						}
+					}
+				}
+
+				for (int32 Idx = 0; Idx < NumConstraints; Idx++)
+				{
+					Constraints[Idx].Location = ATM.TransformPosition(Constraints[Idx].LocalLocation);
+					Constraints[Idx].Phi = BObject->PhiWithNormal(BTM.InverseTransformPosition(Constraints[Idx].Location), StubNormal);
+
+					if (Constraints[Idx].Phi < 0.f && Manifold)
+					{
+						Manifold->SetTimestamp(INT_MAX);
+					}
+				}
+
+				return;
+			}
+		}
+
+		ensureMsgf(false, TEXT("Unsupported convex to convex constraint."));
+
+	}
+
+	template CollisionResolutionConvexConvex<float, 3>;
+
+}
diff --git a/Engine/Source/Runtime/Experimental/Chaos/Private/Chaos/PBDCollisionConstraint.cpp b/Engine/Source/Runtime/Experimental/Chaos/Private/Chaos/PBDCollisionConstraint.cpp
--- a/Engine/Source/Runtime/Experimental/Chaos/Private/Chaos/PBDCollisionConstraint.cpp
+++ b/Engine/Source/Runtime/Experimental/Chaos/Private/Chaos/PBDCollisionConstraint.cpp
@@ -19,6 +19,7 @@
 #include "Chaos/PBDRigidsSOAs.h"
 #include "Chaos/ChaosDebugDraw.h"
 #include "Chaos/ISpatialAccelerationCollection.h"
+#include "Chaos/CollisionResolutionConvexConvex.h"
 
 #if INTEL_ISPC
 #include "PBDCollisionConstraint.ispc.generated.h"
@@ -33,8 +34,8 @@ FAutoConsoleVariableRef CVarConstraintBPBVHDepth(TEXT("p.ConstraintBPBVHDepth"),
 int32 BPTreeOfGrids = 1;
 FAutoConsoleVariableRef CVarBPTreeOfGrids(TEXT("p.BPTreeOfGrids"), BPTreeOfGrids, TEXT("Whether to use a seperate tree of grids for bp"));
 
-int32 CollisionHistoryLifespanCVar = 1;
-FAutoConsoleVariableRef CVarCollisionHistoryLifespan(TEXT("p.CollisionHistoryLifespan"), CollisionHistoryLifespanCVar, TEXT("Number of iterations to cache a collision point.[def:5]"));
+int32 CollisionManifoldLifespanCVar = 1;
+FAutoConsoleVariableRef CVarCollisionManifoldLifespan(TEXT("p.CollisionManifoldLifespan"), CollisionManifoldLifespanCVar, TEXT("Number of iterations to cache a collision point.[def:5]"));
 
 int32 CollisionManifoldSamplesCVar = 5;
 FAutoConsoleVariableRef CVarCollisionManifoldSamples(TEXT("p.CollisionManifoldSamples"), CollisionManifoldSamplesCVar, TEXT("Number of sample points on the manifold.[def:6]"));
@@ -61,37 +62,10 @@ CHAOS_API int32 EnableCollisions = 1;
 FAutoConsoleVariableRef CVarEnableCollisions(TEXT("p.EnableCollisions"), EnableCollisions, TEXT("Enable/Disable collisions on the Chaos solver."));
 
 template<ECollisionUpdateType, typename T, int d>
-void UpdateConstraintImp2(const TRigidTransform<T, d>& ParticleTM, const FImplicitObject& LevelsetObject, const TRigidTransform<T, d>& LevelsetTM, const T Thickness, TRigidBodyContactConstraint<T, d>& Constraint, int32 NumConstraints, TPBDCollisionConstraintHistory<T, d>* History);
+void UpdateConstraintImp2(const TRigidTransform<T, d>& ParticleTM, const FImplicitObject& LevelsetObject, const TRigidTransform<T, d>& LevelsetTM, const T Thickness, TRigidBodyContactConstraint<T, d>& Constraint, int32 NumConstraints, TCollisionResolutionManifold<T, d>* Manifold);
 
 template<typename T, int d>
-void ConstructConstraintsImpl(TGeometryParticleHandle<T, d>* Particle0, TGeometryParticleHandle<T, d>* Particle1, const FImplicitObject* Implicit0, const FImplicitObject* Implicit1, const T Thickness, TArray<TRigidBodyContactConstraint<T, d>>& ConstraintBuffer, TPBDCollisionConstraintHistory<T, d>* History);
-
-//
-// Collision Constraint Handle
-//
-
-template<typename T, int d>
-TPBDCollisionConstraintHandle<T, d>::TPBDCollisionConstraintHandle()
-{
-}
-
-template<typename T, int d>
-TPBDCollisionConstraintHandle<T, d>::TPBDCollisionConstraintHandle(FConstraintContainer* InConstraintContainer, int32 InConstraintIndex)
-	: TContainerConstraintHandle<TPBDCollisionConstraint<T, d>>(InConstraintContainer, InConstraintIndex)
-{
-}
-
-template<typename T, int d>
-TRigidBodyContactConstraint<T, d>& TPBDCollisionConstraintHandle<T, d>::GetContact()
-{
-	return ConstraintContainer->Constraints[ConstraintIndex];
-}
-
-template<typename T, int d>
-const TRigidBodyContactConstraint<T, d>& TPBDCollisionConstraintHandle<T, d>::GetContact() const
-{
-	return ConstraintContainer->Constraints[ConstraintIndex];
-}
+void ConstructConstraintsImpl(TGeometryParticleHandle<T, d>* Particle0, TGeometryParticleHandle<T, d>* Particle1, const FImplicitObject* Implicit0, const FImplicitObject* Implicit1, const T Thickness, TArray<TRigidBodyContactConstraint<T, d>>& ConstraintBuffer, TCollisionResolutionManifold<T, d>* Manifold);
 
 
 //
@@ -111,7 +85,7 @@ TPBDCollisionConstraint<T, d>::TPBDCollisionConstraint(const TPBDRigidsSOAs<T,d>
 	, bUseCCD(false)
 	, bEnableCollisions(true)
 	, LifespanCounter(0)
-	, CollisionHistoryLifespan(CollisionHistoryLifespanCVar)
+	, CollisionManifoldLifespan(CollisionManifoldLifespanCVar)
 	, CollisionVelocityInflation(CollisionVelocityInflationCVar)
 	, PostComputeCallback(nullptr)
 	, PostApplyCallback(nullptr)
@@ -134,12 +108,12 @@ void TPBDCollisionConstraint<T, d>::Reset(/*const TPBDRigidParticles<T, d>& InPa
 		}
 	}
 
-	TArray<FConstraintHandleID> HistoryKeys;
-	History.GetKeys(HistoryKeys);
-	for (FConstraintHandleID Key : HistoryKeys)
+	TArray<FConstraintHandleID> ManifoldKeys;
+	Manifolds.GetKeys(ManifoldKeys);
+	for (FConstraintHandleID Key : ManifoldKeys)
 	{
-		FConstraintHistory* Elem = History[Key];
-		if (LifespanCounter > Elem->GetTimestamp())
+		FConstraintManifold* Elem = Manifolds[Key];
+		if (Elem->GetTimestamp()==FLT_MAX || LifespanCounter > Elem->GetTimestamp())
 		{
 			bool bInContact = false;
 			for (FConstraintContainerHandle* Handle : Elem->GetHandles())
@@ -168,16 +142,16 @@ void TPBDCollisionConstraint<T, d>::Reset(/*const TPBDRigidParticles<T, d>& InPa
 template<typename T, int d>
 void TPBDCollisionConstraint<T, d>::RemoveConstraint(int32 Idx)
 {
-	// History
+	// Manifold
 	FConstraintHandleID HandleID = GetConstraintHandleID(Constraints[Idx]);
-	if (FConstraintHistory** Ptr = History.Find(HandleID))
+	if (FConstraintManifold** Ptr = Manifolds.Find(HandleID))
 	{
-		if (FConstraintHistory* HistoryEntry = *Ptr)
+		if (FConstraintManifold* ManifoldEntry = *Ptr)
 		{
-			HistoryEntry->RemoveHandle(Handles[Idx]);
-			if (!HistoryEntry->GetHandles().Num())
+			ManifoldEntry->RemoveHandle(Handles[Idx]);
+			if (!ManifoldEntry->GetHandles().Num())
 			{
-				History.Remove(HandleID);
+				Manifolds.Remove(HandleID);
 			}
 		}
 	}
@@ -1656,7 +1630,7 @@ TArray<Pair<const FImplicitObject*, TRigidTransform<T, d>>> FindRelevantShapes2(
 
 DECLARE_CYCLE_STAT(TEXT("TPBDCollisionConstraint::UpdateUnionUnionConstraint"), STAT_UpdateUnionUnionConstraint, STATGROUP_ChaosWide);
 template<ECollisionUpdateType UpdateType, typename T, int d>
-void UpdateUnionUnionConstraint(const T Thickness, TRigidBodyContactConstraint<T, d>* Constraints, int32 NumConstraints, TPBDCollisionConstraintHistory<T, d>* History)
+void UpdateUnionUnionConstraint(const T Thickness, TRigidBodyContactConstraint<T, d>* Constraints, int32 NumConstraints, TCollisionResolutionManifold<T, d>* Manifold)
 {
 	checkSlow(Constraints != nullptr && NumConstraints >= 1);
 	SCOPE_CYCLE_COUNTER(STAT_UpdateUnionUnionConstraint);
@@ -1684,14 +1658,14 @@ void UpdateUnionUnionConstraint(const T Thickness, TRigidBodyContactConstraint<T
 		{
 			const FImplicitObject& ParticleInnerObj = *ParticlePair.First;
 			const TRigidTransform<T, d> ParticleInnerObjTM = ParticlePair.Second * ParticlesTM;
-			UpdateConstraintImp2<UpdateType>(ParticleInnerObj, ParticleInnerObjTM, LevelsetInnerObj, LevelsetInnerObjTM, Thickness, Constraints, NumConstraints, History);
+			UpdateConstraintImp2<UpdateType>(ParticleInnerObj, ParticleInnerObjTM, LevelsetInnerObj, LevelsetInnerObjTM, Thickness, Constraints, NumConstraints, Manifold);
 		}
 	}
 }
 
 DECLARE_CYCLE_STAT(TEXT("TPBDCollisionConstraint::UpdateSingleUnionConstraint"), STAT_UpdateSingleUnionConstraint, STATGROUP_ChaosWide);
 template<ECollisionUpdateType UpdateType, typename T, int d>
-void UpdateSingleUnionConstraint(const T Thickness, TRigidBodyContactConstraint<T, d>* Constraints, int32 NumConstraints, TPBDCollisionConstraintHistory<T, d>* History)
+void UpdateSingleUnionConstraint(const T Thickness, TRigidBodyContactConstraint<T, d>* Constraints, int32 NumConstraints, TCollisionResolutionManifold<T, d>* Manifold)
 {
 	checkSlow(Constraints != nullptr && NumConstraints >= 1);
 	SCOPE_CYCLE_COUNTER(STAT_UpdateSingleUnionConstraint);
@@ -1710,7 +1684,7 @@ void UpdateSingleUnionConstraint(const T Thickness, TRigidBodyContactConstraint<
 	{
 		const FImplicitObject& LevelsetInnerObj = *LevelsetObjPair.First;
 		const TRigidTransform<T, d> LevelsetInnerObjTM = LevelsetObjPair.Second * LevelsetTM;
-		UpdateConstraintImp2<UpdateType>(*ParticleObj, ParticlesTM, LevelsetInnerObj, LevelsetInnerObjTM, Thickness, Constraints, NumConstraints, History);
+		UpdateConstraintImp2<UpdateType>(*ParticleObj, ParticlesTM, LevelsetInnerObj, LevelsetInnerObjTM, Thickness, Constraints, NumConstraints, Manifold);
 	}
 }
 
@@ -1919,95 +1893,9 @@ void UpdateBoxConstraint(const TBox<T, d>& Box1, const TRigidTransform<T, d>& Bo
 	}
 }
 
-DECLARE_CYCLE_STAT(TEXT("UpdateConvexConvexConstraint"), UpdateConvexConvexConstraint, STATGROUP_ChaosWide);
-template<ECollisionUpdateType UpdateType, typename T, int d, typename TGeometryA, typename TGeometryB>
-void UpdateConvexConvexConstraint(const TGeometryA& A, TRigidTransform<T, d> ATM, const TGeometryB& B, TRigidTransform<T, d> BTM, T Thickness, TRigidBodyContactConstraint<T, d>* Constraints, int32 NumConstraints, TPBDCollisionConstraintHistory<T, d>* History)
-{
-	SCOPE_CYCLE_COUNTER(UpdateConvexConvexConstraint);
-
-	ensure(NumConstraints);
-
-	unsigned char AType = A.GetType();
-	unsigned char AStaticType = ImplicitObjectType::IsScaled;
-
-	if (ensure(IsScaled(A.GetType()) && IsScaled(B.GetType())))
-	{
-		const TImplicitObjectScaled<TConvex<T, d>>& AScaled = static_cast<const TImplicitObjectScaled<TConvex<T, d>>&>(A);
-		const TImplicitObjectScaled<TConvex<T, d>>& BScaled = static_cast<const TImplicitObjectScaled<TConvex<T, d>>&>(B);
-
-		TRigidTransform<T, d> BToATM = BTM.GetRelativeTransform(ATM);
-		TRigidTransform<T, d> AToBTM = ATM.GetRelativeTransform(BTM);
-
-		if (ensure(GetInnerType(A.GetType()) == ImplicitObjectType::Convex && GetInnerType(B.GetType()) == ImplicitObjectType::Convex))
-		{
-			const TConvex<T, d>* AObject = static_cast<const TConvex<T, d>*>(AScaled.Object().Get());
-			const TConvex<T, d>* BObject = static_cast<const TConvex<T, d>*>(BScaled.Object().Get());
-
-			const TParticles<T, d>& SurfaceParticles = AObject->GetSurfaceParticles();
-
-			T OutDistance;
-			TVector<T, d> StubNormal;
-			TVector<T, d> OutNearestA;
-			TVector<T, d> OutNearestB;
-			TVector<T, d> OutNormal;
-
-			if (History == nullptr)
-			{
-				if (GJKDistance<T>(*AObject, *BObject, BToATM, OutDistance, OutNearestA, OutNearestB))
-				{
-					TVector<T, d> Normal(0);
-					TVector<T, d> WorldA = ATM.TransformPosition(OutNearestA);
-					TVector<T, d> WorldB = BTM.TransformPosition(OutNearestB);
-					TVector<T, d> CenterLocation = WorldA;
-					T Phi = BObject->PhiWithNormal(BTM.InverseTransformPosition(CenterLocation), Normal);
-
-					T ManifoldSize = A.BoundingBox().Extents().Size() * CollisionManifoldSampleScaleCVar;
-					TVector<T, d> CrossVector = FMath::IsNearlyEqual(FMath::Abs(TVector<T, d>::DotProduct(TVector<T, d>(0, 1, 0), Normal)), T(1)) ?
-						TVector<T, d>(0, 0, 1) :
-						TVector<T, d>(0, 1, 0);
-
-					Constraints[0].LocalLocation = ATM.InverseTransformPosition(CenterLocation);
-					Constraints[0].Normal = Normal;
-
-					if (NumConstraints > 1)
-					{
-						TVector<T, d> DirVector = CrossVector;
-						T Angle = T(360) / T(NumConstraints - 1), LocalPhi;
-						for (int32 Idx = 1; Idx < NumConstraints; Idx++)
-						{
-							CrossVector = CrossVector.RotateAngleAxis(Angle,Normal);
-							TVector<T, d> WorldLocation = CenterLocation + ManifoldSize * CrossVector;
-							Constraints[Idx].LocalLocation = ATM.InverseTransformPosition(WorldLocation);
-
-							LocalPhi = AObject->PhiWithNormal(Constraints[Idx].LocalLocation, StubNormal);
-							Constraints[Idx].LocalLocation -= StubNormal * LocalPhi;
-
-							Constraints[Idx].Normal = Normal;
-						}
-					}
-				}
-			}
-
-			for (int32 Idx = 0; Idx < NumConstraints; Idx++)
-			{
-				Constraints[Idx].Location = ATM.TransformPosition(Constraints[Idx].LocalLocation);
-				Constraints[Idx].Phi = BObject->PhiWithNormal(BTM.InverseTransformPosition(Constraints[Idx].Location), StubNormal);
-
-				if (Constraints[Idx].Phi < 0.f && History)
-				{
-					History->SetTimestamp(INT_MAX);
-				}
-			}
-
-			return;
-		}
-	}
-
-	ensureMsgf(false, TEXT("Unsupported convex to convex constraint."));
-}
 
 template<typename T, int d>
-void ConstructLevelsetConstraints(TGeometryParticleHandle<T, d>* Particle0, TGeometryParticleHandle<T, d>* Particle1, const FImplicitObject* Implicit0, const FImplicitObject* Implicit1, const T Thickness, TArray<TRigidBodyContactConstraint<T, d>>& ConstraintBuffer, TPBDCollisionConstraintHistory<T, d>* History = nullptr)
+void ConstructLevelsetConstraints(TGeometryParticleHandle<T, d>* Particle0, TGeometryParticleHandle<T, d>* Particle1, const FImplicitObject* Implicit0, const FImplicitObject* Implicit1, const T Thickness, TArray<TRigidBodyContactConstraint<T, d>>& ConstraintBuffer, TCollisionResolutionManifold<T, d>* Manifold = nullptr)
 {
 	//todo(ocohen):if both have collision particles, use the one with fewer?
 	if (!Particle1->Geometry() || (Particle0->AsDynamic() && !Particle0->AsDynamic()->CollisionParticlesSize() && Particle0->Geometry() && !Particle0->Geometry()->IsUnderlyingUnion()))
@@ -2031,29 +1919,7 @@ void ConstructLevelsetConstraints(TGeometryParticleHandle<T, d>* Particle0, TGeo
 }
 
 template<typename T, int d>
-void ConstructConvexConvexConstraints(TGeometryParticleHandle<T, d>* Particle0, TGeometryParticleHandle<T, d>* Particle1, const FImplicitObject* Implicit0, const FImplicitObject* Implicit1, const T Thickness, TArray<TRigidBodyContactConstraint<T, d>>& ConstraintBuffer, TPBDCollisionConstraintHistory<T, d>* History = nullptr)
-{
-	if (History)
-	{
-		if (History->ContainsShapeConnection(Implicit0, Implicit1))
-		{
-			return;
-		}
-	}
-
-	for (int i = 0; i < CollisionManifoldSamplesCVar; i++)
-	{
-		TRigidBodyContactConstraint<T, d> Constraint;
-		Constraint.Particle = Particle0;
-		Constraint.Levelset = Particle1;
-		Constraint.Geometry[0] = Implicit0;
-		Constraint.Geometry[1] = Implicit1;
-		ConstraintBuffer.Add(Constraint);
-	}
-}
-
-template<typename T, int d>
-void ConstructBoxConstraints(TGeometryParticleHandle<T, d>* Particle0, TGeometryParticleHandle<T, d>* Particle1, const FImplicitObject* Implicit0, const FImplicitObject* Implicit1, const T Thickness, TArray<TRigidBodyContactConstraint<T, d>>& ConstraintBuffer, TPBDCollisionConstraintHistory<T, d>* History = nullptr)
+void ConstructBoxConstraints(TGeometryParticleHandle<T, d>* Particle0, TGeometryParticleHandle<T, d>* Particle1, const FImplicitObject* Implicit0, const FImplicitObject* Implicit1, const T Thickness, TArray<TRigidBodyContactConstraint<T, d>>& ConstraintBuffer, TCollisionResolutionManifold<T, d>* Manifold = nullptr)
 {
 	TRigidBodyContactConstraint<T, d> Constraint;
 	Constraint.Particle = Particle0;
@@ -2064,7 +1930,7 @@ void ConstructBoxConstraints(TGeometryParticleHandle<T, d>* Particle0, TGeometry
 }
 
 template<typename T, int d>
-void ConstructBoxPlaneConstraints(TGeometryParticleHandle<T, d>* Particle0, TGeometryParticleHandle<T, d>* Particle1, const FImplicitObject* Implicit0, const FImplicitObject* Implicit1, const T Thickness, TArray<TRigidBodyContactConstraint<T, d>>& ConstraintBuffer, TPBDCollisionConstraintHistory<T, d>* History = nullptr)
+void ConstructBoxPlaneConstraints(TGeometryParticleHandle<T, d>* Particle0, TGeometryParticleHandle<T, d>* Particle1, const FImplicitObject* Implicit0, const FImplicitObject* Implicit1, const T Thickness, TArray<TRigidBodyContactConstraint<T, d>>& ConstraintBuffer, TCollisionResolutionManifold<T, d>* Manifold = nullptr)
 {
 	TRigidBodyContactConstraint<T, d> Constraint;
 	Constraint.Particle = Particle0;
@@ -2075,7 +1941,7 @@ void ConstructBoxPlaneConstraints(TGeometryParticleHandle<T, d>* Particle0, TGeo
 }
 
 template<typename T, int d>
-void ConstructSphereConstraints(TGeometryParticleHandle<T, d>* Particle0, TGeometryParticleHandle<T, d>* Particle1, const FImplicitObject* Implicit0, const FImplicitObject* Implicit1, const T Thickness, TArray<TRigidBodyContactConstraint<T, d>>& ConstraintBuffer, TPBDCollisionConstraintHistory<T, d>* History = nullptr)
+void ConstructSphereConstraints(TGeometryParticleHandle<T, d>* Particle0, TGeometryParticleHandle<T, d>* Particle1, const FImplicitObject* Implicit0, const FImplicitObject* Implicit1, const T Thickness, TArray<TRigidBodyContactConstraint<T, d>>& ConstraintBuffer, TCollisionResolutionManifold<T, d>* Manifold = nullptr)
 {
 	TRigidBodyContactConstraint<T, d> Constraint;
 	Constraint.Particle = Particle0;
@@ -2086,7 +1952,7 @@ void ConstructSphereConstraints(TGeometryParticleHandle<T, d>* Particle0, TGeome
 }
 
 template<typename T, int d>
-void ConstructSpherePlaneConstraints(TGeometryParticleHandle<T, d>* Particle0, TGeometryParticleHandle<T, d>* Particle1, const FImplicitObject* Implicit0, const FImplicitObject* Implicit1, const T Thickness, TArray<TRigidBodyContactConstraint<T, d>>& ConstraintBuffer, TPBDCollisionConstraintHistory<T, d>* History = nullptr)
+void ConstructSpherePlaneConstraints(TGeometryParticleHandle<T, d>* Particle0, TGeometryParticleHandle<T, d>* Particle1, const FImplicitObject* Implicit0, const FImplicitObject* Implicit1, const T Thickness, TArray<TRigidBodyContactConstraint<T, d>>& ConstraintBuffer, TCollisionResolutionManifold<T, d>* Manifold = nullptr)
 {
 	TRigidBodyContactConstraint<T, d> Constraint;
 	Constraint.Particle = Particle0;
@@ -2097,7 +1963,7 @@ void ConstructSpherePlaneConstraints(TGeometryParticleHandle<T, d>* Particle0, T
 }
 
 template<typename T, int d>
-void ConstructSphereBoxConstraints(TGeometryParticleHandle<T, d>* Particle0, TGeometryParticleHandle<T, d>* Particle1, const FImplicitObject* Implicit0, const FImplicitObject* Implicit1, const T Thickness, TArray<TRigidBodyContactConstraint<T, d>>& ConstraintBuffer, TPBDCollisionConstraintHistory<T, d>* History = nullptr)
+void ConstructSphereBoxConstraints(TGeometryParticleHandle<T, d>* Particle0, TGeometryParticleHandle<T, d>* Particle1, const FImplicitObject* Implicit0, const FImplicitObject* Implicit1, const T Thickness, TArray<TRigidBodyContactConstraint<T, d>>& ConstraintBuffer, TCollisionResolutionManifold<T, d>* Manifold = nullptr)
 {
 	TRigidBodyContactConstraint<T, d> Constraint;
 	Constraint.Particle = Particle0;
@@ -2108,7 +1974,7 @@ void ConstructSphereBoxConstraints(TGeometryParticleHandle<T, d>* Particle0, TGe
 }
 
 template<typename T, int d>
-void ConstructCapsuleCapsuleConstraints(TGeometryParticleHandle<T, d>* Particle0, TGeometryParticleHandle<T, d>* Particle1, const FImplicitObject* Implicit0, const FImplicitObject* Implicit1, const T Thickness, TArray<TRigidBodyContactConstraint<T, d>>& ConstraintBuffer, TPBDCollisionConstraintHistory<T, d>* History = nullptr)
+void ConstructCapsuleCapsuleConstraints(TGeometryParticleHandle<T, d>* Particle0, TGeometryParticleHandle<T, d>* Particle1, const FImplicitObject* Implicit0, const FImplicitObject* Implicit1, const T Thickness, TArray<TRigidBodyContactConstraint<T, d>>& ConstraintBuffer, TCollisionResolutionManifold<T, d>* Manifold = nullptr)
 {
 	TRigidBodyContactConstraint<T, d> Constraint;
 	Constraint.Particle = Particle0;
@@ -2119,7 +1985,7 @@ void ConstructCapsuleCapsuleConstraints(TGeometryParticleHandle<T, d>* Particle0
 }
 
 template<typename T, int d>
-void ConstructCapsuleBoxConstraints(TGeometryParticleHandle<T, d>* Particle0, TGeometryParticleHandle<T, d>* Particle1, const FImplicitObject* Implicit0, const FImplicitObject* Implicit1, const T Thickness, TArray<TRigidBodyContactConstraint<T, d>>& ConstraintBuffer, TPBDCollisionConstraintHistory<T, d>* History = nullptr)
+void ConstructCapsuleBoxConstraints(TGeometryParticleHandle<T, d>* Particle0, TGeometryParticleHandle<T, d>* Particle1, const FImplicitObject* Implicit0, const FImplicitObject* Implicit1, const T Thickness, TArray<TRigidBodyContactConstraint<T, d>>& ConstraintBuffer, TCollisionResolutionManifold<T, d>* Manifold = nullptr)
 {
 	TRigidBodyContactConstraint<T, d> Constraint;
 	Constraint.Particle = Particle0;
@@ -2130,7 +1996,7 @@ void ConstructCapsuleBoxConstraints(TGeometryParticleHandle<T, d>* Particle0, TG
 }
 
 template<typename T, int d>
-void ConstructSingleUnionConstraints(TGeometryParticleHandle<T, d>* Particle0, TGeometryParticleHandle<T, d>* Particle1, const FImplicitObject* Implicit0, const FImplicitObject* Implicit1, const T Thickness, TArray<TRigidBodyContactConstraint<T, d>>& ConstraintBuffer, TPBDCollisionConstraintHistory<T, d>* History = nullptr)
+void ConstructSingleUnionConstraints(TGeometryParticleHandle<T, d>* Particle0, TGeometryParticleHandle<T, d>* Particle1, const FImplicitObject* Implicit0, const FImplicitObject* Implicit1, const T Thickness, TArray<TRigidBodyContactConstraint<T, d>>& ConstraintBuffer, TCollisionResolutionManifold<T, d>* Manifold = nullptr)
 {
 	TRigidBodyContactConstraint<T, d> Constraint;
 	Constraint.Particle = Particle0;
@@ -2141,7 +2007,7 @@ void ConstructSingleUnionConstraints(TGeometryParticleHandle<T, d>* Particle0, T
 }
 
 template<typename T, int d>
-void ConstructUnionUnionConstraints(TGeometryParticleHandle<T, d>* Particle0, TGeometryParticleHandle<T, d>* Particle1, const FImplicitObject* Implicit0, const FImplicitObject* Implicit1, const T Thickness, TArray<TRigidBodyContactConstraint<T, d>>& ConstraintBuffer, TPBDCollisionConstraintHistory<T, d>* History = nullptr)
+void ConstructUnionUnionConstraints(TGeometryParticleHandle<T, d>* Particle0, TGeometryParticleHandle<T, d>* Particle1, const FImplicitObject* Implicit0, const FImplicitObject* Implicit1, const T Thickness, TArray<TRigidBodyContactConstraint<T, d>>& ConstraintBuffer, TCollisionResolutionManifold<T, d>* Manifold = nullptr)
 {
 	SCOPE_CYCLE_COUNTER(STAT_UpdateUnionUnionConstraint);
 
@@ -2165,13 +2031,13 @@ void ConstructUnionUnionConstraints(TGeometryParticleHandle<T, d>* Particle0, TG
 		for (const Pair<const FImplicitObject*, TRigidTransform<T, d>>& ParticlePair : ParticleShapes)
 		{
 			const FImplicitObject* ParticleInnerObj = ParticlePair.First;
-			ConstructConstraintsImpl<T, d>(Particle0, Particle1, ParticleInnerObj, LevelsetInnerObj, Thickness, ConstraintBuffer, History);
+			ConstructConstraintsImpl<T, d>(Particle0, Particle1, ParticleInnerObj, LevelsetInnerObj, Thickness, ConstraintBuffer, Manifold);
 		}
 	}
 }
 
 template<typename T, int d>
-void ConstructConstraintsImpl(TGeometryParticleHandle<T, d>* Particle0, TGeometryParticleHandle<T, d>* Particle1, const FImplicitObject* Implicit0, const FImplicitObject* Implicit1, const T Thickness, TArray< TRigidBodyContactConstraint<T, d> >& ConstraintBuffer, TPBDCollisionConstraintHistory<T, d>* History)
+void ConstructConstraintsImpl(TGeometryParticleHandle<T, d>* Particle0, TGeometryParticleHandle<T, d>* Particle1, const FImplicitObject* Implicit0, const FImplicitObject* Implicit1, const T Thickness, TArray< TRigidBodyContactConstraint<T, d> >& ConstraintBuffer, TCollisionResolutionManifold<T, d>* Manifold)
 {
 	// TriangleMesh implicit's are for scene query only.
 	if (Implicit0 && GetInnerType(Implicit0->GetType()) == ImplicitObjectType::TriangleMesh) return; //as
@@ -2179,71 +2045,71 @@ void ConstructConstraintsImpl(TGeometryParticleHandle<T, d>* Particle0, TGeometr
 
 	if (!Implicit0 || !Implicit1)
 	{
-		ConstructLevelsetConstraints(Particle0, Particle1, Implicit0, Implicit1, Thickness, ConstraintBuffer, History);
+		ConstructLevelsetConstraints(Particle0, Particle1, Implicit0, Implicit1, Thickness, ConstraintBuffer, Manifold);
 	}
 	else if (Implicit0->GetType() == TBox<T, d>::StaticType() && Implicit1->GetType() == TBox<T, d>::StaticType()) 
 	{
-		ConstructBoxConstraints(Particle0, Particle1, Implicit0, Implicit1, Thickness, ConstraintBuffer, History);
+		ConstructBoxConstraints(Particle0, Particle1, Implicit0, Implicit1, Thickness, ConstraintBuffer, Manifold);
 	}
 	else if (Implicit0->GetType() == TSphere<T, d>::StaticType() && Implicit1->GetType() == TSphere<T, d>::StaticType())
 	{
-		ConstructSphereConstraints(Particle0, Particle1, Implicit0, Implicit1, Thickness, ConstraintBuffer, History);
+		ConstructSphereConstraints(Particle0, Particle1, Implicit0, Implicit1, Thickness, ConstraintBuffer, Manifold);
 	}
 	else if (Implicit0->GetType() == TBox<T, d>::StaticType() && Implicit1->GetType() == TPlane<T, d>::StaticType())
 	{
-		ConstructBoxPlaneConstraints(Particle0, Particle1, Implicit0, Implicit1, Thickness, ConstraintBuffer, History);
+		ConstructBoxPlaneConstraints(Particle0, Particle1, Implicit0, Implicit1, Thickness, ConstraintBuffer, Manifold);
 	}
 	else if (Implicit1->GetType() == TBox<T, d>::StaticType() && Implicit0->GetType() == TPlane<T, d>::StaticType())
 	{
-		ConstructBoxPlaneConstraints(Particle1, Particle0, Implicit0, Implicit1, Thickness, ConstraintBuffer, History);
+		ConstructBoxPlaneConstraints(Particle1, Particle0, Implicit0, Implicit1, Thickness, ConstraintBuffer, Manifold);
 	}
 	else if (Implicit0->GetType() == TSphere<T, d>::StaticType() && Implicit1->GetType() == TPlane<T, d>::StaticType())
 	{
-		ConstructSpherePlaneConstraints(Particle0, Particle1, Implicit0, Implicit1, Thickness, ConstraintBuffer, History);
+		ConstructSpherePlaneConstraints(Particle0, Particle1, Implicit0, Implicit1, Thickness, ConstraintBuffer, Manifold);
 	}
 	else if (Implicit1->GetType() == TSphere<T, d>::StaticType() && Implicit0->GetType() == TPlane<T, d>::StaticType())
 	{
-		ConstructSpherePlaneConstraints(Particle1, Particle0, Implicit0, Implicit1, Thickness, ConstraintBuffer, History);
+		ConstructSpherePlaneConstraints(Particle1, Particle0, Implicit0, Implicit1, Thickness, ConstraintBuffer, Manifold);
 	}
 	else if (Implicit0->GetType() == TSphere<T, d>::StaticType() && Implicit1->GetType() == TBox<T, d>::StaticType())
 	{
-		ConstructSphereBoxConstraints(Particle0, Particle1, Implicit0, Implicit1, Thickness, ConstraintBuffer, History);
+		ConstructSphereBoxConstraints(Particle0, Particle1, Implicit0, Implicit1, Thickness, ConstraintBuffer, Manifold);
 	}
 	else if (Implicit1->GetType() == TSphere<T, d>::StaticType() && Implicit0->GetType() == TBox<T, d>::StaticType())
 	{
-		ConstructSphereBoxConstraints(Particle1, Particle0, Implicit0, Implicit1, Thickness, ConstraintBuffer, History);
+		ConstructSphereBoxConstraints(Particle1, Particle0, Implicit0, Implicit1, Thickness, ConstraintBuffer, Manifold);
 	}
 	else if (Implicit0->GetType() == TCapsule<T>::StaticType() && Implicit1->GetType() == TCapsule<T>::StaticType())
 	{
-		ConstructCapsuleCapsuleConstraints(Particle0, Particle1, Implicit0, Implicit1, Thickness, ConstraintBuffer, History);
+		ConstructCapsuleCapsuleConstraints(Particle0, Particle1, Implicit0, Implicit1, Thickness, ConstraintBuffer, Manifold);
 	}
 	else if (Implicit0->GetType() == TCapsule<T>::StaticType() && Implicit1->GetType() == TBox<T, d>::StaticType())
 	{
-		ConstructCapsuleBoxConstraints(Particle0, Particle1, Implicit0, Implicit1, Thickness, ConstraintBuffer, History);
+		ConstructCapsuleBoxConstraints(Particle0, Particle1, Implicit0, Implicit1, Thickness, ConstraintBuffer, Manifold);
 	}
 	else if (Implicit1->GetType() == TCapsule<T>::StaticType() && Implicit0->GetType() == TBox<T, d>::StaticType())
 	{
-		ConstructCapsuleBoxConstraints(Particle1, Particle0, Implicit0, Implicit1, Thickness, ConstraintBuffer, History);
+		ConstructCapsuleBoxConstraints(Particle1, Particle0, Implicit0, Implicit1, Thickness, ConstraintBuffer, Manifold);
 	}
 	else if (Implicit0->GetType() < TImplicitObjectUnion<T, d>::StaticType() && Implicit1->GetType() == TImplicitObjectUnion<T, d>::StaticType())
 	{
-		ConstructSingleUnionConstraints(Particle0, Particle1, Implicit0, Implicit1, Thickness, ConstraintBuffer, History);
+		ConstructSingleUnionConstraints(Particle0, Particle1, Implicit0, Implicit1, Thickness, ConstraintBuffer, Manifold);
 	}
 	else if (Implicit0->GetType() == TImplicitObjectUnion<T, d>::StaticType() && Implicit1->GetType() < TImplicitObjectUnion<T, d>::StaticType())
 	{
-		ConstructSingleUnionConstraints(Particle1, Particle0, Implicit0, Implicit1, Thickness, ConstraintBuffer, History);
+		ConstructSingleUnionConstraints(Particle1, Particle0, Implicit0, Implicit1, Thickness, ConstraintBuffer, Manifold);
 	}
 	else if (Implicit0->GetType() == TImplicitObjectUnion<T, d>::StaticType() && Implicit1->GetType() == TImplicitObjectUnion<T, d>::StaticType())
 	{
-		ConstructUnionUnionConstraints(Particle0, Particle1, Implicit0, Implicit1, Thickness, ConstraintBuffer, History);
+		ConstructUnionUnionConstraints(Particle0, Particle1, Implicit0, Implicit1, Thickness, ConstraintBuffer, Manifold);
 	}
 	else if (Implicit0->IsConvex() && Implicit1->IsConvex())
 	{
-		ConstructConvexConvexConstraints(Particle0, Particle1, Implicit0, Implicit1, Thickness, ConstraintBuffer, History);
+		CollisionResolutionConvexConvex<T,d>::ConstructConvexConvexConstraints(Particle0, Particle1, Implicit0, Implicit1, Thickness, ConstraintBuffer, Manifold, CollisionManifoldSampleScaleCVar, CollisionManifoldSamplesCVar);
 	}
 	else
 	{
-		ConstructLevelsetConstraints(Particle0, Particle1, Implicit0, Implicit1, Thickness, ConstraintBuffer, History);
+		ConstructLevelsetConstraints(Particle0, Particle1, Implicit0, Implicit1, Thickness, ConstraintBuffer, Manifold);
 	}
 }
 
@@ -2252,11 +2118,11 @@ void TPBDCollisionConstraint<T, d>::ConstructConstraints(TGeometryParticleHandle
 {
 	ensure(Particle0 && Particle1);
 
-	FConstraintHistory* HistoryElement = nullptr;
-	if (FConstraintHistory** Ptr = History.Find(FConstraintHandleID(Particle0, Particle1)))
-		HistoryElement = *Ptr;
+	FConstraintManifold* ManifoldElement = nullptr;
+	if (FConstraintManifold** Ptr = Manifolds.Find(FConstraintHandleID(Particle0, Particle1)))
+		ManifoldElement = *Ptr;
 
-	ConstructConstraintsImpl<T, d>(Particle0, Particle1, Particle0->Geometry().Get(), Particle1->Geometry().Get(), Thickness, ConstraintBuffer, HistoryElement);
+	ConstructConstraintsImpl<T, d>(Particle0, Particle1, Particle0->Geometry().Get(), Particle1->Geometry().Get(), Thickness, ConstraintBuffer, ManifoldElement);
 }
 
 // NOTE: UpdateLevelsetConstraintImp and its <float,3> specialization are here for the Linux build. It looks like
@@ -2275,7 +2141,7 @@ void UpdateLevelsetConstraintImp(const float Thickness, TRigidBodyContactConstra
 
 template<ECollisionUpdateType UpdateType, typename T, int d>
 void UpdateConstraintImp2(const FImplicitObject& ParticleObject, const TRigidTransform<T, d>& ParticleTM, const FImplicitObject& LevelsetObject, const TRigidTransform<T, d>& LevelsetTM, const T Thickness,
-	TRigidBodyContactConstraint<T, d>* Constraints, int32 NumConstraints, TPBDCollisionConstraintHistory<T, d>* History)
+	TRigidBodyContactConstraint<T, d>* Constraints, int32 NumConstraints, TCollisionResolutionManifold<T, d>* Manifold)
 {
 	TRigidBodyContactConstraint<T, d>& Constraint = Constraints[0];
 
@@ -2349,7 +2215,7 @@ void UpdateConstraintImp2(const FImplicitObject& ParticleObject, const TRigidTra
 	}
 	else if (ParticleObject.GetType() < TImplicitObjectUnion<T, d>::StaticType() && LevelsetObject.GetType() == TImplicitObjectUnion<T, d>::StaticType())
 	{
-		return UpdateSingleUnionConstraint<UpdateType>(Thickness, Constraints, NumConstraints, History);
+		return UpdateSingleUnionConstraint<UpdateType>(Thickness, Constraints, NumConstraints, Manifold);
 	}
 	else if (ParticleObject.GetType() == TImplicitObjectUnion<T, d>::StaticType() && LevelsetObject.GetType() < TImplicitObjectUnion<T, d>::StaticType())
 	{
@@ -2357,11 +2223,11 @@ void UpdateConstraintImp2(const FImplicitObject& ParticleObject, const TRigidTra
 	}
 	else if (ParticleObject.GetType() == TImplicitObjectUnion<T, d>::StaticType() && LevelsetObject.GetType() == TImplicitObjectUnion<T, d>::StaticType())
 	{
-		return UpdateUnionUnionConstraint<UpdateType>(Thickness, Constraints, NumConstraints, History);
+		return UpdateUnionUnionConstraint<UpdateType>(Thickness, Constraints, NumConstraints, Manifold);
 	}
 	else if (ParticleObject.IsConvex() && LevelsetObject.IsConvex())
 	{
-		UpdateConvexConvexConstraint<UpdateType>(ParticleObject, ParticleTM, LevelsetObject, LevelsetTM, Thickness, Constraints, NumConstraints, History);
+		CollisionResolutionConvexConvex<T,d>::UpdateConvexConvexConstraint(ParticleObject, ParticleTM, LevelsetObject, LevelsetTM, Thickness, Constraints, NumConstraints, Manifold, CollisionManifoldSampleScaleCVar, CollisionManifoldSamplesCVar);
 	}
 	else if (LevelsetObject.IsUnderlyingUnion())
 	{
@@ -2415,13 +2281,13 @@ void TPBDCollisionConstraint<T, d>::UpdateConstraint(const T Thickness, FRigidBo
 	}
 	else
 	{
-		TPBDCollisionConstraintHistory<T, d>* HistoryData = nullptr;
+		TCollisionResolutionManifold<T, d>* ManifoldData = nullptr;
 		FConstraintHandleID HandleID = GetConstraintHandleID(Constraint);
-		if (History.Contains(HandleID))
+		if (Manifolds.Contains(HandleID))
 		{
-			HistoryData = History[HandleID];
+			ManifoldData = Manifolds[HandleID];
 		}
-		UpdateConstraintImp2<UpdateType>(*Constraint.Particle->Geometry(), ParticleTM, *Constraint.Levelset->Geometry(), LevelsetTM, Thickness, InConstraints, NumConstraints, HistoryData);
+		UpdateConstraintImp2<UpdateType>(*Constraint.Particle->Geometry(), ParticleTM, *Constraint.Levelset->Geometry(), LevelsetTM, Thickness, InConstraints, NumConstraints, ManifoldData);
 	}
 }
 
@@ -2434,6 +2300,6 @@ template void TPBDCollisionConstraint<float, 3>::UpdateLevelsetConstraint<EColli
 template void TPBDCollisionConstraint<float, 3>::UpdateLevelsetConstraint<ECollisionUpdateType::Deepest>(const float Thickness, FRigidBodyContactConstraint& Constraint);
 template void TPBDCollisionConstraint<float, 3>::ComputeConstraints<false>(const TPBDCollisionConstraint<float,3>::FAccelerationStructure&, float Dt);
 template void TPBDCollisionConstraint<float, 3>::ComputeConstraints<true>(const TPBDCollisionConstraint<float, 3>::FAccelerationStructure&, float Dt);
-template void UpdateConstraintImp2<ECollisionUpdateType::Any, float, 3>(const FImplicitObject& ParticleObject, const TRigidTransform<float, 3>& ParticleTM, const FImplicitObject& LevelsetObject, const TRigidTransform<float, 3>& LevelsetTM, const float Thickness, TRigidBodyContactConstraint<float, 3>* Constraints, int32 NumConstraints, TPBDCollisionConstraintHistory<float, 3>* History);
-template void UpdateConstraintImp2<ECollisionUpdateType::Deepest, float, 3>(const FImplicitObject& ParticleObject, const TRigidTransform<float, 3>& ParticleTM, const FImplicitObject& LevelsetObject, const TRigidTransform<float, 3>& LevelsetTM, const float Thickness, TRigidBodyContactConstraint<float, 3>* Constraint, int32 NumConstraints, TPBDCollisionConstraintHistory<float, 3>* History);
+template void UpdateConstraintImp2<ECollisionUpdateType::Any, float, 3>(const FImplicitObject& ParticleObject, const TRigidTransform<float, 3>& ParticleTM, const FImplicitObject& LevelsetObject, const TRigidTransform<float, 3>& LevelsetTM, const float Thickness, TRigidBodyContactConstraint<float, 3>* Constraints, int32 NumConstraints, TCollisionResolutionManifold<float, 3>* Manifold);
+template void UpdateConstraintImp2<ECollisionUpdateType::Deepest, float, 3>(const FImplicitObject& ParticleObject, const TRigidTransform<float, 3>& ParticleTM, const FImplicitObject& LevelsetObject, const TRigidTransform<float, 3>& LevelsetTM, const float Thickness, TRigidBodyContactConstraint<float, 3>* Constraint, int32 NumConstraints, TCollisionResolutionManifold<float, 3>* Manifold);
 }
diff --git a/Engine/Source/Runtime/Experimental/Chaos/Public/Chaos/CollisionResolutionConvexConvex.h b/Engine/Source/Runtime/Experimental/Chaos/Public/Chaos/CollisionResolutionConvexConvex.h
new file mode 100644
--- /dev/null
+++ b/Engine/Source/Runtime/Experimental/Chaos/Public/Chaos/CollisionResolutionConvexConvex.h
@@ -0,0 +1,40 @@
+// Copyright 1998-2019 Epic Games, Inc. All Rights Reserved.
+#pragma once
+
+#include "Chaos/CollisionResolutionManifold.h"
+#include "Chaos/Convex.h"
+#include "Chaos/ParticleHandle.h"
+#include "Chaos/PBDCollisionTypes.h"
+#include "Chaos/PBDCollisionConstraint.h"
+#include "Chaos/Transform.h"
+
+
+namespace Chaos
+{
+
+template<class T, int d>
+class CHAOS_API CollisionResolutionConvexConvex
+{
+public:
+	typedef TConvex<T,d> FConvex;
+	typedef TRigidTransform<T,d> FRigidTransform;
+	typedef TRigidBodyContactConstraint<T,d> FRigidBodyContactConstraint;
+	typedef TCollisionResolutionManifold<T,d> FCollisionResolutionManifold;
+	typedef TGeometryParticleHandle<T,d> FGeometryParticleHandle;
+
+	/*
+	* Build a constraint manifold for the contact shape pair. 
+	*/
+	static void ConstructConvexConvexConstraints(FGeometryParticleHandle* Particle0, FGeometryParticleHandle* Particle1, const FImplicitObject* Implicit0, const FImplicitObject* Implicit1, const float Thickness,
+		TArray<FRigidBodyContactConstraint>& ConstraintBuffer, FCollisionResolutionManifold* Manifold = nullptr, float ManifoldScale = 0.5, int32 ManifoldSamples = 4);
+
+	/*
+	* Update the contact manifold.
+	*/
+	static void UpdateConvexConvexConstraint(const FImplicitObject& A, const FRigidTransform& ATM, const FImplicitObject& B, FRigidTransform BTM, float Thickness, 
+		FRigidBodyContactConstraint* Constraints, int32 NumConstraints, FCollisionResolutionManifold* Manifold, float ManifoldScale = 0.5, int32 ManifoldSamples = 4);
+
+
+};
+
+}
diff --git a/Engine/Source/Runtime/Experimental/Chaos/Public/Chaos/CollisionResolutionManifold.h b/Engine/Source/Runtime/Experimental/Chaos/Public/Chaos/CollisionResolutionManifold.h
new file mode 100644
--- /dev/null
+++ b/Engine/Source/Runtime/Experimental/Chaos/Public/Chaos/CollisionResolutionManifold.h
@@ -0,0 +1,64 @@
+// Copyright 1998-2019 Epic Games, Inc. All Rights Reserved.
+#pragma once
+
+#include "Chaos/ConstraintHandle.h"
+#include "Chaos/PBDCollisionTypes.h"
+#include "Chaos/PBDConstraintContainer.h"
+#include "Framework/BufferedData.h"
+
+#include <memory>
+#include <queue>
+#include <sstream>
+#include "BoundingVolume.h"
+#include "AABBTree.h"
+
+namespace Chaos
+{
+	template<typename T, int d>
+	class TPBDCollisionConstraintHandle;
+
+
+template<class T, int d>
+class CHAOS_API TCollisionResolutionManifold
+{
+	typedef TPair<const FImplicitObject*, const FImplicitObject*> ImplicitPairs;
+
+public:
+	using FConstraintContainerHandle = TPBDCollisionConstraintHandle<T, d>;
+	TCollisionResolutionManifold(const TVector<T, d>& InLocation, const TRotation<T, d>& InRotation, int32 InTimestamp = -INT_MAX)
+		: Timestamp(InTimestamp)
+		, Location(InLocation)
+		, Rotation(InRotation)
+	{}
+
+	const TVector<T, d>& GetLocation() { return Location; }
+	const TRotation<T, d>& GetRotation() { return Rotation; }
+
+	void AddHandle(FConstraintContainerHandle* InHandle)
+	{
+		Implicits.Add(ImplicitPairs(InHandle->GetContact().Geometry[0], InHandle->GetContact().Geometry[1]));
+		Handles.Push(InHandle);
+	}
+	void RemoveHandle(FConstraintContainerHandle* InHandle)
+	{
+		Handles.RemoveSingleSwap(InHandle);
+		Implicits.Remove(ImplicitPairs(InHandle->GetContact().Geometry[0], InHandle->GetContact().Geometry[1]));
+	}
+
+	TArray<FConstraintContainerHandle*>& GetHandles() { return Handles; }
+	const TArray<FConstraintContainerHandle*>& GetHandles() const { return Handles; }
+
+	bool ContainsShapeConnection(const FImplicitObject* Implicit0In, const FImplicitObject* Implicit1In) { return Implicits.Contains(ImplicitPairs(Implicit0In, Implicit1In)); }
+
+	int32 GetTimestamp() const { return Timestamp; }
+	void SetTimestamp(int32 InTimestamp) { Timestamp = InTimestamp; }
+
+private:
+	int32 Timestamp;
+	TVector<T, d> Location;
+	TRotation<T, d> Rotation;
+	TArray<FConstraintContainerHandle*> Handles;
+	TSet<ImplicitPairs> Implicits;
+};
+
+}
diff --git a/Engine/Source/Runtime/Experimental/Chaos/Public/Chaos/PBDCollisionConstraint.h b/Engine/Source/Runtime/Experimental/Chaos/Public/Chaos/PBDCollisionConstraint.h
--- a/Engine/Source/Runtime/Experimental/Chaos/Public/Chaos/PBDCollisionConstraint.h
+++ b/Engine/Source/Runtime/Experimental/Chaos/Public/Chaos/PBDCollisionConstraint.h
@@ -2,6 +2,7 @@
 #pragma once
 
 #include "Chaos/ConstraintHandle.h"
+#include "Chaos/CollisionResolutionManifold.h"
 #include "Chaos/PBDCollisionTypes.h"
 #include "Chaos/PBDConstraintContainer.h"
 #include "Framework/BufferedData.h"
@@ -52,67 +53,7 @@ enum class ECollisionModifierResult
 	Disabled,	/** Collision should be disabled */
 };
 
-template<class T, int d>
-class CHAOS_API TPBDCollisionConstraintHandle : public TContainerConstraintHandle<TPBDCollisionConstraint<T, d>>
-{
-public:
-	using Base = TContainerConstraintHandle<TPBDCollisionConstraint<T, d>>;
-	using FConstraintContainer = TPBDCollisionConstraint<T, d>;
-
-	TPBDCollisionConstraintHandle();
-	TPBDCollisionConstraintHandle(FConstraintContainer* InConstraintContainer, int32 InConstraintIndex);
-
-	TRigidBodyContactConstraint<T, d>& GetContact();
-	const TRigidBodyContactConstraint<T, d>& GetContact() const;
-	void SetConstraintIndex(int32 IndexIn) { ConstraintIndex = IndexIn; }
-
-protected:
-	using Base::ConstraintIndex;
-	using Base::ConstraintContainer;
-};
-
-template<class T, int d>
-class CHAOS_API TPBDCollisionConstraintHistory
-{
-	typedef TPair<const FImplicitObject*, const FImplicitObject*> ImplicitPairs;
-
-public:
-	using FConstraintContainerHandle = TPBDCollisionConstraintHandle<T, d>;
-	TPBDCollisionConstraintHistory(const TVector<T, d>& InLocation, const TRotation<T, d>& InRotation, int32 InTimestamp = -INT_MAX)
-		: Timestamp(InTimestamp)
-		, Location(InLocation)
-		, Rotation(InRotation)
-	{}
 
-	const TVector<T, d>& GetLocation() { return Location; }
-	const TRotation<T, d>& GetRotation() { return Rotation; }
-
-	void AddHandle(FConstraintContainerHandle* InHandle)
-	{
-		Implicits.Add(ImplicitPairs(InHandle->GetContact().Geometry[0], InHandle->GetContact().Geometry[1]));
-		Handles.Push(InHandle);
-	}
-	void RemoveHandle(FConstraintContainerHandle* InHandle)
-	{
-		Handles.RemoveSingleSwap(InHandle);
-		Implicits.Remove(ImplicitPairs(InHandle->GetContact().Geometry[0], InHandle->GetContact().Geometry[1]));
-	}
-
-	TArray<FConstraintContainerHandle*>& GetHandles() { return Handles; }
-	const TArray<FConstraintContainerHandle*>& GetHandles() const { return Handles; }
-
-	bool ContainsShapeConnection(const FImplicitObject* Implicit0In, const FImplicitObject* Implicit1In) { return Implicits.Contains(ImplicitPairs(Implicit0In, Implicit1In)); }
-
-	int32 GetTimestamp() const { return Timestamp; }
-	void SetTimestamp(int32 InTimestamp) { Timestamp = InTimestamp; }
-
-private:
-	int32 Timestamp;
-	TVector<T, d> Location;
-	TRotation<T, d> Rotation;
-	TArray<FConstraintContainerHandle*> Handles;
-	TSet<ImplicitPairs> Implicits;
-};
 
 template<typename T, int d>
 using TRigidBodyContactConstraintsPostComputeCallback = TFunction<void()>;
@@ -141,7 +82,7 @@ public:
 	using FReal = T;
 	static const int Dimensions = d;
 	using FConstraintContainerHandle = TPBDCollisionConstraintHandle<T, d>;
-	using FConstraintHistory = TPBDCollisionConstraintHistory<T, d>;
+	using FConstraintManifold = TCollisionResolutionManifold<T, d>;
 	using FConstraintHandleAllocator = TConstraintHandleAllocator<TPBDCollisionConstraint<T, d>>;
 	using FRigidBodyContactConstraint = TRigidBodyContactConstraint<T, d>;
 	typedef TPair<TGeometryParticleHandle<T, d>*, const TGeometryParticleHandle<T, d>*> FConstraintHandleID;
@@ -338,7 +279,7 @@ private:
 	bool bEnableCollisions;
 
 	int32 LifespanCounter;
-	int32 CollisionHistoryLifespan;
+	int32 CollisionManifoldLifespan;
 	float CollisionVelocityInflation;
 
 	TRigidBodyContactConstraintsPostComputeCallback<T, d> PostComputeCallback;
@@ -348,7 +289,7 @@ private:
 	TArray<FConstraintContainerHandle*> Handles;
 	FConstraintHandleAllocator HandleAllocator;
 
-	TMap<FConstraintHandleID, FConstraintHistory*> History;
+	TMap<FConstraintHandleID, FConstraintManifold*> Manifolds;
 };
 
 extern template void TPBDCollisionConstraint<float, 3>::UpdateConstraint<ECollisionUpdateType::Any>(const float Thickness, FRigidBodyContactConstraint* Constraints, int32 NumConstraints);
diff --git a/Engine/Source/Runtime/Experimental/Chaos/Public/Chaos/PBDCollisionConstraintImp.h b/Engine/Source/Runtime/Experimental/Chaos/Public/Chaos/PBDCollisionConstraintImp.h
--- a/Engine/Source/Runtime/Experimental/Chaos/Public/Chaos/PBDCollisionConstraintImp.h
+++ b/Engine/Source/Runtime/Experimental/Chaos/Public/Chaos/PBDCollisionConstraintImp.h
@@ -305,17 +305,17 @@ namespace Chaos
 					ensure(Handles.Num() == Constraints.Num());
 
 					FConstraintHandleID HandleID = GetConstraintHandleID(Constraints[Idx]);
-					FConstraintHistory* HistoryElement;
-					if (auto ValPtr = History.Find(HandleID))
+					FConstraintManifold* Manifold;
+					if (auto ValPtr = Manifolds.Find(HandleID))
 					{
-						HistoryElement = *ValPtr;
+						Manifold = *ValPtr;
 					}
 					else
 					{
-						HistoryElement = new FConstraintHistory(Constraints[Idx].Particle->X(), Constraints[Idx].Particle->R(), LifespanCounter + CollisionHistoryLifespan);
-						History.Add(HandleID, HistoryElement);
+						Manifold = new FConstraintManifold(Constraints[Idx].Particle->X(), Constraints[Idx].Particle->R(), LifespanCounter + CollisionManifoldLifespan);
+						Manifolds.Add(HandleID, Manifold);
 					}
-					HistoryElement->AddHandle(Handle);
+					Manifold->AddHandle(Handle);
 				}
 				LifespanCounter++;
 			}
diff --git a/Engine/Source/Runtime/Experimental/Chaos/Public/Chaos/PBDCollisionTypes.h b/Engine/Source/Runtime/Experimental/Chaos/Public/Chaos/PBDCollisionTypes.h
--- a/Engine/Source/Runtime/Experimental/Chaos/Public/Chaos/PBDCollisionTypes.h
+++ b/Engine/Source/Runtime/Experimental/Chaos/Public/Chaos/PBDCollisionTypes.h
@@ -1,13 +1,18 @@
 // Copyright 1998-2019 Epic Games, Inc. All Rights Reserved.
 #pragma once
-#include "Chaos/Vector.h"
 #include "Box.h"
 #include "Chaos/ParticleHandleFwd.h"
+#include "Chaos/PBDConstraintContainer.h"
+#include "Chaos/Vector.h"
 
 class UPhysicalMaterial;
 
 namespace Chaos
 {
+
+	template<typename T, int d>
+	class TPBDCollisionConstraint;
+
 /**/
 // @todo(ccaulfield): should these be held SOA style in the PBDCollisionConstraint?
 template<class T, int d>
@@ -43,6 +48,45 @@ struct TRigidBodyContactConstraintPGS
 	TVector<T, d> AccumulatedImpulse;
 };
 
+
+template<class T, int d>
+class CHAOS_API TPBDCollisionConstraintHandle : public TContainerConstraintHandle<TPBDCollisionConstraint<T, d>>
+{
+public:
+	using Base = TContainerConstraintHandle<TPBDCollisionConstraint<T, d>>;
+	using FConstraintContainer = TPBDCollisionConstraint<T, d>;
+
+	TPBDCollisionConstraintHandle() 
+	{}
+	
+	TPBDCollisionConstraintHandle(FConstraintContainer* InConstraintContainer, int32 InConstraintIndex) 
+		: TContainerConstraintHandle<TPBDCollisionConstraint<T, d>>(InConstraintContainer, InConstraintIndex)
+	{
+	}
+
+	TRigidBodyContactConstraint<T, d>& GetContact()
+	{
+		return ConstraintContainer->Constraints[ConstraintIndex];
+	}
+
+	const TRigidBodyContactConstraint<T, d>& GetContact() const
+	{
+		return ConstraintContainer->Constraints[ConstraintIndex];
+	}
+
+	void SetConstraintIndex(int32 IndexIn) 
+	{ 
+		ConstraintIndex = IndexIn; 
+	}
+
+protected:
+	using Base::ConstraintIndex;
+	using Base::ConstraintContainer;
+
+
+};
+
+
 /**
  * Collision event data stored for use by other systems (e.g. Niagara, gameplay events)
  */
