When convex MTDs are requested, fix the non-implemented default case to remove NaNs.

diff --git a/Engine/Source/Runtime/Experimental/Chaos/Public/Chaos/GeometryQueries.h b/Engine/Source/Runtime/Experimental/Chaos/Public/Chaos/GeometryQueries.h
--- a/Engine/Source/Runtime/Experimental/Chaos/Public/Chaos/GeometryQueries.h
+++ b/Engine/Source/Runtime/Experimental/Chaos/Public/Chaos/GeometryQueries.h
@@ -201,8 +201,9 @@ namespace Chaos
 			// Compute MTD in the case of an initial overlap
 			if (bResult && bComputeMTD && OutTime == 0.f)
 			{
-				//ensure(false); // We don't support MTD for non-convex types yet!
-				OutNormal = FVec3(0, 0, 1);
+				ensure(false); // We don't support MTD for non-convex types yet!
+				LocalNormal = FVec3(0, 0, 1);
+				LocalPosition = ATM.GetLocation();
 			}
 		}
 
