Back out changelist 9977278. CIS errors

diff --git a/Engine/Source/Runtime/Experimental/Chaos/Private/Chaos/CollisionResolutionConvexConvex.cpp b/Engine/Source/Runtime/Experimental/Chaos/Private/Chaos/CollisionResolutionConvexConvex.cpp
deleted file mode 100644
diff --git a/Engine/Source/Runtime/Experimental/Chaos/Private/Chaos/PBDCollisionConstraint.cpp b/Engine/Source/Runtime/Experimental/Chaos/Private/Chaos/PBDCollisionConstraint.cpp
--- a/Engine/Source/Runtime/Experimental/Chaos/Private/Chaos/PBDCollisionConstraint.cpp
+++ b/Engine/Source/Runtime/Experimental/Chaos/Private/Chaos/PBDCollisionConstraint.cpp
@@ -19,7 +19,6 @@
 #include "Chaos/PBDRigidsSOAs.h"
 #include "Chaos/ChaosDebugDraw.h"
 #include "Chaos/ISpatialAccelerationCollection.h"
-#include "Chaos/CollisionResolutionConvexConvex.h"
 
 #if INTEL_ISPC
 #include "PBDCollisionConstraint.ispc.generated.h"
@@ -34,8 +33,8 @@ FAutoConsoleVariableRef CVarConstraintBPBVHDepth(TEXT("p.ConstraintBPBVHDepth"),
 int32 BPTreeOfGrids = 1;
 FAutoConsoleVariableRef CVarBPTreeOfGrids(TEXT("p.BPTreeOfGrids"), BPTreeOfGrids, TEXT("Whether to use a seperate tree of grids for bp"));
 
-int32 CollisionManifoldLifespanCVar = 1;
-FAutoConsoleVariableRef CVarCollisionManifoldLifespan(TEXT("p.CollisionManifoldLifespan"), CollisionManifoldLifespanCVar, TEXT("Number of iterations to cache a collision point.[def:5]"));
+int32 CollisionHistoryLifespanCVar = 1;
+FAutoConsoleVariableRef CVarCollisionHistoryLifespan(TEXT("p.CollisionHistoryLifespan"), CollisionHistoryLifespanCVar, TEXT("Number of iterations to cache a collision point.[def:5]"));
 
 int32 CollisionManifoldSamplesCVar = 5;
 FAutoConsoleVariableRef CVarCollisionManifoldSamples(TEXT("p.CollisionManifoldSamples"), CollisionManifoldSamplesCVar, TEXT("Number of sample points on the manifold.[def:6]"));
@@ -62,10 +61,37 @@ CHAOS_API int32 EnableCollisions = 1;
 FAutoConsoleVariableRef CVarEnableCollisions(TEXT("p.EnableCollisions"), EnableCollisions, TEXT("Enable/Disable collisions on the Chaos solver."));
 
 template<ECollisionUpdateType, typename T, int d>
-void UpdateConstraintImp2(const TRigidTransform<T, d>& ParticleTM, const FImplicitObject& LevelsetObject, const TRigidTransform<T, d>& LevelsetTM, const T Thickness, TRigidBodyContactConstraint<T, d>& Constraint, int32 NumConstraints, TCollisionResolutionManifold<T, d>* Manifold);
+void UpdateConstraintImp2(const TRigidTransform<T, d>& ParticleTM, const FImplicitObject& LevelsetObject, const TRigidTransform<T, d>& LevelsetTM, const T Thickness, TRigidBodyContactConstraint<T, d>& Constraint, int32 NumConstraints, TPBDCollisionConstraintHistory<T, d>* History);
 
 template<typename T, int d>
-void ConstructConstraintsImpl(TGeometryParticleHandle<T, d>* Particle0, TGeometryParticleHandle<T, d>* Particle1, const FImplicitObject* Implicit0, const FImplicitObject* Implicit1, const T Thickness, TArray<TRigidBodyContactConstraint<T, d>>& ConstraintBuffer, TCollisionResolutionManifold<T, d>* Manifold);
+void ConstructConstraintsImpl(TGeometryParticleHandle<T, d>* Particle0, TGeometryParticleHandle<T, d>* Particle1, const FImplicitObject* Implicit0, const FImplicitObject* Implicit1, const T Thickness, TArray<TRigidBodyContactConstraint<T, d>>& ConstraintBuffer, TPBDCollisionConstraintHistory<T, d>* History);
+
+//
+// Collision Constraint Handle
+//
+
+template<typename T, int d>
+TPBDCollisionConstraintHandle<T, d>::TPBDCollisionConstraintHandle()
+{
+}
+
+template<typename T, int d>
+TPBDCollisionConstraintHandle<T, d>::TPBDCollisionConstraintHandle(FConstraintContainer* InConstraintContainer, int32 InConstraintIndex)
+	: TContainerConstraintHandle<TPBDCollisionConstraint<T, d>>(InConstraintContainer, InConstraintIndex)
+{
+}
+
+template<typename T, int d>
+TRigidBodyContactConstraint<T, d>& TPBDCollisionConstraintHandle<T, d>::GetContact()
+{
+	return ConstraintContainer->Constraints[ConstraintIndex];
+}
+
+template<typename T, int d>
+const TRigidBodyContactConstraint<T, d>& TPBDCollisionConstraintHandle<T, d>::GetContact() const
+{
+	return ConstraintContainer->Constraints[ConstraintIndex];
+}
 
 
 //
@@ -85,7 +111,7 @@ TPBDCollisionConstraint<T, d>::TPBDCollisionConstraint(const TPBDRigidsSOAs<T,d>
 	, bUseCCD(false)
 	, bEnableCollisions(true)
 	, LifespanCounter(0)
-	, CollisionManifoldLifespan(CollisionManifoldLifespanCVar)
+	, CollisionHistoryLifespan(CollisionHistoryLifespanCVar)
 	, CollisionVelocityInflation(CollisionVelocityInflationCVar)
 	, PostComputeCallback(nullptr)
 	, PostApplyCallback(nullptr)
@@ -108,12 +134,12 @@ void TPBDCollisionConstraint<T, d>::Reset(/*const TPBDRigidParticles<T, d>& InPa
 		}
 	}
 
-	TArray<FConstraintHandleID> ManifoldKeys;
-	Manifolds.GetKeys(ManifoldKeys);
-	for (FConstraintHandleID Key : ManifoldKeys)
+	TArray<FConstraintHandleID> HistoryKeys;
+	History.GetKeys(HistoryKeys);
+	for (FConstraintHandleID Key : HistoryKeys)
 	{
-		FConstraintManifold* Elem = Manifolds[Key];
-		if (Elem->GetTimestamp()==FLT_MAX || LifespanCounter > Elem->GetTimestamp())
+		FConstraintHistory* Elem = History[Key];
+		if (LifespanCounter > Elem->GetTimestamp())
 		{
 			bool bInContact = false;
 			for (FConstraintContainerHandle* Handle : Elem->GetHandles())
@@ -142,16 +168,16 @@ void TPBDCollisionConstraint<T, d>::Reset(/*const TPBDRigidParticles<T, d>& InPa
 template<typename T, int d>
 void TPBDCollisionConstraint<T, d>::RemoveConstraint(int32 Idx)
 {
-	// Manifold
+	// History
 	FConstraintHandleID HandleID = GetConstraintHandleID(Constraints[Idx]);
-	if (FConstraintManifold** Ptr = Manifolds.Find(HandleID))
+	if (FConstraintHistory** Ptr = History.Find(HandleID))
 	{
-		if (FConstraintManifold* ManifoldEntry = *Ptr)
+		if (FConstraintHistory* HistoryEntry = *Ptr)
 		{
-			ManifoldEntry->RemoveHandle(Handles[Idx]);
-			if (!ManifoldEntry->GetHandles().Num())
+			HistoryEntry->RemoveHandle(Handles[Idx]);
+			if (!HistoryEntry->GetHandles().Num())
 			{
-				Manifolds.Remove(HandleID);
+				History.Remove(HandleID);
 			}
 		}
 	}
@@ -1630,7 +1656,7 @@ TArray<Pair<const FImplicitObject*, TRigidTransform<T, d>>> FindRelevantShapes2(
 
 DECLARE_CYCLE_STAT(TEXT("TPBDCollisionConstraint::UpdateUnionUnionConstraint"), STAT_UpdateUnionUnionConstraint, STATGROUP_ChaosWide);
 template<ECollisionUpdateType UpdateType, typename T, int d>
-void UpdateUnionUnionConstraint(const T Thickness, TRigidBodyContactConstraint<T, d>* Constraints, int32 NumConstraints, TCollisionResolutionManifold<T, d>* Manifold)
+void UpdateUnionUnionConstraint(const T Thickness, TRigidBodyContactConstraint<T, d>* Constraints, int32 NumConstraints, TPBDCollisionConstraintHistory<T, d>* History)
 {
 	checkSlow(Constraints != nullptr && NumConstraints >= 1);
 	SCOPE_CYCLE_COUNTER(STAT_UpdateUnionUnionConstraint);
@@ -1658,14 +1684,14 @@ void UpdateUnionUnionConstraint(const T Thickness, TRigidBodyContactConstraint<T
 		{
 			const FImplicitObject& ParticleInnerObj = *ParticlePair.First;
 			const TRigidTransform<T, d> ParticleInnerObjTM = ParticlePair.Second * ParticlesTM;
-			UpdateConstraintImp2<UpdateType>(ParticleInnerObj, ParticleInnerObjTM, LevelsetInnerObj, LevelsetInnerObjTM, Thickness, Constraints, NumConstraints, Manifold);
+			UpdateConstraintImp2<UpdateType>(ParticleInnerObj, ParticleInnerObjTM, LevelsetInnerObj, LevelsetInnerObjTM, Thickness, Constraints, NumConstraints, History);
 		}
 	}
 }
 
 DECLARE_CYCLE_STAT(TEXT("TPBDCollisionConstraint::UpdateSingleUnionConstraint"), STAT_UpdateSingleUnionConstraint, STATGROUP_ChaosWide);
 template<ECollisionUpdateType UpdateType, typename T, int d>
-void UpdateSingleUnionConstraint(const T Thickness, TRigidBodyContactConstraint<T, d>* Constraints, int32 NumConstraints, TCollisionResolutionManifold<T, d>* Manifold)
+void UpdateSingleUnionConstraint(const T Thickness, TRigidBodyContactConstraint<T, d>* Constraints, int32 NumConstraints, TPBDCollisionConstraintHistory<T, d>* History)
 {
 	checkSlow(Constraints != nullptr && NumConstraints >= 1);
 	SCOPE_CYCLE_COUNTER(STAT_UpdateSingleUnionConstraint);
@@ -1684,7 +1710,7 @@ void UpdateSingleUnionConstraint(const T Thickness, TRigidBodyContactConstraint<
 	{
 		const FImplicitObject& LevelsetInnerObj = *LevelsetObjPair.First;
 		const TRigidTransform<T, d> LevelsetInnerObjTM = LevelsetObjPair.Second * LevelsetTM;
-		UpdateConstraintImp2<UpdateType>(*ParticleObj, ParticlesTM, LevelsetInnerObj, LevelsetInnerObjTM, Thickness, Constraints, NumConstraints, Manifold);
+		UpdateConstraintImp2<UpdateType>(*ParticleObj, ParticlesTM, LevelsetInnerObj, LevelsetInnerObjTM, Thickness, Constraints, NumConstraints, History);
 	}
 }
 
@@ -1893,9 +1919,95 @@ void UpdateBoxConstraint(const TBox<T, d>& Box1, const TRigidTransform<T, d>& Bo
 	}
 }
 
+DECLARE_CYCLE_STAT(TEXT("UpdateConvexConvexConstraint"), UpdateConvexConvexConstraint, STATGROUP_ChaosWide);
+template<ECollisionUpdateType UpdateType, typename T, int d, typename TGeometryA, typename TGeometryB>
+void UpdateConvexConvexConstraint(const TGeometryA& A, TRigidTransform<T, d> ATM, const TGeometryB& B, TRigidTransform<T, d> BTM, T Thickness, TRigidBodyContactConstraint<T, d>* Constraints, int32 NumConstraints, TPBDCollisionConstraintHistory<T, d>* History)
+{
+	SCOPE_CYCLE_COUNTER(UpdateConvexConvexConstraint);
+
+	ensure(NumConstraints);
+
+	unsigned char AType = A.GetType();
+	unsigned char AStaticType = ImplicitObjectType::IsScaled;
+
+	if (ensure(IsScaled(A.GetType()) && IsScaled(B.GetType())))
+	{
+		const TImplicitObjectScaled<TConvex<T, d>>& AScaled = static_cast<const TImplicitObjectScaled<TConvex<T, d>>&>(A);
+		const TImplicitObjectScaled<TConvex<T, d>>& BScaled = static_cast<const TImplicitObjectScaled<TConvex<T, d>>&>(B);
+
+		TRigidTransform<T, d> BToATM = BTM.GetRelativeTransform(ATM);
+		TRigidTransform<T, d> AToBTM = ATM.GetRelativeTransform(BTM);
+
+		if (ensure(GetInnerType(A.GetType()) == ImplicitObjectType::Convex && GetInnerType(B.GetType()) == ImplicitObjectType::Convex))
+		{
+			const TConvex<T, d>* AObject = static_cast<const TConvex<T, d>*>(AScaled.Object().Get());
+			const TConvex<T, d>* BObject = static_cast<const TConvex<T, d>*>(BScaled.Object().Get());
+
+			const TParticles<T, d>& SurfaceParticles = AObject->GetSurfaceParticles();
+
+			T OutDistance;
+			TVector<T, d> StubNormal;
+			TVector<T, d> OutNearestA;
+			TVector<T, d> OutNearestB;
+			TVector<T, d> OutNormal;
+
+			if (History == nullptr)
+			{
+				if (GJKDistance<T>(*AObject, *BObject, BToATM, OutDistance, OutNearestA, OutNearestB))
+				{
+					TVector<T, d> Normal(0);
+					TVector<T, d> WorldA = ATM.TransformPosition(OutNearestA);
+					TVector<T, d> WorldB = BTM.TransformPosition(OutNearestB);
+					TVector<T, d> CenterLocation = WorldA;
+					T Phi = BObject->PhiWithNormal(BTM.InverseTransformPosition(CenterLocation), Normal);
+
+					T ManifoldSize = A.BoundingBox().Extents().Size() * CollisionManifoldSampleScaleCVar;
+					TVector<T, d> CrossVector = FMath::IsNearlyEqual(FMath::Abs(TVector<T, d>::DotProduct(TVector<T, d>(0, 1, 0), Normal)), T(1)) ?
+						TVector<T, d>(0, 0, 1) :
+						TVector<T, d>(0, 1, 0);
+
+					Constraints[0].LocalLocation = ATM.InverseTransformPosition(CenterLocation);
+					Constraints[0].Normal = Normal;
+
+					if (NumConstraints > 1)
+					{
+						TVector<T, d> DirVector = CrossVector;
+						T Angle = T(360) / T(NumConstraints - 1), LocalPhi;
+						for (int32 Idx = 1; Idx < NumConstraints; Idx++)
+						{
+							CrossVector = CrossVector.RotateAngleAxis(Angle,Normal);
+							TVector<T, d> WorldLocation = CenterLocation + ManifoldSize * CrossVector;
+							Constraints[Idx].LocalLocation = ATM.InverseTransformPosition(WorldLocation);
+
+							LocalPhi = AObject->PhiWithNormal(Constraints[Idx].LocalLocation, StubNormal);
+							Constraints[Idx].LocalLocation -= StubNormal * LocalPhi;
+
+							Constraints[Idx].Normal = Normal;
+						}
+					}
+				}
+			}
+
+			for (int32 Idx = 0; Idx < NumConstraints; Idx++)
+			{
+				Constraints[Idx].Location = ATM.TransformPosition(Constraints[Idx].LocalLocation);
+				Constraints[Idx].Phi = BObject->PhiWithNormal(BTM.InverseTransformPosition(Constraints[Idx].Location), StubNormal);
+
+				if (Constraints[Idx].Phi < 0.f && History)
+				{
+					History->SetTimestamp(INT_MAX);
+				}
+			}
+
+			return;
+		}
+	}
+
+	ensureMsgf(false, TEXT("Unsupported convex to convex constraint."));
+}
 
 template<typename T, int d>
-void ConstructLevelsetConstraints(TGeometryParticleHandle<T, d>* Particle0, TGeometryParticleHandle<T, d>* Particle1, const FImplicitObject* Implicit0, const FImplicitObject* Implicit1, const T Thickness, TArray<TRigidBodyContactConstraint<T, d>>& ConstraintBuffer, TCollisionResolutionManifold<T, d>* Manifold = nullptr)
+void ConstructLevelsetConstraints(TGeometryParticleHandle<T, d>* Particle0, TGeometryParticleHandle<T, d>* Particle1, const FImplicitObject* Implicit0, const FImplicitObject* Implicit1, const T Thickness, TArray<TRigidBodyContactConstraint<T, d>>& ConstraintBuffer, TPBDCollisionConstraintHistory<T, d>* History = nullptr)
 {
 	//todo(ocohen):if both have collision particles, use the one with fewer?
 	if (!Particle1->Geometry() || (Particle0->AsDynamic() && !Particle0->AsDynamic()->CollisionParticlesSize() && Particle0->Geometry() && !Particle0->Geometry()->IsUnderlyingUnion()))
@@ -1919,7 +2031,29 @@ void ConstructLevelsetConstraints(TGeometryParticleHandle<T, d>* Particle0, TGeo
 }
 
 template<typename T, int d>
-void ConstructBoxConstraints(TGeometryParticleHandle<T, d>* Particle0, TGeometryParticleHandle<T, d>* Particle1, const FImplicitObject* Implicit0, const FImplicitObject* Implicit1, const T Thickness, TArray<TRigidBodyContactConstraint<T, d>>& ConstraintBuffer, TCollisionResolutionManifold<T, d>* Manifold = nullptr)
+void ConstructConvexConvexConstraints(TGeometryParticleHandle<T, d>* Particle0, TGeometryParticleHandle<T, d>* Particle1, const FImplicitObject* Implicit0, const FImplicitObject* Implicit1, const T Thickness, TArray<TRigidBodyContactConstraint<T, d>>& ConstraintBuffer, TPBDCollisionConstraintHistory<T, d>* History = nullptr)
+{
+	if (History)
+	{
+		if (History->ContainsShapeConnection(Implicit0, Implicit1))
+		{
+			return;
+		}
+	}
+
+	for (int i = 0; i < CollisionManifoldSamplesCVar; i++)
+	{
+		TRigidBodyContactConstraint<T, d> Constraint;
+		Constraint.Particle = Particle0;
+		Constraint.Levelset = Particle1;
+		Constraint.Geometry[0] = Implicit0;
+		Constraint.Geometry[1] = Implicit1;
+		ConstraintBuffer.Add(Constraint);
+	}
+}
+
+template<typename T, int d>
+void ConstructBoxConstraints(TGeometryParticleHandle<T, d>* Particle0, TGeometryParticleHandle<T, d>* Particle1, const FImplicitObject* Implicit0, const FImplicitObject* Implicit1, const T Thickness, TArray<TRigidBodyContactConstraint<T, d>>& ConstraintBuffer, TPBDCollisionConstraintHistory<T, d>* History = nullptr)
 {
 	TRigidBodyContactConstraint<T, d> Constraint;
 	Constraint.Particle = Particle0;
@@ -1930,7 +2064,7 @@ void ConstructBoxConstraints(TGeometryParticleHandle<T, d>* Particle0, TGeometry
 }
 
 template<typename T, int d>
-void ConstructBoxPlaneConstraints(TGeometryParticleHandle<T, d>* Particle0, TGeometryParticleHandle<T, d>* Particle1, const FImplicitObject* Implicit0, const FImplicitObject* Implicit1, const T Thickness, TArray<TRigidBodyContactConstraint<T, d>>& ConstraintBuffer, TCollisionResolutionManifold<T, d>* Manifold = nullptr)
+void ConstructBoxPlaneConstraints(TGeometryParticleHandle<T, d>* Particle0, TGeometryParticleHandle<T, d>* Particle1, const FImplicitObject* Implicit0, const FImplicitObject* Implicit1, const T Thickness, TArray<TRigidBodyContactConstraint<T, d>>& ConstraintBuffer, TPBDCollisionConstraintHistory<T, d>* History = nullptr)
 {
 	TRigidBodyContactConstraint<T, d> Constraint;
 	Constraint.Particle = Particle0;
@@ -1941,7 +2075,7 @@ void ConstructBoxPlaneConstraints(TGeometryParticleHandle<T, d>* Particle0, TGeo
 }
 
 template<typename T, int d>
-void ConstructSphereConstraints(TGeometryParticleHandle<T, d>* Particle0, TGeometryParticleHandle<T, d>* Particle1, const FImplicitObject* Implicit0, const FImplicitObject* Implicit1, const T Thickness, TArray<TRigidBodyContactConstraint<T, d>>& ConstraintBuffer, TCollisionResolutionManifold<T, d>* Manifold = nullptr)
+void ConstructSphereConstraints(TGeometryParticleHandle<T, d>* Particle0, TGeometryParticleHandle<T, d>* Particle1, const FImplicitObject* Implicit0, const FImplicitObject* Implicit1, const T Thickness, TArray<TRigidBodyContactConstraint<T, d>>& ConstraintBuffer, TPBDCollisionConstraintHistory<T, d>* History = nullptr)
 {
 	TRigidBodyContactConstraint<T, d> Constraint;
 	Constraint.Particle = Particle0;
@@ -1952,7 +2086,7 @@ void ConstructSphereConstraints(TGeometryParticleHandle<T, d>* Particle0, TGeome
 }
 
 template<typename T, int d>
-void ConstructSpherePlaneConstraints(TGeometryParticleHandle<T, d>* Particle0, TGeometryParticleHandle<T, d>* Particle1, const FImplicitObject* Implicit0, const FImplicitObject* Implicit1, const T Thickness, TArray<TRigidBodyContactConstraint<T, d>>& ConstraintBuffer, TCollisionResolutionManifold<T, d>* Manifold = nullptr)
+void ConstructSpherePlaneConstraints(TGeometryParticleHandle<T, d>* Particle0, TGeometryParticleHandle<T, d>* Particle1, const FImplicitObject* Implicit0, const FImplicitObject* Implicit1, const T Thickness, TArray<TRigidBodyContactConstraint<T, d>>& ConstraintBuffer, TPBDCollisionConstraintHistory<T, d>* History = nullptr)
 {
 	TRigidBodyContactConstraint<T, d> Constraint;
 	Constraint.Particle = Particle0;
@@ -1963,7 +2097,7 @@ void ConstructSpherePlaneConstraints(TGeometryParticleHandle<T, d>* Particle0, T
 }
 
 template<typename T, int d>
-void ConstructSphereBoxConstraints(TGeometryParticleHandle<T, d>* Particle0, TGeometryParticleHandle<T, d>* Particle1, const FImplicitObject* Implicit0, const FImplicitObject* Implicit1, const T Thickness, TArray<TRigidBodyContactConstraint<T, d>>& ConstraintBuffer, TCollisionResolutionManifold<T, d>* Manifold = nullptr)
+void ConstructSphereBoxConstraints(TGeometryParticleHandle<T, d>* Particle0, TGeometryParticleHandle<T, d>* Particle1, const FImplicitObject* Implicit0, const FImplicitObject* Implicit1, const T Thickness, TArray<TRigidBodyContactConstraint<T, d>>& ConstraintBuffer, TPBDCollisionConstraintHistory<T, d>* History = nullptr)
 {
 	TRigidBodyContactConstraint<T, d> Constraint;
 	Constraint.Particle = Particle0;
@@ -1974,7 +2108,7 @@ void ConstructSphereBoxConstraints(TGeometryParticleHandle<T, d>* Particle0, TGe
 }
 
 template<typename T, int d>
-void ConstructCapsuleCapsuleConstraints(TGeometryParticleHandle<T, d>* Particle0, TGeometryParticleHandle<T, d>* Particle1, const FImplicitObject* Implicit0, const FImplicitObject* Implicit1, const T Thickness, TArray<TRigidBodyContactConstraint<T, d>>& ConstraintBuffer, TCollisionResolutionManifold<T, d>* Manifold = nullptr)
+void ConstructCapsuleCapsuleConstraints(TGeometryParticleHandle<T, d>* Particle0, TGeometryParticleHandle<T, d>* Particle1, const FImplicitObject* Implicit0, const FImplicitObject* Implicit1, const T Thickness, TArray<TRigidBodyContactConstraint<T, d>>& ConstraintBuffer, TPBDCollisionConstraintHistory<T, d>* History = nullptr)
 {
 	TRigidBodyContactConstraint<T, d> Constraint;
 	Constraint.Particle = Particle0;
@@ -1985,7 +2119,7 @@ void ConstructCapsuleCapsuleConstraints(TGeometryParticleHandle<T, d>* Particle0
 }
 
 template<typename T, int d>
-void ConstructCapsuleBoxConstraints(TGeometryParticleHandle<T, d>* Particle0, TGeometryParticleHandle<T, d>* Particle1, const FImplicitObject* Implicit0, const FImplicitObject* Implicit1, const T Thickness, TArray<TRigidBodyContactConstraint<T, d>>& ConstraintBuffer, TCollisionResolutionManifold<T, d>* Manifold = nullptr)
+void ConstructCapsuleBoxConstraints(TGeometryParticleHandle<T, d>* Particle0, TGeometryParticleHandle<T, d>* Particle1, const FImplicitObject* Implicit0, const FImplicitObject* Implicit1, const T Thickness, TArray<TRigidBodyContactConstraint<T, d>>& ConstraintBuffer, TPBDCollisionConstraintHistory<T, d>* History = nullptr)
 {
 	TRigidBodyContactConstraint<T, d> Constraint;
 	Constraint.Particle = Particle0;
@@ -1996,7 +2130,7 @@ void ConstructCapsuleBoxConstraints(TGeometryParticleHandle<T, d>* Particle0, TG
 }
 
 template<typename T, int d>
-void ConstructSingleUnionConstraints(TGeometryParticleHandle<T, d>* Particle0, TGeometryParticleHandle<T, d>* Particle1, const FImplicitObject* Implicit0, const FImplicitObject* Implicit1, const T Thickness, TArray<TRigidBodyContactConstraint<T, d>>& ConstraintBuffer, TCollisionResolutionManifold<T, d>* Manifold = nullptr)
+void ConstructSingleUnionConstraints(TGeometryParticleHandle<T, d>* Particle0, TGeometryParticleHandle<T, d>* Particle1, const FImplicitObject* Implicit0, const FImplicitObject* Implicit1, const T Thickness, TArray<TRigidBodyContactConstraint<T, d>>& ConstraintBuffer, TPBDCollisionConstraintHistory<T, d>* History = nullptr)
 {
 	TRigidBodyContactConstraint<T, d> Constraint;
 	Constraint.Particle = Particle0;
@@ -2007,7 +2141,7 @@ void ConstructSingleUnionConstraints(TGeometryParticleHandle<T, d>* Particle0, T
 }
 
 template<typename T, int d>
-void ConstructUnionUnionConstraints(TGeometryParticleHandle<T, d>* Particle0, TGeometryParticleHandle<T, d>* Particle1, const FImplicitObject* Implicit0, const FImplicitObject* Implicit1, const T Thickness, TArray<TRigidBodyContactConstraint<T, d>>& ConstraintBuffer, TCollisionResolutionManifold<T, d>* Manifold = nullptr)
+void ConstructUnionUnionConstraints(TGeometryParticleHandle<T, d>* Particle0, TGeometryParticleHandle<T, d>* Particle1, const FImplicitObject* Implicit0, const FImplicitObject* Implicit1, const T Thickness, TArray<TRigidBodyContactConstraint<T, d>>& ConstraintBuffer, TPBDCollisionConstraintHistory<T, d>* History = nullptr)
 {
 	SCOPE_CYCLE_COUNTER(STAT_UpdateUnionUnionConstraint);
 
@@ -2031,13 +2165,13 @@ void ConstructUnionUnionConstraints(TGeometryParticleHandle<T, d>* Particle0, TG
 		for (const Pair<const FImplicitObject*, TRigidTransform<T, d>>& ParticlePair : ParticleShapes)
 		{
 			const FImplicitObject* ParticleInnerObj = ParticlePair.First;
-			ConstructConstraintsImpl<T, d>(Particle0, Particle1, ParticleInnerObj, LevelsetInnerObj, Thickness, ConstraintBuffer, Manifold);
+			ConstructConstraintsImpl<T, d>(Particle0, Particle1, ParticleInnerObj, LevelsetInnerObj, Thickness, ConstraintBuffer, History);
 		}
 	}
 }
 
 template<typename T, int d>
-void ConstructConstraintsImpl(TGeometryParticleHandle<T, d>* Particle0, TGeometryParticleHandle<T, d>* Particle1, const FImplicitObject* Implicit0, const FImplicitObject* Implicit1, const T Thickness, TArray< TRigidBodyContactConstraint<T, d> >& ConstraintBuffer, TCollisionResolutionManifold<T, d>* Manifold)
+void ConstructConstraintsImpl(TGeometryParticleHandle<T, d>* Particle0, TGeometryParticleHandle<T, d>* Particle1, const FImplicitObject* Implicit0, const FImplicitObject* Implicit1, const T Thickness, TArray< TRigidBodyContactConstraint<T, d> >& ConstraintBuffer, TPBDCollisionConstraintHistory<T, d>* History)
 {
 	// TriangleMesh implicit's are for scene query only.
 	if (Implicit0 && GetInnerType(Implicit0->GetType()) == ImplicitObjectType::TriangleMesh) return; //as
@@ -2045,71 +2179,71 @@ void ConstructConstraintsImpl(TGeometryParticleHandle<T, d>* Particle0, TGeometr
 
 	if (!Implicit0 || !Implicit1)
 	{
-		ConstructLevelsetConstraints(Particle0, Particle1, Implicit0, Implicit1, Thickness, ConstraintBuffer, Manifold);
+		ConstructLevelsetConstraints(Particle0, Particle1, Implicit0, Implicit1, Thickness, ConstraintBuffer, History);
 	}
 	else if (Implicit0->GetType() == TBox<T, d>::StaticType() && Implicit1->GetType() == TBox<T, d>::StaticType()) 
 	{
-		ConstructBoxConstraints(Particle0, Particle1, Implicit0, Implicit1, Thickness, ConstraintBuffer, Manifold);
+		ConstructBoxConstraints(Particle0, Particle1, Implicit0, Implicit1, Thickness, ConstraintBuffer, History);
 	}
 	else if (Implicit0->GetType() == TSphere<T, d>::StaticType() && Implicit1->GetType() == TSphere<T, d>::StaticType())
 	{
-		ConstructSphereConstraints(Particle0, Particle1, Implicit0, Implicit1, Thickness, ConstraintBuffer, Manifold);
+		ConstructSphereConstraints(Particle0, Particle1, Implicit0, Implicit1, Thickness, ConstraintBuffer, History);
 	}
 	else if (Implicit0->GetType() == TBox<T, d>::StaticType() && Implicit1->GetType() == TPlane<T, d>::StaticType())
 	{
-		ConstructBoxPlaneConstraints(Particle0, Particle1, Implicit0, Implicit1, Thickness, ConstraintBuffer, Manifold);
+		ConstructBoxPlaneConstraints(Particle0, Particle1, Implicit0, Implicit1, Thickness, ConstraintBuffer, History);
 	}
 	else if (Implicit1->GetType() == TBox<T, d>::StaticType() && Implicit0->GetType() == TPlane<T, d>::StaticType())
 	{
-		ConstructBoxPlaneConstraints(Particle1, Particle0, Implicit0, Implicit1, Thickness, ConstraintBuffer, Manifold);
+		ConstructBoxPlaneConstraints(Particle1, Particle0, Implicit0, Implicit1, Thickness, ConstraintBuffer, History);
 	}
 	else if (Implicit0->GetType() == TSphere<T, d>::StaticType() && Implicit1->GetType() == TPlane<T, d>::StaticType())
 	{
-		ConstructSpherePlaneConstraints(Particle0, Particle1, Implicit0, Implicit1, Thickness, ConstraintBuffer, Manifold);
+		ConstructSpherePlaneConstraints(Particle0, Particle1, Implicit0, Implicit1, Thickness, ConstraintBuffer, History);
 	}
 	else if (Implicit1->GetType() == TSphere<T, d>::StaticType() && Implicit0->GetType() == TPlane<T, d>::StaticType())
 	{
-		ConstructSpherePlaneConstraints(Particle1, Particle0, Implicit0, Implicit1, Thickness, ConstraintBuffer, Manifold);
+		ConstructSpherePlaneConstraints(Particle1, Particle0, Implicit0, Implicit1, Thickness, ConstraintBuffer, History);
 	}
 	else if (Implicit0->GetType() == TSphere<T, d>::StaticType() && Implicit1->GetType() == TBox<T, d>::StaticType())
 	{
-		ConstructSphereBoxConstraints(Particle0, Particle1, Implicit0, Implicit1, Thickness, ConstraintBuffer, Manifold);
+		ConstructSphereBoxConstraints(Particle0, Particle1, Implicit0, Implicit1, Thickness, ConstraintBuffer, History);
 	}
 	else if (Implicit1->GetType() == TSphere<T, d>::StaticType() && Implicit0->GetType() == TBox<T, d>::StaticType())
 	{
-		ConstructSphereBoxConstraints(Particle1, Particle0, Implicit0, Implicit1, Thickness, ConstraintBuffer, Manifold);
+		ConstructSphereBoxConstraints(Particle1, Particle0, Implicit0, Implicit1, Thickness, ConstraintBuffer, History);
 	}
 	else if (Implicit0->GetType() == TCapsule<T>::StaticType() && Implicit1->GetType() == TCapsule<T>::StaticType())
 	{
-		ConstructCapsuleCapsuleConstraints(Particle0, Particle1, Implicit0, Implicit1, Thickness, ConstraintBuffer, Manifold);
+		ConstructCapsuleCapsuleConstraints(Particle0, Particle1, Implicit0, Implicit1, Thickness, ConstraintBuffer, History);
 	}
 	else if (Implicit0->GetType() == TCapsule<T>::StaticType() && Implicit1->GetType() == TBox<T, d>::StaticType())
 	{
-		ConstructCapsuleBoxConstraints(Particle0, Particle1, Implicit0, Implicit1, Thickness, ConstraintBuffer, Manifold);
+		ConstructCapsuleBoxConstraints(Particle0, Particle1, Implicit0, Implicit1, Thickness, ConstraintBuffer, History);
 	}
 	else if (Implicit1->GetType() == TCapsule<T>::StaticType() && Implicit0->GetType() == TBox<T, d>::StaticType())
 	{
-		ConstructCapsuleBoxConstraints(Particle1, Particle0, Implicit0, Implicit1, Thickness, ConstraintBuffer, Manifold);
+		ConstructCapsuleBoxConstraints(Particle1, Particle0, Implicit0, Implicit1, Thickness, ConstraintBuffer, History);
 	}
 	else if (Implicit0->GetType() < TImplicitObjectUnion<T, d>::StaticType() && Implicit1->GetType() == TImplicitObjectUnion<T, d>::StaticType())
 	{
-		ConstructSingleUnionConstraints(Particle0, Particle1, Implicit0, Implicit1, Thickness, ConstraintBuffer, Manifold);
+		ConstructSingleUnionConstraints(Particle0, Particle1, Implicit0, Implicit1, Thickness, ConstraintBuffer, History);
 	}
 	else if (Implicit0->GetType() == TImplicitObjectUnion<T, d>::StaticType() && Implicit1->GetType() < TImplicitObjectUnion<T, d>::StaticType())
 	{
-		ConstructSingleUnionConstraints(Particle1, Particle0, Implicit0, Implicit1, Thickness, ConstraintBuffer, Manifold);
+		ConstructSingleUnionConstraints(Particle1, Particle0, Implicit0, Implicit1, Thickness, ConstraintBuffer, History);
 	}
 	else if (Implicit0->GetType() == TImplicitObjectUnion<T, d>::StaticType() && Implicit1->GetType() == TImplicitObjectUnion<T, d>::StaticType())
 	{
-		ConstructUnionUnionConstraints(Particle0, Particle1, Implicit0, Implicit1, Thickness, ConstraintBuffer, Manifold);
+		ConstructUnionUnionConstraints(Particle0, Particle1, Implicit0, Implicit1, Thickness, ConstraintBuffer, History);
 	}
 	else if (Implicit0->IsConvex() && Implicit1->IsConvex())
 	{
-		CollisionResolutionConvexConvex<T,d>::ConstructConvexConvexConstraints(Particle0, Particle1, Implicit0, Implicit1, Thickness, ConstraintBuffer, Manifold, CollisionManifoldSampleScaleCVar, CollisionManifoldSamplesCVar);
+		ConstructConvexConvexConstraints(Particle0, Particle1, Implicit0, Implicit1, Thickness, ConstraintBuffer, History);
 	}
 	else
 	{
-		ConstructLevelsetConstraints(Particle0, Particle1, Implicit0, Implicit1, Thickness, ConstraintBuffer, Manifold);
+		ConstructLevelsetConstraints(Particle0, Particle1, Implicit0, Implicit1, Thickness, ConstraintBuffer, History);
 	}
 }
 
@@ -2118,11 +2252,11 @@ void TPBDCollisionConstraint<T, d>::ConstructConstraints(TGeometryParticleHandle
 {
 	ensure(Particle0 && Particle1);
 
-	FConstraintManifold* ManifoldElement = nullptr;
-	if (FConstraintManifold** Ptr = Manifolds.Find(FConstraintHandleID(Particle0, Particle1)))
-		ManifoldElement = *Ptr;
+	FConstraintHistory* HistoryElement = nullptr;
+	if (FConstraintHistory** Ptr = History.Find(FConstraintHandleID(Particle0, Particle1)))
+		HistoryElement = *Ptr;
 
-	ConstructConstraintsImpl<T, d>(Particle0, Particle1, Particle0->Geometry().Get(), Particle1->Geometry().Get(), Thickness, ConstraintBuffer, ManifoldElement);
+	ConstructConstraintsImpl<T, d>(Particle0, Particle1, Particle0->Geometry().Get(), Particle1->Geometry().Get(), Thickness, ConstraintBuffer, HistoryElement);
 }
 
 // NOTE: UpdateLevelsetConstraintImp and its <float,3> specialization are here for the Linux build. It looks like
@@ -2141,7 +2275,7 @@ void UpdateLevelsetConstraintImp(const float Thickness, TRigidBodyContactConstra
 
 template<ECollisionUpdateType UpdateType, typename T, int d>
 void UpdateConstraintImp2(const FImplicitObject& ParticleObject, const TRigidTransform<T, d>& ParticleTM, const FImplicitObject& LevelsetObject, const TRigidTransform<T, d>& LevelsetTM, const T Thickness,
-	TRigidBodyContactConstraint<T, d>* Constraints, int32 NumConstraints, TCollisionResolutionManifold<T, d>* Manifold)
+	TRigidBodyContactConstraint<T, d>* Constraints, int32 NumConstraints, TPBDCollisionConstraintHistory<T, d>* History)
 {
 	TRigidBodyContactConstraint<T, d>& Constraint = Constraints[0];
 
@@ -2215,7 +2349,7 @@ void UpdateConstraintImp2(const FImplicitObject& ParticleObject, const TRigidTra
 	}
 	else if (ParticleObject.GetType() < TImplicitObjectUnion<T, d>::StaticType() && LevelsetObject.GetType() == TImplicitObjectUnion<T, d>::StaticType())
 	{
-		return UpdateSingleUnionConstraint<UpdateType>(Thickness, Constraints, NumConstraints, Manifold);
+		return UpdateSingleUnionConstraint<UpdateType>(Thickness, Constraints, NumConstraints, History);
 	}
 	else if (ParticleObject.GetType() == TImplicitObjectUnion<T, d>::StaticType() && LevelsetObject.GetType() < TImplicitObjectUnion<T, d>::StaticType())
 	{
@@ -2223,11 +2357,11 @@ void UpdateConstraintImp2(const FImplicitObject& ParticleObject, const TRigidTra
 	}
 	else if (ParticleObject.GetType() == TImplicitObjectUnion<T, d>::StaticType() && LevelsetObject.GetType() == TImplicitObjectUnion<T, d>::StaticType())
 	{
-		return UpdateUnionUnionConstraint<UpdateType>(Thickness, Constraints, NumConstraints, Manifold);
+		return UpdateUnionUnionConstraint<UpdateType>(Thickness, Constraints, NumConstraints, History);
 	}
 	else if (ParticleObject.IsConvex() && LevelsetObject.IsConvex())
 	{
-		CollisionResolutionConvexConvex<T,d>::UpdateConvexConvexConstraint(ParticleObject, ParticleTM, LevelsetObject, LevelsetTM, Thickness, Constraints, NumConstraints, Manifold, CollisionManifoldSampleScaleCVar, CollisionManifoldSamplesCVar);
+		UpdateConvexConvexConstraint<UpdateType>(ParticleObject, ParticleTM, LevelsetObject, LevelsetTM, Thickness, Constraints, NumConstraints, History);
 	}
 	else if (LevelsetObject.IsUnderlyingUnion())
 	{
@@ -2281,13 +2415,13 @@ void TPBDCollisionConstraint<T, d>::UpdateConstraint(const T Thickness, FRigidBo
 	}
 	else
 	{
-		TCollisionResolutionManifold<T, d>* ManifoldData = nullptr;
+		TPBDCollisionConstraintHistory<T, d>* HistoryData = nullptr;
 		FConstraintHandleID HandleID = GetConstraintHandleID(Constraint);
-		if (Manifolds.Contains(HandleID))
+		if (History.Contains(HandleID))
 		{
-			ManifoldData = Manifolds[HandleID];
+			HistoryData = History[HandleID];
 		}
-		UpdateConstraintImp2<UpdateType>(*Constraint.Particle->Geometry(), ParticleTM, *Constraint.Levelset->Geometry(), LevelsetTM, Thickness, InConstraints, NumConstraints, ManifoldData);
+		UpdateConstraintImp2<UpdateType>(*Constraint.Particle->Geometry(), ParticleTM, *Constraint.Levelset->Geometry(), LevelsetTM, Thickness, InConstraints, NumConstraints, HistoryData);
 	}
 }
 
@@ -2300,6 +2434,6 @@ template void TPBDCollisionConstraint<float, 3>::UpdateLevelsetConstraint<EColli
 template void TPBDCollisionConstraint<float, 3>::UpdateLevelsetConstraint<ECollisionUpdateType::Deepest>(const float Thickness, FRigidBodyContactConstraint& Constraint);
 template void TPBDCollisionConstraint<float, 3>::ComputeConstraints<false>(const TPBDCollisionConstraint<float,3>::FAccelerationStructure&, float Dt);
 template void TPBDCollisionConstraint<float, 3>::ComputeConstraints<true>(const TPBDCollisionConstraint<float, 3>::FAccelerationStructure&, float Dt);
-template void UpdateConstraintImp2<ECollisionUpdateType::Any, float, 3>(const FImplicitObject& ParticleObject, const TRigidTransform<float, 3>& ParticleTM, const FImplicitObject& LevelsetObject, const TRigidTransform<float, 3>& LevelsetTM, const float Thickness, TRigidBodyContactConstraint<float, 3>* Constraints, int32 NumConstraints, TCollisionResolutionManifold<float, 3>* Manifold);
-template void UpdateConstraintImp2<ECollisionUpdateType::Deepest, float, 3>(const FImplicitObject& ParticleObject, const TRigidTransform<float, 3>& ParticleTM, const FImplicitObject& LevelsetObject, const TRigidTransform<float, 3>& LevelsetTM, const float Thickness, TRigidBodyContactConstraint<float, 3>* Constraint, int32 NumConstraints, TCollisionResolutionManifold<float, 3>* Manifold);
+template void UpdateConstraintImp2<ECollisionUpdateType::Any, float, 3>(const FImplicitObject& ParticleObject, const TRigidTransform<float, 3>& ParticleTM, const FImplicitObject& LevelsetObject, const TRigidTransform<float, 3>& LevelsetTM, const float Thickness, TRigidBodyContactConstraint<float, 3>* Constraints, int32 NumConstraints, TPBDCollisionConstraintHistory<float, 3>* History);
+template void UpdateConstraintImp2<ECollisionUpdateType::Deepest, float, 3>(const FImplicitObject& ParticleObject, const TRigidTransform<float, 3>& ParticleTM, const FImplicitObject& LevelsetObject, const TRigidTransform<float, 3>& LevelsetTM, const float Thickness, TRigidBodyContactConstraint<float, 3>* Constraint, int32 NumConstraints, TPBDCollisionConstraintHistory<float, 3>* History);
 }
diff --git a/Engine/Source/Runtime/Experimental/Chaos/Public/Chaos/CollisionResolutionConvexConvex.h b/Engine/Source/Runtime/Experimental/Chaos/Public/Chaos/CollisionResolutionConvexConvex.h
deleted file mode 100644
diff --git a/Engine/Source/Runtime/Experimental/Chaos/Public/Chaos/CollisionResolutionManifold.h b/Engine/Source/Runtime/Experimental/Chaos/Public/Chaos/CollisionResolutionManifold.h
deleted file mode 100644
diff --git a/Engine/Source/Runtime/Experimental/Chaos/Public/Chaos/PBDCollisionConstraint.h b/Engine/Source/Runtime/Experimental/Chaos/Public/Chaos/PBDCollisionConstraint.h
--- a/Engine/Source/Runtime/Experimental/Chaos/Public/Chaos/PBDCollisionConstraint.h
+++ b/Engine/Source/Runtime/Experimental/Chaos/Public/Chaos/PBDCollisionConstraint.h
@@ -2,7 +2,6 @@
 #pragma once
 
 #include "Chaos/ConstraintHandle.h"
-#include "Chaos/CollisionResolutionManifold.h"
 #include "Chaos/PBDCollisionTypes.h"
 #include "Chaos/PBDConstraintContainer.h"
 #include "Framework/BufferedData.h"
@@ -53,7 +52,67 @@ enum class ECollisionModifierResult
 	Disabled,	/** Collision should be disabled */
 };
 
+template<class T, int d>
+class CHAOS_API TPBDCollisionConstraintHandle : public TContainerConstraintHandle<TPBDCollisionConstraint<T, d>>
+{
+public:
+	using Base = TContainerConstraintHandle<TPBDCollisionConstraint<T, d>>;
+	using FConstraintContainer = TPBDCollisionConstraint<T, d>;
+
+	TPBDCollisionConstraintHandle();
+	TPBDCollisionConstraintHandle(FConstraintContainer* InConstraintContainer, int32 InConstraintIndex);
+
+	TRigidBodyContactConstraint<T, d>& GetContact();
+	const TRigidBodyContactConstraint<T, d>& GetContact() const;
+	void SetConstraintIndex(int32 IndexIn) { ConstraintIndex = IndexIn; }
+
+protected:
+	using Base::ConstraintIndex;
+	using Base::ConstraintContainer;
+};
+
+template<class T, int d>
+class CHAOS_API TPBDCollisionConstraintHistory
+{
+	typedef TPair<const FImplicitObject*, const FImplicitObject*> ImplicitPairs;
+
+public:
+	using FConstraintContainerHandle = TPBDCollisionConstraintHandle<T, d>;
+	TPBDCollisionConstraintHistory(const TVector<T, d>& InLocation, const TRotation<T, d>& InRotation, int32 InTimestamp = -INT_MAX)
+		: Timestamp(InTimestamp)
+		, Location(InLocation)
+		, Rotation(InRotation)
+	{}
 
+	const TVector<T, d>& GetLocation() { return Location; }
+	const TRotation<T, d>& GetRotation() { return Rotation; }
+
+	void AddHandle(FConstraintContainerHandle* InHandle)
+	{
+		Implicits.Add(ImplicitPairs(InHandle->GetContact().Geometry[0], InHandle->GetContact().Geometry[1]));
+		Handles.Push(InHandle);
+	}
+	void RemoveHandle(FConstraintContainerHandle* InHandle)
+	{
+		Handles.RemoveSingleSwap(InHandle);
+		Implicits.Remove(ImplicitPairs(InHandle->GetContact().Geometry[0], InHandle->GetContact().Geometry[1]));
+	}
+
+	TArray<FConstraintContainerHandle*>& GetHandles() { return Handles; }
+	const TArray<FConstraintContainerHandle*>& GetHandles() const { return Handles; }
+
+	bool ContainsShapeConnection(const FImplicitObject* Implicit0In, const FImplicitObject* Implicit1In) { return Implicits.Contains(ImplicitPairs(Implicit0In, Implicit1In)); }
+
+	int32 GetTimestamp() const { return Timestamp; }
+	void SetTimestamp(int32 InTimestamp) { Timestamp = InTimestamp; }
+
+private:
+	int32 Timestamp;
+	TVector<T, d> Location;
+	TRotation<T, d> Rotation;
+	TArray<FConstraintContainerHandle*> Handles;
+	TSet<ImplicitPairs> Implicits;
+};
 
 template<typename T, int d>
 using TRigidBodyContactConstraintsPostComputeCallback = TFunction<void()>;
@@ -82,7 +141,7 @@ public:
 	using FReal = T;
 	static const int Dimensions = d;
 	using FConstraintContainerHandle = TPBDCollisionConstraintHandle<T, d>;
-	using FConstraintManifold = TCollisionResolutionManifold<T, d>;
+	using FConstraintHistory = TPBDCollisionConstraintHistory<T, d>;
 	using FConstraintHandleAllocator = TConstraintHandleAllocator<TPBDCollisionConstraint<T, d>>;
 	using FRigidBodyContactConstraint = TRigidBodyContactConstraint<T, d>;
 	typedef TPair<TGeometryParticleHandle<T, d>*, const TGeometryParticleHandle<T, d>*> FConstraintHandleID;
@@ -279,7 +338,7 @@ private:
 	bool bEnableCollisions;
 
 	int32 LifespanCounter;
-	int32 CollisionManifoldLifespan;
+	int32 CollisionHistoryLifespan;
 	float CollisionVelocityInflation;
 
 	TRigidBodyContactConstraintsPostComputeCallback<T, d> PostComputeCallback;
@@ -289,7 +348,7 @@ private:
 	TArray<FConstraintContainerHandle*> Handles;
 	FConstraintHandleAllocator HandleAllocator;
 
-	TMap<FConstraintHandleID, FConstraintManifold*> Manifolds;
+	TMap<FConstraintHandleID, FConstraintHistory*> History;
 };
 
 extern template void TPBDCollisionConstraint<float, 3>::UpdateConstraint<ECollisionUpdateType::Any>(const float Thickness, FRigidBodyContactConstraint* Constraints, int32 NumConstraints);
diff --git a/Engine/Source/Runtime/Experimental/Chaos/Public/Chaos/PBDCollisionConstraintImp.h b/Engine/Source/Runtime/Experimental/Chaos/Public/Chaos/PBDCollisionConstraintImp.h
--- a/Engine/Source/Runtime/Experimental/Chaos/Public/Chaos/PBDCollisionConstraintImp.h
+++ b/Engine/Source/Runtime/Experimental/Chaos/Public/Chaos/PBDCollisionConstraintImp.h
@@ -305,17 +305,17 @@ namespace Chaos
 					ensure(Handles.Num() == Constraints.Num());
 
 					FConstraintHandleID HandleID = GetConstraintHandleID(Constraints[Idx]);
-					FConstraintManifold* Manifold;
-					if (auto ValPtr = Manifolds.Find(HandleID))
+					FConstraintHistory* HistoryElement;
+					if (auto ValPtr = History.Find(HandleID))
 					{
-						Manifold = *ValPtr;
+						HistoryElement = *ValPtr;
 					}
 					else
 					{
-						Manifold = new FConstraintManifold(Constraints[Idx].Particle->X(), Constraints[Idx].Particle->R(), LifespanCounter + CollisionManifoldLifespan);
-						Manifolds.Add(HandleID, Manifold);
+						HistoryElement = new FConstraintHistory(Constraints[Idx].Particle->X(), Constraints[Idx].Particle->R(), LifespanCounter + CollisionHistoryLifespan);
+						History.Add(HandleID, HistoryElement);
 					}
-					Manifold->AddHandle(Handle);
+					HistoryElement->AddHandle(Handle);
 				}
 				LifespanCounter++;
 			}
diff --git a/Engine/Source/Runtime/Experimental/Chaos/Public/Chaos/PBDCollisionTypes.h b/Engine/Source/Runtime/Experimental/Chaos/Public/Chaos/PBDCollisionTypes.h
--- a/Engine/Source/Runtime/Experimental/Chaos/Public/Chaos/PBDCollisionTypes.h
+++ b/Engine/Source/Runtime/Experimental/Chaos/Public/Chaos/PBDCollisionTypes.h
@@ -1,18 +1,13 @@
 // Copyright 1998-2019 Epic Games, Inc. All Rights Reserved.
 #pragma once
+#include "Chaos/Vector.h"
 #include "Box.h"
 #include "Chaos/ParticleHandleFwd.h"
-#include "Chaos/PBDConstraintContainer.h"
-#include "Chaos/Vector.h"
 
 class UPhysicalMaterial;
 
 namespace Chaos
 {
-
-	template<typename T, int d>
-	class TPBDCollisionConstraint;
-
 /**/
 // @todo(ccaulfield): should these be held SOA style in the PBDCollisionConstraint?
 template<class T, int d>
@@ -48,45 +43,6 @@ struct TRigidBodyContactConstraintPGS
 	TVector<T, d> AccumulatedImpulse;
 };
 
-
-template<class T, int d>
-class CHAOS_API TPBDCollisionConstraintHandle : public TContainerConstraintHandle<TPBDCollisionConstraint<T, d>>
-{
-public:
-	using Base = TContainerConstraintHandle<TPBDCollisionConstraint<T, d>>;
-	using FConstraintContainer = TPBDCollisionConstraint<T, d>;
-
-	TPBDCollisionConstraintHandle() 
-	{}
-	
-	TPBDCollisionConstraintHandle(FConstraintContainer* InConstraintContainer, int32 InConstraintIndex) 
-		: TContainerConstraintHandle<TPBDCollisionConstraint<T, d>>(InConstraintContainer, InConstraintIndex)
-	{
-	}
-
-	TRigidBodyContactConstraint<T, d>& GetContact()
-	{
-		return ConstraintContainer->Constraints[ConstraintIndex];
-	}
-
-	const TRigidBodyContactConstraint<T, d>& GetContact() const
-	{
-		return ConstraintContainer->Constraints[ConstraintIndex];
-	}
-
-	void SetConstraintIndex(int32 IndexIn) 
-	{ 
-		ConstraintIndex = IndexIn; 
-	}
-
-protected:
-	using Base::ConstraintIndex;
-	using Base::ConstraintContainer;
-
-
-};
-
-
 /**
  * Collision event data stored for use by other systems (e.g. Niagara, gameplay events)
  */
