Move inverse length computation into visitor code. Pass extra information so we can do box culling more quickly

diff --git a/Engine/Source/Runtime/Experimental/Chaos/Private/Chaos/TriangleMeshImplicitObject.cpp b/Engine/Source/Runtime/Experimental/Chaos/Private/Chaos/TriangleMeshImplicitObject.cpp
--- a/Engine/Source/Runtime/Experimental/Chaos/Private/Chaos/TriangleMeshImplicitObject.cpp
+++ b/Engine/Source/Runtime/Experimental/Chaos/Private/Chaos/TriangleMeshImplicitObject.cpp
@@ -42,7 +42,7 @@ struct TTriangleMeshRaycastVisitor
 	};
 
 	template <ERaycastType SQType>
-	bool Visit(int32 TriIdx, T& CurLength)
+	bool Visit(int32 TriIdx, FQueryFastData& CurData)
 	{
 		constexpr T Epsilon = 1e-4;
 		constexpr T Epsilon2 = Epsilon * Epsilon;
@@ -73,7 +73,7 @@ struct TTriangleMeshRaycastVisitor
 
 		//Check if we even intersect with triangle plane
 		int32 DummyFaceIndex;
-		if (TriPlane.Raycast(StartPoint, Dir, CurLength, Thickness, Time, RaycastPosition, RaycastNormal, DummyFaceIndex))
+		if (TriPlane.Raycast(StartPoint, Dir, CurData.CurrentLength, Thickness, Time, RaycastPosition, RaycastNormal, DummyFaceIndex))
 		{
 			TVector<T, 3> IntersectionPosition = RaycastPosition;
 			TVector<T, 3> IntersectionNormal = RaycastNormal;
@@ -108,19 +108,19 @@ struct TTriangleMeshRaycastVisitor
 				{
 					TVector<T, 3> ABCapsuleAxis = B - A;
 					T ABHeight = ABCapsuleAxis.SafeNormalize();
-					bBorderIntersections[0] = TCapsule<T>::RaycastFast(Thickness, ABHeight, ABCapsuleAxis, A, B, StartPoint, Dir, CurLength, 0, BorderTimes[0], BorderPositions[0], BorderNormals[0], DummyFaceIndex);
+					bBorderIntersections[0] = TCapsule<T>::RaycastFast(Thickness, ABHeight, ABCapsuleAxis, A, B, StartPoint, Dir, CurData.CurrentLength, 0, BorderTimes[0], BorderPositions[0], BorderNormals[0], DummyFaceIndex);
 				}
 				
 				{
 					TVector<T, 3> BCCapsuleAxis = C - B;
 					T BCHeight = BCCapsuleAxis.SafeNormalize();
-					bBorderIntersections[1] = TCapsule<T>::RaycastFast(Thickness, BCHeight, BCCapsuleAxis, B, C, StartPoint, Dir, CurLength, 0, BorderTimes[1], BorderPositions[1], BorderNormals[1], DummyFaceIndex);
+					bBorderIntersections[1] = TCapsule<T>::RaycastFast(Thickness, BCHeight, BCCapsuleAxis, B, C, StartPoint, Dir, CurData.CurrentLength, 0, BorderTimes[1], BorderPositions[1], BorderNormals[1], DummyFaceIndex);
 				}
 				
 				{
 					TVector<T, 3> ACCapsuleAxis = C - A;
 					T ACHeight = ACCapsuleAxis.SafeNormalize();
-					bBorderIntersections[2] = TCapsule<T>::RaycastFast(Thickness, ACHeight, ACCapsuleAxis, A, C, StartPoint, Dir, CurLength, 0, BorderTimes[2], BorderPositions[2], BorderNormals[2], DummyFaceIndex);
+					bBorderIntersections[2] = TCapsule<T>::RaycastFast(Thickness, ACHeight, ACCapsuleAxis, A, C, StartPoint, Dir, CurData.CurrentLength, 0, BorderTimes[2], BorderPositions[2], BorderNormals[2], DummyFaceIndex);
 				}
 
 				int32 MinBorderIdx = INDEX_NONE;
@@ -163,7 +163,7 @@ struct TTriangleMeshRaycastVisitor
 					OutPosition = IntersectionPosition;
 					OutNormal = RaycastNormal;	//We use the plane normal even when hitting triangle edges. This is to deal with triangles that approximate a single flat surface.
 					OutTime = Time;
-					CurLength = Time;	//prevent future rays from going any farther
+					CurData.SetLength(Time);	//prevent future rays from going any farther
 					OutFaceIndex = TriIdx;
 				}
 			}
@@ -172,14 +172,14 @@ struct TTriangleMeshRaycastVisitor
 		return true;
 	}
 
-	bool VisitRaycast(TSpatialVisitorData<int32> TriIdx, T& CurLength)
+	bool VisitRaycast(TSpatialVisitorData<int32> TriIdx, FQueryFastData& CurData)
 	{
-		return Visit<ERaycastType::Raycast>(TriIdx.Payload, CurLength);
+		return Visit<ERaycastType::Raycast>(TriIdx.Payload, CurData);
 	}
 
-	bool VisitSweep(TSpatialVisitorData<int32> TriIdx, T& CurLength)
+	bool VisitSweep(TSpatialVisitorData<int32> TriIdx, FQueryFastData& CurData)
 	{
-		return Visit<ERaycastType::Sweep>(TriIdx.Payload, CurLength);
+		return Visit<ERaycastType::Sweep>(TriIdx.Payload, CurData);
 	}
 
 	bool VisitOverlap(TSpatialVisitorData<int32> TriIdx)
@@ -378,13 +378,13 @@ struct TTriangleMeshSweepVisitor
 		return true;
 	}
 
-	bool VisitRaycast(const TSpatialVisitorData<int32>& VisitData, T& CurLength)
+	bool VisitRaycast(const TSpatialVisitorData<int32>& VisitData, FQueryFastData& CurData)
 	{
 		check(false);
 		return true;
 	}
 
-	bool VisitSweep(const TSpatialVisitorData<int32>& VisitData, T& CurLength)
+	bool VisitSweep(const TSpatialVisitorData<int32>& VisitData, FQueryFastData& CurData)
 	{
 		const int32 TriIdx = VisitData.Payload;
 
@@ -396,14 +396,14 @@ struct TTriangleMeshSweepVisitor
 			TriMesh.MParticles.X(TriMesh.MElements[TriIdx][1]),
 			TriMesh.MParticles.X(TriMesh.MElements[TriIdx][2]));
 
-		if(GJKRaycast2<T>(Tri, QueryGeom, StartTM, Dir, CurLength, Time, HitPosition, HitNormal, Thickness))
+		if(GJKRaycast2<T>(Tri, QueryGeom, StartTM, Dir, CurData.CurrentLength, Time, HitPosition, HitNormal, Thickness))
 		{
 			if(Time < OutTime)
 			{
 				OutNormal = HitNormal;
 				OutPosition = HitPosition;
 				OutTime = Time;
-				CurLength = Time;
+				CurData.SetLength(Time);
 				OutFaceIndex = TriIdx;
 
 				if(Time == 0)
diff --git a/Engine/Source/Runtime/Experimental/Chaos/Public/Chaos/AABBTree.h b/Engine/Source/Runtime/Experimental/Chaos/Public/Chaos/AABBTree.h
--- a/Engine/Source/Runtime/Experimental/Chaos/Public/Chaos/AABBTree.h
+++ b/Engine/Source/Runtime/Experimental/Chaos/Public/Chaos/AABBTree.h
@@ -23,41 +23,41 @@ DECLARE_CYCLE_STAT(TEXT("AABBTreeGenerateTree"), STAT_AABBTreeGenerateTree, STAT
 DECLARE_CYCLE_STAT(TEXT("AABBTreeTimeSliceSetup"), STAT_AABBTreeTimeSliceSetup, STATGROUP_Chaos);
 DECLARE_CYCLE_STAT(TEXT("AABBTreeInitialTimeSlice"), STAT_AABBTreeInitialTimeSlice, STATGROUP_Chaos);
 
-template <typename T, EAABBQueryType Query>
+template <typename T, typename TQueryFastData, EAABBQueryType Query>
 struct TAABBTreeIntersectionHelper
 {
-	static bool Intersects(const TVector<T, 3>& Start, const TVector<T, 3>& Dir, const TVector<T, 3>& InvDir, const bool* bParallel,
-		T& CurrentLength, T& InvCurrentLength, T& TOI, TVector<T, 3>& OutPosition,
-		const TBox<T, 3>& Bounds, const TBox<T, 3>& QueryBounds, const TVector<T, 3>& QueryHalfExtents);
+	static bool Intersects(const TVector<T, 3>& Start, TQueryFastData& QueryFastData, T& TOI, TVector<T, 3>& OutPosition,
+		const TBox<T, 3>& Bounds, const TBox<T, 3>& QueryBounds, const TVector<T, 3>& QueryHalfExtents)
+	{
+		check(false);
+		return true;
+	}
 };
 
-template <typename T>
-struct TAABBTreeIntersectionHelper<T, EAABBQueryType::Raycast>
+template<>
+struct TAABBTreeIntersectionHelper<FReal, FQueryFastData, EAABBQueryType::Raycast>
 {
-	FORCEINLINE_DEBUGGABLE static bool Intersects(const TVector<T, 3>& Start, const TVector<T, 3>& Dir, const TVector<T, 3>& InvDir, const bool* bParallel,
-		T& CurrentLength, T& InvCurrentLength, T& TOI, TVector<T, 3>& OutPosition,
-		const TBox<T, 3>& Bounds, const TBox<T, 3>& QueryBounds, const TVector<T, 3>& QueryHalfExtents)
+	FORCEINLINE_DEBUGGABLE static bool Intersects(const FVec3& Start, FQueryFastData& QueryFastData, FReal& TOI, FVec3& OutPosition,
+		const TBox<FReal, 3>& Bounds, const TBox<FReal, 3>& QueryBounds, const TVector<FReal, 3>& QueryHalfExtents)
 	{
-		return TBox<T, 3>::RaycastFast(Bounds.Min(), Bounds.Max(), Start, Dir, InvDir, bParallel, CurrentLength, InvCurrentLength, TOI, OutPosition);
+		return TBox<FReal, 3>::RaycastFast(Bounds.Min(), Bounds.Max(), Start, QueryFastData.Dir, QueryFastData.InvDir, QueryFastData.bParallel, QueryFastData.CurrentLength, QueryFastData.InvCurrentLength, TOI, OutPosition);
 	}
 };
 
 template <typename T>
-struct TAABBTreeIntersectionHelper<T, EAABBQueryType::Sweep>
+struct TAABBTreeIntersectionHelper<T, FQueryFastData, EAABBQueryType::Sweep>
 {
-	FORCEINLINE_DEBUGGABLE static bool Intersects(const TVector<T, 3>& Start, const TVector<T, 3>& Dir, const TVector<T, 3>& InvDir, const bool* bParallel,
-		T& CurrentLength, T& InvCurrentLength, T& TOI, TVector<T, 3>& OutPosition,
-		const TBox<T, 3>& Bounds, const TBox<T, 3>& QueryBounds, const TVector<T, 3>& QueryHalfExtents)
+	FORCEINLINE_DEBUGGABLE static bool Intersects(const TVector<T, 3>& Start, FQueryFastData& QueryFastData,
+		T& TOI, TVector<T, 3>& OutPosition, const TBox<T, 3>& Bounds, const TBox<T, 3>& QueryBounds, const TVector<T, 3>& QueryHalfExtents)
 	{
-		return TBox<T, 3>::RaycastFast(Bounds.Min() - QueryHalfExtents, Bounds.Max() + QueryHalfExtents, Start, Dir, InvDir, bParallel, CurrentLength, InvCurrentLength, TOI, OutPosition);
+		return TBox<T, 3>::RaycastFast(Bounds.Min() - QueryHalfExtents, Bounds.Max() + QueryHalfExtents, Start, QueryFastData.Dir, QueryFastData.InvDir, QueryFastData.bParallel, QueryFastData.CurrentLength, QueryFastData.InvCurrentLength, TOI, OutPosition);
 	}
 };
 
 template <typename T>
-struct TAABBTreeIntersectionHelper<T, EAABBQueryType::Overlap>
+struct TAABBTreeIntersectionHelper<T, FQueryFastDataVoid, EAABBQueryType::Overlap>
 {
-	FORCEINLINE_DEBUGGABLE static bool Intersects(const TVector<T, 3>& Start, const TVector<T, 3>& Dir, const TVector<T, 3>& InvDir, const bool* bParallel,
-		T& CurrentLength, T& InvCurrentLength, T& TOI, TVector<T, 3>& OutPosition,
+	FORCEINLINE_DEBUGGABLE static bool Intersects(const TVector<T, 3>& Start, FQueryFastDataVoid& QueryFastData, T& TOI, TVector<T, 3>& OutPosition,
 		const TBox<T, 3>& Bounds, const TBox<T, 3>& QueryBounds, const TVector<T, 3>& QueryHalfExtents)
 	{
 		return QueryBounds.Intersects(Bounds);
@@ -74,16 +74,16 @@ struct TAABBTreeLeafArray
 	{
 	}
 
-	template <typename TSQVisitor>
-	bool RaycastFast(const TVector<T,3>& Start, const TVector<T,3>& Dir, const TVector<T,3>& InvDir, const bool* bParallel, T& CurrentLength, T& InvCurrentLength, TSQVisitor& Visitor) const
+	template <typename TSQVisitor, typename TQueryFastData>
+	bool RaycastFast(const TVector<T,3>& Start, TQueryFastData& QueryFastData, TSQVisitor& Visitor) const
 	{
-		return RaycastSweepImp</*bSweep=*/false>(Start, Dir, InvDir, bParallel, CurrentLength, InvCurrentLength, TVector<T,3>(), Visitor);
+		return RaycastSweepImp</*bSweep=*/false>(Start, QueryFastData, TVector<T,3>(), Visitor);
 	}
 
-	template <typename TSQVisitor>
-	bool SweepFast(const TVector<T, 3>& Start, const TVector<T, 3>& Dir, const TVector<T, 3>& InvDir, const bool* bParallel, T& CurrentLength, T& InvCurrentLength, const TVector<T,3>& QueryHalfExtents, TSQVisitor& Visitor) const
+	template <typename TSQVisitor, typename TQueryFastData>
+	bool SweepFast(const TVector<T, 3>& Start, TQueryFastData& QueryFastData, const TVector<T,3>& QueryHalfExtents, TSQVisitor& Visitor) const
 	{
-		return RaycastSweepImp</*bSweep=*/true>(Start, Dir, InvDir, bParallel, CurrentLength, InvCurrentLength, QueryHalfExtents, Visitor);
+		return RaycastSweepImp</*bSweep=*/true>(Start, QueryFastData, QueryHalfExtents, Visitor);
 	}
 
 	template <typename TSQVisitor>
@@ -104,24 +104,22 @@ struct TAABBTreeLeafArray
 		return true;
 	}
 
-	template <bool bSweep, typename TSQVisitor>
-	bool RaycastSweepImp(const TVector<T, 3>& Start, const TVector<T, 3>& Dir, const TVector<T, 3>& InvDir, const bool* bParallel, T& CurrentLength, T& InvCurrentLength, const TVector<T, 3>& QueryHalfExtents, TSQVisitor& Visitor) const
+	template <bool bSweep, typename TQueryFastData, typename TSQVisitor>
+	bool RaycastSweepImp(const TVector<T, 3>& Start, TQueryFastData& QueryFastData, const TVector<T, 3>& QueryHalfExtents, TSQVisitor& Visitor) const
 	{
 		TVector<T, 3> TmpPosition;
 		T TOI;
 		for (const auto& Elem : Elems)
 		{
 			const auto& InstanceBounds = Elem.Bounds;
-			if (TAABBTreeIntersectionHelper<T, bSweep ? EAABBQueryType::Sweep : EAABBQueryType::Raycast>::Intersects(Start, Dir, InvDir, bParallel, CurrentLength, InvCurrentLength,
-				TOI, TmpPosition, InstanceBounds, TBox<T,3>(), QueryHalfExtents))
+			if (TAABBTreeIntersectionHelper<T, TQueryFastData, bSweep ? EAABBQueryType::Sweep : EAABBQueryType::Raycast>::Intersects(Start, QueryFastData, TOI, TmpPosition, InstanceBounds, TBox<T,3>(), QueryHalfExtents))
 			{
 				TSpatialVisitorData<TPayloadType> VisitData(Elem.Payload, true, InstanceBounds);
-				const bool bContinue = (bSweep && Visitor.VisitSweep(VisitData, CurrentLength)) || (!bSweep &&  Visitor.VisitRaycast(VisitData, CurrentLength));
+				const bool bContinue = (bSweep && Visitor.VisitSweep(VisitData, QueryFastData)) || (!bSweep &&  Visitor.VisitRaycast(VisitData, QueryFastData));
 				if (!bContinue)
 				{
 					return false;
 				}
-				InvCurrentLength = 1 / CurrentLength;
 			}
 		}
 
@@ -315,25 +313,14 @@ public:
 	template <typename SQVisitor>
 	void Raycast(const TVector<T, 3>& Start, const TVector<T, 3>& Dir, const T Length, SQVisitor& Visitor) const
 	{
-		ensure(Length);
-		T CurLength = Length;
-		bool bParallel[3];
-		TVector<T, 3> InvDir;
-
-		T InvLength = 1 / Length;
-		for (int Axis = 0; Axis < 3; ++Axis)
-		{
-			bParallel[Axis] = Dir[Axis] == 0;
-			InvDir[Axis] = bParallel[Axis] ? 0 : 1 / Dir[Axis];
-		}
-
-		QueryImp<EAABBQueryType::Raycast>(Start, Dir, InvDir, bParallel, CurLength, InvLength, TVector<T,3>(), TBox<T,3>(), Visitor);
+		FQueryFastData QueryFastData(Dir, Length);
+		QueryImp<EAABBQueryType::Raycast>(Start, QueryFastData, TVector<T,3>(), TBox<T,3>(), Visitor);
 	}
 
 	template <typename SQVisitor>
-	bool RaycastFast(const TVector<T, 3>& Start, const TVector<T, 3>& Dir, const TVector<T,3>& InvDir, const bool* bParallel, T& CurLength, T& InvLength, SQVisitor& Visitor) const
+	bool RaycastFast(const TVector<T, 3>& Start, FQueryFastData& CurData, SQVisitor& Visitor) const
 	{
-		return QueryImp<EAABBQueryType::Raycast>(Start, Dir, InvDir, bParallel, CurLength, InvLength, TVector<T, 3>(), TBox<T, 3>(), Visitor);
+		return QueryImp<EAABBQueryType::Raycast>(Start, CurData, TVector<T, 3>(), TBox<T, 3>(), Visitor);
 	}
 
 	void Sweep(const TVector<T, 3>& Start, const TVector<T, 3>& Dir, const T Length, const TVector<T, 3> QueryHalfExtents, ISpatialVisitor<TPayloadType, T>& Visitor) const override
@@ -345,25 +332,14 @@ public:
 	template <typename SQVisitor>
 	void Sweep(const TVector<T, 3>& Start, const TVector<T, 3>& Dir, const T Length, const TVector<T, 3> QueryHalfExtents, SQVisitor& Visitor) const
 	{
-		bool bParallel[3];
-		TVector<T, 3> InvDir;
-
-		ensure(Length);
-		T CurLength = Length;
-		T InvLength = 1 / Length;
-		for (int Axis = 0; Axis < 3; ++Axis)
-		{
-			bParallel[Axis] = Dir[Axis] == 0;
-			InvDir[Axis] = bParallel[Axis] ? 0 : 1 / Dir[Axis];
-		}
-
-		QueryImp<EAABBQueryType::Sweep>(Start, Dir, InvDir, bParallel, CurLength, InvLength, QueryHalfExtents, TBox<T, 3>(), Visitor);
+		FQueryFastData QueryFastData(Dir, Length);
+		QueryImp<EAABBQueryType::Sweep>(Start, QueryFastData, QueryHalfExtents, TBox<T, 3>(), Visitor);
 	}
 
 	template <typename SQVisitor>
-	bool SweepFast(const TVector<T, 3>& Start, const TVector<T, 3>& Dir, const TVector<T,3>& InvDir, const bool* bParallel, T& CurLength, T& InvLength, const TVector<T, 3> QueryHalfExtents, SQVisitor& Visitor) const
+	bool SweepFast(const TVector<T, 3>& Start, FQueryFastData& CurData, const TVector<T, 3> QueryHalfExtents, SQVisitor& Visitor) const
 	{
-		return QueryImp<EAABBQueryType::Sweep>(Start, Dir, InvDir, bParallel, CurLength, InvLength, QueryHalfExtents, TBox<T, 3>(), Visitor);
+		return QueryImp<EAABBQueryType::Sweep>(Start,CurData, QueryHalfExtents, TBox<T, 3>(), Visitor);
 	}
 
 	void Overlap(const TBox<T, 3>& QueryBounds, ISpatialVisitor<TPayloadType, T>& Visitor) const override
@@ -382,11 +358,8 @@ public:
 	bool OverlapFast(const TBox<T, 3>& QueryBounds, SQVisitor& Visitor) const
 	{
 		//dummy variables to reuse templated path
-		T Length;
-		T InvLength;
-		TVector<T, 3> InvDir;
-		bool bParallel;
-		return QueryImp<EAABBQueryType::Overlap>(TVector<T, 3>(), TVector<T, 3>(), InvDir, &bParallel, Length, InvLength, TVector<T, 3>(), QueryBounds, Visitor);
+		FQueryFastDataVoid VoidData;
+		return QueryImp<EAABBQueryType::Overlap>(TVector<T, 3>(), VoidData, TVector<T, 3>(), QueryBounds, Visitor);
 	}
 
 	virtual void RemoveElement(const TPayloadType& Payload)
@@ -511,8 +484,8 @@ private:
 	using FElement = TPayloadBoundsElement<TPayloadType, T>;
 	using FNode = TAABBTreeNode<T>;
 
-	template <EAABBQueryType Query, typename SQVisitor>
-	bool QueryImp(const TVector<T, 3>& Start, const TVector<T, 3>& Dir, const TVector<T,3>& InvDir, const bool* bParallel, T& CurrentLength, T& InvCurrentLength, const TVector<T, 3> QueryHalfExtents, const TBox<T,3>& QueryBounds, SQVisitor& Visitor) const
+	template <EAABBQueryType Query, typename TQueryFastData, typename SQVisitor>
+	bool QueryImp(const TVector<T, 3>& Start, TQueryFastData& CurData, const TVector<T, 3> QueryHalfExtents, const TBox<T,3>& QueryBounds, SQVisitor& Visitor) const
 	{
 		TVector<T, 3> TmpPosition;
 		T TOI = 0;
@@ -520,8 +493,7 @@ private:
 		for (const auto& Elem : GlobalPayloads)
 		{
 			const auto& InstanceBounds = Elem.Bounds;
-			if (TAABBTreeIntersectionHelper<T, Query>::Intersects(Start, Dir, InvDir, bParallel, CurrentLength, InvCurrentLength,
-				TOI, TmpPosition, InstanceBounds, QueryBounds, QueryHalfExtents))
+			if (TAABBTreeIntersectionHelper<T, TQueryFastData, Query>::Intersects(Start, CurData, TOI, TmpPosition, InstanceBounds, QueryBounds, QueryHalfExtents))
 			{
 				TSpatialVisitorData<TPayloadType> VisitData(Elem.Payload, true);
 				bool bContinue;
@@ -531,26 +503,20 @@ private:
 				}
 				else
 				{
-					bContinue = Query == EAABBQueryType::Sweep ? Visitor.VisitSweep(VisitData, CurrentLength) : Visitor.VisitRaycast(VisitData, CurrentLength);
+					bContinue = Query == EAABBQueryType::Sweep ? Visitor.VisitSweep(VisitData, CurData) : Visitor.VisitRaycast(VisitData, CurData);
 				}
 
 				if (!bContinue)
 				{
 					return false;
 				}
-
-				if (Query != EAABBQueryType::Overlap)
-				{
-					InvCurrentLength = 1 / CurrentLength;
-				}
 			}
 		}
 
 		for (const auto& Elem : DirtyElements)
 		{
 			const auto& InstanceBounds = Elem.Bounds;
-			if(TAABBTreeIntersectionHelper<T, Query>::Intersects(Start, Dir, InvDir, bParallel, CurrentLength, InvCurrentLength,
-				TOI, TmpPosition, InstanceBounds, QueryBounds, QueryHalfExtents))
+			if(TAABBTreeIntersectionHelper<T, TQueryFastData, Query>::Intersects(Start, CurData, TOI, TmpPosition, InstanceBounds, QueryBounds, QueryHalfExtents))
 			{
 				TSpatialVisitorData<TPayloadType> VisitData(Elem.Payload, true, InstanceBounds);
 				
@@ -561,18 +527,13 @@ private:
 				}
 				else
 				{
-					bContinue = Query == EAABBQueryType::Sweep ? Visitor.VisitSweep(VisitData, CurrentLength) : Visitor.VisitRaycast(VisitData, CurrentLength);
+					bContinue = Query == EAABBQueryType::Sweep ? Visitor.VisitSweep(VisitData, CurData) : Visitor.VisitRaycast(VisitData, CurData);
 				}
 				
 				if (!bContinue)
 				{
 					return false;
 				}
-
-				if (Query != EAABBQueryType::Overlap)
-				{
-					InvCurrentLength = 1 / CurrentLength;
-				}
 			}
 		}
 
@@ -589,7 +550,7 @@ private:
 			const FNodeQueueEntry NodeEntry = NodeStack.Pop(false);
 			if (Query != EAABBQueryType::Overlap)
 			{
-				if (NodeEntry.TOI > CurrentLength)
+				if (NodeEntry.TOI > CurData.CurrentLength)
 				{
 					continue;
 				}
@@ -608,12 +569,12 @@ private:
 				}
 				else if (Query == EAABBQueryType::Sweep)
 				{
-					if (Leaf.SweepFast(Start, Dir, InvDir, bParallel, CurrentLength, InvCurrentLength, QueryHalfExtents, Visitor) == false)
+					if (Leaf.SweepFast(Start, CurData, QueryHalfExtents, Visitor) == false)
 					{
 						return false;
 					}
 				}
-				else if (Leaf.RaycastFast(Start, Dir, InvDir, bParallel, CurrentLength, InvCurrentLength, Visitor) == false)
+				else if (Leaf.RaycastFast(Start, CurData, Visitor) == false)
 				{
 					return false;
 				}
@@ -623,8 +584,7 @@ private:
 				int32 Idx = 0;
 				for (const TBox<T, 3>& AABB : Node.ChildrenBounds)
 				{
-					if(TAABBTreeIntersectionHelper<T, Query>::Intersects(Start, Dir, InvDir, bParallel, CurrentLength, InvCurrentLength,
-						TOI, TmpPosition, AABB, QueryBounds, QueryHalfExtents))
+					if(TAABBTreeIntersectionHelper<T, TQueryFastData, Query>::Intersects(Start, CurData, TOI, TmpPosition, AABB, QueryBounds, QueryHalfExtents))
 					{
 						NodeStack.Add(FNodeQueueEntry{ Node.ChildrenNodes[Idx], TOI });
 					}
@@ -853,12 +813,12 @@ private:
 				CollectedResults.Add(Instance.Payload);
 				return true;
 			}
-			bool VisitSweep(const TSpatialVisitorData<TPayloadType>& Instance, T& CurLength)
+			bool VisitSweep(const TSpatialVisitorData<TPayloadType>& Instance, FQueryFastData& CurData)
 			{
 				check(false);
 				return true;
 			}
-			bool VisitRaycast(const TSpatialVisitorData<TPayloadType>& Instance, T& CurLength)
+			bool VisitRaycast(const TSpatialVisitorData<TPayloadType>& Instance, FQueryFastData& CurData)
 			{
 				check(false);
 				return true;
diff --git a/Engine/Source/Runtime/Experimental/Chaos/Public/Chaos/BoundingVolume.h b/Engine/Source/Runtime/Experimental/Chaos/Public/Chaos/BoundingVolume.h
--- a/Engine/Source/Runtime/Experimental/Chaos/Public/Chaos/BoundingVolume.h
+++ b/Engine/Source/Runtime/Experimental/Chaos/Public/Chaos/BoundingVolume.h
@@ -263,27 +263,16 @@ public:
 	}
 
 	template <typename SQVisitor>
-	bool RaycastFast(const TVector<T,d>& Start, const TVector<T,d>& Dir, const TVector<T,d>& InvDir, const bool* bParallel, T& CurrentLength, T& InvCurrentLength, SQVisitor& Visitor) const
+	bool RaycastFast(const TVector<T,d>& Start, FQueryFastData& CurData, SQVisitor& Visitor) const
 	{
-		return RaycastImp(Start, Dir, InvDir, bParallel, CurrentLength, InvCurrentLength, Visitor);
+		return RaycastImp(Start, CurData, Visitor);
 	}
 
 	template <typename SQVisitor, bool bPruneDuplicates = true>
 	void Raycast(const TVector<T, d>& Start, const TVector<T, d>& Dir, const T Length, SQVisitor& Visitor) const
 	{
-		ensure(Length);
-		T CurLength = Length;
-		bool bParallel[d];
-		TVector<T, d> InvDir;
-
-		T InvLength = 1 / Length;
-		for (int Axis = 0; Axis < d; ++Axis)
-		{
-			bParallel[Axis] = Dir[Axis] == 0;
-			InvDir[Axis] = bParallel[Axis] ? 0 : 1 / Dir[Axis];
-		}
-
-		RaycastImp<SQVisitor, bPruneDuplicates>(Start, Dir, InvDir, bParallel, CurLength, InvLength, Visitor);
+		FQueryFastData CurData(Dir, Length); 
+		RaycastImp<SQVisitor, bPruneDuplicates>(Start, CurData, Visitor);
 	}
 
 	void Sweep(const TVector<T, d>& Start, const TVector<T, d>& Dir, const T Length, const TVector<T, d> QueryHalfExtents, ISpatialVisitor<TPayloadType, T>& Visitor) const override
@@ -293,27 +282,16 @@ public:
 	}
 
 	template <typename SQVisitor>
-	bool SweepFast(const TVector<T,d>& Start, const TVector<T,d>& Dir, const TVector<T,d>& InvDir, const bool* bParallel, T& CurrentLength, T& InvCurrentLength, const TVector<T,d>& QueryHalfExtents, SQVisitor& Visitor) const
+	bool SweepFast(const TVector<T,d>& Start, FQueryFastData& CurData, const TVector<T,d>& QueryHalfExtents, SQVisitor& Visitor) const
 	{
-		return SweepImp(Start, Dir, InvDir, bParallel, CurrentLength, InvCurrentLength, QueryHalfExtents, Visitor);
+		return SweepImp(Start, CurData, QueryHalfExtents, Visitor);
 	}
 
 	template <typename SQVisitor, bool bPruneDuplicates = true>
 	void Sweep(const TVector<T, d>& Start, const TVector<T, d>& Dir, const T Length, const TVector<T, d> QueryHalfExtents, SQVisitor& Visitor) const
 	{
-		ensure(Length);
-		T CurLength = Length;
-		bool bParallel[d];
-		TVector<T, d> InvDir;
-
-		T InvLength = 1 / Length;
-		for (int Axis = 0; Axis < d; ++Axis)
-		{
-			bParallel[Axis] = Dir[Axis] == 0;
-			InvDir[Axis] = bParallel[Axis] ? 0 : 1 / Dir[Axis];
-		}
-
-		SweepImp<SQVisitor, bPruneDuplicates>(Start, Dir, InvDir, bParallel, CurLength, InvLength, QueryHalfExtents, Visitor);
+		FQueryFastData CurData(Dir, Length);
+		SweepImp<SQVisitor, bPruneDuplicates>(Start, CurData, QueryHalfExtents, Visitor);
 	}
 
 	virtual void Overlap(const TBox<T, d>& QueryBounds, ISpatialVisitor<TPayloadType, T>& Visitor) const override
@@ -380,7 +358,7 @@ private:
 	using FPayloadInfo = TBVPayloadInfo<T, d>;
 
 	template <typename SQVisitor, bool bPruneDuplicates = true>
-	bool RaycastImp(const TVector<T, d>& Start, const TVector<T, d>& Dir, const TVector<T, d>& InvDir, const bool* bParallel, T& CurrentLength, T& InvCurrentLength, SQVisitor& Visitor) const
+	bool RaycastImp(const TVector<T, d>& Start, FQueryFastData& CurData, SQVisitor& Visitor) const
 	{
 		TVector<T, d> TmpPosition;
 		T TOI;
@@ -388,11 +366,11 @@ private:
 		for (const auto& Elem : MGlobalPayloads)
 		{
 			const auto& InstanceBounds = Elem.Bounds;
-			if (TBox<T, d>::RaycastFast(InstanceBounds.Min(), InstanceBounds.Max(), Start, Dir, InvDir, bParallel, CurrentLength, InvCurrentLength, TOI, TmpPosition))
+			if (TBox<T, d>::RaycastFast(InstanceBounds.Min(), InstanceBounds.Max(), Start,
+				CurData.Dir, CurData.InvDir, CurData.bParallel, CurData.CurrentLength, CurData.InvCurrentLength, TOI, TmpPosition))
 			{
 				TSpatialVisitorData<TPayloadType> VisitData(Elem.Payload, true, InstanceBounds);
-				const bool bContinue = Visitor.VisitRaycast(VisitData, CurrentLength);
-				InvCurrentLength = 1 / CurrentLength;
+				const bool bContinue = Visitor.VisitRaycast(VisitData, CurData);
 
 				if (!bContinue)
 				{
@@ -404,15 +382,15 @@ private:
 		for (const auto& Elem : MDirtyElements)
 		{
 			const auto& InstanceBounds = Elem.Bounds;
-			if (TBox<T, d>::RaycastFast(InstanceBounds.Min(), InstanceBounds.Max(), Start, Dir, InvDir, bParallel, CurrentLength, InvCurrentLength, TOI, TmpPosition))
+			if (TBox<T, d>::RaycastFast(InstanceBounds.Min(), InstanceBounds.Max(), Start, CurData.Dir,
+				CurData.InvDir, CurData.bParallel, CurData.CurrentLength, CurData.InvCurrentLength, TOI, TmpPosition))
 			{
 				TSpatialVisitorData<TPayloadType> VisitData(Elem.Payload, true, InstanceBounds);
-				const bool bContinue = Visitor.VisitRaycast(VisitData, CurrentLength);
+				const bool bContinue = Visitor.VisitRaycast(VisitData, CurData);
 				if (!bContinue)
 				{
 					return false;
 				}
-				InvCurrentLength = 1 / CurrentLength;
 			}
 		}
 
@@ -424,7 +402,8 @@ private:
 		TBox<T, d> GlobalBounds(MGrid.MinCorner(), MGrid.MaxCorner());
 		TVector<T, d> NextStart;
 		TVector<int32, d> CellIdx;
-		bool bCellsLeft = MElements.Num() && TBox<T, d>::RaycastFast(GlobalBounds.Min(), GlobalBounds.Max(), Start, Dir, InvDir, bParallel, CurrentLength, InvCurrentLength, TOI, NextStart);
+		bool bCellsLeft = MElements.Num() && TBox<T, d>::RaycastFast(GlobalBounds.Min(), GlobalBounds.Max(), Start, CurData.Dir,
+			CurData.InvDir, CurData.bParallel, CurData.CurrentLength, CurData.InvCurrentLength, TOI, NextStart);
 		if (bCellsLeft)
 		{
 			CellIdx = MGrid.Cell(NextStart);
@@ -466,15 +445,15 @@ private:
 						}
 					}
 					const auto& InstanceBounds = Elem.Bounds;
-					if (TBox<T, d>::RaycastFast(InstanceBounds.Min(), InstanceBounds.Max(), Start, Dir, InvDir, bParallel, CurrentLength, InvCurrentLength, TOI, TmpPosition))
+					if (TBox<T, d>::RaycastFast(InstanceBounds.Min(), InstanceBounds.Max(), Start,
+						CurData.Dir, CurData.InvDir, CurData.bParallel, CurData.CurrentLength, CurData.InvCurrentLength, TOI, TmpPosition))
 					{
 						TSpatialVisitorData<TPayloadType> VisitData(Elem.Payload, true, InstanceBounds);
-						const bool bContinue = Visitor.VisitRaycast(VisitData, CurrentLength);
+						const bool bContinue = Visitor.VisitRaycast(VisitData, CurData);
 						if (!bContinue)
 						{
 							return false;
 						}
-						InvCurrentLength = 1 / CurrentLength;	//todo: pass into visitor
 					}
 				}
 
@@ -488,15 +467,15 @@ private:
 				const TVector<T, d>& Dx = MGrid.Dx();
 
 				T Times[3];
-				T BestTime = CurrentLength;
+				T BestTime = CurData.CurrentLength;
 				bool bTerminate = true;
 				for (int Axis = 0; Axis < d; ++Axis)
 				{
-					if (!bParallel[Axis])
+					if (!CurData.bParallel[Axis])
 					{
-						const T CrossPoint = Dir[Axis] > 0 ? CellCenter[Axis] + Dx[Axis] / 2 : CellCenter[Axis] - Dx[Axis] / 2;
+						const T CrossPoint = CurData.Dir[Axis] > 0 ? CellCenter[Axis] + Dx[Axis] / 2 : CellCenter[Axis] - Dx[Axis] / 2;
 						const T Distance = CrossPoint - NextStart[Axis];	//note: CellCenter already has /2, we probably want to use the corner instead
-						const T Time = Distance * InvDir[Axis];
+						const T Time = Distance * CurData.InvDir[Axis];
 						Times[Axis] = Time;
 						if (Time < BestTime)
 						{
@@ -518,14 +497,14 @@ private:
 				for (int Axis = 0; Axis < d; ++Axis)
 				{
 					constexpr T Epsilon = 1e-2;	//if raycast is slightly off we still count it as hitting the cell surface
-					CellIdx[Axis] += (Times[Axis] <= BestTime + Epsilon) ? (Dir[Axis] > 0 ? 1 : -1) : 0;
+					CellIdx[Axis] += (Times[Axis] <= BestTime + Epsilon) ? (CurData.Dir[Axis] > 0 ? 1 : -1) : 0;
 					if (CellIdx[Axis] < 0 || CellIdx[Axis] >= MGrid.Counts()[Axis])
 					{
 						return true;
 					}
 				}
 
-				NextStart = NextStart + Dir * BestTime;
+				NextStart = NextStart + CurData.Dir * BestTime;
 			} while (true);
 		}
 
@@ -533,7 +512,7 @@ private:
 	}
 
 	template <typename SQVisitor, bool bPruneDuplicates = true>
-	bool SweepImp(const TVector<T, d>& Start, const TVector<T, d>& Dir, const TVector<T,d>& InvDir, const bool* bParallel, T& CurrentLength, T& InvCurrentLength, const TVector<T, d> QueryHalfExtents, SQVisitor& Visitor) const
+	bool SweepImp(const TVector<T, d>& Start, FQueryFastData& CurData, const TVector<T, d> QueryHalfExtents, SQVisitor& Visitor) const
 	{
 		T TOI = 0;
 		for (const auto& Elem : MGlobalPayloads)
@@ -542,15 +521,14 @@ private:
 			const TVector<T, d> Min = InstanceBounds.Min() - QueryHalfExtents;
 			const TVector<T, d> Max = InstanceBounds.Max() + QueryHalfExtents;
 			TVector<T, d> TmpPosition;
-			if (TBox<T, d>::RaycastFast(Min, Max, Start, Dir, InvDir, bParallel, CurrentLength, InvCurrentLength, TOI, TmpPosition))
+			if (TBox<T, d>::RaycastFast(Min, Max, Start, CurData.Dir, CurData.InvDir, CurData.bParallel, CurData.CurrentLength, CurData.InvCurrentLength, TOI, TmpPosition))
 			{
 				TSpatialVisitorData<TPayloadType> VisitData(Elem.Payload, true, InstanceBounds);
-				const bool bContinue = Visitor.VisitSweep(VisitData, CurrentLength);
+				const bool bContinue = Visitor.VisitSweep(VisitData, CurData);
 				if (!bContinue)
 				{
 					return false;
 				}
-				InvCurrentLength = 1 / CurrentLength;
 			}
 		}
 
@@ -560,15 +538,14 @@ private:
 			const TVector<T, d> Min = InstanceBounds.Min() - QueryHalfExtents;
 			const TVector<T, d> Max = InstanceBounds.Max() + QueryHalfExtents;
 			TVector<T, d> TmpPosition;
-			if (TBox<T, d>::RaycastFast(Min, Max, Start, Dir, InvDir, bParallel, CurrentLength, InvCurrentLength, TOI, TmpPosition))
+			if (TBox<T, d>::RaycastFast(Min, Max, Start, CurData.Dir, CurData.InvDir, CurData.bParallel, CurData.CurrentLength, CurData.InvCurrentLength, TOI, TmpPosition))
 			{
 				TSpatialVisitorData<TPayloadType> VisitData(Elem.Payload, true, InstanceBounds);
-				const bool bContinue = Visitor.VisitSweep(VisitData, CurrentLength);
+				const bool bContinue = Visitor.VisitSweep(VisitData, CurData);
 				if (!bContinue)
 				{
 					return false;
 				}
-				InvCurrentLength = 1 / CurrentLength;
 			}
 		}
 
@@ -591,7 +568,7 @@ private:
 
 		if (StartMinIndex == StartMaxIndex)
 		{
-			const TVector<T, d> End = Start + CurrentLength * Dir;
+			const TVector<T, d> End = Start + CurData.CurrentLength * CurData.Dir;
 			const TVector<int32, d> EndMinIndex = MGrid.ClampIndex(MGrid.Cell(End  - QueryHalfExtents));
 			const TVector<int32, d> EndMaxIndex = MGrid.ClampIndex(MGrid.Cell(End + QueryHalfExtents));
 			if (StartMinIndex == EndMinIndex && StartMinIndex == EndMaxIndex)
@@ -604,15 +581,14 @@ private:
 					const TBox<T, d>& InstanceBounds = Elem.Bounds;
 					const TVector<T, d> Min = InstanceBounds.Min() - QueryHalfExtents;
 					const TVector<T, d> Max = InstanceBounds.Max() + QueryHalfExtents;
-					if (TBox<T, d>::RaycastFast(Min, Max, Start, Dir, InvDir, bParallel, CurrentLength, InvCurrentLength, TOI, TmpPoint))
+					if (TBox<T, d>::RaycastFast(Min, Max, Start, CurData.Dir, CurData.InvDir, CurData.bParallel, CurData.CurrentLength, CurData.InvCurrentLength, TOI, TmpPoint))
 					{
 						TSpatialVisitorData<TPayloadType> VisitData(Elem.Payload, true, InstanceBounds);
-						const bool bContinue = Visitor.VisitSweep(VisitData, CurrentLength);
+						const bool bContinue = Visitor.VisitSweep(VisitData, CurData);
 						if (!bContinue)
 						{
 							return false;
 						}
-						InvCurrentLength = 1 / CurrentLength;
 					}
 				}
 
@@ -638,7 +614,8 @@ private:
 		}
 
 		TVector<T, d> HitPoint;
-		const bool bInitialHit = TBox<T, d>::RaycastFast(GlobalBounds.Min(), GlobalBounds.Max(), Start, Dir, InvDir, bParallel, CurrentLength, InvCurrentLength, TOI, HitPoint);
+		const bool bInitialHit = TBox<T, d>::RaycastFast(GlobalBounds.Min(), GlobalBounds.Max(),
+			Start, CurData.Dir, CurData.InvDir, CurData.bParallel, CurData.CurrentLength, CurData.InvCurrentLength, TOI, HitPoint);
 		if (bInitialHit)	//NOTE: it's possible to have a non empty IdxsQueue and bInitialHit be false. This is because the IdxsQueue works off clamped cells which we can skip
 		{
 			//Flood fill from inflated cell so that we get all cells along the ray
@@ -657,7 +634,7 @@ private:
 			while (QueueIdx < IdxsQueue.Num())
 			{
 				const FCellIntersection CellIntersection = IdxsQueue[QueueIdx++];
-				if (CellIntersection.TOI > CurrentLength)
+				if (CellIntersection.TOI > CurData.CurrentLength)
 				{
 					continue;
 				}
@@ -701,7 +678,7 @@ private:
 						const TVector<T, d> NeighborCenter = MGrid.Location(NeighborIdx);
 						const TVector<T, d> Min = NeighborCenter - QueryHalfExtents - HalfDx;
 						const TVector<T, d> Max = NeighborCenter + QueryHalfExtents + HalfDx;
-						if (TBox<T, d>::RaycastFast(Min, Max, Start, Dir, InvDir, bParallel, CurrentLength, InvCurrentLength, TOI, HitPoint))
+						if (TBox<T, d>::RaycastFast(Min, Max, Start, CurData.Dir, CurData.InvDir, CurData.bParallel, CurData.CurrentLength, CurData.InvCurrentLength, TOI, HitPoint))
 						{
 							IdxsQueue.Add({ NeighborIdx, TOI });	//should we sort by TOI?
 						}
@@ -742,15 +719,15 @@ private:
 					const TBox<T, d>& InstanceBounds = Elem.Bounds;
 					const TVector<T, d> Min = InstanceBounds.Min() - QueryHalfExtents;
 					const TVector<T, d> Max = InstanceBounds.Max() + QueryHalfExtents;
-					if (TBox<T, d>::RaycastFast(Min, Max, Start, Dir, InvDir, bParallel, CurrentLength, InvCurrentLength, TOI, HitPoint))
+					if (TBox<T, d>::RaycastFast(Min, Max, Start,
+						CurData.Dir, CurData.InvDir, CurData.bParallel, CurData.CurrentLength, CurData.InvCurrentLength, TOI, HitPoint))
 					{
 						TSpatialVisitorData<TPayloadType> VisitData(Elem.Payload, true, InstanceBounds);
-						const bool bContinue = Visitor.VisitSweep(VisitData, CurrentLength);
+						const bool bContinue = Visitor.VisitSweep(VisitData, CurData);
 						if (!bContinue)
 						{
 							return false;
 						}
-						InvCurrentLength = 1 / CurrentLength;
 					}
 				}
 				CellsVisited.Add(CellIntersection.CellIdx);
diff --git a/Engine/Source/Runtime/Experimental/Chaos/Public/Chaos/ISpatialAcceleration.h b/Engine/Source/Runtime/Experimental/Chaos/Public/Chaos/ISpatialAcceleration.h
--- a/Engine/Source/Runtime/Experimental/Chaos/Public/Chaos/ISpatialAcceleration.h
+++ b/Engine/Source/Runtime/Experimental/Chaos/Public/Chaos/ISpatialAcceleration.h
@@ -7,6 +7,49 @@
 namespace Chaos
 {
 
+struct CHAOS_API FQueryFastData
+{
+	FQueryFastData(const FVec3& InDir, const FReal InLength)
+		: bParallel{ InDir[0] == 0, InDir[1] == 0, InDir[2] == 0 }
+		, Dir(InDir)
+		, InvDir(bParallel[0] ? 0 : 1 / Dir[0], bParallel[1] ? 0 : 1 / Dir[1], bParallel[2] ? 0 : 1 / Dir[2])
+	{
+		ensure(InLength);
+		SetLength(InLength);
+	}
+
+	const bool bParallel[3];
+	const FVec3& Dir;
+	const FVec3 InvDir;
+
+	FReal CurrentLength;
+	FReal InvCurrentLength;
+
+	void SetLength(const FReal InLength)
+	{
+		CurrentLength = InLength;
+		if(InLength)
+		{
+			InvCurrentLength = 1 / InLength;
+		}
+	}
+
+protected:
+	FQueryFastData(const FVec3& InDummyDir)
+		: bParallel{}
+		, Dir(InDummyDir)
+		, InvDir()
+	{}
+};
+
+//dummy struct for templatized paths
+struct CHAOS_API FQueryFastDataVoid : public FQueryFastData
+{
+	FQueryFastDataVoid() : FQueryFastData(DummyDir) {}
+	
+	FVec3 DummyDir;
+};
+
 template <typename T, int d>
 class TBox;
 
@@ -69,17 +112,17 @@ public:
 
 	/** Called whenever an instance in the acceleration structure may intersect with a raycast
 		@Instance - the instance we are potentially intersecting with a raycast
-		@CurLength - the length all future intersection tests will use. A blocking intersection should update this
+		@CurData - the current query data. Call SetLength to update the length all future intersection tests will use. A blocking intersection should update this
 		Returns true to continue iterating through the acceleration structure
 	*/
-	virtual bool Raycast(const TSpatialVisitorData<TPayloadType>& Instance, T& CurLength) = 0;
+	virtual bool Raycast(const TSpatialVisitorData<TPayloadType>& Instance, FQueryFastData& CurData) = 0;
 
 	/** Called whenever an instance in the acceleration structure may intersect with a sweep
 		@Instance - the instance we are potentially intersecting with a sweep
 		@CurLength - the length all future intersection tests will use. A blocking intersection should update this
 		Returns true to continue iterating through the acceleration structure
 	*/
-	virtual bool Sweep(const TSpatialVisitorData<TPayloadType>& Instance, T& CurLength) = 0;
+	virtual bool Sweep(const TSpatialVisitorData<TPayloadType>& Instance, FQueryFastData& CurData) = 0;
 };
 
 /**
@@ -264,14 +307,14 @@ public:
 		return Visitor.Overlap(Instance);
 	}
 
-	FORCEINLINE bool VisitRaycast(const TSpatialVisitorData<TPayloadType>& Instance, T& CurLength)
+	FORCEINLINE bool VisitRaycast(const TSpatialVisitorData<TPayloadType>& Instance, FQueryFastData& CurData)
 	{
-		return Visitor.Raycast(Instance, CurLength);
+		return Visitor.Raycast(Instance, CurData);
 	}
 
-	FORCEINLINE bool VisitSweep(const TSpatialVisitorData<TPayloadType>& Instance, T& CurLength)
+	FORCEINLINE bool VisitSweep(const TSpatialVisitorData<TPayloadType>& Instance, FQueryFastData& CurData)
 	{
-		return Visitor.Sweep(Instance, CurLength);
+		return Visitor.Sweep(Instance, CurData);
 	}
 
 private:
diff --git a/Engine/Source/Runtime/Experimental/Chaos/Public/Chaos/ISpatialAccelerationCollection.h b/Engine/Source/Runtime/Experimental/Chaos/Public/Chaos/ISpatialAccelerationCollection.h
--- a/Engine/Source/Runtime/Experimental/Chaos/Public/Chaos/ISpatialAccelerationCollection.h
+++ b/Engine/Source/Runtime/Experimental/Chaos/Public/Chaos/ISpatialAccelerationCollection.h
@@ -211,12 +211,12 @@ template <int TypeIdx, int NumTypes, typename Tuple, typename TPayloadType, type
 struct TSpatialAccelerationCollectionHelper
 {
 	template <typename SQVisitor>
-	static bool RaycastFast(const Tuple& Types, const TVector<T, d>& Start, const TVector<T, d>& Dir, const TVector<T,d>& InvDir, const bool* bParallel, T& CurLength, T& InvCurLength, SQVisitor& Visitor)
+	static bool RaycastFast(const Tuple& Types, const TVector<T, d>& Start, FQueryFastData& CurData, SQVisitor& Visitor)
 	{
 		const auto& Accelerations = GetAccelerationsPerType<TypeIdx>(Types).Objects;
 		for (const auto& Accelerator : Accelerations)
 		{
-			if (Accelerator && !Accelerator->RaycastFast(Start, Dir, InvDir, bParallel, CurLength, InvCurLength, Visitor))
+			if (Accelerator && !Accelerator->RaycastFast(Start, CurData, Visitor))
 			{
 				return false;
 			}
@@ -225,19 +225,19 @@ struct TSpatialAccelerationCollectionHelper
 		constexpr int NextType = TypeIdx + 1;
 		if (NextType < NumTypes)
 		{
-			return TSpatialAccelerationCollectionHelper<NextType < NumTypes ? NextType : 0, NumTypes, Tuple, TPayloadType, T, d>::RaycastFast(Types, Start, Dir, InvDir, bParallel, CurLength, InvCurLength, Visitor);
+			return TSpatialAccelerationCollectionHelper<NextType < NumTypes ? NextType : 0, NumTypes, Tuple, TPayloadType, T, d>::RaycastFast(Types, Start, CurData, Visitor);
 		}
 
 		return true;
 	}
 
 	template <typename SQVisitor>
-	static bool SweepFast(const Tuple& Types, const TVector<T, d>& Start, const TVector<T, d>& Dir, const TVector<T,d>& InvDir, const bool* bParallel, T& CurLength, T& InvCurLength, const TVector<T, d> QueryHalfExtents, SQVisitor& Visitor)
+	static bool SweepFast(const Tuple& Types, const TVector<T, d>& Start, FQueryFastData& CurData, const TVector<T, d> QueryHalfExtents, SQVisitor& Visitor)
 	{
 		const auto& Accelerations = GetAccelerationsPerType<TypeIdx>(Types).Objects;
 		for (const auto& Accelerator : Accelerations)
 		{
-			if (Accelerator && !Accelerator->SweepFast(Start, Dir, InvDir, bParallel, CurLength, InvCurLength, QueryHalfExtents, Visitor))
+			if (Accelerator && !Accelerator->SweepFast(Start, CurData, QueryHalfExtents, Visitor))
 			{
 				return false;
 			}
@@ -246,7 +246,7 @@ struct TSpatialAccelerationCollectionHelper
 		constexpr int NextType = TypeIdx + 1;
 		if (NextType < NumTypes)
 		{
-			return TSpatialAccelerationCollectionHelper < NextType < NumTypes ? NextType : 0, NumTypes, Tuple, TPayloadType, T, d>::SweepFast(Types, Start, Dir, InvDir, bParallel, CurLength, InvCurLength, QueryHalfExtents, Visitor);
+			return TSpatialAccelerationCollectionHelper < NextType < NumTypes ? NextType : 0, NumTypes, Tuple, TPayloadType, T, d>::SweepFast(Types, Start, CurData, QueryHalfExtents, Visitor);
 		}
 
 		return true;
@@ -428,19 +428,8 @@ public:
 	template <typename SQVisitor>
 	void Raycast(const TVector<T, d>& Start, const TVector<T, d>& Dir, const T Length, SQVisitor& Visitor) const
 	{
-		ensure(Length);
-		T CurLength = Length;
-		bool bParallel[3];
-		TVector<T, 3> InvDir;
-
-		T InvLength = 1 / Length;
-		for (int Axis = 0; Axis < 3; ++Axis)
-		{
-			bParallel[Axis] = Dir[Axis] == 0;
-			InvDir[Axis] = bParallel[Axis] ? 0 : 1 / Dir[Axis];
-		}
-
-		TSpatialAccelerationCollectionHelper<0, NumTypes, decltype(Types), TPayloadType, T, d>::RaycastFast(Types, Start, Dir, InvDir, bParallel, CurLength, InvLength, Visitor);
+		FQueryFastData QueryFastData(Dir, Length);
+		TSpatialAccelerationCollectionHelper<0, NumTypes, decltype(Types), TPayloadType, T, d>::RaycastFast(Types, Start, QueryFastData, Visitor);
 	}
 
 	void Sweep(const TVector<T, d>& Start, const TVector<T, d>& Dir, const T Length, const TVector<T, d> QueryHalfExtents, ISpatialVisitor<TPayloadType, T>& Visitor) const override
@@ -452,18 +441,8 @@ public:
 	template <typename SQVisitor>
 	void Sweep(const TVector<T, d>& Start, const TVector<T, d>& Dir, const T Length, const TVector<T, d> QueryHalfExtents, SQVisitor& Visitor) const
 	{
-		ensure(Length);
-		T CurLength = Length;
-		bool bParallel[3];
-		TVector<T, 3> InvDir;
-
-		T InvLength = 1 / Length;
-		for (int Axis = 0; Axis < 3; ++Axis)
-		{
-			bParallel[Axis] = Dir[Axis] == 0;
-			InvDir[Axis] = bParallel[Axis] ? 0 : 1 / Dir[Axis];
-		}
-		TSpatialAccelerationCollectionHelper<0, NumTypes, decltype(Types), TPayloadType, T, d>::SweepFast(Types, Start, Dir, InvDir, bParallel, CurLength, InvLength, QueryHalfExtents, Visitor);
+		FQueryFastData QueryFastData(Dir, Length);
+		TSpatialAccelerationCollectionHelper<0, NumTypes, decltype(Types), TPayloadType, T, d>::SweepFast(Types, Start, QueryFastData, QueryHalfExtents, Visitor);
 	}
 
 	virtual void Overlap(const TBox<T, d>& QueryBounds, ISpatialVisitor<TPayloadType, T>& Visitor) const override
diff --git a/Engine/Source/Runtime/Experimental/Chaos/Public/Chaos/PBDCollisionConstraintImp.h b/Engine/Source/Runtime/Experimental/Chaos/Public/Chaos/PBDCollisionConstraintImp.h
--- a/Engine/Source/Runtime/Experimental/Chaos/Public/Chaos/PBDCollisionConstraintImp.h
+++ b/Engine/Source/Runtime/Experimental/Chaos/Public/Chaos/PBDCollisionConstraintImp.h
@@ -107,13 +107,13 @@ namespace Chaos
 			return true;
 		}
 
-		bool VisitSweep(TSpatialVisitorData<TAccelerationStructureHandle<T, d>>, T Length)
+		bool VisitSweep(TSpatialVisitorData<TAccelerationStructureHandle<T, d>>, FQueryFastData& CurData)
 		{
 			check(false);
 			return false;
 		}
 
-		bool VisitRaycast(TSpatialVisitorData<TAccelerationStructureHandle<T, d>>, T Length)
+		bool VisitRaycast(TSpatialVisitorData<TAccelerationStructureHandle<T, d>>, FQueryFastData& CurData)
 		{
 			check(false);
 			return false;
diff --git a/Engine/Source/Runtime/PhysicsSQ/Private/SQAccelerator.cpp b/Engine/Source/Runtime/PhysicsSQ/Private/SQAccelerator.cpp
--- a/Engine/Source/Runtime/PhysicsSQ/Private/SQAccelerator.cpp
+++ b/Engine/Source/Runtime/PhysicsSQ/Private/SQAccelerator.cpp
@@ -136,20 +136,19 @@ struct TSQVisitor : public Chaos::ISpatialVisitor<TPayload, float>
 		//todo: check THitType is overlap
 	}
 
-	virtual bool Raycast(const Chaos::TSpatialVisitorData<TPayload>& Instance, float& CurLength) override
+	virtual bool Raycast(const Chaos::TSpatialVisitorData<TPayload>& Instance, Chaos::FQueryFastData& CurData) override
 	{
-		return Visit<ESQType::Raycast>(Instance, CurLength);
+		return Visit<ESQType::Raycast>(Instance, &CurData);
 	}
 
-	virtual bool Sweep(const Chaos::TSpatialVisitorData<TPayload>& Instance, float& CurLength) override
+	virtual bool Sweep(const Chaos::TSpatialVisitorData<TPayload>& Instance, Chaos::FQueryFastData& CurData) override
 	{
-		return Visit<ESQType::Sweep>(Instance, CurLength);
+		return Visit<ESQType::Sweep>(Instance, &CurData);
 	}
 
 	virtual bool Overlap(const Chaos::TSpatialVisitorData<TPayload>& Instance) override
 	{
-		float DummyLength;
-		return Visit<ESQType::Overlap>(Instance, DummyLength);
+		return Visit<ESQType::Overlap>(Instance, nullptr);
 	}
 
 private:
@@ -162,12 +161,12 @@ private:
 	};
 
 	template <ESQType SQ>
-	bool Visit(const Chaos::TSpatialVisitorData<TPayload>& Instance, float& CurLength)
+	bool Visit(const Chaos::TSpatialVisitorData<TPayload>& Instance, Chaos::FQueryFastData* CurData)
 	{
 #if CHAOS_DEBUG_DRAW && WITH_CHAOS
 		if (DebugParams.IsDebugQuery() && ChaosSQDrawDebugVisitorQueries)
 		{
-			DebugDraw<SQ>(Instance, CurLength);
+			DebugDraw<SQ>(Instance, CurData->CurrentLength);
 		}
 #endif
 
@@ -212,18 +211,18 @@ private:
 
 					const TVector<float, 3> DirLocal = ActorTM.InverseTransformVectorNoScale(Dir);
 					const TVector<float, 3> StartLocal = ActorTM.InverseTransformPositionNoScale(StartPoint);
-					bHit = Geom->Raycast(StartLocal, DirLocal, CurLength, /*Thickness=*/0.f, Distance, LocalPosition, LocalNormal, FaceIdx);
+					bHit = Geom->Raycast(StartLocal, DirLocal, CurData->CurrentLength, /*Thickness=*/0.f, Distance, LocalPosition, LocalNormal, FaceIdx);
 					if (bHit)
 					{
 						WorldPosition = ActorTM.TransformPositionNoScale(LocalPosition);
 						WorldNormal = ActorTM.TransformVectorNoScale(LocalNormal);
 					}
 				}
-				else if(SQ == ESQType::Sweep && CurLength > 0)
+				else if(SQ == ESQType::Sweep && CurData->CurrentLength > 0)
 				{
-					bHit = SweepQuery(*Geom, ActorTM, *QueryGeom, StartTM, Dir, CurLength, Distance, WorldPosition, WorldNormal, FaceIdx, 0.f, bComputeMTD);
+					bHit = SweepQuery(*Geom, ActorTM, *QueryGeom, StartTM, CurData->Dir, CurData->CurrentLength, Distance, WorldPosition, WorldNormal, FaceIdx, 0.f, bComputeMTD);
 				}
-				else if (SQ == ESQType::Overlap || (SQ == ESQType::Sweep && CurLength == 0))
+				else if (SQ == ESQType::Overlap || (SQ == ESQType::Sweep && CurData->CurrentLength == 0))
 				{
 					bHit = OverlapQuery(*Geom, ActorTM, *QueryGeom, StartTM, /*Thickness=*/0);
 				}
@@ -246,8 +245,8 @@ private:
 
 						if (bBlocker)
 						{
-							CurLength = Distance;
-							if (CurLength == 0 && (SQ == ESQType::Raycast || HitBuffer.WantsSingleResult()))	//raycasts always fail with distance 0, sweeps only matter if we want multi overlaps
+							CurData->SetLength(Distance);
+							if (CurData->CurrentLength == 0 && (SQ == ESQType::Raycast || HitBuffer.WantsSingleResult()))	//raycasts always fail with distance 0, sweeps only matter if we want multi overlaps
 							{
 								return false;	//initial overlap so nothing will be better than this
 							}
