Make sure dirty particles are cleared quickly (not timesliced) after level load completed

diff --git a/Engine/Source/Runtime/Experimental/Chaos/Private/Chaos/PBDRigidsEvolution.cpp b/Engine/Source/Runtime/Experimental/Chaos/Private/Chaos/PBDRigidsEvolution.cpp
--- a/Engine/Source/Runtime/Experimental/Chaos/Private/Chaos/PBDRigidsEvolution.cpp
+++ b/Engine/Source/Runtime/Experimental/Chaos/Private/Chaos/PBDRigidsEvolution.cpp
@@ -29,7 +29,7 @@ namespace Chaos
 			AABBMaxChildrenInLeaf = 500;
 			AABBMaxTreeDepth = 200;
 			MaxPayloadSize = 100000;
-			IterationsPerTimeSlice = 2000;
+			IterationsPerTimeSlice = 20000;
 		}
 	} ConfigSettings;
 
@@ -60,7 +60,7 @@ namespace Chaos
 
 			for (uint16 BucketIdx = 0; BucketIdx < NumBuckets; ++BucketIdx)
 			{
-				Collection->AddSubstructure(CreateAccelerationPerBucket_Threaded(Empty, BucketIdx), BucketIdx);
+				Collection->AddSubstructure(CreateAccelerationPerBucket_Threaded(Empty, BucketIdx, true), BucketIdx);
 			}
 
 			return TUniquePtr<ISpatialAccelerationCollection<TAccelerationStructureHandle<T, d>, T, d>>(Collection);
@@ -71,7 +71,7 @@ namespace Chaos
 			return ConfigSettings.BroadphaseType >= 3 ? 3 : 1;
 		}
 
-		virtual TUniquePtr<ISpatialAcceleration<TAccelerationStructureHandle<T, d>, T, d>> CreateAccelerationPerBucket_Threaded(const TConstParticleView<TSpatialAccelerationCache<T, d>>& Particles, uint16 BucketIdx) override
+		virtual TUniquePtr<ISpatialAcceleration<TAccelerationStructureHandle<T, d>, T, d>> CreateAccelerationPerBucket_Threaded(const TConstParticleView<TSpatialAccelerationCache<T, d>>& Particles, uint16 BucketIdx, bool ForceFullBuild) override
 		{
 			switch (BucketIdx)
 			{
@@ -83,7 +83,7 @@ namespace Chaos
 				}
 				else if (ConfigSettings.BroadphaseType == 1 || ConfigSettings.BroadphaseType == 3)
 				{
-					return MakeUnique<AABBTreeType>(Particles, ConfigSettings.MaxChildrenInLeaf, ConfigSettings.MaxTreeDepth, ConfigSettings.MaxPayloadSize, ConfigSettings.IterationsPerTimeSlice);
+					return MakeUnique<AABBTreeType>(Particles, ConfigSettings.MaxChildrenInLeaf, ConfigSettings.MaxTreeDepth, ConfigSettings.MaxPayloadSize, ForceFullBuild ? 0 : ConfigSettings.IterationsPerTimeSlice);
 				}
 				else if (ConfigSettings.BroadphaseType == 4 || ConfigSettings.BroadphaseType == 2)
 				{
@@ -163,11 +163,13 @@ namespace Chaos
 		ISpatialAccelerationCollectionFactory<T, d>& InSpatialCollectionFactory
 		, const TMap<FSpatialAccelerationIdx, TUniquePtr<TSpatialAccelerationCache<T, d>>>& InSpatialAccelerationCache
 		, TUniquePtr<FAccelerationStructure>& InAccelerationStructure
-		, TUniquePtr<FAccelerationStructure>& InAccelerationStructureCopy)
+		, TUniquePtr<FAccelerationStructure>& InAccelerationStructureCopy
+		, bool InForceFullBuild)
 		: SpatialCollectionFactory(InSpatialCollectionFactory)
 		, SpatialAccelerationCache(InSpatialAccelerationCache)
 		, AccelerationStructure(InAccelerationStructure)
 		, AccelerationStructureCopy(InAccelerationStructureCopy)
+		, IsForceFullBuild(InForceFullBuild)
 	{
 	}
 
@@ -234,7 +236,7 @@ namespace Chaos
 			{
 				SCOPE_CYCLE_COUNTER(STAT_AccelerationStructureTimeSlice);
 
-				AccelerationStructure->GetSubstructure(SpatialIdx)->ProgressAsyncTimeSlicing();
+				AccelerationStructure->GetSubstructure(SpatialIdx)->ProgressAsyncTimeSlicing(IsForceFullBuild);
 
 				// is it still progressing or now complete
 				IsTimeSlicingProgressing = !AccelerationStructure->GetSubstructure(SpatialIdx)->IsAsyncTimeSlicingComplete();
@@ -257,7 +259,7 @@ namespace Chaos
 				SCOPE_CYCLE_COUNTER(STAT_CreateInitialAccelerationStructure);
 
 				auto ParticleView = MakeConstParticleView(MoveTemp(ViewsPerBucket[BucketIdx]));
-				auto NewStruct = SpatialCollectionFactory.CreateAccelerationPerBucket_Threaded(ParticleView, BucketIdx);
+				auto NewStruct = SpatialCollectionFactory.CreateAccelerationPerBucket_Threaded(ParticleView, BucketIdx, IsForceFullBuild);
 
 				// we kicked of the creation of a new structure and it's going to time-slice the work
 				if (!NewStruct->IsAsyncTimeSlicingComplete())
@@ -400,6 +402,14 @@ namespace Chaos
 	{
 		SCOPE_CYCLE_COUNTER(STAT_ComputeIntermediateSpatialAcceleration);
 		CHAOS_SCOPED_TIMER(ComputeIntermediateSpatialAcceleration);
+
+		bool ForceFullBuild = InternalAccelerationQueue.Num() > 1000;
+
+		if (ForceFullBuild)
+		{
+			UE_LOG(LogChaos, Warning, TEXT("ForceFullBuild"));
+		}
+
 		if (!AccelerationStructureTaskComplete)
 		{
 			//initial frame so make empty structures
@@ -438,7 +448,7 @@ namespace Chaos
 			}
 
 			// we run the task for both starting a new accel structure as well as for the timeslicing
-			AccelerationStructureTaskComplete = TGraphTask<FChaosAccelerationStructureTask>::CreateTask().ConstructAndDispatchWhenReady(*SpatialCollectionFactory, SpatialAccelerationCache, AsyncInternalAcceleration, AsyncExternalAcceleration);
+			AccelerationStructureTaskComplete = TGraphTask<FChaosAccelerationStructureTask>::CreateTask().ConstructAndDispatchWhenReady(*SpatialCollectionFactory, SpatialAccelerationCache, AsyncInternalAcceleration, AsyncExternalAcceleration, ForceFullBuild);
 		}
 		else
 		{
diff --git a/Engine/Source/Runtime/Experimental/Chaos/Public/Chaos/AABBTree.h b/Engine/Source/Runtime/Experimental/Chaos/Public/Chaos/AABBTree.h
--- a/Engine/Source/Runtime/Experimental/Chaos/Public/Chaos/AABBTree.h
+++ b/Engine/Source/Runtime/Experimental/Chaos/Public/Chaos/AABBTree.h
@@ -22,6 +22,7 @@ enum class EAABBQueryType
 DECLARE_CYCLE_STAT(TEXT("AABBTreeGenerateTree"), STAT_AABBTreeGenerateTree, STATGROUP_Chaos);
 DECLARE_CYCLE_STAT(TEXT("AABBTreeTimeSliceSetup"), STAT_AABBTreeTimeSliceSetup, STATGROUP_Chaos);
 DECLARE_CYCLE_STAT(TEXT("AABBTreeInitialTimeSlice"), STAT_AABBTreeInitialTimeSlice, STATGROUP_Chaos);
+DECLARE_CYCLE_STAT(TEXT("AABBTreeProgressTimeSlice"), STAT_AABBTreeProgressTimeSlice, STATGROUP_Chaos);
 
 template <typename T, typename TQueryFastData, EAABBQueryType Query>
 struct TAABBTreeIntersectionHelper
@@ -236,8 +237,16 @@ public:
 	{
 	}
 
-	virtual void ProgressAsyncTimeSlicing() override
+	virtual void ProgressAsyncTimeSlicing(bool ForceBuildCompletion) override
 	{
+		SCOPE_CYCLE_COUNTER(STAT_AABBTreeProgressTimeSlice);
+
+		// force is to stop time slicing and complete the rest of the build now
+		if (ForceBuildCompletion)
+		{
+			MaxNumToProcess = 0;
+		}
+
 		// still has work to complete
 		if (!TimeSliceWorkToComplete.IsEmpty())
 		{
@@ -610,7 +619,6 @@ private:
 		DirtyElements.Reset();
 		PayloadToInfo.Reset();
 		NumProcessedThisSlice = 0;
-		check(MaxNumToProcess > 0);
 
 		FullBounds = TBox<T, 3>::EmptyBox();
 
@@ -683,7 +691,7 @@ private:
 		 	Nodes.AddDefaulted();	//todo: remove TBox
 		}
 
-		if (NumProcessedThisSlice >= NumToProcess)
+		if ((MaxNumToProcess > 0) && (NumProcessedThisSlice >= NumToProcess))
 		{
 			// done enough work, capture stack
 			FWorkSnapshot Store;
diff --git a/Engine/Source/Runtime/Experimental/Chaos/Public/Chaos/ISpatialAcceleration.h b/Engine/Source/Runtime/Experimental/Chaos/Public/Chaos/ISpatialAcceleration.h
--- a/Engine/Source/Runtime/Experimental/Chaos/Public/Chaos/ISpatialAcceleration.h
+++ b/Engine/Source/Runtime/Experimental/Chaos/Public/Chaos/ISpatialAcceleration.h
@@ -196,7 +196,7 @@ public:
 	virtual ~ISpatialAcceleration() = default;
 
 	virtual bool IsAsyncTimeSlicingComplete() { return AsyncTimeSlicingComplete; }
-	virtual void ProgressAsyncTimeSlicing() {}
+	virtual void ProgressAsyncTimeSlicing(bool ForceBuildCompletion = false) {}
 	virtual TArray<TPayloadType> FindAllIntersections(const TBox<T, d>& Box) const { check(false); return TArray<TPayloadType>(); }
 
 	virtual void Raycast(const TVector<T, d>& Start, const TVector<T, d>& Dir, const T Length, ISpatialVisitor<TPayloadType, T>& Visitor) const { check(false); }
diff --git a/Engine/Source/Runtime/Experimental/Chaos/Public/Chaos/PBDRigidsEvolution.h b/Engine/Source/Runtime/Experimental/Chaos/Public/Chaos/PBDRigidsEvolution.h
--- a/Engine/Source/Runtime/Experimental/Chaos/Public/Chaos/PBDRigidsEvolution.h
+++ b/Engine/Source/Runtime/Experimental/Chaos/Public/Chaos/PBDRigidsEvolution.h
@@ -199,7 +199,7 @@ struct CHAOS_API ISpatialAccelerationCollectionFactory
 	virtual TUniquePtr<ISpatialAccelerationCollection<TAccelerationStructureHandle<T, d>, T, d>> CreateEmptyCollection() = 0;
 
 	//Chaos creates new acceleration structures per bucket. Factory can change underlying type at runtime as well as number of buckets to AB test
-	virtual TUniquePtr<ISpatialAcceleration<TAccelerationStructureHandle<T, d>, T, d>> CreateAccelerationPerBucket_Threaded(const TConstParticleView<TSpatialAccelerationCache<T, d>>& Particles, uint16 BucketIdx) = 0;
+	virtual TUniquePtr<ISpatialAcceleration<TAccelerationStructureHandle<T, d>, T, d>> CreateAccelerationPerBucket_Threaded(const TConstParticleView<TSpatialAccelerationCache<T, d>>& Particles, uint16 BucketIdx, bool ForceFullBuild) = 0;
 
 	//Mask indicating which bucket is active. Spatial indices in inactive buckets fallback to bucket 0. Bit 0 indicates bucket 0 is active, Bit 1 indicates bucket 1 is active, etc...
 	virtual uint8 GetActiveBucketsMask() const = 0;
@@ -681,7 +681,8 @@ protected:
 		FChaosAccelerationStructureTask(ISpatialAccelerationCollectionFactory<T,d>& InSpatialCollectionFactory
 			, const TMap<FSpatialAccelerationIdx, TUniquePtr<TSpatialAccelerationCache<T,d>>>& InSpatialAccelerationCache
 			, TUniquePtr<FAccelerationStructure>& InAccelerationStructure
-			, TUniquePtr<FAccelerationStructure>& InAccelerationStructureCopy);
+			, TUniquePtr<FAccelerationStructure>& InAccelerationStructureCopy
+			, bool InForceFullBuild);
 		static FORCEINLINE TStatId GetStatId();
 		static FORCEINLINE ENamedThreads::Type GetDesiredThread();
 		static FORCEINLINE ESubsequentsMode::Type GetSubsequentsMode();
@@ -691,6 +692,7 @@ protected:
 		const TMap<FSpatialAccelerationIdx, TUniquePtr<TSpatialAccelerationCache<T, d>>>& SpatialAccelerationCache;
 		TUniquePtr<FAccelerationStructure>& AccelerationStructure;
 		TUniquePtr<FAccelerationStructure>& AccelerationStructureCopy;
+		bool IsForceFullBuild;
 	};
 	FGraphEventRef AccelerationStructureTaskComplete;
 
