Chaos - fix substep time when limiting number of steps

diff --git a/Engine/Source/Runtime/Experimental/Chaos/Private/Chaos/PBDRigidsEvolutionGBF.cpp b/Engine/Source/Runtime/Experimental/Chaos/Private/Chaos/PBDRigidsEvolutionGBF.cpp
--- a/Engine/Source/Runtime/Experimental/Chaos/Private/Chaos/PBDRigidsEvolutionGBF.cpp
+++ b/Engine/Source/Runtime/Experimental/Chaos/Private/Chaos/PBDRigidsEvolutionGBF.cpp
@@ -90,14 +90,26 @@ void SerializeToDisk(TEvolution& Evolution)
 template <typename T, int d>
 void TPBDRigidsEvolutionGBF<T, d>::Advance(const T Dt, const T MaxStepDt, const int32 MaxSteps)
 {
-	const int32 NumSteps = FMath::Clamp(FMath::CeilToInt(Dt / MaxStepDt), 1, MaxSteps);
-	const T StepDt = Dt / (T)NumSteps;
-	for (int32 Step = 0; Step < NumSteps; ++Step)
+	// Determine how many steps we would like to take
+	int32 NumSteps = FMath::CeilToInt(Dt / MaxStepDt);
+	if (NumSteps > 0)
 	{
-		// StepFraction: how much of the remaining time this step represents, used to interpolate kinematic targets
-		// E.g., for 4 steps this will be: 1/4, 1/3, 1/2, 1
-		const float StepFraction = (T)1 / (T)(NumSteps - Step);
-		AdvanceOneTimeStep(StepDt, StepFraction);
+		// Determine the step time
+		const T StepDt = Dt / (T)NumSteps;
+
+		// Limit the number of steps
+		// NOTE: This is after step time calculation so sim will appear to slow down for large Dt
+		// but that is preferable to blowing up from a large timestep.
+		NumSteps = FMath::Clamp(NumSteps, 1, MaxSteps);
+
+		for (int32 Step = 0; Step < NumSteps; ++Step)
+		{
+			// StepFraction: how much of the remaining time this step represents, used to interpolate kinematic targets
+			// E.g., for 4 steps this will be: 1/4, 1/3, 1/2, 1
+			const float StepFraction = (T)1 / (T)(NumSteps - Step);
+		
+			AdvanceOneTimeStep(StepDt, StepFraction);
+		}
 	}
 }
 
