Chaos - fix persistent manifold maintenance

diff --git a/Engine/Source/Runtime/Experimental/Chaos/Private/Chaos/PBDCollisionConstraint.cpp b/Engine/Source/Runtime/Experimental/Chaos/Private/Chaos/PBDCollisionConstraint.cpp
--- a/Engine/Source/Runtime/Experimental/Chaos/Private/Chaos/PBDCollisionConstraint.cpp
+++ b/Engine/Source/Runtime/Experimental/Chaos/Private/Chaos/PBDCollisionConstraint.cpp
@@ -2171,11 +2171,13 @@ void ConstructUnionUnionConstraints(TGeometryParticleHandle<T, d>* Particle0, TG
 }
 
 template<typename T, int d>
-void ConstructConstraintsImpl(TGeometryParticleHandle<T, d>* Particle0, TGeometryParticleHandle<T, d>* Particle1, const FImplicitObject* Implicit0, const FImplicitObject* Implicit1, const T Thickness, TArray< TRigidBodyContactConstraint<T, d> >& ConstraintBuffer, TPBDCollisionConstraintHistory<T, d>* History)
+void ConstructPairConstraintImpl(TGeometryParticleHandle<T, d>* Particle0, TGeometryParticleHandle<T, d>* Particle1, const FImplicitObject* Implicit0, const FImplicitObject* Implicit1, const T Thickness, TArray< TRigidBodyContactConstraint<T, d> >& ConstraintBuffer, TPBDCollisionConstraintHistory<T, d>* History)
 {
-	// TriangleMesh implicit's are for scene query only.
-	if (Implicit0 && GetInnerType(Implicit0->GetType()) == ImplicitObjectType::TriangleMesh) return; //as
-	if (Implicit1 && GetInnerType(Implicit1->GetType()) == ImplicitObjectType::TriangleMesh) return;
+	// See if we already have a constraint for this shape pair
+	if (History && History->ContainsShapeConnection(Implicit0, Implicit1))
+	{
+		return;
+	}
 
 	if (!Implicit0 || !Implicit1)
 	{
@@ -2235,7 +2237,8 @@ void ConstructConstraintsImpl(TGeometryParticleHandle<T, d>* Particle0, TGeometr
 	}
 	else if (Implicit0->GetType() == TImplicitObjectUnion<T, d>::StaticType() && Implicit1->GetType() == TImplicitObjectUnion<T, d>::StaticType())
 	{
-		ConstructUnionUnionConstraints(Particle0, Particle1, Implicit0, Implicit1, Thickness, ConstraintBuffer, History);
+		// Union-union creates multiple manifolds - we should never get here for this pair type. See ConstructConstraintsImpl and ConstructUnionUnionConstraints
+		ensure(false);
 	}
 	else if (Implicit0->IsConvex() && Implicit1->IsConvex())
 	{
@@ -2247,14 +2250,33 @@ void ConstructConstraintsImpl(TGeometryParticleHandle<T, d>* Particle0, TGeometr
 	}
 }
 
+template<typename T, int d>
+void ConstructConstraintsImpl(TGeometryParticleHandle<T, d>* Particle0, TGeometryParticleHandle<T, d>* Particle1, const FImplicitObject* Implicit0, const FImplicitObject* Implicit1, const T Thickness, TArray< TRigidBodyContactConstraint<T, d> >& ConstraintBuffer, TPBDCollisionConstraintHistory<T, d>* History)
+{
+	// TriangleMesh implicits are for scene query only.
+	if (Implicit0 && GetInnerType(Implicit0->GetType()) == ImplicitObjectType::TriangleMesh) return;
+	if (Implicit1 && GetInnerType(Implicit1->GetType()) == ImplicitObjectType::TriangleMesh) return;
+
+	if (Implicit0->GetType() == TImplicitObjectUnion<T, d>::StaticType() && Implicit1->GetType() == TImplicitObjectUnion<T, d>::StaticType())
+	{
+		ConstructUnionUnionConstraints(Particle0, Particle1, Implicit0, Implicit1, Thickness, ConstraintBuffer, History);
+	}
+	else
+	{
+		ConstructPairConstraintImpl(Particle0, Particle1, Implicit0, Implicit1, Thickness, ConstraintBuffer, History);
+	}
+}
+
 template<typename T, int d>
 void TPBDCollisionConstraint<T, d>::ConstructConstraints(TGeometryParticleHandle<T, d>* Particle0, TGeometryParticleHandle<T, d>* Particle1, const T Thickness, TArray<TRigidBodyContactConstraint<T, d>>& ConstraintBuffer)
 {
 	ensure(Particle0 && Particle1);
 
 	FConstraintHistory* HistoryElement = nullptr;
-	if (FConstraintHistory** Ptr = History.Find(FConstraintHandleID(Particle0, Particle1)))
+	if (FConstraintHistory** Ptr = History.Find(FConstraintHandleID::Make(Particle0, Particle1)))
+	{
 		HistoryElement = *Ptr;
+	}
 
 	ConstructConstraintsImpl<T, d>(Particle0, Particle1, Particle0->Geometry().Get(), Particle1->Geometry().Get(), Thickness, ConstraintBuffer, HistoryElement);
 }
diff --git a/Engine/Source/Runtime/Experimental/Chaos/Public/Chaos/PBDCollisionConstraint.h b/Engine/Source/Runtime/Experimental/Chaos/Public/Chaos/PBDCollisionConstraint.h
--- a/Engine/Source/Runtime/Experimental/Chaos/Public/Chaos/PBDCollisionConstraint.h
+++ b/Engine/Source/Runtime/Experimental/Chaos/Public/Chaos/PBDCollisionConstraint.h
@@ -52,6 +52,45 @@ enum class ECollisionModifierResult
 	Disabled,	/** Collision should be disabled */
 };
 
+/**
+ * A pair that can be used as a key, where the key value is independent of the pair order,
+ * i.e., TSymmetricPairKey(A, B) == TSymmetricPairKey(B, A)
+ */
+template<typename T>
+class TSymmetricPairKey
+{
+public:
+	using Item = T;
+	using FPair = TPair<Item, Item>;
+
+	static TSymmetricPairKey Make(const T& Item0, const T& Item1)
+	{
+		T ItemMin = FMath::Min(Item0, Item1);
+		T ItemMax = FMath::Max(Item0, Item1);
+		return TSymmetricPairKey<T>(ItemMin, ItemMax);
+	}
+
+	friend inline uint32 GetTypeHash(const TSymmetricPairKey<T>& Item)
+	{
+		return GetTypeHash(Item.Pair);
+	}
+
+	friend inline bool operator==(const TSymmetricPairKey<T>& L, const TSymmetricPairKey<T>& R)
+	{
+		return L.Pair == R.Pair;
+	}
+
+private:
+	TSymmetricPairKey(const T& Item0, const T& Item1)
+		: Pair(Item0, Item1)
+	{
+		check(Item0 < Item1);
+	}
+
+	FPair Pair;
+};
+
+
 template<class T, int d>
 class CHAOS_API TPBDCollisionConstraintHandle : public TContainerConstraintHandle<TPBDCollisionConstraint<T, d>>
 {
@@ -74,10 +113,11 @@ protected:
 template<class T, int d>
 class CHAOS_API TPBDCollisionConstraintHistory
 {
-	typedef TPair<const FImplicitObject*, const FImplicitObject*> ImplicitPairs;
+	using FImplicitPairKey = TSymmetricPairKey<const FImplicitObject*>;
 
 public:
 	using FConstraintContainerHandle = TPBDCollisionConstraintHandle<T, d>;
+
 	TPBDCollisionConstraintHistory(const TVector<T, d>& InLocation, const TRotation<T, d>& InRotation, int32 InTimestamp = -INT_MAX)
 		: Timestamp(InTimestamp)
 		, Location(InLocation)
@@ -89,19 +129,19 @@ public:
 
 	void AddHandle(FConstraintContainerHandle* InHandle)
 	{
-		Implicits.Add(ImplicitPairs(InHandle->GetContact().Geometry[0], InHandle->GetContact().Geometry[1]));
-		Handles.Push(InHandle);
+		Implicits.Add(FImplicitPairKey::Make(InHandle->GetContact().Geometry[0], InHandle->GetContact().Geometry[1]));
+		ConstraintHandles.Push(InHandle);
 	}
 	void RemoveHandle(FConstraintContainerHandle* InHandle)
 	{
-		Handles.RemoveSingleSwap(InHandle);
-		Implicits.Remove(ImplicitPairs(InHandle->GetContact().Geometry[0], InHandle->GetContact().Geometry[1]));
+		ConstraintHandles.RemoveSingleSwap(InHandle);
+		Implicits.Remove(FImplicitPairKey::Make(InHandle->GetContact().Geometry[0], InHandle->GetContact().Geometry[1]));
 	}
 
-	TArray<FConstraintContainerHandle*>& GetHandles() { return Handles; }
-	const TArray<FConstraintContainerHandle*>& GetHandles() const { return Handles; }
+	TArray<FConstraintContainerHandle*>& GetHandles() { return ConstraintHandles; }
+	const TArray<FConstraintContainerHandle*>& GetHandles() const { return ConstraintHandles; }
 
-	bool ContainsShapeConnection(const FImplicitObject* Implicit0In, const FImplicitObject* Implicit1In) { return Implicits.Contains(ImplicitPairs(Implicit0In, Implicit1In)); }
+	bool ContainsShapeConnection(const FImplicitObject* Implicit0In, const FImplicitObject* Implicit1In) { return Implicits.Contains(FImplicitPairKey::Make(Implicit0In, Implicit1In)); }
 
 	int32 GetTimestamp() const { return Timestamp; }
 	void SetTimestamp(int32 InTimestamp) { Timestamp = InTimestamp; }
@@ -110,8 +150,8 @@ private:
 	int32 Timestamp;
 	TVector<T, d> Location;
 	TRotation<T, d> Rotation;
-	TArray<FConstraintContainerHandle*> Handles;
-	TSet<ImplicitPairs> Implicits;
+	TArray<FConstraintContainerHandle*> ConstraintHandles;
+	TSet<FImplicitPairKey> Implicits;
 };
 
 template<typename T, int d>
@@ -144,7 +184,7 @@ public:
 	using FConstraintHistory = TPBDCollisionConstraintHistory<T, d>;
 	using FConstraintHandleAllocator = TConstraintHandleAllocator<TPBDCollisionConstraint<T, d>>;
 	using FRigidBodyContactConstraint = TRigidBodyContactConstraint<T, d>;
-	typedef TPair<TGeometryParticleHandle<T, d>*, const TGeometryParticleHandle<T, d>*> FConstraintHandleID;
+	using FConstraintHandleID = TSymmetricPairKey<const TGeometryParticleHandle<T, d>*>;
 
 	TPBDCollisionConstraint(const TPBDRigidsSOAs<T,d>& InParticles, TArrayCollectionArray<bool>& Collided, const TArrayCollectionArray<TSerializablePtr<TChaosPhysicsMaterial<T>>>& PerParticleMaterials, const int32 PairIterations = 1, const T Thickness = (T)0);
 	virtual ~TPBDCollisionConstraint() {}
@@ -188,12 +228,12 @@ public:
 
 	FConstraintHandleID GetConstraintHandleID(const FRigidBodyContactConstraint & Constraint) const
 	{
-		return  FConstraintHandleID( Constraint.Particle, Constraint.Levelset );
+		return  FConstraintHandleID::Make( Constraint.Particle, Constraint.Levelset );
 	}
 
 	FConstraintHandleID GetConstraintHandleID(int32 ConstraintIndex) const
 	{
-		return  FConstraintHandleID( Constraints[ConstraintIndex].Particle, Constraints[ConstraintIndex].Levelset );
+		return  FConstraintHandleID::Make( Constraints[ConstraintIndex].Particle, Constraints[ConstraintIndex].Levelset );
 	}
 		
 	const FConstraintContainerHandle* GetConstraintHandle(int32 ConstraintIndex) const
