Chaos : Convex Collisions

- Simple PhiWithNormal collisions

diff --git a/Engine/Source/Runtime/Experimental/Chaos/Private/Chaos/CollisionResolutionConvexConvex.cpp b/Engine/Source/Runtime/Experimental/Chaos/Private/Chaos/CollisionResolutionConvexConvex.cpp
--- a/Engine/Source/Runtime/Experimental/Chaos/Private/Chaos/CollisionResolutionConvexConvex.cpp
+++ b/Engine/Source/Runtime/Experimental/Chaos/Private/Chaos/CollisionResolutionConvexConvex.cpp
@@ -23,15 +23,13 @@ namespace Chaos
 			}
 		}
 
-		for (int i = 0; i < ManifoldSamples; i++)
-		{
-			TRigidBodyContactConstraint<T, d> Constraint;
-			Constraint.Particle = Particle0;
-			Constraint.Levelset = Particle1;
-			Constraint.Geometry[0] = Implicit0;
-			Constraint.Geometry[1] = Implicit1;
-			ConstraintBuffer.Add(Constraint);
-		}
+		TRigidBodyContactConstraint<T, d> Constraint;
+		Constraint.Particle = Particle0;
+		Constraint.Levelset = Particle1;
+		Constraint.Geometry[0] = Implicit0;
+		Constraint.Geometry[1] = Implicit1;
+
+		ConstraintBuffer.Add(Constraint);
 	}
 
 	template<class T, int d>
@@ -45,79 +43,43 @@ namespace Chaos
 	{
 		ensure(NumConstraints);
 
-		unsigned char AType = A.GetType();
-		unsigned char AStaticType = ImplicitObjectType::IsScaled;
-
 		if (ensure(IsScaled(A.GetType()) && IsScaled(B.GetType())))
 		{
-			const TImplicitObjectScaled<TConvex<T, d>>& AScaled = static_cast<const TImplicitObjectScaled<TConvex<T, d>>&>(A);
-			const TImplicitObjectScaled<TConvex<T, d>>& BScaled = static_cast<const TImplicitObjectScaled<TConvex<T, d>>&>(B);
-
 			TRigidTransform<T, d> BToATM = BTM.GetRelativeTransform(ATM);
-			TRigidTransform<T, d> AToBTM = ATM.GetRelativeTransform(BTM);
 
 			if (ensure(GetInnerType(A.GetType()) == ImplicitObjectType::Convex && GetInnerType(B.GetType()) == ImplicitObjectType::Convex))
 			{
-				const TConvex<T, d>* AObject = static_cast<const TConvex<T, d>*>(AScaled.Object().Get());
-				const TConvex<T, d>* BObject = static_cast<const TConvex<T, d>*>(BScaled.Object().Get());
-
-				const TParticles<T, d>& SurfaceParticles = AObject->GetSurfaceParticles();
-
-				T OutDistance;
-				TVector<T, d> StubNormal;
-				TVector<T, d> OutNearestA;
-				TVector<T, d> OutNearestB;
-				TVector<T, d> OutNormal;
-
-				if (Manifold == nullptr)
+				const TConvex<T, d>* AObject = static_cast<const TConvex<T, d>*>(static_cast<const TImplicitObjectScaled<TConvex<T, d>>&>(A).Object().Get());
+				const TConvex<T, d>* BObject = static_cast<const TConvex<T, d>*>(static_cast<const TImplicitObjectScaled<TConvex<T, d>>&>(B).Object().Get());
+				if (AObject && BObject)
 				{
-					if (GJKDistance<T>(*AObject, *BObject, BToATM, OutDistance, OutNearestA, OutNearestB))
-					{
-						TVector<T, d> Normal(0);
-						TVector<T, d> WorldA = ATM.TransformPosition(OutNearestA);
-						TVector<T, d> WorldB = BTM.TransformPosition(OutNearestB);
-						TVector<T, d> CenterLocation = WorldA;
-						T Phi = BObject->PhiWithNormal(BTM.InverseTransformPosition(CenterLocation), Normal);
+					const TParticles<T, d>& SurfaceParticles = AObject->GetSurfaceParticles();
 
-						T ManifoldSize = A.BoundingBox().Extents().Size() * ManifoldScale;
-						TVector<T, d> CrossVector = FMath::IsNearlyEqual(FMath::Abs(TVector<T, d>::DotProduct(TVector<T, d>(0, 1, 0), Normal)), T(1)) ?
-							TVector<T, d>(0, 0, 1) :
-							TVector<T, d>(0, 1, 0);
+					for (int32 Idx = 0; Idx < NumConstraints; Idx++)
+					{
+						Constraints[Idx].bDisabled = true;
+					}
 
-						Constraints[0].LocalLocation = ATM.InverseTransformPosition(CenterLocation);
-						Constraints[0].Normal = Normal;
+					if (GJKIntersection(*AObject, *BObject, BToATM, Thickness))
+					{
+						TRigidTransform<T, d> AToBTM = ATM.GetRelativeTransform(BTM);
 
-						if (NumConstraints > 1)
+						T Phi;
+						TVector<T, d> Normal;
+						for (int32 Idx = 0; Idx < (int32)SurfaceParticles.Size(); Idx++)
 						{
-							TVector<T, d> DirVector = CrossVector;
-							T Angle = T(360) / T(NumConstraints - 1), LocalPhi;
-							for (int32 Idx = 1; Idx < NumConstraints; Idx++)
+							Phi = BObject->PhiWithNormal(AToBTM.TransformPosition(SurfaceParticles.X(Idx)), Normal);
+							if (Phi< Constraints[0].Phi)
 							{
-								CrossVector = CrossVector.RotateAngleAxis(Angle, Normal);
-								TVector<T, d> WorldLocation = CenterLocation + ManifoldSize * CrossVector;
-								Constraints[Idx].LocalLocation = ATM.InverseTransformPosition(WorldLocation);
-
-								LocalPhi = AObject->PhiWithNormal(Constraints[Idx].LocalLocation, StubNormal);
-								Constraints[Idx].LocalLocation -= StubNormal * LocalPhi;
-
-								Constraints[Idx].Normal = Normal;
+								Constraints[0].Phi = Phi;
+								Constraints[0].Location = ATM.TransformPosition(SurfaceParticles.X(Idx));
+								Constraints[0].Normal = BTM.TransformVector(Normal);
 							}
 						}
+						return;
 					}
+					return;
 				}
-
-				for (int32 Idx = 0; Idx < NumConstraints; Idx++)
-				{
-					Constraints[Idx].Location = ATM.TransformPosition(Constraints[Idx].LocalLocation);
-					Constraints[Idx].Phi = BObject->PhiWithNormal(BTM.InverseTransformPosition(Constraints[Idx].Location), StubNormal);
-
-					if (Constraints[Idx].Phi < 0.f && Manifold)
-					{
-						Manifold->SetTimestamp(INT_MAX);
-					}
-				}
-
-				return;
 			}
 		}
 
diff --git a/Engine/Source/Runtime/Experimental/Chaos/Private/Chaos/PBDCollisionConstraint.cpp b/Engine/Source/Runtime/Experimental/Chaos/Private/Chaos/PBDCollisionConstraint.cpp
--- a/Engine/Source/Runtime/Experimental/Chaos/Private/Chaos/PBDCollisionConstraint.cpp
+++ b/Engine/Source/Runtime/Experimental/Chaos/Private/Chaos/PBDCollisionConstraint.cpp
@@ -100,14 +100,17 @@ void TPBDCollisionConstraint<T, d>::Reset(/*const TPBDRigidParticles<T, d>& InPa
 {
 	SCOPE_CYCLE_COUNTER(STAT_CollisionConstraintsReset);
 
+	// @todo(brice) : Manifolds are changing, so just default to collision constraints right now.
+
 	for (int32 Idx = Constraints.Num() - 1; Idx >= 0; Idx--)
 	{
-		if (!bEnableCollisions)
+		//if (!bEnableCollisions)
 		{
 			RemoveConstraint(Idx);
 		}
 	}
 
+	/*
 	TArray<FConstraintHandleID> ManifoldKeys;
 	Manifolds.GetKeys(ManifoldKeys);
 	for (FConstraintHandleID Key : ManifoldKeys)
@@ -134,6 +137,7 @@ void TPBDCollisionConstraint<T, d>::Reset(/*const TPBDRigidParticles<T, d>& InPa
 			}
 		}
 	}
+	*/
 
 	MAngularFriction = 0;
 	bUseCCD = false;
@@ -143,6 +147,7 @@ template<typename T, int d>
 void TPBDCollisionConstraint<T, d>::RemoveConstraint(int32 Idx)
 {
 	// Manifold
+	/*
 	FConstraintHandleID HandleID = GetConstraintHandleID(Constraints[Idx]);
 	if (FConstraintManifold** Ptr = Manifolds.Find(HandleID))
 	{
@@ -155,6 +160,7 @@ void TPBDCollisionConstraint<T, d>::RemoveConstraint(int32 Idx)
 			}
 		}
 	}
+	*/
 
 	HandleAllocator.FreeHandle(Handles[Idx]);
 	Handles.RemoveAtSwap(Idx);
diff --git a/Engine/Source/Runtime/Experimental/Chaos/Public/Chaos/PBDCollisionConstraintImp.h b/Engine/Source/Runtime/Experimental/Chaos/Public/Chaos/PBDCollisionConstraintImp.h
--- a/Engine/Source/Runtime/Experimental/Chaos/Public/Chaos/PBDCollisionConstraintImp.h
+++ b/Engine/Source/Runtime/Experimental/Chaos/Public/Chaos/PBDCollisionConstraintImp.h
@@ -304,6 +304,7 @@ namespace Chaos
 					Queue.Dequeue(Constraints[Idx]);
 					ensure(Handles.Num() == Constraints.Num());
 
+					/*
 					FConstraintHandleID HandleID = GetConstraintHandleID(Constraints[Idx]);
 					FConstraintManifold* Manifold;
 					if (auto ValPtr = Manifolds.Find(HandleID))
@@ -316,6 +317,7 @@ namespace Chaos
 						Manifolds.Add(HandleID, Manifold);
 					}
 					Manifold->AddHandle(Handle);
+					*/
 				}
 				LifespanCounter++;
 			}
