Chaos - 6dof joint angular projection

diff --git a/Engine/Source/Runtime/Engine/Classes/PhysicsEngine/ConstraintInstance.h b/Engine/Source/Runtime/Engine/Classes/PhysicsEngine/ConstraintInstance.h
--- a/Engine/Source/Runtime/Engine/Classes/PhysicsEngine/ConstraintInstance.h
+++ b/Engine/Source/Runtime/Engine/Classes/PhysicsEngine/ConstraintInstance.h
@@ -597,12 +597,24 @@ public:
 	void PostSerialize(const FArchive& Ar);
 #endif
 
+	/** Whether projection is enabled for this constraint */
+	bool IsProjectionEnabled() const
+	{
+		return ProfileInstance.bEnableProjection;
+	}
+
 	/** Turn on linear and angular projection */
 	void EnableProjection();
 
 	/** Turn off linear and angular projection */
 	void DisableProjection();
 
+	/** Whether parent domination is enabled (meaning the parent body cannot be be affected at all by a child) */
+	bool IsParentDominatesEnabled() const
+	{
+		return ProfileInstance.bParentDominates;
+	}
+
 	/** Enable/Disable parent dominates (meaning the parent body cannot be be affected at all by a child) */
 	void EnableParentDominates();
 	void DisableParentDominates();
diff --git a/Engine/Source/Runtime/Engine/Private/PhysicsEngine/ImmediatePhysics/ImmediatePhysicsChaos/ImmediatePhysicsJointHandle_Chaos.cpp b/Engine/Source/Runtime/Engine/Private/PhysicsEngine/ImmediatePhysics/ImmediatePhysicsChaos/ImmediatePhysicsJointHandle_Chaos.cpp
--- a/Engine/Source/Runtime/Engine/Private/PhysicsEngine/ImmediatePhysics/ImmediatePhysicsChaos/ImmediatePhysicsJointHandle_Chaos.cpp
+++ b/Engine/Source/Runtime/Engine/Private/PhysicsEngine/ImmediatePhysics/ImmediatePhysicsChaos/ImmediatePhysicsJointHandle_Chaos.cpp
@@ -21,6 +21,11 @@ namespace ImmediatePhysics_Chaos
 	bool ChaosImmediate_StiffnessUseMass = true;
 	FAutoConsoleVariableRef CVarStiffnessUseMass(TEXT("p.Chaos.ImmPhys.StiffnessUseMass"), ChaosImmediate_StiffnessUseMass, TEXT("Whether to use mass to scale stiffness in the conversion."));
 
+	float ChaosImmediate_JointMinProjection = 0.1f;
+	float ChaosImmediate_JointMaxProjection = 1.0f;
+	FAutoConsoleVariableRef CVarJointMinProjection(TEXT("p.Chaos.ImmPhys.JointMinProjection"), ChaosImmediate_JointMinProjection, TEXT("Joint min projection (for joints with projection disabled)."));
+	FAutoConsoleVariableRef CVarJointMaxProjection(TEXT("p.Chaos.ImmPhys.JointMaxProjection"), ChaosImmediate_JointMaxProjection, TEXT("Joint max projection (for joints with projection enabled)."));
+
 	float ChaosImmediate_DriveStiffnessSourceMin = 100.0f;
 	float ChaosImmediate_DriveStiffnessSourceMax = 2000.0f;
 	float ChaosImmediate_DriveStiffnessTargetMin = 0.1f;
@@ -116,6 +121,8 @@ namespace ImmediatePhysics_Chaos
 			FMath::DegreesToRadians(ConstraintInstance->GetAngularSwing1Limit()),
 			FMath::DegreesToRadians(ConstraintInstance->GetAngularSwing2Limit())
 		};
+		ConstraintSettings.Motion.LinearProjection = ConstraintInstance->IsProjectionEnabled() ? ChaosImmediate_JointMaxProjection : ChaosImmediate_JointMinProjection;
+		ConstraintSettings.Motion.AngularProjection = ConstraintInstance->IsProjectionEnabled() ? ChaosImmediate_JointMaxProjection : ChaosImmediate_JointMinProjection;
 		ConstraintSettings.Motion.bSoftLinearLimitsEnabled = ConstraintInstance->GetIsSoftLinearLimit();
 		ConstraintSettings.Motion.bSoftTwistLimitsEnabled = ConstraintInstance->GetIsSoftTwistLimit();
 		ConstraintSettings.Motion.bSoftSwingLimitsEnabled = ConstraintInstance->GetIsSoftSwingLimit();
diff --git a/Engine/Source/Runtime/Engine/Private/PhysicsEngine/ImmediatePhysics/ImmediatePhysicsChaos/ImmediatePhysicsSimulation_Chaos.cpp b/Engine/Source/Runtime/Engine/Private/PhysicsEngine/ImmediatePhysics/ImmediatePhysicsChaos/ImmediatePhysicsSimulation_Chaos.cpp
--- a/Engine/Source/Runtime/Engine/Private/PhysicsEngine/ImmediatePhysics/ImmediatePhysicsChaos/ImmediatePhysicsSimulation_Chaos.cpp
+++ b/Engine/Source/Runtime/Engine/Private/PhysicsEngine/ImmediatePhysics/ImmediatePhysicsChaos/ImmediatePhysicsSimulation_Chaos.cpp
@@ -45,7 +45,8 @@ int32 ChaosImmediate_Joint_EnableLinearLimits = 1;
 int32 ChaosImmediate_Joint_EnableTwistLimits = 1;
 int32 ChaosImmediate_Joint_EnableSwingLimits = 1;
 int32 ChaosImmediate_Joint_EnableDrives = 1;
-float ChaosImmediate_Joint_Projection = 0.1f;
+float ChaosImmediate_Joint_LinearProjection = 0.0f;
+float ChaosImmediate_Joint_AngularProjection = 0.0f;
 float ChaosImmediate_Joint_Stiffness = 1.0f;
 float ChaosImmediate_Joint_SoftLinearStiffness = 0.0f;
 float ChaosImmediate_Joint_SoftAngularStiffness = 0.0f;
@@ -59,7 +60,8 @@ FAutoConsoleVariableRef CVarChaosImmPhysEnableLinearLimits(TEXT("p.Chaos.ImmPhys
 FAutoConsoleVariableRef CVarChaosImmPhysEnableTwistLimits(TEXT("p.Chaos.ImmPhys.Joint.EnableTwistLimits"), ChaosImmediate_Joint_EnableTwistLimits, TEXT("EnableTwistLimits."));
 FAutoConsoleVariableRef CVarChaosImmPhysEnableSwingLimits(TEXT("p.Chaos.ImmPhys.Joint.EnableSwingLimits"), ChaosImmediate_Joint_EnableSwingLimits, TEXT("EnableSwingLimits."));
 FAutoConsoleVariableRef CVarChaosImmPhysEnableDrives(TEXT("p.Chaos.ImmPhys.Joint.EnableDrives"), ChaosImmediate_Joint_EnableDrives, TEXT("EnableDrives."));
-FAutoConsoleVariableRef CVarChaosImmPhysProjection(TEXT("p.Chaos.ImmPhys.Joint.Projection"), ChaosImmediate_Joint_Projection, TEXT("6Dof joint projection amount override (if > 0)."));
+FAutoConsoleVariableRef CVarChaosImmPhysLinearProjection(TEXT("p.Chaos.ImmPhys.Joint.LinearProjection"), ChaosImmediate_Joint_LinearProjection, TEXT("6Dof joint projection amount override (if > 0)."));
+FAutoConsoleVariableRef CVarChaosImmPhysAngularProjection(TEXT("p.Chaos.ImmPhys.Joint.AngularProjection"), ChaosImmediate_Joint_AngularProjection, TEXT("6Dof joint projection amount override (if > 0)."));
 FAutoConsoleVariableRef CVarChaosImmPhysStiffness(TEXT("p.Chaos.ImmPhys.Joint.Stiffness"), ChaosImmediate_Joint_Stiffness, TEXT("6Dof joint stiffness override (if > 0)."));
 FAutoConsoleVariableRef CVarChaosImmPhysSoftLinearStiffness(TEXT("p.Chaos.ImmPhys.Joint.SoftLinearStiffness"), ChaosImmediate_Joint_SoftLinearStiffness, TEXT("6Dof joint soft linear stiffness override (if > 0)."));
 FAutoConsoleVariableRef CVarChaosImmPhysSoftAngularStiffness(TEXT("p.Chaos.ImmPhys.Joint.SoftAngularStiffness"), ChaosImmediate_Joint_SoftAngularStiffness, TEXT("6Dof joint soft angular stiffness override (if > 0)."));
@@ -421,7 +423,8 @@ namespace ImmediatePhysics_Chaos
 			JointsSettings.bEnableTwistLimits = ChaosImmediate_Joint_EnableTwistLimits != 0;
 			JointsSettings.bEnableSwingLimits = ChaosImmediate_Joint_EnableSwingLimits != 0;
 			JointsSettings.bEnableDrives = ChaosImmediate_Joint_EnableDrives != 0;
-			JointsSettings.Projection = ChaosImmediate_Joint_Projection;
+			JointsSettings.LinearProjection = ChaosImmediate_Joint_LinearProjection;
+			JointsSettings.AngularProjection = ChaosImmediate_Joint_AngularProjection;
 			JointsSettings.Stiffness = ChaosImmediate_Joint_Stiffness;
 			JointsSettings.SoftLinearStiffness = ChaosImmediate_Joint_SoftLinearStiffness;
 			JointsSettings.SoftAngularStiffness = ChaosImmediate_Joint_SoftAngularStiffness;
diff --git a/Engine/Source/Runtime/Experimental/Chaos/Private/Chaos/PBDJointConstraintUtilities.cpp b/Engine/Source/Runtime/Experimental/Chaos/Private/Chaos/PBDJointConstraintUtilities.cpp
--- a/Engine/Source/Runtime/Experimental/Chaos/Private/Chaos/PBDJointConstraintUtilities.cpp
+++ b/Engine/Source/Runtime/Experimental/Chaos/Private/Chaos/PBDJointConstraintUtilities.cpp
@@ -348,53 +348,17 @@ namespace Chaos
 	}
 
 	
-	FReal GetLinearStiffness(
-		const FPBDJointSolverSettings& SolverSettings,
-		const FPBDJointSettings& JointSettings)
-	{
-		const FReal SolverStiffness = (SolverSettings.Stiffness > (FReal)0) ? SolverSettings.Stiffness : JointSettings.Motion.Stiffness;
-		const FReal SoftSolverStiffness = (SolverSettings.SoftLinearStiffness > (FReal)0) ? SolverSettings.SoftLinearStiffness : JointSettings.Motion.SoftLinearStiffness;
-		const bool bIsSoft = JointSettings.Motion.bSoftLinearLimitsEnabled && ((JointSettings.Motion.LinearMotionTypes[0] == EJointMotionType::Limited) || (JointSettings.Motion.LinearMotionTypes[1] == EJointMotionType::Limited) || (JointSettings.Motion.LinearMotionTypes[2] == EJointMotionType::Limited));
-		const FReal Stiffness = bIsSoft ? SolverStiffness * SoftSolverStiffness : SolverStiffness;
-		return Stiffness;
-	}
-
-	
-	FReal GetTwistStiffness(
-		const FPBDJointSolverSettings& SolverSettings,
-		const FPBDJointSettings& JointSettings)
-	{
-		const FReal SolverStiffness = (SolverSettings.Stiffness > (FReal)0) ? SolverSettings.Stiffness : JointSettings.Motion.Stiffness;
-		const FReal SoftSolverStiffness = (SolverSettings.SoftAngularStiffness > (FReal)0) ? SolverSettings.SoftAngularStiffness : JointSettings.Motion.SoftTwistStiffness;
-		const bool bIsSoft = JointSettings.Motion.bSoftTwistLimitsEnabled && (JointSettings.Motion.AngularMotionTypes[(int32)EJointAngularConstraintIndex::Twist] == EJointMotionType::Limited);
-		const FReal Stiffness = bIsSoft ? SolverStiffness * SoftSolverStiffness : SolverStiffness;
-		return Stiffness;
-	}
-
-	
-	FReal GetSwingStiffness(
-		const FPBDJointSolverSettings& SolverSettings,
-		const FPBDJointSettings& JointSettings)
-	{
-		const FReal SolverStiffness = (SolverSettings.Stiffness > (FReal)0) ? SolverSettings.Stiffness : JointSettings.Motion.Stiffness;
-		const FReal SoftSolverStiffness = (SolverSettings.SoftAngularStiffness > (FReal)0) ? SolverSettings.SoftAngularStiffness : JointSettings.Motion.SoftSwingStiffness;
-		const bool bIsSoft = JointSettings.Motion.bSoftSwingLimitsEnabled && ((JointSettings.Motion.AngularMotionTypes[(int32)EJointAngularConstraintIndex::Swing1] == EJointMotionType::Limited) || (JointSettings.Motion.AngularMotionTypes[(int32)EJointAngularConstraintIndex::Swing2] == EJointMotionType::Limited));
-		const FReal Stiffness = bIsSoft ? SolverStiffness * SoftSolverStiffness : SolverStiffness;
-		return Stiffness;
-	}
-
-	
 	void ApplyPositionDelta(
 		const FReal Dt,
 		const FPBDJointSolverSettings& SolverSettings,
 		const FPBDJointSettings& JointSettings,
+		const FReal Stiffness,
 		const int32 Index0,
 		const int32 Index1,
 		FVec3& P0,
 		FVec3& P1,
 		const FVec3& DP0,
-		const FVec3& DP1,
-		const FReal Stiffness)
+		const FVec3& DP1)
 	{
 		P0 = P0 + Stiffness * DP0;
 		P1 = P1 + Stiffness * DP1;
@@ -405,6 +369,7 @@ namespace Chaos
 		const FReal Dt,
 		const FPBDJointSolverSettings& SolverSettings,
 		const FPBDJointSettings& JointSettings,
+		const FReal Stiffness,
 		const int32 Index0,
 		const int32 Index1,
 		FVec3& P0,
@@ -412,13 +377,12 @@ namespace Chaos
 		FVec3& P1,
 		FVec3& V1,
 		const FVec3& DV0,
-		const FVec3& DV1, 
-		const FReal Stiffness)
+		const FVec3& DV1)
 	{
 		V0 = V0 + Stiffness * DV0;
 		V1 = V1 + Stiffness * DV1;
 
-		ApplyPositionDelta(Dt, SolverSettings, JointSettings, Index0, Index1, P0, P1, DV0 * Dt, DV1 * Dt, Stiffness);
+		ApplyPositionDelta(Dt, SolverSettings, JointSettings, Stiffness, Index0, Index1, P0, P1, DV0 * Dt, DV1 * Dt);
 	}
 
 	
@@ -426,13 +390,13 @@ namespace Chaos
 		const FReal Dt,
 		const FPBDJointSolverSettings& SolverSettings,
 		const FPBDJointSettings& JointSettings,
+		const FReal Stiffness,
 		const int32 Index0,
 		const int32 Index1,
 		FRotation3& Q0,
 		FRotation3& Q1,
 		const FVec3& DR0,
-		const FVec3& DR1,
-		const FReal Stiffness)
+		const FVec3& DR1)
 	{
 		const FRotation3 DQ0 = (FRotation3::FromElements(Stiffness * DR0, 0) * Q0) * (FReal)0.5;
 		const FRotation3 DQ1 = (FRotation3::FromElements(Stiffness * DR1, 0) * Q1) * (FReal)0.5;
@@ -447,6 +411,7 @@ namespace Chaos
 		const FReal Dt,
 		const FPBDJointSolverSettings& SolverSettings,
 		const FPBDJointSettings& JointSettings,
+		const FReal Stiffness,
 		const int32 Index0,
 		const int32 Index1,
 		FRotation3& Q0,
@@ -454,13 +419,12 @@ namespace Chaos
 		FRotation3& Q1,
 		FVec3& W1,
 		const FVec3& DW0,
-		const FVec3& DW1,
-		const FReal Stiffness)
+		const FVec3& DW1)
 	{
 		W0 = W0 + Stiffness * DW0;
 		W1 = W1 + Stiffness * DW1;
 
-		ApplyRotationDelta(Dt, SolverSettings, JointSettings, Index0, Index1, Q0, Q1, DW0 * Dt, DW1 * Dt, Stiffness);
+		ApplyRotationDelta(Dt, SolverSettings, JointSettings, Stiffness, Index0, Index1, Q0, Q1, DW0 * Dt, DW1 * Dt);
 	}
 
 
@@ -469,6 +433,7 @@ namespace Chaos
 		const FReal Dt,
 		const FPBDJointSolverSettings& SolverSettings,
 		const FPBDJointSettings& JointSettings,
+		const FReal Stiffness,
 		const int32 Index0,
 		const int32 Index1,
 		FRotation3& Q0,
@@ -481,8 +446,7 @@ namespace Chaos
 		const FMatrix33& InvIL1,
 		const FVec3& Axis0,
 		const FVec3& Axis1,
-		const FReal WC,
-		const FReal Stiffness)
+		const FReal WC)
 	{
 		const FMatrix33 InvI0 = Utilities::Multiply(Q0.ToMatrix(), Utilities::Multiply(InvIL0, Q0.ToMatrix().GetTransposed()));
 		const FMatrix33 InvI1 = Utilities::Multiply(Q1.ToMatrix(), Utilities::Multiply(InvIL1, Q1.ToMatrix().GetTransposed()));
@@ -494,7 +458,7 @@ namespace Chaos
 		const FVec3 DW0 = Utilities::Multiply(InvI0, Axis0) * L * WC;
 		const FVec3 DW1 = -Utilities::Multiply(InvI1, Axis1) * L * WC;
 
-		ApplyRotationVelocityDelta(Dt, SolverSettings, JointSettings, Index0, Index1, Q0, W0, Q1, W1, DW0, DW1, Stiffness);
+		ApplyRotationVelocityDelta(Dt, SolverSettings, JointSettings, Stiffness, Index0, Index1, Q0, W0, Q1, W1, DW0, DW1);
 	}
 
 	
@@ -502,6 +466,7 @@ namespace Chaos
 		const FReal Dt,
 		const FPBDJointSolverSettings& SolverSettings,
 		const FPBDJointSettings& JointSettings,
+		const FReal Stiffness,
 		const int32 Index0,
 		const int32 Index1,
 		FVec3& P0,
@@ -515,8 +480,7 @@ namespace Chaos
 		const FVec3& Axis0,
 		const FReal Angle0,
 		const FVec3& Axis1,
-		const FReal Angle1, 
-		const FReal Stiffness)
+		const FReal Angle1)
 	{
 		const FMatrix33 InvI0 = Utilities::Multiply(Q0.ToMatrix(), Utilities::Multiply(InvIL0, Q0.ToMatrix().GetTransposed()));
 		const FMatrix33 InvI1 = Utilities::Multiply(Q1.ToMatrix(), Utilities::Multiply(InvIL1, Q1.ToMatrix().GetTransposed()));
@@ -529,7 +493,7 @@ namespace Chaos
 		const FVec3 DR0 = Utilities::Multiply(InvI0, Axis0) * L * Angle0;
 		const FVec3 DR1 = Utilities::Multiply(InvI1, Axis1) * L * Angle1;
 
-		ApplyRotationDelta(Dt, SolverSettings, JointSettings, Index0, Index1, Q0, Q1, DR0, DR1, Stiffness);
+		ApplyRotationDelta(Dt, SolverSettings, JointSettings, Stiffness, Index0, Index1, Q0, Q1, DR0, DR1);
 	}
 
 	
@@ -683,6 +647,7 @@ namespace Chaos
 		const FReal Dt,
 		const FPBDJointSolverSettings& SolverSettings,
 		const FPBDJointSettings& JointSettings,
+		const FReal Stiffness,
 		const int32 Index0,
 		const int32 Index1,
 		FVec3& P0,
@@ -724,8 +689,8 @@ namespace Chaos
 		const FVec3 DP1 = -InvM1 * DX;
 		const FVec3 DR0 = Utilities::Multiply(InvI0, FVec3::CrossProduct(X0 - P0, DX));
 		const FVec3 DR1 = Utilities::Multiply(InvI1, FVec3::CrossProduct(X1 - P1, -DX));
-		ApplyPositionDelta(Dt, SolverSettings, JointSettings, Index0, Index1, P0, P1, DP0, DP1, GetLinearStiffness(SolverSettings, JointSettings));
-		ApplyRotationDelta(Dt, SolverSettings, JointSettings, Index0, Index1, Q0, Q1, DR0, DR1, GetLinearStiffness(SolverSettings, JointSettings));
+		ApplyPositionDelta(Dt, SolverSettings, JointSettings, Stiffness, Index0, Index1, P0, P1, DP0, DP1);
+		ApplyRotationDelta(Dt, SolverSettings, JointSettings, Stiffness, Index0, Index1, Q0, Q1, DR0, DR1);
 	}
 
 	
@@ -733,6 +698,7 @@ namespace Chaos
 		const FReal Dt,
 		const FPBDJointSolverSettings& SolverSettings,
 		const FPBDJointSettings& JointSettings,
+		const FReal Stiffness,
 		const int32 Index0,
 		const int32 Index1,
 		FVec3& P0,
@@ -780,8 +746,8 @@ namespace Chaos
 		const FVec3 DW0 = Utilities::Multiply(InvI0, FVec3::CrossProduct(XC0, DL));
 		const FVec3 DW1 = -Utilities::Multiply(InvI1, FVec3::CrossProduct(XC1, DL));
 
-		ApplyVelocityDelta(Dt, SolverSettings, JointSettings, Index0, Index1, P0, V0, P1, V1, DV0, DV1, GetLinearStiffness(SolverSettings, JointSettings));
-		ApplyRotationVelocityDelta(Dt, SolverSettings, JointSettings, Index0, Index1, Q0, W0, Q1, W1, DW0, DW1, GetLinearStiffness(SolverSettings, JointSettings));
+		ApplyVelocityDelta(Dt, SolverSettings, JointSettings, Stiffness, Index0, Index1, P0, V0, P1, V1, DV0, DV1);
+		ApplyRotationVelocityDelta(Dt, SolverSettings, JointSettings, Stiffness, Index0, Index1, Q0, W0, Q1, W1, DW0, DW1);
 	}
 
 
@@ -790,6 +756,7 @@ namespace Chaos
 		const FReal Dt,
 		const FPBDJointSolverSettings& SolverSettings,
 		const FPBDJointSettings& JointSettings,
+		const FReal Stiffness,
 		const int32 Index0,
 		const int32 Index1,
 		FVec3& P0,
@@ -855,7 +822,7 @@ namespace Chaos
 		const FReal DTwistAngle1 = -DTwistAngle;
 
 		// Apply twist correction
-		ApplyRotationDelta(Dt, SolverSettings, JointSettings, Index0, Index1, P0, Q0, P1, Q1, InvM0, InvIL0, InvM1, InvIL1, TwistAxis0, DTwistAngle0, TwistAxis1, DTwistAngle1, GetTwistStiffness(SolverSettings, JointSettings));
+		ApplyRotationDelta(Dt, SolverSettings, JointSettings, Stiffness, Index0, Index1, P0, Q0, P1, Q1, InvM0, InvIL0, InvM1, InvIL1, TwistAxis0, DTwistAngle0, TwistAxis1, DTwistAngle1);
 	}
 
 	
@@ -863,6 +830,7 @@ namespace Chaos
 		const FReal Dt,
 		const FPBDJointSolverSettings& SolverSettings,
 		const FPBDJointSettings& JointSettings,
+		const FReal Stiffness,
 		const int32 Index0,
 		const int32 Index1,
 		FVec3& P0,
@@ -931,7 +899,7 @@ namespace Chaos
 			DW = FMath::Min((FReal)0, WC1 - WC0);
 		}
 
-		ApplyRotationVelocityDelta(Dt, SolverSettings, JointSettings, Index0, Index1, Q0, W0, Q1, W1, InvM0, InvIL0, InvM1, InvIL1, TwistAxis0, TwistAxis1, DW, GetTwistStiffness(SolverSettings, JointSettings));
+		ApplyRotationVelocityDelta(Dt, SolverSettings, JointSettings, Stiffness, Index0, Index1, Q0, W0, Q1, W1, InvM0, InvIL0, InvM1, InvIL1, TwistAxis0, TwistAxis1, DW);
 	}
 
 
@@ -940,6 +908,7 @@ namespace Chaos
 		const FReal Dt,
 		const FPBDJointSolverSettings& SolverSettings,
 		const FPBDJointSettings& JointSettings,
+		const FReal Stiffness,
 		const int32 Index0,
 		const int32 Index1,
 		FVec3& P0,
@@ -1008,7 +977,7 @@ namespace Chaos
 		FReal DSwingAngle1 = -DSwingAngle;
 
 		// Apply swing correction
-		ApplyRotationDelta(Dt, SolverSettings, JointSettings, Index0, Index1, P0, Q0, P1, Q1, InvM0, InvIL0, InvM1, InvIL1, SwingAxis0, DSwingAngle0, SwingAxis1, DSwingAngle1, GetSwingStiffness(SolverSettings, JointSettings));
+		ApplyRotationDelta(Dt, SolverSettings, JointSettings, Stiffness, Index0, Index1, P0, Q0, P1, Q1, InvM0, InvIL0, InvM1, InvIL1, SwingAxis0, DSwingAngle0, SwingAxis1, DSwingAngle1);
 	}
 
 	
@@ -1016,6 +985,7 @@ namespace Chaos
 		const FReal Dt,
 		const FPBDJointSolverSettings& SolverSettings,
 		const FPBDJointSettings& JointSettings,
+		const FReal Stiffness,
 		const int32 Index0,
 		const int32 Index1,
 		FVec3& P0,
@@ -1086,7 +1056,7 @@ namespace Chaos
 			DW = FMath::Min((FReal)0, WC1 - WC0);
 		}
 
-		ApplyRotationVelocityDelta(Dt, SolverSettings, JointSettings, Index0, Index1, Q0, W0, Q1, W1, InvM0, InvIL0, InvM1, InvIL1, SwingAxis0, SwingAxis1, DW, GetSwingStiffness(SolverSettings, JointSettings));
+		ApplyRotationVelocityDelta(Dt, SolverSettings, JointSettings, Stiffness, Index0, Index1, Q0, W0, Q1, W1, InvM0, InvIL0, InvM1, InvIL1, SwingAxis0, SwingAxis1, DW);
 	}
 
 
@@ -1095,6 +1065,7 @@ namespace Chaos
 		const FReal Dt,
 		const FPBDJointSolverSettings& SolverSettings,
 		const FPBDJointSettings& JointSettings,
+		const FReal Stiffness,
 		const int32 Index0,
 		const int32 Index1,
 		const EJointAngularConstraintIndex SwingConstraintIndex,
@@ -1181,7 +1152,7 @@ namespace Chaos
 			FReal DSwingAngle1 = -DSwingAngle;
 
 			// Apply swing correction
-			ApplyRotationDelta(Dt, SolverSettings, JointSettings, Index0, Index1, P0, Q0, P1, Q1, InvM0, InvIL0, InvM1, InvIL1, SwingAxis0, DSwingAngle0, SwingAxis1, DSwingAngle1, GetSwingStiffness(SolverSettings, JointSettings));
+			ApplyRotationDelta(Dt, SolverSettings, JointSettings, Stiffness, Index0, Index1, P0, Q0, P1, Q1, InvM0, InvIL0, InvM1, InvIL1, SwingAxis0, DSwingAngle0, SwingAxis1, DSwingAngle1);
 		}
 	}
 
@@ -1190,6 +1161,7 @@ namespace Chaos
 		const FReal Dt,
 		const FPBDJointSolverSettings& SolverSettings,
 		const FPBDJointSettings& JointSettings,
+		const FReal Stiffness,
 		const int32 Index0,
 		const int32 Index1,
 		const EJointAngularConstraintIndex SwingConstraintIndex,
@@ -1278,7 +1250,7 @@ namespace Chaos
 				DW = FMath::Min((FReal)0, WC1 - WC0);
 			}
 
-			ApplyRotationVelocityDelta(Dt, SolverSettings, JointSettings, Index0, Index1, Q0, W0, Q1, W1, InvM0, InvIL0, InvM1, InvIL1, SwingAxis0, SwingAxis1, DW, GetSwingStiffness(SolverSettings, JointSettings));
+			ApplyRotationVelocityDelta(Dt, SolverSettings, JointSettings, Stiffness, Index0, Index1, Q0, W0, Q1, W1, InvM0, InvIL0, InvM1, InvIL1, SwingAxis0, SwingAxis1, DW);
 		}
 	}
 
@@ -1334,7 +1306,7 @@ namespace Chaos
 		const FReal DTwistAngle0 = DTwistAngle;
 		const FReal DTwistAngle1 = -DTwistAngle;
 
-		ApplyRotationDelta(Dt, SolverSettings, JointSettings, Index0, Index1, P0, Q0, P1, Q1, InvM0, InvIL0, InvM1, InvIL1, TwistAxis0, DTwistAngle0, TwistAxis1, DTwistAngle1, DriveStiffness);
+		ApplyRotationDelta(Dt, SolverSettings, JointSettings, DriveStiffness, Index0, Index1, P0, Q0, P1, Q1, InvM0, InvIL0, InvM1, InvIL1, TwistAxis0, DTwistAngle0, TwistAxis1, DTwistAngle1);
 	}
 
 	
@@ -1391,7 +1363,7 @@ namespace Chaos
 		FReal DSwingAngle1 = -DSwingAngle;
 
 		// Apply swing correction
-		ApplyRotationDelta(Dt, SolverSettings, JointSettings, Index0, Index1, P0, Q0, P1, Q1, InvM0, InvIL0, InvM1, InvIL1, SwingAxis0, DSwingAngle0, SwingAxis1, DSwingAngle1, DriveStiffness);
+		ApplyRotationDelta(Dt, SolverSettings, JointSettings, DriveStiffness, Index0, Index1, P0, Q0, P1, Q1, InvM0, InvIL0, InvM1, InvIL1, SwingAxis0, DSwingAngle0, SwingAxis1, DSwingAngle1);
 	}
 
 	
@@ -1450,6 +1422,7 @@ namespace Chaos
 		const FReal Dt,
 		const FPBDJointSolverSettings& SolverSettings,
 		const FPBDJointSettings& JointSettings,
+		const FReal Stiffness,
 		const int32 Index0,
 		const int32 Index1,
 		FVec3& P0,
@@ -1459,8 +1432,7 @@ namespace Chaos
 		FReal InvM0,
 		const FMatrix33& InvIL0,
 		FReal InvM1,
-		const FMatrix33& InvIL1,
-		const FReal ProjectionFactor)
+		const FMatrix33& InvIL1)
 	{
 		const FRigidTransform3& XL0 = JointSettings.ConstraintFrames[Index0];
 		const FRigidTransform3& XL1 = JointSettings.ConstraintFrames[Index1];
@@ -1469,8 +1441,8 @@ namespace Chaos
 		const FRotation3 R0 = Q0 * XL0.GetRotation();
 		const FVec3 CX = GetLimitedPositionError(JointSettings, R0, X1 - X0);
 
-		const FVec3 DP0 = CX * (ProjectionFactor * InvM0 / (InvM0 + InvM1));
-		const FVec3 DP1 = -CX * (ProjectionFactor * InvM1 / (InvM0 + InvM1));
+		const FVec3 DP0 = CX * (Stiffness * InvM0 / (InvM0 + InvM1));
+		const FVec3 DP1 = -CX * (Stiffness* InvM1 / (InvM0 + InvM1));
 		P0 = P0 + DP0;
 		P1 = P1 + DP1;
 	}
@@ -1480,6 +1452,7 @@ namespace Chaos
 		const FReal Dt,
 		const FPBDJointSolverSettings& SolverSettings,
 		const FPBDJointSettings& JointSettings,
+		const FReal Stiffness,
 		const int32 Index0,
 		const int32 Index1,
 		FVec3& P0,
@@ -1489,10 +1462,9 @@ namespace Chaos
 		FReal InvM0,
 		const FMatrix33& InvIL0,
 		FReal InvM1,
-		const FMatrix33& InvIL1,
-		const FReal ProjectionFactor)
+		const FMatrix33& InvIL1)
 	{
-
+		ApplyJointTwistConstraint(Dt, SolverSettings, JointSettings, Stiffness, Index0, Index1, P0, Q0, P1, Q1, InvM0, InvIL0, InvM1, InvIL1);
 	}
 
 	
@@ -1500,6 +1472,7 @@ namespace Chaos
 		const FReal Dt,
 		const FPBDJointSolverSettings& SolverSettings,
 		const FPBDJointSettings& JointSettings,
+		const FReal Stiffness,
 		const int32 Index0,
 		const int32 Index1,
 		FVec3& P0,
@@ -1509,10 +1482,9 @@ namespace Chaos
 		FReal InvM0,
 		const FMatrix33& InvIL0,
 		FReal InvM1,
-		const FMatrix33& InvIL1,
-		const FReal ProjectionFactor)
+		const FMatrix33& InvIL1)
 	{
-
+		ApplyJointConeConstraint(Dt, SolverSettings, JointSettings, Stiffness, Index0, Index1, P0, Q0, P1, Q1, InvM0, InvIL0, InvM1, InvIL1);
 	}
 
 	
@@ -1520,9 +1492,11 @@ namespace Chaos
 		const FReal Dt,
 		const FPBDJointSolverSettings& SolverSettings,
 		const FPBDJointSettings& JointSettings,
+		const FReal Stiffness,
 		const int32 Index0,
 		const int32 Index1,
-		const EJointAngularConstraintIndex SwingConstraint,
+		const EJointAngularConstraintIndex SwingConstraintIndex,
+		const EJointAngularAxisIndex SwingAxisIndex,
 		FVec3& P0,
 		FRotation3& Q0,
 		FVec3& P1,
@@ -1530,9 +1504,8 @@ namespace Chaos
 		FReal InvM0,
 		const FMatrix33& InvIL0,
 		FReal InvM1,
-		const FMatrix33& InvIL1,
-		const FReal ProjectionFactor)
+		const FMatrix33& InvIL1)
 	{
-
+		ApplyJointSwingConstraint(Dt, SolverSettings, JointSettings, Stiffness, Index0, Index1, SwingConstraintIndex, SwingAxisIndex, P0, Q0, P1, Q1, InvM0, InvIL0, InvM1, InvIL1);
 	}
 }
diff --git a/Engine/Source/Runtime/Experimental/Chaos/Private/Chaos/PBDJointConstraints.cpp b/Engine/Source/Runtime/Experimental/Chaos/Private/Chaos/PBDJointConstraints.cpp
--- a/Engine/Source/Runtime/Experimental/Chaos/Private/Chaos/PBDJointConstraints.cpp
+++ b/Engine/Source/Runtime/Experimental/Chaos/Private/Chaos/PBDJointConstraints.cpp
@@ -16,6 +16,40 @@ namespace Chaos
 {
 	DECLARE_CYCLE_STAT(TEXT("TPBDJointConstraints::Apply"), STAT_ApplyJointConstraints, STATGROUP_Chaos);
 
+	FReal GetLinearStiffness(
+		const FPBDJointSolverSettings& SolverSettings,
+		const FPBDJointSettings& JointSettings)
+	{
+		const FReal SolverStiffness = (SolverSettings.Stiffness > (FReal)0) ? SolverSettings.Stiffness : JointSettings.Motion.Stiffness;
+		const FReal SoftSolverStiffness = (SolverSettings.SoftLinearStiffness > (FReal)0) ? SolverSettings.SoftLinearStiffness : JointSettings.Motion.SoftLinearStiffness;
+		const bool bIsSoft = JointSettings.Motion.bSoftLinearLimitsEnabled && ((JointSettings.Motion.LinearMotionTypes[0] == EJointMotionType::Limited) || (JointSettings.Motion.LinearMotionTypes[1] == EJointMotionType::Limited) || (JointSettings.Motion.LinearMotionTypes[2] == EJointMotionType::Limited));
+		const FReal Stiffness = bIsSoft ? SolverStiffness * SoftSolverStiffness : SolverStiffness;
+		return Stiffness;
+	}
+
+
+	FReal GetTwistStiffness(
+		const FPBDJointSolverSettings& SolverSettings,
+		const FPBDJointSettings& JointSettings)
+	{
+		const FReal SolverStiffness = (SolverSettings.Stiffness > (FReal)0) ? SolverSettings.Stiffness : JointSettings.Motion.Stiffness;
+		const FReal SoftSolverStiffness = (SolverSettings.SoftAngularStiffness > (FReal)0) ? SolverSettings.SoftAngularStiffness : JointSettings.Motion.SoftTwistStiffness;
+		const bool bIsSoft = JointSettings.Motion.bSoftTwistLimitsEnabled && (JointSettings.Motion.AngularMotionTypes[(int32)EJointAngularConstraintIndex::Twist] == EJointMotionType::Limited);
+		const FReal Stiffness = bIsSoft ? SolverStiffness * SoftSolverStiffness : SolverStiffness;
+		return Stiffness;
+	}
+
+
+	FReal GetSwingStiffness(
+		const FPBDJointSolverSettings& SolverSettings,
+		const FPBDJointSettings& JointSettings)
+	{
+		const FReal SolverStiffness = (SolverSettings.Stiffness > (FReal)0) ? SolverSettings.Stiffness : JointSettings.Motion.Stiffness;
+		const FReal SoftSolverStiffness = (SolverSettings.SoftAngularStiffness > (FReal)0) ? SolverSettings.SoftAngularStiffness : JointSettings.Motion.SoftSwingStiffness;
+		const bool bIsSoft = JointSettings.Motion.bSoftSwingLimitsEnabled && ((JointSettings.Motion.AngularMotionTypes[(int32)EJointAngularConstraintIndex::Swing1] == EJointMotionType::Limited) || (JointSettings.Motion.AngularMotionTypes[(int32)EJointAngularConstraintIndex::Swing2] == EJointMotionType::Limited));
+		const FReal Stiffness = bIsSoft ? SolverStiffness * SoftSolverStiffness : SolverStiffness;
+		return Stiffness;
+	}
 
 	//
 	// Constraint Handle
@@ -62,7 +96,8 @@ namespace Chaos
 	
 	FPBDJointMotionSettings::FPBDJointMotionSettings()
 		: Stiffness((FReal)1)
-		, Projection((FReal)0)
+		, LinearProjection((FReal)0)
+		, AngularProjection((FReal)0)
 		, LinearMotionTypes({ EJointMotionType::Locked, EJointMotionType::Locked, EJointMotionType::Locked })
 		, LinearLimit(FLT_MAX)
 		, AngularMotionTypes({ EJointMotionType::Free, EJointMotionType::Free, EJointMotionType::Free })
@@ -85,7 +120,8 @@ namespace Chaos
 	
 	FPBDJointMotionSettings::FPBDJointMotionSettings(const TVector<EJointMotionType, 3>& InLinearMotionTypes, const TVector<EJointMotionType, 3>& InAngularMotionTypes)
 		: Stiffness((FReal)1)
-		, Projection((FReal)0)
+		, LinearProjection((FReal)0)
+		, AngularProjection((FReal)0)
 		, LinearMotionTypes(InLinearMotionTypes)
 		, LinearLimit(FLT_MAX)
 		, AngularMotionTypes({ EJointMotionType::Free, EJointMotionType::Free, EJointMotionType::Free })
@@ -132,7 +168,8 @@ namespace Chaos
 		, bEnableTwistLimits(true)
 		, bEnableSwingLimits(true)
 		, bEnableDrives(true)
-		, Projection((FReal)0)
+		, LinearProjection((FReal)0)
+		, AngularProjection((FReal)0)
 		, Stiffness((FReal)0)
 		, DriveStiffness((FReal)0)
 		, SoftLinearStiffness((FReal)0)
@@ -408,19 +445,16 @@ namespace Chaos
 	
 	void FPBDJointConstraints::ApplyProjection(const FReal Dt, const TArray<FConstraintContainerHandle*>& InConstraintHandles)
 	{
-		if (Settings.Projection > 0)
-		{
-			TArray<FConstraintContainerHandle*> SortedConstraintHandles = InConstraintHandles;
-			SortedConstraintHandles.Sort([](const FConstraintContainerHandle& L, const FConstraintContainerHandle& R)
-				{
-					// Sort bodies from root to leaf
-					return L.GetConstraintLevel() < R.GetConstraintLevel();
-				});
-
-			for (FConstraintContainerHandle* ConstraintHandle : SortedConstraintHandles)
+		TArray<FConstraintContainerHandle*> SortedConstraintHandles = InConstraintHandles;
+		SortedConstraintHandles.Sort([](const FConstraintContainerHandle& L, const FConstraintContainerHandle& R)
 			{
-				ProjectPosition(Dt, ConstraintHandle->GetConstraintIndex(), 0, 1);
-			}
+				// Sort bodies from root to leaf
+				return L.GetConstraintLevel() < R.GetConstraintLevel();
+			});
+
+		for (FConstraintContainerHandle* ConstraintHandle : SortedConstraintHandles)
+		{
+			ProjectPosition(Dt, ConstraintHandle->GetConstraintIndex(), 0, 1);
 		}
 	}
 
@@ -455,6 +489,10 @@ namespace Chaos
 
 		Q1.EnforceShortestArcWith(Q0);
 
+		FReal LinearStiffness = GetLinearStiffness(Settings, JointSettings);
+		FReal TwistStiffness = GetTwistStiffness(Settings, JointSettings);
+		FReal SwingStiffness = GetSwingStiffness(Settings, JointSettings);
+
 		// Adjust mass for stability
 		const int32 Level0 = ConstraintStates[ConstraintIndex].ParticleLevels[Index0];
 		const int32 Level1 = ConstraintStates[ConstraintIndex].ParticleLevels[Index1];
@@ -513,7 +551,7 @@ namespace Chaos
 		{
 			if (TwistMotion != EJointMotionType::Free)
 			{
-				FPBDJointUtilities::ApplyJointTwistVelocityConstraint(Dt, Settings, JointSettings, Index0, Index1, P0, Q0, V0, W0, P1, Q1, V1, W1, InvM0, InvIL0, InvM1, InvIL1);
+				FPBDJointUtilities::ApplyJointTwistVelocityConstraint(Dt, Settings, JointSettings, TwistStiffness, Index0, Index1, P0, Q0, V0, W0, P1, Q1, V1, W1, InvM0, InvIL0, InvM1, InvIL1);
 			}
 		}
 
@@ -523,19 +561,19 @@ namespace Chaos
 			if ((Swing1Motion == EJointMotionType::Limited) && (Swing2Motion == EJointMotionType::Limited))
 			{
 				// Swing Cone
-				FPBDJointUtilities::ApplyJointConeVelocityConstraint(Dt, Settings, JointSettings, Index0, Index1, P0, Q0, V0, W0, P1, Q1, V1, W1, InvM0, InvIL0, InvM1, InvIL1);
+				FPBDJointUtilities::ApplyJointConeVelocityConstraint(Dt, Settings, JointSettings, SwingStiffness, Index0, Index1, P0, Q0, V0, W0, P1, Q1, V1, W1, InvM0, InvIL0, InvM1, InvIL1);
 			}
 			else
 			{
 				if (Swing1Motion != EJointMotionType::Free)
 				{
 					// Swing Arc/Lock
-					FPBDJointUtilities::ApplyJointSwingVelocityConstraint(Dt, Settings, JointSettings, Index0, Index1, EJointAngularConstraintIndex::Swing1, EJointAngularAxisIndex::Swing1, P0, Q0, V0, W0, P1, Q1, V1, W1, InvM0, InvIL0, InvM1, InvIL1);
+					FPBDJointUtilities::ApplyJointSwingVelocityConstraint(Dt, Settings, JointSettings, SwingStiffness, Index0, Index1, EJointAngularConstraintIndex::Swing1, EJointAngularAxisIndex::Swing1, P0, Q0, V0, W0, P1, Q1, V1, W1, InvM0, InvIL0, InvM1, InvIL1);
 				}
 				if (Swing2Motion != EJointMotionType::Free)
 				{
 					// Swing Arc/Lock
-					FPBDJointUtilities::ApplyJointSwingVelocityConstraint(Dt, Settings, JointSettings, Index0, Index1, EJointAngularConstraintIndex::Swing2, EJointAngularAxisIndex::Swing2, P0, Q0, V0, W0, P1, Q1, V1, W1, InvM0, InvIL0, InvM1, InvIL1);
+					FPBDJointUtilities::ApplyJointSwingVelocityConstraint(Dt, Settings, JointSettings, SwingStiffness, Index0, Index1, EJointAngularConstraintIndex::Swing2, EJointAngularAxisIndex::Swing2, P0, Q0, V0, W0, P1, Q1, V1, W1, InvM0, InvIL0, InvM1, InvIL1);
 				}
 			}
 		}
@@ -543,7 +581,7 @@ namespace Chaos
 		// Apply linear velocity  constraints
 		if ((LinearMotion[0] != EJointMotionType::Free) || (LinearMotion[1] != EJointMotionType::Free) || (LinearMotion[2] != EJointMotionType::Free))
 		{
-			FPBDJointUtilities::ApplyJointVelocityConstraint(Dt, Settings, JointSettings, Index0, Index1, P0, Q0, V0, W0, P1, Q1, V1, W1, InvM0, InvIL0, InvM1, InvIL1);
+			FPBDJointUtilities::ApplyJointVelocityConstraint(Dt, Settings, JointSettings, LinearStiffness, Index0, Index1, P0, Q0, V0, W0, P1, Q1, V1, W1, InvM0, InvIL0, InvM1, InvIL1);
 		}
 
 		// Update the particles
@@ -590,6 +628,10 @@ namespace Chaos
 
 		Q1.EnforceShortestArcWith(Q0);
 
+		FReal LinearStiffness = GetLinearStiffness(Settings, JointSettings);
+		FReal TwistStiffness = GetTwistStiffness(Settings, JointSettings);
+		FReal SwingStiffness = GetSwingStiffness(Settings, JointSettings);
+
 		// Adjust mass for stability
 		const int32 Level0 = ConstraintStates[ConstraintIndex].ParticleLevels[Index0];
 		const int32 Level1 = ConstraintStates[ConstraintIndex].ParticleLevels[Index1];
@@ -670,7 +712,7 @@ namespace Chaos
 		{
 			if (TwistMotion != EJointMotionType::Free)
 			{
-				FPBDJointUtilities::ApplyJointTwistConstraint(Dt, Settings, JointSettings, Index0, Index1, P0, Q0, P1, Q1, InvM0, InvIL0, InvM1, InvIL1);
+				FPBDJointUtilities::ApplyJointTwistConstraint(Dt, Settings, JointSettings, TwistStiffness, Index0, Index1, P0, Q0, P1, Q1, InvM0, InvIL0, InvM1, InvIL1);
 			}
 		}
 
@@ -680,19 +722,19 @@ namespace Chaos
 			if ((Swing1Motion == EJointMotionType::Limited) && (Swing2Motion == EJointMotionType::Limited))
 			{
 				// Swing Cone
-				FPBDJointUtilities::ApplyJointConeConstraint(Dt, Settings, JointSettings, Index0, Index1, P0, Q0, P1, Q1, InvM0, InvIL0, InvM1, InvIL1);
+				FPBDJointUtilities::ApplyJointConeConstraint(Dt, Settings, JointSettings, SwingStiffness, Index0, Index1, P0, Q0, P1, Q1, InvM0, InvIL0, InvM1, InvIL1);
 			}
 			else
 			{
 				if (Swing1Motion != EJointMotionType::Free)
 				{
 					// Swing Arc/Lock
-					FPBDJointUtilities::ApplyJointSwingConstraint(Dt, Settings, JointSettings, Index0, Index1, EJointAngularConstraintIndex::Swing1, EJointAngularAxisIndex::Swing1, P0, Q0, P1, Q1, InvM0, InvIL0, InvM1, InvIL1);
+					FPBDJointUtilities::ApplyJointSwingConstraint(Dt, Settings, JointSettings, SwingStiffness, Index0, Index1, EJointAngularConstraintIndex::Swing1, EJointAngularAxisIndex::Swing1, P0, Q0, P1, Q1, InvM0, InvIL0, InvM1, InvIL1);
 				}
 				if (Swing2Motion != EJointMotionType::Free)
 				{
 					// Swing Arc/Lock
-					FPBDJointUtilities::ApplyJointSwingConstraint(Dt, Settings, JointSettings, Index0, Index1, EJointAngularConstraintIndex::Swing2, EJointAngularAxisIndex::Swing2, P0, Q0, P1, Q1, InvM0, InvIL0, InvM1, InvIL1);
+					FPBDJointUtilities::ApplyJointSwingConstraint(Dt, Settings, JointSettings, SwingStiffness, Index0, Index1, EJointAngularConstraintIndex::Swing2, EJointAngularAxisIndex::Swing2, P0, Q0, P1, Q1, InvM0, InvIL0, InvM1, InvIL1);
 				}
 			}
 		}
@@ -700,7 +742,7 @@ namespace Chaos
 		// Apply linear constraints
 		if ((LinearMotion[0] != EJointMotionType::Free) || (LinearMotion[1] != EJointMotionType::Free) || (LinearMotion[2] != EJointMotionType::Free))
 		{
-			FPBDJointUtilities::ApplyJointPositionConstraint(Dt, Settings, JointSettings, Index0, Index1, P0, Q0, P1, Q1, InvM0, InvIL0, InvM1, InvIL1);
+			FPBDJointUtilities::ApplyJointPositionConstraint(Dt, Settings, JointSettings, LinearStiffness, Index0, Index1, P0, Q0, P1, Q1, InvM0, InvIL0, InvM1, InvIL1);
 		}
 
 		// Update the particles
@@ -722,14 +764,15 @@ namespace Chaos
 		const FPBDJointSettings& JointSettings = ConstraintSettings[ConstraintIndex];
 
 		// Scale projection up to ProjectionSetting over NumProjectionIts
-		const FReal ProjectionFactor = (Settings.Projection > 0) ? Settings.Projection : JointSettings.Motion.Projection;
-		if (ProjectionFactor == (FReal)0)
+		const FReal LinearProjectionFactor = (Settings.LinearProjection > 0) ? Settings.LinearProjection : JointSettings.Motion.LinearProjection;
+		const FReal AngularProjectionFactor = (Settings.AngularProjection > 0) ? Settings.AngularProjection : JointSettings.Motion.AngularProjection;
+		if ((LinearProjectionFactor == (FReal)0) && (AngularProjectionFactor == (FReal)0))
 		{
 			return;
 		}
 
 		const TVector<TGeometryParticleHandle<FReal, 3>*, 2>& Constraint = ConstraintParticles[ConstraintIndex];
-		UE_LOG(LogChaosJoint, Verbose, TEXT("Project Joint Constraint %d %f (it = %d / %d)"), ConstraintIndex, ProjectionFactor, It, NumIts);
+		UE_LOG(LogChaosJoint, Verbose, TEXT("Project Joint Constraint %d %f %f (it = %d / %d)"), ConstraintIndex, LinearProjectionFactor, AngularProjectionFactor, It, NumIts);
 
 		// Switch particles - internally we assume the first body is the parent (i.e., the space in which constraint limits are specified)
 		const int32 Index0 = 1;
@@ -760,8 +803,42 @@ namespace Chaos
 			InvIL1 = FMatrix33(0, 0, 0);
 		}
 
-		// Project position error
-		FPBDJointUtilities::ApplyJointPositionProjection(Dt, Settings, JointSettings, Index0, Index1, P0, Q0, P1, Q1, InvM0, InvIL0, InvM1, InvIL1, ProjectionFactor);
+		const TVector<EJointMotionType, 3>& LinearMotion = JointSettings.Motion.LinearMotionTypes;
+		EJointMotionType TwistMotion = JointSettings.Motion.AngularMotionTypes[(int32)EJointAngularConstraintIndex::Twist];
+		EJointMotionType Swing1Motion = JointSettings.Motion.AngularMotionTypes[(int32)EJointAngularConstraintIndex::Swing1];
+		EJointMotionType Swing2Motion = JointSettings.Motion.AngularMotionTypes[(int32)EJointAngularConstraintIndex::Swing2];
+
+		if (AngularProjectionFactor > 0)
+		{
+			// Remove Twist Error
+			FPBDJointUtilities::ApplyJointTwistProjection(Dt, Settings, JointSettings, AngularProjectionFactor, Index0, Index1, P0, Q0, P1, Q1, InvM0, InvIL0, InvM1, InvIL1);
+
+			// Remove Swing Error
+			if ((Swing1Motion == EJointMotionType::Limited) && (Swing2Motion == EJointMotionType::Limited))
+			{
+				FPBDJointUtilities::ApplyJointConeProjection(Dt, Settings, JointSettings, AngularProjectionFactor, Index0, Index1, P0, Q0, P1, Q1, InvM0, InvIL0, InvM1, InvIL1);
+			}
+			else
+			{
+				if (Swing1Motion != EJointMotionType::Free)
+				{
+					FPBDJointUtilities::ApplyJointSwingProjection(Dt, Settings, JointSettings, AngularProjectionFactor, Index0, Index1, EJointAngularConstraintIndex::Swing1, EJointAngularAxisIndex::Swing1, P0, Q0, P1, Q1, InvM0, InvIL0, InvM1, InvIL1);
+				}
+				if (Swing2Motion != EJointMotionType::Free)
+				{
+					FPBDJointUtilities::ApplyJointSwingProjection(Dt, Settings, JointSettings, AngularProjectionFactor, Index0, Index1, EJointAngularConstraintIndex::Swing2, EJointAngularAxisIndex::Swing2, P0, Q0, P1, Q1, InvM0, InvIL0, InvM1, InvIL1);
+				}
+			}
+		}
+
+		// Remove Position Error
+		if (LinearProjectionFactor > 0)
+		{
+			if ((LinearMotion[0] != EJointMotionType::Free) || (LinearMotion[1] != EJointMotionType::Free) || (LinearMotion[2] != EJointMotionType::Free))
+			{
+				FPBDJointUtilities::ApplyJointPositionProjection(Dt, Settings, JointSettings, LinearProjectionFactor, Index0, Index1, P0, Q0, P1, Q1, InvM0, InvIL0, InvM1, InvIL1);
+			}
+		}
 
 		// Update the particles
 		TPBDRigidParticleHandle<FReal, 3>* Rigid0 = ConstraintParticles[ConstraintIndex][Index0]->AsDynamic();
diff --git a/Engine/Source/Runtime/Experimental/Chaos/Public/Chaos/PBDJointConstraintTypes.h b/Engine/Source/Runtime/Experimental/Chaos/Public/Chaos/PBDJointConstraintTypes.h
--- a/Engine/Source/Runtime/Experimental/Chaos/Public/Chaos/PBDJointConstraintTypes.h
+++ b/Engine/Source/Runtime/Experimental/Chaos/Public/Chaos/PBDJointConstraintTypes.h
@@ -63,7 +63,8 @@ namespace Chaos
 		FPBDJointMotionSettings(const TVector<EJointMotionType, 3>& InLinearMotionTypes, const TVector<EJointMotionType, 3>& InAngularMotionTypes);
 
 		FReal Stiffness;
-		FReal Projection;
+		FReal LinearProjection;
+		FReal AngularProjection;
 
 		TVector<EJointMotionType, 3> LinearMotionTypes;
 		FReal LinearLimit;
@@ -121,7 +122,8 @@ namespace Chaos
 		bool bEnableTwistLimits;
 		bool bEnableSwingLimits;
 		bool bEnableDrives;
-		FReal Projection;
+		FReal LinearProjection;
+		FReal AngularProjection;
 		FReal Stiffness;
 		FReal DriveStiffness;
 		FReal SoftLinearStiffness;
diff --git a/Engine/Source/Runtime/Experimental/Chaos/Public/Chaos/PBDJointConstraintUtilities.h b/Engine/Source/Runtime/Experimental/Chaos/Public/Chaos/PBDJointConstraintUtilities.h
--- a/Engine/Source/Runtime/Experimental/Chaos/Public/Chaos/PBDJointConstraintUtilities.h
+++ b/Engine/Source/Runtime/Experimental/Chaos/Public/Chaos/PBDJointConstraintUtilities.h
@@ -95,6 +95,7 @@ namespace Chaos
 			const FReal Dt,
 			const FPBDJointSolverSettings& SolverSettings,
 			const FPBDJointSettings& JointSettings,
+			const FReal Stiffness,
 			const int32 Index0,
 			const int32 Index1,
 			FVec3& P0,
@@ -110,6 +111,7 @@ namespace Chaos
 			const FReal Dt,
 			const FPBDJointSolverSettings& SolverSettings,
 			const FPBDJointSettings& JointSettings,
+			const FReal Stiffness,
 			const int32 Index0,
 			const int32 Index1,
 			FVec3& P0,
@@ -129,6 +131,7 @@ namespace Chaos
 			const FReal Dt,
 			const FPBDJointSolverSettings& SolverSettings,
 			const FPBDJointSettings& JointSettings,
+			const FReal Stiffness,
 			const int32 Index0,
 			const int32 Index1,
 			FVec3& P0,
@@ -144,6 +147,7 @@ namespace Chaos
 			const FReal Dt,
 			const FPBDJointSolverSettings& SolverSettings,
 			const FPBDJointSettings& JointSettings,
+			const FReal Stiffness,
 			const int32 Index0,
 			const int32 Index1,
 			FVec3& P0,
@@ -163,6 +167,7 @@ namespace Chaos
 			const FReal Dt,
 			const FPBDJointSolverSettings& SolverSettings,
 			const FPBDJointSettings& JointSettings,
+			const FReal Stiffness,
 			const int32 Index0,
 			const int32 Index1,
 			FVec3& P0,
@@ -178,6 +183,7 @@ namespace Chaos
 			const FReal Dt,
 			const FPBDJointSolverSettings& SolverSettings,
 			const FPBDJointSettings& JointSettings,
+			const FReal Stiffness,
 			const int32 Index0,
 			const int32 Index1,
 			FVec3& P0,
@@ -197,6 +203,7 @@ namespace Chaos
 			const FReal Dt,
 			const FPBDJointSolverSettings& SolverSettings,
 			const FPBDJointSettings& JointSettings,
+			const FReal Stiffness,
 			const int32 Index0,
 			const int32 Index1,
 			const EJointAngularConstraintIndex SwingConstraintIndex,
@@ -214,6 +221,7 @@ namespace Chaos
 			const FReal Dt,
 			const FPBDJointSolverSettings& SolverSettings,
 			const FPBDJointSettings& JointSettings,
+			const FReal Stiffness,
 			const int32 Index0,
 			const int32 Index1,
 			const EJointAngularConstraintIndex SwingConstraintIndex,
@@ -280,6 +288,7 @@ namespace Chaos
 			const FReal Dt,
 			const FPBDJointSolverSettings& SolverSettings,
 			const FPBDJointSettings& JointSettings,
+			const FReal Stiffness,
 			const int32 Index0,
 			const int32 Index1,
 			FVec3& P0,
@@ -289,13 +298,13 @@ namespace Chaos
 			float InvM0,
 			const FMatrix33& InvIL0,
 			float InvM1,
-			const FMatrix33& InvIL1,
-			const FReal ProjectionFactor);
+			const FMatrix33& InvIL1);
 
 		static CHAOS_API void ApplyJointTwistProjection(
 			const FReal Dt,
 			const FPBDJointSolverSettings& SolverSettings,
 			const FPBDJointSettings& JointSettings,
+			const FReal Stiffness,
 			const int32 Index0,
 			const int32 Index1,
 			FVec3& P0,
@@ -305,13 +314,13 @@ namespace Chaos
 			float InvM0,
 			const FMatrix33& InvIL0,
 			float InvM1,
-			const FMatrix33& InvIL1,
-			const FReal ProjectionFactor);
+			const FMatrix33& InvIL1);
 
 		static CHAOS_API void ApplyJointConeProjection(
 			const FReal Dt,
 			const FPBDJointSolverSettings& SolverSettings,
 			const FPBDJointSettings& JointSettings,
+			const FReal Stiffness,
 			const int32 Index0,
 			const int32 Index1,
 			FVec3& P0,
@@ -321,16 +330,17 @@ namespace Chaos
 			float InvM0,
 			const FMatrix33& InvIL0,
 			float InvM1,
-			const FMatrix33& InvIL1,
-			const FReal ProjectionFactor);
+			const FMatrix33& InvIL1);
 
 		static CHAOS_API void ApplyJointSwingProjection(
 			const FReal Dt,
 			const FPBDJointSolverSettings& SolverSettings,
 			const FPBDJointSettings& JointSettings,
+			const FReal Stiffness,
 			const int32 Index0,
 			const int32 Index1,
-			const EJointAngularConstraintIndex SwingConstraint,
+			const EJointAngularConstraintIndex SwingConstraintIndex,
+			const EJointAngularAxisIndex SwingAxisIndex,
 			FVec3& P0,
 			FRotation3& Q0,
 			FVec3& P1,
@@ -338,8 +348,7 @@ namespace Chaos
 			float InvM0,
 			const FMatrix33& InvIL0,
 			float InvM1,
-			const FMatrix33& InvIL1,
-			const FReal ProjectionFactor);
+			const FMatrix33& InvIL1);
 
 	};
 }
\ No newline at end of file
