Non unity and copy constructor fixes on physics handle storage types

diff --git a/Engine/Source/Runtime/Experimental/Chaos/Public/Chaos/Framework/Handles.h b/Engine/Source/Runtime/Experimental/Chaos/Public/Chaos/Framework/Handles.h
--- a/Engine/Source/Runtime/Experimental/Chaos/Public/Chaos/Framework/Handles.h
+++ b/Engine/Source/Runtime/Experimental/Chaos/Public/Chaos/Framework/Handles.h
@@ -3,6 +3,7 @@
 #pragma once
 
 #include "Chaos/Core.h"
+#include "../ChaosArchive.h"
 
 namespace Chaos
 {
@@ -162,7 +163,7 @@ namespace Chaos
 			MoveFrom(Other);
 		}
 
-		THandleArray& operator =(const THandleArray<ElementType, IndexWidth, GenerationWidth>&)
+		THandleArray& operator =(const THandleArray<ElementType, IndexWidth, GenerationWidth>& Other)
 		{
 			CopyFrom(Other);
 		}
@@ -493,7 +494,7 @@ namespace Chaos
 			MoveFrom(Other);
 		}
 
-		THandleHeap& operator =(const THandleHeap<ElementType, IndexWidth, GenerationWidth>&)
+		THandleHeap& operator =(const THandleHeap<ElementType, IndexWidth, GenerationWidth>& Other)
 		{
 			CopyFrom(Other);
 		}
