Fix unity build after Defines.h cleanup

diff --git a/Engine/Source/Runtime/Experimental/Chaos/Public/Chaos/Rotation.h b/Engine/Source/Runtime/Experimental/Chaos/Public/Chaos/Rotation.h
--- a/Engine/Source/Runtime/Experimental/Chaos/Public/Chaos/Rotation.h
+++ b/Engine/Source/Runtime/Experimental/Chaos/Public/Chaos/Rotation.h
@@ -7,6 +7,7 @@
 #include <cmath>
 #if !COMPILE_WITHOUT_UNREAL_SUPPORT
 #include "Math/Quat.h"
+#include "Math/RotationMatrix.h"
 #else
 #include <array>
 
