Fix array collection includes.

diff --git a/Engine/Source/Runtime/Experimental/ChaosCore/Public/Chaos/ArrayCollectionArray.h b/Engine/Source/Runtime/Experimental/ChaosCore/Public/Chaos/ArrayCollectionArray.h
--- a/Engine/Source/Runtime/Experimental/ChaosCore/Public/Chaos/ArrayCollectionArray.h
+++ b/Engine/Source/Runtime/Experimental/ChaosCore/Public/Chaos/ArrayCollectionArray.h
@@ -1,6 +1,8 @@
 // Copyright 1998-2019 Epic Games, Inc. All Rights Reserved.
 #pragma once
 
+#include "Containers/ArrayView.h"
+
 #include "Chaos/Array.h"
 #include "Chaos/ArrayCollectionArrayBase.h"
 
