Fix TAccelerationStructureHandle hash depending on handle pointer.
TAccelerationStructureHandle comparson compares handle, and falls back on GTParticle if handle is not iniitalized.
Fix bad assumption in PendingSpatialData queueing that could result in failed removal of deleted particle.

diff --git a/Engine/Source/Runtime/Core/Public/UObject/ExternalPhysicsCustomObjectVersion.h b/Engine/Source/Runtime/Core/Public/UObject/ExternalPhysicsCustomObjectVersion.h
--- a/Engine/Source/Runtime/Core/Public/UObject/ExternalPhysicsCustomObjectVersion.h
+++ b/Engine/Source/Runtime/Core/Public/UObject/ExternalPhysicsCustomObjectVersion.h
@@ -50,6 +50,9 @@ struct CORE_API FExternalPhysicsCustomObjectVersion
 
 		// Trimeshes serialize AABBTree
 		TrimeshSerializesAABBTree,
+		
+		// Adds Serialization of HashResult, and separates delete/update TAccelerationStructureHandle in FPendingSpatialData
+		SerializeHashResult,
 
 		// -----<new versions can be added above this line>-------------------------------------------------
 		VersionPlusOne,
diff --git a/Engine/Source/Runtime/Experimental/Chaos/Private/Chaos/GeometryParticles.cpp b/Engine/Source/Runtime/Experimental/Chaos/Private/Chaos/GeometryParticles.cpp
--- a/Engine/Source/Runtime/Experimental/Chaos/Private/Chaos/GeometryParticles.cpp
+++ b/Engine/Source/Runtime/Experimental/Chaos/Private/Chaos/GeometryParticles.cpp
@@ -110,6 +110,21 @@ namespace Chaos
 		auto& SerializableGeometryParticles = AsAlwaysSerializableArray(GeometryParticles->MGeometryParticle);
 		Ar << SerializableGeometryParticles;
 	}
+
+	/*template <typename T, int d, EGeometryParticlesSimType SimType>
+	void TGeometryParticlesImp<T, d, SimType>::SerializeHashResultHelper(FChaosArchive& Ar, TGeometryParticle<T, d>* Particle)
+	{
+		if (Particle)
+		{
+
+			MHashResult.Add();
+		}
+		else
+		{
+			MHashResult.Add(GetHashResultLowLevel());
+		}
+	}*/
+
 	
 	template class TGeometryParticlesImp<float, 3, EGeometryParticlesSimType::RigidBodySim>;
 	template class TGeometryParticlesImp<float, 3, EGeometryParticlesSimType::Other>;
diff --git a/Engine/Source/Runtime/Experimental/Chaos/Private/Chaos/PBDRigidsEvolution.cpp b/Engine/Source/Runtime/Experimental/Chaos/Private/Chaos/PBDRigidsEvolution.cpp
--- a/Engine/Source/Runtime/Experimental/Chaos/Private/Chaos/PBDRigidsEvolution.cpp
+++ b/Engine/Source/Runtime/Experimental/Chaos/Private/Chaos/PBDRigidsEvolution.cpp
@@ -283,7 +283,7 @@ namespace Chaos
 	}
 
 	template<class FPBDRigidsEvolution, class FPBDCollisionConstraint, class T, int d>
-	void TPBDRigidsEvolutionBase<FPBDRigidsEvolution, FPBDCollisionConstraint, T, d>::ApplyParticlePendingData(TGeometryParticleHandle<T, d>* Particle, const FPendingSpatialData& SpatialData, FAccelerationStructure& AccelerationStructure, bool bUpdateCache)
+	void TPBDRigidsEvolutionBase<FPBDRigidsEvolution, FPBDCollisionConstraint, T, d>::ApplyParticlePendingData(const FPendingSpatialData& SpatialData, FAccelerationStructure& AccelerationStructure, bool bUpdateCache)
 	{
 		//Note: we collapsed several update delete events into one struct. If memory is reused this can lead to problems
 		//Luckily there are only 3 states we care about:
@@ -294,11 +294,12 @@ namespace Chaos
 
 		if (SpatialData.bDelete)
 		{
-			AccelerationStructure.RemoveElementFrom(SpatialData.AccelerationHandle, SpatialData.DeletedSpatialIdx);
+			AccelerationStructure.RemoveElementFrom(SpatialData.DeleteAccelerationHandle, SpatialData.DeletedSpatialIdx);
+			TGeometryParticleHandle<T, d>* DeleteParticle = SpatialData.DeleteAccelerationHandle.GetGeometryParticleHandle_PhysicsThread();
 
 			if (bUpdateCache)
 			{
-				if (uint32* InnerIdxPtr = ParticleToCacheInnerIdx.Find(Particle))
+				if (uint32* InnerIdxPtr = ParticleToCacheInnerIdx.Find(DeleteParticle))
 				{
 					const auto SpatialIdx = SpatialData.DeletedSpatialIdx;
 					TSpatialAccelerationCache<T, d>& Cache = *SpatialAccelerationCache.FindChecked(SpatialIdx);	//can't delete from cache that doesn't exist
@@ -310,14 +311,17 @@ namespace Chaos
 					}
 
 					Cache.DestroyElement(CacheInnerIdx);
-					ParticleToCacheInnerIdx.Remove(Particle);
+					ParticleToCacheInnerIdx.Remove(DeleteParticle);
 				}
 			}
 		}
 
 		if (SpatialData.bUpdate)
 		{
-			AccelerationStructure.UpdateElementIn(Particle, Particle->WorldSpaceInflatedBounds(), Particle->HasBounds(), SpatialData.UpdatedSpatialIdx);
+
+			TGeometryParticleHandle<T, d>* UpdateParticle = SpatialData.UpdateAccelerationHandle.GetGeometryParticleHandle_PhysicsThread();
+
+			AccelerationStructure.UpdateElementIn(UpdateParticle, UpdateParticle->WorldSpaceInflatedBounds(), UpdateParticle->HasBounds(), SpatialData.UpdatedSpatialIdx);
 
 			if (bUpdateCache)
 			{
@@ -331,7 +335,7 @@ namespace Chaos
 
 				//make sure in mapping
 				uint32 CacheInnerIdx;
-				if (uint32* CacheInnerIdxPtr = ParticleToCacheInnerIdx.Find(Particle))
+				if (uint32* CacheInnerIdxPtr = ParticleToCacheInnerIdx.Find(UpdateParticle))
 				{
 					CacheInnerIdx = *CacheInnerIdxPtr;
 				}
@@ -339,13 +343,13 @@ namespace Chaos
 				{
 					CacheInnerIdx = Cache.Size();
 					Cache.AddElements(1);
-					ParticleToCacheInnerIdx.Add(Particle, CacheInnerIdx);
+					ParticleToCacheInnerIdx.Add(UpdateParticle, CacheInnerIdx);
 				}
 
 				//update cache entry
-				Cache.HasBounds(CacheInnerIdx) = Particle->HasBounds();
-				Cache.Bounds(CacheInnerIdx) = Particle->WorldSpaceInflatedBounds();
-				Cache.Payload(CacheInnerIdx) = SpatialData.AccelerationHandle;
+				Cache.HasBounds(CacheInnerIdx) = UpdateParticle->HasBounds();
+				Cache.Bounds(CacheInnerIdx) = UpdateParticle->WorldSpaceInflatedBounds();
+				Cache.Payload(CacheInnerIdx) = SpatialData.UpdateAccelerationHandle;
 			}
 		}
 	}
@@ -355,7 +359,7 @@ namespace Chaos
 	{
 		for (auto Itr : InternalAccelerationQueue)
 		{
-			ApplyParticlePendingData(Itr.Key, Itr.Value, *InternalAcceleration, false);
+			ApplyParticlePendingData(Itr.Value, *InternalAcceleration, false);
 		}
 		InternalAccelerationQueue.Empty();
 	}
@@ -365,8 +369,8 @@ namespace Chaos
 	{
 		for (auto Itr : AsyncAccelerationQueue)
 		{
-			ApplyParticlePendingData(Itr.Key, Itr.Value, *AsyncInternalAcceleration, true); //only the first queue needs to update the cached acceleration
-			ApplyParticlePendingData(Itr.Key, Itr.Value, *AsyncExternalAcceleration, false);
+			ApplyParticlePendingData(Itr.Value, *AsyncInternalAcceleration, true); //only the first queue needs to update the cached acceleration
+			ApplyParticlePendingData(Itr.Value, *AsyncExternalAcceleration, false);
 		}
 		AsyncAccelerationQueue.Empty();
 
@@ -380,7 +384,7 @@ namespace Chaos
 	{
 		for (auto Itr : ExternalAccelerationQueue)
 		{
-			ApplyParticlePendingData(Itr.Key, Itr.Value, Acceleration, false);
+			ApplyParticlePendingData(Itr.Value, Acceleration, false);
 		}
 		ExternalAccelerationQueue.Empty();
 	}
diff --git a/Engine/Source/Runtime/Experimental/Chaos/Public/Chaos/GeometryParticles.h b/Engine/Source/Runtime/Experimental/Chaos/Public/Chaos/GeometryParticles.h
--- a/Engine/Source/Runtime/Experimental/Chaos/Public/Chaos/GeometryParticles.h
+++ b/Engine/Source/Runtime/Experimental/Chaos/Public/Chaos/GeometryParticles.h
@@ -97,6 +97,7 @@ namespace Chaos
 			TArrayCollection::AddArray(&MWorldSpaceInflatedBounds);
 			TArrayCollection::AddArray(&MHasBounds);
 			TArrayCollection::AddArray(&MSpatialIdx);
+			TArrayCollection::AddArray(&MHashResult);
 
 			if (IsRigidBodySim())
 			{
@@ -119,6 +120,7 @@ namespace Chaos
 			, MWorldSpaceInflatedBounds(MoveTemp(Other.MWorldSpaceInflatedBounds))
 			, MHasBounds(MoveTemp(Other.MHasBounds))
 			, MSpatialIdx(MoveTemp(Other.MSpatialIdx))
+			, MHashResult(MoveTemp(Other.MHashResult))
 #if CHAOS_DETERMINISTIC
 			, MParticleIDs(MoveTemp(Other.MParticleIDs))
 #endif
@@ -133,6 +135,7 @@ namespace Chaos
 			TArrayCollection::AddArray(&MWorldSpaceInflatedBounds);
 			TArrayCollection::AddArray(&MHasBounds);
 			TArrayCollection::AddArray(&MSpatialIdx);
+			TArrayCollection::AddArray(&MHashResult);
 #if CHAOS_DETERMINISTIC
 			TArrayCollection::AddArray(&MParticleIDs);
 #endif
@@ -159,6 +162,7 @@ namespace Chaos
 			TArrayCollection::AddArray(&MWorldSpaceInflatedBounds);
 			TArrayCollection::AddArray(&MHasBounds);
 			TArrayCollection::AddArray(&MSpatialIdx);
+			TArrayCollection::AddArray(&MHashResult);
 #if CHAOS_DETERMINISTIC
 			TArrayCollection::AddArray(&MParticleIDs);
 #endif
@@ -243,6 +247,16 @@ namespace Chaos
 			return MSpatialIdx[Index];
 		}
 
+		uint32 HashResultLowLevel(const int32 Index) const
+		{
+			return MHashResult[Index];
+		}
+
+		uint32& HashResultLowLevel(const int32 Index)
+		{
+			return MHashResult[Index];
+		}
+
 		CHAOS_API const TBox<T, d>& WorldSpaceInflatedBounds(const int32 Index) const
 		{
 			return MWorldSpaceInflatedBounds[Index];
@@ -313,6 +327,18 @@ namespace Chaos
 			{
 				Ar << MSpatialIdx;
 			}
+
+			if (Ar.CustomVer(FExternalPhysicsCustomObjectVersion::GUID) < FExternalPhysicsCustomObjectVersion::SerializeHashResult)
+			{
+				for (const auto& Particle : MGeometryParticle)
+				{
+					MHashResult.Add(FMath::RandHelper(TNumericLimits<uint32>::Max()));
+				}
+			}
+			else
+			{
+				Ar << MHashResult;
+			}
 		}
 
 		CHAOS_API EParticleType ParticleType() const { return MParticleType; }
@@ -338,6 +364,7 @@ namespace Chaos
 		TArrayCollectionArray<TBox<T, d>> MWorldSpaceInflatedBounds;
 		TArrayCollectionArray<bool> MHasBounds;
 		TArrayCollectionArray<FSpatialAccelerationIdx> MSpatialIdx;
+		TArrayCollectionArray<uint32> MHashResult;
 
 		void UpdateShapesArray(const int32 Index)
 		{
@@ -348,6 +375,7 @@ namespace Chaos
 		friend class TGeometryParticlesImp;
 
 		CHAOS_API void SerializeGeometryParticleHelper(FChaosArchive& Ar, TGeometryParticlesImp<T, d, EGeometryParticlesSimType::RigidBodySim>* GeometryParticles);
+		//CHAOS_API void SerializeHashResultHelper(FChaosArchive& Ar, TGeometryParticle<T, d>* Particle);
 		
 		void SerializeGeometryParticleHelper(FChaosArchive& Ar, TGeometryParticlesImp<T, d, EGeometryParticlesSimType::Other>* GeometryParticles)
 		{
diff --git a/Engine/Source/Runtime/Experimental/Chaos/Public/Chaos/PBDRigidsEvolution.h b/Engine/Source/Runtime/Experimental/Chaos/Public/Chaos/PBDRigidsEvolution.h
--- a/Engine/Source/Runtime/Experimental/Chaos/Public/Chaos/PBDRigidsEvolution.h
+++ b/Engine/Source/Runtime/Experimental/Chaos/Public/Chaos/PBDRigidsEvolution.h
@@ -309,16 +309,19 @@ class TPBDRigidsEvolutionBase
 	FORCEINLINE_DEBUGGABLE void DirtyParticle(TGeometryParticleHandleImp<T, d, bPersistent>& Particle)
 	{
 		FPendingSpatialData& SpatialData = InternalAccelerationQueue.FindOrAdd(Particle.Handle());
-		SpatialData.AccelerationHandle = TAccelerationStructureHandle<T,d>(Particle);
+		SpatialData.UpdateAccelerationHandle = TAccelerationStructureHandle<T,d>(Particle);
 		SpatialData.bUpdate = true;
 		SpatialData.UpdatedSpatialIdx = Particle.SpatialIdx();
 
 		auto& AsyncSpatialData = AsyncAccelerationQueue.FindOrAdd(Particle.Handle());
-		AsyncSpatialData.AccelerationHandle = TAccelerationStructureHandle<T, d>(Particle);
+		AsyncSpatialData.UpdateAccelerationHandle = TAccelerationStructureHandle<T, d>(Particle);
 		AsyncSpatialData.bUpdate = true;
 		AsyncSpatialData.UpdatedSpatialIdx = Particle.SpatialIdx();
-		//question: is it safe to reuse for external? Should probably avoid it
-		ExternalAccelerationQueue.FindOrAdd(Particle.Handle()) = SpatialData;
+
+		auto& ExternalSpatialData = ExternalAccelerationQueue.FindOrAdd(Particle.Handle());
+		ExternalSpatialData.UpdateAccelerationHandle = TAccelerationStructureHandle<T, d>(Particle);
+		ExternalSpatialData.bUpdate = true;
+		ExternalSpatialData.UpdatedSpatialIdx = Particle.SpatialIdx();
 	}
 
 	void DestroyParticle(TGeometryParticleHandle<T, d>* Particle)
@@ -524,7 +527,7 @@ protected:
 	{
 		auto Particle = ParticleHandle.Handle();
 		FPendingSpatialData& AsyncSpatialData = AsyncAccelerationQueue.FindOrAdd(Particle);
-		AsyncSpatialData.AccelerationHandle = TAccelerationStructureHandle<T, d>(ParticleHandle);
+
 		if (!AsyncSpatialData.bDelete)
 		{
 			//There are three cases to consider:
@@ -532,14 +535,20 @@ protected:
 			//Delete followed by any number deletes and updates and finally an update, in that case we must delete the first particle and add the final (so use first index for delete)
 			//Delete followed by multiple updates and or deletes and a final delete. In that case we still only delete the first particle since the final delete is not really needed (add will be cancelled)
 			AsyncSpatialData.DeletedSpatialIdx = ParticleHandle.SpatialIdx();
+
+			// We cannot overwrite this handle if delete is already pending. (If delete is pending, that means this is the third case, cancelling update is sufficient),
+			AsyncSpatialData.DeleteAccelerationHandle = TAccelerationStructureHandle<T, d>(ParticleHandle);
 		}
-		AsyncSpatialData.bUpdate = false;	//don't bother updating since deleting anyway
+
+		AsyncSpatialData.bUpdate = false;
 		AsyncSpatialData.bDelete = true;
+
+		// Delete data should match async, and update is not set, so this operation should be safe.
 		ExternalAccelerationQueue.FindOrAdd(Particle) = AsyncSpatialData;
 
 		//remove particle immediately for intermediate structure
 		InternalAccelerationQueue.Remove(Particle);
-		InternalAcceleration->RemoveElementFrom(AsyncSpatialData.AccelerationHandle, AsyncSpatialData.DeletedSpatialIdx);	//even though we remove immediately, future adds are still pending
+		InternalAcceleration->RemoveElementFrom(AsyncSpatialData.DeleteAccelerationHandle, AsyncSpatialData.DeletedSpatialIdx);	//even though we remove immediately, future adds are still pending
 	}
 
 	void UpdateConstraintPositionBasedState(T Dt)
@@ -622,7 +631,8 @@ protected:
 	/** Used for updating intermediate spatial structures when they are finished */
 	struct FPendingSpatialData
 	{
-		TAccelerationStructureHandle<T, d> AccelerationHandle;
+		TAccelerationStructureHandle<T, d> UpdateAccelerationHandle;
+		TAccelerationStructureHandle<T, d> DeleteAccelerationHandle;
 		FSpatialAccelerationIdx UpdatedSpatialIdx;
 		FSpatialAccelerationIdx DeletedSpatialIdx;	//need both updated and deleted in case memory is reused but a different idx is neede
 		bool bUpdate;
@@ -635,7 +645,18 @@ protected:
 
 		void Serialize(FChaosArchive& Ar)
 		{
-			Ar << AccelerationHandle;
+			Ar.UsingCustomVersion(FExternalPhysicsCustomObjectVersion::GUID);
+			if (Ar.CustomVer(FExternalPhysicsCustomObjectVersion::GUID) >= FExternalPhysicsCustomObjectVersion::SerializeHashResult)
+			{
+				Ar << UpdateAccelerationHandle;
+				Ar << DeleteAccelerationHandle;
+			}
+			else
+			{
+				Ar << UpdateAccelerationHandle;
+				DeleteAccelerationHandle = UpdateAccelerationHandle;
+			}
+
 			Ar << bUpdate;
 			Ar << bDelete;
 
@@ -673,7 +694,7 @@ protected:
 
 	TMap<FSpatialAccelerationIdx, TUniquePtr<TSpatialAccelerationCache<T, d>>> SpatialAccelerationCache;
 
-	FORCEINLINE_DEBUGGABLE void ApplyParticlePendingData(TGeometryParticleHandle<T, d>* Particle, const FPendingSpatialData& PendingData, FAccelerationStructure& SpatialAcceleration, bool bUpdateCache);
+	FORCEINLINE_DEBUGGABLE void ApplyParticlePendingData(const FPendingSpatialData& PendingData, FAccelerationStructure& SpatialAcceleration, bool bUpdateCache);
 
 	class FChaosAccelerationStructureTask
 	{
diff --git a/Engine/Source/Runtime/Experimental/Chaos/Public/Chaos/ParticleDirtyFlags.h b/Engine/Source/Runtime/Experimental/Chaos/Public/Chaos/ParticleDirtyFlags.h
--- a/Engine/Source/Runtime/Experimental/Chaos/Public/Chaos/ParticleDirtyFlags.h
+++ b/Engine/Source/Runtime/Experimental/Chaos/Public/Chaos/ParticleDirtyFlags.h
@@ -30,7 +30,8 @@ namespace Chaos
 		ExternalForce =  1 << 18,
 		ExternalTorque = 1 << 19,
 		GravityEnabled = 1 << 20,
-		SpatialIdx     = 1 << 21
+		SpatialIdx     = 1 << 21,
+		HashResult     = 1 << 22
 	};
 
 	class FParticleDirtyFlags
diff --git a/Engine/Source/Runtime/Experimental/Chaos/Public/Chaos/ParticleHandle.h b/Engine/Source/Runtime/Experimental/Chaos/Public/Chaos/ParticleHandle.h
--- a/Engine/Source/Runtime/Experimental/Chaos/Public/Chaos/ParticleHandle.h
+++ b/Engine/Source/Runtime/Experimental/Chaos/Public/Chaos/ParticleHandle.h
@@ -7,6 +7,7 @@
 #include "Chaos/ParticleIterator.h"
 #include "Chaos/ParticleDirtyFlags.h"
 #include "Chaos/Framework/PhysicsProxyBase.h"
+#include "ChaosCheck.h"
 
 class IPhysicsProxyBase;
 
@@ -48,6 +49,7 @@ void GeometryParticleDefaultConstruct(FConcrete& Concrete, const TGeometryPartic
 	Concrete.SetX(TVector<T, d>(0));
 	Concrete.SetR(TRotation<T, d>::Identity);
 	Concrete.SetSpatialIdx(FSpatialAccelerationIdx{ 0,0 });
+	Concrete.SetHashResultLowLevel(FMath::RandRange(0, FMath::RandHelper(TNumericLimits<int32>::Max())));
 }
 
 template <typename T, int d, typename FConcrete>
@@ -137,17 +139,32 @@ public:
 
 	bool operator==(const TAccelerationStructureHandle<T, d>& Rhs) const
 	{
-		return GeometryParticleHandle == Rhs.GeometryParticleHandle;
+		CHAOS_ENSURE((ExternalGeometryParticle || GeometryParticleHandle));
+		CHAOS_ENSURE((Rhs.ExternalGeometryParticle || Rhs.GeometryParticleHandle));
+
+		if (!GeometryParticleHandle || !Rhs.GeometryParticleHandle)
+		{
+			return (ExternalGeometryParticle == Rhs.ExternalGeometryParticle);
+		}
+		else
+		{
+			return (GeometryParticleHandle == Rhs.GeometryParticleHandle);
+		}
+	}
+
+	bool operator!=(const TAccelerationStructureHandle<T, d>& Rhs) const
+	{
+		return !(*this == Rhs);
 	}
 
 	bool operator<(const TAccelerationStructureHandle<T, d>& Rhs) const
 	{
-		return GeometryParticleHandle < Rhs.GeometryParticleHandle;
+		return HashResult < Rhs.HashResult;
 	}
 
 	uint32 GetTypeHash() const
 	{
-		return ::GetTypeHash((void*)GeometryParticleHandle);
+		return HashResult;
 	}
 
 	void Serialize(FChaosArchive& Ar);
@@ -155,6 +172,8 @@ public:
 private:
 	TGeometryParticle<T, d>* ExternalGeometryParticle;
 	TGeometryParticleHandle<T, d>* GeometryParticleHandle;
+
+	uint32 HashResult;
 };
 
 template <typename T, int d>
@@ -368,6 +387,9 @@ public:
 
 	FSpatialAccelerationIdx SpatialIdx() const { return GeometryParticles->SpatialIdx(ParticleIdx); }
 	void SetSpatialIdx(FSpatialAccelerationIdx Idx) { GeometryParticles->SpatialIdx(ParticleIdx) = Idx; }
+
+	void SetHashResultLowLevel(uint32 Value) { GeometryParticles->HashResultLowLevel(ParticleIdx) = Value; }
+	uint32 GetHashResultLowLevel() const { return GeometryParticles->HashResultLowLevel(ParticleIdx); }
 	
 	EObjectStateType ObjectState() const;
 
@@ -1031,6 +1053,14 @@ public:
 		this->MUserData = InUserData;
 	}
 
+	void SetHashResultLowLevel(uint32 Value)
+	{
+		MarkDirty(EParticleFlags::HashResult);
+		MHashResult = Value;
+	}
+
+	uint32 GetHashResultLowLevel() const { return MHashResult; }
+
 	//Note: this must be called after setting geometry. This API seems bad. Should probably be part of setting geometry
 	void SetShapesArray(TShapesArray<T, d>&& InShapesArray)
 	{
@@ -1101,6 +1131,10 @@ private:
 	TShapesArray<T,d> MShapesArray;
 	FSpatialAccelerationIdx MSpatialIdx;
 
+	// This value is generated and used in AccelerationStructureHandle for a hash, as that handle
+	// sometimes only has a particle pointer from game thread or phys thread, no easy way to get stable hash from that.
+	uint32 MHashResult;
+
 	// Pointer to some arbitrary data associated with the particle, but not used by Chaos. External systems may use this for whatever.
 	void* MUserData;
 
@@ -1141,6 +1175,7 @@ public:
 		, X(TVector<T, d>(0))
 		, R(TRotation<T, d>())
 		, SpatialIdx(FSpatialAccelerationIdx{ 0,0 })
+		, HashResult(0)
 	{}
 
 	TGeometryParticleData(const TGeometryParticle<T, d>& InParticle)
@@ -1149,6 +1184,7 @@ public:
 		, R(InParticle.R())
 		, Geometry(InParticle.GeometrySharedLowLevel())
 		, SpatialIdx(InParticle.SpatialIdx())
+		, HashResult(InParticle.GetHashResultLowLevel())
 	{}
 
 	void Reset() 
@@ -1158,12 +1194,14 @@ public:
 		R = TRotation<T, d>(); 
 		Geometry = TSharedPtr<TImplicitObjectUnion<T, d>, ESPMode::ThreadSafe>();
 		SpatialIdx = FSpatialAccelerationIdx{ 0,0 };
+		HashResult = 0;
 	}
 
 	TVector<T, d> X;
 	TRotation<T, d> R;
 	TSharedPtr<FImplicitObject, ESPMode::ThreadSafe> Geometry;
 	FSpatialAccelerationIdx SpatialIdx;
+	uint32 HashResult;
 };
 
 
@@ -1611,6 +1649,7 @@ template <typename T, int d>
 TAccelerationStructureHandle<T,d>::TAccelerationStructureHandle(TGeometryParticleHandle<T, d>* InHandle)
 	: ExternalGeometryParticle(InHandle->GTGeometryParticle())
 	, GeometryParticleHandle(InHandle)
+	, HashResult(InHandle->GetHashResultLowLevel())
 {
 }
 
@@ -1618,6 +1657,7 @@ template <typename T, int d>
 TAccelerationStructureHandle<T,d>::TAccelerationStructureHandle(TGeometryParticle<T, d>* InGeometryParticle)
 	: ExternalGeometryParticle(InGeometryParticle)
 	, GeometryParticleHandle(InGeometryParticle ? InGeometryParticle->Handle() : nullptr)
+	, HashResult(InGeometryParticle ? InGeometryParticle->GetHashResultLowLevel() : 0)
 {
 }
 
@@ -1626,6 +1666,7 @@ template <bool bPersistent>
 TAccelerationStructureHandle<T, d>::TAccelerationStructureHandle(TGeometryParticleHandleImp<T, d, bPersistent>& InHandle)
 	: ExternalGeometryParticle(InHandle.GTGeometryParticle())
 	, GeometryParticleHandle(InHandle.Handle())
+	, HashResult(InHandle.GetHashResultLowLevel())
 {
 }
 
@@ -1634,6 +1675,7 @@ template <bool bPersistent>
 TAccelerationStructureHandle<T, d>::TAccelerationStructureHandle(TGeometryParticleHandleImp<T, d, bPersistent>* InHandle, TGeometryParticle<T, d>* InGeometryParticle)
 	: ExternalGeometryParticle(InGeometryParticle)
 	, GeometryParticleHandle(InHandle)
+	, HashResult(InGeometryParticle ? InGeometryParticle->GetHashResultLowLevel() : InHandle->GetHashResultLowLevel())
 {
 }
 
diff --git a/Engine/Source/Runtime/Experimental/ChaosSolvers/Private/PhysicsProxy/SingleParticlePhysicsProxy.cpp b/Engine/Source/Runtime/Experimental/ChaosSolvers/Private/PhysicsProxy/SingleParticlePhysicsProxy.cpp
--- a/Engine/Source/Runtime/Experimental/ChaosSolvers/Private/PhysicsProxy/SingleParticlePhysicsProxy.cpp
+++ b/Engine/Source/Runtime/Experimental/ChaosSolvers/Private/PhysicsProxy/SingleParticlePhysicsProxy.cpp
@@ -60,6 +60,7 @@ void FSingleParticlePhysicsProxy<Chaos::TGeometryParticle<float, 3>>::PushToPhys
 			RigidHandle->SetWorldSpaceInflatedBounds(Data->Geometry->BoundingBox().TransformedBox(Chaos::TRigidTransform<float, 3>(Data->X, Data->R)));
 		}
 		RigidHandle->SetSpatialIdx(Data->SpatialIdx);	//todo: this needs to only happen once during initialization
+		RigidHandle->SetHashResultLowLevel(Data->HashResult);
 	}
 }
 
@@ -117,6 +118,7 @@ void FSingleParticlePhysicsProxy<Chaos::TKinematicGeometryParticle<float, 3>>::P
 		RigidHandle->SetW(Data->MW);
 
 		RigidHandle->SetSpatialIdx(Data->SpatialIdx);	//todo: this needs to only happen once during initialization
+		RigidHandle->SetHashResultLowLevel(Data->HashResult);	//todo: this needs to only happen once during initialization
 		
 		if (Data->Geometry && Data->Geometry->HasBoundingBox())
 		{
@@ -185,6 +187,7 @@ void FSingleParticlePhysicsProxy<Chaos::TPBDRigidParticle<float, 3>>::PushToPhys
 		RigidHandle->SetExternalForce(Data->MExternalForce);
 		RigidHandle->SetExternalTorque(Data->MExternalTorque);
 		RigidHandle->SetSpatialIdx(Data->SpatialIdx);	//todo: this needs to only happen once during initialization
+		RigidHandle->SetHashResultLowLevel(Data->HashResult);	//todo: this needs to only happen once during initialization
 		GetSolver()->GetEvolution()->GetGravityForces().SetEnabled(*RigidHandle, Data->MGravityEnabled);
 
 		if (Data->Geometry && Data->Geometry->HasBoundingBox())
