Chaos - fix persistent collision expiry

diff --git a/Engine/Source/Runtime/Experimental/Chaos/Private/Chaos/PBDCollisionConstraint.cpp b/Engine/Source/Runtime/Experimental/Chaos/Private/Chaos/PBDCollisionConstraint.cpp
--- a/Engine/Source/Runtime/Experimental/Chaos/Private/Chaos/PBDCollisionConstraint.cpp
+++ b/Engine/Source/Runtime/Experimental/Chaos/Private/Chaos/PBDCollisionConstraint.cpp
@@ -118,22 +118,10 @@ void TPBDCollisionConstraint<T, d>::Reset(/*const TPBDRigidParticles<T, d>& InPa
 		FConstraintManifold* Elem = Manifolds[Key];
 		if (Elem->GetTimestamp()==FLT_MAX || LifespanCounter > Elem->GetTimestamp())
 		{
-			bool bInContact = false;
-			for (FConstraintContainerHandle* Handle : Elem->GetHandles())
+			const TArray<FConstraintContainerHandle*> HandleArray = Elem->GetHandles();
+			for (FConstraintContainerHandle* Handle : HandleArray)
 			{
-				if (Handle->GetContact().Phi <= 0)
-				{
-					bInContact = true;
-				}
-			}
-
-			if (!bInContact)
-			{
-				const TArray<FConstraintContainerHandle*> HandleArray = Elem->GetHandles();
-				for (FConstraintContainerHandle* Handle : HandleArray)
-				{
-					RemoveConstraint(Handle->GetConstraintIndex());
-				}
+				RemoveConstraint(Handle->GetConstraintIndex());
 			}
 		}
 	}
@@ -2151,6 +2139,13 @@ void TPBDCollisionConstraint<T, d>::ConstructConstraints(TGeometryParticleHandle
 	}
 
 	ConstructConstraintsImpl<T, d>(Particle0, Particle1, Particle0->Geometry().Get(), Particle1->Geometry().Get(), Thickness, ConstraintBuffer, ManifoldElement);
+
+	/*
+	if ((ManifoldElement != nullptr) && (ConstraintBuffer.Num() > 0))
+	{
+		ManifoldElement->SetTimestamp(LifespanCounter + CollisionHistoryLifespan);
+	}
+	*/
 }
 
 // NOTE: UpdateLevelsetConstraintImp and its <float,3> specialization are here for the Linux build. It looks like
