Chaos :

Fix CIS SA Warnings

diff --git a/Engine/Source/Runtime/Experimental/Chaos/Private/Chaos/PBDCollisionConstraint.cpp b/Engine/Source/Runtime/Experimental/Chaos/Private/Chaos/PBDCollisionConstraint.cpp
--- a/Engine/Source/Runtime/Experimental/Chaos/Private/Chaos/PBDCollisionConstraint.cpp
+++ b/Engine/Source/Runtime/Experimental/Chaos/Private/Chaos/PBDCollisionConstraint.cpp
@@ -2117,7 +2117,8 @@ void ConstructConstraintsImpl(TGeometryParticleHandle<T, d>* Particle0, TGeometr
 	if (Implicit0 && GetInnerType(Implicit0->GetType()) == ImplicitObjectType::TriangleMesh) return;
 	if (Implicit1 && GetInnerType(Implicit1->GetType()) == ImplicitObjectType::TriangleMesh) return;
 
-	if (Implicit0->GetType() == TImplicitObjectUnion<T, d>::StaticType() && Implicit1->GetType() == TImplicitObjectUnion<T, d>::StaticType())
+	if (Implicit0 && Implicit0->GetType() == TImplicitObjectUnion<T, d>::StaticType() &&
+		Implicit1 && Implicit1->GetType() == TImplicitObjectUnion<T, d>::StaticType())
 	{
 		ConstructUnionUnionConstraints(Particle0, Particle1, Implicit0, Implicit1, Thickness, ConstraintBuffer, Manifold);
 	}
@@ -2130,15 +2131,16 @@ void ConstructConstraintsImpl(TGeometryParticleHandle<T, d>* Particle0, TGeometr
 template<typename T, int d>
 void TPBDCollisionConstraint<T, d>::ConstructConstraints(TGeometryParticleHandle<T, d>* Particle0, TGeometryParticleHandle<T, d>* Particle1, const T Thickness, TArray<TRigidBodyContactConstraint<T, d>>& ConstraintBuffer)
 {
-	ensure(Particle0 && Particle1);
-
-	FConstraintManifold* ManifoldElement = nullptr;
-	if (FConstraintManifold** Ptr = Manifolds.Find(FConstraintHandleID::Make(Particle0, Particle1)))
+	if (ensure(Particle0 && Particle1))
 	{
-		ManifoldElement = *Ptr;
-	}
+		FConstraintManifold* ManifoldElement = nullptr;
+		if (FConstraintManifold** Ptr = Manifolds.Find(FConstraintHandleID::Make(Particle0, Particle1)))
+		{
+			ManifoldElement = *Ptr;
+		}
 
-	ConstructConstraintsImpl<T, d>(Particle0, Particle1, Particle0->Geometry().Get(), Particle1->Geometry().Get(), Thickness, ConstraintBuffer, ManifoldElement);
+		ConstructConstraintsImpl<T, d>(Particle0, Particle1, Particle0->Geometry().Get(), Particle1->Geometry().Get(), Thickness, ConstraintBuffer, ManifoldElement);
+	}
 
 	/*
 	if ((ManifoldElement != nullptr) && (ConstraintBuffer.Num() > 0))
