Chaos : Convex Collisions

- Added Transformed case to ConvexConvex.

diff --git a/Engine/Source/Runtime/Experimental/Chaos/Private/Chaos/CollisionResolutionConvexConvex.cpp b/Engine/Source/Runtime/Experimental/Chaos/Private/Chaos/CollisionResolutionConvexConvex.cpp
--- a/Engine/Source/Runtime/Experimental/Chaos/Private/Chaos/CollisionResolutionConvexConvex.cpp
+++ b/Engine/Source/Runtime/Experimental/Chaos/Private/Chaos/CollisionResolutionConvexConvex.cpp
@@ -2,6 +2,7 @@
 
 #include "Chaos/CollisionResolutionConvexConvex.h"
 #include "Chaos/ImplicitObjectScaled.h"
+#include "Chaos/ImplicitObjectTransformed.h"
 
 namespace Chaos
 {
@@ -43,48 +44,56 @@ namespace Chaos
 	{
 		ensure(NumConstraints);
 
-		if (ensure(IsScaled(A.GetType()) && IsScaled(B.GetType())))
+		TRigidTransform<T, d> BToATM = BTM.GetRelativeTransform(ATM);
+
+		if (ensure(GetInnerType(A.GetType()) == ImplicitObjectType::Convex && GetInnerType(B.GetType()) == ImplicitObjectType::Convex))
 		{
-			TRigidTransform<T, d> BToATM = BTM.GetRelativeTransform(ATM);
+			const TConvex<T, d>* AObject = nullptr;
+			const TConvex<T, d>* BObject = nullptr;
 
-			if (ensure(GetInnerType(A.GetType()) == ImplicitObjectType::Convex && GetInnerType(B.GetType()) == ImplicitObjectType::Convex))
+			if (IsScaled(A.GetType()) && IsScaled(B.GetType()))
+			{
+				AObject = static_cast<const TConvex<T, d>*>(static_cast<const TImplicitObjectScaled<TConvex<T, d>>&>(A).Object().Get());
+				BObject = static_cast<const TConvex<T, d>*>(static_cast<const TImplicitObjectScaled<TConvex<T, d>>&>(B).Object().Get());
+			}
+			else if (A.GetType() == ImplicitObjectType::Transformed && B.GetType() == ImplicitObjectType::Transformed)
 			{
-				const TConvex<T, d>* AObject = static_cast<const TConvex<T, d>*>(static_cast<const TImplicitObjectScaled<TConvex<T, d>>&>(A).Object().Get());
-				const TConvex<T, d>* BObject = static_cast<const TConvex<T, d>*>(static_cast<const TImplicitObjectScaled<TConvex<T, d>>&>(B).Object().Get());
-				if (AObject && BObject)
+				AObject = static_cast<const TConvex<T, d>*>(static_cast<const TImplicitObjectTransformed<T, d>&>(A).Object().Get());
+				BObject = static_cast<const TConvex<T, d>*>(static_cast<const TImplicitObjectTransformed<T, d>&>(B).Object().Get());
+			}
+
+			if (AObject && BObject)
+			{
+				const TParticles<T, d>& SurfaceParticles = AObject->GetSurfaceParticles();
+
+				for (int32 Idx = 0; Idx < NumConstraints; Idx++)
 				{
-					const TParticles<T, d>& SurfaceParticles = AObject->GetSurfaceParticles();
+					Constraints[Idx].bDisabled = true;
+				}
 
-					for (int32 Idx = 0; Idx < NumConstraints; Idx++)
-					{
-						Constraints[Idx].bDisabled = true;
-					}
+				if (GJKIntersection(*AObject, *BObject, BToATM, Thickness))
+				{
+					TRigidTransform<T, d> AToBTM = ATM.GetRelativeTransform(BTM);
 
-					if (GJKIntersection(*AObject, *BObject, BToATM, Thickness))
+					T Phi;
+					TVector<T, d> Normal;
+					for (int32 Idx = 0; Idx < (int32)SurfaceParticles.Size(); Idx++)
 					{
-						TRigidTransform<T, d> AToBTM = ATM.GetRelativeTransform(BTM);
-
-						T Phi;
-						TVector<T, d> Normal;
-						for (int32 Idx = 0; Idx < (int32)SurfaceParticles.Size(); Idx++)
+						Phi = BObject->PhiWithNormal(AToBTM.TransformPosition(SurfaceParticles.X(Idx)), Normal);
+						if (Phi < Constraints[0].Phi)
 						{
-							Phi = BObject->PhiWithNormal(AToBTM.TransformPosition(SurfaceParticles.X(Idx)), Normal);
-							if (Phi< Constraints[0].Phi)
-							{
-								Constraints[0].Phi = Phi;
-								Constraints[0].Location = ATM.TransformPosition(SurfaceParticles.X(Idx));
-								Constraints[0].Normal = BTM.TransformVector(Normal);
-							}
+							Constraints[0].Phi = Phi;
+							Constraints[0].Location = ATM.TransformPosition(SurfaceParticles.X(Idx));
+							Constraints[0].Normal = BTM.TransformVector(Normal);
 						}
-						return;
 					}
 					return;
 				}
+				return;
 			}
 		}
 
 		ensureMsgf(false, TEXT("Unsupported convex to convex constraint."));
-
 	}
 
 	template class CollisionResolutionConvexConvex<float, 3>;
