First step for EPA. Still buggy and missing some robustness features, but passes most basic test.

diff --git a/Engine/Source/Runtime/Experimental/Chaos/Public/Chaos/EPA.h b/Engine/Source/Runtime/Experimental/Chaos/Public/Chaos/EPA.h
new file mode 100644
--- /dev/null
+++ b/Engine/Source/Runtime/Experimental/Chaos/Public/Chaos/EPA.h
@@ -0,0 +1,199 @@
+// Copyright 1998-2019 Epic Games, Inc. All Rights Reserved.
+#pragma once
+#include "Chaos/Simplex.h"
+#include <queue>
+
+namespace Chaos
+{
+
+template <typename T>
+struct TEPAEntry
+{
+	TVec3<T> Vertices[3];	//vertices of triangle
+	TVec3<T> V;	//closest point on affine hull of triangle
+	T Lambdas[3];	//barycentric coordinates of V (\sum Lambda[i] * Vertices[i])
+	T VSize2;	//distance from V to origin squared
+	TVector<int32,3> AdjIdxs;	//Adjacent triangles ordered CCW
+	TVector<int32,3> AdjIdxInOther;	//Used to find this entry in the corresponding adjacent triangle
+	bool bObsolete;	//indicates that an entry can be skipped (became part of bigger polytope)
+	bool bVInterior;	//indicates whether V is an interior point or not
+
+	bool operator>(const TEPAEntry<T>& Other) const
+	{
+		return VSize2 > Other.VSize2;
+	}
+
+	void Initialize(const TVec3<T>& V0, const TVec3<T>& V1, const TVec3<T>& V2,
+		const TVector<int32,3>& InAdjIdx, const TVector<int32,3>& InAdjIdxInOther)
+	{
+		Vertices[0] = V0;
+		Vertices[1] = V1;
+		Vertices[2] = V2;
+
+		AdjIdxs = InAdjIdx;
+		AdjIdxInOther = InAdjIdxInOther;
+
+		FSimplex Simplex = {0,1,2};
+		V = TriangleSimplexFindOrigin(Vertices, Simplex, Lambdas);
+		VSize2 = V.SizeSquared();
+
+		bVInterior = Simplex.NumVerts == 3;
+		bObsolete = false;
+	}
+};
+
+template <typename T>
+TArray<TEPAEntry<T>> InitializeEPA(const TVec3<T>* Verts, const int32 NumVerts)
+{
+	TArray<TEPAEntry<T>> Entries;
+	switch(NumVerts)
+	{
+		case 4:
+		{
+			Entries.AddUninitialized(4);
+			Entries[0].Initialize(Verts[0], Verts[1], Verts[2], {1,2,3}, {2, 2, 1});
+			Entries[1].Initialize(Verts[0], Verts[3], Verts[1], {3,2,0}, {0,0,0});
+			Entries[2].Initialize(Verts[1], Verts[3], Verts[2], {1,3,0}, {1,2,1});
+			Entries[3].Initialize(Verts[3], Verts[0], Verts[2], {1,0,2}, {0,2,1});
+			return Entries;
+		}
+
+		default: ensure(false);	return Entries;	//todo: handle other cases
+	}
+}
+
+struct FEPAFloodEntry
+{
+	int32 EntryIdx;
+	int32 EdgeIdx;
+};
+
+template <typename T>
+void EPAComputeVisibilityBorder(TArray<TEPAEntry<T>>& Entries, int32 EntryIdx, const TVec3<T>& W, TArray<FEPAFloodEntry>& OutBorderEdges)
+{
+	TArray<FEPAFloodEntry> ToVisitStack;
+	{
+		TEPAEntry<T>& Entry = Entries[EntryIdx];
+		for (int i = 0; i < 3; ++i)
+		{
+			ToVisitStack.Add({ Entry.AdjIdxs[i], Entry.AdjIdxInOther[i] });
+		}
+	}
+
+	while (ToVisitStack.Num())
+	{
+		const FEPAFloodEntry FloodEntry = ToVisitStack.Pop(false);
+		TEPAEntry<T>& Entry = Entries[FloodEntry.EntryIdx];
+		if (!Entry.bObsolete)
+		{
+			const T VDotW = TVec3<T>::DotProduct(Entry.V, W);
+			if (VDotW < Entry.VSize2)	//plane check
+			{
+				//W can't see this triangle so mark the edge as a border
+				OutBorderEdges.Add(FloodEntry);
+			}
+			else
+			{
+				//W can see this triangle so continue flood fill
+				Entry.bObsolete = true;	//no longer needed
+				const int32 Idx0 = FloodEntry.EdgeIdx;
+				const int32 Idx1 = (Idx0 + 1) % 3;
+				const int32 Idx2 = (Idx0 + 2) % 3;
+				ToVisitStack.Add({ Entry.AdjIdxs[Idx1], Entry.AdjIdxInOther[Idx1] });
+				ToVisitStack.Add({ Entry.AdjIdxs[Idx2], Entry.AdjIdxInOther[Idx2] });
+			}
+		}
+	}
+}
+
+template <typename T, typename SupportLambda>
+T EPA(const TVec3<T>* Verts, const int32 NumVerts, const SupportLambda& Support)
+{
+	struct FEPAEntryWrapper
+	{
+		const TArray<TEPAEntry<T>>* Entries;		
+		int32 Idx;
+
+		bool operator>(const FEPAEntryWrapper& Other) const
+		{
+			return (*Entries)[Idx] > (*Entries)[Other.Idx];
+		}
+	};
+
+	TArray<TEPAEntry<T>> Entries = InitializeEPA(Verts, NumVerts);
+	std::priority_queue<FEPAEntryWrapper, std::vector<FEPAEntryWrapper>, std::greater<FEPAEntryWrapper>> Queue;
+	for(int32 Idx = 0; Idx < Entries.Num(); ++Idx)
+	{
+		if(Entries[Idx].bVInterior)
+		{
+			Queue.push(FEPAEntryWrapper {&Entries, Idx});
+		}
+	}
+
+	bool bCloseEnough = false;
+	bool bTerminate = bCloseEnough;
+
+	TArray<FEPAFloodEntry> VisibilityBorder;
+
+	while(Queue.size() && !bTerminate)
+	{
+		int32 EntryIdx = Queue.top().Idx;
+		Queue.pop();
+		TEPAEntry<T>& Entry = Entries[EntryIdx];
+		if (Entry.bObsolete)
+		{
+			continue;
+		}
+
+		const TVec3<T>& V = Entry.V;
+		const TVec3<T> W = Support(V);
+		constexpr T Eps = 1e-4;
+		const T VDotW = TVec3<T>::DotProduct(V, W);
+		bCloseEnough = VDotW < Entry.VSize2 + Eps;	//todo: add relative eps
+		if (!bCloseEnough)
+		{
+			Entry.bObsolete = true;
+			VisibilityBorder.Reset();
+			//question: in the paper we call this per edge which means we traverse new faces for second and third edge, does it matter?
+			EPAComputeVisibilityBorder(Entries, EntryIdx, W, VisibilityBorder);
+			const int32 NumBorderEdges = VisibilityBorder.Num();
+			const int32 FirstIdxInBatch = Entries.Num();
+			int32 NewIdx = FirstIdxInBatch;
+			Entries.AddUninitialized(NumBorderEdges);
+			ensure(NumBorderEdges >= 3);
+			for (int32 VisibilityIdx = 0; VisibilityIdx < NumBorderEdges; ++VisibilityIdx)
+			{
+				//create new entries and update adjacencies
+				const FEPAFloodEntry& BorderInfo = VisibilityBorder[VisibilityIdx];
+				TEPAEntry<T>& NewEntry = Entries[NewIdx];
+				const int32 BorderEntryIdx = BorderInfo.EntryIdx;
+				TEPAEntry<T>& BorderEntry = Entries[BorderEntryIdx];
+				const int32 BorderEdgeIdx0 = BorderInfo.EdgeIdx;
+				const int32 BorderEdgeIdx1 = (BorderEdgeIdx0 + 1) % 3;
+				const int32 NextEntryIdx = (VisibilityIdx + 1) < VisibilityBorder.Num() ? NewIdx + 1 : FirstIdxInBatch;
+				const int32 PrevEntryIdx = NewIdx > FirstIdxInBatch ? NewIdx - 1 : FirstIdxInBatch + NumBorderEdges - 1;
+				NewEntry.Initialize(BorderEntry.Vertices[BorderEdgeIdx1], BorderEntry.Vertices[BorderEdgeIdx0], W,
+					{ BorderEntryIdx, NextEntryIdx, PrevEntryIdx },
+					{ BorderEdgeIdx0, 2, 1 });
+				BorderEntry.AdjIdxs[BorderEdgeIdx0] = NewIdx;
+				BorderEntry.AdjIdxInOther[BorderEdgeIdx0] = 0;
+
+				if (NewEntry.bVInterior)
+				{
+					Queue.push(FEPAEntryWrapper{ &Entries, NewIdx});
+				}
+
+				++NewIdx;
+			}
+		}
+		else
+		{
+			return FMath::Sqrt(Entry.VSize2);
+		}
+
+		bTerminate = bCloseEnough;
+	}
+
+	return 0;
+}
+}
diff --git a/Engine/Source/Runtime/Experimental/Chaos/Public/Chaos/GJK.h b/Engine/Source/Runtime/Experimental/Chaos/Public/Chaos/GJK.h
--- a/Engine/Source/Runtime/Experimental/Chaos/Public/Chaos/GJK.h
+++ b/Engine/Source/Runtime/Experimental/Chaos/Public/Chaos/GJK.h
@@ -1,365 +1,15 @@
 // Copyright 1998-2019 Epic Games, Inc. All Rights Reserved.
 #pragma once
-#include "Chaos/Transform.h"
+
+#include "Chaos/Simplex.h"
+#include "Chaos/EPA.h"
 #include "Chaos/Capsule.h"
 #include "Chaos/Box.h"
 #include "Chaos/Sphere.h"
 
 namespace Chaos
 {
-	template <typename T, int d>
-	TVector<T, d> LineSimplexFindOrigin(const TVector<T, d>* Simplex, int32* Idxs, int32& NumVerts, T* OutBarycentric)
-	{
-		const TVector<T, d>& X0 = Simplex[Idxs[0]];
-		const TVector<T, d>& X1 = Simplex[Idxs[1]];
-		const TVector<T, d> X0ToX1 = X1 - X0;
-
-		//Closest Point = (-X0 dot X1-X0) / ||(X1-X0)||^2 * (X1-X0)
-
-		const TVector<T, d> OriginToX0 = -X0;
-		const T Dot = TVector<T, d>::DotProduct(OriginToX0, X0ToX1);
-
-		if (Dot <= 0)
-		{
-			NumVerts = 1;
-			OutBarycentric[Idxs[0]] = 1;
-			return X0;
-		}
-
-		const T X0ToX1Squared = X0ToX1.SizeSquared();
-
-		if (X0ToX1Squared < Dot || X0ToX1Squared <= std::numeric_limits<float>::min())	//if dividing gives 1+ or the line is degenerate
-		{
-			NumVerts = 1;
-			Idxs[0] = Idxs[1];
-			OutBarycentric[Idxs[1]] = 1;
-			return X1;
-		}
-
-		const T Ratio = Dot / X0ToX1Squared;
-		const TVector<T, d> Closest = Ratio * (X0ToX1)+X0;	//note: this could pass X1 by machine epsilon, but doesn't seem worth it for now
-		OutBarycentric[Idxs[0]] = 1 - Ratio;
-		OutBarycentric[Idxs[1]] = Ratio;
-		return Closest;
-	}
-
-	struct FSimplex
-	{
-		int32 NumVerts;
-		int32 Idxs[4];
-		int32 operator[](int32 Idx) const { return Idxs[Idx]; }
-		int32& operator[](int32 Idx) { return Idxs[Idx]; }
-
-		FSimplex& operator=(const FSimplex& Other)
-		{
-			NumVerts = Other.NumVerts;
-			for (int i = 0; i < 4; ++i)
-			{
-				Idxs[i] = Other.Idxs[i];
-			}
-			return *this;
-		}
-
-		FSimplex(std::initializer_list<int32> InIdxs = {})
-			: NumVerts(InIdxs.size())
-		{
-			check(NumVerts <= 4);
-			int32 i = 0;
-			for (int32 Idx : InIdxs)
-			{
-				Idxs[i++] = Idx;
-			}
-			while (i < 4)
-			{
-				Idxs[i++] = 0;	//some code uses these for lookup regardless of NumVerts. Makes for faster code so just use 0 to make lookup safe
-			}
-		}
-	};
-
-	template <typename T>
-	bool SignMatch(T A, T B)
-	{
-		return (A > 0 && B > 0) || (A < 0 && B < 0);
-	}
-
-	template <typename T>
-	TVector<T, 3> TriangleSimplexFindOrigin(const TVector<T, 3>* Simplex, FSimplex& Idxs, T* OutBarycentric)
-	{
-		/* Project the origin onto the triangle plane:
-		   Let n = (b-a) cross (c-a)
-		   Let the distance from the origin dist = ((-a) dot n / ||n||^2)
-		   Then the projection p = 0 - dist * n = (a dot n) / ||n||^2
-		*/
-
-		const int32 Idx0 = Idxs[0];
-		const int32 Idx1 = Idxs[1];
-		const int32 Idx2 = Idxs[2];
-
-		const TVector<T, 3>& X0 = Simplex[Idx0];
-		const TVector<T, 3>& X1 = Simplex[Idx1];
-		const TVector<T, 3>& X2 = Simplex[Idx2];
-
-		const TVector<T, 3> X0ToX1 = X1 - X0;
-		const TVector<T, 3> X0ToX2 = X2 - X0;
-		const TVector<T, 3> TriNormal = TVector<T, 3>::CrossProduct(X0ToX1, X0ToX2);
-
-		/*
-		   We want |(a dot n) / ||n||^2| < 1 / eps to avoid inf. But note that |a dot n| <= ||a||||n|| and so
-		   |(a dot n) / ||n||^2| <= ||a|| / ||n| < 1 / eps requires that ||eps*a||^2 < ||n||^2
-		*/
-		const T TriNormal2 = TriNormal.SizeSquared();
-		const T Eps2 = (X0*std::numeric_limits<T>::min()).SizeSquared();
-		if (Eps2 >= TriNormal2)	//equality fixes case where both X0 and TriNormal2 are 0
-		{
-			//degenerate triangle so return previous line results
-			Idxs.NumVerts = 2;
-			return LineSimplexFindOrigin(Simplex, Idxs.Idxs, Idxs.NumVerts, OutBarycentric);
-		}
-
-		const TVector<T, 3> TriNormalOverSize2 = TriNormal / TriNormal2;
-		const T SignedDistance = TVector<T, 3>::DotProduct(X0, TriNormalOverSize2);
-		const TVector<T, 3> ProjectedOrigin = TriNormal * SignedDistance;
-
-		/*
-			Let p be the origin projected onto the triangle plane. We can represent the point p in a 2d subspace spanned by the triangle
-			|a_u, b_u, c_u| |lambda_1| = |p_u|
-			|a_v, b_v, c_v| |lambda_2| = |p_v|
-			|1,   1,   1  | |lambda_3| = |1  |
-
-			Cramer's rule gives: lambda_i = det(M_i) / det(M)
-			To choose u and v we simply test x,y,z to see if any of them are linearly independent
-		*/
-
-		T DetM = 0;	//not needed but fixes compiler warning
-		int32 BestAxisU = INDEX_NONE;
-		int32 BestAxisV = 0;	//once best axis is chosen use the axes that go with it in the right order
-
-		{
-			T MaxAbsDetM = 0;	//not needed but fixes compiler warning
-			int32 AxisU = 1;
-			int32 AxisV = 2;
-			for (int32 CurAxis = 0; CurAxis < 3; ++CurAxis)
-			{
-				T TmpDetM = X1[AxisU] * X2[AxisV] - X2[AxisU] * X1[AxisV]
-					+ X2[AxisU] * X0[AxisV] - X0[AxisU] * X2[AxisV]
-					+ X0[AxisU] * X1[AxisV] - X1[AxisU] * X0[AxisV];
-				const T AbsDetM = FMath::Abs(TmpDetM);
-				if (BestAxisU == INDEX_NONE || AbsDetM > MaxAbsDetM)
-				{
-					MaxAbsDetM = AbsDetM;
-					DetM = TmpDetM;
-					BestAxisU = AxisU;
-					BestAxisV = AxisV;
-				}
-				AxisU = AxisV;
-				AxisV = CurAxis;
-			}
-		}
-
-		/*
-			Now solve for the cofactors (i.e. the projected origin replaces the column of each cofactor).
-			Notice that this is really the area of each sub triangle with the projected origin.
-			If the sign of the determinants is different than the sign of the entire triangle determinant then we are outside of the triangle.
-			The conflicting signs indicate which voronoi regions to search
-
-			Cofactor_a =    |p_u b_u c_u|  Cofactor_b =    |a_u p_u c_u|  Cofactor_c = |a_u b_u p_u|
-						 det|p_v b_v c_v|               det|a_v p_v c_v|            det|a_v c_v p_v|
-							|1   1  1   |                  |1   1  1   |               |1   1  1   |
-		*/
-		const TVector<T, 3>& P0 = ProjectedOrigin;
-		const TVector<T, 3> P0ToX0 = X0 - P0;
-		const TVector<T, 3> P0ToX1 = X1 - P0;
-		const TVector<T, 3> P0ToX2 = X2 - P0;
-
-		const T Cofactors[3] = {
-			P0ToX1[BestAxisU] * P0ToX2[BestAxisV] - P0ToX2[BestAxisU] * P0ToX1[BestAxisV],
-			-P0ToX0[BestAxisU] * P0ToX2[BestAxisV] + P0ToX2[BestAxisU] * P0ToX0[BestAxisV],
-			P0ToX0[BestAxisU] * P0ToX1[BestAxisV] - P0ToX1[BestAxisU] * P0ToX0[BestAxisV]
-		};
-
-		bool bSignMatch[3];
-		FSimplex SubSimplices[3] = { {Idx1,Idx2}, {Idx0,Idx2}, {Idx0,Idx1} };
-		TVector<T, 3> ClosestPointSub[3];
-		T SubBarycentric[3][3];
-		int32 ClosestSubIdx = INDEX_NONE;
-		T MinSubDist2 = 0;	//not needed
-		bool bInside = true;
-		for (int32 Idx = 0; Idx < 3; ++Idx)
-		{
-			bSignMatch[Idx] = SignMatch(DetM, Cofactors[Idx]);
-			if (!bSignMatch[Idx])
-			{
-				bInside = false;
-				ClosestPointSub[Idx] = LineSimplexFindOrigin(Simplex, SubSimplices[Idx].Idxs, SubSimplices[Idx].NumVerts, SubBarycentric[Idx]);
-
-				const T Dist2 = ClosestPointSub[Idx].SizeSquared();
-				if (ClosestSubIdx == INDEX_NONE || Dist2 < MinSubDist2)
-				{
-					MinSubDist2 = Dist2;
-					ClosestSubIdx = Idx;
-				}
-			}
-		}
-
-		if (bInside)
-		{
-			//SignMatch ensures that DetM is not 0. The Det_i / Det_m ratio is always between 0-1 because it represents the ratio of areas and Det_m is the total area
-			const T InvDetM = 1 / DetM;
-			T Lambda0 = Cofactors[0] * InvDetM;
-			T Lambda1 = Cofactors[1] * InvDetM;
-			//T Lambda2 = 1 - Lambda1 - Lambda0;
-			T Lambda2 = Cofactors[2] * InvDetM;
-			const TVector<T, 3> ClosestPoint = X0 * Lambda0 + X1 * Lambda1 + X2 * Lambda2;	//could be slightly outside if |lambda1| < 1e-7 or |lambda2| < 1e-7. Should we clamp?
-			OutBarycentric[Idx0] = Lambda0;
-			OutBarycentric[Idx1] = Lambda1;
-			OutBarycentric[Idx2] = Lambda2;
-			return ClosestPoint;
-		}
-		else
-		{
-			Idxs = SubSimplices[ClosestSubIdx];
-			OutBarycentric[Idx0] = SubBarycentric[ClosestSubIdx][Idx0];
-			OutBarycentric[Idx1] = SubBarycentric[ClosestSubIdx][Idx1];
-			OutBarycentric[Idx2] = SubBarycentric[ClosestSubIdx][Idx2];
-			return ClosestPointSub[ClosestSubIdx];
-		}
-	}
-
-	template <typename T>
-	TVector<T, 3> TetrahedronSimplexFindOrigin(const TVector<T, 3>* Simplex, FSimplex& Idxs, T* OutBarycentric)
-	{
-		const int32 Idx0 = Idxs[0];
-		const int32 Idx1 = Idxs[1];
-		const int32 Idx2 = Idxs[2];
-		const int32 Idx3 = Idxs[3];
-
-		const TVector<T, 3>& X0 = Simplex[Idx0];
-		const TVector<T, 3>& X1 = Simplex[Idx1];
-		const TVector<T, 3>& X2 = Simplex[Idx2];
-		const TVector<T, 3>& X3 = Simplex[Idx3];
-
-		//Use signed volumes to determine if origin is inside or outside
-		/*
-			M = [X0x X1x X2x X3x;
-				 X0y X1y X2y X3y;
-				 X0z X1z X2z X3z;
-				 1   1   1   1]
-		*/
-
-		T Cofactors[4];
-		Cofactors[0] = -TVector<T, 3>::DotProduct(X1, TVector<T, 3>::CrossProduct(X2, X3));
-		Cofactors[1] = TVector<T, 3>::DotProduct(X0, TVector<T, 3>::CrossProduct(X2, X3));
-		Cofactors[2] = -TVector<T, 3>::DotProduct(X0, TVector<T, 3>::CrossProduct(X1, X3));
-		Cofactors[3] = TVector<T, 3>::DotProduct(X0, TVector<T, 3>::CrossProduct(X1, X2));
-		T DetM = (Cofactors[0] + Cofactors[1]) + (Cofactors[2] + Cofactors[3]);
-
-		bool bSignMatch[4];
-		FSimplex SubIdxs[4] = { {1,2,3}, {0,2,3}, {0,1,3}, {0,1,2} };
-		TVector<T, 3> ClosestPointSub[4];
-		T SubBarycentric[4][4];
-		int32 ClosestTriangleIdx = INDEX_NONE;
-		T MinTriangleDist2 = 0;
-
-		bool bInside = true;
-		for (int Idx = 0; Idx < 4; ++Idx)
-		{
-			bSignMatch[Idx] = SignMatch(DetM, Cofactors[Idx]);
-			if (!bSignMatch[Idx])
-			{
-				bInside = false;
-				ClosestPointSub[Idx] = TriangleSimplexFindOrigin(Simplex, SubIdxs[Idx], SubBarycentric[Idx]);
-
-				const T Dist2 = ClosestPointSub[Idx].SizeSquared();
-				if (ClosestTriangleIdx == INDEX_NONE || Dist2 < MinTriangleDist2)
-				{
-					MinTriangleDist2 = Dist2;
-					ClosestTriangleIdx = Idx;
-				}
-			}
-		}
-
-		if (bInside)
-		{
-			OutBarycentric[Idx0] = Cofactors[0] / DetM;
-			OutBarycentric[Idx1] = Cofactors[1] / DetM;
-			OutBarycentric[Idx2] = Cofactors[2] / DetM;
-			OutBarycentric[Idx3] = Cofactors[3] / DetM;
-
-			return TVector<T, 3>(0);
-		}
-
-		Idxs = SubIdxs[ClosestTriangleIdx];
-
-		OutBarycentric[Idx0] = SubBarycentric[ClosestTriangleIdx][Idx0];
-		OutBarycentric[Idx1] = SubBarycentric[ClosestTriangleIdx][Idx1];
-		OutBarycentric[Idx2] = SubBarycentric[ClosestTriangleIdx][Idx2];
-		OutBarycentric[Idx3] = SubBarycentric[ClosestTriangleIdx][Idx3];
-
-		return ClosestPointSub[ClosestTriangleIdx];
-	}
-
-	template <typename T>
-	void ReorderGJKArray(T* Data, FSimplex& Idxs)
-	{
-		const T D0 = Data[Idxs[0]];
-		const T D1 = Data[Idxs[1]];
-		const T D2 = Data[Idxs[2]];
-		const T D3 = Data[Idxs[3]];
-		Data[0] = D0;
-		Data[1] = D1;
-		Data[2] = D2;
-		Data[3] = D3;
-	}
-
-	template <typename T>
-	TVector<T, 3> SimplexFindClosestToOrigin(TVector<T, 3>* Simplex, FSimplex& Idxs, T* OutBarycentric, TVector<T, 3>* A = nullptr, TVector<T, 3>* B = nullptr)
-	{
-		TVector<T, 3> ClosestPoint;
-		switch (Idxs.NumVerts)
-		{
-		case 1:
-			OutBarycentric[0] = 1;
-			ClosestPoint = Simplex[Idxs[0]]; break;
-		case 2:
-		{
-			ClosestPoint = LineSimplexFindOrigin(Simplex, Idxs.Idxs, Idxs.NumVerts, OutBarycentric);
-			break;
-		}
-		case 3:
-		{
-			ClosestPoint = TriangleSimplexFindOrigin(Simplex, Idxs, OutBarycentric);
-			break;
-		}
-		case 4:
-		{
-			ClosestPoint = TetrahedronSimplexFindOrigin(Simplex, Idxs, OutBarycentric);
-			break;
-		}
-		default:
-			ensure(false);
-			ClosestPoint = TVector<T, 3>(0);
-		}
-
-		ReorderGJKArray(Simplex, Idxs);
-		ReorderGJKArray(OutBarycentric, Idxs);
-		if (A)
-		{
-			ReorderGJKArray(A, Idxs);
-		}
-
-		if (B)
-		{
-			ReorderGJKArray(B, Idxs);
-		}
-
-		Idxs[0] = 0;
-		Idxs[1] = 1;
-		Idxs[2] = 2;
-		Idxs[3] = 3;
-
-		return ClosestPoint;
-	}
-
+	
 	/** Determines if two convex geometries overlap.
 	 @A The first geometry
 	 @B The second geometry
diff --git a/Engine/Source/Runtime/Experimental/Chaos/Public/Chaos/Simplex.h b/Engine/Source/Runtime/Experimental/Chaos/Public/Chaos/Simplex.h
new file mode 100644
--- /dev/null
+++ b/Engine/Source/Runtime/Experimental/Chaos/Public/Chaos/Simplex.h
@@ -0,0 +1,359 @@
+// Copyright 1998-2019 Epic Games, Inc. All Rights Reserved.
+#pragma once
+#include "Chaos/Transform.h"
+
+namespace Chaos
+{
+	template <typename T, int d>
+	TVector<T, d> LineSimplexFindOrigin(const TVector<T, d>* Simplex, int32* Idxs, int32& NumVerts, T* OutBarycentric)
+	{
+		const TVector<T, d>& X0 = Simplex[Idxs[0]];
+		const TVector<T, d>& X1 = Simplex[Idxs[1]];
+		const TVector<T, d> X0ToX1 = X1 - X0;
+
+		//Closest Point = (-X0 dot X1-X0) / ||(X1-X0)||^2 * (X1-X0)
+
+		const TVector<T, d> OriginToX0 = -X0;
+		const T Dot = TVector<T, d>::DotProduct(OriginToX0, X0ToX1);
+
+		if (Dot <= 0)
+		{
+			NumVerts = 1;
+			OutBarycentric[Idxs[0]] = 1;
+			return X0;
+		}
+
+		const T X0ToX1Squared = X0ToX1.SizeSquared();
+
+		if (X0ToX1Squared < Dot || X0ToX1Squared <= std::numeric_limits<float>::min())	//if dividing gives 1+ or the line is degenerate
+		{
+			NumVerts = 1;
+			Idxs[0] = Idxs[1];
+			OutBarycentric[Idxs[1]] = 1;
+			return X1;
+		}
+
+		const T Ratio = Dot / X0ToX1Squared;
+		const TVector<T, d> Closest = Ratio * (X0ToX1)+X0;	//note: this could pass X1 by machine epsilon, but doesn't seem worth it for now
+		OutBarycentric[Idxs[0]] = 1 - Ratio;
+		OutBarycentric[Idxs[1]] = Ratio;
+		return Closest;
+	}
+
+	struct FSimplex
+	{
+		int32 NumVerts;
+		int32 Idxs[4];
+		int32 operator[](int32 Idx) const { return Idxs[Idx]; }
+		int32& operator[](int32 Idx) { return Idxs[Idx]; }
+
+		FSimplex& operator=(const FSimplex& Other)
+		{
+			NumVerts = Other.NumVerts;
+			for (int i = 0; i < 4; ++i)
+			{
+				Idxs[i] = Other.Idxs[i];
+			}
+			return *this;
+		}
+
+		FSimplex(std::initializer_list<int32> InIdxs = {})
+			: NumVerts(InIdxs.size())
+		{
+			check(NumVerts <= 4);
+			int32 i = 0;
+			for (int32 Idx : InIdxs)
+			{
+				Idxs[i++] = Idx;
+			}
+			while (i < 4)
+			{
+				Idxs[i++] = 0;	//some code uses these for lookup regardless of NumVerts. Makes for faster code so just use 0 to make lookup safe
+			}
+		}
+	};
+
+	template <typename T>
+	bool SignMatch(T A, T B)
+	{
+		return (A > 0 && B > 0) || (A < 0 && B < 0);
+	}
+
+	template <typename T>
+	TVector<T, 3> TriangleSimplexFindOrigin(const TVector<T, 3>* Simplex, FSimplex& Idxs, T* OutBarycentric)
+	{
+		/* Project the origin onto the triangle plane:
+		   Let n = (b-a) cross (c-a)
+		   Let the distance from the origin dist = ((-a) dot n / ||n||^2)
+		   Then the projection p = 0 - dist * n = (a dot n) / ||n||^2
+		*/
+
+		const int32 Idx0 = Idxs[0];
+		const int32 Idx1 = Idxs[1];
+		const int32 Idx2 = Idxs[2];
+
+		const TVector<T, 3>& X0 = Simplex[Idx0];
+		const TVector<T, 3>& X1 = Simplex[Idx1];
+		const TVector<T, 3>& X2 = Simplex[Idx2];
+
+		const TVector<T, 3> X0ToX1 = X1 - X0;
+		const TVector<T, 3> X0ToX2 = X2 - X0;
+		const TVector<T, 3> TriNormal = TVector<T, 3>::CrossProduct(X0ToX1, X0ToX2);
+
+		/*
+		   We want |(a dot n) / ||n||^2| < 1 / eps to avoid inf. But note that |a dot n| <= ||a||||n|| and so
+		   |(a dot n) / ||n||^2| <= ||a|| / ||n| < 1 / eps requires that ||eps*a||^2 < ||n||^2
+		*/
+		const T TriNormal2 = TriNormal.SizeSquared();
+		const T Eps2 = (X0*std::numeric_limits<T>::min()).SizeSquared();
+		if (Eps2 >= TriNormal2)	//equality fixes case where both X0 and TriNormal2 are 0
+		{
+			//degenerate triangle so return previous line results
+			Idxs.NumVerts = 2;
+			return LineSimplexFindOrigin(Simplex, Idxs.Idxs, Idxs.NumVerts, OutBarycentric);
+		}
+
+		const TVector<T, 3> TriNormalOverSize2 = TriNormal / TriNormal2;
+		const T SignedDistance = TVector<T, 3>::DotProduct(X0, TriNormalOverSize2);
+		const TVector<T, 3> ProjectedOrigin = TriNormal * SignedDistance;
+
+		/*
+			Let p be the origin projected onto the triangle plane. We can represent the point p in a 2d subspace spanned by the triangle
+			|a_u, b_u, c_u| |lambda_1| = |p_u|
+			|a_v, b_v, c_v| |lambda_2| = |p_v|
+			|1,   1,   1  | |lambda_3| = |1  |
+
+			Cramer's rule gives: lambda_i = det(M_i) / det(M)
+			To choose u and v we simply test x,y,z to see if any of them are linearly independent
+		*/
+
+		T DetM = 0;	//not needed but fixes compiler warning
+		int32 BestAxisU = INDEX_NONE;
+		int32 BestAxisV = 0;	//once best axis is chosen use the axes that go with it in the right order
+
+		{
+			T MaxAbsDetM = 0;	//not needed but fixes compiler warning
+			int32 AxisU = 1;
+			int32 AxisV = 2;
+			for (int32 CurAxis = 0; CurAxis < 3; ++CurAxis)
+			{
+				T TmpDetM = X1[AxisU] * X2[AxisV] - X2[AxisU] * X1[AxisV]
+					+ X2[AxisU] * X0[AxisV] - X0[AxisU] * X2[AxisV]
+					+ X0[AxisU] * X1[AxisV] - X1[AxisU] * X0[AxisV];
+				const T AbsDetM = FMath::Abs(TmpDetM);
+				if (BestAxisU == INDEX_NONE || AbsDetM > MaxAbsDetM)
+				{
+					MaxAbsDetM = AbsDetM;
+					DetM = TmpDetM;
+					BestAxisU = AxisU;
+					BestAxisV = AxisV;
+				}
+				AxisU = AxisV;
+				AxisV = CurAxis;
+			}
+		}
+
+		/*
+			Now solve for the cofactors (i.e. the projected origin replaces the column of each cofactor).
+			Notice that this is really the area of each sub triangle with the projected origin.
+			If the sign of the determinants is different than the sign of the entire triangle determinant then we are outside of the triangle.
+			The conflicting signs indicate which voronoi regions to search
+
+			Cofactor_a =    |p_u b_u c_u|  Cofactor_b =    |a_u p_u c_u|  Cofactor_c = |a_u b_u p_u|
+						 det|p_v b_v c_v|               det|a_v p_v c_v|            det|a_v c_v p_v|
+							|1   1  1   |                  |1   1  1   |               |1   1  1   |
+		*/
+		const TVector<T, 3>& P0 = ProjectedOrigin;
+		const TVector<T, 3> P0ToX0 = X0 - P0;
+		const TVector<T, 3> P0ToX1 = X1 - P0;
+		const TVector<T, 3> P0ToX2 = X2 - P0;
+
+		const T Cofactors[3] = {
+			P0ToX1[BestAxisU] * P0ToX2[BestAxisV] - P0ToX2[BestAxisU] * P0ToX1[BestAxisV],
+			-P0ToX0[BestAxisU] * P0ToX2[BestAxisV] + P0ToX2[BestAxisU] * P0ToX0[BestAxisV],
+			P0ToX0[BestAxisU] * P0ToX1[BestAxisV] - P0ToX1[BestAxisU] * P0ToX0[BestAxisV]
+		};
+
+		bool bSignMatch[3];
+		FSimplex SubSimplices[3] = { {Idx1,Idx2}, {Idx0,Idx2}, {Idx0,Idx1} };
+		TVector<T, 3> ClosestPointSub[3];
+		T SubBarycentric[3][3];
+		int32 ClosestSubIdx = INDEX_NONE;
+		T MinSubDist2 = 0;	//not needed
+		bool bInside = true;
+		for (int32 Idx = 0; Idx < 3; ++Idx)
+		{
+			bSignMatch[Idx] = SignMatch(DetM, Cofactors[Idx]);
+			if (!bSignMatch[Idx])
+			{
+				bInside = false;
+				ClosestPointSub[Idx] = LineSimplexFindOrigin(Simplex, SubSimplices[Idx].Idxs, SubSimplices[Idx].NumVerts, SubBarycentric[Idx]);
+
+				const T Dist2 = ClosestPointSub[Idx].SizeSquared();
+				if (ClosestSubIdx == INDEX_NONE || Dist2 < MinSubDist2)
+				{
+					MinSubDist2 = Dist2;
+					ClosestSubIdx = Idx;
+				}
+			}
+		}
+
+		if (bInside)
+		{
+			//SignMatch ensures that DetM is not 0. The Det_i / Det_m ratio is always between 0-1 because it represents the ratio of areas and Det_m is the total area
+			const T InvDetM = 1 / DetM;
+			T Lambda0 = Cofactors[0] * InvDetM;
+			T Lambda1 = Cofactors[1] * InvDetM;
+			//T Lambda2 = 1 - Lambda1 - Lambda0;
+			T Lambda2 = Cofactors[2] * InvDetM;
+			const TVector<T, 3> ClosestPoint = X0 * Lambda0 + X1 * Lambda1 + X2 * Lambda2;	//could be slightly outside if |lambda1| < 1e-7 or |lambda2| < 1e-7. Should we clamp?
+			OutBarycentric[Idx0] = Lambda0;
+			OutBarycentric[Idx1] = Lambda1;
+			OutBarycentric[Idx2] = Lambda2;
+			return ClosestPoint;
+		}
+		else
+		{
+			Idxs = SubSimplices[ClosestSubIdx];
+			OutBarycentric[Idx0] = SubBarycentric[ClosestSubIdx][Idx0];
+			OutBarycentric[Idx1] = SubBarycentric[ClosestSubIdx][Idx1];
+			OutBarycentric[Idx2] = SubBarycentric[ClosestSubIdx][Idx2];
+			return ClosestPointSub[ClosestSubIdx];
+		}
+	}
+
+	template <typename T>
+	TVector<T, 3> TetrahedronSimplexFindOrigin(const TVector<T, 3>* Simplex, FSimplex& Idxs, T* OutBarycentric)
+	{
+		const int32 Idx0 = Idxs[0];
+		const int32 Idx1 = Idxs[1];
+		const int32 Idx2 = Idxs[2];
+		const int32 Idx3 = Idxs[3];
+
+		const TVector<T, 3>& X0 = Simplex[Idx0];
+		const TVector<T, 3>& X1 = Simplex[Idx1];
+		const TVector<T, 3>& X2 = Simplex[Idx2];
+		const TVector<T, 3>& X3 = Simplex[Idx3];
+
+		//Use signed volumes to determine if origin is inside or outside
+		/*
+			M = [X0x X1x X2x X3x;
+				 X0y X1y X2y X3y;
+				 X0z X1z X2z X3z;
+				 1   1   1   1]
+		*/
+
+		T Cofactors[4];
+		Cofactors[0] = -TVector<T, 3>::DotProduct(X1, TVector<T, 3>::CrossProduct(X2, X3));
+		Cofactors[1] = TVector<T, 3>::DotProduct(X0, TVector<T, 3>::CrossProduct(X2, X3));
+		Cofactors[2] = -TVector<T, 3>::DotProduct(X0, TVector<T, 3>::CrossProduct(X1, X3));
+		Cofactors[3] = TVector<T, 3>::DotProduct(X0, TVector<T, 3>::CrossProduct(X1, X2));
+		T DetM = (Cofactors[0] + Cofactors[1]) + (Cofactors[2] + Cofactors[3]);
+
+		bool bSignMatch[4];
+		FSimplex SubIdxs[4] = { {1,2,3}, {0,2,3}, {0,1,3}, {0,1,2} };
+		TVector<T, 3> ClosestPointSub[4];
+		T SubBarycentric[4][4];
+		int32 ClosestTriangleIdx = INDEX_NONE;
+		T MinTriangleDist2 = 0;
+
+		bool bInside = true;
+		for (int Idx = 0; Idx < 4; ++Idx)
+		{
+			bSignMatch[Idx] = SignMatch(DetM, Cofactors[Idx]);
+			if (!bSignMatch[Idx])
+			{
+				bInside = false;
+				ClosestPointSub[Idx] = TriangleSimplexFindOrigin(Simplex, SubIdxs[Idx], SubBarycentric[Idx]);
+
+				const T Dist2 = ClosestPointSub[Idx].SizeSquared();
+				if (ClosestTriangleIdx == INDEX_NONE || Dist2 < MinTriangleDist2)
+				{
+					MinTriangleDist2 = Dist2;
+					ClosestTriangleIdx = Idx;
+				}
+			}
+		}
+
+		if (bInside)
+		{
+			OutBarycentric[Idx0] = Cofactors[0] / DetM;
+			OutBarycentric[Idx1] = Cofactors[1] / DetM;
+			OutBarycentric[Idx2] = Cofactors[2] / DetM;
+			OutBarycentric[Idx3] = Cofactors[3] / DetM;
+
+			return TVector<T, 3>(0);
+		}
+
+		Idxs = SubIdxs[ClosestTriangleIdx];
+
+		OutBarycentric[Idx0] = SubBarycentric[ClosestTriangleIdx][Idx0];
+		OutBarycentric[Idx1] = SubBarycentric[ClosestTriangleIdx][Idx1];
+		OutBarycentric[Idx2] = SubBarycentric[ClosestTriangleIdx][Idx2];
+		OutBarycentric[Idx3] = SubBarycentric[ClosestTriangleIdx][Idx3];
+
+		return ClosestPointSub[ClosestTriangleIdx];
+	}
+
+	template <typename T>
+	void ReorderGJKArray(T* Data, FSimplex& Idxs)
+	{
+		const T D0 = Data[Idxs[0]];
+		const T D1 = Data[Idxs[1]];
+		const T D2 = Data[Idxs[2]];
+		const T D3 = Data[Idxs[3]];
+		Data[0] = D0;
+		Data[1] = D1;
+		Data[2] = D2;
+		Data[3] = D3;
+	}
+
+	template <typename T>
+	TVector<T, 3> SimplexFindClosestToOrigin(TVector<T, 3>* Simplex, FSimplex& Idxs, T* OutBarycentric, TVector<T, 3>* A = nullptr, TVector<T, 3>* B = nullptr)
+	{
+		TVector<T, 3> ClosestPoint;
+		switch (Idxs.NumVerts)
+		{
+		case 1:
+			OutBarycentric[0] = 1;
+			ClosestPoint = Simplex[Idxs[0]]; break;
+		case 2:
+		{
+			ClosestPoint = LineSimplexFindOrigin(Simplex, Idxs.Idxs, Idxs.NumVerts, OutBarycentric);
+			break;
+		}
+		case 3:
+		{
+			ClosestPoint = TriangleSimplexFindOrigin(Simplex, Idxs, OutBarycentric);
+			break;
+		}
+		case 4:
+		{
+			ClosestPoint = TetrahedronSimplexFindOrigin(Simplex, Idxs, OutBarycentric);
+			break;
+		}
+		default:
+			ensure(false);
+			ClosestPoint = TVector<T, 3>(0);
+		}
+
+		ReorderGJKArray(Simplex, Idxs);
+		ReorderGJKArray(OutBarycentric, Idxs);
+		if (A)
+		{
+			ReorderGJKArray(A, Idxs);
+		}
+
+		if (B)
+		{
+			ReorderGJKArray(B, Idxs);
+		}
+
+		Idxs[0] = 0;
+		Idxs[1] = 1;
+		Idxs[2] = 2;
+		Idxs[3] = 3;
+
+		return ClosestPoint;
+	}
+}
diff --git a/Engine/Source/Runtime/Experimental/ChaosCore/Public/Chaos/Core.h b/Engine/Source/Runtime/Experimental/ChaosCore/Public/Chaos/Core.h
--- a/Engine/Source/Runtime/Experimental/ChaosCore/Public/Chaos/Core.h
+++ b/Engine/Source/Runtime/Experimental/ChaosCore/Public/Chaos/Core.h
@@ -20,4 +20,8 @@ namespace Chaos
 	using FMatrix33 = PMatrix<FReal, 3, 3>;
 	using FRigidTransform3 = TRigidTransform<FReal, 3>;
 	using FRotation3 = TRotation<FReal, 3>;
+
+
+	template <typename T>
+	using TVec3 = TVector<T, 3>;
 }
\ No newline at end of file
