Test that particle is properly removed from accel structure for an edge case with particle handle pointer being reused.

diff --git a/Engine/Source/Runtime/Experimental/Chaos/Public/Chaos/PBDRigidsEvolution.h b/Engine/Source/Runtime/Experimental/Chaos/Public/Chaos/PBDRigidsEvolution.h
--- a/Engine/Source/Runtime/Experimental/Chaos/Public/Chaos/PBDRigidsEvolution.h
+++ b/Engine/Source/Runtime/Experimental/Chaos/Public/Chaos/PBDRigidsEvolution.h
@@ -14,6 +14,10 @@
 #include "Chaos/PBDRigidsSOAs.h"
 #include "Chaos/ISpatialAccelerationCollection.h"
 
+
+// Declaring so it can be friended for tests.
+namespace ChaosTest { void TestPendingSpatialDataHandlePointerConflict(); } 
+
 namespace Chaos
 {
 template<class T, int d> class TPBDRigidsEvolutionGBF;
@@ -219,6 +223,8 @@ class TPBDRigidsEvolutionBase
 	typedef TFunction<void(const TParticleView<TPBDRigidParticles<T, d>>&, const T)> FUpdatePositionRule;
 	typedef TFunction<void(TPBDRigidParticles<T, d>&, const T, const T, const int32)> FKinematicUpdateRule;
 
+	friend void ChaosTest::TestPendingSpatialDataHandlePointerConflict();
+
 	CHAOS_API TPBDRigidsEvolutionBase(TPBDRigidsSOAs<T, d>& InParticles, int32 InNumIterations = 1, int32 InNumPushOutIterations = 1);
 	CHAOS_API virtual ~TPBDRigidsEvolutionBase();
 
