Remove accidentally submitted comments.

diff --git a/Engine/Source/Runtime/Experimental/Chaos/Private/Chaos/GeometryParticles.cpp b/Engine/Source/Runtime/Experimental/Chaos/Private/Chaos/GeometryParticles.cpp
--- a/Engine/Source/Runtime/Experimental/Chaos/Private/Chaos/GeometryParticles.cpp
+++ b/Engine/Source/Runtime/Experimental/Chaos/Private/Chaos/GeometryParticles.cpp
@@ -111,21 +111,6 @@ namespace Chaos
 		Ar << SerializableGeometryParticles;
 	}
 
-	/*template <typename T, int d, EGeometryParticlesSimType SimType>
-	void TGeometryParticlesImp<T, d, SimType>::SerializeHashResultHelper(FChaosArchive& Ar, TGeometryParticle<T, d>* Particle)
-	{
-		if (Particle)
-		{
-
-			MHashResult.Add();
-		}
-		else
-		{
-			MHashResult.Add(GetHashResultLowLevel());
-		}
-	}*/
-
-	
 	template class TGeometryParticlesImp<float, 3, EGeometryParticlesSimType::RigidBodySim>;
 	template class TGeometryParticlesImp<float, 3, EGeometryParticlesSimType::Other>;
 	template class TPerShapeData<float, 3>;
diff --git a/Engine/Source/Runtime/Experimental/Chaos/Public/Chaos/GeometryParticles.h b/Engine/Source/Runtime/Experimental/Chaos/Public/Chaos/GeometryParticles.h
--- a/Engine/Source/Runtime/Experimental/Chaos/Public/Chaos/GeometryParticles.h
+++ b/Engine/Source/Runtime/Experimental/Chaos/Public/Chaos/GeometryParticles.h
@@ -375,7 +375,6 @@ namespace Chaos
 		friend class TGeometryParticlesImp;
 
 		CHAOS_API void SerializeGeometryParticleHelper(FChaosArchive& Ar, TGeometryParticlesImp<T, d, EGeometryParticlesSimType::RigidBodySim>* GeometryParticles);
-		//CHAOS_API void SerializeHashResultHelper(FChaosArchive& Ar, TGeometryParticle<T, d>* Particle);
 		
 		void SerializeGeometryParticleHelper(FChaosArchive& Ar, TGeometryParticlesImp<T, d, EGeometryParticlesSimType::Other>* GeometryParticles)
 		{
