More fine grained control of acceleration structure time-slicing, plus associated test

diff --git a/Engine/Source/Runtime/Experimental/Chaos/Private/Chaos/PBDRigidsEvolution.cpp b/Engine/Source/Runtime/Experimental/Chaos/Private/Chaos/PBDRigidsEvolution.cpp
--- a/Engine/Source/Runtime/Experimental/Chaos/Private/Chaos/PBDRigidsEvolution.cpp
+++ b/Engine/Source/Runtime/Experimental/Chaos/Private/Chaos/PBDRigidsEvolution.cpp
@@ -29,7 +29,7 @@ namespace Chaos
 			AABBMaxChildrenInLeaf = 500;
 			AABBMaxTreeDepth = 200;
 			MaxPayloadSize = 100000;
-			IterationsPerTimeSlice = 20000;
+			IterationsPerTimeSlice = 40000;
 		}
 	} ConfigSettings;
 
@@ -409,11 +409,6 @@ namespace Chaos
 
 		bool ForceFullBuild = InternalAccelerationQueue.Num() > 1000;
 
-		if (ForceFullBuild)
-		{
-			UE_LOG(LogChaos, Warning, TEXT("ForceFullBuild"));
-		}
-
 		if (!AccelerationStructureTaskComplete)
 		{
 			//initial frame so make empty structures
diff --git a/Engine/Source/Runtime/Experimental/Chaos/Public/Chaos/AABBTree.h b/Engine/Source/Runtime/Experimental/Chaos/Public/Chaos/AABBTree.h
--- a/Engine/Source/Runtime/Experimental/Chaos/Public/Chaos/AABBTree.h
+++ b/Engine/Source/Runtime/Experimental/Chaos/Public/Chaos/AABBTree.h
@@ -23,6 +23,8 @@ DECLARE_CYCLE_STAT(TEXT("AABBTreeGenerateTree"), STAT_AABBTreeGenerateTree, STAT
 DECLARE_CYCLE_STAT(TEXT("AABBTreeTimeSliceSetup"), STAT_AABBTreeTimeSliceSetup, STATGROUP_Chaos);
 DECLARE_CYCLE_STAT(TEXT("AABBTreeInitialTimeSlice"), STAT_AABBTreeInitialTimeSlice, STATGROUP_Chaos);
 DECLARE_CYCLE_STAT(TEXT("AABBTreeProgressTimeSlice"), STAT_AABBTreeProgressTimeSlice, STATGROUP_Chaos);
+//DECLARE_CYCLE_STAT(TEXT("AABBTreeGrowPhase"), STAT_AABBTreeGrowPhase, STATGROUP_Chaos);
+//DECLARE_CYCLE_STAT(TEXT("AABBTreeChildrenPhase"), STAT_AABBTreeChildrenPhase, STATGROUP_Chaos);
 
 template <typename T, typename TQueryFastData, EAABBQueryType Query>
 struct TAABBTreeIntersectionHelper
@@ -225,7 +227,7 @@ public:
 	static constexpr T DefaultMaxPayloadBounds = 100000;
 	static constexpr int32 DefaultMaxChildrenInLeaf = 12;
 	static constexpr int32 DefaultMaxTreeDepth = 16;
-	static constexpr int32 DefaultMaxNumToProcess = 4000000;
+	static constexpr int32 DefaultMaxNumToProcess = 0; // 0 special value for processing all without timeslicing
 	static constexpr ESpatialAcceleration StaticType = TIsSame<TAABBTreeLeafArray<TPayloadType, T>, TLeafType>::Value ? ESpatialAcceleration::AABBTree : 
 		(TIsSame<TBoundingVolume<TPayloadType, T, 3>, TLeafType>::Value ? ESpatialAcceleration::AABBTreeBV : ESpatialAcceleration::Unknown);
 	TAABBTree()
@@ -253,7 +255,15 @@ public:
 			FWorkSnapshot Store;
 			TimeSliceWorkToComplete.Dequeue(Store);
 			NumProcessedThisSlice = 0;
-			SplitNode(Store.Bounds, Store.Elems, Store.NodeLevel, Store.NewNodeIdx);
+			if (Store.TimeslicePhase == eTimeSlicePhase::PHASE1)
+			{
+				SplitNode(Store.Bounds, Store.Elems, Store.NodeLevel, Store.NewNodeIdx);
+			}
+			else
+			{
+				check(Store.TimeslicePhase == eTimeSlicePhase::PHASE2);
+				SplitNode(Store.Bounds, Store.RealBounds, Store.Elems, Store.SplitBoundsSize2, Store.NewNodeIdx, Store.NodeLevel, Store.BoxIdx);
+			}
 		}
 
 		// are done now?
@@ -273,14 +283,7 @@ public:
 		, MaxNumToProcess(InMaxNumToProcess)
 
 	{
-		this->SetAsyncTimeSlicingComplete(false);
 		GenerateTree(Particles);
-
-		/*
-		TBox<T, 3> Tmp = TBox<T, 3>::ZeroBox();
-		DumpTreeToLog(0, -1, Tmp);
-		*/
-
 	}
 
 	template <typename ParticleView>
@@ -291,17 +294,17 @@ public:
 
 	virtual TArray<TPayloadType> FindAllIntersections(const TBox<T, 3>& Box) const override { return FindAllIntersectionsImp(Box); }
 
-	bool DumpTreeToLog(int32 NodeIdx, int32 ParentNode, TBox<T,3>& Bounds)
+	bool GetAsBoundsArray(TArray<TBox<T, 3>>& AllBounds, int32 NodeIdx, int32 ParentNode, TBox<T, 3>& Bounds)
 	{
 		if (Nodes[NodeIdx].bLeaf)
 		{
-			UE_LOG(LogChaos, Log, TEXT("Bounds %s"), Bounds.ToString().GetCharArray().GetData());
+			AllBounds.Add(Bounds);
 			return false;
 		}
 		else
-		{		
-			DumpTree(Nodes[NodeIdx].ChildrenNodes[0], NodeIdx, Nodes[NodeIdx].ChildrenBounds[0]);
-			DumpTree(Nodes[NodeIdx].ChildrenNodes[1], NodeIdx, Nodes[NodeIdx].ChildrenBounds[0]);
+		{
+			GetAsBoundsArray(AllBounds, Nodes[NodeIdx].ChildrenNodes[0], NodeIdx, Nodes[NodeIdx].ChildrenBounds[0]);
+			GetAsBoundsArray(AllBounds, Nodes[NodeIdx].ChildrenNodes[1], NodeIdx, Nodes[NodeIdx].ChildrenBounds[0]);
 		}
 		return true;
 	}
@@ -609,6 +612,7 @@ private:
 	void GenerateTree(const TParticles& Particles)
 	{
 		SCOPE_CYCLE_COUNTER(STAT_AABBTreeGenerateTree);
+		this->SetAsyncTimeSlicingComplete(false);
 
 		TArray<FElement> ElemsWithBounds;
 		ElemsWithBounds.Reserve(Particles.Num());
@@ -675,26 +679,41 @@ private:
 
 	}
 
-
-	void SplitNode(const TBox<T, 3>& Bounds, const TArray<FElement>& Elems, int32 NodeLevel, int32 NewNodeIdx, bool Initial=false)
+	void SplitNode(TBox<T, 3>& SplitBounds, TBox<T, 3>& RealBounds, TArray<FElement>& Children, T& SplitBoundsSize2, int NewNodeIdx, int NodeLevel, int BoxIdx)
 	{
-		const int32 InitialMaxNumToProcess = 100;
-		int32 NumToProcess = MaxNumToProcess;
-		if (Initial)
+		//SCOPE_CYCLE_COUNTER(STAT_AABBTreeChildrenPhase);
+
+		check (Children.Num() > 0)
 		{
-			NumToProcess = InitialMaxNumToProcess;
+			Nodes[NewNodeIdx].bLeaf = false;
+
+			// increment this early so the amount of recursion work left on the stack is smaller
+			NumProcessedThisSlice += Children.Num();
+
+			const int32 ChildIdx = Nodes.Num();
+			Nodes[NewNodeIdx].ChildrenBounds[BoxIdx] = RealBounds;
+			Nodes[NewNodeIdx].ChildrenNodes[BoxIdx] = ChildIdx;
+
+			SplitNode(RealBounds, Children, NodeLevel + 1, ChildIdx);
 		}
 
+	}
+
+	void SplitNode(const TBox<T, 3>& Bounds, const TArray<FElement>& Elems, int32 NodeLevel, int32 NewNodeIdx, bool Initial=false)
+	{
 		// create the actual node space but might no be filled in (YET) due to time slicing exit
 		if (NewNodeIdx >= Nodes.Num())
 		{
 		 	Nodes.AddDefaulted();	//todo: remove TBox
 		}
 
-		if ((MaxNumToProcess > 0) && (NumProcessedThisSlice >= NumToProcess))
+		bool WeAreTimeslicing = (MaxNumToProcess > 0);
+
+		if (WeAreTimeslicing && (NumProcessedThisSlice >= MaxNumToProcess))
 		{
 			// done enough work, capture stack
 			FWorkSnapshot Store;
+			Store.TimeslicePhase = eTimeSlicePhase::PHASE1;
 			Store.Bounds = Bounds;
 			Store.Elems = Elems;			
 			Store.NodeLevel = NodeLevel;
@@ -728,14 +747,7 @@ private:
 
 		const TVector<T, 3> Extents = Bounds.Extents();
 		const int32 MaxAxis = Bounds.LargestAxis();
-
-		struct FSplitInfo
-		{
-			TBox<T, 3> SplitBounds;	//Even split of parent bounds
-			TBox<T, 3> RealBounds;	//Actual bounds as children are added
-			TArray<FElement> Children;
-			T SplitBoundsSize2;
-		};
+		bool ChildrenInBothHalves = false;
 
 		FSplitInfo SplitInfos[2];
 		SplitInfos[0].SplitBounds = TBox<T, 3>(Bounds.Min(), Bounds.Min());
@@ -762,50 +774,81 @@ private:
 			SplitInfo.SplitBoundsSize2 = SplitInfo.SplitBounds.Extents().SizeSquared();
 		}
 		
-		// add all elements to one of the two split infos at this level - root level [ not taking into account the max number allowed or anything
-		for (const FElement& Elem : Elems)
 		{
-			int32 MinBoxIdx = INDEX_NONE;
-			T MinDelta2 = TNumericLimits<T>::Max();
-			int32 BoxIdx = 0;
-			for (const FSplitInfo& SplitInfo : SplitInfos)
+			//SCOPE_CYCLE_COUNTER(STAT_AABBTreeGrowPhase);
+
+			// add all elements to one of the two split infos at this level - root level [ not taking into account the max number allowed or anything
+			for (const FElement& Elem : Elems)
 			{
-				TBox<T, 3> NewBox = SplitInfo.SplitBounds;
-				NewBox.GrowToInclude(Elem.Bounds);
-				const T Delta2 = NewBox.Extents().SizeSquared() - SplitInfo.SplitBoundsSize2;
-				if (Delta2 < MinDelta2)
+				int32 MinBoxIdx = INDEX_NONE;
+				T MinDelta2 = TNumericLimits<T>::Max();
+				int32 BoxIdx = 0;
+				for (const FSplitInfo& SplitInfo : SplitInfos)
 				{
-					MinDelta2 = Delta2;
-					MinBoxIdx = BoxIdx;
+					TBox<T, 3> NewBox = SplitInfo.SplitBounds;
+					NewBox.GrowToInclude(Elem.Bounds);
+					const T Delta2 = NewBox.Extents().SizeSquared() - SplitInfo.SplitBoundsSize2;
+					if (Delta2 < MinDelta2)
+					{
+						MinDelta2 = Delta2;
+						MinBoxIdx = BoxIdx;
+					}
+					++BoxIdx;
 				}
-				++BoxIdx;
+
+				SplitInfos[MinBoxIdx].Children.Add(Elem);
+				SplitInfos[MinBoxIdx].RealBounds.GrowToInclude(Elem.Bounds);
 			}
 
-			SplitInfos[MinBoxIdx].Children.Add(Elem);
-			SplitInfos[MinBoxIdx].RealBounds.GrowToInclude(Elem.Bounds);
-		}
+			NumProcessedThisSlice += Elems.Num();
 
-		// if children in both halves, recurse a level down
-		if (SplitInfos[0].Children.Num() && SplitInfos[1].Children.Num())
-		{
-			Nodes[NewNodeIdx].bLeaf = false;
+			ChildrenInBothHalves = SplitInfos[0].Children.Num() && SplitInfos[1].Children.Num();
 
-			// increment this early so the amount of recursion work left on the stack is smaller
-			NumProcessedThisSlice += SplitInfos[0].Children.Num() + SplitInfos[1].Children.Num();
+			if (ChildrenInBothHalves && WeAreTimeslicing && (NumProcessedThisSlice >= MaxNumToProcess))
+			{
+				for (int32 BoxIdx = 0; BoxIdx < 2; ++BoxIdx)
+				{
+					// done enough work, capture stack
+					FWorkSnapshot Store;
+					Store.TimeslicePhase = eTimeSlicePhase::PHASE2;
+					Store.NodeLevel = NodeLevel;
+					Store.NewNodeIdx = NewNodeIdx;
+					Store.BoxIdx = BoxIdx;
+					Store.Bounds = MoveTemp(SplitInfos[BoxIdx].SplitBounds);
+					Store.Elems = MoveTemp(SplitInfos[BoxIdx].Children);
+					Store.RealBounds = MoveTemp(SplitInfos[BoxIdx].RealBounds);
+					Store.SplitBoundsSize2 = SplitInfos[BoxIdx].SplitBoundsSize2;
+					TimeSliceWorkToComplete.Enqueue(Store);
+				}
+				return; // done enough
+			}	
+		}
+
+		{			
+			//SCOPE_CYCLE_COUNTER(STAT_AABBTreeChildrenPhase);
 
-			for (int32 BoxIdx = 0; BoxIdx < 2; ++BoxIdx)
+			// if children in both halves, recurse a level down
+			if (ChildrenInBothHalves)
 			{
+				Nodes[NewNodeIdx].bLeaf = false;
+
+				// increment this early so the amount of recursion work left on the stack is smaller
+				NumProcessedThisSlice += SplitInfos[0].Children.Num() + SplitInfos[1].Children.Num();
+
+				for (int32 BoxIdx = 0; BoxIdx < 2; ++BoxIdx)
+				{
 					const int32 ChildIdx = Nodes.Num();
 					Nodes[NewNodeIdx].ChildrenBounds[BoxIdx] = SplitInfos[BoxIdx].RealBounds;
 					Nodes[NewNodeIdx].ChildrenNodes[BoxIdx] = ChildIdx;
 
 					SplitNode(SplitInfos[BoxIdx].RealBounds, SplitInfos[BoxIdx].Children, NodeLevel + 1, ChildIdx);
+				}
+			}
+			else
+			{
+				//couldn't split so just make a leaf - THIS COULD CONTAIN MORE THAN MaxChildrenInLeaf!!!
+				MakeLeaf();
 			}
-		}
-		else
-		{
-			//couldn't split so just make a leaf - THIS COULD CONTAIN MORE THAN MaxChildrenInLeaf!!!
-			MakeLeaf();
 		}
 
 		return; // keep working
@@ -859,6 +902,14 @@ private:
 
 	}
 
+	struct FSplitInfo
+	{
+		TBox<T, 3> SplitBounds;	//Even split of parent bounds
+		TBox<T, 3> RealBounds;	//Actual bounds as children are added
+		TArray<FElement> Children;
+		T SplitBoundsSize2;
+	};
+
 	TBox<T, 3> FullBounds;
 	TArray<FNode> Nodes;
 	TArray<TLeafType> Leaves;
@@ -869,14 +920,28 @@ private:
 	int32 MaxTreeDepth;
 	T MaxPayloadBounds;
 	int32 MaxNumToProcess;
+	
+	enum eTimeSlicePhase
+	{
+		PHASE1,
+		PHASE2
+	};
 
 	struct FWorkSnapshot
 	{
+		eTimeSlicePhase TimeslicePhase;
+
+		// used for phase 1 & 2
 		TBox<T, 3> Bounds;
 		TArray<FElement> Elems;
 
 		int32 NodeLevel;
 		int32 NewNodeIdx;
+	
+		// phase 2 data only
+		TBox<T, 3> RealBounds;
+		int8 BoxIdx;
+		T SplitBoundsSize2;
 	};
 
 	int32 NumProcessedThisSlice;
