Chaos: Clean up a couple things related to TAccelerationStructureHandle fix.

diff --git a/Engine/Source/Runtime/Experimental/Chaos/Private/Chaos/GeometryParticles.cpp b/Engine/Source/Runtime/Experimental/Chaos/Private/Chaos/GeometryParticles.cpp
--- a/Engine/Source/Runtime/Experimental/Chaos/Private/Chaos/GeometryParticles.cpp
+++ b/Engine/Source/Runtime/Experimental/Chaos/Private/Chaos/GeometryParticles.cpp
@@ -111,6 +111,20 @@ namespace Chaos
 		Ar << SerializableGeometryParticles;
 	}
 
+	template <typename T, int d, EGeometryParticlesSimType SimType>
+	void TGeometryParticlesImp<T, d, SimType>::SerializeHashResultHelper(FChaosArchive& Ar, TGeometryParticle<T, d>* Particle)
+	{
+		if (Particle)
+		{
+			MHashResult.Add(Particle->GetHashResultLowLevel());
+		}
+		else
+		{
+			MHashResult.Add(FMath::RandHelper(TNumericLimits<uint32>::Max()));
+		}
+	}
+
+	
 	template class TGeometryParticlesImp<float, 3, EGeometryParticlesSimType::RigidBodySim>;
 	template class TGeometryParticlesImp<float, 3, EGeometryParticlesSimType::Other>;
 	template class TPerShapeData<float, 3>;
diff --git a/Engine/Source/Runtime/Experimental/Chaos/Public/Chaos/GeometryParticles.h b/Engine/Source/Runtime/Experimental/Chaos/Public/Chaos/GeometryParticles.h
--- a/Engine/Source/Runtime/Experimental/Chaos/Public/Chaos/GeometryParticles.h
+++ b/Engine/Source/Runtime/Experimental/Chaos/Public/Chaos/GeometryParticles.h
@@ -332,7 +332,7 @@ namespace Chaos
 			{
 				for (const auto& Particle : MGeometryParticle)
 				{
-					MHashResult.Add(FMath::RandHelper(TNumericLimits<uint32>::Max()));
+					SerializeHashResultHelper(Ar, Particle);
 				}
 			}
 			else
@@ -375,6 +375,7 @@ namespace Chaos
 		friend class TGeometryParticlesImp;
 
 		CHAOS_API void SerializeGeometryParticleHelper(FChaosArchive& Ar, TGeometryParticlesImp<T, d, EGeometryParticlesSimType::RigidBodySim>* GeometryParticles);
+		CHAOS_API void SerializeHashResultHelper(FChaosArchive& Ar, TGeometryParticle<T, d>* Particle);
 		
 		void SerializeGeometryParticleHelper(FChaosArchive& Ar, TGeometryParticlesImp<T, d, EGeometryParticlesSimType::Other>* GeometryParticles)
 		{
diff --git a/Engine/Source/Runtime/Experimental/Chaos/Public/Chaos/ParticleHandle.h b/Engine/Source/Runtime/Experimental/Chaos/Public/Chaos/ParticleHandle.h
--- a/Engine/Source/Runtime/Experimental/Chaos/Public/Chaos/ParticleHandle.h
+++ b/Engine/Source/Runtime/Experimental/Chaos/Public/Chaos/ParticleHandle.h
@@ -1675,8 +1675,21 @@ template <bool bPersistent>
 TAccelerationStructureHandle<T, d>::TAccelerationStructureHandle(TGeometryParticleHandleImp<T, d, bPersistent>* InHandle, TGeometryParticle<T, d>* InGeometryParticle)
 	: ExternalGeometryParticle(InGeometryParticle)
 	, GeometryParticleHandle(InHandle)
-	, HashResult(InGeometryParticle ? InGeometryParticle->GetHashResultLowLevel() : InHandle->GetHashResultLowLevel())
+	, HashResult(0)
 {
+	if (GeometryParticleHandle)
+	{
+		HashResult = GeometryParticleHandle->GetHashResultLowLevel();
+	}
+	else if (ExternalGeometryParticle)
+	{
+		HashResult = ExternalGeometryParticle->GetHashResultLowLevel();
+	}
+
+	if (GeometryParticleHandle && ExternalGeometryParticle)
+	{
+		CHAOS_ENSURE(GeometryParticleHandle->GetHashResultLowLevel() == ExternalGeometryParticle->GetHashResultLowLevel());
+	}
 }
 
 template <typename T, int d>
@@ -1684,6 +1697,32 @@ void TAccelerationStructureHandle<T, d>::Serialize(FChaosArchive& Ar)
 {
 	Ar << AsAlwaysSerializable(ExternalGeometryParticle);
 	Ar << AsAlwaysSerializable(GeometryParticleHandle);
+
+	Ar.UsingCustomVersion(FExternalPhysicsCustomObjectVersion::GUID);
+	if (Ar.CustomVer(FExternalPhysicsCustomObjectVersion::GUID) < FExternalPhysicsCustomObjectVersion::SerializeHashResult)
+	{
+		if (GeometryParticleHandle)
+		{
+			HashResult = GeometryParticleHandle->GetHashResultLowLevel();
+		}
+		else if (ExternalGeometryParticle)
+		{
+			HashResult = ExternalGeometryParticle->GetHashResultLowLevel();
+		}
+		else
+		{
+			HashResult = 0;
+		}
+
+		if (GeometryParticleHandle && ExternalGeometryParticle)
+		{
+			CHAOS_ENSURE(GeometryParticleHandle->GetHashResultLowLevel() == ExternalGeometryParticle->GetHashResultLowLevel());
+		}
+	}
+	else
+	{
+		Ar << HashResult;
+	}
 }
 
 template <typename T, int d>
