Chaos: Add ensure to AABBTree SplitNode to prevent using index -1 when facing numerical issues.

diff --git a/Engine/Source/Runtime/Experimental/Chaos/Public/Chaos/AABBTree.h b/Engine/Source/Runtime/Experimental/Chaos/Public/Chaos/AABBTree.h
--- a/Engine/Source/Runtime/Experimental/Chaos/Public/Chaos/AABBTree.h
+++ b/Engine/Source/Runtime/Experimental/Chaos/Public/Chaos/AABBTree.h
@@ -796,8 +796,11 @@ private:
 					++BoxIdx;
 				}
 
-				SplitInfos[MinBoxIdx].Children.Add(Elem);
-				SplitInfos[MinBoxIdx].RealBounds.GrowToInclude(Elem.Bounds);
+				if (CHAOS_ENSURE(MinBoxIdx != INDEX_NONE))
+				{
+					SplitInfos[MinBoxIdx].Children.Add(Elem);
+					SplitInfos[MinBoxIdx].RealBounds.GrowToInclude(Elem.Bounds);
+				}
 			}
 
 			NumProcessedThisSlice += Elems.Num();
