Chaos: Fix Evolution Serialization to not serialize acceleration update queues and flush beforehand. Fixes SQCapture crash.

diff --git a/Engine/Source/Runtime/Core/Public/UObject/ExternalPhysicsCustomObjectVersion.h b/Engine/Source/Runtime/Core/Public/UObject/ExternalPhysicsCustomObjectVersion.h
--- a/Engine/Source/Runtime/Core/Public/UObject/ExternalPhysicsCustomObjectVersion.h
+++ b/Engine/Source/Runtime/Core/Public/UObject/ExternalPhysicsCustomObjectVersion.h
@@ -54,6 +54,9 @@ struct CORE_API FExternalPhysicsCustomObjectVersion
 		// Adds Serialization of HashResult, and separates delete/update TAccelerationStructureHandle in FPendingSpatialData
 		SerializeHashResult,
 
+		// Only serialize internal acceleration structure queue and acceleration structure. No external/Async queues.
+		FlushEvolutionInternalAccelerationQueue,
+
 		// -----<new versions can be added above this line>-------------------------------------------------
 		VersionPlusOne,
 		LatestVersion = VersionPlusOne - 1
diff --git a/Engine/Source/Runtime/Experimental/Chaos/Private/Chaos/PBDRigidsEvolution.cpp b/Engine/Source/Runtime/Experimental/Chaos/Private/Chaos/PBDRigidsEvolution.cpp
--- a/Engine/Source/Runtime/Experimental/Chaos/Private/Chaos/PBDRigidsEvolution.cpp
+++ b/Engine/Source/Runtime/Experimental/Chaos/Private/Chaos/PBDRigidsEvolution.cpp
@@ -520,6 +520,11 @@ namespace Chaos
 
 		if (Ar.CustomVer(FExternalPhysicsCustomObjectVersion::GUID) >= FExternalPhysicsCustomObjectVersion::SerializeEvolutionBV)
 		{
+			if (Ar.CustomVer(FExternalPhysicsCustomObjectVersion::GUID) >= FExternalPhysicsCustomObjectVersion::FlushEvolutionInternalAccelerationQueue)
+			{
+				FlushInternalAccelerationQueue();
+			}
+
 			if (Ar.CustomVer(FExternalPhysicsCustomObjectVersion::GUID) < FExternalPhysicsCustomObjectVersion::SerializeMultiStructures)
 			{
 				//old path assumes single sub-structure
@@ -541,9 +546,12 @@ namespace Chaos
 				SpatialCollectionFactory->Serialize(InternalAcceleration, Ar);
 			}
 
-			SerializePendingMap(Ar, InternalAccelerationQueue);
-			SerializePendingMap(Ar, AsyncAccelerationQueue);
-			SerializePendingMap(Ar, ExternalAccelerationQueue);
+			if (Ar.CustomVer(FExternalPhysicsCustomObjectVersion::GUID) < FExternalPhysicsCustomObjectVersion::FlushEvolutionInternalAccelerationQueue)
+			{
+				SerializePendingMap(Ar, InternalAccelerationQueue);
+				SerializePendingMap(Ar, AsyncAccelerationQueue);
+				SerializePendingMap(Ar, ExternalAccelerationQueue);
+			}
 
 			ScratchExternalAcceleration = AsUniqueSpatialAccelerationChecked<FAccelerationStructure>(InternalAcceleration->Copy());
 		}
