Make EPA work with either winding. Handle case where origin is on simplex boundary

diff --git a/Engine/Source/Runtime/Experimental/Chaos/Public/Chaos/EPA.h b/Engine/Source/Runtime/Experimental/Chaos/Public/Chaos/EPA.h
--- a/Engine/Source/Runtime/Experimental/Chaos/Public/Chaos/EPA.h
+++ b/Engine/Source/Runtime/Experimental/Chaos/Public/Chaos/EPA.h
@@ -9,36 +9,81 @@ namespace Chaos
 template <typename T>
 struct TEPAEntry
 {
-	TVec3<T> Vertices[3];	//vertices of triangle
-	TVec3<T> V;	//closest point on affine hull of triangle
-	T Lambdas[3];	//barycentric coordinates of V (\sum Lambda[i] * Vertices[i])
-	T VSize2;	//distance from V to origin squared
-	TVector<int32,3> AdjIdxs;	//Adjacent triangles ordered CCW
-	TVector<int32,3> AdjIdxInOther;	//Used to find this entry in the corresponding adjacent triangle
+	int32 IdxBuffer[3];
+
+	TVec3<T> PlaneNormal;	//Triangle normal
+	T Distance;	//Triangle distance from origin
+	TVector<int32,3> AdjFaces;	//Adjacent triangles
+	TVector<int32,3> AdjEdges;	//Adjacent edges (idx in adjacent face)
 	bool bObsolete;	//indicates that an entry can be skipped (became part of bigger polytope)
-	bool bVInterior;	//indicates whether V is an interior point or not
 
 	bool operator>(const TEPAEntry<T>& Other) const
 	{
-		return VSize2 > Other.VSize2;
+		return Distance > Other.Distance;
 	}
 
-	void Initialize(const TVec3<T>& V0, const TVec3<T>& V1, const TVec3<T>& V2,
-		const TVector<int32,3>& InAdjIdx, const TVector<int32,3>& InAdjIdxInOther)
+	bool Initialize(const TVec3<T>* Vertices, int32 InIdx0, int32 InIdx1, int32 InIdx2, const TVector<int32,3>& InAdjFaces, const TVector<int32,3>& InAdjEdges)
 	{
-		Vertices[0] = V0;
-		Vertices[1] = V1;
-		Vertices[2] = V2;
+		const TVec3<T>& V0 = Vertices[InIdx0];
+		const TVec3<T>& V1 = Vertices[InIdx1];
+		const TVec3<T>& V2 = Vertices[InIdx2];
 
-		AdjIdxs = InAdjIdx;
-		AdjIdxInOther = InAdjIdxInOther;
+		const TVec3<T> V0V1 = V1 - V0;
+		const TVec3<T> V0V2 = V2 - V0;
+		const TVec3<T> Norm = TVec3<T>::CrossProduct(V0V1, V0V2);
+		PlaneNormal = Norm.GetSafeNormal();
+		constexpr T Eps = 1e-4;
+		if (PlaneNormal.SizeSquared() < Eps)
+		{
+			return false;
+		}
+		
+		IdxBuffer[0] = InIdx0;
+		IdxBuffer[1] = InIdx1;
+		IdxBuffer[2] = InIdx2;
 
-		FSimplex Simplex = {0,1,2};
-		V = TriangleSimplexFindOrigin(Vertices, Simplex, Lambdas);
-		VSize2 = V.SizeSquared();
+		AdjFaces = InAdjFaces;
+		AdjEdges = InAdjEdges;
 
-		bVInterior = Simplex.NumVerts == 3;
+		Distance = TVec3<T>::DotProduct(PlaneNormal, V0);
 		bObsolete = false;
+
+		return true;
+	}
+
+	void SwapWinding(TEPAEntry* Entries)
+	{
+		//change vertex order
+		std::swap(IdxBuffer[0], IdxBuffer[1]);
+
+		//edges went from 0,1,2 to 1,0,2
+		//0th edge/face is the same (0,1 becomes 1,0)
+		//1th edge/face is now (0,2 instead of 1,2)
+		//2nd edge/face is now (2,1 instead of 2,0)
+
+		//update the adjacent face's adjacent edge first
+		auto UpdateAdjEdge = [Entries, this](int32 Old, int32 New)
+		{
+			TEPAEntry& AdjFace = Entries[AdjFaces[Old]];
+			int32& StaleAdjIdx = AdjFace.AdjEdges[AdjEdges[Old]];
+			check(StaleAdjIdx == Old);
+			StaleAdjIdx = New;
+		};
+		
+		UpdateAdjEdge(1, 2);
+		UpdateAdjEdge(2, 1);
+		
+		//now swap the actual edges and faces
+		std::swap(AdjFaces[1], AdjFaces[2]);
+		std::swap(AdjEdges[1], AdjEdges[2]);
+
+		PlaneNormal = -PlaneNormal;
+		Distance = -Distance;
+	}
+
+	T DistanceToPlane(const TVec3<T>& X) const
+	{
+		return TVec3<T>::DotProduct(PlaneNormal, X) - Distance;
 	}
 };
 
@@ -50,11 +95,23 @@ TArray<TEPAEntry<T>> InitializeEPA(const TVec3<T>* Verts, const int32 NumVerts)
 	{
 		case 4:
 		{
+			//make sure all triangle normals are facing out
 			Entries.AddUninitialized(4);
-			Entries[0].Initialize(Verts[0], Verts[1], Verts[2], {1,2,3}, {2, 2, 1});
-			Entries[1].Initialize(Verts[0], Verts[3], Verts[1], {3,2,0}, {0,0,0});
-			Entries[2].Initialize(Verts[1], Verts[3], Verts[2], {1,3,0}, {1,2,1});
-			Entries[3].Initialize(Verts[3], Verts[0], Verts[2], {1,0,2}, {0,2,1});
+
+			ensure(Entries[0].Initialize(Verts, 1, 2, 3, { 3, 1, 2 }, { 1,1, 1 }));
+			ensure(Entries[1].Initialize(Verts, 0, 3, 2, { 2,0,3 }, { 2, 1, 0 }));
+			ensure(Entries[2].Initialize(Verts, 0, 1, 3, { 3,0, 1 }, { 2,2,0 }));
+			ensure(Entries[3].Initialize(Verts, 0, 2, 1, { 1,0,2 }, { 2,0,0 }));
+
+			if (TVec3<T>::DotProduct(Entries[0].PlaneNormal, Verts[0]) > 0)
+			{
+				//tet faces are pointing inwards
+				for (TEPAEntry<T>& Entry : Entries)
+				{
+					Entry.SwapWinding(Entries.GetData());
+				}
+			}
+
 			return Entries;
 		}
 
@@ -76,7 +133,7 @@ void EPAComputeVisibilityBorder(TArray<TEPAEntry<T>>& Entries, int32 EntryIdx, c
 		TEPAEntry<T>& Entry = Entries[EntryIdx];
 		for (int i = 0; i < 3; ++i)
 		{
-			ToVisitStack.Add({ Entry.AdjIdxs[i], Entry.AdjIdxInOther[i] });
+			ToVisitStack.Add({ Entry.AdjFaces[i], Entry.AdjEdges[i] });
 		}
 	}
 
@@ -86,8 +143,7 @@ void EPAComputeVisibilityBorder(TArray<TEPAEntry<T>>& Entries, int32 EntryIdx, c
 		TEPAEntry<T>& Entry = Entries[FloodEntry.EntryIdx];
 		if (!Entry.bObsolete)
 		{
-			const T VDotW = TVec3<T>::DotProduct(Entry.V, W);
-			if (VDotW < Entry.VSize2)	//plane check
+			if (Entry.DistanceToPlane(W) < 0)
 			{
 				//W can't see this triangle so mark the edge as a border
 				OutBorderEdges.Add(FloodEntry);
@@ -99,15 +155,15 @@ void EPAComputeVisibilityBorder(TArray<TEPAEntry<T>>& Entries, int32 EntryIdx, c
 				const int32 Idx0 = FloodEntry.EdgeIdx;
 				const int32 Idx1 = (Idx0 + 1) % 3;
 				const int32 Idx2 = (Idx0 + 2) % 3;
-				ToVisitStack.Add({ Entry.AdjIdxs[Idx1], Entry.AdjIdxInOther[Idx1] });
-				ToVisitStack.Add({ Entry.AdjIdxs[Idx2], Entry.AdjIdxInOther[Idx2] });
+				ToVisitStack.Add({ Entry.AdjFaces[Idx1], Entry.AdjEdges[Idx1] });
+				ToVisitStack.Add({ Entry.AdjFaces[Idx2], Entry.AdjEdges[Idx2] });
 			}
 		}
 	}
 }
 
 template <typename T, typename SupportLambda>
-T EPA(const TVec3<T>* Verts, const int32 NumVerts, const SupportLambda& Support)
+T EPA(TArray<TVec3<T>> VertsBuffer, const SupportLambda& Support)
 {
 	struct FEPAEntryWrapper
 	{
@@ -120,23 +176,25 @@ T EPA(const TVec3<T>* Verts, const int32 NumVerts, const SupportLambda& Support)
 		}
 	};
 
-	TArray<TEPAEntry<T>> Entries = InitializeEPA(Verts, NumVerts);
+	T UpperBound = TNumericLimits<T>::Max();
+	T LowerBound = TNumericLimits<T>::Lowest();
+
+	TArray<TEPAEntry<T>> Entries = InitializeEPA(VertsBuffer.GetData(), VertsBuffer.Num());
 	std::priority_queue<FEPAEntryWrapper, std::vector<FEPAEntryWrapper>, std::greater<FEPAEntryWrapper>> Queue;
 	for(int32 Idx = 0; Idx < Entries.Num(); ++Idx)
 	{
-		if(Entries[Idx].bVInterior)
-		{
-			Queue.push(FEPAEntryWrapper {&Entries, Idx});
-		}
+		Queue.push(FEPAEntryWrapper {&Entries, Idx});
 	}
 
 	bool bCloseEnough = false;
 	bool bTerminate = bCloseEnough;
 
 	TArray<FEPAFloodEntry> VisibilityBorder;
-
-	while(Queue.size() && !bTerminate)
+	int32 Iteration = -1;
+	while(Queue.size())
 	{
+		++Iteration;
+
 		int32 EntryIdx = Queue.top().Idx;
 		Queue.pop();
 		TEPAEntry<T>& Entry = Entries[EntryIdx];
@@ -145,53 +203,56 @@ T EPA(const TVec3<T>* Verts, const int32 NumVerts, const SupportLambda& Support)
 			continue;
 		}
 
-		const TVec3<T>& V = Entry.V;
-		const TVec3<T> W = Support(V);
 		constexpr T Eps = 1e-4;
-		const T VDotW = TVec3<T>::DotProduct(V, W);
-		bCloseEnough = VDotW < Entry.VSize2 + Eps;	//todo: add relative eps
-		if (!bCloseEnough)
-		{
-			Entry.bObsolete = true;
-			VisibilityBorder.Reset();
-			//question: in the paper we call this per edge which means we traverse new faces for second and third edge, does it matter?
-			EPAComputeVisibilityBorder(Entries, EntryIdx, W, VisibilityBorder);
-			const int32 NumBorderEdges = VisibilityBorder.Num();
-			const int32 FirstIdxInBatch = Entries.Num();
-			int32 NewIdx = FirstIdxInBatch;
-			Entries.AddUninitialized(NumBorderEdges);
-			ensure(NumBorderEdges >= 3);
-			for (int32 VisibilityIdx = 0; VisibilityIdx < NumBorderEdges; ++VisibilityIdx)
-			{
-				//create new entries and update adjacencies
-				const FEPAFloodEntry& BorderInfo = VisibilityBorder[VisibilityIdx];
-				TEPAEntry<T>& NewEntry = Entries[NewIdx];
-				const int32 BorderEntryIdx = BorderInfo.EntryIdx;
-				TEPAEntry<T>& BorderEntry = Entries[BorderEntryIdx];
-				const int32 BorderEdgeIdx0 = BorderInfo.EdgeIdx;
-				const int32 BorderEdgeIdx1 = (BorderEdgeIdx0 + 1) % 3;
-				const int32 NextEntryIdx = (VisibilityIdx + 1) < VisibilityBorder.Num() ? NewIdx + 1 : FirstIdxInBatch;
-				const int32 PrevEntryIdx = NewIdx > FirstIdxInBatch ? NewIdx - 1 : FirstIdxInBatch + NumBorderEdges - 1;
-				NewEntry.Initialize(BorderEntry.Vertices[BorderEdgeIdx1], BorderEntry.Vertices[BorderEdgeIdx0], W,
-					{ BorderEntryIdx, NextEntryIdx, PrevEntryIdx },
-					{ BorderEdgeIdx0, 2, 1 });
-				BorderEntry.AdjIdxs[BorderEdgeIdx0] = NewIdx;
-				BorderEntry.AdjIdxInOther[BorderEdgeIdx0] = 0;
-
-				if (NewEntry.bVInterior)
-				{
-					Queue.push(FEPAEntryWrapper{ &Entries, NewIdx});
-				}
 
-				++NewIdx;
-			}
-		}
-		else
+
+		const TVec3<T> W = Support(Entry.PlaneNormal);
+		const T DistanceToSupportPlane = TVec3<T>::DotProduct(Entry.PlaneNormal, W);
+
+		UpperBound = FMath::Min(UpperBound, DistanceToSupportPlane);
+		ensure(LowerBound <= Entry.Distance);	//todo: remove this check, just for early development
+		LowerBound = Entry.Distance;
+
+		if (FMath::Abs(UpperBound - LowerBound) < Eps)	//todo: add relative?
 		{
-			return FMath::Sqrt(Entry.VSize2);
+			return Entry.Distance;
 		}
 
-		bTerminate = bCloseEnough;
+		const int32 NewVertIdx = VertsBuffer.Add(W);
+
+		Entry.bObsolete = true;
+		VisibilityBorder.Reset();
+		EPAComputeVisibilityBorder(Entries, EntryIdx, W, VisibilityBorder);
+		const int32 NumBorderEdges = VisibilityBorder.Num();
+		const int32 FirstIdxInBatch = Entries.Num();
+		int32 NewIdx = FirstIdxInBatch;
+		Entries.AddUninitialized(NumBorderEdges);
+		ensure(NumBorderEdges >= 3);
+		for (int32 VisibilityIdx = 0; VisibilityIdx < NumBorderEdges; ++VisibilityIdx)
+		{
+			//create new entries and update adjacencies
+			const FEPAFloodEntry& BorderInfo = VisibilityBorder[VisibilityIdx];
+			TEPAEntry<T>& NewEntry = Entries[NewIdx];
+			const int32 BorderEntryIdx = BorderInfo.EntryIdx;
+			TEPAEntry<T>& BorderEntry = Entries[BorderEntryIdx];
+			const int32 BorderEdgeIdx0 = BorderInfo.EdgeIdx;
+			const int32 BorderEdgeIdx1 = (BorderEdgeIdx0 + 1) % 3;
+			const int32 NextEntryIdx = (VisibilityIdx + 1) < VisibilityBorder.Num() ? NewIdx + 1 : FirstIdxInBatch;
+			const int32 PrevEntryIdx = NewIdx > FirstIdxInBatch ? NewIdx - 1 : FirstIdxInBatch + NumBorderEdges - 1;
+			ensure(NewEntry.Initialize(VertsBuffer.GetData(), BorderEntry.IdxBuffer[BorderEdgeIdx1], BorderEntry.IdxBuffer[BorderEdgeIdx0], NewVertIdx,
+				{ BorderEntryIdx, PrevEntryIdx, NextEntryIdx },
+				{ BorderEdgeIdx0, 2, 1 }));	//todo: handle false case which is a degenerate triangle
+			BorderEntry.AdjFaces[BorderEdgeIdx0] = NewIdx;
+			BorderEntry.AdjEdges[BorderEdgeIdx0] = 0;
+
+			ensure(NewEntry.Distance >= LowerBound);	//todo: remove, just used for early dev
+			if (/*NewEntry.Distance >= LowerBound && */NewEntry.Distance <= UpperBound)
+			{
+				Queue.push(FEPAEntryWrapper{ &Entries, NewIdx});
+			}
+
+			++NewIdx;
+		}
 	}
 
 	return 0;
