Chaos : Collision Restructure

- Updated collision array to hold a list of shape manifolds per body pairs.

diff --git a/Engine/Source/Runtime/Experimental/Chaos/Private/Chaos/ChaosDebugDraw.cpp b/Engine/Source/Runtime/Experimental/Chaos/Private/Chaos/ChaosDebugDraw.cpp
--- a/Engine/Source/Runtime/Experimental/Chaos/Private/Chaos/ChaosDebugDraw.cpp
+++ b/Engine/Source/Runtime/Experimental/Chaos/Private/Chaos/ChaosDebugDraw.cpp
@@ -136,13 +136,13 @@ namespace Chaos
 		void DrawCollisionImpl(const FRigidTransform3& SpaceTransform, const TPBDCollisionConstraintHandle<float, 3>* ConstraintHandle, float ColorScale)
 		{
 			const TRigidBodyContactConstraint<FReal, 3>& Contact = ConstraintHandle->GetContact();
-			if (Contact.Phi > 0)
+			if (Contact.GetPhi() > 0)
 			{
 				ColorScale = ColorScale * (FReal)0.1;
 			}
 
-			FVec3 Location = SpaceTransform.TransformPosition(Contact.Location);
-			FVec3 Normal = SpaceTransform.TransformVector(Contact.Normal);
+			FVec3 Location = SpaceTransform.TransformPosition(Contact.GetLocation());
+			FVec3 Normal = SpaceTransform.TransformVector(Contact.GetNormal());
 
 			if (ContactWidth > 0)
 			{
@@ -155,11 +155,11 @@ namespace Chaos
 				FColor C1 = (ColorScale * FColor(255, 0, 0)).ToFColor(false);
 				FDebugDrawQueue::GetInstance().DrawDebugLine(Location, Location + DrawScale * ContactLen * Normal, C1, false, KINDA_SMALL_NUMBER, DrawPriority, LineThickness);
 			}
-			if (ContactPhiWidth > 0 && Contact.Phi < FLT_MAX)
+			if (ContactPhiWidth > 0 && Contact.GetPhi() < FLT_MAX)
 			{
 				FColor C2 = (ColorScale * FColor(128, 128, 0)).ToFColor(false);
 				FMatrix Axes = FRotationMatrix::MakeFromX(Normal);
-				FDebugDrawQueue::GetInstance().DrawDebugCircle(Location - Contact.Phi * Normal, DrawScale * ContactPhiWidth, 12, C2, false, KINDA_SMALL_NUMBER, DrawPriority, LineThickness, Axes.GetUnitAxis(EAxis::Y), Axes.GetUnitAxis(EAxis::Z), false);
+				FDebugDrawQueue::GetInstance().DrawDebugCircle(Location - Contact.GetPhi() * Normal, DrawScale * ContactPhiWidth, 12, C2, false, KINDA_SMALL_NUMBER, DrawPriority, LineThickness, Axes.GetUnitAxis(EAxis::Y), Axes.GetUnitAxis(EAxis::Z), false);
 			}
 		}
 
diff --git a/Engine/Source/Runtime/Experimental/Chaos/Private/Chaos/CollisionResolutionConvexConvex.cpp b/Engine/Source/Runtime/Experimental/Chaos/Private/Chaos/CollisionResolutionConvexConvex.cpp
--- a/Engine/Source/Runtime/Experimental/Chaos/Private/Chaos/CollisionResolutionConvexConvex.cpp
+++ b/Engine/Source/Runtime/Experimental/Chaos/Private/Chaos/CollisionResolutionConvexConvex.cpp
@@ -12,25 +12,16 @@ namespace Chaos
 		FGeometryParticleHandle* Particle0, FGeometryParticleHandle* Particle1,
 		const FImplicitObject* Implicit0, const FImplicitObject* Implicit1,
 		const float Thickness,
-		TArray<FRigidBodyContactConstraint>& ConstraintBuffer,
-		FCollisionResolutionManifold* Manifold,
-		float ManifoldScale, int32 ManifoldSamples)
+		FRigidBodyContactConstraint& Constraint)
 	{
-		if (Manifold)
+		if (Constraint.ContainsManifold(Implicit0, Implicit1))
 		{
-			if (Manifold->ContainsShapeConnection(Implicit0, Implicit1))
-			{
-				return;
-			}
+			return;
 		}
 
-		TRigidBodyContactConstraint<T, d> Constraint;
 		Constraint.Particle = Particle0;
 		Constraint.Levelset = Particle1;
-		Constraint.Geometry[0] = Implicit0;
-		Constraint.Geometry[1] = Implicit1;
-
-		ConstraintBuffer.Add(Constraint);
+		Constraint.AddManifold(Implicit0, Implicit1);
 	}
 
 	template<class T, int d>
@@ -38,12 +29,8 @@ namespace Chaos
 		const FImplicitObject& A, const FRigidTransform& ATM,
 		const FImplicitObject& B, FRigidTransform BTM,
 		float Thickness,
-		FRigidBodyContactConstraint* Constraints, int32 NumConstraints,
-		FCollisionResolutionManifold* Manifold,
-		float ManifoldScale, int32 ManifoldSamples)
+		FRigidBodyContactConstraint& Constraint)
 	{
-		ensure(NumConstraints);
-
 		TRigidTransform<T, d> BToATM = BTM.GetRelativeTransform(ATM);
 
 		if (ensure(GetInnerType(A.GetType()) == ImplicitObjectType::Convex && GetInnerType(B.GetType()) == ImplicitObjectType::Convex))
@@ -66,10 +53,7 @@ namespace Chaos
 			{
 				const TParticles<T, d>& SurfaceParticles = AObject->GetSurfaceParticles();
 
-				for (int32 Idx = 0; Idx < NumConstraints; Idx++)
-				{
-					Constraints[Idx].bDisabled = true;
-				}
+				Constraint.SetDisabled(true);
 
 				if (GJKIntersection(*AObject, *BObject, BToATM, Thickness))
 				{
@@ -80,11 +64,11 @@ namespace Chaos
 					for (int32 Idx = 0; Idx < (int32)SurfaceParticles.Size(); Idx++)
 					{
 						Phi = BObject->PhiWithNormal(AToBTM.TransformPosition(SurfaceParticles.X(Idx)), Normal);
-						if (Phi < Constraints[0].Phi)
+						if (Phi < Constraint.GetPhi())
 						{
-							Constraints[0].Phi = Phi;
-							Constraints[0].Location = ATM.TransformPosition(SurfaceParticles.X(Idx));
-							Constraints[0].Normal = BTM.TransformVector(Normal);
+							Constraint.SetPhi(Phi);
+							Constraint.SetLocation(ATM.TransformPosition(SurfaceParticles.X(Idx)));
+							Constraint.SetNormal(BTM.TransformVector(Normal));
 						}
 					}
 					return;
diff --git a/Engine/Source/Runtime/Experimental/Chaos/Private/Chaos/PBDCollisionConstraint.cpp b/Engine/Source/Runtime/Experimental/Chaos/Private/Chaos/PBDCollisionConstraint.cpp
--- a/Engine/Source/Runtime/Experimental/Chaos/Private/Chaos/PBDCollisionConstraint.cpp
+++ b/Engine/Source/Runtime/Experimental/Chaos/Private/Chaos/PBDCollisionConstraint.cpp
@@ -34,15 +34,6 @@ FAutoConsoleVariableRef CVarConstraintBPBVHDepth(TEXT("p.ConstraintBPBVHDepth"),
 int32 BPTreeOfGrids = 1;
 FAutoConsoleVariableRef CVarBPTreeOfGrids(TEXT("p.BPTreeOfGrids"), BPTreeOfGrids, TEXT("Whether to use a seperate tree of grids for bp"));
 
-int32 CollisionManifoldLifespanCVar = 1;
-FAutoConsoleVariableRef CVarCollisionManifoldLifespan(TEXT("p.CollisionManifoldLifespan"), CollisionManifoldLifespanCVar, TEXT("Number of iterations to cache a collision point.[def:5]"));
-
-int32 CollisionManifoldSamplesCVar = 5;
-FAutoConsoleVariableRef CVarCollisionManifoldSamples(TEXT("p.CollisionManifoldSamples"), CollisionManifoldSamplesCVar, TEXT("Number of sample points on the manifold.[def:6]"));
-
-float CollisionManifoldSampleScaleCVar = 0.7;
-FAutoConsoleVariableRef CVarCollisionManifoldSampleScale(TEXT("p.CollisionManifoldSampleScale"), CollisionManifoldSampleScaleCVar, TEXT("Scale factor for the collision manifold.[def:0.7]"));
-
 float CollisionVelocityInflationCVar = 2.0;
 FAutoConsoleVariableRef CVarCollisionVelocityInflation(TEXT("p.CollisionVelocityInflation"), CollisionVelocityInflationCVar, TEXT("Collision velocity inflation.[def:2.0]"));
 
@@ -62,10 +53,10 @@ CHAOS_API int32 EnableCollisions = 1;
 FAutoConsoleVariableRef CVarEnableCollisions(TEXT("p.EnableCollisions"), EnableCollisions, TEXT("Enable/Disable collisions on the Chaos solver."));
 
 template<ECollisionUpdateType, typename T, int d>
-void UpdateConstraintImp2(const TRigidTransform<T, d>& ParticleTM, const FImplicitObject& LevelsetObject, const TRigidTransform<T, d>& LevelsetTM, const T Thickness, TRigidBodyContactConstraint<T, d>& Constraint, int32 NumConstraints, TCollisionResolutionManifold<T, d>* Manifold);
+void UpdateConstraintImp(const TRigidTransform<T, d>& ParticleTM, const FImplicitObject& LevelsetObject, const TRigidTransform<T, d>& LevelsetTM, const T Thickness, TRigidBodyContactConstraint<T, d>& Constraint);
 
 template<typename T, int d>
-void ConstructConstraintsImpl(TGeometryParticleHandle<T, d>* Particle0, TGeometryParticleHandle<T, d>* Particle1, const FImplicitObject* Implicit0, const FImplicitObject* Implicit1, const T Thickness, TArray<TRigidBodyContactConstraint<T, d>>& ConstraintBuffer, TCollisionResolutionManifold<T, d>* Manifold);
+void ConstructConstraintsImpl(TGeometryParticleHandle<T, d>* Particle0, TGeometryParticleHandle<T, d>* Particle1, const FImplicitObject* Implicit0, const FImplicitObject* Implicit1, const T Thickness, TRigidBodyContactConstraint<T, d>& Constraint);
 
 
 //
@@ -85,7 +76,6 @@ TPBDCollisionConstraint<T, d>::TPBDCollisionConstraint(const TPBDRigidsSOAs<T,d>
 	, bUseCCD(false)
 	, bEnableCollisions(true)
 	, LifespanCounter(0)
-	, CollisionManifoldLifespan(CollisionManifoldLifespanCVar)
 	, CollisionVelocityInflation(CollisionVelocityInflationCVar)
 	, PostComputeCallback(nullptr)
 	, PostApplyCallback(nullptr)
@@ -100,7 +90,6 @@ void TPBDCollisionConstraint<T, d>::Reset(/*const TPBDRigidParticles<T, d>& InPa
 {
 	SCOPE_CYCLE_COUNTER(STAT_CollisionConstraintsReset);
 
-	// @todo(brice) : Manifolds are changing, so just default to collision constraints right now.
 
 	for (int32 Idx = Constraints.Num() - 1; Idx >= 0; Idx--)
 	{
@@ -110,22 +99,6 @@ void TPBDCollisionConstraint<T, d>::Reset(/*const TPBDRigidParticles<T, d>& InPa
 		}
 	}
 
-	/*
-	TArray<FConstraintHandleID> ManifoldKeys;
-	Manifolds.GetKeys(ManifoldKeys);
-	for (FConstraintHandleID Key : ManifoldKeys)
-	{
-		FConstraintManifold* Elem = Manifolds[Key];
-		if (Elem->GetTimestamp()==FLT_MAX || LifespanCounter > Elem->GetTimestamp())
-		{
-			const TArray<FConstraintContainerHandle*> HandleArray = Elem->GetHandles();
-			for (FConstraintContainerHandle* Handle : HandleArray)
-			{
-				RemoveConstraint(Handle->GetConstraintIndex());
-			}
-		}
-	}
-	*/
 
 	MAngularFriction = 0;
 	bUseCCD = false;
@@ -134,22 +107,6 @@ void TPBDCollisionConstraint<T, d>::Reset(/*const TPBDRigidParticles<T, d>& InPa
 template<typename T, int d>
 void TPBDCollisionConstraint<T, d>::RemoveConstraint(int32 Idx)
 {
-	// Manifold
-	/*
-	FConstraintHandleID HandleID = GetConstraintHandleID(Constraints[Idx]);
-	if (FConstraintManifold** Ptr = Manifolds.Find(HandleID))
-	{
-		if (FConstraintManifold* ManifoldEntry = *Ptr)
-		{
-			ManifoldEntry->RemoveHandle(Handles[Idx]);
-			if (!ManifoldEntry->GetHandles().Num())
-			{
-				Manifolds.Remove(HandleID);
-			}
-		}
-	}
-	*/
-
 	HandleAllocator.FreeHandle(Handles[Idx]);
 	Handles.RemoveAtSwap(Idx);
 	Constraints.RemoveAtSwap(Idx);
@@ -546,8 +503,6 @@ TVector<T, d> GetEnergyClampedImpulse(const TRigidBodyContactConstraint<T, d>& C
 template<typename T, int d>
 void TPBDCollisionConstraint<T, d>::Apply(const T Dt, FRigidBodyContactConstraint& Constraint)
 {
-	if (Constraint.bDisabled) return;
-
 	TGenericParticleHandle<T, d> Particle0 = TGenericParticleHandle<T, d>(Constraint.Particle);
 	TGenericParticleHandle<T, d> Particle1 = TGenericParticleHandle<T, d>(Constraint.Levelset);
 	TPBDRigidParticleHandle<T, d>* PBDRigid0 = Particle0->AsDynamic();
@@ -564,8 +519,8 @@ void TPBDCollisionConstraint<T, d>::Apply(const T Dt, FRigidBodyContactConstrain
 		ensure(!PBDRigid0 || PBDRigid0->Sleeping());
 		return;
 	}
-	UpdateConstraint<ECollisionUpdateType::Deepest>(MThickness, &Constraint, 1);
-	if (Constraint.Phi >= MThickness)
+	UpdateConstraint<ECollisionUpdateType::Deepest>(MThickness, Constraint);
+	if (Constraint.GetPhi() >= MThickness)
 	{
 		return;
 	}
@@ -587,12 +542,14 @@ void TPBDCollisionConstraint<T, d>::Apply(const T Dt, FRigidBodyContactConstrain
 	TSerializablePtr<FChaosPhysicsMaterial> PhysicsMaterial0 = Particle0->AuxilaryValue(MPhysicsMaterials);
 	TSerializablePtr<FChaosPhysicsMaterial> PhysicsMaterial1 = Particle1->AuxilaryValue(MPhysicsMaterials);
 
-	TVector<T, d> VectorToPoint1 = Constraint.Location - P0;
-	TVector<T, d> VectorToPoint2 = Constraint.Location - P1;
+	TContactData<T,d> & Contact = Constraint.ShapeManifold.Manifold;
+
+	TVector<T, d> VectorToPoint1 = Contact.Location - P0;
+	TVector<T, d> VectorToPoint2 = Contact.Location - P1;
 	TVector<T, d> Body1Velocity = V0 + TVector<T, d>::CrossProduct(W0, VectorToPoint1);
 	TVector<T, d> Body2Velocity = V1 + TVector<T, d>::CrossProduct(W1, VectorToPoint2);
 	TVector<T, d> RelativeVelocity = Body1Velocity - Body2Velocity;
-	if (TVector<T, d>::DotProduct(RelativeVelocity, Constraint.Normal) < 0) // ignore separating constraints
+	if (TVector<T, d>::DotProduct(RelativeVelocity, Contact.Normal) < 0) // ignore separating constraints
 	{
 		PMatrix<T, d, d> WorldSpaceInvI1 = PBDRigid0 ? (Q0 * FMatrix::Identity).GetTransposed() * PBDRigid0->InvI() * (Q0 * FMatrix::Identity) : PMatrix<T, d, d>(0);
 		PMatrix<T, d, d> WorldSpaceInvI2 = PBDRigid1 ? (Q1 * FMatrix::Identity).GetTransposed() * PBDRigid1->InvI() * (Q1 * FMatrix::Identity) : PMatrix<T, d, d>(0);
@@ -633,17 +590,17 @@ void TPBDCollisionConstraint<T, d>::Apply(const T Dt, FRigidBodyContactConstrain
 
 		if (Friction)
 		{
-			T RelativeNormalVelocity = TVector<T, d>::DotProduct(RelativeVelocity, Constraint.Normal);
+			T RelativeNormalVelocity = TVector<T, d>::DotProduct(RelativeVelocity, Contact.Normal);
 			if (RelativeNormalVelocity > 0)
 			{
 				RelativeNormalVelocity = 0;
 			}
-			TVector<T, d> VelocityChange = -(Restitution * RelativeNormalVelocity * Constraint.Normal + RelativeVelocity);
-			T NormalVelocityChange = TVector<T, d>::DotProduct(VelocityChange, Constraint.Normal);
+			TVector<T, d> VelocityChange = -(Restitution * RelativeNormalVelocity * Contact.Normal + RelativeVelocity);
+			T NormalVelocityChange = TVector<T, d>::DotProduct(VelocityChange, Contact.Normal);
 			PMatrix<T, d, d> FactorInverse = Factor.Inverse();
 			TVector<T, d> MinimalImpulse = FactorInverse * VelocityChange;
-			const T MinimalImpulseDotNormal = TVector<T, d>::DotProduct(MinimalImpulse, Constraint.Normal);
-			const T TangentialSize = (MinimalImpulse - MinimalImpulseDotNormal * Constraint.Normal).Size();
+			const T MinimalImpulseDotNormal = TVector<T, d>::DotProduct(MinimalImpulse, Contact.Normal);
+			const T TangentialSize = (MinimalImpulse - MinimalImpulseDotNormal * Contact.Normal).Size();
 			if (TangentialSize <= Friction * MinimalImpulseDotNormal)
 			{
 				//within friction cone so just solve for static friction stopping the object
@@ -651,9 +608,9 @@ void TPBDCollisionConstraint<T, d>::Apply(const T Dt, FRigidBodyContactConstrain
 				if (MAngularFriction)
 				{
 					TVector<T, d> RelativeAngularVelocity = W0 - W1;
-					T AngularNormal = TVector<T, d>::DotProduct(RelativeAngularVelocity, Constraint.Normal);
-					TVector<T, d> AngularTangent = RelativeAngularVelocity - AngularNormal * Constraint.Normal;
-					TVector<T, d> FinalAngularVelocity = FMath::Sign(AngularNormal) * FMath::Max((T)0, FMath::Abs(AngularNormal) - MAngularFriction * NormalVelocityChange) * Constraint.Normal + FMath::Max((T)0, AngularTangent.Size() - MAngularFriction * NormalVelocityChange) * AngularTangent.GetSafeNormal();
+					T AngularNormal = TVector<T, d>::DotProduct(RelativeAngularVelocity, Contact.Normal);
+					TVector<T, d> AngularTangent = RelativeAngularVelocity - AngularNormal * Contact.Normal;
+					TVector<T, d> FinalAngularVelocity = FMath::Sign(AngularNormal) * FMath::Max((T)0, FMath::Abs(AngularNormal) - MAngularFriction * NormalVelocityChange) * Contact.Normal + FMath::Max((T)0, AngularTangent.Size() - MAngularFriction * NormalVelocityChange) * AngularTangent.GetSafeNormal();
 					TVector<T, d> Delta = FinalAngularVelocity - RelativeAngularVelocity;
 					if (!PBDRigid0 && PBDRigid1)
 					{
@@ -691,9 +648,9 @@ void TPBDCollisionConstraint<T, d>::Apply(const T Dt, FRigidBodyContactConstrain
 			else
 			{
 				//outside friction cone, solve for normal relative velocity and keep tangent at cone edge
-				TVector<T, d> Tangent = (RelativeVelocity - TVector<T, d>::DotProduct(RelativeVelocity, Constraint.Normal) * Constraint.Normal).GetSafeNormal();
-				TVector<T, d> DirectionalFactor = Factor * (Constraint.Normal - Friction * Tangent);
-				T ImpulseDenominator = TVector<T, d>::DotProduct(Constraint.Normal, DirectionalFactor);
+				TVector<T, d> Tangent = (RelativeVelocity - TVector<T, d>::DotProduct(RelativeVelocity, Contact.Normal) * Contact.Normal).GetSafeNormal();
+				TVector<T, d> DirectionalFactor = Factor * (Contact.Normal - Friction * Tangent);
+				T ImpulseDenominator = TVector<T, d>::DotProduct(Contact.Normal, DirectionalFactor);
 				if (!ensureMsgf(FMath::Abs(ImpulseDenominator) > SMALL_NUMBER, TEXT("Constraint:%s\n\nParticle:%s\n\nLevelset:%s\n\nDirectionalFactor:%s, ImpulseDenominator:%f"),
 					*Constraint.ToString(),
 					*Particle0->ToString(),
@@ -704,18 +661,18 @@ void TPBDCollisionConstraint<T, d>::Apply(const T Dt, FRigidBodyContactConstrain
 				}
 
 				const T ImpulseMag = -(1 + Restitution) * RelativeNormalVelocity / ImpulseDenominator;
-				Impulse = ImpulseMag * (Constraint.Normal - Friction * Tangent);
+				Impulse = ImpulseMag * (Contact.Normal - Friction * Tangent);
 			}
 		}
 		else
 		{
-			T ImpulseDenominator = TVector<T, d>::DotProduct(Constraint.Normal, Factor * Constraint.Normal);
-			TVector<T, d> ImpulseNumerator = -(1 + Restitution) * TVector<T, d>::DotProduct(RelativeVelocity, Constraint.Normal)* Constraint.Normal;
+			T ImpulseDenominator = TVector<T, d>::DotProduct(Contact.Normal, Factor * Contact.Normal);
+			TVector<T, d> ImpulseNumerator = -(1 + Restitution) * TVector<T, d>::DotProduct(RelativeVelocity, Contact.Normal)* Contact.Normal;
 			if (!ensureMsgf(FMath::Abs(ImpulseDenominator) > SMALL_NUMBER, TEXT("Constraint:%s\n\nParticle:%s\n\nLevelset:%s\n\nFactor*Constraint.Normal:%s, ImpulseDenominator:%f"),
 				*Constraint.ToString(),
 				*Particle0->ToString(),
 				*Particle1->ToString(),
-				*(Factor * Constraint.Normal).ToString(), ImpulseDenominator))
+				*(Factor * Contact.Normal).ToString(), ImpulseDenominator))
 			{
 				ImpulseDenominator = (T)1;
 			}
@@ -768,15 +725,6 @@ void TPBDCollisionConstraint<T, d>::Apply(const T Dt, const TArray<FConstraintCo
 	}
 }
 
-template<typename T, int d>
-void TPBDCollisionConstraint<T, d>::Disable(const TArray<FConstraintContainerHandle*>& InConstraintHandles)
-{
-	for(int32 Idx = InConstraintHandles.Num()-1; Idx>=0; Idx--) 
-	{
-		check(InConstraintHandles[Idx] != nullptr);
-		Constraints[InConstraintHandles[Idx]->GetConstraintIndex()].bDisabled = true;
-	}
-}
 
 template<typename T, int d>
 void TPBDCollisionConstraint<T, d>::ApplyPushOut(const T Dt, FRigidBodyContactConstraint& Constraint, const TSet<TGeometryParticleHandle<T,d>*>& IsTemporarilyStatic, int32 Iteration, int32 NumIterations, bool &NeedsAnotherIteration)
@@ -814,8 +762,11 @@ void TPBDCollisionConstraint<T, d>::ApplyPushOut(const T Dt, FRigidBodyContactCo
 
 	for (int32 PairIteration = 0; PairIteration < MPairIterations; ++PairIteration)
 	{
-		UpdateConstraint<ECollisionUpdateType::Deepest>(MThickness, &Constraint, 1);
-		if (Constraint.Phi >= MThickness)
+		UpdateConstraint<ECollisionUpdateType::Deepest>(MThickness, Constraint);
+
+		const TContactData<T,d> & Contact = Constraint.ShapeManifold.Manifold;
+
+		if (Contact.Phi >= MThickness)
 		{
 			break;
 		}
@@ -828,8 +779,8 @@ void TPBDCollisionConstraint<T, d>::ApplyPushOut(const T Dt, FRigidBodyContactCo
 		NeedsAnotherIteration = true;
 		PMatrix<T, d, d> WorldSpaceInvI1 = PBDRigid0? (Q0 * FMatrix::Identity).GetTransposed() * PBDRigid0->InvI() * (Q0 * FMatrix::Identity) : PMatrix<T, d, d>(0);
 		PMatrix<T, d, d> WorldSpaceInvI2 = PBDRigid1 ? (Q1 * FMatrix::Identity).GetTransposed() * PBDRigid1->InvI() * (Q1 * FMatrix::Identity) : PMatrix<T, d, d>(0);
-		TVector<T, d> VectorToPoint1 = Constraint.Location - P0;
-		TVector<T, d> VectorToPoint2 = Constraint.Location - P1;
+		TVector<T, d> VectorToPoint1 = Contact.Location - P0;
+		TVector<T, d> VectorToPoint2 = Contact.Location - P1;
 		PMatrix<T, d, d> Factor =
 			(PBDRigid0 ? ComputeFactorMatrix3(VectorToPoint1, WorldSpaceInvI1, PBDRigid0->InvM()) : PMatrix<T, d, d>(0)) +
 			(PBDRigid1 ? ComputeFactorMatrix3(VectorToPoint2, WorldSpaceInvI2, PBDRigid1->InvM()) : PMatrix<T, d, d>(0));
@@ -840,16 +791,16 @@ void TPBDCollisionConstraint<T, d>::ApplyPushOut(const T Dt, FRigidBodyContactCo
 		TVector<T, d> Body1Velocity = V0 + TVector<T, d>::CrossProduct(W0, VectorToPoint1);
 		TVector<T, d> Body2Velocity = V1 + TVector<T, d>::CrossProduct(W1, VectorToPoint2);
 		TVector<T, d> RelativeVelocity = Body1Velocity - Body2Velocity;
-		const T RelativeVelocityDotNormal = TVector<T, d>::DotProduct(RelativeVelocity, Constraint.Normal);
+		const T RelativeVelocityDotNormal = TVector<T, d>::DotProduct(RelativeVelocity, Contact.Normal);
 		if (RelativeVelocityDotNormal < 0)
 		{
-			T ImpulseDenominator = TVector<T, d>::DotProduct(Constraint.Normal, Factor * Constraint.Normal);
-			TVector<T, d> ImpulseNumerator = -TVector<T, d>::DotProduct(RelativeVelocity, Constraint.Normal) * Constraint.Normal * ScalingFactor;
-			if (!ensureMsgf(FMath::Abs(ImpulseDenominator) > SMALL_NUMBER, TEXT("ApplyPushout Constraint:%s\n\nParticle:%s\n\nLevelset:%s\n\nFactor*Constraint.Normal:%s, ImpulseDenominator:%f"),
+			T ImpulseDenominator = TVector<T, d>::DotProduct(Contact.Normal, Factor * Contact.Normal);
+			TVector<T, d> ImpulseNumerator = -TVector<T, d>::DotProduct(RelativeVelocity, Contact.Normal) * Contact.Normal * ScalingFactor;
+			if (!ensureMsgf(FMath::Abs(ImpulseDenominator) > SMALL_NUMBER, TEXT("ApplyPushout Constraint:%s\n\nParticle:%s\n\nLevelset:%s\n\nFactor*Contact.Normal:%s, ImpulseDenominator:%f"),
 				*Constraint.ToString(),
 				*Particle0->ToString(),
 				*Particle1->ToString(),
-				*(Factor*Constraint.Normal).ToString(), ImpulseDenominator))
+				*(Factor*Contact.Normal).ToString(), ImpulseDenominator))
 			{
 				ImpulseDenominator = (T)1;
 			}
@@ -875,7 +826,7 @@ void TPBDCollisionConstraint<T, d>::ApplyPushOut(const T Dt, FRigidBodyContactCo
 		}
 
 
-		TVector<T, d> Impulse = PMatrix<T, d, d>(Factor.Inverse()) * ((-Constraint.Phi + MThickness) * ScalingFactor * Constraint.Normal);
+		TVector<T, d> Impulse = PMatrix<T, d, d>(Factor.Inverse()) * ((-Contact.Phi + MThickness) * ScalingFactor * Contact.Normal);
 		TVector<T, d> AngularImpulse1 = TVector<T, d>::CrossProduct(VectorToPoint1, Impulse);
 		TVector<T, d> AngularImpulse2 = TVector<T, d>::CrossProduct(VectorToPoint2, -Impulse);
 		if (!IsTemporarilyStatic0 && PBDRigid0)
@@ -943,10 +894,10 @@ void UpdateLevelsetConstraintHelperCCD(const TRigidParticles<T, d>& InParticles,
 		if (PointPair.Second)
 		{
 			const TVector<T, d> WorldSpaceDelta = WorldSpacePointEnd - TVector<T, d>(LocalToWorld2.TransformPosition(PointPair.First));
-			Constraint.Phi = -WorldSpaceDelta.Size();
-			Constraint.Normal = LocalToWorld2.TransformVector(InParticles.Geometry(Constraint.LevelsetIndex)->Normal(PointPair.First));
+			Contact.Phi = -WorldSpaceDelta.Size();
+			Contact.Normal = LocalToWorld2.TransformVector(InParticles.Geometry(Constraint.LevelsetIndex)->Normal(PointPair.First));
 			// @todo(mlentine): Should we be using the actual collision point or that point evolved to the current time step?
-			Constraint.Location = WorldSpacePointEnd;
+			Contact.Location = WorldSpacePointEnd;
 		}
 	}
 }
@@ -958,11 +909,13 @@ bool SampleObjectHelper2(const FImplicitObject& Object, const TRigidTransform<T,
 	TVector<T, d> LocalPoint = SampleToObjectTransform.TransformPositionNoScale(SampleParticle);
 	TVector<T, d> LocalNormal;
 	T LocalPhi = Object.PhiWithNormal(LocalPoint, LocalNormal);
-	if (LocalPhi < Constraint.Phi)
+
+	TContactData<T,d> & Contact = Constraint.ShapeManifold.Manifold;
+	if (LocalPhi < Contact.Phi)
 	{
-		Constraint.Phi = LocalPhi;
-		Constraint.Normal = ObjectTransform.TransformVectorNoScale(LocalNormal);
-		Constraint.Location = ObjectTransform.TransformPositionNoScale(LocalPoint);
+		Contact.Phi = LocalPhi;
+		Contact.Normal = ObjectTransform.TransformVectorNoScale(LocalNormal);
+		Contact.Location = ObjectTransform.TransformPositionNoScale(LocalPoint);
 		return true;
 	}
 	return false;
@@ -974,9 +927,11 @@ bool SampleObjectNoNormal2(const FImplicitObject& Object, const TRigidTransform<
 	TVector<T, d> LocalPoint = SampleToObjectTransform.TransformPositionNoScale(SampleParticle);
 	TVector<T, d> LocalNormal;
 	T LocalPhi = Object.PhiWithNormal(LocalPoint, LocalNormal);
-	if (LocalPhi < Constraint.Phi)
+
+	TContactData<T,d> & Contact = Constraint.ShapeManifold.Manifold;
+	if (LocalPhi < Contact.Phi)
 	{
-		Constraint.Phi = LocalPhi;
+		Contact.Phi = LocalPhi;
 		return true;
 	}
 	return false;
@@ -989,9 +944,11 @@ bool SampleObjectNormalAverageHelper2(const FImplicitObject& Object, const TRigi
 	TVector<T, d> LocalNormal;
 	T LocalPhi = Object.PhiWithNormal(LocalPoint, LocalNormal);
 	T LocalThickness = LocalPhi - Thickness;
+
+	TContactData<T,d> & Contact = Constraint.ShapeManifold.Manifold;
 	if (LocalThickness < -KINDA_SMALL_NUMBER)
 	{
-		Constraint.Location += LocalPoint * LocalThickness;
+		Contact.Location += LocalPoint * LocalThickness;
 		TotalThickness += LocalThickness;
 		return true;
 	}
@@ -1011,17 +968,22 @@ FAutoConsoleVariableRef CVarNormalAveraging2(TEXT("p.NormalAveraging2"), NormalA
 int32 SampleMinParticlesForAcceleration2 = 2048;
 FAutoConsoleVariableRef CVarSampleMinParticlesForAcceleration2(TEXT("p.SampleMinParticlesForAcceleration2"), SampleMinParticlesForAcceleration2, TEXT("The minimum number of particles needed before using an acceleration structure when sampling"));
 
+
 template <ECollisionUpdateType UpdateType, typename T, int d>
 void SampleObject2(const FImplicitObject& Object, const TRigidTransform<T, d>& ObjectTransform, const TBVHParticles<T, d>& SampleParticles, const TRigidTransform<T, d>& SampleParticlesTransform, T Thickness, TRigidBodyContactConstraint<T, d>& Constraint)
 {
 	SCOPE_CYCLE_COUNTER(STAT_SampleObject);
-	TRigidBodyContactConstraint<T, d> AvgConstraint;
+	TRigidBodyContactConstraint<float, 3> AvgConstraint = Constraint;
+
+	TContactData<float, 3> & Contact = Constraint.ShapeManifold.Manifold;
+	TContactData<float, 3> & AvgContact = AvgConstraint.ShapeManifold.Manifold;
+
 	AvgConstraint.Particle = Constraint.Particle;
 	AvgConstraint.Levelset = Constraint.Levelset;
-	AvgConstraint.Location = TVector<T, d>::ZeroVector;
-	AvgConstraint.Normal = TVector<T, d>::ZeroVector;
-	AvgConstraint.Phi = Thickness;
-	T TotalThickness = T(0);
+	AvgContact.Location = TVector<float, 3>::ZeroVector;
+	AvgContact.Normal = TVector<float, 3>::ZeroVector;
+	AvgContact.Phi = Thickness;
+	float TotalThickness = float(0);
 
 	int32 DeepestParticle = -1;
 	const int32 NumParticles = SampleParticles.Size();
@@ -1052,7 +1014,7 @@ void SampleObject2(const FImplicitObject& Object, const TRigidTransform<T, d>& O
 						DeepestParticle = i;
 						if (UpdateType == ECollisionUpdateType::Any)
 						{
-							Constraint.Phi = AvgConstraint.Phi;
+							Contact.Phi = AvgContact.Phi;
 							return;
 						}
 					}
@@ -1076,7 +1038,7 @@ void SampleObject2(const FImplicitObject& Object, const TRigidTransform<T, d>& O
 					DeepestParticle = i;
 					if (UpdateType == ECollisionUpdateType::Any)
 					{
-						Constraint.Phi = AvgConstraint.Phi;
+						Contact.Phi = AvgContact.Phi;
 						return;
 					}
 				}
@@ -1088,44 +1050,49 @@ void SampleObject2(const FImplicitObject& Object, const TRigidTransform<T, d>& O
 	{
 		if (TotalThickness < -KINDA_SMALL_NUMBER)
 		{
-			TVector<T, d> LocalPoint = AvgConstraint.Location / TotalThickness;
+			TVector<T, d> LocalPoint = AvgContact.Location / TotalThickness;
 			TVector<T, d> LocalNormal;
 			const T NewPhi = Object.PhiWithNormal(LocalPoint, LocalNormal);
-			if (NewPhi < Constraint.Phi)
+			if (NewPhi < Contact.Phi)
 			{
-				Constraint.Phi = NewPhi;
-				Constraint.Location = ObjectTransform.TransformPositionNoScale(LocalPoint);
-				Constraint.Normal = ObjectTransform.TransformVectorNoScale(LocalNormal);
+				Contact.Phi = NewPhi;
+				Contact.Location = ObjectTransform.TransformPositionNoScale(LocalPoint);
+				Contact.Normal = ObjectTransform.TransformVectorNoScale(LocalNormal);
 			}
 		}
 		else
 		{
-			check(AvgConstraint.Phi >= Thickness);
+			check(AvgContact.Phi >= Thickness);
 		}
 	}
-	else if(AvgConstraint.Phi < Constraint.Phi)
+	else if(AvgContact.Phi < Contact.Phi)
 	{
 		check(DeepestParticle >= 0);
 		TVector<T,d> LocalPoint = SampleToObjectTM.TransformPositionNoScale(SampleParticles.X(DeepestParticle));
 		TVector<T, d> LocalNormal;
-		Constraint.Phi = Object.PhiWithNormal(LocalPoint, LocalNormal);
-		Constraint.Location = ObjectTransform.TransformPositionNoScale(LocalPoint);
-		Constraint.Normal = ObjectTransform.TransformVectorNoScale(LocalNormal);
+		Contact.Phi = Object.PhiWithNormal(LocalPoint, LocalNormal);
+		Contact.Location = ObjectTransform.TransformPositionNoScale(LocalPoint);
+		Contact.Normal = ObjectTransform.TransformVectorNoScale(LocalNormal);
 			
 	}
 }
 
+
 #if INTEL_ISPC
 template<ECollisionUpdateType UpdateType>
 void SampleObject2(const FImplicitObject& Object, const TRigidTransform<float, 3>& ObjectTransform, const TBVHParticles<float, 3>& SampleParticles, const TRigidTransform<float, 3>& SampleParticlesTransform, float Thickness, TRigidBodyContactConstraint<float, 3>& Constraint)
 {
 	SCOPE_CYCLE_COUNTER(STAT_SampleObject);
-	TRigidBodyContactConstraint<float, 3> AvgConstraint;
+	TRigidBodyContactConstraint<float, 3> AvgConstraint = Constraint;
+
+	TContactData<float, 3> & Contact = Constraint.ShapeManifold.Manifold;
+	TContactData<float, 3> & AvgContact = AvgConstraint.ShapeManifold.Manifold;
+
 	AvgConstraint.Particle = Constraint.Particle;
 	AvgConstraint.Levelset = Constraint.Levelset;
-	AvgConstraint.Location = TVector<float, 3>::ZeroVector;
-	AvgConstraint.Normal = TVector<float, 3>::ZeroVector;
-	AvgConstraint.Phi = Thickness;
+	AvgContact.Location = TVector<float, 3>::ZeroVector;
+	AvgContact.Normal = TVector<float, 3>::ZeroVector;
+	AvgContact.Phi = Thickness;
 	float TotalThickness = float(0);
 
 	int32 DeepestParticle = -1;
@@ -1165,7 +1132,7 @@ void SampleObject2(const FImplicitObject& Object, const TRigidTransform<float, 3
 						&PotentialParticles[0],
 						Thickness,
 						TotalThickness,
-						(ispc::FVector&)AvgConstraint.Location,
+						(ispc::FVector&)AvgContact.Location,
 						PotentialParticles.Num());
 				}
 				else
@@ -1180,12 +1147,12 @@ void SampleObject2(const FImplicitObject& Object, const TRigidTransform<float, 3
 						(ispc::FVector*)&SampleParticles.XArray()[0],
 						&PotentialParticles[0],
 						DeepestParticle,
-						AvgConstraint.Phi,
+						AvgContact.Phi,
 						PotentialParticles.Num());
 
 					if (UpdateType == ECollisionUpdateType::Any)
 					{
-						Constraint.Phi = AvgConstraint.Phi;
+						Contact.Phi = AvgContact.Phi;
 						return;
 					}
 				}
@@ -1205,7 +1172,7 @@ void SampleObject2(const FImplicitObject& Object, const TRigidTransform<float, 3
 						&PotentialParticles[0],
 						Thickness,
 						TotalThickness,
-						(ispc::FVector&)AvgConstraint.Location,
+						(ispc::FVector&)AvgContact.Location,
 						PotentialParticles.Num());
 				}
 				else
@@ -1217,12 +1184,12 @@ void SampleObject2(const FImplicitObject& Object, const TRigidTransform<float, 3
 						(ispc::FVector*)&SampleParticles.XArray()[0],
 						&PotentialParticles[0],
 						DeepestParticle,
-						AvgConstraint.Phi,
+						AvgContact.Phi,
 						PotentialParticles.Num());
 
 					if (UpdateType == ECollisionUpdateType::Any)
 					{
-						Constraint.Phi = AvgConstraint.Phi;
+						Contact.Phi = AvgContact.Phi;
 						return;
 					}
 				}
@@ -1243,7 +1210,7 @@ void SampleObject2(const FImplicitObject& Object, const TRigidTransform<float, 3
 							DeepestParticle = i;
 							if (UpdateType == ECollisionUpdateType::Any)
 							{
-								Constraint.Phi = AvgConstraint.Phi;
+								Contact.Phi = AvgContact.Phi;
 								return;
 							}
 						}
@@ -1272,7 +1239,7 @@ void SampleObject2(const FImplicitObject& Object, const TRigidTransform<float, 3
 					(ispc::FVector*)&SampleParticles.XArray()[0],
 					Thickness,
 					TotalThickness,
-					(ispc::FVector&)AvgConstraint.Location,
+					(ispc::FVector&)AvgContact.Location,
 					NumParticles);
 			}
 			else
@@ -1286,12 +1253,12 @@ void SampleObject2(const FImplicitObject& Object, const TRigidTransform<float, 3
 					(ispc::FTransform&)SampleToObjectTM,
 					(ispc::FVector*)&SampleParticles.XArray()[0],
 					DeepestParticle,
-					AvgConstraint.Phi,
+					AvgContact.Phi,
 					NumParticles);
 
 				if (UpdateType == ECollisionUpdateType::Any)
 				{
-					Constraint.Phi = AvgConstraint.Phi;
+					Contact.Phi = AvgContact.Phi;
 					return;
 				}
 			}
@@ -1309,7 +1276,7 @@ void SampleObject2(const FImplicitObject& Object, const TRigidTransform<float, 3
 					(ispc::FVector*)&SampleParticles.XArray()[0],
 					Thickness,
 					TotalThickness,
-					(ispc::FVector&)AvgConstraint.Location,
+					(ispc::FVector&)AvgContact.Location,
 					NumParticles);
 			}
 			else
@@ -1320,12 +1287,12 @@ void SampleObject2(const FImplicitObject& Object, const TRigidTransform<float, 3
 					(ispc::FTransform&)SampleToObjectTM,
 					(ispc::FVector*)&SampleParticles.XArray()[0],
 					DeepestParticle,
-					AvgConstraint.Phi,
+					AvgContact.Phi,
 					NumParticles);
 
 				if (UpdateType == ECollisionUpdateType::Any)
 				{
-					Constraint.Phi = AvgConstraint.Phi;
+					Contact.Phi = AvgContact.Phi;
 					return;
 				}
 			}
@@ -1343,7 +1310,7 @@ void SampleObject2(const FImplicitObject& Object, const TRigidTransform<float, 3
 					(ispc::FVector*)&SampleParticles.XArray()[0],
 					Thickness,
 					TotalThickness,
-					(ispc::FVector&)AvgConstraint.Location,
+					(ispc::FVector&)AvgContact.Location,
 					NumParticles);
 			}
 			else
@@ -1354,12 +1321,12 @@ void SampleObject2(const FImplicitObject& Object, const TRigidTransform<float, 3
 					(ispc::FTransform&)SampleToObjectTM,
 					(ispc::FVector*)&SampleParticles.XArray()[0],
 					DeepestParticle,
-					AvgConstraint.Phi,
+					AvgContact.Phi,
 					NumParticles);
 
 				if (UpdateType == ECollisionUpdateType::Any)
 				{
-					Constraint.Phi = AvgConstraint.Phi;
+					Contact.Phi = AvgContact.Phi;
 					return;
 				}
 			}
@@ -1380,7 +1347,7 @@ void SampleObject2(const FImplicitObject& Object, const TRigidTransform<float, 3
 						DeepestParticle = i;
 						if (UpdateType == ECollisionUpdateType::Any)
 						{
-							Constraint.Phi = AvgConstraint.Phi;
+							Contact.Phi = AvgContact.Phi;
 							return;
 						}
 					}
@@ -1393,29 +1360,29 @@ void SampleObject2(const FImplicitObject& Object, const TRigidTransform<float, 3
 	{
 		if (TotalThickness < -KINDA_SMALL_NUMBER)
 		{
-			TVector<float, 3> LocalPoint = AvgConstraint.Location / TotalThickness;
+			TVector<float, 3> LocalPoint = AvgContact.Location / TotalThickness;
 			TVector<float, 3> LocalNormal;
 			const float NewPhi = Object.PhiWithNormal(LocalPoint, LocalNormal);
-			if (NewPhi < Constraint.Phi)
+			if (NewPhi < Contact.Phi)
 			{
-				Constraint.Phi = NewPhi;
-				Constraint.Location = ObjectTransform.TransformPositionNoScale(LocalPoint);
-				Constraint.Normal = ObjectTransform.TransformVectorNoScale(LocalNormal);
+				Contact.Phi = NewPhi;
+				Contact.Location = ObjectTransform.TransformPositionNoScale(LocalPoint);
+				Contact.Normal = ObjectTransform.TransformVectorNoScale(LocalNormal);
 			}
 		}
 		else
 		{
-			check(AvgConstraint.Phi >= Thickness);
+			check(AvgContact.Phi >= Thickness);
 		}
 	}
-	else if (AvgConstraint.Phi < Constraint.Phi)
+	else if (AvgContact.Phi < Contact.Phi)
 	{
 		check(DeepestParticle >= 0);
 		TVector<float, 3> LocalPoint = SampleToObjectTM.TransformPositionNoScale(SampleParticles.X(DeepestParticle));
 		TVector<float, 3> LocalNormal;
-		Constraint.Phi = Object.PhiWithNormal(LocalPoint, LocalNormal);
-		Constraint.Location = ObjectTransform.TransformPositionNoScale(LocalPoint);
-		Constraint.Normal = ObjectTransform.TransformVectorNoScale(LocalNormal);
+		Contact.Phi = Object.PhiWithNormal(LocalPoint, LocalNormal);
+		Contact.Location = ObjectTransform.TransformPositionNoScale(LocalPoint);
+		Contact.Normal = ObjectTransform.TransformVectorNoScale(LocalNormal);
 	}
 }
 #endif
@@ -1423,7 +1390,8 @@ void SampleObject2(const FImplicitObject& Object, const TRigidTransform<float, 3
 template <typename T, int d>
 bool UpdateBoxPlaneConstraint(const TBox<T, d>& Box, const TRigidTransform<T, d>& BoxTransform, const TPlane<T, d>& Plane, const TRigidTransform<T, d>& PlaneTransform, T Thickness, TRigidBodyContactConstraint<T, d>& Constraint)
 {
-	
+	TContactData<T,d> & Contact = Constraint.ShapeManifold.Manifold;
+
 #if USING_CODE_ANALYSIS
 	MSVC_PRAGMA( warning( push ) )
 	MSVC_PRAGMA( warning( disable : ALL_CODE_ANALYSIS_WARNINGS ) )
@@ -1455,16 +1423,16 @@ bool UpdateBoxPlaneConstraint(const TBox<T, d>& Box, const TRigidTransform<T, d>
 	{
 		TVector<T, d> Normal;
 		const T NewPhi = Plane.PhiWithNormal(Corners[i], Normal);
-		if (NewPhi < Constraint.Phi + Epsilon)
+		if (NewPhi < Contact.Phi + Epsilon)
 		{
-			if (NewPhi <= Constraint.Phi - Epsilon)
+			if (NewPhi <= Contact.Phi - Epsilon)
 			{
 				NumConstraints = 0;
 			}
-			Constraint.Phi = NewPhi;
-			Constraint.Normal = PlaneTransform.TransformVector(Normal);
-			Constraint.Location = PlaneTransform.TransformPosition(Corners[i]);
-			PotentialConstraints[NumConstraints++] = Constraint.Location;
+			Contact.Phi = NewPhi;
+			Contact.Normal = PlaneTransform.TransformVector(Normal);
+			Contact.Location = PlaneTransform.TransformPosition(Corners[i]);
+			PotentialConstraints[NumConstraints++] = Contact.Location;
 			bApplied = true;
 		}
 	}
@@ -1475,7 +1443,7 @@ bool UpdateBoxPlaneConstraint(const TBox<T, d>& Box, const TRigidTransform<T, d>
 		{
 			AverageLocation += PotentialConstraints[ConstraintIdx];
 		}
-		Constraint.Location = AverageLocation / NumConstraints;
+		Contact.Location = AverageLocation / NumConstraints;
 	}
 
 	return bApplied;
@@ -1484,22 +1452,26 @@ bool UpdateBoxPlaneConstraint(const TBox<T, d>& Box, const TRigidTransform<T, d>
 template <typename T, int d>
 void UpdateSphereConstraint(const TSphere<T, d>& Sphere1, const TRigidTransform<T, d>& Sphere1Transform, const TSphere<T, d>& Sphere2, const TRigidTransform<T, d>& Sphere2Transform, T Thickness, TRigidBodyContactConstraint<T, d>& Constraint)
 {
+	TContactData<T,d> & Contact = Constraint.ShapeManifold.Manifold;
+
 	const TVector<T, d> Center1 = Sphere1Transform.TransformPosition(Sphere1.GetCenter());
 	const TVector<T, d> Center2 = Sphere2Transform.TransformPosition(Sphere2.GetCenter());
 	const TVector<T, d> Direction = Center1 - Center2;
 	const T Size = Direction.Size();
 	const T NewPhi = Size - (Sphere1.GetRadius() + Sphere2.GetRadius());
-	if (NewPhi < Constraint.Phi)
+	if (NewPhi < Contact.Phi)
 	{
-		Constraint.Normal = Size > SMALL_NUMBER ? Direction / Size : TVector<T, d>(0, 0, 1);
-		Constraint.Phi = NewPhi;
-		Constraint.Location = Center1 - Sphere1.GetRadius() * Constraint.Normal;
+		Contact.Normal = Size > SMALL_NUMBER ? Direction / Size : TVector<T, d>(0, 0, 1);
+		Contact.Phi = NewPhi;
+		Contact.Location = Center1 - Sphere1.GetRadius() * Contact.Normal;
 	}
 }
 
 template <typename T, int d>
 void UpdateSpherePlaneConstraint(const TSphere<T, d>& Sphere, const TRigidTransform<T, d>& SphereTransform, const TPlane<T, d>& Plane, const TRigidTransform<T, d>& PlaneTransform, T Thickness, TRigidBodyContactConstraint<T, d>& Constraint)
 {
+	TContactData<T,d> & Contact = Constraint.ShapeManifold.Manifold;
+
 	const TRigidTransform<T, d> SphereToPlaneTransform(PlaneTransform.Inverse() * SphereTransform);
 	const TVector<T, d> SphereCenter = SphereToPlaneTransform.TransformPosition(Sphere.GetCenter());
 
@@ -1507,17 +1479,19 @@ void UpdateSpherePlaneConstraint(const TSphere<T, d>& Sphere, const TRigidTransf
 	T NewPhi = Plane.PhiWithNormal(SphereCenter, NewNormal);
 	NewPhi -= Sphere.GetRadius();
 
-	if (NewPhi < Constraint.Phi)
+	if (NewPhi < Contact.Phi)
 	{
-		Constraint.Phi = NewPhi;
-		Constraint.Normal = PlaneTransform.TransformVectorNoScale(NewNormal);
-		Constraint.Location = SphereCenter - Constraint.Normal * Sphere.GetRadius();
+		Contact.Phi = NewPhi;
+		Contact.Normal = PlaneTransform.TransformVectorNoScale(NewNormal);
+		Contact.Location = SphereCenter - Contact.Normal * Sphere.GetRadius();
 	}
 }
 
 template <typename T, int d>
 void UpdateSphereBoxConstraint(const TSphere<T, d>& Sphere, const TRigidTransform<T, d>& SphereTransform, const TBox<T, d>& Box, const TRigidTransform<T, d>& BoxTransform, T Thickness, TRigidBodyContactConstraint<T, d>& Constraint)
 {
+	TContactData<T,d> & Contact = Constraint.ShapeManifold.Manifold;
+
 	const TRigidTransform<T, d> SphereToBoxTransform(SphereTransform * BoxTransform.Inverse());
 	const TVector<T, d> SphereCenterInBox = SphereToBoxTransform.TransformPosition(Sphere.GetCenter());
 
@@ -1525,17 +1499,19 @@ void UpdateSphereBoxConstraint(const TSphere<T, d>& Sphere, const TRigidTransfor
 	T NewPhi = Box.PhiWithNormal(SphereCenterInBox, NewNormal);
 	NewPhi -= Sphere.GetRadius();
 
-	if (NewPhi < Constraint.Phi)
+	if (NewPhi < Contact.Phi)
 	{
-		Constraint.Phi = NewPhi;
-		Constraint.Normal = BoxTransform.TransformVectorNoScale(NewNormal);
-		Constraint.Location = SphereTransform.TransformPosition(Sphere.GetCenter()) - Constraint.Normal * Sphere.GetRadius();
+		Contact.Phi = NewPhi;
+		Contact.Normal = BoxTransform.TransformVectorNoScale(NewNormal);
+		Contact.Location = SphereTransform.TransformPosition(Sphere.GetCenter()) - Contact.Normal * Sphere.GetRadius();
 	}
 }
 
 template <typename T, int d>
 void UpdateCapsuleCapsuleConstraint(const TCapsule<T>& A, const TRigidTransform<T, d>& ATransform, const TCapsule<T>& B, const TRigidTransform<T, d>& BTransform, T Thickness, TRigidBodyContactConstraint<T, d>& Constraint)
 {
+	TContactData<T,d> & Contact = Constraint.ShapeManifold.Manifold;
+
 	FVector A1 = ATransform.TransformPosition(A.GetX1());
 	FVector A2 = ATransform.TransformPosition(A.GetX2());
 	FVector B1 = BTransform.TransformPosition(B.GetX1());
@@ -1548,9 +1524,9 @@ void UpdateCapsuleCapsuleConstraint(const TCapsule<T>& A, const TRigidTransform<
 	if (DeltaLen > KINDA_SMALL_NUMBER)
 	{
 		TVector<T, d> Dir = Delta / DeltaLen;
-		Constraint.Phi = DeltaLen - (A.GetRadius() + B.GetRadius());
-		Constraint.Normal = -Dir;
-		Constraint.Location = P1 + Dir * A.GetRadius();
+		Contact.Phi = DeltaLen - (A.GetRadius() + B.GetRadius());
+		Contact.Normal = -Dir;
+		Contact.Location = P1 + Dir * A.GetRadius();
 	}
 }
 
@@ -1558,6 +1534,7 @@ template <typename T, int d>
 void UpdateCapsuleBoxConstraint(const TCapsule<T>& A, const TRigidTransform<T, d>& ATransform, const TBox<T, d>& B, const TRigidTransform<T, d>& BTransform, T Thickness, TRigidBodyContactConstraint<T, d>& Constraint)
 {
 	// @todo(ccaulfield): Add custom capsule-box collision
+	TContactData<T,d> & Contact = Constraint.ShapeManifold.Manifold;
 
 	TRigidTransform<T, d> BToATransform = BTransform.GetRelativeTransform(ATransform);
 
@@ -1569,9 +1546,9 @@ void UpdateCapsuleBoxConstraint(const TCapsule<T>& A, const TRigidTransform<T, d
 		TVector<T, d> NearPointAWorld = ATransform.TransformPosition(NearPointALocal);
 		TVector<T, d> NearPointBWorld = BTransform.TransformPosition(NearPointBLocal);
 		TVector<T, d> NearPointBtoAWorld = NearPointAWorld - NearPointBWorld;
-		Constraint.Phi = NearPointDistance;
-		Constraint.Normal = NearPointBtoAWorld.GetSafeNormal();
-		Constraint.Location = NearPointAWorld;
+		Contact.Phi = NearPointDistance;
+		Contact.Normal = NearPointBtoAWorld.GetSafeNormal();
+		Contact.Location = NearPointAWorld;
 	}
 	else
 	{
@@ -1583,7 +1560,7 @@ void UpdateCapsuleBoxConstraint(const TCapsule<T>& A, const TRigidTransform<T, d
 			if (SampleObjectHelper2(A, ATransform, BToATransform, BParticles[ParticleIndex], Thickness, Constraint))
 			{
 				// SampleObjectHelper2 expects A to be the box, so reverse the results
-				Constraint.Normal = -Constraint.Normal;
+				Contact.Normal = -Contact.Normal;
 			}
 		}
 	}
@@ -1624,13 +1601,12 @@ TArray<Pair<const FImplicitObject*, TRigidTransform<T, d>>> FindRelevantShapes2(
 
 DECLARE_CYCLE_STAT(TEXT("TPBDCollisionConstraint::UpdateUnionUnionConstraint"), STAT_UpdateUnionUnionConstraint, STATGROUP_ChaosWide);
 template<ECollisionUpdateType UpdateType, typename T, int d>
-void UpdateUnionUnionConstraint(const T Thickness, TRigidBodyContactConstraint<T, d>* Constraints, int32 NumConstraints, TCollisionResolutionManifold<T, d>* Manifold)
+void UpdateUnionUnionConstraint(const T Thickness, TRigidBodyContactConstraint<T, d>& Constraint)
 {
-	checkSlow(Constraints != nullptr && NumConstraints >= 1);
 	SCOPE_CYCLE_COUNTER(STAT_UpdateUnionUnionConstraint);
 
-	TGenericParticleHandle<T, d> Particle0 = Constraints[0].Particle;
-	TGenericParticleHandle<T, d> Particle1 = Constraints[0].Levelset;
+	TGenericParticleHandle<T, d> Particle0 = Constraint.Particle;
+	TGenericParticleHandle<T, d> Particle1 = Constraint.Levelset;
 
 	TRigidTransform<T, d> ParticlesTM = TRigidTransform<T, d>(Particle0->P(), Particle0->Q());
 	TRigidTransform<T, d> LevelsetTM = TRigidTransform<T, d>(Particle1->P(), Particle1->Q());
@@ -1652,20 +1628,19 @@ void UpdateUnionUnionConstraint(const T Thickness, TRigidBodyContactConstraint<T
 		{
 			const FImplicitObject& ParticleInnerObj = *ParticlePair.First;
 			const TRigidTransform<T, d> ParticleInnerObjTM = ParticlePair.Second * ParticlesTM;
-			UpdateConstraintImp2<UpdateType>(ParticleInnerObj, ParticleInnerObjTM, LevelsetInnerObj, LevelsetInnerObjTM, Thickness, Constraints, NumConstraints, Manifold);
+			UpdateConstraintImp<UpdateType>(ParticleInnerObj, ParticleInnerObjTM, LevelsetInnerObj, LevelsetInnerObjTM, Thickness, Constraint);
 		}
 	}
 }
 
 DECLARE_CYCLE_STAT(TEXT("TPBDCollisionConstraint::UpdateSingleUnionConstraint"), STAT_UpdateSingleUnionConstraint, STATGROUP_ChaosWide);
 template<ECollisionUpdateType UpdateType, typename T, int d>
-void UpdateSingleUnionConstraint(const T Thickness, TRigidBodyContactConstraint<T, d>* Constraints, int32 NumConstraints, TCollisionResolutionManifold<T, d>* Manifold)
+void UpdateSingleUnionConstraint(const T Thickness, TRigidBodyContactConstraint<T, d>& Constraint)
 {
-	checkSlow(Constraints != nullptr && NumConstraints >= 1);
 	SCOPE_CYCLE_COUNTER(STAT_UpdateSingleUnionConstraint);
 
-	TGenericParticleHandle<T, d> Particle0 = Constraints[0].Particle;
-	TGenericParticleHandle<T, d> Particle1 = Constraints[0].Levelset;
+	TGenericParticleHandle<T, d> Particle0 = Constraint.Particle;
+	TGenericParticleHandle<T, d> Particle1 = Constraint.Levelset;
 
 	TRigidTransform<T, d> ParticlesTM = TRigidTransform<T, d>(Particle0->P(), Particle0->Q());
 	TRigidTransform<T, d> LevelsetTM = TRigidTransform<T, d>(Particle1->P(), Particle1->Q());
@@ -1678,7 +1653,7 @@ void UpdateSingleUnionConstraint(const T Thickness, TRigidBodyContactConstraint<
 	{
 		const FImplicitObject& LevelsetInnerObj = *LevelsetObjPair.First;
 		const TRigidTransform<T, d> LevelsetInnerObjTM = LevelsetObjPair.Second * LevelsetTM;
-		UpdateConstraintImp2<UpdateType>(*ParticleObj, ParticlesTM, LevelsetInnerObj, LevelsetInnerObjTM, Thickness, Constraints, NumConstraints, Manifold);
+		UpdateConstraintImp<UpdateType>(*ParticleObj, ParticlesTM, LevelsetInnerObj, LevelsetInnerObjTM, Thickness, Constraint);
 	}
 }
 
@@ -1708,6 +1683,7 @@ void TPBDCollisionConstraint<T, d>::UpdateLevelsetConstraint(const T Thickness,
 
 	if(SampleParticles)
 	{
+		asdf
 		SampleObject2<UpdateType>(*Particle1->Geometry(), LevelsetTM, *SampleParticles, ParticlesTM, Thickness, Constraint);
 	}
 }
@@ -1748,7 +1724,7 @@ void UpdateUnionLevelsetConstraint(const T Thickness, TRigidBodyContactConstrain
 				const FImplicitObject* Object = LevelsetObjPair.First;
 				const TRigidTransform<T, d> ObjectTM = LevelsetObjPair.Second * LevelsetTM;
 				SampleObject2<UpdateType>(*Object, ObjectTM, SampleParticles, ParticlesTM, Thickness, Constraint);
-				if (UpdateType == ECollisionUpdateType::Any && Constraint.Phi < Thickness)
+				if (UpdateType == ECollisionUpdateType::Any && Constraint.GetPhi() < Thickness)
 				{
 					return;
 				}
@@ -1840,6 +1816,8 @@ void UpdateLevelsetUnionConstraint(const T Thickness, TRigidBodyContactConstrain
 template <typename T, int d>
 void UpdateBoxConstraint(const TBox<T, d>& Box1, const TRigidTransform<T, d>& Box1Transform, const TBox<T, d>& Box2, const TRigidTransform<T, d>& Box2Transform, T Thickness, TRigidBodyContactConstraint<T, d>& Constraint)
 {
+	TContactData<T,d> & Contact = Constraint.ShapeManifold.Manifold;
+
 	TBox<T,d> Box2SpaceBox1 = Box1.TransformedBox(Box1Transform * Box2Transform.Inverse());
 	TBox<T,d> Box1SpaceBox2 = Box2.TransformedBox(Box2Transform * Box1Transform.Inverse());
 	Box2SpaceBox1.Thicken(Thickness);
@@ -1858,21 +1836,21 @@ void UpdateBoxConstraint(const TBox<T, d>& Box1, const TRigidTransform<T, d>& Bo
 			if(Size < (Sphere1.GetRadius() + Sphere2.GetRadius()))
 			{
 				const T NewPhi = Size - (Sphere1.GetRadius() + Sphere2.GetRadius());;
-				if (NewPhi < Constraint.Phi)
+				if (NewPhi < Contact.Phi)
 				{
 					bDeepOverlap = true;
-					Constraint.Normal = Size > SMALL_NUMBER ? Direction / Size : TVector<T, d>(0, 0, 1);
-					Constraint.Phi = NewPhi;
-					Constraint.Location = Sphere1.GetCenter() - Sphere1.GetRadius() * Constraint.Normal;
+					Contact.Normal = Size > SMALL_NUMBER ? Direction / Size : TVector<T, d>(0, 0, 1);
+					Contact.Phi = NewPhi;
+					Contact.Location = Sphere1.GetCenter() - Sphere1.GetRadius() * Contact.Normal;
 				}
 			}
 		}
-		if (!bDeepOverlap || Constraint.Phi >= 0)
+		if (!bDeepOverlap || Contact.Phi >= 0)
 		{
 			//if we didn't have deep penetration use signed distance per particle. If we did have deep penetration but the spheres did not overlap use signed distance per particle
 
 			//UpdateLevelsetConstraintGJK(InParticles, Thickness, Constraint);
-			//check(Constraint.Phi < MThickness);
+			//check(Contact.Phi < MThickness);
 			// For now revert to doing all points vs lsv check until we can figure out a good way to get the deepest point without needing this
 			{
 				const TArray<TVector<T, d>> SampleParticles = Box1.ComputeSamplePoints();
@@ -1889,119 +1867,89 @@ void UpdateBoxConstraint(const TBox<T, d>& Box1, const TRigidTransform<T, d>& Bo
 
 
 template<typename T, int d>
-void ConstructLevelsetConstraints(TGeometryParticleHandle<T, d>* Particle0, TGeometryParticleHandle<T, d>* Particle1, const FImplicitObject* Implicit0, const FImplicitObject* Implicit1, const T Thickness, TArray<TRigidBodyContactConstraint<T, d>>& ConstraintBuffer, TCollisionResolutionManifold<T, d>* Manifold = nullptr)
+void ConstructLevelsetConstraints(TGeometryParticleHandle<T, d>* Particle0, TGeometryParticleHandle<T, d>* Particle1, const FImplicitObject* Implicit0, const FImplicitObject* Implicit1, const T Thickness, TRigidBodyContactConstraint<T, d>& Constraint)
 {
-	//todo(ocohen):if both have collision particles, use the one with fewer?
+	
 	if (!Particle1->Geometry() || (Particle0->AsDynamic() && !Particle0->AsDynamic()->CollisionParticlesSize() && Particle0->Geometry() && !Particle0->Geometry()->IsUnderlyingUnion()))
 	{
-		TRigidBodyContactConstraint<T, d> Constraint;
 		Constraint.Particle = Particle1;
 		Constraint.Levelset = Particle0;
-		Constraint.Geometry[0] = Implicit1;
-		Constraint.Geometry[1] = Implicit0;
-		ConstraintBuffer.Add(Constraint);
+		Constraint.AddManifold(Implicit1, Implicit0);
 	}
 	else
 	{
-		TRigidBodyContactConstraint<T, d> Constraint;
 		Constraint.Particle = Particle0;
 		Constraint.Levelset = Particle1;
-		Constraint.Geometry[0] = Implicit0;
-		Constraint.Geometry[1] = Implicit1;
-		ConstraintBuffer.Add(Constraint);
+		Constraint.AddManifold(Implicit0, Implicit1);
 	}
 }
 
 template<typename T, int d>
-void ConstructBoxConstraints(TGeometryParticleHandle<T, d>* Particle0, TGeometryParticleHandle<T, d>* Particle1, const FImplicitObject* Implicit0, const FImplicitObject* Implicit1, const T Thickness, TArray<TRigidBodyContactConstraint<T, d>>& ConstraintBuffer, TCollisionResolutionManifold<T, d>* Manifold = nullptr)
+void ConstructBoxConstraints(TGeometryParticleHandle<T, d>* Particle0, TGeometryParticleHandle<T, d>* Particle1, const FImplicitObject* Implicit0, const FImplicitObject* Implicit1, const T Thickness, TRigidBodyContactConstraint<T, d>& Constraint)
 {
-	TRigidBodyContactConstraint<T, d> Constraint;
 	Constraint.Particle = Particle0;
 	Constraint.Levelset = Particle1;
-	Constraint.Geometry[0] = Implicit0;
-	Constraint.Geometry[1] = Implicit1;
-	ConstraintBuffer.Add(Constraint);
+	Constraint.AddManifold(Implicit0, Implicit1);
 }
 
 template<typename T, int d>
-void ConstructBoxPlaneConstraints(TGeometryParticleHandle<T, d>* Particle0, TGeometryParticleHandle<T, d>* Particle1, const FImplicitObject* Implicit0, const FImplicitObject* Implicit1, const T Thickness, TArray<TRigidBodyContactConstraint<T, d>>& ConstraintBuffer, TCollisionResolutionManifold<T, d>* Manifold = nullptr)
+void ConstructBoxPlaneConstraints(TGeometryParticleHandle<T, d>* Particle0, TGeometryParticleHandle<T, d>* Particle1, const FImplicitObject* Implicit0, const FImplicitObject* Implicit1, const T Thickness, TRigidBodyContactConstraint<T, d>& Constraint)
 {
-	TRigidBodyContactConstraint<T, d> Constraint;
 	Constraint.Particle = Particle0;
 	Constraint.Levelset = Particle1;
-	Constraint.Geometry[0] = Implicit0;
-	Constraint.Geometry[1] = Implicit1;
-	ConstraintBuffer.Add(Constraint);
+	Constraint.AddManifold(Implicit0, Implicit1);
 }
 
 template<typename T, int d>
-void ConstructSphereConstraints(TGeometryParticleHandle<T, d>* Particle0, TGeometryParticleHandle<T, d>* Particle1, const FImplicitObject* Implicit0, const FImplicitObject* Implicit1, const T Thickness, TArray<TRigidBodyContactConstraint<T, d>>& ConstraintBuffer, TCollisionResolutionManifold<T, d>* Manifold = nullptr)
+void ConstructSphereConstraints(TGeometryParticleHandle<T, d>* Particle0, TGeometryParticleHandle<T, d>* Particle1, const FImplicitObject* Implicit0, const FImplicitObject* Implicit1, const T Thickness, TRigidBodyContactConstraint<T, d>& Constraint)
 {
-	TRigidBodyContactConstraint<T, d> Constraint;
 	Constraint.Particle = Particle0;
 	Constraint.Levelset = Particle1;
-	Constraint.Geometry[0] = Implicit0;
-	Constraint.Geometry[1] = Implicit1;
-	ConstraintBuffer.Add(Constraint);
+	Constraint.AddManifold(Implicit0, Implicit1);
 }
 
 template<typename T, int d>
-void ConstructSpherePlaneConstraints(TGeometryParticleHandle<T, d>* Particle0, TGeometryParticleHandle<T, d>* Particle1, const FImplicitObject* Implicit0, const FImplicitObject* Implicit1, const T Thickness, TArray<TRigidBodyContactConstraint<T, d>>& ConstraintBuffer, TCollisionResolutionManifold<T, d>* Manifold = nullptr)
+void ConstructSpherePlaneConstraints(TGeometryParticleHandle<T, d>* Particle0, TGeometryParticleHandle<T, d>* Particle1, const FImplicitObject* Implicit0, const FImplicitObject* Implicit1, const T Thickness, TRigidBodyContactConstraint<T, d>& Constraint)
 {
-	TRigidBodyContactConstraint<T, d> Constraint;
 	Constraint.Particle = Particle0;
 	Constraint.Levelset = Particle1;
-	Constraint.Geometry[0] = Implicit0;
-	Constraint.Geometry[1] = Implicit1;
-	ConstraintBuffer.Add(Constraint);
+	Constraint.AddManifold(Implicit0, Implicit1);
 }
 
 template<typename T, int d>
-void ConstructSphereBoxConstraints(TGeometryParticleHandle<T, d>* Particle0, TGeometryParticleHandle<T, d>* Particle1, const FImplicitObject* Implicit0, const FImplicitObject* Implicit1, const T Thickness, TArray<TRigidBodyContactConstraint<T, d>>& ConstraintBuffer, TCollisionResolutionManifold<T, d>* Manifold = nullptr)
+void ConstructSphereBoxConstraints(TGeometryParticleHandle<T, d>* Particle0, TGeometryParticleHandle<T, d>* Particle1, const FImplicitObject* Implicit0, const FImplicitObject* Implicit1, const T Thickness, TRigidBodyContactConstraint<T, d>& Constraint)
 {
-	TRigidBodyContactConstraint<T, d> Constraint;
 	Constraint.Particle = Particle0;
 	Constraint.Levelset = Particle1;
-	Constraint.Geometry[0] = Implicit0;
-	Constraint.Geometry[1] = Implicit1;
-	ConstraintBuffer.Add(Constraint);
+	Constraint.AddManifold(Implicit0, Implicit1);
 }
 
 template<typename T, int d>
-void ConstructCapsuleCapsuleConstraints(TGeometryParticleHandle<T, d>* Particle0, TGeometryParticleHandle<T, d>* Particle1, const FImplicitObject* Implicit0, const FImplicitObject* Implicit1, const T Thickness, TArray<TRigidBodyContactConstraint<T, d>>& ConstraintBuffer, TCollisionResolutionManifold<T, d>* Manifold = nullptr)
+void ConstructCapsuleCapsuleConstraints(TGeometryParticleHandle<T, d>* Particle0, TGeometryParticleHandle<T, d>* Particle1, const FImplicitObject* Implicit0, const FImplicitObject* Implicit1, const T Thickness, TRigidBodyContactConstraint<T, d>& Constraint)
 {
-	TRigidBodyContactConstraint<T, d> Constraint;
 	Constraint.Particle = Particle0;
 	Constraint.Levelset = Particle1;
-	Constraint.Geometry[0] = Implicit0;
-	Constraint.Geometry[1] = Implicit1;
-	ConstraintBuffer.Add(Constraint);
+	Constraint.AddManifold(Implicit0, Implicit1);
 }
 
 template<typename T, int d>
-void ConstructCapsuleBoxConstraints(TGeometryParticleHandle<T, d>* Particle0, TGeometryParticleHandle<T, d>* Particle1, const FImplicitObject* Implicit0, const FImplicitObject* Implicit1, const T Thickness, TArray<TRigidBodyContactConstraint<T, d>>& ConstraintBuffer, TCollisionResolutionManifold<T, d>* Manifold = nullptr)
+void ConstructCapsuleBoxConstraints(TGeometryParticleHandle<T, d>* Particle0, TGeometryParticleHandle<T, d>* Particle1, const FImplicitObject* Implicit0, const FImplicitObject* Implicit1, const T Thickness, TRigidBodyContactConstraint<T, d>& Constraint)
 {
-	TRigidBodyContactConstraint<T, d> Constraint;
 	Constraint.Particle = Particle0;
 	Constraint.Levelset = Particle1;
-	Constraint.Geometry[0] = Implicit0;
-	Constraint.Geometry[1] = Implicit1;
-	ConstraintBuffer.Add(Constraint);
+	Constraint.AddManifold(Implicit0, Implicit1);
 }
 
 template<typename T, int d>
-void ConstructSingleUnionConstraints(TGeometryParticleHandle<T, d>* Particle0, TGeometryParticleHandle<T, d>* Particle1, const FImplicitObject* Implicit0, const FImplicitObject* Implicit1, const T Thickness, TArray<TRigidBodyContactConstraint<T, d>>& ConstraintBuffer, TCollisionResolutionManifold<T, d>* Manifold = nullptr)
+void ConstructSingleUnionConstraints(TGeometryParticleHandle<T, d>* Particle0, TGeometryParticleHandle<T, d>* Particle1, const FImplicitObject* Implicit0, const FImplicitObject* Implicit1, const T Thickness, TRigidBodyContactConstraint<T, d>& Constraint)
 {
-	TRigidBodyContactConstraint<T, d> Constraint;
 	Constraint.Particle = Particle0;
 	Constraint.Levelset = Particle1;
-	Constraint.Geometry[0] = Implicit0;
-	Constraint.Geometry[1] = Implicit1;
-	ConstraintBuffer.Add(Constraint);
+	Constraint.AddManifold(Implicit0, Implicit1);
 }
 
 template<typename T, int d>
-void ConstructUnionUnionConstraints(TGeometryParticleHandle<T, d>* Particle0, TGeometryParticleHandle<T, d>* Particle1, const FImplicitObject* Implicit0, const FImplicitObject* Implicit1, const T Thickness, TArray<TRigidBodyContactConstraint<T, d>>& ConstraintBuffer, TCollisionResolutionManifold<T, d>* Manifold = nullptr)
+void ConstructUnionUnionConstraints(TGeometryParticleHandle<T, d>* Particle0, TGeometryParticleHandle<T, d>* Particle1, const FImplicitObject* Implicit0, const FImplicitObject* Implicit1, const T Thickness, TRigidBodyContactConstraint<T, d>& Constraint)
 {
 	SCOPE_CYCLE_COUNTER(STAT_UpdateUnionUnionConstraint);
 
@@ -2025,75 +1973,75 @@ void ConstructUnionUnionConstraints(TGeometryParticleHandle<T, d>* Particle0, TG
 		for (const Pair<const FImplicitObject*, TRigidTransform<T, d>>& ParticlePair : ParticleShapes)
 		{
 			const FImplicitObject* ParticleInnerObj = ParticlePair.First;
-			ConstructConstraintsImpl<T, d>(Particle0, Particle1, ParticleInnerObj, LevelsetInnerObj, Thickness, ConstraintBuffer, Manifold);
+			ConstructConstraintsImpl<T, d>(Particle0, Particle1, ParticleInnerObj, LevelsetInnerObj, Thickness, Constraint);
 		}
 	}
 }
 
 template<typename T, int d>
-void ConstructPairConstraintImpl(TGeometryParticleHandle<T, d>* Particle0, TGeometryParticleHandle<T, d>* Particle1, const FImplicitObject* Implicit0, const FImplicitObject* Implicit1, const T Thickness, TArray< TRigidBodyContactConstraint<T, d> >& ConstraintBuffer, TCollisionResolutionManifold<T, d>* Manifold)
+void ConstructPairConstraintImpl(TGeometryParticleHandle<T, d>* Particle0, TGeometryParticleHandle<T, d>* Particle1, const FImplicitObject* Implicit0, const FImplicitObject* Implicit1, const T Thickness, TRigidBodyContactConstraint<T, d>& Constraint)
 {
 	// See if we already have a constraint for this shape pair
-	if (Manifold && Manifold->ContainsShapeConnection(Implicit0, Implicit1))
+	if (Constraint.ContainsManifold(Implicit0, Implicit1))
 	{
 		return;
 	}
 
 	if (!Implicit0 || !Implicit1)
 	{
-		ConstructLevelsetConstraints(Particle0, Particle1, Implicit0, Implicit1, Thickness, ConstraintBuffer, Manifold);
+		ConstructLevelsetConstraints(Particle0, Particle1, Implicit0, Implicit1, Thickness, Constraint);
 	}
 	else if (Implicit0->GetType() == TBox<T, d>::StaticType() && Implicit1->GetType() == TBox<T, d>::StaticType()) 
 	{
-		ConstructBoxConstraints(Particle0, Particle1, Implicit0, Implicit1, Thickness, ConstraintBuffer, Manifold);
+		ConstructBoxConstraints(Particle0, Particle1, Implicit0, Implicit1, Thickness, Constraint);
 	}
 	else if (Implicit0->GetType() == TSphere<T, d>::StaticType() && Implicit1->GetType() == TSphere<T, d>::StaticType())
 	{
-		ConstructSphereConstraints(Particle0, Particle1, Implicit0, Implicit1, Thickness, ConstraintBuffer, Manifold);
+		ConstructSphereConstraints(Particle0, Particle1, Implicit0, Implicit1, Thickness, Constraint);
 	}
 	else if (Implicit0->GetType() == TBox<T, d>::StaticType() && Implicit1->GetType() == TPlane<T, d>::StaticType())
 	{
-		ConstructBoxPlaneConstraints(Particle0, Particle1, Implicit0, Implicit1, Thickness, ConstraintBuffer, Manifold);
+		ConstructBoxPlaneConstraints(Particle0, Particle1, Implicit0, Implicit1, Thickness, Constraint);
 	}
 	else if (Implicit1->GetType() == TBox<T, d>::StaticType() && Implicit0->GetType() == TPlane<T, d>::StaticType())
 	{
-		ConstructBoxPlaneConstraints(Particle1, Particle0, Implicit0, Implicit1, Thickness, ConstraintBuffer, Manifold);
+		ConstructBoxPlaneConstraints(Particle1, Particle0, Implicit0, Implicit1, Thickness, Constraint);
 	}
 	else if (Implicit0->GetType() == TSphere<T, d>::StaticType() && Implicit1->GetType() == TPlane<T, d>::StaticType())
 	{
-		ConstructSpherePlaneConstraints(Particle0, Particle1, Implicit0, Implicit1, Thickness, ConstraintBuffer, Manifold);
+		ConstructSpherePlaneConstraints(Particle0, Particle1, Implicit0, Implicit1, Thickness, Constraint);
 	}
 	else if (Implicit1->GetType() == TSphere<T, d>::StaticType() && Implicit0->GetType() == TPlane<T, d>::StaticType())
 	{
-		ConstructSpherePlaneConstraints(Particle1, Particle0, Implicit0, Implicit1, Thickness, ConstraintBuffer, Manifold);
+		ConstructSpherePlaneConstraints(Particle1, Particle0, Implicit0, Implicit1, Thickness, Constraint);
 	}
 	else if (Implicit0->GetType() == TSphere<T, d>::StaticType() && Implicit1->GetType() == TBox<T, d>::StaticType())
 	{
-		ConstructSphereBoxConstraints(Particle0, Particle1, Implicit0, Implicit1, Thickness, ConstraintBuffer, Manifold);
+		ConstructSphereBoxConstraints(Particle0, Particle1, Implicit0, Implicit1, Thickness, Constraint);
 	}
 	else if (Implicit1->GetType() == TSphere<T, d>::StaticType() && Implicit0->GetType() == TBox<T, d>::StaticType())
 	{
-		ConstructSphereBoxConstraints(Particle1, Particle0, Implicit0, Implicit1, Thickness, ConstraintBuffer, Manifold);
+		ConstructSphereBoxConstraints(Particle1, Particle0, Implicit0, Implicit1, Thickness, Constraint);
 	}
 	else if (Implicit0->GetType() == TCapsule<T>::StaticType() && Implicit1->GetType() == TCapsule<T>::StaticType())
 	{
-		ConstructCapsuleCapsuleConstraints(Particle0, Particle1, Implicit0, Implicit1, Thickness, ConstraintBuffer, Manifold);
+		ConstructCapsuleCapsuleConstraints(Particle0, Particle1, Implicit0, Implicit1, Thickness, Constraint);
 	}
 	else if (Implicit0->GetType() == TCapsule<T>::StaticType() && Implicit1->GetType() == TBox<T, d>::StaticType())
 	{
-		ConstructCapsuleBoxConstraints(Particle0, Particle1, Implicit0, Implicit1, Thickness, ConstraintBuffer, Manifold);
+		ConstructCapsuleBoxConstraints(Particle0, Particle1, Implicit0, Implicit1, Thickness, Constraint);
 	}
 	else if (Implicit1->GetType() == TCapsule<T>::StaticType() && Implicit0->GetType() == TBox<T, d>::StaticType())
 	{
-		ConstructCapsuleBoxConstraints(Particle1, Particle0, Implicit0, Implicit1, Thickness, ConstraintBuffer, Manifold);
+		ConstructCapsuleBoxConstraints(Particle1, Particle0, Implicit0, Implicit1, Thickness, Constraint);
 	}
 	else if (Implicit0->GetType() < TImplicitObjectUnion<T, d>::StaticType() && Implicit1->GetType() == TImplicitObjectUnion<T, d>::StaticType())
 	{
-		ConstructSingleUnionConstraints(Particle0, Particle1, Implicit0, Implicit1, Thickness, ConstraintBuffer, Manifold);
+		ConstructSingleUnionConstraints(Particle0, Particle1, Implicit0, Implicit1, Thickness, Constraint);
 	}
 	else if (Implicit0->GetType() == TImplicitObjectUnion<T, d>::StaticType() && Implicit1->GetType() < TImplicitObjectUnion<T, d>::StaticType())
 	{
-		ConstructSingleUnionConstraints(Particle1, Particle0, Implicit0, Implicit1, Thickness, ConstraintBuffer, Manifold);
+		ConstructSingleUnionConstraints(Particle1, Particle0, Implicit0, Implicit1, Thickness, Constraint);
 	}
 	else if (Implicit0->GetType() == TImplicitObjectUnion<T, d>::StaticType() && Implicit1->GetType() == TImplicitObjectUnion<T, d>::StaticType())
 	{
@@ -2102,16 +2050,16 @@ void ConstructPairConstraintImpl(TGeometryParticleHandle<T, d>* Particle0, TGeom
 	}
 	else if (Implicit0->IsConvex() && Implicit1->IsConvex())
 	{
-		CollisionResolutionConvexConvex<T,d>::ConstructConvexConvexConstraints(Particle0, Particle1, Implicit0, Implicit1, Thickness, ConstraintBuffer, Manifold, CollisionManifoldSampleScaleCVar, CollisionManifoldSamplesCVar);
+		CollisionResolutionConvexConvex<T,d>::ConstructConvexConvexConstraints(Particle0, Particle1, Implicit0, Implicit1, Thickness, Constraint);
 	}
 	else
 	{
-		ConstructLevelsetConstraints(Particle0, Particle1, Implicit0, Implicit1, Thickness, ConstraintBuffer, Manifold);
+		ConstructLevelsetConstraints(Particle0, Particle1, Implicit0, Implicit1, Thickness, Constraint);
 	}
 }
 
 template<typename T, int d>
-void ConstructConstraintsImpl(TGeometryParticleHandle<T, d>* Particle0, TGeometryParticleHandle<T, d>* Particle1, const FImplicitObject* Implicit0, const FImplicitObject* Implicit1, const T Thickness, TArray< TRigidBodyContactConstraint<T, d> >& ConstraintBuffer, TCollisionResolutionManifold<T, d>* Manifold)
+void ConstructConstraintsImpl(TGeometryParticleHandle<T, d>* Particle0, TGeometryParticleHandle<T, d>* Particle1, const FImplicitObject* Implicit0, const FImplicitObject* Implicit1, const T Thickness, TRigidBodyContactConstraint<T, d> & Constraint)
 {
 	// TriangleMesh implicits are for scene query only.
 	if (Implicit0 && GetInnerType(Implicit0->GetType()) == ImplicitObjectType::TriangleMesh) return;
@@ -2120,34 +2068,21 @@ void ConstructConstraintsImpl(TGeometryParticleHandle<T, d>* Particle0, TGeometr
 	if (Implicit0 && Implicit0->GetType() == TImplicitObjectUnion<T, d>::StaticType() &&
 		Implicit1 && Implicit1->GetType() == TImplicitObjectUnion<T, d>::StaticType())
 	{
-		ConstructUnionUnionConstraints(Particle0, Particle1, Implicit0, Implicit1, Thickness, ConstraintBuffer, Manifold);
+		ConstructUnionUnionConstraints(Particle0, Particle1, Implicit0, Implicit1, Thickness, Constraint);
 	}
 	else
 	{
-		ConstructPairConstraintImpl(Particle0, Particle1, Implicit0, Implicit1, Thickness, ConstraintBuffer, Manifold);
+		ConstructPairConstraintImpl(Particle0, Particle1, Implicit0, Implicit1, Thickness, Constraint);
 	}
 }
 
 template<typename T, int d>
-void TPBDCollisionConstraint<T, d>::ConstructConstraints(TGeometryParticleHandle<T, d>* Particle0, TGeometryParticleHandle<T, d>* Particle1, const T Thickness, TArray<TRigidBodyContactConstraint<T, d>>& ConstraintBuffer)
+void TPBDCollisionConstraint<T, d>::ConstructConstraints(TGeometryParticleHandle<T, d>* Particle0, TGeometryParticleHandle<T, d>* Particle1, const T Thickness, TRigidBodyContactConstraint<T, d>& Constraint)
 {
 	if (ensure(Particle0 && Particle1))
 	{
-		FConstraintManifold* ManifoldElement = nullptr;
-		if (FConstraintManifold** Ptr = Manifolds.Find(FConstraintHandleID::Make(Particle0, Particle1)))
-		{
-			ManifoldElement = *Ptr;
-		}
-
-		ConstructConstraintsImpl<T, d>(Particle0, Particle1, Particle0->Geometry().Get(), Particle1->Geometry().Get(), Thickness, ConstraintBuffer, ManifoldElement);
-	}
-
-	/*
-	if ((ManifoldElement != nullptr) && (ConstraintBuffer.Num() > 0))
-	{
-		ManifoldElement->SetTimestamp(LifespanCounter + CollisionHistoryLifespan);
+		ConstructConstraintsImpl<T, d>(Particle0, Particle1, Particle0->Geometry().Get(), Particle1->Geometry().Get(), Thickness, Constraint);
 	}
-	*/
 }
 
 // NOTE: UpdateLevelsetConstraintImp and its <float,3> specialization are here for the Linux build. It looks like
@@ -2165,11 +2100,8 @@ void UpdateLevelsetConstraintImp(const float Thickness, TRigidBodyContactConstra
 }
 
 template<ECollisionUpdateType UpdateType, typename T, int d>
-void UpdateConstraintImp2(const FImplicitObject& ParticleObject, const TRigidTransform<T, d>& ParticleTM, const FImplicitObject& LevelsetObject, const TRigidTransform<T, d>& LevelsetTM, const T Thickness,
-	TRigidBodyContactConstraint<T, d>* Constraints, int32 NumConstraints, TCollisionResolutionManifold<T, d>* Manifold)
+void UpdateConstraintImp(const FImplicitObject& ParticleObject, const TRigidTransform<T, d>& ParticleTM, const FImplicitObject& LevelsetObject, const TRigidTransform<T, d>& LevelsetTM, const T Thickness, TRigidBodyContactConstraint<T, d>& Constraint)
 {
-	TRigidBodyContactConstraint<T, d>& Constraint = Constraints[0];
-
 	if (ParticleObject.GetType() == TBox<T, d>::StaticType() && LevelsetObject.GetType() == TBox<T, d>::StaticType())
 	{
 		UpdateBoxConstraint(*ParticleObject.template GetObject<TBox<T,d>>(), ParticleTM, *LevelsetObject.template GetObject<TBox<T, d>>(), LevelsetTM, Thickness, Constraint);
@@ -2202,45 +2134,45 @@ void UpdateConstraintImp2(const FImplicitObject& ParticleObject, const TRigidTra
 	{
 		TRigidBodyContactConstraint<T, d> TmpConstraint = Constraint;
 		UpdateBoxPlaneConstraint(*LevelsetObject.template GetObject<TBox<T, d>>(), LevelsetTM, *ParticleObject.template GetObject<TPlane<T, d>>(), ParticleTM, Thickness, TmpConstraint);
-		if (TmpConstraint.Phi < Constraint.Phi)
+		if (TmpConstraint.GetPhi() < Constraint.GetPhi())
 		{
 			Constraint = TmpConstraint;
-			Constraint.Normal = -Constraint.Normal;
+			Constraint.SetNormal(-Constraint.GetNormal());
 		}
 	}
 	else if (ParticleObject.GetType() == TPlane<T, d>::StaticType() && LevelsetObject.GetType() == TSphere<T, d>::StaticType())
 	{
 		TRigidBodyContactConstraint<T, d> TmpConstraint = Constraint;
 		UpdateSpherePlaneConstraint(*LevelsetObject.template GetObject<TSphere<T, d>>(), LevelsetTM, *ParticleObject.template GetObject<TPlane<T, d>>(), ParticleTM, Thickness, TmpConstraint);
-		if (TmpConstraint.Phi < Constraint.Phi)
+		if (TmpConstraint.GetPhi() < Constraint.GetPhi())
 		{
 			Constraint = TmpConstraint;
-			Constraint.Normal = -Constraint.Normal;
+			Constraint.SetNormal(-Constraint.GetNormal());
 		}
 	}
 	else if (ParticleObject.GetType() == TBox<T, d>::StaticType() && LevelsetObject.GetType() == TSphere<T, d>::StaticType())
 	{
 		TRigidBodyContactConstraint<T, d> TmpConstraint = Constraint;
 		UpdateSphereBoxConstraint(*LevelsetObject.template GetObject<TSphere<T, d>>(), LevelsetTM, *ParticleObject.template GetObject<TBox<T, d>>(), ParticleTM, Thickness, TmpConstraint);
-		if (TmpConstraint.Phi < Constraint.Phi)
+		if (TmpConstraint.GetPhi() < Constraint.GetPhi())
 		{
 			Constraint = TmpConstraint;
-			Constraint.Normal = -Constraint.Normal;
+			Constraint.SetNormal(-Constraint.GetNormal());
 		}
 	}
 	else if (ParticleObject.GetType() == TBox<T, d>::StaticType() && LevelsetObject.GetType() == TCapsule<T>::StaticType())
 	{
 		TRigidBodyContactConstraint<T, d> TmpConstraint = Constraint;
 		UpdateCapsuleBoxConstraint(*LevelsetObject.template GetObject<TCapsule<T>>(), LevelsetTM, *ParticleObject.template GetObject<TBox<T, d>>(), ParticleTM, Thickness, TmpConstraint);
-		if (TmpConstraint.Phi < Constraint.Phi)
+		if (TmpConstraint.GetPhi() < Constraint.GetPhi())
 		{
 			Constraint = TmpConstraint;
-			Constraint.Normal = -Constraint.Normal;
+			Constraint.SetNormal(-Constraint.GetNormal());
 		}
 	}
 	else if (ParticleObject.GetType() < TImplicitObjectUnion<T, d>::StaticType() && LevelsetObject.GetType() == TImplicitObjectUnion<T, d>::StaticType())
 	{
-		return UpdateSingleUnionConstraint<UpdateType>(Thickness, Constraints, NumConstraints, Manifold);
+		return UpdateSingleUnionConstraint<UpdateType>(Thickness, Constraint);
 	}
 	else if (ParticleObject.GetType() == TImplicitObjectUnion<T, d>::StaticType() && LevelsetObject.GetType() < TImplicitObjectUnion<T, d>::StaticType())
 	{
@@ -2248,11 +2180,11 @@ void UpdateConstraintImp2(const FImplicitObject& ParticleObject, const TRigidTra
 	}
 	else if (ParticleObject.GetType() == TImplicitObjectUnion<T, d>::StaticType() && LevelsetObject.GetType() == TImplicitObjectUnion<T, d>::StaticType())
 	{
-		return UpdateUnionUnionConstraint<UpdateType>(Thickness, Constraints, NumConstraints, Manifold);
+		return UpdateUnionUnionConstraint<UpdateType>(Thickness, Constraint);
 	}
 	else if (ParticleObject.IsConvex() && LevelsetObject.IsConvex())
 	{
-		CollisionResolutionConvexConvex<T,d>::UpdateConvexConvexConstraint(ParticleObject, ParticleTM, LevelsetObject, LevelsetTM, Thickness, Constraints, NumConstraints, Manifold, CollisionManifoldSampleScaleCVar, CollisionManifoldSamplesCVar);
+		CollisionResolutionConvexConvex<T,d>::UpdateConvexConvexConstraint(ParticleObject, ParticleTM, LevelsetObject, LevelsetTM, Thickness, Constraint);
 	}
 	else if (LevelsetObject.IsUnderlyingUnion())
 	{
@@ -2272,21 +2204,11 @@ DECLARE_CYCLE_STAT(TEXT("TPBDCollisionConstraint::UpdateConstraint"), STAT_Updat
 
 template<typename T, int d>
 template<ECollisionUpdateType UpdateType>
-void TPBDCollisionConstraint<T, d>::UpdateConstraint(const T Thickness, FRigidBodyContactConstraint* InConstraints, int32 NumConstraints)
+void TPBDCollisionConstraint<T, d>::UpdateConstraint(const T Thickness, FRigidBodyContactConstraint& Constraint)
 {
 	SCOPE_CYCLE_COUNTER(STAT_UpdateConstraint);
 
-	// Reset Phi and validate that the constraints operate on the same bodies.
-	check(NumConstraints >= 1);
-	check(InConstraints != nullptr);
-	for (int32 i = 0; i < NumConstraints; i++)
-	{
-		InConstraints[i].Phi = Thickness;
-		ensure(InConstraints[i % NumConstraints].Particle == InConstraints[(i + 1) % NumConstraints].Particle);
-		ensure(InConstraints[i % NumConstraints].Levelset == InConstraints[(i + 1) % NumConstraints].Levelset);
-	}
-
-	FRigidBodyContactConstraint& Constraint = InConstraints[0];
+	Constraint.ResetPhi(Thickness);	
 	const TRigidTransform<T, d> ParticleTM = GetTransform(Constraint.Particle);
 	const TRigidTransform<T, d> LevelsetTM = GetTransform(Constraint.Levelset);
 
@@ -2306,25 +2228,19 @@ void TPBDCollisionConstraint<T, d>::UpdateConstraint(const T Thickness, FRigidBo
 	}
 	else
 	{
-		TCollisionResolutionManifold<T, d>* ManifoldData = nullptr;
-		FConstraintHandleID HandleID = GetConstraintHandleID(Constraint);
-		if (Manifolds.Contains(HandleID))
-		{
-			ManifoldData = Manifolds[HandleID];
-		}
-		UpdateConstraintImp2<UpdateType>(*Constraint.Particle->Geometry(), ParticleTM, *Constraint.Levelset->Geometry(), LevelsetTM, Thickness, InConstraints, NumConstraints, ManifoldData);
+		UpdateConstraintImp<UpdateType>(*Constraint.Particle->Geometry(), ParticleTM, *Constraint.Levelset->Geometry(), LevelsetTM, Thickness, Constraint);
 	}
 }
 
 template class TPBDCollisionConstraintHandle<float, 3>;
 template class TAccelerationStructureHandle<float, 3>;
 template class CHAOS_API TPBDCollisionConstraint<float, 3>;
-template void TPBDCollisionConstraint<float, 3>::UpdateConstraint<ECollisionUpdateType::Any>(const float Thickness, FRigidBodyContactConstraint* Constraints, int32 NumConstraints);
-template void TPBDCollisionConstraint<float, 3>::UpdateConstraint<ECollisionUpdateType::Deepest>(const float Thickness, FRigidBodyContactConstraint* Constraints, int32 NumConstraints);
+template void TPBDCollisionConstraint<float, 3>::UpdateConstraint<ECollisionUpdateType::Any>(const float Thickness, FRigidBodyContactConstraint& Constraints);
+template void TPBDCollisionConstraint<float, 3>::UpdateConstraint<ECollisionUpdateType::Deepest>(const float Thickness, FRigidBodyContactConstraint& Constraints);
 template void TPBDCollisionConstraint<float, 3>::UpdateLevelsetConstraint<ECollisionUpdateType::Any>(const float Thickness, FRigidBodyContactConstraint& Constraint);
 template void TPBDCollisionConstraint<float, 3>::UpdateLevelsetConstraint<ECollisionUpdateType::Deepest>(const float Thickness, FRigidBodyContactConstraint& Constraint);
 template void TPBDCollisionConstraint<float, 3>::ComputeConstraints<false>(const TPBDCollisionConstraint<float,3>::FAccelerationStructure&, float Dt);
 template void TPBDCollisionConstraint<float, 3>::ComputeConstraints<true>(const TPBDCollisionConstraint<float, 3>::FAccelerationStructure&, float Dt);
-template void UpdateConstraintImp2<ECollisionUpdateType::Any, float, 3>(const FImplicitObject& ParticleObject, const TRigidTransform<float, 3>& ParticleTM, const FImplicitObject& LevelsetObject, const TRigidTransform<float, 3>& LevelsetTM, const float Thickness, TRigidBodyContactConstraint<float, 3>* Constraints, int32 NumConstraints, TCollisionResolutionManifold<float, 3>* Manifold);
-template void UpdateConstraintImp2<ECollisionUpdateType::Deepest, float, 3>(const FImplicitObject& ParticleObject, const TRigidTransform<float, 3>& ParticleTM, const FImplicitObject& LevelsetObject, const TRigidTransform<float, 3>& LevelsetTM, const float Thickness, TRigidBodyContactConstraint<float, 3>* Constraint, int32 NumConstraints, TCollisionResolutionManifold<float, 3>* Manifold);
+template void UpdateConstraintImp<ECollisionUpdateType::Any, float, 3>(const FImplicitObject& ParticleObject, const TRigidTransform<float, 3>& ParticleTM, const FImplicitObject& LevelsetObject, const TRigidTransform<float, 3>& LevelsetTM, const float Thickness, TRigidBodyContactConstraint<float, 3>& Constraints);
+template void UpdateConstraintImp<ECollisionUpdateType::Deepest, float, 3>(const FImplicitObject& ParticleObject, const TRigidTransform<float, 3>& ParticleTM, const FImplicitObject& LevelsetObject, const TRigidTransform<float, 3>& LevelsetTM, const float Thickness, TRigidBodyContactConstraint<float, 3>& Constraint);
 }
diff --git a/Engine/Source/Runtime/Experimental/Chaos/Private/Chaos/PBDCollisionConstraintUtil.cpp b/Engine/Source/Runtime/Experimental/Chaos/Private/Chaos/PBDCollisionConstraintUtil.cpp
--- a/Engine/Source/Runtime/Experimental/Chaos/Private/Chaos/PBDCollisionConstraintUtil.cpp
+++ b/Engine/Source/Runtime/Experimental/Chaos/Private/Chaos/PBDCollisionConstraintUtil.cpp
@@ -52,7 +52,7 @@ namespace Chaos
 		int32 NumberOfCellsXYZ = NumberOfCellsX * NumberOfCellsY * NumberOfCellsZ;
 		for (int32 IdxConstraint = 0; IdxConstraint < ConstraintsArray.Num(); ++IdxConstraint)
 		{
-			FVector Location = (FVector)ConstraintsArray[IdxConstraint].Location - BoundingBox.Min + FVector(0.5f * CellSize);
+			FVector Location = (FVector)ConstraintsArray[IdxConstraint].GetLocation() - BoundingBox.Min + FVector(0.5f * CellSize);
 			int32 HashTableIdx = (int32)(Location.X * CellSizeInv) +
 								 (int32)(Location.Y * CellSizeInv) * NumberOfCellsX +
 								 (int32)(Location.Z * CellSizeInv) * NumberOfCellsXY;
diff --git a/Engine/Source/Runtime/Experimental/Chaos/Private/Chaos/PBDRigidClustering.cpp b/Engine/Source/Runtime/Experimental/Chaos/Private/Chaos/PBDRigidClustering.cpp
--- a/Engine/Source/Runtime/Experimental/Chaos/Private/Chaos/PBDRigidClustering.cpp
+++ b/Engine/Source/Runtime/Experimental/Chaos/Private/Chaos/PBDRigidClustering.cpp
@@ -1758,7 +1758,7 @@ namespace Chaos
 	template<class T, int d>
 	TVector<T, d> GetContactLocation(const TRigidBodyContactConstraint<T, d>& Contact)
 	{
-		return Contact.Location;
+		return Contact.GetLocation();
 	}
 
 	template<class T, int d>
diff --git a/Engine/Source/Runtime/Experimental/Chaos/Public/Chaos/CollisionResolutionConvexConvex.h b/Engine/Source/Runtime/Experimental/Chaos/Public/Chaos/CollisionResolutionConvexConvex.h
--- a/Engine/Source/Runtime/Experimental/Chaos/Public/Chaos/CollisionResolutionConvexConvex.h
+++ b/Engine/Source/Runtime/Experimental/Chaos/Public/Chaos/CollisionResolutionConvexConvex.h
@@ -16,23 +16,22 @@ template<class T, int d>
 class CHAOS_API CollisionResolutionConvexConvex
 {
 public:
-	typedef TConvex<T,d> FConvex;
-	typedef TRigidTransform<T,d> FRigidTransform;
-	typedef TRigidBodyContactConstraint<T,d> FRigidBodyContactConstraint;
-	typedef TCollisionResolutionManifold<T,d> FCollisionResolutionManifold;
-	typedef TGeometryParticleHandle<T,d> FGeometryParticleHandle;
+	using FConvex = TConvex<T,d>;
+	using FRigidTransform = TRigidTransform<T,d>;
+	using FRigidBodyContactConstraint = TRigidBodyContactConstraint<T, d>;
+	using FGeometryParticleHandle = TGeometryParticleHandle<T, d>;
 
 	/*
 	* Build a constraint manifold for the contact shape pair. 
 	*/
 	static void ConstructConvexConvexConstraints(FGeometryParticleHandle* Particle0, FGeometryParticleHandle* Particle1, const FImplicitObject* Implicit0, const FImplicitObject* Implicit1, const float Thickness,
-		TArray<FRigidBodyContactConstraint>& ConstraintBuffer, FCollisionResolutionManifold* Manifold = nullptr, float ManifoldScale = 0.5, int32 ManifoldSamples = 4);
+		FRigidBodyContactConstraint& ConstraintBuffer);
 
 	/*
 	* Update the contact manifold.
 	*/
 	static void UpdateConvexConvexConstraint(const FImplicitObject& A, const FRigidTransform& ATM, const FImplicitObject& B, FRigidTransform BTM, float Thickness, 
-		FRigidBodyContactConstraint* Constraints, int32 NumConstraints, FCollisionResolutionManifold* Manifold, float ManifoldScale = 0.5, int32 ManifoldSamples = 4);
+		FRigidBodyContactConstraint& Constraints);
 
 
 };
diff --git a/Engine/Source/Runtime/Experimental/Chaos/Public/Chaos/CollisionResolutionManifold.h b/Engine/Source/Runtime/Experimental/Chaos/Public/Chaos/CollisionResolutionManifold.h
--- a/Engine/Source/Runtime/Experimental/Chaos/Public/Chaos/CollisionResolutionManifold.h
+++ b/Engine/Source/Runtime/Experimental/Chaos/Public/Chaos/CollisionResolutionManifold.h
@@ -14,89 +14,7 @@
 
 namespace Chaos
 {
-	template<typename T, int d>
-	class TPBDCollisionConstraintHandle;
 
-	/**
-	 * A pair that can be used as a key, where the key value is independent of the pair order,
-	 * i.e., TSymmetricPairKey(A, B) == TSymmetricPairKey(B, A)
-	 */
-	template<typename T>
-	class TSymmetricPairKey
-	{
-	public:
-		using Item = T;
-		using FPair = TPair<Item, Item>;
 
-		static TSymmetricPairKey Make(const T& Item0, const T& Item1)
-		{
-			T ItemMin = FMath::Min(Item0, Item1);
-			T ItemMax = FMath::Max(Item0, Item1);
-			return TSymmetricPairKey<T>(ItemMin, ItemMax);
-		}
-
-		friend inline uint32 GetTypeHash(const TSymmetricPairKey<T>& Item)
-		{
-			return GetTypeHash(Item.Pair);
-		}
-
-		friend inline bool operator==(const TSymmetricPairKey<T>& L, const TSymmetricPairKey<T>& R)
-		{
-			return L.Pair == R.Pair;
-		}
-
-	private:
-		TSymmetricPairKey(const T& Item0, const T& Item1)
-			: Pair(Item0, Item1)
-		{
-			check(Item0 < Item1);
-		}
-
-		FPair Pair;
-	};
-
-
-	template<class T, int d>
-	class CHAOS_API TCollisionResolutionManifold
-	{
-		using FImplicitPairKey = TSymmetricPairKey<const FImplicitObject*>;
-
-	public:
-		using FConstraintContainerHandle = TPBDCollisionConstraintHandle<T, d>;
-		TCollisionResolutionManifold(const TVector<T, d>& InLocation, const TRotation<T, d>& InRotation, int32 InTimestamp = -INT_MAX)
-			: Timestamp(InTimestamp)
-			, Location(InLocation)
-			, Rotation(InRotation)
-		{}
-
-		const TVector<T, d>& GetLocation() { return Location; }
-		const TRotation<T, d>& GetRotation() { return Rotation; }
-
-		void AddHandle(FConstraintContainerHandle* InHandle)
-		{
-			Implicits.Add(FImplicitPairKey::Make(InHandle->GetContact().Geometry[0], InHandle->GetContact().Geometry[1]));
-			ConstraintHandles.Push(InHandle);
-		}
-		void RemoveHandle(FConstraintContainerHandle* InHandle)
-		{
-			ConstraintHandles.RemoveSingleSwap(InHandle);
-			Implicits.Remove(FImplicitPairKey::Make(InHandle->GetContact().Geometry[0], InHandle->GetContact().Geometry[1]));
-		}
-
-		TArray<FConstraintContainerHandle*>& GetHandles() { return ConstraintHandles; }
-		const TArray<FConstraintContainerHandle*>& GetHandles() const { return ConstraintHandles; }
-
-		bool ContainsShapeConnection(const FImplicitObject* Implicit0In, const FImplicitObject* Implicit1In) { return Implicits.Contains(FImplicitPairKey::Make(Implicit0In, Implicit1In)); }
-
-		int32 GetTimestamp() const { return Timestamp; }
-		void SetTimestamp(int32 InTimestamp) { Timestamp = InTimestamp; }
-
-	private:
-		int32 Timestamp;
-		TVector<T, d> Location;
-		TRotation<T, d> Rotation;
-		TArray<FConstraintContainerHandle*> ConstraintHandles;
-		TSet< FImplicitPairKey > Implicits;
-	};
 
 }
diff --git a/Engine/Source/Runtime/Experimental/Chaos/Public/Chaos/PBDCollisionConstraint.h b/Engine/Source/Runtime/Experimental/Chaos/Public/Chaos/PBDCollisionConstraint.h
--- a/Engine/Source/Runtime/Experimental/Chaos/Public/Chaos/PBDCollisionConstraint.h
+++ b/Engine/Source/Runtime/Experimental/Chaos/Public/Chaos/PBDCollisionConstraint.h
@@ -82,10 +82,9 @@ public:
 	using FReal = T;
 	static const int Dimensions = d;
 	using FConstraintContainerHandle = TPBDCollisionConstraintHandle<T, d>;
-	using FConstraintManifold = TCollisionResolutionManifold<T, d>;
 	using FConstraintHandleAllocator = TConstraintHandleAllocator<TPBDCollisionConstraint<T, d>>;
 	using FRigidBodyContactConstraint = TRigidBodyContactConstraint<T, d>;
-	using FConstraintHandleID = TSymmetricPairKey<const TGeometryParticleHandle<T, d>*>;
+	using FConstraintHandleID = TPair<const TGeometryParticleHandle<T, d>*, const TGeometryParticleHandle<T, d>*>;
 
 	TPBDCollisionConstraint(const TPBDRigidsSOAs<T,d>& InParticles, TArrayCollectionArray<bool>& Collided, const TArrayCollectionArray<TSerializablePtr<FChaosPhysicsMaterial>>& PerParticleMaterials, const int32 PairIterations = 1, const T Thickness = (T)0);
 	virtual ~TPBDCollisionConstraint() {}
@@ -129,12 +128,12 @@ public:
 
 	FConstraintHandleID GetConstraintHandleID(const FRigidBodyContactConstraint & Constraint) const
 	{
-		return  FConstraintHandleID::Make( Constraint.Particle, Constraint.Levelset );
+		return  FConstraintHandleID( Constraint.Particle, Constraint.Levelset );
 	}
 
 	FConstraintHandleID GetConstraintHandleID(int32 ConstraintIndex) const
 	{
-		return  FConstraintHandleID::Make( Constraints[ConstraintIndex].Particle, Constraints[ConstraintIndex].Levelset );
+		return  FConstraintHandleID( Constraints[ConstraintIndex].Particle, Constraints[ConstraintIndex].Levelset );
 	}
 		
 	const FConstraintContainerHandle* GetConstraintHandle(int32 ConstraintIndex) const
@@ -201,11 +200,6 @@ public:
 	 // @todo(ccaulfield): this runs wide. The serial/parallel decision should be in the ConstraintRule
 	void Apply(const T Dt, const TArray<FConstraintContainerHandle*>& InConstraintHandles, const int32 It, const int32 NumIts);
 
-	/**
-	 * Set disabled attribute on all input handles
-	 */
-	void Disable(const TArray<FConstraintContainerHandle*>& InConstraintHandles);
-		
 	/**
 	 * Generate all contact constraints.
 	 */
@@ -255,9 +249,9 @@ private:
 	void ComputeConstraints(const FAccelerationStructure& AccelerationStructure, T Dt);
 
 	template<ECollisionUpdateType>
-	void UpdateConstraint(const T Thickness, FRigidBodyContactConstraint* Constraints, int32 NumConstraints);
+	void UpdateConstraint(const T Thickness, FRigidBodyContactConstraint& Constraint);
 
-	void ConstructConstraints(TGeometryParticleHandle<T, d>* Particle0, TGeometryParticleHandle<T, d>* Particle1, const T Thickness, TArray< TRigidBodyContactConstraint<T,d> >& ConstraintBuffer);
+	void ConstructConstraints(TGeometryParticleHandle<T, d>* Particle0, TGeometryParticleHandle<T, d>* Particle1, const T Thickness, TRigidBodyContactConstraint<T,d>& Constraint);
 
 
 	template<typename SPATIAL_ACCELERATION>
@@ -279,7 +273,6 @@ private:
 	bool bEnableCollisions;
 
 	int32 LifespanCounter;
-	int32 CollisionManifoldLifespan;
 	float CollisionVelocityInflation;
 
 	TRigidBodyContactConstraintsPostComputeCallback<T, d> PostComputeCallback;
@@ -289,11 +282,11 @@ private:
 	TArray<FConstraintContainerHandle*> Handles;
 	FConstraintHandleAllocator HandleAllocator;
 
-	TMap<FConstraintHandleID, FConstraintManifold*> Manifolds;
+	TMap< FConstraintHandleID, FConstraintContainerHandle* > Manifolds;
 };
 
-extern template void TPBDCollisionConstraint<float, 3>::UpdateConstraint<ECollisionUpdateType::Any>(const float Thickness, FRigidBodyContactConstraint* Constraints, int32 NumConstraints);
-extern template void TPBDCollisionConstraint<float, 3>::UpdateConstraint<ECollisionUpdateType::Deepest>(const float Thickness, FRigidBodyContactConstraint* Constraints, int32 NumConstraints);
+extern template void TPBDCollisionConstraint<float, 3>::UpdateConstraint<ECollisionUpdateType::Any>(const float Thickness, FRigidBodyContactConstraint& Constraint);
+extern template void TPBDCollisionConstraint<float, 3>::UpdateConstraint<ECollisionUpdateType::Deepest>(const float Thickness, FRigidBodyContactConstraint& Constraint);
 extern template void TPBDCollisionConstraint<float, 3>::UpdateLevelsetConstraint<ECollisionUpdateType::Any>(const float Thickness, FRigidBodyContactConstraint& Constraint);
 extern template void TPBDCollisionConstraint<float, 3>::UpdateLevelsetConstraint<ECollisionUpdateType::Deepest>(const float Thickness, FRigidBodyContactConstraint& Constraint);
 extern template void TPBDCollisionConstraint<float, 3>::ComputeConstraints<false>(const TPBDCollisionConstraint<float, 3>::FAccelerationStructure&, float Dt);
diff --git a/Engine/Source/Runtime/Experimental/Chaos/Public/Chaos/PBDCollisionConstraintImp.h b/Engine/Source/Runtime/Experimental/Chaos/Public/Chaos/PBDCollisionConstraintImp.h
--- a/Engine/Source/Runtime/Experimental/Chaos/Public/Chaos/PBDCollisionConstraintImp.h
+++ b/Engine/Source/Runtime/Experimental/Chaos/Public/Chaos/PBDCollisionConstraintImp.h
@@ -203,9 +203,6 @@ namespace Chaos
 				int32 CountNP = 0;
 				int32 RejectedNP = 0;
 #endif
-				TArray<TRigidBodyContactConstraint<T, d>> TempConstraintBuffer;
-				TempConstraintBuffer.Reserve(4);
-
 				const int32 NumPotentials = PotentialIntersections.Num();
 				for (int32 i = 0; i < NumPotentials; ++i)
 				{
@@ -255,8 +252,9 @@ namespace Chaos
 					const TVector<T, d> Box2Thickness = Particle2.AsDynamic() ? ComputeThickness(*Particle2.AsDynamic(), Dt) : TVector<T, d>(0);
 					const T UseThickness = FMath::Max(Box1Thickness, Box2Thickness.Size());// + MThickness
 
-					TempConstraintBuffer.Reset();
-					ConstructConstraints(Particle1.Handle(), Particle2.Handle(), UseThickness, TempConstraintBuffer);
+
+					TRigidBodyContactConstraint<T, d> Constraint;
+					ConstructConstraints(Particle1.Handle(), Particle2.Handle(), UseThickness, Constraint);
 
 #if !UE_BUILD_SHIPPING
 					if (bGatherStats)
@@ -264,23 +262,20 @@ namespace Chaos
 						++CountNP;
 					}
 #endif
-					if (TempConstraintBuffer.Num())
+					if (Constraint.ContainsManifold())
 					{
-						UpdateConstraint<ECollisionUpdateType::Any>(UseThickness, &TempConstraintBuffer[0], TempConstraintBuffer.Num());
+						UpdateConstraint<ECollisionUpdateType::Any>(UseThickness, Constraint);
 					}
 
-					for (auto& Constraint : TempConstraintBuffer)
+					if (Constraint.GetPhi() < UseThickness)
+					{
+						Queue.Enqueue(Constraint);
+					}
+					else
 					{
-						if (Constraint.Phi < UseThickness)
-						{
-							Queue.Enqueue(Constraint);
-						}
-						else
-						{
 #if !UE_BUILD_SHIPPING
-							++RejectedNP;
+						++RejectedNP;
 #endif
-						}
 					}
 				}
 
diff --git a/Engine/Source/Runtime/Experimental/Chaos/Public/Chaos/PBDCollisionTypes.h b/Engine/Source/Runtime/Experimental/Chaos/Public/Chaos/PBDCollisionTypes.h
--- a/Engine/Source/Runtime/Experimental/Chaos/Public/Chaos/PBDCollisionTypes.h
+++ b/Engine/Source/Runtime/Experimental/Chaos/Public/Chaos/PBDCollisionTypes.h
@@ -13,503 +13,602 @@ namespace Chaos
 	template<typename T, int d>
 	class TPBDCollisionConstraint;
 
-/**/
-// @todo(ccaulfield): should these be held SOA style in the PBDCollisionConstraint?
-template<class T, int d>
-struct TRigidBodyContactConstraint
-{
-	TRigidBodyContactConstraint() : bDisabled(false), Normal(0), AccumulatedImpulse(0.f) {}
-	TGeometryParticleHandle<T, d>* Particle;
-	TGeometryParticleHandle<T, d>* Levelset;
-	const FImplicitObject* Geometry[2];
-
-	bool bDisabled;
-	TVector<T, d> Normal;
-	TVector<T, d> Location;
-	TVector<T, d> LocalLocation;
-	T Phi;
-	TVector<T, d> AccumulatedImpulse;
-
-	FString ToString() const
-	{
-		return FString::Printf(TEXT("Particle:%s, Levelset:%s, Normal:%s, Location:%s, Phi:%f, AccumulatedImpulse:%s"), *Particle->ToString(), *Levelset->ToString(), *Normal.ToString(), *Location.ToString(), Phi, *AccumulatedImpulse.ToString());
-	}
-};
 
-template<class T, int d>
-struct TRigidBodyContactConstraintPGS
-{
-	TRigidBodyContactConstraintPGS() : AccumulatedImpulse(0.f) {}
-	TGeometryParticleHandle<T, d>* Particle;
-	TGeometryParticleHandle<T, d>* Levelset;
-	TArray<TVector<T, d>> Normal;
-	TArray<TVector<T, d>> Location;
-	TArray<T> Phi;
-	TVector<T, d> AccumulatedImpulse;
-};
-
-
-template<class T, int d>
-class CHAOS_API TPBDCollisionConstraintHandle : public TContainerConstraintHandle<TPBDCollisionConstraint<T, d>>
-{
-public:
-	using Base = TContainerConstraintHandle<TPBDCollisionConstraint<T, d>>;
-	using FConstraintContainer = TPBDCollisionConstraint<T, d>;
-
-	TPBDCollisionConstraintHandle() 
-	{}
-	
-	TPBDCollisionConstraintHandle(FConstraintContainer* InConstraintContainer, int32 InConstraintIndex) 
-		: TContainerConstraintHandle<TPBDCollisionConstraint<T, d>>(InConstraintContainer, InConstraintIndex)
+	/*
+	*
+	*  Contact manifold stored in the local space of the
+	*  Target object.
+	*
+	*/
+	template<class T, int d>
+	class TContactData
 	{
-	}
+	public:
 
-	TRigidBodyContactConstraint<T, d>& GetContact()
-	{
-		return ConstraintContainer->Constraints[ConstraintIndex];
-	}
+		TContactData() : bDisabled(true), Normal(0), Phi(FLT_MAX) {};
 
-	const TRigidBodyContactConstraint<T, d>& GetContact() const
+		//void AddPoint(TVector<T, d> Point, TVector<T, d> Normal){}
+
+		bool bDisabled;
+		//int NumContacts;
+		TVector<T, d> Normal;
+		TVector<T, d> Location;
+		T Phi;
+	};
+
+	/*
+	*
+	*/
+	template<class T, int d>
+	class TConvexManifold
+	{
+	public:
+		using FContactData = TContactData<T, d>;
+
+		TConvexManifold(int32 InTimestamp = -INT_MAX, const FImplicitObject* InImplicit0 = nullptr, const FImplicitObject* InImplicit1 = nullptr)
+			: Timestamp(InTimestamp)
+		{
+			Implicit[0] = InImplicit0;
+			Implicit[1] = InImplicit1;
+		}
+
+		int32 Timestamp;
+		FContactData Manifold;
+		const FImplicitObject* Implicit[2];
+	};
+
+
+	/*
+	*
+	*/
+	template<class T, int d>
+	class TRigidBodyContactConstraint
 	{
-		return ConstraintContainer->Constraints[ConstraintIndex];
-	}
+	public:
+		using FGeometryParticleHandle = TGeometryParticleHandle<T, d>;
+		using FConvexManifold = TConvexManifold<T, d>;
 
-	void SetConstraintIndex(int32 IndexIn) 
-	{ 
-		ConstraintIndex = IndexIn; 
-	}
+		TRigidBodyContactConstraint() : AccumulatedImpulse(0) {}
 
-protected:
-	using Base::ConstraintIndex;
-	using Base::ConstraintContainer;
+		//API
+		void ResetPhi(T InPhi) { SetPhi(InPhi); }
+		bool ContainsManifold(const FImplicitObject* B = nullptr, const FImplicitObject* A = nullptr) const { return (A==nullptr&&B==nullptr)?(ShapeManifold.Implicit[0] != nullptr && ShapeManifold.Implicit[1] != nullptr):(ShapeManifold.Implicit[0]==A && ShapeManifold.Implicit[1]==B); }
 
+		void SetDisabled(bool bInDisabled, const FImplicitObject* B = nullptr, const FImplicitObject* A = nullptr) { ShapeManifold.Manifold.bDisabled = bInDisabled; }
+		TVector<T, d> GetDisabled(const FImplicitObject* B = nullptr, const FImplicitObject* A = nullptr) const { return ShapeManifold.Manifold.bDisabled; }
 
-};
 
+		void SetNormal(const TVector<T, d> & InNormal, const FImplicitObject* B = nullptr, const FImplicitObject* A = nullptr) { ShapeManifold.Manifold.Normal = InNormal; }
+		TVector<T, d> GetNormal(const FImplicitObject* B = nullptr, const FImplicitObject* A = nullptr) const { return ShapeManifold.Manifold.Normal; }
 
-/**
- * Collision event data stored for use by other systems (e.g. Niagara, gameplay events)
- */
-template<class T, int d>
-struct TCollisionData
-{
-	TCollisionData()
-		: Location(TVector<T, d>((T)0.0))
-		, AccumulatedImpulse(TVector<T, d>((T)0.0))
-		, Normal(TVector<T, d>((T)0.0))
-		, Velocity1(TVector<T, d>((T)0.0))
-		, Velocity2(TVector<T, d>((T)0.0))
-		, AngularVelocity1(TVector<T, d>((T)0.0))
-		, AngularVelocity2(TVector<T, d>((T)0.0))
-		, Mass1((T)0.0)
-		, Mass2((T)0.0)
-		, PenetrationDepth((T)0.0)
-		, Particle(nullptr)
-		, Levelset(nullptr)
-		, ParticleIndexMesh(INDEX_NONE)
-		, LevelsetIndexMesh(INDEX_NONE)
-	{}
-
-	TCollisionData(TVector<T, d> InLocation
-		, TVector<T, d> InAccumulatedImpulse
-		, TVector<T, d> InNormal
-		, TVector<T, d> InVelocity1
-		, TVector<T, d> InVelocity2
-		, TVector<T, d> InAngularVelocity1
-		, TVector<T, d> InAngularVelocity2
-		, T InMass1
-		, T InMass2
-		, T InPenetrationDepth
-		, TGeometryParticle<T, d>* InParticle
-		, TGeometryParticle<T, d>* InLevelset
-		, int32 InParticleIndexMesh
-		, int32 InLevelsetIndexMesh)
-		: Location(InLocation)
-		, AccumulatedImpulse(InAccumulatedImpulse)
-		, Normal(InNormal)
-		, Velocity1(InVelocity1)
-		, Velocity2(InVelocity2)
-		, AngularVelocity1(InAngularVelocity1)
-		, AngularVelocity2(InAngularVelocity2)
-		, Mass1(InMass1)
-		, Mass2(InMass2)
-		, PenetrationDepth(InPenetrationDepth)
-		, Particle(InParticle)
-		, Levelset(InLevelset)
-		, ParticleIndexMesh(InParticleIndexMesh)
-		, LevelsetIndexMesh(InLevelsetIndexMesh)
-	{}
-
-	TVector<T, d> Location;
-	TVector<T, d> AccumulatedImpulse;
-	TVector<T, d> Normal;
-	TVector<T, d> Velocity1, Velocity2;
-	TVector<T, d> AngularVelocity1, AngularVelocity2;
-	T Mass1, Mass2;
-	T PenetrationDepth;
-	TGeometryParticle<T, d>* Particle;
-	TGeometryParticle<T, d>* Levelset;
-	// @todo(ccaulfield): CHAOS_PARTICLEHANDLE_TODO
-	int32 ParticleIndexMesh, LevelsetIndexMesh; // If ParticleIndex points to a cluster then this index will point to an actual mesh in the cluster
-												// It is important to be able to get extra data from the component
-
-};
-
-
-/*
-CollisionData used in the subsystems
-*/
-template<class T, int d>
-struct TCollisionDataExt
-{
-	TCollisionDataExt()
-		: Location(TVector<T, d>((T)0.0))
-		, AccumulatedImpulse(TVector<T, d>((T)0.0))
-		, Normal(TVector<T, d>((T)0.0))
-		, Velocity1(TVector<T, d>((T)0.0))
-		, Velocity2(TVector<T, d>((T)0.0))
-		, AngularVelocity1(TVector<T, d>((T)0.0))
-		, AngularVelocity2(TVector<T, d>((T)0.0))
-		, Mass1((T)0.0)
-		, Mass2((T)0.0)
-		, Particle(nullptr)
-		, Levelset(nullptr)
-		, ParticleIndexMesh(INDEX_NONE)
-		, LevelsetIndexMesh(INDEX_NONE)
-		, BoundingboxVolume((T)-1.0)
-		, BoundingboxExtentMin((T)-1.0)
-		, BoundingboxExtentMax((T)-1.0)
-		, SurfaceType(-1)
-	{}
-
-	TCollisionDataExt(TVector<T, d> InLocation
-		, TVector<T, d> InAccumulatedImpulse
-		, TVector<T, d> InNormal
-		, TVector<T, d> InVelocity1
-		, TVector<T, d> InVelocity2
-		, TVector<T, d> InAngularVelocity1
-		, TVector<T, d> InAngularVelocity2
-		, T InMass1
-		, T InMass2
-		, TGeometryParticle<T, d>* InParticle
-		, TGeometryParticle<T, d>* InLevelset
-		, int32 InParticleIndexMesh
-		, int32 InLevelsetIndexMesh
-		, float InBoundingboxVolume
-		, float InBoundingboxExtentMin
-		, float InBoundingboxExtentMax
-		, int32 InSurfaceType)
-		: Location(InLocation)
-		, AccumulatedImpulse(InAccumulatedImpulse)
-		, Normal(InNormal)
-		, Velocity1(InVelocity1)
-		, Velocity2(InVelocity2)
-		, AngularVelocity1(InAngularVelocity1)
-		, AngularVelocity2(InAngularVelocity2)
-		, Mass1(InMass1)
-		, Mass2(InMass2)
-		, Particle(InParticle)
-		, Levelset(InLevelset)
-		, ParticleIndexMesh(InParticleIndexMesh)
-		, LevelsetIndexMesh(InLevelsetIndexMesh)
-		, BoundingboxVolume(InBoundingboxVolume)
-		, BoundingboxExtentMin(InBoundingboxExtentMin)
-		, BoundingboxExtentMax(InBoundingboxExtentMax)
-		, SurfaceType(InSurfaceType)
-	{}
-
-	TCollisionDataExt(const TCollisionData<T, d>& InCollisionData)
-		: Location(InCollisionData.Location)
-		, AccumulatedImpulse(InCollisionData.AccumulatedImpulse)
-		, Normal(InCollisionData.Normal)
-		, Velocity1(InCollisionData.Velocity1)
-		, Velocity2(InCollisionData.Velocity2)
-		, AngularVelocity1(InCollisionData.AngularVelocity1)
-		, AngularVelocity2(InCollisionData.AngularVelocity2)
-		, Mass1(InCollisionData.Mass1)
-		, Mass2(InCollisionData.Mass2)
-		, Particle(InCollisionData.Particle)
-		, Levelset(InCollisionData.Levelset)
-		, ParticleIndexMesh(InCollisionData.ParticleIndexMesh)
-		, LevelsetIndexMesh(InCollisionData.LevelsetIndexMesh)
-		, BoundingboxVolume((T)-1.0)
-		, BoundingboxExtentMin((T)-1.0)
-		, BoundingboxExtentMax((T)-1.0)
-		, SurfaceType(-1)
-	{}
-
-	TVector<T, d> Location;
-	TVector<T, d> AccumulatedImpulse;
-	TVector<T, d> Normal;
-	TVector<T, d> Velocity1, Velocity2;
-	TVector<T, d> AngularVelocity1, AngularVelocity2;
-	T Mass1, Mass2;
-	TGeometryParticle<T, d>* Particle;
-	TGeometryParticle<T, d>* Levelset;
-	int32 ParticleIndexMesh, LevelsetIndexMesh;
-	float BoundingboxVolume;
-	float BoundingboxExtentMin, BoundingboxExtentMax;
-	int32 SurfaceType;
-};
-
-/*
-BreakingData passed from the physics solver to subsystems
-*/
-template<class T, int d>
-struct TBreakingData
-{
-	TBreakingData()
-		: Location(TVector<T, d>((T)0.0))
-		, Velocity(TVector<T, d>((T)0.0))
-		, AngularVelocity(TVector<T, d>((T)0.0))
-		, Mass((T)0.0)
-		, Particle(nullptr)
-		, ParticleIndex(INDEX_NONE)
-		, ParticleIndexMesh(INDEX_NONE)
-		, BoundingBox(TBox<T, d>(TVector<T, d>((T)0.0), TVector<T, d>((T)0.0)))
-	{}
-
-	TVector<T, d> Location;
-	TVector<T, d> Velocity;
-	TVector<T, d> AngularVelocity;
-	T Mass;
-	TGeometryParticle<T, d>* Particle;
-	int32 ParticleIndex; //#todo: remove this in favor of TGeometryParticle?
-	int32 ParticleIndexMesh; // If ParticleIndex points to a cluster then this index will point to an actual mesh in the cluster
-							 // It is important to be able to get extra data from the component
-	Chaos::TBox<T, d> BoundingBox;
-};
-
-
-
-/*
-BreakingData used in the subsystems
-*/
-template<class T, int d>
-struct TBreakingDataExt
-{
-	TBreakingDataExt()
-		: Location(TVector<T, d>((T)0.0))
-		, Velocity(TVector<T, d>((T)0.0))
-		, AngularVelocity(TVector<T, d>((T)0.0))
-		, Mass((T)0.0)
-		, ParticleIndex(INDEX_NONE)
-		, ParticleIndexMesh(INDEX_NONE)
-		, BoundingboxVolume((T)-1.0)
-		, BoundingboxExtentMin((T)-1.0)
-		, BoundingboxExtentMax((T)-1.0)
-		, SurfaceType(-1)
-	{}
-
-	TBreakingDataExt(TVector<T, d> InLocation
-		, TVector<T, d> InVelocity
-		, TVector<T, d> InAngularVelocity
-		, T InMass
-		, int32 InParticleIndex
-		, int32 InParticleIndexMesh
-		, float InBoundingboxVolume
-		, float InBoundingboxExtentMin
-		, float InBoundingboxExtentMax
-		, int32 InSurfaceType)
-		: Location(InLocation)
-		, Velocity(InVelocity)
-		, AngularVelocity(InAngularVelocity)
-		, Mass(InMass)
-		, ParticleIndex(InParticleIndex)
-		, ParticleIndexMesh(InParticleIndexMesh)
-		, BoundingboxVolume(InBoundingboxVolume)
-		, BoundingboxExtentMin(InBoundingboxExtentMin)
-		, BoundingboxExtentMax(InBoundingboxExtentMax)
-		, SurfaceType(InSurfaceType)
-	{}
-
-	TBreakingDataExt(const TBreakingData<float, 3>& InBreakingData)
-		: Location(InBreakingData.Location)
-		, Velocity(InBreakingData.Velocity)
-		, AngularVelocity(InBreakingData.AngularVelocity)
-		, Mass(InBreakingData.Mass)
-		, ParticleIndex(InBreakingData.ParticleIndex)
-		, ParticleIndexMesh(InBreakingData.ParticleIndexMesh)
-		, BoundingboxVolume((T)-1.0)
-		, BoundingboxExtentMin((T)-1.0)
-		, BoundingboxExtentMax((T)-1.0)
-		, SurfaceType(-1)
-	{}
-
-	TVector<T, d> Location;
-	TVector<T, d> Velocity;
-	TVector<T, d> AngularVelocity;
-	T Mass;
-	int32 ParticleIndex;
-	int32 ParticleIndexMesh; // If ParticleIndex points to a cluster then this index will point to an actual mesh in the cluster
-							 // It is important to be able to get extra data from the component
-	float BoundingboxVolume;
-	float BoundingboxExtentMin, BoundingboxExtentMax;
-	int32 SurfaceType;
-
-	FVector TransformTranslation;
-	FQuat TransformRotation;
-	FVector TransformScale;
-
-	FBox BoundingBox;
-
-	// Please don't be tempted to add the code below back. Holding onto a UObject pointer without the GC knowing about it is 
-	// not a safe thing to do.
-	//UPhysicalMaterial* PhysicalMaterialTest;
-	FName PhysicalMaterialName;
-};
-
-/*
-TrailingData passed from the physics solver to subsystems
-*/
-template<class T, int d>
-struct TTrailingData
-{
-	TTrailingData()
-		: Location(TVector<T, d>((T)0.0))
-		, Velocity(TVector<T, d>((T)0.0))
-		, AngularVelocity(TVector<T, d>((T)0.0))
-		, Mass((T)0.0)
-		, Particle(nullptr)
-		, ParticleIndexMesh(INDEX_NONE)
-		, BoundingBox(TBox<T, d>(TVector<T, d>((T)0.0), TVector<T, d>((T)0.0)))
-	{}
-
-	TTrailingData(TVector<T, d> InLocation
-		, TVector<T, d> InVelocity
-		, TVector<T, d> InAngularVelocity
-		, T InMass
-		, TGeometryParticle<T, d>* InParticle
-		, int32 InParticleIndexMesh
-		, float InBoundingboxVolume
-		, float InBoundingboxExtentMin
-		, float InBoundingboxExtentMax
-		, int32 InSurfaceType
-		, Chaos::TBox<T, d>& InBoundingBox)
-		: Location(InLocation)
-		, Velocity(InVelocity)
-		, AngularVelocity(InAngularVelocity)
-		, Mass(InMass)
-		, Particle(InParticle)
-		, ParticleIndexMesh(InParticleIndexMesh)
-		, BoundingBox(InBoundingBox)
-	{}
-
-	TVector<T, d> Location;
-	TVector<T, d> Velocity;
-	TVector<T, d> AngularVelocity;
-	T Mass;
-	TGeometryParticle<T, d>* Particle;
-	int32 ParticleIndexMesh; // If ParticleIndex points to a cluster then this index will point to an actual mesh in the cluster
-							 // It is important to be able to get extra data from the component
-	Chaos::TBox<T, d> BoundingBox;
-
-	friend inline uint32 GetTypeHash(const TTrailingData& Other)
+		void SetLocation(const TVector<T, d> & InLocation, const FImplicitObject* B = nullptr, const FImplicitObject* A = nullptr, int32 Index=0) { ShapeManifold.Manifold.Location = InLocation; }
+		TVector<T, d> GetLocation(const FImplicitObject* B = nullptr, const FImplicitObject* A = nullptr) const { return ShapeManifold.Manifold.Location; }
+
+		void SetPhi(T InPhi, const FImplicitObject* B = nullptr, const FImplicitObject* A = nullptr, int32 Index = 0) { ShapeManifold.Manifold.Phi = InPhi; }
+		T GetPhi(const FImplicitObject* B = nullptr, const FImplicitObject* A = nullptr) const { return ShapeManifold.Manifold.Phi; }
+
+		void AddManifold(const FImplicitObject* B = nullptr, const FImplicitObject* A = nullptr) { ShapeManifold.Implicit[0] = A; ShapeManifold.Implicit[1] = B; }
+
+		FConvexManifold ShapeManifold;
+		FGeometryParticleHandle* Particle;
+		FGeometryParticleHandle* Levelset;
+		TVector<T, d> AccumulatedImpulse;
+
+		FString ToString() const
+		{
+			return FString::Printf(TEXT("Particle:%s, Levelset:%s, AccumulatedImpulse:%s"), *Particle->ToString(), *Levelset->ToString(), *AccumulatedImpulse.ToString());
+		}
+
+	};
+
+	/*
+	template<class T, int d>
+	struct TCollisionResolutionManifold
 	{
-		return ::GetTypeHash(Other.Particle);
-	}
+		bool ContainsShapeConnection(const FImplicitObject* Implicit0, const FImplicitObject* Implicit1) { return false; }
 
-	friend bool operator==(const TTrailingData& A, const TTrailingData& B)
+		TRigidBodyContactConvexConstraint2<T, d> A;
+	};
+	*/
+
+	/**
+	// @todo(ccaulfield): should these be held SOA style in the PBDCollisionConstraint?
+	template<class T, int d>
+	struct TRigidBodyContactConstraint
 	{
-		return A.Particle == B.Particle;
-	}
-};
-
-/*
-TrailingData used in subsystems
-*/
-template<class T, int d>
-struct TTrailingDataExt
-{
-	TTrailingDataExt()
-		: Location(TVector<T, d>((T)0.0))
-		, Velocity(TVector<T, d>((T)0.0))
-		, AngularVelocity(TVector<T, d>((T)0.0))
-		, Mass((T)0.0)
-		, Particle(nullptr)
-		, ParticleIndexMesh(INDEX_NONE)
-		, BoundingboxVolume((T)-1.0)
-		, BoundingboxExtentMin((T)-1.0)
-		, BoundingboxExtentMax((T)-1.0)
-		, SurfaceType(-1)
-	{}
-
-	TTrailingDataExt(TVector<T, d> InLocation
-		, TVector<T, d> InVelocity
-		, TVector<T, d> InAngularVelocity
-		, T InMass
-		, TGeometryParticle<T, d>* InParticle
-		, int32 InParticleIndexMesh
-		, float InBoundingboxVolume
-		, float InBoundingboxExtentMin
-		, float InBoundingboxExtentMax
-		, int32 InSurfaceType)
-		: Location(InLocation)
-		, Velocity(InVelocity)
-		, AngularVelocity(InAngularVelocity)
-		, Mass(InMass)
-		, Particle(InParticle)
-		, ParticleIndexMesh(InParticleIndexMesh)
-		, BoundingboxVolume(InBoundingboxVolume)
-		, BoundingboxExtentMin(InBoundingboxExtentMin)
-		, BoundingboxExtentMax(InBoundingboxExtentMax)
-		, SurfaceType(InSurfaceType)
-	{}
-
-	TTrailingDataExt(const TTrailingData<float, 3>& InTrailingData)
-		: Location(InTrailingData.Location)
-		, Velocity(InTrailingData.Velocity)
-		, AngularVelocity(InTrailingData.AngularVelocity)
-		, Mass(InTrailingData.Mass)
-		, Particle(InTrailingData.Particle)
-		, ParticleIndexMesh(InTrailingData.ParticleIndexMesh)
-		, BoundingboxVolume((T)-1.0)
-		, BoundingboxExtentMin((T)-1.0)
-		, BoundingboxExtentMax((T)-1.0)
-		, SurfaceType(-1)
-	{}
-
-	TVector<T, d> Location;
-	TVector<T, d> Velocity;
-	TVector<T, d> AngularVelocity;
-	T Mass;
-	TGeometryParticle<T, d>* Particle;
-	int32 ParticleIndexMesh; // If ParticleIndex points to a cluster then this index will point to an actual mesh in the cluster
-							 // It is important to be able to get extra data from the component
-	float BoundingboxVolume;
-	float BoundingboxExtentMin, BoundingboxExtentMax;
-	int32 SurfaceType;
-
-	friend inline uint32 GetTypeHash(const TTrailingDataExt& Other)
+		TRigidBodyContactConstraint() : bDisabled(false), Normal(0), AccumulatedImpulse(0.f) {}
+		TGeometryParticleHandle<T, d>* Particle;
+		TGeometryParticleHandle<T, d>* Levelset;
+		const FImplicitObject* Geometry[2];
+
+		bool bDisabled;
+		TVector<T, d> Normal;
+		TVector<T, d> Location;
+		TVector<T, d> LocalLocation;
+		T Phi;
+		TVector<T, d> AccumulatedImpulse;
+
+		FString ToString() const
+		{
+			return FString::Printf(TEXT("Particle:%s, Levelset:%s, Normal:%s, Location:%s, Phi:%f, AccumulatedImpulse:%s"), *Particle->ToString(), *Levelset->ToString(), *Normal.ToString(), *Location.ToString(), Phi, *AccumulatedImpulse.ToString());
+		}
+	};
+	*/
+
+	template<class T, int d>
+	struct TRigidBodyContactConstraintPGS
 	{
-		return ::GetTypeHash(Other.Particle);
-	}
-
-	friend bool operator==(const TTrailingDataExt& A, const TTrailingDataExt& B)
+		TRigidBodyContactConstraintPGS() : AccumulatedImpulse(0.f) {}
+		TGeometryParticleHandle<T, d>* Particle;
+		TGeometryParticleHandle<T, d>* Levelset;
+		TArray<TVector<T, d>> Normal;
+		TArray<TVector<T, d>> Location;
+		TArray<T> Phi;
+		TVector<T, d> AccumulatedImpulse;
+	};
+
+
+	template<class T, int d>
+	class CHAOS_API TPBDCollisionConstraintHandle : public TContainerConstraintHandle<TPBDCollisionConstraint<T, d>>
 	{
-		return A.Particle == B.Particle;
-	}
-};
+	public:
+		using Base = TContainerConstraintHandle<TPBDCollisionConstraint<T, d>>;
+		using FConstraintContainer = TPBDCollisionConstraint<T, d>;
 
+		TPBDCollisionConstraintHandle()
+		{}
 
-template<class T, int d>
-struct TSleepingData
-{
-	TSleepingData()
-		: Particle(nullptr)
-		, Sleeping(true)
-	{}
-
-	TSleepingData(
-		  TGeometryParticle<T, d>* InParticle
-		, bool InSleeping) 
-		: Particle(InParticle)
-		, Sleeping(InSleeping)
-	{}
-
-	TGeometryParticle<T, d>* Particle;
-	bool Sleeping;	// if !Sleeping == Awake
-};
+		TPBDCollisionConstraintHandle(FConstraintContainer* InConstraintContainer, int32 InConstraintIndex)
+			: TContainerConstraintHandle<TPBDCollisionConstraint<T, d>>(InConstraintContainer, InConstraintIndex)
+		{
+		}
+
+		TRigidBodyContactConstraint<T, d>& GetContact()
+		{
+			return ConstraintContainer->Constraints[ConstraintIndex];
+		}
+
+		const TRigidBodyContactConstraint<T, d>& GetContact() const
+		{
+			return ConstraintContainer->Constraints[ConstraintIndex];
+		}
+
+		void SetConstraintIndex(int32 IndexIn)
+		{
+			ConstraintIndex = IndexIn;
+		}
+
+	protected:
+		using Base::ConstraintIndex;
+		using Base::ConstraintContainer;
+
+
+	};
+
+
+	/**
+	 * Collision event data stored for use by other systems (e.g. Niagara, gameplay events)
+	 */
+	template<class T, int d>
+	struct TCollisionData
+	{
+		TCollisionData()
+			: Location(TVector<T, d>((T)0.0))
+			, AccumulatedImpulse(TVector<T, d>((T)0.0))
+			, Normal(TVector<T, d>((T)0.0))
+			, Velocity1(TVector<T, d>((T)0.0))
+			, Velocity2(TVector<T, d>((T)0.0))
+			, AngularVelocity1(TVector<T, d>((T)0.0))
+			, AngularVelocity2(TVector<T, d>((T)0.0))
+			, Mass1((T)0.0)
+			, Mass2((T)0.0)
+			, PenetrationDepth((T)0.0)
+			, Particle(nullptr)
+			, Levelset(nullptr)
+			, ParticleIndexMesh(INDEX_NONE)
+			, LevelsetIndexMesh(INDEX_NONE)
+		{}
+
+		TCollisionData(TVector<T, d> InLocation
+			, TVector<T, d> InAccumulatedImpulse
+			, TVector<T, d> InNormal
+			, TVector<T, d> InVelocity1
+			, TVector<T, d> InVelocity2
+			, TVector<T, d> InAngularVelocity1
+			, TVector<T, d> InAngularVelocity2
+			, T InMass1
+			, T InMass2
+			, T InPenetrationDepth
+			, TGeometryParticle<T, d>* InParticle
+			, TGeometryParticle<T, d>* InLevelset
+			, int32 InParticleIndexMesh
+			, int32 InLevelsetIndexMesh)
+			: Location(InLocation)
+			, AccumulatedImpulse(InAccumulatedImpulse)
+			, Normal(InNormal)
+			, Velocity1(InVelocity1)
+			, Velocity2(InVelocity2)
+			, AngularVelocity1(InAngularVelocity1)
+			, AngularVelocity2(InAngularVelocity2)
+			, Mass1(InMass1)
+			, Mass2(InMass2)
+			, PenetrationDepth(InPenetrationDepth)
+			, Particle(InParticle)
+			, Levelset(InLevelset)
+			, ParticleIndexMesh(InParticleIndexMesh)
+			, LevelsetIndexMesh(InLevelsetIndexMesh)
+		{}
+
+		TVector<T, d> Location;
+		TVector<T, d> AccumulatedImpulse;
+		TVector<T, d> Normal;
+		TVector<T, d> Velocity1, Velocity2;
+		TVector<T, d> AngularVelocity1, AngularVelocity2;
+		T Mass1, Mass2;
+		T PenetrationDepth;
+		TGeometryParticle<T, d>* Particle;
+		TGeometryParticle<T, d>* Levelset;
+		// @todo(ccaulfield): CHAOS_PARTICLEHANDLE_TODO
+		int32 ParticleIndexMesh, LevelsetIndexMesh; // If ParticleIndex points to a cluster then this index will point to an actual mesh in the cluster
+													// It is important to be able to get extra data from the component
+
+	};
+
+
+	/*
+	CollisionData used in the subsystems
+	*/
+	template<class T, int d>
+	struct TCollisionDataExt
+	{
+		TCollisionDataExt()
+			: Location(TVector<T, d>((T)0.0))
+			, AccumulatedImpulse(TVector<T, d>((T)0.0))
+			, Normal(TVector<T, d>((T)0.0))
+			, Velocity1(TVector<T, d>((T)0.0))
+			, Velocity2(TVector<T, d>((T)0.0))
+			, AngularVelocity1(TVector<T, d>((T)0.0))
+			, AngularVelocity2(TVector<T, d>((T)0.0))
+			, Mass1((T)0.0)
+			, Mass2((T)0.0)
+			, Particle(nullptr)
+			, Levelset(nullptr)
+			, ParticleIndexMesh(INDEX_NONE)
+			, LevelsetIndexMesh(INDEX_NONE)
+			, BoundingboxVolume((T)-1.0)
+			, BoundingboxExtentMin((T)-1.0)
+			, BoundingboxExtentMax((T)-1.0)
+			, SurfaceType(-1)
+		{}
+
+		TCollisionDataExt(TVector<T, d> InLocation
+			, TVector<T, d> InAccumulatedImpulse
+			, TVector<T, d> InNormal
+			, TVector<T, d> InVelocity1
+			, TVector<T, d> InVelocity2
+			, TVector<T, d> InAngularVelocity1
+			, TVector<T, d> InAngularVelocity2
+			, T InMass1
+			, T InMass2
+			, TGeometryParticle<T, d>* InParticle
+			, TGeometryParticle<T, d>* InLevelset
+			, int32 InParticleIndexMesh
+			, int32 InLevelsetIndexMesh
+			, float InBoundingboxVolume
+			, float InBoundingboxExtentMin
+			, float InBoundingboxExtentMax
+			, int32 InSurfaceType)
+			: Location(InLocation)
+			, AccumulatedImpulse(InAccumulatedImpulse)
+			, Normal(InNormal)
+			, Velocity1(InVelocity1)
+			, Velocity2(InVelocity2)
+			, AngularVelocity1(InAngularVelocity1)
+			, AngularVelocity2(InAngularVelocity2)
+			, Mass1(InMass1)
+			, Mass2(InMass2)
+			, Particle(InParticle)
+			, Levelset(InLevelset)
+			, ParticleIndexMesh(InParticleIndexMesh)
+			, LevelsetIndexMesh(InLevelsetIndexMesh)
+			, BoundingboxVolume(InBoundingboxVolume)
+			, BoundingboxExtentMin(InBoundingboxExtentMin)
+			, BoundingboxExtentMax(InBoundingboxExtentMax)
+			, SurfaceType(InSurfaceType)
+		{}
+
+		TCollisionDataExt(const TCollisionData<T, d>& InCollisionData)
+			: Location(InCollisionData.Location)
+			, AccumulatedImpulse(InCollisionData.AccumulatedImpulse)
+			, Normal(InCollisionData.Normal)
+			, Velocity1(InCollisionData.Velocity1)
+			, Velocity2(InCollisionData.Velocity2)
+			, AngularVelocity1(InCollisionData.AngularVelocity1)
+			, AngularVelocity2(InCollisionData.AngularVelocity2)
+			, Mass1(InCollisionData.Mass1)
+			, Mass2(InCollisionData.Mass2)
+			, Particle(InCollisionData.Particle)
+			, Levelset(InCollisionData.Levelset)
+			, ParticleIndexMesh(InCollisionData.ParticleIndexMesh)
+			, LevelsetIndexMesh(InCollisionData.LevelsetIndexMesh)
+			, BoundingboxVolume((T)-1.0)
+			, BoundingboxExtentMin((T)-1.0)
+			, BoundingboxExtentMax((T)-1.0)
+			, SurfaceType(-1)
+		{}
+
+		TVector<T, d> Location;
+		TVector<T, d> AccumulatedImpulse;
+		TVector<T, d> Normal;
+		TVector<T, d> Velocity1, Velocity2;
+		TVector<T, d> AngularVelocity1, AngularVelocity2;
+		T Mass1, Mass2;
+		TGeometryParticle<T, d>* Particle;
+		TGeometryParticle<T, d>* Levelset;
+		int32 ParticleIndexMesh, LevelsetIndexMesh;
+		float BoundingboxVolume;
+		float BoundingboxExtentMin, BoundingboxExtentMax;
+		int32 SurfaceType;
+	};
+
+	/*
+	BreakingData passed from the physics solver to subsystems
+	*/
+	template<class T, int d>
+	struct TBreakingData
+	{
+		TBreakingData()
+			: Location(TVector<T, d>((T)0.0))
+			, Velocity(TVector<T, d>((T)0.0))
+			, AngularVelocity(TVector<T, d>((T)0.0))
+			, Mass((T)0.0)
+			, Particle(nullptr)
+			, ParticleIndex(INDEX_NONE)
+			, ParticleIndexMesh(INDEX_NONE)
+			, BoundingBox(TBox<T, d>(TVector<T, d>((T)0.0), TVector<T, d>((T)0.0)))
+		{}
+
+		TVector<T, d> Location;
+		TVector<T, d> Velocity;
+		TVector<T, d> AngularVelocity;
+		T Mass;
+		TGeometryParticle<T, d>* Particle;
+		int32 ParticleIndex; //#todo: remove this in favor of TGeometryParticle?
+		int32 ParticleIndexMesh; // If ParticleIndex points to a cluster then this index will point to an actual mesh in the cluster
+								 // It is important to be able to get extra data from the component
+		Chaos::TBox<T, d> BoundingBox;
+	};
+
+
+
+	/*
+	BreakingData used in the subsystems
+	*/
+	template<class T, int d>
+	struct TBreakingDataExt
+	{
+		TBreakingDataExt()
+			: Location(TVector<T, d>((T)0.0))
+			, Velocity(TVector<T, d>((T)0.0))
+			, AngularVelocity(TVector<T, d>((T)0.0))
+			, Mass((T)0.0)
+			, ParticleIndex(INDEX_NONE)
+			, ParticleIndexMesh(INDEX_NONE)
+			, BoundingboxVolume((T)-1.0)
+			, BoundingboxExtentMin((T)-1.0)
+			, BoundingboxExtentMax((T)-1.0)
+			, SurfaceType(-1)
+		{}
+
+		TBreakingDataExt(TVector<T, d> InLocation
+			, TVector<T, d> InVelocity
+			, TVector<T, d> InAngularVelocity
+			, T InMass
+			, int32 InParticleIndex
+			, int32 InParticleIndexMesh
+			, float InBoundingboxVolume
+			, float InBoundingboxExtentMin
+			, float InBoundingboxExtentMax
+			, int32 InSurfaceType)
+			: Location(InLocation)
+			, Velocity(InVelocity)
+			, AngularVelocity(InAngularVelocity)
+			, Mass(InMass)
+			, ParticleIndex(InParticleIndex)
+			, ParticleIndexMesh(InParticleIndexMesh)
+			, BoundingboxVolume(InBoundingboxVolume)
+			, BoundingboxExtentMin(InBoundingboxExtentMin)
+			, BoundingboxExtentMax(InBoundingboxExtentMax)
+			, SurfaceType(InSurfaceType)
+		{}
+
+		TBreakingDataExt(const TBreakingData<float, 3>& InBreakingData)
+			: Location(InBreakingData.Location)
+			, Velocity(InBreakingData.Velocity)
+			, AngularVelocity(InBreakingData.AngularVelocity)
+			, Mass(InBreakingData.Mass)
+			, ParticleIndex(InBreakingData.ParticleIndex)
+			, ParticleIndexMesh(InBreakingData.ParticleIndexMesh)
+			, BoundingboxVolume((T)-1.0)
+			, BoundingboxExtentMin((T)-1.0)
+			, BoundingboxExtentMax((T)-1.0)
+			, SurfaceType(-1)
+		{}
+
+		TVector<T, d> Location;
+		TVector<T, d> Velocity;
+		TVector<T, d> AngularVelocity;
+		T Mass;
+		int32 ParticleIndex;
+		int32 ParticleIndexMesh; // If ParticleIndex points to a cluster then this index will point to an actual mesh in the cluster
+								 // It is important to be able to get extra data from the component
+		float BoundingboxVolume;
+		float BoundingboxExtentMin, BoundingboxExtentMax;
+		int32 SurfaceType;
+
+		FVector TransformTranslation;
+		FQuat TransformRotation;
+		FVector TransformScale;
+
+		FBox BoundingBox;
+
+		// Please don't be tempted to add the code below back. Holding onto a UObject pointer without the GC knowing about it is 
+		// not a safe thing to do.
+		//UPhysicalMaterial* PhysicalMaterialTest;
+		FName PhysicalMaterialName;
+	};
+
+	/*
+	TrailingData passed from the physics solver to subsystems
+	*/
+	template<class T, int d>
+	struct TTrailingData
+	{
+		TTrailingData()
+			: Location(TVector<T, d>((T)0.0))
+			, Velocity(TVector<T, d>((T)0.0))
+			, AngularVelocity(TVector<T, d>((T)0.0))
+			, Mass((T)0.0)
+			, Particle(nullptr)
+			, ParticleIndexMesh(INDEX_NONE)
+			, BoundingBox(TBox<T, d>(TVector<T, d>((T)0.0), TVector<T, d>((T)0.0)))
+		{}
+
+		TTrailingData(TVector<T, d> InLocation
+			, TVector<T, d> InVelocity
+			, TVector<T, d> InAngularVelocity
+			, T InMass
+			, TGeometryParticle<T, d>* InParticle
+			, int32 InParticleIndexMesh
+			, float InBoundingboxVolume
+			, float InBoundingboxExtentMin
+			, float InBoundingboxExtentMax
+			, int32 InSurfaceType
+			, Chaos::TBox<T, d>& InBoundingBox)
+			: Location(InLocation)
+			, Velocity(InVelocity)
+			, AngularVelocity(InAngularVelocity)
+			, Mass(InMass)
+			, Particle(InParticle)
+			, ParticleIndexMesh(InParticleIndexMesh)
+			, BoundingBox(InBoundingBox)
+		{}
+
+		TVector<T, d> Location;
+		TVector<T, d> Velocity;
+		TVector<T, d> AngularVelocity;
+		T Mass;
+		TGeometryParticle<T, d>* Particle;
+		int32 ParticleIndexMesh; // If ParticleIndex points to a cluster then this index will point to an actual mesh in the cluster
+								 // It is important to be able to get extra data from the component
+		Chaos::TBox<T, d> BoundingBox;
+
+		friend inline uint32 GetTypeHash(const TTrailingData& Other)
+		{
+			return ::GetTypeHash(Other.Particle);
+		}
+
+		friend bool operator==(const TTrailingData& A, const TTrailingData& B)
+		{
+			return A.Particle == B.Particle;
+		}
+	};
+
+	/*
+	TrailingData used in subsystems
+	*/
+	template<class T, int d>
+	struct TTrailingDataExt
+	{
+		TTrailingDataExt()
+			: Location(TVector<T, d>((T)0.0))
+			, Velocity(TVector<T, d>((T)0.0))
+			, AngularVelocity(TVector<T, d>((T)0.0))
+			, Mass((T)0.0)
+			, Particle(nullptr)
+			, ParticleIndexMesh(INDEX_NONE)
+			, BoundingboxVolume((T)-1.0)
+			, BoundingboxExtentMin((T)-1.0)
+			, BoundingboxExtentMax((T)-1.0)
+			, SurfaceType(-1)
+		{}
+
+		TTrailingDataExt(TVector<T, d> InLocation
+			, TVector<T, d> InVelocity
+			, TVector<T, d> InAngularVelocity
+			, T InMass
+			, TGeometryParticle<T, d>* InParticle
+			, int32 InParticleIndexMesh
+			, float InBoundingboxVolume
+			, float InBoundingboxExtentMin
+			, float InBoundingboxExtentMax
+			, int32 InSurfaceType)
+			: Location(InLocation)
+			, Velocity(InVelocity)
+			, AngularVelocity(InAngularVelocity)
+			, Mass(InMass)
+			, Particle(InParticle)
+			, ParticleIndexMesh(InParticleIndexMesh)
+			, BoundingboxVolume(InBoundingboxVolume)
+			, BoundingboxExtentMin(InBoundingboxExtentMin)
+			, BoundingboxExtentMax(InBoundingboxExtentMax)
+			, SurfaceType(InSurfaceType)
+		{}
+
+		TTrailingDataExt(const TTrailingData<float, 3>& InTrailingData)
+			: Location(InTrailingData.Location)
+			, Velocity(InTrailingData.Velocity)
+			, AngularVelocity(InTrailingData.AngularVelocity)
+			, Mass(InTrailingData.Mass)
+			, Particle(InTrailingData.Particle)
+			, ParticleIndexMesh(InTrailingData.ParticleIndexMesh)
+			, BoundingboxVolume((T)-1.0)
+			, BoundingboxExtentMin((T)-1.0)
+			, BoundingboxExtentMax((T)-1.0)
+			, SurfaceType(-1)
+		{}
+
+		TVector<T, d> Location;
+		TVector<T, d> Velocity;
+		TVector<T, d> AngularVelocity;
+		T Mass;
+		TGeometryParticle<T, d>* Particle;
+		int32 ParticleIndexMesh; // If ParticleIndex points to a cluster then this index will point to an actual mesh in the cluster
+								 // It is important to be able to get extra data from the component
+		float BoundingboxVolume;
+		float BoundingboxExtentMin, BoundingboxExtentMax;
+		int32 SurfaceType;
+
+		friend inline uint32 GetTypeHash(const TTrailingDataExt& Other)
+		{
+			return ::GetTypeHash(Other.Particle);
+		}
+
+		friend bool operator==(const TTrailingDataExt& A, const TTrailingDataExt& B)
+		{
+			return A.Particle == B.Particle;
+		}
+	};
+
+
+	template<class T, int d>
+	struct TSleepingData
+	{
+		TSleepingData()
+			: Particle(nullptr)
+			, Sleeping(true)
+		{}
+
+		TSleepingData(
+			TGeometryParticle<T, d>* InParticle
+			, bool InSleeping)
+			: Particle(InParticle)
+			, Sleeping(InSleeping)
+		{}
+
+		TGeometryParticle<T, d>* Particle;
+		bool Sleeping;	// if !Sleeping == Awake
+	};
 
 
 }
diff --git a/Engine/Source/Runtime/Experimental/ChaosSolvers/Private/EventDefaults.cpp b/Engine/Source/Runtime/Experimental/ChaosSolvers/Private/EventDefaults.cpp
--- a/Engine/Source/Runtime/Experimental/ChaosSolvers/Private/EventDefaults.cpp
+++ b/Engine/Source/Runtime/Experimental/ChaosSolvers/Private/EventDefaults.cpp
@@ -72,7 +72,7 @@ namespace Chaos
 					if(Constraint.Particle->Handle()->GTGeometryParticle()->Proxy != nullptr)
 
 					{
-						if(ensure(!Constraint.AccumulatedImpulse.ContainsNaN() && FMath::IsFinite(Constraint.Phi)))
+						if(ensure(!Constraint.AccumulatedImpulse.ContainsNaN() && FMath::IsFinite(Constraint.GetPhi())))
 						{
 							TGeometryParticleHandle<float, 3>* Particle0 = Constraint.Particle;
 							TGeometryParticleHandle<float, 3>* Particle1 = Constraint.Levelset;
@@ -83,8 +83,8 @@ namespace Chaos
 
 							if(!Constraint.AccumulatedImpulse.IsZero() && Body0)
 							{
-								if(ensure(!Constraint.Location.ContainsNaN() &&
-									!Constraint.Normal.ContainsNaN()) &&
+								if(ensure(!Constraint.GetLocation().ContainsNaN() &&
+									!Constraint.GetNormal().ContainsNaN()) &&
 									!Body0->V().ContainsNaN() &&
 									!Body0->W().ContainsNaN() &&
 									(Body1 == nullptr || ((!Body1->V().ContainsNaN()) && !Body1->W().ContainsNaN())))
@@ -109,10 +109,10 @@ namespace Chaos
 						TGeometryParticleHandle<float, 3>* Particle1 = Constraint.Levelset;
 
 						TCollisionData<float, 3> Data;
-						Data.Location = Constraint.Location;
+						Data.Location = Constraint.GetLocation();
 						Data.AccumulatedImpulse = Constraint.AccumulatedImpulse;
-						Data.Normal = Constraint.Normal;
-						Data.PenetrationDepth = Constraint.Phi;
+						Data.Normal = Constraint.GetNormal();
+						Data.PenetrationDepth = Constraint.GetPhi();
 						Data.Particle = Particle0->Handle()->GTGeometryParticle();
 						Data.Levelset = Particle1->Handle()->GTGeometryParticle();
 
