-Fix Typo

diff --git a/Engine/Source/Runtime/Experimental/Chaos/Private/Chaos/PBDCollisionConstraint.cpp b/Engine/Source/Runtime/Experimental/Chaos/Private/Chaos/PBDCollisionConstraint.cpp
--- a/Engine/Source/Runtime/Experimental/Chaos/Private/Chaos/PBDCollisionConstraint.cpp
+++ b/Engine/Source/Runtime/Experimental/Chaos/Private/Chaos/PBDCollisionConstraint.cpp
@@ -1683,7 +1683,6 @@ void TPBDCollisionConstraint<T, d>::UpdateLevelsetConstraint(const T Thickness,
 
 	if(SampleParticles)
 	{
-		asdf
 		SampleObject2<UpdateType>(*Particle1->Geometry(), LevelsetTM, *SampleParticles, ParticlesTM, Thickness, Constraint);
 	}
 }
