Chaos : Collisions

Removed empty file.

diff --git a/Engine/Source/Runtime/Experimental/Chaos/Public/Chaos/CollisionResolutionConvexConvex.h b/Engine/Source/Runtime/Experimental/Chaos/Public/Chaos/CollisionResolutionConvexConvex.h
--- a/Engine/Source/Runtime/Experimental/Chaos/Public/Chaos/CollisionResolutionConvexConvex.h
+++ b/Engine/Source/Runtime/Experimental/Chaos/Public/Chaos/CollisionResolutionConvexConvex.h
@@ -1,7 +1,6 @@
 // Copyright 1998-2019 Epic Games, Inc. All Rights Reserved.
 #pragma once
 
-#include "Chaos/CollisionResolutionManifold.h"
 #include "Chaos/Convex.h"
 #include "Chaos/ParticleHandle.h"
 #include "Chaos/PBDCollisionTypes.h"
diff --git a/Engine/Source/Runtime/Experimental/Chaos/Public/Chaos/CollisionResolutionManifold.h b/Engine/Source/Runtime/Experimental/Chaos/Public/Chaos/CollisionResolutionManifold.h
deleted file mode 100644
diff --git a/Engine/Source/Runtime/Experimental/Chaos/Public/Chaos/PBDCollisionConstraint.h b/Engine/Source/Runtime/Experimental/Chaos/Public/Chaos/PBDCollisionConstraint.h
--- a/Engine/Source/Runtime/Experimental/Chaos/Public/Chaos/PBDCollisionConstraint.h
+++ b/Engine/Source/Runtime/Experimental/Chaos/Public/Chaos/PBDCollisionConstraint.h
@@ -2,7 +2,6 @@
 #pragma once
 
 #include "Chaos/ConstraintHandle.h"
-#include "Chaos/CollisionResolutionManifold.h"
 #include "Chaos/PBDCollisionTypes.h"
 #include "Chaos/PBDConstraintContainer.h"
 #include "Framework/BufferedData.h"
diff --git a/Engine/Source/Runtime/Experimental/Chaos/Public/Chaos/PBDCollisionTypes.h b/Engine/Source/Runtime/Experimental/Chaos/Public/Chaos/PBDCollisionTypes.h
--- a/Engine/Source/Runtime/Experimental/Chaos/Public/Chaos/PBDCollisionTypes.h
+++ b/Engine/Source/Runtime/Experimental/Chaos/Public/Chaos/PBDCollisionTypes.h
@@ -101,39 +101,6 @@ namespace Chaos
 
 	};
 
-	/*
-	template<class T, int d>
-	struct TCollisionResolutionManifold
-	{
-		bool ContainsShapeConnection(const FImplicitObject* Implicit0, const FImplicitObject* Implicit1) { return false; }
-
-		TRigidBodyContactConvexConstraint2<T, d> A;
-	};
-	*/
-
-	/**
-	// @todo(ccaulfield): should these be held SOA style in the PBDCollisionConstraint?
-	template<class T, int d>
-	struct TRigidBodyContactConstraint
-	{
-		TRigidBodyContactConstraint() : bDisabled(false), Normal(0), AccumulatedImpulse(0.f) {}
-		TGeometryParticleHandle<T, d>* Particle;
-		TGeometryParticleHandle<T, d>* Levelset;
-		const FImplicitObject* Geometry[2];
-
-		bool bDisabled;
-		TVector<T, d> Normal;
-		TVector<T, d> Location;
-		TVector<T, d> LocalLocation;
-		T Phi;
-		TVector<T, d> AccumulatedImpulse;
-
-		FString ToString() const
-		{
-			return FString::Printf(TEXT("Particle:%s, Levelset:%s, Normal:%s, Location:%s, Phi:%f, AccumulatedImpulse:%s"), *Particle->ToString(), *Levelset->ToString(), *Normal.ToString(), *Location.ToString(), Phi, *AccumulatedImpulse.ToString());
-		}
-	};
-	*/
 
 	template<class T, int d>
 	struct TRigidBodyContactConstraintPGS
