Chaos: Cache per shape bounding box.

diff --git a/Engine/Source/Runtime/Core/Public/UObject/ExternalPhysicsCustomObjectVersion.h b/Engine/Source/Runtime/Core/Public/UObject/ExternalPhysicsCustomObjectVersion.h
--- a/Engine/Source/Runtime/Core/Public/UObject/ExternalPhysicsCustomObjectVersion.h
+++ b/Engine/Source/Runtime/Core/Public/UObject/ExternalPhysicsCustomObjectVersion.h
@@ -57,6 +57,9 @@ struct CORE_API FExternalPhysicsCustomObjectVersion
 		// Only serialize internal acceleration structure queue and acceleration structure. No external/Async queues.
 		FlushEvolutionInternalAccelerationQueue,
 
+		// Serialize world space bounds of TPerShapeData
+		SerializeShapeWorldSpaceBounds,
+
 		// -----<new versions can be added above this line>-------------------------------------------------
 		VersionPlusOne,
 		LatestVersion = VersionPlusOne - 1
diff --git a/Engine/Source/Runtime/Experimental/Chaos/Private/Chaos/GeometryParticles.cpp b/Engine/Source/Runtime/Experimental/Chaos/Private/Chaos/GeometryParticles.cpp
--- a/Engine/Source/Runtime/Experimental/Chaos/Private/Chaos/GeometryParticles.cpp
+++ b/Engine/Source/Runtime/Experimental/Chaos/Private/Chaos/GeometryParticles.cpp
@@ -62,9 +62,21 @@ namespace Chaos
 	template <typename T, int d>
 	void TPerShapeData<T, d>::Serialize(FChaosArchive& Ar)
 	{
+		Ar.UsingCustomVersion(FExternalPhysicsCustomObjectVersion::GUID);
+
 		Ar << Geometry;
 		Ar << QueryData;
 		Ar << SimData;
+
+		if (Ar.CustomVer(FExternalPhysicsCustomObjectVersion::GUID) >= FExternalPhysicsCustomObjectVersion::SerializeShapeWorldSpaceBounds)
+		{
+			Ar << WorldSpaceInflatedShapeBounds;
+		}
+		else
+		{
+			// This should be set by particle serializing this TPerShapeData.
+			WorldSpaceInflatedShapeBounds = TAABB<FReal, 3>(FVec3(0.0f, 0.0f, 0.0f), FVec3(0.0f, 0.0f, 0.0f));
+		}
 	}
 
 
diff --git a/Engine/Source/Runtime/Experimental/Chaos/Public/Chaos/AABB.h b/Engine/Source/Runtime/Experimental/Chaos/Public/Chaos/AABB.h
--- a/Engine/Source/Runtime/Experimental/Chaos/Public/Chaos/AABB.h
+++ b/Engine/Source/Runtime/Experimental/Chaos/Public/Chaos/AABB.h
@@ -484,6 +484,12 @@ namespace Chaos
 		TVector<T, d> MMin, MMax;
 	};
 
+	FORCEINLINE FChaosArchive& operator<<(FChaosArchive& Ar, TAABB<FReal, 3>& AABB)
+	{
+		AABB.Serialize(Ar);
+		return Ar;
+	}
+
 	template<typename T>
 	struct TAABBSpecializeSamplingHelper<T, 2>
 	{
diff --git a/Engine/Source/Runtime/Experimental/Chaos/Public/Chaos/GeometryParticles.h b/Engine/Source/Runtime/Experimental/Chaos/Public/Chaos/GeometryParticles.h
--- a/Engine/Source/Runtime/Experimental/Chaos/Public/Chaos/GeometryParticles.h
+++ b/Engine/Source/Runtime/Experimental/Chaos/Public/Chaos/GeometryParticles.h
@@ -32,6 +32,7 @@ namespace Chaos
 		FCollisionFilterData SimData;
 		void* UserData;
 		TSerializablePtr<FImplicitObject> Geometry;
+		TAABB<FReal, 3> WorldSpaceInflatedShapeBounds;
 		
 
 		static TPerShapeData<T, d>* SerializationFactory(FChaosArchive& Ar, TPerShapeData<T, d>*);
@@ -262,9 +263,17 @@ namespace Chaos
 			return MWorldSpaceInflatedBounds[Index];
 		}
 
-		CHAOS_API TBox<T, d>& WorldSpaceInflatedBounds(const int32 Index)
+		CHAOS_API void SetWorldSpaceInflatedBounds(const int32 Index, const TBox<T, d>& Bounds)
 		{
-			return MWorldSpaceInflatedBounds[Index];
+			MWorldSpaceInflatedBounds[Index] = Bounds;
+
+			const TShapesArray<FReal, 3>& Shapes = ShapesArray(Index);
+			for (const auto& Shape : Shapes)
+			{
+				const TRigidTransform<FReal, 3> ActorTM(X(Index), R(Index));
+				TAABB<FReal, 3> Bound = Shape->Geometry->BoundingBox().GetAABB().TransformedAABB(ActorTM);
+				Shape->WorldSpaceInflatedShapeBounds = Bound;
+			}
 		}
 
 		const TArray<TSerializablePtr<FImplicitObject>>& GetAllGeometry() const { return MGeometry; }
@@ -304,6 +313,22 @@ namespace Chaos
 				Ar << MLocalBounds;
 				Ar << MWorldSpaceInflatedBounds;
 				Ar << MHasBounds;
+
+				if (Ar.CustomVer(FExternalPhysicsCustomObjectVersion::GUID) < FExternalPhysicsCustomObjectVersion::SerializeShapeWorldSpaceBounds)
+				{
+					// Assign shape bounds to actor bounds for backwards compat.
+					for (int32 Idx = 0; Idx < MGeometry.Num(); ++Idx)
+					{
+						if (MHasBounds[Idx])
+						{
+							auto& ShapesArray = MShapesArray[Idx];
+							for (auto &Shape : ShapesArray)
+							{
+								Shape->WorldSpaceInflatedShapeBounds = MWorldSpaceInflatedBounds[Idx].GetAABB();
+							}
+						}
+					}
+				}
 			}
 			else
 			{
@@ -315,6 +340,12 @@ namespace Chaos
 					{
 						MLocalBounds[Idx] = MGeometry[Idx]->BoundingBox();
 						MWorldSpaceInflatedBounds[Idx] = MLocalBounds[Idx].TransformedBox(TRigidTransform<T,d>(X(Idx), R(Idx)));	//ignore velocity too, really just trying to get something reasonable
+
+						auto& ShapesArray = MShapesArray[Idx];
+						for (auto &Shape : ShapesArray)
+						{
+							Shape->WorldSpaceInflatedShapeBounds = MWorldSpaceInflatedBounds[Idx].GetAABB();
+						}
 					}
 				}
 			}
diff --git a/Engine/Source/Runtime/Experimental/Chaos/Public/Chaos/ParticleHandle.h b/Engine/Source/Runtime/Experimental/Chaos/Public/Chaos/ParticleHandle.h
--- a/Engine/Source/Runtime/Experimental/Chaos/Public/Chaos/ParticleHandle.h
+++ b/Engine/Source/Runtime/Experimental/Chaos/Public/Chaos/ParticleHandle.h
@@ -380,7 +380,10 @@ public:
 	void SetLocalBounds(const TBox<T, d>& NewBounds) { GeometryParticles->LocalBounds(ParticleIdx) = NewBounds; }
 
 	const TBox<T, d>& WorldSpaceInflatedBounds() const { return GeometryParticles->WorldSpaceInflatedBounds(ParticleIdx); }
-	void SetWorldSpaceInflatedBounds(const TBox<T,d>& WorldSpaceInflatedBounds) { GeometryParticles->WorldSpaceInflatedBounds(ParticleIdx) = WorldSpaceInflatedBounds; }
+	void SetWorldSpaceInflatedBounds(const TBox<T, d>& WorldSpaceInflatedBounds)
+	{
+		GeometryParticles->SetWorldSpaceInflatedBounds(ParticleIdx, WorldSpaceInflatedBounds);
+	}
 
 	bool HasBounds() const { return GeometryParticles->HasBounds(ParticleIdx); }
 	void SetHasBounds(bool bHasBounds) { GeometryParticles->HasBounds(ParticleIdx) = bHasBounds; }
diff --git a/Engine/Source/Runtime/PhysicsCore/Private/PhysTestSerializer.cpp b/Engine/Source/Runtime/PhysicsCore/Private/PhysTestSerializer.cpp
--- a/Engine/Source/Runtime/PhysicsCore/Private/PhysTestSerializer.cpp
+++ b/Engine/Source/Runtime/PhysicsCore/Private/PhysTestSerializer.cpp
@@ -307,6 +307,17 @@ void FPhysTestSerializer::CreateChaosData()
 					GTParticle->SetGeometry(MakeUnique<TImplicitObjectUnion<float, 3>>(MoveTemp(Geoms)));
 					Particle->SetGeometry(GTParticle->Geometry());
 				}
+
+				// Fixup bounds
+				auto Geom = GTParticle->Geometry();
+				if (Geom->HasBoundingBox())
+				{
+					auto& ShapeArray = GTParticle->ShapesArray();
+					for (auto& Shape : ShapeArray)
+					{
+						Shape->WorldSpaceInflatedShapeBounds = Geom->BoundingBox().GetAABB().TransformedAABB(TRigidTransform<FReal, 3>(Particle->X(), Particle->R()));
+					}
+				}
 			}
 
 			int32 ShapeIdx = 0;
diff --git a/Engine/Source/Runtime/PhysicsSQ/Private/SQAccelerator.cpp b/Engine/Source/Runtime/PhysicsSQ/Private/SQAccelerator.cpp
--- a/Engine/Source/Runtime/PhysicsSQ/Private/SQAccelerator.cpp
+++ b/Engine/Source/Runtime/PhysicsSQ/Private/SQAccelerator.cpp
@@ -189,9 +189,16 @@ private:
 
 			if (bTestShapeBounds)
 			{
-				const auto& GeomBounds = Geom->BoundingBox();
-				auto InflatedWorldBounds = SQ == ESQType::Raycast ? GeomBounds.GetAABB().TransformedAABB(ActorTM) :
-					TAABB<FReal, 3>(GeomBounds.Min() - HalfExtents, GeomBounds.Max() + HalfExtents).TransformedAABB(ActorTM);
+				TAABB<FReal, 3> InflatedWorldBounds;
+				if (SQ == ESQType::Raycast)
+				{
+					InflatedWorldBounds = Shape->WorldSpaceInflatedShapeBounds;
+				}
+				else
+				{
+					InflatedWorldBounds = TAABB<FReal, 3>(Shape->WorldSpaceInflatedShapeBounds.Min() - HalfExtents, Shape->WorldSpaceInflatedShapeBounds.Max() + HalfExtents);
+				}
+
 				if (SQ != ESQType::Overlap)
 				{
 					//todo: use fast raycast
