Fix FTripleBuffer (NextReadIndex was never used, and the ReadIndex was never changed), and add FGuardedTipleBuffer that tracks the valid state of stored buffers and returns nullptr if there is no new data in the buffer.

diff --git a/Engine/Source/Runtime/Experimental/Chaos/Public/Chaos/Framework/MultiBufferResource.h b/Engine/Source/Runtime/Experimental/Chaos/Public/Chaos/Framework/MultiBufferResource.h
--- a/Engine/Source/Runtime/Experimental/Chaos/Public/Chaos/Framework/MultiBufferResource.h
+++ b/Engine/Source/Runtime/Experimental/Chaos/Public/Chaos/Framework/MultiBufferResource.h
@@ -10,10 +10,11 @@ namespace Chaos
 		Undefined = 0,
 		Single,
 		Double,
-		Triple
+		Triple,
+		TripleGuarded
 	};
 
-	template <typename ResourceType>
+	template<typename ResourceType>
 	class IBufferResource
 	{
 	public:
@@ -31,7 +32,7 @@ namespace Chaos
 	/**
 	 * Single Buffer Implementation
 	 */
-	template <typename ResourceType>
+	template<typename ResourceType>
 	class FSingleBuffer : public IBufferResource<ResourceType>
 	{
 	public:
@@ -54,14 +55,14 @@ namespace Chaos
 	/*
 	* Double Buffer Implementation - Not thread-safe requires external locks
 	*/
-	template <typename ResourceType>
+	template<typename ResourceType>
 	class FDoubleBuffer : public IBufferResource<ResourceType>
 	{
 	public:
 
-		FDoubleBuffer() 
-			: Data_Producer(&Data1)
-			, Data_Consumer(&Data2)
+		FDoubleBuffer()
+		    : Data_Producer(&Data1)
+		    , Data_Consumer(&Data2)
 		{}
 		virtual ~FDoubleBuffer() {}
 
@@ -95,15 +96,14 @@ namespace Chaos
 	/*
 	* Triple Buffer Implementation - Not thread-safe requires external locks
 	*/
-	template <typename ResourceType>
+	template<typename ResourceType>
 	class FTripleBuffer : public IBufferResource<ResourceType>
 	{
 	public:
 
 		FTripleBuffer()
-			: WriteIndex(1)
-			, ReadIndex(0)
-			, NextReadIndex(0)
+		    : WriteIndex(1)
+		    , ReadIndex(0)
 		{
 		}
 		virtual ~FTripleBuffer() {}
@@ -119,7 +119,7 @@ namespace Chaos
 			int32 CurrentWriteIdx = GetWriteIndex();
 			int32 FreeIdx = 3 - (CurrentReadIdx + CurrentWriteIdx);
 
-			NextReadIndex.Store(CurrentWriteIdx);
+			ReadIndex.Store(CurrentWriteIdx);
 			WriteIndex.Store(FreeIdx);
 
 			checkSlow(GetReadIndex() != GetWriteIndex());
@@ -132,10 +132,154 @@ namespace Chaos
 		ResourceType Data[3];
 		TAtomic<int32> WriteIndex;
 		TAtomic<int32> ReadIndex;
-		TAtomic<int32> NextReadIndex;
-	};	
+	};
 	//////////////////////////////////////////////////////////////////////////
+	/**
+	 * Triple buffer based on a single atomic variable, that guards against 
+	 * the consumer thread using old values.
+	 *
+	 * Not thread-safe, requires external locks.
+	 */
+	template<typename ResourceType>
+	class FGuardedTripleBuffer : public IBufferResource<ResourceType>
+	{
+	public:
+		/** This class implements a circular buffer access pattern, such that during
+		 * normal serial operation each buffer will be used.  However, we can alter
+		 * how we use the API to use only 2 or 1 buffer:
+		 * 
+		 *    SERIAL ACCESS BUFFER USE - NORMAL OPERATION:
+		 *            P  C  I
+		 *     Init:  0, 1, 2
+		 *     
+		 *     AccP: *0, 1, 2 - *write access
+		 *     Flip:  2, 1, 0
+		 *     Cons: 2, ^0, 1 - ^read access
+		 *     
+		 *     AccP: *2, 0, 1
+		 *     Flip:  1, 0, 2
+		 *     Cons: 1, ^2, 0
+		 *     
+		 *     AccP: *1, 2, 0
+		 *     Flip:  0, 2, 1
+		 *     Cons: 0, ^1, 2
+		 * 
+		 * 
+		 *    SERIAL ACCESS BUFFER USE - DOUBLE BUFFER EMULATION:
+		 *     Init:  0, 1, 2
+		 *     
+		 *     AccP: *0, 1, 2 - *write access
+		 *     Flip:  2, 1, 0
+		 *     Cons: 2, ^0, 1 - ^read access
+		 *     Flip:  1, 0, 2		<- Add an extra flip and only use 2 buffers.
+		 *     
+		 *     AccP: *1, 0, 2
+		 *     Flip:  2, 0, 1
+		 *     Cons: 2, ^1, 0
+		 *     Flip:  0, 1, 2
+		 * 
+		 *
+		 *    SERIAL ACCESS BUFFER USE - SINGLE BUFFER EMULATION:
+		 *     Init:  0, 1, 2
+		 *
+		 *     AccP: *0, 1, 2 - *write access
+		 *     Flip:  2, 1, 0
+		 *     Cons: 2, ^0, 1 - ^read access
+		 *     Cons:  2, 1, 0		<- Add an extra consume...
+		 *     Flip:  0, 1, 2		<- and an extra flip and only use 1 buffer.
+		 */
 
+		FGuardedTripleBuffer()
+		    : ProducerThreadBuffer(&Buffers[0])
+		    , ConsumerThreadBuffer(&Buffers[1])
+		    , Interchange(&Buffers[2])
+		{}
+
+		virtual EMultiBufferMode GetBufferMode() override { return EMultiBufferMode::Triple; }
+
+		/**
+		 * Get the current producer buffer for writing.
+		 *
+		 * The returned pointer will never be null.
+		 *
+		 * This function should only be called by the producer thread.
+		 */
+		virtual ResourceType* AccessProducerBuffer() override
+		{
+			return ProducerThreadBuffer->Value.Get();
+		}
+
+		/**
+		 * Get the current producer buffer.
+		 *
+		 * The returned pointer will never be null.
+		 *
+		 * This function should only be called by the producer thread.
+		 */
+		virtual const ResourceType* GetProducerBuffer() const override
+		{
+			return ProducerThreadBuffer->Value.Get();
+		}
+
+		/**
+		 * Get an updated buffer for the consuming thread to read from.
+		 *
+		 * The returned pointer may be null if the producer thread hasn't done 
+		 * an update since the last time this function was called.
+		 *
+		 * This function should only be called by the consumer thread.
+		 */
+		virtual const ResourceType* GetConsumerBuffer() const override
+		{
+			ConsumerThreadBuffer = Interchange.Exchange(ConsumerThreadBuffer);
+			// Return null if the value hasn't changed.
+			if (!ConsumerThreadBuffer->bValid)
+			{
+				return nullptr;
+			}
+			// Mark the buffer as invalid so we don't use this value again.
+			ConsumerThreadBuffer->bValid = false;
+			return ConsumerThreadBuffer->Value.Get();
+		}
+
+		/**
+		 * Make the current producer buffer available to the consumer thread.
+		 *
+		 * This function should only be called by the producer thread.
+		 */
+		virtual void FlipProducer()
+		{
+			ProducerThreadBuffer->bValid = true; // In with the new...
+			ProducerThreadBuffer = Interchange.Exchange(ProducerThreadBuffer);
+			ProducerThreadBuffer->bValid = false; // Out with the old.
+		}
+
+	private:
+		struct ResourceTypeWrapper
+		{
+			ResourceTypeWrapper()
+			    : Value(new ResourceType())
+			    , bValid(false)
+			{}
+			ResourceTypeWrapper(const ResourceTypeWrapper&) = delete;
+			ResourceTypeWrapper(ResourceTypeWrapper&& Other)
+			    : Value(MoveTemp(Other.Value))
+			    , bValid(MoveTemp(Other.bValid))
+			{}
+
+			void operator=(const ResourceTypeWrapper& Other) = delete;
+
+			TUniquePtr<ResourceType> Value;
+			bool bValid;
+		};
+
+		ResourceTypeWrapper Buffers[3] = {ResourceTypeWrapper(), ResourceTypeWrapper(), ResourceTypeWrapper()};
+		ResourceTypeWrapper* ProducerThreadBuffer;
+		mutable ResourceTypeWrapper* ConsumerThreadBuffer;
+		mutable TAtomic<ResourceTypeWrapper*> Interchange;
+	};
+
+	//////////////////////////////////////////////////////////////////////////
 
 	template<typename ResourceType>
 	class FMultiBufferFactory
@@ -163,9 +307,15 @@ namespace Chaos
 			}
 			break;
 
+			case EMultiBufferMode::TripleGuarded:
+			{
+				return MakeUnique<FGuardedTripleBuffer<ResourceType>>();
+			}
+			break;
+
 			default:
 				checkf(false, TEXT("FMultiBufferFactory Unexpected buffer mode"));
-			break;
+				break;
 
 			}
 
