Add move constructor with TArray<T>.  Adds Count() which returns a tally of how many entires matched.  Adds Fill() which fills the array with a value.

diff --git a/Engine/Source/Runtime/Experimental/GeometryCollectionCore/Public/GeometryCollection/ManagedArray.h b/Engine/Source/Runtime/Experimental/GeometryCollectionCore/Public/GeometryCollection/ManagedArray.h
--- a/Engine/Source/Runtime/Experimental/GeometryCollectionCore/Public/GeometryCollection/ManagedArray.h
+++ b/Engine/Source/Runtime/Experimental/GeometryCollectionCore/Public/GeometryCollection/ManagedArray.h
@@ -209,6 +209,9 @@ public:
 	FORCEINLINE TManagedArrayBase(TManagedArrayBase<ElementType>&& Other)
 		: Array(MoveTemp(Other.Array))
 	{}
+	FORCEINLINE TManagedArrayBase(TArray<ElementType>&& Other)
+		: Array(MoveTemp(Other))
+	{}
 
 	/**
 	* Assignment operator
@@ -266,7 +269,16 @@ public:
 
 		Resize(Size);
 		InitHelper(Array, NewTypedArray, Size);
-	};
+	}
+
+	/**
+	 * Fill the array with \p Value.
+	 */
+	void Fill(const ElementType& Value)
+	{
+		for (int32 Idx = 0; Idx < Array.Num(); ++Idx)
+			Array[Idx] = Value;
+	}
 
 #if 0
 	virtual void Swap(int32 Index1, int32 Index2) override
@@ -282,7 +294,7 @@ public:
 		TManagedArrayBase<ElementType>& NewTypedArray = static_cast<TManagedArrayBase<ElementType>& >(NewArray);
 
 		Exchange(*this, NewTypedArray);
-	};
+	}
 
 	/**
 	* Returning a reference to the element at index.
@@ -361,6 +373,16 @@ public:
 		return Array.Find(Item);
 	}
 
+	/**
+	* Count the number of entries match \p Item.
+	*/
+	int32 Count(const ElementType& Item) const
+	{
+		int32 Num = 0;
+		for (int32 Idx = 0; Idx < Array.Num(); ++Idx)
+			Num += Array[Idx] == Item ? 1 : 0;
+		return Num;
+	}
 
 	/**
 	* Checks if index is in array range.
@@ -492,6 +514,10 @@ public:
 		: TManagedArrayBase<InElementType>(MoveTemp(Other))
 	{}
 
+	FORCEINLINE TManagedArray(TArray<InElementType>&& Other)
+		: TManagedArrayBase<InElementType>(MoveTemp(Other))
+	{}
+
 	FORCEINLINE TManagedArray& operator=(TManagedArray<InElementType>&& Other)
 	{
 		TManagedArrayBase<InElementType>::operator=(MoveTemp(Other));
@@ -519,6 +545,9 @@ public:
 
 	FORCEINLINE TManagedArray(const TManagedArray<int32>& Other) = delete;
 	FORCEINLINE TManagedArray(TManagedArray<int32>&& Other) = default;
+	FORCEINLINE TManagedArray(TArray<int32>&& Other)
+		: TManagedArrayBase<int32>(MoveTemp(Other))
+	{}
 	FORCEINLINE TManagedArray& operator=(TManagedArray<int32>&& Other) = default;
 
 	virtual ~TManagedArray()
@@ -571,6 +600,9 @@ public:
 
 	FORCEINLINE TManagedArray(const TManagedArray<TSet<int32>>& Other) = delete;
 	FORCEINLINE TManagedArray(TManagedArray<TSet<int32>>&& Other) = default;
+	FORCEINLINE TManagedArray(TArray<TSet<int32>>&& Other)
+		: TManagedArrayBase<TSet<int32>>(MoveTemp(Other))
+	{}
 	FORCEINLINE TManagedArray& operator=(TManagedArray<TSet<int32>>&& Other) = default;
 
 	virtual ~TManagedArray()
@@ -638,6 +670,9 @@ public:
 
 	FORCEINLINE TManagedArray(const TManagedArray<FIntVector>& Other) = delete;
 	FORCEINLINE TManagedArray(TManagedArray<FIntVector>&& Other) = default;
+	FORCEINLINE TManagedArray(TArray<FIntVector>&& Other)
+		: TManagedArrayBase<FIntVector>(MoveTemp(Other))
+	{}
 	FORCEINLINE TManagedArray& operator=(TManagedArray<FIntVector>&& Other) = default;
 
 	virtual ~TManagedArray()
diff --git a/Engine/Source/Runtime/Experimental/GeometryCollectionCore/Public/GeometryCollection/ManagedArrayCollection.h b/Engine/Source/Runtime/Experimental/GeometryCollectionCore/Public/GeometryCollection/ManagedArrayCollection.h
--- a/Engine/Source/Runtime/Experimental/GeometryCollectionCore/Public/GeometryCollection/ManagedArrayCollection.h
+++ b/Engine/Source/Runtime/Experimental/GeometryCollectionCore/Public/GeometryCollection/ManagedArrayCollection.h
@@ -106,10 +106,10 @@ public:
 			Map.Add(FManagedArrayCollection::MakeMapKey(Name, Group), MoveTemp(Value));
 		}
 		return GetAttribute<T>(Name, Group);
-	};
+	}
 
 	/**
-	* Add an external attribute of Type(T) to the group for size management. Lifetime is managed by the caller, must make sure the array is alive when collection is
+	* Add an external attribute of Type(T) to the group for size management. Lifetime is managed by the caller, must make sure the array is alive when the collection is
 	* @param Name - The name of the attribute
 	* @param Group - The group that manages the attribute
 	* @param ValueIn - The array to be managed
@@ -137,7 +137,7 @@ public:
 		}
 		Value.bExternalValue = true;
 		Map.Add(FManagedArrayCollection::MakeMapKey(Name, Group), MoveTemp(Value));
-	};
+	}
 
 	/**
 	* Create a group on the collection. Adding attribute will also create unknown groups.
