Hacking at GeometryCollection debug drawing, trying to replace the int32 rigid body id with a FGuid.  It's not there yet, but this change gets things building.

diff --git a/Engine/Plugins/Experimental/GeometryCollectionPlugin/Source/GeometryCollectionEditor/Private/GeometryCollection/DetailCustomizations/SelectedRigidBodyCustomization.cpp b/Engine/Plugins/Experimental/GeometryCollectionPlugin/Source/GeometryCollectionEditor/Private/GeometryCollection/DetailCustomizations/SelectedRigidBodyCustomization.cpp
--- a/Engine/Plugins/Experimental/GeometryCollectionPlugin/Source/GeometryCollectionEditor/Private/GeometryCollection/DetailCustomizations/SelectedRigidBodyCustomization.cpp
+++ b/Engine/Plugins/Experimental/GeometryCollectionPlugin/Source/GeometryCollectionEditor/Private/GeometryCollection/DetailCustomizations/SelectedRigidBodyCustomization.cpp
@@ -197,6 +197,7 @@ void FSelectedRigidBodyCustomization::CustomizeChildren(TSharedRef<IPropertyHand
 
 void FSelectedRigidBodyCustomization::GetSelectedGeometryCollectionCluster(TSharedRef<IPropertyHandle> PropertyHandleActor, TSharedRef<IPropertyHandle> PropertyHandleId, const UGeometryCollectionComponent*& OutGeometryCollectionComponent, int32& OutTransformIndex)
 {
+#if TODO_REIMPLEMENT_RIGID_CLUSTERING
 	OutGeometryCollectionComponent = nullptr;
 	OutTransformIndex = INDEX_NONE;
 
@@ -225,10 +226,12 @@ void FSelectedRigidBodyCustomization::GetSelectedGeometryCollectionCluster(TShar
 		}
 	}
 	UE_CLOG(OutGeometryCollectionComponent && OutGeometryCollectionComponent->GetOwner(), LogSelectedRigidBodyCustomization, VeryVerbose, TEXT("Component actor %s, TransformIndex %d."), *OutGeometryCollectionComponent->GetOwner()->GetName(), OutTransformIndex);
+#endif // TODO_REIMPLEMENT_RIGID_CLUSTERING
 }
 
 int32 FSelectedRigidBodyCustomization::GetParentClusterRigidBodyId(TSharedRef<IPropertyHandle> PropertyHandleActor, TSharedRef<IPropertyHandle> PropertyHandleId)
 {
+#if TODO_REIMPLEMENT_RIGID_CLUSTERING
 	const UGeometryCollectionComponent* GeometryCollectionComponent;
 	int32 TransformIndex;
 	GetSelectedGeometryCollectionCluster(PropertyHandleActor, PropertyHandleId, GeometryCollectionComponent, TransformIndex);
@@ -240,11 +243,13 @@ int32 FSelectedRigidBodyCustomization::GetParentClusterRigidBodyId(TSharedRef<IP
 			return GeometryCollectionComponent->GetRigidBodyIdArray()[ParentTransformIndex];
 		}
 	}
+#endif // TODO_REIMPLEMENT_RIGID_CLUSTERING
 	return INDEX_NONE;
 }
 
 int32 FSelectedRigidBodyCustomization::GetChildClusterRigidBodyId(TSharedRef<IPropertyHandle> PropertyHandleActor, TSharedRef<IPropertyHandle> PropertyHandleId)
 {
+#if TODO_REIMPLEMENT_RIGID_CLUSTERING
 	const UGeometryCollectionComponent* GeometryCollectionComponent;
 	int32 TransformIndex;
 	GetSelectedGeometryCollectionCluster(PropertyHandleActor, PropertyHandleId, GeometryCollectionComponent, TransformIndex);
@@ -261,11 +266,13 @@ int32 FSelectedRigidBodyCustomization::GetChildClusterRigidBodyId(TSharedRef<IPr
 			}
 		}
 	}
+#endif // TODO_REIMPLEMENT_RIGID_CLUSTERING
 	return INDEX_NONE;
 }
 
 int32 FSelectedRigidBodyCustomization::GetPreviousClusteredSiblingRigidBodyId(TSharedRef<IPropertyHandle> PropertyHandleActor, TSharedRef<IPropertyHandle> PropertyHandleId)
 {
+#if TODO_REIMPLEMENT_RIGID_CLUSTERING
 	const UGeometryCollectionComponent* GeometryCollectionComponent;
 	int32 TransformIndex;
 	GetSelectedGeometryCollectionCluster(PropertyHandleActor, PropertyHandleId, GeometryCollectionComponent, TransformIndex);
@@ -291,11 +298,13 @@ int32 FSelectedRigidBodyCustomization::GetPreviousClusteredSiblingRigidBodyId(TS
 			}
 		}
 	}
+#endif // TODO_REIMPLEMENT_RIGID_CLUSTERING
 	return INDEX_NONE;
 }
 
 int32 FSelectedRigidBodyCustomization::GetNextClusteredSiblingRigidBodyId(TSharedRef<IPropertyHandle> PropertyHandleActor, TSharedRef<IPropertyHandle> PropertyHandleId)
 {
+#if TODO_REIMPLEMENT_RIGID_CLUSTERING
 	const UGeometryCollectionComponent* GeometryCollectionComponent;
 	int32 TransformIndex;
 	GetSelectedGeometryCollectionCluster(PropertyHandleActor, PropertyHandleId, GeometryCollectionComponent, TransformIndex);
@@ -324,6 +333,7 @@ int32 FSelectedRigidBodyCustomization::GetNextClusteredSiblingRigidBodyId(TShare
 			}
 		}
 	}
+#endif // TODO_REIMPLEMENT_RIGID_CLUSTERING
 	return INDEX_NONE;
 }
 
diff --git a/Engine/Plugins/Experimental/GeometryCollectionPlugin/Source/GeometryCollectionEditor/Private/GeometryCollection/GeometryCollectionSelectRigidBodyEdMode.cpp b/Engine/Plugins/Experimental/GeometryCollectionPlugin/Source/GeometryCollectionEditor/Private/GeometryCollection/GeometryCollectionSelectRigidBodyEdMode.cpp
--- a/Engine/Plugins/Experimental/GeometryCollectionPlugin/Source/GeometryCollectionEditor/Private/GeometryCollection/GeometryCollectionSelectRigidBodyEdMode.cpp
+++ b/Engine/Plugins/Experimental/GeometryCollectionPlugin/Source/GeometryCollectionEditor/Private/GeometryCollection/GeometryCollectionSelectRigidBodyEdMode.cpp
@@ -156,17 +156,20 @@ bool FGeometryCollectionSelectRigidBodyEdMode::HandleClick(FEditorViewportClient
 				}
 
 				// Retrieve the rigid body id
-				const TManagedArray<int32>& RigidBodyIdArray = GeometryCollectionComponent->GetRigidBodyIdArray();
-				const int32 RigidBodyId = (TransformIndex != INDEX_NONE && ensure(TransformIndex < RigidBodyIdArray.Num())) ? RigidBodyIdArray[TransformIndex]: INDEX_NONE;
+				//const TManagedArray<int32>& RigidBodyIdArray = GeometryCollectionComponent->GetRigidBodyIdArray();
+				const TManagedArray<FGuid>& RigidBodyIdArray = GeometryCollectionComponent->GetRigidBodyGuidArray();
+				//const int32 RigidBodyId = (TransformIndex != INDEX_NONE && ensure(TransformIndex < RigidBodyIdArray.Num())) ? RigidBodyIdArray[TransformIndex]: INDEX_NONE;
+				const FGuid& RigidBodyId = (TransformIndex != INDEX_NONE && ensure(TransformIndex < RigidBodyIdArray.Num())) ? RigidBodyIdArray[TransformIndex] : FGuid();
 
 				// Update the rigid body id property
-				if (RigidBodyId != INDEX_NONE)
+				//if (RigidBodyId != INDEX_NONE)
+				if (RigidBodyId.IsValid())
 				{
-					UE_LOG(LogGeometryCollectionSelectRigidBodyEdMode, Verbose, TEXT("Hit GeometryCollectionActor %s at rigid body %d."), *GeometryCollectionActor->GetName(), RigidBodyId);
+					UE_LOG(LogGeometryCollectionSelectRigidBodyEdMode, Verbose, TEXT("Hit GeometryCollectionActor %s at rigid body %s."), *GeometryCollectionActor->GetName(), *RigidBodyId.ToString());
 
 					if (const TSharedPtr<IPropertyHandle> PropertyHandleIdPin = PropertyHandleId.Pin())
 					{
-						PropertyHandleIdPin->SetValue(RigidBodyId);
+						PropertyHandleIdPin->SetValue(RigidBodyId.ToString());
 					}
 					if (const TSharedPtr<IPropertyHandle> PropertyHandleSolverPin = PropertyHandleSolver.Pin())
 					{
diff --git a/Engine/Source/Runtime/Experimental/GeometryCollectionEngine/Private/GeometryCollection/GeometryCollectionActor.cpp b/Engine/Source/Runtime/Experimental/GeometryCollectionEngine/Private/GeometryCollection/GeometryCollectionActor.cpp
--- a/Engine/Source/Runtime/Experimental/GeometryCollectionEngine/Private/GeometryCollection/GeometryCollectionActor.cpp
+++ b/Engine/Source/Runtime/Experimental/GeometryCollectionEngine/Private/GeometryCollection/GeometryCollectionActor.cpp
@@ -67,7 +67,8 @@ bool LowLevelRaycastImp(const Chaos::TVector<float, 3>& Start, const Chaos::TVec
 {
 	using namespace Chaos;
 	//todo(ocohen): need to add thread safety / lock semantics
-	const TManagedArray<int32>& RigidBodyIdArray = GeomCollectionActor.GetGeometryCollectionComponent()->GetRigidBodyIdArray();
+	//const TManagedArray<int32>& RigidBodyIdArray = GeomCollectionActor.GetGeometryCollectionComponent()->GetRigidBodyIdArray();
+	const TManagedArray<FGuid>& RigidBodyIdArray = GeomCollectionActor.GetGeometryCollectionComponent()->GetRigidBodyGuidArray();
 	const TSharedPtr<FPhysScene_Chaos> Scene = GeomCollectionActor.GetGeometryCollectionComponent()->GetPhysicsScene();
 	ensure(Scene);
 
diff --git a/Engine/Source/Runtime/Experimental/GeometryCollectionEngine/Private/GeometryCollection/GeometryCollectionComponent.cpp b/Engine/Source/Runtime/Experimental/GeometryCollectionEngine/Private/GeometryCollection/GeometryCollectionComponent.cpp
--- a/Engine/Source/Runtime/Experimental/GeometryCollectionEngine/Private/GeometryCollection/GeometryCollectionComponent.cpp
+++ b/Engine/Source/Runtime/Experimental/GeometryCollectionEngine/Private/GeometryCollection/GeometryCollectionComponent.cpp
@@ -1227,7 +1227,7 @@ void UGeometryCollectionComponent::OnCreatePhysicsState()
 #if GEOMETRYCOLLECTION_DEBUG_DRAW
 				const bool bHasNumParticlesChanged = (NumParticlesAdded != Results.NumParticlesAdded);  // Needs to be evaluated before NumParticlesAdded gets updated
 #endif  // #if GEOMETRYCOLLECTION_DEBUG_DRAW
-				RigidBodyIds.Init(Results.RigidBodyIds);
+				//RigidBodyIds.Init(Results.RigidBodyIds);
 				BaseRigidBodyIndex = Results.BaseIndex;
 				NumParticlesAdded = Results.NumParticlesAdded;
 				DisabledFlags = Results.DisabledStates;
diff --git a/Engine/Source/Runtime/Experimental/GeometryCollectionEngine/Private/GeometryCollection/GeometryCollectionDebugDrawActor.cpp b/Engine/Source/Runtime/Experimental/GeometryCollectionEngine/Private/GeometryCollection/GeometryCollectionDebugDrawActor.cpp
--- a/Engine/Source/Runtime/Experimental/GeometryCollectionEngine/Private/GeometryCollection/GeometryCollectionDebugDrawActor.cpp
+++ b/Engine/Source/Runtime/Experimental/GeometryCollectionEngine/Private/GeometryCollection/GeometryCollectionDebugDrawActor.cpp
@@ -113,6 +113,7 @@ namespace GeometryCollectionDebugDrawActorCVars
 	// Console variables, also exposed as settings in this actor
 	static TAutoConsoleVariable<FString> SelectedRigidBodySolver (TEXT("p.gc.SelectedRigidBodySolver" ), GeometryCollectionDebugDrawActorConstants::SelectedRigidBodySolverDefault , TEXT("Geometry Collection debug draw, visualize debug informations for the selected rigid body solver.\nDefault = None"), ECVF_Cheat);
 	static TAutoConsoleVariable<int32  > SelectedRigidBodyId     (TEXT("p.gc.SelectedRigidBodyId"     ), GeometryCollectionDebugDrawActorConstants::SelectedRigidBodyIdDefault     , TEXT("Geometry Collection debug draw, visualize debug informations for the selected rigid body ids.\nDefault = -1"), ECVF_Cheat);
+	//static TAutoConsoleVariable<FGuid  > SelectedRigidBodyId     (TEXT("p.gc.SelectedRigidBodyId"     ), GeometryCollectionDebugDrawActorConstants::SelectedRigidBodyIdDefault     , TEXT("Geometry Collection debug draw, visualize debug informations for the selected rigid body ids.\nDefault = 0:0:0:0"), ECVF_Cheat);
 	static TAutoConsoleVariable<int32  > DebugDrawWholeCollection(TEXT("p.gc.DebugDrawWholeCollection"), GeometryCollectionDebugDrawActorConstants::DebugDrawWholeCollectionDefault, TEXT("Geometry Collection debug draw, show debug visualization for the rest of the geometry collection related to the current rigid body id selection.\nDefault = 0"), ECVF_Cheat);
 	static TAutoConsoleVariable<int32  > DebugDrawHierarchy      (TEXT("p.gc.DebugDrawHierarchy"      ), GeometryCollectionDebugDrawActorConstants::DebugDrawHierarchyDefault      , TEXT("Geometry Collection debug draw, show debug visualization for the top level node rather than the bottom leaf nodes of a cluster's hierarchy..\nDefault = 0"), ECVF_Cheat);
 	static TAutoConsoleVariable<int32  > DebugDrawClustering     (TEXT("p.gc.DebugDrawClustering"     ), GeometryCollectionDebugDrawActorConstants::DebugDrawClusteringDefault     , TEXT("Geometry Collection debug draw, show debug visualization for all clustered children associated to the current rigid body id selection.\nDefault = 0"), ECVF_Cheat);
@@ -334,6 +335,7 @@ void AGeometryCollectionDebugDrawActor::Tick(float DeltaSeconds)
 #if GEOMETRYCOLLECTION_DEBUG_DRAW
 	// Check badly synced collections in case it is still looking for an id match
 	if (World && SelectedRigidBody.Id != INDEX_NONE && !SelectedRigidBody.GeometryCollection)
+	//if (World && SelectedRigidBody.Id.IsValid() && !SelectedRigidBody.GeometryCollection)
 	{
 #if TODO_REIMPLEMENT_GET_RIGID_PARTICLES
 		// Check the id is within the selected solver range
@@ -399,6 +401,7 @@ void AGeometryCollectionDebugDrawActor::PostLoad()
 
 	GeometryCollectionDebugDrawActorCVars::SelectedRigidBodySolver ->Set(*SelectedRigidBody.GetSolverName(), SetBy);
 	GeometryCollectionDebugDrawActorCVars::SelectedRigidBodyId     ->Set( SelectedRigidBody.Id             , SetBy);
+	//GeometryCollectionDebugDrawActorCVars::SelectedRigidBodyId     ->Set(*SelectedRigidBody.Id.ToString()  , SetBy);
 	GeometryCollectionDebugDrawActorCVars::DebugDrawWholeCollection->Set(int32(bDebugDrawWholeCollection  ), SetBy);
 	GeometryCollectionDebugDrawActorCVars::DebugDrawHierarchy      ->Set(int32(bDebugDrawHierarchy        ), SetBy);
 	GeometryCollectionDebugDrawActorCVars::DebugDrawClustering     ->Set(int32(bDebugDrawClustering       ), SetBy);
@@ -584,8 +587,10 @@ void AGeometryCollectionDebugDrawActor::PostEditChangeProperty(FPropertyChangedE
 
 	if      (PropertyName == GET_MEMBER_NAME_CHECKED(AGeometryCollectionDebugDrawActor, SelectedRigidBody        )) { GeometryCollectionDebugDrawActorCVars::SelectedRigidBodySolver ->Set(*SelectedRigidBody.GetSolverName(), SetBy);
 																													  GeometryCollectionDebugDrawActorCVars::SelectedRigidBodyId     ->Set( SelectedRigidBody.Id             , SetBy); }
+//																													  GeometryCollectionDebugDrawActorCVars::SelectedRigidBodyId     ->Set(*SelectedRigidBody.Id.ToString()  , SetBy); }
 	else if (PropertyName == GET_MEMBER_NAME_CHECKED(FGeometryCollectionDebugDrawActorSelectedRigidBody, Solver  )) { GeometryCollectionDebugDrawActorCVars::SelectedRigidBodySolver ->Set(*SelectedRigidBody.GetSolverName(), SetBy); }
 	else if (PropertyName == GET_MEMBER_NAME_CHECKED(FGeometryCollectionDebugDrawActorSelectedRigidBody, Id      )) { GeometryCollectionDebugDrawActorCVars::SelectedRigidBodyId     ->Set( SelectedRigidBody.Id             , SetBy); }
+//	else if (PropertyName == GET_MEMBER_NAME_CHECKED(FGeometryCollectionDebugDrawActorSelectedRigidBody, Id      )) { GeometryCollectionDebugDrawActorCVars::SelectedRigidBodyId     ->Set(*SelectedRigidBody.Id.ToString()  , SetBy); }
 	else if (PropertyName == GET_MEMBER_NAME_CHECKED(AGeometryCollectionDebugDrawActor, bDebugDrawWholeCollection)) { GeometryCollectionDebugDrawActorCVars::DebugDrawWholeCollection->Set(int32(bDebugDrawWholeCollection  ), SetBy); }
 	else if (PropertyName == GET_MEMBER_NAME_CHECKED(AGeometryCollectionDebugDrawActor, bDebugDrawHierarchy      )) { GeometryCollectionDebugDrawActorCVars::DebugDrawHierarchy      ->Set(int32(bDebugDrawHierarchy        ), SetBy); }
 	else if (PropertyName == GET_MEMBER_NAME_CHECKED(AGeometryCollectionDebugDrawActor, bDebugDrawClustering     )) { GeometryCollectionDebugDrawActorCVars::DebugDrawClustering     ->Set(int32(bDebugDrawClustering       ), SetBy); bForceVisibilityUpdate = true; }
@@ -639,7 +644,8 @@ bool AGeometryCollectionDebugDrawActor::CanEditChange(const UProperty* InPropert
 	if (PropertyName == GET_MEMBER_NAME_CHECKED(AGeometryCollectionDebugDrawActor, bDebugDrawWholeCollection) ||
 		PropertyName == GET_MEMBER_NAME_CHECKED(AGeometryCollectionDebugDrawActor, bShowSingleFace))
 	{
-		return SelectedRigidBody.Id != -1;
+		return SelectedRigidBody.Id != INDEX_NONE;
+		//return SelectedRigidBody.Id.IsValid();
 	}
 	return Super::CanEditChange(InProperty);
 }
@@ -1993,7 +1999,8 @@ FTransform AGeometryCollectionDebugDrawActor::GetParticleTransformNoChecks(const
 	return FTransform(Rotation, Translation);
 }
 
-void AGeometryCollectionDebugDrawActor::DrawRigidBodiesId(const UGeometryCollectionComponent* GeometryCollectionComponent, const FGeometryCollectionParticlesData& ParticlesData, const TManagedArray<int32>& RigidBodyIdArray, const FColor& Color)
+//void AGeometryCollectionDebugDrawActor::DrawRigidBodiesId(const UGeometryCollectionComponent* GeometryCollectionComponent, const FGeometryCollectionParticlesData& ParticlesData, const TManagedArray<int32>& RigidBodyIdArray, const FColor& Color)
+void AGeometryCollectionDebugDrawActor::DrawRigidBodiesId(const UGeometryCollectionComponent* GeometryCollectionComponent, const FGeometryCollectionParticlesData& ParticlesData, const TManagedArray<FGuid>& RigidBodyIdArray, const FColor& Color)
 {
 #if ENABLE_DRAW_DEBUG
 	check(GeometryCollectionComponent);
@@ -2034,7 +2041,8 @@ void AGeometryCollectionDebugDrawActor::DrawRigidBodiesId(const UGeometryCollect
 #endif  // #if ENABLE_DRAW_DEBUG
 }
 
-void AGeometryCollectionDebugDrawActor::DrawRigidBodyId(const UGeometryCollectionComponent* GeometryCollectionComponent, int32 TransformIndex, const FGeometryCollectionParticlesData& ParticlesData, const TManagedArray<int32>& RigidBodyIdArray, const FColor& Color)
+//void AGeometryCollectionDebugDrawActor::DrawRigidBodyId(const UGeometryCollectionComponent* GeometryCollectionComponent, int32 TransformIndex, const FGeometryCollectionParticlesData& ParticlesData, const TManagedArray<int32>& RigidBodyIdArray, const FColor& Color)
+void AGeometryCollectionDebugDrawActor::DrawRigidBodyId(const UGeometryCollectionComponent* GeometryCollectionComponent, int32 TransformIndex, const FGeometryCollectionParticlesData& ParticlesData, const TManagedArray<FGuid>& RigidBodyIdArray, const FColor& Color)
 {
 #if ENABLE_DRAW_DEBUG
 	check(GeometryCollectionComponent);
@@ -2066,10 +2074,12 @@ void AGeometryCollectionDebugDrawActor::DrawRigidBodyId(const UGeometryCollectio
 #endif  // #if ENABLE_DRAW_DEBUG
 }
 
-void AGeometryCollectionDebugDrawActor::DrawRigidBodyIdNoChecks(const UGeometryCollectionComponent* GeometryCollectionComponent, int32 TransformIndex, const FGeometryCollectionParticlesData& ParticlesData, const TManagedArray<int32>& RigidBodyIdArray, const FColor& Color)
+//void AGeometryCollectionDebugDrawActor::DrawRigidBodyIdNoChecks(const UGeometryCollectionComponent* GeometryCollectionComponent, int32 TransformIndex, const FGeometryCollectionParticlesData& ParticlesData, const TManagedArray<int32>& RigidBodyIdArray, const FColor& Color)
+void AGeometryCollectionDebugDrawActor::DrawRigidBodyIdNoChecks(const UGeometryCollectionComponent* GeometryCollectionComponent, int32 TransformIndex, const FGeometryCollectionParticlesData& ParticlesData, const TManagedArray<FGuid>& RigidBodyIdArray, const FColor& Color)
 {
 #if ENABLE_DRAW_DEBUG
-	const int32 RigidBodyId = RigidBodyIdArray[TransformIndex];
+	//const int32 RigidBodyId = RigidBodyIdArray[TransformIndex];
+	const FGuid& RigidBodyId = RigidBodyIdArray[TransformIndex];
 
 	// Retrieve particle transform
 	const FTransform Transform = GetParticleTransformNoChecks(GeometryCollectionComponent, TransformIndex, ParticlesData);
@@ -2081,7 +2091,8 @@ void AGeometryCollectionDebugDrawActor::DrawRigidBodyIdNoChecks(const UGeometryC
 	const FColor& DisabledColor = bIsDisabled ? FColor::Silver: Color;
 
 	// Draw rigid body id
-	const FString Text = FString::Printf(TEXT("%d"), RigidBodyId);
+	//const FString Text = FString::Printf(TEXT("%d"), RigidBodyId);
+	const FString Text = RigidBodyId.ToString();
 	AddDebugText(Text, Position, DisabledColor, TextScale, bTextShadow);
 #endif  // #if ENABLE_DRAW_DEBUG
 }
@@ -2556,7 +2567,8 @@ void AGeometryCollectionDebugDrawActor::DrawRigidBodyInfoNoChecks(const UGeometr
 #endif  // #if ENABLE_DRAW_DEBUG
 }
 
-void AGeometryCollectionDebugDrawActor::DrawConnectivityEdges(const UGeometryCollectionComponent* GeometryCollectionComponent, const FGeometryCollectionParticlesData& ParticlesData, const TManagedArray<int32>& RigidBodyIdArray)
+//void AGeometryCollectionDebugDrawActor::DrawConnectivityEdges(const UGeometryCollectionComponent* GeometryCollectionComponent, const FGeometryCollectionParticlesData& ParticlesData, const TManagedArray<int32>& RigidBodyIdArray)
+void AGeometryCollectionDebugDrawActor::DrawConnectivityEdges(const UGeometryCollectionComponent* GeometryCollectionComponent, const FGeometryCollectionParticlesData& ParticlesData, const TManagedArray<FGuid>& RigidBodyIdArray)
 {
 #if ENABLE_DRAW_DEBUG
 	check(GeometryCollectionComponent);
@@ -2591,7 +2603,8 @@ void AGeometryCollectionDebugDrawActor::DrawConnectivityEdges(const UGeometryCol
 #endif  // #if ENABLE_DRAW_DEBUG
 }
 
-void AGeometryCollectionDebugDrawActor::DrawConnectivityEdges(const UGeometryCollectionComponent* GeometryCollectionComponent, int32 TransformIndex, const FGeometryCollectionParticlesData& ParticlesData, const TManagedArray<int32>& RigidBodyIdArray, FColor HSVColor)
+//void AGeometryCollectionDebugDrawActor::DrawConnectivityEdges(const UGeometryCollectionComponent* GeometryCollectionComponent, int32 TransformIndex, const FGeometryCollectionParticlesData& ParticlesData, const TManagedArray<int32>& RigidBodyIdArray, FColor HSVColor)
+void AGeometryCollectionDebugDrawActor::DrawConnectivityEdges(const UGeometryCollectionComponent* GeometryCollectionComponent, int32 TransformIndex, const FGeometryCollectionParticlesData& ParticlesData, const TManagedArray<FGuid>& RigidBodyIdArray, FColor HSVColor)
 {
 #if ENABLE_DRAW_DEBUG
 	check(GeometryCollectionComponent);
@@ -2631,9 +2644,11 @@ void AGeometryCollectionDebugDrawActor::DrawConnectivityEdges(const UGeometryCol
 #endif  // #if ENABLE_DRAW_DEBUG
 }
 
-void AGeometryCollectionDebugDrawActor::DrawConnectivityEdgesNoChecks(const UGeometryCollectionComponent* GeometryCollectionComponent, int32 TransformIndex, const FGeometryCollectionParticlesData& ParticlesData, const TManagedArray<int32>& RigidBodyIdArray, const FColor& Color)
+//void AGeometryCollectionDebugDrawActor::DrawConnectivityEdgesNoChecks(const UGeometryCollectionComponent* GeometryCollectionComponent, int32 TransformIndex, const FGeometryCollectionParticlesData& ParticlesData, const TManagedArray<int32>& RigidBodyIdArray, const FColor& Color)
+void AGeometryCollectionDebugDrawActor::DrawConnectivityEdgesNoChecks(const UGeometryCollectionComponent* GeometryCollectionComponent, int32 TransformIndex, const FGeometryCollectionParticlesData& ParticlesData, const TManagedArray<FGuid>& RigidBodyIdArray, const FColor& Color)
 {
 #if ENABLE_DRAW_DEBUG
+#if TODO_REIMPLEMENT_RIGID_CLUSTERING
 	const UWorld* const World = GetWorld();
 
 	// Get parent index
@@ -2692,6 +2707,7 @@ void AGeometryCollectionDebugDrawActor::DrawConnectivityEdgesNoChecks(const UGeo
 			bNeedsDebugLinesFlush = true;
 		}
 	}
+#endif // TODO_REIMPLEMENT_RIGID_CLUSTERING
 #endif  // #if ENABLE_DRAW_DEBUG
 }
 
diff --git a/Engine/Source/Runtime/Experimental/GeometryCollectionEngine/Private/GeometryCollection/GeometryCollectionDebugDrawComponent.cpp b/Engine/Source/Runtime/Experimental/GeometryCollectionEngine/Private/GeometryCollection/GeometryCollectionDebugDrawComponent.cpp
--- a/Engine/Source/Runtime/Experimental/GeometryCollectionEngine/Private/GeometryCollection/GeometryCollectionDebugDrawComponent.cpp
+++ b/Engine/Source/Runtime/Experimental/GeometryCollectionEngine/Private/GeometryCollection/GeometryCollectionDebugDrawComponent.cpp
@@ -18,6 +18,7 @@ DEFINE_LOG_CATEGORY_STATIC(LogGeometryCollectionDebugDraw, Log, All);
 // Static member variables
 UGeometryCollectionDebugDrawComponent* UGeometryCollectionDebugDrawComponent::RenderLevelSetOwner = nullptr;
 int32 UGeometryCollectionDebugDrawComponent::LastRenderedId = INDEX_NONE;
+//FGuid UGeometryCollectionDebugDrawComponent::LastRenderedId = FGuid();
 
 UGeometryCollectionDebugDrawComponent::UGeometryCollectionDebugDrawComponent(const FObjectInitializer& ObjectInitializer)
 	: Super(ObjectInitializer)
@@ -28,6 +29,7 @@ UGeometryCollectionDebugDrawComponent::UGeometryCollectionDebugDrawComponent(con
 	, ParticlesData()
 	, ParentCheckSum(0)
 	, SelectedRigidBodyId(INDEX_NONE)
+	//, SelectedRigidBodyId(FGuid())
 	, SelectedTransformIndex(INDEX_NONE)
 	, HiddenTransformIndex(INDEX_NONE)
 	, bWasVisible(true)
@@ -58,6 +60,7 @@ void UGeometryCollectionDebugDrawComponent::BeginPlay()
 	{
 		// Reset all index variables
 		SelectedRigidBodyId = INDEX_NONE;
+		//SelectedRigidBodyId = FGuid();
 		SelectedTransformIndex = INDEX_NONE;
 		HiddenTransformIndex = INDEX_NONE;
 		bWasVisible = true;
@@ -92,6 +95,7 @@ void UGeometryCollectionDebugDrawComponent::EndPlay(EEndPlayReason::Type ReasonE
 		if (RenderLevelSetOwner == this)
 		{
 			LastRenderedId = INDEX_NONE;
+			//LastRenderedId = FGuid();
 		}
 
 		// Refresh states from end-play properties
@@ -283,6 +287,7 @@ void UGeometryCollectionDebugDrawComponent::DebugDrawTick()
 	const bool bIsOneSelected = bIsSelected && !GeometryCollectionDebugDrawActor->bDebugDrawWholeCollection;
 	const bool bAreAllSelected = (bIsSelected && GeometryCollectionDebugDrawActor->bDebugDrawWholeCollection) ||
 		(GeometryCollectionDebugDrawActor->SelectedRigidBody.Id == INDEX_NONE &&
+		//(GeometryCollectionDebugDrawActor->SelectedRigidBody.Id.IsValid() == false &&
 			GeometryCollectionDebugDrawActor->SelectedRigidBody.Solver == GeometryCollectionComponent->ChaosSolverActor);
 
 	// Compute world space geometry and cluster transforms
@@ -415,6 +420,8 @@ void UGeometryCollectionDebugDrawComponent::UpdateSelectedTransformIndex()
 	if (!GeometryCollectionDebugDrawActor)
 	{
 		SelectedTransformIndex = SelectedRigidBodyId = INDEX_NONE;
+		//SelectedTransformIndex = INDEX_NONE;
+		//SelectedRigidBodyId = FGuid();
 		return;
 	}
 
@@ -431,6 +438,7 @@ void UGeometryCollectionDebugDrawComponent::UpdateSelectedTransformIndex()
 
 	// Simple test to allow for an early exit when nothing has been selected
 	if (GeometryCollectionDebugDrawActor->SelectedRigidBody.Id == INDEX_NONE ||
+	//if (GeometryCollectionDebugDrawActor->SelectedRigidBody.Id.IsValid() == false ||
 		GeometryCollectionDebugDrawActor->SelectedRigidBody.Solver != GeometryCollectionComponent->ChaosSolverActor)
 	{
 		SelectedRigidBodyId = GeometryCollectionDebugDrawActor->SelectedRigidBody.Id;
@@ -440,7 +448,8 @@ void UGeometryCollectionDebugDrawComponent::UpdateSelectedTransformIndex()
 
 	// Check rigid body id sync
 	// Note that this test alone isn't enough to ensure that the rigid body ids are valid.
-	const TManagedArray<int32>& RigidBodyIds = GeometryCollectionComponent->RigidBodyIds;
+	//const TManagedArray<int32>& RigidBodyIds = GeometryCollectionComponent->RigidBodyIds;
+	const TManagedArray<FGuid>& RigidBodyIds = GeometryCollectionComponent->GetRigidBodyGuidArray();
 	if (RigidBodyIds.Num() == 0)
 	{
 		bHasIncompleteRigidBodyIdSync = !!GeometryCollectionComponent->GetTransformArray().Num();
@@ -457,14 +466,16 @@ void UGeometryCollectionDebugDrawComponent::UpdateSelectedTransformIndex()
 	for (int32 TransformIndex = 0; TransformIndex < RigidBodyIds.Num(); ++TransformIndex)
 	{
 		// Is this the selected id?
-		if (RigidBodyIds[TransformIndex] == GeometryCollectionDebugDrawActor->SelectedRigidBody.Id)
+		//if (RigidBodyIds[TransformIndex] == GeometryCollectionDebugDrawActor->SelectedRigidBody.Id)
+		if (static_cast<int32>(GetTypeHash(RigidBodyIds[TransformIndex])) == GeometryCollectionDebugDrawActor->SelectedRigidBody.Id)
 		{
 			SelectedTransformIndex = TransformIndex;
 			bHasIncompleteRigidBodyIdSync = false;  // Found it, the wait for a sync can be canceled
 			break;
 		}
 		// Check the reason behind any invalid index
-		if (RigidBodyIds[TransformIndex] == INDEX_NONE)
+		//if (RigidBodyIds[TransformIndex] == INDEX_NONE)
+		if (!RigidBodyIds[TransformIndex].IsValid())
 		{
 			// Look for detached clusters in order to differentiate un-synced vs empty cluster rigid body ids.
 			int32 ChildTransformIndex = TransformIndex;
@@ -568,6 +579,7 @@ void UGeometryCollectionDebugDrawComponent::UpdateGeometryVisibility(bool bForce
 		const bool bIsSelected = (SelectedTransformIndex != INDEX_NONE);
 		const bool bAreAllSelected = (bIsSelected && GeometryCollectionDebugDrawActor->bDebugDrawWholeCollection) ||
 			(GeometryCollectionDebugDrawActor->SelectedRigidBody.Id == INDEX_NONE && 
+			//(GeometryCollectionDebugDrawActor->SelectedRigidBody.Id.IsValid() == false && 
 			 GeometryCollectionDebugDrawActor->SelectedRigidBody.Solver == GeometryCollectionComponent->ChaosSolverActor);
 		const bool bAreAnySelected = bIsSelected || bAreAllSelected;
 
@@ -676,6 +688,7 @@ void UGeometryCollectionDebugDrawComponent::UpdateTickStatus()
 		// Check whether anything from this component is selected for debug drawing
 		const bool bAreAnySelected = (SelectedTransformIndex != INDEX_NONE ||
 			(GeometryCollectionDebugDrawActor->SelectedRigidBody.Id == INDEX_NONE &&
+			//(GeometryCollectionDebugDrawActor->SelectedRigidBody.Id.IsValid() == false &&
 			 GeometryCollectionDebugDrawActor->SelectedRigidBody.Solver == GeometryCollectionComponent->ChaosSolverActor));
 
 		bIsEnabled =
@@ -719,7 +732,8 @@ void UGeometryCollectionDebugDrawComponent::DebugDrawChaosTick()
 	check(Actor);
 
 	// Retrieve synced particle and clustering data
-	const TManagedArray<int32>& RigidBodyIds = GeometryCollectionComponent->RigidBodyIds;
+	//const TManagedArray<int32>& RigidBodyIds = GeometryCollectionComponent->RigidBodyIds;
+	const TManagedArray<FGuid>& RigidBodyIds = GeometryCollectionComponent->GetRigidBodyGuidArray();
 	const FGeometryCollectionPhysicsProxy* const PhysicsProxy = GeometryCollectionComponent->GetPhysicsProxy();
 	if (PhysicsProxy)
 	{
@@ -753,6 +767,7 @@ void UGeometryCollectionDebugDrawComponent::DebugDrawChaosTick()
 			// If the level set index has changed at run time, then reload the volume
 			// because someone wants to visualize another piece
 			const bool bLevelSetTextureDirty = (RenderLevelSetOwner != this || LastRenderedId == INDEX_NONE || LastRenderedId != SelectedRigidBodyId);
+			//const bool bLevelSetTextureDirty = (RenderLevelSetOwner != this || LastRenderedId.IsValid() == false || LastRenderedId != SelectedRigidBodyId);
 			if (!bLevelSetTextureDirty)
 			{
 				// If we are only updating the transform, or also loading the volume
@@ -770,6 +785,7 @@ void UGeometryCollectionDebugDrawComponent::DebugDrawChaosTick()
 				{
 					UE_LOG(LogGeometryCollectionDebugDraw, Warning, TEXT("Levelset generation failed: %s"), *GetFullName());
 					LastRenderedId = INDEX_NONE;
+					//LastRenderedId = FGuid();
 				}
 				else
 				{
@@ -789,6 +805,7 @@ void UGeometryCollectionDebugDrawComponent::DebugDrawChaosTick()
 	const bool bIsOneSelected = bIsSelected && !GeometryCollectionDebugDrawActor->bDebugDrawWholeCollection;
 	const bool bAreAllSelected = (bIsSelected && GeometryCollectionDebugDrawActor->bDebugDrawWholeCollection) ||
 		(GeometryCollectionDebugDrawActor->SelectedRigidBody.Id == INDEX_NONE &&
+		//(GeometryCollectionDebugDrawActor->SelectedRigidBody.Id.IsValid() == false &&
 		 GeometryCollectionDebugDrawActor->SelectedRigidBody.Solver == GeometryCollectionComponent->ChaosSolverActor);
 
 	if (bIsOneSelected)
@@ -870,6 +887,7 @@ void UGeometryCollectionDebugDrawComponent::UpdateLevelSetVisibility()
 	const bool bShowCollision = bIsSelected && GeometryCollectionDebugDrawActor && GeometryCollectionDebugDrawActor->bShowRigidBodyCollision;
 	
 	if (RenderLevelSetOwner == this && (LastRenderedId == INDEX_NONE || !bShowCollision))
+	//if (RenderLevelSetOwner == this && (LastRenderedId.IsValid() == false || !bShowCollision))
 	{
 		// Disable rendering
 		GeometryCollectionRenderLevelSetActor->SetEnabled(false);
@@ -877,6 +895,7 @@ void UGeometryCollectionDebugDrawComponent::UpdateLevelSetVisibility()
 		// Disown renderer
 		RenderLevelSetOwner = nullptr;
 		LastRenderedId = INDEX_NONE;
+		//LastRenderedId = FGuid();
 	}
 }
 #endif  // #if GEOMETRYCOLLECTION_DEBUG_DRAW
diff --git a/Engine/Source/Runtime/Experimental/GeometryCollectionEngine/Private/GeometryCollection/GeometryCollectionParticlesData.cpp b/Engine/Source/Runtime/Experimental/GeometryCollectionEngine/Private/GeometryCollection/GeometryCollectionParticlesData.cpp
--- a/Engine/Source/Runtime/Experimental/GeometryCollectionEngine/Private/GeometryCollection/GeometryCollectionParticlesData.cpp
+++ b/Engine/Source/Runtime/Experimental/GeometryCollectionEngine/Private/GeometryCollection/GeometryCollectionParticlesData.cpp
@@ -26,7 +26,8 @@ TGeometryCollectionParticlesData<T, d>::TGeometryCollectionParticlesData()
 }
 
 template<class T, int d>
-void TGeometryCollectionParticlesData<T, d>::Sync(const Chaos::FPhysicsSolver* Solver, const TManagedArray<int32>& RigidBodyIds)
+//void TGeometryCollectionParticlesData<T, d>::Sync(const Chaos::FPhysicsSolver* Solver, const TManagedArray<int32>& RigidBodyIds)
+void TGeometryCollectionParticlesData<T, d>::Sync(const Chaos::FPhysicsSolver* Solver, const TManagedArray<FGuid>& RigidBodyIds)
 {
 	// No point in calling twice the sync function within the same frame
 	if (!ensureMsgf(SyncFrame != GFrameCounter, TEXT("Sync should not happen twice during the same tick.")))
@@ -135,10 +136,13 @@ void TGeometryCollectionParticlesData<T, d>::FData::Reset(EGeometryCollectionPar
 }
 
 template<class T, int d>
-void TGeometryCollectionParticlesData<T, d>::FData::Copy(EGeometryCollectionParticlesData Data, const Chaos::FPhysicsSolver* Solver, const TManagedArray<int32>& RigidBodyIds)
+void TGeometryCollectionParticlesData<T, d>::FData::Copy(EGeometryCollectionParticlesData Data, const Chaos::FPhysicsSolver* Solver, const TManagedArray<FGuid>& RigidBodyIds)
 {
 	check(Solver);
 #if TODO_REIMPLEMENT_GET_RIGID_PARTICLES
+	// 10.31.2019 Ryan - This code uses the RigidBodyIds as indices, and that's not going to work anymore.
+	// We need to figure that out...
+
 	const Chaos::TPBDRigidParticles<T, d>& Particles = Solver->GetRigidParticles();
 	const Chaos::TPBDRigidClustering<FPBDRigidsEvolution, FPBDCollisionConstraints, T, d>& Clustering = Solver->GetRigidClustering();
 
diff --git a/Engine/Source/Runtime/Experimental/GeometryCollectionEngine/Public/GeometryCollection/GeometryCollectionComponent.h b/Engine/Source/Runtime/Experimental/GeometryCollectionEngine/Public/GeometryCollection/GeometryCollectionComponent.h
--- a/Engine/Source/Runtime/Experimental/GeometryCollectionEngine/Public/GeometryCollection/GeometryCollectionComponent.h
+++ b/Engine/Source/Runtime/Experimental/GeometryCollectionEngine/Public/GeometryCollection/GeometryCollectionComponent.h
@@ -506,7 +506,8 @@ public:
 #endif  // #if !(UE_BUILD_SHIPPING || UE_BUILD_TEST)
 
 	/**/
-	const TManagedArray<int32>& GetRigidBodyIdArray() const { return RigidBodyIds; }
+	//const TManagedArray<int32>& GetRigidBodyIdArray() const { return RigidBodyIds; }
+	const TManagedArray<FGuid>& GetRigidBodyGuidArray() const { return RestCollection->GetGeometryCollection()->GetAttribute<FGuid>(FName("GUID"), FGeometryCollection::TransformGroup); }
 	const TArray<bool>& GetDisabledFlags() const { return DisabledFlags; }
 
 	virtual void OnCreatePhysicsState() override;
@@ -515,7 +516,7 @@ public:
 	virtual bool HasValidPhysicsState() const override;
 
 	// Mirrored from the proxy on a sync
-	TManagedArray<int32> RigidBodyIds;
+	//TManagedArray<int32> RigidBodyIds;
 	TArray<bool> DisabledFlags;
 	int32 BaseRigidBodyIndex;
 	int32 NumParticlesAdded;
diff --git a/Engine/Source/Runtime/Experimental/GeometryCollectionEngine/Public/GeometryCollection/GeometryCollectionDebugDrawActor.h b/Engine/Source/Runtime/Experimental/GeometryCollectionEngine/Public/GeometryCollection/GeometryCollectionDebugDrawActor.h
--- a/Engine/Source/Runtime/Experimental/GeometryCollectionEngine/Public/GeometryCollection/GeometryCollectionDebugDrawActor.h
+++ b/Engine/Source/Runtime/Experimental/GeometryCollectionEngine/Public/GeometryCollection/GeometryCollectionDebugDrawActor.h
@@ -54,10 +54,12 @@ struct FGeometryCollectionDebugDrawActorSelectedRigidBody
 	GENERATED_USTRUCT_BODY()
 
 	explicit FGeometryCollectionDebugDrawActorSelectedRigidBody(int32 InId = -1) : Id(InId), Solver(nullptr), GeometryCollection(nullptr) {}
+	//explicit FGeometryCollectionDebugDrawActorSelectedRigidBody(FGuid InId = FGuid()) : Id(InId), Solver(nullptr), GeometryCollection(nullptr) {}
 
 	/** Id of the selected rigid body whose to visualize debug informations. Use -1 to visualize all Geometry Collections. */
 	UPROPERTY(EditAnywhere, Category = "Selected Rigid Body", meta = (ClampMin="-1"))
 	int32 Id;
+	//FGuid Id;
 
 	/** Chaos RBD Solver. Will use the world's default solver actor if null. */
 	UPROPERTY(EditAnywhere, Category = "Selected Rigid Body")
@@ -434,10 +436,12 @@ public:
 	static FTransform GetParticleTransform(const UGeometryCollectionComponent* GeometryCollectionComponent, int32 TransformIndex, const FGeometryCollectionParticlesData& ParticlesData);
 
 	/** Draw Chaos' rigid body ids. */
-	void DrawRigidBodiesId(const UGeometryCollectionComponent* GeometryCollectionComponent,  const FGeometryCollectionParticlesData& ParticlesData, const TManagedArray<int32>& RigidBodyIdArray, const FColor& Color);
+	//void DrawRigidBodiesId(const UGeometryCollectionComponent* GeometryCollectionComponent,  const FGeometryCollectionParticlesData& ParticlesData, const TManagedArray<int32>& RigidBodyIdArray, const FColor& Color);
+	void DrawRigidBodiesId(const UGeometryCollectionComponent* GeometryCollectionComponent,  const FGeometryCollectionParticlesData& ParticlesData, const TManagedArray<FGuid>& RigidBodyIdArray, const FColor& Color);
 
 	/** Draw Chaos' rigid body id. */
-	void DrawRigidBodyId(const UGeometryCollectionComponent* GeometryCollectionComponent, int32 TransformIndex, const FGeometryCollectionParticlesData& ParticlesData, const TManagedArray<int32>& RigidBodyIdArray, const FColor& Color);
+	//void DrawRigidBodyId(const UGeometryCollectionComponent* GeometryCollectionComponent, int32 TransformIndex, const FGeometryCollectionParticlesData& ParticlesData, const TManagedArray<int32>& RigidBodyIdArray, const FColor& Color);
+	void DrawRigidBodyId(const UGeometryCollectionComponent* GeometryCollectionComponent, int32 TransformIndex, const FGeometryCollectionParticlesData& ParticlesData, const TManagedArray<FGuid>& RigidBodyIdArray, const FColor& Color);
 
 	/** Draw Chaos' rigid body transform. */
 	void DrawRigidBodiesTransform(const UGeometryCollectionComponent* GeometryCollectionComponent, const FGeometryCollectionParticlesData& ParticlesData, float Scale);
@@ -464,10 +468,12 @@ public:
 	void DrawRigidBodyInfo(const UGeometryCollectionComponent* GeometryCollectionComponent, int32 TransformIndex, const FGeometryCollectionParticlesData& ParticlesData, const FColor& Color);
 
 	/** Draw Chaos' rigid clustering's connectivities edges */
-	void DrawConnectivityEdges(const UGeometryCollectionComponent* GeometryCollectionComponent, const FGeometryCollectionParticlesData& ParticlesData, const TManagedArray<int32>& RigidBodyIdArray);
+	//void DrawConnectivityEdges(const UGeometryCollectionComponent* GeometryCollectionComponent, const FGeometryCollectionParticlesData& ParticlesData, const TManagedArray<int32>& RigidBodyIdArray);
+	void DrawConnectivityEdges(const UGeometryCollectionComponent* GeometryCollectionComponent, const FGeometryCollectionParticlesData& ParticlesData, const TManagedArray<FGuid>& RigidBodyIdArray);
 
 	/** Draw Chaos' single rigid clustering's connectivity edges. */
-	void DrawConnectivityEdges(const UGeometryCollectionComponent* GeometryCollectionComponent, int32 TransformIndex, const FGeometryCollectionParticlesData& ParticlesData, const TManagedArray<int32>& RigidBodyIdArray, FColor HSVColor = FColor(157, 160, 128));
+	//void DrawConnectivityEdges(const UGeometryCollectionComponent* GeometryCollectionComponent, int32 TransformIndex, const FGeometryCollectionParticlesData& ParticlesData, const TManagedArray<int32>& RigidBodyIdArray, FColor HSVColor = FColor(157, 160, 128));
+	void DrawConnectivityEdges(const UGeometryCollectionComponent* GeometryCollectionComponent, int32 TransformIndex, const FGeometryCollectionParticlesData& ParticlesData, const TManagedArray<FGuid>& RigidBodyIdArray, FColor HSVColor = FColor(157, 160, 128));
 
 	/** Draw Chaos' rigid body velocity. */
 	void DrawRigidBodiesVelocity(const UGeometryCollectionComponent* GeometryCollectionComponent, const FGeometryCollectionParticlesData& ParticlesData, const FColor& Color);
@@ -514,7 +520,8 @@ private:
 	static FTransform GetParticleTransformNoChecks(const UGeometryCollectionComponent* GeometryCollectionComponent, int32 TransformIndex, const FGeometryCollectionParticlesData& ParticlesData);
 
 	/** Draw Chaos' single rigid body id without synchronization checks. */
-	void DrawRigidBodyIdNoChecks(const UGeometryCollectionComponent* GeometryCollectionComponent, int32 TransformIndex, const FGeometryCollectionParticlesData& ParticlesData, const TManagedArray<int32>& RigidBodyIdArray, const FColor& Color);
+	//void DrawRigidBodyIdNoChecks(const UGeometryCollectionComponent* GeometryCollectionComponent, int32 TransformIndex, const FGeometryCollectionParticlesData& ParticlesData, const TManagedArray<int32>& RigidBodyIdArray, const FColor& Color);
+	void DrawRigidBodyIdNoChecks(const UGeometryCollectionComponent* GeometryCollectionComponent, int32 TransformIndex, const FGeometryCollectionParticlesData& ParticlesData, const TManagedArray<FGuid>& RigidBodyIdArray, const FColor& Color);
 
 	/** Draw Chaos' single rigid body transform without synchronization checks. */
 	void DrawRigidBodyTransformNoChecks(const UGeometryCollectionComponent* GeometryCollectionComponent, int32 TransformIndex, const FGeometryCollectionParticlesData& ParticlesData, float Scale);
@@ -529,7 +536,8 @@ private:
 	void DrawRigidBodyInfoNoChecks(const UGeometryCollectionComponent* GeometryCollectionComponent, int32 TransformIndex, const FGeometryCollectionParticlesData& ParticlesData, const FColor& Color);
 
 	/** Draw Chaos' single rigid clustering's connectivity edges without synchronization checks. */
-	void DrawConnectivityEdgesNoChecks(const UGeometryCollectionComponent* GeometryCollectionComponent, int32 TransformIndex, const FGeometryCollectionParticlesData& ParticlesData, const TManagedArray<int32>& RigidBodyIdArray, const FColor& Color);
+	//void DrawConnectivityEdgesNoChecks(const UGeometryCollectionComponent* GeometryCollectionComponent, int32 TransformIndex, const FGeometryCollectionParticlesData& ParticlesData, const TManagedArray<int32>& RigidBodyIdArray, const FColor& Color);
+	void DrawConnectivityEdgesNoChecks(const UGeometryCollectionComponent* GeometryCollectionComponent, int32 TransformIndex, const FGeometryCollectionParticlesData& ParticlesData, const TManagedArray<FGuid>& RigidBodyIdArray, const FColor& Color);
 
 	/** Draw Chaos' single rigid body velocity without synchronization checks. */
 	void DrawRigidBodyVelocityNoChecks(const UGeometryCollectionComponent* GeometryCollectionComponent, int32 TransformIndex, const FGeometryCollectionParticlesData& ParticlesData, const FColor& Color);
diff --git a/Engine/Source/Runtime/Experimental/GeometryCollectionEngine/Public/GeometryCollection/GeometryCollectionDebugDrawComponent.h b/Engine/Source/Runtime/Experimental/GeometryCollectionEngine/Public/GeometryCollection/GeometryCollectionDebugDrawComponent.h
--- a/Engine/Source/Runtime/Experimental/GeometryCollectionEngine/Public/GeometryCollection/GeometryCollectionDebugDrawComponent.h
+++ b/Engine/Source/Runtime/Experimental/GeometryCollectionEngine/Public/GeometryCollection/GeometryCollectionDebugDrawComponent.h
@@ -98,11 +98,13 @@ private:
 private:
 	static UGeometryCollectionDebugDrawComponent* RenderLevelSetOwner;
 	static int32 LastRenderedId;
+	//static FGuid LastRenderedId;
 
 #if GEOMETRYCOLLECTION_DEBUG_DRAW
 	FGeometryCollectionParticlesData ParticlesData;
 	int32 ParentCheckSum;
 	int32 SelectedRigidBodyId;
+	//FGuid SelectedRigidBodyId;
 	int32 SelectedTransformIndex;
 	int32 HiddenTransformIndex;
 	bool bWasVisible;
diff --git a/Engine/Source/Runtime/Experimental/GeometryCollectionEngine/Public/GeometryCollection/GeometryCollectionParticlesData.h b/Engine/Source/Runtime/Experimental/GeometryCollectionEngine/Public/GeometryCollection/GeometryCollectionParticlesData.h
--- a/Engine/Source/Runtime/Experimental/GeometryCollectionEngine/Public/GeometryCollection/GeometryCollectionParticlesData.h
+++ b/Engine/Source/Runtime/Experimental/GeometryCollectionEngine/Public/GeometryCollection/GeometryCollectionParticlesData.h
@@ -78,7 +78,8 @@ public:
 	bool RequestSyncedData(EGeometryCollectionParticlesData Data) const { SetDataSyncFlag(Data); return HasSyncedData(Data); }
 
 	/** Copy the data of the specified set of particles/rigid body ids to this object. */
-	void Sync(const Chaos::FPhysicsSolver* Solver, const TManagedArray<int32>& RigidBodyIds);
+	//void Sync(const Chaos::FPhysicsSolver* Solver, const TManagedArray<int32>& RigidBodyIds);
+	void Sync(const Chaos::FPhysicsSolver* Solver, const TManagedArray<FGuid>& RigidBodyIds);
 
 	const Chaos::TVector<T, d>               & GetX                     (int32 Index) const { check(HasSyncedData(EGeometryCollectionParticlesData::X                     )); return BufferedData.GetGameDataForRead().X                     [Index]; }
 	const Chaos::TRotation<T, d>             & GetR                     (int32 Index) const { check(HasSyncedData(EGeometryCollectionParticlesData::R                     )); return BufferedData.GetGameDataForRead().R                     [Index]; }
@@ -165,7 +166,7 @@ private:
 		void Reset(EGeometryCollectionParticlesData Data);
 
 		/** Copy the specified particle information for the specified range of rigid body id. */
-		void Copy(EGeometryCollectionParticlesData Data, const Chaos::FPhysicsSolver* Solver, const TManagedArray<int32>& RigidBodyIds);
+		void Copy(EGeometryCollectionParticlesData Data, const Chaos::FPhysicsSolver* Solver, const TManagedArray<FGuid>& RigidBodyIds);
 
 		/** Return a string with the entire set of value for the synced data of the specified particle. */
 		FString ToString(int32 Index, const TCHAR* Separator) const;
