Make sure EPA has max iterations. Added tests for deep sphere and capsule penetration

diff --git a/Engine/Source/Runtime/Experimental/Chaos/Public/Chaos/EPA.h b/Engine/Source/Runtime/Experimental/Chaos/Public/Chaos/EPA.h
--- a/Engine/Source/Runtime/Experimental/Chaos/Public/Chaos/EPA.h
+++ b/Engine/Source/Runtime/Experimental/Chaos/Public/Chaos/EPA.h
@@ -85,6 +85,34 @@ struct TEPAEntry
 	{
 		return TVec3<T>::DotProduct(PlaneNormal, X) - Distance;
 	}
+
+	bool IsOriginProjectedInside(const TVec3<T>* VertsBuffer) const
+	{
+		//Compare the projected point (PlaneNormal) to the triangle in the plane
+		const TVec3<T> PA = VertsBuffer[IdxBuffer[0]] - PlaneNormal;
+		const TVec3<T> PB = VertsBuffer[IdxBuffer[1]] - PlaneNormal;
+		const TVec3<T> PC = VertsBuffer[IdxBuffer[2]] - PlaneNormal;
+
+		const TVec3<T> PACNormal = TVec3<T>::CrossProduct(PA, PC);
+		const T PACSign = TVec3<T>::DotProduct(PACNormal, PlaneNormal);
+		const TVec3<T> PCBNormal = TVec3<T>::CrossProduct(PC, PB);
+		const T PCBSign = TVec3<T>::DotProduct(PCBNormal, PlaneNormal);
+
+		if(PACSign < 0 && PCBSign > 0 || PACSign > 0 && PCBSign < 0)
+		{
+			return false;
+		}
+
+		const TVec3<T> PBANormal = TVec3<T>::CrossProduct(PB, PA);
+		const T PBASign = TVec3<T>::DotProduct(PBANormal, PlaneNormal);
+
+		if(PACSign < 0 && PBASign > 0 || PACSign > 0 && PBASign < 0)
+		{
+			return false;
+		}
+
+		return true;
+	}
 };
 
 template <typename T>
@@ -183,38 +211,59 @@ T EPA(TArray<TVec3<T>> VertsBuffer, const SupportLambda& Support)
 	std::priority_queue<FEPAEntryWrapper, std::vector<FEPAEntryWrapper>, std::greater<FEPAEntryWrapper>> Queue;
 	for(int32 Idx = 0; Idx < Entries.Num(); ++Idx)
 	{
-		Queue.push(FEPAEntryWrapper {&Entries, Idx});
+		if(Entries[Idx].IsOriginProjectedInside(VertsBuffer.GetData()))
+		{
+			Queue.push(FEPAEntryWrapper {&Entries, Idx});
+		}
 	}
 
 	bool bCloseEnough = false;
 	bool bTerminate = bCloseEnough;
+	
+	//TEPAEntry<T> BestEntry;
+	//BestEntry.Distance = 0;
+
+	TEPAEntry<T> LastEntry;
+	LastEntry.Distance = 0;
 
 	TArray<FEPAFloodEntry> VisibilityBorder;
-	int32 Iteration = -1;
-	while(Queue.size())
+	int32 Iteration = 0;
+	int32 constexpr MaxIterations = 128;
+	while(Queue.size() && Iteration++ < MaxIterations)
 	{
-		++Iteration;
-
 		int32 EntryIdx = Queue.top().Idx;
 		Queue.pop();
 		TEPAEntry<T>& Entry = Entries[EntryIdx];
+		bool bBadFace = Entry.IsOriginProjectedInside(VertsBuffer.GetData());
+		{
+			//UE_LOG(LogChaos, Warning, TEXT("%d BestW:%f, Distance:%f, bObsolete:%d, InTriangle:%d"),
+			//	Iteration, BestW.Size(), Entry.Distance, Entry.bObsolete, bBadFace);
+		}
 		if (Entry.bObsolete)
 		{
 			continue;
 		}
 
-		constexpr T Eps = 1e-4;
+		LastEntry = Entry;
 
+		constexpr T Eps = 1e-1;
 
 		const TVec3<T> W = Support(Entry.PlaneNormal);
 		const T DistanceToSupportPlane = TVec3<T>::DotProduct(Entry.PlaneNormal, W);
 
-		UpperBound = FMath::Min(UpperBound, DistanceToSupportPlane);
-		ensure(LowerBound <= Entry.Distance);	//todo: remove this check, just for early development
+		if(DistanceToSupportPlane < UpperBound)
+		{
+			UpperBound = DistanceToSupportPlane;
+			//Remember the entry that gave us the lowest upper bound and use it in case we have to terminate early
+			//This can result in very deep planes. Ideally we'd just use the plane formed at W, but it's not clear how you get back points in A, B for that
+			//BestEntry = Entry;
+		}
+
 		LowerBound = Entry.Distance;
 
-		if (FMath::Abs(UpperBound - LowerBound) < Eps)	//todo: add relative?
+		if (FMath::Abs(UpperBound - LowerBound) < Eps)	//todo: add relative error?
 		{
+			//UE_LOG(LogChaos, Warning, TEXT("Iteration:%d"), Iteration);
 			return Entry.Distance;
 		}
 
@@ -239,22 +288,26 @@ T EPA(TArray<TVec3<T>> VertsBuffer, const SupportLambda& Support)
 			const int32 BorderEdgeIdx1 = (BorderEdgeIdx0 + 1) % 3;
 			const int32 NextEntryIdx = (VisibilityIdx + 1) < VisibilityBorder.Num() ? NewIdx + 1 : FirstIdxInBatch;
 			const int32 PrevEntryIdx = NewIdx > FirstIdxInBatch ? NewIdx - 1 : FirstIdxInBatch + NumBorderEdges - 1;
-			ensure(NewEntry.Initialize(VertsBuffer.GetData(), BorderEntry.IdxBuffer[BorderEdgeIdx1], BorderEntry.IdxBuffer[BorderEdgeIdx0], NewVertIdx,
+			const bool bValidTri = NewEntry.Initialize(VertsBuffer.GetData(), BorderEntry.IdxBuffer[BorderEdgeIdx1], BorderEntry.IdxBuffer[BorderEdgeIdx0], NewVertIdx,
 				{ BorderEntryIdx, PrevEntryIdx, NextEntryIdx },
-				{ BorderEdgeIdx0, 2, 1 }));	//todo: handle false case which is a degenerate triangle
+				{ BorderEdgeIdx0, 2, 1 });
 			BorderEntry.AdjFaces[BorderEdgeIdx0] = NewIdx;
 			BorderEntry.AdjEdges[BorderEdgeIdx0] = 0;
 
-			ensure(NewEntry.Distance >= LowerBound);	//todo: remove, just used for early dev
-			if (/*NewEntry.Distance >= LowerBound && */NewEntry.Distance <= UpperBound)
+			//We should never need to check the lower bound, but in the case of bad precision this can happen
+			//We simply ignore this direction as it likely has even more bad precision
+			if (bValidTri && NewEntry.Distance >= LowerBound && NewEntry.Distance <= UpperBound)
 			{
-				Queue.push(FEPAEntryWrapper{ &Entries, NewIdx});
+				if(NewEntry.IsOriginProjectedInside(VertsBuffer.GetData()))
+				{
+					Queue.push(FEPAEntryWrapper{ &Entries, NewIdx});
+				}
 			}
 
 			++NewIdx;
 		}
 	}
 
-	return 0;
+	return LastEntry.Distance;
 }
 }
diff --git a/Engine/Source/Runtime/Experimental/Chaos/Public/Chaos/GJK.h b/Engine/Source/Runtime/Experimental/Chaos/Public/Chaos/GJK.h
--- a/Engine/Source/Runtime/Experimental/Chaos/Public/Chaos/GJK.h
+++ b/Engine/Source/Runtime/Experimental/Chaos/Public/Chaos/GJK.h
@@ -2,7 +2,6 @@
 #pragma once
 
 #include "Chaos/Simplex.h"
-#include "Chaos/EPA.h"
 #include "Chaos/Capsule.h"
 #include "Chaos/Box.h"
 #include "Chaos/Sphere.h"
