Changed EPA API to be more useful. Returning penetration and normal, witness points are still uninitialized

diff --git a/Engine/Source/Runtime/Experimental/Chaos/Public/Chaos/EPA.h b/Engine/Source/Runtime/Experimental/Chaos/Public/Chaos/EPA.h
--- a/Engine/Source/Runtime/Experimental/Chaos/Public/Chaos/EPA.h
+++ b/Engine/Source/Runtime/Experimental/Chaos/Public/Chaos/EPA.h
@@ -6,6 +6,12 @@
 namespace Chaos
 {
 
+template <typename T>
+FORCEINLINE const TVec3<T> MinkowskiVert(const TVec3<T>* VertsABuffer, const TVec3<T>* VertsBBuffer, const int32 Idx)
+{
+	return VertsABuffer[Idx] - VertsBBuffer[Idx];
+}
+
 template <typename T>
 struct TEPAEntry
 {
@@ -22,11 +28,11 @@ struct TEPAEntry
 		return Distance > Other.Distance;
 	}
 
-	bool Initialize(const TVec3<T>* Vertices, int32 InIdx0, int32 InIdx1, int32 InIdx2, const TVector<int32,3>& InAdjFaces, const TVector<int32,3>& InAdjEdges)
+	bool Initialize(const TVec3<T>* VerticesA, const TVec3<T>* VerticesB, int32 InIdx0, int32 InIdx1, int32 InIdx2, const TVector<int32,3>& InAdjFaces, const TVector<int32,3>& InAdjEdges)
 	{
-		const TVec3<T>& V0 = Vertices[InIdx0];
-		const TVec3<T>& V1 = Vertices[InIdx1];
-		const TVec3<T>& V2 = Vertices[InIdx2];
+		const TVec3<T>& V0 = MinkowskiVert(VerticesA, VerticesB, InIdx0);
+		const TVec3<T>& V1 = MinkowskiVert(VerticesA, VerticesB, InIdx1);
+		const TVec3<T>& V2 = MinkowskiVert(VerticesA, VerticesB, InIdx2);
 
 		const TVec3<T> V0V1 = V1 - V0;
 		const TVec3<T> V0V2 = V2 - V0;
@@ -86,12 +92,12 @@ struct TEPAEntry
 		return TVec3<T>::DotProduct(PlaneNormal, X) - Distance;
 	}
 
-	bool IsOriginProjectedInside(const TVec3<T>* VertsBuffer) const
+	bool IsOriginProjectedInside(const TVec3<T>* VertsABuffer, const TVec3<T>* VertsBBuffer) const
 	{
 		//Compare the projected point (PlaneNormal) to the triangle in the plane
-		const TVec3<T> PA = VertsBuffer[IdxBuffer[0]] - PlaneNormal;
-		const TVec3<T> PB = VertsBuffer[IdxBuffer[1]] - PlaneNormal;
-		const TVec3<T> PC = VertsBuffer[IdxBuffer[2]] - PlaneNormal;
+		const TVec3<T> PA = MinkowskiVert(VertsABuffer, VertsBBuffer, IdxBuffer[0]) - PlaneNormal;
+		const TVec3<T> PB = MinkowskiVert(VertsABuffer, VertsBBuffer, IdxBuffer[1]) - PlaneNormal;
+		const TVec3<T> PC = MinkowskiVert(VertsABuffer, VertsBBuffer, IdxBuffer[2]) - PlaneNormal;
 
 		const TVec3<T> PACNormal = TVec3<T>::CrossProduct(PA, PC);
 		const T PACSign = TVec3<T>::DotProduct(PACNormal, PlaneNormal);
@@ -116,7 +122,7 @@ struct TEPAEntry
 };
 
 template <typename T>
-TArray<TEPAEntry<T>> InitializeEPA(const TVec3<T>* Verts, const int32 NumVerts)
+TArray<TEPAEntry<T>> InitializeEPA(const TVec3<T>* VertsA, const TVec3<T>* VertsB, const int32 NumVerts)
 {
 	TArray<TEPAEntry<T>> Entries;
 	switch(NumVerts)
@@ -126,12 +132,12 @@ TArray<TEPAEntry<T>> InitializeEPA(const TVec3<T>* Verts, const int32 NumVerts)
 			//make sure all triangle normals are facing out
 			Entries.AddUninitialized(4);
 
-			ensure(Entries[0].Initialize(Verts, 1, 2, 3, { 3, 1, 2 }, { 1,1, 1 }));
-			ensure(Entries[1].Initialize(Verts, 0, 3, 2, { 2,0,3 }, { 2, 1, 0 }));
-			ensure(Entries[2].Initialize(Verts, 0, 1, 3, { 3,0, 1 }, { 2,2,0 }));
-			ensure(Entries[3].Initialize(Verts, 0, 2, 1, { 1,0,2 }, { 2,0,0 }));
+			ensure(Entries[0].Initialize(VertsA, VertsB, 1, 2, 3, { 3, 1, 2 }, { 1,1, 1 }));
+			ensure(Entries[1].Initialize(VertsA, VertsB, 0, 3, 2, { 2,0,3 }, { 2, 1, 0 }));
+			ensure(Entries[2].Initialize(VertsA, VertsB, 0, 1, 3, { 3,0, 1 }, { 2,2,0 }));
+			ensure(Entries[3].Initialize(VertsA, VertsB, 0, 2, 1, { 1,0,2 }, { 2,0,0 }));
 
-			if (TVec3<T>::DotProduct(Entries[0].PlaneNormal, Verts[0]) > 0)
+			if (TVec3<T>::DotProduct(Entries[0].PlaneNormal, MinkowskiVert(VertsA, VertsB, 0)) > 0)
 			{
 				//tet faces are pointing inwards
 				for (TEPAEntry<T>& Entry : Entries)
@@ -190,8 +196,14 @@ void EPAComputeVisibilityBorder(TArray<TEPAEntry<T>>& Entries, int32 EntryIdx, c
 	}
 }
 
-template <typename T, typename SupportLambda>
-T EPA(TArray<TVec3<T>> VertsBuffer, const SupportLambda& Support)
+enum EPAResult
+{
+	Ok,
+	MaxIterations,
+};
+
+template <typename T, typename SupportALambda, typename SupportBLambda>
+EPAResult EPA(TArray<TVec3<T>>& VertsABuffer, TArray<TVec3<T>>& VertsBBuffer, const SupportALambda& SupportA, const SupportBLambda& SupportB, T& OutPenetration, TVec3<T>& OutDir, TVec3<T>& WitnessA, TVec3<T>& WitnessB)
 {
 	struct FEPAEntryWrapper
 	{
@@ -207,11 +219,11 @@ T EPA(TArray<TVec3<T>> VertsBuffer, const SupportLambda& Support)
 	T UpperBound = TNumericLimits<T>::Max();
 	T LowerBound = TNumericLimits<T>::Lowest();
 
-	TArray<TEPAEntry<T>> Entries = InitializeEPA(VertsBuffer.GetData(), VertsBuffer.Num());
+	TArray<TEPAEntry<T>> Entries = InitializeEPA(VertsABuffer.GetData(), VertsBBuffer.GetData(), VertsABuffer.Num());
 	std::priority_queue<FEPAEntryWrapper, std::vector<FEPAEntryWrapper>, std::greater<FEPAEntryWrapper>> Queue;
 	for(int32 Idx = 0; Idx < Entries.Num(); ++Idx)
 	{
-		if(Entries[Idx].IsOriginProjectedInside(VertsBuffer.GetData()))
+		if(Entries[Idx].IsOriginProjectedInside(VertsABuffer.GetData(), VertsBBuffer.GetData()))
 		{
 			Queue.push(FEPAEntryWrapper {&Entries, Idx});
 		}
@@ -234,7 +246,7 @@ T EPA(TArray<TVec3<T>> VertsBuffer, const SupportLambda& Support)
 		int32 EntryIdx = Queue.top().Idx;
 		Queue.pop();
 		TEPAEntry<T>& Entry = Entries[EntryIdx];
-		bool bBadFace = Entry.IsOriginProjectedInside(VertsBuffer.GetData());
+		bool bBadFace = Entry.IsOriginProjectedInside(VertsABuffer.GetData(), VertsBBuffer.GetData());
 		{
 			//UE_LOG(LogChaos, Warning, TEXT("%d BestW:%f, Distance:%f, bObsolete:%d, InTriangle:%d"),
 			//	Iteration, BestW.Size(), Entry.Distance, Entry.bObsolete, bBadFace);
@@ -246,9 +258,11 @@ T EPA(TArray<TVec3<T>> VertsBuffer, const SupportLambda& Support)
 
 		LastEntry = Entry;
 
-		constexpr T Eps = 1e-1;
+		constexpr T Eps = 1e-2;
 
-		const TVec3<T> W = Support(Entry.PlaneNormal);
+		const TVec3<T> ASupport = SupportA(Entry.PlaneNormal);
+		const TVec3<T> BSupport = SupportB(-Entry.PlaneNormal);
+		const TVec3<T> W = ASupport - BSupport;
 		const T DistanceToSupportPlane = TVec3<T>::DotProduct(Entry.PlaneNormal, W);
 
 		if(DistanceToSupportPlane < UpperBound)
@@ -264,10 +278,13 @@ T EPA(TArray<TVec3<T>> VertsBuffer, const SupportLambda& Support)
 		if (FMath::Abs(UpperBound - LowerBound) < Eps)	//todo: add relative error?
 		{
 			//UE_LOG(LogChaos, Warning, TEXT("Iteration:%d"), Iteration);
-			return Entry.Distance;
+			OutPenetration = Entry.Distance;
+			OutDir = Entry.PlaneNormal;
+			return EPAResult::Ok;
 		}
 
-		const int32 NewVertIdx = VertsBuffer.Add(W);
+		const int32 NewVertIdx = VertsABuffer.Add(ASupport);
+		VertsBBuffer.Add(BSupport);
 
 		Entry.bObsolete = true;
 		VisibilityBorder.Reset();
@@ -288,7 +305,7 @@ T EPA(TArray<TVec3<T>> VertsBuffer, const SupportLambda& Support)
 			const int32 BorderEdgeIdx1 = (BorderEdgeIdx0 + 1) % 3;
 			const int32 NextEntryIdx = (VisibilityIdx + 1) < VisibilityBorder.Num() ? NewIdx + 1 : FirstIdxInBatch;
 			const int32 PrevEntryIdx = NewIdx > FirstIdxInBatch ? NewIdx - 1 : FirstIdxInBatch + NumBorderEdges - 1;
-			const bool bValidTri = NewEntry.Initialize(VertsBuffer.GetData(), BorderEntry.IdxBuffer[BorderEdgeIdx1], BorderEntry.IdxBuffer[BorderEdgeIdx0], NewVertIdx,
+			const bool bValidTri = NewEntry.Initialize(VertsABuffer.GetData(), VertsBBuffer.GetData(), BorderEntry.IdxBuffer[BorderEdgeIdx1], BorderEntry.IdxBuffer[BorderEdgeIdx0], NewVertIdx,
 				{ BorderEntryIdx, PrevEntryIdx, NextEntryIdx },
 				{ BorderEdgeIdx0, 2, 1 });
 			BorderEntry.AdjFaces[BorderEdgeIdx0] = NewIdx;
@@ -298,7 +315,7 @@ T EPA(TArray<TVec3<T>> VertsBuffer, const SupportLambda& Support)
 			//We simply ignore this direction as it likely has even more bad precision
 			if (bValidTri && NewEntry.Distance >= LowerBound && NewEntry.Distance <= UpperBound)
 			{
-				if(NewEntry.IsOriginProjectedInside(VertsBuffer.GetData()))
+				if(NewEntry.IsOriginProjectedInside(VertsABuffer.GetData(), VertsBBuffer.GetData()))
 				{
 					Queue.push(FEPAEntryWrapper{ &Entries, NewIdx});
 				}
@@ -308,6 +325,8 @@ T EPA(TArray<TVec3<T>> VertsBuffer, const SupportLambda& Support)
 		}
 	}
 
-	return LastEntry.Distance;
+	OutPenetration = LastEntry.Distance;
+	OutDir = LastEntry.PlaneNormal;
+	return EPAResult::MaxIterations;
 }
 }
