Fix uninitialized warning

diff --git a/Engine/Source/Runtime/Experimental/Chaos/Public/Chaos/EPA.h b/Engine/Source/Runtime/Experimental/Chaos/Public/Chaos/EPA.h
--- a/Engine/Source/Runtime/Experimental/Chaos/Public/Chaos/EPA.h
+++ b/Engine/Source/Runtime/Experimental/Chaos/Public/Chaos/EPA.h
@@ -237,6 +237,7 @@ EPAResult EPA(TArray<TVec3<T>>& VertsABuffer, TArray<TVec3<T>>& VertsBBuffer, co
 
 	TEPAEntry<T> LastEntry;
 	LastEntry.Distance = 0;
+	LastEntry.PlaneNormal = TVec3<T>(0,0,1);
 
 	TArray<FEPAFloodEntry> VisibilityBorder;
 	int32 Iteration = 0;
