EPA returns witness points

diff --git a/Engine/Source/Runtime/Experimental/Chaos/Public/Chaos/EPA.h b/Engine/Source/Runtime/Experimental/Chaos/Public/Chaos/EPA.h
--- a/Engine/Source/Runtime/Experimental/Chaos/Public/Chaos/EPA.h
+++ b/Engine/Source/Runtime/Experimental/Chaos/Public/Chaos/EPA.h
@@ -196,6 +196,37 @@ void EPAComputeVisibilityBorder(TArray<TEPAEntry<T>>& Entries, int32 EntryIdx, c
 	}
 }
 
+template <typename T>
+void ComputeEPAResults(const TVec3<T>* VertsA, const TVec3<T>* VertsB, const TEPAEntry<T>& Entry, T& OutPenetration, TVec3<T>& OutDir, TVec3<T>& OutA, TVec3<T>& OutB)
+{
+	FSimplex SimplexIDs({ 0,1,2 });
+	TVec3<T> As[4] = { VertsA[Entry.IdxBuffer[0]], VertsA[Entry.IdxBuffer[1]], VertsA[Entry.IdxBuffer[2]] };
+	TVec3<T> Bs[4] = { VertsB[Entry.IdxBuffer[0]], VertsB[Entry.IdxBuffer[1]], VertsB[Entry.IdxBuffer[2]] };
+	TVec3<T> Simplex[4] = { As[0] - Bs[0], As[1] - Bs[1], As[2] - Bs[2] };
+	T Barycentric[4];
+
+	OutDir = SimplexFindClosestToOrigin(Simplex, SimplexIDs, Barycentric, As, Bs);
+	OutPenetration = OutDir.Size();
+
+	if (OutPenetration < 1e-4)	//if closest point is on the origin (edge case when surface is right on the origin)
+	{
+		OutDir = Entry.PlaneNormal;	//just fall back on plane normal
+	}
+	else
+	{
+		OutDir /= OutPenetration;
+	}
+
+	OutA = TVec3<T>(0);
+	OutB = TVec3<T>(0);
+
+	for (int i = 0; i < SimplexIDs.NumVerts; ++i)
+	{
+		OutA += As[i] * Barycentric[i];
+		OutB += Bs[i] * Barycentric[i];
+	}
+}
+
 enum EPAResult
 {
 	Ok,
@@ -279,8 +310,7 @@ EPAResult EPA(TArray<TVec3<T>>& VertsABuffer, TArray<TVec3<T>>& VertsBBuffer, co
 		if (FMath::Abs(UpperBound - LowerBound) < Eps)	//todo: add relative error?
 		{
 			//UE_LOG(LogChaos, Warning, TEXT("Iteration:%d"), Iteration);
-			OutPenetration = Entry.Distance;
-			OutDir = Entry.PlaneNormal;
+			ComputeEPAResults(VertsABuffer.GetData(), VertsBBuffer.GetData(), Entry, OutPenetration, OutDir, WitnessA, WitnessB);
 			return EPAResult::Ok;
 		}
 
@@ -326,8 +356,8 @@ EPAResult EPA(TArray<TVec3<T>>& VertsABuffer, TArray<TVec3<T>>& VertsBBuffer, co
 		}
 	}
 
-	OutPenetration = LastEntry.Distance;
-	OutDir = LastEntry.PlaneNormal;
+	ComputeEPAResults(VertsABuffer.GetData(), VertsBBuffer.GetData(), LastEntry, OutPenetration, OutDir, WitnessA, WitnessB);
+	
 	return EPAResult::MaxIterations;
 }
 }
