Fix clang

diff --git a/Engine/Source/Runtime/Experimental/Chaos/Public/Chaos/EPA.h b/Engine/Source/Runtime/Experimental/Chaos/Public/Chaos/EPA.h
--- a/Engine/Source/Runtime/Experimental/Chaos/Public/Chaos/EPA.h
+++ b/Engine/Source/Runtime/Experimental/Chaos/Public/Chaos/EPA.h
@@ -104,7 +104,7 @@ struct TEPAEntry
 		const TVec3<T> PCBNormal = TVec3<T>::CrossProduct(PC, PB);
 		const T PCBSign = TVec3<T>::DotProduct(PCBNormal, PlaneNormal);
 
-		if(PACSign < 0 && PCBSign > 0 || PACSign > 0 && PCBSign < 0)
+		if((PACSign < 0 && PCBSign > 0) || (PACSign > 0 && PCBSign < 0))
 		{
 			return false;
 		}
@@ -112,7 +112,7 @@ struct TEPAEntry
 		const TVec3<T> PBANormal = TVec3<T>::CrossProduct(PB, PA);
 		const T PBASign = TVec3<T>::DotProduct(PBANormal, PlaneNormal);
 
-		if(PACSign < 0 && PBASign > 0 || PACSign > 0 && PBASign < 0)
+		if((PACSign < 0 && PBASign > 0) || (PACSign > 0 && PBASign < 0))
 		{
 			return false;
 		}
