Chaos: Resubmit 10086509 with missing shape bound initialization paths.

diff --git a/Engine/Source/Runtime/Core/Public/UObject/ExternalPhysicsCustomObjectVersion.h b/Engine/Source/Runtime/Core/Public/UObject/ExternalPhysicsCustomObjectVersion.h
--- a/Engine/Source/Runtime/Core/Public/UObject/ExternalPhysicsCustomObjectVersion.h
+++ b/Engine/Source/Runtime/Core/Public/UObject/ExternalPhysicsCustomObjectVersion.h
@@ -57,6 +57,9 @@ struct CORE_API FExternalPhysicsCustomObjectVersion
 		// Only serialize internal acceleration structure queue and acceleration structure. No external/Async queues.
 		FlushEvolutionInternalAccelerationQueue,
 
+		// Serialize world space bounds of TPerShapeData
+		SerializeShapeWorldSpaceBounds,
+
 		// -----<new versions can be added above this line>-------------------------------------------------
 		VersionPlusOne,
 		LatestVersion = VersionPlusOne - 1
diff --git a/Engine/Source/Runtime/Engine/Private/PhysicsEngine/BodySetup.cpp b/Engine/Source/Runtime/Engine/Private/PhysicsEngine/BodySetup.cpp
--- a/Engine/Source/Runtime/Engine/Private/PhysicsEngine/BodySetup.cpp
+++ b/Engine/Source/Runtime/Engine/Private/PhysicsEngine/BodySetup.cpp
@@ -815,6 +815,7 @@ void UBodySetup::AddShapesToRigidActor_AssumesLocked(
 	AddParams.SimpleMaterial = SimpleMaterial;
 	AddParams.ComplexMaterials = TArrayView<UPhysicalMaterial*>(ComplexMaterials);
 	AddParams.LocalTransform = RelativeTM;
+	AddParams.WorldTransform = OwningInstance->GetUnrealWorldTransform();
 	AddParams.Geometry = &AggGeom;
 #if WITH_PHYSX
 	AddParams.TriMeshes = TArrayView<PxTriangleMesh*>(TriMeshes);
diff --git a/Engine/Source/Runtime/Engine/Private/PhysicsEngine/Experimental/ChaosInterfaceUtils.cpp b/Engine/Source/Runtime/Engine/Private/PhysicsEngine/Experimental/ChaosInterfaceUtils.cpp
--- a/Engine/Source/Runtime/Engine/Private/PhysicsEngine/Experimental/ChaosInterfaceUtils.cpp
+++ b/Engine/Source/Runtime/Engine/Private/PhysicsEngine/Experimental/ChaosInterfaceUtils.cpp
@@ -146,6 +146,12 @@ namespace ChaosInterface
 			NewShape->Geometry = InGeom;
 			NewShape->QueryData = bComplexShape ? InParams.CollisionData.CollisionFilterData.QueryComplexFilter : InParams.CollisionData.CollisionFilterData.QuerySimpleFilter;
 			NewShape->SimData = InParams.CollisionData.CollisionFilterData.SimFilter;
+
+			if (InGeom->HasBoundingBox())
+			{
+				NewShape->WorldSpaceInflatedShapeBounds = InGeom->BoundingBox().GetAABB().TransformedAABB(InParams.WorldTransform);
+			}
+
 			return NewShape;
 		};
 
diff --git a/Engine/Source/Runtime/Engine/Private/PhysicsEngine/ImmediatePhysics/ImmediatePhysicsChaos/ImmediatePhysicsActorHandle_Chaos.cpp b/Engine/Source/Runtime/Engine/Private/PhysicsEngine/ImmediatePhysics/ImmediatePhysicsChaos/ImmediatePhysicsActorHandle_Chaos.cpp
--- a/Engine/Source/Runtime/Engine/Private/PhysicsEngine/ImmediatePhysics/ImmediatePhysicsChaos/ImmediatePhysicsActorHandle_Chaos.cpp
+++ b/Engine/Source/Runtime/Engine/Private/PhysicsEngine/ImmediatePhysics/ImmediatePhysicsChaos/ImmediatePhysicsActorHandle_Chaos.cpp
@@ -167,6 +167,7 @@ namespace ImmediatePhysics_Chaos
 		//AddParams.SimpleMaterial = SimpleMaterial;
 		//AddParams.ComplexMaterials = TArrayView<UPhysicalMaterial*>(ComplexMaterials);
 		AddParams.LocalTransform = Chaos::TRigidTransform<float, 3>(OutCoMTransform.GetRotation().Inverse() * -OutCoMTransform.GetTranslation(), OutCoMTransform.GetRotation().Inverse());
+		AddParams.WorldTransform = BodyInstance->GetUnrealWorldTransform();
 		AddParams.Geometry = &BodySetup->AggGeom;
 #if WITH_PHYSX
 		AddParams.TriMeshes = TArrayView<PxTriangleMesh*>(BodySetup->TriMeshes);
diff --git a/Engine/Source/Runtime/Engine/Public/Physics/PhysicsInterfaceTypes.h b/Engine/Source/Runtime/Engine/Public/Physics/PhysicsInterfaceTypes.h
--- a/Engine/Source/Runtime/Engine/Public/Physics/PhysicsInterfaceTypes.h
+++ b/Engine/Source/Runtime/Engine/Public/Physics/PhysicsInterfaceTypes.h
@@ -49,6 +49,7 @@ struct FGeometryAddParams
 	UPhysicalMaterial* SimpleMaterial;
 	TArrayView<UPhysicalMaterial*> ComplexMaterials;
 	FTransform LocalTransform;
+	FTransform WorldTransform;
 	FKAggregateGeom* Geometry;
 	// FPhysicsInterfaceTriMesh - Per implementation
 #if WITH_PHYSX
diff --git a/Engine/Source/Runtime/Experimental/Chaos/Private/Chaos/GeometryParticles.cpp b/Engine/Source/Runtime/Experimental/Chaos/Private/Chaos/GeometryParticles.cpp
--- a/Engine/Source/Runtime/Experimental/Chaos/Private/Chaos/GeometryParticles.cpp
+++ b/Engine/Source/Runtime/Experimental/Chaos/Private/Chaos/GeometryParticles.cpp
@@ -8,7 +8,7 @@
 namespace Chaos
 {
 	template <typename T, int d>
-	void UpdateShapesArrayFromGeometry(TShapesArray<T, d>& ShapesArray, TSerializablePtr<FImplicitObject> Geometry)
+	void UpdateShapesArrayFromGeometry(TShapesArray<T, d>& ShapesArray, TSerializablePtr<FImplicitObject> Geometry, const FRigidTransform3& ActorTM)
 	{
 		if(Geometry)
 		{
@@ -29,6 +29,14 @@ namespace Chaos
 				ShapesArray[0] = TPerShapeData<T, d>::CreatePerShapeData();
 				ShapesArray[0]->Geometry = Geometry;
 			}
+
+			if (Geometry->HasBoundingBox())
+			{
+				for (auto& Shape : ShapesArray)
+				{
+					Shape->WorldSpaceInflatedShapeBounds = Geometry->BoundingBox().GetAABB().TransformedAABB(ActorTM);
+				}
+			}
 		}
 		else
 		{
@@ -62,9 +70,21 @@ namespace Chaos
 	template <typename T, int d>
 	void TPerShapeData<T, d>::Serialize(FChaosArchive& Ar)
 	{
+		Ar.UsingCustomVersion(FExternalPhysicsCustomObjectVersion::GUID);
+
 		Ar << Geometry;
 		Ar << QueryData;
 		Ar << SimData;
+
+		if (Ar.CustomVer(FExternalPhysicsCustomObjectVersion::GUID) >= FExternalPhysicsCustomObjectVersion::SerializeShapeWorldSpaceBounds)
+		{
+			Ar << WorldSpaceInflatedShapeBounds;
+		}
+		else
+		{
+			// This should be set by particle serializing this TPerShapeData.
+			WorldSpaceInflatedShapeBounds = TAABB<FReal, 3>(FVec3(0.0f, 0.0f, 0.0f), FVec3(0.0f, 0.0f, 0.0f));
+		}
 	}
 
 
@@ -128,5 +148,5 @@ namespace Chaos
 	template class TGeometryParticlesImp<float, 3, EGeometryParticlesSimType::RigidBodySim>;
 	template class TGeometryParticlesImp<float, 3, EGeometryParticlesSimType::Other>;
 	template class TPerShapeData<float, 3>;
-	template void UpdateShapesArrayFromGeometry(TShapesArray<float, 3>& ShapesArray, TSerializablePtr<FImplicitObject> Geometry);
+	template void UpdateShapesArrayFromGeometry(TShapesArray<float, 3>& ShapesArray, TSerializablePtr<FImplicitObject> Geometry, const FRigidTransform3& ActorTM);
 }
diff --git a/Engine/Source/Runtime/Experimental/Chaos/Public/Chaos/AABB.h b/Engine/Source/Runtime/Experimental/Chaos/Public/Chaos/AABB.h
--- a/Engine/Source/Runtime/Experimental/Chaos/Public/Chaos/AABB.h
+++ b/Engine/Source/Runtime/Experimental/Chaos/Public/Chaos/AABB.h
@@ -484,6 +484,12 @@ namespace Chaos
 		TVector<T, d> MMin, MMax;
 	};
 
+	FORCEINLINE FChaosArchive& operator<<(FChaosArchive& Ar, TAABB<FReal, 3>& AABB)
+	{
+		AABB.Serialize(Ar);
+		return Ar;
+	}
+
 	template<typename T>
 	struct TAABBSpecializeSamplingHelper<T, 2>
 	{
diff --git a/Engine/Source/Runtime/Experimental/Chaos/Public/Chaos/GeometryParticles.h b/Engine/Source/Runtime/Experimental/Chaos/Public/Chaos/GeometryParticles.h
--- a/Engine/Source/Runtime/Experimental/Chaos/Public/Chaos/GeometryParticles.h
+++ b/Engine/Source/Runtime/Experimental/Chaos/Public/Chaos/GeometryParticles.h
@@ -32,6 +32,7 @@ namespace Chaos
 		FCollisionFilterData SimData;
 		void* UserData;
 		TSerializablePtr<FImplicitObject> Geometry;
+		TAABB<FReal, 3> WorldSpaceInflatedShapeBounds;
 		
 
 		static TPerShapeData<T, d>* SerializationFactory(FChaosArchive& Ar, TPerShapeData<T, d>*);
@@ -53,9 +54,9 @@ namespace Chaos
 	using TShapesArray = TArray<TUniquePtr<TPerShapeData<T, d>>, TInlineAllocator<1>>;
 
 	template <typename T, int d>
-	void UpdateShapesArrayFromGeometry(TShapesArray<T, d>& ShapesArray, TSerializablePtr<FImplicitObject> Geometry);
+	void UpdateShapesArrayFromGeometry(TShapesArray<T, d>& ShapesArray, TSerializablePtr<FImplicitObject> Geometry, const FRigidTransform3& ActorTM);
 
-	extern template void CHAOS_API UpdateShapesArrayFromGeometry(TShapesArray<float, 3>& ShapesArray, TSerializablePtr<FImplicitObject> Geometry);
+	extern template void CHAOS_API UpdateShapesArrayFromGeometry(TShapesArray<float, 3>& ShapesArray, TSerializablePtr<FImplicitObject> Geometry,  const FRigidTransform3& ActorTM);
 
 #if CHAOS_DETERMINISTIC
 	using FParticleID = int32;	//Used to break ties when determinism is needed. Should not be used for anything else
@@ -262,9 +263,20 @@ namespace Chaos
 			return MWorldSpaceInflatedBounds[Index];
 		}
 
-		CHAOS_API TBox<T, d>& WorldSpaceInflatedBounds(const int32 Index)
+		CHAOS_API void SetWorldSpaceInflatedBounds(const int32 Index, const TBox<T, d>& Bounds)
 		{
-			return MWorldSpaceInflatedBounds[Index];
+			MWorldSpaceInflatedBounds[Index] = Bounds;
+
+			const TShapesArray<FReal, 3>& Shapes = ShapesArray(Index);
+			for (const auto& Shape : Shapes)
+			{
+				if (Shape->Geometry->HasBoundingBox())
+				{
+					const TRigidTransform<FReal, 3> ActorTM(X(Index), R(Index));
+					TAABB<FReal, 3> Bound = Shape->Geometry->BoundingBox().GetAABB().TransformedAABB(ActorTM);
+					Shape->WorldSpaceInflatedShapeBounds = Bound;
+				}
+			}
 		}
 
 		const TArray<TSerializablePtr<FImplicitObject>>& GetAllGeometry() const { return MGeometry; }
@@ -304,6 +316,14 @@ namespace Chaos
 				Ar << MLocalBounds;
 				Ar << MWorldSpaceInflatedBounds;
 				Ar << MHasBounds;
+
+				if (Ar.CustomVer(FExternalPhysicsCustomObjectVersion::GUID) < FExternalPhysicsCustomObjectVersion::SerializeShapeWorldSpaceBounds)
+				{
+					for (int32 Idx = 0; Idx < MShapesArray.Num(); ++Idx)
+					{
+						SetWorldSpaceInflatedBounds(Idx, MWorldSpaceInflatedBounds[Idx]);
+					}
+				}
 			}
 			else
 			{
@@ -314,7 +334,8 @@ namespace Chaos
 					if (MHasBounds[Idx])
 					{
 						MLocalBounds[Idx] = MGeometry[Idx]->BoundingBox();
-						MWorldSpaceInflatedBounds[Idx] = MLocalBounds[Idx].TransformedBox(TRigidTransform<T,d>(X(Idx), R(Idx)));	//ignore velocity too, really just trying to get something reasonable
+						//ignore velocity too, really just trying to get something reasonable)
+						SetWorldSpaceInflatedBounds(Idx, MLocalBounds[Idx].TransformedBox(TRigidTransform<T,d>(X(Idx), R(Idx))));
 					}
 				}
 			}
@@ -368,7 +389,7 @@ namespace Chaos
 
 		void UpdateShapesArray(const int32 Index)
 		{
-			UpdateShapesArrayFromGeometry(MShapesArray[Index], MGeometry[Index]);
+			UpdateShapesArrayFromGeometry(MShapesArray[Index], MGeometry[Index], FRigidTransform3(X(Index), R(Index)));
 		}
 
 		template <typename T2, int d2, EGeometryParticlesSimType SimType2>
diff --git a/Engine/Source/Runtime/Experimental/Chaos/Public/Chaos/ParticleHandle.h b/Engine/Source/Runtime/Experimental/Chaos/Public/Chaos/ParticleHandle.h
--- a/Engine/Source/Runtime/Experimental/Chaos/Public/Chaos/ParticleHandle.h
+++ b/Engine/Source/Runtime/Experimental/Chaos/Public/Chaos/ParticleHandle.h
@@ -380,7 +380,10 @@ public:
 	void SetLocalBounds(const TBox<T, d>& NewBounds) { GeometryParticles->LocalBounds(ParticleIdx) = NewBounds; }
 
 	const TBox<T, d>& WorldSpaceInflatedBounds() const { return GeometryParticles->WorldSpaceInflatedBounds(ParticleIdx); }
-	void SetWorldSpaceInflatedBounds(const TBox<T,d>& WorldSpaceInflatedBounds) { GeometryParticles->WorldSpaceInflatedBounds(ParticleIdx) = WorldSpaceInflatedBounds; }
+	void SetWorldSpaceInflatedBounds(const TBox<T, d>& WorldSpaceInflatedBounds)
+	{
+		GeometryParticles->SetWorldSpaceInflatedBounds(ParticleIdx, WorldSpaceInflatedBounds);
+	}
 
 	bool HasBounds() const { return GeometryParticles->HasBounds(ParticleIdx); }
 	void SetHasBounds(bool bHasBounds) { GeometryParticles->HasBounds(ParticleIdx) = bHasBounds; }
@@ -997,6 +1000,18 @@ public:
 		Ar << MShapesArray;
 		Ar << Type;
 		//Ar << MDirtyFlags;
+
+		Ar.UsingCustomVersion(FExternalPhysicsCustomObjectVersion::GUID);
+		if (Ar.CustomVer(FExternalPhysicsCustomObjectVersion::GUID) < FExternalPhysicsCustomObjectVersion::SerializeShapeWorldSpaceBounds)
+		{
+			if (MGeometry->HasBoundingBox())
+			{
+				for (auto& Shape : MShapesArray)
+				{
+					Shape->WorldSpaceInflatedShapeBounds = Shape->Geometry->BoundingBox().GetAABB().TransformedAABB(FRigidTransform3(MX, MR));
+				}
+			}
+		}
 	}
 
 	virtual bool IsParticleValid() const
@@ -1153,7 +1168,7 @@ protected:
 
 	void UpdateShapesArray()
 	{
-		UpdateShapesArrayFromGeometry<T, d>(MShapesArray, MakeSerializable(MGeometry));
+		UpdateShapesArrayFromGeometry<T, d>(MShapesArray, MakeSerializable(MGeometry), FRigidTransform3(X(), R()));
 	}
 };
 
diff --git a/Engine/Source/Runtime/PhysicsCore/Private/PhysTestSerializer.cpp b/Engine/Source/Runtime/PhysicsCore/Private/PhysTestSerializer.cpp
--- a/Engine/Source/Runtime/PhysicsCore/Private/PhysTestSerializer.cpp
+++ b/Engine/Source/Runtime/PhysicsCore/Private/PhysTestSerializer.cpp
@@ -307,6 +307,17 @@ void FPhysTestSerializer::CreateChaosData()
 					GTParticle->SetGeometry(MakeUnique<TImplicitObjectUnion<float, 3>>(MoveTemp(Geoms)));
 					Particle->SetGeometry(GTParticle->Geometry());
 				}
+
+				// Fixup bounds
+				auto Geom = GTParticle->Geometry();
+				if (Geom->HasBoundingBox())
+				{
+					auto& ShapeArray = GTParticle->ShapesArray();
+					for (auto& Shape : ShapeArray)
+					{
+						Shape->WorldSpaceInflatedShapeBounds = Geom->BoundingBox().GetAABB().TransformedAABB(TRigidTransform<FReal, 3>(Particle->X(), Particle->R()));
+					}
+				}
 			}
 
 			int32 ShapeIdx = 0;
diff --git a/Engine/Source/Runtime/PhysicsSQ/Private/SQAccelerator.cpp b/Engine/Source/Runtime/PhysicsSQ/Private/SQAccelerator.cpp
--- a/Engine/Source/Runtime/PhysicsSQ/Private/SQAccelerator.cpp
+++ b/Engine/Source/Runtime/PhysicsSQ/Private/SQAccelerator.cpp
@@ -189,9 +189,16 @@ private:
 
 			if (bTestShapeBounds)
 			{
-				const auto& GeomBounds = Geom->BoundingBox();
-				auto InflatedWorldBounds = SQ == ESQType::Raycast ? GeomBounds.GetAABB().TransformedAABB(ActorTM) :
-					TAABB<FReal, 3>(GeomBounds.Min() - HalfExtents, GeomBounds.Max() + HalfExtents).TransformedAABB(ActorTM);
+				TAABB<FReal, 3> InflatedWorldBounds;
+				if (SQ == ESQType::Raycast)
+				{
+					InflatedWorldBounds = Shape->WorldSpaceInflatedShapeBounds;
+				}
+				else
+				{
+					InflatedWorldBounds = TAABB<FReal, 3>(Shape->WorldSpaceInflatedShapeBounds.Min() - HalfExtents, Shape->WorldSpaceInflatedShapeBounds.Max() + HalfExtents);
+
+				}
 				if (SQ != ESQType::Overlap)
 				{
 					//todo: use fast raycast
