Chaos : Collisions

- Moved external collision structures to a separate file (ExternalCollisionData)
- Moved Constraint Particle and Levelset members into a fixed array on the constraint.

diff --git a/Engine/Source/Runtime/Engine/Private/PhysicsEngine/ImmediatePhysics/ImmediatePhysicsChaos/ImmediatePhysicsSimulation_Chaos.cpp b/Engine/Source/Runtime/Engine/Private/PhysicsEngine/ImmediatePhysics/ImmediatePhysicsChaos/ImmediatePhysicsSimulation_Chaos.cpp
--- a/Engine/Source/Runtime/Engine/Private/PhysicsEngine/ImmediatePhysics/ImmediatePhysicsChaos/ImmediatePhysicsSimulation_Chaos.cpp
+++ b/Engine/Source/Runtime/Engine/Private/PhysicsEngine/ImmediatePhysics/ImmediatePhysicsChaos/ImmediatePhysicsSimulation_Chaos.cpp
@@ -147,7 +147,7 @@ namespace ImmediatePhysics_Chaos
 				Evolution->GetCollisionConstraints().ApplyCollisionModifier(
 					[this](TRigidBodyContactConstraint<float, 3>& Constraint)
 					{
-						if (ShouldIgnoreCollisionConstraint(Constraint.Particle, Constraint.Levelset, IgnoreCollisionParticlePairTable))
+						if (ShouldIgnoreCollisionConstraint(Constraint.Particle[0], Constraint.Particle[1], IgnoreCollisionParticlePairTable))
 						{
 							return ECollisionModifierResult::Disabled;
 						}
diff --git a/Engine/Source/Runtime/Experimental/Chaos/Private/Chaos/CollisionResolutionConvexConvex.cpp b/Engine/Source/Runtime/Experimental/Chaos/Private/Chaos/CollisionResolutionConvexConvex.cpp
--- a/Engine/Source/Runtime/Experimental/Chaos/Private/Chaos/CollisionResolutionConvexConvex.cpp
+++ b/Engine/Source/Runtime/Experimental/Chaos/Private/Chaos/CollisionResolutionConvexConvex.cpp
@@ -19,8 +19,8 @@ namespace Chaos
 			return;
 		}
 
-		Constraint.Particle = Particle0;
-		Constraint.Levelset = Particle1;
+		Constraint.Particle[0] = Particle0;
+		Constraint.Particle[1] = Particle1;
 		Constraint.AddManifold(Implicit0, Implicit1);
 	}
 
diff --git a/Engine/Source/Runtime/Experimental/Chaos/Private/Chaos/PBDCollisionConstraint.cpp b/Engine/Source/Runtime/Experimental/Chaos/Private/Chaos/PBDCollisionConstraint.cpp
--- a/Engine/Source/Runtime/Experimental/Chaos/Private/Chaos/PBDCollisionConstraint.cpp
+++ b/Engine/Source/Runtime/Experimental/Chaos/Private/Chaos/PBDCollisionConstraint.cpp
@@ -125,7 +125,7 @@ void TPBDCollisionConstraint<T, d>::RemoveConstraints(const TSet<TGeometryPartic
 	{
 		for (int32 Idx = Constraints.Num() - 1; Idx >= 0; Idx--)
 		{
-			if (Constraints[Idx].Levelset == ParticleHandle || Constraints[Idx].Particle == ParticleHandle)
+			if (Constraints[Idx].Particle[1] == ParticleHandle || Constraints[Idx].Particle[0] == ParticleHandle)
 			{
 				RemoveConstraint(Idx);
 			}
@@ -419,7 +419,7 @@ void TPBDCollisionConstraint<T, d>::UpdateConstraints(/*const TPBDRigidParticles
 		for(int32 i = 0; i < Constraints.Num(); ++i)
 		{
 			const FRigidBodyContactConstraint& Constraint = Constraints[i];
-			if(InParticles.Disabled(Constraint.ParticleIndex) || InParticles.Disabled(Constraint.LevelsetIndex))
+			if(InParticles.Disabled(Constraint.ParticleIndex) || InParticles.Disabled(Constraint.ParticleIndex))
 			{
 				Constraints.RemoveAtSwap(i);
 				HandleAllocator.FreeHandle(Handles[i]);
@@ -469,8 +469,8 @@ PMatrix<T, 3, 3> ComputeFactorMatrix3(const TVector<T, 3>& V, const PMatrix<T, 3
 template<typename T, int d>
 TVector<T, d> GetEnergyClampedImpulse(const TRigidBodyContactConstraint<T, d>& Constraint, const TVector<T, d>& Impulse, const TVector<T, d>& VectorToPoint1, const TVector<T, d>& VectorToPoint2, const TVector<T, d>& Velocity1, const TVector<T, d>& Velocity2)
 {
-	TPBDRigidParticleHandle<T, d>* PBDRigid0 = Constraint.Particle->AsDynamic();
-	TPBDRigidParticleHandle<T, d>* PBDRigid1 = Constraint.Levelset->AsDynamic();
+	TPBDRigidParticleHandle<T, d>* PBDRigid0 = Constraint.Particle[0]->AsDynamic();
+	TPBDRigidParticleHandle<T, d>* PBDRigid1 = Constraint.Particle[1]->AsDynamic();
 
 	TVector<T, d> Jr0, Jr1, IInvJr0, IInvJr1;
 	T ImpulseRatioNumerator0 = 0, ImpulseRatioNumerator1 = 0, ImpulseRatioDenom0 = 0, ImpulseRatioDenom1 = 0;
@@ -503,8 +503,8 @@ TVector<T, d> GetEnergyClampedImpulse(const TRigidBodyContactConstraint<T, d>& C
 template<typename T, int d>
 void TPBDCollisionConstraint<T, d>::Apply(const T Dt, FRigidBodyContactConstraint& Constraint)
 {
-	TGenericParticleHandle<T, d> Particle0 = TGenericParticleHandle<T, d>(Constraint.Particle);
-	TGenericParticleHandle<T, d> Particle1 = TGenericParticleHandle<T, d>(Constraint.Levelset);
+	TGenericParticleHandle<T, d> Particle0 = TGenericParticleHandle<T, d>(Constraint.Particle[0]);
+	TGenericParticleHandle<T, d> Particle1 = TGenericParticleHandle<T, d>(Constraint.Particle[1]);
 	TPBDRigidParticleHandle<T, d>* PBDRigid0 = Particle0->AsDynamic();
 	TPBDRigidParticleHandle<T, d>* PBDRigid1 = Particle1->AsDynamic();
 
@@ -729,8 +729,8 @@ void TPBDCollisionConstraint<T, d>::Apply(const T Dt, const TArray<FConstraintCo
 template<typename T, int d>
 void TPBDCollisionConstraint<T, d>::ApplyPushOut(const T Dt, FRigidBodyContactConstraint& Constraint, const TSet<TGeometryParticleHandle<T,d>*>& IsTemporarilyStatic, int32 Iteration, int32 NumIterations, bool &NeedsAnotherIteration)
 {
-	TGeometryParticleHandle<T, d>* Particle0 = Constraint.Particle;
-	TGeometryParticleHandle<T, d>* Particle1 = Constraint.Levelset;
+	TGeometryParticleHandle<T, d>* Particle0 = Constraint.Particle[0];
+	TGeometryParticleHandle<T, d>* Particle1 = Constraint.Particle[1];
 	TPBDRigidParticleHandle<T, d>* PBDRigid0 = Particle0->AsDynamic();
 	TPBDRigidParticleHandle<T, d>* PBDRigid1 = Particle1->AsDynamic();
 
@@ -978,8 +978,8 @@ void SampleObject2(const FImplicitObject& Object, const TRigidTransform<T, d>& O
 	TContactData<float, 3> & Contact = Constraint.ShapeManifold.Manifold;
 	TContactData<float, 3> & AvgContact = AvgConstraint.ShapeManifold.Manifold;
 
-	AvgConstraint.Particle = Constraint.Particle;
-	AvgConstraint.Levelset = Constraint.Levelset;
+	AvgConstraint.Particle[0] = Constraint.Particle[0];
+	AvgConstraint.Particle[1] = Constraint.Particle[1];
 	AvgContact.Location = TVector<float, 3>::ZeroVector;
 	AvgContact.Normal = TVector<float, 3>::ZeroVector;
 	AvgContact.Phi = Thickness;
@@ -1088,8 +1088,8 @@ void SampleObject2(const FImplicitObject& Object, const TRigidTransform<float, 3
 	TContactData<float, 3> & Contact = Constraint.ShapeManifold.Manifold;
 	TContactData<float, 3> & AvgContact = AvgConstraint.ShapeManifold.Manifold;
 
-	AvgConstraint.Particle = Constraint.Particle;
-	AvgConstraint.Levelset = Constraint.Levelset;
+	AvgConstraint.Particle[0] = Constraint.Particle[0];
+	AvgConstraint.Particle[1] = Constraint.Particle[1];
 	AvgContact.Location = TVector<float, 3>::ZeroVector;
 	AvgContact.Normal = TVector<float, 3>::ZeroVector;
 	AvgContact.Phi = Thickness;
@@ -1605,8 +1605,8 @@ void UpdateUnionUnionConstraint(const T Thickness, TRigidBodyContactConstraint<T
 {
 	SCOPE_CYCLE_COUNTER(STAT_UpdateUnionUnionConstraint);
 
-	TGenericParticleHandle<T, d> Particle0 = Constraint.Particle;
-	TGenericParticleHandle<T, d> Particle1 = Constraint.Levelset;
+	TGenericParticleHandle<T, d> Particle0 = Constraint.Particle[0];
+	TGenericParticleHandle<T, d> Particle1 = Constraint.Particle[1];
 
 	TRigidTransform<T, d> ParticlesTM = TRigidTransform<T, d>(Particle0->P(), Particle0->Q());
 	TRigidTransform<T, d> LevelsetTM = TRigidTransform<T, d>(Particle1->P(), Particle1->Q());
@@ -1639,8 +1639,8 @@ void UpdateSingleUnionConstraint(const T Thickness, TRigidBodyContactConstraint<
 {
 	SCOPE_CYCLE_COUNTER(STAT_UpdateSingleUnionConstraint);
 
-	TGenericParticleHandle<T, d> Particle0 = Constraint.Particle;
-	TGenericParticleHandle<T, d> Particle1 = Constraint.Levelset;
+	TGenericParticleHandle<T, d> Particle0 = Constraint.Particle[0];
+	TGenericParticleHandle<T, d> Particle1 = Constraint.Particle[1];
 
 	TRigidTransform<T, d> ParticlesTM = TRigidTransform<T, d>(Particle0->P(), Particle0->Q());
 	TRigidTransform<T, d> LevelsetTM = TRigidTransform<T, d>(Particle1->P(), Particle1->Q());
@@ -1664,14 +1664,14 @@ void TPBDCollisionConstraint<T, d>::UpdateLevelsetConstraint(const T Thickness,
 {
 	SCOPE_CYCLE_COUNTER(STAT_UpdateLevelsetConstraint);
 	
-	TGenericParticleHandle<T, d> Particle0 = Constraint.Particle;
+	TGenericParticleHandle<T, d> Particle0 = Constraint.Particle[0];
 	TRigidTransform<T, d> ParticlesTM = TRigidTransform<T, d>(Particle0->P(), Particle0->Q());
 	if (!(ensure(!FMath::IsNaN(ParticlesTM.GetTranslation().X)) && ensure(!FMath::IsNaN(ParticlesTM.GetTranslation().Y)) && ensure(!FMath::IsNaN(ParticlesTM.GetTranslation().Z))))
 	{
 		return;
 	}
 
-	TGenericParticleHandle<T, d> Particle1 = Constraint.Levelset;
+	TGenericParticleHandle<T, d> Particle1 = Constraint.Particle[1];
 	TRigidTransform<T, d> LevelsetTM = TRigidTransform<T, d>(Particle1->P(), Particle1->Q());
 	if (!(ensure(!FMath::IsNaN(LevelsetTM.GetTranslation().X)) && ensure(!FMath::IsNaN(LevelsetTM.GetTranslation().Y)) && ensure(!FMath::IsNaN(LevelsetTM.GetTranslation().Z))))
 	{
@@ -1693,8 +1693,8 @@ void UpdateUnionLevelsetConstraint(const T Thickness, TRigidBodyContactConstrain
 {
 	SCOPE_CYCLE_COUNTER(STAT_UpdateUnionLevelsetConstraint);
 
-	TGenericParticleHandle<T, d> Particle0 = Constraint.Particle;
-	TGenericParticleHandle<T, d> Particle1 = Constraint.Levelset;
+	TGenericParticleHandle<T, d> Particle0 = Constraint.Particle[0];
+	TGenericParticleHandle<T, d> Particle1 = Constraint.Particle[1];
 
 	TRigidTransform<T, d> ParticlesTM = TRigidTransform<T, d>(Particle0->P(), Particle0->Q());
 	TRigidTransform<T, d> LevelsetTM = TRigidTransform<T, d>(Particle1->P(), Particle1->Q());
@@ -1769,8 +1769,8 @@ void UpdateLevelsetUnionConstraint(const T Thickness, TRigidBodyContactConstrain
 {
 	SCOPE_CYCLE_COUNTER(STAT_UpdateLevelsetUnionConstraint);
 
-	TGenericParticleHandle<T, d> Particle0 = Constraint.Particle;
-	TGenericParticleHandle<T, d> Particle1 = Constraint.Levelset;
+	TGenericParticleHandle<T, d> Particle0 = Constraint.Particle[0];
+	TGenericParticleHandle<T, d> Particle1 = Constraint.Particle[1];
 
 	TRigidTransform<T, d> ParticlesTM = TRigidTransform<T, d>(Particle0->P(), Particle0->Q());
 	TRigidTransform<T, d> LevelsetTM = TRigidTransform<T, d>(Particle1->P(), Particle1->Q());
@@ -1871,14 +1871,14 @@ void ConstructLevelsetConstraints(TGeometryParticleHandle<T, d>* Particle0, TGeo
 	
 	if (!Particle1->Geometry() || (Particle0->AsDynamic() && !Particle0->AsDynamic()->CollisionParticlesSize() && Particle0->Geometry() && !Particle0->Geometry()->IsUnderlyingUnion()))
 	{
-		Constraint.Particle = Particle1;
-		Constraint.Levelset = Particle0;
+		Constraint.Particle[0] = Particle1;
+		Constraint.Particle[1] = Particle0;
 		Constraint.AddManifold(Implicit1, Implicit0);
 	}
 	else
 	{
-		Constraint.Particle = Particle0;
-		Constraint.Levelset = Particle1;
+		Constraint.Particle[0] = Particle0;
+		Constraint.Particle[1] = Particle1;
 		Constraint.AddManifold(Implicit0, Implicit1);
 	}
 }
@@ -1886,64 +1886,64 @@ void ConstructLevelsetConstraints(TGeometryParticleHandle<T, d>* Particle0, TGeo
 template<typename T, int d>
 void ConstructBoxConstraints(TGeometryParticleHandle<T, d>* Particle0, TGeometryParticleHandle<T, d>* Particle1, const FImplicitObject* Implicit0, const FImplicitObject* Implicit1, const T Thickness, TRigidBodyContactConstraint<T, d>& Constraint)
 {
-	Constraint.Particle = Particle0;
-	Constraint.Levelset = Particle1;
+	Constraint.Particle[0] = Particle0;
+	Constraint.Particle[1] = Particle1;
 	Constraint.AddManifold(Implicit0, Implicit1);
 }
 
 template<typename T, int d>
 void ConstructBoxPlaneConstraints(TGeometryParticleHandle<T, d>* Particle0, TGeometryParticleHandle<T, d>* Particle1, const FImplicitObject* Implicit0, const FImplicitObject* Implicit1, const T Thickness, TRigidBodyContactConstraint<T, d>& Constraint)
 {
-	Constraint.Particle = Particle0;
-	Constraint.Levelset = Particle1;
+	Constraint.Particle[0] = Particle0;
+	Constraint.Particle[1] = Particle1;
 	Constraint.AddManifold(Implicit0, Implicit1);
 }
 
 template<typename T, int d>
 void ConstructSphereConstraints(TGeometryParticleHandle<T, d>* Particle0, TGeometryParticleHandle<T, d>* Particle1, const FImplicitObject* Implicit0, const FImplicitObject* Implicit1, const T Thickness, TRigidBodyContactConstraint<T, d>& Constraint)
 {
-	Constraint.Particle = Particle0;
-	Constraint.Levelset = Particle1;
+	Constraint.Particle[0] = Particle0;
+	Constraint.Particle[1] = Particle1;
 	Constraint.AddManifold(Implicit0, Implicit1);
 }
 
 template<typename T, int d>
 void ConstructSpherePlaneConstraints(TGeometryParticleHandle<T, d>* Particle0, TGeometryParticleHandle<T, d>* Particle1, const FImplicitObject* Implicit0, const FImplicitObject* Implicit1, const T Thickness, TRigidBodyContactConstraint<T, d>& Constraint)
 {
-	Constraint.Particle = Particle0;
-	Constraint.Levelset = Particle1;
+	Constraint.Particle[0] = Particle0;
+	Constraint.Particle[1] = Particle1;
 	Constraint.AddManifold(Implicit0, Implicit1);
 }
 
 template<typename T, int d>
 void ConstructSphereBoxConstraints(TGeometryParticleHandle<T, d>* Particle0, TGeometryParticleHandle<T, d>* Particle1, const FImplicitObject* Implicit0, const FImplicitObject* Implicit1, const T Thickness, TRigidBodyContactConstraint<T, d>& Constraint)
 {
-	Constraint.Particle = Particle0;
-	Constraint.Levelset = Particle1;
+	Constraint.Particle[0] = Particle0;
+	Constraint.Particle[1] = Particle1;
 	Constraint.AddManifold(Implicit0, Implicit1);
 }
 
 template<typename T, int d>
 void ConstructCapsuleCapsuleConstraints(TGeometryParticleHandle<T, d>* Particle0, TGeometryParticleHandle<T, d>* Particle1, const FImplicitObject* Implicit0, const FImplicitObject* Implicit1, const T Thickness, TRigidBodyContactConstraint<T, d>& Constraint)
 {
-	Constraint.Particle = Particle0;
-	Constraint.Levelset = Particle1;
+	Constraint.Particle[0] = Particle0;
+	Constraint.Particle[1] = Particle1;
 	Constraint.AddManifold(Implicit0, Implicit1);
 }
 
 template<typename T, int d>
 void ConstructCapsuleBoxConstraints(TGeometryParticleHandle<T, d>* Particle0, TGeometryParticleHandle<T, d>* Particle1, const FImplicitObject* Implicit0, const FImplicitObject* Implicit1, const T Thickness, TRigidBodyContactConstraint<T, d>& Constraint)
 {
-	Constraint.Particle = Particle0;
-	Constraint.Levelset = Particle1;
+	Constraint.Particle[0] = Particle0;
+	Constraint.Particle[1] = Particle1;
 	Constraint.AddManifold(Implicit0, Implicit1);
 }
 
 template<typename T, int d>
 void ConstructSingleUnionConstraints(TGeometryParticleHandle<T, d>* Particle0, TGeometryParticleHandle<T, d>* Particle1, const FImplicitObject* Implicit0, const FImplicitObject* Implicit1, const T Thickness, TRigidBodyContactConstraint<T, d>& Constraint)
 {
-	Constraint.Particle = Particle0;
-	Constraint.Levelset = Particle1;
+	Constraint.Particle[0] = Particle0;
+	Constraint.Particle[1] = Particle1;
 	Constraint.AddManifold(Implicit0, Implicit1);
 }
 
@@ -2208,14 +2208,14 @@ void TPBDCollisionConstraint<T, d>::UpdateConstraint(const T Thickness, FRigidBo
 	SCOPE_CYCLE_COUNTER(STAT_UpdateConstraint);
 
 	Constraint.ResetPhi(Thickness);	
-	const TRigidTransform<T, d> ParticleTM = GetTransform(Constraint.Particle);
-	const TRigidTransform<T, d> LevelsetTM = GetTransform(Constraint.Levelset);
+	const TRigidTransform<T, d> ParticleTM = GetTransform(Constraint.Particle[0]);
+	const TRigidTransform<T, d> LevelsetTM = GetTransform(Constraint.Particle[1]);
 
-	if (!Constraint.Particle->Geometry())
+	if (!Constraint.Particle[0]->Geometry())
 	{
-		if (Constraint.Levelset->Geometry())
+		if (Constraint.Particle[1]->Geometry())
 		{
-			if (!Constraint.Levelset->Geometry()->IsUnderlyingUnion())
+			if (!Constraint.Particle[1]->Geometry()->IsUnderlyingUnion())
 			{
 				TPBDCollisionConstraint<T, d>::UpdateLevelsetConstraint<UpdateType>(Thickness, Constraint);
 			}
@@ -2227,7 +2227,7 @@ void TPBDCollisionConstraint<T, d>::UpdateConstraint(const T Thickness, FRigidBo
 	}
 	else
 	{
-		UpdateConstraintImp<UpdateType>(*Constraint.Particle->Geometry(), ParticleTM, *Constraint.Levelset->Geometry(), LevelsetTM, Thickness, Constraint);
+		UpdateConstraintImp<UpdateType>(*Constraint.Particle[0]->Geometry(), ParticleTM, *Constraint.Particle[1]->Geometry(), LevelsetTM, Thickness, Constraint);
 	}
 }
 
diff --git a/Engine/Source/Runtime/Experimental/Chaos/Private/Chaos/PBDRigidClustering.cpp b/Engine/Source/Runtime/Experimental/Chaos/Private/Chaos/PBDRigidClustering.cpp
--- a/Engine/Source/Runtime/Experimental/Chaos/Private/Chaos/PBDRigidClustering.cpp
+++ b/Engine/Source/Runtime/Experimental/Chaos/Private/Chaos/PBDRigidClustering.cpp
@@ -673,14 +673,14 @@ namespace Chaos
 				}
 			};
 
-			if (auto ChildrenPtr = MParentToChildren.Find(Contact.Particle))
+			if (auto ChildrenPtr = MParentToChildren.Find(Contact.Particle[0]))
 			{
-				ComputeStrainLambda(Contact.Particle->AsClustered(), *ChildrenPtr);
+				ComputeStrainLambda(Contact.Particle[0]->AsClustered(), *ChildrenPtr);
 			}
 
-			if (auto ChildrenPtr = MParentToChildren.Find(Contact.Levelset))
+			if (auto ChildrenPtr = MParentToChildren.Find(Contact.Particle[1]))
 			{
-				ComputeStrainLambda(Contact.Levelset->AsClustered(), *ChildrenPtr);
+				ComputeStrainLambda(Contact.Particle[1]->AsClustered(), *ChildrenPtr);
 			}
 
 			MCollisionImpulseArrayDirty = true;
diff --git a/Engine/Source/Runtime/Experimental/Chaos/Public/Chaos/ExternalCollisionData.h b/Engine/Source/Runtime/Experimental/Chaos/Public/Chaos/ExternalCollisionData.h
new file mode 100644
--- /dev/null
+++ b/Engine/Source/Runtime/Experimental/Chaos/Public/Chaos/ExternalCollisionData.h
@@ -0,0 +1,437 @@
+// Copyright 1998-2019 Epic Games, Inc. All Rights Reserved.
+#pragma once
+#include "Box.h"
+#include "Chaos/ParticleHandleFwd.h"
+#include "Chaos/PBDConstraintContainer.h"
+#include "Chaos/Vector.h"
+
+class UPhysicalMaterial;
+
+namespace Chaos
+{
+	/**
+	 * Collision event data stored for use by other systems (e.g. Niagara, gameplay events)
+	 */
+	template<class T, int d>
+	struct TCollisionData
+	{
+		TCollisionData()
+			: Location(TVector<T, d>((T)0.0))
+			, AccumulatedImpulse(TVector<T, d>((T)0.0))
+			, Normal(TVector<T, d>((T)0.0))
+			, Velocity1(TVector<T, d>((T)0.0))
+			, Velocity2(TVector<T, d>((T)0.0))
+			, AngularVelocity1(TVector<T, d>((T)0.0))
+			, AngularVelocity2(TVector<T, d>((T)0.0))
+			, Mass1((T)0.0)
+			, Mass2((T)0.0)
+			, PenetrationDepth((T)0.0)
+			, Particle(nullptr)
+			, Levelset(nullptr)
+			, ParticleIndexMesh(INDEX_NONE)
+			, LevelsetIndexMesh(INDEX_NONE)
+		{}
+
+		TCollisionData(TVector<T, d> InLocation
+			, TVector<T, d> InAccumulatedImpulse
+			, TVector<T, d> InNormal
+			, TVector<T, d> InVelocity1
+			, TVector<T, d> InVelocity2
+			, TVector<T, d> InAngularVelocity1
+			, TVector<T, d> InAngularVelocity2
+			, T InMass1
+			, T InMass2
+			, T InPenetrationDepth
+			, TGeometryParticle<T, d>* InParticle
+			, TGeometryParticle<T, d>* InLevelset
+			, int32 InParticleIndexMesh
+			, int32 InLevelsetIndexMesh)
+			: Location(InLocation)
+			, AccumulatedImpulse(InAccumulatedImpulse)
+			, Normal(InNormal)
+			, Velocity1(InVelocity1)
+			, Velocity2(InVelocity2)
+			, AngularVelocity1(InAngularVelocity1)
+			, AngularVelocity2(InAngularVelocity2)
+			, Mass1(InMass1)
+			, Mass2(InMass2)
+			, PenetrationDepth(InPenetrationDepth)
+			, Particle(InParticle)
+			, Levelset(InLevelset)
+			, ParticleIndexMesh(InParticleIndexMesh)
+			, LevelsetIndexMesh(InLevelsetIndexMesh)
+		{}
+
+		TVector<T, d> Location;
+		TVector<T, d> AccumulatedImpulse;
+		TVector<T, d> Normal;
+		TVector<T, d> Velocity1, Velocity2;
+		TVector<T, d> AngularVelocity1, AngularVelocity2;
+		T Mass1, Mass2;
+		T PenetrationDepth;
+		TGeometryParticle<T, d>* Particle;
+		TGeometryParticle<T, d>* Levelset;
+		// @todo(ccaulfield): CHAOS_PARTICLEHANDLE_TODO
+		int32 ParticleIndexMesh, LevelsetIndexMesh; // If ParticleIndex points to a cluster then this index will point to an actual mesh in the cluster
+													// It is important to be able to get extra data from the component
+
+	};
+
+
+	/*
+	CollisionData used in the subsystems
+	*/
+	template<class T, int d>
+	struct TCollisionDataExt
+	{
+		TCollisionDataExt()
+			: Location(TVector<T, d>((T)0.0))
+			, AccumulatedImpulse(TVector<T, d>((T)0.0))
+			, Normal(TVector<T, d>((T)0.0))
+			, Velocity1(TVector<T, d>((T)0.0))
+			, Velocity2(TVector<T, d>((T)0.0))
+			, AngularVelocity1(TVector<T, d>((T)0.0))
+			, AngularVelocity2(TVector<T, d>((T)0.0))
+			, Mass1((T)0.0)
+			, Mass2((T)0.0)
+			, Particle(nullptr)
+			, Levelset(nullptr)
+			, ParticleIndexMesh(INDEX_NONE)
+			, LevelsetIndexMesh(INDEX_NONE)
+			, BoundingboxVolume((T)-1.0)
+			, BoundingboxExtentMin((T)-1.0)
+			, BoundingboxExtentMax((T)-1.0)
+			, SurfaceType(-1)
+		{}
+
+		TCollisionDataExt(TVector<T, d> InLocation
+			, TVector<T, d> InAccumulatedImpulse
+			, TVector<T, d> InNormal
+			, TVector<T, d> InVelocity1
+			, TVector<T, d> InVelocity2
+			, TVector<T, d> InAngularVelocity1
+			, TVector<T, d> InAngularVelocity2
+			, T InMass1
+			, T InMass2
+			, TGeometryParticle<T, d>* InParticle
+			, TGeometryParticle<T, d>* InLevelset
+			, int32 InParticleIndexMesh
+			, int32 InLevelsetIndexMesh
+			, float InBoundingboxVolume
+			, float InBoundingboxExtentMin
+			, float InBoundingboxExtentMax
+			, int32 InSurfaceType)
+			: Location(InLocation)
+			, AccumulatedImpulse(InAccumulatedImpulse)
+			, Normal(InNormal)
+			, Velocity1(InVelocity1)
+			, Velocity2(InVelocity2)
+			, AngularVelocity1(InAngularVelocity1)
+			, AngularVelocity2(InAngularVelocity2)
+			, Mass1(InMass1)
+			, Mass2(InMass2)
+			, Particle(InParticle)
+			, Levelset(InLevelset)
+			, ParticleIndexMesh(InParticleIndexMesh)
+			, LevelsetIndexMesh(InLevelsetIndexMesh)
+			, BoundingboxVolume(InBoundingboxVolume)
+			, BoundingboxExtentMin(InBoundingboxExtentMin)
+			, BoundingboxExtentMax(InBoundingboxExtentMax)
+			, SurfaceType(InSurfaceType)
+		{}
+
+		TCollisionDataExt(const TCollisionData<T, d>& InCollisionData)
+			: Location(InCollisionData.Location)
+			, AccumulatedImpulse(InCollisionData.AccumulatedImpulse)
+			, Normal(InCollisionData.Normal)
+			, Velocity1(InCollisionData.Velocity1)
+			, Velocity2(InCollisionData.Velocity2)
+			, AngularVelocity1(InCollisionData.AngularVelocity1)
+			, AngularVelocity2(InCollisionData.AngularVelocity2)
+			, Mass1(InCollisionData.Mass1)
+			, Mass2(InCollisionData.Mass2)
+			, Particle(InCollisionData.Particle)
+			, Levelset(InCollisionData.Levelset)
+			, ParticleIndexMesh(InCollisionData.ParticleIndexMesh)
+			, LevelsetIndexMesh(InCollisionData.LevelsetIndexMesh)
+			, BoundingboxVolume((T)-1.0)
+			, BoundingboxExtentMin((T)-1.0)
+			, BoundingboxExtentMax((T)-1.0)
+			, SurfaceType(-1)
+		{}
+
+		TVector<T, d> Location;
+		TVector<T, d> AccumulatedImpulse;
+		TVector<T, d> Normal;
+		TVector<T, d> Velocity1, Velocity2;
+		TVector<T, d> AngularVelocity1, AngularVelocity2;
+		T Mass1, Mass2;
+		TGeometryParticle<T, d>* Particle;
+		TGeometryParticle<T, d>* Levelset;
+		int32 ParticleIndexMesh, LevelsetIndexMesh;
+		float BoundingboxVolume;
+		float BoundingboxExtentMin, BoundingboxExtentMax;
+		int32 SurfaceType;
+	};
+
+	/*
+	BreakingData passed from the physics solver to subsystems
+	*/
+	template<class T, int d>
+	struct TBreakingData
+	{
+		TBreakingData()
+			: Location(TVector<T, d>((T)0.0))
+			, Velocity(TVector<T, d>((T)0.0))
+			, AngularVelocity(TVector<T, d>((T)0.0))
+			, Mass((T)0.0)
+			, Particle(nullptr)
+			, ParticleIndex(INDEX_NONE)
+			, ParticleIndexMesh(INDEX_NONE)
+			, BoundingBox(TBox<T, d>(TVector<T, d>((T)0.0), TVector<T, d>((T)0.0)))
+		{}
+
+		TVector<T, d> Location;
+		TVector<T, d> Velocity;
+		TVector<T, d> AngularVelocity;
+		T Mass;
+		TGeometryParticle<T, d>* Particle;
+		int32 ParticleIndex; //#todo: remove this in favor of TGeometryParticle?
+		int32 ParticleIndexMesh; // If ParticleIndex points to a cluster then this index will point to an actual mesh in the cluster
+								 // It is important to be able to get extra data from the component
+		Chaos::TBox<T, d> BoundingBox;
+	};
+
+
+
+	/*
+	BreakingData used in the subsystems
+	*/
+	template<class T, int d>
+	struct TBreakingDataExt
+	{
+		TBreakingDataExt()
+			: Location(TVector<T, d>((T)0.0))
+			, Velocity(TVector<T, d>((T)0.0))
+			, AngularVelocity(TVector<T, d>((T)0.0))
+			, Mass((T)0.0)
+			, ParticleIndex(INDEX_NONE)
+			, ParticleIndexMesh(INDEX_NONE)
+			, BoundingboxVolume((T)-1.0)
+			, BoundingboxExtentMin((T)-1.0)
+			, BoundingboxExtentMax((T)-1.0)
+			, SurfaceType(-1)
+		{}
+
+		TBreakingDataExt(TVector<T, d> InLocation
+			, TVector<T, d> InVelocity
+			, TVector<T, d> InAngularVelocity
+			, T InMass
+			, int32 InParticleIndex
+			, int32 InParticleIndexMesh
+			, float InBoundingboxVolume
+			, float InBoundingboxExtentMin
+			, float InBoundingboxExtentMax
+			, int32 InSurfaceType)
+			: Location(InLocation)
+			, Velocity(InVelocity)
+			, AngularVelocity(InAngularVelocity)
+			, Mass(InMass)
+			, ParticleIndex(InParticleIndex)
+			, ParticleIndexMesh(InParticleIndexMesh)
+			, BoundingboxVolume(InBoundingboxVolume)
+			, BoundingboxExtentMin(InBoundingboxExtentMin)
+			, BoundingboxExtentMax(InBoundingboxExtentMax)
+			, SurfaceType(InSurfaceType)
+		{}
+
+		TBreakingDataExt(const TBreakingData<float, 3>& InBreakingData)
+			: Location(InBreakingData.Location)
+			, Velocity(InBreakingData.Velocity)
+			, AngularVelocity(InBreakingData.AngularVelocity)
+			, Mass(InBreakingData.Mass)
+			, ParticleIndex(InBreakingData.ParticleIndex)
+			, ParticleIndexMesh(InBreakingData.ParticleIndexMesh)
+			, BoundingboxVolume((T)-1.0)
+			, BoundingboxExtentMin((T)-1.0)
+			, BoundingboxExtentMax((T)-1.0)
+			, SurfaceType(-1)
+		{}
+
+		TVector<T, d> Location;
+		TVector<T, d> Velocity;
+		TVector<T, d> AngularVelocity;
+		T Mass;
+		int32 ParticleIndex;
+		int32 ParticleIndexMesh; // If ParticleIndex points to a cluster then this index will point to an actual mesh in the cluster
+								 // It is important to be able to get extra data from the component
+		float BoundingboxVolume;
+		float BoundingboxExtentMin, BoundingboxExtentMax;
+		int32 SurfaceType;
+
+		FVector TransformTranslation;
+		FQuat TransformRotation;
+		FVector TransformScale;
+
+		FBox BoundingBox;
+
+		// Please don't be tempted to add the code below back. Holding onto a UObject pointer without the GC knowing about it is 
+		// not a safe thing to do.
+		//UPhysicalMaterial* PhysicalMaterialTest;
+		FName PhysicalMaterialName;
+	};
+
+	/*
+	TrailingData passed from the physics solver to subsystems
+	*/
+	template<class T, int d>
+	struct TTrailingData
+	{
+		TTrailingData()
+			: Location(TVector<T, d>((T)0.0))
+			, Velocity(TVector<T, d>((T)0.0))
+			, AngularVelocity(TVector<T, d>((T)0.0))
+			, Mass((T)0.0)
+			, Particle(nullptr)
+			, ParticleIndexMesh(INDEX_NONE)
+			, BoundingBox(TBox<T, d>(TVector<T, d>((T)0.0), TVector<T, d>((T)0.0)))
+		{}
+
+		TTrailingData(TVector<T, d> InLocation
+			, TVector<T, d> InVelocity
+			, TVector<T, d> InAngularVelocity
+			, T InMass
+			, TGeometryParticle<T, d>* InParticle
+			, int32 InParticleIndexMesh
+			, float InBoundingboxVolume
+			, float InBoundingboxExtentMin
+			, float InBoundingboxExtentMax
+			, int32 InSurfaceType
+			, Chaos::TBox<T, d>& InBoundingBox)
+			: Location(InLocation)
+			, Velocity(InVelocity)
+			, AngularVelocity(InAngularVelocity)
+			, Mass(InMass)
+			, Particle(InParticle)
+			, ParticleIndexMesh(InParticleIndexMesh)
+			, BoundingBox(InBoundingBox)
+		{}
+
+		TVector<T, d> Location;
+		TVector<T, d> Velocity;
+		TVector<T, d> AngularVelocity;
+		T Mass;
+		TGeometryParticle<T, d>* Particle;
+		int32 ParticleIndexMesh; // If ParticleIndex points to a cluster then this index will point to an actual mesh in the cluster
+								 // It is important to be able to get extra data from the component
+		Chaos::TBox<T, d> BoundingBox;
+
+		friend inline uint32 GetTypeHash(const TTrailingData& Other)
+		{
+			return ::GetTypeHash(Other.Particle);
+		}
+
+		friend bool operator==(const TTrailingData& A, const TTrailingData& B)
+		{
+			return A.Particle == B.Particle;
+		}
+	};
+
+	/*
+	TrailingData used in subsystems
+	*/
+	template<class T, int d>
+	struct TTrailingDataExt
+	{
+		TTrailingDataExt()
+			: Location(TVector<T, d>((T)0.0))
+			, Velocity(TVector<T, d>((T)0.0))
+			, AngularVelocity(TVector<T, d>((T)0.0))
+			, Mass((T)0.0)
+			, Particle(nullptr)
+			, ParticleIndexMesh(INDEX_NONE)
+			, BoundingboxVolume((T)-1.0)
+			, BoundingboxExtentMin((T)-1.0)
+			, BoundingboxExtentMax((T)-1.0)
+			, SurfaceType(-1)
+		{}
+
+		TTrailingDataExt(TVector<T, d> InLocation
+			, TVector<T, d> InVelocity
+			, TVector<T, d> InAngularVelocity
+			, T InMass
+			, TGeometryParticle<T, d>* InParticle
+			, int32 InParticleIndexMesh
+			, float InBoundingboxVolume
+			, float InBoundingboxExtentMin
+			, float InBoundingboxExtentMax
+			, int32 InSurfaceType)
+			: Location(InLocation)
+			, Velocity(InVelocity)
+			, AngularVelocity(InAngularVelocity)
+			, Mass(InMass)
+			, Particle(InParticle)
+			, ParticleIndexMesh(InParticleIndexMesh)
+			, BoundingboxVolume(InBoundingboxVolume)
+			, BoundingboxExtentMin(InBoundingboxExtentMin)
+			, BoundingboxExtentMax(InBoundingboxExtentMax)
+			, SurfaceType(InSurfaceType)
+		{}
+
+		TTrailingDataExt(const TTrailingData<float, 3>& InTrailingData)
+			: Location(InTrailingData.Location)
+			, Velocity(InTrailingData.Velocity)
+			, AngularVelocity(InTrailingData.AngularVelocity)
+			, Mass(InTrailingData.Mass)
+			, Particle(InTrailingData.Particle)
+			, ParticleIndexMesh(InTrailingData.ParticleIndexMesh)
+			, BoundingboxVolume((T)-1.0)
+			, BoundingboxExtentMin((T)-1.0)
+			, BoundingboxExtentMax((T)-1.0)
+			, SurfaceType(-1)
+		{}
+
+		TVector<T, d> Location;
+		TVector<T, d> Velocity;
+		TVector<T, d> AngularVelocity;
+		T Mass;
+		TGeometryParticle<T, d>* Particle;
+		int32 ParticleIndexMesh; // If ParticleIndex points to a cluster then this index will point to an actual mesh in the cluster
+								 // It is important to be able to get extra data from the component
+		float BoundingboxVolume;
+		float BoundingboxExtentMin, BoundingboxExtentMax;
+		int32 SurfaceType;
+
+		friend inline uint32 GetTypeHash(const TTrailingDataExt& Other)
+		{
+			return ::GetTypeHash(Other.Particle);
+		}
+
+		friend bool operator==(const TTrailingDataExt& A, const TTrailingDataExt& B)
+		{
+			return A.Particle == B.Particle;
+		}
+	};
+
+
+	template<class T, int d>
+	struct TSleepingData
+	{
+		TSleepingData()
+			: Particle(nullptr)
+			, Sleeping(true)
+		{}
+
+		TSleepingData(
+			TGeometryParticle<T, d>* InParticle
+			, bool InSleeping)
+			: Particle(InParticle)
+			, Sleeping(InSleeping)
+		{}
+
+		TGeometryParticle<T, d>* Particle;
+		bool Sleeping;	// if !Sleeping == Awake
+	};
+
+
+}
diff --git a/Engine/Source/Runtime/Experimental/Chaos/Public/Chaos/PBDCollisionConstraint.h b/Engine/Source/Runtime/Experimental/Chaos/Public/Chaos/PBDCollisionConstraint.h
--- a/Engine/Source/Runtime/Experimental/Chaos/Public/Chaos/PBDCollisionConstraint.h
+++ b/Engine/Source/Runtime/Experimental/Chaos/Public/Chaos/PBDCollisionConstraint.h
@@ -127,12 +127,12 @@ public:
 
 	FConstraintHandleID GetConstraintHandleID(const FRigidBodyContactConstraint & Constraint) const
 	{
-		return  FConstraintHandleID( Constraint.Particle, Constraint.Levelset );
+		return  FConstraintHandleID( Constraint.Particle[0], Constraint.Particle[1] );
 	}
 
 	FConstraintHandleID GetConstraintHandleID(int32 ConstraintIndex) const
 	{
-		return  FConstraintHandleID( Constraints[ConstraintIndex].Particle, Constraints[ConstraintIndex].Levelset );
+		return  FConstraintHandleID( Constraints[ConstraintIndex].Particle[0], Constraints[ConstraintIndex].Particle[1] );
 	}
 		
 	const FConstraintContainerHandle* GetConstraintHandle(int32 ConstraintIndex) const
@@ -185,7 +185,7 @@ public:
 	 */
 	TVector<TGeometryParticleHandle<T, d>*, 2> GetConstrainedParticles(int32 ConstraintIndex) const
 	{
-		return { Constraints[ConstraintIndex].Particle, Constraints[ConstraintIndex].Levelset };
+		return { Constraints[ConstraintIndex].Particle[0], Constraints[ConstraintIndex].Particle[1] };
 	}
 
 
@@ -228,7 +228,7 @@ public:
 	/**
 	 * Get the particles that are affected by the specified constraint.
 	 */
-	TVector<TGeometryParticleHandle<T, d>*, 2> ConstraintParticles(int32 ContactIndex) const { return { Constraints[ContactIndex].Particle, Constraints[ContactIndex].Levelset }; }
+	TVector<TGeometryParticleHandle<T, d>*, 2> ConstraintParticles(int32 ContactIndex) const { return { Constraints[ContactIndex].Particle[0], Constraints[ContactIndex].Particle[1] }; }
 
 	//NOTE: this should not be called by anyone other than ISpatialAccelerationCollection and CollisionConstraints - todo: make private with friends?
 	template <bool bGatherStats, typename SPATIAL_ACCELERATION>
diff --git a/Engine/Source/Runtime/Experimental/Chaos/Public/Chaos/PBDCollisionConstraintUtil.h b/Engine/Source/Runtime/Experimental/Chaos/Public/Chaos/PBDCollisionConstraintUtil.h
--- a/Engine/Source/Runtime/Experimental/Chaos/Public/Chaos/PBDCollisionConstraintUtil.h
+++ b/Engine/Source/Runtime/Experimental/Chaos/Public/Chaos/PBDCollisionConstraintUtil.h
@@ -1,6 +1,7 @@
 // Copyright 1998-2019 Epic Games, Inc. All Rights Reserved.
 #pragma once
 
+#include "Chaos/ExternalCollisionData.h"
 #include "Chaos/PBDCollisionConstraint.h"
 #include "Chaos/PBDCollisionTypes.h"
 
diff --git a/Engine/Source/Runtime/Experimental/Chaos/Public/Chaos/PBDCollisionTypes.h b/Engine/Source/Runtime/Experimental/Chaos/Public/Chaos/PBDCollisionTypes.h
--- a/Engine/Source/Runtime/Experimental/Chaos/Public/Chaos/PBDCollisionTypes.h
+++ b/Engine/Source/Runtime/Experimental/Chaos/Public/Chaos/PBDCollisionTypes.h
@@ -54,7 +54,7 @@ namespace Chaos
 
 		int32 Timestamp;
 		FContactData Manifold;
-		const FImplicitObject* Implicit[2];
+		const FImplicitObject* Implicit[2]; // {Of Particle[0], Of Particle[1]}
 	};
 
 
@@ -90,13 +90,12 @@ namespace Chaos
 		void AddManifold(const FImplicitObject* B = nullptr, const FImplicitObject* A = nullptr) { ShapeManifold.Implicit[0] = A; ShapeManifold.Implicit[1] = B; }
 
 		FConvexManifold ShapeManifold;
-		FGeometryParticleHandle* Particle;
-		FGeometryParticleHandle* Levelset;
+		FGeometryParticleHandle* Particle[2]; // { Point, Volume } 
 		TVector<T, d> AccumulatedImpulse;
 
 		FString ToString() const
 		{
-			return FString::Printf(TEXT("Particle:%s, Levelset:%s, AccumulatedImpulse:%s"), *Particle->ToString(), *Levelset->ToString(), *AccumulatedImpulse.ToString());
+			return FString::Printf(TEXT("Particle:%s, Levelset:%s, AccumulatedImpulse:%s"), *Particle[0]->ToString(), *Particle[1]->ToString(), *AccumulatedImpulse.ToString());
 		}
 
 	};
@@ -152,430 +151,4 @@ namespace Chaos
 
 	};
 
-
-	/**
-	 * Collision event data stored for use by other systems (e.g. Niagara, gameplay events)
-	 */
-	template<class T, int d>
-	struct TCollisionData
-	{
-		TCollisionData()
-			: Location(TVector<T, d>((T)0.0))
-			, AccumulatedImpulse(TVector<T, d>((T)0.0))
-			, Normal(TVector<T, d>((T)0.0))
-			, Velocity1(TVector<T, d>((T)0.0))
-			, Velocity2(TVector<T, d>((T)0.0))
-			, AngularVelocity1(TVector<T, d>((T)0.0))
-			, AngularVelocity2(TVector<T, d>((T)0.0))
-			, Mass1((T)0.0)
-			, Mass2((T)0.0)
-			, PenetrationDepth((T)0.0)
-			, Particle(nullptr)
-			, Levelset(nullptr)
-			, ParticleIndexMesh(INDEX_NONE)
-			, LevelsetIndexMesh(INDEX_NONE)
-		{}
-
-		TCollisionData(TVector<T, d> InLocation
-			, TVector<T, d> InAccumulatedImpulse
-			, TVector<T, d> InNormal
-			, TVector<T, d> InVelocity1
-			, TVector<T, d> InVelocity2
-			, TVector<T, d> InAngularVelocity1
-			, TVector<T, d> InAngularVelocity2
-			, T InMass1
-			, T InMass2
-			, T InPenetrationDepth
-			, TGeometryParticle<T, d>* InParticle
-			, TGeometryParticle<T, d>* InLevelset
-			, int32 InParticleIndexMesh
-			, int32 InLevelsetIndexMesh)
-			: Location(InLocation)
-			, AccumulatedImpulse(InAccumulatedImpulse)
-			, Normal(InNormal)
-			, Velocity1(InVelocity1)
-			, Velocity2(InVelocity2)
-			, AngularVelocity1(InAngularVelocity1)
-			, AngularVelocity2(InAngularVelocity2)
-			, Mass1(InMass1)
-			, Mass2(InMass2)
-			, PenetrationDepth(InPenetrationDepth)
-			, Particle(InParticle)
-			, Levelset(InLevelset)
-			, ParticleIndexMesh(InParticleIndexMesh)
-			, LevelsetIndexMesh(InLevelsetIndexMesh)
-		{}
-
-		TVector<T, d> Location;
-		TVector<T, d> AccumulatedImpulse;
-		TVector<T, d> Normal;
-		TVector<T, d> Velocity1, Velocity2;
-		TVector<T, d> AngularVelocity1, AngularVelocity2;
-		T Mass1, Mass2;
-		T PenetrationDepth;
-		TGeometryParticle<T, d>* Particle;
-		TGeometryParticle<T, d>* Levelset;
-		// @todo(ccaulfield): CHAOS_PARTICLEHANDLE_TODO
-		int32 ParticleIndexMesh, LevelsetIndexMesh; // If ParticleIndex points to a cluster then this index will point to an actual mesh in the cluster
-													// It is important to be able to get extra data from the component
-
-	};
-
-
-	/*
-	CollisionData used in the subsystems
-	*/
-	template<class T, int d>
-	struct TCollisionDataExt
-	{
-		TCollisionDataExt()
-			: Location(TVector<T, d>((T)0.0))
-			, AccumulatedImpulse(TVector<T, d>((T)0.0))
-			, Normal(TVector<T, d>((T)0.0))
-			, Velocity1(TVector<T, d>((T)0.0))
-			, Velocity2(TVector<T, d>((T)0.0))
-			, AngularVelocity1(TVector<T, d>((T)0.0))
-			, AngularVelocity2(TVector<T, d>((T)0.0))
-			, Mass1((T)0.0)
-			, Mass2((T)0.0)
-			, Particle(nullptr)
-			, Levelset(nullptr)
-			, ParticleIndexMesh(INDEX_NONE)
-			, LevelsetIndexMesh(INDEX_NONE)
-			, BoundingboxVolume((T)-1.0)
-			, BoundingboxExtentMin((T)-1.0)
-			, BoundingboxExtentMax((T)-1.0)
-			, SurfaceType(-1)
-		{}
-
-		TCollisionDataExt(TVector<T, d> InLocation
-			, TVector<T, d> InAccumulatedImpulse
-			, TVector<T, d> InNormal
-			, TVector<T, d> InVelocity1
-			, TVector<T, d> InVelocity2
-			, TVector<T, d> InAngularVelocity1
-			, TVector<T, d> InAngularVelocity2
-			, T InMass1
-			, T InMass2
-			, TGeometryParticle<T, d>* InParticle
-			, TGeometryParticle<T, d>* InLevelset
-			, int32 InParticleIndexMesh
-			, int32 InLevelsetIndexMesh
-			, float InBoundingboxVolume
-			, float InBoundingboxExtentMin
-			, float InBoundingboxExtentMax
-			, int32 InSurfaceType)
-			: Location(InLocation)
-			, AccumulatedImpulse(InAccumulatedImpulse)
-			, Normal(InNormal)
-			, Velocity1(InVelocity1)
-			, Velocity2(InVelocity2)
-			, AngularVelocity1(InAngularVelocity1)
-			, AngularVelocity2(InAngularVelocity2)
-			, Mass1(InMass1)
-			, Mass2(InMass2)
-			, Particle(InParticle)
-			, Levelset(InLevelset)
-			, ParticleIndexMesh(InParticleIndexMesh)
-			, LevelsetIndexMesh(InLevelsetIndexMesh)
-			, BoundingboxVolume(InBoundingboxVolume)
-			, BoundingboxExtentMin(InBoundingboxExtentMin)
-			, BoundingboxExtentMax(InBoundingboxExtentMax)
-			, SurfaceType(InSurfaceType)
-		{}
-
-		TCollisionDataExt(const TCollisionData<T, d>& InCollisionData)
-			: Location(InCollisionData.Location)
-			, AccumulatedImpulse(InCollisionData.AccumulatedImpulse)
-			, Normal(InCollisionData.Normal)
-			, Velocity1(InCollisionData.Velocity1)
-			, Velocity2(InCollisionData.Velocity2)
-			, AngularVelocity1(InCollisionData.AngularVelocity1)
-			, AngularVelocity2(InCollisionData.AngularVelocity2)
-			, Mass1(InCollisionData.Mass1)
-			, Mass2(InCollisionData.Mass2)
-			, Particle(InCollisionData.Particle)
-			, Levelset(InCollisionData.Levelset)
-			, ParticleIndexMesh(InCollisionData.ParticleIndexMesh)
-			, LevelsetIndexMesh(InCollisionData.LevelsetIndexMesh)
-			, BoundingboxVolume((T)-1.0)
-			, BoundingboxExtentMin((T)-1.0)
-			, BoundingboxExtentMax((T)-1.0)
-			, SurfaceType(-1)
-		{}
-
-		TVector<T, d> Location;
-		TVector<T, d> AccumulatedImpulse;
-		TVector<T, d> Normal;
-		TVector<T, d> Velocity1, Velocity2;
-		TVector<T, d> AngularVelocity1, AngularVelocity2;
-		T Mass1, Mass2;
-		TGeometryParticle<T, d>* Particle;
-		TGeometryParticle<T, d>* Levelset;
-		int32 ParticleIndexMesh, LevelsetIndexMesh;
-		float BoundingboxVolume;
-		float BoundingboxExtentMin, BoundingboxExtentMax;
-		int32 SurfaceType;
-	};
-
-	/*
-	BreakingData passed from the physics solver to subsystems
-	*/
-	template<class T, int d>
-	struct TBreakingData
-	{
-		TBreakingData()
-			: Location(TVector<T, d>((T)0.0))
-			, Velocity(TVector<T, d>((T)0.0))
-			, AngularVelocity(TVector<T, d>((T)0.0))
-			, Mass((T)0.0)
-			, Particle(nullptr)
-			, ParticleIndex(INDEX_NONE)
-			, ParticleIndexMesh(INDEX_NONE)
-			, BoundingBox(TBox<T, d>(TVector<T, d>((T)0.0), TVector<T, d>((T)0.0)))
-		{}
-
-		TVector<T, d> Location;
-		TVector<T, d> Velocity;
-		TVector<T, d> AngularVelocity;
-		T Mass;
-		TGeometryParticle<T, d>* Particle;
-		int32 ParticleIndex; //#todo: remove this in favor of TGeometryParticle?
-		int32 ParticleIndexMesh; // If ParticleIndex points to a cluster then this index will point to an actual mesh in the cluster
-								 // It is important to be able to get extra data from the component
-		Chaos::TBox<T, d> BoundingBox;
-	};
-
-
-
-	/*
-	BreakingData used in the subsystems
-	*/
-	template<class T, int d>
-	struct TBreakingDataExt
-	{
-		TBreakingDataExt()
-			: Location(TVector<T, d>((T)0.0))
-			, Velocity(TVector<T, d>((T)0.0))
-			, AngularVelocity(TVector<T, d>((T)0.0))
-			, Mass((T)0.0)
-			, ParticleIndex(INDEX_NONE)
-			, ParticleIndexMesh(INDEX_NONE)
-			, BoundingboxVolume((T)-1.0)
-			, BoundingboxExtentMin((T)-1.0)
-			, BoundingboxExtentMax((T)-1.0)
-			, SurfaceType(-1)
-		{}
-
-		TBreakingDataExt(TVector<T, d> InLocation
-			, TVector<T, d> InVelocity
-			, TVector<T, d> InAngularVelocity
-			, T InMass
-			, int32 InParticleIndex
-			, int32 InParticleIndexMesh
-			, float InBoundingboxVolume
-			, float InBoundingboxExtentMin
-			, float InBoundingboxExtentMax
-			, int32 InSurfaceType)
-			: Location(InLocation)
-			, Velocity(InVelocity)
-			, AngularVelocity(InAngularVelocity)
-			, Mass(InMass)
-			, ParticleIndex(InParticleIndex)
-			, ParticleIndexMesh(InParticleIndexMesh)
-			, BoundingboxVolume(InBoundingboxVolume)
-			, BoundingboxExtentMin(InBoundingboxExtentMin)
-			, BoundingboxExtentMax(InBoundingboxExtentMax)
-			, SurfaceType(InSurfaceType)
-		{}
-
-		TBreakingDataExt(const TBreakingData<float, 3>& InBreakingData)
-			: Location(InBreakingData.Location)
-			, Velocity(InBreakingData.Velocity)
-			, AngularVelocity(InBreakingData.AngularVelocity)
-			, Mass(InBreakingData.Mass)
-			, ParticleIndex(InBreakingData.ParticleIndex)
-			, ParticleIndexMesh(InBreakingData.ParticleIndexMesh)
-			, BoundingboxVolume((T)-1.0)
-			, BoundingboxExtentMin((T)-1.0)
-			, BoundingboxExtentMax((T)-1.0)
-			, SurfaceType(-1)
-		{}
-
-		TVector<T, d> Location;
-		TVector<T, d> Velocity;
-		TVector<T, d> AngularVelocity;
-		T Mass;
-		int32 ParticleIndex;
-		int32 ParticleIndexMesh; // If ParticleIndex points to a cluster then this index will point to an actual mesh in the cluster
-								 // It is important to be able to get extra data from the component
-		float BoundingboxVolume;
-		float BoundingboxExtentMin, BoundingboxExtentMax;
-		int32 SurfaceType;
-
-		FVector TransformTranslation;
-		FQuat TransformRotation;
-		FVector TransformScale;
-
-		FBox BoundingBox;
-
-		// Please don't be tempted to add the code below back. Holding onto a UObject pointer without the GC knowing about it is 
-		// not a safe thing to do.
-		//UPhysicalMaterial* PhysicalMaterialTest;
-		FName PhysicalMaterialName;
-	};
-
-	/*
-	TrailingData passed from the physics solver to subsystems
-	*/
-	template<class T, int d>
-	struct TTrailingData
-	{
-		TTrailingData()
-			: Location(TVector<T, d>((T)0.0))
-			, Velocity(TVector<T, d>((T)0.0))
-			, AngularVelocity(TVector<T, d>((T)0.0))
-			, Mass((T)0.0)
-			, Particle(nullptr)
-			, ParticleIndexMesh(INDEX_NONE)
-			, BoundingBox(TBox<T, d>(TVector<T, d>((T)0.0), TVector<T, d>((T)0.0)))
-		{}
-
-		TTrailingData(TVector<T, d> InLocation
-			, TVector<T, d> InVelocity
-			, TVector<T, d> InAngularVelocity
-			, T InMass
-			, TGeometryParticle<T, d>* InParticle
-			, int32 InParticleIndexMesh
-			, float InBoundingboxVolume
-			, float InBoundingboxExtentMin
-			, float InBoundingboxExtentMax
-			, int32 InSurfaceType
-			, Chaos::TBox<T, d>& InBoundingBox)
-			: Location(InLocation)
-			, Velocity(InVelocity)
-			, AngularVelocity(InAngularVelocity)
-			, Mass(InMass)
-			, Particle(InParticle)
-			, ParticleIndexMesh(InParticleIndexMesh)
-			, BoundingBox(InBoundingBox)
-		{}
-
-		TVector<T, d> Location;
-		TVector<T, d> Velocity;
-		TVector<T, d> AngularVelocity;
-		T Mass;
-		TGeometryParticle<T, d>* Particle;
-		int32 ParticleIndexMesh; // If ParticleIndex points to a cluster then this index will point to an actual mesh in the cluster
-								 // It is important to be able to get extra data from the component
-		Chaos::TBox<T, d> BoundingBox;
-
-		friend inline uint32 GetTypeHash(const TTrailingData& Other)
-		{
-			return ::GetTypeHash(Other.Particle);
-		}
-
-		friend bool operator==(const TTrailingData& A, const TTrailingData& B)
-		{
-			return A.Particle == B.Particle;
-		}
-	};
-
-	/*
-	TrailingData used in subsystems
-	*/
-	template<class T, int d>
-	struct TTrailingDataExt
-	{
-		TTrailingDataExt()
-			: Location(TVector<T, d>((T)0.0))
-			, Velocity(TVector<T, d>((T)0.0))
-			, AngularVelocity(TVector<T, d>((T)0.0))
-			, Mass((T)0.0)
-			, Particle(nullptr)
-			, ParticleIndexMesh(INDEX_NONE)
-			, BoundingboxVolume((T)-1.0)
-			, BoundingboxExtentMin((T)-1.0)
-			, BoundingboxExtentMax((T)-1.0)
-			, SurfaceType(-1)
-		{}
-
-		TTrailingDataExt(TVector<T, d> InLocation
-			, TVector<T, d> InVelocity
-			, TVector<T, d> InAngularVelocity
-			, T InMass
-			, TGeometryParticle<T, d>* InParticle
-			, int32 InParticleIndexMesh
-			, float InBoundingboxVolume
-			, float InBoundingboxExtentMin
-			, float InBoundingboxExtentMax
-			, int32 InSurfaceType)
-			: Location(InLocation)
-			, Velocity(InVelocity)
-			, AngularVelocity(InAngularVelocity)
-			, Mass(InMass)
-			, Particle(InParticle)
-			, ParticleIndexMesh(InParticleIndexMesh)
-			, BoundingboxVolume(InBoundingboxVolume)
-			, BoundingboxExtentMin(InBoundingboxExtentMin)
-			, BoundingboxExtentMax(InBoundingboxExtentMax)
-			, SurfaceType(InSurfaceType)
-		{}
-
-		TTrailingDataExt(const TTrailingData<float, 3>& InTrailingData)
-			: Location(InTrailingData.Location)
-			, Velocity(InTrailingData.Velocity)
-			, AngularVelocity(InTrailingData.AngularVelocity)
-			, Mass(InTrailingData.Mass)
-			, Particle(InTrailingData.Particle)
-			, ParticleIndexMesh(InTrailingData.ParticleIndexMesh)
-			, BoundingboxVolume((T)-1.0)
-			, BoundingboxExtentMin((T)-1.0)
-			, BoundingboxExtentMax((T)-1.0)
-			, SurfaceType(-1)
-		{}
-
-		TVector<T, d> Location;
-		TVector<T, d> Velocity;
-		TVector<T, d> AngularVelocity;
-		T Mass;
-		TGeometryParticle<T, d>* Particle;
-		int32 ParticleIndexMesh; // If ParticleIndex points to a cluster then this index will point to an actual mesh in the cluster
-								 // It is important to be able to get extra data from the component
-		float BoundingboxVolume;
-		float BoundingboxExtentMin, BoundingboxExtentMax;
-		int32 SurfaceType;
-
-		friend inline uint32 GetTypeHash(const TTrailingDataExt& Other)
-		{
-			return ::GetTypeHash(Other.Particle);
-		}
-
-		friend bool operator==(const TTrailingDataExt& A, const TTrailingDataExt& B)
-		{
-			return A.Particle == B.Particle;
-		}
-	};
-
-
-	template<class T, int d>
-	struct TSleepingData
-	{
-		TSleepingData()
-			: Particle(nullptr)
-			, Sleeping(true)
-		{}
-
-		TSleepingData(
-			TGeometryParticle<T, d>* InParticle
-			, bool InSleeping)
-			: Particle(InParticle)
-			, Sleeping(InSleeping)
-		{}
-
-		TGeometryParticle<T, d>* Particle;
-		bool Sleeping;	// if !Sleeping == Awake
-	};
-
-
 }
diff --git a/Engine/Source/Runtime/Experimental/Chaos/Public/Chaos/PBDRigidClustering.h b/Engine/Source/Runtime/Experimental/Chaos/Public/Chaos/PBDRigidClustering.h
--- a/Engine/Source/Runtime/Experimental/Chaos/Public/Chaos/PBDRigidClustering.h
+++ b/Engine/Source/Runtime/Experimental/Chaos/Public/Chaos/PBDRigidClustering.h
@@ -1,11 +1,12 @@
 // Copyright 1998-2019 Epic Games, Inc. All Rights Reserved.
 #pragma once
 
+#include "Chaos/ClusterCreationParameters.h"
+#include "Chaos/ExternalCollisionData.h"
 #include "Chaos/PBDRigidClusteredParticles.h"
-#include "Chaos/Transform.h"
 #include "Chaos/PBDCollisionTypes.h"
+#include "Chaos/Transform.h"
 #include "Chaos/TriangleMesh.h"
-#include "Chaos/ClusterCreationParameters.h"
 #include "Framework/BufferedData.h"
 
 #define TODO_CONVERT_GEOMETRY_COLLECTION_PARTICLE_INDICES_TO_PARTICLE_POINTERS 0
diff --git a/Engine/Source/Runtime/Experimental/ChaosSolvers/Private/EventDefaults.cpp b/Engine/Source/Runtime/Experimental/ChaosSolvers/Private/EventDefaults.cpp
--- a/Engine/Source/Runtime/Experimental/ChaosSolvers/Private/EventDefaults.cpp
+++ b/Engine/Source/Runtime/Experimental/ChaosSolvers/Private/EventDefaults.cpp
@@ -69,13 +69,13 @@ namespace Chaos
 
 					// Since Clustered GCs can be unioned the particleIndex representing the union 
 					// is not associated with a PhysicsProxy
-					if(Constraint.Particle->Handle()->GTGeometryParticle()->Proxy != nullptr)
+					if(Constraint.Particle[0]->Handle()->GTGeometryParticle()->Proxy != nullptr)
 
 					{
 						if(ensure(!Constraint.AccumulatedImpulse.ContainsNaN() && FMath::IsFinite(Constraint.GetPhi())))
 						{
-							TGeometryParticleHandle<float, 3>* Particle0 = Constraint.Particle;
-							TGeometryParticleHandle<float, 3>* Particle1 = Constraint.Levelset;
+							TGeometryParticleHandle<float, 3>* Particle0 = Constraint.Particle[0];
+							TGeometryParticleHandle<float, 3>* Particle1 = Constraint.Particle[1];
 							TKinematicGeometryParticleHandle<float, 3>* Body0 = Particle0->AsKinematic();
 
 							// presently when a rigidbody or kinematic hits static geometry then Body1 is null
@@ -105,8 +105,8 @@ namespace Chaos
 					{
 						Chaos::TPBDCollisionConstraint<float, 3>::FRigidBodyContactConstraint const& Constraint = AllConstraintsArray[ValidCollisionIndices[IdxCollision]];
 
-						TGeometryParticleHandle<float, 3>* Particle0 = Constraint.Particle;
-						TGeometryParticleHandle<float, 3>* Particle1 = Constraint.Levelset;
+						TGeometryParticleHandle<float, 3>* Particle0 = Constraint.Particle[0];
+						TGeometryParticleHandle<float, 3>* Particle1 = Constraint.Particle[1];
 
 						TCollisionData<float, 3> Data;
 						Data.Location = Constraint.GetLocation();
diff --git a/Engine/Source/Runtime/Experimental/ChaosSolvers/Public/EventsData.h b/Engine/Source/Runtime/Experimental/ChaosSolvers/Public/EventsData.h
--- a/Engine/Source/Runtime/Experimental/ChaosSolvers/Public/EventsData.h
+++ b/Engine/Source/Runtime/Experimental/ChaosSolvers/Public/EventsData.h
@@ -2,6 +2,7 @@
 
 #pragma once
 
+#include "Chaos/ExternalCollisionData.h"
 #include "Chaos/PBDCollisionTypes.h"
 #include "Chaos/Framework/PhysicsProxy.h"
 
