Add SOA for kinematic particles which are dynamic under the hood

diff --git a/Engine/Source/Runtime/Core/Public/UObject/ExternalPhysicsCustomObjectVersion.h b/Engine/Source/Runtime/Core/Public/UObject/ExternalPhysicsCustomObjectVersion.h
--- a/Engine/Source/Runtime/Core/Public/UObject/ExternalPhysicsCustomObjectVersion.h
+++ b/Engine/Source/Runtime/Core/Public/UObject/ExternalPhysicsCustomObjectVersion.h
@@ -60,6 +60,9 @@ struct CORE_API FExternalPhysicsCustomObjectVersion
 		// Serialize world space bounds of TPerShapeData
 		SerializeShapeWorldSpaceBounds,
 
+		// Add an SOA which contains particles with full dynamic data but which are in the kinematic object state.
+		AddDynamicKinematicSOA,
+
 		// -----<new versions can be added above this line>-------------------------------------------------
 		VersionPlusOne,
 		LatestVersion = VersionPlusOne - 1
diff --git a/Engine/Source/Runtime/Engine/Private/PhysicsEngine/Experimental/PhysInterface_Chaos.cpp b/Engine/Source/Runtime/Engine/Private/PhysicsEngine/Experimental/PhysInterface_Chaos.cpp
--- a/Engine/Source/Runtime/Engine/Private/PhysicsEngine/Experimental/PhysInterface_Chaos.cpp
+++ b/Engine/Source/Runtime/Engine/Private/PhysicsEngine/Experimental/PhysInterface_Chaos.cpp
@@ -101,21 +101,26 @@ FPhysInterface_Chaos::~FPhysInterface_Chaos()
 void FPhysInterface_Chaos::CreateActor(const FActorCreationParams& InParams, FPhysicsActorHandle& Handle)
 {
 	LLM_SCOPE(ELLMTag::Chaos);
-
-	// Create the new particle
+	
+	// Set object state based on the requested particle type
 	if (InParams.bStatic)
 	{
-		Handle = Chaos::TGeometryParticle<float, 3>::CreateParticle().Release();	//todo: should BodyInstance use a unique ptr to manage this memory?
-	}
-	else if (InParams.BodyInstance && InParams.BodyInstance->ShouldInstanceSimulatingPhysics())
-	{
-		TUniquePtr<Chaos::TPBDRigidParticle<float, 3>> RigidHandle = Chaos::TPBDRigidParticle<float, 3>::CreateParticle();
-		RigidHandle->SetGravityEnabled(InParams.bEnableGravity);
-		Handle = RigidHandle.Release();	//todo: should BodyInstance use a unique ptr to manage this memory?
+		Handle = Chaos::TGeometryParticle<float, 3>::CreateParticle().Release();
 	}
 	else
 	{
-		Handle = Chaos::TKinematicGeometryParticle<float, 3>::CreateParticle().Release(); //todo: should BodyInstance use a unique ptr to manage this memory?
+		// Create an underlying dynamic particle
+		Chaos::TPBDRigidParticle<float, 3>* RigidHandle = Chaos::TPBDRigidParticle<float, 3>::CreateParticle().Release(); //todo: should BodyInstance use a unique ptr to manage this memory?
+		Handle = RigidHandle;
+		RigidHandle->SetGravityEnabled(InParams.bEnableGravity);
+		if (InParams.BodyInstance && InParams.BodyInstance->ShouldInstanceSimulatingPhysics())
+		{
+			RigidHandle->SetObjectState(Chaos::EObjectStateType::Dynamic);
+		}
+		else
+		{
+			RigidHandle->SetObjectState(Chaos::EObjectStateType::Kinematic);
+		}
 	}
 
 	// Set up the new particle's game-thread data. This will be sent to physics-thread when
@@ -294,7 +299,18 @@ void FPhysInterface_Chaos::WakeUp_AssumesLocked(const FPhysicsActorHandle& InAct
 
 void FPhysInterface_Chaos::SetIsKinematic_AssumesLocked(const FPhysicsActorHandle& InActorReference, bool bIsKinematic)
 {
-	// #todo : Implement
+	if (Chaos::TPBDRigidParticle<float, 3>* Particle = InActorReference->AsDynamic())
+	{
+		const Chaos::EObjectStateType NewState
+			= bIsKinematic
+			? Chaos::EObjectStateType::Kinematic
+			: Chaos::EObjectStateType::Dynamic;
+		Particle->SetObjectState(NewState);
+	}
+	else
+	{
+		ensureMsgf(false, TEXT("Can only set kinematic state of underlying dynamic particles"));
+	}
 }
 
 void FPhysInterface_Chaos::SetCcdEnabled_AssumesLocked(const FPhysicsActorHandle& InActorReference, bool bIsCcdEnabled)
diff --git a/Engine/Source/Runtime/Experimental/Chaos/Public/Chaos/PBDRigidsSOAs.h b/Engine/Source/Runtime/Experimental/Chaos/Public/Chaos/PBDRigidsSOAs.h
--- a/Engine/Source/Runtime/Experimental/Chaos/Public/Chaos/PBDRigidsSOAs.h
+++ b/Engine/Source/Runtime/Experimental/Chaos/Public/Chaos/PBDRigidsSOAs.h
@@ -2,6 +2,7 @@
 #pragma once
 
 #include "Chaos/ParticleHandle.h"
+#include "UObject/ExternalPhysicsCustomObjectVersion.h"
 
 namespace Chaos
 {
@@ -24,6 +25,7 @@ public:
 
 		DynamicDisabledParticles = MakeUnique<TPBDRigidParticles<T, d>>();
 		DynamicParticles = MakeUnique<TPBDRigidParticles<T, d>>();
+		DynamicKinematicParticles = MakeUnique<TPBDRigidParticles<T, d>>();
 
 		ClusteredParticles = MakeUnique< TPBDRigidClusteredParticles<T, d>>();
 		ClusteredParticles->RemoveParticleBehavior() = ERemoveParticleBehavior::Remove;	//clustered particles maintain relative ordering
@@ -140,7 +142,7 @@ public:
 			}
 			else
 			{
-				Particle->MoveToSOA(*DynamicParticles);
+				SetDynamicParticleSOA(Particle->AsDynamic());
 			}
 
 			if (!PBDRigid->Sleeping())
@@ -201,6 +203,30 @@ public:
 		}
 	}
 
+	void SetDynamicParticleSOA(TPBDRigidParticleHandle<T, d>* Particle)
+	{
+		const EObjectStateType State = Particle->ObjectState();
+
+		// Move to appropriate dynamic SOA
+		switch (State)
+		{
+		case EObjectStateType::Kinematic:
+			Particle->MoveToSOA(*DynamicKinematicParticles);
+			break;
+
+		case EObjectStateType::Dynamic:
+			Particle->MoveToSOA(*DynamicParticles);
+			break;
+
+		default:
+			// TODO: Special SOAs for sleeping and static particles?
+			Particle->MoveToSOA(*DynamicParticles);
+			break;
+		}
+
+		UpdateViews();
+	}
+
 	void Serialize(FChaosArchive& Ar)
 	{
 		static const FName SOAsName = TEXT("PBDRigidsSOAs");
@@ -214,6 +240,13 @@ public:
 		Ar << KinematicDisabledParticles;
 		Ar << DynamicParticles;
 		Ar << DynamicDisabledParticles;
+
+		Ar.UsingCustomVersion(FExternalPhysicsCustomObjectVersion::GUID);
+		if (Ar.CustomVer(FExternalPhysicsCustomObjectVersion::GUID) >= FExternalPhysicsCustomObjectVersion::AddDynamicKinematicSOA)
+		{
+			Ar << DynamicKinematicParticles;
+		}
+
 		ensure(ClusteredParticles->Size() == 0);	//not supported yet
 		//Ar << ClusteredParticles;
 
@@ -338,7 +371,7 @@ private:
 	{
 		//build various views. Group SOA types together for better branch prediction
 		{
-			TArray<TSOAView<TGeometryParticles<T, d>>> TmpArray = { StaticParticles.Get(), KinematicParticles.Get(), DynamicParticles.Get(), {&NonDisabledClusteredArray} };
+			TArray<TSOAView<TGeometryParticles<T, d>>> TmpArray = { StaticParticles.Get(), KinematicParticles.Get(), DynamicParticles.Get(), DynamicKinematicParticles.Get(), {&NonDisabledClusteredArray} };
 			NonDisabledView = MakeParticleView(MoveTemp(TmpArray));
 		}
 		{
@@ -351,11 +384,11 @@ private:
 		}
 		{
 			TArray<TSOAView<TGeometryParticles<T, d>>> TmpArray = { StaticParticles.Get(), StaticDisabledParticles.Get(), KinematicParticles.Get(), KinematicDisabledParticles.Get(),
-				DynamicParticles.Get(), DynamicDisabledParticles.Get(), ClusteredParticles.Get() };
+				DynamicParticles.Get(), DynamicDisabledParticles.Get(), DynamicKinematicParticles.Get(), ClusteredParticles.Get() };
 			AllParticlesView = MakeParticleView(MoveTemp(TmpArray));
 		}
 		{
-			TArray<TSOAView<TKinematicGeometryParticles<T, d>>> TmpArray = { KinematicParticles.Get() };
+			TArray<TSOAView<TKinematicGeometryParticles<T, d>>> TmpArray = { KinematicParticles.Get(), DynamicKinematicParticles.Get() };
 			ActiveKinematicParticlesView = MakeParticleView(MoveTemp(TmpArray));
 		}
 	}
@@ -368,6 +401,7 @@ private:
 	TUniquePtr<TKinematicGeometryParticles<T, d>> KinematicDisabledParticles;
 
 	TUniquePtr<TPBDRigidParticles<T, d>> DynamicParticles;
+	TUniquePtr<TPBDRigidParticles<T, d>> DynamicKinematicParticles;
 	TUniquePtr<TPBDRigidParticles<T, d>> DynamicDisabledParticles;
 
 	TUniquePtr<TPBDRigidClusteredParticles<T, d>> ClusteredParticles;
diff --git a/Engine/Source/Runtime/Experimental/Chaos/Public/Chaos/ParticleDirtyFlags.h b/Engine/Source/Runtime/Experimental/Chaos/Public/Chaos/ParticleDirtyFlags.h
--- a/Engine/Source/Runtime/Experimental/Chaos/Public/Chaos/ParticleDirtyFlags.h
+++ b/Engine/Source/Runtime/Experimental/Chaos/Public/Chaos/ParticleDirtyFlags.h
@@ -50,6 +50,11 @@ namespace Chaos
 			return (Bits & (int32)CheckBits) != 0;
 		}
 
+		bool IsDirty(const int32 CheckBits) const
+		{
+			return (Bits & CheckBits) != 0;
+		}
+
 		void MarkDirty(const EParticleFlags DirtyBits)
 		{
 			Bits |= (int32)DirtyBits;
diff --git a/Engine/Source/Runtime/Experimental/Chaos/Public/Chaos/ParticleHandle.h b/Engine/Source/Runtime/Experimental/Chaos/Public/Chaos/ParticleHandle.h
--- a/Engine/Source/Runtime/Experimental/Chaos/Public/Chaos/ParticleHandle.h
+++ b/Engine/Source/Runtime/Experimental/Chaos/Public/Chaos/ParticleHandle.h
@@ -77,7 +77,7 @@ void PBDRigidParticleHandleImpDefaultConstruct(TPBDRigidParticleHandleImp<T, d,
 	Concrete.SetInvM(1);
 	Concrete.SetI(PMatrix<T, d, d>(1, 1, 1));
 	Concrete.SetInvI(PMatrix<T, d, d>(1, 1, 1));
-	Concrete.SetObjectState(Params.bStartSleeping ? EObjectStateType::Sleeping : EObjectStateType::Dynamic);
+	Concrete.SetObjectStateLowLevel(Params.bStartSleeping ? EObjectStateType::Sleeping : EObjectStateType::Dynamic);
 }
 
 template <typename T, int d>
@@ -661,7 +661,7 @@ public:
 	void SetToBeRemovedOnFracture(const bool bToBeRemovedOnFracture) { PBDRigidParticles->ToBeRemovedOnFracture(ParticleIdx) = bToBeRemovedOnFracture; }
 
 	EObjectStateType ObjectState() const { return PBDRigidParticles->ObjectState(ParticleIdx); }
-	void SetObjectState(EObjectStateType InState) { PBDRigidParticles->SetObjectState(ParticleIdx, InState); }
+	void SetObjectStateLowLevel(EObjectStateType InState) { PBDRigidParticles->SetObjectState(ParticleIdx, InState); }
 	void SetSleeping(bool bSleeping) { PBDRigidParticles->SetSleeping(ParticleIdx, bSleeping); }
 
 	//Really only useful when using a transient handle
@@ -1119,6 +1119,11 @@ public:
 		return this->MDirtyFlags.IsDirty(CheckBits);
 	}
 
+	const FParticleDirtyFlags& DirtyFlags() const
+	{
+		return MDirtyFlags;
+	}
+
 	void ClearDirtyFlags()
 	{
 		this->MDirtyFlags.Clear();
@@ -1191,6 +1196,7 @@ public:
 		, R(TRotation<T, d>())
 		, SpatialIdx(FSpatialAccelerationIdx{ 0,0 })
 		, HashResult(0)
+		, DirtyFlags()
 	{}
 
 	TGeometryParticleData(const TGeometryParticle<T, d>& InParticle)
@@ -1200,6 +1206,7 @@ public:
 		, Geometry(InParticle.GeometrySharedLowLevel())
 		, SpatialIdx(InParticle.SpatialIdx())
 		, HashResult(InParticle.GetHashResultLowLevel())
+		, DirtyFlags(InParticle.DirtyFlags())
 	{}
 
 	void Reset() 
@@ -1210,6 +1217,7 @@ public:
 		Geometry = TSharedPtr<TImplicitObjectUnion<T, d>, ESPMode::ThreadSafe>();
 		SpatialIdx = FSpatialAccelerationIdx{ 0,0 };
 		HashResult = 0;
+		DirtyFlags.Clear();
 	}
 
 	TVector<T, d> X;
@@ -1217,6 +1225,7 @@ public:
 	TSharedPtr<FImplicitObject, ESPMode::ThreadSafe> Geometry;
 	FSpatialAccelerationIdx SpatialIdx;
 	uint32 HashResult;
+	FParticleDirtyFlags DirtyFlags;
 };
 
 
diff --git a/Engine/Source/Runtime/Experimental/Chaos/Public/Chaos/ParticleIterator.h b/Engine/Source/Runtime/Experimental/Chaos/Public/Chaos/ParticleIterator.h
--- a/Engine/Source/Runtime/Experimental/Chaos/Public/Chaos/ParticleIterator.h
+++ b/Engine/Source/Runtime/Experimental/Chaos/Public/Chaos/ParticleIterator.h
@@ -349,6 +349,7 @@ public:
 			if (CurHandleIdx < CurHandlesArray->Num())
 			{
 				TransientHandle.ParticleIdx = (*CurHandlesArray)[CurHandleIdx]->ParticleIdx;
+				TransientHandle.GeometryParticles = (*CurHandlesArray)[CurHandleIdx]->GeometryParticles;
 			}
 			else
 			{
@@ -414,7 +415,11 @@ protected:
 	void RangedForValidation() const
 	{
 #if PARTICLE_ITERATOR_RANGED_FOR_CHECK
-		if (TransientHandle.GeometryParticles)
+		if (CurHandlesArray)
+		{
+			check(DirtyValidationCount == CurHandlesArray->Num());
+		}
+		else if (TransientHandle.GeometryParticles)
 		{
 			check(DirtyValidationCount != INDEX_NONE);
 			check(TransientHandle.GeometryParticles->DirtyValidationCount() == DirtyValidationCount && TEXT("Iterating over particles while modifying the underlying SOA. Consider delaying any operations that require a Handle*"));
@@ -432,7 +437,14 @@ protected:
 	{
 
 #if PARTICLE_ITERATOR_RANGED_FOR_CHECK
-		DirtyValidationCount = TransientHandle.GeometryParticles ? TransientHandle.GeometryParticles->DirtyValidationCount() : INDEX_NONE;
+		if (CurHandlesArray)
+		{
+			DirtyValidationCount = CurHandlesArray->Num();
+		}
+		else
+		{
+			DirtyValidationCount = TransientHandle.GeometryParticles ? TransientHandle.GeometryParticles->DirtyValidationCount() : INDEX_NONE;
+		}
 #endif
 	}
 
diff --git a/Engine/Source/Runtime/Experimental/ChaosSolvers/Private/PhysicsProxy/SingleParticlePhysicsProxy.cpp b/Engine/Source/Runtime/Experimental/ChaosSolvers/Private/PhysicsProxy/SingleParticlePhysicsProxy.cpp
--- a/Engine/Source/Runtime/Experimental/ChaosSolvers/Private/PhysicsProxy/SingleParticlePhysicsProxy.cpp
+++ b/Engine/Source/Runtime/Experimental/ChaosSolvers/Private/PhysicsProxy/SingleParticlePhysicsProxy.cpp
@@ -175,32 +175,88 @@ void FSingleParticlePhysicsProxy<Chaos::TPBDRigidParticle<float, 3>>::PushToPhys
 		typedef Chaos::TPBDRigidParticle<float, 3> PARTICLE_TYPE;
 		const auto* Data = static_cast<const typename PARTICLE_TYPE::FData*>(InData);
 
-		RigidHandle->SetX(Data->X);
-		RigidHandle->SetR(Data->R);
-		RigidHandle->SetSharedGeometry(Data->Geometry);
-		RigidHandle->SetV(Data->MV);
-		RigidHandle->SetW(Data->MW);
-		RigidHandle->SetM(Data->MM);
-		RigidHandle->SetInvM(Data->MInvM);
-		RigidHandle->SetI(Data->MI);
-		RigidHandle->SetInvI(Data->MInvI);
-		RigidHandle->SetExternalForce(Data->MExternalForce);
-		RigidHandle->SetExternalTorque(Data->MExternalTorque);
+		//
+		// TODO: Split all these setters into separate functions so that
+		//       we don't have to check for each dirty type. We don't
+		//       *actually* have to check each dirty type, but some of
+		//       these set operations require extra actions. This avoids
+		//       performing them for every set.
+		//
+
+		// TODO: Since this can't change, it doesn't have a dirty flag.
+		// This should be moved to initialization.
 		RigidHandle->SetSpatialIdx(Data->SpatialIdx);	//todo: this needs to only happen once during initialization
 		RigidHandle->SetHashResultLowLevel(Data->HashResult);	//todo: this needs to only happen once during initialization
-		GetSolver()->GetEvolution()->GetGravityForces().SetEnabled(*RigidHandle, Data->MGravityEnabled);
 
-		if (Data->Geometry && Data->Geometry->HasBoundingBox())
+		if (Data->DirtyFlags.IsDirty(Chaos::EParticleFlags::X))
 		{
-			RigidHandle->SetHasBounds(true);
-			RigidHandle->SetLocalBounds(Data->Geometry->BoundingBox());
-			Chaos::TBox<float, 3> WorldSpaceBox = Data->Geometry->BoundingBox().TransformedBox(Chaos::TRigidTransform<float, 3>(Data->X, Data->R));
-			WorldSpaceBox.ThickenSymmetrically(Data->MV);
-			RigidHandle->SetWorldSpaceInflatedBounds(WorldSpaceBox);
+			RigidHandle->SetX(Data->X);
+			RigidHandle->SetP(Data->X);
+		}
+		if (Data->DirtyFlags.IsDirty(Chaos::EParticleFlags::R))
+		{
+			RigidHandle->SetR(Data->R);
+			RigidHandle->SetQ(Data->R);
+		}
+		if (Data->DirtyFlags.IsDirty(Chaos::EParticleFlags::Geometry))
+		{
+			RigidHandle->SetSharedGeometry(Data->Geometry);
+		}
+		if (Data->DirtyFlags.IsDirty(Chaos::EParticleFlags::V))
+		{
+			RigidHandle->SetV(Data->MV);
+		}
+		if (Data->DirtyFlags.IsDirty(Chaos::EParticleFlags::W))
+		{
+			RigidHandle->SetW(Data->MW);
+		}
+		if (Data->DirtyFlags.IsDirty(Chaos::EParticleFlags::M))
+		{
+			RigidHandle->SetM(Data->MM);
+		}
+		if (Data->DirtyFlags.IsDirty(Chaos::EParticleFlags::InvM))
+		{
+			RigidHandle->SetInvM(Data->MInvM);
+		}
+		if (Data->DirtyFlags.IsDirty(Chaos::EParticleFlags::I))
+		{
+			RigidHandle->SetI(Data->MI);
+		}
+		if (Data->DirtyFlags.IsDirty(Chaos::EParticleFlags::InvI))
+		{
+			RigidHandle->SetInvI(Data->MInvI);
 		}
-		else
+		if (Data->DirtyFlags.IsDirty(Chaos::EParticleFlags::ExternalForce))
+		{
+			RigidHandle->SetExternalForce(Data->MExternalForce);
+		}
+		if (Data->DirtyFlags.IsDirty(Chaos::EParticleFlags::ExternalTorque))
+		{
+			RigidHandle->SetExternalTorque(Data->MExternalTorque);
+		}
+		if (Data->DirtyFlags.IsDirty(Chaos::EParticleFlags::ObjectState))
+		{
+			GetSolver()->GetEvolution()->SetParticleObjectState(RigidHandle, Data->MObjectState);
+		}
+		if (Data->DirtyFlags.IsDirty(Chaos::EParticleFlags::GravityEnabled))
+		{
+			GetSolver()->GetEvolution()->GetGravityForces().SetEnabled(*RigidHandle, Data->MGravityEnabled);
+		}
+
+		if (Data->DirtyFlags.IsDirty(Chaos::EParticleFlags::X | Chaos::EParticleFlags::R | Chaos::EParticleFlags::V))
 		{
-			//todo: compute bounds based on sample points
+			if (Data->Geometry && Data->Geometry->HasBoundingBox())
+			{
+				RigidHandle->SetHasBounds(true);
+				RigidHandle->SetLocalBounds(Data->Geometry->BoundingBox());
+				Chaos::TBox<float, 3> WorldSpaceBox = Data->Geometry->BoundingBox().TransformedBox(Chaos::TRigidTransform<float, 3>(Data->X, Data->R));
+				WorldSpaceBox.ThickenSymmetrically(Data->MV);
+				RigidHandle->SetWorldSpaceInflatedBounds(WorldSpaceBox);
+			}
+			else
+			{
+				//todo: compute bounds based on sample points
+			}
 		}
 	}
 }
