Add missing function to evolution

diff --git a/Engine/Source/Runtime/Experimental/Chaos/Public/Chaos/PBDRigidsEvolution.h b/Engine/Source/Runtime/Experimental/Chaos/Public/Chaos/PBDRigidsEvolution.h
--- a/Engine/Source/Runtime/Experimental/Chaos/Public/Chaos/PBDRigidsEvolution.h
+++ b/Engine/Source/Runtime/Experimental/Chaos/Public/Chaos/PBDRigidsEvolution.h
@@ -344,6 +344,12 @@ class TPBDRigidsEvolutionBase
 		DirtyParticle(*ParticleAdded);
 	}
 
+	CHAOS_API void SetParticleObjectState(TPBDRigidParticleHandle<T, d>* Particle, EObjectStateType ObjectState)
+	{
+		Particle->SetObjectStateLowLevel(ObjectState);
+		Particles.SetDynamicParticleSOA(Particle);
+	}
+
 	CHAOS_API void DisableParticles(const TSet<TGeometryParticleHandle<T,d>*>& InParticles)
 	{
 		for (TGeometryParticleHandle<T, d>* Particle : InParticles)
