#UE4 IWYU fix

diff --git a/Engine/Source/Runtime/Experimental/Chaos/Public/Chaos/ClusterCreationParameters.h b/Engine/Source/Runtime/Experimental/Chaos/Public/Chaos/ClusterCreationParameters.h
--- a/Engine/Source/Runtime/Experimental/Chaos/Public/Chaos/ClusterCreationParameters.h
+++ b/Engine/Source/Runtime/Experimental/Chaos/Public/Chaos/ClusterCreationParameters.h
@@ -1,6 +1,9 @@
 // Copyright 1998-2019 Epic Games, Inc. All Rights Reserved.
+
 #pragma once
 
+#include "Chaos/Defines.h"
+
 namespace Chaos
 {
 	template<class T, int d>
