Avoid expensive acceleration structure copy (particularly on server) when running single threaded

diff --git a/Engine/Source/Runtime/Engine/Private/PhysicsEngine/Experimental/PhysScene_Chaos.cpp b/Engine/Source/Runtime/Engine/Private/PhysicsEngine/Experimental/PhysScene_Chaos.cpp
--- a/Engine/Source/Runtime/Engine/Private/PhysicsEngine/Experimental/PhysScene_Chaos.cpp
+++ b/Engine/Source/Runtime/Engine/Private/PhysicsEngine/Experimental/PhysScene_Chaos.cpp
@@ -821,7 +821,7 @@ Chaos::ISpatialAcceleration<Chaos::TAccelerationStructureHandle<float, 3>, float
 
 void FPhysScene_Chaos::CopySolverAccelerationStructure()
 {
-	if(SceneSolver)
+	if (SceneSolver && GetDispatcher()->GetMode() != Chaos::EThreadingMode::SingleThread)
 	{
 		ExternalDataLock.WriteLock();
 		SceneSolver->GetEvolution()->UpdateExternalAccelerationStructure(SolverAccelerationStructure);
diff --git a/Engine/Source/Runtime/Experimental/Chaos/Private/Chaos/PBDRigidsEvolution.cpp b/Engine/Source/Runtime/Experimental/Chaos/Private/Chaos/PBDRigidsEvolution.cpp
--- a/Engine/Source/Runtime/Experimental/Chaos/Private/Chaos/PBDRigidsEvolution.cpp
+++ b/Engine/Source/Runtime/Experimental/Chaos/Private/Chaos/PBDRigidsEvolution.cpp
@@ -118,9 +118,10 @@ namespace Chaos
 	};
 
 	template<class FPBDRigidsEvolution, class FPBDCollisionConstraint, class T, int d>
-	TPBDRigidsEvolutionBase<FPBDRigidsEvolution, FPBDCollisionConstraint, T, d>::TPBDRigidsEvolutionBase(TPBDRigidsSOAs<T, d>& InParticles, int32 InNumIterations, int32 InNumPushOutIterations)
+	TPBDRigidsEvolutionBase<FPBDRigidsEvolution, FPBDCollisionConstraint, T, d>::TPBDRigidsEvolutionBase(TPBDRigidsSOAs<T, d>& InParticles, int32 InNumIterations, int32 InNumPushOutIterations, bool InIsSingleThreaded)
 		: Particles(InParticles)
 		, bExternalReady(false)
+		, bIsSingleThreaded(InIsSingleThreaded)
 		, Clustering(static_cast<FPBDRigidsEvolution&>(*this), Particles.GetClusteredParticles())
 		, NumIterations(InNumIterations)
 		, NumPushOutIterations(InNumPushOutIterations)
@@ -164,12 +165,14 @@ namespace Chaos
 		, const TMap<FSpatialAccelerationIdx, TUniquePtr<TSpatialAccelerationCache<T, d>>>& InSpatialAccelerationCache
 		, TUniquePtr<FAccelerationStructure>& InAccelerationStructure
 		, TUniquePtr<FAccelerationStructure>& InAccelerationStructureCopy
-		, bool InForceFullBuild)
+		, bool InForceFullBuild
+		, bool InIsSingleThreaded)
 		: SpatialCollectionFactory(InSpatialCollectionFactory)
 		, SpatialAccelerationCache(InSpatialAccelerationCache)
 		, AccelerationStructure(InAccelerationStructure)
 		, AccelerationStructureCopy(InAccelerationStructureCopy)
 		, IsForceFullBuild(InForceFullBuild)
+		, bIsSingleThreaded(InIsSingleThreaded)
 	{
 	}
 
@@ -274,13 +277,17 @@ namespace Chaos
 
 		AccelerationStructure->SetAllAsyncTrasksComplete(!IsTimeSlicingProgressing);
 
+		// If it's not progressing then it is finished so we can perform the final copy if required
 		if (!IsTimeSlicingProgressing)
 		{
+			if (!bIsSingleThreaded)
+			{
 			// This operation is slow!
 			SCOPE_CYCLE_COUNTER(STAT_CopyAccelerationStructure);
 			AccelerationStructureCopy = AsUniqueSpatialAccelerationChecked<FAccelerationStructure>(AccelerationStructure->Copy());
 		}
 	}
+	}
 
 	template<class FPBDRigidsEvolution, class FPBDCollisionConstraint, class T, int d>
 	void TPBDRigidsEvolutionBase<FPBDRigidsEvolution, FPBDCollisionConstraint, T, d>::ApplyParticlePendingData(const FPendingSpatialData& SpatialData, FAccelerationStructure& AccelerationStructure, bool bUpdateCache)
@@ -370,8 +377,11 @@ namespace Chaos
 		for (auto Itr : AsyncAccelerationQueue)
 		{
 			ApplyParticlePendingData(Itr.Value, *AsyncInternalAcceleration, true); //only the first queue needs to update the cached acceleration
+			if (!bIsSingleThreaded)
+			{
 			ApplyParticlePendingData(Itr.Value, *AsyncExternalAcceleration, false);
 		}
+		}
 		AsyncAccelerationQueue.Empty();
 
 		//other queues are no longer needed since we've flushed all operations and now have a pristine structure
@@ -414,13 +424,20 @@ namespace Chaos
 			//initial frame so make empty structures
 
 			InternalAcceleration = TUniquePtr<FAccelerationStructure>(SpatialCollectionFactory->CreateEmptyCollection());
-			ScratchExternalAcceleration = TUniquePtr<FAccelerationStructure>(SpatialCollectionFactory->CreateEmptyCollection());
 			AsyncInternalAcceleration = TUniquePtr<FAccelerationStructure>(SpatialCollectionFactory->CreateEmptyCollection());
+			if (!bIsSingleThreaded)
+			{
+			ScratchExternalAcceleration = TUniquePtr<FAccelerationStructure>(SpatialCollectionFactory->CreateEmptyCollection());
 			AsyncExternalAcceleration = TUniquePtr<FAccelerationStructure>(SpatialCollectionFactory->CreateEmptyCollection());
+			}
 			FlushInternalAccelerationQueue();
+
+			if (!bIsSingleThreaded)
+			{
 			FlushExternalAccelerationQueue(*ScratchExternalAcceleration);
 			bExternalReady = true;
 		}
+		}
 
 		if (bBlock)
 		{
@@ -442,12 +459,18 @@ namespace Chaos
 
 				//swap acceleration structure for new one
 				std::swap(InternalAcceleration, AsyncInternalAcceleration);	//swap to avoid free on sync part as this can be expensive
+
+				if (!bIsSingleThreaded)
+				{
 				std::swap(ScratchExternalAcceleration, AsyncExternalAcceleration);
+				}
 				bExternalReady = true;
 			}
 
 			// we run the task for both starting a new accel structure as well as for the timeslicing
-			AccelerationStructureTaskComplete = TGraphTask<FChaosAccelerationStructureTask>::CreateTask().ConstructAndDispatchWhenReady(*SpatialCollectionFactory, SpatialAccelerationCache, AsyncInternalAcceleration, AsyncExternalAcceleration, ForceFullBuild);
+			// we run the task for both starting a new accel structure as well as for the timeslicing
+			AccelerationStructureTaskComplete = TGraphTask<FChaosAccelerationStructureTask>::CreateTask().ConstructAndDispatchWhenReady(*SpatialCollectionFactory, SpatialAccelerationCache, AsyncInternalAcceleration, AsyncExternalAcceleration, ForceFullBuild, bIsSingleThreaded);
+
 		}
 		else
 		{
@@ -459,6 +482,9 @@ namespace Chaos
 	void TPBDRigidsEvolutionBase<FPBDRigidsEvolution, FPBDCollisionConstraint, T, d>::UpdateExternalAccelerationStructure(TUniquePtr<ISpatialAccelerationCollection<TAccelerationStructureHandle<T, d>, T, d>>& StructToUpdate)
 	{
 		DECLARE_SCOPE_CYCLE_COUNTER(TEXT("CreateExternalAccelerationStructure"), STAT_CreateExternalAccelerationStructure, STATGROUP_Physics);
+
+		check(!bIsSingleThreaded);
+
 		if (bExternalReady)
 		{
 			std::swap(StructToUpdate, ScratchExternalAcceleration);
diff --git a/Engine/Source/Runtime/Experimental/Chaos/Private/Chaos/PBDRigidsEvolutionGBF.cpp b/Engine/Source/Runtime/Experimental/Chaos/Private/Chaos/PBDRigidsEvolutionGBF.cpp
--- a/Engine/Source/Runtime/Experimental/Chaos/Private/Chaos/PBDRigidsEvolutionGBF.cpp
+++ b/Engine/Source/Runtime/Experimental/Chaos/Private/Chaos/PBDRigidsEvolutionGBF.cpp
@@ -259,8 +259,8 @@ void TPBDRigidsEvolutionGBF<T, d>::AdvanceOneTimeStep(const T Dt, const T StepFr
 }
 
 template <typename T, int d>
-TPBDRigidsEvolutionGBF<T, d>::TPBDRigidsEvolutionGBF(TPBDRigidsSOAs<T, d>& InParticles, int32 InNumIterations)
-	: Base(InParticles, InNumIterations)
+TPBDRigidsEvolutionGBF<T, d>::TPBDRigidsEvolutionGBF(TPBDRigidsSOAs<T, d>& InParticles, int32 InNumIterations, bool InIsSingleThreaded)
+	: Base(InParticles, InNumIterations, 1, InIsSingleThreaded)
 	, CollisionConstraints(InParticles, Collided, PhysicsMaterials, DefaultNumPushOutPairIterations)
 	, CollisionRule(CollisionConstraints, DefaultNumPushOutIterations)
 	, PostIntegrateCallback(nullptr)
diff --git a/Engine/Source/Runtime/Experimental/Chaos/Public/Chaos/PBDRigidsEvolution.h b/Engine/Source/Runtime/Experimental/Chaos/Public/Chaos/PBDRigidsEvolution.h
--- a/Engine/Source/Runtime/Experimental/Chaos/Public/Chaos/PBDRigidsEvolution.h
+++ b/Engine/Source/Runtime/Experimental/Chaos/Public/Chaos/PBDRigidsEvolution.h
@@ -225,7 +225,7 @@ class TPBDRigidsEvolutionBase
 
 	friend void ChaosTest::TestPendingSpatialDataHandlePointerConflict();
 
-	CHAOS_API TPBDRigidsEvolutionBase(TPBDRigidsSOAs<T, d>& InParticles, int32 InNumIterations = 1, int32 InNumPushOutIterations = 1);
+	CHAOS_API TPBDRigidsEvolutionBase(TPBDRigidsSOAs<T, d>& InParticles, int32 InNumIterations = 1, int32 InNumPushOutIterations = 1, bool InIsSingleThreaded = false);
 	CHAOS_API virtual ~TPBDRigidsEvolutionBase();
 
 	CHAOS_API TArray<TGeometryParticleHandle<T, d>*> CreateStaticParticles(int32 NumParticles, const TGeometryParticleParameters<T, d>& Params = TGeometryParticleParameters<T, d>())
@@ -637,6 +637,7 @@ protected:
 	TUniquePtr<FAccelerationStructure> AsyncExternalAcceleration;
 	TUniquePtr<FAccelerationStructure> ScratchExternalAcceleration;
 	bool bExternalReady;
+	bool bIsSingleThreaded;
 
 	TPBDRigidClustering<FPBDRigidsEvolution, FPBDCollisionConstraint, T, d> Clustering;
 
@@ -715,7 +716,8 @@ protected:
 			, const TMap<FSpatialAccelerationIdx, TUniquePtr<TSpatialAccelerationCache<T,d>>>& InSpatialAccelerationCache
 			, TUniquePtr<FAccelerationStructure>& InAccelerationStructure
 			, TUniquePtr<FAccelerationStructure>& InAccelerationStructureCopy
-			, bool InForceFullBuild);
+			, bool InForceFullBuild
+			, bool InIsSingleThreaded);
 		static FORCEINLINE TStatId GetStatId();
 		static FORCEINLINE ENamedThreads::Type GetDesiredThread();
 		static FORCEINLINE ESubsequentsMode::Type GetSubsequentsMode();
@@ -726,6 +728,7 @@ protected:
 		TUniquePtr<FAccelerationStructure>& AccelerationStructure;
 		TUniquePtr<FAccelerationStructure>& AccelerationStructureCopy;
 		bool IsForceFullBuild;
+		bool bIsSingleThreaded;
 	};
 	FGraphEventRef AccelerationStructureTaskComplete;
 
diff --git a/Engine/Source/Runtime/Experimental/Chaos/Public/Chaos/PBDRigidsEvolutionGBF.h b/Engine/Source/Runtime/Experimental/Chaos/Public/Chaos/PBDRigidsEvolutionGBF.h
--- a/Engine/Source/Runtime/Experimental/Chaos/Public/Chaos/PBDRigidsEvolutionGBF.h
+++ b/Engine/Source/Runtime/Experimental/Chaos/Public/Chaos/PBDRigidsEvolutionGBF.h
@@ -55,7 +55,7 @@ public:
 	static constexpr int32 DefaultNumPushOutIterations = 5;
 	static constexpr int32 DefaultNumPushOutPairIterations = 2;
 
-	CHAOS_API TPBDRigidsEvolutionGBF(TPBDRigidsSOAs<T, d>& InParticles, int32 InNumIterations = DefaultNumIterations);
+	CHAOS_API TPBDRigidsEvolutionGBF(TPBDRigidsSOAs<T, d>& InParticles, int32 InNumIterations = DefaultNumIterations, bool InIsSingleThreaded = false);
 	CHAOS_API ~TPBDRigidsEvolutionGBF() {}
 
 	void SetPostIntegrateCallback(const TPBDRigidsEvolutionCallback<T, d>& Cb)
diff --git a/Engine/Source/Runtime/Experimental/ChaosSolvers/Private/PBDRigidsSolver.cpp b/Engine/Source/Runtime/Experimental/ChaosSolvers/Private/PBDRigidsSolver.cpp
--- a/Engine/Source/Runtime/Experimental/ChaosSolvers/Private/PBDRigidsSolver.cpp
+++ b/Engine/Source/Runtime/Experimental/ChaosSolvers/Private/PBDRigidsSolver.cpp
@@ -127,7 +127,7 @@ namespace Chaos
 		, bHasFloor(true)
 		, bIsFloorAnalytic(false)
 		, FloorHeight(0.f)
-		, MEvolution(new FPBDRigidsEvolution(Particles, ChaosSolverCollisionDefaultIterationsCVar))
+		, MEvolution(new FPBDRigidsEvolution(Particles, ChaosSolverCollisionDefaultIterationsCVar, BufferingModeIn == Chaos::EMultiBufferMode::Single))
 		, MEventManager(new FEventManager(BufferingModeIn))
 		, MSolverEventFilters(new FSolverEventFilters())
 		, MActiveParticlesBuffer(new FActiveParticlesBuffer(BufferingModeIn))
@@ -350,7 +350,7 @@ namespace Chaos
 		CurrentFrame = 0;
 		MMaxDeltaTime = 1;
 		TimeStepMultiplier = 1;
-		MEvolution = TUniquePtr<FPBDRigidsEvolution>(new FPBDRigidsEvolution(Particles, ChaosSolverCollisionDefaultIterationsCVar)); 
+		MEvolution = TUniquePtr<FPBDRigidsEvolution>(new FPBDRigidsEvolution(Particles, ChaosSolverCollisionDefaultIterationsCVar, BufferMode == EMultiBufferMode::Single)); 
 
 		FEventDefaults::RegisterSystemEvents(*GetEventManager());
 	}
