FIxed MTD for gjkraycast. Added unit tests. Starting to use EPA, but currently disabled

diff --git a/Engine/Source/Runtime/Experimental/Chaos/Public/Chaos/GJK.h b/Engine/Source/Runtime/Experimental/Chaos/Public/Chaos/GJK.h
--- a/Engine/Source/Runtime/Experimental/Chaos/Public/Chaos/GJK.h
+++ b/Engine/Source/Runtime/Experimental/Chaos/Public/Chaos/GJK.h
@@ -5,6 +5,7 @@
 #include "Chaos/Capsule.h"
 #include "Chaos/Box.h"
 #include "Chaos/Sphere.h"
+#include "Chaos/EPA.h"
 
 namespace Chaos
 {
@@ -247,12 +248,23 @@ namespace Chaos
 		FSimplex SimplexIDs;
 		const TRotation<T, 3> BToARotation = StartTM.GetRotation();
 		const TRotation<T, 3> AToBRotation = BToARotation.Inverse();
-		TVector<T, 3> SupportA = A.Support2(InitialDir);
+
+		auto SupportAFunc = [&A](const TVec3<T>& V)
+		{
+			return A.Support2(V);
+		};
+
+		auto SupportBFunc = [&B, &AToBRotation, &BToARotation](const TVec3<T>& V)
+		{
+			const TVector<T, 3> VInB = AToBRotation * V;
+			const TVector<T, 3> SupportBLocal = B.Support2(VInB);
+			return BToARotation * SupportBLocal;
+		};
+
+		TVector<T, 3> SupportA = SupportAFunc(InitialDir);
 		As[0] = SupportA;
 
-		const TVector<T, 3> InitialDirInB = AToBRotation * (-InitialDir);
-		const TVector<T, 3> InitialSupportBLocal = B.Support2(InitialDirInB);
-		TVector<T, 3> SupportB = BToARotation * InitialSupportBLocal;
+		TVector<T, 3> SupportB = SupportBFunc(-InitialDir);
 		Bs[0] = SupportB;
 
 		T Lambda = 0;
@@ -260,9 +272,13 @@ namespace Chaos
 		TVector<T, 3> V = X - (SupportA - SupportB);
 		TVector<T, 3> Normal(0,0,1);
 
-		bool bCloseEnough = V.SizeSquared() < Inflation2;
+		const T InitialPreDist2 = V.SizeSquared();
+		constexpr T Eps2 = 1e-6;
+		//mtd needs to find closest point even in inflation region, so can only skip if we found the closest points
+		bool bCloseEnough = InitialPreDist2 < Inflation2 && (!bComputeMTD || InitialPreDist2 < Eps2);
 		bool bDegenerate = false;
 		bool bTerminate = bCloseEnough;
+		bool bInflatedCloseEnough = bCloseEnough;
 		int NumIterations = 0;
 		T InGJKPreDist2 = TNumericLimits<T>::Max();
 		while (!bTerminate)
@@ -273,10 +289,8 @@ namespace Chaos
 				break;	//if taking too long just stop. This should never happen
 			}
 
-			SupportA = A.Support2(V);	//todo: add thickness to quadratic geometry to avoid quadratic vs quadratic when possible
-			const TVector<T, 3> VInB = AToBRotation * (-V);
-			const TVector<T, 3> SupportBLocal = B.Support2(VInB);
-			SupportB = BToARotation * SupportBLocal;
+			SupportA = SupportAFunc(V);
+			SupportB = SupportBFunc(-V);
 			const TVector<T, 3> P = SupportA - SupportB;
 			const TVector<T, 3> W = X - P;
 			SimplexIDs[SimplexIDs.NumVerts] = SimplexIDs.NumVerts;	//is this needed?
@@ -286,6 +300,7 @@ namespace Chaos
 			V = V.GetUnsafeNormal();
 
 			const T VDotW = TVector<T, 3>::DotProduct(V, W);
+
 			if (VDotW > Inflation)
 			{
 				const T VDotRayDir = TVector<T, 3>::DotProduct(V, RayDir);
@@ -316,6 +331,7 @@ namespace Chaos
 					Simplex[SimplexIDs.NumVerts++] = X - P;
 
 					InGJKPreDist2 = TNumericLimits<T>::Max();	//translated origin so restart gjk search
+					bInflatedCloseEnough = false;
 				}
 			}
 			else
@@ -323,12 +339,34 @@ namespace Chaos
 				Simplex[SimplexIDs.NumVerts++] = W;	//this is really X - P which is what we need for simplex computation
 			}
 
-			V = SimplexFindClosestToOrigin(Simplex, SimplexIDs, Barycentric, As, Bs);
+			if (bInflatedCloseEnough)
+			{
+				//Inflated shapes are close enough, but we want MTD so we need to find closest point on core shape
+				const T VDotW2 = VDotW * VDotW;
+				bCloseEnough = InGJKPreDist2 <= Eps2 + VDotW2;	//todo: relative error
+			}
 
-			T NewDist2 = V.SizeSquared();	//todo: relative error
-			bCloseEnough = NewDist2 < Inflation2;
-			bDegenerate = NewDist2 >= InGJKPreDist2;
-			InGJKPreDist2 = NewDist2;
+			if (!bCloseEnough)
+			{
+				V = SimplexFindClosestToOrigin(Simplex, SimplexIDs, Barycentric, As, Bs);
+				T NewDist2 = V.SizeSquared();	//todo: relative error
+				bCloseEnough = NewDist2 < Inflation2;
+				bDegenerate = NewDist2 >= InGJKPreDist2;
+				InGJKPreDist2 = NewDist2;
+
+
+				if (bComputeMTD && bCloseEnough && Lambda == 0 && Inflation2 > 0 && SimplexIDs.NumVerts < 4)
+				{
+					//For mtd of inflated shapes we have to find the closest point, so we have to keep going
+					bCloseEnough = false;
+					bInflatedCloseEnough = true;
+				}
+			}
+			else
+			{
+				//It must be that we want MTD and we can terminate. However, we must make one final call to fixup the simplex
+				V = SimplexFindClosestToOrigin(Simplex, SimplexIDs, Barycentric, As, Bs);
+			}
 			bTerminate = bCloseEnough || bDegenerate;
 		}
 
@@ -349,31 +387,62 @@ namespace Chaos
 		}
 		else if (bComputeMTD)
 		{
-			if (InGJKPreDist2)
+			if (InGJKPreDist2 > 1e-6 && InGJKPreDist2 < TNumericLimits<T>::Max())
 			{
 				OutNormal = Normal;
 				TVector<T, 3> ClosestA(0);
 				TVector<T, 3> ClosestB(0);
 
-				for (int i = 0; i < SimplexIDs.NumVerts; ++i)
+				if (NumIterations)
 				{
-					ClosestA += As[i] * Barycentric[i];
-					ClosestB += Bs[i] * Barycentric[i];
+					for (int i = 0; i < SimplexIDs.NumVerts; ++i)
+					{
+						ClosestA += As[i] * Barycentric[i];
+						ClosestB += Bs[i] * Barycentric[i];
+					}
 				}
-
+				else
+				{
+					//didn't even go into gjk loop
+					ClosestA = As[0];
+					ClosestB = Bs[0];
+				}
+				
+				const TVec3<T> ClosestBInA = StartTM.TransformPosition(ClosestB);
 				const T InGJKPreDist = FMath::Sqrt(InGJKPreDist2);
-				OutNormal = (ClosestA - ClosestB).GetUnsafeNormal();	//question: should we just use InGJKPreDist2?
+				OutNormal = (ClosestBInA - ClosestA).GetUnsafeNormal();	//question: should we just use InGJKPreDist2?
+				const T Penetration = ThicknessA + ThicknessB - InGJKPreDist;
 				const TVector<T, 3> ClosestLocal = ClosestB - OutNormal * ThicknessB;
 
-				OutPosition = StartPoint + ClosestLocal;
-				OutTime = InGJKPreDist - ThicknessA - ThicknessB;
+				OutPosition = StartPoint + ClosestLocal + OutNormal * Penetration;
+				OutTime = -Penetration;
 			}
 			else
 			{
-				ensure(false);	//todo: use EPA
-				OutTime = 0;
-				OutPosition = TVector<T, 3>(0);
-				OutNormal = { 0,0,1 };
+				//todo: use EPA
+				if (false && SimplexIDs.NumVerts == 4)
+				{
+					TArray<TVec3<T>> VertsA = { As[0], As[1], As[2], As[3] };
+					TArray<TVec3<T>> VertsB = { Bs[0], Bs[1], Bs[2], Bs[3] };
+
+					T Penetration;
+					TVec3<T> MTD, ClosestA, ClosestB;
+					EPA(VertsA, VertsB, SupportAFunc, SupportBFunc, Penetration, MTD, ClosestA, ClosestB);
+
+					OutNormal = (ClosestA - ClosestB);
+					OutNormal.SafeNormalize();
+					const TVector<T, 3> ClosestLocal = ClosestB - OutNormal * ThicknessB;
+
+					OutPosition = StartPoint + ClosestLocal;
+					OutTime = -Penetration - ThicknessA - ThicknessB;
+				}
+				else
+				{
+					ensure(false);	//todo: EPA with num verts < 4
+					OutTime = 0;
+					OutPosition = TVector<T, 3>(0);
+					OutNormal = { 0,0,1 };
+				}
 			}
 		}
 
