Initialize EPA when 2 or 3 verts are given. Not completely robust, but good start. Added tests

diff --git a/Engine/Source/Runtime/Experimental/Chaos/Public/Chaos/EPA.h b/Engine/Source/Runtime/Experimental/Chaos/Public/Chaos/EPA.h
--- a/Engine/Source/Runtime/Experimental/Chaos/Public/Chaos/EPA.h
+++ b/Engine/Source/Runtime/Experimental/Chaos/Public/Chaos/EPA.h
@@ -121,36 +121,109 @@ struct TEPAEntry
 	}
 };
 
-template <typename T>
-TArray<TEPAEntry<T>> InitializeEPA(const TVec3<T>* VertsA, const TVec3<T>* VertsB, const int32 NumVerts)
+template <typename T, typename SupportALambda, typename SupportBLambda >
+TArray<TEPAEntry<T>> InitializeEPA(TArray<TVec3<T>>& VertsA, TArray<TVec3<T>>& VertsB, const SupportALambda& SupportA, const SupportBLambda& SupportB)
 {
 	TArray<TEPAEntry<T>> Entries;
+	const int32 NumVerts = VertsA.Num();
+	check(VertsB.Num() == NumVerts);
+
+	auto AddFartherPoint = [&](const TVec3<T>& Dir)
+	{
+		const TVec3<T> NegDir = -Dir;
+		const TVec3<T> A0 = SupportA(Dir);	//should we have a function that does both directions at once?
+		const TVec3<T> A1 = SupportA(NegDir);
+		const TVec3<T> B0 = SupportB(NegDir);
+		const TVec3<T> B1 = SupportB(Dir);
+
+		const TVec3<T> W0 = A0 - B0;
+		const TVec3<T> W1 = A1 - B1;
+
+		const T Dist0 = TVec3<T>::DotProduct(W0, Dir);
+		const T Dist1 = TVec3<T>::DotProduct(W1, NegDir);
+
+		if (Dist1 >= Dist0)
+		{
+			VertsA.Add(A1);
+			VertsB.Add(B1);
+		}
+		else
+		{
+			VertsA.Add(A0);
+			VertsB.Add(B0);
+		}
+	};
+
 	switch(NumVerts)
 	{
-		case 4:
+		case 1: check(false);	break;	//This case is not supported - for GJK it means two surfaces are touching so the penetration is 0
+		case 3:
 		{
-			//make sure all triangle normals are facing out
+			//triangle, add farthest point along normal
+			
 			Entries.AddUninitialized(4);
+			ensure(Entries[3].Initialize(VertsA.GetData(), VertsB.GetData(), 0, 2, 1, { 1,0,2 }, { 2,0,0 }));
+
+			const TEPAEntry<T>& Base = Entries[3];
+
+			AddFartherPoint(Base.PlaneNormal);
 
-			ensure(Entries[0].Initialize(VertsA, VertsB, 1, 2, 3, { 3, 1, 2 }, { 1,1, 1 }));
-			ensure(Entries[1].Initialize(VertsA, VertsB, 0, 3, 2, { 2,0,3 }, { 2, 1, 0 }));
-			ensure(Entries[2].Initialize(VertsA, VertsB, 0, 1, 3, { 3,0, 1 }, { 2,2,0 }));
-			ensure(Entries[3].Initialize(VertsA, VertsB, 0, 2, 1, { 1,0,2 }, { 2,0,0 }));
+			ensure(Entries[0].Initialize(VertsA.GetData(), VertsB.GetData(), 1, 2, 3, { 3, 1, 2 }, { 1,1, 1 }));
+			ensure(Entries[1].Initialize(VertsA.GetData(), VertsB.GetData(), 0, 3, 2, { 2,0,3 }, { 2, 1, 0 }));
+			ensure(Entries[2].Initialize(VertsA.GetData(), VertsB.GetData(), 0, 1, 3, { 3,0, 1 }, { 2,2,0 }));
 
-			if (TVec3<T>::DotProduct(Entries[0].PlaneNormal, MinkowskiVert(VertsA, VertsB, 0)) > 0)
+			break;
+		}
+		case 2:
+		{
+			//line, add farthest point along most orthogonal axes
+
+			TVec3<T> Dir = MinkowskiVert(VertsA.GetData(), VertsB.GetData(), 1) - MinkowskiVert(VertsA.GetData(), VertsB.GetData(), 0);
+			ensure(Dir.SafeNormalize());
+			//find most opposing axis
+			int32 BestAxis = 0;
+			T MinVal = TNumericLimits<T>::Max();
+			for (int32 Axis = 0; Axis < 3; ++Axis)
 			{
-				//tet faces are pointing inwards
-				for (TEPAEntry<T>& Entry : Entries)
+				const T AbsVal = FMath::Abs(Dir[Axis]);
+				if (MinVal > AbsVal)
 				{
-					Entry.SwapWinding(Entries.GetData());
+					BestAxis = Axis;
+					MinVal = AbsVal;
 				}
 			}
+			const TVec3<T> OtherAxis = TVec3<T>::AxisVector(BestAxis);
+			const TVec3<T> Orthog = TVec3<T>::CrossProduct(Dir, OtherAxis);
 
-			return Entries;
+			AddFartherPoint(OtherAxis);
+			AddFartherPoint(Orthog);
 		}
+		case 4:
+		{
+			//make sure all triangle normals are facing out
+			Entries.AddUninitialized(4);
 
-		default: ensure(false);	return Entries;	//todo: handle other cases
+			ensure(Entries[0].Initialize(VertsA.GetData(), VertsB.GetData(), 1, 2, 3, { 3, 1, 2 }, { 1,1, 1 }));
+			ensure(Entries[1].Initialize(VertsA.GetData(), VertsB.GetData(), 0, 3, 2, { 2,0,3 }, { 2, 1, 0 }));
+			ensure(Entries[2].Initialize(VertsA.GetData(), VertsB.GetData(), 0, 1, 3, { 3,0, 1 }, { 2,2,0 }));
+			ensure(Entries[3].Initialize(VertsA.GetData(), VertsB.GetData(), 0, 2, 1, { 1,0,2 }, { 2,0,0 }));
+
+			break;
+		}
+
+		default: ensure(false);
 	}
+
+	if (TVec3<T>::DotProduct(Entries[0].PlaneNormal, MinkowskiVert(VertsA.GetData(), VertsB.GetData(), 0)) > 0)
+	{
+		//tet faces are pointing inwards
+		for (TEPAEntry<T>& Entry : Entries)
+		{
+			Entry.SwapWinding(Entries.GetData());
+		}
+	}
+
+	return Entries;
 }
 
 struct FEPAFloodEntry
@@ -250,7 +323,7 @@ EPAResult EPA(TArray<TVec3<T>>& VertsABuffer, TArray<TVec3<T>>& VertsBBuffer, co
 	T UpperBound = TNumericLimits<T>::Max();
 	T LowerBound = TNumericLimits<T>::Lowest();
 
-	TArray<TEPAEntry<T>> Entries = InitializeEPA(VertsABuffer.GetData(), VertsBBuffer.GetData(), VertsABuffer.Num());
+	TArray<TEPAEntry<T>> Entries = InitializeEPA(VertsABuffer, VertsBBuffer, SupportA, SupportB);
 	std::priority_queue<FEPAEntryWrapper, std::vector<FEPAEntryWrapper>, std::greater<FEPAEntryWrapper>> Queue;
 	for(int32 Idx = 0; Idx < Entries.Num(); ++Idx)
 	{
