-Fix Non-Unity CIS

diff --git a/Engine/Source/Editor/Persona/Private/PreviewSceneCustomizations.cpp b/Engine/Source/Editor/Persona/Private/PreviewSceneCustomizations.cpp
--- a/Engine/Source/Editor/Persona/Private/PreviewSceneCustomizations.cpp
+++ b/Engine/Source/Editor/Persona/Private/PreviewSceneCustomizations.cpp
@@ -28,6 +28,7 @@
 #include "PhysicsEngine/PhysicsAsset.h"
 #include "Algo/Sort.h"
 #include "ScopedTransaction.h"
+#include "Features/IModularFeatures.h"
 
 #define LOCTEXT_NAMESPACE "PreviewSceneCustomizations"
 
diff --git a/Engine/Source/Runtime/Engine/Private/PhysicsEngine/Experimental/PhysInterface_Chaos.cpp b/Engine/Source/Runtime/Engine/Private/PhysicsEngine/Experimental/PhysInterface_Chaos.cpp
--- a/Engine/Source/Runtime/Engine/Private/PhysicsEngine/Experimental/PhysInterface_Chaos.cpp
+++ b/Engine/Source/Runtime/Engine/Private/PhysicsEngine/Experimental/PhysInterface_Chaos.cpp
@@ -35,6 +35,7 @@
 #include "Physics/PhysicsFiltering.h"
 #include "Collision/CollisionConversions.h"
 #include "PhysicsInterfaceUtilsCore.h"
+#include "Components/SkeletalMeshComponent.h"
 
 #if WITH_PHYSX
 #include "geometry/PxConvexMesh.h"
diff --git a/Engine/Source/Runtime/Experimental/Chaos/Private/Chaos/ClusterCreationParameters.cpp b/Engine/Source/Runtime/Experimental/Chaos/Private/Chaos/ClusterCreationParameters.cpp
new file mode 100644
--- /dev/null
+++ b/Engine/Source/Runtime/Experimental/Chaos/Private/Chaos/ClusterCreationParameters.cpp
@@ -0,0 +1,9 @@
+// Copyright 1998-2019 Epic Games, Inc. All Rights Reserved.
+
+#pragma once
+#include "Chaos/ClusterCreationParameters.h"
+
+namespace Chaos
+{
+	template<> struct FClusterCreationParameters<float>;
+}
diff --git a/Engine/Source/Runtime/Experimental/Chaos/Public/Chaos/PBDRigidClustering.h b/Engine/Source/Runtime/Experimental/Chaos/Public/Chaos/PBDRigidClustering.h
--- a/Engine/Source/Runtime/Experimental/Chaos/Public/Chaos/PBDRigidClustering.h
+++ b/Engine/Source/Runtime/Experimental/Chaos/Public/Chaos/PBDRigidClustering.h
@@ -1,12 +1,12 @@
 // Copyright 1998-2019 Epic Games, Inc. All Rights Reserved.
 #pragma once
 
-#include "Chaos/ClusterCreationParameters.h"
-#include "Chaos/ExternalCollisionData.h"
 #include "Chaos/PBDRigidClusteredParticles.h"
-#include "Chaos/PBDCollisionTypes.h"
 #include "Chaos/Transform.h"
+#include "Chaos/PBDCollisionTypes.h"
+#include "Chaos/ExternalCollisionData.h"
 #include "Chaos/TriangleMesh.h"
+#include "Chaos/ClusterCreationParameters.h"
 #include "Framework/BufferedData.h"
 
 #define TODO_CONVERT_GEOMETRY_COLLECTION_PARTICLE_INDICES_TO_PARTICLE_POINTERS 0
