Remove dirty checks for masses - this broke object mass initialization

diff --git a/Engine/Source/Runtime/Experimental/ChaosSolvers/Private/PhysicsProxy/SingleParticlePhysicsProxy.cpp b/Engine/Source/Runtime/Experimental/ChaosSolvers/Private/PhysicsProxy/SingleParticlePhysicsProxy.cpp
--- a/Engine/Source/Runtime/Experimental/ChaosSolvers/Private/PhysicsProxy/SingleParticlePhysicsProxy.cpp
+++ b/Engine/Source/Runtime/Experimental/ChaosSolvers/Private/PhysicsProxy/SingleParticlePhysicsProxy.cpp
@@ -210,22 +210,12 @@ void FSingleParticlePhysicsProxy<Chaos::TPBDRigidParticle<float, 3>>::PushToPhys
 		{
 			RigidHandle->SetW(Data->MW);
 		}
-		if (Data->DirtyFlags.IsDirty(Chaos::EParticleFlags::M))
-		{
-			RigidHandle->SetM(Data->MM);
-		}
-		if (Data->DirtyFlags.IsDirty(Chaos::EParticleFlags::InvM))
-		{
-			RigidHandle->SetInvM(Data->MInvM);
-		}
-		if (Data->DirtyFlags.IsDirty(Chaos::EParticleFlags::I))
-		{
-			RigidHandle->SetI(Data->MI);
-		}
-		if (Data->DirtyFlags.IsDirty(Chaos::EParticleFlags::InvI))
-		{
-			RigidHandle->SetInvI(Data->MInvI);
-		}
+
+		RigidHandle->SetM(Data->MM);
+		RigidHandle->SetInvM(Data->MInvM);
+		RigidHandle->SetI(Data->MI);
+		RigidHandle->SetInvI(Data->MInvI);
+
 		if (Data->DirtyFlags.IsDirty(Chaos::EParticleFlags::ExternalForce))
 		{
 			RigidHandle->SetExternalForce(Data->MExternalForce);
