Chaos
- fix collision angular impulse for second body
- fix multi-capsule union collision depth

diff --git a/Engine/Source/Runtime/Experimental/Chaos/Private/Chaos/PBDCollisionConstraint.cpp b/Engine/Source/Runtime/Experimental/Chaos/Private/Chaos/PBDCollisionConstraint.cpp
--- a/Engine/Source/Runtime/Experimental/Chaos/Private/Chaos/PBDCollisionConstraint.cpp
+++ b/Engine/Source/Runtime/Experimental/Chaos/Private/Chaos/PBDCollisionConstraint.cpp
@@ -25,6 +25,8 @@
 #include "PBDCollisionConstraint.ispc.generated.h"
 #endif
 
+//#pragma optimize("", off)
+
 int32 CollisionParticlesBVHDepth = 4;
 FAutoConsoleVariableRef CVarCollisionParticlesBVHDepth(TEXT("p.CollisionParticlesBVHDepth"), CollisionParticlesBVHDepth, TEXT("The maximum depth for collision particles bvh"));
 
@@ -680,26 +682,30 @@ void TPBDCollisionConstraint<T, d>::Apply(const T Dt, FRigidBodyContactConstrain
 		}
 		Impulse = GetEnergyClampedImpulse(Constraint, Impulse, VectorToPoint1, VectorToPoint2, Body1Velocity, Body2Velocity);
 		Constraint.AccumulatedImpulse += Impulse;
-		TVector<T, d> AngularImpulse1 = TVector<T, d>::CrossProduct(VectorToPoint1, Impulse) + AngularImpulse;
-		TVector<T, d> AngularImpulse2 = TVector<T, d>::CrossProduct(VectorToPoint2, -Impulse) - AngularImpulse;
 		if (PBDRigid0)
 		{
 			// Velocity update for next step
-			PBDRigid0->V() += PBDRigid0->InvM() * Impulse;
-			PBDRigid0->W() += WorldSpaceInvI1 * AngularImpulse1;
+			TVector<T, d> AngularImpulse1 = TVector<T, d>::CrossProduct(VectorToPoint1, Impulse) + AngularImpulse;
+			FVec3 DV1 = PBDRigid0->InvM() * Impulse;
+			FVec3 DW1 = WorldSpaceInvI1 * AngularImpulse1;
+			PBDRigid0->V() += DV1;
+			PBDRigid0->W() += DW1;
 			// Position update as part of pbd
-			PBDRigid0->P() += (PBDRigid0->InvM() * Impulse) * Dt;
-			PBDRigid0->Q() += TRotation<T, d>::FromElements(WorldSpaceInvI1 * AngularImpulse1, 0.f) * Q0 * Dt * T(0.5);
+			PBDRigid0->P() += DV1 * Dt;
+			PBDRigid0->Q() += TRotation<T, d>::FromElements(DW1, 0.f) * Q0 * Dt * T(0.5);
 			PBDRigid0->Q().Normalize();
 		}
 		if (PBDRigid1)
 		{
 			// Velocity update for next step
-			PBDRigid1->V() -= PBDRigid1->InvM() * Impulse;
-			PBDRigid1->W() += WorldSpaceInvI2 * AngularImpulse2;
+			TVector<T, d> AngularImpulse2 = TVector<T, d>::CrossProduct(VectorToPoint2, -Impulse) - AngularImpulse;
+			FVec3 DV2 = -PBDRigid1->InvM() * Impulse;
+			FVec3 DW2 = WorldSpaceInvI2 * AngularImpulse2;
+			PBDRigid1->V() += DV2;
+			PBDRigid1->W() += DW2;
 			// Position update as part of pbd
-			PBDRigid1->P() -= (PBDRigid1->InvM() * Impulse) * Dt;
-			PBDRigid1->Q() += TRotation<T, d>::FromElements(WorldSpaceInvI2 * AngularImpulse2, 0.f) * Q0 * Dt * T(0.5);
+			PBDRigid1->P() += DV2 * Dt;
+			PBDRigid1->Q() += TRotation<T, d>::FromElements(DW2, 0.f) * Q1 * Dt * T(0.5);
 			PBDRigid1->Q().Normalize();
 		}
 	}
@@ -1523,10 +1529,14 @@ void UpdateCapsuleCapsuleConstraint(const TCapsule<T>& A, const TRigidTransform<
 	T DeltaLen = Delta.Size();
 	if (DeltaLen > KINDA_SMALL_NUMBER)
 	{
-		TVector<T, d> Dir = Delta / DeltaLen;
-		Contact.Phi = DeltaLen - (A.GetRadius() + B.GetRadius());
-		Contact.Normal = -Dir;
-		Contact.Location = P1 + Dir * A.GetRadius();
+		T NewPhi = DeltaLen - (A.GetRadius() + B.GetRadius());
+		if (NewPhi < Contact.Phi)
+		{
+			TVector<T, d> Dir = Delta / DeltaLen;
+			Contact.Phi = NewPhi;
+			Contact.Normal = -Dir;
+			Contact.Location = P1 + Dir * A.GetRadius();
+		}
 	}
 }
 
