#UE4 Fixed CHAOS_CHECKED

diff --git a/Engine/Source/Runtime/Experimental/ChaosSolverEngine/Private/Chaos/ChaosSolverActor.cpp b/Engine/Source/Runtime/Experimental/ChaosSolverEngine/Private/Chaos/ChaosSolverActor.cpp
--- a/Engine/Source/Runtime/Experimental/ChaosSolverEngine/Private/Chaos/ChaosSolverActor.cpp
+++ b/Engine/Source/Runtime/Experimental/ChaosSolverEngine/Private/Chaos/ChaosSolverActor.cpp
@@ -206,7 +206,7 @@ AChaosSolverActor::AChaosSolverActor(const FObjectInitializer& ObjectInitializer
 	// @question(Benn) : Does this need to be created on the Physics thread using a queued command?
 	PhysScene = MakeShareable(new FPhysScene_Chaos(this
 #if CHAOS_CHECKED
-		, TEXT("Solver Actor Physics"))
+		, TEXT("Solver Actor Physics")
 #endif
 	);
 	Solver = PhysScene->GetSolver();
