#UE4 Better fix for CHAOS_CHECKED

diff --git a/Engine/Source/Runtime/Experimental/ChaosSolverEngine/Private/Chaos/ChaosSolverActor.cpp b/Engine/Source/Runtime/Experimental/ChaosSolverEngine/Private/Chaos/ChaosSolverActor.cpp
--- a/Engine/Source/Runtime/Experimental/ChaosSolverEngine/Private/Chaos/ChaosSolverActor.cpp
+++ b/Engine/Source/Runtime/Experimental/ChaosSolverEngine/Private/Chaos/ChaosSolverActor.cpp
@@ -208,7 +208,7 @@ AChaosSolverActor::AChaosSolverActor(const FObjectInitializer& ObjectInitializer
 #if CHAOS_CHECKED
 		, TEXT("Solver Actor Physics")
 #endif
-	);
+	));
 	Solver = PhysScene->GetSolver();
 	// Ticking setup for collision/breaking notifies
 	PrimaryActorTick.TickGroup = TG_PostPhysics;
