Reenable overlap of sphere optimization to use point with thickness

diff --git a/Engine/Source/Runtime/Experimental/Chaos/Public/Chaos/GeometryQueries.h b/Engine/Source/Runtime/Experimental/Chaos/Public/Chaos/GeometryQueries.h
--- a/Engine/Source/Runtime/Experimental/Chaos/Public/Chaos/GeometryQueries.h
+++ b/Engine/Source/Runtime/Experimental/Chaos/Public/Chaos/GeometryQueries.h
@@ -53,12 +53,12 @@ namespace Chaos
 
 		const FRigidTransform3 BToATM = BTM.GetRelativeTransform(ATM);
 
-		if(false && BType == ImplicitObjectType::Sphere)	//todo: fix
+		if(BType == ImplicitObjectType::Sphere)
 		{
-			ensure(false);	//todo: override helper
-			/*const TSphere<T, d>& BSphere = static_cast<const TSphere<T, d>&>(B);
-			const TVector<T, d> PtInA = BToATM.TransformPositionNoScale(BSphere.GetCenter());
-			return A.Overlap(PtInA, Thickness + BSphere.GetRadius());*/
+			const FImplicitObject& BBase = static_cast<const FImplicitObject&>(B);
+			const TSphere<FReal, 3>& BSphere = static_cast<const TSphere<FReal, 3>&>(BBase);
+			const FVec3 PtInA = BToATM.TransformPositionNoScale(BSphere.GetCenter());
+			return A.Overlap(PtInA, Thickness + BSphere.GetRadius());
 		}
 		//todo: A is a sphere
 		else if (A.IsConvex())
