Chaos: Fix THeightField OverlapGeomImp using TParticles/TConvex instead of TTriangle for GJKIntersection.

diff --git a/Engine/Source/Runtime/Experimental/Chaos/Private/Chaos/HeightField.cpp b/Engine/Source/Runtime/Experimental/Chaos/Private/Chaos/HeightField.cpp
--- a/Engine/Source/Runtime/Experimental/Chaos/Private/Chaos/HeightField.cpp
+++ b/Engine/Source/Runtime/Experimental/Chaos/Private/Chaos/HeightField.cpp
@@ -1014,16 +1014,7 @@ namespace Chaos
 			//However, maybe we should check if it's behind the triangle plane. Also, we should enforce this winding in some way
 			const TVector<T, 3> Offset = TVector<T, 3>::CrossProduct(AB, AC);
 
-			// Ugly but required for now until we have an easier way to do tri collisons 
-			TParticles<T, 3> Particles;
-			Particles.AddParticles(3);
-
-			Particles.X(0) = A;
-			Particles.X(1) = B;
-			Particles.X(2) = C;
-
-			TConvex<T, 3> TriangleConvex(Particles);
-
+			TTriangle<T> TriangleConvex(A, B, C);
 			return GJKIntersection(TriangleConvex, QueryGeom, QueryTM, Thickness, Offset);
 		};
 
