Make EPA robust to initial conditions like touching edges and planer minkwoski volume

diff --git a/Engine/Source/Runtime/Experimental/Chaos/Public/Chaos/EPA.h b/Engine/Source/Runtime/Experimental/Chaos/Public/Chaos/EPA.h
--- a/Engine/Source/Runtime/Experimental/Chaos/Public/Chaos/EPA.h
+++ b/Engine/Source/Runtime/Experimental/Chaos/Public/Chaos/EPA.h
@@ -154,74 +154,95 @@ TArray<TEPAEntry<T>> InitializeEPA(TArray<TVec3<T>>& VertsA, TArray<TVec3<T>>& V
 		}
 	};
 
+	Entries.AddUninitialized(4);
+
+	bool bValid = false;
+
 	switch(NumVerts)
 	{
-		case 1: check(false);	break;	//This case is not supported - for GJK it means two surfaces are touching so the penetration is 0
-		case 3:
+		case 1:
 		{
-			//triangle, add farthest point along normal
-			
-			Entries.AddUninitialized(4);
-			ensure(Entries[3].Initialize(VertsA.GetData(), VertsB.GetData(), 0, 2, 1, { 1,0,2 }, { 2,0,0 }));
-
-			const TEPAEntry<T>& Base = Entries[3];
-
-			AddFartherPoint(Base.PlaneNormal);
-
-			ensure(Entries[0].Initialize(VertsA.GetData(), VertsB.GetData(), 1, 2, 3, { 3, 1, 2 }, { 1,1, 1 }));
-			ensure(Entries[1].Initialize(VertsA.GetData(), VertsB.GetData(), 0, 3, 2, { 2,0,3 }, { 2, 1, 0 }));
-			ensure(Entries[2].Initialize(VertsA.GetData(), VertsB.GetData(), 0, 1, 3, { 3,0, 1 }, { 2,2,0 }));
-
+			ensure(false);	//This case is not supported - for GJK it means two surfaces are touching so the penetration is 0
 			break;
 		}
 		case 2:
 		{
 			//line, add farthest point along most orthogonal axes
-
 			TVec3<T> Dir = MinkowskiVert(VertsA.GetData(), VertsB.GetData(), 1) - MinkowskiVert(VertsA.GetData(), VertsB.GetData(), 0);
-			ensure(Dir.SafeNormalize());
-			//find most opposing axis
-			int32 BestAxis = 0;
-			T MinVal = TNumericLimits<T>::Max();
-			for (int32 Axis = 0; Axis < 3; ++Axis)
+
+			bValid = Dir.SizeSquared() > 1e-4;
+			if (ensure(bValid))	//two verts given should be distinct
 			{
-				const T AbsVal = FMath::Abs(Dir[Axis]);
-				if (MinVal > AbsVal)
+				//find most opposing axis
+				int32 BestAxis = 0;
+				T MinVal = TNumericLimits<T>::Max();
+				for (int32 Axis = 0; Axis < 3; ++Axis)
 				{
-					BestAxis = Axis;
-					MinVal = AbsVal;
+					const T AbsVal = FMath::Abs(Dir[Axis]);
+					if (MinVal > AbsVal)
+					{
+						BestAxis = Axis;
+						MinVal = AbsVal;
+					}
 				}
-			}
-			const TVec3<T> OtherAxis = TVec3<T>::AxisVector(BestAxis);
-			const TVec3<T> Orthog = TVec3<T>::CrossProduct(Dir, OtherAxis);
+				const TVec3<T> OtherAxis = TVec3<T>::AxisVector(BestAxis);
+				const TVec3<T> Orthog = TVec3<T>::CrossProduct(Dir, OtherAxis);
+				const TVec3<T> Orthog2 = TVec3<T>::CrossProduct(Orthog, Dir);
+
+				AddFartherPoint(Orthog);
+				AddFartherPoint(Orthog2);
 
-			AddFartherPoint(OtherAxis);
-			AddFartherPoint(Orthog);
+				bValid = Entries[0].Initialize(VertsA.GetData(), VertsB.GetData(), 1, 2, 3, { 3, 1, 2 }, { 1,1, 1 });
+				bValid &= Entries[1].Initialize(VertsA.GetData(), VertsB.GetData(), 0, 3, 2, { 2,0,3 }, { 2, 1, 0 });
+				bValid &= Entries[2].Initialize(VertsA.GetData(), VertsB.GetData(), 0, 1, 3, { 3,0, 1 }, { 2,2,0 });
+				bValid &= Entries[3].Initialize(VertsA.GetData(), VertsB.GetData(), 0, 2, 1, { 1,0,2 }, { 2,0,0 });
+			}
+			break;
 		}
-		case 4:
+		case 3:
 		{
-			//make sure all triangle normals are facing out
-			Entries.AddUninitialized(4);
+			//triangle, add farthest point along normal
+			bValid = Entries[3].Initialize(VertsA.GetData(), VertsB.GetData(), 0, 2, 1, { 1,0,2 }, { 2,0,0 });	
+			if (ensure(bValid)) //input verts must form a valid triangle
+			{
+				const TEPAEntry<T>& Base = Entries[3];
 
-			ensure(Entries[0].Initialize(VertsA.GetData(), VertsB.GetData(), 1, 2, 3, { 3, 1, 2 }, { 1,1, 1 }));
-			ensure(Entries[1].Initialize(VertsA.GetData(), VertsB.GetData(), 0, 3, 2, { 2,0,3 }, { 2, 1, 0 }));
-			ensure(Entries[2].Initialize(VertsA.GetData(), VertsB.GetData(), 0, 1, 3, { 3,0, 1 }, { 2,2,0 }));
-			ensure(Entries[3].Initialize(VertsA.GetData(), VertsB.GetData(), 0, 2, 1, { 1,0,2 }, { 2,0,0 }));
+				AddFartherPoint(Base.PlaneNormal);
 
+				bValid = Entries[0].Initialize(VertsA.GetData(), VertsB.GetData(), 1, 2, 3, { 3, 1, 2 }, { 1,1, 1 });
+				bValid &= Entries[1].Initialize(VertsA.GetData(), VertsB.GetData(), 0, 3, 2, { 2,0,3 }, { 2, 1, 0 });
+				bValid &= Entries[2].Initialize(VertsA.GetData(), VertsB.GetData(), 0, 1, 3, { 3,0, 1 }, { 2,2,0 });
+			}
+			break;
+		}
+		case 4:
+		{
+			bValid = Entries[0].Initialize(VertsA.GetData(), VertsB.GetData(), 1, 2, 3, { 3, 1, 2 }, { 1,1, 1 });
+			bValid &= Entries[1].Initialize(VertsA.GetData(), VertsB.GetData(), 0, 3, 2, { 2,0,3 }, { 2, 1, 0 });
+			bValid &= Entries[2].Initialize(VertsA.GetData(), VertsB.GetData(), 0, 1, 3, { 3,0, 1 }, { 2,2,0 });
+			bValid &= Entries[3].Initialize(VertsA.GetData(), VertsB.GetData(), 0, 2, 1, { 1,0,2 }, { 2,0,0 });
+			ensure(bValid);	//expect user to give us valid tetrahedron
 			break;
 		}
 
 		default: ensure(false);
 	}
 
-	if (TVec3<T>::DotProduct(Entries[0].PlaneNormal, MinkowskiVert(VertsA.GetData(), VertsB.GetData(), 0)) > 0)
+	if (bValid)
 	{
-		//tet faces are pointing inwards
-		for (TEPAEntry<T>& Entry : Entries)
+		//make sure normals are pointing out of tetrahedron
+		if (TVec3<T>::DotProduct(Entries[0].PlaneNormal, MinkowskiVert(VertsA.GetData(), VertsB.GetData(), 0)) > 0)
 		{
-			Entry.SwapWinding(Entries.GetData());
+			for (TEPAEntry<T>& Entry : Entries)
+			{
+				Entry.SwapWinding(Entries.GetData());
+			}
 		}
 	}
+	else
+	{
+		Entries.SetNum(0);
+	}
 
 	return Entries;
 }
@@ -300,10 +321,11 @@ void ComputeEPAResults(const TVec3<T>* VertsA, const TVec3<T>* VertsB, const TEP
 	}
 }
 
-enum EPAResult
+enum class EPAResult
 {
 	Ok,
 	MaxIterations,
+	BadInitialSimplex
 };
 
 template <typename T, typename SupportALambda, typename SupportBLambda>
@@ -324,6 +346,17 @@ EPAResult EPA(TArray<TVec3<T>>& VertsABuffer, TArray<TVec3<T>>& VertsBBuffer, co
 	T LowerBound = TNumericLimits<T>::Lowest();
 
 	TArray<TEPAEntry<T>> Entries = InitializeEPA(VertsABuffer, VertsBBuffer, SupportA, SupportB);
+
+	if (Entries.Num() < 4)
+	{
+		//either degenerate or a touching hit. Either way return penetration 0
+		OutPenetration = 0;
+		OutDir = TVec3<T>(0, 0, 1);
+		WitnessA = TVec3<T>(0);
+		WitnessB = TVec3<T>(0);
+		return EPAResult::BadInitialSimplex;
+	}
+
 	std::priority_queue<FEPAEntryWrapper, std::vector<FEPAEntryWrapper>, std::greater<FEPAEntryWrapper>> Queue;
 	for(int32 Idx = 0; Idx < Entries.Num(); ++Idx)
 	{
