GJKRaycast uses EPA when needed.

diff --git a/Engine/Source/Runtime/Experimental/Chaos/Public/Chaos/EPA.h b/Engine/Source/Runtime/Experimental/Chaos/Public/Chaos/EPA.h
--- a/Engine/Source/Runtime/Experimental/Chaos/Public/Chaos/EPA.h
+++ b/Engine/Source/Runtime/Experimental/Chaos/Public/Chaos/EPA.h
@@ -162,7 +162,7 @@ TArray<TEPAEntry<T>> InitializeEPA(TArray<TVec3<T>>& VertsA, TArray<TVec3<T>>& V
 	{
 		case 1:
 		{
-			ensure(false);	//This case is not supported - for GJK it means two surfaces are touching so the penetration is 0
+			//assuming it's a touching hit at origin
 			break;
 		}
 		case 2:
diff --git a/Engine/Source/Runtime/Experimental/Chaos/Public/Chaos/GJK.h b/Engine/Source/Runtime/Experimental/Chaos/Public/Chaos/GJK.h
--- a/Engine/Source/Runtime/Experimental/Chaos/Public/Chaos/GJK.h
+++ b/Engine/Source/Runtime/Experimental/Chaos/Public/Chaos/GJK.h
@@ -339,7 +339,7 @@ namespace Chaos
 				Simplex[SimplexIDs.NumVerts++] = W;	//this is really X - P which is what we need for simplex computation
 			}
 
-			if (bInflatedCloseEnough)
+			if (bInflatedCloseEnough && VDotW >= 0)
 			{
 				//Inflated shapes are close enough, but we want MTD so we need to find closest point on core shape
 				const T VDotW2 = VDotW * VDotW;
@@ -420,25 +420,46 @@ namespace Chaos
 			else
 			{
 				//todo: use EPA
-				if (false && SimplexIDs.NumVerts == 4)
-				{
-					TArray<TVec3<T>> VertsA = { As[0], As[1], As[2], As[3] };
-					TArray<TVec3<T>> VertsB = { Bs[0], Bs[1], Bs[2], Bs[3] };
+				TArray<TVec3<T>> VertsA;
+				TArray<TVec3<T>> VertsB;
 
-					T Penetration;
-					TVec3<T> MTD, ClosestA, ClosestB;
-					EPA(VertsA, VertsB, SupportAFunc, SupportBFunc, Penetration, MTD, ClosestA, ClosestB);
+				VertsA.Reserve(8);
+				VertsB.Reserve(8);
+
+				if (NumIterations)
+				{
+					for (int i = 0; i < SimplexIDs.NumVerts; ++i)
+					{
+						VertsA.Add(As[i]);
+						const TVec3<T> BInA = StartTM.TransformPosition(Bs[i]);
+						VertsB.Add(BInA);
+					}
 
-					OutNormal = (ClosestA - ClosestB);
-					OutNormal.SafeNormalize();
-					const TVector<T, 3> ClosestLocal = ClosestB - OutNormal * ThicknessB;
+					auto SupportBInAFunc = [&](const TVec3<T>& Dir)
+					{
+						const TVec3<T> SupportB = SupportBFunc(Dir);
+						return StartTM.TransformPosition(SupportB);
+					};
 
-					OutPosition = StartPoint + ClosestLocal;
-					OutTime = -Penetration - ThicknessA - ThicknessB;
+					T Penetration;
+					TVec3<T> MTD, ClosestA, ClosestBInA;
+					if (EPA(VertsA, VertsB, SupportAFunc, SupportBInAFunc, Penetration, MTD, ClosestA, ClosestBInA) == EPAResult::BadInitialSimplex)
+					{
+						//assume touching hit
+						OutTime = 0;
+						OutPosition = TVector<T, 3>(0);
+						OutNormal = { 0,0,1 };
+					}
+					else
+					{
+						OutNormal = MTD;
+						OutTime = -Penetration - Inflation;
+						OutPosition = ClosestA;
+					}
 				}
 				else
 				{
-					ensure(false);	//todo: EPA with num verts < 4
+					//didn't even go into gjk loop, touching hit
 					OutTime = 0;
 					OutPosition = TVector<T, 3>(0);
 					OutNormal = { 0,0,1 };
