Fix MTD when inflated elements have core shapes just touching

diff --git a/Engine/Source/Runtime/Experimental/Chaos/Public/Chaos/GJK.h b/Engine/Source/Runtime/Experimental/Chaos/Public/Chaos/GJK.h
--- a/Engine/Source/Runtime/Experimental/Chaos/Public/Chaos/GJK.h
+++ b/Engine/Source/Runtime/Experimental/Chaos/Public/Chaos/GJK.h
@@ -443,26 +443,26 @@ namespace Chaos
 
 					T Penetration;
 					TVec3<T> MTD, ClosestA, ClosestBInA;
-					if (EPA(VertsA, VertsB, SupportAFunc, SupportBInAFunc, Penetration, MTD, ClosestA, ClosestBInA) == EPAResult::BadInitialSimplex)
-					{
-						//assume touching hit
-						OutTime = 0;
-						OutPosition = TVector<T, 3>(0);
-						OutNormal = { 0,0,1 };
-					}
-					else
+					if (EPA(VertsA, VertsB, SupportAFunc, SupportBInAFunc, Penetration, MTD, ClosestA, ClosestBInA) != EPAResult::BadInitialSimplex)
 					{
 						OutNormal = MTD;
 						OutTime = -Penetration - Inflation;
 						OutPosition = ClosestA;
 					}
+					else
+					{
+						//assume touching hit
+						OutTime = -Inflation;
+						OutNormal = { 0,0,1 };
+						OutPosition = As[0] + OutNormal * ThicknessA;
+					}
 				}
 				else
 				{
 					//didn't even go into gjk loop, touching hit
-					OutTime = 0;
-					OutPosition = TVector<T, 3>(0);
+					OutTime = -Inflation;
 					OutNormal = { 0,0,1 };
+					OutPosition = As[0] + OutNormal * ThicknessA;
 				}
 			}
 		}
