Chaos : Collisions

- Moving out some functions from the CollisionConstraint to help reduce the file size.

diff --git a/Engine/Source/Runtime/Experimental/Chaos/Private/Chaos/CollisionResolutionAlgo.cpp b/Engine/Source/Runtime/Experimental/Chaos/Private/Chaos/CollisionResolutionAlgo.cpp
new file mode 100644
--- /dev/null
+++ b/Engine/Source/Runtime/Experimental/Chaos/Private/Chaos/CollisionResolutionAlgo.cpp
@@ -0,0 +1,1310 @@
+// Copyright 1998-2019 Epic Games, Inc. All Rights Reserved.
+
+#include "Chaos/CollisionResolutionAlgo.h"
+
+#include "Chaos/ChaosPerfTest.h"
+#include "Chaos/CollisionResolutionConvexConvex.h"
+#include "Chaos/Defines.h"
+#include "Chaos/ImplicitObjectUnion.h"
+#include "Chaos/ImplicitObjectScaled.h"
+#include "Chaos/ISpatialAccelerationCollection.h"
+#include "Chaos/Levelset.h"
+#include "Chaos/Pair.h"
+#include "Chaos/Sphere.h"
+#include "Chaos/Capsule.h"
+#include "Chaos/Transform.h"
+#include "ChaosLog.h"
+#include "ChaosStats.h"
+
+#if INTEL_ISPC
+#include "PBDCollisionConstraint.ispc.generated.h"
+#endif
+
+namespace Chaos
+{
+	template<typename T, int d>
+	void ConstructLevelsetConstraints(TGeometryParticleHandle<T, d>* Particle0, TGeometryParticleHandle<T, d>* Particle1, const FImplicitObject* Implicit0, const FImplicitObject* Implicit1, const T Thickness, TRigidBodyContactConstraint<T, d>& Constraint)
+	{
+
+		if (!Particle1->Geometry() || (Particle0->AsDynamic() && !Particle0->AsDynamic()->CollisionParticlesSize() && Particle0->Geometry() && !Particle0->Geometry()->IsUnderlyingUnion()))
+		{
+			Constraint.Particle[0] = Particle1;
+			Constraint.Particle[1] = Particle0;
+			Constraint.AddManifold(Implicit1, Implicit0);
+		}
+		else
+		{
+			Constraint.Particle[0] = Particle0;
+			Constraint.Particle[1] = Particle1;
+			Constraint.AddManifold(Implicit0, Implicit1);
+		}
+	}
+
+	template<typename T, int d>
+	void ConstructBoxConstraints(TGeometryParticleHandle<T, d>* Particle0, TGeometryParticleHandle<T, d>* Particle1, const FImplicitObject* Implicit0, const FImplicitObject* Implicit1, const T Thickness, TRigidBodyContactConstraint<T, d>& Constraint)
+	{
+		Constraint.Particle[0] = Particle0;
+		Constraint.Particle[1] = Particle1;
+		Constraint.AddManifold(Implicit0, Implicit1);
+	}
+
+	template<typename T, int d>
+	void ConstructBoxPlaneConstraints(TGeometryParticleHandle<T, d>* Particle0, TGeometryParticleHandle<T, d>* Particle1, const FImplicitObject* Implicit0, const FImplicitObject* Implicit1, const T Thickness, TRigidBodyContactConstraint<T, d>& Constraint)
+	{
+		Constraint.Particle[0] = Particle0;
+		Constraint.Particle[1] = Particle1;
+		Constraint.AddManifold(Implicit0, Implicit1);
+	}
+
+	template<typename T, int d>
+	void ConstructSphereConstraints(TGeometryParticleHandle<T, d>* Particle0, TGeometryParticleHandle<T, d>* Particle1, const FImplicitObject* Implicit0, const FImplicitObject* Implicit1, const T Thickness, TRigidBodyContactConstraint<T, d>& Constraint)
+	{
+		Constraint.Particle[0] = Particle0;
+		Constraint.Particle[1] = Particle1;
+		Constraint.AddManifold(Implicit0, Implicit1);
+	}
+
+	template<typename T, int d>
+	void ConstructSpherePlaneConstraints(TGeometryParticleHandle<T, d>* Particle0, TGeometryParticleHandle<T, d>* Particle1, const FImplicitObject* Implicit0, const FImplicitObject* Implicit1, const T Thickness, TRigidBodyContactConstraint<T, d>& Constraint)
+	{
+		Constraint.Particle[0] = Particle0;
+		Constraint.Particle[1] = Particle1;
+		Constraint.AddManifold(Implicit0, Implicit1);
+	}
+
+	template<typename T, int d>
+	void ConstructSphereBoxConstraints(TGeometryParticleHandle<T, d>* Particle0, TGeometryParticleHandle<T, d>* Particle1, const FImplicitObject* Implicit0, const FImplicitObject* Implicit1, const T Thickness, TRigidBodyContactConstraint<T, d>& Constraint)
+	{
+		Constraint.Particle[0] = Particle0;
+		Constraint.Particle[1] = Particle1;
+		Constraint.AddManifold(Implicit0, Implicit1);
+	}
+
+	template<typename T, int d>
+	void ConstructCapsuleCapsuleConstraints(TGeometryParticleHandle<T, d>* Particle0, TGeometryParticleHandle<T, d>* Particle1, const FImplicitObject* Implicit0, const FImplicitObject* Implicit1, const T Thickness, TRigidBodyContactConstraint<T, d>& Constraint)
+	{
+		Constraint.Particle[0] = Particle0;
+		Constraint.Particle[1] = Particle1;
+		Constraint.AddManifold(Implicit0, Implicit1);
+	}
+
+	template<typename T, int d>
+	void ConstructCapsuleBoxConstraints(TGeometryParticleHandle<T, d>* Particle0, TGeometryParticleHandle<T, d>* Particle1, const FImplicitObject* Implicit0, const FImplicitObject* Implicit1, const T Thickness, TRigidBodyContactConstraint<T, d>& Constraint)
+	{
+		Constraint.Particle[0] = Particle0;
+		Constraint.Particle[1] = Particle1;
+		Constraint.AddManifold(Implicit0, Implicit1);
+	}
+
+	template<typename T, int d>
+	void ConstructSingleUnionConstraints(TGeometryParticleHandle<T, d>* Particle0, TGeometryParticleHandle<T, d>* Particle1, const FImplicitObject* Implicit0, const FImplicitObject* Implicit1, const T Thickness, TRigidBodyContactConstraint<T, d>& Constraint)
+	{
+		Constraint.Particle[0] = Particle0;
+		Constraint.Particle[1] = Particle1;
+		Constraint.AddManifold(Implicit0, Implicit1);
+	}
+
+	template<typename T, int d>
+	void ConstructUnionUnionConstraints(TGeometryParticleHandle<T, d>* Particle0, TGeometryParticleHandle<T, d>* Particle1, const FImplicitObject* Implicit0, const FImplicitObject* Implicit1, const T Thickness, TRigidBodyContactConstraint<T, d>& Constraint)
+	{
+		SCOPE_CYCLE_COUNTER(STAT_UpdateUnionUnionConstraint);
+
+		const FImplicitObject* ParticleObj = Particle0->Geometry().Get();
+		const FImplicitObject* LevelsetObj = Particle1->Geometry().Get();
+
+		TRigidTransform<T, d> ParticlesTM = TRigidTransform<T, d>(Particle0->X(), Particle0->R());
+		TRigidTransform<T, d> LevelsetTM = TRigidTransform<T, d>(Particle1->X(), Particle1->R());
+
+		const TArray<Pair<const FImplicitObject*, TRigidTransform<T, d>>> LevelsetShapes = FindRelevantShapes2(ParticleObj, ParticlesTM, *LevelsetObj, LevelsetTM, Thickness);
+
+		for (const Pair<const FImplicitObject*, TRigidTransform<T, d>>& LevelsetObjPair : LevelsetShapes)
+		{
+			const FImplicitObject* LevelsetInnerObj = LevelsetObjPair.First;
+			const TRigidTransform<T, d>& LevelsetInnerObjTM = LevelsetObjPair.Second * LevelsetTM;
+
+			//now find all particle inner objects
+			const TArray<Pair<const FImplicitObject*, TRigidTransform<T, d>>> ParticleShapes = FindRelevantShapes2(LevelsetInnerObj, LevelsetInnerObjTM, *ParticleObj, ParticlesTM, Thickness);
+
+			//for each inner obj pair, update constraint
+			for (const Pair<const FImplicitObject*, TRigidTransform<T, d>>& ParticlePair : ParticleShapes)
+			{
+				const FImplicitObject* ParticleInnerObj = ParticlePair.First;
+				ConstructConstraintsImpl<T, d>(Particle0, Particle1, ParticleInnerObj, LevelsetInnerObj, Thickness, Constraint);
+			}
+		}
+	}
+
+	template<typename T, int d>
+	void ConstructPairConstraintImpl(TGeometryParticleHandle<T, d>* Particle0, TGeometryParticleHandle<T, d>* Particle1, const FImplicitObject* Implicit0, const FImplicitObject* Implicit1, const T Thickness, TRigidBodyContactConstraint<T, d>& Constraint)
+	{
+		// See if we already have a constraint for this shape pair
+		if (Constraint.ContainsManifold(Implicit0, Implicit1))
+		{
+			return;
+		}
+
+		if (!Implicit0 || !Implicit1)
+		{
+			ConstructLevelsetConstraints(Particle0, Particle1, Implicit0, Implicit1, Thickness, Constraint);
+		}
+		if (Implicit0->GetType() == TBox<T, d>::StaticType() && Implicit1->GetType() == TBox<T, d>::StaticType())
+		{
+			ConstructBoxConstraints(Particle0, Particle1, Implicit0, Implicit1, Thickness, Constraint);
+		}
+		else if (Implicit0->GetType() == TSphere<T, d>::StaticType() && Implicit1->GetType() == TSphere<T, d>::StaticType())
+		{
+			ConstructSphereConstraints(Particle0, Particle1, Implicit0, Implicit1, Thickness, Constraint);
+		}
+		else if (Implicit0->GetType() == TBox<T, d>::StaticType() && Implicit1->GetType() == TPlane<T, d>::StaticType())
+		{
+			ConstructBoxPlaneConstraints(Particle0, Particle1, Implicit0, Implicit1, Thickness, Constraint);
+		}
+		else if (Implicit1->GetType() == TBox<T, d>::StaticType() && Implicit0->GetType() == TPlane<T, d>::StaticType())
+		{
+			ConstructBoxPlaneConstraints(Particle1, Particle0, Implicit0, Implicit1, Thickness, Constraint);
+		}
+		else if (Implicit0->GetType() == TSphere<T, d>::StaticType() && Implicit1->GetType() == TPlane<T, d>::StaticType())
+		{
+			ConstructSpherePlaneConstraints(Particle0, Particle1, Implicit0, Implicit1, Thickness, Constraint);
+		}
+		else if (Implicit1->GetType() == TSphere<T, d>::StaticType() && Implicit0->GetType() == TPlane<T, d>::StaticType())
+		{
+			ConstructSpherePlaneConstraints(Particle1, Particle0, Implicit0, Implicit1, Thickness, Constraint);
+		}
+		else if (Implicit0->GetType() == TSphere<T, d>::StaticType() && Implicit1->GetType() == TBox<T, d>::StaticType())
+		{
+			ConstructSphereBoxConstraints(Particle0, Particle1, Implicit0, Implicit1, Thickness, Constraint);
+		}
+		else if (Implicit1->GetType() == TSphere<T, d>::StaticType() && Implicit0->GetType() == TBox<T, d>::StaticType())
+		{
+			ConstructSphereBoxConstraints(Particle1, Particle0, Implicit0, Implicit1, Thickness, Constraint);
+		}
+		else if (Implicit0->GetType() == TCapsule<T>::StaticType() && Implicit1->GetType() == TCapsule<T>::StaticType())
+		{
+			ConstructCapsuleCapsuleConstraints(Particle0, Particle1, Implicit0, Implicit1, Thickness, Constraint);
+		}
+		else if (Implicit0->GetType() == TCapsule<T>::StaticType() && Implicit1->GetType() == TBox<T, d>::StaticType())
+		{
+			ConstructCapsuleBoxConstraints(Particle0, Particle1, Implicit0, Implicit1, Thickness, Constraint);
+		}
+		else if (Implicit1->GetType() == TCapsule<T>::StaticType() && Implicit0->GetType() == TBox<T, d>::StaticType())
+		{
+			ConstructCapsuleBoxConstraints(Particle1, Particle0, Implicit0, Implicit1, Thickness, Constraint);
+		}
+		else if (Implicit0->GetType() < TImplicitObjectUnion<T, d>::StaticType() && Implicit1->GetType() == TImplicitObjectUnion<T, d>::StaticType())
+		{
+			ConstructSingleUnionConstraints(Particle0, Particle1, Implicit0, Implicit1, Thickness, Constraint);
+		}
+		else if (Implicit0->GetType() == TImplicitObjectUnion<T, d>::StaticType() && Implicit1->GetType() < TImplicitObjectUnion<T, d>::StaticType())
+		{
+			ConstructSingleUnionConstraints(Particle1, Particle0, Implicit0, Implicit1, Thickness, Constraint);
+		}
+		else if (Implicit0->GetType() == TImplicitObjectUnion<T, d>::StaticType() && Implicit1->GetType() == TImplicitObjectUnion<T, d>::StaticType())
+		{
+			// Union-union creates multiple manifolds - we should never get here for this pair type. See ConstructConstraintsImpl and ConstructUnionUnionConstraints
+			ensure(false);
+		}
+		else if (Implicit0->IsConvex() && Implicit1->IsConvex())
+		{
+			CollisionResolutionConvexConvex<T, d>::ConstructConvexConvexConstraints(Particle0, Particle1, Implicit0, Implicit1, Thickness, Constraint);
+		}
+		else
+		{
+			ConstructLevelsetConstraints(Particle0, Particle1, Implicit0, Implicit1, Thickness, Constraint);
+		}
+	}
+
+	template <typename T, int d>
+	bool SampleObjectHelper2(const FImplicitObject& Object, const TRigidTransform<T, d>& ObjectTransform, const TRigidTransform<T, d>& SampleToObjectTransform, const TVector<T, d>& SampleParticle, T Thickness, TRigidBodyContactConstraint<T, d>& Constraint)
+	{
+		TVector<T, d> LocalPoint = SampleToObjectTransform.TransformPositionNoScale(SampleParticle);
+		TVector<T, d> LocalNormal;
+		T LocalPhi = Object.PhiWithNormal(LocalPoint, LocalNormal);
+
+		TContactData<T, d> & Contact = Constraint.ShapeManifold.Manifold;
+		if (LocalPhi < Contact.Phi)
+		{
+			Contact.Phi = LocalPhi;
+			Contact.Normal = ObjectTransform.TransformVectorNoScale(LocalNormal);
+			Contact.Location = ObjectTransform.TransformPositionNoScale(LocalPoint);
+			return true;
+		}
+		return false;
+	}
+
+	template <typename T, int d>
+	bool SampleObjectNoNormal2(const FImplicitObject& Object, const TRigidTransform<T, d>& ObjectTransform, const TRigidTransform<T, d>& SampleToObjectTransform, const TVector<T, d>& SampleParticle, T Thickness, TRigidBodyContactConstraint<T, d>& Constraint)
+	{
+		TVector<T, d> LocalPoint = SampleToObjectTransform.TransformPositionNoScale(SampleParticle);
+		TVector<T, d> LocalNormal;
+		T LocalPhi = Object.PhiWithNormal(LocalPoint, LocalNormal);
+
+		TContactData<T, d> & Contact = Constraint.ShapeManifold.Manifold;
+		if (LocalPhi < Contact.Phi)
+		{
+			Contact.Phi = LocalPhi;
+			return true;
+		}
+		return false;
+	}
+
+	template <typename T, int d>
+	bool SampleObjectNormalAverageHelper2(const FImplicitObject& Object, const TRigidTransform<T, d>& ObjectTransform, const TRigidTransform<T, d>& SampleToObjectTransform, const TVector<T, d>& SampleParticle, T Thickness, T& TotalThickness, TRigidBodyContactConstraint<T, d>& Constraint)
+	{
+		TVector<T, d> LocalPoint = SampleToObjectTransform.TransformPositionNoScale(SampleParticle);
+		TVector<T, d> LocalNormal;
+		T LocalPhi = Object.PhiWithNormal(LocalPoint, LocalNormal);
+		T LocalThickness = LocalPhi - Thickness;
+
+		TContactData<T, d> & Contact = Constraint.ShapeManifold.Manifold;
+		if (LocalThickness < -KINDA_SMALL_NUMBER)
+		{
+			Contact.Location += LocalPoint * LocalThickness;
+			TotalThickness += LocalThickness;
+			return true;
+		}
+		return false;
+	}
+
+	DECLARE_CYCLE_STAT(TEXT("TPBDCollisionConstraint::UpdateLevelsetPartial"), STAT_UpdateLevelsetPartial, STATGROUP_ChaosWide);
+	DECLARE_CYCLE_STAT(TEXT("TPBDCollisionConstraint::UpdateLevelsetFindParticles"), STAT_UpdateLevelsetFindParticles, STATGROUP_ChaosWide);
+	DECLARE_CYCLE_STAT(TEXT("TPBDCollisionConstraint::UpdateLevelsetBVHTraversal"), STAT_UpdateLevelsetBVHTraversal, STATGROUP_ChaosWide);
+	DECLARE_CYCLE_STAT(TEXT("TPBDCollisionConstraint::UpdateLevelsetSignedDistance"), STAT_UpdateLevelsetSignedDistance, STATGROUP_ChaosWide);
+	DECLARE_CYCLE_STAT(TEXT("TPBDCollisionConstraint::UpdateLevelsetAll"), STAT_UpdateLevelsetAll, STATGROUP_ChaosWide);
+	DECLARE_CYCLE_STAT(TEXT("TPBDCollisionConstraint::SampleObject"), STAT_SampleObject, STATGROUP_ChaosWide);
+
+	int32 NormalAveraging2 = 1;
+	FAutoConsoleVariableRef CVarNormalAveraging2(TEXT("p.NormalAveraging2"), NormalAveraging2, TEXT(""));
+
+	int32 SampleMinParticlesForAcceleration2 = 2048;
+	FAutoConsoleVariableRef CVarSampleMinParticlesForAcceleration2(TEXT("p.SampleMinParticlesForAcceleration2"), SampleMinParticlesForAcceleration2, TEXT("The minimum number of particles needed before using an acceleration structure when sampling"));
+
+
+	template <ECollisionUpdateType UpdateType, typename T, int d>
+	void SampleObject2(const FImplicitObject& Object, const TRigidTransform<T, d>& ObjectTransform, const TBVHParticles<T, d>& SampleParticles, const TRigidTransform<T, d>& SampleParticlesTransform, T Thickness, TRigidBodyContactConstraint<T, d>& Constraint)
+	{
+		SCOPE_CYCLE_COUNTER(STAT_SampleObject);
+		TRigidBodyContactConstraint<float, 3> AvgConstraint = Constraint;
+
+		TContactData<float, 3> & Contact = Constraint.ShapeManifold.Manifold;
+		TContactData<float, 3> & AvgContact = AvgConstraint.ShapeManifold.Manifold;
+
+		AvgConstraint.Particle[0] = Constraint.Particle[0];
+		AvgConstraint.Particle[1] = Constraint.Particle[1];
+		AvgContact.Location = TVector<float, 3>::ZeroVector;
+		AvgContact.Normal = TVector<float, 3>::ZeroVector;
+		AvgContact.Phi = Thickness;
+		float TotalThickness = float(0);
+
+		int32 DeepestParticle = -1;
+		const int32 NumParticles = SampleParticles.Size();
+
+		const TRigidTransform<T, d> & SampleToObjectTM = SampleParticlesTransform.GetRelativeTransform(ObjectTransform);
+		if (NumParticles > SampleMinParticlesForAcceleration2 && Object.HasBoundingBox())
+		{
+			SCOPE_CYCLE_COUNTER(STAT_UpdateLevelsetPartial);
+			TBox<T, d> ImplicitBox = Object.BoundingBox().TransformedBox(ObjectTransform.GetRelativeTransform(SampleParticlesTransform));
+			ImplicitBox.Thicken(Thickness);
+			TArray<int32> PotentialParticles;
+			{
+				SCOPE_CYCLE_COUNTER(STAT_UpdateLevelsetFindParticles);
+				PotentialParticles = SampleParticles.FindAllIntersections(ImplicitBox);
+			}
+			{
+				SCOPE_CYCLE_COUNTER(STAT_UpdateLevelsetSignedDistance);
+				for (int32 i : PotentialParticles)
+				{
+					if (NormalAveraging2 && UpdateType != ECollisionUpdateType::Any)	//if we just want one don't bother with normal
+					{
+						SampleObjectNormalAverageHelper2(Object, ObjectTransform, SampleToObjectTM, SampleParticles.X(i), Thickness, TotalThickness, AvgConstraint);
+					}
+					else
+					{
+						if (SampleObjectNoNormal2(Object, ObjectTransform, SampleToObjectTM, SampleParticles.X(i), Thickness, AvgConstraint))
+						{
+							DeepestParticle = i;
+							if (UpdateType == ECollisionUpdateType::Any)
+							{
+								Contact.Phi = AvgContact.Phi;
+								return;
+							}
+						}
+					}
+				}
+			}
+		}
+		else
+		{
+			SCOPE_CYCLE_COUNTER(STAT_UpdateLevelsetAll);
+			for (int32 i = 0; i < NumParticles; ++i)
+			{
+				if (NormalAveraging2 && UpdateType != ECollisionUpdateType::Any)	//if we just want one don't bother with normal
+				{
+					const bool bInside = SampleObjectNormalAverageHelper2(Object, ObjectTransform, SampleToObjectTM, SampleParticles.X(i), Thickness, TotalThickness, AvgConstraint);
+				}
+				else
+				{
+					if (SampleObjectNoNormal2(Object, ObjectTransform, SampleToObjectTM, SampleParticles.X(i), Thickness, AvgConstraint))
+					{
+						DeepestParticle = i;
+						if (UpdateType == ECollisionUpdateType::Any)
+						{
+							Contact.Phi = AvgContact.Phi;
+							return;
+						}
+					}
+				}
+			}
+		}
+
+		if (NormalAveraging2)
+		{
+			if (TotalThickness < -KINDA_SMALL_NUMBER)
+			{
+				TVector<T, d> LocalPoint = AvgContact.Location / TotalThickness;
+				TVector<T, d> LocalNormal;
+				const T NewPhi = Object.PhiWithNormal(LocalPoint, LocalNormal);
+				if (NewPhi < Contact.Phi)
+				{
+					Contact.Phi = NewPhi;
+					Contact.Location = ObjectTransform.TransformPositionNoScale(LocalPoint);
+					Contact.Normal = ObjectTransform.TransformVectorNoScale(LocalNormal);
+				}
+			}
+			else
+			{
+				check(AvgContact.Phi >= Thickness);
+			}
+		}
+		else if (AvgContact.Phi < Contact.Phi)
+		{
+			check(DeepestParticle >= 0);
+			TVector<T, d> LocalPoint = SampleToObjectTM.TransformPositionNoScale(SampleParticles.X(DeepestParticle));
+			TVector<T, d> LocalNormal;
+			Contact.Phi = Object.PhiWithNormal(LocalPoint, LocalNormal);
+			Contact.Location = ObjectTransform.TransformPositionNoScale(LocalPoint);
+			Contact.Normal = ObjectTransform.TransformVectorNoScale(LocalNormal);
+
+		}
+	}
+
+
+#if INTEL_ISPC
+	template<ECollisionUpdateType UpdateType>
+	void SampleObject2(const FImplicitObject& Object, const TRigidTransform<float, 3>& ObjectTransform, const TBVHParticles<float, 3>& SampleParticles, const TRigidTransform<float, 3>& SampleParticlesTransform, float Thickness, TRigidBodyContactConstraint<float, 3>& Constraint)
+	{
+		SCOPE_CYCLE_COUNTER(STAT_SampleObject);
+		TRigidBodyContactConstraint<float, 3> AvgConstraint = Constraint;
+
+		TContactData<float, 3> & Contact = Constraint.ShapeManifold.Manifold;
+		TContactData<float, 3> & AvgContact = AvgConstraint.ShapeManifold.Manifold;
+
+		AvgConstraint.Particle[0] = Constraint.Particle[0];
+		AvgConstraint.Particle[1] = Constraint.Particle[1];
+		AvgContact.Location = TVector<float, 3>::ZeroVector;
+		AvgContact.Normal = TVector<float, 3>::ZeroVector;
+		AvgContact.Phi = Thickness;
+		float TotalThickness = float(0);
+
+		int32 DeepestParticle = -1;
+
+		const TRigidTransform<float, 3>& SampleToObjectTM = SampleParticlesTransform.GetRelativeTransform(ObjectTransform);
+		int32 NumParticles = SampleParticles.Size();
+
+		if (NumParticles > SampleMinParticlesForAcceleration2 && Object.HasBoundingBox())
+		{
+			SCOPE_CYCLE_COUNTER(STAT_UpdateLevelsetPartial);
+			TBox<float, 3> ImplicitBox = Object.BoundingBox().TransformedBox(ObjectTransform.GetRelativeTransform(SampleParticlesTransform));
+			ImplicitBox.Thicken(Thickness);
+			TArray<int32> PotentialParticles;
+			{
+				SCOPE_CYCLE_COUNTER(STAT_UpdateLevelsetFindParticles);
+				PotentialParticles = SampleParticles.FindAllIntersections(ImplicitBox);
+			}
+			{
+				SCOPE_CYCLE_COUNTER(STAT_UpdateLevelsetSignedDistance);
+
+				if (Object.GetType(true) == ImplicitObjectType::LevelSet && PotentialParticles.Num() > 0)
+				{
+					//QUICK_SCOPE_CYCLE_COUNTER(STAT_LevelSet);
+					const TLevelSet<float, 3>* LevelSet = Object.GetObject<Chaos::TLevelSet<float, 3>>();
+					const TUniformGrid<float, 3>& Grid = LevelSet->GetGrid();
+
+					if (NormalAveraging2 && UpdateType != ECollisionUpdateType::Any)
+					{
+						ispc::SampleLevelSetNormalAverage(
+							(ispc::FVector&)Grid.MinCorner(),
+							(ispc::FVector&)Grid.MaxCorner(),
+							(ispc::FVector&)Grid.Dx(),
+							(ispc::FIntVector&)Grid.Counts(),
+							(ispc::TArrayND*)&LevelSet->GetPhiArray(),
+							(ispc::FTransform&)SampleToObjectTM,
+							(ispc::FVector*)&SampleParticles.XArray()[0],
+							&PotentialParticles[0],
+							Thickness,
+							TotalThickness,
+							(ispc::FVector&)AvgContact.Location,
+							PotentialParticles.Num());
+					}
+					else
+					{
+						ispc::SampleLevelSetNoNormal(
+							(ispc::FVector&)Grid.MinCorner(),
+							(ispc::FVector&)Grid.MaxCorner(),
+							(ispc::FVector&)Grid.Dx(),
+							(ispc::FIntVector&)Grid.Counts(),
+							(ispc::TArrayND*)&LevelSet->GetPhiArray(),
+							(ispc::FTransform&)SampleToObjectTM,
+							(ispc::FVector*)&SampleParticles.XArray()[0],
+							&PotentialParticles[0],
+							DeepestParticle,
+							AvgContact.Phi,
+							PotentialParticles.Num());
+
+						if (UpdateType == ECollisionUpdateType::Any)
+						{
+							Contact.Phi = AvgContact.Phi;
+							return;
+						}
+					}
+				}
+				else if (Object.GetType(true) == ImplicitObjectType::Box && PotentialParticles.Num() > 0)
+				{
+					//QUICK_SCOPE_CYCLE_COUNTER(STAT_Box);
+					const TBox<float, 3>* Box = Object.GetObject<Chaos::TBox<float, 3>>();
+
+					if (NormalAveraging2 && UpdateType != ECollisionUpdateType::Any)
+					{
+						ispc::SampleBoxNormalAverage(
+							(ispc::FVector&)Box->Min(),
+							(ispc::FVector&)Box->Max(),
+							(ispc::FTransform&)SampleToObjectTM,
+							(ispc::FVector*)&SampleParticles.XArray()[0],
+							&PotentialParticles[0],
+							Thickness,
+							TotalThickness,
+							(ispc::FVector&)AvgContact.Location,
+							PotentialParticles.Num());
+					}
+					else
+					{
+						ispc::SampleBoxNoNormal(
+							(ispc::FVector&)Box->Min(),
+							(ispc::FVector&)Box->Max(),
+							(ispc::FTransform&)SampleToObjectTM,
+							(ispc::FVector*)&SampleParticles.XArray()[0],
+							&PotentialParticles[0],
+							DeepestParticle,
+							AvgContact.Phi,
+							PotentialParticles.Num());
+
+						if (UpdateType == ECollisionUpdateType::Any)
+						{
+							Contact.Phi = AvgContact.Phi;
+							return;
+						}
+					}
+				}
+				else
+				{
+					//QUICK_SCOPE_CYCLE_COUNTER(STAT_Other);
+					for (int32 i : PotentialParticles)
+					{
+						if (NormalAveraging2 && UpdateType != ECollisionUpdateType::Any)
+						{
+							SampleObjectNormalAverageHelper2(Object, ObjectTransform, SampleToObjectTM, SampleParticles.X(i), Thickness, TotalThickness, AvgConstraint);
+						}
+						else
+						{
+							if (SampleObjectNoNormal2(Object, ObjectTransform, SampleToObjectTM, SampleParticles.X(i), Thickness, AvgConstraint))
+							{
+								DeepestParticle = i;
+								if (UpdateType == ECollisionUpdateType::Any)
+								{
+									Contact.Phi = AvgContact.Phi;
+									return;
+								}
+							}
+						}
+					}
+				}
+			}
+		}
+		else
+		{
+			SCOPE_CYCLE_COUNTER(STAT_UpdateLevelsetAll);
+			if (Object.GetType(true) == ImplicitObjectType::LevelSet && NumParticles > 0)
+			{
+				const TLevelSet<float, 3>* LevelSet = Object.GetObject<Chaos::TLevelSet<float, 3>>();
+				const TUniformGrid<float, 3>& Grid = LevelSet->GetGrid();
+
+				if (NormalAveraging2 && UpdateType != ECollisionUpdateType::Any)
+				{
+					ispc::SampleLevelSetNormalAverageAll(
+						(ispc::FVector&)Grid.MinCorner(),
+						(ispc::FVector&)Grid.MaxCorner(),
+						(ispc::FVector&)Grid.Dx(),
+						(ispc::FIntVector&)Grid.Counts(),
+						(ispc::TArrayND*)&LevelSet->GetPhiArray(),
+						(ispc::FTransform&)SampleToObjectTM,
+						(ispc::FVector*)&SampleParticles.XArray()[0],
+						Thickness,
+						TotalThickness,
+						(ispc::FVector&)AvgContact.Location,
+						NumParticles);
+				}
+				else
+				{
+					ispc::SampleLevelSetNoNormalAll(
+						(ispc::FVector&)Grid.MinCorner(),
+						(ispc::FVector&)Grid.MaxCorner(),
+						(ispc::FVector&)Grid.Dx(),
+						(ispc::FIntVector&)Grid.Counts(),
+						(ispc::TArrayND*)&LevelSet->GetPhiArray(),
+						(ispc::FTransform&)SampleToObjectTM,
+						(ispc::FVector*)&SampleParticles.XArray()[0],
+						DeepestParticle,
+						AvgContact.Phi,
+						NumParticles);
+
+					if (UpdateType == ECollisionUpdateType::Any)
+					{
+						Contact.Phi = AvgContact.Phi;
+						return;
+					}
+				}
+			}
+			else if (Object.GetType(true) == ImplicitObjectType::Plane && NumParticles > 0)
+			{
+				const TPlane<float, 3>* Plane = Object.GetObject<Chaos::TPlane<float, 3>>();
+
+				if (NormalAveraging2 && UpdateType != ECollisionUpdateType::Any)
+				{
+					ispc::SamplePlaneNormalAverageAll(
+						(ispc::FVector&)Plane->Normal(),
+						(ispc::FVector&)Plane->X(),
+						(ispc::FTransform&)SampleToObjectTM,
+						(ispc::FVector*)&SampleParticles.XArray()[0],
+						Thickness,
+						TotalThickness,
+						(ispc::FVector&)AvgContact.Location,
+						NumParticles);
+				}
+				else
+				{
+					ispc::SamplePlaneNoNormalAll(
+						(ispc::FVector&)Plane->Normal(),
+						(ispc::FVector&)Plane->X(),
+						(ispc::FTransform&)SampleToObjectTM,
+						(ispc::FVector*)&SampleParticles.XArray()[0],
+						DeepestParticle,
+						AvgContact.Phi,
+						NumParticles);
+
+					if (UpdateType == ECollisionUpdateType::Any)
+					{
+						Contact.Phi = AvgContact.Phi;
+						return;
+					}
+				}
+			}
+			else if (Object.GetType(true) == ImplicitObjectType::Box && NumParticles > 0)
+			{
+				const TBox<float, 3>* Box = Object.GetObject<Chaos::TBox<float, 3>>();
+
+				if (NormalAveraging2 && UpdateType != ECollisionUpdateType::Any)
+				{
+					ispc::SampleBoxNormalAverageAll(
+						(ispc::FVector&)Box->Min(),
+						(ispc::FVector&)Box->Max(),
+						(ispc::FTransform&)SampleToObjectTM,
+						(ispc::FVector*)&SampleParticles.XArray()[0],
+						Thickness,
+						TotalThickness,
+						(ispc::FVector&)AvgContact.Location,
+						NumParticles);
+				}
+				else
+				{
+					ispc::SampleBoxNoNormalAll(
+						(ispc::FVector&)Box->Min(),
+						(ispc::FVector&)Box->Max(),
+						(ispc::FTransform&)SampleToObjectTM,
+						(ispc::FVector*)&SampleParticles.XArray()[0],
+						DeepestParticle,
+						AvgContact.Phi,
+						NumParticles);
+
+					if (UpdateType == ECollisionUpdateType::Any)
+					{
+						Contact.Phi = AvgContact.Phi;
+						return;
+					}
+				}
+			}
+			else
+			{
+				//QUICK_SCOPE_CYCLE_COUNTER(STAT_Other);
+				for (int32 i = 0; i < NumParticles; ++i)
+				{
+					if (NormalAveraging2 && UpdateType != ECollisionUpdateType::Any)
+					{
+						SampleObjectNormalAverageHelper2(Object, ObjectTransform, SampleToObjectTM, SampleParticles.X(i), Thickness, TotalThickness, AvgConstraint);
+					}
+					else
+					{
+						if (SampleObjectNoNormal2(Object, ObjectTransform, SampleToObjectTM, SampleParticles.X(i), Thickness, AvgConstraint))
+						{
+							DeepestParticle = i;
+							if (UpdateType == ECollisionUpdateType::Any)
+							{
+								Contact.Phi = AvgContact.Phi;
+								return;
+							}
+						}
+					}
+				}
+			}
+		}
+
+		if (NormalAveraging2)
+		{
+			if (TotalThickness < -KINDA_SMALL_NUMBER)
+			{
+				TVector<float, 3> LocalPoint = AvgContact.Location / TotalThickness;
+				TVector<float, 3> LocalNormal;
+				const float NewPhi = Object.PhiWithNormal(LocalPoint, LocalNormal);
+				if (NewPhi < Contact.Phi)
+				{
+					Contact.Phi = NewPhi;
+					Contact.Location = ObjectTransform.TransformPositionNoScale(LocalPoint);
+					Contact.Normal = ObjectTransform.TransformVectorNoScale(LocalNormal);
+				}
+			}
+			else
+			{
+				check(AvgContact.Phi >= Thickness);
+			}
+		}
+		else if (AvgContact.Phi < Contact.Phi)
+		{
+			check(DeepestParticle >= 0);
+			TVector<float, 3> LocalPoint = SampleToObjectTM.TransformPositionNoScale(SampleParticles.X(DeepestParticle));
+			TVector<float, 3> LocalNormal;
+			Contact.Phi = Object.PhiWithNormal(LocalPoint, LocalNormal);
+			Contact.Location = ObjectTransform.TransformPositionNoScale(LocalPoint);
+			Contact.Normal = ObjectTransform.TransformVectorNoScale(LocalNormal);
+		}
+	}
+#endif
+
+	template<ECollisionUpdateType UpdateType, typename T, int d>
+	void UpdateConstraintImp(const FImplicitObject& ParticleObject, const TRigidTransform<T, d>& ParticleTM, const FImplicitObject& LevelsetObject, const TRigidTransform<T, d>& LevelsetTM, const T Thickness, TRigidBodyContactConstraint<T, d>& Constraint)
+	{
+		if (ParticleObject.GetType() == TBox<T, d>::StaticType() && LevelsetObject.GetType() == TBox<T, d>::StaticType())
+		{
+			UpdateBoxConstraint(*ParticleObject.template GetObject<TBox<T, d>>(), ParticleTM, *LevelsetObject.template GetObject<TBox<T, d>>(), LevelsetTM, Thickness, Constraint);
+		}
+		else if (ParticleObject.GetType() == TSphere<T, d>::StaticType() && LevelsetObject.GetType() == TSphere<T, d>::StaticType())
+		{
+			UpdateSphereConstraint(*ParticleObject.template GetObject<TSphere<T, d>>(), ParticleTM, *LevelsetObject.template GetObject<TSphere<T, d>>(), LevelsetTM, Thickness, Constraint);
+		}
+		else if (ParticleObject.GetType() == TBox<T, d>::StaticType() && LevelsetObject.GetType() == TPlane<T, d>::StaticType())
+		{
+			UpdateBoxPlaneConstraint(*ParticleObject.template GetObject<TBox<T, d>>(), ParticleTM, *LevelsetObject.template GetObject<TPlane<T, d>>(), LevelsetTM, Thickness, Constraint);
+		}
+		else if (ParticleObject.GetType() == TSphere<T, d>::StaticType() && LevelsetObject.GetType() == TPlane<T, d>::StaticType())
+		{
+			UpdateSpherePlaneConstraint(*ParticleObject.template GetObject<TSphere<T, d>>(), ParticleTM, *LevelsetObject.template GetObject<TPlane<T, d>>(), LevelsetTM, Thickness, Constraint);
+		}
+		else if (ParticleObject.GetType() == TSphere<T, d>::StaticType() && LevelsetObject.GetType() == TBox<T, d>::StaticType())
+		{
+			UpdateSphereBoxConstraint(*ParticleObject.template GetObject<TSphere<T, d>>(), ParticleTM, *LevelsetObject.template GetObject<TBox<T, d>>(), LevelsetTM, Thickness, Constraint);
+		}
+		else if (ParticleObject.GetType() == TCapsule<T>::StaticType() && LevelsetObject.GetType() == TCapsule<T>::StaticType())
+		{
+			UpdateCapsuleCapsuleConstraint(*ParticleObject.template GetObject<TCapsule<T>>(), ParticleTM, *LevelsetObject.template GetObject<TCapsule<T>>(), LevelsetTM, Thickness, Constraint);
+		}
+		else if (ParticleObject.GetType() == TCapsule<T>::StaticType() && LevelsetObject.GetType() == TBox<T, d>::StaticType())
+		{
+			UpdateCapsuleBoxConstraint(*ParticleObject.template GetObject<TCapsule<T>>(), ParticleTM, *LevelsetObject.template GetObject<TBox<T, d>>(), LevelsetTM, Thickness, Constraint);
+		}
+		else if (ParticleObject.GetType() == TPlane<T, d>::StaticType() && LevelsetObject.GetType() == TBox<T, d>::StaticType())
+		{
+			TRigidBodyContactConstraint<T, d> TmpConstraint = Constraint;
+			UpdateBoxPlaneConstraint(*LevelsetObject.template GetObject<TBox<T, d>>(), LevelsetTM, *ParticleObject.template GetObject<TPlane<T, d>>(), ParticleTM, Thickness, TmpConstraint);
+			if (TmpConstraint.GetPhi() < Constraint.GetPhi())
+			{
+				Constraint = TmpConstraint;
+				Constraint.SetNormal(-Constraint.GetNormal());
+			}
+		}
+		else if (ParticleObject.GetType() == TPlane<T, d>::StaticType() && LevelsetObject.GetType() == TSphere<T, d>::StaticType())
+		{
+			TRigidBodyContactConstraint<T, d> TmpConstraint = Constraint;
+			UpdateSpherePlaneConstraint(*LevelsetObject.template GetObject<TSphere<T, d>>(), LevelsetTM, *ParticleObject.template GetObject<TPlane<T, d>>(), ParticleTM, Thickness, TmpConstraint);
+			if (TmpConstraint.GetPhi() < Constraint.GetPhi())
+			{
+				Constraint = TmpConstraint;
+				Constraint.SetNormal(-Constraint.GetNormal());
+			}
+		}
+		else if (ParticleObject.GetType() == TBox<T, d>::StaticType() && LevelsetObject.GetType() == TSphere<T, d>::StaticType())
+		{
+			TRigidBodyContactConstraint<T, d> TmpConstraint = Constraint;
+			UpdateSphereBoxConstraint(*LevelsetObject.template GetObject<TSphere<T, d>>(), LevelsetTM, *ParticleObject.template GetObject<TBox<T, d>>(), ParticleTM, Thickness, TmpConstraint);
+			if (TmpConstraint.GetPhi() < Constraint.GetPhi())
+			{
+				Constraint = TmpConstraint;
+				Constraint.SetNormal(-Constraint.GetNormal());
+			}
+		}
+		else if (ParticleObject.GetType() == TBox<T, d>::StaticType() && LevelsetObject.GetType() == TCapsule<T>::StaticType())
+		{
+			TRigidBodyContactConstraint<T, d> TmpConstraint = Constraint;
+			UpdateCapsuleBoxConstraint(*LevelsetObject.template GetObject<TCapsule<T>>(), LevelsetTM, *ParticleObject.template GetObject<TBox<T, d>>(), ParticleTM, Thickness, TmpConstraint);
+			if (TmpConstraint.GetPhi() < Constraint.GetPhi())
+			{
+				Constraint = TmpConstraint;
+				Constraint.SetNormal(-Constraint.GetNormal());
+			}
+		}
+		else if (ParticleObject.GetType() < TImplicitObjectUnion<T, d>::StaticType() && LevelsetObject.GetType() == TImplicitObjectUnion<T, d>::StaticType())
+		{
+			return UpdateSingleUnionConstraint<UpdateType>(Thickness, Constraint);
+		}
+		else if (ParticleObject.GetType() == TImplicitObjectUnion<T, d>::StaticType() && LevelsetObject.GetType() < TImplicitObjectUnion<T, d>::StaticType())
+		{
+			check(false);	//should not be possible to get this ordering (see ComputeConstraint)
+		}
+		else if (ParticleObject.GetType() == TImplicitObjectUnion<T, d>::StaticType() && LevelsetObject.GetType() == TImplicitObjectUnion<T, d>::StaticType())
+		{
+			return UpdateUnionUnionConstraint<UpdateType>(Thickness, Constraint);
+		}
+		else if (ParticleObject.IsConvex() && LevelsetObject.IsConvex())
+		{
+			CollisionResolutionConvexConvex<T, d>::UpdateConvexConvexConstraint(ParticleObject, ParticleTM, LevelsetObject, LevelsetTM, Thickness, Constraint);
+		}
+		else if (LevelsetObject.IsUnderlyingUnion())
+		{
+			UpdateUnionLevelsetConstraint<UpdateType>(Thickness, Constraint);
+		}
+		else if (ParticleObject.IsUnderlyingUnion())
+		{
+			UpdateLevelsetUnionConstraint<UpdateType>(Thickness, Constraint);
+		}
+		else
+		{
+			UpdateLevelsetConstraint<UpdateType>(Thickness, Constraint);
+		}
+	}
+
+	template<typename T, int d>
+	void ConstructConstraintsImpl(TGeometryParticleHandle<T, d>* Particle0, TGeometryParticleHandle<T, d>* Particle1, const FImplicitObject* Implicit0, const FImplicitObject* Implicit1, const T Thickness, TRigidBodyContactConstraint<T, d> & Constraint)
+	{
+		// TriangleMesh implicits are for scene query only.
+		if (Implicit0 && GetInnerType(Implicit0->GetType()) == ImplicitObjectType::TriangleMesh) return;
+		if (Implicit1 && GetInnerType(Implicit1->GetType()) == ImplicitObjectType::TriangleMesh) return;
+
+		if (Implicit0 && Implicit0->GetType() == TImplicitObjectUnion<T, d>::StaticType() &&
+			Implicit1 && Implicit1->GetType() == TImplicitObjectUnion<T, d>::StaticType())
+		{
+			ConstructUnionUnionConstraints(Particle0, Particle1, Implicit0, Implicit1, Thickness, Constraint);
+		}
+		else
+		{
+			ConstructPairConstraintImpl(Particle0, Particle1, Implicit0, Implicit1, Thickness, Constraint);
+		}
+	}
+
+	template <typename T, int d>
+	bool UpdateBoxPlaneConstraint(const TBox<T, d>& Box, const TRigidTransform<T, d>& BoxTransform, const TPlane<T, d>& Plane, const TRigidTransform<T, d>& PlaneTransform, T Thickness, TRigidBodyContactConstraint<T, d>& Constraint)
+	{
+		TContactData<T, d> & Contact = Constraint.ShapeManifold.Manifold;
+
+#if USING_CODE_ANALYSIS
+		MSVC_PRAGMA(warning(push))
+			MSVC_PRAGMA(warning(disable : ALL_CODE_ANALYSIS_WARNINGS))
+#endif	// USING_CODE_ANALYSIS
+
+			bool bApplied = false;
+		const TRigidTransform<T, d> BoxToPlaneTransform(BoxTransform.GetRelativeTransform(PlaneTransform));
+		const TVector<T, d> Extents = Box.Extents();
+		constexpr int32 NumCorners = 2 + 2 * d;
+		constexpr T Epsilon = KINDA_SMALL_NUMBER;
+
+		TVector<T, d> Corners[NumCorners];
+		int32 CornerIdx = 0;
+		Corners[CornerIdx++] = BoxToPlaneTransform.TransformPosition(Box.Max());
+		Corners[CornerIdx++] = BoxToPlaneTransform.TransformPosition(Box.Min());
+		for (int32 j = 0; j < d; ++j)
+		{
+			Corners[CornerIdx++] = BoxToPlaneTransform.TransformPosition(Box.Min() + TVector<T, d>::AxisVector(j) * Extents);
+			Corners[CornerIdx++] = BoxToPlaneTransform.TransformPosition(Box.Max() - TVector<T, d>::AxisVector(j) * Extents);
+		}
+
+#if USING_CODE_ANALYSIS
+		MSVC_PRAGMA(warning(pop))
+#endif	// USING_CODE_ANALYSIS
+
+			TVector<T, d> PotentialConstraints[NumCorners];
+		int32 NumConstraints = 0;
+		for (int32 i = 0; i < NumCorners; ++i)
+		{
+			TVector<T, d> Normal;
+			const T NewPhi = Plane.PhiWithNormal(Corners[i], Normal);
+			if (NewPhi < Contact.Phi + Epsilon)
+			{
+				if (NewPhi <= Contact.Phi - Epsilon)
+				{
+					NumConstraints = 0;
+				}
+				Contact.Phi = NewPhi;
+				Contact.Normal = PlaneTransform.TransformVector(Normal);
+				Contact.Location = PlaneTransform.TransformPosition(Corners[i]);
+				PotentialConstraints[NumConstraints++] = Contact.Location;
+				bApplied = true;
+			}
+		}
+		if (NumConstraints > 1)
+		{
+			TVector<T, d> AverageLocation(0);
+			for (int32 ConstraintIdx = 0; ConstraintIdx < NumConstraints; ++ConstraintIdx)
+			{
+				AverageLocation += PotentialConstraints[ConstraintIdx];
+			}
+			Contact.Location = AverageLocation / NumConstraints;
+		}
+
+		return bApplied;
+	}
+
+	template <typename T, int d>
+	void UpdateSphereConstraint(const TSphere<T, d>& Sphere1, const TRigidTransform<T, d>& Sphere1Transform, const TSphere<T, d>& Sphere2, const TRigidTransform<T, d>& Sphere2Transform, T Thickness, TRigidBodyContactConstraint<T, d>& Constraint)
+	{
+		TContactData<T, d> & Contact = Constraint.ShapeManifold.Manifold;
+
+		const TVector<T, d> Center1 = Sphere1Transform.TransformPosition(Sphere1.GetCenter());
+		const TVector<T, d> Center2 = Sphere2Transform.TransformPosition(Sphere2.GetCenter());
+		const TVector<T, d> Direction = Center1 - Center2;
+		const T Size = Direction.Size();
+		const T NewPhi = Size - (Sphere1.GetRadius() + Sphere2.GetRadius());
+		if (NewPhi < Contact.Phi)
+		{
+			Contact.Normal = Size > SMALL_NUMBER ? Direction / Size : TVector<T, d>(0, 0, 1);
+			Contact.Phi = NewPhi;
+			Contact.Location = Center1 - Sphere1.GetRadius() * Contact.Normal;
+		}
+	}
+
+	template <typename T, int d>
+	void UpdateSpherePlaneConstraint(const TSphere<T, d>& Sphere, const TRigidTransform<T, d>& SphereTransform, const TPlane<T, d>& Plane, const TRigidTransform<T, d>& PlaneTransform, T Thickness, TRigidBodyContactConstraint<T, d>& Constraint)
+	{
+		TContactData<T, d> & Contact = Constraint.ShapeManifold.Manifold;
+
+		const TRigidTransform<T, d> SphereToPlaneTransform(PlaneTransform.Inverse() * SphereTransform);
+		const TVector<T, d> SphereCenter = SphereToPlaneTransform.TransformPosition(Sphere.GetCenter());
+
+		TVector<T, d> NewNormal;
+		T NewPhi = Plane.PhiWithNormal(SphereCenter, NewNormal);
+		NewPhi -= Sphere.GetRadius();
+
+		if (NewPhi < Contact.Phi)
+		{
+			Contact.Phi = NewPhi;
+			Contact.Normal = PlaneTransform.TransformVectorNoScale(NewNormal);
+			Contact.Location = SphereCenter - Contact.Normal * Sphere.GetRadius();
+		}
+	}
+
+	template <typename T, int d>
+	void UpdateSphereBoxConstraint(const TSphere<T, d>& Sphere, const TRigidTransform<T, d>& SphereTransform, const TBox<T, d>& Box, const TRigidTransform<T, d>& BoxTransform, T Thickness, TRigidBodyContactConstraint<T, d>& Constraint)
+	{
+		TContactData<T, d> & Contact = Constraint.ShapeManifold.Manifold;
+
+		const TRigidTransform<T, d> SphereToBoxTransform(SphereTransform * BoxTransform.Inverse());
+		const TVector<T, d> SphereCenterInBox = SphereToBoxTransform.TransformPosition(Sphere.GetCenter());
+
+		TVector<T, d> NewNormal;
+		T NewPhi = Box.PhiWithNormal(SphereCenterInBox, NewNormal);
+		NewPhi -= Sphere.GetRadius();
+
+		if (NewPhi < Contact.Phi)
+		{
+			Contact.Phi = NewPhi;
+			Contact.Normal = BoxTransform.TransformVectorNoScale(NewNormal);
+			Contact.Location = SphereTransform.TransformPosition(Sphere.GetCenter()) - Contact.Normal * Sphere.GetRadius();
+		}
+	}
+
+	template <typename T, int d>
+	void UpdateCapsuleCapsuleConstraint(const TCapsule<T>& A, const TRigidTransform<T, d>& ATransform, const TCapsule<T>& B, const TRigidTransform<T, d>& BTransform, T Thickness, TRigidBodyContactConstraint<T, d>& Constraint)
+	{
+		TContactData<T, d> & Contact = Constraint.ShapeManifold.Manifold;
+
+		FVector A1 = ATransform.TransformPosition(A.GetX1());
+		FVector A2 = ATransform.TransformPosition(A.GetX2());
+		FVector B1 = BTransform.TransformPosition(B.GetX1());
+		FVector B2 = BTransform.TransformPosition(B.GetX2());
+		FVector P1, P2;
+		FMath::SegmentDistToSegmentSafe(A1, A2, B1, B2, P1, P2);
+
+		TVector<T, d> Delta = P2 - P1;
+		T DeltaLen = Delta.Size();
+		if (DeltaLen > KINDA_SMALL_NUMBER)
+		{
+			T NewPhi = DeltaLen - (A.GetRadius() + B.GetRadius());
+			if (NewPhi < Contact.Phi)
+			{
+				TVector<T, d> Dir = Delta / DeltaLen;
+				Contact.Phi = NewPhi;
+				Contact.Normal = -Dir;
+				Contact.Location = P1 + Dir * A.GetRadius();
+			}
+		}
+	}
+
+	template <typename T, int d>
+	void UpdateCapsuleBoxConstraint(const TCapsule<T>& A, const TRigidTransform<T, d>& ATransform, const TBox<T, d>& B, const TRigidTransform<T, d>& BTransform, T Thickness, TRigidBodyContactConstraint<T, d>& Constraint)
+	{
+		// @todo(ccaulfield): Add custom capsule-box collision
+		TContactData<T, d> & Contact = Constraint.ShapeManifold.Manifold;
+
+		TRigidTransform<T, d> BToATransform = BTransform.GetRelativeTransform(ATransform);
+
+		// Use GJK to track closest points (not strictly necessary yet)
+		TVector<T, d> NearPointALocal, NearPointBLocal;
+		T NearPointDistance;
+		if (GJKDistance<T>(A, B, BToATransform, NearPointDistance, NearPointALocal, NearPointBLocal))
+		{
+			TVector<T, d> NearPointAWorld = ATransform.TransformPosition(NearPointALocal);
+			TVector<T, d> NearPointBWorld = BTransform.TransformPosition(NearPointBLocal);
+			TVector<T, d> NearPointBtoAWorld = NearPointAWorld - NearPointBWorld;
+			Contact.Phi = NearPointDistance;
+			Contact.Normal = NearPointBtoAWorld.GetSafeNormal();
+			Contact.Location = NearPointAWorld;
+		}
+		else
+		{
+			// Use box particle samples against the implicit capsule
+			const TArray<TVector<T, d>> BParticles = B.ComputeSamplePoints();
+			const int32 NumParticles = BParticles.Num();
+			for (int32 ParticleIndex = 0; ParticleIndex < NumParticles; ++ParticleIndex)
+			{
+				if (SampleObjectHelper2(A, ATransform, BToATransform, BParticles[ParticleIndex], Thickness, Constraint))
+				{
+					// SampleObjectHelper2 expects A to be the box, so reverse the results
+					Contact.Normal = -Contact.Normal;
+				}
+			}
+		}
+
+	}
+
+	DECLARE_CYCLE_STAT(TEXT("TPBDCollisionConstraint::FindRelevantShapes"), STAT_FindRelevantShapes, STATGROUP_ChaosWide);
+	template <typename T, int d>
+	TArray<Pair<const FImplicitObject*, TRigidTransform<T, d>>> FindRelevantShapes2(const FImplicitObject* ParticleObj, const TRigidTransform<T, d>& ParticlesTM, const FImplicitObject& LevelsetObj, const TRigidTransform<T, d>& LevelsetTM, const T Thickness)
+	{
+		SCOPE_CYCLE_COUNTER(STAT_FindRelevantShapes);
+		TArray<Pair<const FImplicitObject*, TRigidTransform<T, d>>> RelevantShapes;
+		//find all levelset inner objects
+		if (ParticleObj)
+		{
+			if (ParticleObj->HasBoundingBox())
+			{
+				const TRigidTransform<T, d> ParticlesToLevelsetTM = ParticlesTM.GetRelativeTransform(LevelsetTM);
+				TBox<T, d> ParticleBoundsInLevelset = ParticleObj->BoundingBox().TransformedBox(ParticlesToLevelsetTM);
+				ParticleBoundsInLevelset.Thicken(Thickness);
+				{
+					LevelsetObj.FindAllIntersectingObjects(RelevantShapes, ParticleBoundsInLevelset);
+				}
+			}
+			else
+			{
+				LevelsetObj.AccumulateAllImplicitObjects(RelevantShapes, TRigidTransform<T, d>::Identity);
+			}
+		}
+		else
+		{
+			//todo:compute bounds
+			LevelsetObj.AccumulateAllImplicitObjects(RelevantShapes, TRigidTransform<T, d>::Identity);
+		}
+
+		return RelevantShapes;
+	}
+
+	DECLARE_CYCLE_STAT(TEXT("TPBDCollisionConstraint::UpdateUnionUnionConstraint"), STAT_UpdateUnionUnionConstraint, STATGROUP_ChaosWide);
+	template<ECollisionUpdateType UpdateType, typename T, int d>
+	void UpdateUnionUnionConstraint(const T Thickness, TRigidBodyContactConstraint<T, d>& Constraint)
+	{
+		SCOPE_CYCLE_COUNTER(STAT_UpdateUnionUnionConstraint);
+
+		TGenericParticleHandle<T, d> Particle0 = Constraint.Particle[0];
+		TGenericParticleHandle<T, d> Particle1 = Constraint.Particle[1];
+
+		TRigidTransform<T, d> ParticlesTM = TRigidTransform<T, d>(Particle0->P(), Particle0->Q());
+		TRigidTransform<T, d> LevelsetTM = TRigidTransform<T, d>(Particle1->P(), Particle1->Q());
+
+		const FImplicitObject* ParticleObj = Particle0->Geometry().Get();
+		const FImplicitObject* LevelsetObj = Particle1->Geometry().Get();
+		const TArray<Pair<const FImplicitObject*, TRigidTransform<T, d>>> LevelsetShapes = FindRelevantShapes2(ParticleObj, ParticlesTM, *LevelsetObj, LevelsetTM, Thickness);
+
+		for (const Pair<const FImplicitObject*, TRigidTransform<T, d>>& LevelsetObjPair : LevelsetShapes)
+		{
+			const FImplicitObject& LevelsetInnerObj = *LevelsetObjPair.First;
+			const TRigidTransform<T, d>& LevelsetInnerObjTM = LevelsetObjPair.Second * LevelsetTM;
+
+			//now find all particle inner objects
+			const TArray<Pair<const FImplicitObject*, TRigidTransform<T, d>>> ParticleShapes = FindRelevantShapes2(&LevelsetInnerObj, LevelsetInnerObjTM, *ParticleObj, ParticlesTM, Thickness);
+
+			//for each inner obj pair, update constraint
+			for (const Pair<const FImplicitObject*, TRigidTransform<T, d>>& ParticlePair : ParticleShapes)
+			{
+				const FImplicitObject& ParticleInnerObj = *ParticlePair.First;
+				const TRigidTransform<T, d> ParticleInnerObjTM = ParticlePair.Second * ParticlesTM;
+				UpdateConstraintImp<UpdateType>(ParticleInnerObj, ParticleInnerObjTM, LevelsetInnerObj, LevelsetInnerObjTM, Thickness, Constraint);
+			}
+		}
+	}
+
+	DECLARE_CYCLE_STAT(TEXT("TPBDCollisionConstraint::UpdateSingleUnionConstraint"), STAT_UpdateSingleUnionConstraint, STATGROUP_ChaosWide);
+	template<ECollisionUpdateType UpdateType, typename T, int d>
+	void UpdateSingleUnionConstraint(const T Thickness, TRigidBodyContactConstraint<T, d>& Constraint)
+	{
+		SCOPE_CYCLE_COUNTER(STAT_UpdateSingleUnionConstraint);
+
+		TGenericParticleHandle<T, d> Particle0 = Constraint.Particle[0];
+		TGenericParticleHandle<T, d> Particle1 = Constraint.Particle[1];
+
+		TRigidTransform<T, d> ParticlesTM = TRigidTransform<T, d>(Particle0->P(), Particle0->Q());
+		TRigidTransform<T, d> LevelsetTM = TRigidTransform<T, d>(Particle1->P(), Particle1->Q());
+
+		const FImplicitObject* ParticleObj = Particle0->Geometry().Get();
+		const FImplicitObject* LevelsetObj = Particle1->Geometry().Get();
+		const TArray<Pair<const FImplicitObject*, TRigidTransform<T, d>>> LevelsetShapes = FindRelevantShapes2(ParticleObj, ParticlesTM, *LevelsetObj, LevelsetTM, Thickness);
+
+		for (const Pair<const FImplicitObject*, TRigidTransform<T, d>>& LevelsetObjPair : LevelsetShapes)
+		{
+			const FImplicitObject& LevelsetInnerObj = *LevelsetObjPair.First;
+			const TRigidTransform<T, d> LevelsetInnerObjTM = LevelsetObjPair.Second * LevelsetTM;
+			UpdateConstraintImp<UpdateType>(*ParticleObj, ParticlesTM, LevelsetInnerObj, LevelsetInnerObjTM, Thickness, Constraint);
+		}
+	}
+
+
+
+	DECLARE_CYCLE_STAT(TEXT("TPBDCollisionConstraint::UpdateUnionLevelsetConstraint"), STAT_UpdateUnionLevelsetConstraint, STATGROUP_ChaosWide);
+	template<ECollisionUpdateType UpdateType, typename T, int d>
+	void UpdateUnionLevelsetConstraint(const T Thickness, TRigidBodyContactConstraint<T, d>& Constraint)
+	{
+		SCOPE_CYCLE_COUNTER(STAT_UpdateUnionLevelsetConstraint);
+
+		TGenericParticleHandle<T, d> Particle0 = Constraint.Particle[0];
+		TGenericParticleHandle<T, d> Particle1 = Constraint.Particle[1];
+
+		TRigidTransform<T, d> ParticlesTM = TRigidTransform<T, d>(Particle0->P(), Particle0->Q());
+		TRigidTransform<T, d> LevelsetTM = TRigidTransform<T, d>(Particle1->P(), Particle1->Q());
+
+		if (!(ensure(!FMath::IsNaN(ParticlesTM.GetTranslation().X)) && ensure(!FMath::IsNaN(ParticlesTM.GetTranslation().Y)) && ensure(!FMath::IsNaN(ParticlesTM.GetTranslation().Z))))
+		{
+			return;
+		}
+
+		if (!(ensure(!FMath::IsNaN(LevelsetTM.GetTranslation().X)) && ensure(!FMath::IsNaN(LevelsetTM.GetTranslation().Y)) && ensure(!FMath::IsNaN(LevelsetTM.GetTranslation().Z))))
+		{
+			return;
+		}
+
+		const FImplicitObject* ParticleObj = Particle0->Geometry().Get();
+		const FImplicitObject* LevelsetObj = Particle1->Geometry().Get();
+		TArray<Pair<const FImplicitObject*, TRigidTransform<T, d>>> LevelsetShapes = FindRelevantShapes2(ParticleObj, ParticlesTM, *LevelsetObj, LevelsetTM, Thickness);
+
+		if (LevelsetShapes.Num() && Particle0->CollisionParticles().Get())
+		{
+			const TBVHParticles<T, d>& SampleParticles = *Particle0->CollisionParticles().Get();
+			if (SampleParticles.Size())
+			{
+				for (const Pair<const FImplicitObject*, TRigidTransform<T, d>>& LevelsetObjPair : LevelsetShapes)
+				{
+					const FImplicitObject* Object = LevelsetObjPair.First;
+					const TRigidTransform<T, d> ObjectTM = LevelsetObjPair.Second * LevelsetTM;
+					SampleObject2<UpdateType>(*Object, ObjectTM, SampleParticles, ParticlesTM, Thickness, Constraint);
+					if (UpdateType == ECollisionUpdateType::Any && Constraint.GetPhi() < Thickness)
+					{
+						return;
+					}
+				}
+			}
+#if CHAOS_PARTICLEHANDLE_TODO
+			else if (ParticleObj && ParticleObj->IsUnderlyingUnion())
+			{
+				const TImplicitObjectUnion<T, d>* UnionObj = static_cast<const TImplicitObjectUnion<T, d>*>(ParticleObj);
+				//need to traverse shapes to get their collision particles
+				for (const Pair<const FImplicitObject*, TRigidTransform<T, d>>& LevelsetObjPair : LevelsetShapes)
+				{
+					const FImplicitObject* LevelsetInnerObject = LevelsetObjPair.First;
+					const TRigidTransform<T, d> LevelsetInnerObjectTM = LevelsetObjPair.Second * LevelsetTM;
+
+					TArray<Pair<const FImplicitObject*, TRigidTransform<T, d>>> ParticleShapes = FindRelevantShapes2(LevelsetInnerObject, LevelsetInnerObjectTM, *ParticleObj, ParticlesTM, Thickness);
+					for (const Pair<const FImplicitObject*, TRigidTransform<T, d>>& ParticleObjPair : ParticleShapes)
+					{
+						const FImplicitObject* ParticleInnerObject = ParticleObjPair.First;
+						const TRigidTransform<T, d> ParticleInnerObjectTM = ParticleObjPair.Second * ParticlesTM;
+
+						if (const int32* OriginalIdx = UnionObj->MCollisionParticleLookupHack.Find(ParticleInnerObject))
+						{
+							const TBVHParticles<T, d>& InnerSampleParticles = *InParticles.CollisionParticles(*OriginalIdx).Get();
+							SampleObject2<UpdateType>(*LevelsetInnerObject, LevelsetInnerObjectTM, InnerSampleParticles, ParticleInnerObjectTM, Thickness, Constraint);
+							if (UpdateType == ECollisionUpdateType::Any && Constraint.Phi < Thickness)
+							{
+								return;
+							}
+						}
+					}
+
+				}
+			}
+#endif
+		}
+	}
+
+
+	DECLARE_CYCLE_STAT(TEXT("TPBDCollisionConstraint::UpdateLevelsetUnionConstraint"), STAT_UpdateLevelsetUnionConstraint, STATGROUP_ChaosWide);
+	template<ECollisionUpdateType UpdateType, typename T, int d>
+	void UpdateLevelsetUnionConstraint(const T Thickness, TRigidBodyContactConstraint<T, d>& Constraint)
+	{
+		SCOPE_CYCLE_COUNTER(STAT_UpdateLevelsetUnionConstraint);
+
+		TGenericParticleHandle<T, d> Particle0 = Constraint.Particle[0];
+		TGenericParticleHandle<T, d> Particle1 = Constraint.Particle[1];
+
+		TRigidTransform<T, d> ParticlesTM = TRigidTransform<T, d>(Particle0->P(), Particle0->Q());
+		TRigidTransform<T, d> LevelsetTM = TRigidTransform<T, d>(Particle1->P(), Particle1->Q());
+
+		const FImplicitObject* ParticleObj = Particle0->Geometry().Get();
+		const FImplicitObject* LevelsetObj = Particle1->Geometry().Get();
+
+		if (!(ensure(!FMath::IsNaN(ParticlesTM.GetTranslation().X)) && ensure(!FMath::IsNaN(ParticlesTM.GetTranslation().Y)) && ensure(!FMath::IsNaN(ParticlesTM.GetTranslation().Z))))
+		{
+			return;
+		}
+
+		if (!(ensure(!FMath::IsNaN(LevelsetTM.GetTranslation().X)) && ensure(!FMath::IsNaN(LevelsetTM.GetTranslation().Y)) && ensure(!FMath::IsNaN(LevelsetTM.GetTranslation().Z))))
+		{
+			return;
+		}
+
+#if CHAOS_PARTICLEHANDLE_TODO
+		TArray<Pair<const FImplicitObject*, TRigidTransform<T, d>>> ParticleShapes = FindRelevantShapes2(LevelsetObj, LevelsetTM, *ParticleObj, ParticlesTM, Thickness);
+		check(ParticleObj->IsUnderlyingUnion());
+		const TImplicitObjectUnion<T, d>* UnionObj = static_cast<const TImplicitObjectUnion<T, d>*>(ParticleObj);
+		for (const Pair<const FImplicitObject*, TRigidTransform<T, d>>& ParticleObjPair : ParticleShapes)
+		{
+			const FImplicitObject* Object = ParticleObjPair.First;
+
+			if (const int32* OriginalIdx = UnionObj->MCollisionParticleLookupHack.Find(Object))
+			{
+				const TBVHParticles<T, d>& SampleParticles = *InParticles.CollisionParticles(*OriginalIdx).Get();
+				const TRigidTransform<T, d> ObjectTM = ParticleObjPair.Second * ParticlesTM;
+
+				SampleObject2<UpdateType>(*LevelsetObj, LevelsetTM, SampleParticles, ObjectTM, Thickness, Constraint);
+				if (UpdateType == ECollisionUpdateType::Any && Constraint.Phi < Thickness)
+				{
+					return;
+				}
+			}
+		}
+#endif
+	}
+
+
+	template <typename T, int d>
+	void UpdateBoxConstraint(const TBox<T, d>& Box1, const TRigidTransform<T, d>& Box1Transform, const TBox<T, d>& Box2, const TRigidTransform<T, d>& Box2Transform, T Thickness, TRigidBodyContactConstraint<T, d>& Constraint)
+	{
+		TContactData<T, d> & Contact = Constraint.ShapeManifold.Manifold;
+
+		TBox<T, d> Box2SpaceBox1 = Box1.TransformedBox(Box1Transform * Box2Transform.Inverse());
+		TBox<T, d> Box1SpaceBox2 = Box2.TransformedBox(Box2Transform * Box1Transform.Inverse());
+		Box2SpaceBox1.Thicken(Thickness);
+		Box1SpaceBox2.Thicken(Thickness);
+		if (Box1SpaceBox2.Intersects(Box1) && Box2SpaceBox1.Intersects(Box2))
+		{
+			const TVector<T, d> Box1Center = (Box1Transform * Box2Transform.Inverse()).TransformPosition(Box1.Center());
+			bool bDeepOverlap = false;
+			if (Box2.SignedDistance(Box1Center) < 0)
+			{
+				//If Box1 is overlapping Box2 by this much the signed distance approach will fail (box1 gets sucked into box2). In this case just use two spheres
+				TSphere<T, d> Sphere1(Box1Transform.TransformPosition(Box1.Center()), Box1.Extents().Min() / 2);
+				TSphere<T, d> Sphere2(Box2Transform.TransformPosition(Box2.Center()), Box2.Extents().Min() / 2);
+				const TVector<T, d> Direction = Sphere1.GetCenter() - Sphere2.GetCenter();
+				T Size = Direction.Size();
+				if (Size < (Sphere1.GetRadius() + Sphere2.GetRadius()))
+				{
+					const T NewPhi = Size - (Sphere1.GetRadius() + Sphere2.GetRadius());;
+					if (NewPhi < Contact.Phi)
+					{
+						bDeepOverlap = true;
+						Contact.Normal = Size > SMALL_NUMBER ? Direction / Size : TVector<T, d>(0, 0, 1);
+						Contact.Phi = NewPhi;
+						Contact.Location = Sphere1.GetCenter() - Sphere1.GetRadius() * Contact.Normal;
+					}
+				}
+			}
+			if (!bDeepOverlap || Contact.Phi >= 0)
+			{
+				//if we didn't have deep penetration use signed distance per particle. If we did have deep penetration but the spheres did not overlap use signed distance per particle
+
+				//UpdateLevelsetConstraintGJK(InParticles, Thickness, Constraint);
+				//check(Contact.Phi < MThickness);
+				// For now revert to doing all points vs lsv check until we can figure out a good way to get the deepest point without needing this
+				{
+					const TArray<TVector<T, d>> SampleParticles = Box1.ComputeSamplePoints();
+					const TRigidTransform<T, d> Box1ToBox2Transform = Box1Transform.GetRelativeTransform(Box2Transform);
+					int32 NumParticles = SampleParticles.Num();
+					for (int32 i = 0; i < NumParticles; ++i)
+					{
+						SampleObjectHelper2(Box2, Box2Transform, Box1ToBox2Transform, SampleParticles[i], Thickness, Constraint);
+					}
+				}
+			}
+		}
+	}
+
+	DECLARE_CYCLE_STAT(TEXT("TPBDCollisionConstraint::UpdateLevelsetConstraint"), STAT_UpdateLevelsetConstraint, STATGROUP_ChaosWide);
+	template<ECollisionUpdateType UpdateType, typename T, int d>
+	void UpdateLevelsetConstraint(const T Thickness, TRigidBodyContactConstraint<float, 3>& Constraint)
+	{
+		SCOPE_CYCLE_COUNTER(STAT_UpdateLevelsetConstraint);
+
+		TGenericParticleHandle<T, d> Particle0 = Constraint.Particle[0];
+		TRigidTransform<T, d> ParticlesTM = TRigidTransform<T, d>(Particle0->P(), Particle0->Q());
+		if (!(ensure(!FMath::IsNaN(ParticlesTM.GetTranslation().X)) && ensure(!FMath::IsNaN(ParticlesTM.GetTranslation().Y)) && ensure(!FMath::IsNaN(ParticlesTM.GetTranslation().Z))))
+		{
+			return;
+		}
+
+		TGenericParticleHandle<T, d> Particle1 = Constraint.Particle[1];
+		TRigidTransform<T, d> LevelsetTM = TRigidTransform<T, d>(Particle1->P(), Particle1->Q());
+		if (!(ensure(!FMath::IsNaN(LevelsetTM.GetTranslation().X)) && ensure(!FMath::IsNaN(LevelsetTM.GetTranslation().Y)) && ensure(!FMath::IsNaN(LevelsetTM.GetTranslation().Z))))
+		{
+			return;
+		}
+
+		const TBVHParticles<T, d>* SampleParticles = nullptr;
+		SampleParticles = Particle0->CollisionParticles().Get();
+
+		if (SampleParticles)
+		{
+			SampleObject2<UpdateType>(*Particle1->Geometry(), LevelsetTM, *SampleParticles, ParticlesTM, Thickness, Constraint);
+		}
+	}
+
+	template void ConstructConstraintsImpl<float,3>(TGeometryParticleHandle<float, 3>* Particle0, TGeometryParticleHandle<float, 3>* Particle1, const FImplicitObject* Implicit0, const FImplicitObject* Implicit1, const float Thickness, TRigidBodyContactConstraint<float, 3> & Constraint);
+	template void UpdateUnionLevelsetConstraint<ECollisionUpdateType::Any, float, 3>(const float Thickness, TRigidBodyContactConstraint<float, 3>& Constraint);
+	template void UpdateUnionLevelsetConstraint<ECollisionUpdateType::Deepest, float, 3>(const float Thickness, TRigidBodyContactConstraint<float, 3>& Constraint);
+	template void UpdateConstraintImp<ECollisionUpdateType::Any, float, 3>(const FImplicitObject& ParticleObject, const TRigidTransform<float, 3>& ParticleTM, const FImplicitObject& LevelsetObject, const TRigidTransform<float, 3>& LevelsetTM, const float Thickness, TRigidBodyContactConstraint<float, 3>& Constraints);
+	template void UpdateConstraintImp<ECollisionUpdateType::Deepest, float, 3>(const FImplicitObject& ParticleObject, const TRigidTransform<float, 3>& ParticleTM, const FImplicitObject& LevelsetObject, const TRigidTransform<float, 3>& LevelsetTM, const float Thickness, TRigidBodyContactConstraint<float, 3>& Constraint);
+	template void UpdateLevelsetConstraint<ECollisionUpdateType::Any, float, 3>(const float Thickness, TRigidBodyContactConstraint<float, 3>& Constraint);
+	template void UpdateLevelsetConstraint<ECollisionUpdateType::Deepest, float, 3>(const float Thickness, TRigidBodyContactConstraint<float, 3>& Constraint);
+
+}
diff --git a/Engine/Source/Runtime/Experimental/Chaos/Private/Chaos/PBDCollisionConstraint.cpp b/Engine/Source/Runtime/Experimental/Chaos/Private/Chaos/PBDCollisionConstraint.cpp
--- a/Engine/Source/Runtime/Experimental/Chaos/Private/Chaos/PBDCollisionConstraint.cpp
+++ b/Engine/Source/Runtime/Experimental/Chaos/Private/Chaos/PBDCollisionConstraint.cpp
@@ -2,24 +2,26 @@
 
 #include "Chaos/PBDCollisionConstraint.h"
 #include "Chaos/PBDCollisionConstraintImp.h"
+
+#include "Chaos/Capsule.h"
 #include "Chaos/ChaosPerfTest.h"
+#include "Chaos/ChaosDebugDraw.h"
+#include "Chaos/CollisionResolutionAlgo.h"
+#include "Chaos/CollisionResolutionConvexConvex.h"
 #include "Chaos/Defines.h"
 #include "Chaos/GeometryQueries.h"
 #include "Chaos/ImplicitObjectUnion.h"
 #include "Chaos/ImplicitObjectScaled.h"
+#include "Chaos/ISpatialAccelerationCollection.h"
 #include "Chaos/Levelset.h"
 #include "Chaos/Pair.h"
-#include "Chaos/Capsule.h"
+#include "Chaos/PBDRigidsSOAs.h"
 #include "Chaos/Sphere.h"
 #include "Chaos/Transform.h"
 #include "ChaosLog.h"
 #include "ChaosStats.h"
 #include "Containers/Queue.h"
 #include "ProfilingDebugging/ScopedTimers.h"
-#include "Chaos/PBDRigidsSOAs.h"
-#include "Chaos/ChaosDebugDraw.h"
-#include "Chaos/ISpatialAccelerationCollection.h"
-#include "Chaos/CollisionResolutionConvexConvex.h"
 
 #if INTEL_ISPC
 #include "PBDCollisionConstraint.ispc.generated.h"
@@ -51,2205 +53,686 @@ namespace Chaos
 
 namespace Chaos
 {
-CHAOS_API int32 EnableCollisions = 1;
-FAutoConsoleVariableRef CVarEnableCollisions(TEXT("p.EnableCollisions"), EnableCollisions, TEXT("Enable/Disable collisions on the Chaos solver."));
-
-template<ECollisionUpdateType, typename T, int d>
-void UpdateConstraintImp(const TRigidTransform<T, d>& ParticleTM, const FImplicitObject& LevelsetObject, const TRigidTransform<T, d>& LevelsetTM, const T Thickness, TRigidBodyContactConstraint<T, d>& Constraint);
-
-template<typename T, int d>
-void ConstructConstraintsImpl(TGeometryParticleHandle<T, d>* Particle0, TGeometryParticleHandle<T, d>* Particle1, const FImplicitObject* Implicit0, const FImplicitObject* Implicit1, const T Thickness, TRigidBodyContactConstraint<T, d>& Constraint);
-
-
-//
-// Collision Constraint Container
-//
-
-template<typename T, int d>
-TPBDCollisionConstraint<T, d>::TPBDCollisionConstraint(const TPBDRigidsSOAs<T,d>& InParticles, TArrayCollectionArray<bool>& Collided, const TArrayCollectionArray<TSerializablePtr<FChaosPhysicsMaterial>>& InPerParticleMaterials, const int32 PairIterations /*= 1*/, const T Thickness /*= (T)0*/)
-	: Particles(InParticles)
-	, SpatialAcceleration(nullptr)
-	, MCollided(Collided)
-	, MPhysicsMaterials(InPerParticleMaterials)
-	, bEnableVelocitySolve(true)
-	, MPairIterations(PairIterations)
-	, MThickness(Thickness)
-	, MAngularFriction(0)
-	, bUseCCD(false)
-	, bEnableCollisions(true)
-	, LifespanCounter(0)
-	, CollisionVelocityInflation(CollisionVelocityInflationCVar)
-	, PostComputeCallback(nullptr)
-	, PostApplyCallback(nullptr)
-	, PostApplyPushOutCallback(nullptr)
-{
-}
+	CHAOS_API int32 EnableCollisions = 1;
+	FAutoConsoleVariableRef CVarEnableCollisions(TEXT("p.EnableCollisions"), EnableCollisions, TEXT("Enable/Disable collisions on the Chaos solver."));
 
-DECLARE_CYCLE_STAT(TEXT("TPBDCollisionConstraint::Reset"), STAT_CollisionConstraintsReset, STATGROUP_Chaos);
 
-template<typename T, int d>
-void TPBDCollisionConstraint<T, d>::Reset(/*const TPBDRigidParticles<T, d>& InParticles, const TArray<int32>& InIndices*/)
-{
-	SCOPE_CYCLE_COUNTER(STAT_CollisionConstraintsReset);
 
+	//
+	// Collision Constraint Container
+	//
 
-	for (int32 Idx = Constraints.Num() - 1; Idx >= 0; Idx--)
+	template<typename T, int d>
+	TPBDCollisionConstraint<T, d>::TPBDCollisionConstraint(const TPBDRigidsSOAs<T, d>& InParticles, TArrayCollectionArray<bool>& Collided, const TArrayCollectionArray<TSerializablePtr<FChaosPhysicsMaterial>>& InPerParticleMaterials, const int32 PairIterations /*= 1*/, const T Thickness /*= (T)0*/)
+		: Particles(InParticles)
+		, SpatialAcceleration(nullptr)
+		, MCollided(Collided)
+		, MPhysicsMaterials(InPerParticleMaterials)
+		, bEnableVelocitySolve(true)
+		, MPairIterations(PairIterations)
+		, MThickness(Thickness)
+		, MAngularFriction(0)
+		, bUseCCD(false)
+		, bEnableCollisions(true)
+		, LifespanCounter(0)
+		, CollisionVelocityInflation(CollisionVelocityInflationCVar)
+		, PostComputeCallback(nullptr)
+		, PostApplyCallback(nullptr)
+		, PostApplyPushOutCallback(nullptr)
 	{
-		//if (!bEnableCollisions)
-		{
-			RemoveConstraint(Idx);
-		}
 	}
 
+	DECLARE_CYCLE_STAT(TEXT("TPBDCollisionConstraint::Reset"), STAT_CollisionConstraintsReset, STATGROUP_Chaos);
 
-	MAngularFriction = 0;
-	bUseCCD = false;
-}
-
-template<typename T, int d>
-void TPBDCollisionConstraint<T, d>::RemoveConstraint(int32 Idx)
-{
-	HandleAllocator.FreeHandle(Handles[Idx]);
-	Handles.RemoveAtSwap(Idx);
-	Constraints.RemoveAtSwap(Idx);
-	if (Idx < Constraints.Num())
+	template<typename T, int d>
+	void TPBDCollisionConstraint<T, d>::Reset(/*const TPBDRigidParticles<T, d>& InParticles, const TArray<int32>& InIndices*/)
 	{
-		Handles[Idx]->SetConstraintIndex(Idx);
-	}
-	ensure(Handles.Num() == Constraints.Num());
-}
+		SCOPE_CYCLE_COUNTER(STAT_CollisionConstraintsReset);
+
 
-template<typename T, int d>
-void TPBDCollisionConstraint<T, d>::RemoveConstraints(const TSet<TGeometryParticleHandle<T, d>*>&  InHandleSet)
-{
-	const TArray<TGeometryParticleHandle<T, d>*> HandleArray = InHandleSet.Array();
-	for (auto ParticleHandle : HandleArray)
-	{
 		for (int32 Idx = Constraints.Num() - 1; Idx >= 0; Idx--)
 		{
-			if (Constraints[Idx].Particle[1] == ParticleHandle || Constraints[Idx].Particle[0] == ParticleHandle)
+			//if (!bEnableCollisions)
 			{
 				RemoveConstraint(Idx);
 			}
 		}
+
+
+		MAngularFriction = 0;
+		bUseCCD = false;
 	}
-}
 
-template<typename T, int d>
-void TPBDCollisionConstraint<T, d>::ApplyCollisionModifier(const TFunction<ECollisionModifierResult(FRigidBodyContactConstraint& Constraint)>& CollisionModifier)
-{
-	for (int ConstraintIndex = 0; ConstraintIndex < NumConstraints(); ++ConstraintIndex)
+	template<typename T, int d>
+	void TPBDCollisionConstraint<T, d>::RemoveConstraint(int32 Idx)
 	{
-		ECollisionModifierResult Result = CollisionModifier(Constraints[ConstraintIndex]);
-		if (Result == ECollisionModifierResult::Disabled)
+		HandleAllocator.FreeHandle(Handles[Idx]);
+		Handles.RemoveAtSwap(Idx);
+		Constraints.RemoveAtSwap(Idx);
+		if (Idx < Constraints.Num())
 		{
-			RemoveConstraint(ConstraintIndex);
-			--ConstraintIndex;
+			Handles[Idx]->SetConstraintIndex(Idx);
 		}
+		ensure(Handles.Num() == Constraints.Num());
 	}
-}
-
-template<typename T, int d>
-void TPBDCollisionConstraint<T, d>::SetPostComputeCallback(const TRigidBodyContactConstraintsPostComputeCallback<T, d>& Callback)
-{
-	PostComputeCallback = Callback;
-}
-
-template<typename T, int d>
-void TPBDCollisionConstraint<T, d>::ClearPostComputeCallback()
-{
-	PostComputeCallback = nullptr;
-}
-
-template<typename T, int d>
-void TPBDCollisionConstraint<T, d>::SetPostApplyCallback(const TRigidBodyContactConstraintsPostApplyCallback<T, d>& Callback)
-{
-	PostApplyCallback = Callback;
-}
-
-template<typename T, int d>
-void TPBDCollisionConstraint<T, d>::ClearPostApplyCallback()
-{
-	PostApplyCallback = nullptr;
-}
-
-template<typename T, int d>
-void TPBDCollisionConstraint<T, d>::SetPostApplyPushOutCallback(const TRigidBodyContactConstraintsPostApplyPushOutCallback<T, d>& Callback)
-{
-	PostApplyPushOutCallback = Callback;
-}
-
-template<typename T, int d>
-void TPBDCollisionConstraint<T, d>::ClearPostApplyPushOutCallback()
-{
-	PostApplyPushOutCallback = nullptr;
-}
-
-template<typename T, int d>
-void TPBDCollisionConstraint<T, d>::UpdatePositionBasedState(/*const TPBDRigidParticles<T, d>& InParticles, const TArray<int32>& InIndices,*/ const T Dt)
-{
-	Reset();
 
-	if (bEnableCollisions)
+	template<typename T, int d>
+	void TPBDCollisionConstraint<T, d>::RemoveConstraints(const TSet<TGeometryParticleHandle<T, d>*>&  InHandleSet)
 	{
-#if !UE_BUILD_SHIPPING
-		if (PendingHierarchyDump)
-		{
-			ComputeConstraints<true>(*SpatialAcceleration, Dt);
-		}
-		else
-#endif
+		const TArray<TGeometryParticleHandle<T, d>*> HandleArray = InHandleSet.Array();
+		for (auto ParticleHandle : HandleArray)
 		{
-			ComputeConstraints(*SpatialAcceleration, Dt);
+			for (int32 Idx = Constraints.Num() - 1; Idx >= 0; Idx--)
+			{
+				if (Constraints[Idx].Particle[1] == ParticleHandle || Constraints[Idx].Particle[0] == ParticleHandle)
+				{
+					RemoveConstraint(Idx);
+				}
+			}
 		}
 	}
-}
-
-DEFINE_STAT(STAT_ComputeConstraints);
-DEFINE_STAT(STAT_ComputeConstraintsNP);
-DEFINE_STAT(STAT_ComputeConstraintsBP);
-DEFINE_STAT(STAT_ComputeConstraintsSU);
 
-template<typename T, int d>
-template <bool bGatherStats>
-void TPBDCollisionConstraint<T, d>::ComputeConstraints(const FAccelerationStructure& AccelerationStructure, T Dt)
-{
-	if (const auto AABBTree = AccelerationStructure.template As<TAABBTree<TAccelerationStructureHandle<T, d>, TAABBTreeLeafArray<TAccelerationStructureHandle<T, d>, T>, T>>())
+	template<typename T, int d>
+	void TPBDCollisionConstraint<T, d>::ApplyCollisionModifier(const TFunction<ECollisionModifierResult(FRigidBodyContactConstraint& Constraint)>& CollisionModifier)
 	{
-		ComputeConstraintsHelperLowLevel<bGatherStats>(*AABBTree, Dt);
+		for (int ConstraintIndex = 0; ConstraintIndex < NumConstraints(); ++ConstraintIndex)
+		{
+			ECollisionModifierResult Result = CollisionModifier(Constraints[ConstraintIndex]);
+			if (Result == ECollisionModifierResult::Disabled)
+			{
+				RemoveConstraint(ConstraintIndex);
+				--ConstraintIndex;
+			}
+		}
 	}
-	else if(const auto BV = AccelerationStructure.template As<TBoundingVolume<TAccelerationStructureHandle<T, d>, T, d>>())
+
+	template<typename T, int d>
+	void TPBDCollisionConstraint<T, d>::SetPostComputeCallback(const TRigidBodyContactConstraintsPostComputeCallback<T, d>& Callback)
 	{
-		ComputeConstraintsHelperLowLevel<bGatherStats>(*BV, Dt);
+		PostComputeCallback = Callback;
 	}
-	else if (const auto AABBTreeBV = AccelerationStructure.template As<TAABBTree<TAccelerationStructureHandle<T, d>, TBoundingVolume<TAccelerationStructureHandle<T, d>, T, d>, T>>())
+
+	template<typename T, int d>
+	void TPBDCollisionConstraint<T, d>::ClearPostComputeCallback()
 	{
-		ComputeConstraintsHelperLowLevel<bGatherStats>(*AABBTreeBV, Dt);
+		PostComputeCallback = nullptr;
 	}
-	else if(const auto Collection = AccelerationStructure.template As<ISpatialAccelerationCollection<TAccelerationStructureHandle<T, d>, T, d>>())
+
+	template<typename T, int d>
+	void TPBDCollisionConstraint<T, d>::SetPostApplyCallback(const TRigidBodyContactConstraintsPostApplyCallback<T, d>& Callback)
 	{
-		if (bGatherStats)
-		{
-			Collection->PBDComputeConstraintsLowLevel_GatherStats(*this, Dt);
-		}
-		else
-		{
-			Collection->PBDComputeConstraintsLowLevel(*this, Dt);
-		}
+		PostApplyCallback = Callback;
 	}
-	else
+
+	template<typename T, int d>
+	void TPBDCollisionConstraint<T, d>::ClearPostApplyCallback()
 	{
-		check(false);  //question: do we want to support a dynamic dispatch version?
+		PostApplyCallback = nullptr;
 	}
-	
-	if (PostComputeCallback != nullptr)
+
+	template<typename T, int d>
+	void TPBDCollisionConstraint<T, d>::SetPostApplyPushOutCallback(const TRigidBodyContactConstraintsPostApplyPushOutCallback<T, d>& Callback)
 	{
-		PostComputeCallback();
+		PostApplyPushOutCallback = Callback;
 	}
-}
-
-
-DECLARE_CYCLE_STAT(TEXT("UpdateConstraints"), STAT_UpdateConstraints, STATGROUP_Chaos);
-
-template <typename T, int d>
-template<typename SPATIAL_ACCELERATION>
-void TPBDCollisionConstraint<T, d>::UpdateConstraintsHelper(/*const TPBDRigidParticles<T, d>& InParticles, const TArray<int32>& InIndices,*/ T Dt, const TSet<TGeometryParticleHandle<T, d>*>& AddedParticles, SPATIAL_ACCELERATION& InSpatialAcceleration)
-{
-#if CHAOS_PARTICLEHANDLE_TODO
-	SCOPE_CYCLE_COUNTER(STAT_UpdateConstraints);
-	double Time = 0;
-	FDurationTimer Timer(Time);
-
-	TArray<uint32> AddedParticlesArray = AddedParticles.Array();
-	TArray<uint32> NewActiveIndices = ActiveParticles;
-	NewActiveIndices.Append(AddedParticlesArray);
-
-	//
-	// Broad phase
-	//
 
+	template<typename T, int d>
+	void TPBDCollisionConstraint<T, d>::ClearPostApplyPushOutCallback()
 	{
-		//QUICK_SCOPE_CYCLE_COUNTER(Reinitialize);
-		// @todo(mlentine): We only need to construct the hierarchy for the islands we care about
-		//TBoundingVolumeHierarchy<TPBDRigidParticles<T, d>, TArray<int32>, T, d> LocalHierarchy(InParticles, ActiveParticles, true, Dt * BoundsThicknessMultiplier); 	//todo(ocohen): should we pass MThickness into this structure?
-		SpatialAcceleration.Reinitialize((const TArray<uint32>&)NewActiveIndices, true, Dt * BoundsThicknessMultiplier); //todo: faster path when adding just a few
-		Timer.Stop();
-		UE_LOG(LogChaos, Verbose, TEXT("\tPBDCollisionConstraint Construct Hierarchy %f"), Time);
+		PostApplyPushOutCallback = nullptr;
 	}
 
-	//
-	// Narrow phase
-	//
+	template<typename T, int d>
+	void TPBDCollisionConstraint<T, d>::UpdatePositionBasedState(/*const TPBDRigidParticles<T, d>& InParticles, const TArray<int32>& InIndices,*/ const T Dt)
+	{
+		Reset();
 
-	FCriticalSection CriticalSection;
-	Time = 0;
-	Timer.Start();
-	
-	TQueue<TRigidBodyContactConstraint<T, d>, EQueueMode::Mpsc> Queue;	//todo(ocohen): use per thread buffer instead, need better support than ParallelFor for this
-	//SpatialAcceleration.AddElements(AddedParticlesArray);	not supported by bvh so just reinitializing. Should probably improve this later
-	PhysicsParallelFor(AddedParticlesArray.Num(), [&](int32 Index) {
-		int32 Body1Index = AddedParticlesArray[Index];
-		if (InParticles.Disabled(Body1Index))
+		if (bEnableCollisions)
 		{
-			return;
+#if !UE_BUILD_SHIPPING
+			if (PendingHierarchyDump)
+			{
+				ComputeConstraints<true>(*SpatialAcceleration, Dt);
+			}
+			else
+#endif
+			{
+				ComputeConstraints(*SpatialAcceleration, Dt);
+			}
 		}
-		if (InParticles.InvM(Body1Index) == 0)
+	}
+
+	DEFINE_STAT(STAT_ComputeConstraints);
+	DEFINE_STAT(STAT_ComputeConstraintsNP);
+	DEFINE_STAT(STAT_ComputeConstraintsBP);
+	DEFINE_STAT(STAT_ComputeConstraintsSU);
+
+	template<typename T, int d>
+	template <bool bGatherStats>
+	void TPBDCollisionConstraint<T, d>::ComputeConstraints(const FAccelerationStructure& AccelerationStructure, T Dt)
+	{
+		if (const auto AABBTree = AccelerationStructure.template As<TAABBTree<TAccelerationStructureHandle<T, d>, TAABBTreeLeafArray<TAccelerationStructureHandle<T, d>, T>, T>>())
 		{
-			return;
+			ComputeConstraintsHelperLowLevel<bGatherStats>(*AABBTree, Dt);
 		}
-		TArray<int32> PotentialIntersections;
-		TBox<T, d> Box1;
-		T Box1Thickness = (T)0;
-
-		const bool bBody1Bounded = HasBoundingBox(InParticles, Body1Index);
-		if (bBody1Bounded)
+		else if (const auto BV = AccelerationStructure.template As<TBoundingVolume<TAccelerationStructureHandle<T, d>, T, d>>())
 		{
-			Box1 = SpatialAcceleration.GetWorldSpaceBoundingBox(InParticles, Body1Index);
-			Box1Thickness = ComputeThickness(InParticles, Dt, Body1Index).Size();
-			PotentialIntersections = SpatialAcceleration.FindAllIntersections(Box1);
+			ComputeConstraintsHelperLowLevel<bGatherStats>(*BV, Dt);
 		}
-		else
+		else if (const auto AABBTreeBV = AccelerationStructure.template As<TAABBTree<TAccelerationStructureHandle<T, d>, TBoundingVolume<TAccelerationStructureHandle<T, d>, T, d>, T>>())
 		{
-			PotentialIntersections = SpatialAcceleration.GlobalObjects();
+			ComputeConstraintsHelperLowLevel<bGatherStats>(*AABBTreeBV, Dt);
 		}
-		for (int32 i = 0; i < PotentialIntersections.Num(); ++i)
+		else if (const auto Collection = AccelerationStructure.template As<ISpatialAccelerationCollection<TAccelerationStructureHandle<T, d>, T, d>>())
 		{
-			int32 Body2Index = PotentialIntersections[i];
-			const bool bBody2Bounded = HasBoundingBox(InParticles, Body2Index);
-
-			if(InParticles.Disabled(Body2Index))
-			{
-				// Can't collide with disabled objects
-				continue;
-			}
-
-			if (Body1Index == Body2Index || ((bBody1Bounded == bBody2Bounded) && AddedParticles.Contains(Body2Index) && AddedParticles.Contains(Body1Index) && Body2Index > Body1Index))
-			{
-				continue;
-			}
-
-			if (InParticles.InvM(Body1Index) && InParticles.InvM(Body2Index) && (InParticles.Island(Body1Index) != InParticles.Island(Body2Index)))	//todo(ocohen): this is a hack - we should not even consider dynamics from other islands
-			{
-				continue;
-			}
-
-			if (!InParticles.Geometry(Body1Index) && !InParticles.Geometry(Body2Index))
-			{
-				continue;
-			}
-
-			if (bBody1Bounded && bBody2Bounded)
-			{	
-				const TBox<T, d>& Box2 = SpatialAcceleration.GetWorldSpaceBoundingBox(InParticles, Body2Index);
-				if (!Box1.Intersects(Box2))
-				{
-					continue;
-				}
-			}
-
-			//todo(ocohen): this should not be needed in theory, but in practice we accidentally merge islands. We should be doing this test within an island for clusters
-			if (InParticles.Island(Body1Index) >= 0 && InParticles.Island(Body2Index) >= 0 && InParticles.Island(Body1Index) != InParticles.Island(Body2Index))
+			if (bGatherStats)
 			{
-				continue;
+				Collection->PBDComputeConstraintsLowLevel_GatherStats(*this, Dt);
 			}
-
-			const TVector<T, d> Box2Thickness = ComputeThickness(InParticles, Dt, Body2Index);
-			const T UseThickness = FMath::Max(Box1Thickness, Box2Thickness.Size());// + MThickness
-
-			auto Constraint = ComputeConstraint(InParticles, Body1Index, Body2Index, UseThickness);
-
-			//if (true || !InParticles.Geometry(Body1Index)->HasBoundingBox() || !InParticles.Geometry(Body2Index)->HasBoundingBox())
+			else
 			{
-				//use narrow phase to determine if constraint is needed. Without this we can't do shock propagation
-				if (ComputeConstraintsUseAny)
-				{
-					UpdateConstraint<ECollisionUpdateType::Any>(InParticles, UseThickness, Constraint);
-				}
-				else
-				{
-					UpdateConstraint<ECollisionUpdateType::Deepest>(InParticles, UseThickness, Constraint);
-				}
-				if (Constraint.Phi < UseThickness)
-				{
-					Queue.Enqueue(Constraint);
-				}
+				Collection->PBDComputeConstraintsLowLevel(*this, Dt);
 			}
-			/*else
-			{
-			CriticalSection.Lock();
-			Constraints.Add(Constraint);
-			CriticalSection.Unlock();
-			}*/
 		}
-	});
-	while (!Queue.IsEmpty())
-	{
-		const TRigidBodyContactConstraint<T, d> * Constraint = Queue.Peek();
-		FConstraintHandleID HandleID = GetConstraintHandleID(*Constraint);
-		if (Handles.Contains(HandleID))
+		else
 		{
-			FConstraintContainerHandle* Handle = Handles[HandleID];
-			int32 Idx = Handle->GetConstraintIndex();
-			Queue.Dequeue(Constraints[Idx]);
-			Constraints[Idx].Lifespan = LifespanCounter;
+			check(false);  //question: do we want to support a dynamic dispatch version?
 		}
-		else
+
+		if (PostComputeCallback != nullptr)
 		{
-			int32 Idx = Constraints.AddUninitialized(1);
-			Queue.Dequeue(Constraints[Idx]);
-			Handles.Add(GetConstraintHandleID(Idx), HandleAllocator.AllocHandle(this, Idx));
-			Constraints[Idx].Lifespan = LifespanCounter;
+			PostComputeCallback();
 		}
 	}
-	LifespanCounter++;
 
-	Timer.Stop();
-	UE_LOG(LogChaos, Verbose, TEXT("\tPBDCollisionConstraint Update %d Constraints with Potential Collisions %f"), Constraints.Num(), Time);
-#endif
-}
 
-DECLARE_CYCLE_STAT(TEXT("TPBDCollisionConstraint::ReconcileUpdatedConstraints"), STAT_ReconcileConstraints, STATGROUP_ChaosWide);
-template<typename T, int d>
-void TPBDCollisionConstraint<T, d>::UpdateConstraints(/*const TPBDRigidParticles<T, d>& InParticles, const TArray<int32>& InIndices,*/ T Dt, const TSet<TGeometryParticleHandle<T, d>*>& AddedParticles)
-{
-#if CHAOS_PARTICLEHANDLE_TODO
+
+	// @todo(ccaulfield): This is duplicated in JointConstraints - move to a utility file
+	template<typename T>
+	PMatrix<T, 3, 3> ComputeFactorMatrix3(const TVector<T, 3>& V, const PMatrix<T, 3, 3>& M, const T& Im)
+	{
+		// Rigid objects rotational contribution to the impulse.
+		// Vx*M*VxT+Im
+		check(Im > FLT_MIN)
+			return PMatrix<T, 3, 3>(
+				-V[2] * (-V[2] * M.M[1][1] + V[1] * M.M[2][1]) + V[1] * (-V[2] * M.M[2][1] + V[1] * M.M[2][2]) + Im,
+				V[2] * (-V[2] * M.M[1][0] + V[1] * M.M[2][0]) - V[0] * (-V[2] * M.M[2][1] + V[1] * M.M[2][2]),
+				-V[1] * (-V[2] * M.M[1][0] + V[1] * M.M[2][0]) + V[0] * (-V[2] * M.M[1][1] + V[1] * M.M[2][1]),
+				V[2] * (V[2] * M.M[0][0] - V[0] * M.M[2][0]) - V[0] * (V[2] * M.M[2][0] - V[0] * M.M[2][2]) + Im,
+				-V[1] * (V[2] * M.M[0][0] - V[0] * M.M[2][0]) + V[0] * (V[2] * M.M[1][0] - V[0] * M.M[2][1]),
+				-V[1] * (-V[1] * M.M[0][0] + V[0] * M.M[1][0]) + V[0] * (-V[1] * M.M[1][0] + V[0] * M.M[1][1]) + Im);
+	}
+
+	template<typename T, int d>
+	TVector<T, d> GetEnergyClampedImpulse(const TRigidBodyContactConstraint<T, d>& Constraint, const TVector<T, d>& Impulse, const TVector<T, d>& VectorToPoint1, const TVector<T, d>& VectorToPoint2, const TVector<T, d>& Velocity1, const TVector<T, d>& Velocity2)
 	{
-		SCOPE_CYCLE_COUNTER(STAT_ReconcileConstraints);
+		TPBDRigidParticleHandle<T, d>* PBDRigid0 = Constraint.Particle[0]->AsDynamic();
+		TPBDRigidParticleHandle<T, d>* PBDRigid1 = Constraint.Particle[1]->AsDynamic();
 
-		// Updating post-clustering, we will have invalid constraints
-		int32 NumRemovedConstraints = 0;
-		for(int32 i = 0; i < Constraints.Num(); ++i)
+		TVector<T, d> Jr0, Jr1, IInvJr0, IInvJr1;
+		T ImpulseRatioNumerator0 = 0, ImpulseRatioNumerator1 = 0, ImpulseRatioDenom0 = 0, ImpulseRatioDenom1 = 0;
+		T ImpulseSize = Impulse.SizeSquared();
+		TVector<T, d> KinematicVelocity = !PBDRigid0 ? Velocity1 : !PBDRigid1 ? Velocity2 : TVector<T, d>(0);
+		if (PBDRigid0)
 		{
-			const FRigidBodyContactConstraint& Constraint = Constraints[i];
-			if(InParticles.Disabled(Constraint.ParticleIndex) || InParticles.Disabled(Constraint.ParticleIndex))
-			{
-				Constraints.RemoveAtSwap(i);
-				HandleAllocator.FreeHandle(Handles[i]);
-				Handles.RemoveAtSwap(i);
-				if (i < Handles.Num())
-				{
-					SetConstraintIndex(Handles[i], i);
-				}
-				++NumRemovedConstraints;
-				i--;
-			}
+			Jr0 = TVector<T, d>::CrossProduct(VectorToPoint1, Impulse);
+			IInvJr0 = PBDRigid0->Q().RotateVector(PBDRigid0->InvI() * PBDRigid0->Q().UnrotateVector(Jr0));
+			ImpulseRatioNumerator0 = TVector<T, d>::DotProduct(Impulse, PBDRigid0->V() - KinematicVelocity) + TVector<T, d>::DotProduct(IInvJr0, PBDRigid0->W());
+			ImpulseRatioDenom0 = ImpulseSize / PBDRigid0->M() + TVector<T, d>::DotProduct(Jr0, IInvJr0);
 		}
-
-		if(NumRemovedConstraints > 0)
+		if (PBDRigid1)
 		{
-			UE_LOG(LogChaos, Verbose, TEXT("TPBDCollisionConstraint::UpdateConstraints - Needed to remove %d constraints because they contained disabled particles."), NumRemovedConstraints);
+			Jr1 = TVector<T, d>::CrossProduct(VectorToPoint2, Impulse);
+			IInvJr1 = PBDRigid1->Q().RotateVector(PBDRigid1->InvI() * PBDRigid1->Q().UnrotateVector(Jr1));
+			ImpulseRatioNumerator1 = TVector<T, d>::DotProduct(Impulse, PBDRigid1->V() - KinematicVelocity) + TVector<T, d>::DotProduct(IInvJr1, PBDRigid1->W());
+			ImpulseRatioDenom1 = ImpulseSize / PBDRigid1->M() + TVector<T, d>::DotProduct(Jr1, IInvJr1);
 		}
+		T Numerator = -2 * (ImpulseRatioNumerator0 - ImpulseRatioNumerator1);
+		if (Numerator < 0)
+		{
+			return TVector<T, d>(0);
+		}
+		check(Numerator >= 0);
+		T Denominator = ImpulseRatioDenom0 + ImpulseRatioDenom1;
+		return Numerator < Denominator ? (Impulse * Numerator / Denominator) : Impulse;
 	}
 
-	if (BPTreeOfGrids)
-	{
-		UpdateConstraintsHelper(InParticles, InIndices, Dt, AddedParticles, ActiveParticles, SpatialAccelerationResource2);
-	}
-	else
-	{
-		UpdateConstraintsHelper(InParticles, InIndices, Dt, AddedParticles, ActiveParticles, SpatialAccelerationResource);
-	}
-#endif
-}
-
-// @todo(ccaulfield): This is duplicated in JointConstraints - move to a utility file
-template<typename T>
-PMatrix<T, 3, 3> ComputeFactorMatrix3(const TVector<T, 3>& V, const PMatrix<T, 3, 3>& M, const T& Im)
-{
-	// Rigid objects rotational contribution to the impulse.
-	// Vx*M*VxT+Im
-	check(Im > FLT_MIN)
-	return PMatrix<T, 3, 3>(
-		-V[2] * (-V[2] * M.M[1][1] + V[1] * M.M[2][1]) + V[1] * (-V[2] * M.M[2][1] + V[1] * M.M[2][2]) + Im,
-		V[2] * (-V[2] * M.M[1][0] + V[1] * M.M[2][0]) - V[0] * (-V[2] * M.M[2][1] + V[1] * M.M[2][2]),
-		-V[1] * (-V[2] * M.M[1][0] + V[1] * M.M[2][0]) + V[0] * (-V[2] * M.M[1][1] + V[1] * M.M[2][1]),
-		V[2] * (V[2] * M.M[0][0] - V[0] * M.M[2][0]) - V[0] * (V[2] * M.M[2][0] - V[0] * M.M[2][2]) + Im,
-		-V[1] * (V[2] * M.M[0][0] - V[0] * M.M[2][0]) + V[0] * (V[2] * M.M[1][0] - V[0] * M.M[2][1]),
-		-V[1] * (-V[1] * M.M[0][0] + V[0] * M.M[1][0]) + V[0] * (-V[1] * M.M[1][0] + V[0] * M.M[1][1]) + Im);
-}
-
-template<typename T, int d>
-TVector<T, d> GetEnergyClampedImpulse(const TRigidBodyContactConstraint<T, d>& Constraint, const TVector<T, d>& Impulse, const TVector<T, d>& VectorToPoint1, const TVector<T, d>& VectorToPoint2, const TVector<T, d>& Velocity1, const TVector<T, d>& Velocity2)
-{
-	TPBDRigidParticleHandle<T, d>* PBDRigid0 = Constraint.Particle[0]->AsDynamic();
-	TPBDRigidParticleHandle<T, d>* PBDRigid1 = Constraint.Particle[1]->AsDynamic();
-
-	TVector<T, d> Jr0, Jr1, IInvJr0, IInvJr1;
-	T ImpulseRatioNumerator0 = 0, ImpulseRatioNumerator1 = 0, ImpulseRatioDenom0 = 0, ImpulseRatioDenom1 = 0;
-	T ImpulseSize = Impulse.SizeSquared();
-	TVector<T, d> KinematicVelocity = !PBDRigid0 ? Velocity1 : !PBDRigid1 ? Velocity2 : TVector<T, d>(0);
-	if (PBDRigid0)
-	{
-		Jr0 = TVector<T, d>::CrossProduct(VectorToPoint1, Impulse);
-		IInvJr0 = PBDRigid0->Q().RotateVector(PBDRigid0->InvI() * PBDRigid0->Q().UnrotateVector(Jr0));
-		ImpulseRatioNumerator0 = TVector<T, d>::DotProduct(Impulse, PBDRigid0->V() - KinematicVelocity) + TVector<T, d>::DotProduct(IInvJr0, PBDRigid0->W());
-		ImpulseRatioDenom0 = ImpulseSize / PBDRigid0->M() + TVector<T, d>::DotProduct(Jr0, IInvJr0);
-	}
-	if (PBDRigid1)
-	{
-		Jr1 = TVector<T, d>::CrossProduct(VectorToPoint2, Impulse);
-		IInvJr1 = PBDRigid1->Q().RotateVector(PBDRigid1->InvI() * PBDRigid1->Q().UnrotateVector(Jr1));
-		ImpulseRatioNumerator1 = TVector<T, d>::DotProduct(Impulse, PBDRigid1->V() - KinematicVelocity) + TVector<T, d>::DotProduct(IInvJr1, PBDRigid1->W());
-		ImpulseRatioDenom1 = ImpulseSize / PBDRigid1->M() + TVector<T, d>::DotProduct(Jr1, IInvJr1);
-	}
-	T Numerator = -2 * (ImpulseRatioNumerator0 - ImpulseRatioNumerator1);
-	if (Numerator < 0)
+	template<typename T, int d>
+	void TPBDCollisionConstraint<T, d>::UpdateConstraints(T Dt, const TSet<TGeometryParticleHandle<T, d>*>& AddedParticles)
 	{
-		return TVector<T, d>(0);
+		// Clustering uses update constraints to force a re-evaluation. 
 	}
-	check(Numerator >= 0);
-	T Denominator = ImpulseRatioDenom0 + ImpulseRatioDenom1;
-	return Numerator < Denominator ? (Impulse * Numerator / Denominator) : Impulse;
-}
 
-template<typename T, int d>
-void TPBDCollisionConstraint<T, d>::Apply(const T Dt, FRigidBodyContactConstraint& Constraint)
-{
-	TGenericParticleHandle<T, d> Particle0 = TGenericParticleHandle<T, d>(Constraint.Particle[0]);
-	TGenericParticleHandle<T, d> Particle1 = TGenericParticleHandle<T, d>(Constraint.Particle[1]);
-	TPBDRigidParticleHandle<T, d>* PBDRigid0 = Particle0->AsDynamic();
-	TPBDRigidParticleHandle<T, d>* PBDRigid1 = Particle1->AsDynamic();
 
-	// @todo(ccaulfield): I think we should never get this? Revisit after particle handle refactor
-	if (Particle0->Sleeping())
-	{
-		ensure(!PBDRigid1 || PBDRigid1->Sleeping());
-		return;
-	}
-	if (Particle1->Sleeping())
+	template<typename T, int d>
+	void TPBDCollisionConstraint<T, d>::Apply(const T Dt, FRigidBodyContactConstraint& Constraint)
 	{
-		ensure(!PBDRigid0 || PBDRigid0->Sleeping());
-		return;
-	}
-	UpdateConstraint<ECollisionUpdateType::Deepest>(MThickness, Constraint);
-	if (Constraint.GetPhi() >= MThickness)
-	{
-		return;
-	}
+		TGenericParticleHandle<T, d> Particle0 = TGenericParticleHandle<T, d>(Constraint.Particle[0]);
+		TGenericParticleHandle<T, d> Particle1 = TGenericParticleHandle<T, d>(Constraint.Particle[1]);
+		TPBDRigidParticleHandle<T, d>* PBDRigid0 = Particle0->AsDynamic();
+		TPBDRigidParticleHandle<T, d>* PBDRigid1 = Particle1->AsDynamic();
 
-	// @todo(ccaulfield): CHAOS_PARTICLEHANDLE_TODO what's the best way to manage external per-particle data?
-	Particle0->AuxilaryValue(MCollided) = true;
-	Particle1->AuxilaryValue(MCollided) = true;
-
-	// @todo(ccaulfield): CHAOS_PARTICLEHANDLE_TODO split function to avoid ifs
-	const TVector<T, d> ZeroVector = TVector<T, d>(0);
-	const TRotation<T, d>& Q0 = Particle0->Q();
-	const TRotation<T, d>& Q1 = Particle1->Q();
-	const TVector<T, d>& P0 = Particle0->P();
-	const TVector<T, d>& P1 = Particle1->P();
-	const TVector<T, d>& V0 = Particle0->V();
-	const TVector<T, d>& V1 = Particle1->V();
-	const TVector<T, d>& W0 = Particle0->W();
-	const TVector<T, d>& W1 = Particle1->W();
-	TSerializablePtr<FChaosPhysicsMaterial> PhysicsMaterial0 = Particle0->AuxilaryValue(MPhysicsMaterials);
-	TSerializablePtr<FChaosPhysicsMaterial> PhysicsMaterial1 = Particle1->AuxilaryValue(MPhysicsMaterials);
-
-	TContactData<T,d> & Contact = Constraint.ShapeManifold.Manifold;
-
-	TVector<T, d> VectorToPoint1 = Contact.Location - P0;
-	TVector<T, d> VectorToPoint2 = Contact.Location - P1;
-	TVector<T, d> Body1Velocity = V0 + TVector<T, d>::CrossProduct(W0, VectorToPoint1);
-	TVector<T, d> Body2Velocity = V1 + TVector<T, d>::CrossProduct(W1, VectorToPoint2);
-	TVector<T, d> RelativeVelocity = Body1Velocity - Body2Velocity;
-	if (TVector<T, d>::DotProduct(RelativeVelocity, Contact.Normal) < 0) // ignore separating constraints
-	{
-		PMatrix<T, d, d> WorldSpaceInvI1 = PBDRigid0 ? (Q0 * FMatrix::Identity).GetTransposed() * PBDRigid0->InvI() * (Q0 * FMatrix::Identity) : PMatrix<T, d, d>(0);
-		PMatrix<T, d, d> WorldSpaceInvI2 = PBDRigid1 ? (Q1 * FMatrix::Identity).GetTransposed() * PBDRigid1->InvI() * (Q1 * FMatrix::Identity) : PMatrix<T, d, d>(0);
-		PMatrix<T, d, d> Factor =
-			(PBDRigid0 ? ComputeFactorMatrix3(VectorToPoint1, WorldSpaceInvI1, PBDRigid0->InvM()) : PMatrix<T, d, d>(0)) +
-			(PBDRigid1 ? ComputeFactorMatrix3(VectorToPoint2, WorldSpaceInvI2, PBDRigid1->InvM()) : PMatrix<T, d, d>(0));
-		TVector<T, d> Impulse;
-		TVector<T, d> AngularImpulse(0);
-
-		// Resting contact if very close to the surface
-		T Restitution = (T)0;
-		T Friction = (T)0;
-		bool bApplyRestitution = (RelativeVelocity.Size() > (2 * 980 * Dt));
-		if (PhysicsMaterial0 && PhysicsMaterial1)
+		// @todo(ccaulfield): I think we should never get this? Revisit after particle handle refactor
+		if (Particle0->Sleeping())
 		{
-			if (bApplyRestitution)
-			{
-				Restitution = FMath::Min(PhysicsMaterial0->Restitution, PhysicsMaterial1->Restitution);
-			}
-			Friction = FMath::Max(PhysicsMaterial0->Friction, PhysicsMaterial1->Friction);
+			ensure(!PBDRigid1 || PBDRigid1->Sleeping());
+			return;
 		}
-		else if (PhysicsMaterial0)
+		if (Particle1->Sleeping())
 		{
-			if (bApplyRestitution)
-			{
-				Restitution = PhysicsMaterial0->Restitution;
-			}
-			Friction = PhysicsMaterial0->Friction;
+			ensure(!PBDRigid0 || PBDRigid0->Sleeping());
+			return;
 		}
-		else if (PhysicsMaterial1)
+		UpdateConstraint<ECollisionUpdateType::Deepest>(MThickness, Constraint);
+		if (Constraint.GetPhi() >= MThickness)
 		{
-			if (bApplyRestitution)
-			{
-				Restitution = PhysicsMaterial1->Restitution;
-			}
-			Friction = PhysicsMaterial1->Friction;
+			return;
 		}
 
-		if (Friction)
-		{
-			T RelativeNormalVelocity = TVector<T, d>::DotProduct(RelativeVelocity, Contact.Normal);
-			if (RelativeNormalVelocity > 0)
+		// @todo(ccaulfield): CHAOS_PARTICLEHANDLE_TODO what's the best way to manage external per-particle data?
+		Particle0->AuxilaryValue(MCollided) = true;
+		Particle1->AuxilaryValue(MCollided) = true;
+
+		// @todo(ccaulfield): CHAOS_PARTICLEHANDLE_TODO split function to avoid ifs
+		const TVector<T, d> ZeroVector = TVector<T, d>(0);
+		const TRotation<T, d>& Q0 = Particle0->Q();
+		const TRotation<T, d>& Q1 = Particle1->Q();
+		const TVector<T, d>& P0 = Particle0->P();
+		const TVector<T, d>& P1 = Particle1->P();
+		const TVector<T, d>& V0 = Particle0->V();
+		const TVector<T, d>& V1 = Particle1->V();
+		const TVector<T, d>& W0 = Particle0->W();
+		const TVector<T, d>& W1 = Particle1->W();
+		TSerializablePtr<FChaosPhysicsMaterial> PhysicsMaterial0 = Particle0->AuxilaryValue(MPhysicsMaterials);
+		TSerializablePtr<FChaosPhysicsMaterial> PhysicsMaterial1 = Particle1->AuxilaryValue(MPhysicsMaterials);
+
+		TContactData<T, d> & Contact = Constraint.ShapeManifold.Manifold;
+
+		TVector<T, d> VectorToPoint1 = Contact.Location - P0;
+		TVector<T, d> VectorToPoint2 = Contact.Location - P1;
+		TVector<T, d> Body1Velocity = V0 + TVector<T, d>::CrossProduct(W0, VectorToPoint1);
+		TVector<T, d> Body2Velocity = V1 + TVector<T, d>::CrossProduct(W1, VectorToPoint2);
+		TVector<T, d> RelativeVelocity = Body1Velocity - Body2Velocity;
+		if (TVector<T, d>::DotProduct(RelativeVelocity, Contact.Normal) < 0) // ignore separating constraints
+		{
+			PMatrix<T, d, d> WorldSpaceInvI1 = PBDRigid0 ? (Q0 * FMatrix::Identity).GetTransposed() * PBDRigid0->InvI() * (Q0 * FMatrix::Identity) : PMatrix<T, d, d>(0);
+			PMatrix<T, d, d> WorldSpaceInvI2 = PBDRigid1 ? (Q1 * FMatrix::Identity).GetTransposed() * PBDRigid1->InvI() * (Q1 * FMatrix::Identity) : PMatrix<T, d, d>(0);
+			PMatrix<T, d, d> Factor =
+				(PBDRigid0 ? ComputeFactorMatrix3(VectorToPoint1, WorldSpaceInvI1, PBDRigid0->InvM()) : PMatrix<T, d, d>(0)) +
+				(PBDRigid1 ? ComputeFactorMatrix3(VectorToPoint2, WorldSpaceInvI2, PBDRigid1->InvM()) : PMatrix<T, d, d>(0));
+			TVector<T, d> Impulse;
+			TVector<T, d> AngularImpulse(0);
+
+			// Resting contact if very close to the surface
+			T Restitution = (T)0;
+			T Friction = (T)0;
+			bool bApplyRestitution = (RelativeVelocity.Size() > (2 * 980 * Dt));
+			if (PhysicsMaterial0 && PhysicsMaterial1)
+			{
+				if (bApplyRestitution)
+				{
+					Restitution = FMath::Min(PhysicsMaterial0->Restitution, PhysicsMaterial1->Restitution);
+				}
+				Friction = FMath::Max(PhysicsMaterial0->Friction, PhysicsMaterial1->Friction);
+			}
+			else if (PhysicsMaterial0)
+			{
+				if (bApplyRestitution)
+				{
+					Restitution = PhysicsMaterial0->Restitution;
+				}
+				Friction = PhysicsMaterial0->Friction;
+			}
+			else if (PhysicsMaterial1)
 			{
-				RelativeNormalVelocity = 0;
+				if (bApplyRestitution)
+				{
+					Restitution = PhysicsMaterial1->Restitution;
+				}
+				Friction = PhysicsMaterial1->Friction;
 			}
-			TVector<T, d> VelocityChange = -(Restitution * RelativeNormalVelocity * Contact.Normal + RelativeVelocity);
-			T NormalVelocityChange = TVector<T, d>::DotProduct(VelocityChange, Contact.Normal);
-			PMatrix<T, d, d> FactorInverse = Factor.Inverse();
-			TVector<T, d> MinimalImpulse = FactorInverse * VelocityChange;
-			const T MinimalImpulseDotNormal = TVector<T, d>::DotProduct(MinimalImpulse, Contact.Normal);
-			const T TangentialSize = (MinimalImpulse - MinimalImpulseDotNormal * Contact.Normal).Size();
-			if (TangentialSize <= Friction * MinimalImpulseDotNormal)
+
+			if (Friction)
 			{
-				//within friction cone so just solve for static friction stopping the object
-				Impulse = MinimalImpulse;
-				if (MAngularFriction)
+				T RelativeNormalVelocity = TVector<T, d>::DotProduct(RelativeVelocity, Contact.Normal);
+				if (RelativeNormalVelocity > 0)
 				{
-					TVector<T, d> RelativeAngularVelocity = W0 - W1;
-					T AngularNormal = TVector<T, d>::DotProduct(RelativeAngularVelocity, Contact.Normal);
-					TVector<T, d> AngularTangent = RelativeAngularVelocity - AngularNormal * Contact.Normal;
-					TVector<T, d> FinalAngularVelocity = FMath::Sign(AngularNormal) * FMath::Max((T)0, FMath::Abs(AngularNormal) - MAngularFriction * NormalVelocityChange) * Contact.Normal + FMath::Max((T)0, AngularTangent.Size() - MAngularFriction * NormalVelocityChange) * AngularTangent.GetSafeNormal();
-					TVector<T, d> Delta = FinalAngularVelocity - RelativeAngularVelocity;
-					if (!PBDRigid0 && PBDRigid1)
-					{
-						PMatrix<T, d, d> WorldSpaceI2 = (Q1 * FMatrix::Identity) * PBDRigid1->I() * (Q1 * FMatrix::Identity).GetTransposed();
-						TVector<T, d> ImpulseDelta = PBDRigid1->M() * TVector<T, d>::CrossProduct(VectorToPoint2, Delta);
-						Impulse += ImpulseDelta;
-						AngularImpulse += WorldSpaceI2 * Delta - TVector<T, d>::CrossProduct(VectorToPoint2, ImpulseDelta);
-					}
-					else if (PBDRigid0 && !PBDRigid1)
+					RelativeNormalVelocity = 0;
+				}
+				TVector<T, d> VelocityChange = -(Restitution * RelativeNormalVelocity * Contact.Normal + RelativeVelocity);
+				T NormalVelocityChange = TVector<T, d>::DotProduct(VelocityChange, Contact.Normal);
+				PMatrix<T, d, d> FactorInverse = Factor.Inverse();
+				TVector<T, d> MinimalImpulse = FactorInverse * VelocityChange;
+				const T MinimalImpulseDotNormal = TVector<T, d>::DotProduct(MinimalImpulse, Contact.Normal);
+				const T TangentialSize = (MinimalImpulse - MinimalImpulseDotNormal * Contact.Normal).Size();
+				if (TangentialSize <= Friction * MinimalImpulseDotNormal)
+				{
+					//within friction cone so just solve for static friction stopping the object
+					Impulse = MinimalImpulse;
+					if (MAngularFriction)
 					{
-						PMatrix<T, d, d> WorldSpaceI1 = (Q0 * FMatrix::Identity) * PBDRigid0->I() * (Q0 * FMatrix::Identity).GetTransposed();
-						TVector<T, d> ImpulseDelta = PBDRigid0->M() * TVector<T, d>::CrossProduct(VectorToPoint1, Delta);
-						Impulse += ImpulseDelta;
-						AngularImpulse += WorldSpaceI1 * Delta - TVector<T, d>::CrossProduct(VectorToPoint1, ImpulseDelta);
+						TVector<T, d> RelativeAngularVelocity = W0 - W1;
+						T AngularNormal = TVector<T, d>::DotProduct(RelativeAngularVelocity, Contact.Normal);
+						TVector<T, d> AngularTangent = RelativeAngularVelocity - AngularNormal * Contact.Normal;
+						TVector<T, d> FinalAngularVelocity = FMath::Sign(AngularNormal) * FMath::Max((T)0, FMath::Abs(AngularNormal) - MAngularFriction * NormalVelocityChange) * Contact.Normal + FMath::Max((T)0, AngularTangent.Size() - MAngularFriction * NormalVelocityChange) * AngularTangent.GetSafeNormal();
+						TVector<T, d> Delta = FinalAngularVelocity - RelativeAngularVelocity;
+						if (!PBDRigid0 && PBDRigid1)
+						{
+							PMatrix<T, d, d> WorldSpaceI2 = (Q1 * FMatrix::Identity) * PBDRigid1->I() * (Q1 * FMatrix::Identity).GetTransposed();
+							TVector<T, d> ImpulseDelta = PBDRigid1->M() * TVector<T, d>::CrossProduct(VectorToPoint2, Delta);
+							Impulse += ImpulseDelta;
+							AngularImpulse += WorldSpaceI2 * Delta - TVector<T, d>::CrossProduct(VectorToPoint2, ImpulseDelta);
+						}
+						else if (PBDRigid0 && !PBDRigid1)
+						{
+							PMatrix<T, d, d> WorldSpaceI1 = (Q0 * FMatrix::Identity) * PBDRigid0->I() * (Q0 * FMatrix::Identity).GetTransposed();
+							TVector<T, d> ImpulseDelta = PBDRigid0->M() * TVector<T, d>::CrossProduct(VectorToPoint1, Delta);
+							Impulse += ImpulseDelta;
+							AngularImpulse += WorldSpaceI1 * Delta - TVector<T, d>::CrossProduct(VectorToPoint1, ImpulseDelta);
+						}
+						else if (PBDRigid0 && PBDRigid1)
+						{
+							PMatrix<T, d, d> Cross1(0, VectorToPoint1.Z, -VectorToPoint1.Y, -VectorToPoint1.Z, 0, VectorToPoint1.X, VectorToPoint1.Y, -VectorToPoint1.X, 0);
+							PMatrix<T, d, d> Cross2(0, VectorToPoint2.Z, -VectorToPoint2.Y, -VectorToPoint2.Z, 0, VectorToPoint2.X, VectorToPoint2.Y, -VectorToPoint2.X, 0);
+							PMatrix<T, d, d> CrossI1 = Cross1 * WorldSpaceInvI1;
+							PMatrix<T, d, d> CrossI2 = Cross2 * WorldSpaceInvI2;
+							PMatrix<T, d, d> Diag1 = CrossI1 * Cross1.GetTransposed() + CrossI2 * Cross2.GetTransposed();
+							Diag1.M[0][0] += PBDRigid0->InvM() + PBDRigid1->InvM();
+							Diag1.M[1][1] += PBDRigid0->InvM() + PBDRigid1->InvM();
+							Diag1.M[2][2] += PBDRigid0->InvM() + PBDRigid1->InvM();
+							PMatrix<T, d, d> OffDiag1 = (CrossI1 + CrossI2) * -1;
+							PMatrix<T, d, d> Diag2 = (WorldSpaceInvI1 + WorldSpaceInvI2).Inverse();
+							PMatrix<T, d, d> OffDiag1Diag2 = OffDiag1 * Diag2;
+							TVector<T, d> ImpulseDelta = PMatrix<T, d, d>((Diag1 - OffDiag1Diag2 * OffDiag1.GetTransposed()).Inverse())* ((OffDiag1Diag2 * -1) * Delta);
+							Impulse += ImpulseDelta;
+							AngularImpulse += Diag2 * (Delta - PMatrix<T, d, d>(OffDiag1.GetTransposed()) * ImpulseDelta);
+						}
 					}
-					else if (PBDRigid0 && PBDRigid1)
+				}
+				else
+				{
+					//outside friction cone, solve for normal relative velocity and keep tangent at cone edge
+					TVector<T, d> Tangent = (RelativeVelocity - TVector<T, d>::DotProduct(RelativeVelocity, Contact.Normal) * Contact.Normal).GetSafeNormal();
+					TVector<T, d> DirectionalFactor = Factor * (Contact.Normal - Friction * Tangent);
+					T ImpulseDenominator = TVector<T, d>::DotProduct(Contact.Normal, DirectionalFactor);
+					if (!ensureMsgf(FMath::Abs(ImpulseDenominator) > SMALL_NUMBER, TEXT("Constraint:%s\n\nParticle:%s\n\nLevelset:%s\n\nDirectionalFactor:%s, ImpulseDenominator:%f"),
+						*Constraint.ToString(),
+						*Particle0->ToString(),
+						*Particle1->ToString(),
+						*DirectionalFactor.ToString(), ImpulseDenominator))
 					{
-						PMatrix<T, d, d> Cross1(0, VectorToPoint1.Z, -VectorToPoint1.Y, -VectorToPoint1.Z, 0, VectorToPoint1.X, VectorToPoint1.Y, -VectorToPoint1.X, 0);
-						PMatrix<T, d, d> Cross2(0, VectorToPoint2.Z, -VectorToPoint2.Y, -VectorToPoint2.Z, 0, VectorToPoint2.X, VectorToPoint2.Y, -VectorToPoint2.X, 0);
-						PMatrix<T, d, d> CrossI1 = Cross1 * WorldSpaceInvI1;
-						PMatrix<T, d, d> CrossI2 = Cross2 * WorldSpaceInvI2;
-						PMatrix<T, d, d> Diag1 = CrossI1 * Cross1.GetTransposed() + CrossI2 * Cross2.GetTransposed();
-						Diag1.M[0][0] += PBDRigid0->InvM() + PBDRigid1->InvM();
-						Diag1.M[1][1] += PBDRigid0->InvM() + PBDRigid1->InvM();
-						Diag1.M[2][2] += PBDRigid0->InvM() + PBDRigid1->InvM();
-						PMatrix<T, d, d> OffDiag1 = (CrossI1 + CrossI2) * -1;
-						PMatrix<T, d, d> Diag2 = (WorldSpaceInvI1 + WorldSpaceInvI2).Inverse();
-						PMatrix<T, d, d> OffDiag1Diag2 = OffDiag1 * Diag2;
-						TVector<T, d> ImpulseDelta = PMatrix<T, d, d>((Diag1 - OffDiag1Diag2 * OffDiag1.GetTransposed()).Inverse())* ((OffDiag1Diag2 * -1) * Delta);
-						Impulse += ImpulseDelta;
-						AngularImpulse += Diag2 * (Delta - PMatrix<T, d, d>(OffDiag1.GetTransposed()) * ImpulseDelta);
+						ImpulseDenominator = (T)1;
 					}
+
+					const T ImpulseMag = -(1 + Restitution) * RelativeNormalVelocity / ImpulseDenominator;
+					Impulse = ImpulseMag * (Contact.Normal - Friction * Tangent);
 				}
 			}
 			else
 			{
-				//outside friction cone, solve for normal relative velocity and keep tangent at cone edge
-				TVector<T, d> Tangent = (RelativeVelocity - TVector<T, d>::DotProduct(RelativeVelocity, Contact.Normal) * Contact.Normal).GetSafeNormal();
-				TVector<T, d> DirectionalFactor = Factor * (Contact.Normal - Friction * Tangent);
-				T ImpulseDenominator = TVector<T, d>::DotProduct(Contact.Normal, DirectionalFactor);
-				if (!ensureMsgf(FMath::Abs(ImpulseDenominator) > SMALL_NUMBER, TEXT("Constraint:%s\n\nParticle:%s\n\nLevelset:%s\n\nDirectionalFactor:%s, ImpulseDenominator:%f"),
+				T ImpulseDenominator = TVector<T, d>::DotProduct(Contact.Normal, Factor * Contact.Normal);
+				TVector<T, d> ImpulseNumerator = -(1 + Restitution) * TVector<T, d>::DotProduct(RelativeVelocity, Contact.Normal)* Contact.Normal;
+				if (!ensureMsgf(FMath::Abs(ImpulseDenominator) > SMALL_NUMBER, TEXT("Constraint:%s\n\nParticle:%s\n\nLevelset:%s\n\nFactor*Constraint.Normal:%s, ImpulseDenominator:%f"),
 					*Constraint.ToString(),
 					*Particle0->ToString(),
 					*Particle1->ToString(),
-					*DirectionalFactor.ToString(), ImpulseDenominator))
+					*(Factor * Contact.Normal).ToString(), ImpulseDenominator))
 				{
 					ImpulseDenominator = (T)1;
 				}
-
-				const T ImpulseMag = -(1 + Restitution) * RelativeNormalVelocity / ImpulseDenominator;
-				Impulse = ImpulseMag * (Contact.Normal - Friction * Tangent);
-			}
-		}
-		else
-		{
-			T ImpulseDenominator = TVector<T, d>::DotProduct(Contact.Normal, Factor * Contact.Normal);
-			TVector<T, d> ImpulseNumerator = -(1 + Restitution) * TVector<T, d>::DotProduct(RelativeVelocity, Contact.Normal)* Contact.Normal;
-			if (!ensureMsgf(FMath::Abs(ImpulseDenominator) > SMALL_NUMBER, TEXT("Constraint:%s\n\nParticle:%s\n\nLevelset:%s\n\nFactor*Constraint.Normal:%s, ImpulseDenominator:%f"),
-				*Constraint.ToString(),
-				*Particle0->ToString(),
-				*Particle1->ToString(),
-				*(Factor * Contact.Normal).ToString(), ImpulseDenominator))
-			{
-				ImpulseDenominator = (T)1;
-			}
-			Impulse = ImpulseNumerator / ImpulseDenominator;
-		}
-		Impulse = GetEnergyClampedImpulse(Constraint, Impulse, VectorToPoint1, VectorToPoint2, Body1Velocity, Body2Velocity);
-		Constraint.AccumulatedImpulse += Impulse;
-		if (PBDRigid0)
-		{
-			// Velocity update for next step
-			TVector<T, d> AngularImpulse1 = TVector<T, d>::CrossProduct(VectorToPoint1, Impulse) + AngularImpulse;
-			FVec3 DV1 = PBDRigid0->InvM() * Impulse;
-			FVec3 DW1 = WorldSpaceInvI1 * AngularImpulse1;
-			PBDRigid0->V() += DV1;
-			PBDRigid0->W() += DW1;
-			// Position update as part of pbd
-			PBDRigid0->P() += DV1 * Dt;
-			PBDRigid0->Q() += TRotation<T, d>::FromElements(DW1, 0.f) * Q0 * Dt * T(0.5);
-			PBDRigid0->Q().Normalize();
+				Impulse = ImpulseNumerator / ImpulseDenominator;
+			}
+			Impulse = GetEnergyClampedImpulse(Constraint, Impulse, VectorToPoint1, VectorToPoint2, Body1Velocity, Body2Velocity);
+			Constraint.AccumulatedImpulse += Impulse;
+			if (PBDRigid0)
+			{
+				// Velocity update for next step
+				TVector<T, d> AngularImpulse1 = TVector<T, d>::CrossProduct(VectorToPoint1, Impulse) + AngularImpulse;
+				FVec3 DV1 = PBDRigid0->InvM() * Impulse;
+				FVec3 DW1 = WorldSpaceInvI1 * AngularImpulse1;
+				PBDRigid0->V() += DV1;
+				PBDRigid0->W() += DW1;
+				// Position update as part of pbd
+				PBDRigid0->P() += DV1 * Dt;
+				PBDRigid0->Q() += TRotation<T, d>::FromElements(DW1, 0.f) * Q0 * Dt * T(0.5);
+				PBDRigid0->Q().Normalize();
+			}
+			if (PBDRigid1)
+			{
+				// Velocity update for next step
+				TVector<T, d> AngularImpulse2 = TVector<T, d>::CrossProduct(VectorToPoint2, -Impulse) - AngularImpulse;
+				FVec3 DV2 = -PBDRigid1->InvM() * Impulse;
+				FVec3 DW2 = WorldSpaceInvI2 * AngularImpulse2;
+				PBDRigid1->V() += DV2;
+				PBDRigid1->W() += DW2;
+				// Position update as part of pbd
+				PBDRigid1->P() += DV2 * Dt;
+				PBDRigid1->Q() += TRotation<T, d>::FromElements(DW2, 0.f) * Q1 * Dt * T(0.5);
+				PBDRigid1->Q().Normalize();
+			}
+		}
+	}
+
+	DECLARE_CYCLE_STAT(TEXT("TPBDCollisionConstraint::Apply"), STAT_Apply, STATGROUP_Chaos);
+	template<typename T, int d>
+	void TPBDCollisionConstraint<T, d>::Apply(const T Dt, const TArray<FConstraintContainerHandle*>& InConstraintHandles, const int32 It, const int32 NumIts)
+	{
+		SCOPE_CYCLE_COUNTER(STAT_Apply);
+		if (bEnableVelocitySolve)
+		{
+			PhysicsParallelFor(InConstraintHandles.Num(), [&](int32 ConstraintHandleIndex) {
+				FConstraintContainerHandle* ConstraintHandle = InConstraintHandles[ConstraintHandleIndex];
+				check(ConstraintHandle != nullptr);
+				Apply(Dt, Constraints[ConstraintHandle->GetConstraintIndex()]);
+			}, bDisableCollisionParallelFor);
 		}
-		if (PBDRigid1)
+
+		if (PostApplyCallback != nullptr)
 		{
-			// Velocity update for next step
-			TVector<T, d> AngularImpulse2 = TVector<T, d>::CrossProduct(VectorToPoint2, -Impulse) - AngularImpulse;
-			FVec3 DV2 = -PBDRigid1->InvM() * Impulse;
-			FVec3 DW2 = WorldSpaceInvI2 * AngularImpulse2;
-			PBDRigid1->V() += DV2;
-			PBDRigid1->W() += DW2;
-			// Position update as part of pbd
-			PBDRigid1->P() += DV2 * Dt;
-			PBDRigid1->Q() += TRotation<T, d>::FromElements(DW2, 0.f) * Q1 * Dt * T(0.5);
-			PBDRigid1->Q().Normalize();
+			PostApplyCallback(Dt, InConstraintHandles);
 		}
 	}
-}
-
-DECLARE_CYCLE_STAT(TEXT("TPBDCollisionConstraint::Apply"), STAT_Apply, STATGROUP_Chaos);
-template<typename T, int d>
-void TPBDCollisionConstraint<T, d>::Apply(const T Dt, const TArray<FConstraintContainerHandle*>& InConstraintHandles, const int32 It, const int32 NumIts)
-{
-	SCOPE_CYCLE_COUNTER(STAT_Apply);
-	if (bEnableVelocitySolve)
-	{
-		PhysicsParallelFor(InConstraintHandles.Num(), [&](int32 ConstraintHandleIndex) {
-			FConstraintContainerHandle* ConstraintHandle = InConstraintHandles[ConstraintHandleIndex];
-			check(ConstraintHandle != nullptr);
-			Apply(Dt, Constraints[ConstraintHandle->GetConstraintIndex()]);
-			}, bDisableCollisionParallelFor);
-	}
-
-	if (PostApplyCallback != nullptr)
-	{
-		PostApplyCallback(Dt, InConstraintHandles);
-	}
-}
-
-
-template<typename T, int d>
-void TPBDCollisionConstraint<T, d>::ApplyPushOut(const T Dt, FRigidBodyContactConstraint& Constraint, const TSet<TGeometryParticleHandle<T,d>*>& IsTemporarilyStatic, int32 Iteration, int32 NumIterations, bool &NeedsAnotherIteration)
-{
-	TGeometryParticleHandle<T, d>* Particle0 = Constraint.Particle[0];
-	TGeometryParticleHandle<T, d>* Particle1 = Constraint.Particle[1];
-	TPBDRigidParticleHandle<T, d>* PBDRigid0 = Particle0->AsDynamic();
-	TPBDRigidParticleHandle<T, d>* PBDRigid1 = Particle1->AsDynamic();
 
-	// @todo(ccaulfield): I think we should never get this? Revisit after particle handle refactor
-	if (Particle0->Sleeping())
-	{
-		ensure(!PBDRigid1 || PBDRigid1->Sleeping());
-		return;
-	}
-	if (Particle1->Sleeping())
-	{
-		ensure(!PBDRigid0 || PBDRigid0->Sleeping());
-		return;
-	}
 
-	const TVector<T, d> ZeroVector = TVector<T, d>(0);
-	const TRotation<T, d>& Q0 = PBDRigid0 ? PBDRigid0->Q() : Particle0->R();
-	const TRotation<T, d>& Q1 = PBDRigid1 ? PBDRigid1->Q() : Particle1->R();
-	const TVector<T, d>& P0 = PBDRigid0 ? PBDRigid0->P() : Particle0->X();
-	const TVector<T, d>& P1 = PBDRigid1 ? PBDRigid1->P() : Particle1->X();
-	const TVector<T, d>& V0 = PBDRigid0 ? PBDRigid0->V() : ZeroVector;
-	const TVector<T, d>& V1 = PBDRigid1 ? PBDRigid1->V() : ZeroVector;
-	const TVector<T, d>& W0 = PBDRigid0 ? PBDRigid0->W() : ZeroVector;
-	const TVector<T, d>& W1 = PBDRigid1 ? PBDRigid1->W() : ZeroVector;
-	TSerializablePtr<FChaosPhysicsMaterial> PhysicsMaterial0 = Particle0->AuxilaryValue(MPhysicsMaterials);
-	TSerializablePtr<FChaosPhysicsMaterial> PhysicsMaterial1 = Particle1->AuxilaryValue(MPhysicsMaterials);
-	const bool IsTemporarilyStatic0 = IsTemporarilyStatic.Contains(Particle0);
-	const bool IsTemporarilyStatic1 = IsTemporarilyStatic.Contains(Particle1);
-
-	for (int32 PairIteration = 0; PairIteration < MPairIterations; ++PairIteration)
+	template<typename T, int d>
+	void TPBDCollisionConstraint<T, d>::ApplyPushOut(const T Dt, FRigidBodyContactConstraint& Constraint, const TSet<TGeometryParticleHandle<T, d>*>& IsTemporarilyStatic, int32 Iteration, int32 NumIterations, bool &NeedsAnotherIteration)
 	{
-		UpdateConstraint<ECollisionUpdateType::Deepest>(MThickness, Constraint);
+		TGeometryParticleHandle<T, d>* Particle0 = Constraint.Particle[0];
+		TGeometryParticleHandle<T, d>* Particle1 = Constraint.Particle[1];
+		TPBDRigidParticleHandle<T, d>* PBDRigid0 = Particle0->AsDynamic();
+		TPBDRigidParticleHandle<T, d>* PBDRigid1 = Particle1->AsDynamic();
 
-		const TContactData<T,d> & Contact = Constraint.ShapeManifold.Manifold;
-
-		if (Contact.Phi >= MThickness)
+		// @todo(ccaulfield): I think we should never get this? Revisit after particle handle refactor
+		if (Particle0->Sleeping())
 		{
-			break;
+			ensure(!PBDRigid1 || PBDRigid1->Sleeping());
+			return;
 		}
-
-		if ((!PBDRigid0 || IsTemporarilyStatic0) && (!PBDRigid1 || IsTemporarilyStatic1))
+		if (Particle1->Sleeping())
 		{
-			break;
+			ensure(!PBDRigid0 || PBDRigid0->Sleeping());
+			return;
 		}
 
-		NeedsAnotherIteration = true;
-		PMatrix<T, d, d> WorldSpaceInvI1 = PBDRigid0? (Q0 * FMatrix::Identity).GetTransposed() * PBDRigid0->InvI() * (Q0 * FMatrix::Identity) : PMatrix<T, d, d>(0);
-		PMatrix<T, d, d> WorldSpaceInvI2 = PBDRigid1 ? (Q1 * FMatrix::Identity).GetTransposed() * PBDRigid1->InvI() * (Q1 * FMatrix::Identity) : PMatrix<T, d, d>(0);
-		TVector<T, d> VectorToPoint1 = Contact.Location - P0;
-		TVector<T, d> VectorToPoint2 = Contact.Location - P1;
-		PMatrix<T, d, d> Factor =
-			(PBDRigid0 ? ComputeFactorMatrix3(VectorToPoint1, WorldSpaceInvI1, PBDRigid0->InvM()) : PMatrix<T, d, d>(0)) +
-			(PBDRigid1 ? ComputeFactorMatrix3(VectorToPoint2, WorldSpaceInvI2, PBDRigid1->InvM()) : PMatrix<T, d, d>(0));
-		T Numerator = FMath::Min((T)(Iteration + 2), (T)NumIterations);
-		T ScalingFactor = Numerator / (T)NumIterations;
+		const TVector<T, d> ZeroVector = TVector<T, d>(0);
+		const TRotation<T, d>& Q0 = PBDRigid0 ? PBDRigid0->Q() : Particle0->R();
+		const TRotation<T, d>& Q1 = PBDRigid1 ? PBDRigid1->Q() : Particle1->R();
+		const TVector<T, d>& P0 = PBDRigid0 ? PBDRigid0->P() : Particle0->X();
+		const TVector<T, d>& P1 = PBDRigid1 ? PBDRigid1->P() : Particle1->X();
+		const TVector<T, d>& V0 = PBDRigid0 ? PBDRigid0->V() : ZeroVector;
+		const TVector<T, d>& V1 = PBDRigid1 ? PBDRigid1->V() : ZeroVector;
+		const TVector<T, d>& W0 = PBDRigid0 ? PBDRigid0->W() : ZeroVector;
+		const TVector<T, d>& W1 = PBDRigid1 ? PBDRigid1->W() : ZeroVector;
+		TSerializablePtr<FChaosPhysicsMaterial> PhysicsMaterial0 = Particle0->AuxilaryValue(MPhysicsMaterials);
+		TSerializablePtr<FChaosPhysicsMaterial> PhysicsMaterial1 = Particle1->AuxilaryValue(MPhysicsMaterials);
+		const bool IsTemporarilyStatic0 = IsTemporarilyStatic.Contains(Particle0);
+		const bool IsTemporarilyStatic1 = IsTemporarilyStatic.Contains(Particle1);
 
-		//if pushout is needed we better fix relative velocity along normal. Treat it as if 0 restitution
-		TVector<T, d> Body1Velocity = V0 + TVector<T, d>::CrossProduct(W0, VectorToPoint1);
-		TVector<T, d> Body2Velocity = V1 + TVector<T, d>::CrossProduct(W1, VectorToPoint2);
-		TVector<T, d> RelativeVelocity = Body1Velocity - Body2Velocity;
-		const T RelativeVelocityDotNormal = TVector<T, d>::DotProduct(RelativeVelocity, Contact.Normal);
-		if (RelativeVelocityDotNormal < 0)
-		{
-			T ImpulseDenominator = TVector<T, d>::DotProduct(Contact.Normal, Factor * Contact.Normal);
-			TVector<T, d> ImpulseNumerator = -TVector<T, d>::DotProduct(RelativeVelocity, Contact.Normal) * Contact.Normal * ScalingFactor;
-			if (!ensureMsgf(FMath::Abs(ImpulseDenominator) > SMALL_NUMBER, TEXT("ApplyPushout Constraint:%s\n\nParticle:%s\n\nLevelset:%s\n\nFactor*Contact.Normal:%s, ImpulseDenominator:%f"),
-				*Constraint.ToString(),
-				*Particle0->ToString(),
-				*Particle1->ToString(),
-				*(Factor*Contact.Normal).ToString(), ImpulseDenominator))
-			{
-				ImpulseDenominator = (T)1;
-			}
-
-			TVector<T, d> VelocityFixImpulse = ImpulseNumerator / ImpulseDenominator;
-			VelocityFixImpulse = GetEnergyClampedImpulse(Constraint, VelocityFixImpulse, VectorToPoint1, VectorToPoint2, Body1Velocity, Body2Velocity);
-			Constraint.AccumulatedImpulse += VelocityFixImpulse;	//question: should we track this?
-			if (!IsTemporarilyStatic0 && PBDRigid0)
-			{
-				TVector<T, d> AngularImpulse = TVector<T, d>::CrossProduct(VectorToPoint1, VelocityFixImpulse);
-				PBDRigid0->V() += PBDRigid0->InvM() * VelocityFixImpulse;
-				PBDRigid0->W() += WorldSpaceInvI1 * AngularImpulse;
-
-			}
-
-			if (!IsTemporarilyStatic1 && PBDRigid1)
-			{
-				TVector<T, d> AngularImpulse = TVector<T, d>::CrossProduct(VectorToPoint2, -VelocityFixImpulse);
-				PBDRigid1->V() -= PBDRigid1->InvM() * VelocityFixImpulse;
-				PBDRigid1->W() += WorldSpaceInvI2 * AngularImpulse;
-			}
-
-		}
-
-
-		TVector<T, d> Impulse = PMatrix<T, d, d>(Factor.Inverse()) * ((-Contact.Phi + MThickness) * ScalingFactor * Contact.Normal);
-		TVector<T, d> AngularImpulse1 = TVector<T, d>::CrossProduct(VectorToPoint1, Impulse);
-		TVector<T, d> AngularImpulse2 = TVector<T, d>::CrossProduct(VectorToPoint2, -Impulse);
-		if (!IsTemporarilyStatic0 && PBDRigid0)
-		{
-			PBDRigid0->P() += PBDRigid0->InvM() * Impulse;
-			PBDRigid0->Q() = TRotation<T, d>::FromVector(WorldSpaceInvI1 * AngularImpulse1) * Q0;
-			PBDRigid0->Q().Normalize();
-		}
-		if (!IsTemporarilyStatic1 && PBDRigid1)
-		{
-			PBDRigid1->P() -= PBDRigid1->InvM() * Impulse;
-			PBDRigid1->Q() = TRotation<T, d>::FromVector(WorldSpaceInvI2 * AngularImpulse2) * Q1;
-			PBDRigid1->Q().Normalize();
-		}
-	}
-}
-
-DECLARE_CYCLE_STAT(TEXT("TPBDCollisionConstraint::ApplyPushOut"), STAT_ApplyPushOut, STATGROUP_Chaos);
-template<typename T, int d>
-bool TPBDCollisionConstraint<T, d>::ApplyPushOut(const T Dt, const TArray<FConstraintContainerHandle*>& InConstraintHandles, const TSet<TGeometryParticleHandle<T,d>*>& IsTemporarilyStatic, int32 Iteration, int32 NumIterations)
-{
-	SCOPE_CYCLE_COUNTER(STAT_ApplyPushOut);
-
-	bool NeedsAnotherIteration = false;
-
-	if (MPairIterations > 0)
-	{
-		PhysicsParallelFor(InConstraintHandles.Num(), [&](int32 ConstraintHandleIndex) {
-			FConstraintContainerHandle* ConstraintHandle = InConstraintHandles[ConstraintHandleIndex];
-			check(ConstraintHandle != nullptr);
-			ApplyPushOut(Dt, Constraints[ConstraintHandle->GetConstraintIndex()], IsTemporarilyStatic, Iteration, NumIterations, NeedsAnotherIteration);
-			}, bDisableCollisionParallelFor);
-	}
-
-	if (PostApplyPushOutCallback != nullptr)
-	{
-		PostApplyPushOutCallback(Dt, InConstraintHandles, NeedsAnotherIteration);
-	}
-
-	return NeedsAnotherIteration;
-}
-
-
-template<typename T, int d>
-TRigidTransform<T, d> GetTransform(const TGeometryParticleHandle<T, d>* Particle)
-{
-	//#if CHAOS_PARTICLEHANDLE_TODO
-	
-	TGenericParticleHandle<T, d> Generic = const_cast<TGeometryParticleHandle<T, d>*>(Particle);	//TODO: give a const version of the generic API
-	return TRigidTransform<T, d>(Generic->P(), Generic->Q());
-}
-
-#if 0
-template<typename T, int d>
-void UpdateLevelsetConstraintHelperCCD(const TRigidParticles<T, d>& InParticles, const int32 j, const TRigidTransform<T, d>& LocalToWorld1, const TRigidTransform<T, d>& LocalToWorld2, const T Thickness, TRigidBodyContactConstraint<T, d>& Constraint)
-{
-	if (InParticles.CollisionParticles(Constraint.ParticleIndex))
-	{
-		const TRigidTransform<T, d> PreviousLocalToWorld1 = GetTransform(InParticles, Constraint.ParticleIndex);
-		TVector<T, d> WorldSpacePointStart = PreviousLocalToWorld1.TransformPosition(InParticles.CollisionParticles(Constraint.ParticleIndex)->X(j));
-		TVector<T, d> WorldSpacePointEnd = LocalToWorld1.TransformPosition(InParticles.CollisionParticles(Constraint.ParticleIndex)->X(j));
-		TVector<T, d> Body2SpacePointStart = LocalToWorld2.InverseTransformPosition(WorldSpacePointStart);
-		TVector<T, d> Body2SpacePointEnd = LocalToWorld2.InverseTransformPosition(WorldSpacePointEnd);
-		Pair<TVector<T, d>, bool> PointPair = InParticles.Geometry(Constraint.LevelsetIndex)->FindClosestIntersection(Body2SpacePointStart, Body2SpacePointEnd, Thickness);
-		if (PointPair.Second)
+		for (int32 PairIteration = 0; PairIteration < MPairIterations; ++PairIteration)
 		{
-			const TVector<T, d> WorldSpaceDelta = WorldSpacePointEnd - TVector<T, d>(LocalToWorld2.TransformPosition(PointPair.First));
-			Contact.Phi = -WorldSpaceDelta.Size();
-			Contact.Normal = LocalToWorld2.TransformVector(InParticles.Geometry(Constraint.LevelsetIndex)->Normal(PointPair.First));
-			// @todo(mlentine): Should we be using the actual collision point or that point evolved to the current time step?
-			Contact.Location = WorldSpacePointEnd;
-		}
-	}
-}
-#endif
-
-template <typename T, int d>
-bool SampleObjectHelper2(const FImplicitObject& Object, const TRigidTransform<T,d>& ObjectTransform, const TRigidTransform<T, d>& SampleToObjectTransform, const TVector<T, d>& SampleParticle, T Thickness, TRigidBodyContactConstraint<T, d>& Constraint)
-{
-	TVector<T, d> LocalPoint = SampleToObjectTransform.TransformPositionNoScale(SampleParticle);
-	TVector<T, d> LocalNormal;
-	T LocalPhi = Object.PhiWithNormal(LocalPoint, LocalNormal);
-
-	TContactData<T,d> & Contact = Constraint.ShapeManifold.Manifold;
-	if (LocalPhi < Contact.Phi)
-	{
-		Contact.Phi = LocalPhi;
-		Contact.Normal = ObjectTransform.TransformVectorNoScale(LocalNormal);
-		Contact.Location = ObjectTransform.TransformPositionNoScale(LocalPoint);
-		return true;
-	}
-	return false;
-}
-
-template <typename T, int d>
-bool SampleObjectNoNormal2(const FImplicitObject& Object, const TRigidTransform<T,d>& ObjectTransform, const TRigidTransform<T, d>& SampleToObjectTransform, const TVector<T, d>& SampleParticle, T Thickness, TRigidBodyContactConstraint<T, d>& Constraint)
-{
-	TVector<T, d> LocalPoint = SampleToObjectTransform.TransformPositionNoScale(SampleParticle);
-	TVector<T, d> LocalNormal;
-	T LocalPhi = Object.PhiWithNormal(LocalPoint, LocalNormal);
-
-	TContactData<T,d> & Contact = Constraint.ShapeManifold.Manifold;
-	if (LocalPhi < Contact.Phi)
-	{
-		Contact.Phi = LocalPhi;
-		return true;
-	}
-	return false;
-}
-
-template <typename T, int d>
-bool SampleObjectNormalAverageHelper2(const FImplicitObject& Object, const TRigidTransform<T, d>& ObjectTransform, const TRigidTransform<T, d>& SampleToObjectTransform, const TVector<T, d>& SampleParticle, T Thickness, T& TotalThickness, TRigidBodyContactConstraint<T, d>& Constraint)
-{
-	TVector<T, d> LocalPoint = SampleToObjectTransform.TransformPositionNoScale(SampleParticle);
-	TVector<T, d> LocalNormal;
-	T LocalPhi = Object.PhiWithNormal(LocalPoint, LocalNormal);
-	T LocalThickness = LocalPhi - Thickness;
-
-	TContactData<T,d> & Contact = Constraint.ShapeManifold.Manifold;
-	if (LocalThickness < -KINDA_SMALL_NUMBER)
-	{
-		Contact.Location += LocalPoint * LocalThickness;
-		TotalThickness += LocalThickness;
-		return true;
-	}
-	return false;
-}
-
-DECLARE_CYCLE_STAT(TEXT("TPBDCollisionConstraint::UpdateLevelsetPartial"), STAT_UpdateLevelsetPartial, STATGROUP_ChaosWide);
-DECLARE_CYCLE_STAT(TEXT("TPBDCollisionConstraint::UpdateLevelsetFindParticles"), STAT_UpdateLevelsetFindParticles, STATGROUP_ChaosWide);
-DECLARE_CYCLE_STAT(TEXT("TPBDCollisionConstraint::UpdateLevelsetBVHTraversal"), STAT_UpdateLevelsetBVHTraversal, STATGROUP_ChaosWide);
-DECLARE_CYCLE_STAT(TEXT("TPBDCollisionConstraint::UpdateLevelsetSignedDistance"), STAT_UpdateLevelsetSignedDistance, STATGROUP_ChaosWide);
-DECLARE_CYCLE_STAT(TEXT("TPBDCollisionConstraint::UpdateLevelsetAll"), STAT_UpdateLevelsetAll, STATGROUP_ChaosWide);
-DECLARE_CYCLE_STAT(TEXT("TPBDCollisionConstraint::SampleObject"), STAT_SampleObject, STATGROUP_ChaosWide);
-
-int32 NormalAveraging2 = 1;
-FAutoConsoleVariableRef CVarNormalAveraging2(TEXT("p.NormalAveraging2"), NormalAveraging2, TEXT(""));
-
-int32 SampleMinParticlesForAcceleration2 = 2048;
-FAutoConsoleVariableRef CVarSampleMinParticlesForAcceleration2(TEXT("p.SampleMinParticlesForAcceleration2"), SampleMinParticlesForAcceleration2, TEXT("The minimum number of particles needed before using an acceleration structure when sampling"));
-
-
-template <ECollisionUpdateType UpdateType, typename T, int d>
-void SampleObject2(const FImplicitObject& Object, const TRigidTransform<T, d>& ObjectTransform, const TBVHParticles<T, d>& SampleParticles, const TRigidTransform<T, d>& SampleParticlesTransform, T Thickness, TRigidBodyContactConstraint<T, d>& Constraint)
-{
-	SCOPE_CYCLE_COUNTER(STAT_SampleObject);
-	TRigidBodyContactConstraint<float, 3> AvgConstraint = Constraint;
+			UpdateConstraint<ECollisionUpdateType::Deepest>(MThickness, Constraint);
 
-	TContactData<float, 3> & Contact = Constraint.ShapeManifold.Manifold;
-	TContactData<float, 3> & AvgContact = AvgConstraint.ShapeManifold.Manifold;
+			const TContactData<T, d> & Contact = Constraint.ShapeManifold.Manifold;
 
-	AvgConstraint.Particle[0] = Constraint.Particle[0];
-	AvgConstraint.Particle[1] = Constraint.Particle[1];
-	AvgContact.Location = TVector<float, 3>::ZeroVector;
-	AvgContact.Normal = TVector<float, 3>::ZeroVector;
-	AvgContact.Phi = Thickness;
-	float TotalThickness = float(0);
-
-	int32 DeepestParticle = -1;
-	const int32 NumParticles = SampleParticles.Size();
-
-	const TRigidTransform<T, d> & SampleToObjectTM = SampleParticlesTransform.GetRelativeTransform(ObjectTransform);
-	if (NumParticles > SampleMinParticlesForAcceleration2 && Object.HasBoundingBox())
-	{
-		SCOPE_CYCLE_COUNTER(STAT_UpdateLevelsetPartial);
-		TBox<T, d> ImplicitBox = Object.BoundingBox().TransformedBox(ObjectTransform.GetRelativeTransform(SampleParticlesTransform));
-		ImplicitBox.Thicken(Thickness);
-		TArray<int32> PotentialParticles;
-		{
-			SCOPE_CYCLE_COUNTER(STAT_UpdateLevelsetFindParticles);
-			PotentialParticles = SampleParticles.FindAllIntersections(ImplicitBox);
-		}
-		{
-			SCOPE_CYCLE_COUNTER(STAT_UpdateLevelsetSignedDistance);
-			for (int32 i : PotentialParticles)
-			{
-				if (NormalAveraging2 && UpdateType != ECollisionUpdateType::Any)	//if we just want one don't bother with normal
-				{
-					SampleObjectNormalAverageHelper2(Object, ObjectTransform, SampleToObjectTM, SampleParticles.X(i), Thickness, TotalThickness, AvgConstraint);
-				}
-				else
-				{
-					if (SampleObjectNoNormal2(Object, ObjectTransform, SampleToObjectTM, SampleParticles.X(i), Thickness, AvgConstraint))
-					{
-						DeepestParticle = i;
-						if (UpdateType == ECollisionUpdateType::Any)
-						{
-							Contact.Phi = AvgContact.Phi;
-							return;
-						}
-					}
-				}
-			}
-		}
-	}
-	else
-	{
-		SCOPE_CYCLE_COUNTER(STAT_UpdateLevelsetAll);
-		for (int32 i = 0; i < NumParticles; ++i)
-		{
-			if (NormalAveraging2 && UpdateType != ECollisionUpdateType::Any)	//if we just want one don't bother with normal
-			{
-				const bool bInside = SampleObjectNormalAverageHelper2(Object, ObjectTransform, SampleToObjectTM, SampleParticles.X(i), Thickness, TotalThickness, AvgConstraint);
-			}
-			else
+			if (Contact.Phi >= MThickness)
 			{
-				if (SampleObjectNoNormal2(Object, ObjectTransform, SampleToObjectTM, SampleParticles.X(i), Thickness, AvgConstraint))
-				{
-					DeepestParticle = i;
-					if (UpdateType == ECollisionUpdateType::Any)
-					{
-						Contact.Phi = AvgContact.Phi;
-						return;
-					}
-				}
+				break;
 			}
-		}
-	}
 
-	if (NormalAveraging2)
-	{
-		if (TotalThickness < -KINDA_SMALL_NUMBER)
-		{
-			TVector<T, d> LocalPoint = AvgContact.Location / TotalThickness;
-			TVector<T, d> LocalNormal;
-			const T NewPhi = Object.PhiWithNormal(LocalPoint, LocalNormal);
-			if (NewPhi < Contact.Phi)
+			if ((!PBDRigid0 || IsTemporarilyStatic0) && (!PBDRigid1 || IsTemporarilyStatic1))
 			{
-				Contact.Phi = NewPhi;
-				Contact.Location = ObjectTransform.TransformPositionNoScale(LocalPoint);
-				Contact.Normal = ObjectTransform.TransformVectorNoScale(LocalNormal);
+				break;
 			}
-		}
-		else
-		{
-			check(AvgContact.Phi >= Thickness);
-		}
-	}
-	else if(AvgContact.Phi < Contact.Phi)
-	{
-		check(DeepestParticle >= 0);
-		TVector<T,d> LocalPoint = SampleToObjectTM.TransformPositionNoScale(SampleParticles.X(DeepestParticle));
-		TVector<T, d> LocalNormal;
-		Contact.Phi = Object.PhiWithNormal(LocalPoint, LocalNormal);
-		Contact.Location = ObjectTransform.TransformPositionNoScale(LocalPoint);
-		Contact.Normal = ObjectTransform.TransformVectorNoScale(LocalNormal);
-			
-	}
-}
 
+			NeedsAnotherIteration = true;
+			PMatrix<T, d, d> WorldSpaceInvI1 = PBDRigid0 ? (Q0 * FMatrix::Identity).GetTransposed() * PBDRigid0->InvI() * (Q0 * FMatrix::Identity) : PMatrix<T, d, d>(0);
+			PMatrix<T, d, d> WorldSpaceInvI2 = PBDRigid1 ? (Q1 * FMatrix::Identity).GetTransposed() * PBDRigid1->InvI() * (Q1 * FMatrix::Identity) : PMatrix<T, d, d>(0);
+			TVector<T, d> VectorToPoint1 = Contact.Location - P0;
+			TVector<T, d> VectorToPoint2 = Contact.Location - P1;
+			PMatrix<T, d, d> Factor =
+				(PBDRigid0 ? ComputeFactorMatrix3(VectorToPoint1, WorldSpaceInvI1, PBDRigid0->InvM()) : PMatrix<T, d, d>(0)) +
+				(PBDRigid1 ? ComputeFactorMatrix3(VectorToPoint2, WorldSpaceInvI2, PBDRigid1->InvM()) : PMatrix<T, d, d>(0));
+			T Numerator = FMath::Min((T)(Iteration + 2), (T)NumIterations);
+			T ScalingFactor = Numerator / (T)NumIterations;
 
-#if INTEL_ISPC
-template<ECollisionUpdateType UpdateType>
-void SampleObject2(const FImplicitObject& Object, const TRigidTransform<float, 3>& ObjectTransform, const TBVHParticles<float, 3>& SampleParticles, const TRigidTransform<float, 3>& SampleParticlesTransform, float Thickness, TRigidBodyContactConstraint<float, 3>& Constraint)
-{
-	SCOPE_CYCLE_COUNTER(STAT_SampleObject);
-	TRigidBodyContactConstraint<float, 3> AvgConstraint = Constraint;
-
-	TContactData<float, 3> & Contact = Constraint.ShapeManifold.Manifold;
-	TContactData<float, 3> & AvgContact = AvgConstraint.ShapeManifold.Manifold;
-
-	AvgConstraint.Particle[0] = Constraint.Particle[0];
-	AvgConstraint.Particle[1] = Constraint.Particle[1];
-	AvgContact.Location = TVector<float, 3>::ZeroVector;
-	AvgContact.Normal = TVector<float, 3>::ZeroVector;
-	AvgContact.Phi = Thickness;
-	float TotalThickness = float(0);
-
-	int32 DeepestParticle = -1;
-
-	const TRigidTransform<float, 3>& SampleToObjectTM = SampleParticlesTransform.GetRelativeTransform(ObjectTransform);
-	int32 NumParticles = SampleParticles.Size();
-
-	if (NumParticles > SampleMinParticlesForAcceleration2 && Object.HasBoundingBox())
-	{
-		SCOPE_CYCLE_COUNTER(STAT_UpdateLevelsetPartial);
-		TBox<float, 3> ImplicitBox = Object.BoundingBox().TransformedBox(ObjectTransform.GetRelativeTransform(SampleParticlesTransform));
-		ImplicitBox.Thicken(Thickness);
-		TArray<int32> PotentialParticles;
-		{
-			SCOPE_CYCLE_COUNTER(STAT_UpdateLevelsetFindParticles);
-			PotentialParticles = SampleParticles.FindAllIntersections(ImplicitBox);
-		}
-		{
-			SCOPE_CYCLE_COUNTER(STAT_UpdateLevelsetSignedDistance);
-
-			if (Object.GetType(true) == ImplicitObjectType::LevelSet && PotentialParticles.Num() > 0)
+			//if pushout is needed we better fix relative velocity along normal. Treat it as if 0 restitution
+			TVector<T, d> Body1Velocity = V0 + TVector<T, d>::CrossProduct(W0, VectorToPoint1);
+			TVector<T, d> Body2Velocity = V1 + TVector<T, d>::CrossProduct(W1, VectorToPoint2);
+			TVector<T, d> RelativeVelocity = Body1Velocity - Body2Velocity;
+			const T RelativeVelocityDotNormal = TVector<T, d>::DotProduct(RelativeVelocity, Contact.Normal);
+			if (RelativeVelocityDotNormal < 0)
 			{
-				//QUICK_SCOPE_CYCLE_COUNTER(STAT_LevelSet);
-				const TLevelSet<float, 3>* LevelSet = Object.GetObject<Chaos::TLevelSet<float, 3>>();
-				const TUniformGrid<float, 3>& Grid = LevelSet->GetGrid();
-
-				if (NormalAveraging2 && UpdateType != ECollisionUpdateType::Any)
-				{
-					ispc::SampleLevelSetNormalAverage(
-						(ispc::FVector&)Grid.MinCorner(),
-						(ispc::FVector&)Grid.MaxCorner(),
-						(ispc::FVector&)Grid.Dx(),
-						(ispc::FIntVector&)Grid.Counts(),
-						(ispc::TArrayND*)&LevelSet->GetPhiArray(),
-						(ispc::FTransform&)SampleToObjectTM,
-						(ispc::FVector*)&SampleParticles.XArray()[0],
-						&PotentialParticles[0],
-						Thickness,
-						TotalThickness,
-						(ispc::FVector&)AvgContact.Location,
-						PotentialParticles.Num());
-				}
-				else
+				T ImpulseDenominator = TVector<T, d>::DotProduct(Contact.Normal, Factor * Contact.Normal);
+				TVector<T, d> ImpulseNumerator = -TVector<T, d>::DotProduct(RelativeVelocity, Contact.Normal) * Contact.Normal * ScalingFactor;
+				if (!ensureMsgf(FMath::Abs(ImpulseDenominator) > SMALL_NUMBER, TEXT("ApplyPushout Constraint:%s\n\nParticle:%s\n\nLevelset:%s\n\nFactor*Contact.Normal:%s, ImpulseDenominator:%f"),
+					*Constraint.ToString(),
+					*Particle0->ToString(),
+					*Particle1->ToString(),
+					*(Factor*Contact.Normal).ToString(), ImpulseDenominator))
 				{
-					ispc::SampleLevelSetNoNormal(
-						(ispc::FVector&)Grid.MinCorner(),
-						(ispc::FVector&)Grid.MaxCorner(),
-						(ispc::FVector&)Grid.Dx(),
-						(ispc::FIntVector&)Grid.Counts(),
-						(ispc::TArrayND*)&LevelSet->GetPhiArray(),
-						(ispc::FTransform&)SampleToObjectTM,
-						(ispc::FVector*)&SampleParticles.XArray()[0],
-						&PotentialParticles[0],
-						DeepestParticle,
-						AvgContact.Phi,
-						PotentialParticles.Num());
-
-					if (UpdateType == ECollisionUpdateType::Any)
-					{
-						Contact.Phi = AvgContact.Phi;
-						return;
-					}
+					ImpulseDenominator = (T)1;
 				}
-			}
-			else if (Object.GetType(true) == ImplicitObjectType::Box && PotentialParticles.Num() > 0)
-			{
-				//QUICK_SCOPE_CYCLE_COUNTER(STAT_Box);
-				const TBox<float, 3>* Box = Object.GetObject<Chaos::TBox<float, 3>>();
 
-				if (NormalAveraging2 && UpdateType != ECollisionUpdateType::Any)
-				{
-					ispc::SampleBoxNormalAverage(
-						(ispc::FVector&)Box->Min(),
-						(ispc::FVector&)Box->Max(),
-						(ispc::FTransform&)SampleToObjectTM,
-						(ispc::FVector*)&SampleParticles.XArray()[0],
-						&PotentialParticles[0],
-						Thickness,
-						TotalThickness,
-						(ispc::FVector&)AvgContact.Location,
-						PotentialParticles.Num());
-				}
-				else
-				{
-					ispc::SampleBoxNoNormal(
-						(ispc::FVector&)Box->Min(),
-						(ispc::FVector&)Box->Max(),
-						(ispc::FTransform&)SampleToObjectTM,
-						(ispc::FVector*)&SampleParticles.XArray()[0],
-						&PotentialParticles[0],
-						DeepestParticle,
-						AvgContact.Phi,
-						PotentialParticles.Num());
-
-					if (UpdateType == ECollisionUpdateType::Any)
-					{
-						Contact.Phi = AvgContact.Phi;
-						return;
-					}
-				}
-			}
-			else
-			{
-				//QUICK_SCOPE_CYCLE_COUNTER(STAT_Other);
-				for (int32 i : PotentialParticles)
+				TVector<T, d> VelocityFixImpulse = ImpulseNumerator / ImpulseDenominator;
+				VelocityFixImpulse = GetEnergyClampedImpulse(Constraint, VelocityFixImpulse, VectorToPoint1, VectorToPoint2, Body1Velocity, Body2Velocity);
+				Constraint.AccumulatedImpulse += VelocityFixImpulse;	//question: should we track this?
+				if (!IsTemporarilyStatic0 && PBDRigid0)
 				{
-					if (NormalAveraging2 && UpdateType != ECollisionUpdateType::Any)
-					{
-						SampleObjectNormalAverageHelper2(Object, ObjectTransform, SampleToObjectTM, SampleParticles.X(i), Thickness, TotalThickness, AvgConstraint);
-					}
-					else
-					{
-						if (SampleObjectNoNormal2(Object, ObjectTransform, SampleToObjectTM, SampleParticles.X(i), Thickness, AvgConstraint))
-						{
-							DeepestParticle = i;
-							if (UpdateType == ECollisionUpdateType::Any)
-							{
-								Contact.Phi = AvgContact.Phi;
-								return;
-							}
-						}
-					}
-				}
-			}
-		}
-	}
-	else
-	{
-		SCOPE_CYCLE_COUNTER(STAT_UpdateLevelsetAll);
-		if (Object.GetType(true) == ImplicitObjectType::LevelSet && NumParticles > 0)
-		{
-			const TLevelSet<float, 3>* LevelSet = Object.GetObject<Chaos::TLevelSet<float, 3>>();
-			const TUniformGrid<float, 3>& Grid = LevelSet->GetGrid();
+					TVector<T, d> AngularImpulse = TVector<T, d>::CrossProduct(VectorToPoint1, VelocityFixImpulse);
+					PBDRigid0->V() += PBDRigid0->InvM() * VelocityFixImpulse;
+					PBDRigid0->W() += WorldSpaceInvI1 * AngularImpulse;
 
-			if (NormalAveraging2 && UpdateType != ECollisionUpdateType::Any)
-			{
-				ispc::SampleLevelSetNormalAverageAll(
-					(ispc::FVector&)Grid.MinCorner(),
-					(ispc::FVector&)Grid.MaxCorner(),
-					(ispc::FVector&)Grid.Dx(),
-					(ispc::FIntVector&)Grid.Counts(),
-					(ispc::TArrayND*)&LevelSet->GetPhiArray(),
-					(ispc::FTransform&)SampleToObjectTM,
-					(ispc::FVector*)&SampleParticles.XArray()[0],
-					Thickness,
-					TotalThickness,
-					(ispc::FVector&)AvgContact.Location,
-					NumParticles);
-			}
-			else
-			{
-				ispc::SampleLevelSetNoNormalAll(
-					(ispc::FVector&)Grid.MinCorner(),
-					(ispc::FVector&)Grid.MaxCorner(),
-					(ispc::FVector&)Grid.Dx(),
-					(ispc::FIntVector&)Grid.Counts(),
-					(ispc::TArrayND*)&LevelSet->GetPhiArray(),
-					(ispc::FTransform&)SampleToObjectTM,
-					(ispc::FVector*)&SampleParticles.XArray()[0],
-					DeepestParticle,
-					AvgContact.Phi,
-					NumParticles);
-
-				if (UpdateType == ECollisionUpdateType::Any)
-				{
-					Contact.Phi = AvgContact.Phi;
-					return;
 				}
-			}
-		}
-		else if (Object.GetType(true) == ImplicitObjectType::Plane && NumParticles > 0)
-		{
-			const TPlane<float, 3>* Plane = Object.GetObject<Chaos::TPlane<float, 3>>();
 
-			if (NormalAveraging2 && UpdateType != ECollisionUpdateType::Any)
-			{
-				ispc::SamplePlaneNormalAverageAll(
-					(ispc::FVector&)Plane->Normal(),
-					(ispc::FVector&)Plane->X(),
-					(ispc::FTransform&)SampleToObjectTM,
-					(ispc::FVector*)&SampleParticles.XArray()[0],
-					Thickness,
-					TotalThickness,
-					(ispc::FVector&)AvgContact.Location,
-					NumParticles);
-			}
-			else
-			{
-				ispc::SamplePlaneNoNormalAll(
-					(ispc::FVector&)Plane->Normal(),
-					(ispc::FVector&)Plane->X(),
-					(ispc::FTransform&)SampleToObjectTM,
-					(ispc::FVector*)&SampleParticles.XArray()[0],
-					DeepestParticle,
-					AvgContact.Phi,
-					NumParticles);
-
-				if (UpdateType == ECollisionUpdateType::Any)
+				if (!IsTemporarilyStatic1 && PBDRigid1)
 				{
-					Contact.Phi = AvgContact.Phi;
-					return;
+					TVector<T, d> AngularImpulse = TVector<T, d>::CrossProduct(VectorToPoint2, -VelocityFixImpulse);
+					PBDRigid1->V() -= PBDRigid1->InvM() * VelocityFixImpulse;
+					PBDRigid1->W() += WorldSpaceInvI2 * AngularImpulse;
 				}
-			}
-		}
-		else if (Object.GetType(true) == ImplicitObjectType::Box && NumParticles > 0)
-		{
-			const TBox<float, 3>* Box = Object.GetObject<Chaos::TBox<float, 3>>();
 
-			if (NormalAveraging2 && UpdateType != ECollisionUpdateType::Any)
-			{
-				ispc::SampleBoxNormalAverageAll(
-					(ispc::FVector&)Box->Min(),
-					(ispc::FVector&)Box->Max(),
-					(ispc::FTransform&)SampleToObjectTM,
-					(ispc::FVector*)&SampleParticles.XArray()[0],
-					Thickness,
-					TotalThickness,
-					(ispc::FVector&)AvgContact.Location,
-					NumParticles);
-			}
-			else
-			{
-				ispc::SampleBoxNoNormalAll(
-					(ispc::FVector&)Box->Min(),
-					(ispc::FVector&)Box->Max(),
-					(ispc::FTransform&)SampleToObjectTM,
-					(ispc::FVector*)&SampleParticles.XArray()[0],
-					DeepestParticle,
-					AvgContact.Phi,
-					NumParticles);
-
-				if (UpdateType == ECollisionUpdateType::Any)
-				{
-					Contact.Phi = AvgContact.Phi;
-					return;
-				}
 			}
-		}
-		else
-		{
-			//QUICK_SCOPE_CYCLE_COUNTER(STAT_Other);
-			for (int32 i = 0; i < NumParticles; ++i)
+
+
+			TVector<T, d> Impulse = PMatrix<T, d, d>(Factor.Inverse()) * ((-Contact.Phi + MThickness) * ScalingFactor * Contact.Normal);
+			TVector<T, d> AngularImpulse1 = TVector<T, d>::CrossProduct(VectorToPoint1, Impulse);
+			TVector<T, d> AngularImpulse2 = TVector<T, d>::CrossProduct(VectorToPoint2, -Impulse);
+			if (!IsTemporarilyStatic0 && PBDRigid0)
 			{
-				if (NormalAveraging2 && UpdateType != ECollisionUpdateType::Any)
-				{
-					SampleObjectNormalAverageHelper2(Object, ObjectTransform, SampleToObjectTM, SampleParticles.X(i), Thickness, TotalThickness, AvgConstraint);
-				}
-				else
-				{
-					if (SampleObjectNoNormal2(Object, ObjectTransform, SampleToObjectTM, SampleParticles.X(i), Thickness, AvgConstraint))
-					{
-						DeepestParticle = i;
-						if (UpdateType == ECollisionUpdateType::Any)
-						{
-							Contact.Phi = AvgContact.Phi;
-							return;
-						}
-					}
-				}
+				PBDRigid0->P() += PBDRigid0->InvM() * Impulse;
+				PBDRigid0->Q() = TRotation<T, d>::FromVector(WorldSpaceInvI1 * AngularImpulse1) * Q0;
+				PBDRigid0->Q().Normalize();
 			}
-		}
-	}
-
-	if (NormalAveraging2)
-	{
-		if (TotalThickness < -KINDA_SMALL_NUMBER)
-		{
-			TVector<float, 3> LocalPoint = AvgContact.Location / TotalThickness;
-			TVector<float, 3> LocalNormal;
-			const float NewPhi = Object.PhiWithNormal(LocalPoint, LocalNormal);
-			if (NewPhi < Contact.Phi)
+			if (!IsTemporarilyStatic1 && PBDRigid1)
 			{
-				Contact.Phi = NewPhi;
-				Contact.Location = ObjectTransform.TransformPositionNoScale(LocalPoint);
-				Contact.Normal = ObjectTransform.TransformVectorNoScale(LocalNormal);
+				PBDRigid1->P() -= PBDRigid1->InvM() * Impulse;
+				PBDRigid1->Q() = TRotation<T, d>::FromVector(WorldSpaceInvI2 * AngularImpulse2) * Q1;
+				PBDRigid1->Q().Normalize();
 			}
 		}
-		else
-		{
-			check(AvgContact.Phi >= Thickness);
-		}
-	}
-	else if (AvgContact.Phi < Contact.Phi)
-	{
-		check(DeepestParticle >= 0);
-		TVector<float, 3> LocalPoint = SampleToObjectTM.TransformPositionNoScale(SampleParticles.X(DeepestParticle));
-		TVector<float, 3> LocalNormal;
-		Contact.Phi = Object.PhiWithNormal(LocalPoint, LocalNormal);
-		Contact.Location = ObjectTransform.TransformPositionNoScale(LocalPoint);
-		Contact.Normal = ObjectTransform.TransformVectorNoScale(LocalNormal);
 	}
-}
-#endif
 
-template <typename T, int d>
-bool UpdateBoxPlaneConstraint(const TBox<T, d>& Box, const TRigidTransform<T, d>& BoxTransform, const TPlane<T, d>& Plane, const TRigidTransform<T, d>& PlaneTransform, T Thickness, TRigidBodyContactConstraint<T, d>& Constraint)
-{
-	TContactData<T,d> & Contact = Constraint.ShapeManifold.Manifold;
-
-#if USING_CODE_ANALYSIS
-	MSVC_PRAGMA( warning( push ) )
-	MSVC_PRAGMA( warning( disable : ALL_CODE_ANALYSIS_WARNINGS ) )
-#endif	// USING_CODE_ANALYSIS
-
-	bool bApplied = false;
-	const TRigidTransform<T, d> BoxToPlaneTransform(BoxTransform.GetRelativeTransform(PlaneTransform));
-	const TVector<T, d> Extents = Box.Extents();
-	constexpr int32 NumCorners = 2 + 2 * d;
-	constexpr T Epsilon = KINDA_SMALL_NUMBER;
-
-	TVector<T, d> Corners[NumCorners];
-	int32 CornerIdx = 0;
-	Corners[CornerIdx++] = BoxToPlaneTransform.TransformPosition(Box.Max());
-	Corners[CornerIdx++] = BoxToPlaneTransform.TransformPosition(Box.Min());
-	for (int32 j = 0; j < d; ++j)
+	DECLARE_CYCLE_STAT(TEXT("TPBDCollisionConstraint::ApplyPushOut"), STAT_ApplyPushOut, STATGROUP_Chaos);
+	template<typename T, int d>
+	bool TPBDCollisionConstraint<T, d>::ApplyPushOut(const T Dt, const TArray<FConstraintContainerHandle*>& InConstraintHandles, const TSet<TGeometryParticleHandle<T, d>*>& IsTemporarilyStatic, int32 Iteration, int32 NumIterations)
 	{
-		Corners[CornerIdx++] = BoxToPlaneTransform.TransformPosition(Box.Min() + TVector<T, d>::AxisVector(j) * Extents);
-		Corners[CornerIdx++] = BoxToPlaneTransform.TransformPosition(Box.Max() - TVector<T, d>::AxisVector(j) * Extents);
-	}
+		SCOPE_CYCLE_COUNTER(STAT_ApplyPushOut);
 
-#if USING_CODE_ANALYSIS
-	MSVC_PRAGMA( warning( pop ) )
-#endif	// USING_CODE_ANALYSIS
+		bool NeedsAnotherIteration = false;
 
-	TVector<T, d> PotentialConstraints[NumCorners];
-	int32 NumConstraints = 0;
-	for (int32 i = 0; i < NumCorners; ++i)
-	{
-		TVector<T, d> Normal;
-		const T NewPhi = Plane.PhiWithNormal(Corners[i], Normal);
-		if (NewPhi < Contact.Phi + Epsilon)
+		if (MPairIterations > 0)
 		{
-			if (NewPhi <= Contact.Phi - Epsilon)
-			{
-				NumConstraints = 0;
-			}
-			Contact.Phi = NewPhi;
-			Contact.Normal = PlaneTransform.TransformVector(Normal);
-			Contact.Location = PlaneTransform.TransformPosition(Corners[i]);
-			PotentialConstraints[NumConstraints++] = Contact.Location;
-			bApplied = true;
+			PhysicsParallelFor(InConstraintHandles.Num(), [&](int32 ConstraintHandleIndex) {
+				FConstraintContainerHandle* ConstraintHandle = InConstraintHandles[ConstraintHandleIndex];
+				check(ConstraintHandle != nullptr);
+				ApplyPushOut(Dt, Constraints[ConstraintHandle->GetConstraintIndex()], IsTemporarilyStatic, Iteration, NumIterations, NeedsAnotherIteration);
+			}, bDisableCollisionParallelFor);
 		}
-	}
-	if (NumConstraints > 1)
-	{
-		TVector<T, d> AverageLocation(0);
-		for (int32 ConstraintIdx = 0; ConstraintIdx < NumConstraints; ++ConstraintIdx)
+
+		if (PostApplyPushOutCallback != nullptr)
 		{
-			AverageLocation += PotentialConstraints[ConstraintIdx];
+			PostApplyPushOutCallback(Dt, InConstraintHandles, NeedsAnotherIteration);
 		}
-		Contact.Location = AverageLocation / NumConstraints;
-	}
-
-	return bApplied;
-}
 
-template <typename T, int d>
-void UpdateSphereConstraint(const TSphere<T, d>& Sphere1, const TRigidTransform<T, d>& Sphere1Transform, const TSphere<T, d>& Sphere2, const TRigidTransform<T, d>& Sphere2Transform, T Thickness, TRigidBodyContactConstraint<T, d>& Constraint)
-{
-	TContactData<T,d> & Contact = Constraint.ShapeManifold.Manifold;
-
-	const TVector<T, d> Center1 = Sphere1Transform.TransformPosition(Sphere1.GetCenter());
-	const TVector<T, d> Center2 = Sphere2Transform.TransformPosition(Sphere2.GetCenter());
-	const TVector<T, d> Direction = Center1 - Center2;
-	const T Size = Direction.Size();
-	const T NewPhi = Size - (Sphere1.GetRadius() + Sphere2.GetRadius());
-	if (NewPhi < Contact.Phi)
-	{
-		Contact.Normal = Size > SMALL_NUMBER ? Direction / Size : TVector<T, d>(0, 0, 1);
-		Contact.Phi = NewPhi;
-		Contact.Location = Center1 - Sphere1.GetRadius() * Contact.Normal;
+		return NeedsAnotherIteration;
 	}
-}
-
-template <typename T, int d>
-void UpdateSpherePlaneConstraint(const TSphere<T, d>& Sphere, const TRigidTransform<T, d>& SphereTransform, const TPlane<T, d>& Plane, const TRigidTransform<T, d>& PlaneTransform, T Thickness, TRigidBodyContactConstraint<T, d>& Constraint)
-{
-	TContactData<T,d> & Contact = Constraint.ShapeManifold.Manifold;
 
-	const TRigidTransform<T, d> SphereToPlaneTransform(PlaneTransform.Inverse() * SphereTransform);
-	const TVector<T, d> SphereCenter = SphereToPlaneTransform.TransformPosition(Sphere.GetCenter());
 
-	TVector<T, d> NewNormal;
-	T NewPhi = Plane.PhiWithNormal(SphereCenter, NewNormal);
-	NewPhi -= Sphere.GetRadius();
-
-	if (NewPhi < Contact.Phi)
+	template<typename T, int d>
+	TRigidTransform<T, d> GetTransform(const TGeometryParticleHandle<T, d>* Particle)
 	{
-		Contact.Phi = NewPhi;
-		Contact.Normal = PlaneTransform.TransformVectorNoScale(NewNormal);
-		Contact.Location = SphereCenter - Contact.Normal * Sphere.GetRadius();
-	}
-}
-
-template <typename T, int d>
-void UpdateSphereBoxConstraint(const TSphere<T, d>& Sphere, const TRigidTransform<T, d>& SphereTransform, const TBox<T, d>& Box, const TRigidTransform<T, d>& BoxTransform, T Thickness, TRigidBodyContactConstraint<T, d>& Constraint)
-{
-	TContactData<T,d> & Contact = Constraint.ShapeManifold.Manifold;
-
-	const TRigidTransform<T, d> SphereToBoxTransform(SphereTransform * BoxTransform.Inverse());
-	const TVector<T, d> SphereCenterInBox = SphereToBoxTransform.TransformPosition(Sphere.GetCenter());
+		//#if CHAOS_PARTICLEHANDLE_TODO
 
-	TVector<T, d> NewNormal;
-	T NewPhi = Box.PhiWithNormal(SphereCenterInBox, NewNormal);
-	NewPhi -= Sphere.GetRadius();
-
-	if (NewPhi < Contact.Phi)
-	{
-		Contact.Phi = NewPhi;
-		Contact.Normal = BoxTransform.TransformVectorNoScale(NewNormal);
-		Contact.Location = SphereTransform.TransformPosition(Sphere.GetCenter()) - Contact.Normal * Sphere.GetRadius();
+		TGenericParticleHandle<T, d> Generic = const_cast<TGeometryParticleHandle<T, d>*>(Particle);	//TODO: give a const version of the generic API
+		return TRigidTransform<T, d>(Generic->P(), Generic->Q());
 	}
-}
 
-template <typename T, int d>
-void UpdateCapsuleCapsuleConstraint(const TCapsule<T>& A, const TRigidTransform<T, d>& ATransform, const TCapsule<T>& B, const TRigidTransform<T, d>& BTransform, T Thickness, TRigidBodyContactConstraint<T, d>& Constraint)
-{
-	TContactData<T,d> & Contact = Constraint.ShapeManifold.Manifold;
-
-	FVector A1 = ATransform.TransformPosition(A.GetX1());
-	FVector A2 = ATransform.TransformPosition(A.GetX2());
-	FVector B1 = BTransform.TransformPosition(B.GetX1());
-	FVector B2 = BTransform.TransformPosition(B.GetX2());
-	FVector P1, P2;
-	FMath::SegmentDistToSegmentSafe(A1, A2, B1, B2, P1, P2);
-
-	TVector<T, d> Delta = P2 - P1;
-	T DeltaLen = Delta.Size();
-	if (DeltaLen > KINDA_SMALL_NUMBER)
+	template<typename T, int d>
+	void TPBDCollisionConstraint<T, d>::ConstructConstraints(TGeometryParticleHandle<T, d>* Particle0, TGeometryParticleHandle<T, d>* Particle1, const T Thickness, TRigidBodyContactConstraint<T, d>& Constraint)
 	{
-		T NewPhi = DeltaLen - (A.GetRadius() + B.GetRadius());
-		if (NewPhi < Contact.Phi)
+		if (ensure(Particle0 && Particle1))
 		{
-			TVector<T, d> Dir = Delta / DeltaLen;
-			Contact.Phi = NewPhi;
-			Contact.Normal = -Dir;
-			Contact.Location = P1 + Dir * A.GetRadius();
+			ConstructConstraintsImpl<T, d>(Particle0, Particle1, Particle0->Geometry().Get(), Particle1->Geometry().Get(), Thickness, Constraint);
 		}
 	}
-}
-
-template <typename T, int d>
-void UpdateCapsuleBoxConstraint(const TCapsule<T>& A, const TRigidTransform<T, d>& ATransform, const TBox<T, d>& B, const TRigidTransform<T, d>& BTransform, T Thickness, TRigidBodyContactConstraint<T, d>& Constraint)
-{
-	// @todo(ccaulfield): Add custom capsule-box collision
-	TContactData<T,d> & Contact = Constraint.ShapeManifold.Manifold;
 
-	TRigidTransform<T, d> BToATransform = BTransform.GetRelativeTransform(ATransform);
+	DECLARE_CYCLE_STAT(TEXT("TPBDCollisionConstraint::UpdateConstraint"), STAT_UpdateConstraint, STATGROUP_ChaosWide);
 
-	// Use GJK to track closest points (not strictly necessary yet)
-	TVector<T, d> NearPointALocal, NearPointBLocal;
-	T NearPointDistance;
-	if (GJKDistance<T>(A, B, BToATransform, NearPointDistance, NearPointALocal, NearPointBLocal))
-	{
-		TVector<T, d> NearPointAWorld = ATransform.TransformPosition(NearPointALocal);
-		TVector<T, d> NearPointBWorld = BTransform.TransformPosition(NearPointBLocal);
-		TVector<T, d> NearPointBtoAWorld = NearPointAWorld - NearPointBWorld;
-		Contact.Phi = NearPointDistance;
-		Contact.Normal = NearPointBtoAWorld.GetSafeNormal();
-		Contact.Location = NearPointAWorld;
-	}
-	else
+	template<typename T, int d>
+	template<ECollisionUpdateType UpdateType>
+	void TPBDCollisionConstraint<T, d>::UpdateConstraint(const T Thickness, FRigidBodyContactConstraint& Constraint)
 	{
-		// Use box particle samples against the implicit capsule
-		const TArray<TVector<T, d>> BParticles = B.ComputeSamplePoints();
-		const int32 NumParticles = BParticles.Num();
-		for (int32 ParticleIndex = 0; ParticleIndex < NumParticles; ++ParticleIndex)
-		{
-			if (SampleObjectHelper2(A, ATransform, BToATransform, BParticles[ParticleIndex], Thickness, Constraint))
-			{
-				// SampleObjectHelper2 expects A to be the box, so reverse the results
-				Contact.Normal = -Contact.Normal;
-			}
-		}
-	}
+		SCOPE_CYCLE_COUNTER(STAT_UpdateConstraint);
 
-}
+		Constraint.ResetPhi(Thickness);
+		const TRigidTransform<T, d> ParticleTM = GetTransform(Constraint.Particle[0]);
+		const TRigidTransform<T, d> LevelsetTM = GetTransform(Constraint.Particle[1]);
 
-DECLARE_CYCLE_STAT(TEXT("TPBDCollisionConstraint::FindRelevantShapes"), STAT_FindRelevantShapes, STATGROUP_ChaosWide);
-template <typename T, int d>
-TArray<Pair<const FImplicitObject*, TRigidTransform<T, d>>> FindRelevantShapes2(const FImplicitObject* ParticleObj, const TRigidTransform<T,d>& ParticlesTM, const FImplicitObject& LevelsetObj, const TRigidTransform<T,d>& LevelsetTM, const T Thickness)
-{
-	SCOPE_CYCLE_COUNTER(STAT_FindRelevantShapes);
-	TArray<Pair<const FImplicitObject*, TRigidTransform<T, d>>> RelevantShapes;
-	//find all levelset inner objects
-	if (ParticleObj)
-	{
-		if (ParticleObj->HasBoundingBox())
+		if (!Constraint.Particle[0]->Geometry())
 		{
-			const TRigidTransform<T, d> ParticlesToLevelsetTM = ParticlesTM.GetRelativeTransform(LevelsetTM);
-			TBox<T, d> ParticleBoundsInLevelset = ParticleObj->BoundingBox().TransformedBox(ParticlesToLevelsetTM);
-			ParticleBoundsInLevelset.Thicken(Thickness);
+			if (Constraint.Particle[1]->Geometry())
 			{
-				LevelsetObj.FindAllIntersectingObjects(RelevantShapes, ParticleBoundsInLevelset);
+				if (!Constraint.Particle[1]->Geometry()->IsUnderlyingUnion())
+				{
+					UpdateLevelsetConstraint<UpdateType>(Thickness, Constraint);
+				}
+				else
+				{
+					UpdateUnionLevelsetConstraint<UpdateType>(Thickness, Constraint);
+				}
 			}
 		}
 		else
 		{
-			LevelsetObj.AccumulateAllImplicitObjects(RelevantShapes, TRigidTransform<T, d>::Identity);
+			UpdateConstraintImp<UpdateType>(*Constraint.Particle[0]->Geometry(), ParticleTM, *Constraint.Particle[1]->Geometry(), LevelsetTM, Thickness, Constraint);
 		}
 	}
-	else
-	{
-		//todo:compute bounds
-		LevelsetObj.AccumulateAllImplicitObjects(RelevantShapes, TRigidTransform<T, d>::Identity);
-	}
-
-	return RelevantShapes;
-}
-
-DECLARE_CYCLE_STAT(TEXT("TPBDCollisionConstraint::UpdateUnionUnionConstraint"), STAT_UpdateUnionUnionConstraint, STATGROUP_ChaosWide);
-template<ECollisionUpdateType UpdateType, typename T, int d>
-void UpdateUnionUnionConstraint(const T Thickness, TRigidBodyContactConstraint<T, d>& Constraint)
-{
-	SCOPE_CYCLE_COUNTER(STAT_UpdateUnionUnionConstraint);
-
-	TGenericParticleHandle<T, d> Particle0 = Constraint.Particle[0];
-	TGenericParticleHandle<T, d> Particle1 = Constraint.Particle[1];
-
-	TRigidTransform<T, d> ParticlesTM = TRigidTransform<T, d>(Particle0->P(), Particle0->Q());
-	TRigidTransform<T, d> LevelsetTM = TRigidTransform<T, d>(Particle1->P(), Particle1->Q());
-
-	const FImplicitObject* ParticleObj = Particle0->Geometry().Get();
-	const FImplicitObject* LevelsetObj = Particle1->Geometry().Get();
-	const TArray<Pair<const FImplicitObject*, TRigidTransform<T, d>>> LevelsetShapes = FindRelevantShapes2(ParticleObj, ParticlesTM, *LevelsetObj, LevelsetTM, Thickness);
-
-	for (const Pair<const FImplicitObject*, TRigidTransform<T, d>>& LevelsetObjPair : LevelsetShapes)
-	{
-		const FImplicitObject& LevelsetInnerObj = *LevelsetObjPair.First;
-		const TRigidTransform<T, d>& LevelsetInnerObjTM = LevelsetObjPair.Second * LevelsetTM;
-
-		//now find all particle inner objects
-		const TArray<Pair<const FImplicitObject*, TRigidTransform<T, d>>> ParticleShapes = FindRelevantShapes2(&LevelsetInnerObj, LevelsetInnerObjTM, *ParticleObj, ParticlesTM, Thickness);
-
-		//for each inner obj pair, update constraint
-		for (const Pair<const FImplicitObject*, TRigidTransform<T, d>>& ParticlePair : ParticleShapes)
-		{
-			const FImplicitObject& ParticleInnerObj = *ParticlePair.First;
-			const TRigidTransform<T, d> ParticleInnerObjTM = ParticlePair.Second * ParticlesTM;
-			UpdateConstraintImp<UpdateType>(ParticleInnerObj, ParticleInnerObjTM, LevelsetInnerObj, LevelsetInnerObjTM, Thickness, Constraint);
-		}
-	}
-}
-
-DECLARE_CYCLE_STAT(TEXT("TPBDCollisionConstraint::UpdateSingleUnionConstraint"), STAT_UpdateSingleUnionConstraint, STATGROUP_ChaosWide);
-template<ECollisionUpdateType UpdateType, typename T, int d>
-void UpdateSingleUnionConstraint(const T Thickness, TRigidBodyContactConstraint<T, d>& Constraint)
-{
-	SCOPE_CYCLE_COUNTER(STAT_UpdateSingleUnionConstraint);
-
-	TGenericParticleHandle<T, d> Particle0 = Constraint.Particle[0];
-	TGenericParticleHandle<T, d> Particle1 = Constraint.Particle[1];
-
-	TRigidTransform<T, d> ParticlesTM = TRigidTransform<T, d>(Particle0->P(), Particle0->Q());
-	TRigidTransform<T, d> LevelsetTM = TRigidTransform<T, d>(Particle1->P(), Particle1->Q());
-
-	const FImplicitObject* ParticleObj = Particle0->Geometry().Get();
-	const FImplicitObject* LevelsetObj = Particle1->Geometry().Get();
-	const TArray<Pair<const FImplicitObject*, TRigidTransform<T, d>>> LevelsetShapes = FindRelevantShapes2(ParticleObj, ParticlesTM, *LevelsetObj, LevelsetTM, Thickness);
-	
-	for (const Pair<const FImplicitObject*, TRigidTransform<T, d>>& LevelsetObjPair : LevelsetShapes)
-	{
-		const FImplicitObject& LevelsetInnerObj = *LevelsetObjPair.First;
-		const TRigidTransform<T, d> LevelsetInnerObjTM = LevelsetObjPair.Second * LevelsetTM;
-		UpdateConstraintImp<UpdateType>(*ParticleObj, ParticlesTM, LevelsetInnerObj, LevelsetInnerObjTM, Thickness, Constraint);
-	}
-}
-
-DECLARE_CYCLE_STAT(TEXT("TPBDCollisionConstraint::UpdateLevelsetConstraint"), STAT_UpdateLevelsetConstraint, STATGROUP_ChaosWide);
-template<typename T, int d>
-template<ECollisionUpdateType UpdateType>
-void TPBDCollisionConstraint<T, d>::UpdateLevelsetConstraint(const T Thickness, FRigidBodyContactConstraint& Constraint)
-{
-	SCOPE_CYCLE_COUNTER(STAT_UpdateLevelsetConstraint);
-	
-	TGenericParticleHandle<T, d> Particle0 = Constraint.Particle[0];
-	TRigidTransform<T, d> ParticlesTM = TRigidTransform<T, d>(Particle0->P(), Particle0->Q());
-	if (!(ensure(!FMath::IsNaN(ParticlesTM.GetTranslation().X)) && ensure(!FMath::IsNaN(ParticlesTM.GetTranslation().Y)) && ensure(!FMath::IsNaN(ParticlesTM.GetTranslation().Z))))
-	{
-		return;
-	}
-
-	TGenericParticleHandle<T, d> Particle1 = Constraint.Particle[1];
-	TRigidTransform<T, d> LevelsetTM = TRigidTransform<T, d>(Particle1->P(), Particle1->Q());
-	if (!(ensure(!FMath::IsNaN(LevelsetTM.GetTranslation().X)) && ensure(!FMath::IsNaN(LevelsetTM.GetTranslation().Y)) && ensure(!FMath::IsNaN(LevelsetTM.GetTranslation().Z))))
-	{
-		return;
-	}
-
-	const TBVHParticles<T, d>* SampleParticles = nullptr;
-	SampleParticles = Particle0->CollisionParticles().Get();
-
-	if(SampleParticles)
-	{
-		SampleObject2<UpdateType>(*Particle1->Geometry(), LevelsetTM, *SampleParticles, ParticlesTM, Thickness, Constraint);
-	}
-}
-
-DECLARE_CYCLE_STAT(TEXT("TPBDCollisionConstraint::UpdateUnionLevelsetConstraint"), STAT_UpdateUnionLevelsetConstraint, STATGROUP_ChaosWide);
-template<ECollisionUpdateType UpdateType, typename T, int d>
-void UpdateUnionLevelsetConstraint(const T Thickness, TRigidBodyContactConstraint<T, d>& Constraint)
-{
-	SCOPE_CYCLE_COUNTER(STAT_UpdateUnionLevelsetConstraint);
-
-	TGenericParticleHandle<T, d> Particle0 = Constraint.Particle[0];
-	TGenericParticleHandle<T, d> Particle1 = Constraint.Particle[1];
-
-	TRigidTransform<T, d> ParticlesTM = TRigidTransform<T, d>(Particle0->P(), Particle0->Q());
-	TRigidTransform<T, d> LevelsetTM = TRigidTransform<T, d>(Particle1->P(), Particle1->Q());
-
-	if (!(ensure(!FMath::IsNaN(ParticlesTM.GetTranslation().X)) && ensure(!FMath::IsNaN(ParticlesTM.GetTranslation().Y)) && ensure(!FMath::IsNaN(ParticlesTM.GetTranslation().Z))))
-	{
-		return;
-	}
-
-	if (!(ensure(!FMath::IsNaN(LevelsetTM.GetTranslation().X)) && ensure(!FMath::IsNaN(LevelsetTM.GetTranslation().Y)) && ensure(!FMath::IsNaN(LevelsetTM.GetTranslation().Z))))
-	{
-		return;
-	}
-
-	const FImplicitObject* ParticleObj = Particle0->Geometry().Get();
-	const FImplicitObject* LevelsetObj = Particle1->Geometry().Get();
-	TArray<Pair<const FImplicitObject*, TRigidTransform<T, d>>> LevelsetShapes = FindRelevantShapes2(ParticleObj, ParticlesTM, *LevelsetObj, LevelsetTM, Thickness);
-
-	if (LevelsetShapes.Num() && Particle0->CollisionParticles().Get())
-	{
-		const TBVHParticles<T, d>& SampleParticles = *Particle0->CollisionParticles().Get();
-		if (SampleParticles.Size())
-		{
-			for (const Pair<const FImplicitObject*, TRigidTransform<T, d>>& LevelsetObjPair : LevelsetShapes)
-			{
-				const FImplicitObject* Object = LevelsetObjPair.First;
-				const TRigidTransform<T, d> ObjectTM = LevelsetObjPair.Second * LevelsetTM;
-				SampleObject2<UpdateType>(*Object, ObjectTM, SampleParticles, ParticlesTM, Thickness, Constraint);
-				if (UpdateType == ECollisionUpdateType::Any && Constraint.GetPhi() < Thickness)
-				{
-					return;
-				}
-			}
-		}
-#if CHAOS_PARTICLEHANDLE_TODO
-		else if (ParticleObj && ParticleObj->IsUnderlyingUnion())
-		{
-			const TImplicitObjectUnion<T, d>* UnionObj = static_cast<const TImplicitObjectUnion<T, d>*>(ParticleObj);
-			//need to traverse shapes to get their collision particles
-			for (const Pair<const FImplicitObject*, TRigidTransform<T, d>>& LevelsetObjPair : LevelsetShapes)
-			{
-				const FImplicitObject* LevelsetInnerObject = LevelsetObjPair.First;
-				const TRigidTransform<T, d> LevelsetInnerObjectTM = LevelsetObjPair.Second * LevelsetTM;
-
-				TArray<Pair<const FImplicitObject*, TRigidTransform<T, d>>> ParticleShapes = FindRelevantShapes2(LevelsetInnerObject, LevelsetInnerObjectTM, *ParticleObj, ParticlesTM, Thickness);
-				for (const Pair<const FImplicitObject*, TRigidTransform<T, d>>& ParticleObjPair : ParticleShapes)
-				{
-					const FImplicitObject* ParticleInnerObject = ParticleObjPair.First;
-					const TRigidTransform<T, d> ParticleInnerObjectTM = ParticleObjPair.Second * ParticlesTM;
-
-					if (const int32* OriginalIdx = UnionObj->MCollisionParticleLookupHack.Find(ParticleInnerObject))
-					{
-						const TBVHParticles<T, d>& InnerSampleParticles = *InParticles.CollisionParticles(*OriginalIdx).Get();
-						SampleObject2<UpdateType>(*LevelsetInnerObject, LevelsetInnerObjectTM, InnerSampleParticles, ParticleInnerObjectTM, Thickness, Constraint);
-						if (UpdateType == ECollisionUpdateType::Any && Constraint.Phi < Thickness)
-						{
-							return;
-						}
-					}
-				}
-
-			}
-		}
-#endif
-	}
-}
-
-
-DECLARE_CYCLE_STAT(TEXT("TPBDCollisionConstraint::UpdateLevelsetUnionConstraint"), STAT_UpdateLevelsetUnionConstraint, STATGROUP_ChaosWide);
-template<ECollisionUpdateType UpdateType, typename T, int d>
-void UpdateLevelsetUnionConstraint(const T Thickness, TRigidBodyContactConstraint<T, d>& Constraint)
-{
-	SCOPE_CYCLE_COUNTER(STAT_UpdateLevelsetUnionConstraint);
-
-	TGenericParticleHandle<T, d> Particle0 = Constraint.Particle[0];
-	TGenericParticleHandle<T, d> Particle1 = Constraint.Particle[1];
-
-	TRigidTransform<T, d> ParticlesTM = TRigidTransform<T, d>(Particle0->P(), Particle0->Q());
-	TRigidTransform<T, d> LevelsetTM = TRigidTransform<T, d>(Particle1->P(), Particle1->Q());
-
-	const FImplicitObject* ParticleObj = Particle0->Geometry().Get();
-	const FImplicitObject* LevelsetObj = Particle1->Geometry().Get();
-
-	if (!(ensure(!FMath::IsNaN(ParticlesTM.GetTranslation().X)) && ensure(!FMath::IsNaN(ParticlesTM.GetTranslation().Y)) && ensure(!FMath::IsNaN(ParticlesTM.GetTranslation().Z))))
-	{
-		return;
-	}
-
-	if (!(ensure(!FMath::IsNaN(LevelsetTM.GetTranslation().X)) && ensure(!FMath::IsNaN(LevelsetTM.GetTranslation().Y)) && ensure(!FMath::IsNaN(LevelsetTM.GetTranslation().Z))))
-	{
-		return;
-	}
-
-#if CHAOS_PARTICLEHANDLE_TODO
-	TArray<Pair<const FImplicitObject*, TRigidTransform<T, d>>> ParticleShapes = FindRelevantShapes2(LevelsetObj, LevelsetTM, *ParticleObj, ParticlesTM, Thickness);
-	check(ParticleObj->IsUnderlyingUnion());
-	const TImplicitObjectUnion<T, d>* UnionObj = static_cast<const TImplicitObjectUnion<T, d>*>(ParticleObj);
-	for (const Pair<const FImplicitObject*, TRigidTransform<T, d>>& ParticleObjPair : ParticleShapes)
-	{
-		const FImplicitObject* Object = ParticleObjPair.First;
-
-		if (const int32* OriginalIdx = UnionObj->MCollisionParticleLookupHack.Find(Object))
-		{
-			const TBVHParticles<T, d>& SampleParticles = *InParticles.CollisionParticles(*OriginalIdx).Get();
-			const TRigidTransform<T, d> ObjectTM = ParticleObjPair.Second * ParticlesTM;
-
-			SampleObject2<UpdateType>(*LevelsetObj, LevelsetTM, SampleParticles, ObjectTM, Thickness, Constraint);
-			if (UpdateType == ECollisionUpdateType::Any && Constraint.Phi < Thickness)
-			{
-				return;
-			}
-		}
-	}
-#endif
-}
-
-
-template <typename T, int d>
-void UpdateBoxConstraint(const TBox<T, d>& Box1, const TRigidTransform<T, d>& Box1Transform, const TBox<T, d>& Box2, const TRigidTransform<T, d>& Box2Transform, T Thickness, TRigidBodyContactConstraint<T, d>& Constraint)
-{
-	TContactData<T,d> & Contact = Constraint.ShapeManifold.Manifold;
-
-	TBox<T,d> Box2SpaceBox1 = Box1.TransformedBox(Box1Transform * Box2Transform.Inverse());
-	TBox<T,d> Box1SpaceBox2 = Box2.TransformedBox(Box2Transform * Box1Transform.Inverse());
-	Box2SpaceBox1.Thicken(Thickness);
-	Box1SpaceBox2.Thicken(Thickness);
-	if (Box1SpaceBox2.Intersects(Box1) && Box2SpaceBox1.Intersects(Box2))
-	{
-		const TVector<T, d> Box1Center = (Box1Transform * Box2Transform.Inverse()).TransformPosition(Box1.Center());
-		bool bDeepOverlap = false;
-		if (Box2.SignedDistance(Box1Center) < 0)
-		{
-			//If Box1 is overlapping Box2 by this much the signed distance approach will fail (box1 gets sucked into box2). In this case just use two spheres
-			TSphere<T, d> Sphere1(Box1Transform.TransformPosition(Box1.Center()), Box1.Extents().Min() / 2);
-			TSphere<T, d> Sphere2(Box2Transform.TransformPosition(Box2.Center()), Box2.Extents().Min() / 2);
-			const TVector<T, d> Direction = Sphere1.GetCenter() - Sphere2.GetCenter();
-			T Size = Direction.Size();
-			if(Size < (Sphere1.GetRadius() + Sphere2.GetRadius()))
-			{
-				const T NewPhi = Size - (Sphere1.GetRadius() + Sphere2.GetRadius());;
-				if (NewPhi < Contact.Phi)
-				{
-					bDeepOverlap = true;
-					Contact.Normal = Size > SMALL_NUMBER ? Direction / Size : TVector<T, d>(0, 0, 1);
-					Contact.Phi = NewPhi;
-					Contact.Location = Sphere1.GetCenter() - Sphere1.GetRadius() * Contact.Normal;
-				}
-			}
-		}
-		if (!bDeepOverlap || Contact.Phi >= 0)
-		{
-			//if we didn't have deep penetration use signed distance per particle. If we did have deep penetration but the spheres did not overlap use signed distance per particle
-
-			//UpdateLevelsetConstraintGJK(InParticles, Thickness, Constraint);
-			//check(Contact.Phi < MThickness);
-			// For now revert to doing all points vs lsv check until we can figure out a good way to get the deepest point without needing this
-			{
-				const TArray<TVector<T, d>> SampleParticles = Box1.ComputeSamplePoints();
-				const TRigidTransform<T, d> Box1ToBox2Transform = Box1Transform.GetRelativeTransform(Box2Transform);
-				int32 NumParticles = SampleParticles.Num();
-				for (int32 i = 0; i < NumParticles; ++i)
-				{
-					SampleObjectHelper2(Box2, Box2Transform, Box1ToBox2Transform, SampleParticles[i], Thickness, Constraint);
-				}
-			}
-		}
-	}
-}
-
-
-template<typename T, int d>
-void ConstructLevelsetConstraints(TGeometryParticleHandle<T, d>* Particle0, TGeometryParticleHandle<T, d>* Particle1, const FImplicitObject* Implicit0, const FImplicitObject* Implicit1, const T Thickness, TRigidBodyContactConstraint<T, d>& Constraint)
-{
-	
-	if (!Particle1->Geometry() || (Particle0->AsDynamic() && !Particle0->AsDynamic()->CollisionParticlesSize() && Particle0->Geometry() && !Particle0->Geometry()->IsUnderlyingUnion()))
-	{
-		Constraint.Particle[0] = Particle1;
-		Constraint.Particle[1] = Particle0;
-		Constraint.AddManifold(Implicit1, Implicit0);
-	}
-	else
-	{
-		Constraint.Particle[0] = Particle0;
-		Constraint.Particle[1] = Particle1;
-		Constraint.AddManifold(Implicit0, Implicit1);
-	}
-}
-
-template<typename T, int d>
-void ConstructBoxConstraints(TGeometryParticleHandle<T, d>* Particle0, TGeometryParticleHandle<T, d>* Particle1, const FImplicitObject* Implicit0, const FImplicitObject* Implicit1, const T Thickness, TRigidBodyContactConstraint<T, d>& Constraint)
-{
-	Constraint.Particle[0] = Particle0;
-	Constraint.Particle[1] = Particle1;
-	Constraint.AddManifold(Implicit0, Implicit1);
-}
-
-template<typename T, int d>
-void ConstructBoxPlaneConstraints(TGeometryParticleHandle<T, d>* Particle0, TGeometryParticleHandle<T, d>* Particle1, const FImplicitObject* Implicit0, const FImplicitObject* Implicit1, const T Thickness, TRigidBodyContactConstraint<T, d>& Constraint)
-{
-	Constraint.Particle[0] = Particle0;
-	Constraint.Particle[1] = Particle1;
-	Constraint.AddManifold(Implicit0, Implicit1);
-}
-
-template<typename T, int d>
-void ConstructSphereConstraints(TGeometryParticleHandle<T, d>* Particle0, TGeometryParticleHandle<T, d>* Particle1, const FImplicitObject* Implicit0, const FImplicitObject* Implicit1, const T Thickness, TRigidBodyContactConstraint<T, d>& Constraint)
-{
-	Constraint.Particle[0] = Particle0;
-	Constraint.Particle[1] = Particle1;
-	Constraint.AddManifold(Implicit0, Implicit1);
-}
-
-template<typename T, int d>
-void ConstructSpherePlaneConstraints(TGeometryParticleHandle<T, d>* Particle0, TGeometryParticleHandle<T, d>* Particle1, const FImplicitObject* Implicit0, const FImplicitObject* Implicit1, const T Thickness, TRigidBodyContactConstraint<T, d>& Constraint)
-{
-	Constraint.Particle[0] = Particle0;
-	Constraint.Particle[1] = Particle1;
-	Constraint.AddManifold(Implicit0, Implicit1);
-}
-
-template<typename T, int d>
-void ConstructSphereBoxConstraints(TGeometryParticleHandle<T, d>* Particle0, TGeometryParticleHandle<T, d>* Particle1, const FImplicitObject* Implicit0, const FImplicitObject* Implicit1, const T Thickness, TRigidBodyContactConstraint<T, d>& Constraint)
-{
-	Constraint.Particle[0] = Particle0;
-	Constraint.Particle[1] = Particle1;
-	Constraint.AddManifold(Implicit0, Implicit1);
-}
-
-template<typename T, int d>
-void ConstructCapsuleCapsuleConstraints(TGeometryParticleHandle<T, d>* Particle0, TGeometryParticleHandle<T, d>* Particle1, const FImplicitObject* Implicit0, const FImplicitObject* Implicit1, const T Thickness, TRigidBodyContactConstraint<T, d>& Constraint)
-{
-	Constraint.Particle[0] = Particle0;
-	Constraint.Particle[1] = Particle1;
-	Constraint.AddManifold(Implicit0, Implicit1);
-}
-
-template<typename T, int d>
-void ConstructCapsuleBoxConstraints(TGeometryParticleHandle<T, d>* Particle0, TGeometryParticleHandle<T, d>* Particle1, const FImplicitObject* Implicit0, const FImplicitObject* Implicit1, const T Thickness, TRigidBodyContactConstraint<T, d>& Constraint)
-{
-	Constraint.Particle[0] = Particle0;
-	Constraint.Particle[1] = Particle1;
-	Constraint.AddManifold(Implicit0, Implicit1);
-}
-
-template<typename T, int d>
-void ConstructSingleUnionConstraints(TGeometryParticleHandle<T, d>* Particle0, TGeometryParticleHandle<T, d>* Particle1, const FImplicitObject* Implicit0, const FImplicitObject* Implicit1, const T Thickness, TRigidBodyContactConstraint<T, d>& Constraint)
-{
-	Constraint.Particle[0] = Particle0;
-	Constraint.Particle[1] = Particle1;
-	Constraint.AddManifold(Implicit0, Implicit1);
-}
-
-template<typename T, int d>
-void ConstructUnionUnionConstraints(TGeometryParticleHandle<T, d>* Particle0, TGeometryParticleHandle<T, d>* Particle1, const FImplicitObject* Implicit0, const FImplicitObject* Implicit1, const T Thickness, TRigidBodyContactConstraint<T, d>& Constraint)
-{
-	SCOPE_CYCLE_COUNTER(STAT_UpdateUnionUnionConstraint);
-
-	const FImplicitObject* ParticleObj = Particle0->Geometry().Get();
-	const FImplicitObject* LevelsetObj = Particle1->Geometry().Get();
-
-	TRigidTransform<T, d> ParticlesTM = TRigidTransform<T, d>(Particle0->X(), Particle0->R());
-	TRigidTransform<T, d> LevelsetTM = TRigidTransform<T, d>(Particle1->X(), Particle1->R());
-
-	const TArray<Pair<const FImplicitObject*, TRigidTransform<T, d>>> LevelsetShapes = FindRelevantShapes2(ParticleObj, ParticlesTM, *LevelsetObj, LevelsetTM, Thickness);
-
-	for (const Pair<const FImplicitObject*, TRigidTransform<T, d>>& LevelsetObjPair : LevelsetShapes)
-	{
-		const FImplicitObject* LevelsetInnerObj = LevelsetObjPair.First;
-		const TRigidTransform<T, d>& LevelsetInnerObjTM = LevelsetObjPair.Second * LevelsetTM;
-
-		//now find all particle inner objects
-		const TArray<Pair<const FImplicitObject*, TRigidTransform<T, d>>> ParticleShapes = FindRelevantShapes2(LevelsetInnerObj, LevelsetInnerObjTM, *ParticleObj, ParticlesTM, Thickness);
-
-		//for each inner obj pair, update constraint
-		for (const Pair<const FImplicitObject*, TRigidTransform<T, d>>& ParticlePair : ParticleShapes)
-		{
-			const FImplicitObject* ParticleInnerObj = ParticlePair.First;
-			ConstructConstraintsImpl<T, d>(Particle0, Particle1, ParticleInnerObj, LevelsetInnerObj, Thickness, Constraint);
-		}
-	}
-}
-
-template<typename T, int d>
-void ConstructPairConstraintImpl(TGeometryParticleHandle<T, d>* Particle0, TGeometryParticleHandle<T, d>* Particle1, const FImplicitObject* Implicit0, const FImplicitObject* Implicit1, const T Thickness, TRigidBodyContactConstraint<T, d>& Constraint)
-{
-	// See if we already have a constraint for this shape pair
-	if (Constraint.ContainsManifold(Implicit0, Implicit1))
-	{
-		return;
-	}
-
-	if (!Implicit0 || !Implicit1)
-	{
-		ConstructLevelsetConstraints(Particle0, Particle1, Implicit0, Implicit1, Thickness, Constraint);
-	}
-	else if (Implicit0->GetType() == TBox<T, d>::StaticType() && Implicit1->GetType() == TBox<T, d>::StaticType()) 
-	{
-		ConstructBoxConstraints(Particle0, Particle1, Implicit0, Implicit1, Thickness, Constraint);
-	}
-	else if (Implicit0->GetType() == TSphere<T, d>::StaticType() && Implicit1->GetType() == TSphere<T, d>::StaticType())
-	{
-		ConstructSphereConstraints(Particle0, Particle1, Implicit0, Implicit1, Thickness, Constraint);
-	}
-	else if (Implicit0->GetType() == TBox<T, d>::StaticType() && Implicit1->GetType() == TPlane<T, d>::StaticType())
-	{
-		ConstructBoxPlaneConstraints(Particle0, Particle1, Implicit0, Implicit1, Thickness, Constraint);
-	}
-	else if (Implicit1->GetType() == TBox<T, d>::StaticType() && Implicit0->GetType() == TPlane<T, d>::StaticType())
-	{
-		ConstructBoxPlaneConstraints(Particle1, Particle0, Implicit0, Implicit1, Thickness, Constraint);
-	}
-	else if (Implicit0->GetType() == TSphere<T, d>::StaticType() && Implicit1->GetType() == TPlane<T, d>::StaticType())
-	{
-		ConstructSpherePlaneConstraints(Particle0, Particle1, Implicit0, Implicit1, Thickness, Constraint);
-	}
-	else if (Implicit1->GetType() == TSphere<T, d>::StaticType() && Implicit0->GetType() == TPlane<T, d>::StaticType())
-	{
-		ConstructSpherePlaneConstraints(Particle1, Particle0, Implicit0, Implicit1, Thickness, Constraint);
-	}
-	else if (Implicit0->GetType() == TSphere<T, d>::StaticType() && Implicit1->GetType() == TBox<T, d>::StaticType())
-	{
-		ConstructSphereBoxConstraints(Particle0, Particle1, Implicit0, Implicit1, Thickness, Constraint);
-	}
-	else if (Implicit1->GetType() == TSphere<T, d>::StaticType() && Implicit0->GetType() == TBox<T, d>::StaticType())
-	{
-		ConstructSphereBoxConstraints(Particle1, Particle0, Implicit0, Implicit1, Thickness, Constraint);
-	}
-	else if (Implicit0->GetType() == TCapsule<T>::StaticType() && Implicit1->GetType() == TCapsule<T>::StaticType())
-	{
-		ConstructCapsuleCapsuleConstraints(Particle0, Particle1, Implicit0, Implicit1, Thickness, Constraint);
-	}
-	else if (Implicit0->GetType() == TCapsule<T>::StaticType() && Implicit1->GetType() == TBox<T, d>::StaticType())
-	{
-		ConstructCapsuleBoxConstraints(Particle0, Particle1, Implicit0, Implicit1, Thickness, Constraint);
-	}
-	else if (Implicit1->GetType() == TCapsule<T>::StaticType() && Implicit0->GetType() == TBox<T, d>::StaticType())
-	{
-		ConstructCapsuleBoxConstraints(Particle1, Particle0, Implicit0, Implicit1, Thickness, Constraint);
-	}
-	else if (Implicit0->GetType() < TImplicitObjectUnion<T, d>::StaticType() && Implicit1->GetType() == TImplicitObjectUnion<T, d>::StaticType())
-	{
-		ConstructSingleUnionConstraints(Particle0, Particle1, Implicit0, Implicit1, Thickness, Constraint);
-	}
-	else if (Implicit0->GetType() == TImplicitObjectUnion<T, d>::StaticType() && Implicit1->GetType() < TImplicitObjectUnion<T, d>::StaticType())
-	{
-		ConstructSingleUnionConstraints(Particle1, Particle0, Implicit0, Implicit1, Thickness, Constraint);
-	}
-	else if (Implicit0->GetType() == TImplicitObjectUnion<T, d>::StaticType() && Implicit1->GetType() == TImplicitObjectUnion<T, d>::StaticType())
-	{
-		// Union-union creates multiple manifolds - we should never get here for this pair type. See ConstructConstraintsImpl and ConstructUnionUnionConstraints
-		ensure(false);
-	}
-	else if (Implicit0->IsConvex() && Implicit1->IsConvex())
-	{
-		CollisionResolutionConvexConvex<T,d>::ConstructConvexConvexConstraints(Particle0, Particle1, Implicit0, Implicit1, Thickness, Constraint);
-	}
-	else
-	{
-		ConstructLevelsetConstraints(Particle0, Particle1, Implicit0, Implicit1, Thickness, Constraint);
-	}
-}
-
-template<typename T, int d>
-void ConstructConstraintsImpl(TGeometryParticleHandle<T, d>* Particle0, TGeometryParticleHandle<T, d>* Particle1, const FImplicitObject* Implicit0, const FImplicitObject* Implicit1, const T Thickness, TRigidBodyContactConstraint<T, d> & Constraint)
-{
-	// TriangleMesh implicits are for scene query only.
-	if (Implicit0 && GetInnerType(Implicit0->GetType()) == ImplicitObjectType::TriangleMesh) return;
-	if (Implicit1 && GetInnerType(Implicit1->GetType()) == ImplicitObjectType::TriangleMesh) return;
-
-	if (Implicit0 && Implicit0->GetType() == TImplicitObjectUnion<T, d>::StaticType() &&
-		Implicit1 && Implicit1->GetType() == TImplicitObjectUnion<T, d>::StaticType())
-	{
-		ConstructUnionUnionConstraints(Particle0, Particle1, Implicit0, Implicit1, Thickness, Constraint);
-	}
-	else
-	{
-		ConstructPairConstraintImpl(Particle0, Particle1, Implicit0, Implicit1, Thickness, Constraint);
-	}
-}
-
-template<typename T, int d>
-void TPBDCollisionConstraint<T, d>::ConstructConstraints(TGeometryParticleHandle<T, d>* Particle0, TGeometryParticleHandle<T, d>* Particle1, const T Thickness, TRigidBodyContactConstraint<T, d>& Constraint)
-{
-	if (ensure(Particle0 && Particle1))
-	{
-		ConstructConstraintsImpl<T, d>(Particle0, Particle1, Particle0->Geometry().Get(), Particle1->Geometry().Get(), Thickness, Constraint);
-	}
-}
-
-// NOTE: UpdateLevelsetConstraintImp and its <float,3> specialization are here for the Linux build. It looks like
-// GCC does not resolve TPBDCollisionConstraint<T, d> correctly and it won't compile without them.
-template<ECollisionUpdateType UpdateType, typename T, int d>
-void UpdateLevelsetConstraintImp(const T Thickness, TRigidBodyContactConstraint<T, d>& Constraint)
-{
-	TPBDCollisionConstraint<T, d>::UpdateLevelsetConstraint<UpdateType>(Thickness, Constraint);
-}
-
-template<ECollisionUpdateType UpdateType>
-void UpdateLevelsetConstraintImp(const float Thickness, TRigidBodyContactConstraint<float, 3>& Constraint)
-{
-	TPBDCollisionConstraint<float, 3>::UpdateLevelsetConstraint<UpdateType>(Thickness, Constraint);
-}
-
-template<ECollisionUpdateType UpdateType, typename T, int d>
-void UpdateConstraintImp(const FImplicitObject& ParticleObject, const TRigidTransform<T, d>& ParticleTM, const FImplicitObject& LevelsetObject, const TRigidTransform<T, d>& LevelsetTM, const T Thickness, TRigidBodyContactConstraint<T, d>& Constraint)
-{
-	if (ParticleObject.GetType() == TBox<T, d>::StaticType() && LevelsetObject.GetType() == TBox<T, d>::StaticType())
-	{
-		UpdateBoxConstraint(*ParticleObject.template GetObject<TBox<T,d>>(), ParticleTM, *LevelsetObject.template GetObject<TBox<T, d>>(), LevelsetTM, Thickness, Constraint);
-	}
-	else if (ParticleObject.GetType() == TSphere<T, d>::StaticType() && LevelsetObject.GetType() == TSphere<T, d>::StaticType())
-	{
-		UpdateSphereConstraint(*ParticleObject.template GetObject<TSphere<T, d>>(), ParticleTM, *LevelsetObject.template GetObject<TSphere<T, d>>(), LevelsetTM, Thickness, Constraint);
-	}
-	else if (ParticleObject.GetType() == TBox<T, d>::StaticType() && LevelsetObject.GetType() == TPlane<T, d>::StaticType())
-	{
-		UpdateBoxPlaneConstraint(*ParticleObject.template GetObject<TBox<T, d>>(), ParticleTM, *LevelsetObject.template GetObject<TPlane<T, d>>(), LevelsetTM, Thickness, Constraint);
-	}
-	else if (ParticleObject.GetType() == TSphere<T, d>::StaticType() && LevelsetObject.GetType() == TPlane<T, d>::StaticType())
-	{
-		UpdateSpherePlaneConstraint(*ParticleObject.template GetObject<TSphere<T, d>>(), ParticleTM, *LevelsetObject.template GetObject<TPlane<T, d>>(), LevelsetTM, Thickness, Constraint);
-	}
-	else if (ParticleObject.GetType() == TSphere<T, d>::StaticType() && LevelsetObject.GetType() == TBox<T, d>::StaticType())
-	{
-		UpdateSphereBoxConstraint(*ParticleObject.template GetObject<TSphere<T, d>>(), ParticleTM, *LevelsetObject.template GetObject<TBox<T, d>>(), LevelsetTM, Thickness, Constraint);
-	}
-	else if (ParticleObject.GetType() == TCapsule<T>::StaticType() && LevelsetObject.GetType() == TCapsule<T>::StaticType())
-	{
-		UpdateCapsuleCapsuleConstraint(*ParticleObject.template GetObject<TCapsule<T>>(), ParticleTM, *LevelsetObject.template GetObject<TCapsule<T>>(), LevelsetTM, Thickness, Constraint);
-	}
-	else if (ParticleObject.GetType() == TCapsule<T>::StaticType() && LevelsetObject.GetType() == TBox<T, d>::StaticType())
-	{
-		UpdateCapsuleBoxConstraint(*ParticleObject.template GetObject<TCapsule<T>>(), ParticleTM, *LevelsetObject.template GetObject<TBox<T, d>>(), LevelsetTM, Thickness, Constraint);
-	}
-	else if (ParticleObject.GetType() == TPlane<T, d>::StaticType() && LevelsetObject.GetType() == TBox<T, d>::StaticType())
-	{
-		TRigidBodyContactConstraint<T, d> TmpConstraint = Constraint;
-		UpdateBoxPlaneConstraint(*LevelsetObject.template GetObject<TBox<T, d>>(), LevelsetTM, *ParticleObject.template GetObject<TPlane<T, d>>(), ParticleTM, Thickness, TmpConstraint);
-		if (TmpConstraint.GetPhi() < Constraint.GetPhi())
-		{
-			Constraint = TmpConstraint;
-			Constraint.SetNormal(-Constraint.GetNormal());
-		}
-	}
-	else if (ParticleObject.GetType() == TPlane<T, d>::StaticType() && LevelsetObject.GetType() == TSphere<T, d>::StaticType())
-	{
-		TRigidBodyContactConstraint<T, d> TmpConstraint = Constraint;
-		UpdateSpherePlaneConstraint(*LevelsetObject.template GetObject<TSphere<T, d>>(), LevelsetTM, *ParticleObject.template GetObject<TPlane<T, d>>(), ParticleTM, Thickness, TmpConstraint);
-		if (TmpConstraint.GetPhi() < Constraint.GetPhi())
-		{
-			Constraint = TmpConstraint;
-			Constraint.SetNormal(-Constraint.GetNormal());
-		}
-	}
-	else if (ParticleObject.GetType() == TBox<T, d>::StaticType() && LevelsetObject.GetType() == TSphere<T, d>::StaticType())
-	{
-		TRigidBodyContactConstraint<T, d> TmpConstraint = Constraint;
-		UpdateSphereBoxConstraint(*LevelsetObject.template GetObject<TSphere<T, d>>(), LevelsetTM, *ParticleObject.template GetObject<TBox<T, d>>(), ParticleTM, Thickness, TmpConstraint);
-		if (TmpConstraint.GetPhi() < Constraint.GetPhi())
-		{
-			Constraint = TmpConstraint;
-			Constraint.SetNormal(-Constraint.GetNormal());
-		}
-	}
-	else if (ParticleObject.GetType() == TBox<T, d>::StaticType() && LevelsetObject.GetType() == TCapsule<T>::StaticType())
-	{
-		TRigidBodyContactConstraint<T, d> TmpConstraint = Constraint;
-		UpdateCapsuleBoxConstraint(*LevelsetObject.template GetObject<TCapsule<T>>(), LevelsetTM, *ParticleObject.template GetObject<TBox<T, d>>(), ParticleTM, Thickness, TmpConstraint);
-		if (TmpConstraint.GetPhi() < Constraint.GetPhi())
-		{
-			Constraint = TmpConstraint;
-			Constraint.SetNormal(-Constraint.GetNormal());
-		}
-	}
-	else if (ParticleObject.GetType() < TImplicitObjectUnion<T, d>::StaticType() && LevelsetObject.GetType() == TImplicitObjectUnion<T, d>::StaticType())
-	{
-		return UpdateSingleUnionConstraint<UpdateType>(Thickness, Constraint);
-	}
-	else if (ParticleObject.GetType() == TImplicitObjectUnion<T, d>::StaticType() && LevelsetObject.GetType() < TImplicitObjectUnion<T, d>::StaticType())
-	{
-		check(false);	//should not be possible to get this ordering (see ComputeConstraint)
-	}
-	else if (ParticleObject.GetType() == TImplicitObjectUnion<T, d>::StaticType() && LevelsetObject.GetType() == TImplicitObjectUnion<T, d>::StaticType())
-	{
-		return UpdateUnionUnionConstraint<UpdateType>(Thickness, Constraint);
-	}
-	else if (ParticleObject.IsConvex() && LevelsetObject.IsConvex())
-	{
-		CollisionResolutionConvexConvex<T,d>::UpdateConvexConvexConstraint(ParticleObject, ParticleTM, LevelsetObject, LevelsetTM, Thickness, Constraint);
-	}
-	else if (LevelsetObject.IsUnderlyingUnion())
-	{
-		UpdateUnionLevelsetConstraint<UpdateType>(Thickness, Constraint);
-	}
-	else if (ParticleObject.IsUnderlyingUnion())
-	{
-		UpdateLevelsetUnionConstraint<UpdateType>(Thickness, Constraint);
-	}
-	else
-	{
-		UpdateLevelsetConstraintImp<UpdateType>(Thickness, Constraint);
-	}
-}
-
-DECLARE_CYCLE_STAT(TEXT("TPBDCollisionConstraint::UpdateConstraint"), STAT_UpdateConstraint, STATGROUP_ChaosWide);
-
-template<typename T, int d>
-template<ECollisionUpdateType UpdateType>
-void TPBDCollisionConstraint<T, d>::UpdateConstraint(const T Thickness, FRigidBodyContactConstraint& Constraint)
-{
-	SCOPE_CYCLE_COUNTER(STAT_UpdateConstraint);
-
-	Constraint.ResetPhi(Thickness);	
-	const TRigidTransform<T, d> ParticleTM = GetTransform(Constraint.Particle[0]);
-	const TRigidTransform<T, d> LevelsetTM = GetTransform(Constraint.Particle[1]);
-
-	if (!Constraint.Particle[0]->Geometry())
-	{
-		if (Constraint.Particle[1]->Geometry())
-		{
-			if (!Constraint.Particle[1]->Geometry()->IsUnderlyingUnion())
-			{
-				TPBDCollisionConstraint<T, d>::UpdateLevelsetConstraint<UpdateType>(Thickness, Constraint);
-			}
-			else
-			{
-				UpdateUnionLevelsetConstraint<UpdateType>(Thickness, Constraint);
-			}
-		}
-	}
-	else
-	{
-		UpdateConstraintImp<UpdateType>(*Constraint.Particle[0]->Geometry(), ParticleTM, *Constraint.Particle[1]->Geometry(), LevelsetTM, Thickness, Constraint);
-	}
-}
 
-template class TPBDCollisionConstraintHandle<float, 3>;
-template class TAccelerationStructureHandle<float, 3>;
-template class CHAOS_API TPBDCollisionConstraint<float, 3>;
-template void TPBDCollisionConstraint<float, 3>::UpdateConstraint<ECollisionUpdateType::Any>(const float Thickness, FRigidBodyContactConstraint& Constraints);
-template void TPBDCollisionConstraint<float, 3>::UpdateConstraint<ECollisionUpdateType::Deepest>(const float Thickness, FRigidBodyContactConstraint& Constraints);
-template void TPBDCollisionConstraint<float, 3>::UpdateLevelsetConstraint<ECollisionUpdateType::Any>(const float Thickness, FRigidBodyContactConstraint& Constraint);
-template void TPBDCollisionConstraint<float, 3>::UpdateLevelsetConstraint<ECollisionUpdateType::Deepest>(const float Thickness, FRigidBodyContactConstraint& Constraint);
-template void TPBDCollisionConstraint<float, 3>::ComputeConstraints<false>(const TPBDCollisionConstraint<float,3>::FAccelerationStructure&, float Dt);
-template void TPBDCollisionConstraint<float, 3>::ComputeConstraints<true>(const TPBDCollisionConstraint<float, 3>::FAccelerationStructure&, float Dt);
-template void UpdateConstraintImp<ECollisionUpdateType::Any, float, 3>(const FImplicitObject& ParticleObject, const TRigidTransform<float, 3>& ParticleTM, const FImplicitObject& LevelsetObject, const TRigidTransform<float, 3>& LevelsetTM, const float Thickness, TRigidBodyContactConstraint<float, 3>& Constraints);
-template void UpdateConstraintImp<ECollisionUpdateType::Deepest, float, 3>(const FImplicitObject& ParticleObject, const TRigidTransform<float, 3>& ParticleTM, const FImplicitObject& LevelsetObject, const TRigidTransform<float, 3>& LevelsetTM, const float Thickness, TRigidBodyContactConstraint<float, 3>& Constraint);
+	template class TPBDCollisionConstraintHandle<float, 3>;
+	template class TAccelerationStructureHandle<float, 3>;
+	template class CHAOS_API TPBDCollisionConstraint<float, 3>;
+	template void TPBDCollisionConstraint<float, 3>::UpdateConstraint<ECollisionUpdateType::Any>(const float Thickness, FRigidBodyContactConstraint& Constraints);
+	template void TPBDCollisionConstraint<float, 3>::UpdateConstraint<ECollisionUpdateType::Deepest>(const float Thickness, FRigidBodyContactConstraint& Constraints);
+	template void TPBDCollisionConstraint<float, 3>::ComputeConstraints<false>(const TPBDCollisionConstraint<float, 3>::FAccelerationStructure&, float Dt);
+	template void TPBDCollisionConstraint<float, 3>::ComputeConstraints<true>(const TPBDCollisionConstraint<float, 3>::FAccelerationStructure&, float Dt);
 }
diff --git a/Engine/Source/Runtime/Experimental/Chaos/Public/Chaos/CollisionResolutionAlgo.h b/Engine/Source/Runtime/Experimental/Chaos/Public/Chaos/CollisionResolutionAlgo.h
new file mode 100644
--- /dev/null
+++ b/Engine/Source/Runtime/Experimental/Chaos/Public/Chaos/CollisionResolutionAlgo.h
@@ -0,0 +1,119 @@
+// Copyright 1998-2019 Epic Games, Inc. All Rights Reserved.
+
+#pragma once
+#include "Chaos/PBDCollisionTypes.h"
+#include "Chaos/Pair.h"
+
+namespace Chaos
+{
+
+	template <typename T, int d>
+	class TBox;
+
+	template <typename T>
+	class TCapsule;
+
+	template <typename T, int d>
+	class TPlane;
+
+	class FImplicitObject;
+
+	template<typename T, int d>
+	class TRigidBodyContactConstraint;
+
+	template <typename T, int d>
+	class TRigidTransform;
+
+	template <typename T, int d>
+	class TSphere;
+
+	//
+	// Utility
+	//
+	//template <typename T, int d>
+	//TArray<Pair<const FImplicitObject*, TRigidTransform<T, d>>> FindRelevantShapes2(const FImplicitObject* ParticleObj, const TRigidTransform<T, d>& ParticlesTM, const FImplicitObject& LevelsetObj, const TRigidTransform<T, d>& LevelsetTM, const T Thickness);
+
+
+	//
+	// Construct Constraints
+	//
+	template<typename T, int d>
+	void ConstructLevelsetConstraints(TGeometryParticleHandle<T, d>* Particle0, TGeometryParticleHandle<T, d>* Particle1, const FImplicitObject* Implicit0, const FImplicitObject* Implicit1, const T Thickness, TRigidBodyContactConstraint<T, d>& Constraint);
+
+	template<typename T, int d>
+	void ConstructBoxConstraints(TGeometryParticleHandle<T, d>* Particle0, TGeometryParticleHandle<T, d>* Particle1, const FImplicitObject* Implicit0, const FImplicitObject* Implicit1, const T Thickness, TRigidBodyContactConstraint<T, d>& Constraint);
+
+	template<typename T, int d>
+	void ConstructBoxPlaneConstraints(TGeometryParticleHandle<T, d>* Particle0, TGeometryParticleHandle<T, d>* Particle1, const FImplicitObject* Implicit0, const FImplicitObject* Implicit1, const T Thickness, TRigidBodyContactConstraint<T, d>& Constraint);
+
+	template<typename T, int d>
+	void ConstructSphereConstraints(TGeometryParticleHandle<T, d>* Particle0, TGeometryParticleHandle<T, d>* Particle1, const FImplicitObject* Implicit0, const FImplicitObject* Implicit1, const T Thickness, TRigidBodyContactConstraint<T, d>& Constraint);
+
+	template<typename T, int d>
+	void ConstructSpherePlaneConstraints(TGeometryParticleHandle<T, d>* Particle0, TGeometryParticleHandle<T, d>* Particle1, const FImplicitObject* Implicit0, const FImplicitObject* Implicit1, const T Thickness, TRigidBodyContactConstraint<T, d>& Constraint);
+
+	template<typename T, int d>
+	void ConstructSphereBoxConstraints(TGeometryParticleHandle<T, d>* Particle0, TGeometryParticleHandle<T, d>* Particle1, const FImplicitObject* Implicit0, const FImplicitObject* Implicit1, const T Thickness, TRigidBodyContactConstraint<T, d>& Constraint);
+
+	template<typename T, int d>
+	void ConstructCapsuleCapsuleConstraints(TGeometryParticleHandle<T, d>* Particle0, TGeometryParticleHandle<T, d>* Particle1, const FImplicitObject* Implicit0, const FImplicitObject* Implicit1, const T Thickness, TRigidBodyContactConstraint<T, d>& Constraint);
+
+	template<typename T, int d>
+	void ConstructCapsuleBoxConstraints(TGeometryParticleHandle<T, d>* Particle0, TGeometryParticleHandle<T, d>* Particle1, const FImplicitObject* Implicit0, const FImplicitObject* Implicit1, const T Thickness, TRigidBodyContactConstraint<T, d>& Constraint);
+
+	template<typename T, int d>
+	void ConstructSingleUnionConstraints(TGeometryParticleHandle<T, d>* Particle0, TGeometryParticleHandle<T, d>* Particle1, const FImplicitObject* Implicit0, const FImplicitObject* Implicit1, const T Thickness, TRigidBodyContactConstraint<T, d>& Constraint);
+
+	template<typename T, int d>
+	void ConstructUnionUnionConstraints(TGeometryParticleHandle<T, d>* Particle0, TGeometryParticleHandle<T, d>* Particle1, const FImplicitObject* Implicit0, const FImplicitObject* Implicit1, const T Thickness, TRigidBodyContactConstraint<T, d>& Constraint);
+
+	template<typename T, int d>
+	void ConstructPairConstraintImpl(TGeometryParticleHandle<T, d>* Particle0, TGeometryParticleHandle<T, d>* Particle1, const FImplicitObject* Implicit0, const FImplicitObject* Implicit1, const T Thickness, TRigidBodyContactConstraint<T, d>& Constraint);
+
+	template<typename T, int d>
+	void ConstructConstraintsImpl(TGeometryParticleHandle<T, d>* Particle0, TGeometryParticleHandle<T, d>* Particle1, const FImplicitObject* Implicit0, const FImplicitObject* Implicit1, const T Thickness, TRigidBodyContactConstraint<T, d> & Constraint);
+
+
+	//
+	// Update Constraints
+	//
+	template<ECollisionUpdateType UpdateType, typename T = float, int d = 3>
+	void CHAOS_API UpdateLevelsetConstraint(const T Thickness, TRigidBodyContactConstraint<float, 3>& Constraint);
+
+	template<ECollisionUpdateType UpdateType, typename T, int d>
+	void UpdateConstraintImp(const FImplicitObject& ParticleObject, const TRigidTransform<T, d>& ParticleTM, const FImplicitObject& LevelsetObject, const TRigidTransform<T, d>& LevelsetTM, const T Thickness, TRigidBodyContactConstraint<T, d>& Constraint);
+
+	template <typename T, int d>
+	bool UpdateBoxPlaneConstraint(const TBox<T, d>& Box, const TRigidTransform<T, d>& BoxTransform, const TPlane<T, d>& Plane, const TRigidTransform<T, d>& PlaneTransform, T Thickness, TRigidBodyContactConstraint<T, d>& Constraint);
+
+	template <typename T, int d>
+	void UpdateSphereConstraint(const TSphere<T, d>& Sphere1, const TRigidTransform<T, d>& Sphere1Transform, const TSphere<T, d>& Sphere2, const TRigidTransform<T, d>& Sphere2Transform, T Thickness, TRigidBodyContactConstraint<T, d>& Constraint);
+
+	template <typename T, int d>
+	void UpdateSpherePlaneConstraint(const TSphere<T, d>& Sphere, const TRigidTransform<T, d>& SphereTransform, const TPlane<T, d>& Plane, const TRigidTransform<T, d>& PlaneTransform, T Thickness, TRigidBodyContactConstraint<T, d>& Constraint);
+
+	template <typename T, int d>
+	void UpdateSphereBoxConstraint(const TSphere<T, d>& Sphere, const TRigidTransform<T, d>& SphereTransform, const TBox<T, d>& Box, const TRigidTransform<T, d>& BoxTransform, T Thickness, TRigidBodyContactConstraint<T, d>& Constraint);
+
+	template <typename T, int d>
+	void UpdateCapsuleCapsuleConstraint(const TCapsule<T>& A, const TRigidTransform<T, d>& ATransform, const TCapsule<T>& B, const TRigidTransform<T, d>& BTransform, T Thickness, TRigidBodyContactConstraint<T, d>& Constraint);
+
+	template <typename T, int d>
+	void UpdateCapsuleBoxConstraint(const TCapsule<T>& A, const TRigidTransform<T, d>& ATransform, const TBox<T, d>& B, const TRigidTransform<T, d>& BTransform, T Thickness, TRigidBodyContactConstraint<T, d>& Constraint);
+
+	template<ECollisionUpdateType UpdateType, typename T, int d>
+	void UpdateUnionUnionConstraint(const T Thickness, TRigidBodyContactConstraint<T, d>& Constraint);
+
+	template<ECollisionUpdateType UpdateType, typename T, int d>
+	void UpdateSingleUnionConstraint(const T Thickness, TRigidBodyContactConstraint<T, d>& Constraint);
+
+	template<ECollisionUpdateType UpdateType, typename T, int d>
+	void UpdateUnionLevelsetConstraint(const T Thickness, TRigidBodyContactConstraint<T, d>& Constraint);
+
+	template<ECollisionUpdateType UpdateType, typename T, int d>
+	void UpdateLevelsetUnionConstraint(const T Thickness, TRigidBodyContactConstraint<T, d>& Constraint);
+
+	template <typename T, int d>
+	void UpdateBoxConstraint(const TBox<T, d>& Box1, const TRigidTransform<T, d>& Box1Transform, const TBox<T, d>& Box2, const TRigidTransform<T, d>& Box2Transform, T Thickness, TRigidBodyContactConstraint<T, d>& Constraint);
+
+}
diff --git a/Engine/Source/Runtime/Experimental/Chaos/Public/Chaos/PBDCollisionConstraint.h b/Engine/Source/Runtime/Experimental/Chaos/Public/Chaos/PBDCollisionConstraint.h
--- a/Engine/Source/Runtime/Experimental/Chaos/Public/Chaos/PBDCollisionConstraint.h
+++ b/Engine/Source/Runtime/Experimental/Chaos/Public/Chaos/PBDCollisionConstraint.h
@@ -14,280 +14,252 @@
 
 namespace Chaos
 {
-template<typename T, int d>
-class TPBDCollisionConstraintAccessor;
+	template<typename T, int d>
+	class TPBDCollisionConstraintAccessor;
 
-template<typename T, int d>
-class TPBDCollisionConstraint;
+	template<typename T, int d>
+	class TPBDCollisionConstraint;
 
-template <typename T, int d>
-class TRigidTransform;
+	template <typename T, int d>
+	class TRigidTransform;
 
-class FImplicitObject;
+	class FImplicitObject;
 
-template <typename T, int d>
-class TBVHParticles;
+	template <typename T, int d>
+	class TBVHParticles;
 
-template <typename T, int d>
-class TBox;
+	template <typename T, int d>
+	class TBox;
 
-template<class T>
-class TChaosPhysicsMaterial;
+	template<class T>
+	class TChaosPhysicsMaterial;
 
-template <typename T, int d>
-class TPBDRigidsSOAs;
+	template <typename T, int d>
+	class TPBDRigidsSOAs;
 
-/** Specifies the type of work we should do*/
-enum class ECollisionUpdateType
-{
-	Any,	//stop if we have at least one deep penetration. Does not compute location or normal
-	Deepest	//find the deepest penetration. Compute location and normal
-};
-
-/** Return value of the collision modification callback */
-enum class ECollisionModifierResult
-{
-	Unchanged,	/** No change to the collision */
-	Modified,	/** Modified the collision, but want it to remain enabled */
-	Disabled,	/** Collision should be disabled */
-};
-
-
-
-template<typename T, int d>
-using TRigidBodyContactConstraintsPostComputeCallback = TFunction<void()>;
-
-template<typename T, int d>
-using TRigidBodyContactConstraintsPostApplyCallback = TFunction<void(const T Dt, const TArray<TPBDCollisionConstraintHandle<T, d>*>& InConstraintHandles)>;
-
-template<typename T, int d>
-using TRigidBodyContactConstraintsPostApplyPushOutCallback = TFunction<void(const T Dt, const TArray<TPBDCollisionConstraintHandle<T, d>*>& InConstraintHandles, bool bRequiresAnotherIteration)>;
-
-/**
- * Manages a set of contact constraints:
- *	- Performs collision detection to generate constraints.
- *	- Responsible for applying corrections to particles affected by the constraints.
- * @todo(ccaulfield): rename to TPBDCollisionConstraints
- * @todo(ccaulfield): remove TPBDCollisionConstraintAccessor
- * @todo(ccaulfield): separate collision detection from constraint container
- */
-template<typename T, int d>
-class CHAOS_API TPBDCollisionConstraint : public FPBDConstraintContainer
-{
-public:
-	using Base = FPBDConstraintContainer;
-	friend class TPBDCollisionConstraintAccessor<T, d>;
-	friend class TPBDCollisionConstraintHandle<T, d>;
-	using FReal = T;
-	static const int Dimensions = d;
-	using FConstraintContainerHandle = TPBDCollisionConstraintHandle<T, d>;
-	using FConstraintHandleAllocator = TConstraintHandleAllocator<TPBDCollisionConstraint<T, d>>;
-	using FRigidBodyContactConstraint = TRigidBodyContactConstraint<T, d>;
-	using FConstraintHandleID = TPair<const TGeometryParticleHandle<T, d>*, const TGeometryParticleHandle<T, d>*>;
-
-	TPBDCollisionConstraint(const TPBDRigidsSOAs<T,d>& InParticles, TArrayCollectionArray<bool>& Collided, const TArrayCollectionArray<TSerializablePtr<FChaosPhysicsMaterial>>& PerParticleMaterials, const int32 PairIterations = 1, const T Thickness = (T)0);
-	virtual ~TPBDCollisionConstraint() {}
-
-	//
-	// Constraint Container API
-	//
-
-	void SetThickness(T InThickness)
-	{
-		MThickness = InThickness;
-	}
+	template<typename T, int d>
+	using TRigidBodyContactConstraintsPostComputeCallback = TFunction<void()>;
 
-	void SetVelocityInflation(float ScaleFactor)
-	{
-		CollisionVelocityInflation = ScaleFactor;
-	}
-
-	void SetVelocitySolveEnabled(bool bInEnableVelocitySolve)
-	{
-		bEnableVelocitySolve = bInEnableVelocitySolve;
-	}
+	template<typename T, int d>
+	using TRigidBodyContactConstraintsPostApplyCallback = TFunction<void(const T Dt, const TArray<TPBDCollisionConstraintHandle<T, d>*>& InConstraintHandles)>;
 
-	void SetPushOutPairIterations(int32 InPairIterations)
-	{ 
-		MPairIterations = InPairIterations;
-	}
-
-	void SetCollisionsEnabled(bool bInEnableCollisions)
-	{
-		bEnableCollisions = bInEnableCollisions;
-	}
+	template<typename T, int d>
+	using TRigidBodyContactConstraintsPostApplyPushOutCallback = TFunction<void(const T Dt, const TArray<TPBDCollisionConstraintHandle<T, d>*>& InConstraintHandles, bool bRequiresAnotherIteration)>;
 
 	/**
-	 * Get the number of constraints.
+	 * Manages a set of contact constraints:
+	 *	- Performs collision detection to generate constraints.
+	 *	- Responsible for applying corrections to particles affected by the constraints.
+	 * @todo(ccaulfield): rename to TPBDCollisionConstraints
+	 * @todo(ccaulfield): remove TPBDCollisionConstraintAccessor
+	 * @todo(ccaulfield): separate collision detection from constraint container
 	 */
-	int32 NumConstraints() const
-	{
-		return Constraints.Num();
-	}
-
-	FConstraintHandleID GetConstraintHandleID(const FRigidBodyContactConstraint & Constraint) const
+	template<typename T, int d>
+	class CHAOS_API TPBDCollisionConstraint : public FPBDConstraintContainer
 	{
-		return  FConstraintHandleID( Constraint.Particle[0], Constraint.Particle[1] );
-	}
-
-	FConstraintHandleID GetConstraintHandleID(int32 ConstraintIndex) const
-	{
-		return  FConstraintHandleID( Constraints[ConstraintIndex].Particle[0], Constraints[ConstraintIndex].Particle[1] );
-	}
-		
-	const FConstraintContainerHandle* GetConstraintHandle(int32 ConstraintIndex) const
-	{
-		return Handles[ConstraintIndex];
-	}
-
-	FConstraintContainerHandle* GetConstraintHandle(int32 ConstraintIndex)
-	{
-		return Handles[ ConstraintIndex ];
-	}
-
-	void RemoveConstraint(int32 Idx);
-
-	/**
-	* Remove the constraints associated with the ParticleHandle.
-	*/
-	void RemoveConstraints(const TSet<TGeometryParticleHandle<T, d>*>&  ParticleHandle);
-
-	/**
-	 * Apply a modifier to the constraints and specify which constraints should be disabled.
-	 * You would probably call this in the PostComputeCallback. Prefer this to calling RemoveConstraints in a loop, 
-	 * so you don't have to worry about constraint iterator/indices changing.
-	 */
-	void ApplyCollisionModifier(const TFunction<ECollisionModifierResult(FRigidBodyContactConstraint& Constraint)>& CollisionModifier);
-
-	/**
-	 * Set the callback used just after contacts are generated at the start of a frame tick.
-	 * This can be used to modify or disable constraints (via ApplyCollisionModifier).
-	 */
-	void SetPostComputeCallback(const TRigidBodyContactConstraintsPostComputeCallback<T, d>& Callback);
-
-	/**
-	 * Remove the constraint callback.
-	 */
-	void ClearPostComputeCallback();
-
-	void SetPostApplyCallback(const TRigidBodyContactConstraintsPostApplyCallback<T, d>& Callback);
-	void ClearPostApplyCallback();
-
-	void SetPostApplyPushOutCallback(const TRigidBodyContactConstraintsPostApplyPushOutCallback<T, d>& Callback);
-	void ClearPostApplyPushOutCallback();
-
-	//
-	// Constraint API
-	//
-
-	/**
-	 * Get the particles that are affected by the specified constraint.
-	 */
-	TVector<TGeometryParticleHandle<T, d>*, 2> GetConstrainedParticles(int32 ConstraintIndex) const
-	{
-		return { Constraints[ConstraintIndex].Particle[0], Constraints[ConstraintIndex].Particle[1] };
-	}
-
-
-	//
-	// Island Rule API
-	//
-
-	/**
-	 * Apply corrections for the specified list of constraints.
-	 */
-	 // @todo(ccaulfield): this runs wide. The serial/parallel decision should be in the ConstraintRule
-	void Apply(const T Dt, const TArray<FConstraintContainerHandle*>& InConstraintHandles, const int32 It, const int32 NumIts);
-
-	/**
-	 * Generate all contact constraints.
-	 */
-	void UpdatePositionBasedState(/*const TPBDRigidParticles<T, d>& InParticles, const TArray<int32>& InIndices,*/ const T Dt);
-
-	/** 
-	 * Apply push out for the specified list of constraints.
-	 * Return true if we need another iteration 
-	 */
-	 // @todo(ccaulfield): this runs wide. The serial/parallel decision should be in the ConstraintRule
-	bool ApplyPushOut(const T Dt, const TArray<FConstraintContainerHandle*>& InConstraintHandles, const TSet<TGeometryParticleHandle<T,d>*>& IsTemporarilyStatic, int32 Iteration, int32 NumIterations);
-
-	void UpdateConstraints(/*const TPBDRigidParticles<T, d>& InParticles, const TArray<int32>& InIndices,*/ T Dt, const TSet<TGeometryParticleHandle<T, d>*>& AddedParticles);
-
-	const TArray<FRigidBodyContactConstraint>& GetAllConstraints() const { return Constraints; }
-
-	//using FAccelerationStructure = TBoundingVolume<TAccelerationStructureHandle<T, d>, T, d>;
-	//using FAccelerationStructure = TAABBTree<TAccelerationStructureHandle<T, d>, TAABBTreeLeafArray<TAccelerationStructureHandle<T, d>, T>, T>;
-	using FAccelerationStructure = ISpatialAcceleration< TAccelerationStructureHandle<T, d>, T, d>;
-	void SetSpatialAcceleration(const FAccelerationStructure* InSpatialAcceleration) { SpatialAcceleration = InSpatialAcceleration; }
-
-	// @todo(ccaulfield) We should probably wrap the LevelsetConstraint functions in a utility class and remove them from here.
-	// They are currently public for the Linux build, and member functions for Headless Chaos (TPBDCollisionConstraintAccessor). 
-	template<ECollisionUpdateType>
-	static void UpdateLevelsetConstraint(const T Thickness, FRigidBodyContactConstraint& Constraint);
-
-	/**
-	 * Get the particles that are affected by the specified constraint.
-	 */
-	TVector<TGeometryParticleHandle<T, d>*, 2> ConstraintParticles(int32 ContactIndex) const { return { Constraints[ContactIndex].Particle[0], Constraints[ContactIndex].Particle[1] }; }
-
-	//NOTE: this should not be called by anyone other than ISpatialAccelerationCollection and CollisionConstraints - todo: make private with friends?
-	template <bool bGatherStats, typename SPATIAL_ACCELERATION>
-	void ComputeConstraintsHelperLowLevel(const SPATIAL_ACCELERATION& SpatialAcceleration, T Dt);
-	
-protected:
-	using Base::GetConstraintIndex;
-	using Base::SetConstraintIndex;
-
-private:
-	void Reset(/*const TPBDRigidParticles<T, d>& InParticles, const TArray<int32>& InIndices*/);
-
-	void Apply(const T Dt, FRigidBodyContactConstraint& Constraint);
-	void ApplyPushOut(const T Dt, FRigidBodyContactConstraint& Constraint, const TSet<TGeometryParticleHandle<T,d>*>& IsTemporarilyStatic, int32 Iteration, int32 NumIterations, bool& NeedsAnotherIteration);
-
-	template <bool bGatherStats = false>
-	void ComputeConstraints(const FAccelerationStructure& AccelerationStructure, T Dt);
-
-	template<ECollisionUpdateType>
-	void UpdateConstraint(const T Thickness, FRigidBodyContactConstraint& Constraint);
-
-	void ConstructConstraints(TGeometryParticleHandle<T, d>* Particle0, TGeometryParticleHandle<T, d>* Particle1, const T Thickness, TRigidBodyContactConstraint<T,d>& Constraint);
-
-
-	template<typename SPATIAL_ACCELERATION>
-	void UpdateConstraintsHelper(/*const TPBDRigidParticles<T, d>& InParticles, const TArray<int32>& InIndices,*/ T Dt, const TSet<TGeometryParticleHandle<T, d>*>& AddedParticles, SPATIAL_ACCELERATION& SpatialAcceleration);
-
-	const TPBDRigidsSOAs<T,d>& Particles;
-
-	// @todo(ccaulfield): move spatial acceleration out of constraint container and make shareable
-	const FAccelerationStructure* SpatialAcceleration;
+	public:
+		using Base = FPBDConstraintContainer;
+		friend class TPBDCollisionConstraintAccessor<T, d>;
+		friend class TPBDCollisionConstraintHandle<T, d>;
+		using FReal = T;
+		static const int Dimensions = d;
+		using FConstraintContainerHandle = TPBDCollisionConstraintHandle<T, d>;
+		using FConstraintHandleAllocator = TConstraintHandleAllocator<TPBDCollisionConstraint<T, d>>;
+		using FRigidBodyContactConstraint = TRigidBodyContactConstraint<T, d>;
+		using FConstraintHandleID = TPair<const TGeometryParticleHandle<T, d>*, const TGeometryParticleHandle<T, d>*>;
+
+		TPBDCollisionConstraint(const TPBDRigidsSOAs<T, d>& InParticles, TArrayCollectionArray<bool>& Collided, const TArrayCollectionArray<TSerializablePtr<FChaosPhysicsMaterial>>& PerParticleMaterials, const int32 PairIterations = 1, const T Thickness = (T)0);
+		virtual ~TPBDCollisionConstraint() {}
+
+		//
+		// Constraint Container API
+		//
+
+		void SetThickness(T InThickness)
+		{
+			MThickness = InThickness;
+		}
+
+		void SetVelocityInflation(float ScaleFactor)
+		{
+			CollisionVelocityInflation = ScaleFactor;
+		}
+
+		void SetVelocitySolveEnabled(bool bInEnableVelocitySolve)
+		{
+			bEnableVelocitySolve = bInEnableVelocitySolve;
+		}
+
+		void SetPushOutPairIterations(int32 InPairIterations)
+		{
+			MPairIterations = InPairIterations;
+		}
+
+		void SetCollisionsEnabled(bool bInEnableCollisions)
+		{
+			bEnableCollisions = bInEnableCollisions;
+		}
+
+		/**
+		 * Get the number of constraints.
+		 */
+		int32 NumConstraints() const
+		{
+			return Constraints.Num();
+		}
+
+		FConstraintHandleID GetConstraintHandleID(const FRigidBodyContactConstraint & Constraint) const
+		{
+			return  FConstraintHandleID(Constraint.Particle[0], Constraint.Particle[1]);
+		}
+
+		FConstraintHandleID GetConstraintHandleID(int32 ConstraintIndex) const
+		{
+			return  FConstraintHandleID(Constraints[ConstraintIndex].Particle[0], Constraints[ConstraintIndex].Particle[1]);
+		}
+
+		const FConstraintContainerHandle* GetConstraintHandle(int32 ConstraintIndex) const
+		{
+			return Handles[ConstraintIndex];
+		}
+
+		FConstraintContainerHandle* GetConstraintHandle(int32 ConstraintIndex)
+		{
+			return Handles[ConstraintIndex];
+		}
+
+		void RemoveConstraint(int32 Idx);
+
+		/**
+		* Remove the constraints associated with the ParticleHandle.
+		*/
+		void RemoveConstraints(const TSet<TGeometryParticleHandle<T, d>*>&  ParticleHandle);
+
+		/**
+		 * Apply a modifier to the constraints and specify which constraints should be disabled.
+		 * You would probably call this in the PostComputeCallback. Prefer this to calling RemoveConstraints in a loop,
+		 * so you don't have to worry about constraint iterator/indices changing.
+		 */
+		void ApplyCollisionModifier(const TFunction<ECollisionModifierResult(FRigidBodyContactConstraint& Constraint)>& CollisionModifier);
+
+		/**
+		 * Set the callback used just after contacts are generated at the start of a frame tick.
+		 * This can be used to modify or disable constraints (via ApplyCollisionModifier).
+		 */
+		void SetPostComputeCallback(const TRigidBodyContactConstraintsPostComputeCallback<T, d>& Callback);
+
+		/**
+		 * Remove the constraint callback.
+		 */
+		void ClearPostComputeCallback();
+
+		void SetPostApplyCallback(const TRigidBodyContactConstraintsPostApplyCallback<T, d>& Callback);
+		void ClearPostApplyCallback();
+
+		void SetPostApplyPushOutCallback(const TRigidBodyContactConstraintsPostApplyPushOutCallback<T, d>& Callback);
+		void ClearPostApplyPushOutCallback();
+
+		//
+		// Constraint API
+		//
+
+		/**
+		 * Get the particles that are affected by the specified constraint.
+		 */
+		TVector<TGeometryParticleHandle<T, d>*, 2> GetConstrainedParticles(int32 ConstraintIndex) const
+		{
+			return { Constraints[ConstraintIndex].Particle[0], Constraints[ConstraintIndex].Particle[1] };
+		}
+
+
+		//
+		// Island Rule API
+		//
+
+		/**
+		 * Apply corrections for the specified list of constraints.
+		 */
+		 // @todo(ccaulfield): this runs wide. The serial/parallel decision should be in the ConstraintRule
+		void Apply(const T Dt, const TArray<FConstraintContainerHandle*>& InConstraintHandles, const int32 It, const int32 NumIts);
+
+		/**
+		 * Generate all contact constraints.
+		 */
+		void UpdatePositionBasedState(/*const TPBDRigidParticles<T, d>& InParticles, const TArray<int32>& InIndices,*/ const T Dt);
+
+		/**
+		 * Apply push out for the specified list of constraints.
+		 * Return true if we need another iteration
+		 */
+		 // @todo(ccaulfield): this runs wide. The serial/parallel decision should be in the ConstraintRule
+		bool ApplyPushOut(const T Dt, const TArray<FConstraintContainerHandle*>& InConstraintHandles, const TSet<TGeometryParticleHandle<T, d>*>& IsTemporarilyStatic, int32 Iteration, int32 NumIterations);
+
+		void UpdateConstraints(T Dt, const TSet<TGeometryParticleHandle<T, d>*>& AddedParticles);
+
+		const TArray<FRigidBodyContactConstraint>& GetAllConstraints() const { return Constraints; }
+
+		//using FAccelerationStructure = TBoundingVolume<TAccelerationStructureHandle<T, d>, T, d>;
+		//using FAccelerationStructure = TAABBTree<TAccelerationStructureHandle<T, d>, TAABBTreeLeafArray<TAccelerationStructureHandle<T, d>, T>, T>;
+		using FAccelerationStructure = ISpatialAcceleration< TAccelerationStructureHandle<T, d>, T, d>;
+		void SetSpatialAcceleration(const FAccelerationStructure* InSpatialAcceleration) { SpatialAcceleration = InSpatialAcceleration; }
+
+		/**
+		 * Get the particles that are affected by the specified constraint.
+		 */
+		TVector<TGeometryParticleHandle<T, d>*, 2> ConstraintParticles(int32 ContactIndex) const { return { Constraints[ContactIndex].Particle[0], Constraints[ContactIndex].Particle[1] }; }
+
+		//NOTE: this should not be called by anyone other than ISpatialAccelerationCollection and CollisionConstraints - todo: make private with friends?
+		template <bool bGatherStats, typename SPATIAL_ACCELERATION>
+		void ComputeConstraintsHelperLowLevel(const SPATIAL_ACCELERATION& SpatialAcceleration, T Dt);
+
+	protected:
+		using Base::GetConstraintIndex;
+		using Base::SetConstraintIndex;
+
+	private:
+		void Reset(/*const TPBDRigidParticles<T, d>& InParticles, const TArray<int32>& InIndices*/);
+
+		void Apply(const T Dt, FRigidBodyContactConstraint& Constraint);
+		void ApplyPushOut(const T Dt, FRigidBodyContactConstraint& Constraint, const TSet<TGeometryParticleHandle<T, d>*>& IsTemporarilyStatic, int32 Iteration, int32 NumIterations, bool& NeedsAnotherIteration);
+
+		template <bool bGatherStats = false>
+		void ComputeConstraints(const FAccelerationStructure& AccelerationStructure, T Dt);
+
+		template<ECollisionUpdateType>
+		void UpdateConstraint(const T Thickness, FRigidBodyContactConstraint& Constraint);
 
-	TArray<FRigidBodyContactConstraint> Constraints;
-	TArrayCollectionArray<bool>& MCollided;
-	const TArrayCollectionArray<TSerializablePtr<FChaosPhysicsMaterial>>& MPhysicsMaterials;
-	bool bEnableVelocitySolve;
-	int32 MPairIterations;
-	T MThickness;
-	T MAngularFriction;
-	bool bUseCCD;
-	bool bEnableCollisions;
+		void ConstructConstraints(TGeometryParticleHandle<T, d>* Particle0, TGeometryParticleHandle<T, d>* Particle1, const T Thickness, TRigidBodyContactConstraint<T, d>& Constraint);
+
+		const TPBDRigidsSOAs<T, d>& Particles;
+
+		// @todo(ccaulfield): move spatial acceleration out of constraint container and make shareable
+		const FAccelerationStructure* SpatialAcceleration;
+
+		TArray<FRigidBodyContactConstraint> Constraints;
+		TArrayCollectionArray<bool>& MCollided;
+		const TArrayCollectionArray<TSerializablePtr<FChaosPhysicsMaterial>>& MPhysicsMaterials;
+		bool bEnableVelocitySolve;
+		int32 MPairIterations;
+		T MThickness;
+		T MAngularFriction;
+		bool bUseCCD;
+		bool bEnableCollisions;
 
-	int32 LifespanCounter;
-	float CollisionVelocityInflation;
+		int32 LifespanCounter;
+		float CollisionVelocityInflation;
 
-	TRigidBodyContactConstraintsPostComputeCallback<T, d> PostComputeCallback;
-	TRigidBodyContactConstraintsPostApplyCallback<T, d> PostApplyCallback;
-	TRigidBodyContactConstraintsPostApplyPushOutCallback<T, d> PostApplyPushOutCallback;
+		TRigidBodyContactConstraintsPostComputeCallback<T, d> PostComputeCallback;
+		TRigidBodyContactConstraintsPostApplyCallback<T, d> PostApplyCallback;
+		TRigidBodyContactConstraintsPostApplyPushOutCallback<T, d> PostApplyPushOutCallback;
 
-	TArray<FConstraintContainerHandle*> Handles;
-	FConstraintHandleAllocator HandleAllocator;
+		TArray<FConstraintContainerHandle*> Handles;
+		FConstraintHandleAllocator HandleAllocator;
 
-	TMap< FConstraintHandleID, FConstraintContainerHandle* > Manifolds;
-};
+		TMap< FConstraintHandleID, FConstraintContainerHandle* > Manifolds;
+	};
 
-extern template void TPBDCollisionConstraint<float, 3>::UpdateConstraint<ECollisionUpdateType::Any>(const float Thickness, FRigidBodyContactConstraint& Constraint);
-extern template void TPBDCollisionConstraint<float, 3>::UpdateConstraint<ECollisionUpdateType::Deepest>(const float Thickness, FRigidBodyContactConstraint& Constraint);
-extern template void TPBDCollisionConstraint<float, 3>::UpdateLevelsetConstraint<ECollisionUpdateType::Any>(const float Thickness, FRigidBodyContactConstraint& Constraint);
-extern template void TPBDCollisionConstraint<float, 3>::UpdateLevelsetConstraint<ECollisionUpdateType::Deepest>(const float Thickness, FRigidBodyContactConstraint& Constraint);
-extern template void TPBDCollisionConstraint<float, 3>::ComputeConstraints<false>(const TPBDCollisionConstraint<float, 3>::FAccelerationStructure&, float Dt);
-extern template void TPBDCollisionConstraint<float, 3>::ComputeConstraints<true>(const TPBDCollisionConstraint<float, 3>::FAccelerationStructure&, float Dt);
+	extern template void TPBDCollisionConstraint<float, 3>::UpdateConstraint<ECollisionUpdateType::Any>(const float Thickness, FRigidBodyContactConstraint& Constraint);
+	extern template void TPBDCollisionConstraint<float, 3>::UpdateConstraint<ECollisionUpdateType::Deepest>(const float Thickness, FRigidBodyContactConstraint& Constraint);
+	extern template void TPBDCollisionConstraint<float, 3>::ComputeConstraints<false>(const TPBDCollisionConstraint<float, 3>::FAccelerationStructure&, float Dt);
+	extern template void TPBDCollisionConstraint<float, 3>::ComputeConstraints<true>(const TPBDCollisionConstraint<float, 3>::FAccelerationStructure&, float Dt);
 }
diff --git a/Engine/Source/Runtime/Experimental/Chaos/Public/Chaos/PBDCollisionTypes.h b/Engine/Source/Runtime/Experimental/Chaos/Public/Chaos/PBDCollisionTypes.h
--- a/Engine/Source/Runtime/Experimental/Chaos/Public/Chaos/PBDCollisionTypes.h
+++ b/Engine/Source/Runtime/Experimental/Chaos/Public/Chaos/PBDCollisionTypes.h
@@ -15,6 +15,22 @@ namespace Chaos
 	class TPBDCollisionConstraint;
 
 
+	/** Specifies the type of work we should do*/
+	enum class CHAOS_API ECollisionUpdateType
+	{
+		Any,	//stop if we have at least one deep penetration. Does not compute location or normal
+		Deepest	//find the deepest penetration. Compute location and normal
+	};
+
+	/** Return value of the collision modification callback */
+	enum class CHAOS_API ECollisionModifierResult
+	{
+		Unchanged,	/** No change to the collision */
+		Modified,	/** Modified the collision, but want it to remain enabled */
+		Disabled,	/** Collision should be disabled */
+	};
+
+
 	/*
 	*
 	*  Contact manifold stored in the local space of the
@@ -73,7 +89,7 @@ namespace Chaos
 
 		//API
 		void ResetPhi(T InPhi) { SetPhi(InPhi); }
-		bool ContainsManifold(const FImplicitObject* B = nullptr, const FImplicitObject* A = nullptr) const { return (A==nullptr&&B==nullptr)?(ShapeManifold.Implicit[0] != nullptr && ShapeManifold.Implicit[1] != nullptr):(ShapeManifold.Implicit[0]==A && ShapeManifold.Implicit[1]==B); }
+		bool ContainsManifold(const FImplicitObject* B = nullptr, const FImplicitObject* A = nullptr) const { return (A == nullptr&&B == nullptr) ? (ShapeManifold.Implicit[0] != nullptr && ShapeManifold.Implicit[1] != nullptr) : (ShapeManifold.Implicit[0] == A && ShapeManifold.Implicit[1] == B); }
 
 		void SetDisabled(bool bInDisabled, const FImplicitObject* B = nullptr, const FImplicitObject* A = nullptr) { ShapeManifold.Manifold.bDisabled = bInDisabled; }
 		TVector<T, d> GetDisabled(const FImplicitObject* B = nullptr, const FImplicitObject* A = nullptr) const { return ShapeManifold.Manifold.bDisabled; }
@@ -82,7 +98,7 @@ namespace Chaos
 		void SetNormal(const TVector<T, d> & InNormal, const FImplicitObject* B = nullptr, const FImplicitObject* A = nullptr) { ShapeManifold.Manifold.Normal = InNormal; }
 		TVector<T, d> GetNormal(const FImplicitObject* B = nullptr, const FImplicitObject* A = nullptr) const { return ShapeManifold.Manifold.Normal; }
 
-		void SetLocation(const TVector<T, d> & InLocation, const FImplicitObject* B = nullptr, const FImplicitObject* A = nullptr, int32 Index=0) { ShapeManifold.Manifold.Location = InLocation; }
+		void SetLocation(const TVector<T, d> & InLocation, const FImplicitObject* B = nullptr, const FImplicitObject* A = nullptr, int32 Index = 0) { ShapeManifold.Manifold.Location = InLocation; }
 		TVector<T, d> GetLocation(const FImplicitObject* B = nullptr, const FImplicitObject* A = nullptr) const { return ShapeManifold.Manifold.Location; }
 
 		void SetPhi(T InPhi, const FImplicitObject* B = nullptr, const FImplicitObject* A = nullptr, int32 Index = 0) { ShapeManifold.Manifold.Phi = InPhi; }
