- Fix Bad Merge

diff --git a/Engine/Source/Runtime/Experimental/Chaos/Private/Chaos/ClusterCreationParameters.cpp b/Engine/Source/Runtime/Experimental/Chaos/Private/Chaos/ClusterCreationParameters.cpp
--- a/Engine/Source/Runtime/Experimental/Chaos/Private/Chaos/ClusterCreationParameters.cpp
+++ b/Engine/Source/Runtime/Experimental/Chaos/Private/Chaos/ClusterCreationParameters.cpp
@@ -1,6 +1,5 @@
 // Copyright 1998-2019 Epic Games, Inc. All Rights Reserved.
 
-#pragma once
 #include "Chaos/ClusterCreationParameters.h"
 
 namespace Chaos
