Specitive dll clang fix.

diff --git a/Engine/Source/Runtime/Experimental/Chaos/Public/Chaos/GeometryParticles.h b/Engine/Source/Runtime/Experimental/Chaos/Public/Chaos/GeometryParticles.h
--- a/Engine/Source/Runtime/Experimental/Chaos/Public/Chaos/GeometryParticles.h
+++ b/Engine/Source/Runtime/Experimental/Chaos/Public/Chaos/GeometryParticles.h
@@ -248,12 +248,12 @@ namespace Chaos
 			return MSpatialIdx[Index];
 		}
 
-		uint32 HashResultLowLevel(const int32 Index) const
+		CHAOS_API uint32 HashResultLowLevel(const int32 Index) const
 		{
 			return MHashResult[Index];
 		}
 
-		uint32& HashResultLowLevel(const int32 Index)
+		CHAOS_API uint32& HashResultLowLevel(const int32 Index)
 		{
 			return MHashResult[Index];
 		}
