Fix crash due to uninitialized memory

diff --git a/Engine/Source/Runtime/Experimental/Chaos/Public/Chaos/EPA.h b/Engine/Source/Runtime/Experimental/Chaos/Public/Chaos/EPA.h
--- a/Engine/Source/Runtime/Experimental/Chaos/Public/Chaos/EPA.h
+++ b/Engine/Source/Runtime/Experimental/Chaos/Public/Chaos/EPA.h
@@ -372,9 +372,7 @@ EPAResult EPA(TArray<TVec3<T>>& VertsABuffer, TArray<TVec3<T>>& VertsBBuffer, co
 	//TEPAEntry<T> BestEntry;
 	//BestEntry.Distance = 0;
 
-	TEPAEntry<T> LastEntry;
-	LastEntry.Distance = 0;
-	LastEntry.PlaneNormal = TVec3<T>(0,0,1);
+	TEPAEntry<T> LastEntry = Entries[0];
 
 	TArray<FEPAFloodEntry> VisibilityBorder;
 	int32 Iteration = 0;
