Fix shadow variable

diff --git a/Engine/Source/Runtime/Experimental/Chaos/Public/Chaos/GJK.h b/Engine/Source/Runtime/Experimental/Chaos/Public/Chaos/GJK.h
--- a/Engine/Source/Runtime/Experimental/Chaos/Public/Chaos/GJK.h
+++ b/Engine/Source/Runtime/Experimental/Chaos/Public/Chaos/GJK.h
@@ -437,8 +437,8 @@ namespace Chaos
 
 					auto SupportBInAFunc = [&](const TVec3<T>& Dir)
 					{
-						const TVec3<T> SupportB = SupportBFunc(Dir);
-						return StartTM.TransformPosition(SupportB);
+						const TVec3<T> SupportBLocal = SupportBFunc(Dir);
+						return StartTM.TransformPosition(SupportBLocal);
 					};
 
 					T Penetration;
