-Fic CIS Non-Unity

diff --git a/Engine/Source/Runtime/Experimental/Chaos/Private/Chaos/CollisionResolutionAlgo.cpp b/Engine/Source/Runtime/Experimental/Chaos/Private/Chaos/CollisionResolutionAlgo.cpp
--- a/Engine/Source/Runtime/Experimental/Chaos/Private/Chaos/CollisionResolutionAlgo.cpp
+++ b/Engine/Source/Runtime/Experimental/Chaos/Private/Chaos/CollisionResolutionAlgo.cpp
@@ -15,6 +15,7 @@
 #include "Chaos/Transform.h"
 #include "ChaosLog.h"
 #include "ChaosStats.h"
+#include "Stats/Stats.h" 
 
 #if INTEL_ISPC
 #include "PBDCollisionConstraint.ispc.generated.h"
