Chaos: Add CellBounds test for HeightField Raycast.

diff --git a/Engine/Source/Runtime/Experimental/Chaos/Private/Chaos/HeightField.cpp b/Engine/Source/Runtime/Experimental/Chaos/Private/Chaos/HeightField.cpp
--- a/Engine/Source/Runtime/Experimental/Chaos/Private/Chaos/HeightField.cpp
+++ b/Engine/Source/Runtime/Experimental/Chaos/Private/Chaos/HeightField.cpp
@@ -568,6 +568,21 @@ namespace Chaos
 		TVector<T, 2> ClippedFlatRayStart;
 		TVector<T, 2> ClippedFlatRayEnd;
 
+		// Data for fast box cast
+		TVector<T, 3> Min, Max, HitPoint;
+		float ToI;
+		bool bParallel[3];
+		TVector<T, 3> InvDir;
+
+		T InvCurrentLength = 1 / CurrentLength;
+		for(int Axis = 0; Axis < 3; ++Axis)
+		{
+			bParallel[Axis] = Dir[Axis] == 0;
+			InvDir[Axis] = bParallel[Axis] ? 0 : 1 / Dir[Axis];
+		}
+
+
+
 		if(GetFlatBounds().ClipLine(StartPoint, StartPoint + Dir * Length, ClippedFlatRayStart, ClippedFlatRayEnd))
 		{
 			// The line is now valid and is entirely enclosed by the bounds (and thus, the grid)
@@ -578,7 +593,13 @@ namespace Chaos
 				const int32 QueryY = (int32)(ClippedFlatRayStart[1] / GeomData.Scale[1]);
 				const TVector<int32, 2> Cell = FlatGrid.ClampIndex({QueryX, QueryY});
 
-				Visitor.VisitRaycast(Cell[1] * (GeomData.NumCols - 1) + Cell[0], CurrentLength);
+				if (GetCellBounds3DScaled(Cell, Min, Max))
+				{
+					if (TBox<T, 3>::RaycastFast(Min, Max, StartPoint, Dir, InvDir, bParallel, CurrentLength, InvCurrentLength, ToI, HitPoint))
+					{
+						Visitor.VisitRaycast(Cell[1] * (GeomData.NumCols - 1) + Cell[0], CurrentLength);
+					}
+				}
 			}
 			else
 			{
@@ -601,11 +622,26 @@ namespace Chaos
 				{
 					// Visit the selected cell, in this case no need to rasterize because we never
 					// leave the initial cell
-					Visitor.VisitRaycast(StartCell[1] * (GeomData.NumCols - 1) + StartCell[0], CurrentLength);
+					if (GetCellBounds3DScaled(StartCell, Min, Max))
+					{
+						// Check cell bounds
+						if (TBox<T, 3>::RaycastFast(Min, Max, StartPoint, Dir, InvDir, bParallel, CurrentLength, InvCurrentLength, ToI, HitPoint))
+						{
+							Visitor.VisitRaycast(StartCell[1] * (GeomData.NumCols - 1) + StartCell[0], CurrentLength);
+						}
+					}
 				}
 				else
 				{
-					Visitor.VisitRaycast(StartCell[1] * (GeomData.NumCols - 1) + StartCell[0], CurrentLength);
+
+					if (GetCellBounds3DScaled(StartCell, Min, Max))
+					{
+						// Check cell bounds
+						if (TBox<T, 3>::RaycastFast(Min, Max, StartPoint, Dir, InvDir, bParallel, CurrentLength, InvCurrentLength, ToI, HitPoint))
+						{
+							Visitor.VisitRaycast(StartCell[1] * (GeomData.NumCols - 1) + StartCell[0], CurrentLength);
+						}
+					}
 
 					do 
 					{
@@ -623,11 +659,18 @@ namespace Chaos
 							StartCell[1] += DirY;
 						}
 
-						// Visit the selected cell
-						bool bContinue = Visitor.VisitRaycast(StartCell[1] * (GeomData.NumCols - 1) + StartCell[0], CurrentLength);
-						if(!bContinue)
+						if (GetCellBounds3DScaled(StartCell, Min, Max))
 						{
-							return true;
+							// Check cell bounds
+							if (TBox<T, 3>::RaycastFast(Min, Max, StartPoint, Dir, InvDir, bParallel, CurrentLength, InvCurrentLength, ToI, HitPoint))
+							{
+								// Visit the selected cell
+								bool bContinue = Visitor.VisitRaycast(StartCell[1] * (GeomData.NumCols - 1) + StartCell[0], CurrentLength);
+								if (!bContinue)
+								{
+									return true;
+								}
+							}
 						}
 					} while(StartCell != EndCell);
 				}
