Attempt to fix mac without breaking windows.

diff --git a/Engine/Source/Runtime/Experimental/Chaos/Private/Chaos/BVHParticles.cpp b/Engine/Source/Runtime/Experimental/Chaos/Private/Chaos/BVHParticles.cpp
new file mode 100644
--- /dev/null
+++ b/Engine/Source/Runtime/Experimental/Chaos/Private/Chaos/BVHParticles.cpp
@@ -0,0 +1,80 @@
+#include "Chaos/BVHParticles.h"
+#include "Chaos/BoundingVolumeHierarchy.h"
+
+using namespace Chaos;
+
+template<class T, int d>
+TBVHParticles<T,d>::TBVHParticles()
+	: TParticles<T, d>()
+	, MBVH(new TBoundingVolumeHierarchy<TParticles<T, d>, TArray<int32>, T, d>(*this, CollisionParticlesBVHDepth))
+{
+}
+
+template<class T, int d>
+TBVHParticles<T,d>::TBVHParticles(TBVHParticles<T, d>&& Other)
+	: TParticles<T, d>(MoveTemp(Other))
+	, MBVH(new TBoundingVolumeHierarchy<TParticles<T, d>, TArray<int32>, T, d>(MoveTemp(*Other.MBVH)))
+{
+}
+
+template<class T, int d>
+TBVHParticles<T,d>::TBVHParticles(TParticles<T, d>&& Other)
+	: TParticles<T, d>(MoveTemp(Other))
+	, MBVH(new TBoundingVolumeHierarchy<TParticles<T, d>, TArray<int32>, T, d>(*this, CollisionParticlesBVHDepth))
+{
+}
+
+template<class T, int d>
+TBVHParticles<T,d>::~TBVHParticles()
+{
+    delete MBVH;
+}
+
+template<class T, int d>
+TBVHParticles<T,d>& TBVHParticles<T,d>::operator=(const TBVHParticles<T, d>& Other)
+{
+	*this = TBVHParticles(Other);
+	return *this;
+}
+
+template<class T, int d>
+TBVHParticles<T,d>& TBVHParticles<T,d>::operator=(TBVHParticles<T, d>&& Other)
+{
+	*MBVH = MoveTemp(*Other.MBVH);
+	TParticles<T, d>::operator=(static_cast<TParticles<T, d>&&>(Other));
+	return *this;
+}
+
+template<class T, int d>
+TBVHParticles<T,d>::TBVHParticles(const TBVHParticles<T, d>& Other)
+	: TParticles<T, d>()
+{
+	AddParticles(Other.Size());
+	for (int32 i = Other.Size() - 1; 0 <= i; i--)
+	{
+		X(i) = Other.X(i);
+	}
+    delete MBVH;
+	MBVH = new TBoundingVolumeHierarchy<TParticles<T, d>, TArray<int32>, T, d>(*this, CollisionParticlesBVHDepth);
+}
+
+template<class T, int d>
+void TBVHParticles<T,d>::UpdateAccelerationStructures()
+{
+	MBVH->UpdateHierarchy();
+}
+
+template<class T, int d>
+const TArray<int32> TBVHParticles<T,d>::FindAllIntersections(const TBox<T, d>& Object) const
+{
+	return MBVH->FindAllIntersections(Object);
+}
+
+template<class T, int d>
+void TBVHParticles<T,d>::Serialize(FChaosArchive& Ar)
+{
+	TParticles<T, d>::Serialize(Ar);
+	Ar << *MBVH;
+}
+
+template class Chaos::TBVHParticles<float,3>;
diff --git a/Engine/Source/Runtime/Experimental/Chaos/Private/Chaos/BoundingVolume.cpp b/Engine/Source/Runtime/Experimental/Chaos/Private/Chaos/BoundingVolume.cpp
--- a/Engine/Source/Runtime/Experimental/Chaos/Private/Chaos/BoundingVolume.cpp
+++ b/Engine/Source/Runtime/Experimental/Chaos/Private/Chaos/BoundingVolume.cpp
@@ -37,4 +37,7 @@ namespace Chaos
 
 	template class CHAOS_API Chaos::ISpatialAcceleration<int32, float, 3>;
 	template class CHAOS_API Chaos::ISpatialAcceleration<TAccelerationStructureHandle<float,3>, float, 3>;
-}
\ No newline at end of file
+
+    template class CHAOS_API Chaos::TBoundingVolume<int32,float,3>;
+    template class CHAOS_API Chaos::TBoundingVolume<TAccelerationStructureHandle<float,3>,float,3>;
+}
diff --git a/Engine/Source/Runtime/Experimental/Chaos/Private/Chaos/BoundingVolumeHierarchy.cpp b/Engine/Source/Runtime/Experimental/Chaos/Private/Chaos/BoundingVolumeHierarchy.cpp
--- a/Engine/Source/Runtime/Experimental/Chaos/Private/Chaos/BoundingVolumeHierarchy.cpp
+++ b/Engine/Source/Runtime/Experimental/Chaos/Private/Chaos/BoundingVolumeHierarchy.cpp
@@ -680,10 +680,10 @@ void TBoundingVolumeHierarchy<OBJECT_ARRAY, LEAF_TYPE, T, d>::Serialize(FArchive
 }
 }
 
-template class CHAOS_API Chaos::TBoundingVolume<int32, float, 3>;
-template class CHAOS_API Chaos::TBoundingVolumeHierarchy<TArray<Chaos::TSphere<float, 3>*>, TArray<int32>, float, 3>;
-template class CHAOS_API Chaos::TBoundingVolumeHierarchy<Chaos::TPBDRigidParticles<float, 3>, TArray<int32>, float, 3>;
-template class CHAOS_API Chaos::TBoundingVolumeHierarchy<Chaos::TParticles<float, 3>, TArray<int32>, float, 3>;
-template class CHAOS_API Chaos::TBoundingVolumeHierarchy<Chaos::TGeometryParticles<float, 3>, TArray<int32>, float, 3>;
-template class CHAOS_API Chaos::TBoundingVolumeHierarchy<Chaos::TPBDRigidParticles<float, 3>, TBoundingVolume<TPBDRigidParticleHandle<float,3>*, float, 3>, float, 3>;
-template class CHAOS_API Chaos::TBoundingVolumeHierarchy<Chaos::TGeometryParticles<float, 3>, TBoundingVolume<TGeometryParticleHandle<float,3>*, float, 3>, float, 3>;
+//template class CHAOS_API Chaos::TBoundingVolume<int32, float, 3>;
+template class Chaos::TBoundingVolumeHierarchy<TArray<Chaos::TSphere<float, 3>*>, TArray<int32>, float, 3>;
+template class Chaos::TBoundingVolumeHierarchy<Chaos::TPBDRigidParticles<float, 3>, TArray<int32>, float, 3>;
+template class Chaos::TBoundingVolumeHierarchy<Chaos::TParticles<float, 3>, TArray<int32>, float, 3>;
+template class Chaos::TBoundingVolumeHierarchy<Chaos::TGeometryParticles<float, 3>, TArray<int32>, float, 3>;
+template class Chaos::TBoundingVolumeHierarchy<Chaos::TPBDRigidParticles<float, 3>, TBoundingVolume<TPBDRigidParticleHandle<float,3>*, float, 3>, float, 3>;
+template class Chaos::TBoundingVolumeHierarchy<Chaos::TGeometryParticles<float, 3>, TBoundingVolume<TGeometryParticleHandle<float,3>*, float, 3>, float, 3>;
diff --git a/Engine/Source/Runtime/Experimental/Chaos/Public/Chaos/BVHParticles.h b/Engine/Source/Runtime/Experimental/Chaos/Public/Chaos/BVHParticles.h
--- a/Engine/Source/Runtime/Experimental/Chaos/Public/Chaos/BVHParticles.h
+++ b/Engine/Source/Runtime/Experimental/Chaos/Public/Chaos/BVHParticles.h
@@ -1,7 +1,6 @@
 // Copyright 1998-2019 Epic Games, Inc. All Rights Reserved.
 #pragma once
 
-#include "Chaos/BoundingVolumeHierarchy.h"
 #include "Chaos/Particles.h"
 #include "ChaosArchive.h"
 #include "Chaos/ParticleHandleFwd.h"
@@ -10,6 +9,13 @@ extern int32 CHAOS_API CollisionParticlesBVHDepth;
 
 namespace Chaos
 {
+
+    template<class OBJECT_ARRAY, class LEAF_TYPE, class T, int d>
+    class TBoundingVolumeHierarchy;
+
+    template<class T, int d>
+    class TBox;
+
 	template<class T, int d>
 	class TBVHParticles final /*Note: removing this final has implications for serialization. See TImplicitObject*/ : public TParticles<T, d>
 	{
@@ -18,77 +24,38 @@ namespace Chaos
 		using TParticles<T, d>::X;
 		using TParticles<T, d>::AddParticles;
 
-		TBVHParticles()
-		    : TParticles<T, d>()
-			, MBVH(*this, CollisionParticlesBVHDepth)
-		{}
-		TBVHParticles(TBVHParticles<T, d>&& Other)
-		    : TParticles<T, d>(MoveTemp(Other))
-			, MBVH(MoveTemp(Other.MBVH))
-		{}
-		TBVHParticles(TParticles<T, d>&& Other)
-		    : TParticles<T, d>(MoveTemp(Other))
-		    , MBVH(*this, CollisionParticlesBVHDepth)
-		{}
-
-		TBVHParticles& operator=(const TBVHParticles<T, d>& Other)
-		{
-			*this = TBVHParticles(Other);
-			return *this;
-		}
+		CHAOS_API TBVHParticles();
+		TBVHParticles(TBVHParticles<T, d>&& Other);
+		TBVHParticles(TParticles<T, d>&& Other);
+        CHAOS_API ~TBVHParticles();
 
-		TBVHParticles& operator=(TBVHParticles<T, d>&& Other)
-		{
-			MBVH = MoveTemp(Other.MBVH);
-			TParticles<T, d>::operator=(static_cast<TParticles<T, d>&&>(Other));
-			return *this;
-		}
+	    CHAOS_API TBVHParticles& operator=(const TBVHParticles<T, d>& Other);
+	    CHAOS_API TBVHParticles& operator=(TBVHParticles<T, d>&& Other);
 
-		TBVHParticles* NewCopy()
+		CHAOS_API TBVHParticles* NewCopy()
 		{
 			return new TBVHParticles(*this);
 		}
 
-		void UpdateAccelerationStructures()
-		{
-			MBVH.UpdateHierarchy();
-		}
-
-		const TArray<int32> FindAllIntersections(const TBox<T, d>& Object) const
-		{
-			return MBVH.FindAllIntersections(Object);
-		}
+		CHAOS_API void UpdateAccelerationStructures();
+		const TArray<int32> FindAllIntersections(const TBox<T, d>& Object) const;
 
 		static TBVHParticles<T,d>* SerializationFactory(FChaosArchive& Ar, TBVHParticles<T,d>* BVHParticles)
 		{
 			return Ar.IsLoading() ? new TBVHParticles<T, d>() : nullptr;
 		}
 
-		void Serialize(FChaosArchive& Ar)
-		{
-			TParticles<T, d>::Serialize(Ar);
-			Ar << MBVH;
-		}
+		CHAOS_API void Serialize(FChaosArchive& Ar);
 
-		virtual void Serialize(FArchive& Ar)
+		CHAOS_API virtual void Serialize(FArchive& Ar)
 		{
 			check(false); //Aggregate simplicial require FChaosArchive - check false by default
 		}
 
 	private:
-		TBVHParticles(const TBVHParticles<T, d>& Other)
-			: TParticles<T, d>()
-			, MBVH(*this, CollisionParticlesBVHDepth)
-		{
-			AddParticles(Other.Size());
-			for (int32 i = Other.Size() - 1; 0 <= i; i--)
-			{
-				X(i) = Other.X(i);
-			}
-			MBVH = TBoundingVolumeHierarchy<TParticles<T, d>, TArray<int32>, T, d>(*this, CollisionParticlesBVHDepth);
-		}
+		CHAOS_API TBVHParticles(const TBVHParticles<T, d>& Other);
 
-		TBoundingVolumeHierarchy<TParticles<T, d>, TArray<int32>, T, d> MBVH;
+		TBoundingVolumeHierarchy<TParticles<T, d>, TArray<int32>, T, d>* MBVH;
 	};
 
 	template<typename T, int d>
diff --git a/Engine/Source/Runtime/Experimental/Chaos/Public/Chaos/BoundingVolume.h b/Engine/Source/Runtime/Experimental/Chaos/Public/Chaos/BoundingVolume.h
--- a/Engine/Source/Runtime/Experimental/Chaos/Public/Chaos/BoundingVolume.h
+++ b/Engine/Source/Runtime/Experimental/Chaos/Public/Chaos/BoundingVolume.h
@@ -11,6 +11,7 @@
 #include "Chaos/UniformGrid.h"
 #include "Chaos/ISpatialAcceleration.h"
 #include "ChaosStats.h"
+#include "ChaosLog.h"
 #include "HAL/IConsoleManager.h"
 #include "Chaos/ParticleHandleFwd.h"
 #include "Templates/Models.h"
@@ -1178,6 +1179,14 @@ private:
 	bool bIsEmpty;
 };
 
+#if PLATFORM_MAC || PLATFORM_LINUX
+    extern template class CHAOS_API TBoundingVolume<int32,float,3>;
+    extern template class CHAOS_API TBoundingVolume<TAccelerationStructureHandle<float,3>,float,3>;
+#else
+    extern template class TBoundingVolume<int32,float,3>;
+    extern template class TBoundingVolume<TAccelerationStructureHandle<float,3>,float,3>;
+#endif
+
 template<typename TPayloadType, class T, int d>
 FArchive& operator<<(FArchive& Ar, TBoundingVolume<TPayloadType, T, d>& BoundingVolume)
 {
diff --git a/Engine/Source/Runtime/Experimental/Chaos/Public/Chaos/BoundingVolumeHierarchy.h b/Engine/Source/Runtime/Experimental/Chaos/Public/Chaos/BoundingVolumeHierarchy.h
--- a/Engine/Source/Runtime/Experimental/Chaos/Public/Chaos/BoundingVolumeHierarchy.h
+++ b/Engine/Source/Runtime/Experimental/Chaos/Public/Chaos/BoundingVolumeHierarchy.h
@@ -5,6 +5,8 @@
 #include "Chaos/Box.h"
 #include "Chaos/Defines.h"
 #include "Chaos/GeometryParticles.h"
+#include "Chaos/ParticleHandleFwd.h"
+#include "Chaos/BoundingVolume.h"
 #include "Chaos/Transform.h"
 #include "ChaosLog.h"
 #include "Chaos/ISpatialAcceleration.h"
@@ -33,8 +35,6 @@ struct CComplexBVHLeaf
 	auto Requires(typename T::TPayloadType) ->void;
 };
 
-#define MIN_NUM_OBJECTS 5
-
 namespace Chaos
 {
 
@@ -56,7 +56,7 @@ FArchive& operator<<(FArchive& Ar, TBVHNode<T,d>& LeafNode)
 }
 
 template<class OBJECT_ARRAY, class LEAF_TYPE, class T, int d>
-class CHAOS_API TBoundingVolumeHierarchy final : public ISpatialAcceleration<int32, T,d>
+class TBoundingVolumeHierarchy final : public ISpatialAcceleration<int32, T,d>
 {
   public:
 	static constexpr int32 DefaultMaxLevels = 12;
@@ -65,16 +65,16 @@ class CHAOS_API TBoundingVolumeHierarchy final : public ISpatialAcceleration<int
 	static constexpr T DefaultDt = 0;
 	using TPayloadType = typename TBVHLeafTraits<LEAF_TYPE, TModels<CComplexBVHLeaf, LEAF_TYPE>::Value>::TPayloadType;
 
-	TBoundingVolumeHierarchy(const OBJECT_ARRAY& Objects, const int32 MaxLevels = DefaultMaxLevels, const bool bUseVelocity = DefaultUseVelocity, const T Dt = DefaultDt);
-	TBoundingVolumeHierarchy(const OBJECT_ARRAY& Objects, const TArray<uint32>& ActiveIndices, const int32 MaxLevels = DefaultMaxLevels, const bool bUseVelocity = DefaultUseVelocity, const T Dt = DefaultDt);
+	CHAOS_API TBoundingVolumeHierarchy(const OBJECT_ARRAY& Objects, const int32 MaxLevels = DefaultMaxLevels, const bool bUseVelocity = DefaultUseVelocity, const T Dt = DefaultDt);
+	CHAOS_API TBoundingVolumeHierarchy(const OBJECT_ARRAY& Objects, const TArray<uint32>& ActiveIndices, const int32 MaxLevels = DefaultMaxLevels, const bool bUseVelocity = DefaultUseVelocity, const T Dt = DefaultDt);
 	
 	TBoundingVolumeHierarchy(const TBoundingVolumeHierarchy<OBJECT_ARRAY, LEAF_TYPE, T, d>& Other) = delete;
-	TBoundingVolumeHierarchy(TBoundingVolumeHierarchy<OBJECT_ARRAY, LEAF_TYPE, T, d>&& Other)
+	CHAOS_API TBoundingVolumeHierarchy(TBoundingVolumeHierarchy<OBJECT_ARRAY, LEAF_TYPE, T, d>&& Other)
 	    : MObjects(Other.MObjects), MGlobalObjects(MoveTemp(Other.MGlobalObjects)), MWorldSpaceBoxes(MoveTemp(Other.MWorldSpaceBoxes)), MMaxLevels(Other.MMaxLevels), Elements(MoveTemp(Other.Elements)), Leafs(MoveTemp(Other.Leafs))
 	{
 	}
 
-	TBoundingVolumeHierarchy& operator=(TBoundingVolumeHierarchy<OBJECT_ARRAY, LEAF_TYPE, T, d>&& Other)
+	CHAOS_API TBoundingVolumeHierarchy& operator=(TBoundingVolumeHierarchy<OBJECT_ARRAY, LEAF_TYPE, T, d>&& Other)
 	{
 		MObjects = Other.MObjects;
 		MGlobalObjects = MoveTemp(Other.MGlobalObjects);
@@ -85,15 +85,15 @@ class CHAOS_API TBoundingVolumeHierarchy final : public ISpatialAcceleration<int
 		return *this;
 	}
 
-	void UpdateHierarchy(const bool AllowMultipleSplitting = DefaultAllowMultipleSplitting, const bool bUseVelocity = DefaultUseVelocity, const T Dt = DefaultDt);
-	void UpdateHierarchy(const TArray<uint32>& ActiveIndices, const bool AllowMultipleSplitting = DefaultAllowMultipleSplitting, const bool bUseVelocity = DefaultUseVelocity, const T Dt = DefaultDt);
+	CHAOS_API void UpdateHierarchy(const bool AllowMultipleSplitting = DefaultAllowMultipleSplitting, const bool bUseVelocity = DefaultUseVelocity, const T Dt = DefaultDt);
+	CHAOS_API void UpdateHierarchy(const TArray<uint32>& ActiveIndices, const bool AllowMultipleSplitting = DefaultAllowMultipleSplitting, const bool bUseVelocity = DefaultUseVelocity, const T Dt = DefaultDt);
 
-	void Reinitialize(bool bUseVelocity, T Dt)
+	CHAOS_API void Reinitialize(bool bUseVelocity, T Dt)
 	{
 		UpdateHierarchy(DefaultAllowMultipleSplitting, bUseVelocity, Dt);
 	}
 
-	void Reinitialize(const TArray<uint32>& ActiveIndices, bool bUseVelocity, T Dt)
+	CHAOS_API void Reinitialize(const TArray<uint32>& ActiveIndices, bool bUseVelocity, T Dt)
 	{
 		UpdateHierarchy(ActiveIndices, DefaultAllowMultipleSplitting, bUseVelocity, Dt);
 	}
@@ -114,33 +114,33 @@ class CHAOS_API TBoundingVolumeHierarchy final : public ISpatialAcceleration<int
 	}
 
 	// Begin ISpatialAcceleration interface
-	TArray<int32> FindAllIntersections(const TBox<T, d>& Box) const { return FindAllIntersectionsImp(Box); }
-	TArray<int32> FindAllIntersections(const TSpatialRay<T,d>& Ray) const { return FindAllIntersectionsImp(Ray); }
-	TArray<int32> FindAllIntersections(const TVector<T, d>& Point) const { return FindAllIntersectionsImp(Point); }
-	TArray<int32> FindAllIntersections(const TGeometryParticles<T, d>& InParticles, const int32 i) const;
+	CHAOS_API TArray<int32> FindAllIntersections(const TBox<T, d>& Box) const { return FindAllIntersectionsImp(Box); }
+	CHAOS_API TArray<int32> FindAllIntersections(const TSpatialRay<T,d>& Ray) const { return FindAllIntersectionsImp(Ray); }
+	CHAOS_API TArray<int32> FindAllIntersections(const TVector<T, d>& Point) const { return FindAllIntersectionsImp(Point); }
+	CHAOS_API TArray<int32> FindAllIntersections(const TGeometryParticles<T, d>& InParticles, const int32 i) const;
 	// End ISpatialAcceleration interface
 
-	const TArray<int32>& GlobalObjects() const
+	CHAOS_API const TArray<int32>& GlobalObjects() const
 	{
 		return MGlobalObjects;
 	}
 
 	// TODO(mlentine): Need to move this elsewhere; probably on CollisionConstraint
-	const TBox<T, d>& GetWorldSpaceBoundingBox(const TGeometryParticles<T, d>& InParticles, const int32 Index);
+	CHAOS_API const TBox<T, d>& GetWorldSpaceBoundingBox(const TGeometryParticles<T, d>& InParticles, const int32 Index);
 
 #if !UE_BUILD_SHIPPING
-	virtual void DebugDraw(ISpacialDebugDrawInterface<T>* InInterface) const override;
+	CHAOS_API virtual void DebugDraw(ISpacialDebugDrawInterface<T>* InInterface) const override;
 #endif
 
-	void Serialize(FArchive& Ar);
+	CHAOS_API void Serialize(FArchive& Ar);
 
-	virtual void Serialize(FChaosArchive& Ar) override
+	CHAOS_API virtual void Serialize(FChaosArchive& Ar) override
 	{
 		check(false);
 	}
 
   private:
-	void PrintTree(FString Prefix, const TBVHNode<T,d>* MyNode) const
+	CHAOS_API void PrintTree(FString Prefix, const TBVHNode<T,d>* MyNode) const
 	{
 		UE_LOG(LogChaos, Verbose, TEXT("%sNode has Box: (%f, %f, %f) to (%f, %f, %f) with %d Children"), *Prefix, MyNode->MMin[0], MyNode->MMin[1], MyNode->MMin[2], MyNode->MMax[0], MyNode->MMax[1], MyNode->MMax[2], MyNode->MChildren.Num());
 		for (auto& Child : MyNode->MChildren)
@@ -149,16 +149,16 @@ class CHAOS_API TBoundingVolumeHierarchy final : public ISpatialAcceleration<int
 		}
 	}
 
-	TArray<int32> FindAllIntersectionsHelper(const TBVHNode<T,d>& MyNode, const TVector<T, d>& Point) const;
-	TArray<int32> FindAllIntersectionsHelper(const TBVHNode<T,d>& MyNode, const TBox<T, d>& ObjectBox) const;
-	TArray<int32> FindAllIntersectionsHelper(const TBVHNode<T, d>& MyNode, const TSpatialRay<T, d>& Ray) const;
+	CHAOS_API TArray<int32> FindAllIntersectionsHelper(const TBVHNode<T,d>& MyNode, const TVector<T, d>& Point) const;
+	CHAOS_API TArray<int32> FindAllIntersectionsHelper(const TBVHNode<T,d>& MyNode, const TBox<T, d>& ObjectBox) const;
+	CHAOS_API TArray<int32> FindAllIntersectionsHelper(const TBVHNode<T, d>& MyNode, const TSpatialRay<T, d>& Ray) const;
 
 	template <typename QUERY_OBJECT>
-	void FindAllIntersectionsHelperRecursive(const TBVHNode<T,d>& MyNode, const QUERY_OBJECT& ObjectBox, TArray<int32>& AccumulateElements) const;
-	void UpdateHierarchyImp(const TArray<int32>& AllObjects, const bool bAllowMultipleSplitting, const bool bUseVelocity, const T Dt);
+	CHAOS_API void FindAllIntersectionsHelperRecursive(const TBVHNode<T,d>& MyNode, const QUERY_OBJECT& ObjectBox, TArray<int32>& AccumulateElements) const;
+	CHAOS_API void UpdateHierarchyImp(const TArray<int32>& AllObjects, const bool bAllowMultipleSplitting, const bool bUseVelocity, const T Dt);
 
-	int32 GenerateNextLevel(const TVector<T, d>& GlobalMin, const TVector<T, d>& GlobalMax, const TArray<int32>& Objects, const int32 Axis, const int32 Level, const bool AllowMultipleSplitting);
-	int32 GenerateNextLevel(const TVector<T, d>& GlobalMin, const TVector<T, d>& GlobalMax, const TArray<int32>& Objects, const int32 Level);
+	CHAOS_API int32 GenerateNextLevel(const TVector<T, d>& GlobalMin, const TVector<T, d>& GlobalMax, const TArray<int32>& Objects, const int32 Axis, const int32 Level, const bool AllowMultipleSplitting);
+	CHAOS_API int32 GenerateNextLevel(const TVector<T, d>& GlobalMin, const TVector<T, d>& GlobalMax, const TArray<int32>& Objects, const int32 Level);
 
 	OBJECT_ARRAY const* MObjects;
 	TArray<int32> MGlobalObjects;
@@ -171,8 +171,19 @@ class CHAOS_API TBoundingVolumeHierarchy final : public ISpatialAcceleration<int
 };
 
 #if PLATFORM_MAC || PLATFORM_LINUX
-extern template class CHAOS_API Chaos::TBoundingVolumeHierarchy<Chaos::TParticles<float, 3>, TArray<int32>, float, 3>;
-extern template class CHAOS_API Chaos::TBoundingVolumeHierarchy<Chaos::TGeometryParticles<float, 3>, TArray<int32>, float, 3>;
+    extern template class CHAOS_API TBoundingVolumeHierarchy<TArray<Chaos::TSphere<float, 3>*>, TArray<int32>, float, 3>;
+    extern template class CHAOS_API TBoundingVolumeHierarchy<Chaos::TPBDRigidParticles<float, 3>, TArray<int32>, float, 3>;
+    extern template class CHAOS_API TBoundingVolumeHierarchy<Chaos::TParticles<float, 3>, TArray<int32>, float, 3>;
+    extern template class CHAOS_API TBoundingVolumeHierarchy<Chaos::TGeometryParticles<float, 3>, TArray<int32>, float, 3>;
+    extern template class CHAOS_API TBoundingVolumeHierarchy<Chaos::TPBDRigidParticles<float, 3>, TBoundingVolume<TPBDRigidParticleHandle<float,3>*, float, 3>, float, 3>;
+    extern template class CHAOS_API TBoundingVolumeHierarchy<Chaos::TGeometryParticles<float, 3>, TBoundingVolume<TGeometryParticleHandle<float,3>*, float, 3>, float, 3>;
+#else
+    extern template class TBoundingVolumeHierarchy<TArray<Chaos::TSphere<float, 3>*>, TArray<int32>, float, 3>;
+    extern template class TBoundingVolumeHierarchy<Chaos::TPBDRigidParticles<float, 3>, TArray<int32>, float, 3>;
+    extern template class TBoundingVolumeHierarchy<Chaos::TParticles<float, 3>, TArray<int32>, float, 3>;
+    extern template class TBoundingVolumeHierarchy<Chaos::TGeometryParticles<float, 3>, TArray<int32>, float, 3>;
+    extern template class TBoundingVolumeHierarchy<Chaos::TPBDRigidParticles<float, 3>, TBoundingVolume<TPBDRigidParticleHandle<float,3>*, float, 3>, float, 3>;
+    extern template class TBoundingVolumeHierarchy<Chaos::TGeometryParticles<float, 3>, TBoundingVolume<TGeometryParticleHandle<float,3>*, float, 3>, float, 3>;
 #endif
 
 template<class OBJECT_ARRAY, class LEAF_TYPE, class T, int d>
diff --git a/Engine/Source/Runtime/Experimental/Chaos/Public/Chaos/BoundingVolumeUtilities.h b/Engine/Source/Runtime/Experimental/Chaos/Public/Chaos/BoundingVolumeUtilities.h
--- a/Engine/Source/Runtime/Experimental/Chaos/Public/Chaos/BoundingVolumeUtilities.h
+++ b/Engine/Source/Runtime/Experimental/Chaos/Public/Chaos/BoundingVolumeUtilities.h
@@ -12,6 +12,8 @@
 #include "Chaos/Sphere.h"
 #include "Chaos/Transform.h"
 
+#define MIN_NUM_OBJECTS 5
+
 namespace Chaos
 {
 
diff --git a/Engine/Source/Runtime/Experimental/Chaos/Public/Chaos/ImplicitObjectUnion.h b/Engine/Source/Runtime/Experimental/Chaos/Public/Chaos/ImplicitObjectUnion.h
--- a/Engine/Source/Runtime/Experimental/Chaos/Public/Chaos/ImplicitObjectUnion.h
+++ b/Engine/Source/Runtime/Experimental/Chaos/Public/Chaos/ImplicitObjectUnion.h
@@ -3,14 +3,18 @@
 
 #include "Chaos/Array.h"
 #include "Chaos/ImplicitObject.h"
+#include "Chaos/ISpatialAcceleration.h"
+#include "Chaos/GeometryParticles.h"
 
-#include <memory>
-#include "BoundingVolumeHierarchy.h"
 #include "ImplicitObjectTransformed.h"
 #include "ChaosArchive.h"
 
 namespace Chaos
 {
+
+template<class OBJECT_ARRAY, class LEAF_TYPE, class T, int d>
+class TBoundingVolumeHierarchy;
+
 template<class T, int d>
 class TImplicitObjectUnion : public FImplicitObject
 {
@@ -18,56 +22,10 @@ class TImplicitObjectUnion : public FImplicitObject
 
 	using FImplicitObject::GetTypeName;
 
-	TImplicitObjectUnion(TArray<TUniquePtr<FImplicitObject>>&& Objects, const TArray<int32>& OriginalParticleLookupHack = TArray<int32>())
-	    : FImplicitObject(EImplicitObject::HasBoundingBox, ImplicitObjectType::Union)
-	    , MObjects(MoveTemp(Objects))
-		, Hierarchy(GeomParticles)
-	    , MLocalBoundingBox()
-		, bHierarchyBuilt(false)
-		, MOriginalParticleLookupHack(OriginalParticleLookupHack)
-	{
-		ensure(MObjects.Num());
-		for (int32 i = 0; i < MObjects.Num(); ++i)
-		{
-			if (i > 0)
-			{
-				MLocalBoundingBox.GrowToInclude(MObjects[i]->BoundingBox());
-			}
-			else
-			{
-				MLocalBoundingBox = MObjects[i]->BoundingBox();
-			}
-			check(MOriginalParticleLookupHack.Num() == 0 || MOriginalParticleLookupHack.Num() == MObjects.Num());
-			if (MOriginalParticleLookupHack.Num() > 0)
-			{
-				//this whole part sucks, only needed because of how we get union children. Need to refactor and enforce no unions of unions
-				if (const TImplicitObjectTransformed<T, d>* Transformed = MObjects[i]->template GetObject<const TImplicitObjectTransformed<T, d>>())
-				{
-					MCollisionParticleLookupHack.Add(Transformed->GetTransformedObject(), MOriginalParticleLookupHack[i]);
-				}
-				else
-				{
-					ensure(false);	//shouldn't be here
-				}
-			}
-		}
-
-		CacheAllImplicitObjects();
-
-	}
+	CHAOS_API TImplicitObjectUnion(TArray<TUniquePtr<FImplicitObject>>&& Objects, const TArray<int32>& OriginalParticleLookupHack = TArray<int32>());
 	TImplicitObjectUnion(const TImplicitObjectUnion<T, d>& Other) = delete;
-	TImplicitObjectUnion(TImplicitObjectUnion<T, d>&& Other)
-	    : FImplicitObject(EImplicitObject::HasBoundingBox)
-	    , MObjects(MoveTemp(Other.MObjects))
-		, GeomParticles(MoveTemp(Other.GeomParticles))
-		, Hierarchy(MoveTemp(Other.Hierarchy))
-	    , MLocalBoundingBox(MoveTemp(Other.MLocalBoundingBox))
-		, bHierarchyBuilt(Other.bHierarchyBuilt)
-		, MOriginalParticleLookupHack(MoveTemp(Other.MOriginalParticleLookupHack))
-		, MCollisionParticleLookupHack(MoveTemp(Other.MCollisionParticleLookupHack))
-	{
-	}
-	virtual ~TImplicitObjectUnion() {}
+	TImplicitObjectUnion(TImplicitObjectUnion<T, d>&& Other);
+	virtual ~TImplicitObjectUnion();
 
 	FORCEINLINE static constexpr EImplicitObjectType StaticType()
 	{
@@ -128,114 +86,10 @@ class TImplicitObjectUnion : public FImplicitObject
 		}
 	}
 
-	virtual void FindAllIntersectingObjects(TArray < Pair<const FImplicitObject*, TRigidTransform<T, d>>>& Out, const TBox<T, d>& LocalBounds) const
-	{
-		if (bHierarchyBuilt)
-		{
-			TArray<int32> Overlaps = Hierarchy.FindAllIntersections(LocalBounds);
-			Out.Reserve(Out.Num() + Overlaps.Num());
-			for (int32 Idx : Overlaps)
-			{
-				const FImplicitObject* Obj = GeomParticles.Geometry(Idx).Get();
-				Out.Add(MakePair(Obj, TRigidTransform<T, d>(GeomParticles.X(Idx), GeomParticles.R(Idx))));
-			}
-		}
-		else
-		{
-			for (const TUniquePtr<FImplicitObject>& Object : MObjects)
-			{
-				Object->FindAllIntersectingObjects(Out, LocalBounds);
-			}
-		}
-	}
-
-	TArray<int32> FindAllIntersectingChildren(const TBox<T, d>& LocalBounds) const
-	{
-		TArray<int32> IntersectingChildren;
-		if (bHierarchyBuilt) //todo: make this work when hierarchy is not built
-		{
-			IntersectingChildren = Hierarchy.FindAllIntersections(LocalBounds);
-			for (int32 i = IntersectingChildren.Num() - 1; i >= 0; --i)
-			{
-				const int32 Idx = IntersectingChildren[i];
-				if (Idx < MOriginalParticleLookupHack.Num())
-				{
-					IntersectingChildren[i] = MOriginalParticleLookupHack[Idx];
-				}
-				else
-				{
-					IntersectingChildren.RemoveAtSwap(i);
-				}
-			}
-			/*for (int32& Idx : IntersectingChildren)
-			{
-				Idx = MOriginalParticleLookupHack[Idx];
-			}*/
-		}
-		else
-		{
-			IntersectingChildren = MOriginalParticleLookupHack;
-		}
-
-		return IntersectingChildren;
-	}
-
-	TArray<int32> FindAllIntersectingChildren(const TSpatialRay<T, d>& LocalRay) const
-	{
-		TArray<int32> IntersectingChildren;
-		if (bHierarchyBuilt) //todo: make this work when hierarchy is not built
-		{
-			IntersectingChildren = Hierarchy.FindAllIntersections(LocalRay);
-			for (int32 i = IntersectingChildren.Num() - 1; i >= 0; --i)
-			{
-				const int32 Idx = IntersectingChildren[i];
-				if (Idx < MOriginalParticleLookupHack.Num())
-				{
-					IntersectingChildren[i] = MOriginalParticleLookupHack[Idx];
-				}
-				else
-				{
-					IntersectingChildren.RemoveAtSwap(i);
-				}
-			}
-			/*for (int32& Idx : IntersectingChildren)
-			{
-				Idx = MOriginalParticleLookupHack[Idx];
-			}*/
-		}
-		else
-		{
-			IntersectingChildren = MOriginalParticleLookupHack;
-		}
-
-		return IntersectingChildren;
-	}
-
-	virtual void CacheAllImplicitObjects()
-	{
-		TArray < Pair<TSerializablePtr<FImplicitObject>, TRigidTransform<T, d>>> SubObjects;
-		AccumulateAllSerializableImplicitObjectsHelper(SubObjects, TRigidTransform<T, d>::Identity);
-		//build hierarchy
-		{
-			const int32 NumObjects = SubObjects.Num();
-			constexpr int32 MinSubObjectsToCache = 8;	//todo(make this tunable?)
-			if (NumObjects > MinSubObjectsToCache)
-			{
-				GeomParticles.Resize(NumObjects);
-				for (int32 i = 0; i < NumObjects; ++i)
-				{
-					GeomParticles.X(i) = SubObjects[i].Second.GetLocation();
-					GeomParticles.R(i) = SubObjects[i].Second.GetRotation();
-					GeomParticles.SetGeometry(i, SubObjects[i].First);
-					//check(!SubObjects[i].First->IsUnderlyingUnion());	//we don't support union of unions
-				}
-
-				TBoundingVolumeHierarchy<TGeometryParticles<T, d>, TArray<int32>, T, d> NewHierarchy(GeomParticles, 1);
-				Hierarchy = MoveTemp(NewHierarchy);
-				bHierarchyBuilt = true;
-			}
-		}
-	}
+	virtual void FindAllIntersectingObjects(TArray < Pair<const FImplicitObject*, TRigidTransform<T, d>>>& Out, const TBox<T, d>& LocalBounds) const;
+	TArray<int32> FindAllIntersectingChildren(const TBox<T, d>& LocalBounds) const;
+	TArray<int32> FindAllIntersectingChildren(const TSpatialRay<T, d>& LocalRay) const;
+	virtual void CacheAllImplicitObjects();
 
 	virtual bool Raycast(const TVector<T, d>& StartPoint, const TVector<T, d>& Dir, const T Length, const T Thickness, T& OutTime, TVector<T, d>& OutPosition, TVector<T, d>& OutNormal, int32& OutFaceIndex) const override
 	{
@@ -278,12 +132,7 @@ class TImplicitObjectUnion : public FImplicitObject
 		return false;
 	}
 
-	virtual void Serialize(FChaosArchive& Ar) override
-	{
-		FChaosArchiveScopedMemory ScopedMemory(Ar, GetTypeName(), false);
-		FImplicitObject::SerializeImp(Ar);
-		Ar << MObjects << MLocalBoundingBox << GeomParticles << Hierarchy << bHierarchyBuilt;
-	}
+	virtual void Serialize(FChaosArchive& Ar) override;
 
 	virtual bool IsValidGeometry() const
 	{
@@ -331,14 +180,14 @@ private:
   private:
 	TArray<TUniquePtr<FImplicitObject>> MObjects;
 	TGeometryParticles<T, d> GeomParticles;
-	TBoundingVolumeHierarchy<TGeometryParticles<T, d>, TArray<int32>, T, d> Hierarchy;
+	TBoundingVolumeHierarchy<TGeometryParticles<T, d>, TArray<int32>, T, d>* Hierarchy;
 	TBox<T, d> MLocalBoundingBox;
 	bool bHierarchyBuilt;
 	TArray<int32> MOriginalParticleLookupHack;	//temp hack for finding original particles
 
 
 	//needed for serialization
-	TImplicitObjectUnion() : FImplicitObject(EImplicitObject::HasBoundingBox, ImplicitObjectType::Union), Hierarchy(GeomParticles, 1){}
+	TImplicitObjectUnion();
 	friend FImplicitObject;	//needed for serialization
 public:
 	TMap<const FImplicitObject*, int32> MCollisionParticleLookupHack;	//temp hack for finding collision particles
diff --git a/Engine/Source/Runtime/Experimental/Chaos/Public/Chaos/PBDRigidClusteredParticles.h b/Engine/Source/Runtime/Experimental/Chaos/Public/Chaos/PBDRigidClusteredParticles.h
--- a/Engine/Source/Runtime/Experimental/Chaos/Public/Chaos/PBDRigidClusteredParticles.h
+++ b/Engine/Source/Runtime/Experimental/Chaos/Public/Chaos/PBDRigidClusteredParticles.h
@@ -7,6 +7,9 @@
 
 namespace Chaos
 {
+template<class T, int d>
+class TImplicitObjectUnion; 
+
 // ClusterId 
 //    Used within the clustering system to describe the clustering hierarchy. The ID will
 //  store the children IDs, and a Parent ID. When Id==NONE_INDEX the cluster is not
diff --git a/Engine/Source/Runtime/Experimental/Chaos/Public/Chaos/PBDRigidParticles.h b/Engine/Source/Runtime/Experimental/Chaos/Public/Chaos/PBDRigidParticles.h
--- a/Engine/Source/Runtime/Experimental/Chaos/Public/Chaos/PBDRigidParticles.h
+++ b/Engine/Source/Runtime/Experimental/Chaos/Public/Chaos/PBDRigidParticles.h
@@ -42,29 +42,37 @@ class TPBDRigidParticles : public TRigidParticles<T, d>
 		TArrayCollection::AddArray(&MPreW);
 	}
 
-		CHAOS_API virtual ~TPBDRigidParticles()
-		{}
+	CHAOS_API virtual ~TPBDRigidParticles()
+	{}
 
+	CHAOS_API const TVector<T, d>& P(const int32 index) const { return MP[index]; }
+	CHAOS_API TVector<T, d>& P(const int32 index) { return MP[index]; }
 
-	const TVector<T, d>& P(const int32 index) const { return MP[index]; }
-	TVector<T, d>& P(const int32 index) { return MP[index]; }
+	CHAOS_API const TRotation<T, d>& Q(const int32 index) const { return MQ[index]; }
+	CHAOS_API TRotation<T, d>& Q(const int32 index) { return MQ[index]; }
 
-	const TRotation<T, d>& Q(const int32 index) const { return MQ[index]; }
-	TRotation<T, d>& Q(const int32 index) { return MQ[index]; }
+	CHAOS_API const TVector<T, d>& PreV(const int32 index) const { return MPreV[index]; }
+	CHAOS_API TVector<T, d>& PreV(const int32 index) { return MPreV[index]; }
 
-	const TVector<T, d>& PreV(const int32 index) const { return MPreV[index]; }
-	TVector<T, d>& PreV(const int32 index) { return MPreV[index]; }
+	CHAOS_API const TVector<T, d>& PreW(const int32 index) const { return MPreW[index]; }
+	CHAOS_API TVector<T, d>& PreW(const int32 index) { return MPreW[index]; }
 
-	const TVector<T, d>& PreW(const int32 index) const { return MPreW[index]; }
-	TVector<T, d>& PreW(const int32 index) { return MPreW[index]; }
+#if PLATFORM_LINUX || PLATFORM_MAC
+    // Must be reinterpret cast instead of static_cast as it's a forward declare
+	typedef TPBDRigidParticleHandle<T, d> THandleType;
+	CHAOS_API const THandleType* Handle(int32 Index) const { return reinterpret_cast<const THandleType*>(TGeometryParticles<T,d>::Handle(Index)); }
 
+	//cannot be reference because double pointer would allow for badness, but still useful to have non const access to handle
+	CHAOS_API THandleType* Handle(int32 Index) { return reinterpret_cast<THandleType*>(TGeometryParticles<T, d>::Handle(Index)); }
+#else
 	typedef TPBDRigidParticleHandle<T, d> THandleType;
-	const THandleType* Handle(int32 Index) const { return static_cast<const THandleType*>(TGeometryParticles<T,d>::Handle(Index)); }
+	CHAOS_API const THandleType* Handle(int32 Index) const { return static_cast<const THandleType*>(TGeometryParticles<T,d>::Handle(Index)); }
 
 	//cannot be reference because double pointer would allow for badness, but still useful to have non const access to handle
-	THandleType* Handle(int32 Index) { return static_cast<THandleType*>(TGeometryParticles<T, d>::Handle(Index)); }
+	CHAOS_API THandleType* Handle(int32 Index) { return static_cast<THandleType*>(TGeometryParticles<T, d>::Handle(Index)); }
+#endif
 
-	void SetSleeping(int32 Index, bool bSleeping)
+	CHAOS_API void SetSleeping(int32 Index, bool bSleeping)
 	{
 		if (Sleeping(Index) && bSleeping == false)
 		{
@@ -74,7 +82,11 @@ class TPBDRigidParticles : public TRigidParticles<T, d>
 
 		if (Sleeping(Index) != bSleeping)
 		{
+#if PLATFORM_MAC || PLATFORM_LINUX
+			TGeometryParticleHandle<T, d>* Particle = reinterpret_cast<TGeometryParticleHandle<T, d>*>(this->Handle(Index));
+#else
 			TGeometryParticleHandle<T, d>* Particle = this->Handle(Index);
+#endif
 			this->AddSleepData(Particle);
 		}
 
@@ -90,7 +102,7 @@ class TPBDRigidParticles : public TRigidParticles<T, d>
 		}
 	}
 
-	void SetObjectState(int32 Index, EObjectStateType InObjectState)
+	CHAOS_API void SetObjectState(int32 Index, EObjectStateType InObjectState)
 	{
 		const EObjectStateType CurrentState = this->ObjectState(Index);
 
@@ -130,7 +142,11 @@ class TPBDRigidParticles : public TRigidParticles<T, d>
 		const bool bNewSleeping = InObjectState == EObjectStateType::Sleeping;
 		if(bCurrentSleeping != bNewSleeping)
 		{
+#if PLATFORM_MAC || PLATFORM_LINUX
+			TGeometryParticleHandle<T, d>* Particle = reinterpret_cast<TGeometryParticleHandle<T, d>*>(this->Handle(Index));
+#else
 			TGeometryParticleHandle<T, d>* Particle = this->Handle(Index);
+#endif
 			this->AddSleepData(Particle);
 		}
 
@@ -156,6 +172,8 @@ class TPBDRigidParticles : public TRigidParticles<T, d>
 	TArrayCollectionArray<TVector<T, d>> MPreW;
 };
 
+extern template class TPBDRigidParticles<float,3>;
+
 template <typename T, int d>
 FChaosArchive& operator<<(FChaosArchive& Ar, TPBDRigidParticles<T, d>& Particles)
 {
diff --git a/Engine/Source/Runtime/Experimental/GeometryCollectionCore/Public/GeometryCollection/ManagedArrayTypes.h b/Engine/Source/Runtime/Experimental/GeometryCollectionCore/Public/GeometryCollection/ManagedArrayTypes.h
--- a/Engine/Source/Runtime/Experimental/GeometryCollectionCore/Public/GeometryCollection/ManagedArrayTypes.h
+++ b/Engine/Source/Runtime/Experimental/GeometryCollectionCore/Public/GeometryCollection/ManagedArrayTypes.h
@@ -10,6 +10,7 @@
 #include "Math/Vector.h"
 #include "Math/Vector2D.h"
 #include "Math/Box.h"
+#include "Chaos/ImplicitObject.h"
 #include "Chaos/BVHParticles.h"
 
 inline FArchive& operator<<(FArchive& Ar, TArray<FVector>*& ValueIn)
@@ -96,4 +97,4 @@ inline FManagedArrayBase* CopyManagedTypedArray(EManagedArrayType ArrayType, FMa
 	}
 	check(false);
 	return nullptr;
-}
\ No newline at end of file
+}
