Missed file.

diff --git a/Engine/Source/Runtime/Experimental/Chaos/Private/Chaos/ImplicitObjectUnion.cpp b/Engine/Source/Runtime/Experimental/Chaos/Private/Chaos/ImplicitObjectUnion.cpp
new file mode 100644
--- /dev/null
+++ b/Engine/Source/Runtime/Experimental/Chaos/Private/Chaos/ImplicitObjectUnion.cpp
@@ -0,0 +1,188 @@
+// Copyright 1998-2019 Epic Games, Inc. All Rights Reserved.
+#include "Chaos/ImplicitObjectUnion.h"
+#include "Chaos/BoundingVolumeHierarchy.h"
+
+using namespace Chaos;
+
+template<class T, int d>
+TImplicitObjectUnion<T,d>::TImplicitObjectUnion(TArray<TUniquePtr<FImplicitObject>>&& Objects, const TArray<int32>& OriginalParticleLookupHack)
+	: FImplicitObject(EImplicitObject::HasBoundingBox, ImplicitObjectType::Union)
+	, MObjects(MoveTemp(Objects))
+	, Hierarchy(new TBoundingVolumeHierarchy<TGeometryParticles<T, d>, TArray<int32>, T, d>(GeomParticles))
+	, MLocalBoundingBox()
+	, bHierarchyBuilt(false)
+	, MOriginalParticleLookupHack(OriginalParticleLookupHack)
+{
+	ensure(MObjects.Num());
+	for (int32 i = 0; i < MObjects.Num(); ++i)
+	{
+		if (i > 0)
+		{
+			MLocalBoundingBox.GrowToInclude(MObjects[i]->BoundingBox());
+		}
+		else
+		{
+			MLocalBoundingBox = MObjects[i]->BoundingBox();
+		}
+		check(MOriginalParticleLookupHack.Num() == 0 || MOriginalParticleLookupHack.Num() == MObjects.Num());
+		if (MOriginalParticleLookupHack.Num() > 0)
+		{
+			//this whole part sucks, only needed because of how we get union children. Need to refactor and enforce no unions of unions
+			if (const TImplicitObjectTransformed<T, d>* Transformed = MObjects[i]->template GetObject<const TImplicitObjectTransformed<T, d>>())
+			{
+				MCollisionParticleLookupHack.Add(Transformed->GetTransformedObject(), MOriginalParticleLookupHack[i]);
+			}
+			else
+			{
+				ensure(false);	//shouldn't be here
+			}
+		}
+	}
+
+	CacheAllImplicitObjects();
+}
+
+template<class T, int d>
+TImplicitObjectUnion<T,d>::TImplicitObjectUnion(TImplicitObjectUnion<T, d>&& Other)
+	: FImplicitObject(EImplicitObject::HasBoundingBox)
+	, MObjects(MoveTemp(Other.MObjects))
+	, GeomParticles(MoveTemp(Other.GeomParticles))
+	, Hierarchy(new TBoundingVolumeHierarchy<TGeometryParticles<T, d>, TArray<int32>, T, d>(MoveTemp(*Other.Hierarchy)))
+	, MLocalBoundingBox(MoveTemp(Other.MLocalBoundingBox))
+	, bHierarchyBuilt(Other.bHierarchyBuilt)
+	, MOriginalParticleLookupHack(MoveTemp(Other.MOriginalParticleLookupHack))
+	, MCollisionParticleLookupHack(MoveTemp(Other.MCollisionParticleLookupHack))
+{
+}
+
+template<class T, int d>
+TImplicitObjectUnion<T, d>::~TImplicitObjectUnion()
+{
+	delete Hierarchy;
+}
+
+template<class T, int d>
+void TImplicitObjectUnion<T,d>::FindAllIntersectingObjects(TArray < Pair<const FImplicitObject*, TRigidTransform<T, d>>>& Out, const TBox<T, d>& LocalBounds) const
+{
+	if (bHierarchyBuilt)
+	{
+		TArray<int32> Overlaps = Hierarchy->FindAllIntersections(LocalBounds);
+		Out.Reserve(Out.Num() + Overlaps.Num());
+		for (int32 Idx : Overlaps)
+		{
+			const FImplicitObject* Obj = GeomParticles.Geometry(Idx).Get();
+			Out.Add(MakePair(Obj, TRigidTransform<T, d>(GeomParticles.X(Idx), GeomParticles.R(Idx))));
+		}
+	}
+	else
+	{
+		for (const TUniquePtr<FImplicitObject>& Object : MObjects)
+		{
+			Object->FindAllIntersectingObjects(Out, LocalBounds);
+		}
+	}
+}
+
+template<class T, int d>
+TArray<int32> TImplicitObjectUnion<T,d>::FindAllIntersectingChildren(const TBox<T, d>& LocalBounds) const
+{
+	TArray<int32> IntersectingChildren;
+	if (bHierarchyBuilt) //todo: make this work when hierarchy is not built
+	{
+		IntersectingChildren = Hierarchy->FindAllIntersections(LocalBounds);
+		for (int32 i = IntersectingChildren.Num() - 1; i >= 0; --i)
+		{
+			const int32 Idx = IntersectingChildren[i];
+			if (Idx < MOriginalParticleLookupHack.Num())
+			{
+				IntersectingChildren[i] = MOriginalParticleLookupHack[Idx];
+			}
+			else
+			{
+				IntersectingChildren.RemoveAtSwap(i);
+			}
+		}
+		/*for (int32& Idx : IntersectingChildren)
+		{
+			Idx = MOriginalParticleLookupHack[Idx];
+		}*/
+	}
+	else
+	{
+		IntersectingChildren = MOriginalParticleLookupHack;
+	}
+
+	return IntersectingChildren;
+}
+
+template<class T, int d>
+TArray<int32> TImplicitObjectUnion<T,d>::FindAllIntersectingChildren(const TSpatialRay<T, d>& LocalRay) const
+{
+	TArray<int32> IntersectingChildren;
+	if (bHierarchyBuilt) //todo: make this work when hierarchy is not built
+	{
+		IntersectingChildren = Hierarchy->FindAllIntersections(LocalRay);
+		for (int32 i = IntersectingChildren.Num() - 1; i >= 0; --i)
+		{
+			const int32 Idx = IntersectingChildren[i];
+			if (Idx < MOriginalParticleLookupHack.Num())
+			{
+				IntersectingChildren[i] = MOriginalParticleLookupHack[Idx];
+			}
+			else
+			{
+				IntersectingChildren.RemoveAtSwap(i);
+			}
+		}
+		/*for (int32& Idx : IntersectingChildren)
+		{
+			Idx = MOriginalParticleLookupHack[Idx];
+		}*/
+	}
+	else
+	{
+		IntersectingChildren = MOriginalParticleLookupHack;
+	}
+
+	return IntersectingChildren;
+}
+
+template<class T, int d>
+void TImplicitObjectUnion<T,d>::CacheAllImplicitObjects()
+{
+	TArray < Pair<TSerializablePtr<FImplicitObject>, TRigidTransform<T, d>>> SubObjects;
+	AccumulateAllSerializableImplicitObjectsHelper(SubObjects, TRigidTransform<T, d>::Identity);
+	//build hierarchy
+	{
+		const int32 NumObjects = SubObjects.Num();
+		constexpr int32 MinSubObjectsToCache = 8;	//todo(make this tunable?)
+		if (NumObjects > MinSubObjectsToCache)
+		{
+			GeomParticles.Resize(NumObjects);
+			for (int32 i = 0; i < NumObjects; ++i)
+			{
+				GeomParticles.X(i) = SubObjects[i].Second.GetLocation();
+				GeomParticles.R(i) = SubObjects[i].Second.GetRotation();
+				GeomParticles.SetGeometry(i, SubObjects[i].First);
+				//check(!SubObjects[i].First->IsUnderlyingUnion());	//we don't support union of unions
+			}
+
+			TBoundingVolumeHierarchy<TGeometryParticles<T, d>, TArray<int32>, T, d> NewHierarchy(GeomParticles, 1);
+			*Hierarchy = MoveTemp(NewHierarchy);
+			bHierarchyBuilt = true;
+		}
+	}
+}
+
+template<class T, int d>
+void TImplicitObjectUnion<T,d>::Serialize(FChaosArchive& Ar)
+{
+	FChaosArchiveScopedMemory ScopedMemory(Ar, GetTypeName(), false);
+	FImplicitObject::SerializeImp(Ar);
+	Ar << MObjects << MLocalBoundingBox << GeomParticles << *Hierarchy << bHierarchyBuilt;
+}
+
+template<class T, int d>
+TImplicitObjectUnion<T, d>::TImplicitObjectUnion() : FImplicitObject(EImplicitObject::HasBoundingBox, ImplicitObjectType::Union), Hierarchy(new TBoundingVolumeHierarchy<TGeometryParticles<T, d>, TArray<int32>, T, d>(GeomParticles, 1)){}
+
+template class TImplicitObjectUnion<float, 3>;
