Adding missing namespace.

diff --git a/Engine/Source/Runtime/Experimental/Chaos/Private/Chaos/ImplicitObjectUnion.cpp b/Engine/Source/Runtime/Experimental/Chaos/Private/Chaos/ImplicitObjectUnion.cpp
--- a/Engine/Source/Runtime/Experimental/Chaos/Private/Chaos/ImplicitObjectUnion.cpp
+++ b/Engine/Source/Runtime/Experimental/Chaos/Private/Chaos/ImplicitObjectUnion.cpp
@@ -185,4 +185,4 @@ void TImplicitObjectUnion<T,d>::Serialize(FChaosArchive& Ar)
 template<class T, int d>
 TImplicitObjectUnion<T, d>::TImplicitObjectUnion() : FImplicitObject(EImplicitObject::HasBoundingBox, ImplicitObjectType::Union), Hierarchy(new TBoundingVolumeHierarchy<TGeometryParticles<T, d>, TArray<int32>, T, d>(GeomParticles, 1)){}
 
-template class TImplicitObjectUnion<float, 3>;
+template class Chaos::TImplicitObjectUnion<float, 3>;
