Fix GJKRaycast passing the wrong verts into EPA

diff --git a/Engine/Source/Runtime/Experimental/Chaos/Public/Chaos/GJK.h b/Engine/Source/Runtime/Experimental/Chaos/Public/Chaos/GJK.h
--- a/Engine/Source/Runtime/Experimental/Chaos/Public/Chaos/GJK.h
+++ b/Engine/Source/Runtime/Experimental/Chaos/Public/Chaos/GJK.h
@@ -431,19 +431,21 @@ namespace Chaos
 					for (int i = 0; i < SimplexIDs.NumVerts; ++i)
 					{
 						VertsA.Add(As[i]);
-						const TVec3<T> BInA = StartTM.TransformPosition(Bs[i]);
-						VertsB.Add(BInA);
+						const TVec3<T> BAtOrigin = Bs[i] + X;
+						VertsB.Add(BAtOrigin);
 					}
 
-					auto SupportBInAFunc = [&](const TVec3<T>& Dir)
+
+					auto SupportBAtOriginFunc = [&](const TVec3<T>& Dir)
 					{
-						const TVec3<T> SupportBLocal = SupportBFunc(Dir);
-						return StartTM.TransformPosition(SupportBLocal);
+						const TVector<T, 3> DirInB = AToBRotation * Dir;
+						const TVector<T, 3> SupportBLocal = B.Support2(DirInB);
+						return StartTM.TransformPositionNoScale(SupportBLocal);
 					};
 
 					T Penetration;
 					TVec3<T> MTD, ClosestA, ClosestBInA;
-					if (EPA(VertsA, VertsB, SupportAFunc, SupportBInAFunc, Penetration, MTD, ClosestA, ClosestBInA) != EPAResult::BadInitialSimplex)
+					if (EPA(VertsA, VertsB, SupportA, SupportBAtOriginFunc, Penetration, MTD, ClosestA, ClosestBInA) != EPAResult::BadInitialSimplex)
 					{
 						OutNormal = MTD;
 						OutTime = -Penetration - Inflation;
