Merging CL#10166083

Fix CIS

diff --git a/Engine/Source/Runtime/Experimental/Chaos/Public/Chaos/GJK.h b/Engine/Source/Runtime/Experimental/Chaos/Public/Chaos/GJK.h
--- a/Engine/Source/Runtime/Experimental/Chaos/Public/Chaos/GJK.h
+++ b/Engine/Source/Runtime/Experimental/Chaos/Public/Chaos/GJK.h
@@ -445,7 +445,7 @@ namespace Chaos
 
 					T Penetration;
 					TVec3<T> MTD, ClosestA, ClosestBInA;
-					if (EPA(VertsA, VertsB, SupportA, SupportBAtOriginFunc, Penetration, MTD, ClosestA, ClosestBInA) != EPAResult::BadInitialSimplex)
+					if (EPA(VertsA, VertsB, SupportAFunc, SupportBAtOriginFunc, Penetration, MTD, ClosestA, ClosestBInA) != EPAResult::BadInitialSimplex)
 					{
 						OutNormal = MTD;
 						OutTime = -Penetration - Inflation;
