Fix copyright.

diff --git a/Engine/Source/Runtime/Experimental/Chaos/Private/Chaos/BVHParticles.cpp b/Engine/Source/Runtime/Experimental/Chaos/Private/Chaos/BVHParticles.cpp
--- a/Engine/Source/Runtime/Experimental/Chaos/Private/Chaos/BVHParticles.cpp
+++ b/Engine/Source/Runtime/Experimental/Chaos/Private/Chaos/BVHParticles.cpp
@@ -1,3 +1,5 @@
+// Copyright 1998-2019 Epic Games, Inc. All Rights Reserved.
+
 #include "Chaos/BVHParticles.h"
 #include "Chaos/BoundingVolumeHierarchy.h"
 
