#UE4 CIS fix

diff --git a/Engine/Source/Runtime/Experimental/Chaos/Private/Chaos/CollisionResolutionAlgo.cpp b/Engine/Source/Runtime/Experimental/Chaos/Private/Chaos/CollisionResolutionAlgo.cpp
--- a/Engine/Source/Runtime/Experimental/Chaos/Private/Chaos/CollisionResolutionAlgo.cpp
+++ b/Engine/Source/Runtime/Experimental/Chaos/Private/Chaos/CollisionResolutionAlgo.cpp
@@ -109,7 +109,7 @@ namespace Chaos
 	template<typename T, int d>
 	void ConstructUnionUnionConstraints(TGeometryParticleHandle<T, d>* Particle0, TGeometryParticleHandle<T, d>* Particle1, const FImplicitObject* Implicit0, const FImplicitObject* Implicit1, const T Thickness, TRigidBodyContactConstraint<T, d>& Constraint)
 	{
-		SCOPE_CYCLE_COUNTER(STAT_UpdateUnionUnionConstraint);
+		//SCOPE_CYCLE_COUNTER(STAT_UpdateUnionUnionConstraint);
 
 		const FImplicitObject* ParticleObj = Particle0->Geometry().Get();
 		const FImplicitObject* LevelsetObj = Particle1->Geometry().Get();
