#UE4 Nonunity fix

diff --git a/Engine/Source/Runtime/Experimental/Chaos/Private/Chaos/PBDConstraintGraph.cpp b/Engine/Source/Runtime/Experimental/Chaos/Private/Chaos/PBDConstraintGraph.cpp
--- a/Engine/Source/Runtime/Experimental/Chaos/Private/Chaos/PBDConstraintGraph.cpp
+++ b/Engine/Source/Runtime/Experimental/Chaos/Private/Chaos/PBDConstraintGraph.cpp
@@ -9,6 +9,7 @@
 #include "Containers/Queue.h"
 #include "Chaos/ParticleHandle.h"
 #include "Chaos/PBDRigidsSOAs.h"
+#include "ChaosLog.h"
 
 using namespace Chaos;
 
