#UE4 Disable a few pragmas since some platforms dont compile with them

diff --git a/Engine/Source/Runtime/Experimental/Chaos/Public/Chaos/PBDRigidParticles.h b/Engine/Source/Runtime/Experimental/Chaos/Public/Chaos/PBDRigidParticles.h
--- a/Engine/Source/Runtime/Experimental/Chaos/Public/Chaos/PBDRigidParticles.h
+++ b/Engine/Source/Runtime/Experimental/Chaos/Public/Chaos/PBDRigidParticles.h
@@ -5,8 +5,10 @@
 #include "Chaos/RigidParticles.h"
 #include "Chaos/Rotation.h"
 
+#if !PLATFORM_PS4
 #pragma warning(push)
 #pragma warning(disable:4946)
+#endif
 
 namespace Chaos
 {
@@ -169,4 +171,6 @@ FChaosArchive& operator<<(FChaosArchive& Ar, TPBDRigidParticles<T, d>& Particles
 }
 }
 
+#if !PLATFORM_PS4
 #pragma warning(pop)
+#endif
