#UE4 IWYU fix

diff --git a/Engine/Source/Runtime/Experimental/Chaos/Private/Chaos/PBDConstraintColor.cpp b/Engine/Source/Runtime/Experimental/Chaos/Private/Chaos/PBDConstraintColor.cpp
--- a/Engine/Source/Runtime/Experimental/Chaos/Private/Chaos/PBDConstraintColor.cpp
+++ b/Engine/Source/Runtime/Experimental/Chaos/Private/Chaos/PBDConstraintColor.cpp
@@ -11,6 +11,7 @@
 #include "ProfilingDebugging/ScopedTimers.h"
 #include "ChaosStats.h"
 #include "Containers/Queue.h"
+#include "ChaosLog.h"
 
 #include <memory>
 #include <queue>
