Allow heightfield and trimesh to compute mtd for sweeps

diff --git a/Engine/Source/Runtime/Experimental/Chaos/Private/Chaos/HeightField.cpp b/Engine/Source/Runtime/Experimental/Chaos/Private/Chaos/HeightField.cpp
--- a/Engine/Source/Runtime/Experimental/Chaos/Private/Chaos/HeightField.cpp
+++ b/Engine/Source/Runtime/Experimental/Chaos/Private/Chaos/HeightField.cpp
@@ -193,7 +193,7 @@ namespace Chaos
 	{
 	public:
 
-		THeightfieldSweepVisitor(const typename THeightField<T>::FDataType* InData, const GeomQueryType& InQueryGeom, const TRigidTransform<T, 3>& InStartTM, const TVector<T, 3>& InDir, const T InThickness)
+		THeightfieldSweepVisitor(const typename THeightField<T>::FDataType* InData, const GeomQueryType& InQueryGeom, const TRigidTransform<T, 3>& InStartTM, const TVector<T, 3>& InDir, const T InThickness, bool InComputeMTD)
 			: OutTime(TNumericLimits<T>::Max())
 			, OutFaceIndex(INDEX_NONE)
 			, HfData(InData)
@@ -201,6 +201,7 @@ namespace Chaos
 			, OtherGeom(InQueryGeom)
 			, Dir(InDir)
 			, Thickness(InThickness)
+			, bComputeMTD(InComputeMTD)
 		{}
 
 		bool VisitSweep(int32 Payload, T& CurrentLength)
@@ -221,20 +222,22 @@ namespace Chaos
 				T Time;
 				TVector<T, 3> HitPosition;
 				TVector<T, 3> HitNormal;
-				if(GJKRaycast2<T>(Convex, OtherGeom, StartTM, Dir, CurrentLength, Time, HitPosition, HitNormal, Thickness))
+				if(GJKRaycast2<T>(Convex, OtherGeom, StartTM, Dir, CurrentLength, Time, HitPosition, HitNormal, Thickness, bComputeMTD))
 				{
 					if(Time < OutTime)
 					{
 						OutNormal = HitNormal;
 						OutPosition = HitPosition;
 						OutTime = Time;
-						CurrentLength = Time;
 						OutFaceIndex = FaceIndex;
 
-						if(Time == 0)
+						if(Time <= 0)	//initial overlap or MTD, so stop
 						{
+							CurrentLength = 0;
 							return false;
 						}
+
+						CurrentLength = Time;
 					}
 				}
 
@@ -244,10 +247,13 @@ namespace Chaos
 			TVector<T, 3> Points[4];
 			HfData->GetPointsScaled(FullIndex, Points);
 
-			TestTriangle(Payload * 2, Points[0], Points[1], Points[2]);
-			TestTriangle(Payload * 2 + 1, Points[2], Points[1], Points[3]);
+			bool bContinue = TestTriangle(Payload * 2, Points[0], Points[1], Points[2]);
+			if (bContinue)
+			{
+				TestTriangle(Payload * 2 + 1, Points[2], Points[1], Points[3]);
+			}
 
-			return OutTime != 0;
+			return OutTime > 0;
 		}
 
 		T OutTime;
@@ -262,6 +268,7 @@ namespace Chaos
 		const GeomQueryType& OtherGeom;
 		const TVector<T, 3>& Dir;
 		const T Thickness;
+		bool bComputeMTD;
 
 	};
 
@@ -1160,10 +1167,10 @@ namespace Chaos
 
 	template <typename T>
 	template <typename QueryGeomType>
-	bool THeightField<T>::SweepGeomImp(const QueryGeomType& QueryGeom, const TRigidTransform<T, 3>& StartTM, const TVector<T, 3>& Dir, const T Length, T& OutTime, TVector<T, 3>& OutPosition, TVector<T, 3>& OutNormal, int32& OutFaceIndex, const T Thickness) const
+	bool THeightField<T>::SweepGeomImp(const QueryGeomType& QueryGeom, const TRigidTransform<T, 3>& StartTM, const TVector<T, 3>& Dir, const T Length, T& OutTime, TVector<T, 3>& OutPosition, TVector<T, 3>& OutNormal, int32& OutFaceIndex, const T Thickness, bool bComputeMTD) const
 	{
 		bool bHit = false;
-		THeightfieldSweepVisitor<QueryGeomType, T> SQVisitor(&GeomData, QueryGeom, StartTM, Dir, Thickness);
+		THeightfieldSweepVisitor<QueryGeomType, T> SQVisitor(&GeomData, QueryGeom, StartTM, Dir, Thickness, bComputeMTD);
 		const TBox<T, 3> QueryBounds = QueryGeom.BoundingBox();
 		const TVector<T, 3> StartPoint = StartTM.TransformPositionNoScale(QueryBounds.Center());
 
@@ -1183,51 +1190,51 @@ namespace Chaos
 	}
 
 	template <typename T>
-	bool THeightField<T>::SweepGeom(const TSphere<T, 3>& QueryGeom, const TRigidTransform<T, 3>& StartTM, const TVector<T, 3>& Dir, const T Length, T& OutTime, TVector<T, 3>& OutPosition, TVector<T, 3>& OutNormal, int32& OutFaceIndex, const T Thickness) const
+	bool THeightField<T>::SweepGeom(const TSphere<T, 3>& QueryGeom, const TRigidTransform<T, 3>& StartTM, const TVector<T, 3>& Dir, const T Length, T& OutTime, TVector<T, 3>& OutPosition, TVector<T, 3>& OutNormal, int32& OutFaceIndex, const T Thickness, bool bComputeMTD) const
 	{
-		return SweepGeomImp(QueryGeom, StartTM, Dir, Length, OutTime, OutPosition, OutNormal, OutFaceIndex, Thickness);
+		return SweepGeomImp(QueryGeom, StartTM, Dir, Length, OutTime, OutPosition, OutNormal, OutFaceIndex, Thickness, bComputeMTD);
 	}
 
 	template <typename T>
-	bool THeightField<T>::SweepGeom(const TBox<T, 3>& QueryGeom, const TRigidTransform<T, 3>& StartTM, const TVector<T, 3>& Dir, const T Length, T& OutTime, TVector<T, 3>& OutPosition, TVector<T, 3>& OutNormal, int32& OutFaceIndex, const T Thickness) const
+	bool THeightField<T>::SweepGeom(const TBox<T, 3>& QueryGeom, const TRigidTransform<T, 3>& StartTM, const TVector<T, 3>& Dir, const T Length, T& OutTime, TVector<T, 3>& OutPosition, TVector<T, 3>& OutNormal, int32& OutFaceIndex, const T Thickness, bool bComputeMTD) const
 	{
-		return SweepGeomImp(QueryGeom, StartTM, Dir, Length, OutTime, OutPosition, OutNormal, OutFaceIndex, Thickness);
+		return SweepGeomImp(QueryGeom, StartTM, Dir, Length, OutTime, OutPosition, OutNormal, OutFaceIndex, Thickness, bComputeMTD);
 	}
 
 	template <typename T>
-	bool THeightField<T>::SweepGeom(const TCapsule<T>& QueryGeom, const TRigidTransform<T, 3>& StartTM, const TVector<T, 3>& Dir, const T Length, T& OutTime, TVector<T, 3>& OutPosition, TVector<T, 3>& OutNormal, int32& OutFaceIndex, const T Thickness) const
+	bool THeightField<T>::SweepGeom(const TCapsule<T>& QueryGeom, const TRigidTransform<T, 3>& StartTM, const TVector<T, 3>& Dir, const T Length, T& OutTime, TVector<T, 3>& OutPosition, TVector<T, 3>& OutNormal, int32& OutFaceIndex, const T Thickness, bool bComputeMTD) const
 	{
-		return SweepGeomImp(QueryGeom, StartTM, Dir, Length, OutTime, OutPosition, OutNormal, OutFaceIndex, Thickness);
+		return SweepGeomImp(QueryGeom, StartTM, Dir, Length, OutTime, OutPosition, OutNormal, OutFaceIndex, Thickness, bComputeMTD);
 	}
 
 	template <typename T>
-	bool THeightField<T>::SweepGeom(const TConvex<T, 3>& QueryGeom, const TRigidTransform<T, 3>& StartTM, const TVector<T, 3>& Dir, const T Length, T& OutTime, TVector<T, 3>& OutPosition, TVector<T, 3>& OutNormal, int32& OutFaceIndex, const T Thickness) const
+	bool THeightField<T>::SweepGeom(const TConvex<T, 3>& QueryGeom, const TRigidTransform<T, 3>& StartTM, const TVector<T, 3>& Dir, const T Length, T& OutTime, TVector<T, 3>& OutPosition, TVector<T, 3>& OutNormal, int32& OutFaceIndex, const T Thickness, bool bComputeMTD) const
 	{
-		return SweepGeomImp(QueryGeom, StartTM, Dir, Length, OutTime, OutPosition, OutNormal, OutFaceIndex, Thickness);
+		return SweepGeomImp(QueryGeom, StartTM, Dir, Length, OutTime, OutPosition, OutNormal, OutFaceIndex, Thickness, bComputeMTD);
 	}
 
 	template <typename T>
-	bool THeightField<T>::SweepGeom(const TImplicitObjectScaled<TSphere<T, 3>>& QueryGeom, const TRigidTransform<T, 3>& StartTM, const TVector<T, 3>& Dir, const T Length, T& OutTime, TVector<T, 3>& OutPosition, TVector<T, 3>& OutNormal, int32& OutFaceIndex, const T Thickness) const
+	bool THeightField<T>::SweepGeom(const TImplicitObjectScaled<TSphere<T, 3>>& QueryGeom, const TRigidTransform<T, 3>& StartTM, const TVector<T, 3>& Dir, const T Length, T& OutTime, TVector<T, 3>& OutPosition, TVector<T, 3>& OutNormal, int32& OutFaceIndex, const T Thickness, bool bComputeMTD) const
 	{
-		return SweepGeomImp(QueryGeom, StartTM, Dir, Length, OutTime, OutPosition, OutNormal, OutFaceIndex, Thickness);
+		return SweepGeomImp(QueryGeom, StartTM, Dir, Length, OutTime, OutPosition, OutNormal, OutFaceIndex, Thickness, bComputeMTD);
 	}
 
 	template <typename T>
-	bool THeightField<T>::SweepGeom(const TImplicitObjectScaled<TBox<T, 3>>& QueryGeom, const TRigidTransform<T, 3>& StartTM, const TVector<T, 3>& Dir, const T Length, T& OutTime, TVector<T, 3>& OutPosition, TVector<T, 3>& OutNormal, int32& OutFaceIndex, const T Thickness) const
+	bool THeightField<T>::SweepGeom(const TImplicitObjectScaled<TBox<T, 3>>& QueryGeom, const TRigidTransform<T, 3>& StartTM, const TVector<T, 3>& Dir, const T Length, T& OutTime, TVector<T, 3>& OutPosition, TVector<T, 3>& OutNormal, int32& OutFaceIndex, const T Thickness, bool bComputeMTD) const
 	{
-		return SweepGeomImp(QueryGeom, StartTM, Dir, Length, OutTime, OutPosition, OutNormal, OutFaceIndex, Thickness);
+		return SweepGeomImp(QueryGeom, StartTM, Dir, Length, OutTime, OutPosition, OutNormal, OutFaceIndex, Thickness, bComputeMTD);
 	}
 
 	template <typename T>
-	bool THeightField<T>::SweepGeom(const TImplicitObjectScaled<TCapsule<T>>& QueryGeom, const TRigidTransform<T, 3>& StartTM, const TVector<T, 3>& Dir, const T Length, T& OutTime, TVector<T, 3>& OutPosition, TVector<T, 3>& OutNormal, int32& OutFaceIndex, const T Thickness) const
+	bool THeightField<T>::SweepGeom(const TImplicitObjectScaled<TCapsule<T>>& QueryGeom, const TRigidTransform<T, 3>& StartTM, const TVector<T, 3>& Dir, const T Length, T& OutTime, TVector<T, 3>& OutPosition, TVector<T, 3>& OutNormal, int32& OutFaceIndex, const T Thickness, bool bComputeMTD) const
 	{
-		return SweepGeomImp(QueryGeom, StartTM, Dir, Length, OutTime, OutPosition, OutNormal, OutFaceIndex, Thickness);
+		return SweepGeomImp(QueryGeom, StartTM, Dir, Length, OutTime, OutPosition, OutNormal, OutFaceIndex, Thickness, bComputeMTD);
 	}
 
 	template <typename T>
-	bool THeightField<T>::SweepGeom(const TImplicitObjectScaled<TConvex<T, 3>>& QueryGeom, const TRigidTransform<T, 3>& StartTM, const TVector<T, 3>& Dir, const T Length, T& OutTime, TVector<T, 3>& OutPosition, TVector<T, 3>& OutNormal, int32& OutFaceIndex, const T Thickness) const
+	bool THeightField<T>::SweepGeom(const TImplicitObjectScaled<TConvex<T, 3>>& QueryGeom, const TRigidTransform<T, 3>& StartTM, const TVector<T, 3>& Dir, const T Length, T& OutTime, TVector<T, 3>& OutPosition, TVector<T, 3>& OutNormal, int32& OutFaceIndex, const T Thickness, bool bComputeMTD) const
 	{
-		return SweepGeomImp(QueryGeom, StartTM, Dir, Length, OutTime, OutPosition, OutNormal, OutFaceIndex, Thickness);
+		return SweepGeomImp(QueryGeom, StartTM, Dir, Length, OutTime, OutPosition, OutNormal, OutFaceIndex, Thickness, bComputeMTD);
 	}
 
 	template <typename T>
diff --git a/Engine/Source/Runtime/Experimental/Chaos/Private/Chaos/TriangleMeshImplicitObject.cpp b/Engine/Source/Runtime/Experimental/Chaos/Private/Chaos/TriangleMeshImplicitObject.cpp
--- a/Engine/Source/Runtime/Experimental/Chaos/Private/Chaos/TriangleMeshImplicitObject.cpp
+++ b/Engine/Source/Runtime/Experimental/Chaos/Private/Chaos/TriangleMeshImplicitObject.cpp
@@ -362,12 +362,13 @@ bool TTriangleMeshImplicitObject<T>::OverlapGeom(const TImplicitObjectScaled<TIm
 template <typename QueryGeomType, typename T>
 struct TTriangleMeshSweepVisitor
 {
-	TTriangleMeshSweepVisitor(const TTriangleMeshImplicitObject<T>& InTriMesh, const QueryGeomType& InQueryGeom, const TRigidTransform<T,3>& InStartTM, const TVector<T,3>& InDir, const T InThickness)
+	TTriangleMeshSweepVisitor(const TTriangleMeshImplicitObject<T>& InTriMesh, const QueryGeomType& InQueryGeom, const TRigidTransform<T,3>& InStartTM, const TVector<T,3>& InDir, const T InThickness, const bool InComputeMTD)
 	: TriMesh(InTriMesh)
 	, StartTM(InStartTM)
 	, QueryGeom(InQueryGeom)
 	, Dir(InDir)
 	, Thickness(InThickness)
+	, bComputeMTD(InComputeMTD)
 	, OutTime(TNumericLimits<T>::Max())
 	{
 	}
@@ -396,21 +397,24 @@ struct TTriangleMeshSweepVisitor
 			TriMesh.MParticles.X(TriMesh.MElements[TriIdx][1]),
 			TriMesh.MParticles.X(TriMesh.MElements[TriIdx][2]));
 
-		if(GJKRaycast2<T>(Tri, QueryGeom, StartTM, Dir, CurData.CurrentLength, Time, HitPosition, HitNormal, Thickness))
+		if(GJKRaycast2<T>(Tri, QueryGeom, StartTM, Dir, CurData.CurrentLength, Time, HitPosition, HitNormal, Thickness, bComputeMTD))
 		{
 			if(Time < OutTime)
 			{
 				OutNormal = HitNormal;
 				OutPosition = HitPosition;
 				OutTime = Time;
-				CurData.SetLength(Time);
 				OutFaceIndex = TriIdx;
 
-				if(Time == 0)
+				if(Time <= 0)	//MTD or initial overlap
 				{
+					CurData.SetLength(0);
+
 					//initial overlap, no one will beat this
 					return false;
 				}
+
+				CurData.SetLength(Time);
 			}
 		}
 
@@ -422,6 +426,7 @@ struct TTriangleMeshSweepVisitor
 	const QueryGeomType& QueryGeom;
 	const TVector<T, 3>& Dir;
 	const T Thickness;
+	const bool bComputeMTD;
 
 	T OutTime;
 	TVector<T, 3> OutPosition;
@@ -431,10 +436,10 @@ struct TTriangleMeshSweepVisitor
 
 template <typename T>
 template <typename QueryGeomType>
-bool TTriangleMeshImplicitObject<T>::SweepGeomImp(const QueryGeomType& QueryGeom, const TRigidTransform<T, 3>& StartTM, const TVector<T, 3>& Dir, const T Length, T& OutTime, TVector<T, 3>& OutPosition, TVector<T, 3>& OutNormal, int32& OutFaceIndex, const T Thickness) const
+bool TTriangleMeshImplicitObject<T>::SweepGeomImp(const QueryGeomType& QueryGeom, const TRigidTransform<T, 3>& StartTM, const TVector<T, 3>& Dir, const T Length, T& OutTime, TVector<T, 3>& OutPosition, TVector<T, 3>& OutNormal, int32& OutFaceIndex, const T Thickness, const bool bComputeMTD) const
 {
 	bool bHit = false;
-	TTriangleMeshSweepVisitor<QueryGeomType, T> SQVisitor(*this, QueryGeom, StartTM, Dir, Thickness);
+	TTriangleMeshSweepVisitor<QueryGeomType, T> SQVisitor(*this, QueryGeom, StartTM, Dir, Thickness, bComputeMTD);
 	const TBox<T, 3> QueryBounds = QueryGeom.BoundingBox();
 	const TVector<T, 3> StartPoint = StartTM.TransformPositionNoScale(QueryBounds.Center());
 	const TVector<T, 3> Inflation = QueryBounds.Extents() * 0.5 + TVector<T, 3>(Thickness);
@@ -452,51 +457,51 @@ bool TTriangleMeshImplicitObject<T>::SweepGeomImp(const QueryGeomType& QueryGeom
 }
 
 template <typename T>
-bool TTriangleMeshImplicitObject<T>::SweepGeom(const TSphere<T,3>& QueryGeom, const TRigidTransform<T, 3>& StartTM, const TVector<T, 3>& Dir, const T Length, T& OutTime, TVector<T, 3>& OutPosition, TVector<T, 3>& OutNormal, int32& OutFaceIndex, const T Thickness) const
+bool TTriangleMeshImplicitObject<T>::SweepGeom(const TSphere<T,3>& QueryGeom, const TRigidTransform<T, 3>& StartTM, const TVector<T, 3>& Dir, const T Length, T& OutTime, TVector<T, 3>& OutPosition, TVector<T, 3>& OutNormal, int32& OutFaceIndex, const T Thickness, const bool bComputeMTD) const
 {
-	return SweepGeomImp(QueryGeom, StartTM, Dir, Length, OutTime, OutPosition, OutNormal, OutFaceIndex, Thickness);
+	return SweepGeomImp(QueryGeom, StartTM, Dir, Length, OutTime, OutPosition, OutNormal, OutFaceIndex, Thickness, bComputeMTD);
 }
 
 template <typename T>
-bool TTriangleMeshImplicitObject<T>::SweepGeom(const TBox<T, 3>& QueryGeom, const TRigidTransform<T, 3>& StartTM, const TVector<T, 3>& Dir, const T Length, T& OutTime, TVector<T, 3>& OutPosition, TVector<T, 3>& OutNormal, int32& OutFaceIndex, const T Thickness) const
+bool TTriangleMeshImplicitObject<T>::SweepGeom(const TBox<T, 3>& QueryGeom, const TRigidTransform<T, 3>& StartTM, const TVector<T, 3>& Dir, const T Length, T& OutTime, TVector<T, 3>& OutPosition, TVector<T, 3>& OutNormal, int32& OutFaceIndex, const T Thickness, const bool bComputeMTD) const
 {
-	return SweepGeomImp(QueryGeom, StartTM, Dir, Length, OutTime, OutPosition, OutNormal, OutFaceIndex, Thickness);
+	return SweepGeomImp(QueryGeom, StartTM, Dir, Length, OutTime, OutPosition, OutNormal, OutFaceIndex, Thickness, bComputeMTD);
 }
 
 template <typename T>
-bool TTriangleMeshImplicitObject<T>::SweepGeom(const TCapsule<T>& QueryGeom, const TRigidTransform<T, 3>& StartTM, const TVector<T, 3>& Dir, const T Length, T& OutTime, TVector<T, 3>& OutPosition, TVector<T, 3>& OutNormal, int32& OutFaceIndex, const T Thickness) const
+bool TTriangleMeshImplicitObject<T>::SweepGeom(const TCapsule<T>& QueryGeom, const TRigidTransform<T, 3>& StartTM, const TVector<T, 3>& Dir, const T Length, T& OutTime, TVector<T, 3>& OutPosition, TVector<T, 3>& OutNormal, int32& OutFaceIndex, const T Thickness, const bool bComputeMTD) const
 {
-	return SweepGeomImp(QueryGeom, StartTM, Dir, Length, OutTime, OutPosition, OutNormal, OutFaceIndex, Thickness);
+	return SweepGeomImp(QueryGeom, StartTM, Dir, Length, OutTime, OutPosition, OutNormal, OutFaceIndex, Thickness, bComputeMTD);
 }
 
 template <typename T>
-bool TTriangleMeshImplicitObject<T>::SweepGeom(const TConvex<T, 3>& QueryGeom, const TRigidTransform<T, 3>& StartTM, const TVector<T, 3>& Dir, const T Length, T& OutTime, TVector<T, 3>& OutPosition, TVector<T, 3>& OutNormal, int32& OutFaceIndex, const T Thickness) const
+bool TTriangleMeshImplicitObject<T>::SweepGeom(const TConvex<T, 3>& QueryGeom, const TRigidTransform<T, 3>& StartTM, const TVector<T, 3>& Dir, const T Length, T& OutTime, TVector<T, 3>& OutPosition, TVector<T, 3>& OutNormal, int32& OutFaceIndex, const T Thickness, const bool bComputeMTD) const
 {
-	return SweepGeomImp(QueryGeom, StartTM, Dir, Length, OutTime, OutPosition, OutNormal, OutFaceIndex, Thickness);
+	return SweepGeomImp(QueryGeom, StartTM, Dir, Length, OutTime, OutPosition, OutNormal, OutFaceIndex, Thickness, bComputeMTD);
 }
 
 template <typename T>
-bool TTriangleMeshImplicitObject<T>::SweepGeom(const TImplicitObjectScaled<TSphere<T, 3>>& QueryGeom, const TRigidTransform<T, 3>& StartTM, const TVector<T, 3>& Dir, const T Length, T& OutTime, TVector<T, 3>& OutPosition, TVector<T, 3>& OutNormal, int32& OutFaceIndex, const T Thickness) const
+bool TTriangleMeshImplicitObject<T>::SweepGeom(const TImplicitObjectScaled<TSphere<T, 3>>& QueryGeom, const TRigidTransform<T, 3>& StartTM, const TVector<T, 3>& Dir, const T Length, T& OutTime, TVector<T, 3>& OutPosition, TVector<T, 3>& OutNormal, int32& OutFaceIndex, const T Thickness, const bool bComputeMTD) const
 {
-	return SweepGeomImp(QueryGeom, StartTM, Dir, Length, OutTime, OutPosition, OutNormal, OutFaceIndex, Thickness);
+	return SweepGeomImp(QueryGeom, StartTM, Dir, Length, OutTime, OutPosition, OutNormal, OutFaceIndex, Thickness, bComputeMTD);
 }
 
 template <typename T>
-bool TTriangleMeshImplicitObject<T>::SweepGeom(const TImplicitObjectScaled<TBox<T, 3>>& QueryGeom, const TRigidTransform<T, 3>& StartTM, const TVector<T, 3>& Dir, const T Length, T& OutTime, TVector<T, 3>& OutPosition, TVector<T, 3>& OutNormal, int32& OutFaceIndex, const T Thickness) const
+bool TTriangleMeshImplicitObject<T>::SweepGeom(const TImplicitObjectScaled<TBox<T, 3>>& QueryGeom, const TRigidTransform<T, 3>& StartTM, const TVector<T, 3>& Dir, const T Length, T& OutTime, TVector<T, 3>& OutPosition, TVector<T, 3>& OutNormal, int32& OutFaceIndex, const T Thickness, const bool bComputeMTD) const
 {
-	return SweepGeomImp(QueryGeom, StartTM, Dir, Length, OutTime, OutPosition, OutNormal, OutFaceIndex, Thickness);
+	return SweepGeomImp(QueryGeom, StartTM, Dir, Length, OutTime, OutPosition, OutNormal, OutFaceIndex, Thickness, bComputeMTD);
 }
 
 template <typename T>
-bool TTriangleMeshImplicitObject<T>::SweepGeom(const TImplicitObjectScaled<TCapsule<T>>& QueryGeom, const TRigidTransform<T, 3>& StartTM, const TVector<T, 3>& Dir, const T Length, T& OutTime, TVector<T, 3>& OutPosition, TVector<T, 3>& OutNormal, int32& OutFaceIndex, const T Thickness) const
+bool TTriangleMeshImplicitObject<T>::SweepGeom(const TImplicitObjectScaled<TCapsule<T>>& QueryGeom, const TRigidTransform<T, 3>& StartTM, const TVector<T, 3>& Dir, const T Length, T& OutTime, TVector<T, 3>& OutPosition, TVector<T, 3>& OutNormal, int32& OutFaceIndex, const T Thickness, const bool bComputeMTD) const
 {
-	return SweepGeomImp(QueryGeom, StartTM, Dir, Length, OutTime, OutPosition, OutNormal, OutFaceIndex, Thickness);
+	return SweepGeomImp(QueryGeom, StartTM, Dir, Length, OutTime, OutPosition, OutNormal, OutFaceIndex, Thickness, bComputeMTD);
 }
 
 template <typename T>
-bool TTriangleMeshImplicitObject<T>::SweepGeom(const TImplicitObjectScaled<TConvex<T, 3>>& QueryGeom, const TRigidTransform<T, 3>& StartTM, const TVector<T, 3>& Dir, const T Length, T& OutTime, TVector<T, 3>& OutPosition, TVector<T, 3>& OutNormal, int32& OutFaceIndex, const T Thickness) const
+bool TTriangleMeshImplicitObject<T>::SweepGeom(const TImplicitObjectScaled<TConvex<T, 3>>& QueryGeom, const TRigidTransform<T, 3>& StartTM, const TVector<T, 3>& Dir, const T Length, T& OutTime, TVector<T, 3>& OutPosition, TVector<T, 3>& OutNormal, int32& OutFaceIndex, const T Thickness, const bool bComputeMTD) const
 {
-	return SweepGeomImp(QueryGeom, StartTM, Dir, Length, OutTime, OutPosition, OutNormal, OutFaceIndex, Thickness);
+	return SweepGeomImp(QueryGeom, StartTM, Dir, Length, OutTime, OutPosition, OutNormal, OutFaceIndex, Thickness, bComputeMTD);
 }
 
 template <typename T>
diff --git a/Engine/Source/Runtime/Experimental/Chaos/Public/Chaos/GeometryQueries.h b/Engine/Source/Runtime/Experimental/Chaos/Public/Chaos/GeometryQueries.h
--- a/Engine/Source/Runtime/Experimental/Chaos/Public/Chaos/GeometryQueries.h
+++ b/Engine/Source/Runtime/Experimental/Chaos/Public/Chaos/GeometryQueries.h
@@ -119,7 +119,7 @@ namespace Chaos
 		const FRigidTransform3 BToATM = BTM.GetRelativeTransform(ATM);
 		const FVec3 LocalDir = ATM.InverseTransformVectorNoScale(Dir);
 
-		bool bSweepAsRaycast = BType == ImplicitObjectType::Sphere;
+		bool bSweepAsRaycast = BType == ImplicitObjectType::Sphere && !bComputeMTD;
 		if (bSweepAsRaycast && IsScaled(AType))
 		{
 			const auto& Scaled = TImplicitObjectScaledGeneric<FReal, 3>::AsScaledChecked(A);
@@ -175,20 +175,20 @@ namespace Chaos
 			case ImplicitObjectType::HeightField:
 			{
 				const auto& AHeightField = static_cast<const THeightField<FReal>&>(A);
-				bResult = AHeightField.SweepGeom(B, BToATM, LocalDir, Length, OutTime, LocalPosition, LocalNormal, OutFaceIndex, Thickness);
+				bResult = AHeightField.SweepGeom(B, BToATM, LocalDir, Length, OutTime, LocalPosition, LocalNormal, OutFaceIndex, Thickness, bComputeMTD);
 				break;
 			}
 			case ImplicitObjectType::TriangleMesh:
 			{
 				const auto& ATriangleMesh = static_cast<const TTriangleMeshImplicitObject<FReal>&>(A);
-				bResult = ATriangleMesh.SweepGeom(B, BToATM, LocalDir, Length, OutTime, LocalPosition, LocalNormal, OutFaceIndex, Thickness);
+				bResult = ATriangleMesh.SweepGeom(B, BToATM, LocalDir, Length, OutTime, LocalPosition, LocalNormal, OutFaceIndex, Thickness, bComputeMTD);
 				break;
 			}
 			default:
 				if (IsScaled(AType))
 				{
 					const auto& AScaled = TImplicitObjectScaled<TTriangleMeshImplicitObject<FReal>>::AsScaledChecked(A);
-					bResult = AScaled.LowLevelSweepGeom(B, BToATM, LocalDir, Length, OutTime, LocalPosition, LocalNormal, OutFaceIndex, Thickness);
+					bResult = AScaled.LowLevelSweepGeom(B, BToATM, LocalDir, Length, OutTime, LocalPosition, LocalNormal, OutFaceIndex, Thickness, bComputeMTD);
 					break;
 				}
 				else
@@ -197,18 +197,10 @@ namespace Chaos
 				}
 
 			}
-
-			// Compute MTD in the case of an initial overlap
-			if (bResult && bComputeMTD && OutTime == 0.f)
-			{
-				ensure(false); // We don't support MTD for non-convex types yet!
-				LocalNormal = FVec3(0, 0, 1);
-				LocalPosition = ATM.GetLocation();
-			}
 		}
 
 		//put back into world space
-		if (OutTime > 0 || bComputeMTD)
+		if (bResult && (OutTime > 0 || bComputeMTD))
 		{
 			OutNormal = ATM.TransformVectorNoScale(LocalNormal);
 			OutPosition = ATM.TransformPositionNoScale(LocalPosition);
diff --git a/Engine/Source/Runtime/Experimental/Chaos/Public/Chaos/HeightField.h b/Engine/Source/Runtime/Experimental/Chaos/Public/Chaos/HeightField.h
--- a/Engine/Source/Runtime/Experimental/Chaos/Public/Chaos/HeightField.h
+++ b/Engine/Source/Runtime/Experimental/Chaos/Public/Chaos/HeightField.h
@@ -56,14 +56,14 @@ namespace Chaos
 		bool OverlapGeom(const TImplicitObjectScaled<TCapsule<T>>& QueryGeom, const TRigidTransform<T, 3>& QueryTM, const T Thickness) const;
 		bool OverlapGeom(const TImplicitObjectScaled<TConvex<T, 3>>& QueryGeom, const TRigidTransform<T, 3>& QueryTM, const T Thickness) const;
 
-		bool SweepGeom(const TSphere<T, 3>& QueryGeom, const TRigidTransform<T, 3>& StartTM, const TVector<T, 3>& Dir, const T Length, T& OutTime, TVector<T, 3>& OutPosition, TVector<T, 3>& OutNormal, int32& OutFaceIndex, const T Thickness = 0) const;
-		bool SweepGeom(const TBox<T, 3>& QueryGeom, const TRigidTransform<T, 3>& StartTM, const TVector<T, 3>& Dir, const T Length, T& OutTime, TVector<T, 3>& OutPosition, TVector<T, 3>& OutNormal, int32& OutFaceIndex, const T Thickness = 0) const;
-		bool SweepGeom(const TCapsule<T>& QueryGeom, const TRigidTransform<T, 3>& StartTM, const TVector<T, 3>& Dir, const T Length, T& OutTime, TVector<T, 3>& OutPosition, TVector<T, 3>& OutNormal, int32& OutFaceIndex, const T Thickness = 0) const;
-		bool SweepGeom(const TConvex<T, 3>& QueryGeom, const TRigidTransform<T, 3>& StartTM, const TVector<T, 3>& Dir, const T Length, T& OutTime, TVector<T, 3>& OutPosition, TVector<T, 3>& OutNormal, int32& OutFaceIndex, const T Thickness = 0) const;
-		bool SweepGeom(const TImplicitObjectScaled<TSphere<T, 3>>& QueryGeom, const TRigidTransform<T, 3>& StartTM, const TVector<T, 3>& Dir, const T Length, T& OutTime, TVector<T, 3>& OutPosition, TVector<T, 3>& OutNormal, int32& OutFaceIndex, const T Thickness = 0) const;
-		bool SweepGeom(const TImplicitObjectScaled<TBox<T, 3>>& QueryGeom, const TRigidTransform<T, 3>& StartTM, const TVector<T, 3>& Dir, const T Length, T& OutTime, TVector<T, 3>& OutPosition, TVector<T, 3>& OutNormal, int32& OutFaceIndex, const T Thickness = 0) const;
-		bool SweepGeom(const TImplicitObjectScaled<TCapsule<T>>& QueryGeom, const TRigidTransform<T, 3>& StartTM, const TVector<T, 3>& Dir, const T Length, T& OutTime, TVector<T, 3>& OutPosition, TVector<T, 3>& OutNormal, int32& OutFaceIndex, const T Thickness = 0) const;
-		bool SweepGeom(const TImplicitObjectScaled<TConvex<T, 3>>& QueryGeom, const TRigidTransform<T, 3>& StartTM, const TVector<T, 3>& Dir, const T Length, T& OutTime, TVector<T, 3>& OutPosition, TVector<T, 3>& OutNormal, int32& OutFaceIndex, const T Thickness = 0) const;
+		bool SweepGeom(const TSphere<T, 3>& QueryGeom, const TRigidTransform<T, 3>& StartTM, const TVector<T, 3>& Dir, const T Length, T& OutTime, TVector<T, 3>& OutPosition, TVector<T, 3>& OutNormal, int32& OutFaceIndex, const T Thickness = 0, bool bComputeMTD = false) const;
+		bool SweepGeom(const TBox<T, 3>& QueryGeom, const TRigidTransform<T, 3>& StartTM, const TVector<T, 3>& Dir, const T Length, T& OutTime, TVector<T, 3>& OutPosition, TVector<T, 3>& OutNormal, int32& OutFaceIndex, const T Thickness = 0, bool bComputeMTD = false) const;
+		bool SweepGeom(const TCapsule<T>& QueryGeom, const TRigidTransform<T, 3>& StartTM, const TVector<T, 3>& Dir, const T Length, T& OutTime, TVector<T, 3>& OutPosition, TVector<T, 3>& OutNormal, int32& OutFaceIndex, const T Thickness = 0, bool bComputeMTD = false) const;
+		bool SweepGeom(const TConvex<T, 3>& QueryGeom, const TRigidTransform<T, 3>& StartTM, const TVector<T, 3>& Dir, const T Length, T& OutTime, TVector<T, 3>& OutPosition, TVector<T, 3>& OutNormal, int32& OutFaceIndex, const T Thickness = 0, bool bComputeMTD = false) const;
+		bool SweepGeom(const TImplicitObjectScaled<TSphere<T, 3>>& QueryGeom, const TRigidTransform<T, 3>& StartTM, const TVector<T, 3>& Dir, const T Length, T& OutTime, TVector<T, 3>& OutPosition, TVector<T, 3>& OutNormal, int32& OutFaceIndex, const T Thickness = 0, bool bComputeMTD = false) const;
+		bool SweepGeom(const TImplicitObjectScaled<TBox<T, 3>>& QueryGeom, const TRigidTransform<T, 3>& StartTM, const TVector<T, 3>& Dir, const T Length, T& OutTime, TVector<T, 3>& OutPosition, TVector<T, 3>& OutNormal, int32& OutFaceIndex, const T Thickness = 0, bool bComputeMTD = false) const;
+		bool SweepGeom(const TImplicitObjectScaled<TCapsule<T>>& QueryGeom, const TRigidTransform<T, 3>& StartTM, const TVector<T, 3>& Dir, const T Length, T& OutTime, TVector<T, 3>& OutPosition, TVector<T, 3>& OutNormal, int32& OutFaceIndex, const T Thickness = 0, bool bComputeMTD = false) const;
+		bool SweepGeom(const TImplicitObjectScaled<TConvex<T, 3>>& QueryGeom, const TRigidTransform<T, 3>& StartTM, const TVector<T, 3>& Dir, const T Length, T& OutTime, TVector<T, 3>& OutPosition, TVector<T, 3>& OutNormal, int32& OutFaceIndex, const T Thickness = 0, bool bComputeMTD = false) const;
 
 		virtual int32 FindMostOpposingFace(const TVector<T, 3>& Position, const TVector<T, 3>& UnitDir, int32 HintFaceIndex, T SearchDist) const override;
 		virtual TVector<T, 3> FindGeometryOpposingNormal(const TVector<T, 3>& DenormDir, int32 FaceIndex, const TVector<T, 3>& OriginalNormal) const override;
@@ -446,7 +446,7 @@ namespace Chaos
 		bool OverlapGeomImp(const QueryGeomType& QueryGeom, const TRigidTransform<T, 3>& QueryTM, const T Thickness) const;
 
 		template <typename QueryGeomType>
-		bool SweepGeomImp(const QueryGeomType& QueryGeom, const TRigidTransform<T, 3>& StartTM, const TVector<T, 3>& Dir, const T Length, T& OutTime, TVector<T, 3>& OutPosition, TVector<T, 3>& OutNormal, int32& OutFaceIndex, const T Thickness) const;
+		bool SweepGeomImp(const QueryGeomType& QueryGeom, const TRigidTransform<T, 3>& StartTM, const TVector<T, 3>& Dir, const T Length, T& OutTime, TVector<T, 3>& OutPosition, TVector<T, 3>& OutNormal, int32& OutFaceIndex, const T Thickness, bool bComputeMTD) const;
 
 	};
 }
diff --git a/Engine/Source/Runtime/Experimental/Chaos/Public/Chaos/ISpatialAcceleration.h b/Engine/Source/Runtime/Experimental/Chaos/Public/Chaos/ISpatialAcceleration.h
--- a/Engine/Source/Runtime/Experimental/Chaos/Public/Chaos/ISpatialAcceleration.h
+++ b/Engine/Source/Runtime/Experimental/Chaos/Public/Chaos/ISpatialAcceleration.h
@@ -25,15 +25,24 @@ struct CHAOS_API FQueryFastData
 	FReal CurrentLength;
 	FReal InvCurrentLength;
 
+	#pragma warning (push)
+	#pragma warning(disable:4723)
+	//compiler complaining about divide by 0, but we are guarding against it.
+	//Seems like it's trying to evaluate div independent of the guard?
+
 	void SetLength(const FReal InLength)
 	{
 		CurrentLength = InLength;
+
 		if(InLength)
 		{
 			InvCurrentLength = 1 / InLength;
 		}
 	}
 
+	#pragma warning(pop)
+
+
 protected:
 	FQueryFastData(const FVec3& InDummyDir)
 		: bParallel{}
diff --git a/Engine/Source/Runtime/Experimental/Chaos/Public/Chaos/ImplicitObjectScaled.h b/Engine/Source/Runtime/Experimental/Chaos/Public/Chaos/ImplicitObjectScaled.h
--- a/Engine/Source/Runtime/Experimental/Chaos/Public/Chaos/ImplicitObjectScaled.h
+++ b/Engine/Source/Runtime/Experimental/Chaos/Public/Chaos/ImplicitObjectScaled.h
@@ -291,7 +291,7 @@ public:
 
 	/** This is a low level function and assumes the internal object has a SweepGeom function. Should not be called directly. See GeometryQueries.h : SweepQuery */
 	template <typename QueryGeomType>
-	bool LowLevelSweepGeom(const QueryGeomType& B, const TRigidTransform<T, d>& BToATM, const TVector<T, d>& LocalDir, const T Length, T& OutTime, TVector<T, d>& LocalPosition, TVector<T, d>& LocalNormal, int32& OutFaceIndex, T Thickness = 0) const
+	bool LowLevelSweepGeom(const QueryGeomType& B, const TRigidTransform<T, d>& BToATM, const TVector<T, d>& LocalDir, const T Length, T& OutTime, TVector<T, d>& LocalPosition, TVector<T, d>& LocalNormal, int32& OutFaceIndex, T Thickness = 0, bool bComputeMTD = false) const
 	{
 		ensure(Length > 0);
 		ensure(FMath::IsNearlyEqual(LocalDir.SizeSquared(), 1, KINDA_SMALL_NUMBER));
@@ -313,7 +313,7 @@ public:
 
 			TRigidTransform<T, d> BToATMNoScale(BToATM.GetLocation() * MInvScale, BToATM.GetRotation());
 			
-			if (MObject->SweepGeom(ScaledB, BToATMNoScale, UnscaledDir, UnscaledLength, UnscaledTime, UnscaledPosition, UnscaledNormal, OutFaceIndex, MInternalThickness + Thickness))
+			if (MObject->SweepGeom(ScaledB, BToATMNoScale, UnscaledDir, UnscaledLength, UnscaledTime, UnscaledPosition, UnscaledNormal, OutFaceIndex, MInternalThickness + Thickness, bComputeMTD))
 			{
 				OutTime = LengthScaleInv * UnscaledTime;
 				LocalPosition = MScale * UnscaledPosition;
diff --git a/Engine/Source/Runtime/Experimental/Chaos/Public/Chaos/TriangleMeshImplicitObject.h b/Engine/Source/Runtime/Experimental/Chaos/Public/Chaos/TriangleMeshImplicitObject.h
--- a/Engine/Source/Runtime/Experimental/Chaos/Public/Chaos/TriangleMeshImplicitObject.h
+++ b/Engine/Source/Runtime/Experimental/Chaos/Public/Chaos/TriangleMeshImplicitObject.h
@@ -53,14 +53,14 @@ namespace Chaos
 		bool OverlapGeom(const TImplicitObjectScaled<TConvex<T, 3>>& QueryGeom, const TRigidTransform<T, 3>& QueryTM, const T Thickness) const;
 		bool OverlapGeom(const TImplicitObjectScaled<TImplicitObjectScaled<TConvex<T, 3>>>& QueryGeom, const TRigidTransform<T, 3>& QueryTM, const T Thickness) const;
 
-		bool SweepGeom(const TSphere<T,3>& QueryGeom, const TRigidTransform<T, 3>& StartTM, const TVector<T,3>& Dir, const T Length, T& OutTime, TVector<T,3>& OutPosition, TVector<T,3>& OutNormal, int32& OutFaceIndex, const T Thickness = 0) const;
-		bool SweepGeom(const TBox<T, 3>& QueryGeom, const TRigidTransform<T, 3>& StartTM, const TVector<T, 3>& Dir, const T Length, T& OutTime, TVector<T, 3>& OutPosition, TVector<T, 3>& OutNormal, int32& OutFaceIndex, const T Thickness = 0) const;
-		bool SweepGeom(const TCapsule<T>& QueryGeom, const TRigidTransform<T, 3>& StartTM, const TVector<T, 3>& Dir, const T Length, T& OutTime, TVector<T, 3>& OutPosition, TVector<T, 3>& OutNormal, int32& OutFaceIndex, const T Thickness = 0) const;
-		bool SweepGeom(const TConvex<T, 3>& QueryGeom, const TRigidTransform<T, 3>& StartTM, const TVector<T, 3>& Dir, const T Length, T& OutTime, TVector<T, 3>& OutPosition, TVector<T, 3>& OutNormal, int32& OutFaceIndex, const T Thickness = 0) const;
-		bool SweepGeom(const TImplicitObjectScaled<TSphere<T, 3>>& QueryGeom, const TRigidTransform<T, 3>& StartTM, const TVector<T, 3>& Dir, const T Length, T& OutTime, TVector<T, 3>& OutPosition, TVector<T, 3>& OutNormal, int32& OutFaceIndex, const T Thickness = 0) const;
-		bool SweepGeom(const TImplicitObjectScaled<TBox<T, 3>>& QueryGeom, const TRigidTransform<T, 3>& StartTM, const TVector<T, 3>& Dir, const T Length, T& OutTime, TVector<T, 3>& OutPosition, TVector<T, 3>& OutNormal, int32& OutFaceIndex, const T Thickness = 0) const;
-		bool SweepGeom(const TImplicitObjectScaled<TCapsule<T>>& QueryGeom, const TRigidTransform<T, 3>& StartTM, const TVector<T, 3>& Dir, const T Length, T& OutTime, TVector<T, 3>& OutPosition, TVector<T, 3>& OutNormal, int32& OutFaceIndex, const T Thickness = 0) const;
-		bool SweepGeom(const TImplicitObjectScaled<TConvex<T, 3>>& QueryGeom, const TRigidTransform<T, 3>& StartTM, const TVector<T, 3>& Dir, const T Length, T& OutTime, TVector<T, 3>& OutPosition, TVector<T, 3>& OutNormal, int32& OutFaceIndex, const T Thickness = 0) const;
+		bool SweepGeom(const TSphere<T,3>& QueryGeom, const TRigidTransform<T, 3>& StartTM, const TVector<T,3>& Dir, const T Length, T& OutTime, TVector<T,3>& OutPosition, TVector<T,3>& OutNormal, int32& OutFaceIndex, const T Thickness = 0, const bool bComputeMTD = false) const;
+		bool SweepGeom(const TBox<T, 3>& QueryGeom, const TRigidTransform<T, 3>& StartTM, const TVector<T, 3>& Dir, const T Length, T& OutTime, TVector<T, 3>& OutPosition, TVector<T, 3>& OutNormal, int32& OutFaceIndex, const T Thickness = 0, const bool bComputeMTD = false) const;
+		bool SweepGeom(const TCapsule<T>& QueryGeom, const TRigidTransform<T, 3>& StartTM, const TVector<T, 3>& Dir, const T Length, T& OutTime, TVector<T, 3>& OutPosition, TVector<T, 3>& OutNormal, int32& OutFaceIndex, const T Thickness = 0, const bool bComputeMTD = false) const;
+		bool SweepGeom(const TConvex<T, 3>& QueryGeom, const TRigidTransform<T, 3>& StartTM, const TVector<T, 3>& Dir, const T Length, T& OutTime, TVector<T, 3>& OutPosition, TVector<T, 3>& OutNormal, int32& OutFaceIndex, const T Thickness = 0, const bool bComputeMTD = false) const;
+		bool SweepGeom(const TImplicitObjectScaled<TSphere<T, 3>>& QueryGeom, const TRigidTransform<T, 3>& StartTM, const TVector<T, 3>& Dir, const T Length, T& OutTime, TVector<T, 3>& OutPosition, TVector<T, 3>& OutNormal, int32& OutFaceIndex, const T Thickness = 0, const bool bComputeMTD = false) const;
+		bool SweepGeom(const TImplicitObjectScaled<TBox<T, 3>>& QueryGeom, const TRigidTransform<T, 3>& StartTM, const TVector<T, 3>& Dir, const T Length, T& OutTime, TVector<T, 3>& OutPosition, TVector<T, 3>& OutNormal, int32& OutFaceIndex, const T Thickness = 0, const bool bComputeMTD = false) const;
+		bool SweepGeom(const TImplicitObjectScaled<TCapsule<T>>& QueryGeom, const TRigidTransform<T, 3>& StartTM, const TVector<T, 3>& Dir, const T Length, T& OutTime, TVector<T, 3>& OutPosition, TVector<T, 3>& OutNormal, int32& OutFaceIndex, const T Thickness = 0, const bool bComputeMTD = false) const;
+		bool SweepGeom(const TImplicitObjectScaled<TConvex<T, 3>>& QueryGeom, const TRigidTransform<T, 3>& StartTM, const TVector<T, 3>& Dir, const T Length, T& OutTime, TVector<T, 3>& OutPosition, TVector<T, 3>& OutNormal, int32& OutFaceIndex, const T Thickness = 0, const bool bComputeMTD = false) const;
 
 		virtual int32 FindMostOpposingFace(const TVector<T, 3>& Position, const TVector<T, 3>& UnitDir, int32 HintFaceIndex, T SearchDistance) const override;
 		virtual TVector<T, 3> FindGeometryOpposingNormal(const TVector<T, 3>& DenormDir, int32 FaceIndex, const TVector<T, 3>& OriginalNormal) const override;
@@ -185,7 +185,7 @@ namespace Chaos
 		bool OverlapGeomImp(const QueryGeomType& QueryGeom, const TRigidTransform<T, 3>& QueryTM, const T Thickness) const;
 
 		template <typename QueryGeomType>
-		bool SweepGeomImp(const QueryGeomType& QueryGeom, const TRigidTransform<T, 3>& StartTM, const TVector<T, 3>& Dir, const T Length, T& OutTime, TVector<T, 3>& OutPosition, TVector<T, 3>& OutNormal, int32& OutFaceIndex, const T Thickness) const;
+		bool SweepGeomImp(const QueryGeomType& QueryGeom, const TRigidTransform<T, 3>& StartTM, const TVector<T, 3>& Dir, const T Length, T& OutTime, TVector<T, 3>& OutPosition, TVector<T, 3>& OutNormal, int32& OutFaceIndex, const T Thickness, const bool bComputeMTD) const;
 	};
 }
 
diff --git a/Engine/Source/Runtime/PhysicsSQ/Private/SQAccelerator.cpp b/Engine/Source/Runtime/PhysicsSQ/Private/SQAccelerator.cpp
--- a/Engine/Source/Runtime/PhysicsSQ/Private/SQAccelerator.cpp
+++ b/Engine/Source/Runtime/PhysicsSQ/Private/SQAccelerator.cpp
@@ -282,7 +282,7 @@ private:
 
 						if (bBlocker)
 						{
-							CurData->SetLength(Distance);
+							CurData->SetLength(FMath::Max(0.f, Distance));	//Max is needed for MTD which returns negative distance
 							if (CurData->CurrentLength == 0 && (SQ == ESQType::Raycast || HitBuffer.WantsSingleResult()))	//raycasts always fail with distance 0, sweeps only matter if we want multi overlaps
 							{
 								return false;	//initial overlap so nothing will be better than this
