Convert GJKIntersection to use polytopes when possible. Means we can't support non uniform spheres and capsules, not sure if we care.
Also added unit test from previously failing EPA. The test just makes sure our ensures pass

diff --git a/Engine/Source/Runtime/Experimental/Chaos/Public/Chaos/GJK.h b/Engine/Source/Runtime/Experimental/Chaos/Public/Chaos/GJK.h
--- a/Engine/Source/Runtime/Experimental/Chaos/Public/Chaos/GJK.h
+++ b/Engine/Source/Runtime/Experimental/Chaos/Public/Chaos/GJK.h
@@ -20,10 +20,14 @@ namespace Chaos
 	 @return True if the geometries overlap, False otherwise */
 
 	template <typename T, typename TGeometryA, typename TGeometryB>
-	bool GJKIntersection(const TGeometryA& A, const TGeometryB& B, const TRigidTransform<T, 3>& BToATM, const T ThicknessA = 0, const TVector<T, 3>& InitialDir = TVector<T, 3>(1, 0, 0), const T ThicknessB = 0)
+	bool GJKIntersection(const TGeometryA& A, const TGeometryB& B, const TRigidTransform<T, 3>& BToATM, const T InThicknessA = 0, const TVector<T, 3>& InitialDir = TVector<T, 3>(1, 0, 0), const T InThicknessB = 0)
 	{
-		check(A.IsConvex() && B.IsConvex());
 		TVector<T, 3> V = -InitialDir;
+		if (V.SafeNormalize() == 0)
+		{
+			V = TVec3<T>(-1, 0, 0);
+		}
+
 		FSimplex SimplexIDs;
 		TVector<T, 3> Simplex[4];
 		T Barycentric[4] = { -1,-1,-1,-1 };	//not needed, but compiler warns
@@ -32,6 +36,10 @@ namespace Chaos
 		bool bNearZero = false;
 		int NumIterations = 0;
 		T PrevDist2 = FLT_MAX;
+		const T ThicknessA = A.GetMargin() + InThicknessA;
+		const T ThicknessB = B.GetMargin() + InThicknessB;
+		const T Inflation = ThicknessA + ThicknessB + 1e-3;
+		const T Inflation2 = Inflation * Inflation;
 		do
 		{
 			if (!ensure(NumIterations++ < 32))	//todo: take this out
@@ -39,13 +47,13 @@ namespace Chaos
 				break;	//if taking too long just stop. This should never happen
 			}
 			const TVector<T, 3> NegV = -V;
-			const TVector<T, 3> SupportA = A.Support(NegV, ThicknessA);	//todo: add thickness to quadratic geometry to avoid quadratic vs quadratic when possible
+			const TVector<T, 3> SupportA = A.Support2(NegV);
 			const TVector<T, 3> VInB = AToBRotation * V;
-			const TVector<T, 3> SupportBLocal = B.Support(VInB, ThicknessB);
+			const TVector<T, 3> SupportBLocal = B.Support2(VInB);
 			const TVector<T, 3> SupportB = BToATM.TransformPositionNoScale(SupportBLocal);
 			const TVector<T, 3> W = SupportA - SupportB;
 
-			if (TVector<T, 3>::DotProduct(V, W) > 0)
+			if (TVector<T, 3>::DotProduct(V, W) > Inflation)
 			{
 				return false;
 			}
@@ -56,7 +64,7 @@ namespace Chaos
 			V = SimplexFindClosestToOrigin(Simplex, SimplexIDs, Barycentric);
 
 			T NewDist2 = V.SizeSquared();
-			bNearZero = NewDist2 < 1e-6;
+			bNearZero = NewDist2 < Inflation2;
 
 			//as simplices become degenerate we will stop making progress. This is a side-effect of precision, in that case take V as the current best approximation
 			//question: should we take previous v in case it's better?
@@ -65,6 +73,11 @@ namespace Chaos
 
 			PrevDist2 = NewDist2;
 
+			if (!bTerminate)
+			{
+				V /= FMath::Sqrt(NewDist2);
+			}
+
 		} while (!bTerminate);
 
 		return bNearZero;
diff --git a/Engine/Source/Runtime/Experimental/Chaos/Public/Chaos/ImplicitObjectScaled.h b/Engine/Source/Runtime/Experimental/Chaos/Public/Chaos/ImplicitObjectScaled.h
--- a/Engine/Source/Runtime/Experimental/Chaos/Public/Chaos/ImplicitObjectScaled.h
+++ b/Engine/Source/Runtime/Experimental/Chaos/Public/Chaos/ImplicitObjectScaled.h
@@ -414,7 +414,16 @@ public:
 		return MObject->Support2(Direction * MScale) * MScale;
 	}
 
-	FORCEINLINE T GetMargin() const { return MObject->GetMargin(); }
+	FORCEINLINE T GetMargin() const
+	{
+		if (T UnscaledMargin = MObject->GetMargin())
+		{
+			ensure(MScale[0] == MScale[1] && MScale[1] == MScale[2]);
+			return UnscaledMargin * MScale[0];
+		}
+
+		return 0;
+	}
 
 	const TVector<T, d>& GetScale() const { return MScale; }
 	void SetScale(const TVector<T, d>& Scale)
