Fix clang

diff --git a/Engine/Source/Runtime/Experimental/Chaos/Public/Chaos/ISpatialAcceleration.h b/Engine/Source/Runtime/Experimental/Chaos/Public/Chaos/ISpatialAcceleration.h
--- a/Engine/Source/Runtime/Experimental/Chaos/Public/Chaos/ISpatialAcceleration.h
+++ b/Engine/Source/Runtime/Experimental/Chaos/Public/Chaos/ISpatialAcceleration.h
@@ -25,8 +25,10 @@ struct CHAOS_API FQueryFastData
 	FReal CurrentLength;
 	FReal InvCurrentLength;
 
+#if PLATFORM_WINDOWS
 	#pragma warning (push)
 	#pragma warning(disable:4723)
+#endif
 	//compiler complaining about divide by 0, but we are guarding against it.
 	//Seems like it's trying to evaluate div independent of the guard?
 
@@ -40,7 +42,9 @@ struct CHAOS_API FQueryFastData
 		}
 	}
 
+#if PLATFORM_WINDOWS
 	#pragma warning(pop)
+#endif
 
 
 protected:
