Fix for assert fail in IslandParticleArrayIndex PBDConstrintGraph.cpp. Island held inside Particle not up to date with graph.

diff --git a/Engine/Source/Runtime/Experimental/Chaos/Private/Chaos/PBDConstraintGraph.cpp b/Engine/Source/Runtime/Experimental/Chaos/Private/Chaos/PBDConstraintGraph.cpp
--- a/Engine/Source/Runtime/Experimental/Chaos/Private/Chaos/PBDConstraintGraph.cpp
+++ b/Engine/Source/Runtime/Experimental/Chaos/Private/Chaos/PBDConstraintGraph.cpp
@@ -131,6 +131,11 @@ void FPBDConstraintGraph::InitializeGraph(const TParticleView<TGeometryParticles
 				{
 					Nodes[UpdatedNodes[Index]].Island = INDEX_NONE;
 					Nodes[UpdatedNodes[Index]].Edges.Empty();
+					auto* Particle = Nodes[UpdatedNodes[Index]].Particle;
+					if (TPBDRigidParticleHandle<FReal, 3>* PBDRigid = Particle->AsDynamic())
+					{
+						PBDRigid->Island() = INDEX_NONE;
+					}
 				}
 			});
 		UpdatedNodes.Empty();
@@ -682,8 +687,8 @@ void FPBDConstraintGraph::DisableParticle(TGeometryParticleHandle<FReal, 3>* Par
 			}
 		}
 
-		ParticleRemove(Particle);
 	}
+	ParticleRemove(Particle);
 }
 
 
diff --git a/Engine/Source/Runtime/Experimental/Chaos/Public/Chaos/PBDConstraintGraph.h b/Engine/Source/Runtime/Experimental/Chaos/Public/Chaos/PBDConstraintGraph.h
--- a/Engine/Source/Runtime/Experimental/Chaos/Public/Chaos/PBDConstraintGraph.h
+++ b/Engine/Source/Runtime/Experimental/Chaos/Public/Chaos/PBDConstraintGraph.h
@@ -165,8 +165,6 @@ namespace Chaos
 		//Remove particle from constraint, maybe rethink some of these names
 		void RemoveParticle(TGeometryParticleHandle<FReal, 3>* Particle)
 		{
-			ParticleRemove(Particle);
-
 			//we know it just removes, probably rename later
 			DisableParticle(Particle);
 		}
