Fix EPA crashing when getting degenerate triangles

diff --git a/Engine/Source/Runtime/Experimental/Chaos/Public/Chaos/EPA.h b/Engine/Source/Runtime/Experimental/Chaos/Public/Chaos/EPA.h
--- a/Engine/Source/Runtime/Experimental/Chaos/Public/Chaos/EPA.h
+++ b/Engine/Source/Runtime/Experimental/Chaos/Public/Chaos/EPA.h
@@ -293,6 +293,7 @@ void EPAComputeVisibilityBorder(TArray<TEPAEntry<T>>& Entries, int32 EntryIdx, c
 template <typename T>
 void ComputeEPAResults(const TVec3<T>* VertsA, const TVec3<T>* VertsB, const TEPAEntry<T>& Entry, T& OutPenetration, TVec3<T>& OutDir, TVec3<T>& OutA, TVec3<T>& OutB)
 {
+	//NOTE: We use this function as fallback when robustness breaks. So - do not assume adjacency is valid as these may be new uninitialized traingles that failed
 	FSimplex SimplexIDs({ 0,1,2 });
 	TVec3<T> As[4] = { VertsA[Entry.IdxBuffer[0]], VertsA[Entry.IdxBuffer[1]], VertsA[Entry.IdxBuffer[2]] };
 	TVec3<T> Bs[4] = { VertsB[Entry.IdxBuffer[0]], VertsB[Entry.IdxBuffer[1]], VertsB[Entry.IdxBuffer[2]] };
@@ -328,6 +329,10 @@ enum class EPAResult
 	BadInitialSimplex
 };
 
+#ifndef DEBUG_EPA
+#define DEBUG_EPA 0
+#endif
+
 template <typename T, typename SupportALambda, typename SupportBLambda>
 EPAResult EPA(TArray<TVec3<T>>& VertsABuffer, TArray<TVec3<T>>& VertsBBuffer, const SupportALambda& SupportA, const SupportBLambda& SupportB, T& OutPenetration, TVec3<T>& OutDir, TVec3<T>& WitnessA, TVec3<T>& WitnessB)
 {
@@ -345,6 +350,8 @@ EPAResult EPA(TArray<TVec3<T>>& VertsABuffer, TArray<TVec3<T>>& VertsBBuffer, co
 	T UpperBound = TNumericLimits<T>::Max();
 	T LowerBound = TNumericLimits<T>::Lowest();
 
+	constexpr T Eps = 1e-2;
+
 	TArray<TEPAEntry<T>> Entries = InitializeEPA(VertsABuffer, VertsBBuffer, SupportA, SupportB);
 
 	if (Entries.Num() < 4)
@@ -357,17 +364,25 @@ EPAResult EPA(TArray<TVec3<T>>& VertsABuffer, TArray<TVec3<T>>& VertsBBuffer, co
 		return EPAResult::BadInitialSimplex;
 	}
 
+
+#if DEBUG_EPA
+	TArray<TVec3<T>> VertsWBuffer;
+	for (int32 Idx = 0; Idx < 4; ++Idx)
+	{
+		VertsWBuffer.Add(MinkowskiVert(VertsABuffer.GetData(), VertsBBuffer.GetData(), Idx));
+	}
+#endif
+	
 	std::priority_queue<FEPAEntryWrapper, std::vector<FEPAEntryWrapper>, std::greater<FEPAEntryWrapper>> Queue;
 	for(int32 Idx = 0; Idx < Entries.Num(); ++Idx)
 	{
+		//ensure(Entries[Idx].Distance > -Eps);
 		if(Entries[Idx].IsOriginProjectedInside(VertsABuffer.GetData(), VertsBBuffer.GetData()))
 		{
 			Queue.push(FEPAEntryWrapper {&Entries, Idx});
 		}
 	}
 
-	bool bCloseEnough = false;
-	bool bTerminate = bCloseEnough;
 	
 	//TEPAEntry<T> BestEntry;
 	//BestEntry.Distance = 0;
@@ -382,7 +397,7 @@ EPAResult EPA(TArray<TVec3<T>>& VertsABuffer, TArray<TVec3<T>>& VertsBBuffer, co
 		int32 EntryIdx = Queue.top().Idx;
 		Queue.pop();
 		TEPAEntry<T>& Entry = Entries[EntryIdx];
-		bool bBadFace = Entry.IsOriginProjectedInside(VertsABuffer.GetData(), VertsBBuffer.GetData());
+		//bool bBadFace = Entry.IsOriginProjectedInside(VertsABuffer.GetData(), VertsBBuffer.GetData());
 		{
 			//UE_LOG(LogChaos, Warning, TEXT("%d BestW:%f, Distance:%f, bObsolete:%d, InTriangle:%d"),
 			//	Iteration, BestW.Size(), Entry.Distance, Entry.bObsolete, bBadFace);
@@ -394,13 +409,12 @@ EPAResult EPA(TArray<TVec3<T>>& VertsABuffer, TArray<TVec3<T>>& VertsBBuffer, co
 
 		LastEntry = Entry;
 
-		constexpr T Eps = 1e-2;
 
 		const TVec3<T> ASupport = SupportA(Entry.PlaneNormal);
 		const TVec3<T> BSupport = SupportB(-Entry.PlaneNormal);
 		const TVec3<T> W = ASupport - BSupport;
 		const T DistanceToSupportPlane = TVec3<T>::DotProduct(Entry.PlaneNormal, W);
-
+		//ensure(DistanceToSupportPlane > -Eps);
 		if(DistanceToSupportPlane < UpperBound)
 		{
 			UpperBound = DistanceToSupportPlane;
@@ -411,7 +425,7 @@ EPAResult EPA(TArray<TVec3<T>>& VertsABuffer, TArray<TVec3<T>>& VertsBBuffer, co
 
 		LowerBound = Entry.Distance;
 
-		if (FMath::Abs(UpperBound - LowerBound) < Eps)	//todo: add relative error?
+		if (FMath::Abs(UpperBound - LowerBound) < Eps * LowerBound)
 		{
 			//UE_LOG(LogChaos, Warning, TEXT("Iteration:%d"), Iteration);
 			ComputeEPAResults(VertsABuffer.GetData(), VertsBBuffer.GetData(), Entry, OutPenetration, OutDir, WitnessA, WitnessB);
@@ -421,6 +435,10 @@ EPAResult EPA(TArray<TVec3<T>>& VertsABuffer, TArray<TVec3<T>>& VertsBBuffer, co
 		const int32 NewVertIdx = VertsABuffer.Add(ASupport);
 		VertsBBuffer.Add(BSupport);
 
+#if DEBUG_EPA
+		VertsWBuffer.Add(MinkowskiVert(VertsABuffer.GetData(), VertsBBuffer.GetData(), NewVertIdx));
+#endif
+
 		Entry.bObsolete = true;
 		VisibilityBorder.Reset();
 		EPAComputeVisibilityBorder(Entries, EntryIdx, W, VisibilityBorder);
@@ -428,35 +446,55 @@ EPAResult EPA(TArray<TVec3<T>>& VertsABuffer, TArray<TVec3<T>>& VertsBBuffer, co
 		const int32 FirstIdxInBatch = Entries.Num();
 		int32 NewIdx = FirstIdxInBatch;
 		Entries.AddUninitialized(NumBorderEdges);
-		ensure(NumBorderEdges >= 3);
-		for (int32 VisibilityIdx = 0; VisibilityIdx < NumBorderEdges; ++VisibilityIdx)
+		if (NumBorderEdges >= 3)
 		{
-			//create new entries and update adjacencies
-			const FEPAFloodEntry& BorderInfo = VisibilityBorder[VisibilityIdx];
-			TEPAEntry<T>& NewEntry = Entries[NewIdx];
-			const int32 BorderEntryIdx = BorderInfo.EntryIdx;
-			TEPAEntry<T>& BorderEntry = Entries[BorderEntryIdx];
-			const int32 BorderEdgeIdx0 = BorderInfo.EdgeIdx;
-			const int32 BorderEdgeIdx1 = (BorderEdgeIdx0 + 1) % 3;
-			const int32 NextEntryIdx = (VisibilityIdx + 1) < VisibilityBorder.Num() ? NewIdx + 1 : FirstIdxInBatch;
-			const int32 PrevEntryIdx = NewIdx > FirstIdxInBatch ? NewIdx - 1 : FirstIdxInBatch + NumBorderEdges - 1;
-			const bool bValidTri = NewEntry.Initialize(VertsABuffer.GetData(), VertsBBuffer.GetData(), BorderEntry.IdxBuffer[BorderEdgeIdx1], BorderEntry.IdxBuffer[BorderEdgeIdx0], NewVertIdx,
-				{ BorderEntryIdx, PrevEntryIdx, NextEntryIdx },
-				{ BorderEdgeIdx0, 2, 1 });
-			BorderEntry.AdjFaces[BorderEdgeIdx0] = NewIdx;
-			BorderEntry.AdjEdges[BorderEdgeIdx0] = 0;
-
-			//We should never need to check the lower bound, but in the case of bad precision this can happen
-			//We simply ignore this direction as it likely has even more bad precision
-			if (bValidTri && NewEntry.Distance >= LowerBound && NewEntry.Distance <= UpperBound)
+			bool bTerminate = false;
+			for (int32 VisibilityIdx = 0; VisibilityIdx < NumBorderEdges; ++VisibilityIdx)
 			{
-				if(NewEntry.IsOriginProjectedInside(VertsABuffer.GetData(), VertsBBuffer.GetData()))
+				//create new entries and update adjacencies
+				const FEPAFloodEntry& BorderInfo = VisibilityBorder[VisibilityIdx];
+				TEPAEntry<T>& NewEntry = Entries[NewIdx];
+				const int32 BorderEntryIdx = BorderInfo.EntryIdx;
+				TEPAEntry<T>& BorderEntry = Entries[BorderEntryIdx];
+				const int32 BorderEdgeIdx0 = BorderInfo.EdgeIdx;
+				const int32 BorderEdgeIdx1 = (BorderEdgeIdx0 + 1) % 3;
+				const int32 NextEntryIdx = (VisibilityIdx + 1) < VisibilityBorder.Num() ? NewIdx + 1 : FirstIdxInBatch;
+				const int32 PrevEntryIdx = NewIdx > FirstIdxInBatch ? NewIdx - 1 : FirstIdxInBatch + NumBorderEdges - 1;
+				const bool bValidTri = NewEntry.Initialize(VertsABuffer.GetData(), VertsBBuffer.GetData(), BorderEntry.IdxBuffer[BorderEdgeIdx1], BorderEntry.IdxBuffer[BorderEdgeIdx0], NewVertIdx,
+					{ BorderEntryIdx, PrevEntryIdx, NextEntryIdx },
+					{ BorderEdgeIdx0, 2, 1 });
+				BorderEntry.AdjFaces[BorderEdgeIdx0] = NewIdx;
+				BorderEntry.AdjEdges[BorderEdgeIdx0] = 0;
+				
+				if (!bValidTri)
+				{
+					//couldn't properly expand polytope, so just stop
+					bTerminate = true;
+					break;
+				}
+
+				//We should never need to check the lower bound, but in the case of bad precision this can happen
+				//We simply ignore this direction as it likely has even more bad precision
+				if (bValidTri && NewEntry.Distance >= LowerBound && NewEntry.Distance <= UpperBound)
 				{
-					Queue.push(FEPAEntryWrapper{ &Entries, NewIdx});
+					if (NewEntry.IsOriginProjectedInside(VertsABuffer.GetData(), VertsBBuffer.GetData()))
+					{
+						Queue.push(FEPAEntryWrapper{ &Entries, NewIdx });
+					}
 				}
+
+				++NewIdx;
 			}
 
-			++NewIdx;
+			if (bTerminate)
+			{
+				break;
+			}
+		}
+		else
+		{
+			//couldn't properly expand polytope, just stop now
+			break;
 		}
 	}
 
