Use absolute value of margin

diff --git a/Engine/Source/Runtime/Experimental/Chaos/Public/Chaos/ImplicitObjectScaled.h b/Engine/Source/Runtime/Experimental/Chaos/Public/Chaos/ImplicitObjectScaled.h
--- a/Engine/Source/Runtime/Experimental/Chaos/Public/Chaos/ImplicitObjectScaled.h
+++ b/Engine/Source/Runtime/Experimental/Chaos/Public/Chaos/ImplicitObjectScaled.h
@@ -419,7 +419,7 @@ public:
 		if (T UnscaledMargin = MObject->GetMargin())
 		{
 			ensure(MScale[0] == MScale[1] && MScale[1] == MScale[2]);
-			return UnscaledMargin * MScale[0];
+			return UnscaledMargin * FMath::Abs(MScale[0]);
 		}
 
 		return 0;
