Add some ensures around handle array access to get more info from bad handles

diff --git a/Engine/Source/Runtime/Experimental/Chaos/Public/Chaos/Framework/Handles.h b/Engine/Source/Runtime/Experimental/Chaos/Public/Chaos/Framework/Handles.h
--- a/Engine/Source/Runtime/Experimental/Chaos/Public/Chaos/Framework/Handles.h
+++ b/Engine/Source/Runtime/Experimental/Chaos/Public/Chaos/Framework/Handles.h
@@ -226,20 +226,40 @@ namespace Chaos
 
 		ElementType* Get(FHandle InHandle) const
 		{
-			if(InHandle.IsValid() && Validity[InHandle.Index])
+			if(InHandle.IsValid())
 			{
-				const FHandleEntry& Entry = HandleEntries[InHandle.Index];
-				return Entry.Generation == InHandle.Generation ? &(Data[InHandle.Index]) : nullptr;
+				if(Validity.IsValidIndex(InHandle.Index))
+				{
+					if(Validity[InHandle.Index])
+					{
+						const FHandleEntry& Entry = HandleEntries[InHandle.Index];
+						return Entry.Generation == InHandle.Generation ? &(Data[InHandle.Index]) : nullptr;
+					}
+				}
+				else
+				{
+					ensureMsgf(false, TEXT("Failed to access handle (%u, %u). NumEntries = %d, NumFlags = %d"), InHandle.Index, InHandle.Generation, HandleEntries.Num(), Validity.Num());
+				}
 			}
 			return nullptr;
 		}
 
 		const ElementType* Get(FConstHandle InHandle) const
 		{
-			if(InHandle.IsValid() && Validity[InHandle.Index])
+			if (InHandle.IsValid())
 			{
-				const FHandleEntry& Entry = HandleEntries[InHandle.Index];
-				return Entry.Generation == InHandle.Generation ? &(Data[InHandle.Index]) : nullptr;
+				if (Validity.IsValidIndex(InHandle.Index))
+				{
+					if (Validity[InHandle.Index])
+					{
+						const FHandleEntry& Entry = HandleEntries[InHandle.Index];
+						return Entry.Generation == InHandle.Generation ? &(Data[InHandle.Index]) : nullptr;
+					}
+				}
+				else
+				{
+					ensureMsgf(false, TEXT("Failed to access handle (%u, %u). NumEntries = %d, NumFlags = %d"), InHandle.Index, InHandle.Generation, HandleEntries.Num(), Validity.Num());
+				}
 			}
 			return nullptr;
 		}
