Fix non uniform triangle mesh forcing capsule and sphere to become non uniform in gjk

diff --git a/Engine/Source/Runtime/Experimental/Chaos/Private/Chaos/TriangleMeshImplicitObject.cpp b/Engine/Source/Runtime/Experimental/Chaos/Private/Chaos/TriangleMeshImplicitObject.cpp
--- a/Engine/Source/Runtime/Experimental/Chaos/Private/Chaos/TriangleMeshImplicitObject.cpp
+++ b/Engine/Source/Runtime/Experimental/Chaos/Private/Chaos/TriangleMeshImplicitObject.cpp
@@ -265,22 +265,29 @@ bool TTriangleMeshImplicitObject<T>::Overlap(const TVector<T, 3>& Point, const T
 	return false;
 }
 
+template <typename QueryGeomType, typename T>
+const QueryGeomType& GetVertsAndQueryGeomHelper(const QueryGeomType& QueryGeom, const int32 TriIdx, const TParticles<T, 3>& Particles, const TArray<TVector<int32, 3>>& Elements, TVec3<T>& OutA, TVec3<T>& OutB, TVec3<T>& OutC)
+{
+	OutA = Particles.X(Elements[TriIdx][0]);
+	OutB = Particles.X(Elements[TriIdx][1]);
+	OutC = Particles.X(Elements[TriIdx][2]);
+	return QueryGeom;
+}
+
+template <typename QueryGeomType, typename T>
+const QueryGeomType& GetVertsAndQueryGeomHelper(const TImplicitObjectScaled<QueryGeomType>& QueryGeom, const int32 TriIdx, const TParticles<T, 3>& Particles, const TArray<TVector<int32, 3>>& Elements, TVec3<T>& OutA, TVec3<T>& OutB, TVec3<T>& OutC)
+{
+	const TVec3<T> InvScale = QueryGeom.GetInvScale();
+	OutA = Particles.X(Elements[TriIdx][0]) * InvScale;
+	OutB = Particles.X(Elements[TriIdx][1]) * InvScale;
+	OutC = Particles.X(Elements[TriIdx][2]) * InvScale;
+	return *QueryGeom.GetUnscaledObject();
+}
+
 template <typename T>
 template <typename QueryGeomType>
 bool TTriangleMeshImplicitObject<T>::OverlapGeomImp(const QueryGeomType& QueryGeom, const TRigidTransform<T, 3>& QueryTM, const T Thickness) const
 {
-	auto OverlapTriangle = [&](const TVector<T, 3>& A, const TVector<T, 3>& B, const TVector<T, 3>& C) -> bool
-	{
-		const TVector<T, 3> AB = B - A;
-		const TVector<T, 3> AC = C - A;
-
-		//It's most likely that the query object is in front of the triangle since queries tend to be on the outside.
-		//However, maybe we should check if it's behind the triangle plane. Also, we should enforce this winding in some way
-		const TVector<T, 3> Offset = TVector<T, 3>::CrossProduct(AB, AC);
-
-		return GJKIntersection(TTriangle<T>(A, B, C), QueryGeom, QueryTM, Thickness, Offset);
-	};
-
 	bool bResult = false;
 	TBox<T, 3> QueryBounds = QueryGeom.BoundingBox();
 	QueryBounds.Thicken(Thickness);
@@ -290,11 +297,17 @@ bool TTriangleMeshImplicitObject<T>::OverlapGeomImp(const QueryGeomType& QueryGe
 
 	for (int32 TriIdx : PotentialIntersections)
 	{
-	const TVector<T, 3>& A = MParticles.X(MElements[TriIdx][0]);
-	const TVector<T, 3>& B = MParticles.X(MElements[TriIdx][1]);
-	const TVector<T, 3>& C = MParticles.X(MElements[TriIdx][2]);
+		TVec3<T> A, B, C;
+		const auto & InnerQueryGeom = GetVertsAndQueryGeomHelper(QueryGeom, TriIdx, MParticles, MElements, A, B, C);
 
-		if (OverlapTriangle(A, B, C))
+		const TVector<T, 3> AB = B - A;
+		const TVector<T, 3> AC = C - A;
+
+		//It's most likely that the query object is in front of the triangle since queries tend to be on the outside.
+		//However, maybe we should check if it's behind the triangle plane. Also, we should enforce this winding in some way
+		const TVector<T, 3> Offset = TVector<T, 3>::CrossProduct(AB, AC);
+
+		if (GJKIntersection(TTriangle<T>(A, B, C), InnerQueryGeom, QueryTM, Thickness, Offset))
 		{
 			return true;
 		}
@@ -393,11 +406,12 @@ struct TTriangleMeshSweepVisitor
 		TVector<T, 3> HitPosition;
 		TVector<T, 3> HitNormal;
 
-		TTriangle<T> Tri(TriMesh.MParticles.X(TriMesh.MElements[TriIdx][0]),
-			TriMesh.MParticles.X(TriMesh.MElements[TriIdx][1]),
-			TriMesh.MParticles.X(TriMesh.MElements[TriIdx][2]));
+		TVec3<T> A, B, C;
+		const auto & InnerQueryGeom = GetVertsAndQueryGeomHelper(QueryGeom, TriIdx, TriMesh.MParticles, TriMesh.MElements, A, B, C);
+
+		TTriangle<T> Tri(A, B, C);
 
-		if(GJKRaycast2<T>(Tri, QueryGeom, StartTM, Dir, CurData.CurrentLength, Time, HitPosition, HitNormal, Thickness, bComputeMTD))
+		if(GJKRaycast2<T>(Tri, InnerQueryGeom, StartTM, Dir, CurData.CurrentLength, Time, HitPosition, HitNormal, Thickness, bComputeMTD))
 		{
 			if(Time < OutTime)
 			{
diff --git a/Engine/Source/Runtime/Experimental/Chaos/Public/Chaos/ImplicitObjectScaled.h b/Engine/Source/Runtime/Experimental/Chaos/Public/Chaos/ImplicitObjectScaled.h
--- a/Engine/Source/Runtime/Experimental/Chaos/Public/Chaos/ImplicitObjectScaled.h
+++ b/Engine/Source/Runtime/Experimental/Chaos/Public/Chaos/ImplicitObjectScaled.h
@@ -426,6 +426,7 @@ public:
 	}
 
 	const TVector<T, d>& GetScale() const { return MScale; }
+	const TVector<T, d>& GetInvScale() const { return MInvScale; }
 	void SetScale(const TVector<T, d>& Scale)
 	{
 		constexpr T MinMagnitude = 1e-4;
