Avoid ensures being inlined in exported function
#rnx

diff --git a/Engine/Source/Runtime/Experimental/Chaos/Private/Chaos/PBDRigidParticles.cpp b/Engine/Source/Runtime/Experimental/Chaos/Private/Chaos/PBDRigidParticles.cpp
--- a/Engine/Source/Runtime/Experimental/Chaos/Private/Chaos/PBDRigidParticles.cpp
+++ b/Engine/Source/Runtime/Experimental/Chaos/Private/Chaos/PBDRigidParticles.cpp
@@ -2,4 +2,10 @@
 
 #include "Chaos/PBDRigidParticles.h"
 
-template class Chaos::TPBDRigidParticles<float,3>;
+void Chaos::EnsureSleepingObjectState(EObjectStateType ObjectState)
+{
+	ensure(ObjectState != EObjectStateType::Kinematic);
+	ensure(ObjectState != EObjectStateType::Static);
+}
+
+template class Chaos::TPBDRigidParticles<float, 3>;
diff --git a/Engine/Source/Runtime/Experimental/Chaos/Public/Chaos/PBDRigidParticles.h b/Engine/Source/Runtime/Experimental/Chaos/Public/Chaos/PBDRigidParticles.h
--- a/Engine/Source/Runtime/Experimental/Chaos/Public/Chaos/PBDRigidParticles.h
+++ b/Engine/Source/Runtime/Experimental/Chaos/Public/Chaos/PBDRigidParticles.h
@@ -15,6 +15,8 @@ namespace Chaos
 template<class T, int d>
 class TPBDRigidsEvolution;
 
+CHAOS_API void EnsureSleepingObjectState(EObjectStateType ObjectState);
+
 template<class T, int d>
 class TPBDRigidParticles : public TRigidParticles<T, d>
 {
@@ -90,8 +92,7 @@ class TPBDRigidParticles : public TRigidParticles<T, d>
 
 		if (bSleeping)
 		{
-			ensure(this->ObjectState(Index) != EObjectStateType::Kinematic);
-			ensure(this->ObjectState(Index) != EObjectStateType::Static);
+			EnsureSleepingObjectState(this->ObjectState(Index));
 		}
 	}
 
