Event timings reduced

diff --git a/Engine/Source/Runtime/Experimental/ChaosSolvers/Private/EventDefaults.cpp b/Engine/Source/Runtime/Experimental/ChaosSolvers/Private/EventDefaults.cpp
--- a/Engine/Source/Runtime/Experimental/ChaosSolvers/Private/EventDefaults.cpp
+++ b/Engine/Source/Runtime/Experimental/ChaosSolvers/Private/EventDefaults.cpp
@@ -187,8 +187,8 @@ namespace Chaos
 			SCOPE_CYCLE_COUNTER(STAT_GatherBreakingEvent);
 
 			// #todo: This isn't working - SolverActor parameters are set on a solver but it is currently a different solver that is simulating!!
-			//if (!Solver->GetEventFilters()->IsBreakingEventEnabled())
-			//	return;
+			if (!Solver->GetEventFilters()->IsBreakingEventEnabled())
+				return;
 
 			FBreakingDataArray& AllBreakingDataArray = BreakingEventData.BreakingData.AllBreakingsArray;
 			TMap<IPhysicsProxyBase*, TArray<int32>>& AllBreakingIndicesByPhysicsProxy = BreakingEventData.PhysicsProxyToBreakingIndices.PhysicsProxyToIndicesMap;
@@ -267,8 +267,8 @@ namespace Chaos
 			check(Solver);
 
 			// #todo: This isn't working - SolverActor parameters are set on a solver but it is currently a different solver that is simulating!!
-			//if (!Solver->GetEventFilters()->IsTrailingEventEnabled())
-			//	return;
+			if (!Solver->GetEventFilters()->IsTrailingEventEnabled())
+				return;
 
 			const FPBDRigidsSolver::FPBDRigidsEvolution* Evolution = Solver->GetEvolution();
 
