Fix bounds check for EPA

diff --git a/Engine/Source/Runtime/Experimental/Chaos/Public/Chaos/EPA.h b/Engine/Source/Runtime/Experimental/Chaos/Public/Chaos/EPA.h
--- a/Engine/Source/Runtime/Experimental/Chaos/Public/Chaos/EPA.h
+++ b/Engine/Source/Runtime/Experimental/Chaos/Public/Chaos/EPA.h
@@ -425,7 +425,7 @@ EPAResult EPA(TArray<TVec3<T>>& VertsABuffer, TArray<TVec3<T>>& VertsBBuffer, co
 
 		LowerBound = Entry.Distance;
 
-		if (FMath::Abs(UpperBound - LowerBound) < Eps * LowerBound)
+		if (FMath::Abs(UpperBound - LowerBound) <= Eps * LowerBound)
 		{
 			//UE_LOG(LogChaos, Warning, TEXT("Iteration:%d"), Iteration);
 			ComputeEPAResults(VertsABuffer.GetData(), VertsBBuffer.GetData(), Entry, OutPenetration, OutDir, WitnessA, WitnessB);
