Added GJKPenetration. Good enough for now, should refactor and merge some of these later

diff --git a/Engine/Source/Runtime/Experimental/Chaos/Public/Chaos/GJK.h b/Engine/Source/Runtime/Experimental/Chaos/Public/Chaos/GJK.h
--- a/Engine/Source/Runtime/Experimental/Chaos/Public/Chaos/GJK.h
+++ b/Engine/Source/Runtime/Experimental/Chaos/Public/Chaos/GJK.h
@@ -83,6 +83,154 @@ namespace Chaos
 		return bNearZero;
 	}
 
+	template <typename T, typename TGeometryA, typename TGeometryB>
+	bool GJKPenetration(const TGeometryA& A, const TGeometryB& B, const TRigidTransform<T, 3>& BToATM, T& OutPenetration, TVec3<T>& OutClosestA, TVec3<T>& OutClosestB, TVec3<T>& OutNormal, const T InThicknessA = 0, const TVector<T, 3>& InitialDir = TVector<T, 3>(1, 0, 0), const T InThicknessB = 0)
+	{
+		auto SupportAFunc = [&A](const TVec3<T>& V)
+		{
+			return A.Support2(V);
+		};
+
+		const TRotation<T, 3> AToBRotation = BToATM.GetRotation().Inverse();
+
+
+		auto SupportBFunc = [&B, &BToATM, &AToBRotation](const TVec3<T>& V)
+		{
+			const TVector<T, 3> VInB = AToBRotation * V;
+			const TVector<T, 3> SupportBLocal = B.Support2(VInB);
+			return BToATM.TransformPositionNoScale(SupportBLocal);
+		};
+
+		//todo: refactor all of these similar functions
+		TVector<T, 3> V = -InitialDir;
+		if (V.SafeNormalize() == 0)
+		{
+			V = TVec3<T>(-1, 0, 0);
+		}
+
+		TVec3<T> As[4];
+		TVec3<T> Bs[4];
+
+		FSimplex SimplexIDs;
+		TVector<T, 3> Simplex[4];
+		T Barycentric[4] = { -1,-1,-1,-1 };	//not needed, but compiler warns
+		bool bTerminate;
+		bool bNearZero = false;
+		int NumIterations = 0;
+		T PrevDist2 = FLT_MAX;
+		const T ThicknessA = A.GetMargin() + InThicknessA;
+		const T ThicknessB = B.GetMargin() + InThicknessB;
+		const T Inflation = ThicknessA + ThicknessB + 1e-3;
+		const T Inflation2 = Inflation * Inflation;
+		const T Eps2 = 1e-6;
+		do
+		{
+			if (!ensure(NumIterations++ < 32))	//todo: take this out
+			{
+				break;	//if taking too long just stop. This should never happen
+			}
+			const TVector<T, 3> NegV = -V;
+			const TVector<T, 3> SupportA = SupportAFunc(NegV);
+			const TVector<T, 3> SupportB = SupportBFunc(V);
+			const TVector<T, 3> W = SupportA - SupportB;
+
+			if (TVector<T, 3>::DotProduct(V, W) > Inflation)
+			{
+				return false;
+			}
+
+			SimplexIDs[SimplexIDs.NumVerts] = SimplexIDs.NumVerts;
+			As[SimplexIDs.NumVerts] = SupportA;
+			Bs[SimplexIDs.NumVerts] = SupportB;
+			Simplex[SimplexIDs.NumVerts++] = W;
+
+			V = SimplexFindClosestToOrigin(Simplex, SimplexIDs, Barycentric, As, Bs);
+
+			T NewDist2 = V.SizeSquared();
+			bNearZero = NewDist2 < Eps2;	//want to get the closest point for MTD
+
+			//as simplices become degenerate we will stop making progress. This is a side-effect of precision, in that case take V as the current best approximation
+			//question: should we take previous v in case it's better?
+			const bool bMadeProgress = NewDist2 < PrevDist2;
+			bTerminate = bNearZero || !bMadeProgress;
+
+			PrevDist2 = NewDist2;
+
+			if (!bTerminate)
+			{
+				V /= FMath::Sqrt(NewDist2);
+			}
+
+		} while (!bTerminate);
+
+		if (PrevDist2 > Eps2)
+		{
+			//generally this happens when shapes are inflated.
+			TVector<T, 3> ClosestA(0);
+			TVector<T, 3> ClosestBInA(0);
+
+			for (int i = 0; i < SimplexIDs.NumVerts; ++i)
+			{
+				ClosestA += As[i] * Barycentric[i];
+				ClosestBInA += Bs[i] * Barycentric[i];
+			}
+
+			
+			const T PreDist = FMath::Sqrt(PrevDist2);
+			OutNormal = (ClosestBInA - ClosestA).GetUnsafeNormal();	//question: should we just use PreDist2?
+			const T Penetration = ThicknessA + ThicknessB - PreDist;
+			OutPenetration = Penetration;
+			OutClosestA = ClosestA + OutNormal * ThicknessA;
+			OutClosestB = ClosestBInA - OutNormal * ThicknessB;
+
+		}
+		else
+		{
+			TArray<TVec3<T>> VertsA;
+			TArray<TVec3<T>> VertsB;
+
+			VertsA.Reserve(8);
+			VertsB.Reserve(8);
+
+			for (int i = 0; i < SimplexIDs.NumVerts; ++i)
+			{
+				VertsA.Add(As[i]);
+				VertsB.Add(Bs[i]);
+			}
+
+			T Penetration;
+			TVec3<T> MTD, ClosestA, ClosestBInA;
+			if (EPA(VertsA, VertsB, SupportAFunc, SupportBFunc, Penetration, MTD, ClosestA, ClosestBInA) != EPAResult::BadInitialSimplex)
+			{
+				OutNormal = MTD;
+				OutPenetration = Penetration + Inflation;
+				OutClosestA = ClosestA + OutNormal * ThicknessA;
+				OutClosestB = ClosestBInA - OutNormal * ThicknessB;
+			}
+			else
+			{
+				//assume touching hit
+
+				ClosestA = TVec3<T>(0);
+				ClosestBInA = TVec3<T>(0);
+
+				for (int i = 0; i < SimplexIDs.NumVerts; ++i)
+				{
+					ClosestA += As[i] * Barycentric[i];
+					ClosestBInA += Bs[i] * Barycentric[i];
+				}
+
+				OutPenetration = Inflation;
+				OutNormal = { 0,0,1 };
+				OutClosestA = ClosestA + OutNormal * ThicknessA;
+				OutClosestB = ClosestBInA - OutNormal * ThicknessB;
+				return false;
+			}
+		}
+
+		return true;
+	}
+
 
 	/** Sweeps one geometry against the other
 	 @A The first geometry
