Chaos : Collisions

Removed public access to the indices of the collision constraint, iterating constraints should use the handles instead.

diff --git a/Engine/Source/Runtime/Engine/Private/PhysicsEngine/ImmediatePhysics/ImmediatePhysicsChaos/ImmediatePhysicsSimulation_Chaos.cpp b/Engine/Source/Runtime/Engine/Private/PhysicsEngine/ImmediatePhysics/ImmediatePhysicsChaos/ImmediatePhysicsSimulation_Chaos.cpp
--- a/Engine/Source/Runtime/Engine/Private/PhysicsEngine/ImmediatePhysics/ImmediatePhysicsChaos/ImmediatePhysicsSimulation_Chaos.cpp
+++ b/Engine/Source/Runtime/Engine/Private/PhysicsEngine/ImmediatePhysics/ImmediatePhysicsChaos/ImmediatePhysicsSimulation_Chaos.cpp
@@ -145,16 +145,17 @@ namespace ImmediatePhysics_Chaos
 			[this]()
 			{
 				Evolution->GetCollisionConstraints().ApplyCollisionModifier(
-					[this](TRigidBodyContactConstraint<float, 3>& Constraint)
+					[this](const TPBDCollisionConstraint<FReal, Dimensions>::FConstraintContainerHandle* Handle)
 					{
-						if (ShouldIgnoreCollisionConstraint(Constraint.Particle[0], Constraint.Particle[1], IgnoreCollisionParticlePairTable))
+						TVector<const TGeometryParticleHandle<FReal, Dimensions>*, 2> ConstrainedParticles = Handle->GetConstrainedParticles();
+						if (ShouldIgnoreCollisionConstraint(ConstrainedParticles[0], ConstrainedParticles[1], IgnoreCollisionParticlePairTable))
 						{
 							return ECollisionModifierResult::Disabled;
 						}
 						return ECollisionModifierResult::Unchanged;
 					});
 			});
-		
+
 #if CHAOS_DEBUG_DRAW
 		Evolution->SetPostIntegrateCallback(
 			[this]()
diff --git a/Engine/Source/Runtime/Experimental/Chaos/Private/Chaos/ChaosDebugDraw.cpp b/Engine/Source/Runtime/Experimental/Chaos/Private/Chaos/ChaosDebugDraw.cpp
--- a/Engine/Source/Runtime/Experimental/Chaos/Private/Chaos/ChaosDebugDraw.cpp
+++ b/Engine/Source/Runtime/Experimental/Chaos/Private/Chaos/ChaosDebugDraw.cpp
@@ -135,31 +135,34 @@ namespace Chaos
 
 		void DrawCollisionImpl(const FRigidTransform3& SpaceTransform, const TPBDCollisionConstraintHandle<float, 3>* ConstraintHandle, float ColorScale)
 		{
-			const TRigidBodyContactConstraint<FReal, 3>& Contact = ConstraintHandle->GetContact();
-			if (Contact.GetPhi() > 0)
+			if (ConstraintHandle->GetType() == TCollisionConstraintBase<float,3>::FType::SinglePoint)
 			{
-				ColorScale = ColorScale * (FReal)0.1;
-			}
+				const TRigidBodyContactConstraint<FReal, 3>& Contact = ConstraintHandle->GetContact< TRigidBodyContactConstraint<float, 3> >();
+				if (Contact.GetPhi() > 0)
+				{
+					ColorScale = ColorScale * (FReal)0.1;
+				}
 
-			FVec3 Location = SpaceTransform.TransformPosition(Contact.GetLocation());
-			FVec3 Normal = SpaceTransform.TransformVector(Contact.GetNormal());
+				FVec3 Location = SpaceTransform.TransformPosition(Contact.GetLocation());
+				FVec3 Normal = SpaceTransform.TransformVector(Contact.GetNormal());
 
-			if (ContactWidth > 0)
-			{
-				FColor C0 = (ColorScale * FColor(128, 0, 0)).ToFColor(false);
-				FMatrix Axes = FRotationMatrix::MakeFromX(Normal);
-				FDebugDrawQueue::GetInstance().DrawDebugCircle(Location, DrawScale * ContactWidth, 12, C0, false, KINDA_SMALL_NUMBER, DrawPriority, LineThickness, Axes.GetUnitAxis(EAxis::Y), Axes.GetUnitAxis(EAxis::Z), false);
-			}
-			if (ContactLen > 0)
-			{
-				FColor C1 = (ColorScale * FColor(255, 0, 0)).ToFColor(false);
-				FDebugDrawQueue::GetInstance().DrawDebugLine(Location, Location + DrawScale * ContactLen * Normal, C1, false, KINDA_SMALL_NUMBER, DrawPriority, LineThickness);
-			}
-			if (ContactPhiWidth > 0 && Contact.GetPhi() < FLT_MAX)
-			{
-				FColor C2 = (ColorScale * FColor(128, 128, 0)).ToFColor(false);
-				FMatrix Axes = FRotationMatrix::MakeFromX(Normal);
-				FDebugDrawQueue::GetInstance().DrawDebugCircle(Location - Contact.GetPhi() * Normal, DrawScale * ContactPhiWidth, 12, C2, false, KINDA_SMALL_NUMBER, DrawPriority, LineThickness, Axes.GetUnitAxis(EAxis::Y), Axes.GetUnitAxis(EAxis::Z), false);
+				if (ContactWidth > 0)
+				{
+					FColor C0 = (ColorScale * FColor(128, 0, 0)).ToFColor(false);
+					FMatrix Axes = FRotationMatrix::MakeFromX(Normal);
+					FDebugDrawQueue::GetInstance().DrawDebugCircle(Location, DrawScale * ContactWidth, 12, C0, false, KINDA_SMALL_NUMBER, DrawPriority, LineThickness, Axes.GetUnitAxis(EAxis::Y), Axes.GetUnitAxis(EAxis::Z), false);
+				}
+				if (ContactLen > 0)
+				{
+					FColor C1 = (ColorScale * FColor(255, 0, 0)).ToFColor(false);
+					FDebugDrawQueue::GetInstance().DrawDebugLine(Location, Location + DrawScale * ContactLen * Normal, C1, false, KINDA_SMALL_NUMBER, DrawPriority, LineThickness);
+				}
+				if (ContactPhiWidth > 0 && Contact.GetPhi() < FLT_MAX)
+				{
+					FColor C2 = (ColorScale * FColor(128, 128, 0)).ToFColor(false);
+					FMatrix Axes = FRotationMatrix::MakeFromX(Normal);
+					FDebugDrawQueue::GetInstance().DrawDebugCircle(Location - Contact.GetPhi() * Normal, DrawScale * ContactPhiWidth, 12, C2, false, KINDA_SMALL_NUMBER, DrawPriority, LineThickness, Axes.GetUnitAxis(EAxis::Y), Axes.GetUnitAxis(EAxis::Z), false);
+				}
 			}
 		}
 
@@ -211,10 +214,6 @@ namespace Chaos
 
 #endif
 
-		//
-		//
-		//
-
 		void DrawParticleShapes(const FRigidTransform3& SpaceTransform, const TParticleView<TGeometryParticles<float, 3>>& ParticlesView, float ColorScale, bool bDrawKinemtatic, bool bDrawDynamic)
 		{
 #if CHAOS_DEBUG_DRAW
@@ -284,12 +283,11 @@ namespace Chaos
 #if CHAOS_DEBUG_DRAW
 			if (FDebugDrawQueue::IsDebugDrawingEnabled())
 			{
-				for (int32 ConstraintIndex = 0; ConstraintIndex < Collisions.NumConstraints(); ++ConstraintIndex)
+				for (const Chaos::TPBDCollisionConstraintHandle<float, 3> * ConstraintHandle : Collisions.GetConstConstraintHandles())
 				{
-					TVector<TGeometryParticleHandle<float, 3>*, 2> ConstrainedParticles = Collisions.GetConstrainedParticles(ConstraintIndex);
+					TVector<const TGeometryParticleHandle<float, 3>*, 2> ConstrainedParticles = ConstraintHandle->GetConstrainedParticles();
 					if ((ConstrainedParticles[0] == Particle) || (ConstrainedParticles[1] == Particle))
 					{
-						const typename TPBDCollisionConstraint<float, 3>::FConstraintContainerHandle* ConstraintHandle = Collisions.GetConstraintHandle(ConstraintIndex);
 						DrawCollisionImpl(SpaceTransform, ConstraintHandle, 1.0f);
 					}
 				}
@@ -302,9 +300,9 @@ namespace Chaos
 #if CHAOS_DEBUG_DRAW
 			if (FDebugDrawQueue::IsDebugDrawingEnabled())
 			{
-				for (int32 ConstraintIndex = 0; ConstraintIndex < Collisions.NumConstraints(); ++ConstraintIndex)
+				for (const Chaos::TPBDCollisionConstraintHandle<float, 3> * ConstraintHandle : Collisions.GetConstConstraintHandles())
 				{
-					DrawCollisionImpl(SpaceTransform, Collisions.GetConstraintHandle(ConstraintIndex), ColorScale);
+					DrawCollisionImpl(SpaceTransform, ConstraintHandle, ColorScale);
 				}
 			}
 #endif
diff --git a/Engine/Source/Runtime/Experimental/Chaos/Private/Chaos/PBD6DJointConstraints.cpp b/Engine/Source/Runtime/Experimental/Chaos/Private/Chaos/PBD6DJointConstraints.cpp
--- a/Engine/Source/Runtime/Experimental/Chaos/Private/Chaos/PBD6DJointConstraints.cpp
+++ b/Engine/Source/Runtime/Experimental/Chaos/Private/Chaos/PBD6DJointConstraints.cpp
@@ -22,7 +22,8 @@ namespace Chaos
 	}
 
 
-	FPBD6DJointConstraintHandle::FPBD6DJointConstraintHandle(FConstraintContainer* InConstraintContainer, int32 InConstraintIndex) : TContainerConstraintHandle<FPBD6DJointConstraints>(InConstraintContainer, InConstraintIndex)
+	FPBD6DJointConstraintHandle::FPBD6DJointConstraintHandle(FConstraintContainer* InConstraintContainer, int32 InConstraintIndex)
+		: TContainerConstraintHandle<FPBD6DJointConstraints>(InConstraintContainer, InConstraintIndex)
 	{
 	}
 
@@ -44,6 +45,12 @@ namespace Chaos
 		return ConstraintContainer->GetConstraintLevel(ConstraintIndex);
 	}
 
+	TVector<TGeometryParticleHandle<float,3>*, 2> FPBD6DJointConstraintHandle::GetConstrainedParticles() const 
+	{ 
+		return ConstraintContainer->GetConstrainedParticles(ConstraintIndex); 
+	}
+
+
 	//
 	// Constraint Math
 	//
diff --git a/Engine/Source/Runtime/Experimental/Chaos/Private/Chaos/PBDCollisionConstraint.cpp b/Engine/Source/Runtime/Experimental/Chaos/Private/Chaos/PBDCollisionConstraint.cpp
--- a/Engine/Source/Runtime/Experimental/Chaos/Private/Chaos/PBDCollisionConstraint.cpp
+++ b/Engine/Source/Runtime/Experimental/Chaos/Private/Chaos/PBDCollisionConstraint.cpp
@@ -65,7 +65,6 @@ namespace Chaos
 	template<typename T, int d>
 	TPBDCollisionConstraint<T, d>::TPBDCollisionConstraint(const TPBDRigidsSOAs<T, d>& InParticles, TArrayCollectionArray<bool>& Collided, const TArrayCollectionArray<TSerializablePtr<FChaosPhysicsMaterial>>& InPerParticleMaterials, const int32 PairIterations /*= 1*/, const T Thickness /*= (T)0*/)
 		: Particles(InParticles)
-		, SpatialAcceleration(nullptr)
 		, MCollided(Collided)
 		, MPhysicsMaterials(InPerParticleMaterials)
 		, bEnableVelocitySolve(true)
@@ -79,22 +78,24 @@ namespace Chaos
 		, PostComputeCallback(nullptr)
 		, PostApplyCallback(nullptr)
 		, PostApplyPushOutCallback(nullptr)
+		, SpatialAcceleration(nullptr)
 	{
 	}
 
 	DECLARE_CYCLE_STAT(TEXT("TPBDCollisionConstraint::Reset"), STAT_CollisionConstraintsReset, STATGROUP_Chaos);
 
 	template<typename T, int d>
-	void TPBDCollisionConstraint<T, d>::Reset(/*const TPBDRigidParticles<T, d>& InParticles, const TArray<int32>& InIndices*/)
+	void TPBDCollisionConstraint<T, d>::Reset()
 	{
 		SCOPE_CYCLE_COUNTER(STAT_CollisionConstraintsReset);
 
+		TArray<FConstraintContainerHandle*> CopyOfHandles = Handles;
 
-		for (int32 Idx = Constraints.Num() - 1; Idx >= 0; Idx--)
+		for (FConstraintContainerHandle* ContactHandle : CopyOfHandles)
 		{
 			//if (!bEnableCollisions)
 			{
-				RemoveConstraint(Idx);
+				RemoveConstraint(ContactHandle);
 			}
 		}
 
@@ -104,16 +105,20 @@ namespace Chaos
 	}
 
 	template<typename T, int d>
-	void TPBDCollisionConstraint<T, d>::RemoveConstraint(int32 Idx)
+	void TPBDCollisionConstraint<T, d>::RemoveConstraint(FConstraintContainerHandle* Handle)
 	{
-		HandleAllocator.FreeHandle(Handles[Idx]);
+		int32 Idx = Handle->GetConstraintIndex();
+		typename TCollisionConstraintBase<T, d>::FType ConstraintType = Handle->GetType();
+
 		Handles.RemoveAtSwap(Idx);
 		Constraints.RemoveAtSwap(Idx);
 		if (Idx < Constraints.Num())
 		{
-			Handles[Idx]->SetConstraintIndex(Idx);
+			Handles[Idx]->SetConstraintIndex(Idx, ConstraintType);
 		}
+
 		ensure(Handles.Num() == Constraints.Num());
+		HandleAllocator.FreeHandle(Handle);
 	}
 
 	template<typename T, int d>
@@ -122,26 +127,30 @@ namespace Chaos
 		const TArray<TGeometryParticleHandle<T, d>*> HandleArray = InHandleSet.Array();
 		for (auto ParticleHandle : HandleArray)
 		{
-			for (int32 Idx = Constraints.Num() - 1; Idx >= 0; Idx--)
+			TArray<FConstraintContainerHandle*> CopyOfHandles = Handles;
+
+			for (FConstraintContainerHandle* ContactHandle : CopyOfHandles)
 			{
-				if (Constraints[Idx].Particle[1] == ParticleHandle || Constraints[Idx].Particle[0] == ParticleHandle)
+				TVector<TGeometryParticleHandle<T, d>*, 2> ConstraintParticles = ContactHandle->GetConstrainedParticles();
+				if (ConstraintParticles[1] == ParticleHandle || ConstraintParticles[0] == ParticleHandle)
 				{
-					RemoveConstraint(Idx);
+					RemoveConstraint(ContactHandle);
 				}
 			}
 		}
 	}
 
 	template<typename T, int d>
-	void TPBDCollisionConstraint<T, d>::ApplyCollisionModifier(const TFunction<ECollisionModifierResult(FRigidBodyContactConstraint& Constraint)>& CollisionModifier)
+	void TPBDCollisionConstraint<T, d>::ApplyCollisionModifier(const TFunction<ECollisionModifierResult(const FConstraintContainerHandle* Handle)>& CollisionModifier)
 	{
-		for (int ConstraintIndex = 0; ConstraintIndex < NumConstraints(); ++ConstraintIndex)
+		TArray<FConstraintContainerHandle*> CopyOfHandles = Handles;
+
+		for (FConstraintContainerHandle* ContactHandle : CopyOfHandles)
 		{
-			ECollisionModifierResult Result = CollisionModifier(Constraints[ConstraintIndex]);
+			ECollisionModifierResult Result = CollisionModifier(ContactHandle);
 			if (Result == ECollisionModifierResult::Disabled)
 			{
-				RemoveConstraint(ConstraintIndex);
-				--ConstraintIndex;
+				RemoveConstraint(ContactHandle);
 			}
 		}
 	}
@@ -535,7 +544,7 @@ namespace Chaos
 
 
 	template<typename T, int d>
-	void TPBDCollisionConstraint<T, d>::ApplyPushOut(const T Dt, FRigidBodyContactConstraint& Constraint, const TSet<TGeometryParticleHandle<T, d>*>& IsTemporarilyStatic, int32 Iteration, int32 NumIterations, bool &NeedsAnotherIteration)
+	void TPBDCollisionConstraint<T, d>::ApplyPushOut(const T Dt, FRigidBodyContactConstraint& Constraint, const TSet<const TGeometryParticleHandle<T, d>*>& IsTemporarilyStatic, int32 Iteration, int32 NumIterations, bool &NeedsAnotherIteration)
 	{
 		TGeometryParticleHandle<T, d>* Particle0 = Constraint.Particle[0];
 		TGeometryParticleHandle<T, d>* Particle1 = Constraint.Particle[1];
@@ -654,12 +663,11 @@ namespace Chaos
 
 	DECLARE_CYCLE_STAT(TEXT("TPBDCollisionConstraint::ApplyPushOut"), STAT_ApplyPushOut, STATGROUP_Chaos);
 	template<typename T, int d>
-	bool TPBDCollisionConstraint<T, d>::ApplyPushOut(const T Dt, const TArray<FConstraintContainerHandle*>& InConstraintHandles, const TSet<TGeometryParticleHandle<T, d>*>& IsTemporarilyStatic, int32 Iteration, int32 NumIterations)
+	bool TPBDCollisionConstraint<T, d>::ApplyPushOut(const T Dt, const TArray<FConstraintContainerHandle*>& InConstraintHandles, const TSet< const TGeometryParticleHandle<T, d>*>& IsTemporarilyStatic, int32 Iteration, int32 NumIterations)
 	{
 		SCOPE_CYCLE_COUNTER(STAT_ApplyPushOut);
 
 		bool NeedsAnotherIteration = false;
-
 		if (MPairIterations > 0)
 		{
 			PhysicsParallelFor(InConstraintHandles.Num(), [&](int32 ConstraintHandleIndex) {
@@ -673,7 +681,6 @@ namespace Chaos
 		{
 			PostApplyPushOutCallback(Dt, InConstraintHandles, NeedsAnotherIteration);
 		}
-
 		return NeedsAnotherIteration;
 	}
 
@@ -681,14 +688,12 @@ namespace Chaos
 	template<typename T, int d>
 	TRigidTransform<T, d> GetTransform(const TGeometryParticleHandle<T, d>* Particle)
 	{
-		//#if CHAOS_PARTICLEHANDLE_TODO
-
 		TGenericParticleHandle<T, d> Generic = const_cast<TGeometryParticleHandle<T, d>*>(Particle);	//TODO: give a const version of the generic API
 		return TRigidTransform<T, d>(Generic->P(), Generic->Q());
 	}
 
 	template<typename T, int d>
-	void TPBDCollisionConstraint<T, d>::ConstructConstraints(TGeometryParticleHandle<T, d>* Particle0, TGeometryParticleHandle<T, d>* Particle1, const T Thickness, TRigidBodyContactConstraint<T, d>& Constraint)
+	void TPBDCollisionConstraint<T, d>::ConstructConstraints(TGeometryParticleHandle<T, d>* Particle0, TGeometryParticleHandle<T, d>* Particle1, const T Thickness, TRigidBodyContactConstraint<T, d> & Constraint)
 	{
 		if (ensure(Particle0 && Particle1))
 		{
@@ -731,8 +736,8 @@ namespace Chaos
 	template class TPBDCollisionConstraintHandle<float, 3>;
 	template class TAccelerationStructureHandle<float, 3>;
 	template class CHAOS_API TPBDCollisionConstraint<float, 3>;
-	template void TPBDCollisionConstraint<float, 3>::UpdateConstraint<ECollisionUpdateType::Any>(const float Thickness, FRigidBodyContactConstraint& Constraints);
-	template void TPBDCollisionConstraint<float, 3>::UpdateConstraint<ECollisionUpdateType::Deepest>(const float Thickness, FRigidBodyContactConstraint& Constraints);
+	template void TPBDCollisionConstraint<float, 3>::UpdateConstraint<ECollisionUpdateType::Any>(const float Thickness, FRigidBodyContactConstraint& Constraint);
+	template void TPBDCollisionConstraint<float, 3>::UpdateConstraint<ECollisionUpdateType::Deepest>(const float Thickness, FRigidBodyContactConstraint& Constraint);
 	template void TPBDCollisionConstraint<float, 3>::ComputeConstraints<false>(const TPBDCollisionConstraint<float, 3>::FAccelerationStructure&, float Dt);
 	template void TPBDCollisionConstraint<float, 3>::ComputeConstraints<true>(const TPBDCollisionConstraint<float, 3>::FAccelerationStructure&, float Dt);
 }
diff --git a/Engine/Source/Runtime/Experimental/Chaos/Private/Chaos/PBDCollisionTypes.cpp b/Engine/Source/Runtime/Experimental/Chaos/Private/Chaos/PBDCollisionTypes.cpp
new file mode 100644
--- /dev/null
+++ b/Engine/Source/Runtime/Experimental/Chaos/Private/Chaos/PBDCollisionTypes.cpp
@@ -0,0 +1,24 @@
+// Copyright 1998-2019 Epic Games, Inc. All Rights Reserved.
+
+#include "Chaos/PBDCollisionTypes.h"
+#include "Chaos/PBDCollisionConstraint.h"
+
+namespace Chaos
+{
+
+	template<class T, int d> template<class TYPE>
+	TYPE& TPBDCollisionConstraintHandle<T, d>::GetContact()
+	{
+		return ConstraintContainer->Constraints[ConstraintIndex];
+	}
+
+	template<class T, int d> template<class TYPE>
+	const TYPE& TPBDCollisionConstraintHandle<T, d>::GetContact() const
+	{
+		return ConstraintContainer->Constraints[ConstraintIndex];
+	}
+
+	template class CHAOS_API TPBDCollisionConstraintHandle<float, 3>;
+	template CHAOS_API TRigidBodyContactConstraint<float, 3>& TPBDCollisionConstraintHandle<float, 3>::GetContact<TRigidBodyContactConstraint<float, 3>>();
+	template CHAOS_API const TRigidBodyContactConstraint<float, 3>& TPBDCollisionConstraintHandle<float, 3>::GetContact<TRigidBodyContactConstraint<float, 3>>() const;
+};
diff --git a/Engine/Source/Runtime/Experimental/Chaos/Private/Chaos/PBDConstraintRule.cpp b/Engine/Source/Runtime/Experimental/Chaos/Private/Chaos/PBDConstraintRule.cpp
--- a/Engine/Source/Runtime/Experimental/Chaos/Private/Chaos/PBDConstraintRule.cpp
+++ b/Engine/Source/Runtime/Experimental/Chaos/Private/Chaos/PBDConstraintRule.cpp
@@ -7,13 +7,63 @@
 #include "Chaos/PBDJointConstraints.h"
 #include "Chaos/PBDPositionConstraints.h"
 #include "Chaos/PBDRigidDynamicSpringConstraints.h"
-
+#include "Chaos/PBDRigidSpringConstraints.h"
 
 namespace Chaos
 {
+
+	template<class T_CONSTRAINTS>
+	TPBDConstraintGraphRuleImpl<T_CONSTRAINTS>::TPBDConstraintGraphRuleImpl(FConstraints& InConstraints, int32 InPriority)
+		: FPBDConstraintGraphRule(InPriority)
+		, Constraints(InConstraints)
+		, ConstraintGraph(nullptr)
+	{
+	}
+
+
+	template<class T_CONSTRAINTS>
+	void TPBDConstraintGraphRuleImpl<T_CONSTRAINTS>::BindToGraph(FPBDConstraintGraph& InContactGraph, uint32 InContainerId)
+	{
+		ConstraintGraph = &InContactGraph;
+		ContainerId = InContainerId;
+	}
+
+	template<class T_CONSTRAINTS>
+	void TPBDConstraintGraphRuleImpl<T_CONSTRAINTS>::UpdatePositionBasedState(const FReal Dt)
+	{
+		Constraints.UpdatePositionBasedState(Dt);
+	}
+
+
+	template<class T_CONSTRAINTS>
+	void TPBDConstraintGraphRuleImpl<T_CONSTRAINTS>::AddToGraph()
+	{
+		ConstraintGraph->ReserveConstraints(Constraints.NumConstraints());
+
+		for (typename FConstraints::FConstraintContainerHandle * ConstraintHandle : Constraints.GetConstraintHandles())
+		{
+			ConstraintGraph->AddConstraint(ContainerId, ConstraintHandle, ConstraintHandle->GetConstrainedParticles());
+		}
+
+	}
+
+	template<class T_CONSTRAINTS>
+	int32 TPBDConstraintGraphRuleImpl<T_CONSTRAINTS>::NumConstraints() const
+	{ 
+		return Constraints.NumConstraints(); 
+	}
+
+	template class TPBDConstraintGraphRuleImpl< TPBDCollisionConstraint<float, 3> >;
+	template class TPBDConstraintGraphRuleImpl< FPBD6DJointConstraints >;
+	template class TPBDConstraintGraphRuleImpl< FPBDJointConstraints >;
+	template class TPBDConstraintGraphRuleImpl< TPBDPositionConstraints<float, 3> >;
+	template class TPBDConstraintGraphRuleImpl< TPBDRigidDynamicSpringConstraints<float, 3> >;
+	template class TPBDConstraintGraphRuleImpl< TPBDRigidSpringConstraints<float, 3> >;
+
 	template class TPBDConstraintColorRule<TPBDCollisionConstraint<float, 3>>;
 	template class TPBDConstraintIslandRule<FPBD6DJointConstraints>;
 	template class TPBDConstraintIslandRule<FPBDJointConstraints>;
 	template class TPBDConstraintIslandRule<TPBDPositionConstraints<float, 3>>;
 	template class TPBDConstraintIslandRule<TPBDRigidDynamicSpringConstraints<float, 3>>;
+	template class TPBDConstraintIslandRule<TPBDRigidSpringConstraints<float, 3>>;
 }
\ No newline at end of file
diff --git a/Engine/Source/Runtime/Experimental/Chaos/Private/Chaos/PBDJointConstraints.cpp b/Engine/Source/Runtime/Experimental/Chaos/Private/Chaos/PBDJointConstraints.cpp
--- a/Engine/Source/Runtime/Experimental/Chaos/Private/Chaos/PBDJointConstraints.cpp
+++ b/Engine/Source/Runtime/Experimental/Chaos/Private/Chaos/PBDJointConstraints.cpp
@@ -61,7 +61,8 @@ namespace Chaos
 	}
 
 	
-	FPBDJointConstraintHandle::FPBDJointConstraintHandle(FConstraintContainer* InConstraintContainer, int32 InConstraintIndex) : TContainerConstraintHandle<FPBDJointConstraints>(InConstraintContainer, InConstraintIndex)
+	FPBDJointConstraintHandle::FPBDJointConstraintHandle(FConstraintContainer* InConstraintContainer, int32 InConstraintIndex)
+		: TContainerConstraintHandle<FPBDJointConstraints>(InConstraintContainer, InConstraintIndex)
 	{
 	}
 
@@ -89,6 +90,11 @@ namespace Chaos
 		return ConstraintContainer->GetConstraintSettings(ConstraintIndex);
 	}
 
+	TVector<TGeometryParticleHandle<float,3>*, 2> FPBDJointConstraintHandle::GetConstrainedParticles() const 
+	{ 
+		return ConstraintContainer->GetConstrainedParticles(ConstraintIndex); 
+	}
+
 	//
 	// Constraint Settings
 	//
diff --git a/Engine/Source/Runtime/Experimental/Chaos/Private/Chaos/PBDPositionConstraints.cpp b/Engine/Source/Runtime/Experimental/Chaos/Private/Chaos/PBDPositionConstraints.cpp
new file mode 100644
--- /dev/null
+++ b/Engine/Source/Runtime/Experimental/Chaos/Private/Chaos/PBDPositionConstraints.cpp
@@ -0,0 +1,22 @@
+// Copyright 1998-2019 Epic Games, Inc. All Rights Reserved.
+
+#include "Chaos/PBDPositionConstraints.h"
+
+
+namespace Chaos
+{
+
+	template<class T, int d>
+	void TPBDPositionConstraints<T, d>::Apply(const T Dt, const TArray<FConstraintContainerHandle*>& ConstraintHandles, const int32 It, const int32 NumIts) const
+	{
+		for (FConstraintContainerHandle* ConstraintHandle : ConstraintHandles)
+		{
+			ApplySingle(Dt, ConstraintHandle->GetConstraintIndex());
+		}
+	}
+
+
+
+	template class TPBDPositionConstraintHandle<float, 3>;
+	template class TPBDPositionConstraints<float, 3>;
+}
diff --git a/Engine/Source/Runtime/Experimental/Chaos/Private/Chaos/PBDRigidClustering.cpp b/Engine/Source/Runtime/Experimental/Chaos/Private/Chaos/PBDRigidClustering.cpp
--- a/Engine/Source/Runtime/Experimental/Chaos/Private/Chaos/PBDRigidClustering.cpp
+++ b/Engine/Source/Runtime/Experimental/Chaos/Private/Chaos/PBDRigidClustering.cpp
@@ -631,16 +631,16 @@ namespace Chaos
 
 		ResetCollisionImpulseArray();
 
-		for (const typename FPBDCollisionConstraint::FRigidBodyContactConstraint& Contact : CollisionRule.GetAllConstraints())
+		for (const Chaos::TPBDCollisionConstraintHandle<T, 3> * ContactHandle : CollisionRule.GetConstConstraintHandles())
 		{
-			if (Contact.AccumulatedImpulse.Size() < MinImpulseForStrainEval)
+			if (ContactHandle->GetAccumulatedImpulse().Size() < MinImpulseForStrainEval)
 			{
 				continue;
 			}
 
-			auto ComputeStrainLambda = [&](TPBDRigidClusteredParticleHandle<T, d>* Cluster, const TArray<uint32>& ParentToChildren) {
+			auto ComputeStrainLambda = [&](const TPBDRigidClusteredParticleHandle<T, d>* Cluster, const TArray<uint32>& ParentToChildren) {
 				const TRigidTransform<T, d> WorldToClusterTM = TRigidTransform<T, d>(Cluster->P(), Cluster->Q());
-				const TVector<T, d> ContactLocationClusterLocal = WorldToClusterTM.InverseTransformPosition(GetContactLocation(Contact));
+				const TVector<T, d> ContactLocationClusterLocal = WorldToClusterTM.InverseTransformPosition(ContactHandle->GetContactLocation());
 				TBox<T, d> ContactBox(ContactLocationClusterLocal, ContactLocationClusterLocal);
 				ContactBox.Thicken(ClusterDistanceThreshold);
 				if (Cluster->ChildrenSpatial())
@@ -661,26 +661,27 @@ namespace Chaos
 								const TArray<int32> SubIntersections = MParticles.ChildrenSpatial(Child)->FindAllIntersectingChildren(ContactBoxProxy);
 								for (int32 SubChild : SubIntersections)
 								{
-									MParticles.CollisionImpulses(SubChild) += Contact.AccumulatedImpulse.Size();
+									MParticles.CollisionImpulses(SubChild) += ContactHandle->GetAccumulatedImpulse().Size();
 								}
 							}
 						}
 						else
 						{
-							MParticles.CollisionImpulses(Child) += Contact.AccumulatedImpulse.Size();
+							MParticles.CollisionImpulses(Child) += ContactHandle->GetAccumulatedImpulse().Size();
 						}
 					}
 				}
 			};
 
-			if (auto ChildrenPtr = MParentToChildren.Find(Contact.Particle[0]))
+			TVector<const TGeometryParticleHandle<T, d>*, 2> ConstrainedParticles = ContactHandle->GetConstrainedParticles();
+			if (auto ChildrenPtr = MParentToChildren.Find(ConstrainedParticles[0]))
 			{
-				ComputeStrainLambda(Contact.Particle[0]->AsClustered(), *ChildrenPtr);
+				ComputeStrainLambda(ConstrainedParticles[0]->AsClustered(), *ChildrenPtr);
 			}
 
-			if (auto ChildrenPtr = MParentToChildren.Find(Contact.Particle[1]))
+			if (auto ChildrenPtr = MParentToChildren.Find(ConstrainedParticles[1]))
 			{
-				ComputeStrainLambda(Contact.Particle[1]->AsClustered(), *ChildrenPtr);
+				ComputeStrainLambda(ConstrainedParticles[1]->AsClustered(), *ChildrenPtr);
 			}
 
 			MCollisionImpulseArrayDirty = true;
diff --git a/Engine/Source/Runtime/Experimental/Chaos/Private/Chaos/PBDRigidDynamicSpringConstraints.cpp b/Engine/Source/Runtime/Experimental/Chaos/Private/Chaos/PBDRigidDynamicSpringConstraints.cpp
--- a/Engine/Source/Runtime/Experimental/Chaos/Private/Chaos/PBDRigidDynamicSpringConstraints.cpp
+++ b/Engine/Source/Runtime/Experimental/Chaos/Private/Chaos/PBDRigidDynamicSpringConstraints.cpp
@@ -4,6 +4,12 @@
 
 using namespace Chaos;
 
+template<class T, int d>
+TVector<TGeometryParticleHandle<T, d>*, 2> TPBDRigidDynamicSpringConstraintHandle<T,d>::GetConstrainedParticles() const
+{
+	return ConstraintContainer->GetConstrainedParticles(ConstraintIndex);
+}
+
 template<class T, int d>
 void TPBDRigidDynamicSpringConstraints<T, d>::UpdatePositionBasedState(const T Dt)
 {
@@ -153,5 +159,6 @@ void TPBDRigidDynamicSpringConstraints<T, d>::ApplySingle(const T Dt, int32 Cons
 
 namespace Chaos
 {
+	template class TPBDRigidDynamicSpringConstraintHandle<float, 3>;
 	template class TPBDRigidDynamicSpringConstraints<float, 3>;
 }
diff --git a/Engine/Source/Runtime/Experimental/Chaos/Private/Chaos/PBDRigidSpringConstraints.cpp b/Engine/Source/Runtime/Experimental/Chaos/Private/Chaos/PBDRigidSpringConstraints.cpp
--- a/Engine/Source/Runtime/Experimental/Chaos/Private/Chaos/PBDRigidSpringConstraints.cpp
+++ b/Engine/Source/Runtime/Experimental/Chaos/Private/Chaos/PBDRigidSpringConstraints.cpp
@@ -2,85 +2,95 @@
 
 #include "Chaos/PBDRigidSpringConstraints.h"
 
-using namespace Chaos;
 
-template<class T, int d>
-void TPBDRigidSpringConstraints<T, d>::UpdateDistance(int32 ConstraintIndex, const TVector<T, d>& Location0, const TVector<T, d>& Location1)
-{
-	const TVector<TGeometryParticleHandle<T, d>*, 2>& Constraint = Constraints[ConstraintIndex];
-	const TGeometryParticleHandle<T, d>* Particle0 = Constraint[0];
-	const TGeometryParticleHandle<T, d>* Particle1 = Constraint[1];
-
-	Distances[ConstraintIndex][0] = Particle0->R().Inverse().RotateVector(Location0 - Particle0->X());
-	Distances[ConstraintIndex][1] = Particle1->R().Inverse().RotateVector(Location1 - Particle1->X());
-	SpringDistances[ConstraintIndex] = (Location0 - Location1).Size();
-}
-
-template<class T, int d>
-TVector<T, d> TPBDRigidSpringConstraints<T, d>::GetDelta(int32 ConstraintIndex, const TVector<T, d>& WorldSpaceX1, const TVector<T, d>& WorldSpaceX2) const
+namespace Chaos
 {
-	const TVector<TGeometryParticleHandle<T, d>*, 2>& Constraint = Constraints[ConstraintIndex];
-	const TPBDRigidParticleHandle<T, d>* PBDRigid0 = Constraint[0]->AsDynamic();
-	const TPBDRigidParticleHandle<T, d>* PBDRigid1 = Constraint[1]->AsDynamic();
 
-	if (!PBDRigid0 && !PBDRigid1)
+	template<class T, int d>
+	void TPBDRigidSpringConstraints<T, d>::UpdateDistance(int32 ConstraintIndex, const TVector<T, d>& Location0, const TVector<T, d>& Location1)
 	{
-		return TVector<T, d>(0);
-	}
+		const TVector<TGeometryParticleHandle<T, d>*, 2>& Constraint = Constraints[ConstraintIndex];
+		const TGeometryParticleHandle<T, d>* Particle0 = Constraint[0];
+		const TGeometryParticleHandle<T, d>* Particle1 = Constraint[1];
 
-	const TVector<T, d> Difference = WorldSpaceX2 - WorldSpaceX1;
-
-	const float Distance = Difference.Size();
-	check(Distance > 1e-7);
+		Distances[ConstraintIndex][0] = Particle0->R().Inverse().RotateVector(Location0 - Particle0->X());
+		Distances[ConstraintIndex][1] = Particle1->R().Inverse().RotateVector(Location1 - Particle1->X());
+		SpringDistances[ConstraintIndex] = (Location0 - Location1).Size();
+	}
 
-	const TVector<T, d> Direction = Difference / Distance;
-	const TVector<T, d> Delta = (Distance - SpringDistances[ConstraintIndex]) * Direction;
-	const T InvM0 = (PBDRigid0) ? PBDRigid0->InvM() : (T)0;
-	const T InvM1 = (PBDRigid1) ? PBDRigid1->InvM() : (T)0;
-	const T CombinedMass = InvM0 + InvM1;
+	template<class T, int d>
+	TVector<T, d> TPBDRigidSpringConstraints<T, d>::GetDelta(int32 ConstraintIndex, const TVector<T, d>& WorldSpaceX1, const TVector<T, d>& WorldSpaceX2) const
+	{
+		const TVector<TGeometryParticleHandle<T, d>*, 2>& Constraint = Constraints[ConstraintIndex];
+		const TPBDRigidParticleHandle<T, d>* PBDRigid0 = Constraint[0]->AsDynamic();
+		const TPBDRigidParticleHandle<T, d>* PBDRigid1 = Constraint[1]->AsDynamic();
 
-	return Stiffness * Delta / CombinedMass;
-}
+		if (!PBDRigid0 && !PBDRigid1)
+		{
+			return TVector<T, d>(0);
+		}
 
-template<class T, int d>
-void TPBDRigidSpringConstraints<T, d>::ApplySingle(const T Dt, int32 ConstraintIndex) const
-{
-	const FConstrainedParticlePair& Constraint = Constraints[ConstraintIndex];
+		const TVector<T, d> Difference = WorldSpaceX2 - WorldSpaceX1;
 
-	TPBDRigidParticleHandle<T, d>* PBDRigid0 = Constraint[0]->AsDynamic();
-	TPBDRigidParticleHandle<T, d>* PBDRigid1 = Constraint[1]->AsDynamic();
-	check((PBDRigid0 && PBDRigid1 && PBDRigid0->Island() == PBDRigid1->Island()) || (!PBDRigid0 && PBDRigid1) || (PBDRigid0 && PBDRigid1));
+		const float Distance = Difference.Size();
+		check(Distance > 1e-7);
 
-	// @todo(ccaulfield): see if we can eliminate the need for all these ifs
-	const TRotation<T, d> & Q0 = PBDRigid0 ? PBDRigid0->Q() : Constraint[0]->R();
-	const TRotation<T, d> & Q1 = PBDRigid1 ? PBDRigid1->Q() : Constraint[1]->R();
-	const TVector<T, d> & P0 = PBDRigid0 ? PBDRigid0->P() : Constraint[0]->X();
-	const TVector<T, d> & P1 = PBDRigid1 ? PBDRigid1->P() : Constraint[1]->X();
+		const TVector<T, d> Direction = Difference / Distance;
+		const TVector<T, d> Delta = (Distance - SpringDistances[ConstraintIndex]) * Direction;
+		const T InvM0 = (PBDRigid0) ? PBDRigid0->InvM() : (T)0;
+		const T InvM1 = (PBDRigid1) ? PBDRigid1->InvM() : (T)0;
+		const T CombinedMass = InvM0 + InvM1;
 
-	const TVector<T, d> WorldSpaceX1 = Q0.RotateVector(Distances[ConstraintIndex][0]) + P0;
-	const TVector<T, d> WorldSpaceX2 = Q1.RotateVector(Distances[ConstraintIndex][1]) + P1;
-	const PMatrix<T, d, d> WorldSpaceInvI1 = PBDRigid0 ? (PBDRigid0->Q() * FMatrix::Identity) * PBDRigid0->InvI() * (PBDRigid0->Q() * FMatrix::Identity).GetTransposed() : PMatrix<T, d, d>(0);
-	const PMatrix<T, d, d> WorldSpaceInvI2 = PBDRigid1 ? (PBDRigid1->Q() * FMatrix::Identity) * PBDRigid1->InvI() * (PBDRigid1->Q() * FMatrix::Identity).GetTransposed() : PMatrix<T, d, d>(0);
-	const TVector<T, d> Delta = GetDelta(ConstraintIndex, WorldSpaceX1, WorldSpaceX2);
+		return Stiffness * Delta / CombinedMass;
+	}
 
-	if (PBDRigid0)
+	template<class T, int d>
+	void TPBDRigidSpringConstraints<T, d>::Apply(const T Dt, const TArray<FConstraintContainerHandle*>& InConstraintHandles, const int32 It, const int32 NumIts)
 	{
-		const TVector<T, d> Radius = WorldSpaceX1 - PBDRigid0->P();
-		PBDRigid0->P() += PBDRigid0->InvM() * Delta;
-		PBDRigid0->Q() += TRotation<T, d>::FromElements(WorldSpaceInvI1 * TVector<T, d>::CrossProduct(Radius, Delta), 0.f) * PBDRigid0->Q() * T(0.5);
-		PBDRigid0->Q().Normalize();
+		for (FConstraintContainerHandle* ConstraintHandle : InConstraintHandles)
+		{
+			ApplySingle(Dt, ConstraintHandle->GetConstraintIndex());
+		}
 	}
 
-	if (PBDRigid1)
+
+	template<class T, int d>
+	void TPBDRigidSpringConstraints<T, d>::ApplySingle(const T Dt, int32 ConstraintIndex) const
 	{
-		const TVector<T, d> Radius = WorldSpaceX2 - PBDRigid1->P();
-		PBDRigid1->P() -= PBDRigid1->InvM() * Delta;
-		PBDRigid1->Q() += TRotation<T, d>::FromElements(WorldSpaceInvI2 * TVector<T, d>::CrossProduct(Radius, -Delta), 0.f) * PBDRigid1->Q() * T(0.5);
-		PBDRigid1->Q().Normalize();
+		const FConstrainedParticlePair& Constraint = Constraints[ConstraintIndex];
+
+		TPBDRigidParticleHandle<T, d>* PBDRigid0 = Constraint[0]->AsDynamic();
+		TPBDRigidParticleHandle<T, d>* PBDRigid1 = Constraint[1]->AsDynamic();
+		check((PBDRigid0 && PBDRigid1 && PBDRigid0->Island() == PBDRigid1->Island()) || (!PBDRigid0 && PBDRigid1) || (PBDRigid0 && PBDRigid1));
+
+		// @todo(ccaulfield): see if we can eliminate the need for all these ifs
+		const TRotation<T, d> & Q0 = PBDRigid0 ? PBDRigid0->Q() : Constraint[0]->R();
+		const TRotation<T, d> & Q1 = PBDRigid1 ? PBDRigid1->Q() : Constraint[1]->R();
+		const TVector<T, d> & P0 = PBDRigid0 ? PBDRigid0->P() : Constraint[0]->X();
+		const TVector<T, d> & P1 = PBDRigid1 ? PBDRigid1->P() : Constraint[1]->X();
+
+		const TVector<T, d> WorldSpaceX1 = Q0.RotateVector(Distances[ConstraintIndex][0]) + P0;
+		const TVector<T, d> WorldSpaceX2 = Q1.RotateVector(Distances[ConstraintIndex][1]) + P1;
+		const PMatrix<T, d, d> WorldSpaceInvI1 = PBDRigid0 ? (PBDRigid0->Q() * FMatrix::Identity) * PBDRigid0->InvI() * (PBDRigid0->Q() * FMatrix::Identity).GetTransposed() : PMatrix<T, d, d>(0);
+		const PMatrix<T, d, d> WorldSpaceInvI2 = PBDRigid1 ? (PBDRigid1->Q() * FMatrix::Identity) * PBDRigid1->InvI() * (PBDRigid1->Q() * FMatrix::Identity).GetTransposed() : PMatrix<T, d, d>(0);
+		const TVector<T, d> Delta = GetDelta(ConstraintIndex, WorldSpaceX1, WorldSpaceX2);
+
+		if (PBDRigid0)
+		{
+			const TVector<T, d> Radius = WorldSpaceX1 - PBDRigid0->P();
+			PBDRigid0->P() += PBDRigid0->InvM() * Delta;
+			PBDRigid0->Q() += TRotation<T, d>::FromElements(WorldSpaceInvI1 * TVector<T, d>::CrossProduct(Radius, Delta), 0.f) * PBDRigid0->Q() * T(0.5);
+			PBDRigid0->Q().Normalize();
+		}
+
+		if (PBDRigid1)
+		{
+			const TVector<T, d> Radius = WorldSpaceX2 - PBDRigid1->P();
+			PBDRigid1->P() -= PBDRigid1->InvM() * Delta;
+			PBDRigid1->Q() += TRotation<T, d>::FromElements(WorldSpaceInvI2 * TVector<T, d>::CrossProduct(Radius, -Delta), 0.f) * PBDRigid1->Q() * T(0.5);
+			PBDRigid1->Q().Normalize();
+		}
 	}
-}
 
-namespace Chaos
-{
 	template class TPBDRigidSpringConstraints<float, 3>;
 }
diff --git a/Engine/Source/Runtime/Experimental/Chaos/Public/Chaos/ConstraintHandle.h b/Engine/Source/Runtime/Experimental/Chaos/Public/Chaos/ConstraintHandle.h
--- a/Engine/Source/Runtime/Experimental/Chaos/Public/Chaos/ConstraintHandle.h
+++ b/Engine/Source/Runtime/Experimental/Chaos/Public/Chaos/ConstraintHandle.h
@@ -17,13 +17,13 @@ namespace Chaos
 	public:
 		using FGeometryParticleHandle = TGeometryParticleHandle<FReal, 3>;
 
-		FConstraintHandle() : ConstraintIndex(INDEX_NONE) {}
-		FConstraintHandle(int32 InConstraintIndex) : ConstraintIndex(InConstraintIndex) {}
+		FConstraintHandle() : ConstraintIndex(INDEX_NONE) { }
+		FConstraintHandle(int32 InConstraintIndex): ConstraintIndex(InConstraintIndex) {}
 		~FConstraintHandle() {}
 
 		bool IsValid() const
 		{
-			return (ConstraintIndex != INDEX_NONE);
+			return (ConstraintIndex != INDEX_NONE );
 		}
 
 		int32 GetConstraintIndex() const
@@ -32,8 +32,7 @@ namespace Chaos
 		}
 
 		// @todo(ccaulfield): Add checked down - casting to specific constraint - type handle
-		template<typename T> 
-		T* As() { return static_cast<T*>(this); }
+		template<typename T>  T* As() { return static_cast<T*>(this); }
 
 	protected:
 		friend class FPBDConstraintContainer;
@@ -54,12 +53,11 @@ namespace Chaos
 		using FConstraintContainer = T_CONTAINER;
 
 		TContainerConstraintHandle() : ConstraintContainer(nullptr) {}
-		TContainerConstraintHandle(FConstraintContainer* InConstraintContainer, int32 InConstraintIndex) : FConstraintHandle(InConstraintIndex), ConstraintContainer(InConstraintContainer) {}
+		TContainerConstraintHandle(FConstraintContainer* InConstraintContainer, int32 InConstraintIndex) 
+			: FConstraintHandle(InConstraintIndex), ConstraintContainer(InConstraintContainer) {}
 
 		void RemoveConstraint() { ConstraintContainer->RemoveConstraint(ConstraintIndex); }
 
-		TVector<FGeometryParticleHandle*, 2> GetConstrainedParticles() const { return ConstraintContainer->GetConstrainedParticles(ConstraintIndex); }
-
 	protected:
 		using Base::ConstraintIndex;
 		FConstraintContainer* ConstraintContainer;
@@ -79,6 +77,8 @@ namespace Chaos
 		using FConstraintContainerHandle = typename FConstraintContainer::FConstraintContainerHandle;
 
 		FConstraintContainerHandle* AllocHandle(FConstraintContainer* ConstraintContainer, int32 ConstraintIndex) { return new FConstraintContainerHandle(ConstraintContainer, ConstraintIndex); }
+		template<class TYPE>
+		FConstraintContainerHandle* AllocHandle(FConstraintContainer* ConstraintContainer, int32 ConstraintIndex) { return new FConstraintContainerHandle(ConstraintContainer, ConstraintIndex, TYPE::StaticType()); }
 		void FreeHandle(FConstraintContainerHandle* Handle) { delete Handle; }
 	};
 }
diff --git a/Engine/Source/Runtime/Experimental/Chaos/Public/Chaos/PBD6DJointConstraints.h b/Engine/Source/Runtime/Experimental/Chaos/Public/Chaos/PBD6DJointConstraints.h
--- a/Engine/Source/Runtime/Experimental/Chaos/Public/Chaos/PBD6DJointConstraints.h
+++ b/Engine/Source/Runtime/Experimental/Chaos/Public/Chaos/PBD6DJointConstraints.h
@@ -28,7 +28,8 @@ namespace Chaos
 		void CalculateConstraintSpace(FVec3& OutXa, FMatrix33& OutRa, FVec3& OutXb, FMatrix33& OutRb, FVec3& OutCR) const;
 		void SetParticleLevels(const TVector<int32, 2>& ParticleLevels);
 		int32 GetConstraintLevel() const;
-
+		TVector<TGeometryParticleHandle<float,3>*, 2> GetConstrainedParticles() const;
+		
 	protected:
 		using Base::ConstraintIndex;
 		using Base::ConstraintContainer;
@@ -196,6 +197,7 @@ namespace Chaos
 		using FParticlePair = TVector<TGeometryParticleHandle<FReal, 3>*, 2>;
 		using FVectorPair = TVector<FVec3, 2>;
 		using FTransformPair = TVector<FRigidTransform3, 2>;
+		using FHandles = TArray<FConstraintContainerHandle*>;
 
 		FPBD6DJointConstraints(const FPBD6DJointSolverSettings& InSettings);
 
@@ -236,10 +238,19 @@ namespace Chaos
 		//
 		// Constraint API
 		//
+		FHandles& GetConstraintHandles()
+		{
+			return Handles;
+		}
+		const FHandles& GetConstConstraintHandles() const
+		{
+			return Handles;
+		}
 
 		const FConstraintContainerHandle* GetConstraintHandle(int32 ConstraintIndex) const;
 		FConstraintContainerHandle* GetConstraintHandle(int32 ConstraintIndex);
 
+
 		/**
 		 * Get the particles that are affected by the specified constraint.
 		 */
@@ -281,7 +292,7 @@ namespace Chaos
 		TArray<FParticlePair> ConstraintParticles;
 		TArray<FPBD6DJointState> ConstraintStates;
 
-		TArray<FConstraintContainerHandle*> Handles;
+		FHandles Handles;
 		FConstraintHandleAllocator HandleAllocator;
 
 		FD6JointPreApplyCallback PreApplyCallback;
diff --git a/Engine/Source/Runtime/Experimental/Chaos/Public/Chaos/PBDCollisionConstraint.h b/Engine/Source/Runtime/Experimental/Chaos/Public/Chaos/PBDCollisionConstraint.h
--- a/Engine/Source/Runtime/Experimental/Chaos/Public/Chaos/PBDCollisionConstraint.h
+++ b/Engine/Source/Runtime/Experimental/Chaos/Public/Chaos/PBDCollisionConstraint.h
@@ -14,252 +14,236 @@
 
 namespace Chaos
 {
-	template<typename T, int d>
-	class TPBDCollisionConstraintAccessor;
+template<typename T, int d>
+class TPBDCollisionConstraintAccessor;
 
-	template<typename T, int d>
-	class TPBDCollisionConstraint;
+template<typename T, int d>
+class TPBDCollisionConstraint;
 
-	template <typename T, int d>
-	class TRigidTransform;
+template <typename T, int d>
+class TRigidTransform;
 
-	class FImplicitObject;
+class FImplicitObject;
 
-	template <typename T, int d>
-	class TBVHParticles;
+template <typename T, int d>
+class TBVHParticles;
 
-	template <typename T, int d>
-	class TBox;
+template <typename T, int d>
+class TBox;
 
-	template<class T>
-	class TChaosPhysicsMaterial;
+template<class T>
+class TChaosPhysicsMaterial;
 
-	template <typename T, int d>
-	class TPBDRigidsSOAs;
+template <typename T, int d>
+class TPBDRigidsSOAs;
 
-	template<typename T, int d>
-	using TRigidBodyContactConstraintsPostComputeCallback = TFunction<void()>;
+template<typename T, int d>
+using TRigidBodyContactConstraintsPostComputeCallback = TFunction<void()>;
 
-	template<typename T, int d>
-	using TRigidBodyContactConstraintsPostApplyCallback = TFunction<void(const T Dt, const TArray<TPBDCollisionConstraintHandle<T, d>*>& InConstraintHandles)>;
+template<typename T, int d>
+using TRigidBodyContactConstraintsPostApplyCallback = TFunction<void(const T Dt, const TArray<TPBDCollisionConstraintHandle<T, d>*>& InConstraintHandles)>;
 
-	template<typename T, int d>
-	using TRigidBodyContactConstraintsPostApplyPushOutCallback = TFunction<void(const T Dt, const TArray<TPBDCollisionConstraintHandle<T, d>*>& InConstraintHandles, bool bRequiresAnotherIteration)>;
+template<typename T, int d>
+using TRigidBodyContactConstraintsPostApplyPushOutCallback = TFunction<void(const T Dt, const TArray<TPBDCollisionConstraintHandle<T, d>*>& InConstraintHandles, bool bRequiresAnotherIteration)>;
+
+/**
+ * Manages a set of contact constraints:
+ *	- Performs collision detection to generate constraints.
+ *	- Responsible for applying corrections to particles affected by the constraints.
+ * @todo(ccaulfield): rename to TPBDCollisionConstraints
+ * @todo(ccaulfield): remove TPBDCollisionConstraintAccessor
+ * @todo(ccaulfield): separate collision detection from constraint container
+ */
+template<typename T, int d>
+class CHAOS_API TPBDCollisionConstraint : public FPBDConstraintContainer
+{
+public:
+	friend class TPBDCollisionConstraintAccessor<T, d>;
+	friend class TPBDCollisionConstraintHandle<T, d>;
+
+	using Base = FPBDConstraintContainer;
+	using FConstraintHandleID = TPair<const TGeometryParticleHandle<T, d>*, const TGeometryParticleHandle<T, d>*>;
+	using FConstraintContainerHandle = TPBDCollisionConstraintHandle<T, d>; // @todo(brice) rename this to FCollisionConstraintHandle
+	using FConstraintHandleAllocator = TConstraintHandleAllocator<TPBDCollisionConstraint<T, d>>;
+	using FRigidBodyContactConstraint = TRigidBodyContactConstraint<T, d>;
+	using FHandles = TArray<FConstraintContainerHandle*>;
+
+	TPBDCollisionConstraint(const TPBDRigidsSOAs<T,d>& InParticles, TArrayCollectionArray<bool>& Collided, const TArrayCollectionArray<TSerializablePtr<FChaosPhysicsMaterial>>& PerParticleMaterials, const int32 PairIterations = 1, const T Thickness = (T)0);
+	virtual ~TPBDCollisionConstraint() {}
+
+	//
+	// Constraint Container API
+	//
+
+	void SetThickness(T InThickness)
+	{
+		MThickness = InThickness;
+	}
+
+	void SetVelocityInflation(float ScaleFactor)
+	{
+		CollisionVelocityInflation = ScaleFactor;
+	}
+
+	void SetVelocitySolveEnabled(bool bInEnableVelocitySolve)
+	{
+		bEnableVelocitySolve = bInEnableVelocitySolve;
+	}
+
+	void SetPushOutPairIterations(int32 InPairIterations)
+	{ 
+		MPairIterations = InPairIterations;
+	}
+
+	void SetCollisionsEnabled(bool bInEnableCollisions)
+	{
+		bEnableCollisions = bInEnableCollisions;
+	}
+
+	int32 NumConstraints() const
+	{
+		return Constraints.Num();
+	}
+
+	FConstraintHandleID GetConstraintHandleID(const FRigidBodyContactConstraint & Constraint) const
+	{
+		return  FConstraintHandleID( Constraint.Particle[0], Constraint.Particle[1] );
+	}
+
+	FHandles& GetConstraintHandles()
+	{
+		return Handles;
+	}
+	const FHandles& GetConstConstraintHandles() const
+	{
+		return Handles;
+	}
+
+	void RemoveConstraint(FConstraintContainerHandle*);
+
+	/**
+	* Remove the constraints associated with the ParticleHandle.
+	*/
+	void RemoveConstraints(const TSet<TGeometryParticleHandle<T, d>*>&  ParticleHandle);
 
 	/**
-	 * Manages a set of contact constraints:
-	 *	- Performs collision detection to generate constraints.
-	 *	- Responsible for applying corrections to particles affected by the constraints.
-	 * @todo(ccaulfield): rename to TPBDCollisionConstraints
-	 * @todo(ccaulfield): remove TPBDCollisionConstraintAccessor
-	 * @todo(ccaulfield): separate collision detection from constraint container
+	 * Apply a modifier to the constraints and specify which constraints should be disabled.
+	 * You would probably call this in the PostComputeCallback. Prefer this to calling RemoveConstraints in a loop, 
+	 * so you don't have to worry about constraint iterator/indices changing.
 	 */
-	template<typename T, int d>
-	class CHAOS_API TPBDCollisionConstraint : public FPBDConstraintContainer
-	{
-	public:
-		using Base = FPBDConstraintContainer;
-		friend class TPBDCollisionConstraintAccessor<T, d>;
-		friend class TPBDCollisionConstraintHandle<T, d>;
-		using FReal = T;
-		static const int Dimensions = d;
-		using FConstraintContainerHandle = TPBDCollisionConstraintHandle<T, d>;
-		using FConstraintHandleAllocator = TConstraintHandleAllocator<TPBDCollisionConstraint<T, d>>;
-		using FRigidBodyContactConstraint = TRigidBodyContactConstraint<T, d>;
-		using FConstraintHandleID = TPair<const TGeometryParticleHandle<T, d>*, const TGeometryParticleHandle<T, d>*>;
-
-		TPBDCollisionConstraint(const TPBDRigidsSOAs<T, d>& InParticles, TArrayCollectionArray<bool>& Collided, const TArrayCollectionArray<TSerializablePtr<FChaosPhysicsMaterial>>& PerParticleMaterials, const int32 PairIterations = 1, const T Thickness = (T)0);
-		virtual ~TPBDCollisionConstraint() {}
-
-		//
-		// Constraint Container API
-		//
-
-		void SetThickness(T InThickness)
-		{
-			MThickness = InThickness;
-		}
-
-		void SetVelocityInflation(float ScaleFactor)
-		{
-			CollisionVelocityInflation = ScaleFactor;
-		}
-
-		void SetVelocitySolveEnabled(bool bInEnableVelocitySolve)
-		{
-			bEnableVelocitySolve = bInEnableVelocitySolve;
-		}
-
-		void SetPushOutPairIterations(int32 InPairIterations)
-		{
-			MPairIterations = InPairIterations;
-		}
-
-		void SetCollisionsEnabled(bool bInEnableCollisions)
-		{
-			bEnableCollisions = bInEnableCollisions;
-		}
-
-		/**
-		 * Get the number of constraints.
-		 */
-		int32 NumConstraints() const
-		{
-			return Constraints.Num();
-		}
-
-		FConstraintHandleID GetConstraintHandleID(const FRigidBodyContactConstraint & Constraint) const
-		{
-			return  FConstraintHandleID(Constraint.Particle[0], Constraint.Particle[1]);
-		}
-
-		FConstraintHandleID GetConstraintHandleID(int32 ConstraintIndex) const
-		{
-			return  FConstraintHandleID(Constraints[ConstraintIndex].Particle[0], Constraints[ConstraintIndex].Particle[1]);
-		}
-
-		const FConstraintContainerHandle* GetConstraintHandle(int32 ConstraintIndex) const
-		{
-			return Handles[ConstraintIndex];
-		}
-
-		FConstraintContainerHandle* GetConstraintHandle(int32 ConstraintIndex)
-		{
-			return Handles[ConstraintIndex];
-		}
-
-		void RemoveConstraint(int32 Idx);
-
-		/**
-		* Remove the constraints associated with the ParticleHandle.
-		*/
-		void RemoveConstraints(const TSet<TGeometryParticleHandle<T, d>*>&  ParticleHandle);
-
-		/**
-		 * Apply a modifier to the constraints and specify which constraints should be disabled.
-		 * You would probably call this in the PostComputeCallback. Prefer this to calling RemoveConstraints in a loop,
-		 * so you don't have to worry about constraint iterator/indices changing.
-		 */
-		void ApplyCollisionModifier(const TFunction<ECollisionModifierResult(FRigidBodyContactConstraint& Constraint)>& CollisionModifier);
-
-		/**
-		 * Set the callback used just after contacts are generated at the start of a frame tick.
-		 * This can be used to modify or disable constraints (via ApplyCollisionModifier).
-		 */
-		void SetPostComputeCallback(const TRigidBodyContactConstraintsPostComputeCallback<T, d>& Callback);
-
-		/**
-		 * Remove the constraint callback.
-		 */
-		void ClearPostComputeCallback();
-
-		void SetPostApplyCallback(const TRigidBodyContactConstraintsPostApplyCallback<T, d>& Callback);
-		void ClearPostApplyCallback();
-
-		void SetPostApplyPushOutCallback(const TRigidBodyContactConstraintsPostApplyPushOutCallback<T, d>& Callback);
-		void ClearPostApplyPushOutCallback();
-
-		//
-		// Constraint API
-		//
-
-		/**
-		 * Get the particles that are affected by the specified constraint.
-		 */
-		TVector<TGeometryParticleHandle<T, d>*, 2> GetConstrainedParticles(int32 ConstraintIndex) const
-		{
-			return { Constraints[ConstraintIndex].Particle[0], Constraints[ConstraintIndex].Particle[1] };
-		}
-
-
-		//
-		// Island Rule API
-		//
-
-		/**
-		 * Apply corrections for the specified list of constraints.
-		 */
-		 // @todo(ccaulfield): this runs wide. The serial/parallel decision should be in the ConstraintRule
-		void Apply(const T Dt, const TArray<FConstraintContainerHandle*>& InConstraintHandles, const int32 It, const int32 NumIts);
-
-		/**
-		 * Generate all contact constraints.
-		 */
-		void UpdatePositionBasedState(/*const TPBDRigidParticles<T, d>& InParticles, const TArray<int32>& InIndices,*/ const T Dt);
-
-		/**
-		 * Apply push out for the specified list of constraints.
-		 * Return true if we need another iteration
-		 */
-		 // @todo(ccaulfield): this runs wide. The serial/parallel decision should be in the ConstraintRule
-		bool ApplyPushOut(const T Dt, const TArray<FConstraintContainerHandle*>& InConstraintHandles, const TSet<TGeometryParticleHandle<T, d>*>& IsTemporarilyStatic, int32 Iteration, int32 NumIterations);
-
-		void UpdateConstraints(T Dt, const TSet<TGeometryParticleHandle<T, d>*>& AddedParticles);
-
-		const TArray<FRigidBodyContactConstraint>& GetAllConstraints() const { return Constraints; }
-
-		//using FAccelerationStructure = TBoundingVolume<TAccelerationStructureHandle<T, d>, T, d>;
-		//using FAccelerationStructure = TAABBTree<TAccelerationStructureHandle<T, d>, TAABBTreeLeafArray<TAccelerationStructureHandle<T, d>, T>, T>;
-		using FAccelerationStructure = ISpatialAcceleration< TAccelerationStructureHandle<T, d>, T, d>;
-		void SetSpatialAcceleration(const FAccelerationStructure* InSpatialAcceleration) { SpatialAcceleration = InSpatialAcceleration; }
-
-		/**
-		 * Get the particles that are affected by the specified constraint.
-		 */
-		TVector<TGeometryParticleHandle<T, d>*, 2> ConstraintParticles(int32 ContactIndex) const { return { Constraints[ContactIndex].Particle[0], Constraints[ContactIndex].Particle[1] }; }
-
-		//NOTE: this should not be called by anyone other than ISpatialAccelerationCollection and CollisionConstraints - todo: make private with friends?
-		template <bool bGatherStats, typename SPATIAL_ACCELERATION>
-		void ComputeConstraintsHelperLowLevel(const SPATIAL_ACCELERATION& SpatialAcceleration, T Dt);
-
-	protected:
+	void ApplyCollisionModifier(const TFunction<ECollisionModifierResult(const FConstraintContainerHandle* Handle)>& CollisionModifier);
+
+	/**
+	 * Set the callback used just after contacts are generated at the start of a frame tick.
+	 * This can be used to modify or disable constraints (via ApplyCollisionModifier).
+	 */
+	void SetPostComputeCallback(const TRigidBodyContactConstraintsPostComputeCallback<T, d>& Callback);
+	void ClearPostComputeCallback();
+
+	void SetPostApplyCallback(const TRigidBodyContactConstraintsPostApplyCallback<T, d>& Callback);
+	void ClearPostApplyCallback();
+
+	void SetPostApplyPushOutCallback(const TRigidBodyContactConstraintsPostApplyPushOutCallback<T, d>& Callback);
+	void ClearPostApplyPushOutCallback();
+
+	//
+	// Island Rule API
+	//
+
+	/**
+	 * Apply corrections for the specified list of constraints. (Runs Wide!)
+	 *
+	 * @todo(ccaulfield): this runs wide. The serial/parallel decision should be in the ConstraintRule
+	 *
+	 */
+	void Apply(const T Dt, const TArray<FConstraintContainerHandle*>& InConstraintHandles, const int32 It, const int32 NumIts);
+
+	/**
+	 * Generate all contact constraints.
+	 */
+	void UpdatePositionBasedState(const T Dt);
+
+	/**
+	 * Update all constraint values
+	 */
+	void UpdateConstraints(T Dt, const TSet<TGeometryParticleHandle<T, d>*>& AddedParticles);
+
+	/**
+	 * Apply push out for the specified list of constraints.
+	 *  Return true if we need another iteration 
+	 *
+	 * @todo(ccaulfield): this runs wide. The serial/parallel decision should be in the ConstraintRule
+	 *
+	 */
+	bool ApplyPushOut(const T Dt, const TArray<FConstraintContainerHandle*>& InConstraintHandles, const TSet<const TGeometryParticleHandle<T,d>*>& IsTemporarilyStatic, int32 Iteration, int32 NumIterations);
+
+
+	const TArray<FConstraintContainerHandle*>& GetAllConstraintHandles() const { return Handles; }
+
+	//using FAccelerationStructure = TBoundingVolume<TAccelerationStructureHandle<T, d>, T, d>;
+	//using FAccelerationStructure = TAABBTree<TAccelerationStructureHandle<T, d>, TAABBTreeLeafArray<TAccelerationStructureHandle<T, d>, T>, T>;
+	using FAccelerationStructure = ISpatialAcceleration< TAccelerationStructureHandle<T, d>, T, d>;
+	void SetSpatialAcceleration(const FAccelerationStructure* InSpatialAcceleration) { SpatialAcceleration = InSpatialAcceleration; }
+
+	//NOTE: this should not be called by anyone other than ISpatialAccelerationCollection and CollisionConstraints - todo: make private with friends?
+	template <bool bGatherStats, typename SPATIAL_ACCELERATION>
+	void ComputeConstraintsHelperLowLevel(const SPATIAL_ACCELERATION& SpatialAcceleration, T Dt);
+	
+
+protected:
 		using Base::GetConstraintIndex;
 		using Base::SetConstraintIndex;
 
-	private:
-		void Reset(/*const TPBDRigidParticles<T, d>& InParticles, const TArray<int32>& InIndices*/);
+private:
+
+	/*
+	* Reset the constraints. 
+	*/
+	void Reset( );
+
+	void Apply(const T Dt, FRigidBodyContactConstraint& Constraint);
+
+	void ApplyPushOut(const T Dt, FRigidBodyContactConstraint& Constraint, const TSet<const TGeometryParticleHandle<T, d>*>& IsTemporarilyStatic, int32 Iteration, int32 NumIterations, bool &NeedsAnotherIteration);
 
-		void Apply(const T Dt, FRigidBodyContactConstraint& Constraint);
-		void ApplyPushOut(const T Dt, FRigidBodyContactConstraint& Constraint, const TSet<TGeometryParticleHandle<T, d>*>& IsTemporarilyStatic, int32 Iteration, int32 NumIterations, bool& NeedsAnotherIteration);
+	template <bool bGatherStats = false>
+	void ComputeConstraints(const FAccelerationStructure& AccelerationStructure, T Dt);
 
-		template <bool bGatherStats = false>
-		void ComputeConstraints(const FAccelerationStructure& AccelerationStructure, T Dt);
+	template<ECollisionUpdateType>
+	void UpdateConstraint(const T Thickness, FRigidBodyContactConstraint& Constraint);
 
-		template<ECollisionUpdateType>
-		void UpdateConstraint(const T Thickness, FRigidBodyContactConstraint& Constraint);
+	void ConstructConstraints(TGeometryParticleHandle<T, d>* Particle0, TGeometryParticleHandle<T, d>* Particle1, const T Thickness, TRigidBodyContactConstraint<T, d> & Constraint);
 
-		void ConstructConstraints(TGeometryParticleHandle<T, d>* Particle0, TGeometryParticleHandle<T, d>* Particle1, const T Thickness, TRigidBodyContactConstraint<T, d>& Constraint);
 
-		const TPBDRigidsSOAs<T, d>& Particles;
+	const TPBDRigidsSOAs<T,d>& Particles;
 
-		// @todo(ccaulfield): move spatial acceleration out of constraint container and make shareable
-		const FAccelerationStructure* SpatialAcceleration;
 
-		TArray<FRigidBodyContactConstraint> Constraints;
-		TArrayCollectionArray<bool>& MCollided;
-		const TArrayCollectionArray<TSerializablePtr<FChaosPhysicsMaterial>>& MPhysicsMaterials;
-		bool bEnableVelocitySolve;
-		int32 MPairIterations;
-		T MThickness;
-		T MAngularFriction;
-		bool bUseCCD;
-		bool bEnableCollisions;
+	TArray<FConstraintContainerHandle*> Handles;
+	FConstraintHandleAllocator HandleAllocator;
+	TMap< FConstraintHandleID, FConstraintContainerHandle* > Manifolds;
+	TArray<FRigidBodyContactConstraint> Constraints;
 
-		int32 LifespanCounter;
-		float CollisionVelocityInflation;
+	TArrayCollectionArray<bool>& MCollided;
+	const TArrayCollectionArray<TSerializablePtr<FChaosPhysicsMaterial>>& MPhysicsMaterials;
+	bool bEnableVelocitySolve;
+	int32 MPairIterations;
+	T MThickness;
+	T MAngularFriction;
+	bool bUseCCD;
+	bool bEnableCollisions;
 
-		TRigidBodyContactConstraintsPostComputeCallback<T, d> PostComputeCallback;
-		TRigidBodyContactConstraintsPostApplyCallback<T, d> PostApplyCallback;
-		TRigidBodyContactConstraintsPostApplyPushOutCallback<T, d> PostApplyPushOutCallback;
+	int32 LifespanCounter;
+	float CollisionVelocityInflation;
 
-		TArray<FConstraintContainerHandle*> Handles;
-		FConstraintHandleAllocator HandleAllocator;
+	TRigidBodyContactConstraintsPostComputeCallback<T, d> PostComputeCallback;
+	TRigidBodyContactConstraintsPostApplyCallback<T, d> PostApplyCallback;
+	TRigidBodyContactConstraintsPostApplyPushOutCallback<T, d> PostApplyPushOutCallback;
 
-		TMap< FConstraintHandleID, FConstraintContainerHandle* > Manifolds;
-	};
+	// @todo(ccaulfield): move spatial acceleration out of constraint container and make shareable
+	const FAccelerationStructure* SpatialAcceleration;
+};
 
-	extern template void TPBDCollisionConstraint<float, 3>::UpdateConstraint<ECollisionUpdateType::Any>(const float Thickness, FRigidBodyContactConstraint& Constraint);
-	extern template void TPBDCollisionConstraint<float, 3>::UpdateConstraint<ECollisionUpdateType::Deepest>(const float Thickness, FRigidBodyContactConstraint& Constraint);
-	extern template void TPBDCollisionConstraint<float, 3>::ComputeConstraints<false>(const TPBDCollisionConstraint<float, 3>::FAccelerationStructure&, float Dt);
-	extern template void TPBDCollisionConstraint<float, 3>::ComputeConstraints<true>(const TPBDCollisionConstraint<float, 3>::FAccelerationStructure&, float Dt);
+extern template void TPBDCollisionConstraint<float, 3>::UpdateConstraint<ECollisionUpdateType::Any>(const float Thickness, FRigidBodyContactConstraint& Constraint);
+extern template void TPBDCollisionConstraint<float, 3>::UpdateConstraint<ECollisionUpdateType::Deepest>(const float Thickness, FRigidBodyContactConstraint& Constraint);
+extern template void TPBDCollisionConstraint<float, 3>::ComputeConstraints<false>(const TPBDCollisionConstraint<float, 3>::FAccelerationStructure&, float Dt);
+extern template void TPBDCollisionConstraint<float, 3>::ComputeConstraints<true>(const TPBDCollisionConstraint<float, 3>::FAccelerationStructure&, float Dt);
 }
diff --git a/Engine/Source/Runtime/Experimental/Chaos/Public/Chaos/PBDCollisionConstraintImp.h b/Engine/Source/Runtime/Experimental/Chaos/Public/Chaos/PBDCollisionConstraintImp.h
--- a/Engine/Source/Runtime/Experimental/Chaos/Public/Chaos/PBDCollisionConstraintImp.h
+++ b/Engine/Source/Runtime/Experimental/Chaos/Public/Chaos/PBDCollisionConstraintImp.h
@@ -293,26 +293,12 @@ namespace Chaos
 				while (!Queue.IsEmpty())
 				{
 					int32 Idx = Constraints.AddUninitialized(1);
-					FConstraintContainerHandle* Handle = HandleAllocator.AllocHandle(this, Idx);
+					FConstraintContainerHandle* Handle = HandleAllocator.template AllocHandle< TRigidBodyContactConstraint<T,d> >(this, Idx );
 
 					Handles.Add(Handle);
 					Queue.Dequeue(Constraints[Idx]);
 					ensure(Handles.Num() == Constraints.Num());
 
-					/*
-					FConstraintHandleID HandleID = GetConstraintHandleID(Constraints[Idx]);
-					FConstraintManifold* Manifold;
-					if (auto ValPtr = Manifolds.Find(HandleID))
-					{
-						Manifold = *ValPtr;
-					}
-					else
-					{
-						Manifold = new FConstraintManifold(Constraints[Idx].Particle->X(), Constraints[Idx].Particle->R(), LifespanCounter + CollisionManifoldLifespan);
-						Manifolds.Add(HandleID, Manifold);
-					}
-					Manifold->AddHandle(Handle);
-					*/
 				}
 				LifespanCounter++;
 			}
diff --git a/Engine/Source/Runtime/Experimental/Chaos/Public/Chaos/PBDCollisionTypes.h b/Engine/Source/Runtime/Experimental/Chaos/Public/Chaos/PBDCollisionTypes.h
--- a/Engine/Source/Runtime/Experimental/Chaos/Public/Chaos/PBDCollisionTypes.h
+++ b/Engine/Source/Runtime/Experimental/Chaos/Public/Chaos/PBDCollisionTypes.h
@@ -14,6 +14,11 @@ namespace Chaos
 	template<typename T, int d>
 	class TPBDCollisionConstraint;
 
+	template<typename  T, int d>
+	class TPBDCollisionConstraintHandle;
+
+	template<typename T, int d>
+	class TRigidBodyContactConstraint;
 
 	/** Specifies the type of work we should do*/
 	enum class CHAOS_API ECollisionUpdateType
@@ -38,16 +43,13 @@ namespace Chaos
 	*
 	*/
 	template<class T, int d>
-	class TContactData
+	class CHAOS_API TContactData
 	{
 	public:
 
 		TContactData() : bDisabled(true), Normal(0), Phi(FLT_MAX) {};
 
-		//void AddPoint(TVector<T, d> Point, TVector<T, d> Normal){}
-
 		bool bDisabled;
-		//int NumContacts;
 		TVector<T, d> Normal;
 		TVector<T, d> Location;
 		T Phi;
@@ -57,7 +59,7 @@ namespace Chaos
 	*
 	*/
 	template<class T, int d>
-	class TConvexManifold
+	class CHAOS_API TConvexManifold
 	{
 	public:
 		using FContactData = TContactData<T, d>;
@@ -75,13 +77,36 @@ namespace Chaos
 	};
 
 
+	template<class T, int d>
+	class CHAOS_API TCollisionConstraintBase
+	{
+	public:
+
+		enum class FType
+		{
+			None = 0,     // default value also indicates a invalid constraint
+			SinglePoint   // TRigidBodyContactConstraint
+		};
+
+		TCollisionConstraintBase() : Type(FType::None) {}
+		TCollisionConstraintBase(FType InType) : Type(InType) {}
+		FType GetType() const { return Type; }
+
+		template<class AS_T> AS_T * As() {return static_cast<AS_T>(this); }
+		template<class AS_T> const AS_T * As() const { return static_cast<AS_T>(this); }
+
+	private:
+		FType Type;
+	};
+
 	/*
 	*
 	*/
 	template<class T, int d>
-	class TRigidBodyContactConstraint
+	class CHAOS_API TRigidBodyContactConstraint : public TCollisionConstraintBase<T,d>
 	{
 	public:
+		using Base = TCollisionConstraintBase<T, d>;
 		using FGeometryParticleHandle = TGeometryParticleHandle<T, d>;
 		using FConvexManifold = TConvexManifold<T, d>;
 
@@ -115,11 +140,13 @@ namespace Chaos
 			return FString::Printf(TEXT("Particle:%s, Levelset:%s, AccumulatedImpulse:%s"), *Particle[0]->ToString(), *Particle[1]->ToString(), *AccumulatedImpulse.ToString());
 		}
 
+		static typename Base::FType StaticType() { return Base::FType::SinglePoint; };
+
 	};
 
 
 	template<class T, int d>
-	struct TRigidBodyContactConstraintPGS
+	struct CHAOS_API TRigidBodyContactConstraintPGS
 	{
 		TRigidBodyContactConstraintPGS() : AccumulatedImpulse(0.f) {}
 		TGeometryParticleHandle<T, d>* Particle;
@@ -137,35 +164,89 @@ namespace Chaos
 	public:
 		using Base = TContainerConstraintHandle<TPBDCollisionConstraint<T, d>>;
 		using FConstraintContainer = TPBDCollisionConstraint<T, d>;
+		using FConstraintBase = TCollisionConstraintBase<T, d>;
+
 
-		TPBDCollisionConstraintHandle()
+		TPBDCollisionConstraintHandle() 
+			: ConstraintType(FConstraintBase::FType::None)
 		{}
 
-		TPBDCollisionConstraintHandle(FConstraintContainer* InConstraintContainer, int32 InConstraintIndex)
+		TPBDCollisionConstraintHandle(FConstraintContainer* InConstraintContainer, int32 InConstraintIndex, typename FConstraintBase::FType InType)
 			: TContainerConstraintHandle<TPBDCollisionConstraint<T, d>>(InConstraintContainer, InConstraintIndex)
+			, ConstraintType(InType)
 		{
 		}
 
-		TRigidBodyContactConstraint<T, d>& GetContact()
+		// Handle API
+
+		template<class AS_T> CHAOS_API AS_T& GetContact();
+		template<class AS_T> CHAOS_API const AS_T& GetContact() const;
+
+		typename FConstraintBase::FType GetType() const { return ConstraintType; }
+
+		void SetConstraintIndex(int32 IndexIn, typename FConstraintBase::FType InType)
 		{
-			return ConstraintContainer->Constraints[ConstraintIndex];
+			ConstraintIndex = IndexIn;
+			ConstraintType = InType;
 		}
 
-		const TRigidBodyContactConstraint<T, d>& GetContact() const
+		// Contact API
+
+		TVector<T, d> GetContactLocation() const
 		{
-			return ConstraintContainer->Constraints[ConstraintIndex];
+			if (ConstraintType == FConstraintBase::FType::SinglePoint)
+			{
+				return GetContact<TRigidBodyContactConstraint<float, 3> >().GetLocation();
+			}
+			ensure(false);
+			return TVector<T, d>(0);
 		}
 
-		void SetConstraintIndex(int32 IndexIn)
+		TVector<T, d> GetAccumulatedImpulse() const
 		{
-			ConstraintIndex = IndexIn;
+			if (ConstraintType == FConstraintBase::FType::SinglePoint)
+			{
+				return GetContact<TRigidBodyContactConstraint<float, 3> >().AccumulatedImpulse;
+			}
+			ensure(false);
+
+			return TVector<T, d>(0);
+		}
+
+
+		TVector<const TGeometryParticleHandle<T, d>*, 2> GetConstrainedParticles() const
+		{
+			if (ConstraintType == FConstraintBase::FType::SinglePoint)
+			{
+				const TRigidBodyContactConstraint<float, 3>& Contact = GetContact< TRigidBodyContactConstraint<float, 3> >();
+				return { Contact.Particle[0], Contact.Particle[1] };
+			}
+			ensure(false);
+			return { nullptr,nullptr };
+		}
+
+
+
+		TVector<TGeometryParticleHandle<T, d>*, 2> GetConstrainedParticles()
+		{
+			if (ConstraintType == FConstraintBase::FType::SinglePoint)
+			{
+				return { GetContact< TRigidBodyContactConstraint<float, 3> >().Particle[0], GetContact< TRigidBodyContactConstraint<float, 3> >().Particle[1] };
+			}
+			ensure(false);
+			return { nullptr,nullptr };
 		}
 
+
 	protected:
+		typename FConstraintBase::FType ConstraintType;
 		using Base::ConstraintIndex;
 		using Base::ConstraintContainer;
 
 
 	};
 
+	extern template TRigidBodyContactConstraint<float, 3>& TPBDCollisionConstraintHandle<float, 3>::GetContact<TRigidBodyContactConstraint<float, 3>>();
+	extern template const TRigidBodyContactConstraint<float, 3>& TPBDCollisionConstraintHandle<float, 3>::GetContact<TRigidBodyContactConstraint<float, 3>>() const;
+
 }
diff --git a/Engine/Source/Runtime/Experimental/Chaos/Public/Chaos/PBDConstraintRule.h b/Engine/Source/Runtime/Experimental/Chaos/Public/Chaos/PBDConstraintRule.h
--- a/Engine/Source/Runtime/Experimental/Chaos/Public/Chaos/PBDConstraintRule.h
+++ b/Engine/Source/Runtime/Experimental/Chaos/Public/Chaos/PBDConstraintRule.h
@@ -86,34 +86,15 @@ namespace Chaos
 	public:
 		typedef T_CONSTRAINTS FConstraints;
 
-		TPBDConstraintGraphRuleImpl(FConstraints& InConstraints, int32 InPriority)
-			: FPBDConstraintGraphRule(InPriority)
-			, Constraints(InConstraints)
-			, ConstraintGraph(nullptr)
-		{
-		}
+		TPBDConstraintGraphRuleImpl(FConstraints& InConstraints, int32 InPriority);
 
-		virtual void BindToGraph(FPBDConstraintGraph& InContactGraph, uint32 InContainerId) override
-		{
-			ConstraintGraph = &InContactGraph;
-			ContainerId = InContainerId;
-		}
+		virtual void BindToGraph(FPBDConstraintGraph& InContactGraph, uint32 InContainerId) override;
 
-		virtual void UpdatePositionBasedState(const FReal Dt) override
-		{
-			Constraints.UpdatePositionBasedState(Dt);
-		}
+		virtual void UpdatePositionBasedState(const FReal Dt) override;
 
-		virtual void AddToGraph() override
-		{
-			ConstraintGraph->ReserveConstraints(Constraints.NumConstraints());
-			for (int32 ConstraintIndex = 0; ConstraintIndex < Constraints.NumConstraints(); ++ConstraintIndex)
-			{
-				ConstraintGraph->AddConstraint(ContainerId, Constraints.GetConstraintHandle(ConstraintIndex), Constraints.GetConstrainedParticles(ConstraintIndex));
-			}
-		}
+		virtual void AddToGraph() override;
 
-		virtual int32 NumConstraints() const override { return Constraints.NumConstraints(); }
+		virtual int32 NumConstraints() const override;
 
 	protected:
 		FConstraints& Constraints;
@@ -259,7 +240,7 @@ namespace Chaos
 			int32 MaxColor = GraphColor.GetIslandMaxColor(Island);
 			int32 MaxLevel = GraphColor.GetIslandMaxLevel(Island);
 
-			TSet<TGeometryParticleHandle<FReal, 3>*> IsTemporarilyStatic;
+			TSet<const TGeometryParticleHandle<FReal, 3>*> IsTemporarilyStatic;
 			bool bNeedsAnotherIteration = false;
 			for (int32 Level = 0; Level <= MaxLevel; ++Level)
 			{
@@ -281,10 +262,12 @@ namespace Chaos
 				{
 					if (LevelToColorToConstraintListMap[Level].Contains(Color))
 					{
-						for (int32 Edge = 0; Edge < LevelToColorToConstraintListMap[Level][Color].Num(); ++Edge)
+						const TArray<typename FConstraints::FConstraintContainerHandle*>& ConstraintHandles = GetLevelColorConstraints(LevelToColorToConstraintListMap, Level, Color);
+
+						for (int32 Edge = 0; Edge < ConstraintHandles.Num(); ++Edge)
 						{
-							const int32 ConstraintIndex = LevelToColorToConstraintListMap[Level][Color][Edge]->GetConstraintIndex();
-							const TVector<TGeometryParticleHandle<FReal, 3>*, 2> Particles = Constraints.ConstraintParticles(ConstraintIndex);
+							const typename FConstraints::FConstraintContainerHandle* Handle = ConstraintHandles[Edge];
+							TVector<const TGeometryParticleHandle<FReal, 3>*, 2> Particles = Handle->GetConstrainedParticles();
 							if (It == NumIts - 1)
 							{
 								if (Particles[0]->AsDynamic() == nullptr || IsTemporarilyStatic.Contains(Particles[0]))
diff --git a/Engine/Source/Runtime/Experimental/Chaos/Public/Chaos/PBDJointConstraints.h b/Engine/Source/Runtime/Experimental/Chaos/Public/Chaos/PBDJointConstraints.h
--- a/Engine/Source/Runtime/Experimental/Chaos/Public/Chaos/PBDJointConstraints.h
+++ b/Engine/Source/Runtime/Experimental/Chaos/Public/Chaos/PBDJointConstraints.h
@@ -15,6 +15,7 @@
 
 namespace Chaos
 {
+
 	class FPBDJointConstraintHandle : public TContainerConstraintHandle<FPBDJointConstraints>
 	{
 	public:
@@ -28,6 +29,8 @@ namespace Chaos
 		CHAOS_API void SetParticleLevels(const TVector<int32, 2>& ParticleLevels);
 		CHAOS_API int32 GetConstraintLevel() const;
 		CHAOS_API const FPBDJointSettings& GetSettings() const;
+		CHAOS_API TVector<TGeometryParticleHandle<float,3>*, 2> GetConstrainedParticles() const;
+
 	protected:
 		using Base::ConstraintIndex;
 		using Base::ConstraintContainer;
@@ -56,6 +59,7 @@ namespace Chaos
 		using FParticlePair = TVector<TGeometryParticleHandle<FReal, 3>*, 2>;
 		using FVectorPair = TVector<FVec3, 2>;
 		using FTransformPair = TVector<FRigidTransform3, 2>;
+		using FHandles = TArray<FConstraintContainerHandle*>;
 
 		FPBDJointConstraints(const FPBDJointSolverSettings& InSettings = FPBDJointSolverSettings());
 
@@ -97,6 +101,14 @@ namespace Chaos
 		//
 		// Constraint API
 		//
+		FHandles& GetConstraintHandles()
+		{
+			return Handles;
+		}
+		const FHandles& GetConstConstraintHandles() const
+		{
+			return Handles;
+		}
 
 		const FConstraintContainerHandle* GetConstraintHandle(int32 ConstraintIndex) const;
 		FConstraintContainerHandle* GetConstraintHandle(int32 ConstraintIndex);
@@ -142,7 +154,7 @@ namespace Chaos
 		TArray<FParticlePair> ConstraintParticles;
 		TArray<FPBDJointState> ConstraintStates;
 
-		TArray<FConstraintContainerHandle*> Handles;
+		FHandles Handles;
 		FConstraintHandleAllocator HandleAllocator;
 
 		FJointPreApplyCallback PreApplyCallback;
diff --git a/Engine/Source/Runtime/Experimental/Chaos/Public/Chaos/PBDPositionConstraints.h b/Engine/Source/Runtime/Experimental/Chaos/Public/Chaos/PBDPositionConstraints.h
--- a/Engine/Source/Runtime/Experimental/Chaos/Public/Chaos/PBDPositionConstraints.h
+++ b/Engine/Source/Runtime/Experimental/Chaos/Public/Chaos/PBDPositionConstraints.h
@@ -17,9 +17,11 @@ namespace Chaos
 	public:
 		using Base = TContainerConstraintHandle<TPBDPositionConstraints<T, d>>;
 		using FConstraintContainer = TPBDPositionConstraints<T, d>;
+		using FGeometryParticleHandle = TGeometryParticleHandle<T, d>;
 
 		TPBDPositionConstraintHandle() {}
 		TPBDPositionConstraintHandle(FConstraintContainer* InConstraintContainer, int32 InConstraintIndex) : TContainerConstraintHandle<TPBDPositionConstraints<T, d>>(InConstraintContainer, InConstraintIndex) {}
+		TVector<FGeometryParticleHandle*, 2> GetConstrainedParticles() const { return ConstraintContainer->GetConstrainedParticles(ConstraintIndex); }
 
 	protected:
 		using Base::ConstraintIndex;
@@ -35,6 +37,7 @@ namespace Chaos
 		static const int Dimensions = d;
 		using FConstraintContainerHandle = TPBDPositionConstraintHandle<FReal, Dimensions>;
 		using FConstraintHandleAllocator = TConstraintHandleAllocator<TPBDPositionConstraints<FReal, Dimensions>>;
+		using FHandles = TArray<FConstraintContainerHandle*>;
 
 		TPBDPositionConstraints(const T InStiffness = (T)1)
 			: Stiffness(InStiffness)
@@ -114,6 +117,14 @@ namespace Chaos
 		//
 		// Constraint API
 		//
+		FHandles& GetConstraintHandles()
+		{
+			return Handles;
+		}
+		const FHandles& GetConstConstraintHandles() const
+		{
+			return Handles;
+		}
 
 		const FConstraintContainerHandle* GetConstraintHandle(int32 ConstraintIndex) const
 		{
@@ -156,13 +167,7 @@ namespace Chaos
 		{
 		}
 
-		void Apply(const T Dt, const TArray<FConstraintContainerHandle*>& ConstraintHandles, const int32 It, const int32 NumIts) const
-		{
-			for (FConstraintContainerHandle* ConstraintHandle : ConstraintHandles)
-			{
-				ApplySingle(Dt, ConstraintHandle->GetConstraintIndex());
-			}
-		}
+		void Apply(const T Dt, const TArray<FConstraintContainerHandle*>& ConstraintHandles, const int32 It, const int32 NumIts) const;
 
 		bool ApplyPushOut(const T Dt, const TArray<FConstraintContainerHandle*>& InConstraintIndices, const int32 It, const int32 NumIts) const
 		{
@@ -189,7 +194,7 @@ namespace Chaos
 		TArray<TPBDRigidParticleHandle<T,d>*> ConstrainedParticles;
 		T Stiffness;
 
-		TArray<FConstraintContainerHandle*> Handles;
+		FHandles Handles;
 		FConstraintHandleAllocator HandleAllocator;
 	};
 }
diff --git a/Engine/Source/Runtime/Experimental/Chaos/Public/Chaos/PBDRigidClustering.h b/Engine/Source/Runtime/Experimental/Chaos/Public/Chaos/PBDRigidClustering.h
--- a/Engine/Source/Runtime/Experimental/Chaos/Public/Chaos/PBDRigidClustering.h
+++ b/Engine/Source/Runtime/Experimental/Chaos/Public/Chaos/PBDRigidClustering.h
@@ -37,7 +37,8 @@ class CHAOS_API TPBDRigidClustering
 	typedef typename FPBDCollisionConstraint::FRigidBodyContactConstraint FRigidBodyContactConstraint;
 
 public:
-	typedef TMap<TGeometryParticleHandle<T,d>*, TArray<uint32> > FClusterMap;
+	typedef TMap<const TGeometryParticleHandle<T,d>*, TArray<uint32> > FClusterMap;
+	using FCollisionConstraintHandle = TPBDCollisionConstraintHandle<T, d>;
 
 	TPBDRigidClustering(FPBDRigidEvolution& InEvolution, TPBDRigidClusteredParticles<T, d>& InParticles);
 	~TPBDRigidClustering();
diff --git a/Engine/Source/Runtime/Experimental/Chaos/Public/Chaos/PBDRigidDynamicSpringConstraints.h b/Engine/Source/Runtime/Experimental/Chaos/Public/Chaos/PBDRigidDynamicSpringConstraints.h
--- a/Engine/Source/Runtime/Experimental/Chaos/Public/Chaos/PBDRigidDynamicSpringConstraints.h
+++ b/Engine/Source/Runtime/Experimental/Chaos/Public/Chaos/PBDRigidDynamicSpringConstraints.h
@@ -20,6 +20,7 @@ namespace Chaos
 
 		TPBDRigidDynamicSpringConstraintHandle() {}
 		TPBDRigidDynamicSpringConstraintHandle(FConstraintContainer* InConstraintContainer, int32 InConstraintIndex) : TContainerConstraintHandle<TPBDRigidDynamicSpringConstraints<T, d>>(InConstraintContainer, InConstraintIndex) {}
+		TVector<TGeometryParticleHandle<T, d>*, 2> GetConstrainedParticles() const;
 
 	protected:
 		using Base::ConstraintIndex;
@@ -36,6 +37,7 @@ namespace Chaos
 		static const int Dimensions = d;
 		using FConstraintContainerHandle = TPBDRigidDynamicSpringConstraintHandle<FReal, Dimensions>;
 		using FConstraintHandleAllocator = TConstraintHandleAllocator<TPBDRigidDynamicSpringConstraints<FReal, Dimensions>>;
+		using FHandles = TArray<FConstraintContainerHandle*>;
 
 		TPBDRigidDynamicSpringConstraints(const T InStiffness = (T)1)
 			: CreationThreshold(1), MaxSprings(1), Stiffness(InStiffness) 
@@ -136,6 +138,15 @@ namespace Chaos
 		//
 		// Constraint API
 		//
+		FHandles& GetConstraintHandles()
+		{
+			return Handles;
+		}
+		const FHandles& GetConstConstraintHandles() const
+		{
+			return Handles;
+		}
+
 
 		const FConstraintContainerHandle* GetConstraintHandle(int32 ConstraintIndex) const
 		{
@@ -147,6 +158,7 @@ namespace Chaos
 			return Handles[ConstraintIndex];
 		}
 
+
 		/**
 		 * Get the particles that are affected by the specified constraint.
 		 */
@@ -156,6 +168,7 @@ namespace Chaos
 		}
 
 
+
 		//
 		// Island Rule API
 		//
@@ -191,7 +204,7 @@ namespace Chaos
 		int32 MaxSprings;
 		T Stiffness;
 
-		TArray<FConstraintContainerHandle*> Handles;
+		FHandles Handles;
 		FConstraintHandleAllocator HandleAllocator;
 	};
 }
diff --git a/Engine/Source/Runtime/Experimental/Chaos/Public/Chaos/PBDRigidSpringConstraints.h b/Engine/Source/Runtime/Experimental/Chaos/Public/Chaos/PBDRigidSpringConstraints.h
--- a/Engine/Source/Runtime/Experimental/Chaos/Public/Chaos/PBDRigidSpringConstraints.h
+++ b/Engine/Source/Runtime/Experimental/Chaos/Public/Chaos/PBDRigidSpringConstraints.h
@@ -17,12 +17,14 @@ namespace Chaos
 	public:
 		using Base = TContainerConstraintHandle<TPBDRigidSpringConstraints<T, d>>;
 		using FConstraintContainer = TPBDRigidSpringConstraints<T, d>;
+		using FGeometryParticleHandle = TGeometryParticleHandle<T, d>;
 
 		TPBDRigidSpringConstraintHandle() {}
 		TPBDRigidSpringConstraintHandle(FConstraintContainer* InConstraintContainer, int32 InConstraintIndex) : TContainerConstraintHandle<TPBDRigidSpringConstraints<T, d>>(InConstraintContainer, InConstraintIndex) {}
 
 		const TVector<TVector<T, 3>, 2>& GetConstraintPositions() const;
 		void SetConstraintPositions(const TVector<TVector<T, 3>, 2>& ConstraintPositions);
+		TVector<FGeometryParticleHandle*, 2> GetConstrainedParticles() const { return ConstraintContainer->GetConstrainedParticles(ConstraintIndex); }
 
 	protected:
 		using Base::ConstraintIndex;
@@ -40,6 +42,7 @@ namespace Chaos
 		using FConstraintContainerHandle = TPBDRigidSpringConstraintHandle<FReal, Dimensions>;
 		using FConstraintHandleAllocator = TConstraintHandleAllocator<TPBDRigidSpringConstraints<FReal, Dimensions>>;
 		using FConstrainedParticlePair = TVector<TGeometryParticleHandle<T, d>*, 2>;
+		using FHandles = TArray<FConstraintContainerHandle*>;
 
 		TPBDRigidSpringConstraints(const T InStiffness = (T)1)
 			: Stiffness(InStiffness) 
@@ -112,7 +115,8 @@ namespace Chaos
 			// Update the handle for the constraint that was moved
 			if (ConstraintIndex < Handles.Num())
 			{
-				SetConstraintIndex(Handles[ConstraintIndex], ConstraintIndex);
+				FConstraintHandle* Handle = Handles[ConstraintIndex];
+				SetConstraintIndex(Handle, ConstraintIndex);
 			}
 		}
 
@@ -124,6 +128,14 @@ namespace Chaos
 		//
 		// Constraint API
 		//
+		FHandles& GetConstraintHandles()
+		{
+			return Handles;
+		}
+		const FHandles& GetConstConstraintHandles() const
+		{
+			return Handles;
+		}
 
 		const FConstraintContainerHandle* GetConstraintHandle(int32 ConstraintIndex) const
 		{
@@ -168,13 +180,7 @@ namespace Chaos
 		{
 		}
 
-		void Apply(const T Dt, const TArray<FConstraintContainerHandle*>& InConstraintHandles, const int32 It, const int32 NumIts)
-		{
-			for (FConstraintContainerHandle* ConstraintHandle : InConstraintHandles)
-			{
-				ApplySingle(Dt, ConstraintHandle->GetConstraintIndex());
-			}
-		}
+		void Apply(const T Dt, const TArray<FConstraintContainerHandle*>& InConstraintHandles, const int32 It, const int32 NumIts);
 
 		bool ApplyPushOut(const T Dt, const TArray<FConstraintContainerHandle*>& InConstraintHandles, const int32 It, const int32 NumIts)
 		{
diff --git a/Engine/Source/Runtime/Experimental/ChaosSolvers/Private/EventDefaults.cpp b/Engine/Source/Runtime/Experimental/ChaosSolvers/Private/EventDefaults.cpp
--- a/Engine/Source/Runtime/Experimental/ChaosSolvers/Private/EventDefaults.cpp
+++ b/Engine/Source/Runtime/Experimental/ChaosSolvers/Private/EventDefaults.cpp
@@ -6,6 +6,7 @@
 #include "Chaos/PBDRigidsEvolution.h"
 #include "Chaos/PBDRigidClustering.h"
 #include "Chaos/PBDCollisionConstraint.h"
+#include "Chaos/PBDCollisionTypes.h"
 #include "PBDRigidsSolver.h"
 #include "PhysicsProxy/SkeletalMeshPhysicsProxy.h"
 #include "PhysicsProxy/StaticMeshPhysicsProxy.h"
@@ -47,7 +48,7 @@ namespace Chaos
 			const FPBDRigidsSolver::FPBDRigidsEvolution* Evolution = Solver->GetEvolution();
 
 			const FPhysicsSolver::FPBDCollisionConstraints& CollisionRule = Evolution->GetCollisionConstraints();
-			const TArray<FPhysicsSolver::FPBDCollisionConstraints::FRigidBodyContactConstraint>& AllConstraintsArray = CollisionRule.GetAllConstraints();
+
 
 			const TPBDRigidParticles<float, 3>& Particles = Evolution->GetParticles().GetDynamicParticles();
 			const TArrayCollectionArray<ClusterId>& ClusterIdsArray = Evolution->GetRigidClustering().GetClusterIdsArray();
@@ -56,54 +57,57 @@ namespace Chaos
 #endif
 			const Chaos::TPBDRigidClustering<FPhysicsSolver::FPBDRigidsEvolution, FPhysicsSolver::FPBDCollisionConstraints, float, 3>::FClusterMap& ParentToChildrenMap = Evolution->GetRigidClustering().GetChildrenMap();
 
-			if(AllConstraintsArray.Num() > 0)
+			if(CollisionRule.NumConstraints() > 0)
 			{
 				// Get the number of valid constraints (AccumulatedImpulse != 0.f and Phi < 0.f) from AllConstraintsArray
-				TArray<int32> ValidCollisionIndices;
-				ValidCollisionIndices.SetNumUninitialized(AllConstraintsArray.Num());
+				TArray<const Chaos::TPBDCollisionConstraintHandle<float, 3> *> ValidCollisionHandles;
+				ValidCollisionHandles.SetNumUninitialized(CollisionRule.NumConstraints());
 				int32 NumValidCollisions = 0;
 
-				for(int32 Idx = 0; Idx < AllConstraintsArray.Num(); ++Idx)
+				for (const Chaos::TPBDCollisionConstraintHandle<float, 3> * ContactHandle : CollisionRule.GetConstConstraintHandles())
 				{
-					const TRigidBodyContactConstraint<float, 3>& Constraint = AllConstraintsArray[Idx];
+					if (ContactHandle->GetType() == TPBDCollisionConstraintHandle<float, 3>::FConstraintBase::FType::SinglePoint)
+					{
+						const TRigidBodyContactConstraint<float, 3>& Constraint = ContactHandle->GetContact<TRigidBodyContactConstraint<float, 3>>();
 
-					// Since Clustered GCs can be unioned the particleIndex representing the union 
-					// is not associated with a PhysicsProxy
-					if(Constraint.Particle[0]->Handle()->GTGeometryParticle()->Proxy != nullptr)
+						// Since Clustered GCs can be unioned the particleIndex representing the union 
+						// is not associated with a PhysicsProxy
+						if (Constraint.Particle[0]->Handle()->GTGeometryParticle()->Proxy != nullptr)
 
-					{
-						if(ensure(!Constraint.AccumulatedImpulse.ContainsNaN() && FMath::IsFinite(Constraint.GetPhi())))
 						{
-							TGeometryParticleHandle<float, 3>* Particle0 = Constraint.Particle[0];
-							TGeometryParticleHandle<float, 3>* Particle1 = Constraint.Particle[1];
-							TKinematicGeometryParticleHandle<float, 3>* Body0 = Particle0->AsKinematic();
+							if (ensure(!Constraint.AccumulatedImpulse.ContainsNaN() && FMath::IsFinite(Constraint.GetPhi())))
+							{
+								TGeometryParticleHandle<float, 3>* Particle0 = Constraint.Particle[0];
+								TGeometryParticleHandle<float, 3>* Particle1 = Constraint.Particle[1];
+								TKinematicGeometryParticleHandle<float, 3>* Body0 = Particle0->AsKinematic();
 
-							// presently when a rigidbody or kinematic hits static geometry then Body1 is null
-							TKinematicGeometryParticleHandle<float, 3>* Body1 = Particle1->AsKinematic();
+								// presently when a rigidbody or kinematic hits static geometry then Body1 is null
+								TKinematicGeometryParticleHandle<float, 3>* Body1 = Particle1->AsKinematic();
 
-							if(!Constraint.AccumulatedImpulse.IsZero() && Body0)
-							{
-								if(ensure(!Constraint.GetLocation().ContainsNaN() &&
-									!Constraint.GetNormal().ContainsNaN()) &&
-									!Body0->V().ContainsNaN() &&
-									!Body0->W().ContainsNaN() &&
-									(Body1 == nullptr || ((!Body1->V().ContainsNaN()) && !Body1->W().ContainsNaN())))
+								if (!Constraint.AccumulatedImpulse.IsZero() && Body0)
 								{
-									ValidCollisionIndices[NumValidCollisions] = Idx;
-									NumValidCollisions++;
+									if (ensure(!Constraint.GetLocation().ContainsNaN() &&
+										!Constraint.GetNormal().ContainsNaN()) &&
+										!Body0->V().ContainsNaN() &&
+										!Body0->W().ContainsNaN() &&
+										(Body1 == nullptr || ((!Body1->V().ContainsNaN()) && !Body1->W().ContainsNaN())))
+									{
+										ValidCollisionHandles[NumValidCollisions] = ContactHandle;
+										NumValidCollisions++;
+									}
 								}
 							}
 						}
 					}
 				}
 
-				ValidCollisionIndices.SetNum(NumValidCollisions);
+				ValidCollisionHandles.SetNum(NumValidCollisions);
 
-				if(ValidCollisionIndices.Num() > 0)
+				if(ValidCollisionHandles.Num() > 0)
 				{
-					for(int32 IdxCollision = 0; IdxCollision < ValidCollisionIndices.Num(); ++IdxCollision)
+					for(int32 IdxCollision = 0; IdxCollision < ValidCollisionHandles.Num(); ++IdxCollision)
 					{
-						Chaos::TPBDCollisionConstraint<float, 3>::FRigidBodyContactConstraint const& Constraint = AllConstraintsArray[ValidCollisionIndices[IdxCollision]];
+						Chaos::TPBDCollisionConstraint<float, 3>::FRigidBodyContactConstraint const& Constraint = ValidCollisionHandles[IdxCollision]->GetContact<TRigidBodyContactConstraint<float, 3>>();
 
 						TGeometryParticleHandle<float, 3>* Particle0 = Constraint.Particle[0];
 						TGeometryParticleHandle<float, 3>* Particle1 = Constraint.Particle[1];
@@ -174,7 +178,6 @@ namespace Chaos
 					}
 				}
 			}
-
 		});
 	}
 
