Chaos
- pair iterations for joint solver
- pair iterations for collision solver
- maintain accurate velocity in joint Apply step (required to work properly collisions)
- dense matrix for linear equations solving (runtime dimensioned with compile-time maximium)
- Cholesky factorizing and linear equation solver for dense matrix
- dense matrix units tests
- dense matrix natvis
- transient handle natvis

diff --git a/Engine/Source/Runtime/AnimGraphRuntime/Private/BoneControllers/AnimNode_RigidBody.cpp b/Engine/Source/Runtime/AnimGraphRuntime/Private/BoneControllers/AnimNode_RigidBody.cpp
--- a/Engine/Source/Runtime/AnimGraphRuntime/Private/BoneControllers/AnimNode_RigidBody.cpp
+++ b/Engine/Source/Runtime/AnimGraphRuntime/Private/BoneControllers/AnimNode_RigidBody.cpp
@@ -29,7 +29,9 @@ TAutoConsoleVariable<int32> CVarEnableRigidBodyNode(TEXT("p.RigidBodyNode"), 1,
 TAutoConsoleVariable<int32> CVarRigidBodyLODThreshold(TEXT("p.RigidBodyLODThreshold"), -1, TEXT("Max LOD that rigid body node is allowed to run on. Provides a global threshold that overrides per-node the LODThreshold property. -1 means no override."), ECVF_Scalability);
 
 int32 RBAN_MaxSubSteps = 4;
+bool bRBAN_EnableTimeBasedReset = true;
 FAutoConsoleVariableRef CVarRigidBodyNodeMaxSteps(TEXT("p.RigidBodyNode.MaxSubSteps"), RBAN_MaxSubSteps, TEXT("Set the maximum number of simulation steps in the update loop"), ECVF_Default);
+FAutoConsoleVariableRef CVarRigidBodyNodeEnableTimeBasedReset(TEXT("p.RigidBodyNode.EnableTimeBasedReset"), bRBAN_EnableTimeBasedReset, TEXT("If true, Rigid Body nodes are reset when they have not been updated for a while (default true)"), ECVF_Default);
 
 FAnimNode_RigidBody::FAnimNode_RigidBody():
 	QueryParams(NAME_None, FCollisionQueryParams::GetUnknownStatId())
@@ -255,7 +257,7 @@ void FAnimNode_RigidBody::EvaluateSkeletalControl_AnyThread(FComponentSpacePoseC
 	{
 		// Always propagate skip rate as it can go up and down between updates
 		EvalCounter.SetMaxSkippedFrames(Output.AnimInstanceProxy->GetEvaluationCounter().GetMaxSkippedFrames());
-		if(!EvalCounter.WasSynchronizedLastFrame(Output.AnimInstanceProxy->GetEvaluationCounter()))
+		if(!EvalCounter.WasSynchronizedLastFrame(Output.AnimInstanceProxy->GetEvaluationCounter())  && bRBAN_EnableTimeBasedReset)
 		{
 			ResetSimulatedTeleportType = ETeleportType::ResetPhysics;
 		}
diff --git a/Engine/Source/Runtime/AnimGraphRuntime/Private/BoneControllers/AnimNode_RigidBody_Chaos.cpp b/Engine/Source/Runtime/AnimGraphRuntime/Private/BoneControllers/AnimNode_RigidBody_Chaos.cpp
--- a/Engine/Source/Runtime/AnimGraphRuntime/Private/BoneControllers/AnimNode_RigidBody_Chaos.cpp
+++ b/Engine/Source/Runtime/AnimGraphRuntime/Private/BoneControllers/AnimNode_RigidBody_Chaos.cpp
@@ -226,7 +226,8 @@ void FAnimNode_RigidBody_Chaos::EvaluateSkeletalControl_AnyThread(FComponentSpac
 	{
 		// Always propagate skip rate as it can go up and down between updates
 		EvalCounter.SetMaxSkippedFrames(Output.AnimInstanceProxy->GetEvaluationCounter().GetMaxSkippedFrames());
-		if (!EvalCounter.WasSynchronizedLastFrame(Output.AnimInstanceProxy->GetEvaluationCounter()))
+		extern bool bRBAN_EnableTimeBasedReset;
+		if (!EvalCounter.WasSynchronizedLastFrame(Output.AnimInstanceProxy->GetEvaluationCounter()) && bRBAN_EnableTimeBasedReset)
 		{
 			ResetSimulatedTeleportType = ETeleportType::ResetPhysics;
 		}
diff --git a/Engine/Source/Runtime/Engine/Private/PhysicsEngine/ImmediatePhysics/ImmediatePhysicsChaos/ImmediatePhysicsSimulation_Chaos.cpp b/Engine/Source/Runtime/Engine/Private/PhysicsEngine/ImmediatePhysics/ImmediatePhysicsChaos/ImmediatePhysicsSimulation_Chaos.cpp
--- a/Engine/Source/Runtime/Engine/Private/PhysicsEngine/ImmediatePhysics/ImmediatePhysicsChaos/ImmediatePhysicsSimulation_Chaos.cpp
+++ b/Engine/Source/Runtime/Engine/Private/PhysicsEngine/ImmediatePhysics/ImmediatePhysicsChaos/ImmediatePhysicsSimulation_Chaos.cpp
@@ -21,30 +21,31 @@
 // @todo(ccaulfield): remove when finished
 //
 float ChaosImmediate_Evolution_DeltaTime = 0.03f;
-int32 ChaosImmediate_Evolution_Iterations = 6;
-int32 ChaosImmediate_Evolution_PushOutIterations = 5;
+int32 ChaosImmediate_Evolution_Iterations = 2;
+int32 ChaosImmediate_Evolution_PushOutIterations = 10;
 FAutoConsoleVariableRef CVarChaosImmPhysDeltaTime(TEXT("p.Chaos.ImmPhys.DeltaTime"), ChaosImmediate_Evolution_DeltaTime, TEXT("Override chaos immediate physics delta time if non-zero"));
 FAutoConsoleVariableRef CVarChaosImmPhysIterations(TEXT("p.Chaos.ImmPhys.Iterations"), ChaosImmediate_Evolution_Iterations, TEXT("Number of constraint solver loops in immediate physics"));
 FAutoConsoleVariableRef CVarChaosImmPhysPushOutIterations(TEXT("p.Chaos.ImmPhys.PushOutIterations"), ChaosImmediate_Evolution_PushOutIterations, TEXT("Set the ApplyPushOut() (position correction) iteration count"));
 
 int32 ChaosImmediate_Collision_Enabled = 1;
-int32 ChaosImmediate_Collision_ApplyEnabled = 1;
+int32 ChaosImmediate_Collision_PairIterations = 0;
 int32 ChaosImmediate_Collision_PushOutPairIterations = 2;
 int32 ChaosImmediate_Collision_Priority = 1;
 float ChaosImmediate_Collision_Thickness = 0;
-FAutoConsoleVariableRef CVarChaosImmPhysDisableCollisions(TEXT("p.Chaos.ImmPhys.EnableCollisions"), ChaosImmediate_Collision_Enabled, TEXT("Enable/Disable collisions in Immediate Physics."));
-FAutoConsoleVariableRef CVarChaosImmPhysCollisionDisableApply(TEXT("p.Chaos.ImmPhys.CollisionEnableApply"), ChaosImmediate_Collision_ApplyEnabled, TEXT("Enable/Disable the Apply() (velocity, friction, restitution) method for collisions"));
-FAutoConsoleVariableRef CVarChaosImmPhysCollisionPushOutPairIterations(TEXT("p.Chaos.ImmPhys.CollisionPushOutPairIterations"), ChaosImmediate_Collision_PushOutPairIterations, TEXT("Set the ApplyPushOut() internal pair interations (position correction) iteration count (0 to disable)"));
-FAutoConsoleVariableRef CVarChaosImmPhysCollisionPriority(TEXT("p.Chaos.ImmPhys.CollisionPriority"), ChaosImmediate_Collision_Priority, TEXT("Set the Collision constraint sort order (Joints have priority 0)"));
-FAutoConsoleVariableRef CVarChaosImmPhysThickness(TEXT("p.Chaos.ImmPhys.CollisionThickness"), ChaosImmediate_Collision_Thickness, TEXT("ChaosImmediateThickness"));
-
+FAutoConsoleVariableRef CVarChaosImmPhysDisableCollisions(TEXT("p.Chaos.ImmPhys.Collision.Enabled"), ChaosImmediate_Collision_Enabled, TEXT("Enable/Disable collisions in Immediate Physics."));
+FAutoConsoleVariableRef CVarChaosImmPhysCollisionPairIterations(TEXT("p.Chaos.ImmPhys.Collision.PairIterations"), ChaosImmediate_Collision_PairIterations, TEXT("Set the Apply() (velocity correction) internal pair iteration count (0 to disable)"));
+FAutoConsoleVariableRef CVarChaosImmPhysCollisionPushOutPairIterations(TEXT("p.Chaos.ImmPhys.Collision.PushOutPairIterations"), ChaosImmediate_Collision_PushOutPairIterations, TEXT("Set the ApplyPushOut() (position correction) internal pair iteration count (0 to disable)"));
+FAutoConsoleVariableRef CVarChaosImmPhysCollisionPriority(TEXT("p.Chaos.ImmPhys.Collision.Priority"), ChaosImmediate_Collision_Priority, TEXT("Set the Collision constraint sort order (Joints have priority 0)"));
+FAutoConsoleVariableRef CVarChaosImmPhysThickness(TEXT("p.Chaos.ImmPhys.Collision.Thickness"), ChaosImmediate_Collision_Thickness, TEXT("ChaosImmediateThickness"));
+
+int32 ChaosImmediate_Joint_PairIterations = 0;
+int32 ChaosImmediate_Joint_PushOutPairIterations = 2;
 float ChaosImmediate_Joint_SwingTwistAngleTolerance = 1.0e-6f;
-int32 ChaosImmediate_Joint_PositionIterations = 6;
-int32 ChaosImmediate_Joint_EnableVelocitySolve = 1;
-int32 ChaosImmediate_Joint_EnableLinearLimits = 1;
+int32 ChaosImmediate_Joint_EnableVelocitySolve = 0;
 int32 ChaosImmediate_Joint_EnableTwistLimits = 1;
 int32 ChaosImmediate_Joint_EnableSwingLimits = 1;
 int32 ChaosImmediate_Joint_EnableDrives = 1;
+int32 ChaosImmediate_Joint_ProjectionPhase = (int32)Chaos::EJointProjectionPhase::ApplyPushOut;
 float ChaosImmediate_Joint_LinearProjection = 0.0f;
 float ChaosImmediate_Joint_AngularProjection = 0.0f;
 float ChaosImmediate_Joint_Stiffness = 1.0f;
@@ -53,13 +54,14 @@ float ChaosImmediate_Joint_SoftAngularStiffness = 0.0f;
 float ChaosImmediate_Joint_DriveStiffness = 0.0f;
 float ChaosImmediate_Joint_MinParentMassRatio = 0.5f;
 float ChaosImmediate_Joint_MaxInertiaRatio = 5.0f;
+FAutoConsoleVariableRef CVarChaosImmPhysPairIterations(TEXT("p.Chaos.ImmPhys.Joint.PairIterations"), ChaosImmediate_Joint_PairIterations, TEXT("PairIterations."));
+FAutoConsoleVariableRef CVarChaosImmPhysPushOutPairIterations(TEXT("p.Chaos.ImmPhys.Joint.PushOutPairIterations"), ChaosImmediate_Joint_PushOutPairIterations, TEXT("PushOutPairIterations."));
 FAutoConsoleVariableRef CVarChaosImmPhysSwingTwistAngleTolerance(TEXT("p.Chaos.ImmPhys.Joint.SwingTwistAngleTolerance"), ChaosImmediate_Joint_SwingTwistAngleTolerance, TEXT("SwingTwistAngleTolerance."));
-FAutoConsoleVariableRef CVarChaosImmPhysEnablePositionSolve(TEXT("p.Chaos.ImmPhys.Joint.PositionIterations"), ChaosImmediate_Joint_PositionIterations, TEXT("6Dof joint PushOut (Position solve) iterations."));
 FAutoConsoleVariableRef CVarChaosImmPhysEnableVelocitySolve(TEXT("p.Chaos.ImmPhys.Joint.EnableVelocitySolve"), ChaosImmediate_Joint_EnableVelocitySolve, TEXT("EnableVelocitySolve."));
-FAutoConsoleVariableRef CVarChaosImmPhysEnableLinearLimits(TEXT("p.Chaos.ImmPhys.Joint.EnableLinearLimits"), ChaosImmediate_Joint_EnableLinearLimits, TEXT("EnableLinearLimits."));
 FAutoConsoleVariableRef CVarChaosImmPhysEnableTwistLimits(TEXT("p.Chaos.ImmPhys.Joint.EnableTwistLimits"), ChaosImmediate_Joint_EnableTwistLimits, TEXT("EnableTwistLimits."));
 FAutoConsoleVariableRef CVarChaosImmPhysEnableSwingLimits(TEXT("p.Chaos.ImmPhys.Joint.EnableSwingLimits"), ChaosImmediate_Joint_EnableSwingLimits, TEXT("EnableSwingLimits."));
 FAutoConsoleVariableRef CVarChaosImmPhysEnableDrives(TEXT("p.Chaos.ImmPhys.Joint.EnableDrives"), ChaosImmediate_Joint_EnableDrives, TEXT("EnableDrives."));
+FAutoConsoleVariableRef CVarChaosImmPhysProjectionPhase(TEXT("p.Chaos.ImmPhys.Joint.ProjectionPhase"), ChaosImmediate_Joint_ProjectionPhase, TEXT("Projection Phase (0=Never; 1=Apply; 2=ApplyPushOut)."));
 FAutoConsoleVariableRef CVarChaosImmPhysLinearProjection(TEXT("p.Chaos.ImmPhys.Joint.LinearProjection"), ChaosImmediate_Joint_LinearProjection, TEXT("6Dof joint projection amount override (if > 0)."));
 FAutoConsoleVariableRef CVarChaosImmPhysAngularProjection(TEXT("p.Chaos.ImmPhys.Joint.AngularProjection"), ChaosImmediate_Joint_AngularProjection, TEXT("6Dof joint projection amount override (if > 0)."));
 FAutoConsoleVariableRef CVarChaosImmPhysStiffness(TEXT("p.Chaos.ImmPhys.Joint.Stiffness"), ChaosImmediate_Joint_Stiffness, TEXT("6Dof joint stiffness override (if > 0)."));
@@ -87,13 +89,31 @@ int32 ChaosImmediate_DebugDrawCollisions = 0;
 int32 ChaosImmediate_DebugDrawJoints = 0;
 #endif
 
+int32 ChaosImmediate_DebugDrawJointFeatures = Chaos::DebugDraw::EDebugDrawJointFeature::Default;
+
 FAutoConsoleVariableRef CVarChaosImmPhysDebugDrawParticles(TEXT("p.Chaos.ImmPhys.DebugDrawParticles"), ChaosImmediate_DebugDrawParticles, TEXT("Draw Particle Transforms (0 = never; 1 = end of frame; 2 = begin and end of frame; 3 = post-integate, post-apply and post-applypushout;)."));
 FAutoConsoleVariableRef CVarChaosImmPhysDebugDrawShapes(TEXT("p.Chaos.ImmPhys.DebugDrawShapes"), ChaosImmediate_DebugDrawShapes, TEXT("Draw Shapes (0 = never; 1 = end of frame; 2 = begin and end of frame; 3 = post-integate, post-apply and post-applypushout;"));
 FAutoConsoleVariableRef CVarChaosImmPhysDebugDrawCollisions(TEXT("p.Chaos.ImmPhys.DebugDrawCollisions"), ChaosImmediate_DebugDrawCollisions, TEXT("Draw Collisions (0 = never; 1 = end of frame; 2 = begin and end of frame; 3 = post-integate, post-apply and post-applypushout;)"));
 FAutoConsoleVariableRef CVarChaosImmPhysDebugDrawJoints(TEXT("p.Chaos.ImmPhys.DebugDrawJoints"), ChaosImmediate_DebugDrawJoints, TEXT("Draw Joints. (0 = never; 1 = end of frame; 2 = begin and end of frame; 3 = post-integate, post-apply and post-applypushout; 4 = each Apply step)."));
+FAutoConsoleVariableRef CVarChaosImmPhysDebugDrawJointFeatures(TEXT("p.Chaos.ImmPhys.DebugDrawJointFeatures"), ChaosImmediate_DebugDrawJointFeatures, TEXT("Joint features mask (see EDebugDrawJointFeature)."));
 
 namespace ImmediatePhysics_Chaos
 {
+	Chaos::EJointProjectionPhase ToProjectionPhase(const int32 Index)
+	{
+		using namespace Chaos;
+
+		if (Index == 1)
+		{
+			return EJointProjectionPhase::Apply;
+		}
+		else if (Index == 2)
+		{
+			return EJointProjectionPhase::ApplyPushOut;
+		}
+		return EJointProjectionPhase::None;
+	}
+
 	template<typename T, int d>
 	bool ShouldIgnoreCollisionConstraint(
 		const Chaos::TGeometryParticleHandle<T, d>* ParticleA,
@@ -161,15 +181,15 @@ namespace ImmediatePhysics_Chaos
 			[this]()
 			{
 				// Dynamic only - Kinematics get drawn once at the end of frame
-				//DebugDrawIslandParticles(Island, 3, 3, 0.3f, false, true);
-				//DebugDrawIslandConstraints(Island, 3, 3, 0.3f);
+				DebugDrawParticles(3, 3, 0.3f, false, true);
+				DebugDrawConstraints(3, 3, 0.3f);
 			});
 		Evolution->SetPostApplyCallback(
 			[this](int32 Island)
 			{
 				// Dynamic only - Kinematics get drawn once at the end of frame
-				DebugDrawIslandParticles(Island, 3, 3, 0.3f, false, true);
-				DebugDrawIslandConstraints(Island, 3, 3, 0.3f);
+				DebugDrawIslandParticles(Island, 3, 3, 0.6f, false, true);
+				DebugDrawIslandConstraints(Island, 3, 3, 0.6f);
 			});
 		Evolution->SetPostApplyPushOutCallback(
 			[this](int32 Island)
@@ -412,18 +432,21 @@ namespace ImmediatePhysics_Chaos
 				// Round Dt to the nearest multiple of fixed step size...
 				const float NumDts = FMath::RoundToFloat(DeltaTime / ChaosImmediate_Evolution_DeltaTime);
 				DeltaTime = NumDts * ChaosImmediate_Evolution_DeltaTime;
+				MaxDeltaTime = ChaosImmediate_Evolution_DeltaTime;
 			}
 			UE_LOG(LogChaosJoint, Verbose, TEXT("Simulate dt = %f"), DeltaTime);
 
 			FPBDJointSolverSettings JointsSettings = Joints->GetSettings();
+			JointsSettings.ApplyPairIterations = ChaosImmediate_Joint_PairIterations;
+			JointsSettings.ApplyPushOutPairIterations = ChaosImmediate_Joint_PushOutPairIterations;
 			JointsSettings.SwingTwistAngleTolerance = ChaosImmediate_Joint_SwingTwistAngleTolerance;
 			JointsSettings.MinParentMassRatio = ChaosImmediate_Joint_MinParentMassRatio;
 			JointsSettings.MaxInertiaRatio = ChaosImmediate_Joint_MaxInertiaRatio;
 			JointsSettings.bEnableVelocitySolve = ChaosImmediate_Joint_EnableVelocitySolve != 0;
-			JointsSettings.bEnableLinearLimits = ChaosImmediate_Joint_EnableLinearLimits != 0;
 			JointsSettings.bEnableTwistLimits = ChaosImmediate_Joint_EnableTwistLimits != 0;
 			JointsSettings.bEnableSwingLimits = ChaosImmediate_Joint_EnableSwingLimits != 0;
 			JointsSettings.bEnableDrives = ChaosImmediate_Joint_EnableDrives != 0;
+			JointsSettings.ProjectionPhase = ToProjectionPhase(ChaosImmediate_Joint_ProjectionPhase);
 			JointsSettings.LinearProjection = ChaosImmediate_Joint_LinearProjection;
 			JointsSettings.AngularProjection = ChaosImmediate_Joint_AngularProjection;
 			JointsSettings.Stiffness = ChaosImmediate_Joint_Stiffness;
@@ -436,7 +459,7 @@ namespace ImmediatePhysics_Chaos
 			Evolution->SetNumPushOutIterations(ChaosImmediate_Evolution_PushOutIterations);
 
 			Evolution->GetCollisionConstraints().SetThickness(ChaosImmediate_Collision_Thickness);
-			Evolution->GetCollisionConstraints().SetVelocitySolveEnabled(ChaosImmediate_Collision_ApplyEnabled != 0);
+			Evolution->GetCollisionConstraints().SetPairIterations(ChaosImmediate_Collision_PairIterations);
 			Evolution->GetCollisionConstraints().SetPushOutPairIterations(ChaosImmediate_Collision_PushOutPairIterations);
 			Evolution->GetCollisionConstraintsRule().SetPriority(ChaosImmediate_Collision_Priority);
 		
@@ -490,7 +513,7 @@ namespace ImmediatePhysics_Chaos
 			}
 			if ((ChaosImmediate_DebugDrawJoints >= MinDebugLevel) && (ChaosImmediate_DebugDrawJoints <= MaxDebugLevel))
 			{
-				DebugDraw::DrawJointConstraints(SimulationSpaceTransform, *Joints, ColorScale);
+				DebugDraw::DrawJointConstraints(SimulationSpaceTransform, *Joints, ColorScale, (uint32)ChaosImmediate_DebugDrawJointFeatures);
 			}
 		}
 #endif
@@ -533,7 +556,7 @@ namespace ImmediatePhysics_Chaos
 				JointsRule->VisitIslandConstraints(Island,
 					[this, ColorScale](const TArray<FPBDJointConstraintHandle*>& ConstraintHandles)
 					{
-						DebugDraw::DrawJointConstraints(SimulationSpaceTransform, *Joints, ColorScale);
+						DebugDraw::DrawJointConstraints(SimulationSpaceTransform, *Joints, ColorScale, (uint32)ChaosImmediate_DebugDrawJointFeatures);
 					});
 			}
 		}
diff --git a/Engine/Source/Runtime/Experimental/Chaos/Private/Chaos/ChaosDebugDraw.cpp b/Engine/Source/Runtime/Experimental/Chaos/Private/Chaos/ChaosDebugDraw.cpp
--- a/Engine/Source/Runtime/Experimental/Chaos/Private/Chaos/ChaosDebugDraw.cpp
+++ b/Engine/Source/Runtime/Experimental/Chaos/Private/Chaos/ChaosDebugDraw.cpp
@@ -29,6 +29,8 @@ namespace Chaos
 		float DrawScale = 1.0f;
 		float FontHeight = 10.0f;
 		float FontScale = 1.0f;
+		float ShapeThicknesScale = 2.0f;
+		float PointSize = 2.0f;
 
 		FAutoConsoleVariableRef CVarArrowSize(TEXT("p.Chaos.DebugDrawArrowSize"), ArrowSize, TEXT("ArrowSize."));
 		FAutoConsoleVariableRef CVarBodyAxisLen(TEXT("p.Chaos.DebugDrawBodyAxisLen"), BodyAxisLen, TEXT("BodyAxisLen."));
@@ -54,14 +56,14 @@ namespace Chaos
 			{
 				const TSphere<FReal, 3>* Sphere = Shape->template GetObject<TSphere<FReal, 3>>();
 				const FVec3 P = ShapeTransform.TransformPosition(Sphere->GetCenter());
-				FDebugDrawQueue::GetInstance().DrawDebugSphere(P, Sphere->GetRadius(), 20, Color, false, KINDA_SMALL_NUMBER, DrawPriority, LineThickness);
+				FDebugDrawQueue::GetInstance().DrawDebugSphere(P, Sphere->GetRadius(), 20, Color, false, KINDA_SMALL_NUMBER, DrawPriority, ShapeThicknesScale * LineThickness);
 				break;
 			}
 			case ImplicitObjectType::Box:
 			{
 				const TBox<FReal, 3>* Box = Shape->template GetObject<TBox<FReal, 3>>();
 				const FVec3 P = ShapeTransform.TransformPosition(Box->GetCenter());
-				FDebugDrawQueue::GetInstance().DrawDebugBox(P, (FReal)0.5 * Box->Extents(), ShapeTransform.GetRotation(), Color, false, KINDA_SMALL_NUMBER, DrawPriority, LineThickness);
+				FDebugDrawQueue::GetInstance().DrawDebugBox(P, (FReal)0.5 * Box->Extents(), ShapeTransform.GetRotation(), Color, false, KINDA_SMALL_NUMBER, DrawPriority, ShapeThicknesScale * LineThickness);
 				break;
 			}
 			case ImplicitObjectType::Plane:
@@ -71,7 +73,7 @@ namespace Chaos
 				const TCapsule<FReal>* Capsule = Shape->template GetObject<TCapsule<FReal>>();
 				const FVec3 P = ShapeTransform.TransformPosition(Capsule->GetCenter());
 				const FRotation3 Q = ShapeTransform.GetRotation() * FRotationMatrix::MakeFromZ(Capsule->GetAxis());
-				FDebugDrawQueue::GetInstance().DrawDebugCapsule(P, (FReal)0.5 * Capsule->GetHeight() + Capsule->GetRadius(), Capsule->GetRadius(), Q, Color, false, KINDA_SMALL_NUMBER, DrawPriority, LineThickness);
+				FDebugDrawQueue::GetInstance().DrawDebugCapsule(P, (FReal)0.5 * Capsule->GetHeight() + Capsule->GetRadius(), Capsule->GetRadius(), Q, Color, false, KINDA_SMALL_NUMBER, DrawPriority, ShapeThicknesScale * LineThickness);
 				break;
 			}
 			case ImplicitObjectType::Transformed:
@@ -166,7 +168,7 @@ namespace Chaos
 			}
 		}
 
-		void DrawJointConstraintImpl(const FRigidTransform3& SpaceTransform, const FVec3& InXa, const FMatrix33& Ra, const FVec3& InXb, const FMatrix33& Rb, const FVec3& CR, int32 Level, FReal ColorScale)
+		void DrawJointConstraintImpl(const FRigidTransform3& SpaceTransform, const FVec3& InPa, const FVec3& InXa, const FMatrix33& Ra, const FVec3& InPb, const FVec3& InXb, const FMatrix33& Rb, const FVec3& CR, int32 Level, FReal ColorScale, uint32 FeatureMask)
 		{
 			using namespace Chaos::DebugDraw;
 			FColor R = (ColorScale * FColor::Red).ToFColor(false);
@@ -175,15 +177,28 @@ namespace Chaos
 			FColor C = (ColorScale * FColor::Cyan).ToFColor(false);
 			FColor M = (ColorScale * FColor::Magenta).ToFColor(false);
 			FColor Y = (ColorScale * FColor::Yellow).ToFColor(false);
+			FVec3 Pa = SpaceTransform.TransformPosition(InPa);
+			FVec3 Pb = SpaceTransform.TransformPosition(InPb);
 			FVec3 Xa = SpaceTransform.TransformPosition(InXa);
 			FVec3 Xb = SpaceTransform.TransformPosition(InXb);
-			FDebugDrawQueue::GetInstance().DrawDebugDirectionalArrow(Xa, Xa + DrawScale * ConstraintAxisLen * SpaceTransform.TransformVector(Ra.GetAxis(0)), DrawScale * ArrowSize, R, false, KINDA_SMALL_NUMBER, DrawPriority, LineThickness);
-			FDebugDrawQueue::GetInstance().DrawDebugDirectionalArrow(Xa, Xa + DrawScale * ConstraintAxisLen * SpaceTransform.TransformVector(Ra.GetAxis(1)), DrawScale * ArrowSize, G, false, KINDA_SMALL_NUMBER, DrawPriority, LineThickness);
-			FDebugDrawQueue::GetInstance().DrawDebugDirectionalArrow(Xa, Xa + DrawScale * ConstraintAxisLen * SpaceTransform.TransformVector(Ra.GetAxis(2)), DrawScale * ArrowSize, B, false, KINDA_SMALL_NUMBER, DrawPriority, LineThickness);
-			FDebugDrawQueue::GetInstance().DrawDebugDirectionalArrow(Xb, Xb + DrawScale * ConstraintAxisLen * SpaceTransform.TransformVector(Rb.GetAxis(0)), DrawScale * ArrowSize, C, false, KINDA_SMALL_NUMBER, DrawPriority, LineThickness);
-			FDebugDrawQueue::GetInstance().DrawDebugDirectionalArrow(Xb, Xb + DrawScale * ConstraintAxisLen * SpaceTransform.TransformVector(Rb.GetAxis(1)), DrawScale * ArrowSize, M, false, KINDA_SMALL_NUMBER, DrawPriority, LineThickness);
-			FDebugDrawQueue::GetInstance().DrawDebugDirectionalArrow(Xb, Xb + DrawScale * ConstraintAxisLen * SpaceTransform.TransformVector(Rb.GetAxis(2)), DrawScale * ArrowSize, Y, false, KINDA_SMALL_NUMBER, DrawPriority, LineThickness);
-			if (Level >= 0)
+
+			if (FeatureMask & EDebugDrawJointFeature::Connector)
+			{
+				FDebugDrawQueue::GetInstance().DrawDebugLine(Pa, Xa, R, false, KINDA_SMALL_NUMBER, DrawPriority, LineThickness);
+				FDebugDrawQueue::GetInstance().DrawDebugLine(Pb, Xb, C, false, KINDA_SMALL_NUMBER, DrawPriority, LineThickness);
+				FDebugDrawQueue::GetInstance().DrawDebugPoint(Xa, R, false, KINDA_SMALL_NUMBER, DrawPriority, DrawScale * PointSize);
+				FDebugDrawQueue::GetInstance().DrawDebugPoint(Xb, C, false, KINDA_SMALL_NUMBER, DrawPriority, DrawScale * PointSize);
+			}
+			if (FeatureMask & EDebugDrawJointFeature::Axes)
+			{
+				FDebugDrawQueue::GetInstance().DrawDebugDirectionalArrow(Xa, Xa + DrawScale * ConstraintAxisLen * SpaceTransform.TransformVector(Ra.GetAxis(0)), DrawScale * ArrowSize, R, false, KINDA_SMALL_NUMBER, DrawPriority, LineThickness);
+				FDebugDrawQueue::GetInstance().DrawDebugDirectionalArrow(Xa, Xa + DrawScale * ConstraintAxisLen * SpaceTransform.TransformVector(Ra.GetAxis(1)), DrawScale * ArrowSize, G, false, KINDA_SMALL_NUMBER, DrawPriority, LineThickness);
+				FDebugDrawQueue::GetInstance().DrawDebugDirectionalArrow(Xa, Xa + DrawScale * ConstraintAxisLen * SpaceTransform.TransformVector(Ra.GetAxis(2)), DrawScale * ArrowSize, B, false, KINDA_SMALL_NUMBER, DrawPriority, LineThickness);
+				FDebugDrawQueue::GetInstance().DrawDebugDirectionalArrow(Xb, Xb + DrawScale * ConstraintAxisLen * SpaceTransform.TransformVector(Rb.GetAxis(0)), DrawScale * ArrowSize, C, false, KINDA_SMALL_NUMBER, DrawPriority, LineThickness);
+				FDebugDrawQueue::GetInstance().DrawDebugDirectionalArrow(Xb, Xb + DrawScale * ConstraintAxisLen * SpaceTransform.TransformVector(Rb.GetAxis(1)), DrawScale * ArrowSize, M, false, KINDA_SMALL_NUMBER, DrawPriority, LineThickness);
+				FDebugDrawQueue::GetInstance().DrawDebugDirectionalArrow(Xb, Xb + DrawScale * ConstraintAxisLen * SpaceTransform.TransformVector(Rb.GetAxis(2)), DrawScale * ArrowSize, Y, false, KINDA_SMALL_NUMBER, DrawPriority, LineThickness);
+			}
+			if ((FeatureMask & EDebugDrawJointFeature::Level) && (Level >= 0))
 			{
 				FDebugDrawQueue::GetInstance().DrawDebugString(Xb + FontHeight * FVec3(0, 0, 1), FString::Format(TEXT("{0}"), { Level }), nullptr, FColor::Red, KINDA_SMALL_NUMBER, false, FontScale);
 			}
@@ -192,24 +207,28 @@ namespace Chaos
 			//FDebugDrawQueue::GetInstance().DrawDebugString(Xb + 1 * FontHeight * FVec3(0, 0, 1), FString::Format(TEXT("S2= {0}"), { FMath::RadiansToDegrees(CR[(int32)E6DJointAngularConstraintIndex::Swing2]) }), nullptr, FColor::Red, KINDA_SMALL_NUMBER, false, FontScale);
 		}
 
-		void DrawJointConstraintImpl(const FRigidTransform3& SpaceTransform, const FPBDJointConstraintHandle* ConstraintHandle, FReal ColorScale)
+		void DrawJointConstraintImpl(const FRigidTransform3& SpaceTransform, const FPBDJointConstraintHandle* ConstraintHandle, FReal ColorScale, uint32 FeatureMask)
 		{
 			TVector<TGeometryParticleHandle<FReal, 3>*, 2> ConstrainedParticles = ConstraintHandle->GetConstrainedParticles();
 			if (ConstrainedParticles[0]->AsDynamic() || ConstrainedParticles[1]->AsDynamic())
 			{
+				FVec3 Pa = TGenericParticleHandle<FReal, 3>(ConstraintHandle->GetConstrainedParticles()[1])->P();
+				FVec3 Pb = TGenericParticleHandle<FReal, 3>(ConstraintHandle->GetConstrainedParticles()[0])->P();
 				FVec3 Xa, Xb, CR;
 				FMatrix33 Ra, Rb;
 				ConstraintHandle->CalculateConstraintSpace(Xa, Ra, Xb, Rb, CR);
-				DrawJointConstraintImpl(SpaceTransform, Xa, Ra, Xb, Rb, CR, ConstraintHandle->GetConstraintLevel(), ColorScale);
+				DrawJointConstraintImpl(SpaceTransform, Pa, Xa, Ra, Pb, Xb, Rb, CR, ConstraintHandle->GetConstraintLevel(), ColorScale, FeatureMask);
 			}
 		}
 
 		void Draw6DofConstraintImpl(const FRigidTransform3& SpaceTransform, const FPBD6DJointConstraintHandle* ConstraintHandle, FReal ColorScale)
 		{
+			FVec3 Pa = TGenericParticleHandle<FReal, 3>(ConstraintHandle->GetConstrainedParticles()[1])->P();
+			FVec3 Pb = TGenericParticleHandle<FReal, 3>(ConstraintHandle->GetConstrainedParticles()[0])->P();
 			FVec3 Xa, Xb, CR;
 			FMatrix33 Ra, Rb;
 			ConstraintHandle->CalculateConstraintSpace(Xa, Ra, Xb, Rb, CR);
-			DrawJointConstraintImpl(SpaceTransform, Xa, Ra, Xb, Rb, CR, INDEX_NONE, ColorScale);
+			DrawJointConstraintImpl(SpaceTransform, Pa, Xa, Ra, Pb, Xb, Rb, CR, INDEX_NONE, ColorScale, EDebugDrawJointFeature::Default);
 		}
 
 #endif
@@ -322,27 +341,27 @@ namespace Chaos
 		}
 
 
-		void DrawJointConstraints(const FRigidTransform3& SpaceTransform, const TArray<FPBDJointConstraintHandle*>& ConstraintHandles, float ColorScale)
+		void DrawJointConstraints(const FRigidTransform3& SpaceTransform, const TArray<FPBDJointConstraintHandle*>& ConstraintHandles, float ColorScale, uint32 FeatureMask)
 		{
 #if CHAOS_DEBUG_DRAW
 			if (FDebugDrawQueue::IsDebugDrawingEnabled())
 			{
 				for (const FPBDJointConstraintHandle* ConstraintHandle : ConstraintHandles)
 				{
-					DrawJointConstraintImpl(SpaceTransform, ConstraintHandle, ColorScale);
+					DrawJointConstraintImpl(SpaceTransform, ConstraintHandle, ColorScale, FeatureMask);
 				}
 			}
 #endif
 		}
 
-		void DrawJointConstraints(const FRigidTransform3& SpaceTransform, const FPBDJointConstraints& Constraints, float ColorScale)
+		void DrawJointConstraints(const FRigidTransform3& SpaceTransform, const FPBDJointConstraints& Constraints, float ColorScale, uint32 FeatureMask)
 		{
 #if CHAOS_DEBUG_DRAW
 			if (FDebugDrawQueue::IsDebugDrawingEnabled())
 			{
 				for (int32 ConstraintIndex = 0; ConstraintIndex < Constraints.NumConstraints(); ++ConstraintIndex)
 				{
-					DrawJointConstraintImpl(SpaceTransform, Constraints.GetConstraintHandle(ConstraintIndex), ColorScale);
+					DrawJointConstraintImpl(SpaceTransform, Constraints.GetConstraintHandle(ConstraintIndex), ColorScale, FeatureMask);
 				}
 			}
 #endif
diff --git a/Engine/Source/Runtime/Experimental/Chaos/Private/Chaos/DenseMatrix.cpp b/Engine/Source/Runtime/Experimental/Chaos/Private/Chaos/DenseMatrix.cpp
new file mode 100644
--- /dev/null
+++ b/Engine/Source/Runtime/Experimental/Chaos/Private/Chaos/DenseMatrix.cpp
@@ -0,0 +1,2 @@
+// Copyright 1998-2019 Epic Games, Inc. All Rights Reserved.
+#include "Chaos/DenseMatrix.h"
\ No newline at end of file
diff --git a/Engine/Source/Runtime/Experimental/Chaos/Private/Chaos/PBDCollisionConstraint.cpp b/Engine/Source/Runtime/Experimental/Chaos/Private/Chaos/PBDCollisionConstraint.cpp
--- a/Engine/Source/Runtime/Experimental/Chaos/Private/Chaos/PBDCollisionConstraint.cpp
+++ b/Engine/Source/Runtime/Experimental/Chaos/Private/Chaos/PBDCollisionConstraint.cpp
@@ -38,9 +38,12 @@ FAutoConsoleVariableRef CVarConstraintBPBVHDepth(TEXT("p.ConstraintBPBVHDepth"),
 int32 BPTreeOfGrids = 1;
 FAutoConsoleVariableRef CVarBPTreeOfGrids(TEXT("p.BPTreeOfGrids"), BPTreeOfGrids, TEXT("Whether to use a seperate tree of grids for bp"));
 
-float CollisionVelocityInflationCVar = 2.0;
+float CollisionVelocityInflationCVar = 2.0f;
 FAutoConsoleVariableRef CVarCollisionVelocityInflation(TEXT("p.CollisionVelocityInflation"), CollisionVelocityInflationCVar, TEXT("Collision velocity inflation.[def:2.0]"));
 
+float CollisionFrictionOverride = -1.0f;
+FAutoConsoleVariableRef CVarCollisionFrictionOverride(TEXT("p.CollisionFriction"), CollisionFrictionOverride, TEXT("Collision friction for all contacts if >= 0"));
+
 
 extern int32 UseLevelsetCollision;
 
@@ -63,12 +66,18 @@ namespace Chaos
 	//
 
 	template<typename T, int d>
-	TPBDCollisionConstraint<T, d>::TPBDCollisionConstraint(const TPBDRigidsSOAs<T, d>& InParticles, TArrayCollectionArray<bool>& Collided, const TArrayCollectionArray<TSerializablePtr<FChaosPhysicsMaterial>>& InPerParticleMaterials, const int32 PairIterations /*= 1*/, const T Thickness /*= (T)0*/)
+	TPBDCollisionConstraint<T, d>::TPBDCollisionConstraint(
+		const TPBDRigidsSOAs<T, d>& InParticles, 
+		TArrayCollectionArray<bool>& Collided, 
+		const TArrayCollectionArray<TSerializablePtr<FChaosPhysicsMaterial>>& InPerParticleMaterials, 
+		const int32 InApplyPairIterations /*= 1*/, 
+		const int32 InApplyPushOutPairIterations /*= 1*/, 
+		const T Thickness /*= (T)0*/)
 		: Particles(InParticles)
 		, MCollided(Collided)
 		, MPhysicsMaterials(InPerParticleMaterials)
-		, bEnableVelocitySolve(true)
-		, MPairIterations(PairIterations)
+		, MApplyPairIterations(InApplyPairIterations)
+		, MApplyPushOutPairIterations(InApplyPushOutPairIterations)
 		, MThickness(Thickness)
 		, MAngularFriction(0)
 		, bUseCCD(false)
@@ -314,7 +323,7 @@ namespace Chaos
 
 
 	template<typename T, int d>
-	void TPBDCollisionConstraint<T, d>::Apply(const T Dt, FRigidBodyContactConstraint& Constraint)
+	void TPBDCollisionConstraint<T, d>::Apply(const T Dt, FRigidBodyContactConstraint& Constraint, const int32 It, const int32 NumIts)
 	{
 		TGenericParticleHandle<T, d> Particle0 = TGenericParticleHandle<T, d>(Constraint.Particle[0]);
 		TGenericParticleHandle<T, d> Particle1 = TGenericParticleHandle<T, d>(Constraint.Particle[1]);
@@ -332,192 +341,202 @@ namespace Chaos
 			ensure(!PBDRigid0 || PBDRigid0->Sleeping());
 			return;
 		}
-		UpdateConstraint<ECollisionUpdateType::Deepest>(MThickness, Constraint);
-		if (Constraint.GetPhi() >= MThickness)
+
+		for (int32 PairIt = 0; PairIt < MApplyPairIterations; ++PairIt)
 		{
-			return;
-		}
+			UpdateConstraint<ECollisionUpdateType::Deepest>(MThickness, Constraint);
+			if (Constraint.GetPhi() >= MThickness)
+			{
+				return;
+			}
 
-		// @todo(ccaulfield): CHAOS_PARTICLEHANDLE_TODO what's the best way to manage external per-particle data?
-		Particle0->AuxilaryValue(MCollided) = true;
-		Particle1->AuxilaryValue(MCollided) = true;
+			// @todo(ccaulfield): CHAOS_PARTICLEHANDLE_TODO what's the best way to manage external per-particle data?
+			Particle0->AuxilaryValue(MCollided) = true;
+			Particle1->AuxilaryValue(MCollided) = true;
+
+			// @todo(ccaulfield): CHAOS_PARTICLEHANDLE_TODO split function to avoid ifs
+			const TVector<T, d> ZeroVector = TVector<T, d>(0);
+			const TRotation<T, d>& Q0 = Particle0->Q();
+			const TRotation<T, d>& Q1 = Particle1->Q();
+			const TVector<T, d>& P0 = Particle0->P();
+			const TVector<T, d>& P1 = Particle1->P();
+			const TVector<T, d>& V0 = Particle0->V();
+			const TVector<T, d>& V1 = Particle1->V();
+			const TVector<T, d>& W0 = Particle0->W();
+			const TVector<T, d>& W1 = Particle1->W();
+			TSerializablePtr<FChaosPhysicsMaterial> PhysicsMaterial0 = Particle0->AuxilaryValue(MPhysicsMaterials);
+			TSerializablePtr<FChaosPhysicsMaterial> PhysicsMaterial1 = Particle1->AuxilaryValue(MPhysicsMaterials);
+
+			TContactData<T, d> & Contact = Constraint.ShapeManifold.Manifold;
 
-		// @todo(ccaulfield): CHAOS_PARTICLEHANDLE_TODO split function to avoid ifs
-		const TVector<T, d> ZeroVector = TVector<T, d>(0);
-		const TRotation<T, d>& Q0 = Particle0->Q();
-		const TRotation<T, d>& Q1 = Particle1->Q();
-		const TVector<T, d>& P0 = Particle0->P();
-		const TVector<T, d>& P1 = Particle1->P();
-		const TVector<T, d>& V0 = Particle0->V();
-		const TVector<T, d>& V1 = Particle1->V();
-		const TVector<T, d>& W0 = Particle0->W();
-		const TVector<T, d>& W1 = Particle1->W();
-		TSerializablePtr<FChaosPhysicsMaterial> PhysicsMaterial0 = Particle0->AuxilaryValue(MPhysicsMaterials);
-		TSerializablePtr<FChaosPhysicsMaterial> PhysicsMaterial1 = Particle1->AuxilaryValue(MPhysicsMaterials);
-
-		TContactData<T, d> & Contact = Constraint.ShapeManifold.Manifold;
+			TVector<T, d> VectorToPoint1 = Contact.Location - P0;
+			TVector<T, d> VectorToPoint2 = Contact.Location - P1;
+			TVector<T, d> Body1Velocity = V0 + TVector<T, d>::CrossProduct(W0, VectorToPoint1);
+			TVector<T, d> Body2Velocity = V1 + TVector<T, d>::CrossProduct(W1, VectorToPoint2);
+			TVector<T, d> RelativeVelocity = Body1Velocity - Body2Velocity;
+			T RelativeNormalVelocity = TVector<T, d>::DotProduct(RelativeVelocity, Contact.Normal);
 
-		TVector<T, d> VectorToPoint1 = Contact.Location - P0;
-		TVector<T, d> VectorToPoint2 = Contact.Location - P1;
-		TVector<T, d> Body1Velocity = V0 + TVector<T, d>::CrossProduct(W0, VectorToPoint1);
-		TVector<T, d> Body2Velocity = V1 + TVector<T, d>::CrossProduct(W1, VectorToPoint2);
-		TVector<T, d> RelativeVelocity = Body1Velocity - Body2Velocity;
-		if (TVector<T, d>::DotProduct(RelativeVelocity, Contact.Normal) < 0) // ignore separating constraints
-		{
-			PMatrix<T, d, d> WorldSpaceInvI1 = PBDRigid0 ? (Q0 * FMatrix::Identity).GetTransposed() * PBDRigid0->InvI() * (Q0 * FMatrix::Identity) : PMatrix<T, d, d>(0);
-			PMatrix<T, d, d> WorldSpaceInvI2 = PBDRigid1 ? (Q1 * FMatrix::Identity).GetTransposed() * PBDRigid1->InvI() * (Q1 * FMatrix::Identity) : PMatrix<T, d, d>(0);
-			PMatrix<T, d, d> Factor =
-				(PBDRigid0 ? ComputeFactorMatrix3(VectorToPoint1, WorldSpaceInvI1, PBDRigid0->InvM()) : PMatrix<T, d, d>(0)) +
-				(PBDRigid1 ? ComputeFactorMatrix3(VectorToPoint2, WorldSpaceInvI2, PBDRigid1->InvM()) : PMatrix<T, d, d>(0));
-			TVector<T, d> Impulse;
-			TVector<T, d> AngularImpulse(0);
-
-			// Resting contact if very close to the surface
-			T Restitution = (T)0;
-			T Friction = (T)0;
-			bool bApplyRestitution = (RelativeVelocity.Size() > (2 * 980 * Dt));
-			if (PhysicsMaterial0 && PhysicsMaterial1)
+			if (RelativeNormalVelocity < 0) // ignore separating constraints
 			{
-				if (bApplyRestitution)
+				PMatrix<T, d, d> WorldSpaceInvI1 = PBDRigid0 ? (Q0 * FMatrix::Identity).GetTransposed() * PBDRigid0->InvI() * (Q0 * FMatrix::Identity) : PMatrix<T, d, d>(0);
+				PMatrix<T, d, d> WorldSpaceInvI2 = PBDRigid1 ? (Q1 * FMatrix::Identity).GetTransposed() * PBDRigid1->InvI() * (Q1 * FMatrix::Identity) : PMatrix<T, d, d>(0);
+				PMatrix<T, d, d> Factor =
+					(PBDRigid0 ? ComputeFactorMatrix3(VectorToPoint1, WorldSpaceInvI1, PBDRigid0->InvM()) : PMatrix<T, d, d>(0)) +
+					(PBDRigid1 ? ComputeFactorMatrix3(VectorToPoint2, WorldSpaceInvI2, PBDRigid1->InvM()) : PMatrix<T, d, d>(0));
+				TVector<T, d> Impulse;
+				TVector<T, d> AngularImpulse(0);
+
+				// Resting contact if very close to the surface
+				T Restitution = (T)0;
+				T Friction = (T)0;
+				bool bApplyRestitution = (RelativeVelocity.Size() > (2 * 980 * Dt));
+				if (PhysicsMaterial0 && PhysicsMaterial1)
 				{
-					Restitution = FMath::Min(PhysicsMaterial0->Restitution, PhysicsMaterial1->Restitution);
+					if (bApplyRestitution)
+					{
+						Restitution = FMath::Min(PhysicsMaterial0->Restitution, PhysicsMaterial1->Restitution);
+					}
+					Friction = FMath::Max(PhysicsMaterial0->Friction, PhysicsMaterial1->Friction);
 				}
-				Friction = FMath::Max(PhysicsMaterial0->Friction, PhysicsMaterial1->Friction);
-			}
-			else if (PhysicsMaterial0)
-			{
-				if (bApplyRestitution)
+				else if (PhysicsMaterial0)
 				{
-					Restitution = PhysicsMaterial0->Restitution;
+					if (bApplyRestitution)
+					{
+						Restitution = PhysicsMaterial0->Restitution;
+					}
+					Friction = PhysicsMaterial0->Friction;
 				}
-				Friction = PhysicsMaterial0->Friction;
-			}
-			else if (PhysicsMaterial1)
-			{
-				if (bApplyRestitution)
+				else if (PhysicsMaterial1)
 				{
-					Restitution = PhysicsMaterial1->Restitution;
+					if (bApplyRestitution)
+					{
+						Restitution = PhysicsMaterial1->Restitution;
+					}
+					Friction = PhysicsMaterial1->Friction;
 				}
-				Friction = PhysicsMaterial1->Friction;
-			}
 
-			if (Friction)
-			{
-				T RelativeNormalVelocity = TVector<T, d>::DotProduct(RelativeVelocity, Contact.Normal);
-				if (RelativeNormalVelocity > 0)
+				if (CollisionFrictionOverride >= 0)
 				{
-					RelativeNormalVelocity = 0;
+					Friction = CollisionFrictionOverride;
 				}
-				TVector<T, d> VelocityChange = -(Restitution * RelativeNormalVelocity * Contact.Normal + RelativeVelocity);
-				T NormalVelocityChange = TVector<T, d>::DotProduct(VelocityChange, Contact.Normal);
-				PMatrix<T, d, d> FactorInverse = Factor.Inverse();
-				TVector<T, d> MinimalImpulse = FactorInverse * VelocityChange;
-				const T MinimalImpulseDotNormal = TVector<T, d>::DotProduct(MinimalImpulse, Contact.Normal);
-				const T TangentialSize = (MinimalImpulse - MinimalImpulseDotNormal * Contact.Normal).Size();
-				if (TangentialSize <= Friction * MinimalImpulseDotNormal)
+
+				if (Friction)
 				{
-					//within friction cone so just solve for static friction stopping the object
-					Impulse = MinimalImpulse;
-					if (MAngularFriction)
+					if (RelativeNormalVelocity > 0)
 					{
-						TVector<T, d> RelativeAngularVelocity = W0 - W1;
-						T AngularNormal = TVector<T, d>::DotProduct(RelativeAngularVelocity, Contact.Normal);
-						TVector<T, d> AngularTangent = RelativeAngularVelocity - AngularNormal * Contact.Normal;
-						TVector<T, d> FinalAngularVelocity = FMath::Sign(AngularNormal) * FMath::Max((T)0, FMath::Abs(AngularNormal) - MAngularFriction * NormalVelocityChange) * Contact.Normal + FMath::Max((T)0, AngularTangent.Size() - MAngularFriction * NormalVelocityChange) * AngularTangent.GetSafeNormal();
-						TVector<T, d> Delta = FinalAngularVelocity - RelativeAngularVelocity;
-						if (!PBDRigid0 && PBDRigid1)
-						{
-							PMatrix<T, d, d> WorldSpaceI2 = (Q1 * FMatrix::Identity) * PBDRigid1->I() * (Q1 * FMatrix::Identity).GetTransposed();
-							TVector<T, d> ImpulseDelta = PBDRigid1->M() * TVector<T, d>::CrossProduct(VectorToPoint2, Delta);
-							Impulse += ImpulseDelta;
-							AngularImpulse += WorldSpaceI2 * Delta - TVector<T, d>::CrossProduct(VectorToPoint2, ImpulseDelta);
-						}
-						else if (PBDRigid0 && !PBDRigid1)
+						RelativeNormalVelocity = 0;
+					}
+					TVector<T, d> VelocityChange = -(Restitution * RelativeNormalVelocity * Contact.Normal + RelativeVelocity);
+					T NormalVelocityChange = TVector<T, d>::DotProduct(VelocityChange, Contact.Normal);
+					PMatrix<T, d, d> FactorInverse = Factor.Inverse();
+					TVector<T, d> MinimalImpulse = FactorInverse * VelocityChange;
+					const T MinimalImpulseDotNormal = TVector<T, d>::DotProduct(MinimalImpulse, Contact.Normal);
+					const T TangentialSize = (MinimalImpulse - MinimalImpulseDotNormal * Contact.Normal).Size();
+					if (TangentialSize <= Friction * MinimalImpulseDotNormal)
+					{
+						//within friction cone so just solve for static friction stopping the object
+						Impulse = MinimalImpulse;
+						if (MAngularFriction)
 						{
-							PMatrix<T, d, d> WorldSpaceI1 = (Q0 * FMatrix::Identity) * PBDRigid0->I() * (Q0 * FMatrix::Identity).GetTransposed();
-							TVector<T, d> ImpulseDelta = PBDRigid0->M() * TVector<T, d>::CrossProduct(VectorToPoint1, Delta);
-							Impulse += ImpulseDelta;
-							AngularImpulse += WorldSpaceI1 * Delta - TVector<T, d>::CrossProduct(VectorToPoint1, ImpulseDelta);
+							TVector<T, d> RelativeAngularVelocity = W0 - W1;
+							T AngularNormal = TVector<T, d>::DotProduct(RelativeAngularVelocity, Contact.Normal);
+							TVector<T, d> AngularTangent = RelativeAngularVelocity - AngularNormal * Contact.Normal;
+							TVector<T, d> FinalAngularVelocity = FMath::Sign(AngularNormal) * FMath::Max((T)0, FMath::Abs(AngularNormal) - MAngularFriction * NormalVelocityChange) * Contact.Normal + FMath::Max((T)0, AngularTangent.Size() - MAngularFriction * NormalVelocityChange) * AngularTangent.GetSafeNormal();
+							TVector<T, d> Delta = FinalAngularVelocity - RelativeAngularVelocity;
+							if (!PBDRigid0 && PBDRigid1)
+							{
+								PMatrix<T, d, d> WorldSpaceI2 = (Q1 * FMatrix::Identity) * PBDRigid1->I() * (Q1 * FMatrix::Identity).GetTransposed();
+								TVector<T, d> ImpulseDelta = PBDRigid1->M() * TVector<T, d>::CrossProduct(VectorToPoint2, Delta);
+								Impulse += ImpulseDelta;
+								AngularImpulse += WorldSpaceI2 * Delta - TVector<T, d>::CrossProduct(VectorToPoint2, ImpulseDelta);
+							}
+							else if (PBDRigid0 && !PBDRigid1)
+							{
+								PMatrix<T, d, d> WorldSpaceI1 = (Q0 * FMatrix::Identity) * PBDRigid0->I() * (Q0 * FMatrix::Identity).GetTransposed();
+								TVector<T, d> ImpulseDelta = PBDRigid0->M() * TVector<T, d>::CrossProduct(VectorToPoint1, Delta);
+								Impulse += ImpulseDelta;
+								AngularImpulse += WorldSpaceI1 * Delta - TVector<T, d>::CrossProduct(VectorToPoint1, ImpulseDelta);
+							}
+							else if (PBDRigid0 && PBDRigid1)
+							{
+								PMatrix<T, d, d> Cross1(0, VectorToPoint1.Z, -VectorToPoint1.Y, -VectorToPoint1.Z, 0, VectorToPoint1.X, VectorToPoint1.Y, -VectorToPoint1.X, 0);
+								PMatrix<T, d, d> Cross2(0, VectorToPoint2.Z, -VectorToPoint2.Y, -VectorToPoint2.Z, 0, VectorToPoint2.X, VectorToPoint2.Y, -VectorToPoint2.X, 0);
+								PMatrix<T, d, d> CrossI1 = Cross1 * WorldSpaceInvI1;
+								PMatrix<T, d, d> CrossI2 = Cross2 * WorldSpaceInvI2;
+								PMatrix<T, d, d> Diag1 = CrossI1 * Cross1.GetTransposed() + CrossI2 * Cross2.GetTransposed();
+								Diag1.M[0][0] += PBDRigid0->InvM() + PBDRigid1->InvM();
+								Diag1.M[1][1] += PBDRigid0->InvM() + PBDRigid1->InvM();
+								Diag1.M[2][2] += PBDRigid0->InvM() + PBDRigid1->InvM();
+								PMatrix<T, d, d> OffDiag1 = (CrossI1 + CrossI2) * -1;
+								PMatrix<T, d, d> Diag2 = (WorldSpaceInvI1 + WorldSpaceInvI2).Inverse();
+								PMatrix<T, d, d> OffDiag1Diag2 = OffDiag1 * Diag2;
+								TVector<T, d> ImpulseDelta = PMatrix<T, d, d>((Diag1 - OffDiag1Diag2 * OffDiag1.GetTransposed()).Inverse())* ((OffDiag1Diag2 * -1) * Delta);
+								Impulse += ImpulseDelta;
+								AngularImpulse += Diag2 * (Delta - PMatrix<T, d, d>(OffDiag1.GetTransposed()) * ImpulseDelta);
+							}
 						}
-						else if (PBDRigid0 && PBDRigid1)
+					}
+					else
+					{
+						//outside friction cone, solve for normal relative velocity and keep tangent at cone edge
+						TVector<T, d> Tangent = (RelativeVelocity - TVector<T, d>::DotProduct(RelativeVelocity, Contact.Normal) * Contact.Normal).GetSafeNormal();
+						TVector<T, d> DirectionalFactor = Factor * (Contact.Normal - Friction * Tangent);
+						T ImpulseDenominator = TVector<T, d>::DotProduct(Contact.Normal, DirectionalFactor);
+						if (!ensureMsgf(FMath::Abs(ImpulseDenominator) > SMALL_NUMBER, TEXT("Constraint:%s\n\nParticle:%s\n\nLevelset:%s\n\nDirectionalFactor:%s, ImpulseDenominator:%f"),
+							*Constraint.ToString(),
+							*Particle0->ToString(),
+							*Particle1->ToString(),
+							*DirectionalFactor.ToString(), ImpulseDenominator))
 						{
-							PMatrix<T, d, d> Cross1(0, VectorToPoint1.Z, -VectorToPoint1.Y, -VectorToPoint1.Z, 0, VectorToPoint1.X, VectorToPoint1.Y, -VectorToPoint1.X, 0);
-							PMatrix<T, d, d> Cross2(0, VectorToPoint2.Z, -VectorToPoint2.Y, -VectorToPoint2.Z, 0, VectorToPoint2.X, VectorToPoint2.Y, -VectorToPoint2.X, 0);
-							PMatrix<T, d, d> CrossI1 = Cross1 * WorldSpaceInvI1;
-							PMatrix<T, d, d> CrossI2 = Cross2 * WorldSpaceInvI2;
-							PMatrix<T, d, d> Diag1 = CrossI1 * Cross1.GetTransposed() + CrossI2 * Cross2.GetTransposed();
-							Diag1.M[0][0] += PBDRigid0->InvM() + PBDRigid1->InvM();
-							Diag1.M[1][1] += PBDRigid0->InvM() + PBDRigid1->InvM();
-							Diag1.M[2][2] += PBDRigid0->InvM() + PBDRigid1->InvM();
-							PMatrix<T, d, d> OffDiag1 = (CrossI1 + CrossI2) * -1;
-							PMatrix<T, d, d> Diag2 = (WorldSpaceInvI1 + WorldSpaceInvI2).Inverse();
-							PMatrix<T, d, d> OffDiag1Diag2 = OffDiag1 * Diag2;
-							TVector<T, d> ImpulseDelta = PMatrix<T, d, d>((Diag1 - OffDiag1Diag2 * OffDiag1.GetTransposed()).Inverse())* ((OffDiag1Diag2 * -1) * Delta);
-							Impulse += ImpulseDelta;
-							AngularImpulse += Diag2 * (Delta - PMatrix<T, d, d>(OffDiag1.GetTransposed()) * ImpulseDelta);
+							ImpulseDenominator = (T)1;
 						}
+
+						const T ImpulseMag = -(1 + Restitution) * RelativeNormalVelocity / ImpulseDenominator;
+						Impulse = ImpulseMag * (Contact.Normal - Friction * Tangent);
 					}
 				}
 				else
 				{
-					//outside friction cone, solve for normal relative velocity and keep tangent at cone edge
-					TVector<T, d> Tangent = (RelativeVelocity - TVector<T, d>::DotProduct(RelativeVelocity, Contact.Normal) * Contact.Normal).GetSafeNormal();
-					TVector<T, d> DirectionalFactor = Factor * (Contact.Normal - Friction * Tangent);
-					T ImpulseDenominator = TVector<T, d>::DotProduct(Contact.Normal, DirectionalFactor);
-					if (!ensureMsgf(FMath::Abs(ImpulseDenominator) > SMALL_NUMBER, TEXT("Constraint:%s\n\nParticle:%s\n\nLevelset:%s\n\nDirectionalFactor:%s, ImpulseDenominator:%f"),
+					T ImpulseDenominator = TVector<T, d>::DotProduct(Contact.Normal, Factor * Contact.Normal);
+					TVector<T, d> ImpulseNumerator = -(1 + Restitution) * TVector<T, d>::DotProduct(RelativeVelocity, Contact.Normal)* Contact.Normal;
+					if (!ensureMsgf(FMath::Abs(ImpulseDenominator) > SMALL_NUMBER, TEXT("Constraint:%s\n\nParticle:%s\n\nLevelset:%s\n\nFactor*Constraint.Normal:%s, ImpulseDenominator:%f"),
 						*Constraint.ToString(),
 						*Particle0->ToString(),
 						*Particle1->ToString(),
-						*DirectionalFactor.ToString(), ImpulseDenominator))
+						*(Factor * Contact.Normal).ToString(), ImpulseDenominator))
 					{
 						ImpulseDenominator = (T)1;
 					}
-
-					const T ImpulseMag = -(1 + Restitution) * RelativeNormalVelocity / ImpulseDenominator;
-					Impulse = ImpulseMag * (Contact.Normal - Friction * Tangent);
+					Impulse = ImpulseNumerator / ImpulseDenominator;
 				}
-			}
-			else
-			{
-				T ImpulseDenominator = TVector<T, d>::DotProduct(Contact.Normal, Factor * Contact.Normal);
-				TVector<T, d> ImpulseNumerator = -(1 + Restitution) * TVector<T, d>::DotProduct(RelativeVelocity, Contact.Normal)* Contact.Normal;
-				if (!ensureMsgf(FMath::Abs(ImpulseDenominator) > SMALL_NUMBER, TEXT("Constraint:%s\n\nParticle:%s\n\nLevelset:%s\n\nFactor*Constraint.Normal:%s, ImpulseDenominator:%f"),
-					*Constraint.ToString(),
-					*Particle0->ToString(),
-					*Particle1->ToString(),
-					*(Factor * Contact.Normal).ToString(), ImpulseDenominator))
+				Impulse = GetEnergyClampedImpulse(Constraint, Impulse, VectorToPoint1, VectorToPoint2, Body1Velocity, Body2Velocity);
+				Constraint.AccumulatedImpulse += Impulse;
+				if (PBDRigid0)
 				{
-					ImpulseDenominator = (T)1;
+					// Velocity update for next step
+					FVec3 NetAngularImpulse = FVec3::CrossProduct(VectorToPoint1, Impulse) + AngularImpulse;
+					FVec3 DV = PBDRigid0->InvM() * Impulse;
+					FVec3 DW = WorldSpaceInvI1 * NetAngularImpulse;
+					PBDRigid0->V() += DV;
+					PBDRigid0->W() += DW;
+					// Position update as part of pbd
+					PBDRigid0->P() += DV * Dt;
+					PBDRigid0->Q() += FRotation3::FromElements(DW, 0.f) * Q0 * Dt * T(0.5);
+					PBDRigid0->Q().Normalize();
+				}
+				if (PBDRigid1)
+				{
+					// Velocity update for next step
+					FVec3 NetAngularImpulse = FVec3::CrossProduct(VectorToPoint2, -Impulse) - AngularImpulse;
+					FVec3 DV = -PBDRigid1->InvM() * Impulse;
+					FVec3 DW = WorldSpaceInvI2 * NetAngularImpulse;
+					PBDRigid1->V() += DV;
+					PBDRigid1->W() += DW;
+					// Position update as part of pbd
+					PBDRigid1->P() += DV * Dt;
+					PBDRigid1->Q() += FRotation3::FromElements(DW, 0.f) * Q1 * Dt * T(0.5);
+					PBDRigid1->Q().Normalize();
 				}
-				Impulse = ImpulseNumerator / ImpulseDenominator;
-			}
-			Impulse = GetEnergyClampedImpulse(Constraint, Impulse, VectorToPoint1, VectorToPoint2, Body1Velocity, Body2Velocity);
-			Constraint.AccumulatedImpulse += Impulse;
-			if (PBDRigid0)
-			{
-				// Velocity update for next step
-				TVector<T, d> AngularImpulse1 = TVector<T, d>::CrossProduct(VectorToPoint1, Impulse) + AngularImpulse;
-				FVec3 DV1 = PBDRigid0->InvM() * Impulse;
-				FVec3 DW1 = WorldSpaceInvI1 * AngularImpulse1;
-				PBDRigid0->V() += DV1;
-				PBDRigid0->W() += DW1;
-				// Position update as part of pbd
-				PBDRigid0->P() += DV1 * Dt;
-				PBDRigid0->Q() += TRotation<T, d>::FromElements(DW1, 0.f) * Q0 * Dt * T(0.5);
-				PBDRigid0->Q().Normalize();
-			}
-			if (PBDRigid1)
-			{
-				// Velocity update for next step
-				TVector<T, d> AngularImpulse2 = TVector<T, d>::CrossProduct(VectorToPoint2, -Impulse) - AngularImpulse;
-				FVec3 DV2 = -PBDRigid1->InvM() * Impulse;
-				FVec3 DW2 = WorldSpaceInvI2 * AngularImpulse2;
-				PBDRigid1->V() += DV2;
-				PBDRigid1->W() += DW2;
-				// Position update as part of pbd
-				PBDRigid1->P() += DV2 * Dt;
-				PBDRigid1->Q() += TRotation<T, d>::FromElements(DW2, 0.f) * Q1 * Dt * T(0.5);
-				PBDRigid1->Q().Normalize();
 			}
 		}
 	}
@@ -527,12 +546,12 @@ namespace Chaos
 	void TPBDCollisionConstraint<T, d>::Apply(const T Dt, const TArray<FConstraintContainerHandle*>& InConstraintHandles, const int32 It, const int32 NumIts)
 	{
 		SCOPE_CYCLE_COUNTER(STAT_Apply);
-		if (bEnableVelocitySolve)
+		if (MApplyPairIterations > 0)
 		{
 			PhysicsParallelFor(InConstraintHandles.Num(), [&](int32 ConstraintHandleIndex) {
 				FConstraintContainerHandle* ConstraintHandle = InConstraintHandles[ConstraintHandleIndex];
 				check(ConstraintHandle != nullptr);
-				Apply(Dt, Constraints[ConstraintHandle->GetConstraintIndex()]);
+				Apply(Dt, Constraints[ConstraintHandle->GetConstraintIndex()], It, NumIts);
 			}, bDisableCollisionParallelFor);
 		}
 
@@ -577,7 +596,7 @@ namespace Chaos
 		const bool IsTemporarilyStatic0 = IsTemporarilyStatic.Contains(Particle0);
 		const bool IsTemporarilyStatic1 = IsTemporarilyStatic.Contains(Particle1);
 
-		for (int32 PairIteration = 0; PairIteration < MPairIterations; ++PairIteration)
+		for (int32 PairIteration = 0; PairIteration < MApplyPushOutPairIterations; ++PairIteration)
 		{
 			UpdateConstraint<ECollisionUpdateType::Deepest>(MThickness, Constraint);
 
@@ -668,7 +687,7 @@ namespace Chaos
 		SCOPE_CYCLE_COUNTER(STAT_ApplyPushOut);
 
 		bool NeedsAnotherIteration = false;
-		if (MPairIterations > 0)
+		if (MApplyPushOutPairIterations > 0)
 		{
 			PhysicsParallelFor(InConstraintHandles.Num(), [&](int32 ConstraintHandleIndex) {
 				FConstraintContainerHandle* ConstraintHandle = InConstraintHandles[ConstraintHandleIndex];
diff --git a/Engine/Source/Runtime/Experimental/Chaos/Private/Chaos/PBDJointConstraintUtilities.cpp b/Engine/Source/Runtime/Experimental/Chaos/Private/Chaos/PBDJointConstraintUtilities.cpp
--- a/Engine/Source/Runtime/Experimental/Chaos/Private/Chaos/PBDJointConstraintUtilities.cpp
+++ b/Engine/Source/Runtime/Experimental/Chaos/Private/Chaos/PBDJointConstraintUtilities.cpp
@@ -7,7 +7,6 @@
 
 namespace Chaos
 {
-	
 	FVec3 FPBDJointUtilities::ConditionInertia(const FVec3& InI, const FReal MaxRatio)
 	{
 		FReal IMin = InI.Min();
@@ -356,12 +355,19 @@ namespace Chaos
 		const int32 Index0,
 		const int32 Index1,
 		FVec3& P0,
+		FVec3& V0,
 		FVec3& P1,
+		FVec3& V1,
 		const FVec3& DP0,
 		const FVec3& DP1)
 	{
 		P0 = P0 + Stiffness * DP0;
 		P1 = P1 + Stiffness * DP1;
+		if (Dt > SMALL_NUMBER)
+		{
+			V0 = V0 + Stiffness * DP0 / Dt;
+			V1 = V1 + Stiffness * DP1 / Dt;
+		}
 	}
 
 	
@@ -381,8 +387,8 @@ namespace Chaos
 	{
 		V0 = V0 + Stiffness * DV0;
 		V1 = V1 + Stiffness * DV1;
-
-		ApplyPositionDelta(Dt, SolverSettings, JointSettings, Stiffness, Index0, Index1, P0, P1, DV0 * Dt, DV1 * Dt);
+		P0 = P0 + Stiffness * DV0 * Dt;
+		P1 = P1 + Stiffness * DV1 * Dt;
 	}
 
 	
@@ -394,7 +400,9 @@ namespace Chaos
 		const int32 Index0,
 		const int32 Index1,
 		FRotation3& Q0,
+		FVec3& W0,
 		FRotation3& Q1,
+		FVec3& W1,
 		const FVec3& DR0,
 		const FVec3& DR1)
 	{
@@ -403,6 +411,12 @@ namespace Chaos
 		Q0 = (Q0 + DQ0).GetNormalized();
 		Q1 = (Q1 + DQ1).GetNormalized();
 		Q1.EnforceShortestArcWith(Q0);
+		
+		if (Dt > SMALL_NUMBER)
+		{
+			W0 = W0 + Stiffness * DR0 / Dt;
+			W1 = W1 + Stiffness * DR1 / Dt;
+		}
 	}
 
 
@@ -424,7 +438,11 @@ namespace Chaos
 		W0 = W0 + Stiffness * DW0;
 		W1 = W1 + Stiffness * DW1;
 
-		ApplyRotationDelta(Dt, SolverSettings, JointSettings, Stiffness, Index0, Index1, Q0, Q1, DW0 * Dt, DW1 * Dt);
+		const FRotation3 DQ0 = (FRotation3::FromElements(Stiffness * DW0 * Dt, 0) * Q0) * (FReal)0.5;
+		const FRotation3 DQ1 = (FRotation3::FromElements(Stiffness * DW1 * Dt, 0) * Q1) * (FReal)0.5;
+		Q0 = (Q0 + DQ0).GetNormalized();
+		Q1 = (Q1 + DQ1).GetNormalized();
+		Q1.EnforceShortestArcWith(Q0);
 	}
 
 
@@ -469,10 +487,10 @@ namespace Chaos
 		const FReal Stiffness,
 		const int32 Index0,
 		const int32 Index1,
-		FVec3& P0,
 		FRotation3& Q0,
-		FVec3& P1,
+		FVec3& W0,
 		FRotation3& Q1,
+		FVec3& W1,
 		const FReal InvM0,
 		const FMatrix33& InvIL0,
 		const FReal InvM1,
@@ -493,7 +511,7 @@ namespace Chaos
 		const FVec3 DR0 = Utilities::Multiply(InvI0, Axis0) * L * Angle0;
 		const FVec3 DR1 = Utilities::Multiply(InvI1, Axis1) * L * Angle1;
 
-		ApplyRotationDelta(Dt, SolverSettings, JointSettings, Stiffness, Index0, Index1, Q0, Q1, DR0, DR1);
+		ApplyRotationDelta(Dt, SolverSettings, JointSettings, Stiffness, Index0, Index1, Q0, W0, Q1, W1, DR0, DR1);
 	}
 
 	
@@ -652,8 +670,12 @@ namespace Chaos
 		const int32 Index1,
 		FVec3& P0,
 		FRotation3& Q0,
+		FVec3& V0,
+		FVec3& W0,
 		FVec3& P1,
 		FRotation3& Q1,
+		FVec3& V1,
+		FVec3& W1,
 		FReal InvM0,
 		const FMatrix33& InvIL0,
 		FReal InvM1,
@@ -689,8 +711,8 @@ namespace Chaos
 		const FVec3 DP1 = -InvM1 * DX;
 		const FVec3 DR0 = Utilities::Multiply(InvI0, FVec3::CrossProduct(X0 - P0, DX));
 		const FVec3 DR1 = Utilities::Multiply(InvI1, FVec3::CrossProduct(X1 - P1, -DX));
-		ApplyPositionDelta(Dt, SolverSettings, JointSettings, Stiffness, Index0, Index1, P0, P1, DP0, DP1);
-		ApplyRotationDelta(Dt, SolverSettings, JointSettings, Stiffness, Index0, Index1, Q0, Q1, DR0, DR1);
+		ApplyPositionDelta(Dt, SolverSettings, JointSettings, Stiffness, Index0, Index1, P0, V0, P1, V1, DP0, DP1);
+		ApplyRotationDelta(Dt, SolverSettings, JointSettings, Stiffness, Index0, Index1, Q0, W0, Q1, W1, DR0, DR1);
 	}
 
 	
@@ -761,8 +783,12 @@ namespace Chaos
 		const int32 Index1,
 		FVec3& P0,
 		FRotation3& Q0,
+		FVec3& V0,
+		FVec3& W0,
 		FVec3& P1,
 		FRotation3& Q1,
+		FVec3& V1,
+		FVec3& W1,
 		FReal InvM0,
 		const FMatrix33& InvIL0,
 		FReal InvM1,
@@ -822,7 +848,7 @@ namespace Chaos
 		const FReal DTwistAngle1 = -DTwistAngle;
 
 		// Apply twist correction
-		ApplyRotationDelta(Dt, SolverSettings, JointSettings, Stiffness, Index0, Index1, P0, Q0, P1, Q1, InvM0, InvIL0, InvM1, InvIL1, TwistAxis0, DTwistAngle0, TwistAxis1, DTwistAngle1);
+		ApplyRotationDelta(Dt, SolverSettings, JointSettings, Stiffness, Index0, Index1, Q0, W0, Q1, W1, InvM0, InvIL0, InvM1, InvIL1, TwistAxis0, DTwistAngle0, TwistAxis1, DTwistAngle1);
 	}
 
 	
@@ -913,8 +939,12 @@ namespace Chaos
 		const int32 Index1,
 		FVec3& P0,
 		FRotation3& Q0,
+		FVec3& V0,
+		FVec3& W0,
 		FVec3& P1,
 		FRotation3& Q1,
+		FVec3& V1,
+		FVec3& W1,
 		FReal InvM0,
 		const FMatrix33& InvIL0,
 		FReal InvM1,
@@ -977,7 +1007,7 @@ namespace Chaos
 		FReal DSwingAngle1 = -DSwingAngle;
 
 		// Apply swing correction
-		ApplyRotationDelta(Dt, SolverSettings, JointSettings, Stiffness, Index0, Index1, P0, Q0, P1, Q1, InvM0, InvIL0, InvM1, InvIL1, SwingAxis0, DSwingAngle0, SwingAxis1, DSwingAngle1);
+		ApplyRotationDelta(Dt, SolverSettings, JointSettings, Stiffness, Index0, Index1, Q0, W0, Q1, W1, InvM0, InvIL0, InvM1, InvIL1, SwingAxis0, DSwingAngle0, SwingAxis1, DSwingAngle1);
 	}
 
 	
@@ -1072,8 +1102,12 @@ namespace Chaos
 		const EJointAngularAxisIndex SwingAxisIndex,
 		FVec3& P0,
 		FRotation3& Q0,
+		FVec3& V0,
+		FVec3& W0,
 		FVec3& P1,
 		FRotation3& Q1,
+		FVec3& V1,
+		FVec3& W1,
 		FReal InvM0,
 		const FMatrix33& InvIL0,
 		FReal InvM1,
@@ -1152,7 +1186,7 @@ namespace Chaos
 			FReal DSwingAngle1 = -DSwingAngle;
 
 			// Apply swing correction
-			ApplyRotationDelta(Dt, SolverSettings, JointSettings, Stiffness, Index0, Index1, P0, Q0, P1, Q1, InvM0, InvIL0, InvM1, InvIL1, SwingAxis0, DSwingAngle0, SwingAxis1, DSwingAngle1);
+			ApplyRotationDelta(Dt, SolverSettings, JointSettings, Stiffness, Index0, Index1, Q0, W0, Q1, W1, InvM0, InvIL0, InvM1, InvIL1, SwingAxis0, DSwingAngle0, SwingAxis1, DSwingAngle1);
 		}
 	}
 
@@ -1263,8 +1297,12 @@ namespace Chaos
 		const int32 Index1,
 		FVec3& P0,
 		FRotation3& Q0,
+		FVec3& V0,
+		FVec3& W0,
 		FVec3& P1,
 		FRotation3& Q1,
+		FVec3& V1,
+		FVec3& W1,
 		FReal InvM0,
 		const FMatrix33& InvIL0,
 		FReal InvM1,
@@ -1306,7 +1344,7 @@ namespace Chaos
 		const FReal DTwistAngle0 = DTwistAngle;
 		const FReal DTwistAngle1 = -DTwistAngle;
 
-		ApplyRotationDelta(Dt, SolverSettings, JointSettings, DriveStiffness, Index0, Index1, P0, Q0, P1, Q1, InvM0, InvIL0, InvM1, InvIL1, TwistAxis0, DTwistAngle0, TwistAxis1, DTwistAngle1);
+		ApplyRotationDelta(Dt, SolverSettings, JointSettings, DriveStiffness, Index0, Index1, Q0, W0, Q1, W1, InvM0, InvIL0, InvM1, InvIL1, TwistAxis0, DTwistAngle0, TwistAxis1, DTwistAngle1);
 	}
 
 	
@@ -1318,8 +1356,12 @@ namespace Chaos
 		const int32 Index1,
 		FVec3& P0,
 		FRotation3& Q0,
+		FVec3& V0,
+		FVec3& W0,
 		FVec3& P1,
 		FRotation3& Q1,
+		FVec3& V1,
+		FVec3& W1,
 		FReal InvM0,
 		const FMatrix33& InvIL0,
 		FReal InvM1,
@@ -1363,7 +1405,7 @@ namespace Chaos
 		FReal DSwingAngle1 = -DSwingAngle;
 
 		// Apply swing correction
-		ApplyRotationDelta(Dt, SolverSettings, JointSettings, DriveStiffness, Index0, Index1, P0, Q0, P1, Q1, InvM0, InvIL0, InvM1, InvIL1, SwingAxis0, DSwingAngle0, SwingAxis1, DSwingAngle1);
+		ApplyRotationDelta(Dt, SolverSettings, JointSettings, DriveStiffness, Index0, Index1, Q0, W0, Q1, W1, InvM0, InvIL0, InvM1, InvIL1, SwingAxis0, DSwingAngle0, SwingAxis1, DSwingAngle1);
 	}
 
 	
@@ -1375,8 +1417,12 @@ namespace Chaos
 		const int32 Index1,
 		FVec3& P0,
 		FRotation3& Q0,
+		FVec3& V0,
+		FVec3& W0,
 		FVec3& P1,
 		FRotation3& Q1,
+		FVec3& V1,
+		FVec3& W1,
 		FReal InvM0,
 		const FMatrix33& InvIL0,
 		FReal InvM1,
@@ -1464,7 +1510,11 @@ namespace Chaos
 		FReal InvM1,
 		const FMatrix33& InvIL1)
 	{
-		ApplyJointTwistConstraint(Dt, SolverSettings, JointSettings, Stiffness, Index0, Index1, P0, Q0, P1, Q1, InvM0, InvIL0, InvM1, InvIL1);
+		FVec3 V0(0, 0, 0);
+		FVec3 W0(0, 0, 0);
+		FVec3 V1(0, 0, 0);
+		FVec3 W1(0, 0, 0);
+		ApplyJointTwistConstraint(Dt, SolverSettings, JointSettings, Stiffness, Index0, Index1, P0, Q0, V0, W0, P1, Q1, V1, W1, InvM0, InvIL0, InvM1, InvIL1);
 	}
 
 	
@@ -1484,7 +1534,11 @@ namespace Chaos
 		FReal InvM1,
 		const FMatrix33& InvIL1)
 	{
-		ApplyJointConeConstraint(Dt, SolverSettings, JointSettings, Stiffness, Index0, Index1, P0, Q0, P1, Q1, InvM0, InvIL0, InvM1, InvIL1);
+		FVec3 V0(0, 0, 0);
+		FVec3 W0(0, 0, 0);
+		FVec3 V1(0, 0, 0);
+		FVec3 W1(0, 0, 0);
+		ApplyJointConeConstraint(Dt, SolverSettings, JointSettings, Stiffness, Index0, Index1, P0, Q0, V0, W0, P1, Q1, V1, W1, InvM0, InvIL0, InvM1, InvIL1);
 	}
 
 	
@@ -1506,6 +1560,10 @@ namespace Chaos
 		FReal InvM1,
 		const FMatrix33& InvIL1)
 	{
-		ApplyJointSwingConstraint(Dt, SolverSettings, JointSettings, Stiffness, Index0, Index1, SwingConstraintIndex, SwingAxisIndex, P0, Q0, P1, Q1, InvM0, InvIL0, InvM1, InvIL1);
+		FVec3 V0(0, 0, 0);
+		FVec3 W0(0, 0, 0);
+		FVec3 V1(0, 0, 0);
+		FVec3 W1(0, 0, 0);
+		ApplyJointSwingConstraint(Dt, SolverSettings, JointSettings, Stiffness, Index0, Index1, SwingConstraintIndex, SwingAxisIndex, P0, Q0, V0, W0, P1, Q1, V1, W1, InvM0, InvIL0, InvM1, InvIL1);
 	}
 }
diff --git a/Engine/Source/Runtime/Experimental/Chaos/Private/Chaos/PBDJointConstraints.cpp b/Engine/Source/Runtime/Experimental/Chaos/Private/Chaos/PBDJointConstraints.cpp
--- a/Engine/Source/Runtime/Experimental/Chaos/Private/Chaos/PBDJointConstraints.cpp
+++ b/Engine/Source/Runtime/Experimental/Chaos/Private/Chaos/PBDJointConstraints.cpp
@@ -15,6 +15,7 @@
 namespace Chaos
 {
 	DECLARE_CYCLE_STAT(TEXT("TPBDJointConstraints::Apply"), STAT_ApplyJointConstraints, STATGROUP_Chaos);
+	DECLARE_CYCLE_STAT(TEXT("TPBDJointConstraints::Apply"), STAT_ApplyPushOutJointConstraints, STATGROUP_Chaos);
 
 	FReal GetLinearStiffness(
 		const FPBDJointSolverSettings& SolverSettings,
@@ -166,14 +167,16 @@ namespace Chaos
 
 	
 	FPBDJointSolverSettings::FPBDJointSolverSettings()
-		: SwingTwistAngleTolerance((FReal)1.0e-6)
+		: ApplyPairIterations(1)
+		, ApplyPushOutPairIterations(1)
+		, SwingTwistAngleTolerance((FReal)1.0e-6)
 		, MinParentMassRatio(0)
 		, MaxInertiaRatio(0)
 		, bEnableVelocitySolve(false)
-		, bEnableLinearLimits(true)
 		, bEnableTwistLimits(true)
 		, bEnableSwingLimits(true)
 		, bEnableDrives(true)
+		, ProjectionPhase(EJointProjectionPhase::None)
 		, LinearProjection((FReal)0)
 		, AngularProjection((FReal)0)
 		, Stiffness((FReal)0)
@@ -404,11 +407,20 @@ namespace Chaos
 		{
 			if (Settings.bEnableVelocitySolve)
 			{
-				SolveVelocity(Dt, ConstraintHandle->GetConstraintIndex(), It, NumIts);
+				SolveVelocity(Dt, ConstraintHandle->GetConstraintIndex(), Settings.ApplyPairIterations, It, NumIts);
 			}
 			else
 			{
-				SolvePosition(Dt, ConstraintHandle->GetConstraintIndex(), It, NumIts);
+				SolvePosition(Dt, ConstraintHandle->GetConstraintIndex(), Settings.ApplyPairIterations, It, NumIts);
+			}
+		}
+
+		if (Settings.ProjectionPhase == EJointProjectionPhase::Apply)
+		{
+			int32 ProjectionIt = NumIts - 1;
+			if (It == ProjectionIt)
+			{
+				ApplyProjection(Dt, InConstraintHandles);
 			}
 		}
 
@@ -421,10 +433,12 @@ namespace Chaos
 	
 	bool FPBDJointConstraints::ApplyPushOut(const FReal Dt, const TArray<FConstraintContainerHandle*>& InConstraintHandles, const int32 It, const int32 NumIts)
 	{
+		SCOPE_CYCLE_COUNTER(STAT_ApplyPushOutJointConstraints);
+
 		// @todo(ccaulfield): track whether we are sufficiently solved
 		bool bNeedsAnotherIteration = true;
 
-		if (Settings.bEnableVelocitySolve)
+		if (Settings.ApplyPushOutPairIterations > 0)
 		{
 			TArray<FConstraintContainerHandle*> SortedConstraintHandles = InConstraintHandles;
 			SortedConstraintHandles.Sort([](const FConstraintContainerHandle& L, const FConstraintContainerHandle& R)
@@ -435,14 +449,17 @@ namespace Chaos
 
 			for (FConstraintContainerHandle* ConstraintHandle : SortedConstraintHandles)
 			{
-				SolvePosition(Dt, ConstraintHandle->GetConstraintIndex(), It, NumIts);
+				SolvePosition(Dt, ConstraintHandle->GetConstraintIndex(), Settings.ApplyPushOutPairIterations, It, NumIts);
 			}
 		}
 
-		if (It == NumIts - 1)
+		if (Settings.ProjectionPhase == EJointProjectionPhase::ApplyPushOut)
 		{
-			// @todo(ccaulfield): should be called by constraint rule after PushOut
-			ApplyProjection(Dt, InConstraintHandles);
+			int32 ProjectionIt = (Settings.ApplyPushOutPairIterations > 0) ? NumIts - 1 : 0;
+			if (It == ProjectionIt)
+			{
+				ApplyProjection(Dt, InConstraintHandles);
+			}
 		}
 
 		return bNeedsAnotherIteration;
@@ -465,7 +482,7 @@ namespace Chaos
 	}
 
 	
-	void FPBDJointConstraints::SolveVelocity(const FReal Dt, const int32 ConstraintIndex, const int32 It, const int32 NumIts)
+	void FPBDJointConstraints::SolveVelocity(const FReal Dt, const int32 ConstraintIndex, const int32 NumPairIts, const int32 It, const int32 NumIts)
 	{
 		const TVector<TGeometryParticleHandle<FReal, 3>*, 2>& Constraint = ConstraintParticles[ConstraintIndex];
 		UE_LOG(LogChaosJoint, Verbose, TEXT("Solve Joint Constraint %d %s %s (dt = %f; it = %d / %d)"), ConstraintIndex, *Constraint[0]->ToString(), *Constraint[1]->ToString(), Dt, It, NumIts);
@@ -520,74 +537,77 @@ namespace Chaos
 		EJointMotionType Swing1Motion = JointSettings.Motion.AngularMotionTypes[(int32)EJointAngularConstraintIndex::Swing1];
 		EJointMotionType Swing2Motion = JointSettings.Motion.AngularMotionTypes[(int32)EJointAngularConstraintIndex::Swing2];
 
-		// Apply angular drives (NOTE: modifies position, not velocity)
-		if (Settings.bEnableDrives)
+		for (int32 PairIt = 0; PairIt < NumPairIts; ++PairIt)
 		{
-			bool bTwistLocked = TwistMotion == EJointMotionType::Locked;
-			bool bSwing1Locked = Swing1Motion == EJointMotionType::Locked;
-			bool bSwing2Locked = Swing2Motion == EJointMotionType::Locked;
-
-			// No SLerp drive if we have a locked rotation (it will be grayed out in the editor in this case, but could still have been set before the rotation was locked)
-			if (JointSettings.Motion.bAngularSLerpDriveEnabled && !bTwistLocked && !bSwing1Locked && !bSwing2Locked)
+			// Apply angular drives (NOTE: modifies position, not velocity)
+			if (Settings.bEnableDrives)
 			{
-				FPBDJointUtilities::ApplyJointSLerpDrive(Dt, Settings, JointSettings, Index0, Index1, P0, Q0, P1, Q1, InvM0, InvIL0, InvM1, InvIL1);
-			}
+				bool bTwistLocked = TwistMotion == EJointMotionType::Locked;
+				bool bSwing1Locked = Swing1Motion == EJointMotionType::Locked;
+				bool bSwing2Locked = Swing2Motion == EJointMotionType::Locked;
 
-			if (JointSettings.Motion.bAngularTwistDriveEnabled && !bTwistLocked)
-			{
-				FPBDJointUtilities::ApplyJointTwistDrive(Dt, Settings, JointSettings, Index0, Index1, P0, Q0, P1, Q1, InvM0, InvIL0, InvM1, InvIL1);
-			}
+				// No SLerp drive if we have a locked rotation (it will be grayed out in the editor in this case, but could still have been set before the rotation was locked)
+				if (JointSettings.Motion.bAngularSLerpDriveEnabled && !bTwistLocked && !bSwing1Locked && !bSwing2Locked)
+				{
+					FPBDJointUtilities::ApplyJointSLerpDrive(Dt, Settings, JointSettings, Index0, Index1, P0, Q0, V0, W0, P1, Q1, V1, W1, InvM0, InvIL0, InvM1, InvIL1);
+				}
 
-			if (JointSettings.Motion.bAngularSwingDriveEnabled && !bSwing1Locked && !bSwing2Locked)
-			{
-				FPBDJointUtilities::ApplyJointConeDrive(Dt, Settings, JointSettings, Index0, Index1, P0, Q0, P1, Q1, InvM0, InvIL0, InvM1, InvIL1);
-			}
-			else if (JointSettings.Motion.bAngularSwingDriveEnabled && !bSwing1Locked)
-			{
-				//FPBDJointUtilities::ApplyJointSwingDrive(Dt, Settings, JointSettings, Index0, Index1, EJointAngularConstraintIndex::Swing1, P0, Q0, P1, Q1, InvM0, InvIL0, InvM1, InvIL1);
-			}
-			else if (JointSettings.Motion.bAngularSwingDriveEnabled && !bSwing2Locked)
-			{
-				//FPBDJointUtilities::ApplyJointSwingDrive(Dt, Settings, JointSettings, Index0, Index1, EJointAngularConstraintIndex::Swing2, P0, Q0, P1, Q1, InvM0, InvIL0, InvM1, InvIL1);
-			}
-		}
+				if (JointSettings.Motion.bAngularTwistDriveEnabled && !bTwistLocked)
+				{
+					FPBDJointUtilities::ApplyJointTwistDrive(Dt, Settings, JointSettings, Index0, Index1, P0, Q0, V0, W0, P1, Q1, V1, W1, InvM0, InvIL0, InvM1, InvIL1);
+				}
 
-		// Apply twist velocity constraint
-		if (Settings.bEnableTwistLimits)
-		{
-			if (TwistMotion != EJointMotionType::Free)
-			{
-				FPBDJointUtilities::ApplyJointTwistVelocityConstraint(Dt, Settings, JointSettings, TwistStiffness, Index0, Index1, P0, Q0, V0, W0, P1, Q1, V1, W1, InvM0, InvIL0, InvM1, InvIL1);
+				if (JointSettings.Motion.bAngularSwingDriveEnabled && !bSwing1Locked && !bSwing2Locked)
+				{
+					FPBDJointUtilities::ApplyJointConeDrive(Dt, Settings, JointSettings, Index0, Index1, P0, Q0, V0, W0, P1, Q1, V1, W1, InvM0, InvIL0, InvM1, InvIL1);
+				}
+				else if (JointSettings.Motion.bAngularSwingDriveEnabled && !bSwing1Locked)
+				{
+					//FPBDJointUtilities::ApplyJointSwingDrive(Dt, Settings, JointSettings, Index0, Index1, EJointAngularConstraintIndex::Swing1, P0, Q0, P1, Q1, InvM0, InvIL0, InvM1, InvIL1);
+				}
+				else if (JointSettings.Motion.bAngularSwingDriveEnabled && !bSwing2Locked)
+				{
+					//FPBDJointUtilities::ApplyJointSwingDrive(Dt, Settings, JointSettings, Index0, Index1, EJointAngularConstraintIndex::Swing2, P0, Q0, P1, Q1, InvM0, InvIL0, InvM1, InvIL1);
+				}
 			}
-		}
 
-		// Apply swing velocity constraints
-		if (Settings.bEnableSwingLimits)
-		{
-			if ((Swing1Motion == EJointMotionType::Limited) && (Swing2Motion == EJointMotionType::Limited))
+			// Apply twist velocity constraint
+			if (Settings.bEnableTwistLimits)
 			{
-				// Swing Cone
-				FPBDJointUtilities::ApplyJointConeVelocityConstraint(Dt, Settings, JointSettings, SwingStiffness, Index0, Index1, P0, Q0, V0, W0, P1, Q1, V1, W1, InvM0, InvIL0, InvM1, InvIL1);
+				if (TwistMotion != EJointMotionType::Free)
+				{
+					FPBDJointUtilities::ApplyJointTwistVelocityConstraint(Dt, Settings, JointSettings, TwistStiffness, Index0, Index1, P0, Q0, V0, W0, P1, Q1, V1, W1, InvM0, InvIL0, InvM1, InvIL1);
+				}
 			}
-			else
+
+			// Apply swing velocity constraints
+			if (Settings.bEnableSwingLimits)
 			{
-				if (Swing1Motion != EJointMotionType::Free)
+				if ((Swing1Motion == EJointMotionType::Limited) && (Swing2Motion == EJointMotionType::Limited))
 				{
-					// Swing Arc/Lock
-					FPBDJointUtilities::ApplyJointSwingVelocityConstraint(Dt, Settings, JointSettings, SwingStiffness, Index0, Index1, EJointAngularConstraintIndex::Swing1, EJointAngularAxisIndex::Swing1, P0, Q0, V0, W0, P1, Q1, V1, W1, InvM0, InvIL0, InvM1, InvIL1);
+					// Swing Cone
+					FPBDJointUtilities::ApplyJointConeVelocityConstraint(Dt, Settings, JointSettings, SwingStiffness, Index0, Index1, P0, Q0, V0, W0, P1, Q1, V1, W1, InvM0, InvIL0, InvM1, InvIL1);
 				}
-				if (Swing2Motion != EJointMotionType::Free)
+				else
 				{
-					// Swing Arc/Lock
-					FPBDJointUtilities::ApplyJointSwingVelocityConstraint(Dt, Settings, JointSettings, SwingStiffness, Index0, Index1, EJointAngularConstraintIndex::Swing2, EJointAngularAxisIndex::Swing2, P0, Q0, V0, W0, P1, Q1, V1, W1, InvM0, InvIL0, InvM1, InvIL1);
+					if (Swing1Motion != EJointMotionType::Free)
+					{
+						// Swing Arc/Lock
+						FPBDJointUtilities::ApplyJointSwingVelocityConstraint(Dt, Settings, JointSettings, SwingStiffness, Index0, Index1, EJointAngularConstraintIndex::Swing1, EJointAngularAxisIndex::Swing1, P0, Q0, V0, W0, P1, Q1, V1, W1, InvM0, InvIL0, InvM1, InvIL1);
+					}
+					if (Swing2Motion != EJointMotionType::Free)
+					{
+						// Swing Arc/Lock
+						FPBDJointUtilities::ApplyJointSwingVelocityConstraint(Dt, Settings, JointSettings, SwingStiffness, Index0, Index1, EJointAngularConstraintIndex::Swing2, EJointAngularAxisIndex::Swing2, P0, Q0, V0, W0, P1, Q1, V1, W1, InvM0, InvIL0, InvM1, InvIL1);
+					}
 				}
 			}
-		}
 
-		// Apply linear velocity  constraints
-		if ((LinearMotion[0] != EJointMotionType::Free) || (LinearMotion[1] != EJointMotionType::Free) || (LinearMotion[2] != EJointMotionType::Free))
-		{
-			FPBDJointUtilities::ApplyJointVelocityConstraint(Dt, Settings, JointSettings, LinearStiffness, Index0, Index1, P0, Q0, V0, W0, P1, Q1, V1, W1, InvM0, InvIL0, InvM1, InvIL1);
+			// Apply linear velocity  constraints
+			if ((LinearMotion[0] != EJointMotionType::Free) || (LinearMotion[1] != EJointMotionType::Free) || (LinearMotion[2] != EJointMotionType::Free))
+			{
+				FPBDJointUtilities::ApplyJointVelocityConstraint(Dt, Settings, JointSettings, LinearStiffness, Index0, Index1, P0, Q0, V0, W0, P1, Q1, V1, W1, InvM0, InvIL0, InvM1, InvIL1);
+			}
 		}
 
 		// Update the particles
@@ -608,7 +628,7 @@ namespace Chaos
 	}
 
 	
-	void FPBDJointConstraints::SolvePosition(const FReal Dt, const int32 ConstraintIndex, const int32 It, const int32 NumIts)
+	void FPBDJointConstraints::SolvePosition(const FReal Dt, const int32 ConstraintIndex, const int32 NumPairIts, const int32 It, const int32 NumIts)
 	{
 		const TVector<TGeometryParticleHandle<FReal, 3>*, 2>& Constraint = ConstraintParticles[ConstraintIndex];
 		UE_LOG(LogChaosJoint, Verbose, TEXT("Solve Joint Constraint %d %s %s (dt = %f; it = %d / %d)"), ConstraintIndex, *Constraint[0]->ToString(), *Constraint[1]->ToString(), Dt, It, NumIts);
@@ -625,8 +645,12 @@ namespace Chaos
 
 		FVec3 P0 = Particle0->P();
 		FRotation3 Q0 = Particle0->Q();
+		FVec3 V0 = Particle0->V();
+		FVec3 W0 = Particle0->W();
 		FVec3 P1 = Particle1->P();
 		FRotation3 Q1 = Particle1->Q();
+		FVec3 V1 = Particle1->V();
+		FVec3 W1 = Particle1->W();
 		float InvM0 = Particle0->InvM();
 		float InvM1 = Particle1->InvM();
 		FMatrix33 InvIL0 = Particle0->InvI();
@@ -654,101 +678,82 @@ namespace Chaos
 			FPBDJointUtilities::GetConditionedInverseMass(Particle0->M(), Particle0->I().GetDiagonal(), Particle1->M(), Particle1->I().GetDiagonal(), InvM0, InvM1, InvIL0, InvIL1, (FReal)0, Settings.MaxInertiaRatio);
 		}
 
-		// Freeze the closest to kinematic connection (if one is closer than the other)
-		if (Settings.bEnableVelocitySolve && (Level0 != Level1))
-		{
-			const FReal FreezeFactor = (FReal)(NumIts - (It + 1)) / (FReal)NumIts;
-			if (Level0 < Level1)
-			{
-				InvM0 = InvM0 * FreezeFactor * FreezeFactor;
-				InvIL0 = InvIL0 * FreezeFactor * FreezeFactor;
-			}
-			else if (Level1 < Level0)
-			{
-				InvM1 = InvM1 * FreezeFactor * FreezeFactor;
-				InvIL1 = InvIL1 * FreezeFactor * FreezeFactor;
-			}
-		}
-
 		const TVector<EJointMotionType, 3>& LinearMotion = JointSettings.Motion.LinearMotionTypes;
 		EJointMotionType TwistMotion = JointSettings.Motion.AngularMotionTypes[(int32)EJointAngularConstraintIndex::Twist];
 		EJointMotionType Swing1Motion = JointSettings.Motion.AngularMotionTypes[(int32)EJointAngularConstraintIndex::Swing1];
 		EJointMotionType Swing2Motion = JointSettings.Motion.AngularMotionTypes[(int32)EJointAngularConstraintIndex::Swing2];
 
-		// Disable a constraint if it has any linear limits?
-		if (!Settings.bEnableLinearLimits && ((LinearMotion[0] == EJointMotionType::Limited) || (LinearMotion[1] == EJointMotionType::Limited) || (LinearMotion[2] == EJointMotionType::Limited)))
+		for (int32 PairIt = 0; PairIt < NumPairIts; ++PairIt)
 		{
-			return;
-		}
-
-		// Apply angular drives (NOTE: modifies position, not velocity)
-		if (!Settings.bEnableVelocitySolve && Settings.bEnableDrives)
-		{
-			bool bTwistLocked = TwistMotion == EJointMotionType::Locked;
-			bool bSwing1Locked = Swing1Motion == EJointMotionType::Locked;
-			bool bSwing2Locked = Swing2Motion == EJointMotionType::Locked;
-
-			// No SLerp drive if we have a locked rotation (it will be grayed out in the editor in this case, but could still have been set before the rotation was locked)
-			if (JointSettings.Motion.bAngularSLerpDriveEnabled && !bTwistLocked && !bSwing1Locked && !bSwing2Locked)
+			// Apply angular drives (NOTE: modifies position, not velocity)
+			if (!Settings.bEnableVelocitySolve && Settings.bEnableDrives)
 			{
-				FPBDJointUtilities::ApplyJointSLerpDrive(Dt, Settings, JointSettings, Index0, Index1, P0, Q0, P1, Q1, InvM0, InvIL0, InvM1, InvIL1);
-			}
+				bool bTwistLocked = TwistMotion == EJointMotionType::Locked;
+				bool bSwing1Locked = Swing1Motion == EJointMotionType::Locked;
+				bool bSwing2Locked = Swing2Motion == EJointMotionType::Locked;
 
-			if (JointSettings.Motion.bAngularTwistDriveEnabled && !bTwistLocked)
-			{
-				FPBDJointUtilities::ApplyJointTwistDrive(Dt, Settings, JointSettings, Index0, Index1, P0, Q0, P1, Q1, InvM0, InvIL0, InvM1, InvIL1);
-			}
+				// No SLerp drive if we have a locked rotation (it will be grayed out in the editor in this case, but could still have been set before the rotation was locked)
+				if (JointSettings.Motion.bAngularSLerpDriveEnabled && !bTwistLocked && !bSwing1Locked && !bSwing2Locked)
+				{
+					FPBDJointUtilities::ApplyJointSLerpDrive(Dt, Settings, JointSettings, Index0, Index1, P0, Q0, V0, W0, P1, Q1, V1, W1, InvM0, InvIL0, InvM1, InvIL1);
+				}
 
-			if (JointSettings.Motion.bAngularSwingDriveEnabled && !bSwing1Locked && !bSwing2Locked)
-			{
-				FPBDJointUtilities::ApplyJointConeDrive(Dt, Settings, JointSettings, Index0, Index1, P0, Q0, P1, Q1, InvM0, InvIL0, InvM1, InvIL1);
-			}
-			else if (JointSettings.Motion.bAngularSwingDriveEnabled && !bSwing1Locked)
-			{
-				//FPBDJointUtilities::ApplyJointSwingDrive(Dt, Settings, JointSettings, Index0, Index1, EJointAngularConstraintIndex::Swing1, P0, Q0, P1, Q1, InvM0, InvIL0, InvM1, InvIL1);
-			}
-			else if (JointSettings.Motion.bAngularSwingDriveEnabled && !bSwing2Locked)
-			{
-				//FPBDJointUtilities::ApplyJointSwingDrive(Dt, Settings, JointSettings, Index0, Index1, EJointAngularConstraintIndex::Swing2, P0, Q0, P1, Q1, InvM0, InvIL0, InvM1, InvIL1);
-			}
-		}
+				if (JointSettings.Motion.bAngularTwistDriveEnabled && !bTwistLocked)
+				{
+					FPBDJointUtilities::ApplyJointTwistDrive(Dt, Settings, JointSettings, Index0, Index1, P0, Q0, V0, W0, P1, Q1, V1, W1, InvM0, InvIL0, InvM1, InvIL1);
+				}
 
-		// Apply twist constraint
-		if (Settings.bEnableTwistLimits)
-		{
-			if (TwistMotion != EJointMotionType::Free)
-			{
-				FPBDJointUtilities::ApplyJointTwistConstraint(Dt, Settings, JointSettings, TwistStiffness, Index0, Index1, P0, Q0, P1, Q1, InvM0, InvIL0, InvM1, InvIL1);
+				if (JointSettings.Motion.bAngularSwingDriveEnabled && !bSwing1Locked && !bSwing2Locked)
+				{
+					FPBDJointUtilities::ApplyJointConeDrive(Dt, Settings, JointSettings, Index0, Index1, P0, Q0, V0, W0, P1, Q1, V1, W1, InvM0, InvIL0, InvM1, InvIL1);
+				}
+				else if (JointSettings.Motion.bAngularSwingDriveEnabled && !bSwing1Locked)
+				{
+					//FPBDJointUtilities::ApplyJointSwingDrive(Dt, Settings, JointSettings, Index0, Index1, EJointAngularConstraintIndex::Swing1, P0, Q0, P1, Q1, InvM0, InvIL0, InvM1, InvIL1);
+				}
+				else if (JointSettings.Motion.bAngularSwingDriveEnabled && !bSwing2Locked)
+				{
+					//FPBDJointUtilities::ApplyJointSwingDrive(Dt, Settings, JointSettings, Index0, Index1, EJointAngularConstraintIndex::Swing2, P0, Q0, P1, Q1, InvM0, InvIL0, InvM1, InvIL1);
+				}
 			}
-		}
 
-		// Apply swing constraints
-		if (Settings.bEnableSwingLimits)
-		{
-			if ((Swing1Motion == EJointMotionType::Limited) && (Swing2Motion == EJointMotionType::Limited))
+			// Apply twist constraint
+			if (Settings.bEnableTwistLimits)
 			{
-				// Swing Cone
-				FPBDJointUtilities::ApplyJointConeConstraint(Dt, Settings, JointSettings, SwingStiffness, Index0, Index1, P0, Q0, P1, Q1, InvM0, InvIL0, InvM1, InvIL1);
+				if (TwistMotion != EJointMotionType::Free)
+				{
+					FPBDJointUtilities::ApplyJointTwistConstraint(Dt, Settings, JointSettings, TwistStiffness, Index0, Index1, P0, Q0, V0, W0, P1, Q1, V1, W1, InvM0, InvIL0, InvM1, InvIL1);
+				}
 			}
-			else
+
+			// Apply swing constraints
+			if (Settings.bEnableSwingLimits)
 			{
-				if (Swing1Motion != EJointMotionType::Free)
+				if ((Swing1Motion == EJointMotionType::Limited) && (Swing2Motion == EJointMotionType::Limited))
 				{
-					// Swing Arc/Lock
-					FPBDJointUtilities::ApplyJointSwingConstraint(Dt, Settings, JointSettings, SwingStiffness, Index0, Index1, EJointAngularConstraintIndex::Swing1, EJointAngularAxisIndex::Swing1, P0, Q0, P1, Q1, InvM0, InvIL0, InvM1, InvIL1);
+					// Swing Cone
+					FPBDJointUtilities::ApplyJointConeConstraint(Dt, Settings, JointSettings, SwingStiffness, Index0, Index1, P0, Q0, V0, W0, P1, Q1, V1, W1, InvM0, InvIL0, InvM1, InvIL1);
 				}
-				if (Swing2Motion != EJointMotionType::Free)
+				else
 				{
-					// Swing Arc/Lock
-					FPBDJointUtilities::ApplyJointSwingConstraint(Dt, Settings, JointSettings, SwingStiffness, Index0, Index1, EJointAngularConstraintIndex::Swing2, EJointAngularAxisIndex::Swing2, P0, Q0, P1, Q1, InvM0, InvIL0, InvM1, InvIL1);
+					if (Swing1Motion != EJointMotionType::Free)
+					{
+						// Swing Arc/Lock
+						FPBDJointUtilities::ApplyJointSwingConstraint(Dt, Settings, JointSettings, SwingStiffness, Index0, Index1, EJointAngularConstraintIndex::Swing1, EJointAngularAxisIndex::Swing1, P0, Q0, V0, W0, P1, Q1, V1, W1, InvM0, InvIL0, InvM1, InvIL1);
+					}
+					if (Swing2Motion != EJointMotionType::Free)
+					{
+						// Swing Arc/Lock
+						FPBDJointUtilities::ApplyJointSwingConstraint(Dt, Settings, JointSettings, SwingStiffness, Index0, Index1, EJointAngularConstraintIndex::Swing2, EJointAngularAxisIndex::Swing2, P0, Q0, V0, W0, P1, Q1, V1, W1, InvM0, InvIL0, InvM1, InvIL1);
+					}
 				}
 			}
-		}
 
-		// Apply linear constraints
-		if ((LinearMotion[0] != EJointMotionType::Free) || (LinearMotion[1] != EJointMotionType::Free) || (LinearMotion[2] != EJointMotionType::Free))
-		{
-			FPBDJointUtilities::ApplyJointPositionConstraint(Dt, Settings, JointSettings, LinearStiffness, Index0, Index1, P0, Q0, P1, Q1, InvM0, InvIL0, InvM1, InvIL1);
+			// Apply linear constraints
+			if ((LinearMotion[0] != EJointMotionType::Free) || (LinearMotion[1] != EJointMotionType::Free) || (LinearMotion[2] != EJointMotionType::Free))
+			{
+				FPBDJointUtilities::ApplyJointPositionConstraint(Dt, Settings, JointSettings, LinearStiffness, Index0, Index1, P0, Q0, V0, W0, P1, Q1, V1, W1, InvM0, InvIL0, InvM1, InvIL1);
+			}
 		}
 
 		// Update the particles
@@ -756,11 +761,15 @@ namespace Chaos
 		{
 			Rigid0->SetP(P0);
 			Rigid0->SetQ(Q0);
+			Rigid0->SetV(V0);
+			Rigid0->SetW(W0);
 		}
 		if (Rigid1)
 		{
 			Rigid1->SetP(P1);
 			Rigid1->SetQ(Q1);
+			Rigid1->SetV(V1);
+			Rigid1->SetW(W1);
 		}
 	}
 
@@ -816,23 +825,29 @@ namespace Chaos
 
 		if (AngularProjectionFactor > 0)
 		{
-			// Remove Twist Error
-			FPBDJointUtilities::ApplyJointTwistProjection(Dt, Settings, JointSettings, AngularProjectionFactor, Index0, Index1, P0, Q0, P1, Q1, InvM0, InvIL0, InvM1, InvIL1);
-
-			// Remove Swing Error
-			if ((Swing1Motion == EJointMotionType::Limited) && (Swing2Motion == EJointMotionType::Limited))
+			if (Settings.bEnableTwistLimits)
 			{
-				FPBDJointUtilities::ApplyJointConeProjection(Dt, Settings, JointSettings, AngularProjectionFactor, Index0, Index1, P0, Q0, P1, Q1, InvM0, InvIL0, InvM1, InvIL1);
+				// Remove Twist Error
+				FPBDJointUtilities::ApplyJointTwistProjection(Dt, Settings, JointSettings, AngularProjectionFactor, Index0, Index1, P0, Q0, P1, Q1, InvM0, InvIL0, InvM1, InvIL1);
 			}
-			else
+
+			if (Settings.bEnableSwingLimits)
 			{
-				if (Swing1Motion != EJointMotionType::Free)
+				// Remove Swing Error
+				if ((Swing1Motion == EJointMotionType::Limited) && (Swing2Motion == EJointMotionType::Limited))
 				{
-					FPBDJointUtilities::ApplyJointSwingProjection(Dt, Settings, JointSettings, AngularProjectionFactor, Index0, Index1, EJointAngularConstraintIndex::Swing1, EJointAngularAxisIndex::Swing1, P0, Q0, P1, Q1, InvM0, InvIL0, InvM1, InvIL1);
+					FPBDJointUtilities::ApplyJointConeProjection(Dt, Settings, JointSettings, AngularProjectionFactor, Index0, Index1, P0, Q0, P1, Q1, InvM0, InvIL0, InvM1, InvIL1);
 				}
-				if (Swing2Motion != EJointMotionType::Free)
+				else
 				{
-					FPBDJointUtilities::ApplyJointSwingProjection(Dt, Settings, JointSettings, AngularProjectionFactor, Index0, Index1, EJointAngularConstraintIndex::Swing2, EJointAngularAxisIndex::Swing2, P0, Q0, P1, Q1, InvM0, InvIL0, InvM1, InvIL1);
+					if (Swing1Motion != EJointMotionType::Free)
+					{
+						FPBDJointUtilities::ApplyJointSwingProjection(Dt, Settings, JointSettings, AngularProjectionFactor, Index0, Index1, EJointAngularConstraintIndex::Swing1, EJointAngularAxisIndex::Swing1, P0, Q0, P1, Q1, InvM0, InvIL0, InvM1, InvIL1);
+					}
+					if (Swing2Motion != EJointMotionType::Free)
+					{
+						FPBDJointUtilities::ApplyJointSwingProjection(Dt, Settings, JointSettings, AngularProjectionFactor, Index0, Index1, EJointAngularConstraintIndex::Swing2, EJointAngularAxisIndex::Swing2, P0, Q0, P1, Q1, InvM0, InvIL0, InvM1, InvIL1);
+					}
 				}
 			}
 		}
diff --git a/Engine/Source/Runtime/Experimental/Chaos/Private/Chaos/PBDRigidsEvolutionGBF.cpp b/Engine/Source/Runtime/Experimental/Chaos/Private/Chaos/PBDRigidsEvolutionGBF.cpp
--- a/Engine/Source/Runtime/Experimental/Chaos/Private/Chaos/PBDRigidsEvolutionGBF.cpp
+++ b/Engine/Source/Runtime/Experimental/Chaos/Private/Chaos/PBDRigidsEvolutionGBF.cpp
@@ -118,9 +118,6 @@ void TPBDRigidsEvolutionGBF<T, d>::AdvanceOneTimeStep(const T Dt, const T StepFr
 {
 	SCOPE_CYCLE_COUNTER(STAT_Evo_AdvanceOneTimeStep);
 
-	// @todo(ccaulfield): only need this when list has changed
-	ConstraintRules.StableSort();
-
 #if !UE_BUILD_SHIPPING
 	if (SerializeEvolution)
 	{
diff --git a/Engine/Source/Runtime/Experimental/Chaos/Public/Chaos/ChaosDebugDraw.h b/Engine/Source/Runtime/Experimental/Chaos/Public/Chaos/ChaosDebugDraw.h
--- a/Engine/Source/Runtime/Experimental/Chaos/Public/Chaos/ChaosDebugDraw.h
+++ b/Engine/Source/Runtime/Experimental/Chaos/Public/Chaos/ChaosDebugDraw.h
@@ -10,6 +10,18 @@ namespace Chaos
 	namespace DebugDraw
 	{
 #if CHAOS_DEBUG_DRAW
+
+		enum EDebugDrawJointFeature
+		{
+			None = 0,
+			Connector = 1 << 0,
+			Axes = 1 << 1,
+			Level = 1 << 2,
+
+			Default = Axes,
+			All = Connector | Axes | Level
+		};
+
 		CHAOS_API void DrawParticleShapes(const FRigidTransform3& SpaceTransform, const TParticleView<TGeometryParticles<float, 3>>& ParticlesView, float ColorScale, bool bDrawKinematic = true, bool bDrawDynamic = true);
 		CHAOS_API void DrawParticleShapes(const FRigidTransform3& SpaceTransform, const TArray<TGeometryParticleHandle<float, 3>*>& Particles, float ColorScale, bool bDrawKinematic = true, bool bDrawDynamic = true);
 		CHAOS_API void DrawParticleTransforms(const FRigidTransform3& SpaceTransform, const TParticleView<TGeometryParticles<float, 3>>& ParticlesView, float ColorScale, bool bDrawKinematic = true, bool bDrawDynamic = true);
@@ -17,8 +29,8 @@ namespace Chaos
 		CHAOS_API void DrawParticleCollisions(const FRigidTransform3& SpaceTransform, const TGeometryParticleHandle<float, 3>* Particle, const TPBDCollisionConstraint<float, 3>& Collisions);
 		CHAOS_API void DrawCollisions(const FRigidTransform3& SpaceTransform, const TPBDCollisionConstraint<float, 3>& Collisions, float ColorScale);
 		CHAOS_API void DrawCollisions(const FRigidTransform3& SpaceTransform, const TArray<TPBDCollisionConstraintHandle<float, 3>*>& ConstraintHandles, float ColorScale);
-		CHAOS_API void DrawJointConstraints(const FRigidTransform3& SpaceTransform, const TArray<FPBDJointConstraintHandle*>& ConstraintHandles, float ColorScale);
-		CHAOS_API void DrawJointConstraints(const FRigidTransform3& SpaceTransform, const FPBDJointConstraints& Constraints, float ColorScale);
+		CHAOS_API void DrawJointConstraints(const FRigidTransform3& SpaceTransform, const TArray<FPBDJointConstraintHandle*>& ConstraintHandles, float ColorScale, uint32 FeatureMask = EDebugDrawJointFeature::Default);
+		CHAOS_API void DrawJointConstraints(const FRigidTransform3& SpaceTransform, const FPBDJointConstraints& Constraints, float ColorScale, uint32 FeatureMask = EDebugDrawJointFeature::Default);
 		CHAOS_API void Draw6DofConstraints(const FRigidTransform3& SpaceTransform, const TArray<FPBD6DJointConstraintHandle*>& ConstraintHandles, float ColorScale);
 		CHAOS_API void Draw6DofConstraints(const FRigidTransform3& SpaceTransform, const FPBD6DJointConstraints& Constraints, float ColorScale);
 
diff --git a/Engine/Source/Runtime/Experimental/Chaos/Public/Chaos/DenseMatrix.h b/Engine/Source/Runtime/Experimental/Chaos/Public/Chaos/DenseMatrix.h
new file mode 100644
--- /dev/null
+++ b/Engine/Source/Runtime/Experimental/Chaos/Public/Chaos/DenseMatrix.h
@@ -0,0 +1,517 @@
+// Copyright 1998-2019 Epic Games, Inc. All Rights Reserved.
+#pragma once
+
+#include "CoreMinimal.h"
+#include "Chaos/Core.h"
+#include "Chaos/Vector.h"
+#include "Chaos/Matrix.h"
+
+namespace Chaos
+{
+	/**
+	 * A matrix with run-time variable dimensions, up to an element limit defined at compile-time.
+	 *
+	 * Elements are stored in row-major order (i.e., elements in a row are adjacent in memory). Note
+	 * that FMatrix stores elements in column-major order so that we can access the columns quickly
+	 * which is handy when you have rotation matrices and want the spatial axes. We don't care about
+	 * that so we use the more conventional row-major indexing and matching storage.
+	 */
+	template<int32 T_MAXELEMENTS>
+	class TDenseMatrix
+	{
+	public:
+		static const int32 MaxElements = T_MAXELEMENTS;
+
+		TDenseMatrix()
+			: NRows(0)
+			, NCols(0)
+		{
+		}
+
+		TDenseMatrix(const int32 InNRows, const int32 InNCols)
+			: NRows(InNRows)
+			, NCols(InNCols)
+		{
+		}
+
+		TDenseMatrix(const TDenseMatrix<MaxElements>& A)
+		{
+			*this = A;
+		}
+
+		TDenseMatrix<MaxElements>& operator=(const TDenseMatrix<MaxElements>& A)
+		{
+			SetDimensions(A.NumRows(), A.NumColumns());
+			for (int32 Index = 0; Index < NumElements(); ++Index)
+			{
+				M[Index] = A.M[Index];
+			}
+			return *this;
+		}
+
+		int32 NumRows() const
+		{
+			return NRows;
+		}
+
+		int32 NumColumns() const
+		{
+			return NCols;
+		}
+
+		int32 NumElements() const
+		{
+			return NRows * NCols;
+		}
+
+		/**
+		 * Set the dimensions of the matrix. This does not rearrange or set any matrix elements so the matrix should not be used until its values have been set.
+		 */
+		void SetDimensions(const int32 InNumRows, const int32 InNumColumns)
+		{
+			check(InNumRows * InNumColumns <= MaxElements);
+			NRows = InNumRows;
+			NCols = InNumColumns;
+		}
+
+		FORCEINLINE int32 ElementIndex(const int32 RowIndex, const int32 ColumnIndex) const
+		{
+			checkSlow(RowIndex < NumRows());
+			checkSlow(ColumnIndex < NumColumns());
+			return RowIndex * NCols + ColumnIndex;
+		}
+
+		/**
+		 * Return a writable reference to the element at the specified row and column.
+		 */
+		FReal& At(const int32 RowIndex, const int32 ColumnIndex)
+		{
+			return M[ElementIndex(RowIndex, ColumnIndex)];
+		}
+
+		/**
+		 * Return a read-only reference to the element at the specified row and column.
+		 */
+		const FReal& At(const int32 RowIndex, const int32 ColumnIndex) const
+		{
+			return M[ElementIndex(RowIndex, ColumnIndex)];
+		}
+
+		/**
+		 * Set all elements to 'V'.
+		 */
+		void SetAll(FReal V)
+		{
+			for (int32 II = 0; II < NumElements(); ++II)
+			{
+				M[II] = V;
+			}
+		}
+
+		/**
+		 * Set the diagonal elements to 'V'. Does not set off-diagonal elements.
+		 * /see MakeDiagonal
+		 */
+		void SetDiagonal(FReal V)
+		{
+			int32 Num = FMath::Min(NRows, NCols);
+			for (int32 II = 0; II < Num; ++II)
+			{
+				M[ElementIndex(II, II)] = V;
+			}
+		}
+
+
+		//
+		// Factory methods
+		//
+
+		static TDenseMatrix<MaxElements> Make(const int32 InNumRows, const int32 InNumCols)
+		{
+			return TDenseMatrix<MaxElements>(InNumRows, InNumCols);
+		}
+
+
+		static TDenseMatrix<MaxElements> Make(const int32 InNumRows, const int32 InNumCols, const FReal* V, const int32 VLen)
+		{
+			return TDenseMatrix<MaxElements>(InNumRows, InNumCols, V, VLen);
+		}
+
+		static TDenseMatrix<MaxElements> Make(const int32 InNumRows, const int32 InNumCols, std::initializer_list<FReal> InitList)
+		{
+			return TDenseMatrix<MaxElements>(InNumRows, InNumCols, InitList);
+		}
+
+		static TDenseMatrix<MaxElements> Make(const FMatrix33& InM)
+		{
+			// NOTE: UE matrices are Column-major (columns are sequential in memory), but DenseMatrix is Row-major (rows are sequential in memory)
+			TDenseMatrix<MaxElements> M(3, 3);
+			M.At(0, 0) = InM.M[0][0];
+			M.At(0, 1) = InM.M[1][0];
+			M.At(0, 2) = InM.M[2][0];
+			M.At(1, 0) = InM.M[0][1];
+			M.At(1, 1) = InM.M[1][1];
+			M.At(1, 2) = InM.M[2][1];
+			M.At(2, 0) = InM.M[0][2];
+			M.At(2, 1) = InM.M[1][2];
+			M.At(2, 2) = InM.M[2][2];
+			return M;
+		}
+
+		static TDenseMatrix<MaxElements> MakeDiagonal(const int32 InNumRows, const int32 InNumCols, const FReal D)
+		{
+			TDenseMatrix<MaxElements> M(InNumRows, InNumCols);
+			for (int32 I = 0; I < InNumRows; ++I)
+			{
+				for (int32 J = 0; J < InNumCols; ++J)
+				{
+					M.At(I, J) = (I == J) ? D : 0;
+				}
+			}
+			return M;
+		}
+
+		static TDenseMatrix<MaxElements> MakeIdentity(const int32 InDim)
+		{
+			return MakeDiagonal(InDim, InDim, (FReal)1);
+		}
+
+		//
+		// Math operations
+		//
+
+		template<int32 T_EA, int32 T_EB>
+		static TDenseMatrix<MaxElements> Negative(const TDenseMatrix<T_EA>& A)
+		{
+			// @todo(ccaulfield): optimize
+			check(A.NumColumns() == B.NumColumns());
+			check(A.NumRows() == B.NumRows());
+			TDenseMatrix<T_MAXELEMENTS> Result(A.NumRows(), A.NumColumns());
+			for (int32 IRow = 0; IRow < Result.NumRows(); ++IRow)
+			{
+				for (int32 ICol = 0; ICol < Result.NumColumns(); ++ICol)
+				{
+					Result.At(IRow, ICol) = -A.At(IRow, ICol);
+				}
+			}
+			return Result;
+		}
+
+		template<int32 T_EA, int32 T_EB>
+		static TDenseMatrix<MaxElements> Add(const TDenseMatrix<T_EA>& A, const TDenseMatrix<T_EB>& B)
+		{
+			// @todo(ccaulfield): optimize
+			check(A.NumColumns() == B.NumColumns());
+			check(A.NumRows() == B.NumRows());
+			TDenseMatrix<T_MAXELEMENTS> Result(A.NumRows(), A.NumColumns());
+			for (int32 IRow = 0; IRow < Result.NumRows(); ++IRow)
+			{
+				for (int32 ICol = 0; ICol < Result.NumColumns(); ++ICol)
+				{
+					Result.At(IRow, ICol) = A.At(IRow, ICol) + B.At(IRow, ICol);
+				}
+			}
+			return Result;
+		}
+
+		template<int32 T_EA, int32 T_EB>
+		static TDenseMatrix<MaxElements> Subtract(const TDenseMatrix<T_EA>& A, const TDenseMatrix<T_EB>& B)
+		{
+			// @todo(ccaulfield): optimize
+			check(A.NumColumns() == B.NumColumns());
+			check(A.NumRows() == B.NumRows());
+			TDenseMatrix<T_MAXELEMENTS> Result(A.NumRows(), A.NumColumns());
+			for (int32 IRow = 0; IRow < Result.NumRows(); ++IRow)
+			{
+				for (int32 ICol = 0; ICol < Result.NumColumns(); ++ICol)
+				{
+					Result.At(IRow, ICol) = A.At(IRow, ICol) - B.At(IRow, ICol);
+				}
+			}
+			return Result;
+		}
+
+		template<int32 T_EA, int32 T_EB>
+		static TDenseMatrix<MaxElements> MultiplyAB(const TDenseMatrix<T_EA>& A, const TDenseMatrix<T_EB>& B)
+		{
+			// @todo(ccaulfield): optimize
+			check(A.NumColumns() == B.NumRows());
+			TDenseMatrix<T_MAXELEMENTS> Result(A.NumRows(), B.NumColumns());
+			for (int32 IRow = 0; IRow < Result.NumRows(); ++IRow)
+			{
+				for (int32 ICol = 0; ICol < Result.NumColumns(); ++ICol)
+				{
+					FReal V = 0;
+					for (int32 II = 0; II < A.NumColumns(); ++II)
+					{
+						V += A.At(IRow, II) * B.At(II, ICol);
+					}
+					Result.At(IRow, ICol) = V;
+				}
+			}
+			return Result;
+		}
+
+		template<int32 T_EA, int32 T_EB>
+		static TDenseMatrix<MaxElements> MultiplyAtB(const TDenseMatrix<T_EA>& A, const TDenseMatrix<T_EB>& B)
+		{
+			// @todo(ccaulfield): optimize
+			check(A.NumRows() == B.NumRows());
+			TDenseMatrix<T_MAXELEMENTS> Result(A.NumColumns(), B.NumColumns());
+			for (int32 IRow = 0; IRow < Result.NumRows(); ++IRow)
+			{
+				for (int32 ICol = 0; ICol < Result.NumColumns(); ++ICol)
+				{
+					FReal V = 0;
+					for (int32 II = 0; II < A.NumRows(); ++II)
+					{
+						V += A.At(II, IRow) * B.At(II, ICol);
+					}
+					Result.At(IRow, ICol) = V;
+				}
+			}
+			return Result;
+		}
+
+		template<int32 T_EA, int32 T_EB>
+		static TDenseMatrix<MaxElements> MultiplyABt(const TDenseMatrix<T_EA>& A, const TDenseMatrix<T_EB>& B)
+		{
+			// @todo(ccaulfield): optimize
+			check(A.NumColumns() == B.NumColumns());
+			TDenseMatrix<T_MAXELEMENTS> Result(A.NumRows(), B.NumRows());
+			for (int32 IRow = 0; IRow < Result.NumRows(); ++IRow)
+			{
+				for (int32 ICol = 0; ICol < Result.NumColumns(); ++ICol)
+				{
+					FReal V = 0;
+					for (int32 II = 0; II < A.NumColumns(); ++II)
+					{
+						V += A.At(IRow, II) * B.At(ICol, II);
+					}
+					Result.At(IRow, ICol) = V;
+				}
+			}
+			return Result;
+		}
+
+		template<int32 T_EA, int32 T_EB>
+		static TDenseMatrix<MaxElements> MultiplyAtBt(const TDenseMatrix<T_EA>& A, const TDenseMatrix<T_EB>& B)
+		{
+			// @todo(ccaulfield): optimize
+			check(A.NumRows() == B.NumColumns());
+			TDenseMatrix<T_MAXELEMENTS> Result(A.NumColumns(), B.NumRows());
+			for (int32 IRow = 0; IRow < Result.NumRows(); ++IRow)
+			{
+				for (int32 ICol = 0; ICol < Result.NumColumns(); ++ICol)
+				{
+					FReal V = 0;
+					for (int32 II = 0; II < A.NumRows(); ++II)
+					{
+						V += A.At(II, IRow) * B.At(ICol, II);
+					}
+					Result.At(IRow, ICol) = V;
+				}
+			}
+			return Result;
+		}
+
+		template<int32 T_EA, int32 T_EB>
+		static TDenseMatrix<MaxElements> Multiply(const TDenseMatrix<T_EA>& A, const FReal V)
+		{
+			// @todo(ccaulfield): optimize
+			check(A.NumColumns() == B.NumRows());
+			TDenseMatrix<T_MAXELEMENTS> Result(A.NumRows(), B.NumColumns());
+			for (int32 IRow = 0; IRow < Result.NumRows(); ++IRow)
+			{
+				for (int32 ICol = 0; ICol < Result.NumColumns(); ++ICol)
+				{
+					Result.At(IRow, ICol) = Result.At(IRow, ICol) / V;
+				}
+			}
+			return Result;
+		}
+
+		template<int32 T_EA, int32 T_EB>
+		static TDenseMatrix<MaxElements> Multiply(const FReal V, const TDenseMatrix<T_EA>& A)
+		{
+			return Multiply(A, V);
+		}
+
+		template<int32 T_EA, int32 T_EB>
+		static TDenseMatrix<MaxElements> Divide(const TDenseMatrix<T_EA>& A, const FReal V)
+		{
+			// @todo(ccaulfield): optimize
+			check(A.NumColumns() == B.NumRows());
+			TDenseMatrix<T_MAXELEMENTS> Result(A.NumRows(), B.NumColumns());
+			for (int32 IRow = 0; IRow < Result.NumRows(); ++IRow)
+			{
+				for (int32 ICol = 0; ICol < Result.NumColumns(); ++ICol)
+				{
+					Result.At(IRow, ICol) = Result.At(IRow, ICol) / V;
+				}
+			}
+			return Result;
+		}
+
+		template<int32 T_EA, int32 T_EB>
+		static TDenseMatrix<MaxElements> DotProduct(const TDenseMatrix<T_EA>& A, const TDenseMatrix<T_EA>& B)
+		{
+			return MultiplyAtB(A, B)
+		}
+
+	private:
+		TDenseMatrix(const int32 InNRows, const int32 InNCols, const FReal* V, const int32 N)
+			: NRows(InNRows)
+			, NCols(InNCols)
+		{
+			int32 NLimited = FMath::Min<int32>(NumElements(), N);
+			for (int32 I = 0; I < NLimited; ++I)
+			{
+				M[I] = V[I];
+			}
+		}
+
+		TDenseMatrix(const int32 InNRows, const int32 InNCols, std::initializer_list<FReal> InitList)
+			: NRows(InNRows)
+			, NCols(InNCols)
+		{
+			int32 N = FMath::Min<int32>(NumElements(), InitList.size());
+			for (int32 I = 0; I < N; ++I)
+			{
+				M[I] = InitList[I];
+			}
+		}
+
+		FReal M[MaxElements];
+		int32 NRows;
+		int32 NCols;
+	};
+
+	/**
+	 * Methods to solves sets of Linear equations stored as
+	 *	AX = B
+	 * where A is an NxN matrix, and X.B are Nx1 column vectors.
+	 */
+	class FDenseMatrixSolver
+	{
+	public:
+
+		/**
+		 * Overwrite A with its Cholesky Factor (A must be Positive Definite).
+		 * See "Matrix Computations, 4th Edition" Section 4.2, Golub & Van Loan.
+		 *
+		 * The Cholesky Factor of A is G (Gt its transpose), where A = GGt. G is lower triangular.
+		 */
+		template<int32 T_E>
+		static bool CholeskyFactorize(TDenseMatrix<T_E>& A)
+		{
+			check(A.NumRows() == A.NumColumns());
+			const int32 N = A.NumRows();
+			for (int32 I = 0; I < N; ++I)
+			{
+				for (int32 J = I; J < N; ++J)
+				{
+					FReal Sum = A.At(I, J);
+					for (int32 K = I - 1; K >= 0; --K)
+					{
+						Sum -= A.At(I, K) * A.At(J, K);
+					}
+					if (I == J)
+					{
+						if (Sum <= 0)
+						{
+							// Not positive definite (rounding?)
+							return false;
+						}
+						A.At(I, J) = FMath::Sqrt(Sum);
+					}
+					else
+					{
+						A.At(J, I) = Sum / A.At(I, I);
+					}
+				}
+			}
+
+			for (int32 I = 0; I < N; ++I)
+			{
+				for (int32 J = 0; J < I; ++J)
+				{
+					A.At(J, I) = 0;
+				}
+			}
+
+			return true;
+		}
+
+		/**
+		 * This solves AX = B, where A is positive definite and has been Cholesky Factorized to produce G, 
+		 * where A = GGt, G is lower triangular.
+		 *
+		 * This is a helper method for SolvePositiveDefinite, or useful if you need to reuse the 
+		 * Cholesky Factor and therefore calculated it yourself.
+		 *
+		 * \see SolvePositiveDefinite
+		 */
+		template<int32 T_EA, int32 T_EB, int32 T_EX>
+		static void SolveCholeskyFactorized(const TDenseMatrix<T_EA>& G, const TDenseMatrix<T_EB>& B, TDenseMatrix<T_EX>& X)
+		{
+			check(B.NumColumns() == 1);
+			check(G.NumRows() == B.NumRows());
+
+			const int32 N = G.NumRows();
+			X.SetDimensions(N, 1);
+
+			// Solve LY = B (G is lower-triangular)
+			for (int32 I = 0; I < N; ++I)
+			{
+				FReal Sum = B.At(I, 0);
+				for (int32 K = I - 1; K >= 0; --K)
+				{
+					Sum -= G.At(I, K) * X.At(K, 0);
+				}
+				X.At(I, 0) = Sum / G.At(I, I);
+			}
+
+			// Solve LtX = Y (Lt is upper-triangular)
+			for (int32 I = N - 1; I >= 0; --I)
+			{
+				FReal Sum = X.At(I, 0);
+				for (int32 K = I + 1; K < N; ++K)
+				{
+					Sum -= G.At(K, I) * X.At(K, 0);
+				}
+				X.At(I, 0) = Sum / G.At(I, I);
+			}
+		}
+
+		/**
+		 * Solve AX = B, for positive-definite NxN matrix A, and Nx1 column vectors B and X.
+		 *
+		 * For positive definite A, A = GGt, where G is the Cholesky factor and lower trangular.
+		 * We can solve GGtX = B by first solving GY = B, and then GtX = Y.
+		 *
+		 * E.g., this can be used to solve constraint equations of the form
+		 *		J.I.Jt.X = B
+		 * where J is a Jacobian (Jt its transpose), I is an Inverse mas matrix, and B the residual.
+		 * In this case, I is positive definite, and therefore so is JIJt.
+		 *
+		 */
+		template<int32 T_EA, int32 T_EB, int32 T_EX>
+		static bool SolvePositiveDefinite(const TDenseMatrix<T_EA>& A, const TDenseMatrix<T_EB>& B, TDenseMatrix<T_EX>& X)
+		{
+			check(B.NumColumns() == 1);
+			check(A.NumRows() == B.NumRows());
+
+			TDenseMatrix<T_EA> G = A;
+			if (!CholeskyFactorize(G))
+			{
+				// Not positive definite
+				return false;
+			}
+
+			SolveCholeskyFactorized(G, B, X);
+			return true;
+		}
+	};
+
+}
diff --git a/Engine/Source/Runtime/Experimental/Chaos/Public/Chaos/PBDCollisionConstraint.h b/Engine/Source/Runtime/Experimental/Chaos/Public/Chaos/PBDCollisionConstraint.h
--- a/Engine/Source/Runtime/Experimental/Chaos/Public/Chaos/PBDCollisionConstraint.h
+++ b/Engine/Source/Runtime/Experimental/Chaos/Public/Chaos/PBDCollisionConstraint.h
@@ -68,7 +68,7 @@ public:
 	using FRigidBodyContactConstraint = TRigidBodyContactConstraint<T, d>;
 	using FHandles = TArray<FConstraintContainerHandle*>;
 
-	TPBDCollisionConstraint(const TPBDRigidsSOAs<T,d>& InParticles, TArrayCollectionArray<bool>& Collided, const TArrayCollectionArray<TSerializablePtr<FChaosPhysicsMaterial>>& PerParticleMaterials, const int32 PairIterations = 1, const T Thickness = (T)0);
+	TPBDCollisionConstraint(const TPBDRigidsSOAs<T,d>& InParticles, TArrayCollectionArray<bool>& Collided, const TArrayCollectionArray<TSerializablePtr<FChaosPhysicsMaterial>>& PerParticleMaterials, const int32 ApplyPairIterations = 1, const int32 ApplyPushOutPairIterations = 1, const T Thickness = (T)0);
 	virtual ~TPBDCollisionConstraint() {}
 
 	//
@@ -85,14 +85,14 @@ public:
 		CollisionVelocityInflation = ScaleFactor;
 	}
 
-	void SetVelocitySolveEnabled(bool bInEnableVelocitySolve)
+	void SetPairIterations(int32 InPairIterations)
 	{
-		bEnableVelocitySolve = bInEnableVelocitySolve;
+		MApplyPairIterations = InPairIterations;
 	}
 
 	void SetPushOutPairIterations(int32 InPairIterations)
-	{ 
-		MPairIterations = InPairIterations;
+	{
+		MApplyPushOutPairIterations = InPairIterations;
 	}
 
 	void SetCollisionsEnabled(bool bInEnableCollisions)
@@ -201,7 +201,8 @@ private:
 	*/
 	void Reset( );
 
-	void Apply(const T Dt, FRigidBodyContactConstraint& Constraint);
+	void Apply(const T Dt, FRigidBodyContactConstraint& Constraint, const int32 It, const int32 NumIts);
+
 
 	void ApplyPushOut(const T Dt, FRigidBodyContactConstraint& Constraint, const TSet<const TGeometryParticleHandle<T, d>*>& IsTemporarilyStatic, int32 Iteration, int32 NumIterations, bool &NeedsAnotherIteration);
 
@@ -224,8 +225,8 @@ private:
 
 	TArrayCollectionArray<bool>& MCollided;
 	const TArrayCollectionArray<TSerializablePtr<FChaosPhysicsMaterial>>& MPhysicsMaterials;
-	bool bEnableVelocitySolve;
-	int32 MPairIterations;
+	int32 MApplyPairIterations;
+	int32 MApplyPushOutPairIterations;
 	T MThickness;
 	T MAngularFriction;
 	bool bUseCCD;
diff --git a/Engine/Source/Runtime/Experimental/Chaos/Public/Chaos/PBDJointConstraintTypes.h b/Engine/Source/Runtime/Experimental/Chaos/Public/Chaos/PBDJointConstraintTypes.h
--- a/Engine/Source/Runtime/Experimental/Chaos/Public/Chaos/PBDJointConstraintTypes.h
+++ b/Engine/Source/Runtime/Experimental/Chaos/Public/Chaos/PBDJointConstraintTypes.h
@@ -44,6 +44,13 @@ namespace Chaos
 		Swing1 = 2,		// Swing1 Axis = Z
 	};
 
+	enum class EJointProjectionPhase
+	{
+		None,
+		Apply,
+		ApplyPushOut,
+	};
+
 	struct FJointConstants
 	{
 		/** The constraint-space twist axis (X Axis) */
@@ -109,6 +116,10 @@ namespace Chaos
 	public:
 		FPBDJointSolverSettings();
 
+		// Iterations
+		int32 ApplyPairIterations;
+		int32 ApplyPushOutPairIterations;
+
 		// Tolerances
 		FReal SwingTwistAngleTolerance;
 
@@ -118,10 +129,10 @@ namespace Chaos
 
 		// @todo(ccaulfield): remove these TEMP overrides for testing
 		bool bEnableVelocitySolve;
-		bool bEnableLinearLimits;
 		bool bEnableTwistLimits;
 		bool bEnableSwingLimits;
 		bool bEnableDrives;
+		EJointProjectionPhase ProjectionPhase;
 		FReal LinearProjection;
 		FReal AngularProjection;
 		FReal Stiffness;
diff --git a/Engine/Source/Runtime/Experimental/Chaos/Public/Chaos/PBDJointConstraintUtilities.h b/Engine/Source/Runtime/Experimental/Chaos/Public/Chaos/PBDJointConstraintUtilities.h
--- a/Engine/Source/Runtime/Experimental/Chaos/Public/Chaos/PBDJointConstraintUtilities.h
+++ b/Engine/Source/Runtime/Experimental/Chaos/Public/Chaos/PBDJointConstraintUtilities.h
@@ -100,8 +100,12 @@ namespace Chaos
 			const int32 Index1,
 			FVec3& P0,
 			FRotation3& Q0,
+			FVec3& V0,
+			FVec3& W0,
 			FVec3& P1,
 			FRotation3& Q1,
+			FVec3& V1,
+			FVec3& W1,
 			float InvM0,
 			const FMatrix33& InvIL0,
 			float InvM1,
@@ -136,8 +140,12 @@ namespace Chaos
 			const int32 Index1,
 			FVec3& P0,
 			FRotation3& Q0,
+			FVec3& V0,
+			FVec3& W0,
 			FVec3& P1,
 			FRotation3& Q1,
+			FVec3& V1,
+			FVec3& W1,
 			float InvM0,
 			const FMatrix33& InvIL0,
 			float InvM1,
@@ -172,8 +180,12 @@ namespace Chaos
 			const int32 Index1,
 			FVec3& P0,
 			FRotation3& Q0,
+			FVec3& V0,
+			FVec3& W0,
 			FVec3& P1,
 			FRotation3& Q1,
+			FVec3& V1,
+			FVec3& W1,
 			float InvM0,
 			const FMatrix33& InvIL0,
 			float InvM1,
@@ -210,8 +222,12 @@ namespace Chaos
 			const EJointAngularAxisIndex SwingAxisIndex,
 			FVec3& P0,
 			FRotation3& Q0,
+			FVec3& V0,
+			FVec3& W0,
 			FVec3& P1,
 			FRotation3& Q1,
+			FVec3& V1,
+			FVec3& W1,
 			float InvM0,
 			const FMatrix33& InvIL0,
 			float InvM1,
@@ -247,8 +263,12 @@ namespace Chaos
 			const int32 Index1,
 			FVec3& P0,
 			FRotation3& Q0,
+			FVec3& V0,
+			FVec3& W0,
 			FVec3& P1,
 			FRotation3& Q1,
+			FVec3& V1,
+			FVec3& W1,
 			float InvM0,
 			const FMatrix33& InvIL0,
 			float InvM1,
@@ -262,8 +282,12 @@ namespace Chaos
 			const int32 Index1,
 			FVec3& P0,
 			FRotation3& Q0,
+			FVec3& V0,
+			FVec3& W0,
 			FVec3& P1,
 			FRotation3& Q1,
+			FVec3& V1,
+			FVec3& W1,
 			float InvM0,
 			const FMatrix33& InvIL0,
 			float InvM1,
@@ -277,8 +301,12 @@ namespace Chaos
 			const int32 Index1,
 			FVec3& P0,
 			FRotation3& Q0,
+			FVec3& V0,
+			FVec3& W0,
 			FVec3& P1,
 			FRotation3& Q1,
+			FVec3& V1,
+			FVec3& W1,
 			float InvM0,
 			const FMatrix33& InvIL0,
 			float InvM1,
diff --git a/Engine/Source/Runtime/Experimental/Chaos/Public/Chaos/PBDJointConstraints.h b/Engine/Source/Runtime/Experimental/Chaos/Public/Chaos/PBDJointConstraints.h
--- a/Engine/Source/Runtime/Experimental/Chaos/Public/Chaos/PBDJointConstraints.h
+++ b/Engine/Source/Runtime/Experimental/Chaos/Public/Chaos/PBDJointConstraints.h
@@ -144,8 +144,8 @@ namespace Chaos
 
 		void CalculateConstraintSpace(int32 ConstraintIndex, FVec3& OutX0, FMatrix33& OutR0, FVec3& OutX1, FMatrix33& OutR1, FVec3& OutAngles) const;
 
-		void SolveVelocity(const FReal Dt, const int32 ConstraintIndex, const int32 It, const int32 NumIts);
-		void SolvePosition(const FReal Dt, const int32 ConstraintIndex, const int32 It, const int32 NumIts);
+		void SolveVelocity(const FReal Dt, const int32 ConstraintIndex, const int32 NumPairIts, const int32 It, const int32 NumIts);
+		void SolvePosition(const FReal Dt, const int32 ConstraintIndex, const int32 NumPairIts, const int32 It, const int32 NumIts);
 		void ProjectPosition(const FReal Dt, const int32 ConstraintIndex, const int32 It, const int32 NumIts);
 
 		FPBDJointSolverSettings Settings;
diff --git a/Engine/Source/Runtime/Experimental/Chaos/Public/Chaos/PBDRigidsEvolution.h b/Engine/Source/Runtime/Experimental/Chaos/Public/Chaos/PBDRigidsEvolution.h
--- a/Engine/Source/Runtime/Experimental/Chaos/Public/Chaos/PBDRigidsEvolution.h
+++ b/Engine/Source/Runtime/Experimental/Chaos/Public/Chaos/PBDRigidsEvolution.h
@@ -423,15 +423,12 @@ class TPBDRigidsEvolutionBase
 
 	void ApplyConstraints(const T Dt, int32 Island)
 	{
-		for (FPBDConstraintGraphRule* ConstraintRule : ConstraintRules)
-		{
-			ConstraintRule->UpdateAccelerationStructures(Island);
-		}
+		UpdateAccelerationStructures(Island);
 
 		// @todo(ccaulfield): track whether we are sufficiently solved and can early-out
 		for (int i = 0; i < NumIterations; ++i)
 		{
-			for (FPBDConstraintGraphRule* ConstraintRule : ConstraintRules)
+			for (FPBDConstraintGraphRule* ConstraintRule : PrioritizedConstraintRules)
 			{
 				ConstraintRule->ApplyConstraints(Dt, Island, i, NumIterations);
 			}
@@ -579,6 +576,11 @@ protected:
 		{
 			ConstraintRule->AddToGraph();
 		}
+
+		// Apply rules in priority order
+		// @todo(ccaulfield): only really needed when list or priorities change
+		PrioritizedConstraintRules = ConstraintRules;
+		PrioritizedConstraintRules.StableSort();
 	}
 
 	void CreateIslands()
@@ -602,7 +604,7 @@ protected:
 		for (int32 It = 0; bNeedsAnotherIteration && (It < NumPushOutIterations); ++It)
 		{
 			bNeedsAnotherIteration = false;
-			for (FPBDConstraintGraphRule* ConstraintRule : ConstraintRules)
+			for (FPBDConstraintGraphRule* ConstraintRule : PrioritizedConstraintRules)
 			{
 				if (ConstraintRule->ApplyPushOut(Dt, Island, It, NumPushOutIterations))
 				{
@@ -625,6 +627,7 @@ protected:
 	FUpdatePositionRule ParticleUpdatePosition;
 	FKinematicUpdateRule KinematicUpdate;
 	TArray<FPBDConstraintGraphRule*> ConstraintRules;
+	TArray<FPBDConstraintGraphRule*> PrioritizedConstraintRules;
 	FPBDConstraintGraph ConstraintGraph;
 	TArrayCollectionArray<TSerializablePtr<FChaosPhysicsMaterial>> PhysicsMaterials;
 	TArrayCollectionArray<TUniquePtr<FChaosPhysicsMaterial>> PerParticlePhysicsMaterials;
diff --git a/Engine/Source/Runtime/Experimental/Chaos/Public/Chaos/Utilities.h b/Engine/Source/Runtime/Experimental/Chaos/Public/Chaos/Utilities.h
--- a/Engine/Source/Runtime/Experimental/Chaos/Public/Chaos/Utilities.h
+++ b/Engine/Source/Runtime/Experimental/Chaos/Public/Chaos/Utilities.h
@@ -102,6 +102,21 @@ namespace Chaos
 				).GetTransposed();
 		}
 
+		inline PMatrix<float, 3, 3> MultiplyAB(const PMatrix<float, 3, 3>& LIn, const PMatrix<float, 3, 3>& RIn)
+		{
+			return Multiply(LIn, RIn);
+		}
+
+		inline PMatrix<float, 3, 3> MultiplyABt(const PMatrix<float, 3, 3>& LIn, const PMatrix<float, 3, 3>& RIn)
+		{
+			return Multiply(LIn, RIn.GetTransposed());
+		}
+
+		inline PMatrix<float, 3, 3> MultiplyAtB(const PMatrix<float, 3, 3>& LIn, const PMatrix<float, 3, 3>& RIn)
+		{
+			return Multiply(LIn.GetTransposed(), RIn);
+		}
+
 		/**
 		 * Multiple a vector by a matrix: C = L.R
 		 * If L is a rotation matrix, then this will return R rotated by that rotation.
