Chaos - Fix CIS (DenseMatrix)

diff --git a/Engine/Source/Runtime/Experimental/Chaos/Public/Chaos/DenseMatrix.h b/Engine/Source/Runtime/Experimental/Chaos/Public/Chaos/DenseMatrix.h
--- a/Engine/Source/Runtime/Experimental/Chaos/Public/Chaos/DenseMatrix.h
+++ b/Engine/Source/Runtime/Experimental/Chaos/Public/Chaos/DenseMatrix.h
@@ -184,8 +184,6 @@ namespace Chaos
 		static TDenseMatrix<MaxElements> Negative(const TDenseMatrix<T_EA>& A)
 		{
 			// @todo(ccaulfield): optimize
-			check(A.NumColumns() == B.NumColumns());
-			check(A.NumRows() == B.NumRows());
 			TDenseMatrix<T_MAXELEMENTS> Result(A.NumRows(), A.NumColumns());
 			for (int32 IRow = 0; IRow < Result.NumRows(); ++IRow)
 			{
@@ -319,8 +317,7 @@ namespace Chaos
 		static TDenseMatrix<MaxElements> Multiply(const TDenseMatrix<T_EA>& A, const FReal V)
 		{
 			// @todo(ccaulfield): optimize
-			check(A.NumColumns() == B.NumRows());
-			TDenseMatrix<T_MAXELEMENTS> Result(A.NumRows(), B.NumColumns());
+			TDenseMatrix<T_MAXELEMENTS> Result(A.NumRows(), A.NumColumns());
 			for (int32 IRow = 0; IRow < Result.NumRows(); ++IRow)
 			{
 				for (int32 ICol = 0; ICol < Result.NumColumns(); ++ICol)
@@ -341,8 +338,7 @@ namespace Chaos
 		static TDenseMatrix<MaxElements> Divide(const TDenseMatrix<T_EA>& A, const FReal V)
 		{
 			// @todo(ccaulfield): optimize
-			check(A.NumColumns() == B.NumRows());
-			TDenseMatrix<T_MAXELEMENTS> Result(A.NumRows(), B.NumColumns());
+			TDenseMatrix<T_MAXELEMENTS> Result(A.NumRows(), A.NumColumns());
 			for (int32 IRow = 0; IRow < Result.NumRows(); ++IRow)
 			{
 				for (int32 ICol = 0; ICol < Result.NumColumns(); ++ICol)
@@ -356,7 +352,7 @@ namespace Chaos
 		template<int32 T_EA, int32 T_EB>
 		static TDenseMatrix<MaxElements> DotProduct(const TDenseMatrix<T_EA>& A, const TDenseMatrix<T_EA>& B)
 		{
-			return MultiplyAtB(A, B)
+			return MultiplyAtB(A, B);
 		}
 
 	private:
