Chaos - fiox CIS (ChaodDebugDraw)

diff --git a/Engine/Source/Runtime/Engine/Private/PhysicsEngine/ImmediatePhysics/ImmediatePhysicsChaos/ImmediatePhysicsSimulation_Chaos.cpp b/Engine/Source/Runtime/Engine/Private/PhysicsEngine/ImmediatePhysics/ImmediatePhysicsChaos/ImmediatePhysicsSimulation_Chaos.cpp
--- a/Engine/Source/Runtime/Engine/Private/PhysicsEngine/ImmediatePhysics/ImmediatePhysicsChaos/ImmediatePhysicsSimulation_Chaos.cpp
+++ b/Engine/Source/Runtime/Engine/Private/PhysicsEngine/ImmediatePhysics/ImmediatePhysicsChaos/ImmediatePhysicsSimulation_Chaos.cpp
@@ -89,7 +89,7 @@ int32 ChaosImmediate_DebugDrawCollisions = 0;
 int32 ChaosImmediate_DebugDrawJoints = 0;
 #endif
 
-int32 ChaosImmediate_DebugDrawJointFeatures = Chaos::DebugDraw::EDebugDrawJointFeature::Default;
+int32 ChaosImmediate_DebugDrawJointFeatures = (int32)Chaos::DebugDraw::EDebugDrawJointFeature::Default;
 
 FAutoConsoleVariableRef CVarChaosImmPhysDebugDrawParticles(TEXT("p.Chaos.ImmPhys.DebugDrawParticles"), ChaosImmediate_DebugDrawParticles, TEXT("Draw Particle Transforms (0 = never; 1 = end of frame; 2 = begin and end of frame; 3 = post-integate, post-apply and post-applypushout;)."));
 FAutoConsoleVariableRef CVarChaosImmPhysDebugDrawShapes(TEXT("p.Chaos.ImmPhys.DebugDrawShapes"), ChaosImmediate_DebugDrawShapes, TEXT("Draw Shapes (0 = never; 1 = end of frame; 2 = begin and end of frame; 3 = post-integate, post-apply and post-applypushout;"));
diff --git a/Engine/Source/Runtime/Experimental/Chaos/Private/Chaos/ChaosDebugDraw.cpp b/Engine/Source/Runtime/Experimental/Chaos/Private/Chaos/ChaosDebugDraw.cpp
--- a/Engine/Source/Runtime/Experimental/Chaos/Private/Chaos/ChaosDebugDraw.cpp
+++ b/Engine/Source/Runtime/Experimental/Chaos/Private/Chaos/ChaosDebugDraw.cpp
@@ -182,14 +182,14 @@ namespace Chaos
 			FVec3 Xa = SpaceTransform.TransformPosition(InXa);
 			FVec3 Xb = SpaceTransform.TransformPosition(InXb);
 
-			if (FeatureMask & EDebugDrawJointFeature::Connector)
+			if (FeatureMask & (uint32)EDebugDrawJointFeature::Connector)
 			{
 				FDebugDrawQueue::GetInstance().DrawDebugLine(Pa, Xa, R, false, KINDA_SMALL_NUMBER, DrawPriority, LineThickness);
 				FDebugDrawQueue::GetInstance().DrawDebugLine(Pb, Xb, C, false, KINDA_SMALL_NUMBER, DrawPriority, LineThickness);
 				FDebugDrawQueue::GetInstance().DrawDebugPoint(Xa, R, false, KINDA_SMALL_NUMBER, DrawPriority, DrawScale * PointSize);
 				FDebugDrawQueue::GetInstance().DrawDebugPoint(Xb, C, false, KINDA_SMALL_NUMBER, DrawPriority, DrawScale * PointSize);
 			}
-			if (FeatureMask & EDebugDrawJointFeature::Axes)
+			if (FeatureMask & (uint32)EDebugDrawJointFeature::Axes)
 			{
 				FDebugDrawQueue::GetInstance().DrawDebugDirectionalArrow(Xa, Xa + DrawScale * ConstraintAxisLen * SpaceTransform.TransformVector(Ra.GetAxis(0)), DrawScale * ArrowSize, R, false, KINDA_SMALL_NUMBER, DrawPriority, LineThickness);
 				FDebugDrawQueue::GetInstance().DrawDebugDirectionalArrow(Xa, Xa + DrawScale * ConstraintAxisLen * SpaceTransform.TransformVector(Ra.GetAxis(1)), DrawScale * ArrowSize, G, false, KINDA_SMALL_NUMBER, DrawPriority, LineThickness);
@@ -198,7 +198,7 @@ namespace Chaos
 				FDebugDrawQueue::GetInstance().DrawDebugDirectionalArrow(Xb, Xb + DrawScale * ConstraintAxisLen * SpaceTransform.TransformVector(Rb.GetAxis(1)), DrawScale * ArrowSize, M, false, KINDA_SMALL_NUMBER, DrawPriority, LineThickness);
 				FDebugDrawQueue::GetInstance().DrawDebugDirectionalArrow(Xb, Xb + DrawScale * ConstraintAxisLen * SpaceTransform.TransformVector(Rb.GetAxis(2)), DrawScale * ArrowSize, Y, false, KINDA_SMALL_NUMBER, DrawPriority, LineThickness);
 			}
-			if ((FeatureMask & EDebugDrawJointFeature::Level) && (Level >= 0))
+			if ((FeatureMask & (uint32)EDebugDrawJointFeature::Level) && (Level >= 0))
 			{
 				FDebugDrawQueue::GetInstance().DrawDebugString(Xb + FontHeight * FVec3(0, 0, 1), FString::Format(TEXT("{0}"), { Level }), nullptr, FColor::Red, KINDA_SMALL_NUMBER, false, FontScale);
 			}
@@ -228,7 +228,7 @@ namespace Chaos
 			FVec3 Xa, Xb, CR;
 			FMatrix33 Ra, Rb;
 			ConstraintHandle->CalculateConstraintSpace(Xa, Ra, Xb, Rb, CR);
-			DrawJointConstraintImpl(SpaceTransform, Pa, Xa, Ra, Pb, Xb, Rb, CR, INDEX_NONE, ColorScale, EDebugDrawJointFeature::Default);
+			DrawJointConstraintImpl(SpaceTransform, Pa, Xa, Ra, Pb, Xb, Rb, CR, INDEX_NONE, ColorScale, (uint32)EDebugDrawJointFeature::Default);
 		}
 
 #endif
diff --git a/Engine/Source/Runtime/Experimental/Chaos/Public/Chaos/ChaosDebugDraw.h b/Engine/Source/Runtime/Experimental/Chaos/Public/Chaos/ChaosDebugDraw.h
--- a/Engine/Source/Runtime/Experimental/Chaos/Public/Chaos/ChaosDebugDraw.h
+++ b/Engine/Source/Runtime/Experimental/Chaos/Public/Chaos/ChaosDebugDraw.h
@@ -11,7 +11,7 @@ namespace Chaos
 	{
 #if CHAOS_DEBUG_DRAW
 
-		enum EDebugDrawJointFeature
+		enum class EDebugDrawJointFeature
 		{
 			None = 0,
 			Connector = 1 << 0,
@@ -29,8 +29,8 @@ namespace Chaos
 		CHAOS_API void DrawParticleCollisions(const FRigidTransform3& SpaceTransform, const TGeometryParticleHandle<float, 3>* Particle, const TPBDCollisionConstraint<float, 3>& Collisions);
 		CHAOS_API void DrawCollisions(const FRigidTransform3& SpaceTransform, const TPBDCollisionConstraint<float, 3>& Collisions, float ColorScale);
 		CHAOS_API void DrawCollisions(const FRigidTransform3& SpaceTransform, const TArray<TPBDCollisionConstraintHandle<float, 3>*>& ConstraintHandles, float ColorScale);
-		CHAOS_API void DrawJointConstraints(const FRigidTransform3& SpaceTransform, const TArray<FPBDJointConstraintHandle*>& ConstraintHandles, float ColorScale, uint32 FeatureMask = EDebugDrawJointFeature::Default);
-		CHAOS_API void DrawJointConstraints(const FRigidTransform3& SpaceTransform, const FPBDJointConstraints& Constraints, float ColorScale, uint32 FeatureMask = EDebugDrawJointFeature::Default);
+		CHAOS_API void DrawJointConstraints(const FRigidTransform3& SpaceTransform, const TArray<FPBDJointConstraintHandle*>& ConstraintHandles, float ColorScale, uint32 FeatureMask = (uint32)EDebugDrawJointFeature::Default);
+		CHAOS_API void DrawJointConstraints(const FRigidTransform3& SpaceTransform, const FPBDJointConstraints& Constraints, float ColorScale, uint32 FeatureMask = (uint32)EDebugDrawJointFeature::Default);
 		CHAOS_API void Draw6DofConstraints(const FRigidTransform3& SpaceTransform, const TArray<FPBD6DJointConstraintHandle*>& ConstraintHandles, float ColorScale);
 		CHAOS_API void Draw6DofConstraints(const FRigidTransform3& SpaceTransform, const FPBD6DJointConstraints& Constraints, float ColorScale);
 
diff --git a/Engine/Source/Runtime/Experimental/Chaos/Public/Chaos/DenseMatrix.h b/Engine/Source/Runtime/Experimental/Chaos/Public/Chaos/DenseMatrix.h
--- a/Engine/Source/Runtime/Experimental/Chaos/Public/Chaos/DenseMatrix.h
+++ b/Engine/Source/Runtime/Experimental/Chaos/Public/Chaos/DenseMatrix.h
@@ -367,17 +367,6 @@ namespace Chaos
 			}
 		}
 
-		TDenseMatrix(const int32 InNRows, const int32 InNCols, std::initializer_list<FReal> InitList)
-			: NRows(InNRows)
-			, NCols(InNCols)
-		{
-			int32 N = FMath::Min<int32>(NumElements(), InitList.size());
-			for (int32 I = 0; I < N; ++I)
-			{
-				M[I] = InitList[I];
-			}
-		}
-
 		FReal M[MaxElements];
 		int32 NRows;
 		int32 NCols;
