Chaos : Builds

- Fix CIS, export FClusterCreationParameters

diff --git a/Engine/Source/Runtime/Experimental/Chaos/Private/Chaos/ClusterCreationParameters.cpp b/Engine/Source/Runtime/Experimental/Chaos/Private/Chaos/ClusterCreationParameters.cpp
--- a/Engine/Source/Runtime/Experimental/Chaos/Private/Chaos/ClusterCreationParameters.cpp
+++ b/Engine/Source/Runtime/Experimental/Chaos/Private/Chaos/ClusterCreationParameters.cpp
@@ -1,8 +1,9 @@
 // Copyright 1998-2019 Epic Games, Inc. All Rights Reserved.
 
 #include "Chaos/ClusterCreationParameters.h"
+#include "Chaos/BVHParticles.h"
 
 namespace Chaos
 {
-	template<> struct FClusterCreationParameters<float>;
+	template struct CHAOS_API FClusterCreationParameters<float>;
 }
diff --git a/Engine/Source/Runtime/Experimental/Chaos/Public/Chaos/ClusterCreationParameters.h b/Engine/Source/Runtime/Experimental/Chaos/Public/Chaos/ClusterCreationParameters.h
--- a/Engine/Source/Runtime/Experimental/Chaos/Public/Chaos/ClusterCreationParameters.h
+++ b/Engine/Source/Runtime/Experimental/Chaos/Public/Chaos/ClusterCreationParameters.h
@@ -10,7 +10,7 @@ namespace Chaos
 	class TBVHParticles;
 
 	template <typename T>
-	struct FClusterCreationParameters
+	struct CHAOS_API FClusterCreationParameters
 	{
 		enum EConnectionMethod { None = 0, PointImplicit, DelaunayTriangulation, MinimalSpanningSubsetDelaunayTriangulation, PointImplicitAugmentedWithMinimalDelaunay };
 
@@ -22,7 +22,7 @@ namespace Chaos
 			, bool bCopyCollisionParticlesIn = true
 			, bool bGenerateConnectionGraphIn = true
 			, EConnectionMethod ConnectionMethodIn = EConnectionMethod::PointImplicitAugmentedWithMinimalDelaunay
-			, TBVHParticles<float, 3>* CollisionParticlesIn = nullptr
+			, TBVHParticles<T, 3>* CollisionParticlesIn = nullptr
 			, int32 RigidBodyIndexIn = INDEX_NONE
 		)
 			: CoillisionThicknessPercent(CoillisionThicknessPercentIn)
@@ -41,7 +41,7 @@ namespace Chaos
 		bool bCopyCollisionParticles;
 		bool bGenerateConnectionGraph;
 		EConnectionMethod ConnectionMethod;
-		TBVHParticles<float, 3>* CollisionParticles;
+		TBVHParticles<T, 3>* CollisionParticles;
 		int32 RigidBodyIndex;
 	};
 }
