Chaos : Collisions

- Removed TConvexData

diff --git a/Engine/Source/Runtime/Experimental/Chaos/Private/Chaos/CollisionResolutionAlgo.cpp b/Engine/Source/Runtime/Experimental/Chaos/Private/Chaos/CollisionResolutionAlgo.cpp
--- a/Engine/Source/Runtime/Experimental/Chaos/Private/Chaos/CollisionResolutionAlgo.cpp
+++ b/Engine/Source/Runtime/Experimental/Chaos/Private/Chaos/CollisionResolutionAlgo.cpp
@@ -224,7 +224,7 @@ namespace Chaos
 		TVector<T, d> LocalNormal;
 		T LocalPhi = Object.PhiWithNormal(LocalPoint, LocalNormal);
 
-		TContactData<T, d> & Contact = Constraint.ShapeManifold.Manifold;
+		TConvexManifold<T, d> & Contact = Constraint.ShapeManifold;
 		if (LocalPhi < Contact.Phi)
 		{
 			Contact.Phi = LocalPhi;
@@ -242,7 +242,7 @@ namespace Chaos
 		TVector<T, d> LocalNormal;
 		T LocalPhi = Object.PhiWithNormal(LocalPoint, LocalNormal);
 
-		TContactData<T, d> & Contact = Constraint.ShapeManifold.Manifold;
+		TConvexManifold<T, d> & Contact = Constraint.ShapeManifold;
 		if (LocalPhi < Contact.Phi)
 		{
 			Contact.Phi = LocalPhi;
@@ -259,7 +259,7 @@ namespace Chaos
 		T LocalPhi = Object.PhiWithNormal(LocalPoint, LocalNormal);
 		T LocalThickness = LocalPhi - Thickness;
 
-		TContactData<T, d> & Contact = Constraint.ShapeManifold.Manifold;
+		TConvexManifold<T, d> & Contact = Constraint.ShapeManifold;
 		if (LocalThickness < -KINDA_SMALL_NUMBER)
 		{
 			Contact.Location += LocalPoint * LocalThickness;
@@ -289,8 +289,8 @@ namespace Chaos
 		SCOPE_CYCLE_COUNTER(STAT_SampleObject);
 		TRigidBodyContactConstraint<float, 3> AvgConstraint = Constraint;
 
-		TContactData<float, 3> & Contact = Constraint.ShapeManifold.Manifold;
-		TContactData<float, 3> & AvgContact = AvgConstraint.ShapeManifold.Manifold;
+		TConvexManifold<float, 3> & Contact = Constraint.ShapeManifold;
+		TConvexManifold<float, 3> & AvgContact = AvgConstraint.ShapeManifold;
 
 		AvgConstraint.Particle[0] = Constraint.Particle[0];
 		AvgConstraint.Particle[1] = Constraint.Particle[1];
@@ -399,8 +399,8 @@ namespace Chaos
 		SCOPE_CYCLE_COUNTER(STAT_SampleObject);
 		TRigidBodyContactConstraint<float, 3> AvgConstraint = Constraint;
 
-		TContactData<float, 3> & Contact = Constraint.ShapeManifold.Manifold;
-		TContactData<float, 3> & AvgContact = AvgConstraint.ShapeManifold.Manifold;
+		TConvexManifold<float, 3> & Contact = Constraint.ShapeManifold;
+		TConvexManifold<float, 3> & AvgContact = AvgConstraint.ShapeManifold;
 
 		AvgConstraint.Particle[0] = Constraint.Particle[0];
 		AvgConstraint.Particle[1] = Constraint.Particle[1];
@@ -823,7 +823,7 @@ namespace Chaos
 	template <typename T, int d>
 	bool UpdateBoxPlaneConstraint(const TBox<T, d>& Box, const TRigidTransform<T, d>& BoxTransform, const TPlane<T, d>& Plane, const TRigidTransform<T, d>& PlaneTransform, T Thickness, TRigidBodyContactConstraint<T, d>& Constraint)
 	{
-		TContactData<T, d> & Contact = Constraint.ShapeManifold.Manifold;
+		TConvexManifold<T, d> & Contact = Constraint.ShapeManifold;
 
 #if USING_CODE_ANALYSIS
 		MSVC_PRAGMA(warning(push))
@@ -885,7 +885,7 @@ namespace Chaos
 	template <typename T, int d>
 	void UpdateSphereConstraint(const TSphere<T, d>& Sphere1, const TRigidTransform<T, d>& Sphere1Transform, const TSphere<T, d>& Sphere2, const TRigidTransform<T, d>& Sphere2Transform, T Thickness, TRigidBodyContactConstraint<T, d>& Constraint)
 	{
-		TContactData<T, d> & Contact = Constraint.ShapeManifold.Manifold;
+		TConvexManifold<T, d> & Contact = Constraint.ShapeManifold;
 
 		const TVector<T, d> Center1 = Sphere1Transform.TransformPosition(Sphere1.GetCenter());
 		const TVector<T, d> Center2 = Sphere2Transform.TransformPosition(Sphere2.GetCenter());
@@ -903,7 +903,7 @@ namespace Chaos
 	template <typename T, int d>
 	void UpdateSpherePlaneConstraint(const TSphere<T, d>& Sphere, const TRigidTransform<T, d>& SphereTransform, const TPlane<T, d>& Plane, const TRigidTransform<T, d>& PlaneTransform, T Thickness, TRigidBodyContactConstraint<T, d>& Constraint)
 	{
-		TContactData<T, d> & Contact = Constraint.ShapeManifold.Manifold;
+		TConvexManifold<T, d> & Contact = Constraint.ShapeManifold;
 
 		const TRigidTransform<T, d> SphereToPlaneTransform(PlaneTransform.Inverse() * SphereTransform);
 		const TVector<T, d> SphereCenter = SphereToPlaneTransform.TransformPosition(Sphere.GetCenter());
@@ -923,7 +923,7 @@ namespace Chaos
 	template <typename T, int d>
 	void UpdateSphereBoxConstraint(const TSphere<T, d>& Sphere, const TRigidTransform<T, d>& SphereTransform, const TBox<T, d>& Box, const TRigidTransform<T, d>& BoxTransform, T Thickness, TRigidBodyContactConstraint<T, d>& Constraint)
 	{
-		TContactData<T, d> & Contact = Constraint.ShapeManifold.Manifold;
+		TConvexManifold<T, d> & Contact = Constraint.ShapeManifold;
 
 		const TRigidTransform<T, d> SphereToBoxTransform(SphereTransform * BoxTransform.Inverse());
 		const TVector<T, d> SphereCenterInBox = SphereToBoxTransform.TransformPosition(Sphere.GetCenter());
@@ -943,7 +943,7 @@ namespace Chaos
 	template <typename T, int d>
 	void UpdateCapsuleCapsuleConstraint(const TCapsule<T>& A, const TRigidTransform<T, d>& ATransform, const TCapsule<T>& B, const TRigidTransform<T, d>& BTransform, T Thickness, TRigidBodyContactConstraint<T, d>& Constraint)
 	{
-		TContactData<T, d> & Contact = Constraint.ShapeManifold.Manifold;
+		TConvexManifold<T, d> & Contact = Constraint.ShapeManifold;
 
 		FVector A1 = ATransform.TransformPosition(A.GetX1());
 		FVector A2 = ATransform.TransformPosition(A.GetX2());
@@ -971,7 +971,7 @@ namespace Chaos
 	void UpdateCapsuleBoxConstraint(const TCapsule<T>& A, const TRigidTransform<T, d>& ATransform, const TBox<T, d>& B, const TRigidTransform<T, d>& BTransform, T Thickness, TRigidBodyContactConstraint<T, d>& Constraint)
 	{
 		// @todo(ccaulfield): Add custom capsule-box collision
-		TContactData<T, d> & Contact = Constraint.ShapeManifold.Manifold;
+		TConvexManifold<T, d> & Contact = Constraint.ShapeManifold;
 
 		TRigidTransform<T, d> BToATransform = BTransform.GetRelativeTransform(ATransform);
 
@@ -1223,7 +1223,7 @@ namespace Chaos
 	template <typename T, int d>
 	void UpdateBoxConstraint(const TBox<T, d>& Box1, const TRigidTransform<T, d>& Box1Transform, const TBox<T, d>& Box2, const TRigidTransform<T, d>& Box2Transform, T Thickness, TRigidBodyContactConstraint<T, d>& Constraint)
 	{
-		TContactData<T, d> & Contact = Constraint.ShapeManifold.Manifold;
+		TConvexManifold<T, d> & Contact = Constraint.ShapeManifold;
 
 		TBox<T, d> Box2SpaceBox1 = Box1.TransformedBox(Box1Transform * Box2Transform.Inverse());
 		TBox<T, d> Box1SpaceBox2 = Box2.TransformedBox(Box2Transform * Box1Transform.Inverse());
diff --git a/Engine/Source/Runtime/Experimental/Chaos/Private/Chaos/PBDCollisionConstraint.cpp b/Engine/Source/Runtime/Experimental/Chaos/Private/Chaos/PBDCollisionConstraint.cpp
--- a/Engine/Source/Runtime/Experimental/Chaos/Private/Chaos/PBDCollisionConstraint.cpp
+++ b/Engine/Source/Runtime/Experimental/Chaos/Private/Chaos/PBDCollisionConstraint.cpp
@@ -367,7 +367,7 @@ namespace Chaos
 			TSerializablePtr<FChaosPhysicsMaterial> PhysicsMaterial0 = Particle0->AuxilaryValue(MPhysicsMaterials);
 			TSerializablePtr<FChaosPhysicsMaterial> PhysicsMaterial1 = Particle1->AuxilaryValue(MPhysicsMaterials);
 
-			TContactData<T, d> & Contact = Constraint.ShapeManifold.Manifold;
+			TConvexManifold<T, d> & Contact = Constraint.ShapeManifold;
 
 			TVector<T, d> VectorToPoint1 = Contact.Location - P0;
 			TVector<T, d> VectorToPoint2 = Contact.Location - P1;
@@ -600,7 +600,7 @@ namespace Chaos
 		{
 			UpdateConstraint<ECollisionUpdateType::Deepest>(MThickness, Constraint);
 
-			const TContactData<T, d> & Contact = Constraint.ShapeManifold.Manifold;
+			const TConvexManifold<T, d> & Contact = Constraint.ShapeManifold;
 
 			if (Contact.Phi >= MThickness)
 			{
diff --git a/Engine/Source/Runtime/Experimental/Chaos/Public/Chaos/PBDCollisionTypes.h b/Engine/Source/Runtime/Experimental/Chaos/Public/Chaos/PBDCollisionTypes.h
--- a/Engine/Source/Runtime/Experimental/Chaos/Public/Chaos/PBDCollisionTypes.h
+++ b/Engine/Source/Runtime/Experimental/Chaos/Public/Chaos/PBDCollisionTypes.h
@@ -36,25 +36,6 @@ namespace Chaos
 	};
 
 
-	/*
-	*
-	*  Contact manifold stored in the local space of the
-	*  Target object.
-	*
-	*/
-	template<class T, int d>
-	class CHAOS_API TContactData
-	{
-	public:
-
-		TContactData() : bDisabled(true), Normal(0), Phi(FLT_MAX) {};
-
-		bool bDisabled;
-		TVector<T, d> Normal;
-		TVector<T, d> Location;
-		T Phi;
-	};
-
 	/*
 	*
 	*/
@@ -62,17 +43,19 @@ namespace Chaos
 	class CHAOS_API TConvexManifold
 	{
 	public:
-		using FContactData = TContactData<T, d>;
-
 		TConvexManifold(int32 InTimestamp = -INT_MAX, const FImplicitObject* InImplicit0 = nullptr, const FImplicitObject* InImplicit1 = nullptr)
-			: Timestamp(InTimestamp)
+			: bDisabled(true), Timestamp(InTimestamp), Normal(0), Phi(FLT_MAX)
 		{
 			Implicit[0] = InImplicit0;
 			Implicit[1] = InImplicit1;
 		}
 
+		bool bDisabled;
 		int32 Timestamp;
-		FContactData Manifold;
+		TVector<T, d> Normal;
+		TVector<T, d> Location;
+		T Phi;
+
 		const FImplicitObject* Implicit[2]; // {Of Particle[0], Of Particle[1]}
 	};
 
@@ -114,20 +97,20 @@ namespace Chaos
 
 		//API
 		void ResetPhi(T InPhi) { SetPhi(InPhi); }
-		bool ContainsManifold(const FImplicitObject* B = nullptr, const FImplicitObject* A = nullptr) const { return (A == nullptr&&B == nullptr) ? (ShapeManifold.Implicit[0] != nullptr && ShapeManifold.Implicit[1] != nullptr) : (ShapeManifold.Implicit[0] == A && ShapeManifold.Implicit[1] == B); }
+		bool ContainsManifold(const FImplicitObject* B = nullptr, const FImplicitObject* A = nullptr) const { return (A==nullptr&&B==nullptr)?(ShapeManifold.Implicit[0] != nullptr && ShapeManifold.Implicit[1] != nullptr):(ShapeManifold.Implicit[0]==A && ShapeManifold.Implicit[1]==B); }
 
-		void SetDisabled(bool bInDisabled, const FImplicitObject* B = nullptr, const FImplicitObject* A = nullptr) { ShapeManifold.Manifold.bDisabled = bInDisabled; }
-		TVector<T, d> GetDisabled(const FImplicitObject* B = nullptr, const FImplicitObject* A = nullptr) const { return ShapeManifold.Manifold.bDisabled; }
+		void SetDisabled(bool bInDisabled, const FImplicitObject* B = nullptr, const FImplicitObject* A = nullptr) { ShapeManifold.bDisabled = bInDisabled; }
+		TVector<T, d> GetDisabled(const FImplicitObject* B = nullptr, const FImplicitObject* A = nullptr) const { return ShapeManifold.bDisabled; }
 
 
-		void SetNormal(const TVector<T, d> & InNormal, const FImplicitObject* B = nullptr, const FImplicitObject* A = nullptr) { ShapeManifold.Manifold.Normal = InNormal; }
-		TVector<T, d> GetNormal(const FImplicitObject* B = nullptr, const FImplicitObject* A = nullptr) const { return ShapeManifold.Manifold.Normal; }
+		void SetNormal(const TVector<T, d> & InNormal, const FImplicitObject* B = nullptr, const FImplicitObject* A = nullptr) { ShapeManifold.Normal = InNormal; }
+		TVector<T, d> GetNormal(const FImplicitObject* B = nullptr, const FImplicitObject* A = nullptr) const { return ShapeManifold.Normal; }
 
-		void SetLocation(const TVector<T, d> & InLocation, const FImplicitObject* B = nullptr, const FImplicitObject* A = nullptr, int32 Index = 0) { ShapeManifold.Manifold.Location = InLocation; }
-		TVector<T, d> GetLocation(const FImplicitObject* B = nullptr, const FImplicitObject* A = nullptr) const { return ShapeManifold.Manifold.Location; }
+		void SetLocation(const TVector<T, d> & InLocation, const FImplicitObject* B = nullptr, const FImplicitObject* A = nullptr, int32 Index=0) { ShapeManifold.Location = InLocation; }
+		TVector<T, d> GetLocation(const FImplicitObject* B = nullptr, const FImplicitObject* A = nullptr) const { return ShapeManifold.Location; }
 
-		void SetPhi(T InPhi, const FImplicitObject* B = nullptr, const FImplicitObject* A = nullptr, int32 Index = 0) { ShapeManifold.Manifold.Phi = InPhi; }
-		T GetPhi(const FImplicitObject* B = nullptr, const FImplicitObject* A = nullptr) const { return ShapeManifold.Manifold.Phi; }
+		void SetPhi(T InPhi, const FImplicitObject* B = nullptr, const FImplicitObject* A = nullptr, int32 Index = 0) { ShapeManifold.Phi = InPhi; }
+		T GetPhi(const FImplicitObject* B = nullptr, const FImplicitObject* A = nullptr) const { return ShapeManifold.Phi; }
 
 		void AddManifold(const FImplicitObject* B = nullptr, const FImplicitObject* A = nullptr) { ShapeManifold.Implicit[0] = A; ShapeManifold.Implicit[1] = B; }
 
