Chaos : Collisions

- Renamed RigidBodyContatConstraint

diff --git a/Engine/Source/Runtime/Experimental/Chaos/Private/Chaos/ChaosDebugDraw.cpp b/Engine/Source/Runtime/Experimental/Chaos/Private/Chaos/ChaosDebugDraw.cpp
--- a/Engine/Source/Runtime/Experimental/Chaos/Private/Chaos/ChaosDebugDraw.cpp
+++ b/Engine/Source/Runtime/Experimental/Chaos/Private/Chaos/ChaosDebugDraw.cpp
@@ -139,7 +139,7 @@ namespace Chaos
 		{
 			if (ConstraintHandle->GetType() == TCollisionConstraintBase<float,3>::FType::SinglePoint)
 			{
-				const TRigidBodyContactConstraint<FReal, 3>& Contact = ConstraintHandle->GetContact< TRigidBodyContactConstraint<float, 3> >();
+				const TRigidBodySingleContactConstraint<FReal, 3>& Contact = ConstraintHandle->GetContact< TRigidBodySingleContactConstraint<float, 3> >();
 				if (Contact.GetPhi() > 0)
 				{
 					ColorScale = ColorScale * (FReal)0.1;
diff --git a/Engine/Source/Runtime/Experimental/Chaos/Private/Chaos/CollisionResolutionAlgo.cpp b/Engine/Source/Runtime/Experimental/Chaos/Private/Chaos/CollisionResolutionAlgo.cpp
--- a/Engine/Source/Runtime/Experimental/Chaos/Private/Chaos/CollisionResolutionAlgo.cpp
+++ b/Engine/Source/Runtime/Experimental/Chaos/Private/Chaos/CollisionResolutionAlgo.cpp
@@ -24,7 +24,7 @@
 namespace Chaos
 {
 	template<typename T, int d>
-	void ConstructLevelsetConstraints(TGeometryParticleHandle<T, d>* Particle0, TGeometryParticleHandle<T, d>* Particle1, const FImplicitObject* Implicit0, const FImplicitObject* Implicit1, const T Thickness, TRigidBodyContactConstraint<T, d>& Constraint)
+	void ConstructLevelsetConstraints(TGeometryParticleHandle<T, d>* Particle0, TGeometryParticleHandle<T, d>* Particle1, const FImplicitObject* Implicit0, const FImplicitObject* Implicit1, const T Thickness, TRigidBodySingleContactConstraint<T, d>& Constraint)
 	{
 
 		if (!Particle1->Geometry() || (Particle0->AsDynamic() && !Particle0->AsDynamic()->CollisionParticlesSize() && Particle0->Geometry() && !Particle0->Geometry()->IsUnderlyingUnion()))
@@ -42,7 +42,7 @@ namespace Chaos
 	}
 
 	template<typename T, int d>
-	void ConstructBoxConstraints(TGeometryParticleHandle<T, d>* Particle0, TGeometryParticleHandle<T, d>* Particle1, const FImplicitObject* Implicit0, const FImplicitObject* Implicit1, const T Thickness, TRigidBodyContactConstraint<T, d>& Constraint)
+	void ConstructBoxConstraints(TGeometryParticleHandle<T, d>* Particle0, TGeometryParticleHandle<T, d>* Particle1, const FImplicitObject* Implicit0, const FImplicitObject* Implicit1, const T Thickness, TRigidBodySingleContactConstraint<T, d>& Constraint)
 	{
 		Constraint.Particle[0] = Particle0;
 		Constraint.Particle[1] = Particle1;
@@ -50,7 +50,7 @@ namespace Chaos
 	}
 
 	template<typename T, int d>
-	void ConstructBoxPlaneConstraints(TGeometryParticleHandle<T, d>* Particle0, TGeometryParticleHandle<T, d>* Particle1, const FImplicitObject* Implicit0, const FImplicitObject* Implicit1, const T Thickness, TRigidBodyContactConstraint<T, d>& Constraint)
+	void ConstructBoxPlaneConstraints(TGeometryParticleHandle<T, d>* Particle0, TGeometryParticleHandle<T, d>* Particle1, const FImplicitObject* Implicit0, const FImplicitObject* Implicit1, const T Thickness, TRigidBodySingleContactConstraint<T, d>& Constraint)
 	{
 		Constraint.Particle[0] = Particle0;
 		Constraint.Particle[1] = Particle1;
@@ -58,7 +58,7 @@ namespace Chaos
 	}
 
 	template<typename T, int d>
-	void ConstructSphereConstraints(TGeometryParticleHandle<T, d>* Particle0, TGeometryParticleHandle<T, d>* Particle1, const FImplicitObject* Implicit0, const FImplicitObject* Implicit1, const T Thickness, TRigidBodyContactConstraint<T, d>& Constraint)
+	void ConstructSphereConstraints(TGeometryParticleHandle<T, d>* Particle0, TGeometryParticleHandle<T, d>* Particle1, const FImplicitObject* Implicit0, const FImplicitObject* Implicit1, const T Thickness, TRigidBodySingleContactConstraint<T, d>& Constraint)
 	{
 		Constraint.Particle[0] = Particle0;
 		Constraint.Particle[1] = Particle1;
@@ -66,7 +66,7 @@ namespace Chaos
 	}
 
 	template<typename T, int d>
-	void ConstructSpherePlaneConstraints(TGeometryParticleHandle<T, d>* Particle0, TGeometryParticleHandle<T, d>* Particle1, const FImplicitObject* Implicit0, const FImplicitObject* Implicit1, const T Thickness, TRigidBodyContactConstraint<T, d>& Constraint)
+	void ConstructSpherePlaneConstraints(TGeometryParticleHandle<T, d>* Particle0, TGeometryParticleHandle<T, d>* Particle1, const FImplicitObject* Implicit0, const FImplicitObject* Implicit1, const T Thickness, TRigidBodySingleContactConstraint<T, d>& Constraint)
 	{
 		Constraint.Particle[0] = Particle0;
 		Constraint.Particle[1] = Particle1;
@@ -74,7 +74,7 @@ namespace Chaos
 	}
 
 	template<typename T, int d>
-	void ConstructSphereBoxConstraints(TGeometryParticleHandle<T, d>* Particle0, TGeometryParticleHandle<T, d>* Particle1, const FImplicitObject* Implicit0, const FImplicitObject* Implicit1, const T Thickness, TRigidBodyContactConstraint<T, d>& Constraint)
+	void ConstructSphereBoxConstraints(TGeometryParticleHandle<T, d>* Particle0, TGeometryParticleHandle<T, d>* Particle1, const FImplicitObject* Implicit0, const FImplicitObject* Implicit1, const T Thickness, TRigidBodySingleContactConstraint<T, d>& Constraint)
 	{
 		Constraint.Particle[0] = Particle0;
 		Constraint.Particle[1] = Particle1;
@@ -82,7 +82,7 @@ namespace Chaos
 	}
 
 	template<typename T, int d>
-	void ConstructCapsuleCapsuleConstraints(TGeometryParticleHandle<T, d>* Particle0, TGeometryParticleHandle<T, d>* Particle1, const FImplicitObject* Implicit0, const FImplicitObject* Implicit1, const T Thickness, TRigidBodyContactConstraint<T, d>& Constraint)
+	void ConstructCapsuleCapsuleConstraints(TGeometryParticleHandle<T, d>* Particle0, TGeometryParticleHandle<T, d>* Particle1, const FImplicitObject* Implicit0, const FImplicitObject* Implicit1, const T Thickness, TRigidBodySingleContactConstraint<T, d>& Constraint)
 	{
 		Constraint.Particle[0] = Particle0;
 		Constraint.Particle[1] = Particle1;
@@ -90,7 +90,7 @@ namespace Chaos
 	}
 
 	template<typename T, int d>
-	void ConstructCapsuleBoxConstraints(TGeometryParticleHandle<T, d>* Particle0, TGeometryParticleHandle<T, d>* Particle1, const FImplicitObject* Implicit0, const FImplicitObject* Implicit1, const T Thickness, TRigidBodyContactConstraint<T, d>& Constraint)
+	void ConstructCapsuleBoxConstraints(TGeometryParticleHandle<T, d>* Particle0, TGeometryParticleHandle<T, d>* Particle1, const FImplicitObject* Implicit0, const FImplicitObject* Implicit1, const T Thickness, TRigidBodySingleContactConstraint<T, d>& Constraint)
 	{
 		Constraint.Particle[0] = Particle0;
 		Constraint.Particle[1] = Particle1;
@@ -98,7 +98,7 @@ namespace Chaos
 	}
 
 	template<typename T, int d>
-	void ConstructSingleUnionConstraints(TGeometryParticleHandle<T, d>* Particle0, TGeometryParticleHandle<T, d>* Particle1, const FImplicitObject* Implicit0, const FImplicitObject* Implicit1, const T Thickness, TRigidBodyContactConstraint<T, d>& Constraint)
+	void ConstructSingleUnionConstraints(TGeometryParticleHandle<T, d>* Particle0, TGeometryParticleHandle<T, d>* Particle1, const FImplicitObject* Implicit0, const FImplicitObject* Implicit1, const T Thickness, TRigidBodySingleContactConstraint<T, d>& Constraint)
 	{
 		Constraint.Particle[0] = Particle0;
 		Constraint.Particle[1] = Particle1;
@@ -108,7 +108,7 @@ namespace Chaos
 	DECLARE_CYCLE_STAT(TEXT("TPBDCollisionConstraint::UpdateUnionUnionConstraint"), STAT_UpdateUnionUnionConstraint, STATGROUP_ChaosWide);
 
 	template<typename T, int d>
-	void ConstructUnionUnionConstraints(TGeometryParticleHandle<T, d>* Particle0, TGeometryParticleHandle<T, d>* Particle1, const FImplicitObject* Implicit0, const FImplicitObject* Implicit1, const T Thickness, TRigidBodyContactConstraint<T, d>& Constraint)
+	void ConstructUnionUnionConstraints(TGeometryParticleHandle<T, d>* Particle0, TGeometryParticleHandle<T, d>* Particle1, const FImplicitObject* Implicit0, const FImplicitObject* Implicit1, const T Thickness, TRigidBodySingleContactConstraint<T, d>& Constraint)
 	{
 		//SCOPE_CYCLE_COUNTER(STAT_UpdateUnionUnionConstraint);
 
@@ -138,7 +138,7 @@ namespace Chaos
 	}
 
 	template<typename T, int d>
-	void ConstructPairConstraintImpl(TGeometryParticleHandle<T, d>* Particle0, TGeometryParticleHandle<T, d>* Particle1, const FImplicitObject* Implicit0, const FImplicitObject* Implicit1, const T Thickness, TRigidBodyContactConstraint<T, d>& Constraint)
+	void ConstructPairConstraintImpl(TGeometryParticleHandle<T, d>* Particle0, TGeometryParticleHandle<T, d>* Particle1, const FImplicitObject* Implicit0, const FImplicitObject* Implicit1, const T Thickness, TRigidBodySingleContactConstraint<T, d>& Constraint)
 	{
 		// See if we already have a constraint for this shape pair
 		if (Constraint.ContainsManifold(Implicit0, Implicit1))
@@ -218,7 +218,7 @@ namespace Chaos
 	}
 
 	template <typename T, int d>
-	bool SampleObjectHelper2(const FImplicitObject& Object, const TRigidTransform<T, d>& ObjectTransform, const TRigidTransform<T, d>& SampleToObjectTransform, const TVector<T, d>& SampleParticle, T Thickness, TRigidBodyContactConstraint<T, d>& Constraint)
+	bool SampleObjectHelper2(const FImplicitObject& Object, const TRigidTransform<T, d>& ObjectTransform, const TRigidTransform<T, d>& SampleToObjectTransform, const TVector<T, d>& SampleParticle, T Thickness, TRigidBodySingleContactConstraint<T, d>& Constraint)
 	{
 		TVector<T, d> LocalPoint = SampleToObjectTransform.TransformPositionNoScale(SampleParticle);
 		TVector<T, d> LocalNormal;
@@ -236,7 +236,7 @@ namespace Chaos
 	}
 
 	template <typename T, int d>
-	bool SampleObjectNoNormal2(const FImplicitObject& Object, const TRigidTransform<T, d>& ObjectTransform, const TRigidTransform<T, d>& SampleToObjectTransform, const TVector<T, d>& SampleParticle, T Thickness, TRigidBodyContactConstraint<T, d>& Constraint)
+	bool SampleObjectNoNormal2(const FImplicitObject& Object, const TRigidTransform<T, d>& ObjectTransform, const TRigidTransform<T, d>& SampleToObjectTransform, const TVector<T, d>& SampleParticle, T Thickness, TRigidBodySingleContactConstraint<T, d>& Constraint)
 	{
 		TVector<T, d> LocalPoint = SampleToObjectTransform.TransformPositionNoScale(SampleParticle);
 		TVector<T, d> LocalNormal;
@@ -252,7 +252,7 @@ namespace Chaos
 	}
 
 	template <typename T, int d>
-	bool SampleObjectNormalAverageHelper2(const FImplicitObject& Object, const TRigidTransform<T, d>& ObjectTransform, const TRigidTransform<T, d>& SampleToObjectTransform, const TVector<T, d>& SampleParticle, T Thickness, T& TotalThickness, TRigidBodyContactConstraint<T, d>& Constraint)
+	bool SampleObjectNormalAverageHelper2(const FImplicitObject& Object, const TRigidTransform<T, d>& ObjectTransform, const TRigidTransform<T, d>& SampleToObjectTransform, const TVector<T, d>& SampleParticle, T Thickness, T& TotalThickness, TRigidBodySingleContactConstraint<T, d>& Constraint)
 	{
 		TVector<T, d> LocalPoint = SampleToObjectTransform.TransformPositionNoScale(SampleParticle);
 		TVector<T, d> LocalNormal;
@@ -284,10 +284,10 @@ namespace Chaos
 
 
 	template <ECollisionUpdateType UpdateType, typename T, int d>
-	void SampleObject2(const FImplicitObject& Object, const TRigidTransform<T, d>& ObjectTransform, const TBVHParticles<T, d>& SampleParticles, const TRigidTransform<T, d>& SampleParticlesTransform, T Thickness, TRigidBodyContactConstraint<T, d>& Constraint)
+	void SampleObject2(const FImplicitObject& Object, const TRigidTransform<T, d>& ObjectTransform, const TBVHParticles<T, d>& SampleParticles, const TRigidTransform<T, d>& SampleParticlesTransform, T Thickness, TRigidBodySingleContactConstraint<T, d>& Constraint)
 	{
 		SCOPE_CYCLE_COUNTER(STAT_SampleObject);
-		TRigidBodyContactConstraint<float, 3> AvgConstraint = Constraint;
+		TRigidBodySingleContactConstraint<float, 3> AvgConstraint = Constraint;
 
 		TConvexManifold<float, 3> & Contact = Constraint.ShapeManifold;
 		TConvexManifold<float, 3> & AvgContact = AvgConstraint.ShapeManifold;
@@ -394,10 +394,10 @@ namespace Chaos
 
 #if INTEL_ISPC
 	template<ECollisionUpdateType UpdateType>
-	void SampleObject2(const FImplicitObject& Object, const TRigidTransform<float, 3>& ObjectTransform, const TBVHParticles<float, 3>& SampleParticles, const TRigidTransform<float, 3>& SampleParticlesTransform, float Thickness, TRigidBodyContactConstraint<float, 3>& Constraint)
+	void SampleObject2(const FImplicitObject& Object, const TRigidTransform<float, 3>& ObjectTransform, const TBVHParticles<float, 3>& SampleParticles, const TRigidTransform<float, 3>& SampleParticlesTransform, float Thickness, TRigidBodySingleContactConstraint<float, 3>& Constraint)
 	{
 		SCOPE_CYCLE_COUNTER(STAT_SampleObject);
-		TRigidBodyContactConstraint<float, 3> AvgConstraint = Constraint;
+		TRigidBodySingleContactConstraint<float, 3> AvgConstraint = Constraint;
 
 		TConvexManifold<float, 3> & Contact = Constraint.ShapeManifold;
 		TConvexManifold<float, 3> & AvgContact = AvgConstraint.ShapeManifold;
@@ -702,7 +702,7 @@ namespace Chaos
 #endif
 
 	template<ECollisionUpdateType UpdateType, typename T, int d>
-	void UpdateConstraintImp(const FImplicitObject& ParticleObject, const TRigidTransform<T, d>& ParticleTM, const FImplicitObject& LevelsetObject, const TRigidTransform<T, d>& LevelsetTM, const T Thickness, TRigidBodyContactConstraint<T, d>& Constraint)
+	void UpdateConstraintImp(const FImplicitObject& ParticleObject, const TRigidTransform<T, d>& ParticleTM, const FImplicitObject& LevelsetObject, const TRigidTransform<T, d>& LevelsetTM, const T Thickness, TRigidBodySingleContactConstraint<T, d>& Constraint)
 	{
 		if (ParticleObject.GetType() == TBox<T, d>::StaticType() && LevelsetObject.GetType() == TBox<T, d>::StaticType())
 		{
@@ -734,7 +734,7 @@ namespace Chaos
 		}
 		else if (ParticleObject.GetType() == TPlane<T, d>::StaticType() && LevelsetObject.GetType() == TBox<T, d>::StaticType())
 		{
-			TRigidBodyContactConstraint<T, d> TmpConstraint = Constraint;
+			TRigidBodySingleContactConstraint<T, d> TmpConstraint = Constraint;
 			UpdateBoxPlaneConstraint(*LevelsetObject.template GetObject<TBox<T, d>>(), LevelsetTM, *ParticleObject.template GetObject<TPlane<T, d>>(), ParticleTM, Thickness, TmpConstraint);
 			if (TmpConstraint.GetPhi() < Constraint.GetPhi())
 			{
@@ -744,7 +744,7 @@ namespace Chaos
 		}
 		else if (ParticleObject.GetType() == TPlane<T, d>::StaticType() && LevelsetObject.GetType() == TSphere<T, d>::StaticType())
 		{
-			TRigidBodyContactConstraint<T, d> TmpConstraint = Constraint;
+			TRigidBodySingleContactConstraint<T, d> TmpConstraint = Constraint;
 			UpdateSpherePlaneConstraint(*LevelsetObject.template GetObject<TSphere<T, d>>(), LevelsetTM, *ParticleObject.template GetObject<TPlane<T, d>>(), ParticleTM, Thickness, TmpConstraint);
 			if (TmpConstraint.GetPhi() < Constraint.GetPhi())
 			{
@@ -754,7 +754,7 @@ namespace Chaos
 		}
 		else if (ParticleObject.GetType() == TBox<T, d>::StaticType() && LevelsetObject.GetType() == TSphere<T, d>::StaticType())
 		{
-			TRigidBodyContactConstraint<T, d> TmpConstraint = Constraint;
+			TRigidBodySingleContactConstraint<T, d> TmpConstraint = Constraint;
 			UpdateSphereBoxConstraint(*LevelsetObject.template GetObject<TSphere<T, d>>(), LevelsetTM, *ParticleObject.template GetObject<TBox<T, d>>(), ParticleTM, Thickness, TmpConstraint);
 			if (TmpConstraint.GetPhi() < Constraint.GetPhi())
 			{
@@ -764,7 +764,7 @@ namespace Chaos
 		}
 		else if (ParticleObject.GetType() == TBox<T, d>::StaticType() && LevelsetObject.GetType() == TCapsule<T>::StaticType())
 		{
-			TRigidBodyContactConstraint<T, d> TmpConstraint = Constraint;
+			TRigidBodySingleContactConstraint<T, d> TmpConstraint = Constraint;
 			UpdateCapsuleBoxConstraint(*LevelsetObject.template GetObject<TCapsule<T>>(), LevelsetTM, *ParticleObject.template GetObject<TBox<T, d>>(), ParticleTM, Thickness, TmpConstraint);
 			if (TmpConstraint.GetPhi() < Constraint.GetPhi())
 			{
@@ -803,7 +803,7 @@ namespace Chaos
 	}
 
 	template<typename T, int d>
-	void ConstructConstraintsImpl(TGeometryParticleHandle<T, d>* Particle0, TGeometryParticleHandle<T, d>* Particle1, const FImplicitObject* Implicit0, const FImplicitObject* Implicit1, const T Thickness, TRigidBodyContactConstraint<T, d> & Constraint)
+	void ConstructConstraintsImpl(TGeometryParticleHandle<T, d>* Particle0, TGeometryParticleHandle<T, d>* Particle1, const FImplicitObject* Implicit0, const FImplicitObject* Implicit1, const T Thickness, TRigidBodySingleContactConstraint<T, d> & Constraint)
 	{
 		// TriangleMesh implicits are for scene query only.
 		if (Implicit0 && GetInnerType(Implicit0->GetType()) == ImplicitObjectType::TriangleMesh) return;
@@ -821,7 +821,7 @@ namespace Chaos
 	}
 
 	template <typename T, int d>
-	bool UpdateBoxPlaneConstraint(const TBox<T, d>& Box, const TRigidTransform<T, d>& BoxTransform, const TPlane<T, d>& Plane, const TRigidTransform<T, d>& PlaneTransform, T Thickness, TRigidBodyContactConstraint<T, d>& Constraint)
+	bool UpdateBoxPlaneConstraint(const TBox<T, d>& Box, const TRigidTransform<T, d>& BoxTransform, const TPlane<T, d>& Plane, const TRigidTransform<T, d>& PlaneTransform, T Thickness, TRigidBodySingleContactConstraint<T, d>& Constraint)
 	{
 		TConvexManifold<T, d> & Contact = Constraint.ShapeManifold;
 
@@ -883,7 +883,7 @@ namespace Chaos
 	}
 
 	template <typename T, int d>
-	void UpdateSphereConstraint(const TSphere<T, d>& Sphere1, const TRigidTransform<T, d>& Sphere1Transform, const TSphere<T, d>& Sphere2, const TRigidTransform<T, d>& Sphere2Transform, T Thickness, TRigidBodyContactConstraint<T, d>& Constraint)
+	void UpdateSphereConstraint(const TSphere<T, d>& Sphere1, const TRigidTransform<T, d>& Sphere1Transform, const TSphere<T, d>& Sphere2, const TRigidTransform<T, d>& Sphere2Transform, T Thickness, TRigidBodySingleContactConstraint<T, d>& Constraint)
 	{
 		TConvexManifold<T, d> & Contact = Constraint.ShapeManifold;
 
@@ -901,7 +901,7 @@ namespace Chaos
 	}
 
 	template <typename T, int d>
-	void UpdateSpherePlaneConstraint(const TSphere<T, d>& Sphere, const TRigidTransform<T, d>& SphereTransform, const TPlane<T, d>& Plane, const TRigidTransform<T, d>& PlaneTransform, T Thickness, TRigidBodyContactConstraint<T, d>& Constraint)
+	void UpdateSpherePlaneConstraint(const TSphere<T, d>& Sphere, const TRigidTransform<T, d>& SphereTransform, const TPlane<T, d>& Plane, const TRigidTransform<T, d>& PlaneTransform, T Thickness, TRigidBodySingleContactConstraint<T, d>& Constraint)
 	{
 		TConvexManifold<T, d> & Contact = Constraint.ShapeManifold;
 
@@ -921,7 +921,7 @@ namespace Chaos
 	}
 
 	template <typename T, int d>
-	void UpdateSphereBoxConstraint(const TSphere<T, d>& Sphere, const TRigidTransform<T, d>& SphereTransform, const TBox<T, d>& Box, const TRigidTransform<T, d>& BoxTransform, T Thickness, TRigidBodyContactConstraint<T, d>& Constraint)
+	void UpdateSphereBoxConstraint(const TSphere<T, d>& Sphere, const TRigidTransform<T, d>& SphereTransform, const TBox<T, d>& Box, const TRigidTransform<T, d>& BoxTransform, T Thickness, TRigidBodySingleContactConstraint<T, d>& Constraint)
 	{
 		TConvexManifold<T, d> & Contact = Constraint.ShapeManifold;
 
@@ -941,7 +941,7 @@ namespace Chaos
 	}
 
 	template <typename T, int d>
-	void UpdateCapsuleCapsuleConstraint(const TCapsule<T>& A, const TRigidTransform<T, d>& ATransform, const TCapsule<T>& B, const TRigidTransform<T, d>& BTransform, T Thickness, TRigidBodyContactConstraint<T, d>& Constraint)
+	void UpdateCapsuleCapsuleConstraint(const TCapsule<T>& A, const TRigidTransform<T, d>& ATransform, const TCapsule<T>& B, const TRigidTransform<T, d>& BTransform, T Thickness, TRigidBodySingleContactConstraint<T, d>& Constraint)
 	{
 		TConvexManifold<T, d> & Contact = Constraint.ShapeManifold;
 
@@ -968,7 +968,7 @@ namespace Chaos
 	}
 
 	template <typename T, int d>
-	void UpdateCapsuleBoxConstraint(const TCapsule<T>& A, const TRigidTransform<T, d>& ATransform, const TBox<T, d>& B, const TRigidTransform<T, d>& BTransform, T Thickness, TRigidBodyContactConstraint<T, d>& Constraint)
+	void UpdateCapsuleBoxConstraint(const TCapsule<T>& A, const TRigidTransform<T, d>& ATransform, const TBox<T, d>& B, const TRigidTransform<T, d>& BTransform, T Thickness, TRigidBodySingleContactConstraint<T, d>& Constraint)
 	{
 		// @todo(ccaulfield): Add custom capsule-box collision
 		TConvexManifold<T, d> & Contact = Constraint.ShapeManifold;
@@ -1037,7 +1037,7 @@ namespace Chaos
 	}
 
 	template<ECollisionUpdateType UpdateType, typename T, int d>
-	void UpdateUnionUnionConstraint(const T Thickness, TRigidBodyContactConstraint<T, d>& Constraint)
+	void UpdateUnionUnionConstraint(const T Thickness, TRigidBodySingleContactConstraint<T, d>& Constraint)
 	{
 		SCOPE_CYCLE_COUNTER(STAT_UpdateUnionUnionConstraint);
 
@@ -1071,7 +1071,7 @@ namespace Chaos
 
 	DECLARE_CYCLE_STAT(TEXT("TPBDCollisionConstraint::UpdateSingleUnionConstraint"), STAT_UpdateSingleUnionConstraint, STATGROUP_ChaosWide);
 	template<ECollisionUpdateType UpdateType, typename T, int d>
-	void UpdateSingleUnionConstraint(const T Thickness, TRigidBodyContactConstraint<T, d>& Constraint)
+	void UpdateSingleUnionConstraint(const T Thickness, TRigidBodySingleContactConstraint<T, d>& Constraint)
 	{
 		SCOPE_CYCLE_COUNTER(STAT_UpdateSingleUnionConstraint);
 
@@ -1097,7 +1097,7 @@ namespace Chaos
 
 	DECLARE_CYCLE_STAT(TEXT("TPBDCollisionConstraint::UpdateUnionLevelsetConstraint"), STAT_UpdateUnionLevelsetConstraint, STATGROUP_ChaosWide);
 	template<ECollisionUpdateType UpdateType, typename T, int d>
-	void UpdateUnionLevelsetConstraint(const T Thickness, TRigidBodyContactConstraint<T, d>& Constraint)
+	void UpdateUnionLevelsetConstraint(const T Thickness, TRigidBodySingleContactConstraint<T, d>& Constraint)
 	{
 		SCOPE_CYCLE_COUNTER(STAT_UpdateUnionLevelsetConstraint);
 
@@ -1173,7 +1173,7 @@ namespace Chaos
 
 	DECLARE_CYCLE_STAT(TEXT("TPBDCollisionConstraint::UpdateLevelsetUnionConstraint"), STAT_UpdateLevelsetUnionConstraint, STATGROUP_ChaosWide);
 	template<ECollisionUpdateType UpdateType, typename T, int d>
-	void UpdateLevelsetUnionConstraint(const T Thickness, TRigidBodyContactConstraint<T, d>& Constraint)
+	void UpdateLevelsetUnionConstraint(const T Thickness, TRigidBodySingleContactConstraint<T, d>& Constraint)
 	{
 		SCOPE_CYCLE_COUNTER(STAT_UpdateLevelsetUnionConstraint);
 
@@ -1221,7 +1221,7 @@ namespace Chaos
 
 
 	template <typename T, int d>
-	void UpdateBoxConstraint(const TBox<T, d>& Box1, const TRigidTransform<T, d>& Box1Transform, const TBox<T, d>& Box2, const TRigidTransform<T, d>& Box2Transform, T Thickness, TRigidBodyContactConstraint<T, d>& Constraint)
+	void UpdateBoxConstraint(const TBox<T, d>& Box1, const TRigidTransform<T, d>& Box1Transform, const TBox<T, d>& Box2, const TRigidTransform<T, d>& Box2Transform, T Thickness, TRigidBodySingleContactConstraint<T, d>& Constraint)
 	{
 		TConvexManifold<T, d> & Contact = Constraint.ShapeManifold;
 
@@ -1274,7 +1274,7 @@ namespace Chaos
 
 	DECLARE_CYCLE_STAT(TEXT("TPBDCollisionConstraint::UpdateLevelsetConstraint"), STAT_UpdateLevelsetConstraint, STATGROUP_ChaosWide);
 	template<ECollisionUpdateType UpdateType, typename T, int d>
-	void UpdateLevelsetConstraint(const T Thickness, TRigidBodyContactConstraint<float, 3>& Constraint)
+	void UpdateLevelsetConstraint(const T Thickness, TRigidBodySingleContactConstraint<float, 3>& Constraint)
 	{
 		SCOPE_CYCLE_COUNTER(STAT_UpdateLevelsetConstraint);
 
@@ -1301,12 +1301,12 @@ namespace Chaos
 		}
 	}
 
-	template void ConstructConstraintsImpl<float,3>(TGeometryParticleHandle<float, 3>* Particle0, TGeometryParticleHandle<float, 3>* Particle1, const FImplicitObject* Implicit0, const FImplicitObject* Implicit1, const float Thickness, TRigidBodyContactConstraint<float, 3> & Constraint);
-	template void UpdateUnionLevelsetConstraint<ECollisionUpdateType::Any, float, 3>(const float Thickness, TRigidBodyContactConstraint<float, 3>& Constraint);
-	template void UpdateUnionLevelsetConstraint<ECollisionUpdateType::Deepest, float, 3>(const float Thickness, TRigidBodyContactConstraint<float, 3>& Constraint);
-	template void UpdateConstraintImp<ECollisionUpdateType::Any, float, 3>(const FImplicitObject& ParticleObject, const TRigidTransform<float, 3>& ParticleTM, const FImplicitObject& LevelsetObject, const TRigidTransform<float, 3>& LevelsetTM, const float Thickness, TRigidBodyContactConstraint<float, 3>& Constraints);
-	template void UpdateConstraintImp<ECollisionUpdateType::Deepest, float, 3>(const FImplicitObject& ParticleObject, const TRigidTransform<float, 3>& ParticleTM, const FImplicitObject& LevelsetObject, const TRigidTransform<float, 3>& LevelsetTM, const float Thickness, TRigidBodyContactConstraint<float, 3>& Constraint);
-	template void UpdateLevelsetConstraint<ECollisionUpdateType::Any, float, 3>(const float Thickness, TRigidBodyContactConstraint<float, 3>& Constraint);
-	template void UpdateLevelsetConstraint<ECollisionUpdateType::Deepest, float, 3>(const float Thickness, TRigidBodyContactConstraint<float, 3>& Constraint);
+	template void ConstructConstraintsImpl<float,3>(TGeometryParticleHandle<float, 3>* Particle0, TGeometryParticleHandle<float, 3>* Particle1, const FImplicitObject* Implicit0, const FImplicitObject* Implicit1, const float Thickness, TRigidBodySingleContactConstraint<float, 3> & Constraint);
+	template void UpdateUnionLevelsetConstraint<ECollisionUpdateType::Any, float, 3>(const float Thickness, TRigidBodySingleContactConstraint<float, 3>& Constraint);
+	template void UpdateUnionLevelsetConstraint<ECollisionUpdateType::Deepest, float, 3>(const float Thickness, TRigidBodySingleContactConstraint<float, 3>& Constraint);
+	template void UpdateConstraintImp<ECollisionUpdateType::Any, float, 3>(const FImplicitObject& ParticleObject, const TRigidTransform<float, 3>& ParticleTM, const FImplicitObject& LevelsetObject, const TRigidTransform<float, 3>& LevelsetTM, const float Thickness, TRigidBodySingleContactConstraint<float, 3>& Constraints);
+	template void UpdateConstraintImp<ECollisionUpdateType::Deepest, float, 3>(const FImplicitObject& ParticleObject, const TRigidTransform<float, 3>& ParticleTM, const FImplicitObject& LevelsetObject, const TRigidTransform<float, 3>& LevelsetTM, const float Thickness, TRigidBodySingleContactConstraint<float, 3>& Constraint);
+	template void UpdateLevelsetConstraint<ECollisionUpdateType::Any, float, 3>(const float Thickness, TRigidBodySingleContactConstraint<float, 3>& Constraint);
+	template void UpdateLevelsetConstraint<ECollisionUpdateType::Deepest, float, 3>(const float Thickness, TRigidBodySingleContactConstraint<float, 3>& Constraint);
 
 }
diff --git a/Engine/Source/Runtime/Experimental/Chaos/Private/Chaos/PBDCollisionConstraint.cpp b/Engine/Source/Runtime/Experimental/Chaos/Private/Chaos/PBDCollisionConstraint.cpp
--- a/Engine/Source/Runtime/Experimental/Chaos/Private/Chaos/PBDCollisionConstraint.cpp
+++ b/Engine/Source/Runtime/Experimental/Chaos/Private/Chaos/PBDCollisionConstraint.cpp
@@ -282,7 +282,7 @@ namespace Chaos
 	}
 
 	template<typename T, int d>
-	TVector<T, d> GetEnergyClampedImpulse(const TRigidBodyContactConstraint<T, d>& Constraint, const TVector<T, d>& Impulse, const TVector<T, d>& VectorToPoint1, const TVector<T, d>& VectorToPoint2, const TVector<T, d>& Velocity1, const TVector<T, d>& Velocity2)
+	TVector<T, d> GetEnergyClampedImpulse(const TRigidBodySingleContactConstraint<T, d>& Constraint, const TVector<T, d>& Impulse, const TVector<T, d>& VectorToPoint1, const TVector<T, d>& VectorToPoint2, const TVector<T, d>& Velocity1, const TVector<T, d>& Velocity2)
 	{
 		TPBDRigidParticleHandle<T, d>* PBDRigid0 = Constraint.Particle[0]->AsDynamic();
 		TPBDRigidParticleHandle<T, d>* PBDRigid1 = Constraint.Particle[1]->AsDynamic();
@@ -712,7 +712,7 @@ namespace Chaos
 	}
 
 	template<typename T, int d>
-	void TPBDCollisionConstraint<T, d>::ConstructConstraints(TGeometryParticleHandle<T, d>* Particle0, TGeometryParticleHandle<T, d>* Particle1, const T Thickness, TRigidBodyContactConstraint<T, d> & Constraint)
+	void TPBDCollisionConstraint<T, d>::ConstructConstraints(TGeometryParticleHandle<T, d>* Particle0, TGeometryParticleHandle<T, d>* Particle1, const T Thickness, TRigidBodySingleContactConstraint<T, d> & Constraint)
 	{
 		if (ensure(Particle0 && Particle1))
 		{
diff --git a/Engine/Source/Runtime/Experimental/Chaos/Private/Chaos/PBDCollisionTypes.cpp b/Engine/Source/Runtime/Experimental/Chaos/Private/Chaos/PBDCollisionTypes.cpp
--- a/Engine/Source/Runtime/Experimental/Chaos/Private/Chaos/PBDCollisionTypes.cpp
+++ b/Engine/Source/Runtime/Experimental/Chaos/Private/Chaos/PBDCollisionTypes.cpp
@@ -19,6 +19,6 @@ namespace Chaos
 	}
 
 	template class CHAOS_API TPBDCollisionConstraintHandle<float, 3>;
-	template CHAOS_API TRigidBodyContactConstraint<float, 3>& TPBDCollisionConstraintHandle<float, 3>::GetContact<TRigidBodyContactConstraint<float, 3>>();
-	template CHAOS_API const TRigidBodyContactConstraint<float, 3>& TPBDCollisionConstraintHandle<float, 3>::GetContact<TRigidBodyContactConstraint<float, 3>>() const;
+	template CHAOS_API TRigidBodySingleContactConstraint<float, 3>& TPBDCollisionConstraintHandle<float, 3>::GetContact<TRigidBodySingleContactConstraint<float, 3>>();
+	template CHAOS_API const TRigidBodySingleContactConstraint<float, 3>& TPBDCollisionConstraintHandle<float, 3>::GetContact<TRigidBodySingleContactConstraint<float, 3>>() const;
 };
diff --git a/Engine/Source/Runtime/Experimental/Chaos/Private/Chaos/PBDRigidClustering.cpp b/Engine/Source/Runtime/Experimental/Chaos/Private/Chaos/PBDRigidClustering.cpp
--- a/Engine/Source/Runtime/Experimental/Chaos/Private/Chaos/PBDRigidClustering.cpp
+++ b/Engine/Source/Runtime/Experimental/Chaos/Private/Chaos/PBDRigidClustering.cpp
@@ -1757,7 +1757,7 @@ namespace Chaos
 	}
 
 	template<class T, int d>
-	TVector<T, d> GetContactLocation(const TRigidBodyContactConstraint<T, d>& Contact)
+	TVector<T, d> GetContactLocation(const TRigidBodySingleContactConstraint<T, d>& Contact)
 	{
 		return Contact.GetLocation();
 	}
diff --git a/Engine/Source/Runtime/Experimental/Chaos/Public/Chaos/CollisionResolutionAlgo.h b/Engine/Source/Runtime/Experimental/Chaos/Public/Chaos/CollisionResolutionAlgo.h
--- a/Engine/Source/Runtime/Experimental/Chaos/Public/Chaos/CollisionResolutionAlgo.h
+++ b/Engine/Source/Runtime/Experimental/Chaos/Public/Chaos/CollisionResolutionAlgo.h
@@ -19,7 +19,7 @@ namespace Chaos
 	class FImplicitObject;
 
 	template<typename T, int d>
-	class TRigidBodyContactConstraint;
+	class TRigidBodySingleContactConstraint;
 
 	template <typename T, int d>
 	class TRigidTransform;
@@ -38,82 +38,82 @@ namespace Chaos
 	// Construct Constraints
 	//
 	template<typename T, int d>
-	void ConstructLevelsetConstraints(TGeometryParticleHandle<T, d>* Particle0, TGeometryParticleHandle<T, d>* Particle1, const FImplicitObject* Implicit0, const FImplicitObject* Implicit1, const T Thickness, TRigidBodyContactConstraint<T, d>& Constraint);
+	void ConstructLevelsetConstraints(TGeometryParticleHandle<T, d>* Particle0, TGeometryParticleHandle<T, d>* Particle1, const FImplicitObject* Implicit0, const FImplicitObject* Implicit1, const T Thickness, TRigidBodySingleContactConstraint<T, d>& Constraint);
 
 	template<typename T, int d>
-	void ConstructBoxConstraints(TGeometryParticleHandle<T, d>* Particle0, TGeometryParticleHandle<T, d>* Particle1, const FImplicitObject* Implicit0, const FImplicitObject* Implicit1, const T Thickness, TRigidBodyContactConstraint<T, d>& Constraint);
+	void ConstructBoxConstraints(TGeometryParticleHandle<T, d>* Particle0, TGeometryParticleHandle<T, d>* Particle1, const FImplicitObject* Implicit0, const FImplicitObject* Implicit1, const T Thickness, TRigidBodySingleContactConstraint<T, d>& Constraint);
 
 	template<typename T, int d>
-	void ConstructBoxPlaneConstraints(TGeometryParticleHandle<T, d>* Particle0, TGeometryParticleHandle<T, d>* Particle1, const FImplicitObject* Implicit0, const FImplicitObject* Implicit1, const T Thickness, TRigidBodyContactConstraint<T, d>& Constraint);
+	void ConstructBoxPlaneConstraints(TGeometryParticleHandle<T, d>* Particle0, TGeometryParticleHandle<T, d>* Particle1, const FImplicitObject* Implicit0, const FImplicitObject* Implicit1, const T Thickness, TRigidBodySingleContactConstraint<T, d>& Constraint);
 
 	template<typename T, int d>
-	void ConstructSphereConstraints(TGeometryParticleHandle<T, d>* Particle0, TGeometryParticleHandle<T, d>* Particle1, const FImplicitObject* Implicit0, const FImplicitObject* Implicit1, const T Thickness, TRigidBodyContactConstraint<T, d>& Constraint);
+	void ConstructSphereConstraints(TGeometryParticleHandle<T, d>* Particle0, TGeometryParticleHandle<T, d>* Particle1, const FImplicitObject* Implicit0, const FImplicitObject* Implicit1, const T Thickness, TRigidBodySingleContactConstraint<T, d>& Constraint);
 
 	template<typename T, int d>
-	void ConstructSpherePlaneConstraints(TGeometryParticleHandle<T, d>* Particle0, TGeometryParticleHandle<T, d>* Particle1, const FImplicitObject* Implicit0, const FImplicitObject* Implicit1, const T Thickness, TRigidBodyContactConstraint<T, d>& Constraint);
+	void ConstructSpherePlaneConstraints(TGeometryParticleHandle<T, d>* Particle0, TGeometryParticleHandle<T, d>* Particle1, const FImplicitObject* Implicit0, const FImplicitObject* Implicit1, const T Thickness, TRigidBodySingleContactConstraint<T, d>& Constraint);
 
 	template<typename T, int d>
-	void ConstructSphereBoxConstraints(TGeometryParticleHandle<T, d>* Particle0, TGeometryParticleHandle<T, d>* Particle1, const FImplicitObject* Implicit0, const FImplicitObject* Implicit1, const T Thickness, TRigidBodyContactConstraint<T, d>& Constraint);
+	void ConstructSphereBoxConstraints(TGeometryParticleHandle<T, d>* Particle0, TGeometryParticleHandle<T, d>* Particle1, const FImplicitObject* Implicit0, const FImplicitObject* Implicit1, const T Thickness, TRigidBodySingleContactConstraint<T, d>& Constraint);
 
 	template<typename T, int d>
-	void ConstructCapsuleCapsuleConstraints(TGeometryParticleHandle<T, d>* Particle0, TGeometryParticleHandle<T, d>* Particle1, const FImplicitObject* Implicit0, const FImplicitObject* Implicit1, const T Thickness, TRigidBodyContactConstraint<T, d>& Constraint);
+	void ConstructCapsuleCapsuleConstraints(TGeometryParticleHandle<T, d>* Particle0, TGeometryParticleHandle<T, d>* Particle1, const FImplicitObject* Implicit0, const FImplicitObject* Implicit1, const T Thickness, TRigidBodySingleContactConstraint<T, d>& Constraint);
 
 	template<typename T, int d>
-	void ConstructCapsuleBoxConstraints(TGeometryParticleHandle<T, d>* Particle0, TGeometryParticleHandle<T, d>* Particle1, const FImplicitObject* Implicit0, const FImplicitObject* Implicit1, const T Thickness, TRigidBodyContactConstraint<T, d>& Constraint);
+	void ConstructCapsuleBoxConstraints(TGeometryParticleHandle<T, d>* Particle0, TGeometryParticleHandle<T, d>* Particle1, const FImplicitObject* Implicit0, const FImplicitObject* Implicit1, const T Thickness, TRigidBodySingleContactConstraint<T, d>& Constraint);
 
 	template<typename T, int d>
-	void ConstructSingleUnionConstraints(TGeometryParticleHandle<T, d>* Particle0, TGeometryParticleHandle<T, d>* Particle1, const FImplicitObject* Implicit0, const FImplicitObject* Implicit1, const T Thickness, TRigidBodyContactConstraint<T, d>& Constraint);
+	void ConstructSingleUnionConstraints(TGeometryParticleHandle<T, d>* Particle0, TGeometryParticleHandle<T, d>* Particle1, const FImplicitObject* Implicit0, const FImplicitObject* Implicit1, const T Thickness, TRigidBodySingleContactConstraint<T, d>& Constraint);
 
 	template<typename T, int d>
-	void ConstructUnionUnionConstraints(TGeometryParticleHandle<T, d>* Particle0, TGeometryParticleHandle<T, d>* Particle1, const FImplicitObject* Implicit0, const FImplicitObject* Implicit1, const T Thickness, TRigidBodyContactConstraint<T, d>& Constraint);
+	void ConstructUnionUnionConstraints(TGeometryParticleHandle<T, d>* Particle0, TGeometryParticleHandle<T, d>* Particle1, const FImplicitObject* Implicit0, const FImplicitObject* Implicit1, const T Thickness, TRigidBodySingleContactConstraint<T, d>& Constraint);
 
 	template<typename T, int d>
-	void ConstructPairConstraintImpl(TGeometryParticleHandle<T, d>* Particle0, TGeometryParticleHandle<T, d>* Particle1, const FImplicitObject* Implicit0, const FImplicitObject* Implicit1, const T Thickness, TRigidBodyContactConstraint<T, d>& Constraint);
+	void ConstructPairConstraintImpl(TGeometryParticleHandle<T, d>* Particle0, TGeometryParticleHandle<T, d>* Particle1, const FImplicitObject* Implicit0, const FImplicitObject* Implicit1, const T Thickness, TRigidBodySingleContactConstraint<T, d>& Constraint);
 
 	template<typename T, int d>
-	void ConstructConstraintsImpl(TGeometryParticleHandle<T, d>* Particle0, TGeometryParticleHandle<T, d>* Particle1, const FImplicitObject* Implicit0, const FImplicitObject* Implicit1, const T Thickness, TRigidBodyContactConstraint<T, d> & Constraint);
+	void ConstructConstraintsImpl(TGeometryParticleHandle<T, d>* Particle0, TGeometryParticleHandle<T, d>* Particle1, const FImplicitObject* Implicit0, const FImplicitObject* Implicit1, const T Thickness, TRigidBodySingleContactConstraint<T, d> & Constraint);
 
 
 	//
 	// Update Constraints
 	//
 	template<ECollisionUpdateType UpdateType, typename T = float, int d = 3>
-	void CHAOS_API UpdateLevelsetConstraint(const T Thickness, TRigidBodyContactConstraint<float, 3>& Constraint);
+	void CHAOS_API UpdateLevelsetConstraint(const T Thickness, TRigidBodySingleContactConstraint<float, 3>& Constraint);
 
 	template<ECollisionUpdateType UpdateType, typename T, int d>
-	void UpdateConstraintImp(const FImplicitObject& ParticleObject, const TRigidTransform<T, d>& ParticleTM, const FImplicitObject& LevelsetObject, const TRigidTransform<T, d>& LevelsetTM, const T Thickness, TRigidBodyContactConstraint<T, d>& Constraint);
+	void UpdateConstraintImp(const FImplicitObject& ParticleObject, const TRigidTransform<T, d>& ParticleTM, const FImplicitObject& LevelsetObject, const TRigidTransform<T, d>& LevelsetTM, const T Thickness, TRigidBodySingleContactConstraint<T, d>& Constraint);
 
 	template <typename T, int d>
-	bool UpdateBoxPlaneConstraint(const TBox<T, d>& Box, const TRigidTransform<T, d>& BoxTransform, const TPlane<T, d>& Plane, const TRigidTransform<T, d>& PlaneTransform, T Thickness, TRigidBodyContactConstraint<T, d>& Constraint);
+	bool UpdateBoxPlaneConstraint(const TBox<T, d>& Box, const TRigidTransform<T, d>& BoxTransform, const TPlane<T, d>& Plane, const TRigidTransform<T, d>& PlaneTransform, T Thickness, TRigidBodySingleContactConstraint<T, d>& Constraint);
 
 	template <typename T, int d>
-	void UpdateSphereConstraint(const TSphere<T, d>& Sphere1, const TRigidTransform<T, d>& Sphere1Transform, const TSphere<T, d>& Sphere2, const TRigidTransform<T, d>& Sphere2Transform, T Thickness, TRigidBodyContactConstraint<T, d>& Constraint);
+	void UpdateSphereConstraint(const TSphere<T, d>& Sphere1, const TRigidTransform<T, d>& Sphere1Transform, const TSphere<T, d>& Sphere2, const TRigidTransform<T, d>& Sphere2Transform, T Thickness, TRigidBodySingleContactConstraint<T, d>& Constraint);
 
 	template <typename T, int d>
-	void UpdateSpherePlaneConstraint(const TSphere<T, d>& Sphere, const TRigidTransform<T, d>& SphereTransform, const TPlane<T, d>& Plane, const TRigidTransform<T, d>& PlaneTransform, T Thickness, TRigidBodyContactConstraint<T, d>& Constraint);
+	void UpdateSpherePlaneConstraint(const TSphere<T, d>& Sphere, const TRigidTransform<T, d>& SphereTransform, const TPlane<T, d>& Plane, const TRigidTransform<T, d>& PlaneTransform, T Thickness, TRigidBodySingleContactConstraint<T, d>& Constraint);
 
 	template <typename T, int d>
-	void UpdateSphereBoxConstraint(const TSphere<T, d>& Sphere, const TRigidTransform<T, d>& SphereTransform, const TBox<T, d>& Box, const TRigidTransform<T, d>& BoxTransform, T Thickness, TRigidBodyContactConstraint<T, d>& Constraint);
+	void UpdateSphereBoxConstraint(const TSphere<T, d>& Sphere, const TRigidTransform<T, d>& SphereTransform, const TBox<T, d>& Box, const TRigidTransform<T, d>& BoxTransform, T Thickness, TRigidBodySingleContactConstraint<T, d>& Constraint);
 
 	template <typename T, int d>
-	void UpdateCapsuleCapsuleConstraint(const TCapsule<T>& A, const TRigidTransform<T, d>& ATransform, const TCapsule<T>& B, const TRigidTransform<T, d>& BTransform, T Thickness, TRigidBodyContactConstraint<T, d>& Constraint);
+	void UpdateCapsuleCapsuleConstraint(const TCapsule<T>& A, const TRigidTransform<T, d>& ATransform, const TCapsule<T>& B, const TRigidTransform<T, d>& BTransform, T Thickness, TRigidBodySingleContactConstraint<T, d>& Constraint);
 
 	template <typename T, int d>
-	void UpdateCapsuleBoxConstraint(const TCapsule<T>& A, const TRigidTransform<T, d>& ATransform, const TBox<T, d>& B, const TRigidTransform<T, d>& BTransform, T Thickness, TRigidBodyContactConstraint<T, d>& Constraint);
+	void UpdateCapsuleBoxConstraint(const TCapsule<T>& A, const TRigidTransform<T, d>& ATransform, const TBox<T, d>& B, const TRigidTransform<T, d>& BTransform, T Thickness, TRigidBodySingleContactConstraint<T, d>& Constraint);
 
 	template<ECollisionUpdateType UpdateType, typename T, int d>
-	void UpdateUnionUnionConstraint(const T Thickness, TRigidBodyContactConstraint<T, d>& Constraint);
+	void UpdateUnionUnionConstraint(const T Thickness, TRigidBodySingleContactConstraint<T, d>& Constraint);
 
 	template<ECollisionUpdateType UpdateType, typename T, int d>
-	void UpdateSingleUnionConstraint(const T Thickness, TRigidBodyContactConstraint<T, d>& Constraint);
+	void UpdateSingleUnionConstraint(const T Thickness, TRigidBodySingleContactConstraint<T, d>& Constraint);
 
 	template<ECollisionUpdateType UpdateType, typename T, int d>
-	void UpdateUnionLevelsetConstraint(const T Thickness, TRigidBodyContactConstraint<T, d>& Constraint);
+	void UpdateUnionLevelsetConstraint(const T Thickness, TRigidBodySingleContactConstraint<T, d>& Constraint);
 
 	template<ECollisionUpdateType UpdateType, typename T, int d>
-	void UpdateLevelsetUnionConstraint(const T Thickness, TRigidBodyContactConstraint<T, d>& Constraint);
+	void UpdateLevelsetUnionConstraint(const T Thickness, TRigidBodySingleContactConstraint<T, d>& Constraint);
 
 	template <typename T, int d>
-	void UpdateBoxConstraint(const TBox<T, d>& Box1, const TRigidTransform<T, d>& Box1Transform, const TBox<T, d>& Box2, const TRigidTransform<T, d>& Box2Transform, T Thickness, TRigidBodyContactConstraint<T, d>& Constraint);
+	void UpdateBoxConstraint(const TBox<T, d>& Box1, const TRigidTransform<T, d>& Box1Transform, const TBox<T, d>& Box2, const TRigidTransform<T, d>& Box2Transform, T Thickness, TRigidBodySingleContactConstraint<T, d>& Constraint);
 
 }
diff --git a/Engine/Source/Runtime/Experimental/Chaos/Public/Chaos/CollisionResolutionConvexConvex.h b/Engine/Source/Runtime/Experimental/Chaos/Public/Chaos/CollisionResolutionConvexConvex.h
--- a/Engine/Source/Runtime/Experimental/Chaos/Public/Chaos/CollisionResolutionConvexConvex.h
+++ b/Engine/Source/Runtime/Experimental/Chaos/Public/Chaos/CollisionResolutionConvexConvex.h
@@ -17,7 +17,7 @@ class CHAOS_API CollisionResolutionConvexConvex
 public:
 	using FConvex = TConvex<T,d>;
 	using FRigidTransform = TRigidTransform<T,d>;
-	using FRigidBodyContactConstraint = TRigidBodyContactConstraint<T, d>;
+	using FRigidBodyContactConstraint = TRigidBodySingleContactConstraint<T, d>;
 	using FGeometryParticleHandle = TGeometryParticleHandle<T, d>;
 
 	/*
diff --git a/Engine/Source/Runtime/Experimental/Chaos/Public/Chaos/PBDCollisionConstraint.h b/Engine/Source/Runtime/Experimental/Chaos/Public/Chaos/PBDCollisionConstraint.h
--- a/Engine/Source/Runtime/Experimental/Chaos/Public/Chaos/PBDCollisionConstraint.h
+++ b/Engine/Source/Runtime/Experimental/Chaos/Public/Chaos/PBDCollisionConstraint.h
@@ -65,7 +65,7 @@ public:
 	using FConstraintHandleID = TPair<const TGeometryParticleHandle<T, d>*, const TGeometryParticleHandle<T, d>*>;
 	using FConstraintContainerHandle = TPBDCollisionConstraintHandle<T, d>; // @todo(brice) rename this to FCollisionConstraintHandle
 	using FConstraintHandleAllocator = TConstraintHandleAllocator<TPBDCollisionConstraint<T, d>>;
-	using FRigidBodyContactConstraint = TRigidBodyContactConstraint<T, d>;
+	using FRigidBodyContactConstraint = TRigidBodySingleContactConstraint<T, d>;
 	using FHandles = TArray<FConstraintContainerHandle*>;
 
 	TPBDCollisionConstraint(const TPBDRigidsSOAs<T,d>& InParticles, TArrayCollectionArray<bool>& Collided, const TArrayCollectionArray<TSerializablePtr<FChaosPhysicsMaterial>>& PerParticleMaterials, const int32 ApplyPairIterations = 1, const int32 ApplyPushOutPairIterations = 1, const T Thickness = (T)0);
@@ -212,7 +212,7 @@ private:
 	template<ECollisionUpdateType>
 	void UpdateConstraint(const T Thickness, FRigidBodyContactConstraint& Constraint);
 
-	void ConstructConstraints(TGeometryParticleHandle<T, d>* Particle0, TGeometryParticleHandle<T, d>* Particle1, const T Thickness, TRigidBodyContactConstraint<T, d> & Constraint);
+	void ConstructConstraints(TGeometryParticleHandle<T, d>* Particle0, TGeometryParticleHandle<T, d>* Particle1, const T Thickness, TRigidBodySingleContactConstraint<T, d> & Constraint);
 
 
 	const TPBDRigidsSOAs<T,d>& Particles;
diff --git a/Engine/Source/Runtime/Experimental/Chaos/Public/Chaos/PBDCollisionConstraintImp.h b/Engine/Source/Runtime/Experimental/Chaos/Public/Chaos/PBDCollisionConstraintImp.h
--- a/Engine/Source/Runtime/Experimental/Chaos/Public/Chaos/PBDCollisionConstraintImp.h
+++ b/Engine/Source/Runtime/Experimental/Chaos/Public/Chaos/PBDCollisionConstraintImp.h
@@ -152,7 +152,7 @@ namespace Chaos
 			// Narrow phase
 			CHAOS_SCOPED_TIMER(ComputeConstraints_NP);
 
-			TQueue<TRigidBodyContactConstraint<T, d>, EQueueMode::Mpsc> Queue;	//todo(ocohen): use per thread buffer instead, need better support than ParallelFor for this
+			TQueue<TRigidBodySingleContactConstraint<T, d>, EQueueMode::Mpsc> Queue;	//todo(ocohen): use per thread buffer instead, need better support than ParallelFor for this
 			Particles.GetNonDisabledDynamicView().ParallelFor([&](auto& Particle1, int32 ActiveIdxIdx)
 			{
 #if !UE_BUILD_SHIPPING
@@ -253,7 +253,7 @@ namespace Chaos
 					const T UseThickness = FMath::Max(Box1Thickness, Box2Thickness.Size());// + MThickness
 
 
-					TRigidBodyContactConstraint<T, d> Constraint;
+					TRigidBodySingleContactConstraint<T, d> Constraint;
 					ConstructConstraints(Particle1.Handle(), Particle2.Handle(), UseThickness, Constraint);
 
 #if !UE_BUILD_SHIPPING
@@ -293,7 +293,7 @@ namespace Chaos
 				while (!Queue.IsEmpty())
 				{
 					int32 Idx = Constraints.AddUninitialized(1);
-					FConstraintContainerHandle* Handle = HandleAllocator.template AllocHandle< TRigidBodyContactConstraint<T,d> >(this, Idx );
+					FConstraintContainerHandle* Handle = HandleAllocator.template AllocHandle< TRigidBodySingleContactConstraint<T,d> >(this, Idx );
 
 					Handles.Add(Handle);
 					Queue.Dequeue(Constraints[Idx]);
diff --git a/Engine/Source/Runtime/Experimental/Chaos/Public/Chaos/PBDCollisionTypes.h b/Engine/Source/Runtime/Experimental/Chaos/Public/Chaos/PBDCollisionTypes.h
--- a/Engine/Source/Runtime/Experimental/Chaos/Public/Chaos/PBDCollisionTypes.h
+++ b/Engine/Source/Runtime/Experimental/Chaos/Public/Chaos/PBDCollisionTypes.h
@@ -18,7 +18,7 @@ namespace Chaos
 	class TPBDCollisionConstraintHandle;
 
 	template<typename T, int d>
-	class TRigidBodyContactConstraint;
+	class TRigidBodySingleContactConstraint;
 
 	/** Specifies the type of work we should do*/
 	enum class CHAOS_API ECollisionUpdateType
@@ -68,7 +68,7 @@ namespace Chaos
 		enum class FType
 		{
 			None = 0,     // default value also indicates a invalid constraint
-			SinglePoint   // TRigidBodyContactConstraint
+			SinglePoint   // TRigidBodySingleContactConstraint
 		};
 
 		TCollisionConstraintBase() : Type(FType::None) {}
@@ -86,14 +86,14 @@ namespace Chaos
 	*
 	*/
 	template<class T, int d>
-	class CHAOS_API TRigidBodyContactConstraint : public TCollisionConstraintBase<T,d>
+	class CHAOS_API TRigidBodySingleContactConstraint : public TCollisionConstraintBase<T,d>
 	{
 	public:
 		using Base = TCollisionConstraintBase<T, d>;
 		using FGeometryParticleHandle = TGeometryParticleHandle<T, d>;
 		using FConvexManifold = TConvexManifold<T, d>;
 
-		TRigidBodyContactConstraint() : AccumulatedImpulse(0) {}
+		TRigidBodySingleContactConstraint() : AccumulatedImpulse(0) {}
 
 		//API
 		void ResetPhi(T InPhi) { SetPhi(InPhi); }
@@ -179,7 +179,7 @@ namespace Chaos
 		{
 			if (ConstraintType == FConstraintBase::FType::SinglePoint)
 			{
-				return GetContact<TRigidBodyContactConstraint<float, 3> >().GetLocation();
+				return GetContact<TRigidBodySingleContactConstraint<float, 3> >().GetLocation();
 			}
 			ensure(false);
 			return TVector<T, d>(0);
@@ -189,7 +189,7 @@ namespace Chaos
 		{
 			if (ConstraintType == FConstraintBase::FType::SinglePoint)
 			{
-				return GetContact<TRigidBodyContactConstraint<float, 3> >().AccumulatedImpulse;
+				return GetContact<TRigidBodySingleContactConstraint<float, 3> >().AccumulatedImpulse;
 			}
 			ensure(false);
 
@@ -201,7 +201,7 @@ namespace Chaos
 		{
 			if (ConstraintType == FConstraintBase::FType::SinglePoint)
 			{
-				const TRigidBodyContactConstraint<float, 3>& Contact = GetContact< TRigidBodyContactConstraint<float, 3> >();
+				const TRigidBodySingleContactConstraint<float, 3>& Contact = GetContact< TRigidBodySingleContactConstraint<float, 3> >();
 				return { Contact.Particle[0], Contact.Particle[1] };
 			}
 			ensure(false);
@@ -214,7 +214,7 @@ namespace Chaos
 		{
 			if (ConstraintType == FConstraintBase::FType::SinglePoint)
 			{
-				return { GetContact< TRigidBodyContactConstraint<float, 3> >().Particle[0], GetContact< TRigidBodyContactConstraint<float, 3> >().Particle[1] };
+				return { GetContact< TRigidBodySingleContactConstraint<float, 3> >().Particle[0], GetContact< TRigidBodySingleContactConstraint<float, 3> >().Particle[1] };
 			}
 			ensure(false);
 			return { nullptr,nullptr };
@@ -229,7 +229,7 @@ namespace Chaos
 
 	};
 
-	extern template TRigidBodyContactConstraint<float, 3>& TPBDCollisionConstraintHandle<float, 3>::GetContact<TRigidBodyContactConstraint<float, 3>>();
-	extern template const TRigidBodyContactConstraint<float, 3>& TPBDCollisionConstraintHandle<float, 3>::GetContact<TRigidBodyContactConstraint<float, 3>>() const;
+	extern template TRigidBodySingleContactConstraint<float, 3>& TPBDCollisionConstraintHandle<float, 3>::GetContact<TRigidBodySingleContactConstraint<float, 3>>();
+	extern template const TRigidBodySingleContactConstraint<float, 3>& TPBDCollisionConstraintHandle<float, 3>::GetContact<TRigidBodySingleContactConstraint<float, 3>>() const;
 
 }
diff --git a/Engine/Source/Runtime/Experimental/ChaosSolvers/Private/EventDefaults.cpp b/Engine/Source/Runtime/Experimental/ChaosSolvers/Private/EventDefaults.cpp
--- a/Engine/Source/Runtime/Experimental/ChaosSolvers/Private/EventDefaults.cpp
+++ b/Engine/Source/Runtime/Experimental/ChaosSolvers/Private/EventDefaults.cpp
@@ -68,7 +68,7 @@ namespace Chaos
 				{
 					if (ContactHandle->GetType() == TPBDCollisionConstraintHandle<float, 3>::FConstraintBase::FType::SinglePoint)
 					{
-						const TRigidBodyContactConstraint<float, 3>& Constraint = ContactHandle->GetContact<TRigidBodyContactConstraint<float, 3>>();
+						const TRigidBodySingleContactConstraint<float, 3>& Constraint = ContactHandle->GetContact<TRigidBodySingleContactConstraint<float, 3>>();
 
 						// Since Clustered GCs can be unioned the particleIndex representing the union 
 						// is not associated with a PhysicsProxy
@@ -107,7 +107,7 @@ namespace Chaos
 				{
 					for(int32 IdxCollision = 0; IdxCollision < ValidCollisionHandles.Num(); ++IdxCollision)
 					{
-						Chaos::TPBDCollisionConstraint<float, 3>::FRigidBodyContactConstraint const& Constraint = ValidCollisionHandles[IdxCollision]->GetContact<TRigidBodyContactConstraint<float, 3>>();
+						Chaos::TPBDCollisionConstraint<float, 3>::FRigidBodyContactConstraint const& Constraint = ValidCollisionHandles[IdxCollision]->GetContact<TRigidBodySingleContactConstraint<float, 3>>();
 
 						TGeometryParticleHandle<float, 3>* Particle0 = Constraint.Particle[0];
 						TGeometryParticleHandle<float, 3>* Particle1 = Constraint.Particle[1];
