Chaos : Collisions

- More renaming

diff --git a/Engine/Source/Runtime/Experimental/Chaos/Private/Chaos/CollisionResolutionAlgo.cpp b/Engine/Source/Runtime/Experimental/Chaos/Private/Chaos/CollisionResolutionAlgo.cpp
--- a/Engine/Source/Runtime/Experimental/Chaos/Private/Chaos/CollisionResolutionAlgo.cpp
+++ b/Engine/Source/Runtime/Experimental/Chaos/Private/Chaos/CollisionResolutionAlgo.cpp
@@ -224,7 +224,7 @@ namespace Chaos
 		TVector<T, d> LocalNormal;
 		T LocalPhi = Object.PhiWithNormal(LocalPoint, LocalNormal);
 
-		TConvexManifold<T, d> & Contact = Constraint.ShapeManifold;
+		TSingleContactManifold<T, d> & Contact = Constraint.Manifold;
 		if (LocalPhi < Contact.Phi)
 		{
 			Contact.Phi = LocalPhi;
@@ -242,7 +242,7 @@ namespace Chaos
 		TVector<T, d> LocalNormal;
 		T LocalPhi = Object.PhiWithNormal(LocalPoint, LocalNormal);
 
-		TConvexManifold<T, d> & Contact = Constraint.ShapeManifold;
+		TSingleContactManifold<T, d> & Contact = Constraint.Manifold;
 		if (LocalPhi < Contact.Phi)
 		{
 			Contact.Phi = LocalPhi;
@@ -259,7 +259,7 @@ namespace Chaos
 		T LocalPhi = Object.PhiWithNormal(LocalPoint, LocalNormal);
 		T LocalThickness = LocalPhi - Thickness;
 
-		TConvexManifold<T, d> & Contact = Constraint.ShapeManifold;
+		TSingleContactManifold<T, d> & Contact = Constraint.Manifold;
 		if (LocalThickness < -KINDA_SMALL_NUMBER)
 		{
 			Contact.Location += LocalPoint * LocalThickness;
@@ -289,8 +289,8 @@ namespace Chaos
 		SCOPE_CYCLE_COUNTER(STAT_SampleObject);
 		TRigidBodySingleContactConstraint<float, 3> AvgConstraint = Constraint;
 
-		TConvexManifold<float, 3> & Contact = Constraint.ShapeManifold;
-		TConvexManifold<float, 3> & AvgContact = AvgConstraint.ShapeManifold;
+		TSingleContactManifold<float, 3> & Contact = Constraint.Manifold;
+		TSingleContactManifold<float, 3> & AvgContact = AvgConstraint.Manifold;
 
 		AvgConstraint.Particle[0] = Constraint.Particle[0];
 		AvgConstraint.Particle[1] = Constraint.Particle[1];
@@ -399,8 +399,8 @@ namespace Chaos
 		SCOPE_CYCLE_COUNTER(STAT_SampleObject);
 		TRigidBodySingleContactConstraint<float, 3> AvgConstraint = Constraint;
 
-		TConvexManifold<float, 3> & Contact = Constraint.ShapeManifold;
-		TConvexManifold<float, 3> & AvgContact = AvgConstraint.ShapeManifold;
+		TSingleContactManifold<float, 3> & Contact = Constraint.Manifold;
+		TSingleContactManifold<float, 3> & AvgContact = AvgConstraint.Manifold;
 
 		AvgConstraint.Particle[0] = Constraint.Particle[0];
 		AvgConstraint.Particle[1] = Constraint.Particle[1];
@@ -823,7 +823,7 @@ namespace Chaos
 	template <typename T, int d>
 	bool UpdateBoxPlaneConstraint(const TBox<T, d>& Box, const TRigidTransform<T, d>& BoxTransform, const TPlane<T, d>& Plane, const TRigidTransform<T, d>& PlaneTransform, T Thickness, TRigidBodySingleContactConstraint<T, d>& Constraint)
 	{
-		TConvexManifold<T, d> & Contact = Constraint.ShapeManifold;
+		TSingleContactManifold<T, d> & Contact = Constraint.Manifold;
 
 #if USING_CODE_ANALYSIS
 		MSVC_PRAGMA(warning(push))
@@ -885,7 +885,7 @@ namespace Chaos
 	template <typename T, int d>
 	void UpdateSphereConstraint(const TSphere<T, d>& Sphere1, const TRigidTransform<T, d>& Sphere1Transform, const TSphere<T, d>& Sphere2, const TRigidTransform<T, d>& Sphere2Transform, T Thickness, TRigidBodySingleContactConstraint<T, d>& Constraint)
 	{
-		TConvexManifold<T, d> & Contact = Constraint.ShapeManifold;
+		TSingleContactManifold<T, d> & Contact = Constraint.Manifold;
 
 		const TVector<T, d> Center1 = Sphere1Transform.TransformPosition(Sphere1.GetCenter());
 		const TVector<T, d> Center2 = Sphere2Transform.TransformPosition(Sphere2.GetCenter());
@@ -903,7 +903,7 @@ namespace Chaos
 	template <typename T, int d>
 	void UpdateSpherePlaneConstraint(const TSphere<T, d>& Sphere, const TRigidTransform<T, d>& SphereTransform, const TPlane<T, d>& Plane, const TRigidTransform<T, d>& PlaneTransform, T Thickness, TRigidBodySingleContactConstraint<T, d>& Constraint)
 	{
-		TConvexManifold<T, d> & Contact = Constraint.ShapeManifold;
+		TSingleContactManifold<T, d> & Contact = Constraint.Manifold;
 
 		const TRigidTransform<T, d> SphereToPlaneTransform(PlaneTransform.Inverse() * SphereTransform);
 		const TVector<T, d> SphereCenter = SphereToPlaneTransform.TransformPosition(Sphere.GetCenter());
@@ -923,7 +923,7 @@ namespace Chaos
 	template <typename T, int d>
 	void UpdateSphereBoxConstraint(const TSphere<T, d>& Sphere, const TRigidTransform<T, d>& SphereTransform, const TBox<T, d>& Box, const TRigidTransform<T, d>& BoxTransform, T Thickness, TRigidBodySingleContactConstraint<T, d>& Constraint)
 	{
-		TConvexManifold<T, d> & Contact = Constraint.ShapeManifold;
+		TSingleContactManifold<T, d> & Contact = Constraint.Manifold;
 
 		const TRigidTransform<T, d> SphereToBoxTransform(SphereTransform * BoxTransform.Inverse());
 		const TVector<T, d> SphereCenterInBox = SphereToBoxTransform.TransformPosition(Sphere.GetCenter());
@@ -943,7 +943,7 @@ namespace Chaos
 	template <typename T, int d>
 	void UpdateCapsuleCapsuleConstraint(const TCapsule<T>& A, const TRigidTransform<T, d>& ATransform, const TCapsule<T>& B, const TRigidTransform<T, d>& BTransform, T Thickness, TRigidBodySingleContactConstraint<T, d>& Constraint)
 	{
-		TConvexManifold<T, d> & Contact = Constraint.ShapeManifold;
+		TSingleContactManifold<T, d> & Contact = Constraint.Manifold;
 
 		FVector A1 = ATransform.TransformPosition(A.GetX1());
 		FVector A2 = ATransform.TransformPosition(A.GetX2());
@@ -971,7 +971,7 @@ namespace Chaos
 	void UpdateCapsuleBoxConstraint(const TCapsule<T>& A, const TRigidTransform<T, d>& ATransform, const TBox<T, d>& B, const TRigidTransform<T, d>& BTransform, T Thickness, TRigidBodySingleContactConstraint<T, d>& Constraint)
 	{
 		// @todo(ccaulfield): Add custom capsule-box collision
-		TConvexManifold<T, d> & Contact = Constraint.ShapeManifold;
+		TSingleContactManifold<T, d> & Contact = Constraint.Manifold;
 
 		TRigidTransform<T, d> BToATransform = BTransform.GetRelativeTransform(ATransform);
 
@@ -1223,7 +1223,7 @@ namespace Chaos
 	template <typename T, int d>
 	void UpdateBoxConstraint(const TBox<T, d>& Box1, const TRigidTransform<T, d>& Box1Transform, const TBox<T, d>& Box2, const TRigidTransform<T, d>& Box2Transform, T Thickness, TRigidBodySingleContactConstraint<T, d>& Constraint)
 	{
-		TConvexManifold<T, d> & Contact = Constraint.ShapeManifold;
+		TSingleContactManifold<T, d> & Contact = Constraint.Manifold;
 
 		TBox<T, d> Box2SpaceBox1 = Box1.TransformedBox(Box1Transform * Box2Transform.Inverse());
 		TBox<T, d> Box1SpaceBox2 = Box2.TransformedBox(Box2Transform * Box1Transform.Inverse());
diff --git a/Engine/Source/Runtime/Experimental/Chaos/Private/Chaos/PBDCollisionConstraint.cpp b/Engine/Source/Runtime/Experimental/Chaos/Private/Chaos/PBDCollisionConstraint.cpp
--- a/Engine/Source/Runtime/Experimental/Chaos/Private/Chaos/PBDCollisionConstraint.cpp
+++ b/Engine/Source/Runtime/Experimental/Chaos/Private/Chaos/PBDCollisionConstraint.cpp
@@ -367,7 +367,7 @@ namespace Chaos
 			TSerializablePtr<FChaosPhysicsMaterial> PhysicsMaterial0 = Particle0->AuxilaryValue(MPhysicsMaterials);
 			TSerializablePtr<FChaosPhysicsMaterial> PhysicsMaterial1 = Particle1->AuxilaryValue(MPhysicsMaterials);
 
-			TConvexManifold<T, d> & Contact = Constraint.ShapeManifold;
+			typename FRigidBodyContactConstraint::FManifold & Contact = Constraint.Manifold;
 
 			TVector<T, d> VectorToPoint1 = Contact.Location - P0;
 			TVector<T, d> VectorToPoint2 = Contact.Location - P1;
@@ -600,7 +600,7 @@ namespace Chaos
 		{
 			UpdateConstraint<ECollisionUpdateType::Deepest>(MThickness, Constraint);
 
-			const TConvexManifold<T, d> & Contact = Constraint.ShapeManifold;
+			const typename FRigidBodyContactConstraint::FManifold & Contact = Constraint.Manifold;
 
 			if (Contact.Phi >= MThickness)
 			{
diff --git a/Engine/Source/Runtime/Experimental/Chaos/Public/Chaos/PBDCollisionTypes.h b/Engine/Source/Runtime/Experimental/Chaos/Public/Chaos/PBDCollisionTypes.h
--- a/Engine/Source/Runtime/Experimental/Chaos/Public/Chaos/PBDCollisionTypes.h
+++ b/Engine/Source/Runtime/Experimental/Chaos/Public/Chaos/PBDCollisionTypes.h
@@ -40,10 +40,10 @@ namespace Chaos
 	*
 	*/
 	template<class T, int d>
-	class CHAOS_API TConvexManifold
+	class CHAOS_API TSingleContactManifold
 	{
 	public:
-		TConvexManifold(int32 InTimestamp = -INT_MAX, const FImplicitObject* InImplicit0 = nullptr, const FImplicitObject* InImplicit1 = nullptr)
+		TSingleContactManifold(int32 InTimestamp = -INT_MAX, const FImplicitObject* InImplicit0 = nullptr, const FImplicitObject* InImplicit1 = nullptr)
 			: bDisabled(true), Timestamp(InTimestamp), Normal(0), Phi(FLT_MAX)
 		{
 			Implicit[0] = InImplicit0;
@@ -91,30 +91,30 @@ namespace Chaos
 	public:
 		using Base = TCollisionConstraintBase<T, d>;
 		using FGeometryParticleHandle = TGeometryParticleHandle<T, d>;
-		using FConvexManifold = TConvexManifold<T, d>;
+		using FManifold = TSingleContactManifold<T, d>;
 
 		TRigidBodySingleContactConstraint() : AccumulatedImpulse(0) {}
 
 		//API
 		void ResetPhi(T InPhi) { SetPhi(InPhi); }
-		bool ContainsManifold(const FImplicitObject* B = nullptr, const FImplicitObject* A = nullptr) const { return (A==nullptr&&B==nullptr)?(ShapeManifold.Implicit[0] != nullptr && ShapeManifold.Implicit[1] != nullptr):(ShapeManifold.Implicit[0]==A && ShapeManifold.Implicit[1]==B); }
+		bool ContainsManifold(const FImplicitObject* B = nullptr, const FImplicitObject* A = nullptr) const { return (A==nullptr&&B==nullptr)?(Manifold.Implicit[0] != nullptr && Manifold.Implicit[1] != nullptr):(Manifold.Implicit[0]==A && Manifold.Implicit[1]==B); }
 
-		void SetDisabled(bool bInDisabled, const FImplicitObject* B = nullptr, const FImplicitObject* A = nullptr) { ShapeManifold.bDisabled = bInDisabled; }
-		TVector<T, d> GetDisabled(const FImplicitObject* B = nullptr, const FImplicitObject* A = nullptr) const { return ShapeManifold.bDisabled; }
+		void SetDisabled(bool bInDisabled, const FImplicitObject* B = nullptr, const FImplicitObject* A = nullptr) { Manifold.bDisabled = bInDisabled; }
+		TVector<T, d> GetDisabled(const FImplicitObject* B = nullptr, const FImplicitObject* A = nullptr) const { return Manifold.bDisabled; }
 
 
-		void SetNormal(const TVector<T, d> & InNormal, const FImplicitObject* B = nullptr, const FImplicitObject* A = nullptr) { ShapeManifold.Normal = InNormal; }
-		TVector<T, d> GetNormal(const FImplicitObject* B = nullptr, const FImplicitObject* A = nullptr) const { return ShapeManifold.Normal; }
+		void SetNormal(const TVector<T, d> & InNormal, const FImplicitObject* B = nullptr, const FImplicitObject* A = nullptr) { Manifold.Normal = InNormal; }
+		TVector<T, d> GetNormal(const FImplicitObject* B = nullptr, const FImplicitObject* A = nullptr) const { return Manifold.Normal; }
 
-		void SetLocation(const TVector<T, d> & InLocation, const FImplicitObject* B = nullptr, const FImplicitObject* A = nullptr, int32 Index=0) { ShapeManifold.Location = InLocation; }
-		TVector<T, d> GetLocation(const FImplicitObject* B = nullptr, const FImplicitObject* A = nullptr) const { return ShapeManifold.Location; }
+		void SetLocation(const TVector<T, d> & InLocation, const FImplicitObject* B = nullptr, const FImplicitObject* A = nullptr, int32 Index=0) { Manifold.Location = InLocation; }
+		TVector<T, d> GetLocation(const FImplicitObject* B = nullptr, const FImplicitObject* A = nullptr) const { return Manifold.Location; }
 
-		void SetPhi(T InPhi, const FImplicitObject* B = nullptr, const FImplicitObject* A = nullptr, int32 Index = 0) { ShapeManifold.Phi = InPhi; }
-		T GetPhi(const FImplicitObject* B = nullptr, const FImplicitObject* A = nullptr) const { return ShapeManifold.Phi; }
+		void SetPhi(T InPhi, const FImplicitObject* B = nullptr, const FImplicitObject* A = nullptr, int32 Index = 0) { Manifold.Phi = InPhi; }
+		T GetPhi(const FImplicitObject* B = nullptr, const FImplicitObject* A = nullptr) const { return Manifold.Phi; }
 
-		void AddManifold(const FImplicitObject* B = nullptr, const FImplicitObject* A = nullptr) { ShapeManifold.Implicit[0] = A; ShapeManifold.Implicit[1] = B; }
+		void AddManifold(const FImplicitObject* B = nullptr, const FImplicitObject* A = nullptr) { Manifold.Implicit[0] = A; Manifold.Implicit[1] = B; }
 
-		FConvexManifold ShapeManifold;
+		FManifold Manifold;
 		FGeometryParticleHandle* Particle[2]; // { Point, Volume } 
 		TVector<T, d> AccumulatedImpulse;
 
