Chaos - Dense Matrix tests and fixes

diff --git a/Engine/Source/Runtime/Experimental/Chaos/Public/Chaos/DenseMatrix.h b/Engine/Source/Runtime/Experimental/Chaos/Public/Chaos/DenseMatrix.h
--- a/Engine/Source/Runtime/Experimental/Chaos/Public/Chaos/DenseMatrix.h
+++ b/Engine/Source/Runtime/Experimental/Chaos/Public/Chaos/DenseMatrix.h
@@ -49,23 +49,32 @@ namespace Chaos
 			return *this;
 		}
 
+		/**
+		 * The number of rows in the matrix.
+		 */
 		int32 NumRows() const
 		{
 			return NRows;
 		}
 
+		/**
+		 * The number of columns in the matrix.
+		 */
 		int32 NumColumns() const
 		{
 			return NCols;
 		}
 
+		/**
+		 * The number of elements in the matrix.
+		 */
 		int32 NumElements() const
 		{
 			return NRows * NCols;
 		}
 
 		/**
-		 * Set the dimensions of the matrix. This does not rearrange or set any matrix elements so the matrix should not be used until its values have been set.
+		 * Set the dimensions of the matrix, but do not initialize any values.
 		 */
 		void SetDimensions(const int32 InNumRows, const int32 InNumColumns)
 		{
@@ -126,22 +135,39 @@ namespace Chaos
 		// Factory methods
 		//
 
+		/**
+		 * Create a matrix with the specified dimensions, but all elements are uninitialized.
+		 */
 		static TDenseMatrix<MaxElements> Make(const int32 InNumRows, const int32 InNumCols)
 		{
 			return TDenseMatrix<MaxElements>(InNumRows, InNumCols);
 		}
 
 
+		/**
+		 * Create a matrix with the specified elements supplied as an array in row-major order 
+		 * (i.e., the first N elements are for Row 0, the next N for Row 1, etc., where N is the number of columns).
+		 */
 		static TDenseMatrix<MaxElements> Make(const int32 InNumRows, const int32 InNumCols, const FReal* V, const int32 VLen)
 		{
 			return TDenseMatrix<MaxElements>(InNumRows, InNumCols, V, VLen);
 		}
 
-		static TDenseMatrix<MaxElements> Make(const int32 InNumRows, const int32 InNumCols, std::initializer_list<FReal> InitList)
+		/**
+		 * Create a copy of the 3x1 columns vector.
+		 */
+		static TDenseMatrix<MaxElements> Make(const FVec3& InM)
 		{
-			return TDenseMatrix<MaxElements>(InNumRows, InNumCols, InitList);
+			TDenseMatrix<MaxElements> M(3, 1);
+			M.At(0, 0) = InM[0];
+			M.At(1, 0) = InM[1];
+			M.At(2, 0) = InM[2];
+			return M;
 		}
 
+		/**
+		 * Create a copy of the 3x3 matrix.
+		 */
 		static TDenseMatrix<MaxElements> Make(const FMatrix33& InM)
 		{
 			// NOTE: UE matrices are Column-major (columns are sequential in memory), but DenseMatrix is Row-major (rows are sequential in memory)
@@ -158,6 +184,9 @@ namespace Chaos
 			return M;
 		}
 
+		/**
+		 * Create a matrix with all elemets set to zero, except the diagonal elements.
+		 */
 		static TDenseMatrix<MaxElements> MakeDiagonal(const int32 InNumRows, const int32 InNumCols, const FReal D)
 		{
 			TDenseMatrix<MaxElements> M(InNumRows, InNumCols);
@@ -171,6 +200,9 @@ namespace Chaos
 			return M;
 		}
 
+		/**
+		 * Create an identity matrix.
+		 */
 		static TDenseMatrix<MaxElements> MakeIdentity(const int32 InDim)
 		{
 			return MakeDiagonal(InDim, InDim, (FReal)1);
@@ -180,8 +212,11 @@ namespace Chaos
 		// Math operations
 		//
 
-		template<int32 T_EA, int32 T_EB>
-		static TDenseMatrix<MaxElements> Negative(const TDenseMatrix<T_EA>& A)
+		/**
+		 * Copy a matrix and set each element to its negative.
+		 */
+		template<int32 T_EA>
+		static TDenseMatrix<MaxElements> Negate(const TDenseMatrix<T_EA>& A)
 		{
 			// @todo(ccaulfield): optimize
 			TDenseMatrix<T_MAXELEMENTS> Result(A.NumRows(), A.NumColumns());
@@ -195,6 +230,9 @@ namespace Chaos
 			return Result;
 		}
 
+		/**
+		 * Return C = A + B
+		 */
 		template<int32 T_EA, int32 T_EB>
 		static TDenseMatrix<MaxElements> Add(const TDenseMatrix<T_EA>& A, const TDenseMatrix<T_EB>& B)
 		{
@@ -212,6 +250,9 @@ namespace Chaos
 			return Result;
 		}
 
+		/**
+		 * Return C = A - B
+		 */
 		template<int32 T_EA, int32 T_EB>
 		static TDenseMatrix<MaxElements> Subtract(const TDenseMatrix<T_EA>& A, const TDenseMatrix<T_EB>& B)
 		{
@@ -229,6 +270,10 @@ namespace Chaos
 			return Result;
 		}
 
+		/**
+		 * Return C = A x B, the product of A and B where each element of C is Cij = DotProduct(A.GetRow(i), B.GetColumns(j)).
+		 * /see MultiplyAtB, MultiplyABt, MultipltAtBt.
+		 */
 		template<int32 T_EA, int32 T_EB>
 		static TDenseMatrix<MaxElements> MultiplyAB(const TDenseMatrix<T_EA>& A, const TDenseMatrix<T_EB>& B)
 		{
@@ -250,6 +295,10 @@ namespace Chaos
 			return Result;
 		}
 
+		/**
+		 * Return C = Transpose(A) x B.
+		 * /see MultiplyAB, MultiplyABt, MultipltAtBt.
+		 */
 		template<int32 T_EA, int32 T_EB>
 		static TDenseMatrix<MaxElements> MultiplyAtB(const TDenseMatrix<T_EA>& A, const TDenseMatrix<T_EB>& B)
 		{
@@ -271,6 +320,10 @@ namespace Chaos
 			return Result;
 		}
 
+		/**
+		 * Return C = A x Transpose(B).
+		 * /see MultiplyAB, MultiplyAtB, MultipltAtBt.
+		 */
 		template<int32 T_EA, int32 T_EB>
 		static TDenseMatrix<MaxElements> MultiplyABt(const TDenseMatrix<T_EA>& A, const TDenseMatrix<T_EB>& B)
 		{
@@ -292,6 +345,10 @@ namespace Chaos
 			return Result;
 		}
 
+		/**
+		 * Return C = Transpose(A) x Transpose(B).
+		 * /see MultiplyAB, MultiplyAtB, MultipltABt.
+		 */
 		template<int32 T_EA, int32 T_EB>
 		static TDenseMatrix<MaxElements> MultiplyAtBt(const TDenseMatrix<T_EA>& A, const TDenseMatrix<T_EB>& B)
 		{
@@ -313,6 +370,9 @@ namespace Chaos
 			return Result;
 		}
 
+		/**
+		 * Return C = A x V, where A is an MxN matrix, and V a real number.
+		 */
 		template<int32 T_EA, int32 T_EB>
 		static TDenseMatrix<MaxElements> Multiply(const TDenseMatrix<T_EA>& A, const FReal V)
 		{
@@ -328,12 +388,18 @@ namespace Chaos
 			return Result;
 		}
 
+		/**
+		 * Return C = A x V, where A is an MxN matrix, and V a real number.
+		 */
 		template<int32 T_EA, int32 T_EB>
 		static TDenseMatrix<MaxElements> Multiply(const FReal V, const TDenseMatrix<T_EA>& A)
 		{
 			return Multiply(A, V);
 		}
 
+		/**
+		 * Return C = A / V, where A is an MxN matrix, and V a real number.
+		 */
 		template<int32 T_EA, int32 T_EB>
 		static TDenseMatrix<MaxElements> Divide(const TDenseMatrix<T_EA>& A, const FReal V)
 		{
@@ -349,8 +415,11 @@ namespace Chaos
 			return Result;
 		}
 
+		/**
+		 * Return C = At x B. If A and B are column vectors (Nx1 matrices), this is a vector dot product.
+		 */
 		template<int32 T_EA, int32 T_EB>
-		static TDenseMatrix<MaxElements> DotProduct(const TDenseMatrix<T_EA>& A, const TDenseMatrix<T_EA>& B)
+		static TDenseMatrix<MaxElements> DotProduct(const TDenseMatrix<T_EA>& A, const TDenseMatrix<T_EB>& B)
 		{
 			return MultiplyAtB(A, B);
 		}
