Fixed most opposing normal selection for heightfields, incorrect cell calculation was picking the wrong points to calculate the impact normal.

diff --git a/Engine/Source/Runtime/Experimental/Chaos/Private/Chaos/HeightField.cpp b/Engine/Source/Runtime/Experimental/Chaos/Private/Chaos/HeightField.cpp
--- a/Engine/Source/Runtime/Experimental/Chaos/Private/Chaos/HeightField.cpp
+++ b/Engine/Source/Runtime/Experimental/Chaos/Private/Chaos/HeightField.cpp
@@ -985,14 +985,15 @@ namespace Chaos
 		MinCell = FlatGrid.ClampIndex(MinCell);
 		MaxCell = FlatGrid.ClampIndex(MaxCell);
 
-		const int32 DeltaX = MaxCell[0] - MinCell[0];
-		const int32 DeltaY = MaxCell[1] - MinCell[1];
+		// We want to capture the first cell (delta == 0) as well
+		const int32 NumX = MaxCell[0] - MinCell[0] + 1;
+		const int32 NumY = MaxCell[1] - MinCell[1] + 1;
 
-		for(int32 CurrX = 0; CurrX < DeltaX; ++CurrX)
+		for(int32 CurrX = 0; CurrX < NumX; ++CurrX)
 		{
-			for(int32 CurrY = 0; CurrY < DeltaY; ++CurrY)
+			for(int32 CurrY = 0; CurrY < NumY; ++CurrY)
 			{
-				OutInterssctions.Add(TVector<int32, 2>(MinCell[0] + CurrX, MinCell[1] + CurrY));
+				OutInterssctions.Add(FlatGrid.ClampIndex(TVector<int32, 2>(MinCell[0] + CurrX, MinCell[1] + CurrY)));
 			}
 		}
 
@@ -1302,15 +1303,12 @@ namespace Chaos
 		{
 			bool bSecondFace = FaceIndex % 2 != 0;
 
-			if(bSecondFace)
-			{
-				FaceIndex -= 1;
-			}
-			FaceIndex /= 2;
+			int32 CellIndex = FaceIndex / 2;
+			int32 CellY = CellIndex / (GeomData.NumCols - 1);
 
 			TVector<T, 3> Points[4];
 
-			GeomData.GetPointsScaled(FaceIndex, Points);
+			GeomData.GetPointsScaled(CellIndex + CellY, Points);
 
 			TVector<T, 3> A;
 			TVector<T, 3> B;
