Chaos - Add null Particle check in PBDConstraintGaph

diff --git a/Engine/Source/Runtime/Experimental/Chaos/Private/Chaos/PBDConstraintGraph.cpp b/Engine/Source/Runtime/Experimental/Chaos/Private/Chaos/PBDConstraintGraph.cpp
--- a/Engine/Source/Runtime/Experimental/Chaos/Private/Chaos/PBDConstraintGraph.cpp
+++ b/Engine/Source/Runtime/Experimental/Chaos/Private/Chaos/PBDConstraintGraph.cpp
@@ -132,10 +132,13 @@ void FPBDConstraintGraph::InitializeGraph(const TParticleView<TGeometryParticles
 					Nodes[UpdatedNodes[Index]].Island = INDEX_NONE;
 					Nodes[UpdatedNodes[Index]].Edges.Empty();
 					auto* Particle = Nodes[UpdatedNodes[Index]].Particle;
-					if (TPBDRigidParticleHandle<FReal, 3>* PBDRigid = Particle->AsDynamic())
-					{
-						PBDRigid->Island() = INDEX_NONE;
-					}
+				    if (Particle)
+				    {
+					    if (TPBDRigidParticleHandle<FReal, 3>* PBDRigid = Particle->AsDynamic())
+					    {
+						    PBDRigid->Island() = INDEX_NONE;
+					    }
+				    }
 				}
 			});
 		UpdatedNodes.Empty();
