CHAOS - Ensure on null Particle in FPBDConstraintGraph::InitialzieGraph()

diff --git a/Engine/Source/Runtime/Experimental/Chaos/Private/Chaos/PBDConstraintGraph.cpp b/Engine/Source/Runtime/Experimental/Chaos/Private/Chaos/PBDConstraintGraph.cpp
--- a/Engine/Source/Runtime/Experimental/Chaos/Private/Chaos/PBDConstraintGraph.cpp
+++ b/Engine/Source/Runtime/Experimental/Chaos/Private/Chaos/PBDConstraintGraph.cpp
@@ -132,7 +132,7 @@ void FPBDConstraintGraph::InitializeGraph(const TParticleView<TGeometryParticles
 					Nodes[UpdatedNodes[Index]].Island = INDEX_NONE;
 					Nodes[UpdatedNodes[Index]].Edges.Empty();
 					auto* Particle = Nodes[UpdatedNodes[Index]].Particle;
-				    if (Particle)
+				    if (CHAOS_ENSURE(Particle))
 				    {
 					    if (TPBDRigidParticleHandle<FReal, 3>* PBDRigid = Particle->AsDynamic())
 					    {
