Chaos - Renable Shape bounds test in SQ. Add method to update PerShapeData bounds instead of duplicating code.

diff --git a/Engine/Source/Runtime/Engine/Private/PhysicsEngine/Experimental/ChaosInterfaceUtils.cpp b/Engine/Source/Runtime/Engine/Private/PhysicsEngine/Experimental/ChaosInterfaceUtils.cpp
--- a/Engine/Source/Runtime/Engine/Private/PhysicsEngine/Experimental/ChaosInterfaceUtils.cpp
+++ b/Engine/Source/Runtime/Engine/Private/PhysicsEngine/Experimental/ChaosInterfaceUtils.cpp
@@ -146,11 +146,7 @@ namespace ChaosInterface
 			NewShape->Geometry = InGeom;
 			NewShape->QueryData = bComplexShape ? InParams.CollisionData.CollisionFilterData.QueryComplexFilter : InParams.CollisionData.CollisionFilterData.QuerySimpleFilter;
 			NewShape->SimData = InParams.CollisionData.CollisionFilterData.SimFilter;
-
-			if (InGeom->HasBoundingBox())
-			{
-				NewShape->WorldSpaceInflatedShapeBounds = InGeom->BoundingBox().GetAABB().TransformedAABB(InParams.WorldTransform);
-			}
+			NewShape->UpdateShapeBounds(InParams.WorldTransform);
 
 			return NewShape;
 		};
diff --git a/Engine/Source/Runtime/Experimental/Chaos/Private/Chaos/GeometryParticles.cpp b/Engine/Source/Runtime/Experimental/Chaos/Private/Chaos/GeometryParticles.cpp
--- a/Engine/Source/Runtime/Experimental/Chaos/Private/Chaos/GeometryParticles.cpp
+++ b/Engine/Source/Runtime/Experimental/Chaos/Private/Chaos/GeometryParticles.cpp
@@ -34,7 +34,7 @@ namespace Chaos
 			{
 				for (auto& Shape : ShapesArray)
 				{
-					Shape->WorldSpaceInflatedShapeBounds = Geometry->BoundingBox().GetAABB().TransformedAABB(ActorTM);
+					Shape->UpdateShapeBounds(ActorTM);
 				}
 			}
 		}
@@ -61,6 +61,15 @@ namespace Chaos
 		return TUniquePtr<TPerShapeData<T, d>>(new TPerShapeData<T, d>());
 	}
 
+	template<typename T, int d>
+	void TPerShapeData<T, d>::UpdateShapeBounds(const TRigidTransform<T, d>& WorldTM)
+	{
+		if (Geometry && Geometry->HasBoundingBox())
+		{
+			WorldSpaceInflatedShapeBounds = Geometry->BoundingBox().GetAABB().TransformedAABB(WorldTM);
+		}
+	}
+
 	template <typename T, int d>
 	TPerShapeData<T, d>* TPerShapeData<T, d>::SerializationFactory(FChaosArchive& Ar, TPerShapeData<T, d>*)
 	{
diff --git a/Engine/Source/Runtime/Experimental/Chaos/Public/Chaos/GeometryParticles.h b/Engine/Source/Runtime/Experimental/Chaos/Public/Chaos/GeometryParticles.h
--- a/Engine/Source/Runtime/Experimental/Chaos/Public/Chaos/GeometryParticles.h
+++ b/Engine/Source/Runtime/Experimental/Chaos/Public/Chaos/GeometryParticles.h
@@ -36,6 +36,8 @@ namespace Chaos
 		TAABB<FReal, 3> WorldSpaceInflatedShapeBounds;
 		TArray<FMaterialHandle> Materials;
 
+		void UpdateShapeBounds(const TRigidTransform<T, d>& WorldTM);
+
 		static TPerShapeData<T, d>* SerializationFactory(FChaosArchive& Ar, TPerShapeData<T, d>*);
 		void Serialize(FChaosArchive& Ar);
 
@@ -274,8 +276,7 @@ namespace Chaos
 				if (Shape->Geometry->HasBoundingBox())
 				{
 					const TRigidTransform<FReal, 3> ActorTM(X(Index), R(Index));
-					TAABB<FReal, 3> Bound = Shape->Geometry->BoundingBox().GetAABB().TransformedAABB(ActorTM);
-					Shape->WorldSpaceInflatedShapeBounds = Bound;
+					Shape->UpdateShapeBounds(ActorTM);
 				}
 			}
 		}
diff --git a/Engine/Source/Runtime/Experimental/Chaos/Public/Chaos/ParticleHandle.h b/Engine/Source/Runtime/Experimental/Chaos/Public/Chaos/ParticleHandle.h
--- a/Engine/Source/Runtime/Experimental/Chaos/Public/Chaos/ParticleHandle.h
+++ b/Engine/Source/Runtime/Experimental/Chaos/Public/Chaos/ParticleHandle.h
@@ -1008,7 +1008,7 @@ public:
 			{
 				for (auto& Shape : MShapesArray)
 				{
-					Shape->WorldSpaceInflatedShapeBounds = Shape->Geometry->BoundingBox().GetAABB().TransformedAABB(FRigidTransform3(MX, MR));
+					Shape->UpdateShapeBounds(FRigidTransform3(MX, MR));
 				}
 			}
 		}
diff --git a/Engine/Source/Runtime/Landscape/Private/LandscapeCollision.cpp b/Engine/Source/Runtime/Landscape/Private/LandscapeCollision.cpp
--- a/Engine/Source/Runtime/Landscape/Private/LandscapeCollision.cpp
+++ b/Engine/Source/Runtime/Landscape/Private/LandscapeCollision.cpp
@@ -53,6 +53,7 @@
 #include "PhysXToChaosUtil.h"
 #include "Chaos/ParticleHandle.h"
 #include "Chaos/Vector.h"
+#include "Chaos/Core.h"
 #endif
 
 using namespace PhysicsInterfaceTypes;
@@ -484,6 +485,15 @@ void ULandscapeHeightfieldCollisionComponent::OnCreatePhysicsState()
 					PhysHandle->SetGeometry(MakeUnique<Chaos::TImplicitObjectUnion<float, 3>>(MoveTemp(Geoms)));
 				}
 
+				// Construct Shape Bounds
+				for (auto& Shape : ShapeArray)
+				{
+					Chaos::FRigidTransform3 WorldTransform = Chaos::FRigidTransform3(PhysHandle->X(), PhysHandle->R());
+					Shape->UpdateShapeBounds(WorldTransform);
+				}
+
+
+
 				PhysHandle->SetShapesArray(MoveTemp(ShapeArray));
 
 				// Push the actor to the scene
diff --git a/Engine/Source/Runtime/PhysicsSQ/Private/SQAccelerator.cpp b/Engine/Source/Runtime/PhysicsSQ/Private/SQAccelerator.cpp
--- a/Engine/Source/Runtime/PhysicsSQ/Private/SQAccelerator.cpp
+++ b/Engine/Source/Runtime/PhysicsSQ/Private/SQAccelerator.cpp
@@ -180,10 +180,7 @@ private:
 		TGeometryParticle<float, 3>* GeometryParticle = Payload.GetExternalGeometryParticle_ExternalThread();
 		const TShapesArray<float,3>& Shapes = GeometryParticle->ShapesArray();
 
-		//////////////////////////////////////////////////////////////////////////
-		// Temporarily disabled shape bounds checked due to conflicts with some shape unions
-		const bool bTestShapeBounds = false; // Shapes.Num() > 1;
-		//////////////////////////////////////////////////////////////////////////
+		const bool bTestShapeBounds =  Shapes.Num() > 1;
 
 		const TRigidTransform<float, 3> ActorTM(GeometryParticle->X(), GeometryParticle->R());
 
