Potential CIS fix

diff --git a/Engine/Source/Runtime/Experimental/Chaos/Public/Chaos/PBDCollisionTypes.h b/Engine/Source/Runtime/Experimental/Chaos/Public/Chaos/PBDCollisionTypes.h
--- a/Engine/Source/Runtime/Experimental/Chaos/Public/Chaos/PBDCollisionTypes.h
+++ b/Engine/Source/Runtime/Experimental/Chaos/Public/Chaos/PBDCollisionTypes.h
@@ -162,8 +162,8 @@ namespace Chaos
 
 		// Handle API
 
-		template<class AS_T> CHAOS_API AS_T& GetContact();
-		template<class AS_T> CHAOS_API const AS_T& GetContact() const;
+		template<class AS_T> AS_T& GetContact();
+		template<class AS_T> const AS_T& GetContact() const;
 
 		typename FConstraintBase::FType GetType() const { return ConstraintType; }
 
