Chaos - world-space inertia utility method

diff --git a/Engine/Source/Runtime/Experimental/Chaos/Public/Chaos/Utilities.h b/Engine/Source/Runtime/Experimental/Chaos/Public/Chaos/Utilities.h
--- a/Engine/Source/Runtime/Experimental/Chaos/Public/Chaos/Utilities.h
+++ b/Engine/Source/Runtime/Experimental/Chaos/Public/Chaos/Utilities.h
@@ -145,6 +145,16 @@ namespace Chaos
 			return TRigidTransform<float, 3>(L.GetTranslation() + L.GetRotation().RotateVector(R.GetTranslation()), L.GetRotation() * R.GetRotation());
 		}
 
+		/**
+		 * Calculate the world-space inertia (or inverse inertia) for a body with rotation "Q" and local-space inertia/inverse-inertia "I".
+		 */
+		static FMatrix33 ComputeWorldSpaceInertia(const FRotation3& Q, const FMatrix33& I)
+		{
+			FMatrix33 QM = Q.ToMatrix();
+			//return QM * I * QM.GetTransposed();
+			return MultiplyAB(QM, MultiplyABt(I, QM));
+		}
+
 		/**
 		 * Calculate the matrix that maps a constraint position error to constraint position and rotation corrections.
 		 */
