Chaos :

File cleanup.

diff --git a/Engine/Plugins/Experimental/ChaosNiagara/Source/ChaosNiagara/Private/NiagaraDataInterfaceChaosDestruction.cpp b/Engine/Plugins/Experimental/ChaosNiagara/Source/ChaosNiagara/Private/NiagaraDataInterfaceChaosDestruction.cpp
--- a/Engine/Plugins/Experimental/ChaosNiagara/Source/ChaosNiagara/Private/NiagaraDataInterfaceChaosDestruction.cpp
+++ b/Engine/Plugins/Experimental/ChaosNiagara/Source/ChaosNiagara/Private/NiagaraDataInterfaceChaosDestruction.cpp
@@ -9,7 +9,7 @@
 #include "Niagara/Private/NiagaraStats.h"
 #include "Chaos/PBDCollisionConstraintUtil.h"
 #include "PhysicsProxy/GeometryCollectionPhysicsProxy.h"
-#include "Chaos/PBDCollisionTypes.h"
+#include "Chaos/CollisionResolutionTypes.h"
 #include "PhysicalMaterials/PhysicalMaterial.h"
 #include "NiagaraComponent.h"
 #include "NiagaraEmitterInstanceBatcher.h"
diff --git a/Engine/Source/Runtime/Experimental/Chaos/Private/Chaos/CollisionResolutionTypes.cpp b/Engine/Source/Runtime/Experimental/Chaos/Private/Chaos/CollisionResolutionTypes.cpp
new file mode 100644
--- /dev/null
+++ b/Engine/Source/Runtime/Experimental/Chaos/Private/Chaos/CollisionResolutionTypes.cpp
@@ -0,0 +1,24 @@
+// Copyright 1998-2019 Epic Games, Inc. All Rights Reserved.
+
+#include "Chaos/CollisionResolutionTypes.h"
+#include "Chaos/PBDCollisionConstraint.h"
+
+namespace Chaos
+{
+
+	template<class T, int d> template<class TYPE>
+	TYPE& TPBDCollisionConstraintHandle<T, d>::GetContact()
+	{
+		return ConstraintContainer->Constraints[ConstraintIndex];
+	}
+
+	template<class T, int d> template<class TYPE>
+	const TYPE& TPBDCollisionConstraintHandle<T, d>::GetContact() const
+	{
+		return ConstraintContainer->Constraints[ConstraintIndex];
+	}
+
+	template class CHAOS_API TPBDCollisionConstraintHandle<float, 3>;
+	template CHAOS_API TRigidBodySingleContactConstraint<float, 3>& TPBDCollisionConstraintHandle<float, 3>::GetContact<TRigidBodySingleContactConstraint<float, 3>>();
+	template CHAOS_API const TRigidBodySingleContactConstraint<float, 3>& TPBDCollisionConstraintHandle<float, 3>::GetContact<TRigidBodySingleContactConstraint<float, 3>>() const;
+};
diff --git a/Engine/Source/Runtime/Experimental/Chaos/Private/Chaos/PBDCollisionConstraintUtil.cpp b/Engine/Source/Runtime/Experimental/Chaos/Private/Chaos/PBDCollisionConstraintUtil.cpp
--- a/Engine/Source/Runtime/Experimental/Chaos/Private/Chaos/PBDCollisionConstraintUtil.cpp
+++ b/Engine/Source/Runtime/Experimental/Chaos/Private/Chaos/PBDCollisionConstraintUtil.cpp
@@ -11,7 +11,7 @@
 #include "ChaosLog.h"
 #include "ProfilingDebugging/ScopedTimers.h"
 #include "Chaos/ImplicitObjectUnion.h"
-#include "Chaos/PBDCollisionTypes.h"
+#include "Chaos/CollisionResolutionTypes.h"
 
 namespace Chaos
 {
diff --git a/Engine/Source/Runtime/Experimental/Chaos/Private/Chaos/PBDCollisionTypes.cpp b/Engine/Source/Runtime/Experimental/Chaos/Private/Chaos/PBDCollisionTypes.cpp
deleted file mode 100644
diff --git a/Engine/Source/Runtime/Experimental/Chaos/Public/Chaos/CollisionResolutionAlgo.h b/Engine/Source/Runtime/Experimental/Chaos/Public/Chaos/CollisionResolutionAlgo.h
--- a/Engine/Source/Runtime/Experimental/Chaos/Public/Chaos/CollisionResolutionAlgo.h
+++ b/Engine/Source/Runtime/Experimental/Chaos/Public/Chaos/CollisionResolutionAlgo.h
@@ -1,7 +1,7 @@
 // Copyright 1998-2019 Epic Games, Inc. All Rights Reserved.
 
 #pragma once
-#include "Chaos/PBDCollisionTypes.h"
+#include "Chaos/CollisionResolutionTypes.h"
 #include "Chaos/Pair.h"
 
 namespace Chaos
diff --git a/Engine/Source/Runtime/Experimental/Chaos/Public/Chaos/CollisionResolutionConvexConvex.h b/Engine/Source/Runtime/Experimental/Chaos/Public/Chaos/CollisionResolutionConvexConvex.h
--- a/Engine/Source/Runtime/Experimental/Chaos/Public/Chaos/CollisionResolutionConvexConvex.h
+++ b/Engine/Source/Runtime/Experimental/Chaos/Public/Chaos/CollisionResolutionConvexConvex.h
@@ -3,7 +3,7 @@
 
 #include "Chaos/Convex.h"
 #include "Chaos/ParticleHandle.h"
-#include "Chaos/PBDCollisionTypes.h"
+#include "Chaos/CollisionResolutionTypes.h"
 #include "Chaos/PBDCollisionConstraint.h"
 #include "Chaos/Transform.h"
 
diff --git a/Engine/Source/Runtime/Experimental/Chaos/Public/Chaos/CollisionResolutionTypes.h b/Engine/Source/Runtime/Experimental/Chaos/Public/Chaos/CollisionResolutionTypes.h
new file mode 100644
--- /dev/null
+++ b/Engine/Source/Runtime/Experimental/Chaos/Public/Chaos/CollisionResolutionTypes.h
@@ -0,0 +1,235 @@
+// Copyright 1998-2019 Epic Games, Inc. All Rights Reserved.
+#pragma once
+#include "Box.h"
+#include "Chaos/ExternalCollisionData.h"
+#include "Chaos/ParticleHandleFwd.h"
+#include "Chaos/PBDConstraintContainer.h"
+#include "Chaos/Vector.h"
+
+class UPhysicalMaterial;
+
+namespace Chaos
+{
+
+	template<typename T, int d>
+	class TPBDCollisionConstraint;
+
+	template<typename  T, int d>
+	class TPBDCollisionConstraintHandle;
+
+	template<typename T, int d>
+	class TRigidBodySingleContactConstraint;
+
+	/** Specifies the type of work we should do*/
+	enum class CHAOS_API ECollisionUpdateType
+	{
+		Any,	//stop if we have at least one deep penetration. Does not compute location or normal
+		Deepest	//find the deepest penetration. Compute location and normal
+	};
+
+	/** Return value of the collision modification callback */
+	enum class CHAOS_API ECollisionModifierResult
+	{
+		Unchanged,	/** No change to the collision */
+		Modified,	/** Modified the collision, but want it to remain enabled */
+		Disabled,	/** Collision should be disabled */
+	};
+
+
+	/*
+	*
+	*/
+	template<class T, int d>
+	class CHAOS_API TSingleContactManifold
+	{
+	public:
+		TSingleContactManifold(int32 InTimestamp = -INT_MAX, const FImplicitObject* InImplicit0 = nullptr, const FImplicitObject* InImplicit1 = nullptr)
+			: bDisabled(true), Timestamp(InTimestamp), Normal(0), Phi(FLT_MAX)
+		{
+			Implicit[0] = InImplicit0;
+			Implicit[1] = InImplicit1;
+		}
+
+		bool bDisabled;
+		int32 Timestamp;
+		TVector<T, d> Normal;
+		TVector<T, d> Location;
+		T Phi;
+
+		const FImplicitObject* Implicit[2]; // {Of Particle[0], Of Particle[1]}
+	};
+
+
+	template<class T, int d>
+	class CHAOS_API TCollisionConstraintBase
+	{
+	public:
+
+		enum class FType
+		{
+			None = 0,     // default value also indicates a invalid constraint
+			SinglePoint   // TRigidBodySingleContactConstraint
+		};
+
+		TCollisionConstraintBase() : Type(FType::None) {}
+		TCollisionConstraintBase(FType InType) : Type(InType) {}
+		FType GetType() const { return Type; }
+
+		template<class AS_T> AS_T * As() {return static_cast<AS_T>(this); }
+		template<class AS_T> const AS_T * As() const { return static_cast<AS_T>(this); }
+
+	private:
+		FType Type;
+	};
+
+	/*
+	*
+	*/
+	template<class T, int d>
+	class CHAOS_API TRigidBodySingleContactConstraint : public TCollisionConstraintBase<T,d>
+	{
+	public:
+		using Base = TCollisionConstraintBase<T, d>;
+		using FGeometryParticleHandle = TGeometryParticleHandle<T, d>;
+		using FManifold = TSingleContactManifold<T, d>;
+
+		TRigidBodySingleContactConstraint() : AccumulatedImpulse(0) {}
+
+		//API
+		void ResetPhi(T InPhi) { SetPhi(InPhi); }
+		bool ContainsManifold(const FImplicitObject* B = nullptr, const FImplicitObject* A = nullptr) const { return (A==nullptr&&B==nullptr)?(Manifold.Implicit[0] != nullptr && Manifold.Implicit[1] != nullptr):(Manifold.Implicit[0]==A && Manifold.Implicit[1]==B); }
+
+		void SetDisabled(bool bInDisabled, const FImplicitObject* B = nullptr, const FImplicitObject* A = nullptr) { Manifold.bDisabled = bInDisabled; }
+		TVector<T, d> GetDisabled(const FImplicitObject* B = nullptr, const FImplicitObject* A = nullptr) const { return Manifold.bDisabled; }
+
+
+		void SetNormal(const TVector<T, d> & InNormal, const FImplicitObject* B = nullptr, const FImplicitObject* A = nullptr) { Manifold.Normal = InNormal; }
+		TVector<T, d> GetNormal(const FImplicitObject* B = nullptr, const FImplicitObject* A = nullptr) const { return Manifold.Normal; }
+
+		void SetLocation(const TVector<T, d> & InLocation, const FImplicitObject* B = nullptr, const FImplicitObject* A = nullptr, int32 Index=0) { Manifold.Location = InLocation; }
+		TVector<T, d> GetLocation(const FImplicitObject* B = nullptr, const FImplicitObject* A = nullptr) const { return Manifold.Location; }
+
+		void SetPhi(T InPhi, const FImplicitObject* B = nullptr, const FImplicitObject* A = nullptr, int32 Index = 0) { Manifold.Phi = InPhi; }
+		T GetPhi(const FImplicitObject* B = nullptr, const FImplicitObject* A = nullptr) const { return Manifold.Phi; }
+
+		void AddManifold(const FImplicitObject* B = nullptr, const FImplicitObject* A = nullptr) { Manifold.Implicit[0] = A; Manifold.Implicit[1] = B; }
+
+		FManifold Manifold;
+		FGeometryParticleHandle* Particle[2]; // { Point, Volume } 
+		TVector<T, d> AccumulatedImpulse;
+
+		FString ToString() const
+		{
+			return FString::Printf(TEXT("Particle:%s, Levelset:%s, AccumulatedImpulse:%s"), *Particle[0]->ToString(), *Particle[1]->ToString(), *AccumulatedImpulse.ToString());
+		}
+
+		static typename Base::FType StaticType() { return Base::FType::SinglePoint; };
+
+	};
+
+
+	template<class T, int d>
+	struct CHAOS_API TRigidBodyContactConstraintPGS
+	{
+		TRigidBodyContactConstraintPGS() : AccumulatedImpulse(0.f) {}
+		TGeometryParticleHandle<T, d>* Particle;
+		TGeometryParticleHandle<T, d>* Levelset;
+		TArray<TVector<T, d>> Normal;
+		TArray<TVector<T, d>> Location;
+		TArray<T> Phi;
+		TVector<T, d> AccumulatedImpulse;
+	};
+
+
+	template<class T, int d>
+	class CHAOS_API TPBDCollisionConstraintHandle : public TContainerConstraintHandle<TPBDCollisionConstraint<T, d>>
+	{
+	public:
+		using Base = TContainerConstraintHandle<TPBDCollisionConstraint<T, d>>;
+		using FConstraintContainer = TPBDCollisionConstraint<T, d>;
+		using FConstraintBase = TCollisionConstraintBase<T, d>;
+
+
+		TPBDCollisionConstraintHandle() 
+			: ConstraintType(FConstraintBase::FType::None)
+		{}
+
+		TPBDCollisionConstraintHandle(FConstraintContainer* InConstraintContainer, int32 InConstraintIndex, typename FConstraintBase::FType InType)
+			: TContainerConstraintHandle<TPBDCollisionConstraint<T, d>>(InConstraintContainer, InConstraintIndex)
+			, ConstraintType(InType)
+		{
+		}
+
+		// Handle API
+
+		template<class AS_T> AS_T& GetContact();
+		template<class AS_T> const AS_T& GetContact() const;
+
+		typename FConstraintBase::FType GetType() const { return ConstraintType; }
+
+		void SetConstraintIndex(int32 IndexIn, typename FConstraintBase::FType InType)
+		{
+			ConstraintIndex = IndexIn;
+			ConstraintType = InType;
+		}
+
+		// Contact API
+
+		TVector<T, d> GetContactLocation() const
+		{
+			if (ConstraintType == FConstraintBase::FType::SinglePoint)
+			{
+				return GetContact<TRigidBodySingleContactConstraint<float, 3> >().GetLocation();
+			}
+			ensure(false);
+			return TVector<T, d>(0);
+		}
+
+		TVector<T, d> GetAccumulatedImpulse() const
+		{
+			if (ConstraintType == FConstraintBase::FType::SinglePoint)
+			{
+				return GetContact<TRigidBodySingleContactConstraint<float, 3> >().AccumulatedImpulse;
+			}
+			ensure(false);
+
+			return TVector<T, d>(0);
+		}
+
+
+		TVector<const TGeometryParticleHandle<T, d>*, 2> GetConstrainedParticles() const
+		{
+			if (ConstraintType == FConstraintBase::FType::SinglePoint)
+			{
+				const TRigidBodySingleContactConstraint<float, 3>& Contact = GetContact< TRigidBodySingleContactConstraint<float, 3> >();
+				return { Contact.Particle[0], Contact.Particle[1] };
+			}
+			ensure(false);
+			return { nullptr,nullptr };
+		}
+
+
+
+		TVector<TGeometryParticleHandle<T, d>*, 2> GetConstrainedParticles()
+		{
+			if (ConstraintType == FConstraintBase::FType::SinglePoint)
+			{
+				return { GetContact< TRigidBodySingleContactConstraint<float, 3> >().Particle[0], GetContact< TRigidBodySingleContactConstraint<float, 3> >().Particle[1] };
+			}
+			ensure(false);
+			return { nullptr,nullptr };
+		}
+
+
+	protected:
+		typename FConstraintBase::FType ConstraintType;
+		using Base::ConstraintIndex;
+		using Base::ConstraintContainer;
+
+
+	};
+
+	extern template TRigidBodySingleContactConstraint<float, 3>& TPBDCollisionConstraintHandle<float, 3>::GetContact<TRigidBodySingleContactConstraint<float, 3>>();
+	extern template const TRigidBodySingleContactConstraint<float, 3>& TPBDCollisionConstraintHandle<float, 3>::GetContact<TRigidBodySingleContactConstraint<float, 3>>() const;
+
+}
diff --git a/Engine/Source/Runtime/Experimental/Chaos/Public/Chaos/PBDCollisionConstraint.h b/Engine/Source/Runtime/Experimental/Chaos/Public/Chaos/PBDCollisionConstraint.h
--- a/Engine/Source/Runtime/Experimental/Chaos/Public/Chaos/PBDCollisionConstraint.h
+++ b/Engine/Source/Runtime/Experimental/Chaos/Public/Chaos/PBDCollisionConstraint.h
@@ -2,7 +2,7 @@
 #pragma once
 
 #include "Chaos/ConstraintHandle.h"
-#include "Chaos/PBDCollisionTypes.h"
+#include "Chaos/CollisionResolutionTypes.h"
 #include "Chaos/PBDConstraintContainer.h"
 #include "Framework/BufferedData.h"
 
diff --git a/Engine/Source/Runtime/Experimental/Chaos/Public/Chaos/PBDCollisionConstraintPGS.h b/Engine/Source/Runtime/Experimental/Chaos/Public/Chaos/PBDCollisionConstraintPGS.h
--- a/Engine/Source/Runtime/Experimental/Chaos/Public/Chaos/PBDCollisionConstraintPGS.h
+++ b/Engine/Source/Runtime/Experimental/Chaos/Public/Chaos/PBDCollisionConstraintPGS.h
@@ -1,7 +1,7 @@
 // Copyright 1998-2019 Epic Games, Inc. All Rights Reserved.
 #pragma once
 
-#include "Chaos/PBDCollisionTypes.h"
+#include "Chaos/CollisionResolutionTypes.h"
 #include "Chaos/PBDConstraintGraph.h"
 
 #include <memory>
diff --git a/Engine/Source/Runtime/Experimental/Chaos/Public/Chaos/PBDCollisionConstraintSingleContactUtil.h b/Engine/Source/Runtime/Experimental/Chaos/Public/Chaos/PBDCollisionConstraintSingleContactUtil.h
--- a/Engine/Source/Runtime/Experimental/Chaos/Public/Chaos/PBDCollisionConstraintSingleContactUtil.h
+++ b/Engine/Source/Runtime/Experimental/Chaos/Public/Chaos/PBDCollisionConstraintSingleContactUtil.h
@@ -3,7 +3,7 @@
 
 #include "Chaos/ExternalCollisionData.h"
 #include "Chaos/PBDCollisionConstraint.h"
-#include "Chaos/PBDCollisionTypes.h"
+#include "Chaos/CollisionResolutionTypes.h"
 
 namespace Chaos
 {
diff --git a/Engine/Source/Runtime/Experimental/Chaos/Public/Chaos/PBDCollisionConstraintUtil.h b/Engine/Source/Runtime/Experimental/Chaos/Public/Chaos/PBDCollisionConstraintUtil.h
--- a/Engine/Source/Runtime/Experimental/Chaos/Public/Chaos/PBDCollisionConstraintUtil.h
+++ b/Engine/Source/Runtime/Experimental/Chaos/Public/Chaos/PBDCollisionConstraintUtil.h
@@ -3,7 +3,7 @@
 
 #include "Chaos/ExternalCollisionData.h"
 #include "Chaos/PBDCollisionConstraint.h"
-#include "Chaos/PBDCollisionTypes.h"
+#include "Chaos/CollisionResolutionTypes.h"
 
 namespace Chaos
 {
diff --git a/Engine/Source/Runtime/Experimental/Chaos/Public/Chaos/PBDCollisionTypes.h b/Engine/Source/Runtime/Experimental/Chaos/Public/Chaos/PBDCollisionTypes.h
deleted file mode 100644
diff --git a/Engine/Source/Runtime/Experimental/Chaos/Public/Chaos/PBDConstraintRule.h b/Engine/Source/Runtime/Experimental/Chaos/Public/Chaos/PBDConstraintRule.h
--- a/Engine/Source/Runtime/Experimental/Chaos/Public/Chaos/PBDConstraintRule.h
+++ b/Engine/Source/Runtime/Experimental/Chaos/Public/Chaos/PBDConstraintRule.h
@@ -3,7 +3,7 @@
 
 #include "Chaos/Declares.h"
 #include "Chaos/ConstraintHandle.h"
-#include "Chaos/PBDCollisionTypes.h"
+#include "Chaos/CollisionResolutionTypes.h"
 #include "Chaos/PBDConstraintColor.h"
 #include "Chaos/PBDConstraintGraph.h"
 
diff --git a/Engine/Source/Runtime/Experimental/Chaos/Public/Chaos/PBDRigidClustering.h b/Engine/Source/Runtime/Experimental/Chaos/Public/Chaos/PBDRigidClustering.h
--- a/Engine/Source/Runtime/Experimental/Chaos/Public/Chaos/PBDRigidClustering.h
+++ b/Engine/Source/Runtime/Experimental/Chaos/Public/Chaos/PBDRigidClustering.h
@@ -3,7 +3,7 @@
 
 #include "Chaos/PBDRigidClusteredParticles.h"
 #include "Chaos/Transform.h"
-#include "Chaos/PBDCollisionTypes.h"
+#include "Chaos/CollisionResolutionTypes.h"
 #include "Chaos/ExternalCollisionData.h"
 #include "Chaos/TriangleMesh.h"
 #include "Chaos/ClusterCreationParameters.h"
diff --git a/Engine/Source/Runtime/Experimental/ChaosSolvers/Private/EventDefaults.cpp b/Engine/Source/Runtime/Experimental/ChaosSolvers/Private/EventDefaults.cpp
--- a/Engine/Source/Runtime/Experimental/ChaosSolvers/Private/EventDefaults.cpp
+++ b/Engine/Source/Runtime/Experimental/ChaosSolvers/Private/EventDefaults.cpp
@@ -6,7 +6,7 @@
 #include "Chaos/PBDRigidsEvolution.h"
 #include "Chaos/PBDRigidClustering.h"
 #include "Chaos/PBDCollisionConstraint.h"
-#include "Chaos/PBDCollisionTypes.h"
+#include "Chaos/CollisionResolutionTypes.h"
 #include "PBDRigidsSolver.h"
 #include "PhysicsProxy/SkeletalMeshPhysicsProxy.h"
 #include "PhysicsProxy/StaticMeshPhysicsProxy.h"
diff --git a/Engine/Source/Runtime/Experimental/ChaosSolvers/Private/PhysicsProxy/SolverEventFilters.cpp b/Engine/Source/Runtime/Experimental/ChaosSolvers/Private/PhysicsProxy/SolverEventFilters.cpp
--- a/Engine/Source/Runtime/Experimental/ChaosSolvers/Private/PhysicsProxy/SolverEventFilters.cpp
+++ b/Engine/Source/Runtime/Experimental/ChaosSolvers/Private/PhysicsProxy/SolverEventFilters.cpp
@@ -1,7 +1,7 @@
 // Copyright 1998-2019 Epic Games, Inc. All Rights Reserved.
 
 #include "SolverEventFilters.h"
-#include "Chaos/PBDCollisionTypes.h"
+#include "Chaos/CollisionResolutionTypes.h"
 
 namespace Chaos
 {
diff --git a/Engine/Source/Runtime/Experimental/ChaosSolvers/Public/EventsData.h b/Engine/Source/Runtime/Experimental/ChaosSolvers/Public/EventsData.h
--- a/Engine/Source/Runtime/Experimental/ChaosSolvers/Public/EventsData.h
+++ b/Engine/Source/Runtime/Experimental/ChaosSolvers/Public/EventsData.h
@@ -3,7 +3,7 @@
 #pragma once
 
 #include "Chaos/ExternalCollisionData.h"
-#include "Chaos/PBDCollisionTypes.h"
+#include "Chaos/CollisionResolutionTypes.h"
 #include "Chaos/Framework/PhysicsProxy.h"
 
 namespace Chaos
