Chaos :

 - API cleanup

diff --git a/Engine/Source/Runtime/Experimental/Chaos/Private/Chaos/PBDCollisionConstraint.cpp b/Engine/Source/Runtime/Experimental/Chaos/Private/Chaos/PBDCollisionConstraint.cpp
--- a/Engine/Source/Runtime/Experimental/Chaos/Private/Chaos/PBDCollisionConstraint.cpp
+++ b/Engine/Source/Runtime/Experimental/Chaos/Private/Chaos/PBDCollisionConstraint.cpp
@@ -298,7 +298,9 @@ namespace Chaos
 			}
 		}
 
-		Collisions::Apply(Constraint,MThickness, Dt, It, NumIts, MApplyPairIterations, &MCollided, &MPhysicsMaterials, CollisionFrictionOverride, MAngularFriction);
+		Collisions::TSingleContParticleParameters<T> ParticleParameters = { &MCollided, &MPhysicsMaterials, CollisionFrictionOverride, MAngularFriction };
+		Collisions::TSingleContactIterationParameters<T> IterationParameters = { Dt, It, NumIts, MApplyPairIterations, nullptr };
+		Collisions::Apply(Constraint,MThickness, IterationParameters,  ParticleParameters);
 	}
 
 	DECLARE_CYCLE_STAT(TEXT("TPBDCollisionConstraint::Apply"), STAT_Apply, STATGROUP_Chaos);
@@ -347,7 +349,9 @@ namespace Chaos
 			}
 		}
 
-		Collisions::ApplyPushOut(Constraint, MThickness, IsTemporarilyStatic, Dt, Iteration, NumIterations, MApplyPushOutPairIterations, NeedsAnotherIteration, &MPhysicsMaterials);
+		Collisions::TSingleContParticleParameters<T> ParticleParameters = { &MCollided, &MPhysicsMaterials, CollisionFrictionOverride, MAngularFriction };
+		Collisions::TSingleContactIterationParameters<T> IterationParameters = { Dt, Iteration, NumIterations, MApplyPairIterations, &NeedsAnotherIteration };
+		Collisions::ApplyPushOut(Constraint, MThickness, IsTemporarilyStatic, IterationParameters, ParticleParameters);
 	}
 
 	DECLARE_CYCLE_STAT(TEXT("TPBDCollisionConstraint::ApplyPushOut"), STAT_ApplyPushOut, STATGROUP_Chaos);
diff --git a/Engine/Source/Runtime/Experimental/Chaos/Private/Chaos/PBDCollisionConstraintSingleContactUtil.cpp b/Engine/Source/Runtime/Experimental/Chaos/Private/Chaos/PBDCollisionConstraintSingleContactUtil.cpp
--- a/Engine/Source/Runtime/Experimental/Chaos/Private/Chaos/PBDCollisionConstraintSingleContactUtil.cpp
+++ b/Engine/Source/Runtime/Experimental/Chaos/Private/Chaos/PBDCollisionConstraintSingleContactUtil.cpp
@@ -127,19 +127,14 @@ namespace Chaos
 
 
 		template<typename T, int d>
-		void Apply( TRigidBodySingleContactConstraint<T, d>& Constraint, T Thickness, 
-			const T Dt, const int32 It, const int32 NumIts, const int32 PairIterations,
-			TArrayCollectionArray<bool>* Collided,
-			const TArrayCollectionArray<TSerializablePtr<FChaosPhysicsMaterial>>* PhysicsMaterials,
-			T CollisionFrictionOverride,
-			T AngularFriction) 
+		void Apply(TRigidBodySingleContactConstraint<T, d>& Constraint, T Thickness, TSingleContactIterationParameters<T> & IterationParameters, TSingleContParticleParameters<T> & ParticleParameters)
 		{
 			TGenericParticleHandle<T, d> Particle0 = TGenericParticleHandle<T, d>(Constraint.Particle[0]);
 			TGenericParticleHandle<T, d> Particle1 = TGenericParticleHandle<T, d>(Constraint.Particle[1]);
 			TPBDRigidParticleHandle<T, d>* PBDRigid0 = Particle0->AsDynamic();
 			TPBDRigidParticleHandle<T, d>* PBDRigid1 = Particle1->AsDynamic();
 
-			for (int32 PairIt = 0; PairIt < PairIterations; ++PairIt) 
+			for (int32 PairIt = 0; PairIt < IterationParameters.NumPairIterations; ++PairIt)
 			{
 				Collisions::UpdateConstraint<ECollisionUpdateType::Deepest>(Thickness, Constraint);
 				if (Constraint.GetPhi() >= Thickness)
@@ -148,17 +143,17 @@ namespace Chaos
 				}
 
 				// @todo(ccaulfield): CHAOS_PARTICLEHANDLE_TODO what's the best way to manage external per-particle data?
-				if (Collided)
+				if (ParticleParameters.Collided)
 				{
-					Particle0->AuxilaryValue(*Collided) = true;
-					Particle1->AuxilaryValue(*Collided) = true;
+					Particle0->AuxilaryValue(*ParticleParameters.Collided) = true;
+					Particle1->AuxilaryValue(*ParticleParameters.Collided) = true;
 				}
 
 				TSerializablePtr<FChaosPhysicsMaterial> PhysicsMaterial0, PhysicsMaterial1;
-				if (PhysicsMaterials)
+				if (ParticleParameters.PhysicsMaterials)
 				{
-					PhysicsMaterial0 = Particle0->AuxilaryValue(*PhysicsMaterials);
-					PhysicsMaterial1 = Particle1->AuxilaryValue(*PhysicsMaterials);
+					PhysicsMaterial0 = Particle0->AuxilaryValue(*ParticleParameters.PhysicsMaterials);
+					PhysicsMaterial1 = Particle1->AuxilaryValue(*ParticleParameters.PhysicsMaterials);
 				}
 
 				// @todo(ccaulfield): CHAOS_PARTICLEHANDLE_TODO split function to avoid ifs
@@ -194,7 +189,8 @@ namespace Chaos
 					// Resting contact if very close to the surface
 					T Restitution = (T)0;
 					T Friction = (T)0;
-					bool bApplyRestitution = (RelativeVelocity.Size() > (2 * 980 * Dt));
+					T AngularFriction = (T)0;
+					bool bApplyRestitution = (RelativeVelocity.Size() > (2 * 980 * IterationParameters.Dt));
 					if (PhysicsMaterial0 && PhysicsMaterial1)
 					{
 						if (bApplyRestitution)
@@ -220,9 +216,13 @@ namespace Chaos
 						Friction = PhysicsMaterial1->Friction;
 					}
 
-					if (CollisionFrictionOverride >= 0)
+					if (ParticleParameters.FrictionOverride >= 0)
 					{
-						Friction = CollisionFrictionOverride;
+						Friction = ParticleParameters.FrictionOverride;
+					}
+					if (ParticleParameters.AngularFrictionOverride >= 0)
+					{
+						AngularFriction = ParticleParameters.AngularFrictionOverride;
 					}
 
 					if (Friction)
@@ -325,8 +325,8 @@ namespace Chaos
 						PBDRigid0->V() += DV;
 						PBDRigid0->W() += DW;
 						// Position update as part of pbd
-						PBDRigid0->P() += DV * Dt;
-						PBDRigid0->Q() += FRotation3::FromElements(DW, 0.f) * Q0 * Dt * T(0.5);
+						PBDRigid0->P() += DV * IterationParameters.Dt;
+						PBDRigid0->Q() += FRotation3::FromElements(DW, 0.f) * Q0 * IterationParameters.Dt * T(0.5);
 						PBDRigid0->Q().Normalize();
 					}
 					if (PBDRigid1)
@@ -338,24 +338,18 @@ namespace Chaos
 						PBDRigid1->V() += DV;
 						PBDRigid1->W() += DW;
 						// Position update as part of pbd
-						PBDRigid1->P() += DV * Dt;
-						PBDRigid1->Q() += FRotation3::FromElements(DW, 0.f) * Q1 * Dt * T(0.5);
+						PBDRigid1->P() += DV * IterationParameters.Dt;
+						PBDRigid1->Q() += FRotation3::FromElements(DW, 0.f) * Q1 * IterationParameters.Dt * T(0.5);
 						PBDRigid1->Q().Normalize();
 					}
 				}
 			}
 		}
-		template void Apply<float, 3>(TRigidBodySingleContactConstraint<float,3>&, float, 
-			const float, const int32, const int32, const int32,
-			TArrayCollectionArray<bool>*, const TArrayCollectionArray<TSerializablePtr<FChaosPhysicsMaterial>>*,
-			float CollisionFrictionOverride,
-			float AngularFriction);
-
+		template void Apply<float, 3>(TRigidBodySingleContactConstraint<float, 3>&, float, TSingleContactIterationParameters<float> &, TSingleContParticleParameters<float> &);
 
 		template<typename T, int d>
 		void ApplyPushOut(TRigidBodySingleContactConstraint<T, d>& Constraint, T Thickness, const TSet<const TGeometryParticleHandle<T, d>*>& IsTemporarilyStatic,
-			const T Dt, const int32 Iteration, const int32 NumIterations, const int32 PairIterations, bool &NeedsAnotherIteration, 
-			const TArrayCollectionArray<TSerializablePtr<FChaosPhysicsMaterial>>* PhysicsMaterials)
+			TSingleContactIterationParameters<T> & IterationParameters, TSingleContParticleParameters<T> & ParticleParameters)
 		{
 			TGeometryParticleHandle<T, d>* Particle0 = Constraint.Particle[0];
 			TGeometryParticleHandle<T, d>* Particle1 = Constraint.Particle[1];
@@ -375,13 +369,13 @@ namespace Chaos
 			const bool IsTemporarilyStatic1 = IsTemporarilyStatic.Contains(Particle1);
 	
 			TSerializablePtr<FChaosPhysicsMaterial> PhysicsMaterial0, PhysicsMaterial1;
-			if (PhysicsMaterials)
+			if (ParticleParameters.PhysicsMaterials)
 			{
-				PhysicsMaterial0 = Particle0->AuxilaryValue(*PhysicsMaterials);
-				PhysicsMaterial1 = Particle1->AuxilaryValue(*PhysicsMaterials);
+				PhysicsMaterial0 = Particle0->AuxilaryValue(*ParticleParameters.PhysicsMaterials);
+				PhysicsMaterial1 = Particle1->AuxilaryValue(*ParticleParameters.PhysicsMaterials);
 			}
 
-			for (int32 PairIteration = 0; PairIteration < PairIterations; ++PairIteration)
+			for (int32 PairIteration = 0; PairIteration < IterationParameters.NumPairIterations; ++PairIteration)
 			{
 				UpdateConstraint<ECollisionUpdateType::Deepest>(Thickness, Constraint);
 
@@ -397,7 +391,7 @@ namespace Chaos
 					break;
 				}
 
-				NeedsAnotherIteration = true;
+				*IterationParameters.NeedsAnotherIteration = true;
 				PMatrix<T, d, d> WorldSpaceInvI1 = PBDRigid0 ? (Q0 * FMatrix::Identity).GetTransposed() * PBDRigid0->InvI() * (Q0 * FMatrix::Identity) : PMatrix<T, d, d>(0);
 				PMatrix<T, d, d> WorldSpaceInvI2 = PBDRigid1 ? (Q1 * FMatrix::Identity).GetTransposed() * PBDRigid1->InvI() * (Q1 * FMatrix::Identity) : PMatrix<T, d, d>(0);
 				TVector<T, d> VectorToPoint1 = Contact.Location - P0;
@@ -405,8 +399,8 @@ namespace Chaos
 				PMatrix<T, d, d> Factor =
 					(PBDRigid0 ? ComputeFactorMatrix3(VectorToPoint1, WorldSpaceInvI1, PBDRigid0->InvM()) : PMatrix<T, d, d>(0)) +
 					(PBDRigid1 ? ComputeFactorMatrix3(VectorToPoint2, WorldSpaceInvI2, PBDRigid1->InvM()) : PMatrix<T, d, d>(0));
-				T Numerator = FMath::Min((T)(Iteration + 2), (T)NumIterations);
-				T ScalingFactor = Numerator / (T)NumIterations;
+				T Numerator = FMath::Min((T)(IterationParameters.Iteration + 2), (T)IterationParameters.NumIterations);
+				T ScalingFactor = Numerator / (T)IterationParameters.NumIterations;
 
 				//if pushout is needed we better fix relative velocity along normal. Treat it as if 0 restitution
 				TVector<T, d> Body1Velocity = V0 + TVector<T, d>::CrossProduct(W0, VectorToPoint1);
@@ -465,8 +459,7 @@ namespace Chaos
 			}
 		}
 		template void ApplyPushOut<float, 3>(TRigidBodySingleContactConstraint<float,3>&, float , const TSet<const TGeometryParticleHandle<float,3>*>&,
-			const float, const int32, const int32, const int32, bool &,
-			const TArrayCollectionArray<TSerializablePtr<FChaosPhysicsMaterial>>*);
+			TSingleContactIterationParameters<float> & IterationParameters, TSingleContParticleParameters<float> & ParticleParameters);
 
 	} // Collisions
 
diff --git a/Engine/Source/Runtime/Experimental/Chaos/Public/Chaos/PBDCollisionConstraintSingleContactUtil.h b/Engine/Source/Runtime/Experimental/Chaos/Public/Chaos/PBDCollisionConstraintSingleContactUtil.h
--- a/Engine/Source/Runtime/Experimental/Chaos/Public/Chaos/PBDCollisionConstraintSingleContactUtil.h
+++ b/Engine/Source/Runtime/Experimental/Chaos/Public/Chaos/PBDCollisionConstraintSingleContactUtil.h
@@ -9,6 +9,23 @@ namespace Chaos
 {
 	namespace Collisions
 	{
+		template<class T = float>
+		struct TSingleContParticleParameters {
+			TArrayCollectionArray<bool>* Collided;
+			const TArrayCollectionArray<TSerializablePtr<FChaosPhysicsMaterial>>* PhysicsMaterials;
+			T FrictionOverride;
+			T AngularFrictionOverride;
+		};
+
+		template<class T = float>
+		struct TSingleContactIterationParameters {
+			const T Dt; 
+			const int32 Iteration; 
+			const int32 NumIterations; 
+			const int32 NumPairIterations;
+			bool* NeedsAnotherIteration;
+		};
+
 		template<ECollisionUpdateType UpdateType, typename T = float, int d = 3>
 		void Update(const T Thickness, TRigidBodySingleContactConstraint<T, d>& Constraint);
 
@@ -17,14 +34,13 @@ namespace Chaos
 			
 		template<typename T = float, int d = 3>
 		void Apply(TRigidBodySingleContactConstraint<T, d>& Constraint, T Thickness, 
-			const T Dt, const int32 It, const int32 NumIts, const int32 NumPairIterations,
-			TArrayCollectionArray<bool>* Collided, const TArrayCollectionArray<TSerializablePtr<FChaosPhysicsMaterial>>*,T CollisionFrictionOverride, T AngularFriction);
+			TSingleContactIterationParameters<T> & IterationParameters,
+			TSingleContParticleParameters<T> & ParticleParameters);
 			
 		template<typename T=float, int d=3>
-		void ApplyPushOut(TRigidBodySingleContactConstraint<T, d>& Constraint, T Thickness, const TSet<const TGeometryParticleHandle<T, d>*>& IsTemporarilyStatic, 
-			const T Dt, const int32 Iteration, const int32 NumIterations, const int32 NumPairIterations, bool &NeedsAnotherIteration,
-			const TArrayCollectionArray<TSerializablePtr<FChaosPhysicsMaterial>>* PhysicsMaterials);
-
+		void ApplyPushOut(TRigidBodySingleContactConstraint<T, d>& Constraint, T Thickness, const TSet<const TGeometryParticleHandle<T, d>*>& IsTemporarilyStatic,
+			TSingleContactIterationParameters<T> & IterationParameters,
+			TSingleContParticleParameters<T> & ParticleParameters);
 	}
 
 }
