Change rotational Euler velocity steps to use low-velocity approximation and world-space inertia tensors.

diff --git a/Engine/Source/Runtime/Experimental/Chaos/Public/Chaos/PerParticleEulerStepVelocity.h b/Engine/Source/Runtime/Experimental/Chaos/Public/Chaos/PerParticleEulerStepVelocity.h
--- a/Engine/Source/Runtime/Experimental/Chaos/Public/Chaos/PerParticleEulerStepVelocity.h
+++ b/Engine/Source/Runtime/Experimental/Chaos/Public/Chaos/PerParticleEulerStepVelocity.h
@@ -2,6 +2,7 @@
 #pragma once
 
 #include "Chaos/PerParticleRule.h"
+#include "Chaos/Utilities.h"
 
 namespace Chaos
 {
@@ -30,13 +31,23 @@ class TPerParticleEulerStepVelocity : public TPerParticleRule<T, d>
 		if (InParticles.InvM(Index) == 0 || InParticles.Disabled(Index) || InParticles.Sleeping(Index))
 			return;
 		ApplyHelper(InParticles, Dt, Index);
-		TVector<T, d> L = InParticles.I(Index) * InParticles.W(Index);
-		InParticles.W(Index) += InParticles.InvI(Index) * (InParticles.Torque(Index) - TVector<T, d>::CrossProduct(InParticles.W(Index), L)) * Dt;
+
+		//
+		// TODO: This is the first-order approximation.
+		//       If needed, we might eventually want to do a second order Euler's Equation,
+		//       but if we do that we'll need to do a transform into a rotating reference frame.
+		//       Just using W += InvI * (Torque - W x (I * W)) * dt is not correct, since Torque
+		//		 and W are in an inertial frame.
+		//
+		const FMatrix33 WorldInvI = Utilities::ComputeWorldSpaceInertia(InParticles.R(Index), InParticles.InvI(Index));
+		InParticles.W(Index) += WorldInvI * InParticles.Torque(Index) * Dt;
 	}
 	
 	inline void Apply(TTransientPBDRigidParticleHandle<T, d>& Particle, const T Dt) const override //-V762
 	{
 		Particle.V() += Particle.F() * Particle.InvM() * Dt;
+		const FMatrix33 WorldInvI = Utilities::ComputeWorldSpaceInertia(Particle.R(), Particle.InvI());
+		Particle.W() += WorldInvI * Particle.Torque() * Dt;
 	}
 };
 }
