Choas :

Fix CIS - Android

diff --git a/Engine/Source/Runtime/Experimental/Chaos/Private/Chaos/PBDCollisionConstraint.cpp b/Engine/Source/Runtime/Experimental/Chaos/Private/Chaos/PBDCollisionConstraint.cpp
--- a/Engine/Source/Runtime/Experimental/Chaos/Private/Chaos/PBDCollisionConstraint.cpp
+++ b/Engine/Source/Runtime/Experimental/Chaos/Private/Chaos/PBDCollisionConstraint.cpp
@@ -404,7 +404,6 @@ namespace Chaos
 		Collisions::UpdateConstraint<UpdateType>(Thickness, Constraint);
 	}
 
-	template class TPBDCollisionConstraintHandle<float, 3>;
 	template class TAccelerationStructureHandle<float, 3>;
 	template class CHAOS_API TPBDCollisionConstraint<float, 3>;
 	template void TPBDCollisionConstraint<float, 3>::UpdateConstraint<ECollisionUpdateType::Any>(const float Thickness, FRigidBodyContactConstraint& Constraint);
