Chaos :

Added stubbed out  PlaneCollisionConstrain.

diff --git a/Engine/Source/Runtime/Experimental/Chaos/Private/Chaos/ChaosDebugDraw.cpp b/Engine/Source/Runtime/Experimental/Chaos/Private/Chaos/ChaosDebugDraw.cpp
--- a/Engine/Source/Runtime/Experimental/Chaos/Private/Chaos/ChaosDebugDraw.cpp
+++ b/Engine/Source/Runtime/Experimental/Chaos/Private/Chaos/ChaosDebugDraw.cpp
@@ -139,7 +139,7 @@ namespace Chaos
 		{
 			if (ConstraintHandle->GetType() == TCollisionConstraintBase<float,3>::FType::SinglePoint)
 			{
-				const TRigidBodySingleContactConstraint<FReal, 3>& Contact = ConstraintHandle->GetContact< TRigidBodySingleContactConstraint<float, 3> >();
+				const TRigidBodyPointContactConstraint<FReal, 3>& Contact = ConstraintHandle->GetSingleContact();
 				if (Contact.GetPhi() > 0)
 				{
 					ColorScale = ColorScale * (FReal)0.1;
diff --git a/Engine/Source/Runtime/Experimental/Chaos/Private/Chaos/CollisionResolutionAlgo.cpp b/Engine/Source/Runtime/Experimental/Chaos/Private/Chaos/CollisionResolutionAlgo.cpp
--- a/Engine/Source/Runtime/Experimental/Chaos/Private/Chaos/CollisionResolutionAlgo.cpp
+++ b/Engine/Source/Runtime/Experimental/Chaos/Private/Chaos/CollisionResolutionAlgo.cpp
@@ -27,7 +27,7 @@ namespace Chaos
 	{
 
 		template<typename T, int d>
-		void ConstructLevelsetConstraints(TGeometryParticleHandle<T, d>* Particle0, TGeometryParticleHandle<T, d>* Particle1, const FImplicitObject* Implicit0, const FImplicitObject* Implicit1, const T Thickness, TRigidBodySingleContactConstraint<T, d>& Constraint)
+		void ConstructLevelsetConstraints(TGeometryParticleHandle<T, d>* Particle0, TGeometryParticleHandle<T, d>* Particle1, const FImplicitObject* Implicit0, const FImplicitObject* Implicit1, const T Thickness, TRigidBodyPointContactConstraint<T, d>& Constraint)
 		{
 
 			if (!Particle1->Geometry() || (Particle0->AsDynamic() && !Particle0->AsDynamic()->CollisionParticlesSize() && Particle0->Geometry() && !Particle0->Geometry()->IsUnderlyingUnion()))
@@ -45,7 +45,7 @@ namespace Chaos
 		}
 
 		template<typename T, int d>
-		void ConstructBoxConstraints(TGeometryParticleHandle<T, d>* Particle0, TGeometryParticleHandle<T, d>* Particle1, const FImplicitObject* Implicit0, const FImplicitObject* Implicit1, const T Thickness, TRigidBodySingleContactConstraint<T, d>& Constraint)
+		void ConstructBoxConstraints(TGeometryParticleHandle<T, d>* Particle0, TGeometryParticleHandle<T, d>* Particle1, const FImplicitObject* Implicit0, const FImplicitObject* Implicit1, const T Thickness, TRigidBodyPointContactConstraint<T, d>& Constraint)
 		{
 			Constraint.Particle[0] = Particle0;
 			Constraint.Particle[1] = Particle1;
@@ -53,7 +53,7 @@ namespace Chaos
 		}
 
 		template<typename T, int d>
-		void ConstructBoxPlaneConstraints(TGeometryParticleHandle<T, d>* Particle0, TGeometryParticleHandle<T, d>* Particle1, const FImplicitObject* Implicit0, const FImplicitObject* Implicit1, const T Thickness, TRigidBodySingleContactConstraint<T, d>& Constraint)
+		void ConstructBoxPlaneConstraints(TGeometryParticleHandle<T, d>* Particle0, TGeometryParticleHandle<T, d>* Particle1, const FImplicitObject* Implicit0, const FImplicitObject* Implicit1, const T Thickness, TRigidBodyPointContactConstraint<T, d>& Constraint)
 		{
 			Constraint.Particle[0] = Particle0;
 			Constraint.Particle[1] = Particle1;
@@ -61,7 +61,7 @@ namespace Chaos
 		}
 
 		template<typename T, int d>
-		void ConstructSphereConstraints(TGeometryParticleHandle<T, d>* Particle0, TGeometryParticleHandle<T, d>* Particle1, const FImplicitObject* Implicit0, const FImplicitObject* Implicit1, const T Thickness, TRigidBodySingleContactConstraint<T, d>& Constraint)
+		void ConstructSphereConstraints(TGeometryParticleHandle<T, d>* Particle0, TGeometryParticleHandle<T, d>* Particle1, const FImplicitObject* Implicit0, const FImplicitObject* Implicit1, const T Thickness, TRigidBodyPointContactConstraint<T, d>& Constraint)
 		{
 			Constraint.Particle[0] = Particle0;
 			Constraint.Particle[1] = Particle1;
@@ -69,7 +69,7 @@ namespace Chaos
 		}
 
 		template<typename T, int d>
-		void ConstructSpherePlaneConstraints(TGeometryParticleHandle<T, d>* Particle0, TGeometryParticleHandle<T, d>* Particle1, const FImplicitObject* Implicit0, const FImplicitObject* Implicit1, const T Thickness, TRigidBodySingleContactConstraint<T, d>& Constraint)
+		void ConstructSpherePlaneConstraints(TGeometryParticleHandle<T, d>* Particle0, TGeometryParticleHandle<T, d>* Particle1, const FImplicitObject* Implicit0, const FImplicitObject* Implicit1, const T Thickness, TRigidBodyPointContactConstraint<T, d>& Constraint)
 		{
 			Constraint.Particle[0] = Particle0;
 			Constraint.Particle[1] = Particle1;
@@ -77,7 +77,7 @@ namespace Chaos
 		}
 
 		template<typename T, int d>
-		void ConstructSphereBoxConstraints(TGeometryParticleHandle<T, d>* Particle0, TGeometryParticleHandle<T, d>* Particle1, const FImplicitObject* Implicit0, const FImplicitObject* Implicit1, const T Thickness, TRigidBodySingleContactConstraint<T, d>& Constraint)
+		void ConstructSphereBoxConstraints(TGeometryParticleHandle<T, d>* Particle0, TGeometryParticleHandle<T, d>* Particle1, const FImplicitObject* Implicit0, const FImplicitObject* Implicit1, const T Thickness, TRigidBodyPointContactConstraint<T, d>& Constraint)
 		{
 			Constraint.Particle[0] = Particle0;
 			Constraint.Particle[1] = Particle1;
@@ -85,7 +85,7 @@ namespace Chaos
 		}
 
 		template<typename T, int d>
-		void ConstructCapsuleCapsuleConstraints(TGeometryParticleHandle<T, d>* Particle0, TGeometryParticleHandle<T, d>* Particle1, const FImplicitObject* Implicit0, const FImplicitObject* Implicit1, const T Thickness, TRigidBodySingleContactConstraint<T, d>& Constraint)
+		void ConstructCapsuleCapsuleConstraints(TGeometryParticleHandle<T, d>* Particle0, TGeometryParticleHandle<T, d>* Particle1, const FImplicitObject* Implicit0, const FImplicitObject* Implicit1, const T Thickness, TRigidBodyPointContactConstraint<T, d>& Constraint)
 		{
 			Constraint.Particle[0] = Particle0;
 			Constraint.Particle[1] = Particle1;
@@ -93,7 +93,7 @@ namespace Chaos
 		}
 
 		template<typename T, int d>
-		void ConstructCapsuleBoxConstraints(TGeometryParticleHandle<T, d>* Particle0, TGeometryParticleHandle<T, d>* Particle1, const FImplicitObject* Implicit0, const FImplicitObject* Implicit1, const T Thickness, TRigidBodySingleContactConstraint<T, d>& Constraint)
+		void ConstructCapsuleBoxConstraints(TGeometryParticleHandle<T, d>* Particle0, TGeometryParticleHandle<T, d>* Particle1, const FImplicitObject* Implicit0, const FImplicitObject* Implicit1, const T Thickness, TRigidBodyPointContactConstraint<T, d>& Constraint)
 		{
 			Constraint.Particle[0] = Particle0;
 			Constraint.Particle[1] = Particle1;
@@ -101,7 +101,7 @@ namespace Chaos
 		}
 
 		template<typename T, int d>
-		void ConstructSingleUnionConstraints(TGeometryParticleHandle<T, d>* Particle0, TGeometryParticleHandle<T, d>* Particle1, const FImplicitObject* Implicit0, const FImplicitObject* Implicit1, const T Thickness, TRigidBodySingleContactConstraint<T, d>& Constraint)
+		void ConstructSingleUnionConstraints(TGeometryParticleHandle<T, d>* Particle0, TGeometryParticleHandle<T, d>* Particle1, const FImplicitObject* Implicit0, const FImplicitObject* Implicit1, const T Thickness, TRigidBodyPointContactConstraint<T, d>& Constraint)
 		{
 			Constraint.Particle[0] = Particle0;
 			Constraint.Particle[1] = Particle1;
@@ -111,7 +111,7 @@ namespace Chaos
 		DECLARE_CYCLE_STAT(TEXT("TPBDCollisionConstraint::UpdateUnionUnionConstraint"), STAT_UpdateUnionUnionConstraint, STATGROUP_ChaosWide);
 
 		template<typename T, int d>
-		void ConstructUnionUnionConstraints(TGeometryParticleHandle<T, d>* Particle0, TGeometryParticleHandle<T, d>* Particle1, const FImplicitObject* Implicit0, const FImplicitObject* Implicit1, const T Thickness, TRigidBodySingleContactConstraint<T, d>& Constraint)
+		void ConstructUnionUnionConstraints(TGeometryParticleHandle<T, d>* Particle0, TGeometryParticleHandle<T, d>* Particle1, const FImplicitObject* Implicit0, const FImplicitObject* Implicit1, const T Thickness, TRigidBodyPointContactConstraint<T, d>& Constraint)
 		{
 			SCOPE_CYCLE_COUNTER(STAT_UpdateUnionUnionConstraint);
 
@@ -140,7 +140,7 @@ namespace Chaos
 		}
 
 		template<typename T, int d>
-		void ConstructPairConstraintImpl(TGeometryParticleHandle<T, d>* Particle0, TGeometryParticleHandle<T, d>* Particle1, const FImplicitObject* Implicit0, const FImplicitObject* Implicit1, const T Thickness, TRigidBodySingleContactConstraint<T, d>& Constraint)
+		void ConstructPairConstraintImpl(TGeometryParticleHandle<T, d>* Particle0, TGeometryParticleHandle<T, d>* Particle1, const FImplicitObject* Implicit0, const FImplicitObject* Implicit1, const T Thickness, TRigidBodyPointContactConstraint<T, d>& Constraint)
 		{
 			// See if we already have a constraint for this shape pair
 			if (Constraint.ContainsManifold(Implicit0, Implicit1))
@@ -220,13 +220,13 @@ namespace Chaos
 		}
 
 		template <typename T, int d>
-		bool SampleObjectHelper2(const FImplicitObject& Object, const TRigidTransform<T, d>& ObjectTransform, const TRigidTransform<T, d>& SampleToObjectTransform, const TVector<T, d>& SampleParticle, T Thickness, TRigidBodySingleContactConstraint<T, d>& Constraint)
+		bool SampleObjectHelper2(const FImplicitObject& Object, const TRigidTransform<T, d>& ObjectTransform, const TRigidTransform<T, d>& SampleToObjectTransform, const TVector<T, d>& SampleParticle, T Thickness, TRigidBodyPointContactConstraint<T, d>& Constraint)
 		{
 			TVector<T, d> LocalPoint = SampleToObjectTransform.TransformPositionNoScale(SampleParticle);
 			TVector<T, d> LocalNormal;
 			T LocalPhi = Object.PhiWithNormal(LocalPoint, LocalNormal);
 
-			TSingleContactManifold<T, d> & Contact = Constraint.Manifold;
+			TPointContactManifold<T, d> & Contact = Constraint.Manifold;
 			if (LocalPhi < Contact.Phi)
 			{
 				Contact.Phi = LocalPhi;
@@ -238,13 +238,13 @@ namespace Chaos
 		}
 
 		template <typename T, int d>
-		bool SampleObjectNoNormal2(const FImplicitObject& Object, const TRigidTransform<T, d>& ObjectTransform, const TRigidTransform<T, d>& SampleToObjectTransform, const TVector<T, d>& SampleParticle, T Thickness, TRigidBodySingleContactConstraint<T, d>& Constraint)
+		bool SampleObjectNoNormal2(const FImplicitObject& Object, const TRigidTransform<T, d>& ObjectTransform, const TRigidTransform<T, d>& SampleToObjectTransform, const TVector<T, d>& SampleParticle, T Thickness, TRigidBodyPointContactConstraint<T, d>& Constraint)
 		{
 			TVector<T, d> LocalPoint = SampleToObjectTransform.TransformPositionNoScale(SampleParticle);
 			TVector<T, d> LocalNormal;
 			T LocalPhi = Object.PhiWithNormal(LocalPoint, LocalNormal);
 
-			TSingleContactManifold<T, d> & Contact = Constraint.Manifold;
+			TPointContactManifold<T, d> & Contact = Constraint.Manifold;
 			if (LocalPhi < Contact.Phi)
 			{
 				Contact.Phi = LocalPhi;
@@ -254,14 +254,14 @@ namespace Chaos
 		}
 
 		template <typename T, int d>
-		bool SampleObjectNormalAverageHelper2(const FImplicitObject& Object, const TRigidTransform<T, d>& ObjectTransform, const TRigidTransform<T, d>& SampleToObjectTransform, const TVector<T, d>& SampleParticle, T Thickness, T& TotalThickness, TRigidBodySingleContactConstraint<T, d>& Constraint)
+		bool SampleObjectNormalAverageHelper2(const FImplicitObject& Object, const TRigidTransform<T, d>& ObjectTransform, const TRigidTransform<T, d>& SampleToObjectTransform, const TVector<T, d>& SampleParticle, T Thickness, T& TotalThickness, TRigidBodyPointContactConstraint<T, d>& Constraint)
 		{
 			TVector<T, d> LocalPoint = SampleToObjectTransform.TransformPositionNoScale(SampleParticle);
 			TVector<T, d> LocalNormal;
 			T LocalPhi = Object.PhiWithNormal(LocalPoint, LocalNormal);
 			T LocalThickness = LocalPhi - Thickness;
 
-			TSingleContactManifold<T, d> & Contact = Constraint.Manifold;
+			TPointContactManifold<T, d> & Contact = Constraint.Manifold;
 			if (LocalThickness < -KINDA_SMALL_NUMBER)
 			{
 				Contact.Location += LocalPoint * LocalThickness;
@@ -286,13 +286,13 @@ namespace Chaos
 
 
 		template <ECollisionUpdateType UpdateType, typename T, int d>
-		void SampleObject2(const FImplicitObject& Object, const TRigidTransform<T, d>& ObjectTransform, const TBVHParticles<T, d>& SampleParticles, const TRigidTransform<T, d>& SampleParticlesTransform, T Thickness, TRigidBodySingleContactConstraint<T, d>& Constraint)
+		void SampleObject2(const FImplicitObject& Object, const TRigidTransform<T, d>& ObjectTransform, const TBVHParticles<T, d>& SampleParticles, const TRigidTransform<T, d>& SampleParticlesTransform, T Thickness, TRigidBodyPointContactConstraint<T, d>& Constraint)
 		{
 			SCOPE_CYCLE_COUNTER(STAT_SampleObject);
-			TRigidBodySingleContactConstraint<float, 3> AvgConstraint = Constraint;
+			TRigidBodyPointContactConstraint<float, 3> AvgConstraint = Constraint;
 
-			TSingleContactManifold<float, 3> & Contact = Constraint.Manifold;
-			TSingleContactManifold<float, 3> & AvgContact = AvgConstraint.Manifold;
+			TPointContactManifold<float, 3> & Contact = Constraint.Manifold;
+			TPointContactManifold<float, 3> & AvgContact = AvgConstraint.Manifold;
 
 			AvgConstraint.Particle[0] = Constraint.Particle[0];
 			AvgConstraint.Particle[1] = Constraint.Particle[1];
@@ -396,13 +396,13 @@ namespace Chaos
 
 #if INTEL_ISPC
 		template<ECollisionUpdateType UpdateType>
-		void SampleObject2(const FImplicitObject& Object, const TRigidTransform<float, 3>& ObjectTransform, const TBVHParticles<float, 3>& SampleParticles, const TRigidTransform<float, 3>& SampleParticlesTransform, float Thickness, TRigidBodySingleContactConstraint<float, 3>& Constraint)
+		void SampleObject2(const FImplicitObject& Object, const TRigidTransform<float, 3>& ObjectTransform, const TBVHParticles<float, 3>& SampleParticles, const TRigidTransform<float, 3>& SampleParticlesTransform, float Thickness, TRigidBodyPointContactConstraint<float, 3>& Constraint)
 		{
 			SCOPE_CYCLE_COUNTER(STAT_SampleObject);
-			TRigidBodySingleContactConstraint<float, 3> AvgConstraint = Constraint;
+			TRigidBodyPointContactConstraint<float, 3> AvgConstraint = Constraint;
 
-			TSingleContactManifold<float, 3> & Contact = Constraint.Manifold;
-			TSingleContactManifold<float, 3> & AvgContact = AvgConstraint.Manifold;
+			TPointContactManifold<float, 3> & Contact = Constraint.Manifold;
+			TPointContactManifold<float, 3> & AvgContact = AvgConstraint.Manifold;
 
 			AvgConstraint.Particle[0] = Constraint.Particle[0];
 			AvgConstraint.Particle[1] = Constraint.Particle[1];
@@ -704,7 +704,7 @@ namespace Chaos
 #endif
 
 		template<ECollisionUpdateType UpdateType, typename T, int d>
-		void UpdateConstraintImp(const FImplicitObject& ParticleObject, const TRigidTransform<T, d>& ParticleTM, const FImplicitObject& LevelsetObject, const TRigidTransform<T, d>& LevelsetTM, const T Thickness, TRigidBodySingleContactConstraint<T, d>& Constraint)
+		void UpdateConstraintImp(const FImplicitObject& ParticleObject, const TRigidTransform<T, d>& ParticleTM, const FImplicitObject& LevelsetObject, const TRigidTransform<T, d>& LevelsetTM, const T Thickness, TRigidBodyPointContactConstraint<T, d>& Constraint)
 		{
 			if (ParticleObject.GetType() == TBox<T, d>::StaticType() && LevelsetObject.GetType() == TBox<T, d>::StaticType())
 			{
@@ -736,7 +736,7 @@ namespace Chaos
 			}
 			else if (ParticleObject.GetType() == TPlane<T, d>::StaticType() && LevelsetObject.GetType() == TBox<T, d>::StaticType())
 			{
-				TRigidBodySingleContactConstraint<T, d> TmpConstraint = Constraint;
+				TRigidBodyPointContactConstraint<T, d> TmpConstraint = Constraint;
 				UpdateBoxPlaneConstraint(*LevelsetObject.template GetObject<TBox<T, d>>(), LevelsetTM, *ParticleObject.template GetObject<TPlane<T, d>>(), ParticleTM, Thickness, TmpConstraint);
 				if (TmpConstraint.GetPhi() < Constraint.GetPhi())
 				{
@@ -746,7 +746,7 @@ namespace Chaos
 			}
 			else if (ParticleObject.GetType() == TPlane<T, d>::StaticType() && LevelsetObject.GetType() == TSphere<T, d>::StaticType())
 			{
-				TRigidBodySingleContactConstraint<T, d> TmpConstraint = Constraint;
+				TRigidBodyPointContactConstraint<T, d> TmpConstraint = Constraint;
 				UpdateSpherePlaneConstraint(*LevelsetObject.template GetObject<TSphere<T, d>>(), LevelsetTM, *ParticleObject.template GetObject<TPlane<T, d>>(), ParticleTM, Thickness, TmpConstraint);
 				if (TmpConstraint.GetPhi() < Constraint.GetPhi())
 				{
@@ -756,7 +756,7 @@ namespace Chaos
 			}
 			else if (ParticleObject.GetType() == TBox<T, d>::StaticType() && LevelsetObject.GetType() == TSphere<T, d>::StaticType())
 			{
-				TRigidBodySingleContactConstraint<T, d> TmpConstraint = Constraint;
+				TRigidBodyPointContactConstraint<T, d> TmpConstraint = Constraint;
 				UpdateSphereBoxConstraint(*LevelsetObject.template GetObject<TSphere<T, d>>(), LevelsetTM, *ParticleObject.template GetObject<TBox<T, d>>(), ParticleTM, Thickness, TmpConstraint);
 				if (TmpConstraint.GetPhi() < Constraint.GetPhi())
 				{
@@ -766,7 +766,7 @@ namespace Chaos
 			}
 			else if (ParticleObject.GetType() == TBox<T, d>::StaticType() && LevelsetObject.GetType() == TCapsule<T>::StaticType())
 			{
-				TRigidBodySingleContactConstraint<T, d> TmpConstraint = Constraint;
+				TRigidBodyPointContactConstraint<T, d> TmpConstraint = Constraint;
 				UpdateCapsuleBoxConstraint(*LevelsetObject.template GetObject<TCapsule<T>>(), LevelsetTM, *ParticleObject.template GetObject<TBox<T, d>>(), ParticleTM, Thickness, TmpConstraint);
 				if (TmpConstraint.GetPhi() < Constraint.GetPhi())
 				{
@@ -805,7 +805,7 @@ namespace Chaos
 		}
 
 		template<typename T, int d>
-		void ConstructConstraintsImpl(TGeometryParticleHandle<T, d>* Particle0, TGeometryParticleHandle<T, d>* Particle1, const FImplicitObject* Implicit0, const FImplicitObject* Implicit1, const T Thickness, TRigidBodySingleContactConstraint<T, d> & Constraint)
+		void ConstructConstraintsImpl(TGeometryParticleHandle<T, d>* Particle0, TGeometryParticleHandle<T, d>* Particle1, const FImplicitObject* Implicit0, const FImplicitObject* Implicit1, const T Thickness, TRigidBodyPointContactConstraint<T, d> & Constraint)
 		{
 			// TriangleMesh implicits are for scene query only.
 			if (Implicit0 && GetInnerType(Implicit0->GetType()) == ImplicitObjectType::TriangleMesh) return;
@@ -823,9 +823,9 @@ namespace Chaos
 		}
 
 		template <typename T, int d>
-		bool UpdateBoxPlaneConstraint(const TBox<T, d>& Box, const TRigidTransform<T, d>& BoxTransform, const TPlane<T, d>& Plane, const TRigidTransform<T, d>& PlaneTransform, T Thickness, TRigidBodySingleContactConstraint<T, d>& Constraint)
+		bool UpdateBoxPlaneConstraint(const TBox<T, d>& Box, const TRigidTransform<T, d>& BoxTransform, const TPlane<T, d>& Plane, const TRigidTransform<T, d>& PlaneTransform, T Thickness, TRigidBodyPointContactConstraint<T, d>& Constraint)
 		{
-			TSingleContactManifold<T, d> & Contact = Constraint.Manifold;
+			TPointContactManifold<T, d> & Contact = Constraint.Manifold;
 
 #if USING_CODE_ANALYSIS
 			MSVC_PRAGMA(warning(push))
@@ -885,9 +885,9 @@ namespace Chaos
 		}
 
 		template <typename T, int d>
-		void UpdateSphereConstraint(const TSphere<T, d>& Sphere1, const TRigidTransform<T, d>& Sphere1Transform, const TSphere<T, d>& Sphere2, const TRigidTransform<T, d>& Sphere2Transform, T Thickness, TRigidBodySingleContactConstraint<T, d>& Constraint)
+		void UpdateSphereConstraint(const TSphere<T, d>& Sphere1, const TRigidTransform<T, d>& Sphere1Transform, const TSphere<T, d>& Sphere2, const TRigidTransform<T, d>& Sphere2Transform, T Thickness, TRigidBodyPointContactConstraint<T, d>& Constraint)
 		{
-			TSingleContactManifold<T, d> & Contact = Constraint.Manifold;
+			TPointContactManifold<T, d> & Contact = Constraint.Manifold;
 
 			const TVector<T, d> Center1 = Sphere1Transform.TransformPosition(Sphere1.GetCenter());
 			const TVector<T, d> Center2 = Sphere2Transform.TransformPosition(Sphere2.GetCenter());
@@ -903,9 +903,9 @@ namespace Chaos
 		}
 
 		template <typename T, int d>
-		void UpdateSpherePlaneConstraint(const TSphere<T, d>& Sphere, const TRigidTransform<T, d>& SphereTransform, const TPlane<T, d>& Plane, const TRigidTransform<T, d>& PlaneTransform, T Thickness, TRigidBodySingleContactConstraint<T, d>& Constraint)
+		void UpdateSpherePlaneConstraint(const TSphere<T, d>& Sphere, const TRigidTransform<T, d>& SphereTransform, const TPlane<T, d>& Plane, const TRigidTransform<T, d>& PlaneTransform, T Thickness, TRigidBodyPointContactConstraint<T, d>& Constraint)
 		{
-			TSingleContactManifold<T, d> & Contact = Constraint.Manifold;
+			TPointContactManifold<T, d> & Contact = Constraint.Manifold;
 
 			const TRigidTransform<T, d> SphereToPlaneTransform(PlaneTransform.Inverse() * SphereTransform);
 			const TVector<T, d> SphereCenter = SphereToPlaneTransform.TransformPosition(Sphere.GetCenter());
@@ -923,9 +923,9 @@ namespace Chaos
 		}
 
 		template <typename T, int d>
-		void UpdateSphereBoxConstraint(const TSphere<T, d>& Sphere, const TRigidTransform<T, d>& SphereTransform, const TBox<T, d>& Box, const TRigidTransform<T, d>& BoxTransform, T Thickness, TRigidBodySingleContactConstraint<T, d>& Constraint)
+		void UpdateSphereBoxConstraint(const TSphere<T, d>& Sphere, const TRigidTransform<T, d>& SphereTransform, const TBox<T, d>& Box, const TRigidTransform<T, d>& BoxTransform, T Thickness, TRigidBodyPointContactConstraint<T, d>& Constraint)
 		{
-			TSingleContactManifold<T, d> & Contact = Constraint.Manifold;
+			TPointContactManifold<T, d> & Contact = Constraint.Manifold;
 
 			const TRigidTransform<T, d> SphereToBoxTransform(SphereTransform * BoxTransform.Inverse());
 			const TVector<T, d> SphereCenterInBox = SphereToBoxTransform.TransformPosition(Sphere.GetCenter());
@@ -943,9 +943,9 @@ namespace Chaos
 		}
 
 		template <typename T, int d>
-		void UpdateCapsuleCapsuleConstraint(const TCapsule<T>& A, const TRigidTransform<T, d>& ATransform, const TCapsule<T>& B, const TRigidTransform<T, d>& BTransform, T Thickness, TRigidBodySingleContactConstraint<T, d>& Constraint)
+		void UpdateCapsuleCapsuleConstraint(const TCapsule<T>& A, const TRigidTransform<T, d>& ATransform, const TCapsule<T>& B, const TRigidTransform<T, d>& BTransform, T Thickness, TRigidBodyPointContactConstraint<T, d>& Constraint)
 		{
-			TSingleContactManifold<T, d> & Contact = Constraint.Manifold;
+			TPointContactManifold<T, d> & Contact = Constraint.Manifold;
 
 			FVector A1 = ATransform.TransformPosition(A.GetX1());
 			FVector A2 = ATransform.TransformPosition(A.GetX2());
@@ -970,10 +970,10 @@ namespace Chaos
 		}
 
 		template <typename T, int d>
-		void UpdateCapsuleBoxConstraint(const TCapsule<T>& A, const TRigidTransform<T, d>& ATransform, const TBox<T, d>& B, const TRigidTransform<T, d>& BTransform, T Thickness, TRigidBodySingleContactConstraint<T, d>& Constraint)
+		void UpdateCapsuleBoxConstraint(const TCapsule<T>& A, const TRigidTransform<T, d>& ATransform, const TBox<T, d>& B, const TRigidTransform<T, d>& BTransform, T Thickness, TRigidBodyPointContactConstraint<T, d>& Constraint)
 		{
 			// @todo(ccaulfield): Add custom capsule-box collision
-			TSingleContactManifold<T, d> & Contact = Constraint.Manifold;
+			TPointContactManifold<T, d> & Contact = Constraint.Manifold;
 
 			TRigidTransform<T, d> BToATransform = BTransform.GetRelativeTransform(ATransform);
 
@@ -1039,7 +1039,7 @@ namespace Chaos
 		}
 
 		template<ECollisionUpdateType UpdateType, typename T, int d>
-		void UpdateUnionUnionConstraint(const T Thickness, TRigidBodySingleContactConstraint<T, d>& Constraint)
+		void UpdateUnionUnionConstraint(const T Thickness, TRigidBodyPointContactConstraint<T, d>& Constraint)
 		{
 			SCOPE_CYCLE_COUNTER(STAT_UpdateUnionUnionConstraint);
 
@@ -1073,7 +1073,7 @@ namespace Chaos
 
 		DECLARE_CYCLE_STAT(TEXT("TPBDCollisionConstraint::UpdateSingleUnionConstraint"), STAT_UpdateSingleUnionConstraint, STATGROUP_ChaosWide);
 		template<ECollisionUpdateType UpdateType, typename T, int d>
-		void UpdateSingleUnionConstraint(const T Thickness, TRigidBodySingleContactConstraint<T, d>& Constraint)
+		void UpdateSingleUnionConstraint(const T Thickness, TRigidBodyPointContactConstraint<T, d>& Constraint)
 		{
 			SCOPE_CYCLE_COUNTER(STAT_UpdateSingleUnionConstraint);
 
@@ -1099,7 +1099,7 @@ namespace Chaos
 
 		DECLARE_CYCLE_STAT(TEXT("TPBDCollisionConstraint::UpdateUnionLevelsetConstraint"), STAT_UpdateUnionLevelsetConstraint, STATGROUP_ChaosWide);
 		template<ECollisionUpdateType UpdateType, typename T, int d>
-		void UpdateUnionLevelsetConstraint(const T Thickness, TRigidBodySingleContactConstraint<T, d>& Constraint)
+		void UpdateUnionLevelsetConstraint(const T Thickness, TRigidBodyPointContactConstraint<T, d>& Constraint)
 		{
 			SCOPE_CYCLE_COUNTER(STAT_UpdateUnionLevelsetConstraint);
 
@@ -1175,7 +1175,7 @@ namespace Chaos
 
 		DECLARE_CYCLE_STAT(TEXT("TPBDCollisionConstraint::UpdateLevelsetUnionConstraint"), STAT_UpdateLevelsetUnionConstraint, STATGROUP_ChaosWide);
 		template<ECollisionUpdateType UpdateType, typename T, int d>
-		void UpdateLevelsetUnionConstraint(const T Thickness, TRigidBodySingleContactConstraint<T, d>& Constraint)
+		void UpdateLevelsetUnionConstraint(const T Thickness, TRigidBodyPointContactConstraint<T, d>& Constraint)
 		{
 			SCOPE_CYCLE_COUNTER(STAT_UpdateLevelsetUnionConstraint);
 
@@ -1223,9 +1223,9 @@ namespace Chaos
 
 
 		template <typename T, int d>
-		void UpdateBoxConstraint(const TBox<T, d>& Box1, const TRigidTransform<T, d>& Box1Transform, const TBox<T, d>& Box2, const TRigidTransform<T, d>& Box2Transform, T Thickness, TRigidBodySingleContactConstraint<T, d>& Constraint)
+		void UpdateBoxConstraint(const TBox<T, d>& Box1, const TRigidTransform<T, d>& Box1Transform, const TBox<T, d>& Box2, const TRigidTransform<T, d>& Box2Transform, T Thickness, TRigidBodyPointContactConstraint<T, d>& Constraint)
 		{
-			TSingleContactManifold<T, d> & Contact = Constraint.Manifold;
+			TPointContactManifold<T, d> & Contact = Constraint.Manifold;
 
 			TBox<T, d> Box2SpaceBox1 = Box1.TransformedBox(Box1Transform * Box2Transform.Inverse());
 			TBox<T, d> Box1SpaceBox2 = Box2.TransformedBox(Box2Transform * Box1Transform.Inverse());
@@ -1276,7 +1276,7 @@ namespace Chaos
 
 		DECLARE_CYCLE_STAT(TEXT("TPBDCollisionConstraint::UpdateLevelsetConstraint"), STAT_UpdateLevelsetConstraint, STATGROUP_ChaosWide);
 		template<ECollisionUpdateType UpdateType, typename T, int d>
-		void UpdateLevelsetConstraint(const T Thickness, TRigidBodySingleContactConstraint<float, 3>& Constraint)
+		void UpdateLevelsetConstraint(const T Thickness, TRigidBodyPointContactConstraint<float, 3>& Constraint)
 		{
 			SCOPE_CYCLE_COUNTER(STAT_UpdateLevelsetConstraint);
 
@@ -1303,13 +1303,13 @@ namespace Chaos
 			}
 		}
 
-		template void ConstructConstraintsImpl<float, 3>(TGeometryParticleHandle<float, 3>* Particle0, TGeometryParticleHandle<float, 3>* Particle1, const FImplicitObject* Implicit0, const FImplicitObject* Implicit1, const float Thickness, TRigidBodySingleContactConstraint<float, 3> & Constraint);
-		template void UpdateUnionLevelsetConstraint<ECollisionUpdateType::Any, float, 3>(const float Thickness, TRigidBodySingleContactConstraint<float, 3>& Constraint);
-		template void UpdateUnionLevelsetConstraint<ECollisionUpdateType::Deepest, float, 3>(const float Thickness, TRigidBodySingleContactConstraint<float, 3>& Constraint);
-		template void UpdateConstraintImp<ECollisionUpdateType::Any, float, 3>(const FImplicitObject& ParticleObject, const TRigidTransform<float, 3>& ParticleTM, const FImplicitObject& LevelsetObject, const TRigidTransform<float, 3>& LevelsetTM, const float Thickness, TRigidBodySingleContactConstraint<float, 3>& Constraints);
-		template void UpdateConstraintImp<ECollisionUpdateType::Deepest, float, 3>(const FImplicitObject& ParticleObject, const TRigidTransform<float, 3>& ParticleTM, const FImplicitObject& LevelsetObject, const TRigidTransform<float, 3>& LevelsetTM, const float Thickness, TRigidBodySingleContactConstraint<float, 3>& Constraint);
-		template void UpdateLevelsetConstraint<ECollisionUpdateType::Any, float, 3>(const float Thickness, TRigidBodySingleContactConstraint<float, 3>& Constraint);
-		template void UpdateLevelsetConstraint<ECollisionUpdateType::Deepest, float, 3>(const float Thickness, TRigidBodySingleContactConstraint<float, 3>& Constraint);
+		template void ConstructConstraintsImpl<float, 3>(TGeometryParticleHandle<float, 3>* Particle0, TGeometryParticleHandle<float, 3>* Particle1, const FImplicitObject* Implicit0, const FImplicitObject* Implicit1, const float Thickness, TRigidBodyPointContactConstraint<float, 3> & Constraint);
+		template void UpdateUnionLevelsetConstraint<ECollisionUpdateType::Any, float, 3>(const float Thickness, TRigidBodyPointContactConstraint<float, 3>& Constraint);
+		template void UpdateUnionLevelsetConstraint<ECollisionUpdateType::Deepest, float, 3>(const float Thickness, TRigidBodyPointContactConstraint<float, 3>& Constraint);
+		template void UpdateConstraintImp<ECollisionUpdateType::Any, float, 3>(const FImplicitObject& ParticleObject, const TRigidTransform<float, 3>& ParticleTM, const FImplicitObject& LevelsetObject, const TRigidTransform<float, 3>& LevelsetTM, const float Thickness, TRigidBodyPointContactConstraint<float, 3>& Constraints);
+		template void UpdateConstraintImp<ECollisionUpdateType::Deepest, float, 3>(const FImplicitObject& ParticleObject, const TRigidTransform<float, 3>& ParticleTM, const FImplicitObject& LevelsetObject, const TRigidTransform<float, 3>& LevelsetTM, const float Thickness, TRigidBodyPointContactConstraint<float, 3>& Constraint);
+		template void UpdateLevelsetConstraint<ECollisionUpdateType::Any, float, 3>(const float Thickness, TRigidBodyPointContactConstraint<float, 3>& Constraint);
+		template void UpdateLevelsetConstraint<ECollisionUpdateType::Deepest, float, 3>(const float Thickness, TRigidBodyPointContactConstraint<float, 3>& Constraint);
 
 	} // Collisions
 
diff --git a/Engine/Source/Runtime/Experimental/Chaos/Private/Chaos/CollisionResolutionConvexConvex.cpp b/Engine/Source/Runtime/Experimental/Chaos/Private/Chaos/CollisionResolutionConvexConvex.cpp
--- a/Engine/Source/Runtime/Experimental/Chaos/Private/Chaos/CollisionResolutionConvexConvex.cpp
+++ b/Engine/Source/Runtime/Experimental/Chaos/Private/Chaos/CollisionResolutionConvexConvex.cpp
@@ -12,7 +12,7 @@ namespace Chaos
 		FGeometryParticleHandle* Particle0, FGeometryParticleHandle* Particle1,
 		const FImplicitObject* Implicit0, const FImplicitObject* Implicit1,
 		const float Thickness,
-		FRigidBodyContactConstraint& Constraint)
+		FPointContactConstraint& Constraint)
 	{
 		if (Constraint.ContainsManifold(Implicit0, Implicit1))
 		{
@@ -29,7 +29,7 @@ namespace Chaos
 		const FImplicitObject& A, const FRigidTransform& ATM,
 		const FImplicitObject& B, FRigidTransform BTM,
 		float Thickness,
-		FRigidBodyContactConstraint& Constraint)
+		FPointContactConstraint& Constraint)
 	{
 		TRigidTransform<T, d> BToATM = BTM.GetRelativeTransform(ATM);
 
diff --git a/Engine/Source/Runtime/Experimental/Chaos/Private/Chaos/CollisionResolutionTypes.cpp b/Engine/Source/Runtime/Experimental/Chaos/Private/Chaos/CollisionResolutionTypes.cpp
--- a/Engine/Source/Runtime/Experimental/Chaos/Private/Chaos/CollisionResolutionTypes.cpp
+++ b/Engine/Source/Runtime/Experimental/Chaos/Private/Chaos/CollisionResolutionTypes.cpp
@@ -5,20 +5,5 @@
 
 namespace Chaos
 {
-
-	template<class T, int d> template<class TYPE>
-	TYPE& TPBDCollisionConstraintHandle<T, d>::GetContact()
-	{
-		return ConstraintContainer->Constraints[ConstraintIndex];
-	}
-
-	template<class T, int d> template<class TYPE>
-	const TYPE& TPBDCollisionConstraintHandle<T, d>::GetContact() const
-	{
-		return ConstraintContainer->Constraints[ConstraintIndex];
-	}
-
-	template class CHAOS_API TPBDCollisionConstraintHandle<float, 3>;
-	template CHAOS_API TRigidBodySingleContactConstraint<float, 3>& TPBDCollisionConstraintHandle<float, 3>::GetContact<TRigidBodySingleContactConstraint<float, 3>>();
-	template CHAOS_API const TRigidBodySingleContactConstraint<float, 3>& TPBDCollisionConstraintHandle<float, 3>::GetContact<TRigidBodySingleContactConstraint<float, 3>>() const;
+	template class CHAOS_API TPBDCollisionConstraintHandle<float, 3>; 
 };
diff --git a/Engine/Source/Runtime/Experimental/Chaos/Private/Chaos/PBDCollisionConstraint.cpp b/Engine/Source/Runtime/Experimental/Chaos/Private/Chaos/PBDCollisionConstraint.cpp
--- a/Engine/Source/Runtime/Experimental/Chaos/Private/Chaos/PBDCollisionConstraint.cpp
+++ b/Engine/Source/Runtime/Experimental/Chaos/Private/Chaos/PBDCollisionConstraint.cpp
@@ -121,13 +121,13 @@ namespace Chaos
 		typename TCollisionConstraintBase<T, d>::FType ConstraintType = Handle->GetType();
 
 		Handles.RemoveAtSwap(Idx);
-		Constraints.RemoveAtSwap(Idx);
-		if (Idx < Constraints.Num())
+		PointConstraints.RemoveAtSwap(Idx);
+		if (Idx < PointConstraints.Num())
 		{
 			Handles[Idx]->SetConstraintIndex(Idx, ConstraintType);
 		}
 
-		ensure(Handles.Num() == Constraints.Num());
+		ensure(Handles.Num() == PointConstraints.Num());
 		HandleAllocator.FreeHandle(Handle);
 	}
 
@@ -273,7 +273,7 @@ namespace Chaos
 
 
 	template<typename T, int d>
-	void TPBDCollisionConstraint<T, d>::Apply(const T Dt, FRigidBodyContactConstraint& Constraint, const int32 It, const int32 NumIts)
+	void TPBDCollisionConstraint<T, d>::Apply(const T Dt, FPointContactConstraint& Constraint, const int32 It, const int32 NumIts)
 	{
 
 		{
@@ -298,7 +298,7 @@ namespace Chaos
 			}
 		}
 
-		Collisions::TSingleContParticleParameters<T> ParticleParameters = { &MCollided, &MPhysicsMaterials, CollisionFrictionOverride, MAngularFriction };
+		Collisions::TSingleContactParticleParameters<T> ParticleParameters = { &MCollided, &MPhysicsMaterials, CollisionFrictionOverride, MAngularFriction };
 		Collisions::TSingleContactIterationParameters<T> IterationParameters = { Dt, It, NumIts, MApplyPairIterations, nullptr };
 		Collisions::Apply(Constraint,MThickness, IterationParameters,  ParticleParameters);
 	}
@@ -313,7 +313,7 @@ namespace Chaos
 			PhysicsParallelFor(InConstraintHandles.Num(), [&](int32 ConstraintHandleIndex) {
 				FConstraintContainerHandle* ConstraintHandle = InConstraintHandles[ConstraintHandleIndex];
 				check(ConstraintHandle != nullptr);
-				Apply(Dt, Constraints[ConstraintHandle->GetConstraintIndex()], It, NumIts);
+				Apply(Dt, PointConstraints[ConstraintHandle->GetConstraintIndex()], It, NumIts);
 			}, bDisableCollisionParallelFor);
 		}
 
@@ -325,7 +325,7 @@ namespace Chaos
 
 
 	template<typename T, int d>
-	void TPBDCollisionConstraint<T, d>::ApplyPushOut(const T Dt, FRigidBodyContactConstraint& Constraint, const TSet<const TGeometryParticleHandle<T, d>*>& IsTemporarilyStatic, int32 Iteration, int32 NumIterations, bool &NeedsAnotherIteration)
+	void TPBDCollisionConstraint<T, d>::ApplyPushOut(const T Dt, FPointContactConstraint& Constraint, const TSet<const TGeometryParticleHandle<T, d>*>& IsTemporarilyStatic, int32 Iteration, int32 NumIterations, bool &NeedsAnotherIteration)
 	{
 		{
 			// 
@@ -349,7 +349,7 @@ namespace Chaos
 			}
 		}
 
-		Collisions::TSingleContParticleParameters<T> ParticleParameters = { &MCollided, &MPhysicsMaterials, CollisionFrictionOverride, MAngularFriction };
+		Collisions::TSingleContactParticleParameters<T> ParticleParameters = { &MCollided, &MPhysicsMaterials, CollisionFrictionOverride, MAngularFriction };
 		Collisions::TSingleContactIterationParameters<T> IterationParameters = { Dt, Iteration, NumIterations, MApplyPairIterations, &NeedsAnotherIteration };
 		Collisions::ApplyPushOut(Constraint, MThickness, IsTemporarilyStatic, IterationParameters, ParticleParameters);
 	}
@@ -366,7 +366,7 @@ namespace Chaos
 			PhysicsParallelFor(InConstraintHandles.Num(), [&](int32 ConstraintHandleIndex) {
 				FConstraintContainerHandle* ConstraintHandle = InConstraintHandles[ConstraintHandleIndex];
 				check(ConstraintHandle != nullptr);
-				ApplyPushOut(Dt, Constraints[ConstraintHandle->GetConstraintIndex()], IsTemporarilyStatic, Iteration, NumIterations, NeedsAnotherIteration);
+				ApplyPushOut(Dt, PointConstraints[ConstraintHandle->GetConstraintIndex()], IsTemporarilyStatic, Iteration, NumIterations, NeedsAnotherIteration);
 			}, bDisableCollisionParallelFor);
 		}
 
@@ -386,7 +386,7 @@ namespace Chaos
 	}
 
 	template<typename T, int d>
-	void TPBDCollisionConstraint<T, d>::ConstructConstraints(TGeometryParticleHandle<T, d>* Particle0, TGeometryParticleHandle<T, d>* Particle1, const T Thickness, TRigidBodySingleContactConstraint<T, d> & Constraint)
+	void TPBDCollisionConstraint<T, d>::ConstructConstraints(TGeometryParticleHandle<T, d>* Particle0, TGeometryParticleHandle<T, d>* Particle1, const T Thickness, TRigidBodyPointContactConstraint<T, d> & Constraint)
 	{
 		if (ensure(Particle0 && Particle1))
 		{
@@ -398,7 +398,7 @@ namespace Chaos
 
 	template<typename T, int d>
 	template<ECollisionUpdateType UpdateType>
-	void TPBDCollisionConstraint<T, d>::UpdateConstraint(const T Thickness, FRigidBodyContactConstraint& Constraint)
+	void TPBDCollisionConstraint<T, d>::UpdateConstraint(const T Thickness, FPointContactConstraint& Constraint)
 	{
 		SCOPE_CYCLE_COUNTER(STAT_UpdateConstraint);
 		Collisions::UpdateConstraint<UpdateType>(Thickness, Constraint);
@@ -406,8 +406,8 @@ namespace Chaos
 
 	template class TAccelerationStructureHandle<float, 3>;
 	template class CHAOS_API TPBDCollisionConstraint<float, 3>;
-	template void TPBDCollisionConstraint<float, 3>::UpdateConstraint<ECollisionUpdateType::Any>(const float Thickness, FRigidBodyContactConstraint& Constraint);
-	template void TPBDCollisionConstraint<float, 3>::UpdateConstraint<ECollisionUpdateType::Deepest>(const float Thickness, FRigidBodyContactConstraint& Constraint);
+	template void TPBDCollisionConstraint<float, 3>::UpdateConstraint<ECollisionUpdateType::Any>(const float Thickness, FPointContactConstraint& Constraint);
+	template void TPBDCollisionConstraint<float, 3>::UpdateConstraint<ECollisionUpdateType::Deepest>(const float Thickness, FPointContactConstraint& Constraint);
 	template void TPBDCollisionConstraint<float, 3>::ComputeConstraints<false>(const TPBDCollisionConstraint<float, 3>::FAccelerationStructure&, float Dt);
 	template void TPBDCollisionConstraint<float, 3>::ComputeConstraints<true>(const TPBDCollisionConstraint<float, 3>::FAccelerationStructure&, float Dt);
 }
diff --git a/Engine/Source/Runtime/Experimental/Chaos/Private/Chaos/PBDCollisionConstraintPlaneContactUtil.cpp b/Engine/Source/Runtime/Experimental/Chaos/Private/Chaos/PBDCollisionConstraintPlaneContactUtil.cpp
new file mode 100644
--- /dev/null
+++ b/Engine/Source/Runtime/Experimental/Chaos/Private/Chaos/PBDCollisionConstraintPlaneContactUtil.cpp
@@ -0,0 +1,45 @@
+// Copyright 1998-2019 Epic Games, Inc. All Rights Reserved.
+
+#include "Chaos/PBDCollisionConstraintPlaneContactUtil.h"
+#include "Chaos/CollisionResolutionAlgo.h"
+#include "Chaos/Defines.h"
+
+namespace Chaos
+{
+	namespace Collisions
+	{
+		/*
+		template<ECollisionUpdateType UpdateType, typename T, int d>
+		void Update(const T Thickness, TRigidBodyPlaneContactConstraint<T, d>& Constraint)
+		{
+		}
+		template void Update<ECollisionUpdateType::Any, float, 3>(const float, TRigidBodyPlaneContactConstraint<float, 3>&);
+		template void Update<ECollisionUpdateType::Deepest, float, 3>(const float, TRigidBodyPlaneContactConstraint<float, 3>&);
+
+
+		template<ECollisionUpdateType UpdateType, typename T, int d>
+		void UpdateConstraint(const T Thickness, TRigidBodyPlaneContactConstraint<T, d>& Constraint)
+		{
+		}
+		template void UpdateConstraint<ECollisionUpdateType::Any, float, 3>(const float, TRigidBodyPointContactConstraint<float,3>&);
+		template void UpdateConstraint<ECollisionUpdateType::Deepest, float, 3>(const float, TRigidBodyPointContactConstraint<float,3>&);
+
+
+		template<typename T, int d>
+		void Apply(TRigidBodyPlaneContactConstraint<T, d>& Constraint, T Thickness, TSingleContactIterationParameters<T> & IterationParameters, TPlainContactParticleParameters<T> & ParticleParameters)
+		{
+		}
+		template void Apply<float, 3>(TRigidBodyPlaneContactConstraint<float, 3>&, float, TSingleContactIterationParameters<float> &, TSingleContactParticleParameters<float> &);
+
+		template<typename T, int d>
+		void ApplyPushOut(TRigidBodyPointContactConstraint<T, d>& Constraint, T Thickness, const TSet<const TGeometryParticleHandle<T, d>*>& IsTemporarilyStatic,
+			TSingleContactIterationParameters<T> & IterationParameters, TSingleContactParticleParameters<T> & ParticleParameters)
+		{
+		}
+		template void ApplyPushOut<float, 3>(TRigidBodyPointContactConstraint<float,3>&, float , const TSet<const TGeometryParticleHandle<float,3>*>&,
+			TSingleContactIterationParameters<float> & IterationParameters, TSingleContactParticleParameters<float> & ParticleParameters);
+		*/
+	} // Collisions
+
+}// Chaos
+
diff --git a/Engine/Source/Runtime/Experimental/Chaos/Private/Chaos/PBDCollisionConstraintSingleContactUtil.cpp b/Engine/Source/Runtime/Experimental/Chaos/Private/Chaos/PBDCollisionConstraintSingleContactUtil.cpp
--- a/Engine/Source/Runtime/Experimental/Chaos/Private/Chaos/PBDCollisionConstraintSingleContactUtil.cpp
+++ b/Engine/Source/Runtime/Experimental/Chaos/Private/Chaos/PBDCollisionConstraintSingleContactUtil.cpp
@@ -32,7 +32,7 @@ namespace Chaos
 		}
 
 		template<typename T, int d>
-		TVector<T, d> GetEnergyClampedImpulse(const TRigidBodySingleContactConstraint<T, d>& Constraint, const TVector<T, d>& Impulse, const TVector<T, d>& VectorToPoint1, const TVector<T, d>& VectorToPoint2, const TVector<T, d>& Velocity1, const TVector<T, d>& Velocity2)
+		TVector<T, d> GetEnergyClampedImpulse(const TRigidBodyPointContactConstraint<T, d>& Constraint, const TVector<T, d>& Impulse, const TVector<T, d>& VectorToPoint1, const TVector<T, d>& VectorToPoint2, const TVector<T, d>& Velocity1, const TVector<T, d>& Velocity2)
 		{
 			TPBDRigidParticleHandle<T, d>* PBDRigid0 = Constraint.Particle[0]->AsDynamic();
 			TPBDRigidParticleHandle<T, d>* PBDRigid1 = Constraint.Particle[1]->AsDynamic();
@@ -67,7 +67,7 @@ namespace Chaos
 
 
 		template<ECollisionUpdateType UpdateType, typename T, int d>
-		void Update(const T Thickness, TRigidBodySingleContactConstraint<T, d>& Constraint)
+		void Update(const T Thickness, TRigidBodyPointContactConstraint<T, d>& Constraint)
 		{
 			Constraint.ResetPhi(Thickness);
 			const TRigidTransform<T, d> ParticleTM = GetTransform(Constraint.Particle[0]);
@@ -92,12 +92,12 @@ namespace Chaos
 				UpdateConstraintImp<UpdateType>(*Constraint.Particle[0]->Geometry(), ParticleTM, *Constraint.Particle[1]->Geometry(), LevelsetTM, Thickness, Constraint);
 			}
 		}
-		template void Update<ECollisionUpdateType::Any, float, 3>(const float, TRigidBodySingleContactConstraint<float, 3>&);
-		template void Update<ECollisionUpdateType::Deepest, float, 3>(const float, TRigidBodySingleContactConstraint<float, 3>&);
+		template void Update<ECollisionUpdateType::Any, float, 3>(const float, TRigidBodyPointContactConstraint<float, 3>&);
+		template void Update<ECollisionUpdateType::Deepest, float, 3>(const float, TRigidBodyPointContactConstraint<float, 3>&);
 
 
 		template<ECollisionUpdateType UpdateType, typename T, int d>
-		void UpdateConstraint(const T Thickness, TRigidBodySingleContactConstraint<T, d>& Constraint)
+		void UpdateConstraint(const T Thickness, TRigidBodyPointContactConstraint<T, d>& Constraint)
 		{
 			Constraint.ResetPhi(Thickness);
 			const TRigidTransform<T, d> ParticleTM = GetTransform(Constraint.Particle[0]);
@@ -122,12 +122,12 @@ namespace Chaos
 				UpdateConstraintImp<UpdateType>(*Constraint.Particle[0]->Geometry(), ParticleTM, *Constraint.Particle[1]->Geometry(), LevelsetTM, Thickness, Constraint);
 			}
 		}
-		template void UpdateConstraint<ECollisionUpdateType::Any, float, 3>(const float, TRigidBodySingleContactConstraint<float,3>&);
-		template void UpdateConstraint<ECollisionUpdateType::Deepest, float, 3>(const float, TRigidBodySingleContactConstraint<float,3>&);
+		template void UpdateConstraint<ECollisionUpdateType::Any, float, 3>(const float, TRigidBodyPointContactConstraint<float,3>&);
+		template void UpdateConstraint<ECollisionUpdateType::Deepest, float, 3>(const float, TRigidBodyPointContactConstraint<float,3>&);
 
 
 		template<typename T, int d>
-		void Apply(TRigidBodySingleContactConstraint<T, d>& Constraint, T Thickness, TSingleContactIterationParameters<T> & IterationParameters, TSingleContParticleParameters<T> & ParticleParameters)
+		void Apply(TRigidBodyPointContactConstraint<T, d>& Constraint, T Thickness, TSingleContactIterationParameters<T> & IterationParameters, TSingleContactParticleParameters<T> & ParticleParameters)
 		{
 			TGenericParticleHandle<T, d> Particle0 = TGenericParticleHandle<T, d>(Constraint.Particle[0]);
 			TGenericParticleHandle<T, d> Particle1 = TGenericParticleHandle<T, d>(Constraint.Particle[1]);
@@ -167,7 +167,7 @@ namespace Chaos
 				const TVector<T, d>& W0 = Particle0->W();
 				const TVector<T, d>& W1 = Particle1->W();
 
-				typename TRigidBodySingleContactConstraint<T, d>::FManifold & Contact = Constraint.Manifold;
+				typename TRigidBodyPointContactConstraint<T, d>::FManifold & Contact = Constraint.Manifold;
 
 				TVector<T, d> VectorToPoint1 = Contact.Location - P0;
 				TVector<T, d> VectorToPoint2 = Contact.Location - P1;
@@ -345,11 +345,11 @@ namespace Chaos
 				}
 			}
 		}
-		template void Apply<float, 3>(TRigidBodySingleContactConstraint<float, 3>&, float, TSingleContactIterationParameters<float> &, TSingleContParticleParameters<float> &);
+		template void Apply<float, 3>(TRigidBodyPointContactConstraint<float, 3>&, float, TSingleContactIterationParameters<float> &, TSingleContactParticleParameters<float> &);
 
 		template<typename T, int d>
-		void ApplyPushOut(TRigidBodySingleContactConstraint<T, d>& Constraint, T Thickness, const TSet<const TGeometryParticleHandle<T, d>*>& IsTemporarilyStatic,
-			TSingleContactIterationParameters<T> & IterationParameters, TSingleContParticleParameters<T> & ParticleParameters)
+		void ApplyPushOut(TRigidBodyPointContactConstraint<T, d>& Constraint, T Thickness, const TSet<const TGeometryParticleHandle<T, d>*>& IsTemporarilyStatic,
+			TSingleContactIterationParameters<T> & IterationParameters, TSingleContactParticleParameters<T> & ParticleParameters)
 		{
 			TGeometryParticleHandle<T, d>* Particle0 = Constraint.Particle[0];
 			TGeometryParticleHandle<T, d>* Particle1 = Constraint.Particle[1];
@@ -379,7 +379,7 @@ namespace Chaos
 			{
 				UpdateConstraint<ECollisionUpdateType::Deepest>(Thickness, Constraint);
 
-				const typename TRigidBodySingleContactConstraint<T, d>::FManifold & Contact = Constraint.Manifold;
+				const typename TRigidBodyPointContactConstraint<T, d>::FManifold & Contact = Constraint.Manifold;
 
 				if (Contact.Phi >= Thickness)
 				{
@@ -458,8 +458,8 @@ namespace Chaos
 				}
 			}
 		}
-		template void ApplyPushOut<float, 3>(TRigidBodySingleContactConstraint<float,3>&, float , const TSet<const TGeometryParticleHandle<float,3>*>&,
-			TSingleContactIterationParameters<float> & IterationParameters, TSingleContParticleParameters<float> & ParticleParameters);
+		template void ApplyPushOut<float, 3>(TRigidBodyPointContactConstraint<float,3>&, float , const TSet<const TGeometryParticleHandle<float,3>*>&,
+			TSingleContactIterationParameters<float> & IterationParameters, TSingleContactParticleParameters<float> & ParticleParameters);
 
 	} // Collisions
 
diff --git a/Engine/Source/Runtime/Experimental/Chaos/Private/Chaos/PBDCollisionConstraintUtil.cpp b/Engine/Source/Runtime/Experimental/Chaos/Private/Chaos/PBDCollisionConstraintUtil.cpp
--- a/Engine/Source/Runtime/Experimental/Chaos/Private/Chaos/PBDCollisionConstraintUtil.cpp
+++ b/Engine/Source/Runtime/Experimental/Chaos/Private/Chaos/PBDCollisionConstraintUtil.cpp
@@ -15,7 +15,7 @@
 
 namespace Chaos
 {
-	void ComputeHashTable(const TArray<Chaos::TPBDCollisionConstraint<float, 3>::FRigidBodyContactConstraint>& ConstraintsArray,
+	void ComputeHashTable(const TArray<Chaos::TPBDCollisionConstraint<float, 3>::FPointContactConstraint>& ConstraintsArray,
 						  const FBox& BoundingBox, TMultiMap<int32, int32>& HashTableMap, const float SpatialHashRadius)
 	{
 		float CellSize = 2.f * SpatialHashRadius;
diff --git a/Engine/Source/Runtime/Experimental/Chaos/Private/Chaos/PBDRigidClustering.cpp b/Engine/Source/Runtime/Experimental/Chaos/Private/Chaos/PBDRigidClustering.cpp
--- a/Engine/Source/Runtime/Experimental/Chaos/Private/Chaos/PBDRigidClustering.cpp
+++ b/Engine/Source/Runtime/Experimental/Chaos/Private/Chaos/PBDRigidClustering.cpp
@@ -1757,7 +1757,7 @@ namespace Chaos
 	}
 
 	template<class T, int d>
-	TVector<T, d> GetContactLocation(const TRigidBodySingleContactConstraint<T, d>& Contact)
+	TVector<T, d> GetContactLocation(const TRigidBodyPointContactConstraint<T, d>& Contact)
 	{
 		return Contact.GetLocation();
 	}
diff --git a/Engine/Source/Runtime/Experimental/Chaos/Public/Chaos/CollisionResolutionAlgo.h b/Engine/Source/Runtime/Experimental/Chaos/Public/Chaos/CollisionResolutionAlgo.h
--- a/Engine/Source/Runtime/Experimental/Chaos/Public/Chaos/CollisionResolutionAlgo.h
+++ b/Engine/Source/Runtime/Experimental/Chaos/Public/Chaos/CollisionResolutionAlgo.h
@@ -19,7 +19,7 @@ namespace Chaos
 	class FImplicitObject;
 
 	template<typename T, int d>
-	class TRigidBodySingleContactConstraint;
+	class TRigidBodyPointContactConstraint;
 
 	template <typename T, int d>
 	class TRigidTransform;
@@ -39,83 +39,83 @@ namespace Chaos
 		// Construct Constraints
 		//
 		template<typename T, int d>
-		void CHAOS_API ConstructLevelsetConstraints(TGeometryParticleHandle<T, d>* Particle0, TGeometryParticleHandle<T, d>* Particle1, const FImplicitObject* Implicit0, const FImplicitObject* Implicit1, const T Thickness, TRigidBodySingleContactConstraint<T, d>& Constraint);
+		void CHAOS_API ConstructLevelsetConstraints(TGeometryParticleHandle<T, d>* Particle0, TGeometryParticleHandle<T, d>* Particle1, const FImplicitObject* Implicit0, const FImplicitObject* Implicit1, const T Thickness, TRigidBodyPointContactConstraint<T, d>& Constraint);
 
 		template<typename T, int d>
-		void CHAOS_API ConstructBoxConstraints(TGeometryParticleHandle<T, d>* Particle0, TGeometryParticleHandle<T, d>* Particle1, const FImplicitObject* Implicit0, const FImplicitObject* Implicit1, const T Thickness, TRigidBodySingleContactConstraint<T, d>& Constraint);
+		void CHAOS_API ConstructBoxConstraints(TGeometryParticleHandle<T, d>* Particle0, TGeometryParticleHandle<T, d>* Particle1, const FImplicitObject* Implicit0, const FImplicitObject* Implicit1, const T Thickness, TRigidBodyPointContactConstraint<T, d>& Constraint);
 
 		template<typename T, int d>
-		void CHAOS_API ConstructBoxPlaneConstraints(TGeometryParticleHandle<T, d>* Particle0, TGeometryParticleHandle<T, d>* Particle1, const FImplicitObject* Implicit0, const FImplicitObject* Implicit1, const T Thickness, TRigidBodySingleContactConstraint<T, d>& Constraint);
+		void CHAOS_API ConstructBoxPlaneConstraints(TGeometryParticleHandle<T, d>* Particle0, TGeometryParticleHandle<T, d>* Particle1, const FImplicitObject* Implicit0, const FImplicitObject* Implicit1, const T Thickness, TRigidBodyPointContactConstraint<T, d>& Constraint);
 
 		template<typename T, int d>
-		void CHAOS_API ConstructSphereConstraints(TGeometryParticleHandle<T, d>* Particle0, TGeometryParticleHandle<T, d>* Particle1, const FImplicitObject* Implicit0, const FImplicitObject* Implicit1, const T Thickness, TRigidBodySingleContactConstraint<T, d>& Constraint);
+		void CHAOS_API ConstructSphereConstraints(TGeometryParticleHandle<T, d>* Particle0, TGeometryParticleHandle<T, d>* Particle1, const FImplicitObject* Implicit0, const FImplicitObject* Implicit1, const T Thickness, TRigidBodyPointContactConstraint<T, d>& Constraint);
 
 		template<typename T, int d>
-		void CHAOS_API ConstructSpherePlaneConstraints(TGeometryParticleHandle<T, d>* Particle0, TGeometryParticleHandle<T, d>* Particle1, const FImplicitObject* Implicit0, const FImplicitObject* Implicit1, const T Thickness, TRigidBodySingleContactConstraint<T, d>& Constraint);
+		void CHAOS_API ConstructSpherePlaneConstraints(TGeometryParticleHandle<T, d>* Particle0, TGeometryParticleHandle<T, d>* Particle1, const FImplicitObject* Implicit0, const FImplicitObject* Implicit1, const T Thickness, TRigidBodyPointContactConstraint<T, d>& Constraint);
 
 		template<typename T, int d>
-		void CHAOS_API ConstructSphereBoxConstraints(TGeometryParticleHandle<T, d>* Particle0, TGeometryParticleHandle<T, d>* Particle1, const FImplicitObject* Implicit0, const FImplicitObject* Implicit1, const T Thickness, TRigidBodySingleContactConstraint<T, d>& Constraint);
+		void CHAOS_API ConstructSphereBoxConstraints(TGeometryParticleHandle<T, d>* Particle0, TGeometryParticleHandle<T, d>* Particle1, const FImplicitObject* Implicit0, const FImplicitObject* Implicit1, const T Thickness, TRigidBodyPointContactConstraint<T, d>& Constraint);
 
 		template<typename T, int d>
-		void CHAOS_API ConstructCapsuleCapsuleConstraints(TGeometryParticleHandle<T, d>* Particle0, TGeometryParticleHandle<T, d>* Particle1, const FImplicitObject* Implicit0, const FImplicitObject* Implicit1, const T Thickness, TRigidBodySingleContactConstraint<T, d>& Constraint);
+		void CHAOS_API ConstructCapsuleCapsuleConstraints(TGeometryParticleHandle<T, d>* Particle0, TGeometryParticleHandle<T, d>* Particle1, const FImplicitObject* Implicit0, const FImplicitObject* Implicit1, const T Thickness, TRigidBodyPointContactConstraint<T, d>& Constraint);
 
 		template<typename T, int d>
-		void CHAOS_API ConstructCapsuleBoxConstraints(TGeometryParticleHandle<T, d>* Particle0, TGeometryParticleHandle<T, d>* Particle1, const FImplicitObject* Implicit0, const FImplicitObject* Implicit1, const T Thickness, TRigidBodySingleContactConstraint<T, d>& Constraint);
+		void CHAOS_API ConstructCapsuleBoxConstraints(TGeometryParticleHandle<T, d>* Particle0, TGeometryParticleHandle<T, d>* Particle1, const FImplicitObject* Implicit0, const FImplicitObject* Implicit1, const T Thickness, TRigidBodyPointContactConstraint<T, d>& Constraint);
 
 		template<typename T, int d>
-		void CHAOS_API ConstructSingleUnionConstraints(TGeometryParticleHandle<T, d>* Particle0, TGeometryParticleHandle<T, d>* Particle1, const FImplicitObject* Implicit0, const FImplicitObject* Implicit1, const T Thickness, TRigidBodySingleContactConstraint<T, d>& Constraint);
+		void CHAOS_API ConstructSingleUnionConstraints(TGeometryParticleHandle<T, d>* Particle0, TGeometryParticleHandle<T, d>* Particle1, const FImplicitObject* Implicit0, const FImplicitObject* Implicit1, const T Thickness, TRigidBodyPointContactConstraint<T, d>& Constraint);
 
 		template<typename T, int d>
-		void CHAOS_API ConstructUnionUnionConstraints(TGeometryParticleHandle<T, d>* Particle0, TGeometryParticleHandle<T, d>* Particle1, const FImplicitObject* Implicit0, const FImplicitObject* Implicit1, const T Thickness, TRigidBodySingleContactConstraint<T, d>& Constraint);
+		void CHAOS_API ConstructUnionUnionConstraints(TGeometryParticleHandle<T, d>* Particle0, TGeometryParticleHandle<T, d>* Particle1, const FImplicitObject* Implicit0, const FImplicitObject* Implicit1, const T Thickness, TRigidBodyPointContactConstraint<T, d>& Constraint);
 
 		template<typename T, int d>
-		void CHAOS_API ConstructPairConstraintImpl(TGeometryParticleHandle<T, d>* Particle0, TGeometryParticleHandle<T, d>* Particle1, const FImplicitObject* Implicit0, const FImplicitObject* Implicit1, const T Thickness, TRigidBodySingleContactConstraint<T, d>& Constraint);
+		void CHAOS_API ConstructPairConstraintImpl(TGeometryParticleHandle<T, d>* Particle0, TGeometryParticleHandle<T, d>* Particle1, const FImplicitObject* Implicit0, const FImplicitObject* Implicit1, const T Thickness, TRigidBodyPointContactConstraint<T, d>& Constraint);
 
 		template<typename T, int d>
-		void CHAOS_API ConstructConstraintsImpl(TGeometryParticleHandle<T, d>* Particle0, TGeometryParticleHandle<T, d>* Particle1, const FImplicitObject* Implicit0, const FImplicitObject* Implicit1, const T Thickness, TRigidBodySingleContactConstraint<T, d> & Constraint);
+		void CHAOS_API ConstructConstraintsImpl(TGeometryParticleHandle<T, d>* Particle0, TGeometryParticleHandle<T, d>* Particle1, const FImplicitObject* Implicit0, const FImplicitObject* Implicit1, const T Thickness, TRigidBodyPointContactConstraint<T, d> & Constraint);
 
 
 		//
 		// Update Constraints
 		//
 		template<ECollisionUpdateType UpdateType, typename T = float, int d = 3>
-		void CHAOS_API UpdateLevelsetConstraint(const T Thickness, TRigidBodySingleContactConstraint<float, 3>& Constraint);
+		void CHAOS_API UpdateLevelsetConstraint(const T Thickness, TRigidBodyPointContactConstraint<float, 3>& Constraint);
 
 		template<ECollisionUpdateType UpdateType, typename T, int d>
-		void CHAOS_API UpdateConstraintImp(const FImplicitObject& ParticleObject, const TRigidTransform<T, d>& ParticleTM, const FImplicitObject& LevelsetObject, const TRigidTransform<T, d>& LevelsetTM, const T Thickness, TRigidBodySingleContactConstraint<T, d>& Constraint);
+		void CHAOS_API UpdateConstraintImp(const FImplicitObject& ParticleObject, const TRigidTransform<T, d>& ParticleTM, const FImplicitObject& LevelsetObject, const TRigidTransform<T, d>& LevelsetTM, const T Thickness, TRigidBodyPointContactConstraint<T, d>& Constraint);
 
 		template <typename T, int d>
-		bool CHAOS_API UpdateBoxPlaneConstraint(const TBox<T, d>& Box, const TRigidTransform<T, d>& BoxTransform, const TPlane<T, d>& Plane, const TRigidTransform<T, d>& PlaneTransform, T Thickness, TRigidBodySingleContactConstraint<T, d>& Constraint);
+		bool CHAOS_API UpdateBoxPlaneConstraint(const TBox<T, d>& Box, const TRigidTransform<T, d>& BoxTransform, const TPlane<T, d>& Plane, const TRigidTransform<T, d>& PlaneTransform, T Thickness, TRigidBodyPointContactConstraint<T, d>& Constraint);
 
 		template <typename T, int d>
-		void CHAOS_API UpdateSphereConstraint(const TSphere<T, d>& Sphere1, const TRigidTransform<T, d>& Sphere1Transform, const TSphere<T, d>& Sphere2, const TRigidTransform<T, d>& Sphere2Transform, T Thickness, TRigidBodySingleContactConstraint<T, d>& Constraint);
+		void CHAOS_API UpdateSphereConstraint(const TSphere<T, d>& Sphere1, const TRigidTransform<T, d>& Sphere1Transform, const TSphere<T, d>& Sphere2, const TRigidTransform<T, d>& Sphere2Transform, T Thickness, TRigidBodyPointContactConstraint<T, d>& Constraint);
 
 		template <typename T, int d>
-		void CHAOS_API UpdateSpherePlaneConstraint(const TSphere<T, d>& Sphere, const TRigidTransform<T, d>& SphereTransform, const TPlane<T, d>& Plane, const TRigidTransform<T, d>& PlaneTransform, T Thickness, TRigidBodySingleContactConstraint<T, d>& Constraint);
+		void CHAOS_API UpdateSpherePlaneConstraint(const TSphere<T, d>& Sphere, const TRigidTransform<T, d>& SphereTransform, const TPlane<T, d>& Plane, const TRigidTransform<T, d>& PlaneTransform, T Thickness, TRigidBodyPointContactConstraint<T, d>& Constraint);
 
 		template <typename T, int d>
-		void CHAOS_API UpdateSphereBoxConstraint(const TSphere<T, d>& Sphere, const TRigidTransform<T, d>& SphereTransform, const TBox<T, d>& Box, const TRigidTransform<T, d>& BoxTransform, T Thickness, TRigidBodySingleContactConstraint<T, d>& Constraint);
+		void CHAOS_API UpdateSphereBoxConstraint(const TSphere<T, d>& Sphere, const TRigidTransform<T, d>& SphereTransform, const TBox<T, d>& Box, const TRigidTransform<T, d>& BoxTransform, T Thickness, TRigidBodyPointContactConstraint<T, d>& Constraint);
 
 		template <typename T, int d>
-		void CHAOS_API UpdateCapsuleCapsuleConstraint(const TCapsule<T>& A, const TRigidTransform<T, d>& ATransform, const TCapsule<T>& B, const TRigidTransform<T, d>& BTransform, T Thickness, TRigidBodySingleContactConstraint<T, d>& Constraint);
+		void CHAOS_API UpdateCapsuleCapsuleConstraint(const TCapsule<T>& A, const TRigidTransform<T, d>& ATransform, const TCapsule<T>& B, const TRigidTransform<T, d>& BTransform, T Thickness, TRigidBodyPointContactConstraint<T, d>& Constraint);
 
 		template <typename T, int d>
-		void CHAOS_API UpdateCapsuleBoxConstraint(const TCapsule<T>& A, const TRigidTransform<T, d>& ATransform, const TBox<T, d>& B, const TRigidTransform<T, d>& BTransform, T Thickness, TRigidBodySingleContactConstraint<T, d>& Constraint);
+		void CHAOS_API UpdateCapsuleBoxConstraint(const TCapsule<T>& A, const TRigidTransform<T, d>& ATransform, const TBox<T, d>& B, const TRigidTransform<T, d>& BTransform, T Thickness, TRigidBodyPointContactConstraint<T, d>& Constraint);
 
 		template<ECollisionUpdateType UpdateType, typename T, int d>
-		void CHAOS_API UpdateUnionUnionConstraint(const T Thickness, TRigidBodySingleContactConstraint<T, d>& Constraint);
+		void CHAOS_API UpdateUnionUnionConstraint(const T Thickness, TRigidBodyPointContactConstraint<T, d>& Constraint);
 
 		template<ECollisionUpdateType UpdateType, typename T, int d>
-		void CHAOS_API UpdateSingleUnionConstraint(const T Thickness, TRigidBodySingleContactConstraint<T, d>& Constraint);
+		void CHAOS_API UpdateSingleUnionConstraint(const T Thickness, TRigidBodyPointContactConstraint<T, d>& Constraint);
 
 		template<ECollisionUpdateType UpdateType, typename T, int d>
-		void CHAOS_API UpdateUnionLevelsetConstraint(const T Thickness, TRigidBodySingleContactConstraint<T, d>& Constraint);
+		void CHAOS_API UpdateUnionLevelsetConstraint(const T Thickness, TRigidBodyPointContactConstraint<T, d>& Constraint);
 
 		template<ECollisionUpdateType UpdateType, typename T, int d>
-		void CHAOS_API UpdateLevelsetUnionConstraint(const T Thickness, TRigidBodySingleContactConstraint<T, d>& Constraint);
+		void CHAOS_API UpdateLevelsetUnionConstraint(const T Thickness, TRigidBodyPointContactConstraint<T, d>& Constraint);
 
 		template <typename T, int d>
-		void CHAOS_API UpdateBoxConstraint(const TBox<T, d>& Box1, const TRigidTransform<T, d>& Box1Transform, const TBox<T, d>& Box2, const TRigidTransform<T, d>& Box2Transform, T Thickness, TRigidBodySingleContactConstraint<T, d>& Constraint);
+		void CHAOS_API UpdateBoxConstraint(const TBox<T, d>& Box1, const TRigidTransform<T, d>& Box1Transform, const TBox<T, d>& Box2, const TRigidTransform<T, d>& Box2Transform, T Thickness, TRigidBodyPointContactConstraint<T, d>& Constraint);
 
 	}// Collisions
 
diff --git a/Engine/Source/Runtime/Experimental/Chaos/Public/Chaos/CollisionResolutionConvexConvex.h b/Engine/Source/Runtime/Experimental/Chaos/Public/Chaos/CollisionResolutionConvexConvex.h
--- a/Engine/Source/Runtime/Experimental/Chaos/Public/Chaos/CollisionResolutionConvexConvex.h
+++ b/Engine/Source/Runtime/Experimental/Chaos/Public/Chaos/CollisionResolutionConvexConvex.h
@@ -17,20 +17,20 @@ class CHAOS_API CollisionResolutionConvexConvex
 public:
 	using FConvex = TConvex<T,d>;
 	using FRigidTransform = TRigidTransform<T,d>;
-	using FRigidBodyContactConstraint = TRigidBodySingleContactConstraint<T, d>;
+	using FPointContactConstraint = TRigidBodyPointContactConstraint<T, d>;
 	using FGeometryParticleHandle = TGeometryParticleHandle<T, d>;
 
 	/*
 	* Build a constraint manifold for the contact shape pair. 
 	*/
 	static void ConstructConvexConvexConstraints(FGeometryParticleHandle* Particle0, FGeometryParticleHandle* Particle1, const FImplicitObject* Implicit0, const FImplicitObject* Implicit1, const float Thickness,
-		FRigidBodyContactConstraint& ConstraintBuffer);
+		FPointContactConstraint& ConstraintBuffer);
 
 	/*
 	* Update the contact manifold.
 	*/
 	static void UpdateConvexConvexConstraint(const FImplicitObject& A, const FRigidTransform& ATM, const FImplicitObject& B, FRigidTransform BTM, float Thickness, 
-		FRigidBodyContactConstraint& Constraints);
+		FPointContactConstraint& Constraints);
 
 
 };
diff --git a/Engine/Source/Runtime/Experimental/Chaos/Public/Chaos/CollisionResolutionTypes.h b/Engine/Source/Runtime/Experimental/Chaos/Public/Chaos/CollisionResolutionTypes.h
--- a/Engine/Source/Runtime/Experimental/Chaos/Public/Chaos/CollisionResolutionTypes.h
+++ b/Engine/Source/Runtime/Experimental/Chaos/Public/Chaos/CollisionResolutionTypes.h
@@ -18,7 +18,7 @@ namespace Chaos
 	class TPBDCollisionConstraintHandle;
 
 	template<typename T, int d>
-	class TRigidBodySingleContactConstraint;
+	class TRigidBodyPointContactConstraint;
 
 	/** Specifies the type of work we should do*/
 	enum class CHAOS_API ECollisionUpdateType
@@ -36,30 +36,10 @@ namespace Chaos
 	};
 
 
+
 	/*
 	*
 	*/
-	template<class T, int d>
-	class CHAOS_API TSingleContactManifold
-	{
-	public:
-		TSingleContactManifold(int32 InTimestamp = -INT_MAX, const FImplicitObject* InImplicit0 = nullptr, const FImplicitObject* InImplicit1 = nullptr)
-			: bDisabled(true), Timestamp(InTimestamp), Normal(0), Phi(FLT_MAX)
-		{
-			Implicit[0] = InImplicit0;
-			Implicit[1] = InImplicit1;
-		}
-
-		bool bDisabled;
-		int32 Timestamp;
-		TVector<T, d> Normal;
-		TVector<T, d> Location;
-		T Phi;
-
-		const FImplicitObject* Implicit[2]; // {Of Particle[0], Of Particle[1]}
-	};
-
-
 	template<class T, int d>
 	class CHAOS_API TCollisionConstraintBase
 	{
@@ -68,66 +48,158 @@ namespace Chaos
 		enum class FType
 		{
 			None = 0,     // default value also indicates a invalid constraint
-			SinglePoint   // TRigidBodySingleContactConstraint
+			SinglePoint,   // TRigidBodyPointContactConstraint
+			Plane         // TRigidBodyPlaneContactConstraint
 		};
 
 		TCollisionConstraintBase() : Type(FType::None) {}
 		TCollisionConstraintBase(FType InType) : Type(InType) {}
 		FType GetType() const { return Type; }
 
-		template<class AS_T> AS_T * As() {return static_cast<AS_T>(this); }
+		template<class AS_T> AS_T * As() { return static_cast<AS_T>(this); }
 		template<class AS_T> const AS_T * As() const { return static_cast<AS_T>(this); }
 
 	private:
 		FType Type;
 	};
 
+
 	/*
 	*
 	*/
 	template<class T, int d>
-	class CHAOS_API TRigidBodySingleContactConstraint : public TCollisionConstraintBase<T,d>
+	class CHAOS_API TPointContactManifold
+	{
+	public:
+		TPointContactManifold(int32 InTimestamp = -INT_MAX, const FImplicitObject* InImplicit0 = nullptr, const FImplicitObject* InImplicit1 = nullptr)
+			: bDisabled(true), Timestamp(InTimestamp), Normal(0), Phi(FLT_MAX)
+		{
+			Implicit[0] = InImplicit0;
+			Implicit[1] = InImplicit1;
+		}
+
+		bool bDisabled;
+		int32 Timestamp;
+		TVector<T, d> Normal;
+		TVector<T, d> Location;
+		T Phi;
+
+		const FImplicitObject* Implicit[2]; // {Of Particle[0], Of Particle[1]}
+	};
+
+	/*
+	*
+	*/
+	template<class T, int d>
+	class CHAOS_API TRigidBodyPointContactConstraint : public TCollisionConstraintBase<T,d>
 	{
 	public:
 		using Base = TCollisionConstraintBase<T, d>;
 		using FGeometryParticleHandle = TGeometryParticleHandle<T, d>;
-		using FManifold = TSingleContactManifold<T, d>;
+		using FManifold = TPointContactManifold<T, d>;
+
+		TRigidBodyPointContactConstraint() : AccumulatedImpulse(0) {}
+		static typename Base::FType StaticType() { return Base::FType::SinglePoint; };
+
+
+		//
+		// Manidfold management
+		//
+		bool ContainsManifold(const FImplicitObject* B = nullptr, const FImplicitObject* A = nullptr) const { return (A == nullptr&&B == nullptr) ? (Manifold.Implicit[0] != nullptr && Manifold.Implicit[1] != nullptr) : (Manifold.Implicit[0] == A && Manifold.Implicit[1] == B); }
+		void AddManifold(const FImplicitObject* B = nullptr, const FImplicitObject* A = nullptr) { Manifold.Implicit[0] = A; Manifold.Implicit[1] = B; }
 
-		TRigidBodySingleContactConstraint() : AccumulatedImpulse(0) {}
 
-		//API
+		//
+		// API
+		//
 		void ResetPhi(T InPhi) { SetPhi(InPhi); }
-		bool ContainsManifold(const FImplicitObject* B = nullptr, const FImplicitObject* A = nullptr) const { return (A==nullptr&&B==nullptr)?(Manifold.Implicit[0] != nullptr && Manifold.Implicit[1] != nullptr):(Manifold.Implicit[0]==A && Manifold.Implicit[1]==B); }
+
+		void SetPhi(T InPhi, const FImplicitObject* B = nullptr, const FImplicitObject* A = nullptr, int32 Index = 0) { Manifold.Phi = InPhi; }
+		T GetPhi(const FImplicitObject* B = nullptr, const FImplicitObject* A = nullptr) const { return Manifold.Phi; }
 
 		void SetDisabled(bool bInDisabled, const FImplicitObject* B = nullptr, const FImplicitObject* A = nullptr) { Manifold.bDisabled = bInDisabled; }
 		TVector<T, d> GetDisabled(const FImplicitObject* B = nullptr, const FImplicitObject* A = nullptr) const { return Manifold.bDisabled; }
 
-
 		void SetNormal(const TVector<T, d> & InNormal, const FImplicitObject* B = nullptr, const FImplicitObject* A = nullptr) { Manifold.Normal = InNormal; }
 		TVector<T, d> GetNormal(const FImplicitObject* B = nullptr, const FImplicitObject* A = nullptr) const { return Manifold.Normal; }
 
 		void SetLocation(const TVector<T, d> & InLocation, const FImplicitObject* B = nullptr, const FImplicitObject* A = nullptr, int32 Index=0) { Manifold.Location = InLocation; }
 		TVector<T, d> GetLocation(const FImplicitObject* B = nullptr, const FImplicitObject* A = nullptr) const { return Manifold.Location; }
 
-		void SetPhi(T InPhi, const FImplicitObject* B = nullptr, const FImplicitObject* A = nullptr, int32 Index = 0) { Manifold.Phi = InPhi; }
-		T GetPhi(const FImplicitObject* B = nullptr, const FImplicitObject* A = nullptr) const { return Manifold.Phi; }
+		FString ToString() const
+		{
+			return FString::Printf(TEXT("Particle:%s, Levelset:%s, AccumulatedImpulse:%s"), *Particle[0]->ToString(), *Particle[1]->ToString(), *AccumulatedImpulse.ToString());
+		}
 
-		void AddManifold(const FImplicitObject* B = nullptr, const FImplicitObject* A = nullptr) { Manifold.Implicit[0] = A; Manifold.Implicit[1] = B; }
 
 		FManifold Manifold;
 		FGeometryParticleHandle* Particle[2]; // { Point, Volume } 
 		TVector<T, d> AccumulatedImpulse;
 
+	};
+
+
+	/*
+	*
+	*/
+	template<class T, int d>
+	class CHAOS_API TPlaneContactManifold
+	{
+	public:
+		TPlaneContactManifold(int32 InTimestamp = -INT_MAX, const FImplicitObject* InImplicit0 = nullptr, const FImplicitObject* InImplicit1 = nullptr)
+			: bDisabled(true), Timestamp(InTimestamp), Phi(FLT_MAX)
+		{
+			Implicit[0] = InImplicit0;
+			Implicit[1] = InImplicit1;
+		}
+
+		bool bDisabled;
+		int32 Timestamp;
+
+		TPlane<T, d> Plane;
+		TVector<T, d> ManifoldSamples[4]; // iterative plane samples
+
+		TVector<T, d> Location;
+		T Phi;
+
+		const FImplicitObject* Implicit[2]; // {Of Particle[0], Of Particle[1]}
+	};
+
+
+
+	/*
+	*
+	*/
+	template<class T, int d>
+	class CHAOS_API TRigidBodyPlaneContactConstraint : public TCollisionConstraintBase<T, d>
+	{
+	public:
+		using Base = TCollisionConstraintBase<T, d>;
+		using FGeometryParticleHandle = TGeometryParticleHandle<T, d>;
+		using FManifold = TPlaneContactManifold<T, d>;
+
+		TRigidBodyPlaneContactConstraint() : AccumulatedImpulse(0) {}
+		static typename Base::FType StaticType() { return Base::FType::Plane; };
+
+	
+		bool ContainsManifold(const FImplicitObject* B = nullptr, const FImplicitObject* A = nullptr) const { return (A == nullptr&&B == nullptr) ? (Manifold.Implicit[0] != nullptr && Manifold.Implicit[1] != nullptr) : (Manifold.Implicit[0] == A && Manifold.Implicit[1] == B); }
+		void AddManifold(const FImplicitObject* B = nullptr, const FImplicitObject* A = nullptr) { Manifold.Implicit[0] = A; Manifold.Implicit[1] = B; }
+
+
 		FString ToString() const
 		{
 			return FString::Printf(TEXT("Particle:%s, Levelset:%s, AccumulatedImpulse:%s"), *Particle[0]->ToString(), *Particle[1]->ToString(), *AccumulatedImpulse.ToString());
 		}
 
-		static typename Base::FType StaticType() { return Base::FType::SinglePoint; };
+
+		FManifold Manifold;
+		FGeometryParticleHandle* Particle[2]; // { Point, Volume } 
+		TVector<T, d> AccumulatedImpulse;
 
 	};
 
 
+
 	template<class T, int d>
 	struct CHAOS_API TRigidBodyContactConstraintPGS
 	{
@@ -162,8 +234,12 @@ namespace Chaos
 
 		// Handle API
 
-		template<class AS_T> AS_T& GetContact();
-		template<class AS_T> const AS_T& GetContact() const;
+		const TRigidBodyPointContactConstraint<T, d>& GetSingleContact() const { return ConstraintContainer->PointConstraints[ConstraintIndex]; }
+		TRigidBodyPointContactConstraint<T, d>& GetSingleContact() { return ConstraintContainer->PointConstraints[ConstraintIndex]; }
+
+		const TRigidBodyPlaneContactConstraint<T, d>& GetPlaneContact() const { return ConstraintContainer->PlaneConstraints[ConstraintIndex]; }
+		TRigidBodyPlaneContactConstraint<T, d>& GetPlaneContact() { return ConstraintContainer->PlaneConstraints[ConstraintIndex]; }
+
 
 		typename FConstraintBase::FType GetType() const { return ConstraintType; }
 
@@ -179,7 +255,7 @@ namespace Chaos
 		{
 			if (ConstraintType == FConstraintBase::FType::SinglePoint)
 			{
-				return GetContact<TRigidBodySingleContactConstraint<float, 3> >().GetLocation();
+				return GetSingleContact().GetLocation();
 			}
 			ensure(false);
 			return TVector<T, d>(0);
@@ -189,7 +265,7 @@ namespace Chaos
 		{
 			if (ConstraintType == FConstraintBase::FType::SinglePoint)
 			{
-				return GetContact<TRigidBodySingleContactConstraint<float, 3> >().AccumulatedImpulse;
+				return GetSingleContact().AccumulatedImpulse;
 			}
 			ensure(false);
 
@@ -201,7 +277,7 @@ namespace Chaos
 		{
 			if (ConstraintType == FConstraintBase::FType::SinglePoint)
 			{
-				const TRigidBodySingleContactConstraint<float, 3>& Contact = GetContact< TRigidBodySingleContactConstraint<float, 3> >();
+				const TRigidBodyPointContactConstraint<float, 3>& Contact = GetSingleContact();
 				return { Contact.Particle[0], Contact.Particle[1] };
 			}
 			ensure(false);
@@ -214,7 +290,7 @@ namespace Chaos
 		{
 			if (ConstraintType == FConstraintBase::FType::SinglePoint)
 			{
-				return { GetContact< TRigidBodySingleContactConstraint<float, 3> >().Particle[0], GetContact< TRigidBodySingleContactConstraint<float, 3> >().Particle[1] };
+				return { GetSingleContact().Particle[0], GetSingleContact().Particle[1] };
 			}
 			ensure(false);
 			return { nullptr,nullptr };
@@ -229,7 +305,4 @@ namespace Chaos
 
 	};
 
-	extern template TRigidBodySingleContactConstraint<float, 3>& TPBDCollisionConstraintHandle<float, 3>::GetContact<TRigidBodySingleContactConstraint<float, 3>>();
-	extern template const TRigidBodySingleContactConstraint<float, 3>& TPBDCollisionConstraintHandle<float, 3>::GetContact<TRigidBodySingleContactConstraint<float, 3>>() const;
-
 }
diff --git a/Engine/Source/Runtime/Experimental/Chaos/Public/Chaos/PBDCollisionConstraint.h b/Engine/Source/Runtime/Experimental/Chaos/Public/Chaos/PBDCollisionConstraint.h
--- a/Engine/Source/Runtime/Experimental/Chaos/Public/Chaos/PBDCollisionConstraint.h
+++ b/Engine/Source/Runtime/Experimental/Chaos/Public/Chaos/PBDCollisionConstraint.h
@@ -65,7 +65,8 @@ public:
 	using FConstraintHandleID = TPair<const TGeometryParticleHandle<T, d>*, const TGeometryParticleHandle<T, d>*>;
 	using FConstraintContainerHandle = TPBDCollisionConstraintHandle<T, d>; // @todo(brice) rename this to FCollisionConstraintHandle
 	using FConstraintHandleAllocator = TConstraintHandleAllocator<TPBDCollisionConstraint<T, d>>;
-	using FRigidBodyContactConstraint = TRigidBodySingleContactConstraint<T, d>;
+ 	using FPointContactConstraint = TRigidBodyPointContactConstraint<T, d>;
+	using FPlaneContactConstraint = TRigidBodyPlaneContactConstraint<T, d>;
 	using FHandles = TArray<FConstraintContainerHandle*>;
 
 	TPBDCollisionConstraint(const TPBDRigidsSOAs<T,d>& InParticles, TArrayCollectionArray<bool>& Collided, const TArrayCollectionArray<TSerializablePtr<FChaosPhysicsMaterial>>& PerParticleMaterials, const int32 ApplyPairIterations = 1, const int32 ApplyPushOutPairIterations = 1, const T Thickness = (T)0);
@@ -102,10 +103,10 @@ public:
 
 	int32 NumConstraints() const
 	{
-		return Constraints.Num();
+		return Handles.Num();
 	}
 
-	FConstraintHandleID GetConstraintHandleID(const FRigidBodyContactConstraint & Constraint) const
+	FConstraintHandleID GetConstraintHandleID(const FPointContactConstraint & Constraint) const
 	{
 		return  FConstraintHandleID( Constraint.Particle[0], Constraint.Particle[1] );
 	}
@@ -201,27 +202,28 @@ private:
 	*/
 	void Reset( );
 
-	void Apply(const T Dt, FRigidBodyContactConstraint& Constraint, const int32 It, const int32 NumIts);
+	void Apply(const T Dt, FPointContactConstraint& Constraint, const int32 It, const int32 NumIts);
 
 
-	void ApplyPushOut(const T Dt, FRigidBodyContactConstraint& Constraint, const TSet<const TGeometryParticleHandle<T, d>*>& IsTemporarilyStatic, int32 Iteration, int32 NumIterations, bool &NeedsAnotherIteration);
+	void ApplyPushOut(const T Dt, FPointContactConstraint& Constraint, const TSet<const TGeometryParticleHandle<T, d>*>& IsTemporarilyStatic, int32 Iteration, int32 NumIterations, bool &NeedsAnotherIteration);
 
 	template <bool bGatherStats = false>
 	void ComputeConstraints(const FAccelerationStructure& AccelerationStructure, T Dt);
 
 	template<ECollisionUpdateType>
-	void UpdateConstraint(const T Thickness, FRigidBodyContactConstraint& Constraint);
+	void UpdateConstraint(const T Thickness, FPointContactConstraint& Constraint);
 
-	void ConstructConstraints(TGeometryParticleHandle<T, d>* Particle0, TGeometryParticleHandle<T, d>* Particle1, const T Thickness, TRigidBodySingleContactConstraint<T, d> & Constraint);
+	void ConstructConstraints(TGeometryParticleHandle<T, d>* Particle0, TGeometryParticleHandle<T, d>* Particle1, const T Thickness, TRigidBodyPointContactConstraint<T, d> & Constraint);
 
 
 	const TPBDRigidsSOAs<T,d>& Particles;
 
+	TArray<FPointContactConstraint> PointConstraints;
+	TArray<FPlaneContactConstraint> PlaneConstraints;
 
+	TMap< FConstraintHandleID, FConstraintContainerHandle* > Manifolds;
 	TArray<FConstraintContainerHandle*> Handles;
 	FConstraintHandleAllocator HandleAllocator;
-	TMap< FConstraintHandleID, FConstraintContainerHandle* > Manifolds;
-	TArray<FRigidBodyContactConstraint> Constraints;
 
 	TArrayCollectionArray<bool>& MCollided;
 	const TArrayCollectionArray<TSerializablePtr<FChaosPhysicsMaterial>>& MPhysicsMaterials;
@@ -243,8 +245,8 @@ private:
 	const FAccelerationStructure* SpatialAcceleration;
 };
 
-extern template void TPBDCollisionConstraint<float, 3>::UpdateConstraint<ECollisionUpdateType::Any>(const float Thickness, FRigidBodyContactConstraint& Constraint);
-extern template void TPBDCollisionConstraint<float, 3>::UpdateConstraint<ECollisionUpdateType::Deepest>(const float Thickness, FRigidBodyContactConstraint& Constraint);
+extern template void TPBDCollisionConstraint<float, 3>::UpdateConstraint<ECollisionUpdateType::Any>(const float Thickness, FPointContactConstraint& Constraint);
+extern template void TPBDCollisionConstraint<float, 3>::UpdateConstraint<ECollisionUpdateType::Deepest>(const float Thickness, FPointContactConstraint& Constraint);
 extern template void TPBDCollisionConstraint<float, 3>::ComputeConstraints<false>(const TPBDCollisionConstraint<float, 3>::FAccelerationStructure&, float Dt);
 extern template void TPBDCollisionConstraint<float, 3>::ComputeConstraints<true>(const TPBDCollisionConstraint<float, 3>::FAccelerationStructure&, float Dt);
 }
diff --git a/Engine/Source/Runtime/Experimental/Chaos/Public/Chaos/PBDCollisionConstraintImp.h b/Engine/Source/Runtime/Experimental/Chaos/Public/Chaos/PBDCollisionConstraintImp.h
--- a/Engine/Source/Runtime/Experimental/Chaos/Public/Chaos/PBDCollisionConstraintImp.h
+++ b/Engine/Source/Runtime/Experimental/Chaos/Public/Chaos/PBDCollisionConstraintImp.h
@@ -152,7 +152,7 @@ namespace Chaos
 			// Narrow phase
 			CHAOS_SCOPED_TIMER(ComputeConstraints_NP);
 
-			TQueue<TRigidBodySingleContactConstraint<T, d>, EQueueMode::Mpsc> Queue;	//todo(ocohen): use per thread buffer instead, need better support than ParallelFor for this
+			TQueue<TRigidBodyPointContactConstraint<T, d>, EQueueMode::Mpsc> Queue;	//todo(ocohen): use per thread buffer instead, need better support than ParallelFor for this
 			Particles.GetNonDisabledDynamicView().ParallelFor([&](auto& Particle1, int32 ActiveIdxIdx)
 			{
 #if !UE_BUILD_SHIPPING
@@ -253,7 +253,7 @@ namespace Chaos
 					const T UseThickness = FMath::Max(Box1Thickness, Box2Thickness.Size());// + MThickness
 
 
-					TRigidBodySingleContactConstraint<T, d> Constraint;
+					TRigidBodyPointContactConstraint<T, d> Constraint;
 					ConstructConstraints(Particle1.Handle(), Particle2.Handle(), UseThickness, Constraint);
 
 #if !UE_BUILD_SHIPPING
@@ -292,13 +292,11 @@ namespace Chaos
 
 				while (!Queue.IsEmpty())
 				{
-					int32 Idx = Constraints.AddUninitialized(1);
-					FConstraintContainerHandle* Handle = HandleAllocator.template AllocHandle< TRigidBodySingleContactConstraint<T,d> >(this, Idx );
+					int32 Idx = PointConstraints.AddUninitialized(1);
+					FConstraintContainerHandle* Handle = HandleAllocator.template AllocHandle< TRigidBodyPointContactConstraint<T,d> >(this, Idx );
 
 					Handles.Add(Handle);
-					Queue.Dequeue(Constraints[Idx]);
-					ensure(Handles.Num() == Constraints.Num());
-
+					Queue.Dequeue(PointConstraints[Idx]);
 				}
 				LifespanCounter++;
 			}
diff --git a/Engine/Source/Runtime/Experimental/Chaos/Public/Chaos/PBDCollisionConstraintPlaneContactUtil.h b/Engine/Source/Runtime/Experimental/Chaos/Public/Chaos/PBDCollisionConstraintPlaneContactUtil.h
new file mode 100644
--- /dev/null
+++ b/Engine/Source/Runtime/Experimental/Chaos/Public/Chaos/PBDCollisionConstraintPlaneContactUtil.h
@@ -0,0 +1,48 @@
+// Copyright 1998-2019 Epic Games, Inc. All Rights Reserved.
+#pragma once
+
+#include "Chaos/ExternalCollisionData.h"
+#include "Chaos/PBDCollisionConstraint.h"
+#include "Chaos/CollisionResolutionTypes.h"
+
+namespace Chaos
+{
+	namespace Collisions
+	{
+		template<class T = float>
+		struct TPlainContactParticleParameters {
+			TArrayCollectionArray<bool>* Collided;
+			const TArrayCollectionArray<TSerializablePtr<FChaosPhysicsMaterial>>* PhysicsMaterials;
+			T FrictionOverride;
+			T AngularFrictionOverride;
+		};
+
+		template<class T = float>
+		struct TPlaneContactIterationParameters {
+			const T Dt; 
+			const int32 Iteration; 
+			const int32 NumIterations; 
+			const int32 NumPairIterations;
+			bool* NeedsAnotherIteration;
+		};
+
+		/*
+		template<ECollisionUpdateType UpdateType, typename T = float, int d = 3>
+		void Update(const T Thickness, TRigidBodyPointContactConstraint<T, d>& Constraint);
+
+		template<ECollisionUpdateType UpdateType, typename T=float, int d=3>
+		void UpdateConstraint(const T Thickness, TRigidBodyPointContactConstraint<T, d>& Constraint);
+			
+		template<typename T = float, int d = 3>
+		void Apply(TRigidBodyPointContactConstraint<T, d>& Constraint, T Thickness, 
+			TSingleContactIterationParameters<T> & IterationParameters,
+			TSingleContactParticleParameters<T> & ParticleParameters);
+			
+		template<typename T=float, int d=3>
+		void ApplyPushOut(TRigidBodyPointContactConstraint<T, d>& Constraint, T Thickness, const TSet<const TGeometryParticleHandle<T, d>*>& IsTemporarilyStatic,
+			TSingleContactIterationParameters<T> & IterationParameters,
+			TSingleContactParticleParameters<T> & ParticleParameters);
+	*/
+	}
+
+}
diff --git a/Engine/Source/Runtime/Experimental/Chaos/Public/Chaos/PBDCollisionConstraintSingleContactUtil.h b/Engine/Source/Runtime/Experimental/Chaos/Public/Chaos/PBDCollisionConstraintSingleContactUtil.h
--- a/Engine/Source/Runtime/Experimental/Chaos/Public/Chaos/PBDCollisionConstraintSingleContactUtil.h
+++ b/Engine/Source/Runtime/Experimental/Chaos/Public/Chaos/PBDCollisionConstraintSingleContactUtil.h
@@ -10,7 +10,7 @@ namespace Chaos
 	namespace Collisions
 	{
 		template<class T = float>
-		struct TSingleContParticleParameters {
+		struct TSingleContactParticleParameters {
 			TArrayCollectionArray<bool>* Collided;
 			const TArrayCollectionArray<TSerializablePtr<FChaosPhysicsMaterial>>* PhysicsMaterials;
 			T FrictionOverride;
@@ -27,20 +27,20 @@ namespace Chaos
 		};
 
 		template<ECollisionUpdateType UpdateType, typename T = float, int d = 3>
-		void Update(const T Thickness, TRigidBodySingleContactConstraint<T, d>& Constraint);
+		void Update(const T Thickness, TRigidBodyPointContactConstraint<T, d>& Constraint);
 
 		template<ECollisionUpdateType UpdateType, typename T=float, int d=3>
-		void UpdateConstraint(const T Thickness, TRigidBodySingleContactConstraint<T, d>& Constraint);
+		void UpdateConstraint(const T Thickness, TRigidBodyPointContactConstraint<T, d>& Constraint);
 			
 		template<typename T = float, int d = 3>
-		void Apply(TRigidBodySingleContactConstraint<T, d>& Constraint, T Thickness, 
+		void Apply(TRigidBodyPointContactConstraint<T, d>& Constraint, T Thickness, 
 			TSingleContactIterationParameters<T> & IterationParameters,
-			TSingleContParticleParameters<T> & ParticleParameters);
+			TSingleContactParticleParameters<T> & ParticleParameters);
 			
 		template<typename T=float, int d=3>
-		void ApplyPushOut(TRigidBodySingleContactConstraint<T, d>& Constraint, T Thickness, const TSet<const TGeometryParticleHandle<T, d>*>& IsTemporarilyStatic,
+		void ApplyPushOut(TRigidBodyPointContactConstraint<T, d>& Constraint, T Thickness, const TSet<const TGeometryParticleHandle<T, d>*>& IsTemporarilyStatic,
 			TSingleContactIterationParameters<T> & IterationParameters,
-			TSingleContParticleParameters<T> & ParticleParameters);
+			TSingleContactParticleParameters<T> & ParticleParameters);
 	}
 
 }
diff --git a/Engine/Source/Runtime/Experimental/Chaos/Public/Chaos/PBDCollisionConstraintUtil.h b/Engine/Source/Runtime/Experimental/Chaos/Public/Chaos/PBDCollisionConstraintUtil.h
--- a/Engine/Source/Runtime/Experimental/Chaos/Public/Chaos/PBDCollisionConstraintUtil.h
+++ b/Engine/Source/Runtime/Experimental/Chaos/Public/Chaos/PBDCollisionConstraintUtil.h
@@ -7,7 +7,7 @@
 
 namespace Chaos
 {
-	void CHAOS_API ComputeHashTable(const TArray<Chaos::TPBDCollisionConstraint<float, 3>::FRigidBodyContactConstraint>& ConstraintsArray,
+	void CHAOS_API ComputeHashTable(const TArray<Chaos::TPBDCollisionConstraint<float, 3>::FPointContactConstraint>& ConstraintsArray,
 									const FBox& BoundingBox, TMultiMap<int32, int32>& HashTableMap, const float SpatialHashRadius);
 
 	void CHAOS_API ComputeHashTable(const TArray<TCollisionData<float, 3>>& CollisionsArray,
diff --git a/Engine/Source/Runtime/Experimental/Chaos/Public/Chaos/PBDRigidClustering.h b/Engine/Source/Runtime/Experimental/Chaos/Public/Chaos/PBDRigidClustering.h
--- a/Engine/Source/Runtime/Experimental/Chaos/Public/Chaos/PBDRigidClustering.h
+++ b/Engine/Source/Runtime/Experimental/Chaos/Public/Chaos/PBDRigidClustering.h
@@ -34,7 +34,7 @@ public:
 template<class FPBDRigidEvolution, class FPBDCollisionConstraint, class T, int d>
 class CHAOS_API TPBDRigidClustering
 {
-	typedef typename FPBDCollisionConstraint::FRigidBodyContactConstraint FRigidBodyContactConstraint;
+	typedef typename FPBDCollisionConstraint::FPointContactConstraint FPointContactConstraint;
 
 public:
 	typedef TMap<const TGeometryParticleHandle<T,d>*, TArray<uint32> > FClusterMap;
diff --git a/Engine/Source/Runtime/Experimental/ChaosSolvers/Private/EventDefaults.cpp b/Engine/Source/Runtime/Experimental/ChaosSolvers/Private/EventDefaults.cpp
--- a/Engine/Source/Runtime/Experimental/ChaosSolvers/Private/EventDefaults.cpp
+++ b/Engine/Source/Runtime/Experimental/ChaosSolvers/Private/EventDefaults.cpp
@@ -68,7 +68,7 @@ namespace Chaos
 				{
 					if (ContactHandle->GetType() == TPBDCollisionConstraintHandle<float, 3>::FConstraintBase::FType::SinglePoint)
 					{
-						const TRigidBodySingleContactConstraint<float, 3>& Constraint = ContactHandle->GetContact<TRigidBodySingleContactConstraint<float, 3>>();
+						const TRigidBodyPointContactConstraint<float, 3>& Constraint = ContactHandle->GetSingleContact();
 
 						// Since Clustered GCs can be unioned the particleIndex representing the union 
 						// is not associated with a PhysicsProxy
@@ -107,7 +107,7 @@ namespace Chaos
 				{
 					for(int32 IdxCollision = 0; IdxCollision < ValidCollisionHandles.Num(); ++IdxCollision)
 					{
-						Chaos::TPBDCollisionConstraint<float, 3>::FRigidBodyContactConstraint const& Constraint = ValidCollisionHandles[IdxCollision]->GetContact<TRigidBodySingleContactConstraint<float, 3>>();
+						Chaos::TPBDCollisionConstraint<float, 3>::FPointContactConstraint const& Constraint = ValidCollisionHandles[IdxCollision]->GetSingleContact();
 
 						TGeometryParticleHandle<float, 3>* Particle0 = Constraint.Particle[0];
 						TGeometryParticleHandle<float, 3>* Particle1 = Constraint.Particle[1];
diff --git a/Engine/Source/Runtime/Experimental/ChaosSolvers/Private/PhysicsProxy/GeometryCollectionPhysicsProxy.cpp b/Engine/Source/Runtime/Experimental/ChaosSolvers/Private/PhysicsProxy/GeometryCollectionPhysicsProxy.cpp
--- a/Engine/Source/Runtime/Experimental/ChaosSolvers/Private/PhysicsProxy/GeometryCollectionPhysicsProxy.cpp
+++ b/Engine/Source/Runtime/Experimental/ChaosSolvers/Private/PhysicsProxy/GeometryCollectionPhysicsProxy.cpp
@@ -2672,7 +2672,7 @@ FRecordedFrame& FGeometryCollectionPhysicsProxy::InsertRecordedFrame(FRecordedTr
 	return InTrack.Records[NewRecordIndex];
 }
 
-void FGeometryCollectionPhysicsProxy::AddCollisionToCollisionData(FRecordedFrame* ExistingFrame, const FParticlesType& Particles, const Chaos::TPBDCollisionConstraint<float, 3>::FRigidBodyContactConstraint& Constraint)
+void FGeometryCollectionPhysicsProxy::AddCollisionToCollisionData(FRecordedFrame* ExistingFrame, const FParticlesType& Particles, const Chaos::TPBDCollisionConstraint<float, 3>::FPointContactConstraint& Constraint)
 {
 #if TODO_REIMPLEMENT_RIGID_CLUSTERING
 	const Chaos::TArrayCollectionArray<Chaos::ClusterId>& ClusterIdsArray = GetSolver()->GetRigidClustering().GetClusterIdsArray();
diff --git a/Engine/Source/Runtime/Experimental/ChaosSolvers/Public/PhysicsProxy/GeometryCollectionPhysicsProxy.h b/Engine/Source/Runtime/Experimental/ChaosSolvers/Public/PhysicsProxy/GeometryCollectionPhysicsProxy.h
--- a/Engine/Source/Runtime/Experimental/ChaosSolvers/Public/PhysicsProxy/GeometryCollectionPhysicsProxy.h
+++ b/Engine/Source/Runtime/Experimental/ChaosSolvers/Public/PhysicsProxy/GeometryCollectionPhysicsProxy.h
@@ -180,7 +180,7 @@ private:
 
 	void UpdateRecordedState(float SolverTime, const TManagedArray<int32>& RigidBodyID, const TManagedArray<int32>& CollectionClusterID, const Chaos::TArrayCollectionArray<bool>& InternalCluster, const FParticlesType& Particles, const FCollisionConstraintsType& CollisionRule);
 
-	void AddCollisionToCollisionData(FRecordedFrame* ExistingFrame, const FParticlesType& Particles, const Chaos::TPBDCollisionConstraint<float, 3>::FRigidBodyContactConstraint& Constraint);
+	void AddCollisionToCollisionData(FRecordedFrame* ExistingFrame, const FParticlesType& Particles, const Chaos::TPBDCollisionConstraint<float, 3>::FPointContactConstraint& Constraint);
 
 	void UpdateCollisionData(const FParticlesType& Particles, const FCollisionConstraintsType& CollisionRule, FRecordedFrame* ExistingFrame);
 
