Chaos :

File cleanup.

diff --git a/Engine/Source/Runtime/Experimental/Chaos/Private/Chaos/PBDCollisionConstraint.cpp b/Engine/Source/Runtime/Experimental/Chaos/Private/Chaos/PBDCollisionConstraint.cpp
--- a/Engine/Source/Runtime/Experimental/Chaos/Private/Chaos/PBDCollisionConstraint.cpp
+++ b/Engine/Source/Runtime/Experimental/Chaos/Private/Chaos/PBDCollisionConstraint.cpp
@@ -22,7 +22,7 @@
 #include "ChaosStats.h"
 #include "Containers/Queue.h"
 #include "ProfilingDebugging/ScopedTimers.h"
-#include "Chaos/PBDCollisionConstraintSingleContactUtil.h"
+#include "Chaos/PBDCollisionConstraintPointContactUtil.h"
 
 #if INTEL_ISPC
 #include "PBDCollisionConstraint.ispc.generated.h"
@@ -298,8 +298,8 @@ namespace Chaos
 			}
 		}
 
-		Collisions::TSingleContactParticleParameters<T> ParticleParameters = { &MCollided, &MPhysicsMaterials, CollisionFrictionOverride, MAngularFriction };
-		Collisions::TSingleContactIterationParameters<T> IterationParameters = { Dt, It, NumIts, MApplyPairIterations, nullptr };
+		Collisions::TPointContactParticleParameters<T> ParticleParameters = { &MCollided, &MPhysicsMaterials, CollisionFrictionOverride, MAngularFriction };
+		Collisions::TPointContactIterationParameters<T> IterationParameters = { Dt, It, NumIts, MApplyPairIterations, nullptr };
 		Collisions::Apply(Constraint,MThickness, IterationParameters,  ParticleParameters);
 	}
 
@@ -349,8 +349,8 @@ namespace Chaos
 			}
 		}
 
-		Collisions::TSingleContactParticleParameters<T> ParticleParameters = { &MCollided, &MPhysicsMaterials, CollisionFrictionOverride, MAngularFriction };
-		Collisions::TSingleContactIterationParameters<T> IterationParameters = { Dt, Iteration, NumIterations, MApplyPairIterations, &NeedsAnotherIteration };
+		Collisions::TPointContactParticleParameters<T> ParticleParameters = { &MCollided, &MPhysicsMaterials, CollisionFrictionOverride, MAngularFriction };
+		Collisions::TPointContactIterationParameters<T> IterationParameters = { Dt, Iteration, NumIterations, MApplyPairIterations, &NeedsAnotherIteration };
 		Collisions::ApplyPushOut(Constraint, MThickness, IsTemporarilyStatic, IterationParameters, ParticleParameters);
 	}
 
diff --git a/Engine/Source/Runtime/Experimental/Chaos/Private/Chaos/PBDCollisionConstraintPlaneContactUtil.cpp b/Engine/Source/Runtime/Experimental/Chaos/Private/Chaos/PBDCollisionConstraintPlaneContactUtil.cpp
--- a/Engine/Source/Runtime/Experimental/Chaos/Private/Chaos/PBDCollisionConstraintPlaneContactUtil.cpp
+++ b/Engine/Source/Runtime/Experimental/Chaos/Private/Chaos/PBDCollisionConstraintPlaneContactUtil.cpp
@@ -26,18 +26,18 @@ namespace Chaos
 
 
 		template<typename T, int d>
-		void Apply(TRigidBodyPlaneContactConstraint<T, d>& Constraint, T Thickness, TSingleContactIterationParameters<T> & IterationParameters, TPlainContactParticleParameters<T> & ParticleParameters)
+		void Apply(TRigidBodyPlaneContactConstraint<T, d>& Constraint, T Thickness, TPointContactIterationParameters<T> & IterationParameters, TPlainContactParticleParameters<T> & ParticleParameters)
 		{
 		}
-		template void Apply<float, 3>(TRigidBodyPlaneContactConstraint<float, 3>&, float, TSingleContactIterationParameters<float> &, TSingleContactParticleParameters<float> &);
+		template void Apply<float, 3>(TRigidBodyPlaneContactConstraint<float, 3>&, float, TPointContactIterationParameters<float> &, TPointContactParticleParameters<float> &);
 
 		template<typename T, int d>
 		void ApplyPushOut(TRigidBodyPointContactConstraint<T, d>& Constraint, T Thickness, const TSet<const TGeometryParticleHandle<T, d>*>& IsTemporarilyStatic,
-			TSingleContactIterationParameters<T> & IterationParameters, TSingleContactParticleParameters<T> & ParticleParameters)
+			TPointContactIterationParameters<T> & IterationParameters, TPointContactParticleParameters<T> & ParticleParameters)
 		{
 		}
 		template void ApplyPushOut<float, 3>(TRigidBodyPointContactConstraint<float,3>&, float , const TSet<const TGeometryParticleHandle<float,3>*>&,
-			TSingleContactIterationParameters<float> & IterationParameters, TSingleContactParticleParameters<float> & ParticleParameters);
+			TPointContactIterationParameters<float> & IterationParameters, TPointContactParticleParameters<float> & ParticleParameters);
 		*/
 	} // Collisions
 
diff --git a/Engine/Source/Runtime/Experimental/Chaos/Private/Chaos/PBDCollisionConstraintPointContactUtil.cpp b/Engine/Source/Runtime/Experimental/Chaos/Private/Chaos/PBDCollisionConstraintPointContactUtil.cpp
new file mode 100644
--- /dev/null
+++ b/Engine/Source/Runtime/Experimental/Chaos/Private/Chaos/PBDCollisionConstraintPointContactUtil.cpp
@@ -0,0 +1,467 @@
+// Copyright 1998-2019 Epic Games, Inc. All Rights Reserved.
+
+#include "Chaos/PBDCollisionConstraintPointContactUtil.h"
+#include "Chaos/CollisionResolutionAlgo.h"
+#include "Chaos/Defines.h"
+
+namespace Chaos
+{
+	namespace Collisions
+	{
+		template<typename T=float, int d=3>
+		TRigidTransform<T, d> GetTransform(const TGeometryParticleHandle<T, d>* Particle)
+		{
+			TGenericParticleHandle<T, d> Generic = const_cast<TGeometryParticleHandle<T, d>*>(Particle);
+			return TRigidTransform<T, d>(Generic->P(), Generic->Q());
+		}
+
+		// @todo(ccaulfield): This is duplicated in JointConstraints - move to a utility file
+		template<typename T>
+		PMatrix<T, 3, 3> ComputeFactorMatrix3(const TVector<T, 3>& V, const PMatrix<T, 3, 3>& M, const T& Im)
+		{
+			// Rigid objects rotational contribution to the impulse.
+			// Vx*M*VxT+Im
+			check(Im > FLT_MIN)
+				return PMatrix<T, 3, 3>(
+					-V[2] * (-V[2] * M.M[1][1] + V[1] * M.M[2][1]) + V[1] * (-V[2] * M.M[2][1] + V[1] * M.M[2][2]) + Im,
+					V[2] * (-V[2] * M.M[1][0] + V[1] * M.M[2][0]) - V[0] * (-V[2] * M.M[2][1] + V[1] * M.M[2][2]),
+					-V[1] * (-V[2] * M.M[1][0] + V[1] * M.M[2][0]) + V[0] * (-V[2] * M.M[1][1] + V[1] * M.M[2][1]),
+					V[2] * (V[2] * M.M[0][0] - V[0] * M.M[2][0]) - V[0] * (V[2] * M.M[2][0] - V[0] * M.M[2][2]) + Im,
+					-V[1] * (V[2] * M.M[0][0] - V[0] * M.M[2][0]) + V[0] * (V[2] * M.M[1][0] - V[0] * M.M[2][1]),
+					-V[1] * (-V[1] * M.M[0][0] + V[0] * M.M[1][0]) + V[0] * (-V[1] * M.M[1][0] + V[0] * M.M[1][1]) + Im);
+		}
+
+		template<typename T, int d>
+		TVector<T, d> GetEnergyClampedImpulse(const TRigidBodyPointContactConstraint<T, d>& Constraint, const TVector<T, d>& Impulse, const TVector<T, d>& VectorToPoint1, const TVector<T, d>& VectorToPoint2, const TVector<T, d>& Velocity1, const TVector<T, d>& Velocity2)
+		{
+			TPBDRigidParticleHandle<T, d>* PBDRigid0 = Constraint.Particle[0]->AsDynamic();
+			TPBDRigidParticleHandle<T, d>* PBDRigid1 = Constraint.Particle[1]->AsDynamic();
+
+			TVector<T, d> Jr0, Jr1, IInvJr0, IInvJr1;
+			T ImpulseRatioNumerator0 = 0, ImpulseRatioNumerator1 = 0, ImpulseRatioDenom0 = 0, ImpulseRatioDenom1 = 0;
+			T ImpulseSize = Impulse.SizeSquared();
+			TVector<T, d> KinematicVelocity = !PBDRigid0 ? Velocity1 : !PBDRigid1 ? Velocity2 : TVector<T, d>(0);
+			if (PBDRigid0)
+			{
+				Jr0 = TVector<T, d>::CrossProduct(VectorToPoint1, Impulse);
+				IInvJr0 = PBDRigid0->Q().RotateVector(PBDRigid0->InvI() * PBDRigid0->Q().UnrotateVector(Jr0));
+				ImpulseRatioNumerator0 = TVector<T, d>::DotProduct(Impulse, PBDRigid0->V() - KinematicVelocity) + TVector<T, d>::DotProduct(IInvJr0, PBDRigid0->W());
+				ImpulseRatioDenom0 = ImpulseSize / PBDRigid0->M() + TVector<T, d>::DotProduct(Jr0, IInvJr0);
+			}
+			if (PBDRigid1)
+			{
+				Jr1 = TVector<T, d>::CrossProduct(VectorToPoint2, Impulse);
+				IInvJr1 = PBDRigid1->Q().RotateVector(PBDRigid1->InvI() * PBDRigid1->Q().UnrotateVector(Jr1));
+				ImpulseRatioNumerator1 = TVector<T, d>::DotProduct(Impulse, PBDRigid1->V() - KinematicVelocity) + TVector<T, d>::DotProduct(IInvJr1, PBDRigid1->W());
+				ImpulseRatioDenom1 = ImpulseSize / PBDRigid1->M() + TVector<T, d>::DotProduct(Jr1, IInvJr1);
+			}
+			T Numerator = -2 * (ImpulseRatioNumerator0 - ImpulseRatioNumerator1);
+			if (Numerator < 0)
+			{
+				return TVector<T, d>(0);
+			}
+			check(Numerator >= 0);
+			T Denominator = ImpulseRatioDenom0 + ImpulseRatioDenom1;
+			return Numerator < Denominator ? (Impulse * Numerator / Denominator) : Impulse;
+		}
+
+
+		template<ECollisionUpdateType UpdateType, typename T, int d>
+		void Update(const T Thickness, TRigidBodyPointContactConstraint<T, d>& Constraint)
+		{
+			Constraint.ResetPhi(Thickness);
+			const TRigidTransform<T, d> ParticleTM = GetTransform(Constraint.Particle[0]);
+			const TRigidTransform<T, d> LevelsetTM = GetTransform(Constraint.Particle[1]);
+
+			if (!Constraint.Particle[0]->Geometry())
+			{
+				if (Constraint.Particle[1]->Geometry())
+				{
+					if (!Constraint.Particle[1]->Geometry()->IsUnderlyingUnion())
+					{
+						UpdateLevelsetConstraint<UpdateType>(Thickness, Constraint);
+					}
+					else
+					{
+						UpdateUnionLevelsetConstraint<UpdateType>(Thickness, Constraint);
+					}
+				}
+			}
+			else
+			{
+				UpdateConstraintImp<UpdateType>(*Constraint.Particle[0]->Geometry(), ParticleTM, *Constraint.Particle[1]->Geometry(), LevelsetTM, Thickness, Constraint);
+			}
+		}
+		template void Update<ECollisionUpdateType::Any, float, 3>(const float, TRigidBodyPointContactConstraint<float, 3>&);
+		template void Update<ECollisionUpdateType::Deepest, float, 3>(const float, TRigidBodyPointContactConstraint<float, 3>&);
+
+
+		template<ECollisionUpdateType UpdateType, typename T, int d>
+		void UpdateConstraint(const T Thickness, TRigidBodyPointContactConstraint<T, d>& Constraint)
+		{
+			Constraint.ResetPhi(Thickness);
+			const TRigidTransform<T, d> ParticleTM = GetTransform(Constraint.Particle[0]);
+			const TRigidTransform<T, d> LevelsetTM = GetTransform(Constraint.Particle[1]);
+
+			if (!Constraint.Particle[0]->Geometry())
+			{
+				if (Constraint.Particle[1]->Geometry())
+				{
+					if (!Constraint.Particle[1]->Geometry()->IsUnderlyingUnion())
+					{
+						UpdateLevelsetConstraint<UpdateType>(Thickness, Constraint);
+					}
+					else
+					{
+						UpdateUnionLevelsetConstraint<UpdateType>(Thickness, Constraint);
+					}
+				}
+			}
+			else
+			{
+				UpdateConstraintImp<UpdateType>(*Constraint.Particle[0]->Geometry(), ParticleTM, *Constraint.Particle[1]->Geometry(), LevelsetTM, Thickness, Constraint);
+			}
+		}
+		template void UpdateConstraint<ECollisionUpdateType::Any, float, 3>(const float, TRigidBodyPointContactConstraint<float,3>&);
+		template void UpdateConstraint<ECollisionUpdateType::Deepest, float, 3>(const float, TRigidBodyPointContactConstraint<float,3>&);
+
+
+		template<typename T, int d>
+		void Apply(TRigidBodyPointContactConstraint<T, d>& Constraint, T Thickness, TPointContactIterationParameters<T> & IterationParameters, TPointContactParticleParameters<T> & ParticleParameters)
+		{
+			TGenericParticleHandle<T, d> Particle0 = TGenericParticleHandle<T, d>(Constraint.Particle[0]);
+			TGenericParticleHandle<T, d> Particle1 = TGenericParticleHandle<T, d>(Constraint.Particle[1]);
+			TPBDRigidParticleHandle<T, d>* PBDRigid0 = Particle0->AsDynamic();
+			TPBDRigidParticleHandle<T, d>* PBDRigid1 = Particle1->AsDynamic();
+
+			for (int32 PairIt = 0; PairIt < IterationParameters.NumPairIterations; ++PairIt)
+			{
+				Collisions::UpdateConstraint<ECollisionUpdateType::Deepest>(Thickness, Constraint);
+				if (Constraint.GetPhi() >= Thickness)
+				{
+					return;
+				}
+
+				// @todo(ccaulfield): CHAOS_PARTICLEHANDLE_TODO what's the best way to manage external per-particle data?
+				if (ParticleParameters.Collided)
+				{
+					Particle0->AuxilaryValue(*ParticleParameters.Collided) = true;
+					Particle1->AuxilaryValue(*ParticleParameters.Collided) = true;
+				}
+
+				TSerializablePtr<FChaosPhysicsMaterial> PhysicsMaterial0, PhysicsMaterial1;
+				if (ParticleParameters.PhysicsMaterials)
+				{
+					PhysicsMaterial0 = Particle0->AuxilaryValue(*ParticleParameters.PhysicsMaterials);
+					PhysicsMaterial1 = Particle1->AuxilaryValue(*ParticleParameters.PhysicsMaterials);
+				}
+
+				// @todo(ccaulfield): CHAOS_PARTICLEHANDLE_TODO split function to avoid ifs
+				const TVector<T, d> ZeroVector = TVector<T, d>(0);
+				const TRotation<T, d>& Q0 = Particle0->Q();
+				const TRotation<T, d>& Q1 = Particle1->Q();
+				const TVector<T, d>& P0 = Particle0->P();
+				const TVector<T, d>& P1 = Particle1->P();
+				const TVector<T, d>& V0 = Particle0->V();
+				const TVector<T, d>& V1 = Particle1->V();
+				const TVector<T, d>& W0 = Particle0->W();
+				const TVector<T, d>& W1 = Particle1->W();
+
+				typename TRigidBodyPointContactConstraint<T, d>::FManifold & Contact = Constraint.Manifold;
+
+				TVector<T, d> VectorToPoint1 = Contact.Location - P0;
+				TVector<T, d> VectorToPoint2 = Contact.Location - P1;
+				TVector<T, d> Body1Velocity = V0 + TVector<T, d>::CrossProduct(W0, VectorToPoint1);
+				TVector<T, d> Body2Velocity = V1 + TVector<T, d>::CrossProduct(W1, VectorToPoint2);
+				TVector<T, d> RelativeVelocity = Body1Velocity - Body2Velocity;
+				T RelativeNormalVelocity = TVector<T, d>::DotProduct(RelativeVelocity, Contact.Normal);
+
+				if (RelativeNormalVelocity < 0) // ignore separating constraints
+				{
+					PMatrix<T, d, d> WorldSpaceInvI1 = PBDRigid0 ? (Q0 * FMatrix::Identity).GetTransposed() * PBDRigid0->InvI() * (Q0 * FMatrix::Identity) : PMatrix<T, d, d>(0);
+					PMatrix<T, d, d> WorldSpaceInvI2 = PBDRigid1 ? (Q1 * FMatrix::Identity).GetTransposed() * PBDRigid1->InvI() * (Q1 * FMatrix::Identity) : PMatrix<T, d, d>(0);
+					PMatrix<T, d, d> Factor =
+						(PBDRigid0 ? ComputeFactorMatrix3(VectorToPoint1, WorldSpaceInvI1, PBDRigid0->InvM()) : PMatrix<T, d, d>(0)) +
+						(PBDRigid1 ? ComputeFactorMatrix3(VectorToPoint2, WorldSpaceInvI2, PBDRigid1->InvM()) : PMatrix<T, d, d>(0));
+					TVector<T, d> Impulse;
+					TVector<T, d> AngularImpulse(0);
+
+					// Resting contact if very close to the surface
+					T Restitution = (T)0;
+					T Friction = (T)0;
+					T AngularFriction = (T)0;
+					bool bApplyRestitution = (RelativeVelocity.Size() > (2 * 980 * IterationParameters.Dt));
+					if (PhysicsMaterial0 && PhysicsMaterial1)
+					{
+						if (bApplyRestitution)
+						{
+							Restitution = FMath::Min(PhysicsMaterial0->Restitution, PhysicsMaterial1->Restitution);
+						}
+						Friction = FMath::Max(PhysicsMaterial0->Friction, PhysicsMaterial1->Friction);
+					}
+					else if (PhysicsMaterial0)
+					{
+						if (bApplyRestitution)
+						{
+							Restitution = PhysicsMaterial0->Restitution;
+						}
+						Friction = PhysicsMaterial0->Friction;
+					}
+					else if (PhysicsMaterial1)
+					{
+						if (bApplyRestitution)
+						{
+							Restitution = PhysicsMaterial1->Restitution;
+						}
+						Friction = PhysicsMaterial1->Friction;
+					}
+
+					if (ParticleParameters.FrictionOverride >= 0)
+					{
+						Friction = ParticleParameters.FrictionOverride;
+					}
+					if (ParticleParameters.AngularFrictionOverride >= 0)
+					{
+						AngularFriction = ParticleParameters.AngularFrictionOverride;
+					}
+
+					if (Friction)
+					{
+						if (RelativeNormalVelocity > 0)
+						{
+							RelativeNormalVelocity = 0;
+						}
+						TVector<T, d> VelocityChange = -(Restitution * RelativeNormalVelocity * Contact.Normal + RelativeVelocity);
+						T NormalVelocityChange = TVector<T, d>::DotProduct(VelocityChange, Contact.Normal);
+						PMatrix<T, d, d> FactorInverse = Factor.Inverse();
+						TVector<T, d> MinimalImpulse = FactorInverse * VelocityChange;
+						const T MinimalImpulseDotNormal = TVector<T, d>::DotProduct(MinimalImpulse, Contact.Normal);
+						const T TangentialSize = (MinimalImpulse - MinimalImpulseDotNormal * Contact.Normal).Size();
+						if (TangentialSize <= Friction * MinimalImpulseDotNormal)
+						{
+							//within friction cone so just solve for static friction stopping the object
+							Impulse = MinimalImpulse;
+							if (AngularFriction)
+							{
+								TVector<T, d> RelativeAngularVelocity = W0 - W1;
+								T AngularNormal = TVector<T, d>::DotProduct(RelativeAngularVelocity, Contact.Normal);
+								TVector<T, d> AngularTangent = RelativeAngularVelocity - AngularNormal * Contact.Normal;
+								TVector<T, d> FinalAngularVelocity = FMath::Sign(AngularNormal) * FMath::Max((T)0, FMath::Abs(AngularNormal) - AngularFriction * NormalVelocityChange) * Contact.Normal + FMath::Max((T)0, AngularTangent.Size() - AngularFriction * NormalVelocityChange) * AngularTangent.GetSafeNormal();
+								TVector<T, d> Delta = FinalAngularVelocity - RelativeAngularVelocity;
+								if (!PBDRigid0 && PBDRigid1)
+								{
+									PMatrix<T, d, d> WorldSpaceI2 = (Q1 * FMatrix::Identity) * PBDRigid1->I() * (Q1 * FMatrix::Identity).GetTransposed();
+									TVector<T, d> ImpulseDelta = PBDRigid1->M() * TVector<T, d>::CrossProduct(VectorToPoint2, Delta);
+									Impulse += ImpulseDelta;
+									AngularImpulse += WorldSpaceI2 * Delta - TVector<T, d>::CrossProduct(VectorToPoint2, ImpulseDelta);
+								}
+								else if (PBDRigid0 && !PBDRigid1)
+								{
+									PMatrix<T, d, d> WorldSpaceI1 = (Q0 * FMatrix::Identity) * PBDRigid0->I() * (Q0 * FMatrix::Identity).GetTransposed();
+									TVector<T, d> ImpulseDelta = PBDRigid0->M() * TVector<T, d>::CrossProduct(VectorToPoint1, Delta);
+									Impulse += ImpulseDelta;
+									AngularImpulse += WorldSpaceI1 * Delta - TVector<T, d>::CrossProduct(VectorToPoint1, ImpulseDelta);
+								}
+								else if (PBDRigid0 && PBDRigid1)
+								{
+									PMatrix<T, d, d> Cross1(0, VectorToPoint1.Z, -VectorToPoint1.Y, -VectorToPoint1.Z, 0, VectorToPoint1.X, VectorToPoint1.Y, -VectorToPoint1.X, 0);
+									PMatrix<T, d, d> Cross2(0, VectorToPoint2.Z, -VectorToPoint2.Y, -VectorToPoint2.Z, 0, VectorToPoint2.X, VectorToPoint2.Y, -VectorToPoint2.X, 0);
+									PMatrix<T, d, d> CrossI1 = Cross1 * WorldSpaceInvI1;
+									PMatrix<T, d, d> CrossI2 = Cross2 * WorldSpaceInvI2;
+									PMatrix<T, d, d> Diag1 = CrossI1 * Cross1.GetTransposed() + CrossI2 * Cross2.GetTransposed();
+									Diag1.M[0][0] += PBDRigid0->InvM() + PBDRigid1->InvM();
+									Diag1.M[1][1] += PBDRigid0->InvM() + PBDRigid1->InvM();
+									Diag1.M[2][2] += PBDRigid0->InvM() + PBDRigid1->InvM();
+									PMatrix<T, d, d> OffDiag1 = (CrossI1 + CrossI2) * -1;
+									PMatrix<T, d, d> Diag2 = (WorldSpaceInvI1 + WorldSpaceInvI2).Inverse();
+									PMatrix<T, d, d> OffDiag1Diag2 = OffDiag1 * Diag2;
+									TVector<T, d> ImpulseDelta = PMatrix<T, d, d>((Diag1 - OffDiag1Diag2 * OffDiag1.GetTransposed()).Inverse())* ((OffDiag1Diag2 * -1) * Delta);
+									Impulse += ImpulseDelta;
+									AngularImpulse += Diag2 * (Delta - PMatrix<T, d, d>(OffDiag1.GetTransposed()) * ImpulseDelta);
+								}
+							}
+						}
+						else
+						{
+							//outside friction cone, solve for normal relative velocity and keep tangent at cone edge
+							TVector<T, d> Tangent = (RelativeVelocity - TVector<T, d>::DotProduct(RelativeVelocity, Contact.Normal) * Contact.Normal).GetSafeNormal();
+							TVector<T, d> DirectionalFactor = Factor * (Contact.Normal - Friction * Tangent);
+							T ImpulseDenominator = TVector<T, d>::DotProduct(Contact.Normal, DirectionalFactor);
+							if (!ensureMsgf(FMath::Abs(ImpulseDenominator) > SMALL_NUMBER, TEXT("Constraint:%s\n\nParticle:%s\n\nLevelset:%s\n\nDirectionalFactor:%s, ImpulseDenominator:%f"),
+								*Constraint.ToString(),
+								*Particle0->ToString(),
+								*Particle1->ToString(),
+								*DirectionalFactor.ToString(), ImpulseDenominator))
+							{
+								ImpulseDenominator = (T)1;
+							}
+
+							const T ImpulseMag = -(1 + Restitution) * RelativeNormalVelocity / ImpulseDenominator;
+							Impulse = ImpulseMag * (Contact.Normal - Friction * Tangent);
+						}
+					}
+					else
+					{
+						T ImpulseDenominator = TVector<T, d>::DotProduct(Contact.Normal, Factor * Contact.Normal);
+						TVector<T, d> ImpulseNumerator = -(1 + Restitution) * TVector<T, d>::DotProduct(RelativeVelocity, Contact.Normal)* Contact.Normal;
+						if (!ensureMsgf(FMath::Abs(ImpulseDenominator) > SMALL_NUMBER, TEXT("Constraint:%s\n\nParticle:%s\n\nLevelset:%s\n\nFactor*Constraint.Normal:%s, ImpulseDenominator:%f"),
+							*Constraint.ToString(),
+							*Particle0->ToString(),
+							*Particle1->ToString(),
+							*(Factor * Contact.Normal).ToString(), ImpulseDenominator))
+						{
+							ImpulseDenominator = (T)1;
+						}
+						Impulse = ImpulseNumerator / ImpulseDenominator;
+					}
+					Impulse = GetEnergyClampedImpulse(Constraint, Impulse, VectorToPoint1, VectorToPoint2, Body1Velocity, Body2Velocity);
+					Constraint.AccumulatedImpulse += Impulse;
+					if (PBDRigid0)
+					{
+						// Velocity update for next step
+						FVec3 NetAngularImpulse = FVec3::CrossProduct(VectorToPoint1, Impulse) + AngularImpulse;
+						FVec3 DV = PBDRigid0->InvM() * Impulse;
+						FVec3 DW = WorldSpaceInvI1 * NetAngularImpulse;
+						PBDRigid0->V() += DV;
+						PBDRigid0->W() += DW;
+						// Position update as part of pbd
+						PBDRigid0->P() += DV * IterationParameters.Dt;
+						PBDRigid0->Q() += FRotation3::FromElements(DW, 0.f) * Q0 * IterationParameters.Dt * T(0.5);
+						PBDRigid0->Q().Normalize();
+					}
+					if (PBDRigid1)
+					{
+						// Velocity update for next step
+						FVec3 NetAngularImpulse = FVec3::CrossProduct(VectorToPoint2, -Impulse) - AngularImpulse;
+						FVec3 DV = -PBDRigid1->InvM() * Impulse;
+						FVec3 DW = WorldSpaceInvI2 * NetAngularImpulse;
+						PBDRigid1->V() += DV;
+						PBDRigid1->W() += DW;
+						// Position update as part of pbd
+						PBDRigid1->P() += DV * IterationParameters.Dt;
+						PBDRigid1->Q() += FRotation3::FromElements(DW, 0.f) * Q1 * IterationParameters.Dt * T(0.5);
+						PBDRigid1->Q().Normalize();
+					}
+				}
+			}
+		}
+		template void Apply<float, 3>(TRigidBodyPointContactConstraint<float, 3>&, float, TPointContactIterationParameters<float> &, TPointContactParticleParameters<float> &);
+
+		template<typename T, int d>
+		void ApplyPushOut(TRigidBodyPointContactConstraint<T, d>& Constraint, T Thickness, const TSet<const TGeometryParticleHandle<T, d>*>& IsTemporarilyStatic,
+			TPointContactIterationParameters<T> & IterationParameters, TPointContactParticleParameters<T> & ParticleParameters)
+		{
+			TGeometryParticleHandle<T, d>* Particle0 = Constraint.Particle[0];
+			TGeometryParticleHandle<T, d>* Particle1 = Constraint.Particle[1];
+			TPBDRigidParticleHandle<T, d>* PBDRigid0 = Particle0->AsDynamic();
+			TPBDRigidParticleHandle<T, d>* PBDRigid1 = Particle1->AsDynamic();
+
+			const TVector<T, d> ZeroVector = TVector<T, d>(0);
+			const TRotation<T, d>& Q0 = PBDRigid0 ? PBDRigid0->Q() : Particle0->R();
+			const TRotation<T, d>& Q1 = PBDRigid1 ? PBDRigid1->Q() : Particle1->R();
+			const TVector<T, d>& P0 = PBDRigid0 ? PBDRigid0->P() : Particle0->X();
+			const TVector<T, d>& P1 = PBDRigid1 ? PBDRigid1->P() : Particle1->X();
+			const TVector<T, d>& V0 = PBDRigid0 ? PBDRigid0->V() : ZeroVector;
+			const TVector<T, d>& V1 = PBDRigid1 ? PBDRigid1->V() : ZeroVector;
+			const TVector<T, d>& W0 = PBDRigid0 ? PBDRigid0->W() : ZeroVector;
+			const TVector<T, d>& W1 = PBDRigid1 ? PBDRigid1->W() : ZeroVector;
+			const bool IsTemporarilyStatic0 = IsTemporarilyStatic.Contains(Particle0);
+			const bool IsTemporarilyStatic1 = IsTemporarilyStatic.Contains(Particle1);
+	
+			TSerializablePtr<FChaosPhysicsMaterial> PhysicsMaterial0, PhysicsMaterial1;
+			if (ParticleParameters.PhysicsMaterials)
+			{
+				PhysicsMaterial0 = Particle0->AuxilaryValue(*ParticleParameters.PhysicsMaterials);
+				PhysicsMaterial1 = Particle1->AuxilaryValue(*ParticleParameters.PhysicsMaterials);
+			}
+
+			for (int32 PairIteration = 0; PairIteration < IterationParameters.NumPairIterations; ++PairIteration)
+			{
+				UpdateConstraint<ECollisionUpdateType::Deepest>(Thickness, Constraint);
+
+				const typename TRigidBodyPointContactConstraint<T, d>::FManifold & Contact = Constraint.Manifold;
+
+				if (Contact.Phi >= Thickness)
+				{
+					break;
+				}
+
+				if ((!PBDRigid0 || IsTemporarilyStatic0) && (!PBDRigid1 || IsTemporarilyStatic1))
+				{
+					break;
+				}
+
+				*IterationParameters.NeedsAnotherIteration = true;
+				PMatrix<T, d, d> WorldSpaceInvI1 = PBDRigid0 ? (Q0 * FMatrix::Identity).GetTransposed() * PBDRigid0->InvI() * (Q0 * FMatrix::Identity) : PMatrix<T, d, d>(0);
+				PMatrix<T, d, d> WorldSpaceInvI2 = PBDRigid1 ? (Q1 * FMatrix::Identity).GetTransposed() * PBDRigid1->InvI() * (Q1 * FMatrix::Identity) : PMatrix<T, d, d>(0);
+				TVector<T, d> VectorToPoint1 = Contact.Location - P0;
+				TVector<T, d> VectorToPoint2 = Contact.Location - P1;
+				PMatrix<T, d, d> Factor =
+					(PBDRigid0 ? ComputeFactorMatrix3(VectorToPoint1, WorldSpaceInvI1, PBDRigid0->InvM()) : PMatrix<T, d, d>(0)) +
+					(PBDRigid1 ? ComputeFactorMatrix3(VectorToPoint2, WorldSpaceInvI2, PBDRigid1->InvM()) : PMatrix<T, d, d>(0));
+				T Numerator = FMath::Min((T)(IterationParameters.Iteration + 2), (T)IterationParameters.NumIterations);
+				T ScalingFactor = Numerator / (T)IterationParameters.NumIterations;
+
+				//if pushout is needed we better fix relative velocity along normal. Treat it as if 0 restitution
+				TVector<T, d> Body1Velocity = V0 + TVector<T, d>::CrossProduct(W0, VectorToPoint1);
+				TVector<T, d> Body2Velocity = V1 + TVector<T, d>::CrossProduct(W1, VectorToPoint2);
+				TVector<T, d> RelativeVelocity = Body1Velocity - Body2Velocity;
+				const T RelativeVelocityDotNormal = TVector<T, d>::DotProduct(RelativeVelocity, Contact.Normal);
+				if (RelativeVelocityDotNormal < 0)
+				{
+					T ImpulseDenominator = TVector<T, d>::DotProduct(Contact.Normal, Factor * Contact.Normal);
+					TVector<T, d> ImpulseNumerator = -TVector<T, d>::DotProduct(RelativeVelocity, Contact.Normal) * Contact.Normal * ScalingFactor;
+					if (!ensureMsgf(FMath::Abs(ImpulseDenominator) > SMALL_NUMBER, TEXT("ApplyPushout Constraint:%s\n\nParticle:%s\n\nLevelset:%s\n\nFactor*Contact.Normal:%s, ImpulseDenominator:%f"),
+						*Constraint.ToString(),
+						*Particle0->ToString(),
+						*Particle1->ToString(),
+						*(Factor*Contact.Normal).ToString(), ImpulseDenominator))
+					{
+						ImpulseDenominator = (T)1;
+					}
+
+					TVector<T, d> VelocityFixImpulse = ImpulseNumerator / ImpulseDenominator;
+					VelocityFixImpulse = GetEnergyClampedImpulse(Constraint, VelocityFixImpulse, VectorToPoint1, VectorToPoint2, Body1Velocity, Body2Velocity);
+					Constraint.AccumulatedImpulse += VelocityFixImpulse;	//question: should we track this?
+					if (!IsTemporarilyStatic0 && PBDRigid0)
+					{
+						TVector<T, d> AngularImpulse = TVector<T, d>::CrossProduct(VectorToPoint1, VelocityFixImpulse);
+						PBDRigid0->V() += PBDRigid0->InvM() * VelocityFixImpulse;
+						PBDRigid0->W() += WorldSpaceInvI1 * AngularImpulse;
+
+					}
+
+					if (!IsTemporarilyStatic1 && PBDRigid1)
+					{
+						TVector<T, d> AngularImpulse = TVector<T, d>::CrossProduct(VectorToPoint2, -VelocityFixImpulse);
+						PBDRigid1->V() -= PBDRigid1->InvM() * VelocityFixImpulse;
+						PBDRigid1->W() += WorldSpaceInvI2 * AngularImpulse;
+					}
+
+				}
+
+
+				TVector<T, d> Impulse = PMatrix<T, d, d>(Factor.Inverse()) * ((-Contact.Phi + Thickness) * ScalingFactor * Contact.Normal);
+				TVector<T, d> AngularImpulse1 = TVector<T, d>::CrossProduct(VectorToPoint1, Impulse);
+				TVector<T, d> AngularImpulse2 = TVector<T, d>::CrossProduct(VectorToPoint2, -Impulse);
+				if (!IsTemporarilyStatic0 && PBDRigid0)
+				{
+					PBDRigid0->P() += PBDRigid0->InvM() * Impulse;
+					PBDRigid0->Q() = TRotation<T, d>::FromVector(WorldSpaceInvI1 * AngularImpulse1) * Q0;
+					PBDRigid0->Q().Normalize();
+				}
+				if (!IsTemporarilyStatic1 && PBDRigid1)
+				{
+					PBDRigid1->P() -= PBDRigid1->InvM() * Impulse;
+					PBDRigid1->Q() = TRotation<T, d>::FromVector(WorldSpaceInvI2 * AngularImpulse2) * Q1;
+					PBDRigid1->Q().Normalize();
+				}
+			}
+		}
+		template void ApplyPushOut<float, 3>(TRigidBodyPointContactConstraint<float,3>&, float , const TSet<const TGeometryParticleHandle<float,3>*>&,
+			TPointContactIterationParameters<float> & IterationParameters, TPointContactParticleParameters<float> & ParticleParameters);
+
+	} // Collisions
+
+}// Chaos
+
diff --git a/Engine/Source/Runtime/Experimental/Chaos/Private/Chaos/PBDCollisionConstraintSingleContactUtil.cpp b/Engine/Source/Runtime/Experimental/Chaos/Private/Chaos/PBDCollisionConstraintSingleContactUtil.cpp
deleted file mode 100644
diff --git a/Engine/Source/Runtime/Experimental/Chaos/Public/Chaos/PBDCollisionConstraintPlaneContactUtil.h b/Engine/Source/Runtime/Experimental/Chaos/Public/Chaos/PBDCollisionConstraintPlaneContactUtil.h
--- a/Engine/Source/Runtime/Experimental/Chaos/Public/Chaos/PBDCollisionConstraintPlaneContactUtil.h
+++ b/Engine/Source/Runtime/Experimental/Chaos/Public/Chaos/PBDCollisionConstraintPlaneContactUtil.h
@@ -35,13 +35,13 @@ namespace Chaos
 			
 		template<typename T = float, int d = 3>
 		void Apply(TRigidBodyPointContactConstraint<T, d>& Constraint, T Thickness, 
-			TSingleContactIterationParameters<T> & IterationParameters,
-			TSingleContactParticleParameters<T> & ParticleParameters);
+			TPointContactIterationParameters<T> & IterationParameters,
+			TPointContactParticleParameters<T> & ParticleParameters);
 			
 		template<typename T=float, int d=3>
 		void ApplyPushOut(TRigidBodyPointContactConstraint<T, d>& Constraint, T Thickness, const TSet<const TGeometryParticleHandle<T, d>*>& IsTemporarilyStatic,
-			TSingleContactIterationParameters<T> & IterationParameters,
-			TSingleContactParticleParameters<T> & ParticleParameters);
+			TPointContactIterationParameters<T> & IterationParameters,
+			TPointContactParticleParameters<T> & ParticleParameters);
 	*/
 	}
 
diff --git a/Engine/Source/Runtime/Experimental/Chaos/Public/Chaos/PBDCollisionConstraintPointContactUtil.h b/Engine/Source/Runtime/Experimental/Chaos/Public/Chaos/PBDCollisionConstraintPointContactUtil.h
new file mode 100644
--- /dev/null
+++ b/Engine/Source/Runtime/Experimental/Chaos/Public/Chaos/PBDCollisionConstraintPointContactUtil.h
@@ -0,0 +1,46 @@
+// Copyright 1998-2019 Epic Games, Inc. All Rights Reserved.
+#pragma once
+
+#include "Chaos/ExternalCollisionData.h"
+#include "Chaos/PBDCollisionConstraint.h"
+#include "Chaos/CollisionResolutionTypes.h"
+
+namespace Chaos
+{
+	namespace Collisions
+	{
+		template<class T = float>
+		struct TPointContactParticleParameters {
+			TArrayCollectionArray<bool>* Collided;
+			const TArrayCollectionArray<TSerializablePtr<FChaosPhysicsMaterial>>* PhysicsMaterials;
+			T FrictionOverride;
+			T AngularFrictionOverride;
+		};
+
+		template<class T = float>
+		struct TPointContactIterationParameters {
+			const T Dt; 
+			const int32 Iteration; 
+			const int32 NumIterations; 
+			const int32 NumPairIterations;
+			bool* NeedsAnotherIteration;
+		};
+
+		template<ECollisionUpdateType UpdateType, typename T = float, int d = 3>
+		void Update(const T Thickness, TRigidBodyPointContactConstraint<T, d>& Constraint);
+
+		template<ECollisionUpdateType UpdateType, typename T=float, int d=3>
+		void UpdateConstraint(const T Thickness, TRigidBodyPointContactConstraint<T, d>& Constraint);
+			
+		template<typename T = float, int d = 3>
+		void Apply(TRigidBodyPointContactConstraint<T, d>& Constraint, T Thickness, 
+			TPointContactIterationParameters<T> & IterationParameters,
+			TPointContactParticleParameters<T> & ParticleParameters);
+			
+		template<typename T=float, int d=3>
+		void ApplyPushOut(TRigidBodyPointContactConstraint<T, d>& Constraint, T Thickness, const TSet<const TGeometryParticleHandle<T, d>*>& IsTemporarilyStatic,
+			TPointContactIterationParameters<T> & IterationParameters,
+			TPointContactParticleParameters<T> & ParticleParameters);
+	}
+
+}
diff --git a/Engine/Source/Runtime/Experimental/Chaos/Public/Chaos/PBDCollisionConstraintSingleContactUtil.h b/Engine/Source/Runtime/Experimental/Chaos/Public/Chaos/PBDCollisionConstraintSingleContactUtil.h
deleted file mode 100644
