Chaos : Collisions

Interface cleanup.

diff --git a/Engine/Source/Runtime/Experimental/Chaos/Private/Chaos/ChaosDebugDraw.cpp b/Engine/Source/Runtime/Experimental/Chaos/Private/Chaos/ChaosDebugDraw.cpp
--- a/Engine/Source/Runtime/Experimental/Chaos/Private/Chaos/ChaosDebugDraw.cpp
+++ b/Engine/Source/Runtime/Experimental/Chaos/Private/Chaos/ChaosDebugDraw.cpp
@@ -139,7 +139,7 @@ namespace Chaos
 		{
 			if (ConstraintHandle->GetType() == TCollisionConstraintBase<float,3>::FType::SinglePoint)
 			{
-				const TRigidBodyPointContactConstraint<FReal, 3>& Contact = ConstraintHandle->GetSingleContact();
+				const TRigidBodyPointContactConstraint<FReal, 3>& Contact = ConstraintHandle->GetPointContact();
 				if (Contact.GetPhi() > 0)
 				{
 					ColorScale = ColorScale * (FReal)0.1;
diff --git a/Engine/Source/Runtime/Experimental/Chaos/Private/Chaos/PBDCollisionConstraint.cpp b/Engine/Source/Runtime/Experimental/Chaos/Private/Chaos/PBDCollisionConstraint.cpp
--- a/Engine/Source/Runtime/Experimental/Chaos/Private/Chaos/PBDCollisionConstraint.cpp
+++ b/Engine/Source/Runtime/Experimental/Chaos/Private/Chaos/PBDCollisionConstraint.cpp
@@ -68,11 +68,11 @@ namespace Chaos
 
 	template<typename T, int d>
 	TPBDCollisionConstraint<T, d>::TPBDCollisionConstraint(
-		const TPBDRigidsSOAs<T, d>& InParticles, 
-		TArrayCollectionArray<bool>& Collided, 
-		const TArrayCollectionArray<TSerializablePtr<FChaosPhysicsMaterial>>& InPerParticleMaterials, 
-		const int32 InApplyPairIterations /*= 1*/, 
-		const int32 InApplyPushOutPairIterations /*= 1*/, 
+		const TPBDRigidsSOAs<T, d>& InParticles,
+		TArrayCollectionArray<bool>& Collided,
+		const TArrayCollectionArray<TSerializablePtr<FChaosPhysicsMaterial>>& InPerParticleMaterials,
+		const int32 InApplyPairIterations /*= 1*/,
+		const int32 InApplyPushOutPairIterations /*= 1*/,
 		const T Thickness /*= (T)0*/)
 		: Particles(InParticles)
 		, MCollided(Collided)
@@ -92,79 +92,6 @@ namespace Chaos
 	{
 	}
 
-	DECLARE_CYCLE_STAT(TEXT("TPBDCollisionConstraint::Reset"), STAT_CollisionConstraintsReset, STATGROUP_Chaos);
-
-	template<typename T, int d>
-	void TPBDCollisionConstraint<T, d>::Reset()
-	{
-		SCOPE_CYCLE_COUNTER(STAT_CollisionConstraintsReset);
-
-		TArray<FConstraintContainerHandle*> CopyOfHandles = Handles;
-
-		for (FConstraintContainerHandle* ContactHandle : CopyOfHandles)
-		{
-			//if (!bEnableCollisions)
-			{
-				RemoveConstraint(ContactHandle);
-			}
-		}
-
-
-		MAngularFriction = 0;
-		bUseCCD = false;
-	}
-
-	template<typename T, int d>
-	void TPBDCollisionConstraint<T, d>::RemoveConstraint(FConstraintContainerHandle* Handle)
-	{
-		int32 Idx = Handle->GetConstraintIndex();
-		typename TCollisionConstraintBase<T, d>::FType ConstraintType = Handle->GetType();
-
-		Handles.RemoveAtSwap(Idx);
-		PointConstraints.RemoveAtSwap(Idx);
-		if (Idx < PointConstraints.Num())
-		{
-			Handles[Idx]->SetConstraintIndex(Idx, ConstraintType);
-		}
-
-		ensure(Handles.Num() == PointConstraints.Num());
-		HandleAllocator.FreeHandle(Handle);
-	}
-
-	template<typename T, int d>
-	void TPBDCollisionConstraint<T, d>::RemoveConstraints(const TSet<TGeometryParticleHandle<T, d>*>&  InHandleSet)
-	{
-		const TArray<TGeometryParticleHandle<T, d>*> HandleArray = InHandleSet.Array();
-		for (auto ParticleHandle : HandleArray)
-		{
-			TArray<FConstraintContainerHandle*> CopyOfHandles = Handles;
-
-			for (FConstraintContainerHandle* ContactHandle : CopyOfHandles)
-			{
-				TVector<TGeometryParticleHandle<T, d>*, 2> ConstraintParticles = ContactHandle->GetConstrainedParticles();
-				if (ConstraintParticles[1] == ParticleHandle || ConstraintParticles[0] == ParticleHandle)
-				{
-					RemoveConstraint(ContactHandle);
-				}
-			}
-		}
-	}
-
-	template<typename T, int d>
-	void TPBDCollisionConstraint<T, d>::ApplyCollisionModifier(const TFunction<ECollisionModifierResult(const FConstraintContainerHandle* Handle)>& CollisionModifier)
-	{
-		TArray<FConstraintContainerHandle*> CopyOfHandles = Handles;
-
-		for (FConstraintContainerHandle* ContactHandle : CopyOfHandles)
-		{
-			ECollisionModifierResult Result = CollisionModifier(ContactHandle);
-			if (Result == ECollisionModifierResult::Disabled)
-			{
-				RemoveConstraint(ContactHandle);
-			}
-		}
-	}
-
 	template<typename T, int d>
 	void TPBDCollisionConstraint<T, d>::SetPostComputeCallback(const TRigidBodyContactConstraintsPostComputeCallback<T, d>& Callback)
 	{
@@ -201,6 +128,12 @@ namespace Chaos
 		PostApplyPushOutCallback = nullptr;
 	}
 
+	DEFINE_STAT(STAT_ComputeConstraints);
+	DEFINE_STAT(STAT_ComputeConstraintsNP);
+	DEFINE_STAT(STAT_ComputeConstraintsBP);
+	DEFINE_STAT(STAT_ComputeConstraintsSU);
+
+
 	template<typename T, int d>
 	void TPBDCollisionConstraint<T, d>::UpdatePositionBasedState(/*const TPBDRigidParticles<T, d>& InParticles, const TArray<int32>& InIndices,*/ const T Dt)
 	{
@@ -221,10 +154,88 @@ namespace Chaos
 		}
 	}
 
-	DEFINE_STAT(STAT_ComputeConstraints);
-	DEFINE_STAT(STAT_ComputeConstraintsNP);
-	DEFINE_STAT(STAT_ComputeConstraintsBP);
-	DEFINE_STAT(STAT_ComputeConstraintsSU);
+	template<typename T, int d>
+	void TPBDCollisionConstraint<T, d>::ConstructConstraints(TGeometryParticleHandle<T, d>* Particle0, TGeometryParticleHandle<T, d>* Particle1, const T Thickness, TRigidBodyPointContactConstraint<T, d> & Constraint)
+	{
+		if (ensure(Particle0 && Particle1))
+		{
+			Collisions::ConstructConstraintsImpl<T, d>(Particle0, Particle1, Particle0->Geometry().Get(), Particle1->Geometry().Get(), Thickness, Constraint);
+		}
+	}
+
+	DECLARE_CYCLE_STAT(TEXT("TPBDCollisionConstraint::Reset"), STAT_CollisionConstraintsReset, STATGROUP_Chaos);
+	template<typename T, int d>
+	void TPBDCollisionConstraint<T, d>::Reset()
+	{
+		SCOPE_CYCLE_COUNTER(STAT_CollisionConstraintsReset);
+
+		TArray<FConstraintContainerHandle*> CopyOfHandles = Handles;
+
+		for (FConstraintContainerHandle* ContactHandle : CopyOfHandles)
+		{
+			//if (!bEnableCollisions)
+			{
+				RemoveConstraint(ContactHandle);
+			}
+		}
+
+
+		MAngularFriction = 0;
+		bUseCCD = false;
+	}
+
+	template<typename T, int d>
+	void TPBDCollisionConstraint<T, d>::ApplyCollisionModifier(const TFunction<ECollisionModifierResult(const FConstraintContainerHandle* Handle)>& CollisionModifier)
+	{
+		TArray<FConstraintContainerHandle*> CopyOfHandles = Handles;
+
+		for (FConstraintContainerHandle* ContactHandle : CopyOfHandles)
+		{
+			ECollisionModifierResult Result = CollisionModifier(ContactHandle);
+			if (Result == ECollisionModifierResult::Disabled)
+			{
+				RemoveConstraint(ContactHandle);
+			}
+		}
+	}
+
+	template<typename T, int d>
+	void TPBDCollisionConstraint<T, d>::RemoveConstraints(const TSet<TGeometryParticleHandle<T, d>*>&  InHandleSet)
+	{
+		const TArray<TGeometryParticleHandle<T, d>*> HandleArray = InHandleSet.Array();
+		for (auto ParticleHandle : HandleArray)
+		{
+			TArray<FConstraintContainerHandle*> CopyOfHandles = Handles;
+
+			for (FConstraintContainerHandle* ContactHandle : CopyOfHandles)
+			{
+				TVector<TGeometryParticleHandle<T, d>*, 2> ConstraintParticles = ContactHandle->GetConstrainedParticles();
+				if (ConstraintParticles[1] == ParticleHandle || ConstraintParticles[0] == ParticleHandle)
+				{
+					RemoveConstraint(ContactHandle);
+				}
+			}
+		}
+	}
+
+	template<typename T, int d>
+	void TPBDCollisionConstraint<T, d>::RemoveConstraint(FConstraintContainerHandle* Handle)
+	{
+		int32 Idx = Handle->GetConstraintIndex();
+		typename TCollisionConstraintBase<T, d>::FType ConstraintType = Handle->GetType();
+
+		Handles.RemoveAtSwap(Idx);
+		PointConstraints.RemoveAtSwap(Idx);
+		if (Idx < PointConstraints.Num())
+		{
+			Handles[Idx]->SetConstraintIndex(Idx, ConstraintType);
+		}
+
+		ensure(Handles.Num() == PointConstraints.Num());
+		HandleAllocator.FreeHandle(Handle);
+	}
+
+
 
 	template<typename T, int d>
 	template <bool bGatherStats>
@@ -266,46 +277,29 @@ namespace Chaos
 
 
 	template<typename T, int d>
-	void TPBDCollisionConstraint<T, d>::UpdateConstraints(T Dt, const TSet<TGeometryParticleHandle<T, d>*>& AddedParticles)
+	void TPBDCollisionConstraint<T, d>::UpdateConstraints(T Dt, const TSet<TGeometryParticleHandle<T, d>*>& ParticlesSet)
 	{
 		// Clustering uses update constraints to force a re-evaluation. 
 	}
 
-
+	DECLARE_CYCLE_STAT(TEXT("TPBDCollisionConstraint::UpdateConstraint"), STAT_UpdateConstraint, STATGROUP_ChaosWide);
 	template<typename T, int d>
-	void TPBDCollisionConstraint<T, d>::Apply(const T Dt, FPointContactConstraint& Constraint, const int32 It, const int32 NumIts)
+	template<ECollisionUpdateType UpdateType>
+	void TPBDCollisionConstraint<T, d>::UpdateConstraint(const T Thickness, TRigidBodyPointContactConstraint<T, d> & Constraint)
 	{
-
+		SCOPE_CYCLE_COUNTER(STAT_UpdateConstraint);
+		//if (ConstraintHandle->GetType() == FPointContactConstraint::StaticType())
 		{
-			// 
-			// @todo(ccaulfield,brice): I think we should never get this? Revisit after particle handle refactor
-			// 
+			Collisions::UpdateConstraint<UpdateType>(Thickness, Constraint);
+		}
+	}
 
-			TGenericParticleHandle<T, d> Particle0 = TGenericParticleHandle<T, d>(Constraint.Particle[0]);
-			TGenericParticleHandle<T, d> Particle1 = TGenericParticleHandle<T, d>(Constraint.Particle[1]);
-			TPBDRigidParticleHandle<T, d>* PBDRigid0 = Particle0->AsDynamic();
-			TPBDRigidParticleHandle<T, d>* PBDRigid1 = Particle1->AsDynamic();
 
-			if (Particle0->Sleeping())
-			{
-				ensure(!PBDRigid1 || PBDRigid1->Sleeping());
-				return;
-			}
-			if (Particle1->Sleeping())
-			{
-				ensure(!PBDRigid0 || PBDRigid0->Sleeping());
-				return;
-			}
-		}
 
-		Collisions::TPointContactParticleParameters<T> ParticleParameters = { &MCollided, &MPhysicsMaterials, CollisionFrictionOverride, MAngularFriction };
-		Collisions::TPointContactIterationParameters<T> IterationParameters = { Dt, It, NumIts, MApplyPairIterations, nullptr };
-		Collisions::Apply(Constraint,MThickness, IterationParameters,  ParticleParameters);
-	}
 
 	DECLARE_CYCLE_STAT(TEXT("TPBDCollisionConstraint::Apply"), STAT_Apply, STATGROUP_Chaos);
 	template<typename T, int d>
-	void TPBDCollisionConstraint<T, d>::Apply(const T Dt, const TArray<FConstraintContainerHandle*>& InConstraintHandles, const int32 It, const int32 NumIts)
+	void TPBDCollisionConstraint<T, d>::Apply(const T Dt, const TArray<FConstraintContainerHandle*>& InConstraintHandles, const int32 Iterations, const int32 NumIterations)
 	{
 		SCOPE_CYCLE_COUNTER(STAT_Apply);
 		if (MApplyPairIterations > 0)
@@ -313,7 +307,22 @@ namespace Chaos
 			PhysicsParallelFor(InConstraintHandles.Num(), [&](int32 ConstraintHandleIndex) {
 				FConstraintContainerHandle* ConstraintHandle = InConstraintHandles[ConstraintHandleIndex];
 				check(ConstraintHandle != nullptr);
-				Apply(Dt, PointConstraints[ConstraintHandle->GetConstraintIndex()], It, NumIts);
+
+				if (ConstraintHandle->GetType() == FPointContactConstraint::StaticType())
+				{
+					Collisions::TPointContactParticleParameters<T> ParticleParameters = { &MCollided, &MPhysicsMaterials, CollisionFrictionOverride, MAngularFriction };
+					Collisions::TPointContactIterationParameters<T> IterationParameters = { Dt, Iterations, NumIterations, MApplyPairIterations, nullptr };
+					Collisions::Apply(ConstraintHandle->GetPointContact(), MThickness, IterationParameters, ParticleParameters);
+				}
+				else if (ConstraintHandle->GetType() == FPlaneContactConstraint::StaticType())
+				{
+
+				}
+				else
+				{
+					ensureMsgf(false, TEXT("Invalid constraint type"));
+				}
+
 			}, bDisableCollisionParallelFor);
 		}
 
@@ -324,36 +333,6 @@ namespace Chaos
 	}
 
 
-	template<typename T, int d>
-	void TPBDCollisionConstraint<T, d>::ApplyPushOut(const T Dt, FPointContactConstraint& Constraint, const TSet<const TGeometryParticleHandle<T, d>*>& IsTemporarilyStatic, int32 Iteration, int32 NumIterations, bool &NeedsAnotherIteration)
-	{
-		{
-			// 
-			// @todo(ccaulfield,brice): I think we should never get this? Revisit after particle handle refactor
-			// 
-
-			TGenericParticleHandle<T, d> Particle0 = TGenericParticleHandle<T, d>(Constraint.Particle[0]);
-			TGenericParticleHandle<T, d> Particle1 = TGenericParticleHandle<T, d>(Constraint.Particle[1]);
-			TPBDRigidParticleHandle<T, d>* PBDRigid0 = Particle0->AsDynamic();
-			TPBDRigidParticleHandle<T, d>* PBDRigid1 = Particle1->AsDynamic();
-
-			if (Particle0->Sleeping())
-			{
-				ensure(!PBDRigid1 || PBDRigid1->Sleeping());
-				return;
-			}
-			if (Particle1->Sleeping())
-			{
-				ensure(!PBDRigid0 || PBDRigid0->Sleeping());
-				return;
-			}
-		}
-
-		Collisions::TPointContactParticleParameters<T> ParticleParameters = { &MCollided, &MPhysicsMaterials, CollisionFrictionOverride, MAngularFriction };
-		Collisions::TPointContactIterationParameters<T> IterationParameters = { Dt, Iteration, NumIterations, MApplyPairIterations, &NeedsAnotherIteration };
-		Collisions::ApplyPushOut(Constraint, MThickness, IsTemporarilyStatic, IterationParameters, ParticleParameters);
-	}
-
 	DECLARE_CYCLE_STAT(TEXT("TPBDCollisionConstraint::ApplyPushOut"), STAT_ApplyPushOut, STATGROUP_Chaos);
 	template<typename T, int d>
 	bool TPBDCollisionConstraint<T, d>::ApplyPushOut(const T Dt, const TArray<FConstraintContainerHandle*>& InConstraintHandles, const TSet< const TGeometryParticleHandle<T, d>*>& IsTemporarilyStatic, int32 Iteration, int32 NumIterations)
@@ -363,10 +342,27 @@ namespace Chaos
 		bool NeedsAnotherIteration = false;
 		if (MApplyPushOutPairIterations > 0)
 		{
-			PhysicsParallelFor(InConstraintHandles.Num(), [&](int32 ConstraintHandleIndex) {
+			PhysicsParallelFor(InConstraintHandles.Num(), [&](int32 ConstraintHandleIndex)
+			{
 				FConstraintContainerHandle* ConstraintHandle = InConstraintHandles[ConstraintHandleIndex];
 				check(ConstraintHandle != nullptr);
-				ApplyPushOut(Dt, PointConstraints[ConstraintHandle->GetConstraintIndex()], IsTemporarilyStatic, Iteration, NumIterations, NeedsAnotherIteration);
+
+				if (ConstraintHandle->GetType() == FPointContactConstraint::StaticType())
+				{
+					Collisions::TPointContactParticleParameters<T> ParticleParameters = { &MCollided, &MPhysicsMaterials, CollisionFrictionOverride, MAngularFriction };
+					Collisions::TPointContactIterationParameters<T> IterationParameters = { Dt, Iteration, NumIterations, MApplyPairIterations, &NeedsAnotherIteration };
+					Collisions::ApplyPushOut(ConstraintHandle->GetPointContact(), MThickness, IsTemporarilyStatic, IterationParameters, ParticleParameters);
+				}
+				else if (ConstraintHandle->GetType() == FPlaneContactConstraint::StaticType())
+				{
+
+				}
+				else
+				{
+					ensureMsgf(false, TEXT("Invalid constraint type"));
+				}
+
+
 			}, bDisableCollisionParallelFor);
 		}
 
@@ -378,36 +374,11 @@ namespace Chaos
 	}
 
 
-	template<typename T, int d>
-	TRigidTransform<T, d> GetTransform(const TGeometryParticleHandle<T, d>* Particle)
-	{
-		TGenericParticleHandle<T, d> Generic = const_cast<TGeometryParticleHandle<T, d>*>(Particle);	//TODO: give a const version of the generic API
-		return TRigidTransform<T, d>(Generic->P(), Generic->Q());
-	}
-
-	template<typename T, int d>
-	void TPBDCollisionConstraint<T, d>::ConstructConstraints(TGeometryParticleHandle<T, d>* Particle0, TGeometryParticleHandle<T, d>* Particle1, const T Thickness, TRigidBodyPointContactConstraint<T, d> & Constraint)
-	{
-		if (ensure(Particle0 && Particle1))
-		{
-			Collisions::ConstructConstraintsImpl<T, d>(Particle0, Particle1, Particle0->Geometry().Get(), Particle1->Geometry().Get(), Thickness, Constraint);
-		}
-	}
-
-	DECLARE_CYCLE_STAT(TEXT("TPBDCollisionConstraint::UpdateConstraint"), STAT_UpdateConstraint, STATGROUP_ChaosWide);
-
-	template<typename T, int d>
-	template<ECollisionUpdateType UpdateType>
-	void TPBDCollisionConstraint<T, d>::UpdateConstraint(const T Thickness, FPointContactConstraint& Constraint)
-	{
-		SCOPE_CYCLE_COUNTER(STAT_UpdateConstraint);
-		Collisions::UpdateConstraint<UpdateType>(Thickness, Constraint);
-	}
 
 	template class TAccelerationStructureHandle<float, 3>;
 	template class CHAOS_API TPBDCollisionConstraint<float, 3>;
-	template void TPBDCollisionConstraint<float, 3>::UpdateConstraint<ECollisionUpdateType::Any>(const float Thickness, FPointContactConstraint& Constraint);
-	template void TPBDCollisionConstraint<float, 3>::UpdateConstraint<ECollisionUpdateType::Deepest>(const float Thickness, FPointContactConstraint& Constraint);
+	template void TPBDCollisionConstraint<float, 3>::UpdateConstraint<ECollisionUpdateType::Any>(const float Thickness, TRigidBodyPointContactConstraint<float,3> & Constraint);
+	template void TPBDCollisionConstraint<float, 3>::UpdateConstraint<ECollisionUpdateType::Deepest>(const float Thickness, TRigidBodyPointContactConstraint<float,3> & Constraint);
 	template void TPBDCollisionConstraint<float, 3>::ComputeConstraints<false>(const TPBDCollisionConstraint<float, 3>::FAccelerationStructure&, float Dt);
 	template void TPBDCollisionConstraint<float, 3>::ComputeConstraints<true>(const TPBDCollisionConstraint<float, 3>::FAccelerationStructure&, float Dt);
 }
diff --git a/Engine/Source/Runtime/Experimental/Chaos/Public/Chaos/CollisionResolutionTypes.h b/Engine/Source/Runtime/Experimental/Chaos/Public/Chaos/CollisionResolutionTypes.h
--- a/Engine/Source/Runtime/Experimental/Chaos/Public/Chaos/CollisionResolutionTypes.h
+++ b/Engine/Source/Runtime/Experimental/Chaos/Public/Chaos/CollisionResolutionTypes.h
@@ -234,11 +234,11 @@ namespace Chaos
 
 		// Handle API
 
-		const TRigidBodyPointContactConstraint<T, d>& GetSingleContact() const { return ConstraintContainer->PointConstraints[ConstraintIndex]; }
-		TRigidBodyPointContactConstraint<T, d>& GetSingleContact() { return ConstraintContainer->PointConstraints[ConstraintIndex]; }
+		const TRigidBodyPointContactConstraint<T, d>& GetPointContact() const { check(GetType() == FConstraintBase::FType::SinglePoint); return ConstraintContainer->PointConstraints[ConstraintIndex]; }
+		TRigidBodyPointContactConstraint<T, d>& GetPointContact() { check(GetType() == FConstraintBase::FType::SinglePoint); return ConstraintContainer->PointConstraints[ConstraintIndex]; }
 
-		const TRigidBodyPlaneContactConstraint<T, d>& GetPlaneContact() const { return ConstraintContainer->PlaneConstraints[ConstraintIndex]; }
-		TRigidBodyPlaneContactConstraint<T, d>& GetPlaneContact() { return ConstraintContainer->PlaneConstraints[ConstraintIndex]; }
+		const TRigidBodyPlaneContactConstraint<T, d>& GetPlaneContact() const { check(GetType() == FConstraintBase::FType::Plane); return ConstraintContainer->PlaneConstraints[ConstraintIndex]; }
+		TRigidBodyPlaneContactConstraint<T, d>& GetPlaneContact() { check(GetType() == FConstraintBase::FType::Plane); return ConstraintContainer->PlaneConstraints[ConstraintIndex]; }
 
 
 		typename FConstraintBase::FType GetType() const { return ConstraintType; }
@@ -255,7 +255,7 @@ namespace Chaos
 		{
 			if (ConstraintType == FConstraintBase::FType::SinglePoint)
 			{
-				return GetSingleContact().GetLocation();
+				return GetPointContact().GetLocation();
 			}
 			ensure(false);
 			return TVector<T, d>(0);
@@ -265,7 +265,7 @@ namespace Chaos
 		{
 			if (ConstraintType == FConstraintBase::FType::SinglePoint)
 			{
-				return GetSingleContact().AccumulatedImpulse;
+				return GetPointContact().AccumulatedImpulse;
 			}
 			ensure(false);
 
@@ -277,7 +277,7 @@ namespace Chaos
 		{
 			if (ConstraintType == FConstraintBase::FType::SinglePoint)
 			{
-				const TRigidBodyPointContactConstraint<float, 3>& Contact = GetSingleContact();
+				const TRigidBodyPointContactConstraint<float, 3>& Contact = GetPointContact();
 				return { Contact.Particle[0], Contact.Particle[1] };
 			}
 			ensure(false);
@@ -290,7 +290,7 @@ namespace Chaos
 		{
 			if (ConstraintType == FConstraintBase::FType::SinglePoint)
 			{
-				return { GetSingleContact().Particle[0], GetSingleContact().Particle[1] };
+				return { GetPointContact().Particle[0], GetPointContact().Particle[1] };
 			}
 			ensure(false);
 			return { nullptr,nullptr };
diff --git a/Engine/Source/Runtime/Experimental/Chaos/Public/Chaos/PBDCollisionConstraint.h b/Engine/Source/Runtime/Experimental/Chaos/Public/Chaos/PBDCollisionConstraint.h
--- a/Engine/Source/Runtime/Experimental/Chaos/Public/Chaos/PBDCollisionConstraint.h
+++ b/Engine/Source/Runtime/Experimental/Chaos/Public/Chaos/PBDCollisionConstraint.h
@@ -41,10 +41,10 @@
 using TRigidBodyContactConstraintsPostComputeCallback = TFunction<void()>;
 
 template<typename T, int d>
-using TRigidBodyContactConstraintsPostApplyCallback = TFunction<void(const T Dt, const TArray<TPBDCollisionConstraintHandle<T, d>*>& InConstraintHandles)>;
+using TRigidBodyContactConstraintsPostApplyCallback = TFunction<void(const T Dt, const TArray<TPBDCollisionConstraintHandle<T, d>*>&)>;
 
 template<typename T, int d>
-using TRigidBodyContactConstraintsPostApplyPushOutCallback = TFunction<void(const T Dt, const TArray<TPBDCollisionConstraintHandle<T, d>*>& InConstraintHandles, bool bRequiresAnotherIteration)>;
+using TRigidBodyContactConstraintsPostApplyPushOutCallback = TFunction<void(const T Dt, const TArray<TPBDCollisionConstraintHandle<T, d>*>&, bool)>;
 
 /**
  * Manages a set of contact constraints:
@@ -52,7 +52,6 @@
  *	- Responsible for applying corrections to particles affected by the constraints.
  * @todo(ccaulfield): rename to TPBDCollisionConstraints
  * @todo(ccaulfield): remove TPBDCollisionConstraintAccessor
- * @todo(ccaulfield): separate collision detection from constraint container
  */
 template<typename T, int d>
 class CHAOS_API TPBDCollisionConstraint : public FPBDConstraintContainer
@@ -62,20 +61,122 @@
 	friend class TPBDCollisionConstraintHandle<T, d>;
 
 	using Base = FPBDConstraintContainer;
-	using FConstraintHandleID = TPair<const TGeometryParticleHandle<T, d>*, const TGeometryParticleHandle<T, d>*>;
-	using FConstraintContainerHandle = TPBDCollisionConstraintHandle<T, d>; // @todo(brice) rename this to FCollisionConstraintHandle
-	using FConstraintHandleAllocator = TConstraintHandleAllocator<TPBDCollisionConstraint<T, d>>;
+	using FHandles = TArray<TPBDCollisionConstraintHandle<T, d>*>;
+	using FConstraintContainerHandle = TPBDCollisionConstraintHandle<T, d>;
  	using FPointContactConstraint = TRigidBodyPointContactConstraint<T, d>;
 	using FPlaneContactConstraint = TRigidBodyPlaneContactConstraint<T, d>;
-	using FHandles = TArray<FConstraintContainerHandle*>;
+	using FConstraintHandleAllocator = TConstraintHandleAllocator<TPBDCollisionConstraint<T, d>>;
+	using FConstraintHandleID = TPair<const TGeometryParticleHandle<T, d>*, const TGeometryParticleHandle<T, d>*>;
 
-	TPBDCollisionConstraint(const TPBDRigidsSOAs<T,d>& InParticles, TArrayCollectionArray<bool>& Collided, const TArrayCollectionArray<TSerializablePtr<FChaosPhysicsMaterial>>& PerParticleMaterials, const int32 ApplyPairIterations = 1, const int32 ApplyPushOutPairIterations = 1, const T Thickness = (T)0);
+
+	TPBDCollisionConstraint(const TPBDRigidsSOAs<T,d>& InParticles, 
+		TArrayCollectionArray<bool>& Collided, 
+		const TArrayCollectionArray<TSerializablePtr<FChaosPhysicsMaterial>>& PerParticleMaterials, 
+		const int32 ApplyPairIterations = 1, const int32 ApplyPushOutPairIterations = 1, const T Thickness = (T)0);
+
 	virtual ~TPBDCollisionConstraint() {}
 
+
+	/**
+	 * Generate all contact constraints.
+	 */
+	void UpdatePositionBasedState(const T Dt);
+
+
+	/**
+	*  Reset the constraint frame. 
+	*/
+	void Reset();
+
+	/**
+	 * Apply a modifier to the constraints and specify which constraints should be disabled.
+	 * You would probably call this in the PostComputeCallback. Prefer this to calling RemoveConstraints in a loop,
+	 * so you don't have to worry about constraint iterator/indices changing.
+	 */
+	void ApplyCollisionModifier(const TFunction<ECollisionModifierResult(const FConstraintContainerHandle* Handle)>& CollisionModifier);
+
+
+	/**
+	* Remove the constraints associated with the ParticleHandle.
+	*/
+	void RemoveConstraints(const TSet<TGeometryParticleHandle<T, d>*>&  ParticleHandle);
+
+
+	/**
+	* Remove the constraint, update the handle, and any maps. 
+	*/
+	void RemoveConstraint(FConstraintContainerHandle*);
+
+
+	/**
+	*  Process constraints acceleration structures. 
+	*/
+	//using FAccelerationStructure = TBoundingVolume<TAccelerationStructureHandle<T, d>, T, d>;
+	//using FAccelerationStructure = TAABBTree<TAccelerationStructureHandle<T, d>, TAABBTreeLeafArray<TAccelerationStructureHandle<T, d>, T>, T>;
+	using FAccelerationStructure = ISpatialAcceleration< TAccelerationStructureHandle<T, d>, T, d>;
+	void SetSpatialAcceleration(const FAccelerationStructure* InSpatialAcceleration) { SpatialAcceleration = InSpatialAcceleration; }
+
+	template <bool bGatherStats = false>
+	void ComputeConstraints(const FAccelerationStructure& AccelerationStructure, T Dt);
+
+	//NOTE: this should not be called by anyone other than ISpatialAccelerationCollection and CollisionConstraints - todo: make private with friends?
+	template <bool bGatherStats, typename SPATIAL_ACCELERATION>
+	void ComputeConstraintsHelperLowLevel(const SPATIAL_ACCELERATION& SpatialAcceleration, T Dt);
+
+
+	/**
+	 * Update all constraint values
+	 */
+	void UpdateConstraints(T Dt, const TSet<TGeometryParticleHandle<T, d>*>& AddedParticles);
+
+
+
+	/**
+	 * Apply corrections for the specified list of constraints. (Runs Wide!)
+	 *
+	 * @todo(ccaulfield): this runs wide. The serial/parallel decision should be in the ConstraintRule
+	 *
+	 */
+	void Apply(const T Dt, const TArray<FConstraintContainerHandle*>& InConstraintHandles, const int32 It, const int32 NumIts);
+
+
+	/**
+	 * Apply push out for the specified list of constraints.
+	 *  Return true if we need another iteration 
+	 *
+	 * @todo(ccaulfield): this runs wide. The serial/parallel decision should be in the ConstraintRule
+	 *
+	 */
+	bool ApplyPushOut(const T Dt, const TArray<FConstraintContainerHandle*>& InConstraintHandles, 
+		const TSet<const TGeometryParticleHandle<T,d>*>& IsTemporarilyStatic, int32 Iteration, int32 NumIterations);
+
+
+	/**
+	*
+	*  Callbacks
+	*
+	 * Set the callback used just after contacts are generated at the start of a frame tick.
+	 * This can be used to modify or disable constraints (via ApplyCollisionModifier).
+	 */
+	void SetPostComputeCallback(const TRigidBodyContactConstraintsPostComputeCallback<T, d>& Callback);
+	void ClearPostComputeCallback();
+
+	void SetPostApplyCallback(const TRigidBodyContactConstraintsPostApplyCallback<T, d>& Callback);
+	void ClearPostApplyCallback();
+
+	void SetPostApplyPushOutCallback(const TRigidBodyContactConstraintsPostApplyPushOutCallback<T, d>& Callback);
+	void ClearPostApplyPushOutCallback();
+
+
 	//
-	// Constraint Container API
+	// Member Access
 	//
 
+	const TArray<FConstraintContainerHandle*>& GetAllConstraintHandles() const 
+	{ 
+		return Handles; 
+	}
+
 	void SetThickness(T InThickness)
 	{
 		MThickness = InThickness;
@@ -108,7 +209,7 @@
 
 	FConstraintHandleID GetConstraintHandleID(const FPointContactConstraint & Constraint) const
 	{
-		return  FConstraintHandleID( Constraint.Particle[0], Constraint.Particle[1] );
+		return  FConstraintHandleID(Constraint.Particle[0], Constraint.Particle[1]);
 	}
 
 	FHandles& GetConstraintHandles()
@@ -120,76 +221,6 @@
 		return Handles;
 	}
 
-	void RemoveConstraint(FConstraintContainerHandle*);
 
-	/**
-	* Remove the constraints associated with the ParticleHandle.
-	*/
-	void RemoveConstraints(const TSet<TGeometryParticleHandle<T, d>*>&  ParticleHandle);
-
-	/**
-	 * Apply a modifier to the constraints and specify which constraints should be disabled.
-	 * You would probably call this in the PostComputeCallback. Prefer this to calling RemoveConstraints in a loop, 
-	 * so you don't have to worry about constraint iterator/indices changing.
-	 */
-	void ApplyCollisionModifier(const TFunction<ECollisionModifierResult(const FConstraintContainerHandle* Handle)>& CollisionModifier);
-
-	/**
-	 * Set the callback used just after contacts are generated at the start of a frame tick.
-	 * This can be used to modify or disable constraints (via ApplyCollisionModifier).
-	 */
-	void SetPostComputeCallback(const TRigidBodyContactConstraintsPostComputeCallback<T, d>& Callback);
-	void ClearPostComputeCallback();
-
-	void SetPostApplyCallback(const TRigidBodyContactConstraintsPostApplyCallback<T, d>& Callback);
-	void ClearPostApplyCallback();
-
-	void SetPostApplyPushOutCallback(const TRigidBodyContactConstraintsPostApplyPushOutCallback<T, d>& Callback);
-	void ClearPostApplyPushOutCallback();
-
-	//
-	// Island Rule API
-	//
-
-	/**
-	 * Apply corrections for the specified list of constraints. (Runs Wide!)
-	 *
-	 * @todo(ccaulfield): this runs wide. The serial/parallel decision should be in the ConstraintRule
-	 *
-	 */
-	void Apply(const T Dt, const TArray<FConstraintContainerHandle*>& InConstraintHandles, const int32 It, const int32 NumIts);
-
-	/**
-	 * Generate all contact constraints.
-	 */
-	void UpdatePositionBasedState(const T Dt);
-
-	/**
-	 * Update all constraint values
-	 */
-	void UpdateConstraints(T Dt, const TSet<TGeometryParticleHandle<T, d>*>& AddedParticles);
-
-	/**
-	 * Apply push out for the specified list of constraints.
-	 *  Return true if we need another iteration 
-	 *
-	 * @todo(ccaulfield): this runs wide. The serial/parallel decision should be in the ConstraintRule
-	 *
-	 */
-	bool ApplyPushOut(const T Dt, const TArray<FConstraintContainerHandle*>& InConstraintHandles, const TSet<const TGeometryParticleHandle<T,d>*>& IsTemporarilyStatic, int32 Iteration, int32 NumIterations);
-
-
-	const TArray<FConstraintContainerHandle*>& GetAllConstraintHandles() const { return Handles; }
-
-	//using FAccelerationStructure = TBoundingVolume<TAccelerationStructureHandle<T, d>, T, d>;
-	//using FAccelerationStructure = TAABBTree<TAccelerationStructureHandle<T, d>, TAABBTreeLeafArray<TAccelerationStructureHandle<T, d>, T>, T>;
-	using FAccelerationStructure = ISpatialAcceleration< TAccelerationStructureHandle<T, d>, T, d>;
-	void SetSpatialAcceleration(const FAccelerationStructure* InSpatialAcceleration) { SpatialAcceleration = InSpatialAcceleration; }
-
-	//NOTE: this should not be called by anyone other than ISpatialAccelerationCollection and CollisionConstraints - todo: make private with friends?
-	template <bool bGatherStats, typename SPATIAL_ACCELERATION>
-	void ComputeConstraintsHelperLowLevel(const SPATIAL_ACCELERATION& SpatialAcceleration, T Dt);
-	
-
 protected:
 		using Base::GetConstraintIndex;
@@ -197,23 +228,18 @@
 
 private:
 
-	/*
-	* Reset the constraints. 
+	/**
+	* Update the individual constraints.
 	*/
-	void Reset( );
-
-	void Apply(const T Dt, FPointContactConstraint& Constraint, const int32 It, const int32 NumIts);
-
-
-	void ApplyPushOut(const T Dt, FPointContactConstraint& Constraint, const TSet<const TGeometryParticleHandle<T, d>*>& IsTemporarilyStatic, int32 Iteration, int32 NumIterations, bool &NeedsAnotherIteration);
-
-	template <bool bGatherStats = false>
-	void ComputeConstraints(const FAccelerationStructure& AccelerationStructure, T Dt);
-
 	template<ECollisionUpdateType>
-	void UpdateConstraint(const T Thickness, FPointContactConstraint& Constraint);
+	void UpdateConstraint(const T Thickness, TRigidBodyPointContactConstraint<T, d> & Constraint);
 
-	void ConstructConstraints(TGeometryParticleHandle<T, d>* Particle0, TGeometryParticleHandle<T, d>* Particle1, const T Thickness, TRigidBodyPointContactConstraint<T, d> & Constraint);
+	/**
+	*  Build a constraint based on the two particle handles.
+	*/
+	void ConstructConstraints(TGeometryParticleHandle<T, d>* Particle0, TGeometryParticleHandle<T, d>* Particle1, 
+		const T Thickness, TRigidBodyPointContactConstraint<T, d> & Constraint);
 
 
+
 	const TPBDRigidsSOAs<T,d>& Particles;
@@ -245,8 +271,8 @@
 	const FAccelerationStructure* SpatialAcceleration;
 };
 
-extern template void TPBDCollisionConstraint<float, 3>::UpdateConstraint<ECollisionUpdateType::Any>(const float Thickness, FPointContactConstraint& Constraint);
-extern template void TPBDCollisionConstraint<float, 3>::UpdateConstraint<ECollisionUpdateType::Deepest>(const float Thickness, FPointContactConstraint& Constraint);
+extern template void TPBDCollisionConstraint<float, 3>::UpdateConstraint<ECollisionUpdateType::Any>(const float Thickness, TRigidBodyPointContactConstraint<float,3> & Constraint);
+extern template void TPBDCollisionConstraint<float, 3>::UpdateConstraint<ECollisionUpdateType::Deepest>(const float Thickness, TRigidBodyPointContactConstraint<float, 3> & Constraint);
 extern template void TPBDCollisionConstraint<float, 3>::ComputeConstraints<false>(const TPBDCollisionConstraint<float, 3>::FAccelerationStructure&, float Dt);
 extern template void TPBDCollisionConstraint<float, 3>::ComputeConstraints<true>(const TPBDCollisionConstraint<float, 3>::FAccelerationStructure&, float Dt);
 }
diff --git a/Engine/Source/Runtime/Experimental/ChaosSolvers/Private/EventDefaults.cpp b/Engine/Source/Runtime/Experimental/ChaosSolvers/Private/EventDefaults.cpp
--- a/Engine/Source/Runtime/Experimental/ChaosSolvers/Private/EventDefaults.cpp
+++ b/Engine/Source/Runtime/Experimental/ChaosSolvers/Private/EventDefaults.cpp
@@ -68,7 +68,7 @@ namespace Chaos
 				{
 					if (ContactHandle->GetType() == TPBDCollisionConstraintHandle<float, 3>::FConstraintBase::FType::SinglePoint)
 					{
-						const TRigidBodyPointContactConstraint<float, 3>& Constraint = ContactHandle->GetSingleContact();
+						const TRigidBodyPointContactConstraint<float, 3>& Constraint = ContactHandle->GetPointContact();
 
 						// Since Clustered GCs can be unioned the particleIndex representing the union 
 						// is not associated with a PhysicsProxy
@@ -107,7 +107,7 @@ namespace Chaos
 				{
 					for(int32 IdxCollision = 0; IdxCollision < ValidCollisionHandles.Num(); ++IdxCollision)
 					{
-						Chaos::TPBDCollisionConstraint<float, 3>::FPointContactConstraint const& Constraint = ValidCollisionHandles[IdxCollision]->GetSingleContact();
+						Chaos::TPBDCollisionConstraint<float, 3>::FPointContactConstraint const& Constraint = ValidCollisionHandles[IdxCollision]->GetPointContact();
 
 						TGeometryParticleHandle<float, 3>* Particle0 = Constraint.Particle[0];
 						TGeometryParticleHandle<float, 3>* Particle1 = Constraint.Particle[1];
