Integrate CL #10089537 from 4.24 to Main.

diff --git a/Engine/Plugins/Experimental/GeometryCollectionPlugin/Source/GeometryCollectionExampleCore/Private/GeometryCollection/GeometryCollectionExampleSimulation.cpp b/Engine/Plugins/Experimental/GeometryCollectionPlugin/Source/GeometryCollectionExampleCore/Private/GeometryCollection/GeometryCollectionExampleSimulation.cpp
--- a/Engine/Plugins/Experimental/GeometryCollectionPlugin/Source/GeometryCollectionExampleCore/Private/GeometryCollection/GeometryCollectionExampleSimulation.cpp
+++ b/Engine/Plugins/Experimental/GeometryCollectionPlugin/Source/GeometryCollectionExampleCore/Private/GeometryCollection/GeometryCollectionExampleSimulation.cpp
@@ -29,36 +29,59 @@ namespace GeometryCollectionExample
 	template<class T>
 	void RigidBodiesFallingUnderGravity()
 	{
-		TUniquePtr<Chaos::FChaosPhysicsMaterial> PhysicalMaterial = nullptr;
-		TSharedPtr<FGeometryCollection> RestCollection = nullptr;
-		TSharedPtr<FGeometryDynamicCollection> DynamicCollection = nullptr;
+		FChaosSolversModule* Module = FChaosSolversModule::GetModule();
+		Module->ChangeThreadingMode(EChaosThreadingMode::SingleThread);
 
 		//
 		//  Rigid Body Setup
 		//
-		InitCollectionsParameters InitParams = { FTransform::Identity, FVector(1.0), nullptr, (int32)EObjectStateTypeEnum::Chaos_Object_Kinematic };
-		InitCollections(PhysicalMaterial, RestCollection, DynamicCollection, InitParams);
 
-		FGeometryCollectionPhysicsProxy* PhysObject = RigidBodySetup(PhysicalMaterial, RestCollection, DynamicCollection);
+		InitCollectionsParameters InitParams = 
+		{ 
+			FTransform::Identity,	// RestCenter
+			FVector(1.0),			// RestScale
+			nullptr,				// RestInitFunc
+			(int32)EObjectStateTypeEnum::Chaos_Object_Kinematic // DynamicStateDefault
+		};
+
+		TUniquePtr<Chaos::FChaosPhysicsMaterial> PhysicalMaterial = nullptr; // Allocated and zero'ed
+		TSharedPtr<FGeometryCollection> RestCollection = nullptr;				// GeometryCollection::MakeCubeElement(InitParams)
+		TSharedPtr<FGeometryDynamicCollection> GTDynamicCollection = nullptr;		// New'ed w/copy of RestCollection Transform, Parent, Children, SimulationType, and StatusFlags attrs.
+		InitCollections(PhysicalMaterial, RestCollection, GTDynamicCollection, InitParams);
+
+		//
+		// Sim Initialization
+		//
+		
+		// Creates FGeometryCollectionPhysicsProxy with GTDynamicCollection, and 
+		// an InitFunc that sets FSimulationParametes's RestCollection and 
+		// DynamicCollection pointers.  Calls FGeometryCollectionPhysicsProx::Initialize().
+		FGeometryCollectionPhysicsProxy* PhysObject = RigidBodySetup(PhysicalMaterial, RestCollection, GTDynamicCollection);
 
 		Chaos::FPBDRigidsSolver* Solver = FChaosSolversModule::GetModule()->CreateSolver(true);
-#if CHAOS_PARTICLEHANDLE_TODO
 		Solver->RegisterObject(PhysObject);
-#endif
 		Solver->SetHasFloor(false);
 		Solver->SetEnabled(true);
 		PhysObject->ActivateBodies();
 
+		Solver->AddDirtyProxy(PhysObject); // why?
+		Solver->PushPhysicsState(Module->GetDispatcher());
+
 		Solver->AdvanceSolverBy(1 / 24.);
 
-		FinalizeSolver(*Solver);
+		// Calls BufferPhysicsResults(), FlipBuffer(), and PullFromPhysicsState() on each proxy.
+		//FinalizeSolver(*Solver);
+
+		Solver->BufferPhysicsResults();
+		Solver->FlipBuffers();
+		Solver->UpdateGameThreadStructures();
 
 		// never touched
 		TManagedArray<FTransform>& RestTransform = RestCollection->Transform;
 		EXPECT_LT(FMath::Abs(RestTransform[0].GetTranslation().Z), SMALL_THRESHOLD);
 
 		// simulated
-		TManagedArray<FTransform>& Transform = DynamicCollection->Transform;
+		TManagedArray<FTransform>& Transform = GTDynamicCollection->Transform;
 		EXPECT_EQ(Transform.Num(), 1);
 		EXPECT_LT(Transform[0].GetTranslation().Z, 0);
 		
diff --git a/Engine/Source/Runtime/Engine/Private/PhysicsEngine/Experimental/PhysScene_Chaos.cpp b/Engine/Source/Runtime/Engine/Private/PhysicsEngine/Experimental/PhysScene_Chaos.cpp
--- a/Engine/Source/Runtime/Engine/Private/PhysicsEngine/Experimental/PhysScene_Chaos.cpp
+++ b/Engine/Source/Runtime/Engine/Private/PhysicsEngine/Experimental/PhysScene_Chaos.cpp
@@ -536,13 +536,10 @@ void FPhysScene_Chaos::AddObject(UPrimitiveComponent* Component, FGeometryCollec
 {
 	AddToComponentMaps(Component, InObject);
 
-	check(false);
-#if 0
 	Chaos::FPhysicsSolver* Solver = GetSolver();
 	Solver->RegisterObject(InObject);
 
 	AddPhysicsProxy(InObject, Solver, GetDispatcher());
-#endif
 }
 
 void FPhysScene_Chaos::AddObject(UPrimitiveComponent* Component, FFieldSystemPhysicsProxy* InObject)
@@ -680,21 +677,13 @@ void FPhysScene_Chaos::RemoveObject(FGeometryParticlePhysicsProxy* InObject)
 
 void FPhysScene_Chaos::RemoveObject(FGeometryCollectionPhysicsProxy* InObject)
 {
-	check(false);
-#if 0
 	Chaos::FPhysicsSolver* Solver = InObject->GetSolver();
-
-	const int32 NumRemoved = Solver->UnregisterObject(InObject);
-
-	if(NumRemoved == 0)
+	if(Solver && !Solver->UnregisterObject(InObject))
 	{
 		UE_LOG(LogChaos, Warning, TEXT("Attempted to remove an object that wasn't found in its solver's gamethread storage - it's likely the solver has been mistakenly changed."));
 	}
-
 	RemoveFromComponentMaps(InObject);
-
 	RemovePhysicsProxy(InObject, Solver, ChaosModule);
-#endif
 }
 
 void FPhysScene_Chaos::RemoveObject(FFieldSystemPhysicsProxy* InObject)
diff --git a/Engine/Source/Runtime/Experimental/Chaos/Public/Chaos/PBDGeometryCollectionParticles.h b/Engine/Source/Runtime/Experimental/Chaos/Public/Chaos/PBDGeometryCollectionParticles.h
new file mode 100644
--- /dev/null
+++ b/Engine/Source/Runtime/Experimental/Chaos/Public/Chaos/PBDGeometryCollectionParticles.h
@@ -0,0 +1,35 @@
+// Copyright 1998-2019 Epic Games, Inc. All Rights Reserved.
+#pragma once
+
+#include "Chaos/ArrayCollectionArray.h"
+#include "Chaos/PBDRigidParticles.h"
+
+namespace Chaos
+{
+	template<class T, int d>
+	class TPBDGeometryCollectionParticles : public TPBDRigidParticles<T, d>
+	{
+	public:
+		TPBDGeometryCollectionParticles()
+		    : TPBDRigidParticles<T, d>()
+		{
+			InitHelper();
+		}
+		TPBDGeometryCollectionParticles(const TPBDRigidParticles<T, d>& Other) = delete;
+
+		TPBDGeometryCollectionParticles(TPBDRigidParticles<T, d>&& Other)
+		    : TPBDRigidParticles<T, d>(MoveTemp(Other))
+		{
+			InitHelper();
+		}
+		~TPBDGeometryCollectionParticles()
+		{}
+
+		typedef TPBDGeometryCollectionParticleHandle<T, d> THandleType;
+		const THandleType* Handle(int32 Index) const { return static_cast<const THandleType*>(TGeometryParticles<T, d>::Handle(Index)); }
+		THandleType* Handle(int32 Index) { return static_cast<THandleType*>(TGeometryParticles<T, d>::Handle(Index)); }
+
+	private:
+		void InitHelper() { this->MParticleType = EParticleType::GeometryCollection; }
+	};
+} // namespace Chaos
\ No newline at end of file
diff --git a/Engine/Source/Runtime/Experimental/Chaos/Public/Chaos/PBDRigidsSOAs.h b/Engine/Source/Runtime/Experimental/Chaos/Public/Chaos/PBDRigidsSOAs.h
--- a/Engine/Source/Runtime/Experimental/Chaos/Public/Chaos/PBDRigidsSOAs.h
+++ b/Engine/Source/Runtime/Experimental/Chaos/Public/Chaos/PBDRigidsSOAs.h
@@ -29,6 +29,10 @@ public:
 
 		ClusteredParticles = MakeUnique< TPBDRigidClusteredParticles<T, d>>();
 		ClusteredParticles->RemoveParticleBehavior() = ERemoveParticleBehavior::Remove;	//clustered particles maintain relative ordering
+
+		GeometryCollectionDisabledParticles = MakeUnique<TPBDGeometryCollectionParticles<T, d>>();
+		GeometryCollectionParticles = MakeUnique<TPBDGeometryCollectionParticles<T, d>>();
+		
 		UpdateViews();
 	}
 
@@ -54,7 +58,7 @@ public:
 	}
 	TArray<TPBDRigidParticleHandle<T, d>*> CreateDynamicParticles(int32 NumParticles, const TPBDRigidParticleParameters<T, d>& Params = TPBDRigidParticleParameters<T, d>())
 	{
-		auto Results = CreateParticlesHelper<TPBDRigidParticleHandle<T, d>>(NumParticles, Params.bDisabled ? DynamicDisabledParticles : DynamicParticles, Params);;
+		auto Results = CreateParticlesHelper<TPBDRigidParticleHandle<T, d>>(NumParticles, Params.bDisabled ? DynamicDisabledParticles : DynamicParticles, Params);
 
 		if (!Params.bStartSleeping)
 		{
@@ -63,19 +67,32 @@ public:
 		UpdateViews();
 		return Results;
 	}
+	TArray<TPBDGeometryCollectionParticleHandle<T, d>*> CreateGeometryCollectionParticles(int32 NumParticles, const TPBDRigidParticleParameters<T, d>& Params = TPBDRigidParticleParameters<T, d>())
+	{
+		TArray<TPBDGeometryCollectionParticleHandle<T, d>*> Results = CreateParticlesHelper<TPBDGeometryCollectionParticleHandle<T, d>>(
+			NumParticles, Params.bDisabled ? GeometryCollectionDisabledParticles : GeometryCollectionParticles, Params);
+		//TArray<TPBDRigidParticleHandle<T, d>*>& RigidHandles = (TArray<TPBDRigidParticleHandle<T, d>*>*)&Results;//*static_cast<TArray<TPBDRigidParticleHandle<T, d>*>*>(&Results);
+		if (!Params.bStartSleeping)
+		{
+			InsertToMapAndArray(Results, ActiveParticlesToIndex, ActiveParticlesArray);
+			InsertToMapAndArray(Results, ActiveGeometryCollectionToIndex, ActiveGeometryCollectionArray);
+		}
+		UpdateViews();
+		return Results;
+	}
 
 	/** Used specifically by PBDRigidClustering. These have special properties for maintaining relative order, efficiently switching from kinematic to dynamic, disable to enable, etc... */
 	TArray<TPBDRigidClusteredParticleHandle<T, d>*> CreateClusteredParticles(int32 NumParticles, const TPBDRigidParticleParameters<T, d>& Params = TPBDRigidParticleParameters<T, d>())
 	{
 		auto NewClustered = CreateParticlesHelper<TPBDRigidClusteredParticleHandle<T, d>>(NumParticles, ClusteredParticles, Params);
 		
-			if (!Params.bDisabled)
-			{
+		if (!Params.bDisabled)
+		{
 			InsertToMapAndArray(NewClustered, NonDisabledClusteredToIndex, NonDisabledClusteredArray);
-			}
+		}
 
-			if (!Params.bStartSleeping)
-			{
+		if (!Params.bStartSleeping)
+		{
 			InsertToMapAndArray(reinterpret_cast<TArray<TPBDRigidParticleHandle<T,d>*>&>(NewClustered), ActiveParticlesToIndex, ActiveParticlesArray);
 			InsertToMapAndArray(NewClustered, ActiveClusteredToIndex, ActiveClusteredArray);
 		}
@@ -249,6 +266,8 @@ public:
 
 		ensure(ClusteredParticles->Size() == 0);	//not supported yet
 		//Ar << ClusteredParticles;
+		Ar << GeometryCollectionParticles;
+		Ar << GeometryCollectionDisabledParticles;
 
 		SerializeMapAndArray(Ar, ActiveParticlesToIndex, ActiveParticlesArray);
 		//SerializeMapAndArray(Ar, ActiveClusteredToIndex, ActiveClusteredArray);
@@ -281,6 +300,12 @@ public:
 	const TGeometryParticles<T, d>& GetNonDisabledStaticParticles() const { return *StaticParticles; }
 	TGeometryParticles<T, d>& GetNonDisabledStaticParticles() { return *StaticParticles; }
 
+	const TPBDGeometryCollectionParticles<T, d>& GetGeometryCollectionParticles() const { return *GeometryCollectionParticles; }
+	TPBDGeometryCollectionParticles<T, d>& GetGeometryCollectionParticles() { return *GeometryCollectionParticles; }
+
+	const TParticleView<TPBDGeometryCollectionParticles<T, d>>& GetActiveGeometryCollectionParticlesView() const { return ActiveGeometryCollectionParticlesView; }
+	TParticleView<TPBDGeometryCollectionParticles<T, d>>& GetActiveGeometryCollectionParticlesView() { return ActiveGeometryCollectionParticlesView; }
+
 	//TEMP: only needed while clustering code continues to use direct indices
 	const auto& GetActiveClusteredArray() const { return ActiveClusteredArray; }
 	const auto& GetNonDisabledClusteredArray() const { return NonDisabledClusteredArray; }
@@ -316,9 +341,10 @@ private:
 		return ReturnHandles;
 	}
 
-	template <typename TParticle>
-	void InsertToMapAndArray(const TArray<TParticle*>& ParticlesToInsert, TMap<TParticle*, int32>& ParticleToIndex, TArray<TParticle*>& ParticleArray)
+	template <typename TParticle1, typename TParticle2>
+	void InsertToMapAndArray(const TArray<TParticle1*>& ParticlesToInsert, TMap<TParticle2*, int32>& ParticleToIndex, TArray<TParticle2*>& ParticleArray)
 	{
+		// TODO: Compile time check ensuring TParticle2 is derived from TParticle1?
 		int32 NextIdx = ParticleArray.Num();
 		for (auto Particle : ParticlesToInsert)
 		{
@@ -331,11 +357,11 @@ private:
 	static void InsertToMapAndArray(TParticle* Particle, TMap<TParticle*, int32>& ParticleToIndex, TArray<TParticle*>& ParticleArray)
 	{
 		if (ParticleToIndex.Contains(Particle) == false)
-			{
+		{
 			ParticleToIndex.Add(Particle, ParticleArray.Num());
 			ParticleArray.Add(Particle);
-			}
 		}
+	}
 
 	template <typename TParticle>
 	static void RemoveFromMapAndArray(TParticle* Particle, TMap<TParticle*, int32>& ParticleToIndex, TArray<TParticle*>& ParticleArray)
@@ -371,11 +397,11 @@ private:
 	{
 		//build various views. Group SOA types together for better branch prediction
 		{
-			TArray<TSOAView<TGeometryParticles<T, d>>> TmpArray = { StaticParticles.Get(), KinematicParticles.Get(), DynamicParticles.Get(), DynamicKinematicParticles.Get(), {&NonDisabledClusteredArray} };
+			TArray<TSOAView<TGeometryParticles<T, d>>> TmpArray = { StaticParticles.Get(), KinematicParticles.Get(), DynamicParticles.Get(), DynamicKinematicParticles.Get(), {&NonDisabledClusteredArray}, GeometryCollectionParticles.Get() };
 			NonDisabledView = MakeParticleView(MoveTemp(TmpArray));
 		}
 		{
-			TArray<TSOAView<TPBDRigidParticles<T, d>>> TmpArray = { DynamicParticles.Get(), {&NonDisabledClusteredArray} };
+			TArray<TSOAView<TPBDRigidParticles<T, d>>> TmpArray = { DynamicParticles.Get(), {&NonDisabledClusteredArray}, GeometryCollectionParticles.Get() };
 			NonDisabledDynamicView = MakeParticleView(MoveTemp(TmpArray));
 		}
 		{
@@ -384,13 +410,17 @@ private:
 		}
 		{
 			TArray<TSOAView<TGeometryParticles<T, d>>> TmpArray = { StaticParticles.Get(), StaticDisabledParticles.Get(), KinematicParticles.Get(), KinematicDisabledParticles.Get(),
-				DynamicParticles.Get(), DynamicDisabledParticles.Get(), DynamicKinematicParticles.Get(), ClusteredParticles.Get() };
+				DynamicParticles.Get(), DynamicDisabledParticles.Get(), DynamicKinematicParticles.Get(), ClusteredParticles.Get(), GeometryCollectionParticles.Get() };
 			AllParticlesView = MakeParticleView(MoveTemp(TmpArray));
 		}
 		{
 			TArray<TSOAView<TKinematicGeometryParticles<T, d>>> TmpArray = { KinematicParticles.Get(), DynamicKinematicParticles.Get() };
 			ActiveKinematicParticlesView = MakeParticleView(MoveTemp(TmpArray));
 		}
+		{
+			TArray<TSOAView<TPBDGeometryCollectionParticles<T, d>>> TmpArray = { GeometryCollectionParticles.Get()};
+			ActiveGeometryCollectionParticlesView = MakeParticleView(MoveTemp(TmpArray));
+		}
 	}
 
 	//Organized by SOA type
@@ -406,6 +436,11 @@ private:
 
 	TUniquePtr<TPBDRigidClusteredParticles<T, d>> ClusteredParticles;
 
+	TUniquePtr<TPBDGeometryCollectionParticles<T, d>> GeometryCollectionParticles;
+	TUniquePtr<TPBDGeometryCollectionParticles<T, d>> GeometryCollectionDisabledParticles;
+	TMap<TPBDGeometryCollectionParticleHandle<T, d>*, int32> ActiveGeometryCollectionToIndex;
+	TArray<TPBDGeometryCollectionParticleHandle<T, d>*> ActiveGeometryCollectionArray;
+
 	//Utility structures for maintaining an Active particles view
 	TMap<TPBDRigidParticleHandle<T, d>*, int32> ActiveParticlesToIndex;
 	TArray<TPBDRigidParticleHandle<T, d>*> ActiveParticlesArray;
@@ -422,6 +457,7 @@ private:
 	TParticleView<TPBDRigidParticles<T, d>> ActiveParticlesView;						//all particles that are active
 	TParticleView<TGeometryParticles<T, d>> AllParticlesView;							//all particles
 	TParticleView<TKinematicGeometryParticles<T, d>> ActiveKinematicParticlesView;		//all kinematic particles that are not disabled
+	TParticleView<TPBDGeometryCollectionParticles<T, d>> ActiveGeometryCollectionParticlesView; // all geom collection particles that are not disabled
 
 	//Auxiliary data synced with particle handles
 	TGeometryParticleHandles<T, d> ParticleHandles;
diff --git a/Engine/Source/Runtime/Experimental/Chaos/Public/Chaos/ParticleHandle.h b/Engine/Source/Runtime/Experimental/Chaos/Public/Chaos/ParticleHandle.h
--- a/Engine/Source/Runtime/Experimental/Chaos/Public/Chaos/ParticleHandle.h
+++ b/Engine/Source/Runtime/Experimental/Chaos/Public/Chaos/ParticleHandle.h
@@ -3,6 +3,7 @@
 
 #include "Templates/ChooseClass.h"
 #include "Chaos/PBDRigidClusteredParticles.h"
+#include "Chaos/PBDGeometryCollectionParticles.h"
 #include "Chaos/ParticleHandleFwd.h"
 #include "Chaos/ParticleIterator.h"
 #include "Chaos/ParticleDirtyFlags.h"
@@ -743,6 +744,57 @@ public:
 	int32 TransientParticleIndex() const { return ParticleIdx; }
 };
 
+template <typename T, int d, bool bPersistent = true>
+class TPBDGeometryCollectionParticleHandleImp : public TPBDRigidParticleHandleImp<T, d, bPersistent>
+{
+public:
+	using TGeometryParticleHandleImp<T, d, bPersistent>::ParticleIdx;
+	using TGeometryParticleHandleImp<T, d, bPersistent>::PBDRigidParticles;
+	using TGeometryParticleHandleImp<T, d, bPersistent>::Type;
+	using TTransientHandle = TTransientPBDGeometryCollectionParticleHandle<T, d>;
+	using TSOAType = TPBDGeometryCollectionParticles<T, d>;
+
+protected:
+	friend class TGeometryParticleHandleImp<T, d, bPersistent>;
+
+	//needed for serialization
+	TPBDGeometryCollectionParticleHandleImp()
+		: TPBDRigidParticleHandleImp<T, d, bPersistent>()
+	{}
+
+	TPBDGeometryCollectionParticleHandleImp<T, d, bPersistent>(
+		TSerializablePtr<TPBDGeometryCollectionParticles<T, d>> Particles, 
+		int32 InIdx, 
+		int32 InGlobalIdx, 
+		const TPBDRigidParticleParameters<T, d>& Params = TPBDRigidParticleParameters<T, d>())
+		: TPBDRigidParticleHandleImp<T, d, bPersistent>(
+			TSerializablePtr<TPBDRigidParticles<T, d>>(Particles), InIdx, InGlobalIdx, Params)
+	{}
+public:
+
+	static TUniquePtr<TPBDGeometryCollectionParticleHandleImp<T, d, bPersistent>> CreateParticleHandle(
+		TSerializablePtr<TPBDGeometryCollectionParticles<T, d>> InParticles, 
+		int32 InParticleIdx, 
+		int32 InHandleIdx, 
+		const TPBDRigidParticleParameters<T, d>& Params = TPBDRigidParticleParameters<T, d>())
+	{
+		return TGeometryParticleHandleImp<T, d, bPersistent>::CreateParticleHandleHelper(
+			InParticles, InParticleIdx, InHandleIdx, Params);
+	}
+
+	TSerializablePtr<TPBDGeometryCollectionParticleHandleImp<T, d, bPersistent>> ToSerializable() const
+	{
+		TSerializablePtr<TPBDGeometryCollectionParticleHandleImp<T, d, bPersistent>> Serializable;
+		Serializable.SetFromRawLowLevel(this);
+		return Serializable;
+	}
+
+	const TPBDGeometryCollectionParticleHandleImp<T, d, true>* Handle() const { return PBDRigidParticles->Handle(ParticleIdx); }
+	TPBDGeometryCollectionParticleHandleImp<T, d, true>* Handle() { return PBDRigidParticles->Handle(ParticleIdx); }
+
+	static constexpr EParticleType StaticType() { return EParticleType::GeometryCollection; }
+};
+
 template <typename T, int d, bool bPersistent>
 const TKinematicGeometryParticleHandleImp<T, d, bPersistent>* TGeometryParticleHandleImp<T,d, bPersistent>::AsKinematic() const { return Type >= EParticleType::Kinematic ? static_cast<const TKinematicGeometryParticleHandleImp<T, d, bPersistent>*>(this) : nullptr; }
 template <typename T, int d, bool bPersistent>
@@ -786,6 +838,9 @@ FString TGeometryParticleHandleImp<T, d, bPersistent>::ToString() const
 	case EParticleType::Dynamic:
 		return FString::Printf(TEXT("Dynamic[%d]"), ParticleIdx);
 		break;
+	case EParticleType::GeometryCollection:
+		return FString::Printf(TEXT("GeometryCollection[%d]"), ParticleIdx);
+		break;
 	}
 	return FString();
 }
@@ -802,7 +857,7 @@ TGeometryParticleHandleImp<T,d,bPersistent>* TGeometryParticleHandleImp<T,d, bPe
  * A wrapper around any type of particle handle to provide a consistent (read-only) API for all particle types.
  * This can make code simpler because you can write code that is type-agnostic, but it
  * has a cost. Where possible it is better to write code that is specific to the type(s)
- * of particles being operated on. TGenericParticleHandle has pointer symatics, so you can use one wherever
+ * of particles being operated on. TGenericParticleHandle has pointer symantics, so you can use one wherever
  * you have a particle handle pointer;
  *
  */
@@ -938,6 +993,12 @@ template <typename T, int d>
 class TPBDRigidParticle;
 
 
+/**
+ * Base class for transient classes used to communicate simulated particle state 
+ * between game and physics threads, which is managed by proxies.
+ *
+ * Note the lack of virtual api.
+ */
 class FParticleData
 {
 public:
@@ -1452,7 +1513,6 @@ public:
 		this->MExternalTorque = InExternalTorque;
 	}
 
-
 	const PMatrix<T, d, d>& I() const { return MI; }
 	void SetI(const PMatrix<T, d, d>& InI)
 	{
@@ -1531,6 +1591,27 @@ private:
 	bool MGravityEnabled;
 };
 
+template <typename T, int d>
+class TPBDGeometryCollectionParticle : public TPBDRigidParticle<T, d>
+{
+public:
+	typedef TPBDRigidParticleData<T, d> FData;
+	typedef TPBDGeometryCollectionParticleHandle<T, d> FHandle;
+
+	using TGeometryParticle<T, d>::Type;
+public:
+	TPBDGeometryCollectionParticle<T,d>(const TPBDRigidParticleParameters<T,d>& DynamicParams = TPBDRigidParticleParameters<T,d>())
+		: TPBDRigidParticle<T, d>(DynamicParams)
+	{
+		Type = EParticleType::GeometryCollection;
+	}
+
+	static TUniquePtr<TPBDGeometryCollectionParticle<T, d>> CreateParticle(const TPBDRigidParticleParameters<T, d>& DynamicParams = TPBDRigidParticleParameters<T, d>())
+	{
+		return TUniquePtr<TPBDGeometryCollectionParticle<T, d>>(new TPBDGeometryCollectionParticle<T, d>(DynamicParams));
+	}
+};
+
 template <typename T, int d>
 class TPBDRigidParticleData : public TKinematicGeometryParticleData<T, d>
 {
@@ -1660,6 +1741,7 @@ TGeometryParticle<T, d>* TGeometryParticle<T, d>::SerializationFactory(FChaosArc
 	case EParticleType::Static: if (Ar.IsLoading()) { return new TGeometryParticle<T, d>(); } break;
 	case EParticleType::Kinematic: if (Ar.IsLoading()) { return new TKinematicGeometryParticle<T, d>(); } break;
 	case EParticleType::Dynamic: if (Ar.IsLoading()) { return new TPBDRigidParticle<T, d>(); } break;
+	case EParticleType::GeometryCollection: if (Ar.IsLoading()) { return new TPBDGeometryCollectionParticle<T, d>(); } break;
 	default:
 		check(false);
 	}
@@ -1756,5 +1838,5 @@ FChaosArchive& operator<<(FChaosArchive& Ar, TAccelerationStructureHandle<T, d>&
 	return Ar;
 }
 
-}
+} // namespace Chaos
 
diff --git a/Engine/Source/Runtime/Experimental/Chaos/Public/Chaos/ParticleHandleFwd.h b/Engine/Source/Runtime/Experimental/Chaos/Public/Chaos/ParticleHandleFwd.h
--- a/Engine/Source/Runtime/Experimental/Chaos/Public/Chaos/ParticleHandleFwd.h
+++ b/Engine/Source/Runtime/Experimental/Chaos/Public/Chaos/ParticleHandleFwd.h
@@ -6,17 +6,14 @@
 
 namespace Chaos
 {
-	template <typename T, int d, bool bProcessing>
-	class TGeometryParticleHandleImp;
-
-	template <typename T, int d, bool bProcessing>
-	class TKinematicGeometryParticleHandleImp;
+	// TGeometryParticle 
 
-	template <typename T, int d, bool bProcessing>
-	class TPBDRigidParticleHandleImp;
+	template <typename T, int d>
+	class TGeometryParticle;
+	using TGeometryParticleFloat3 = TGeometryParticle<float, 3>;
 
 	template <typename T, int d, bool bProcessing>
-	class TPBDRigidClusteredParticleHandleImp;
+	class TGeometryParticleHandleImp;
 
 	template <typename T, int d>
 	using TGeometryParticleHandle = TGeometryParticleHandleImp<T, d, true>;
@@ -24,39 +21,71 @@ namespace Chaos
 	template <typename T, int d>
 	using TTransientGeometryParticleHandle = TGeometryParticleHandleImp<T, d, false>;
 
+	// TKinematicGeometryParticle
+
+	template <typename T, int d>
+	class TKinematicGeometryParticle;
+
+	template <typename T, int d, bool bProcessing>
+	class TKinematicGeometryParticleHandleImp;
+
 	template <typename T, int d>
 	using TKinematicGeometryParticleHandle = TKinematicGeometryParticleHandleImp<T, d, true>;
 
 	template <typename T, int d>
 	using TTransientKinematicGeometryParticleHandle = TKinematicGeometryParticleHandleImp<T, d, false>;
 
+	// TPBDRigidParticle
+
+	template <typename T, int d>
+	class TPBDRigidParticle;
+
+	template <typename T, int d, bool bProcessing>
+	class TPBDRigidParticleHandleImp;
+
 	template <typename T, int d>
 	using TPBDRigidParticleHandle = TPBDRigidParticleHandleImp<T, d, true>;
+	using TPBDRigidParticleHandleFloat3 = TPBDRigidParticleHandle<float, 3>;
 
 	template <typename T, int d>
 	using TTransientPBDRigidParticleHandle = TPBDRigidParticleHandleImp<T, d, false>;
 
+	// TPBDRigidClusteredParticleHandle
+
+	template <typename T, int d, bool bProcessing>
+	class TPBDRigidClusteredParticleHandleImp;
+
 	template <typename T, int d>
 	using TPBDRigidClusteredParticleHandle = TPBDRigidClusteredParticleHandleImp<T, d, true>;
 
 	template <typename T, int d>
 	using TTransientPBDRigidClusteredParticleHandle = TPBDRigidClusteredParticleHandleImp<T, d, false>;
 
-	template <typename T, int d>
-	class TGenericParticleHandle;
+	// TPBDGeometryCollectionParticleHandle
 
-	template <typename TSOA>
-	class TParticleIterator;
+	template <typename T, int d>
+	class TPBDGeometryCollectionParticle;
 
-	template <typename TSOA>
-	class TConstParticleIterator;
+	template <typename T, int d, bool bPersistent>
+	class TPBDGeometryCollectionParticleHandleImp;
 
 	template <typename T, int d>
-	class TGeometryParticle;
+	using TPBDGeometryCollectionParticleHandle = TPBDGeometryCollectionParticleHandleImp<T, d, true>;
+	using TPBDGeometryCollectionParticleHandleFloat3 = TPBDGeometryCollectionParticleHandle<float, 3>;
 
 	template <typename T, int d>
-	class TKinematicGeometryParticle;
+	using TTransientPBDGeometryCollectionParticleHandle = TPBDGeometryCollectionParticleHandleImp<T, d, false>;
+	
+	// TGenericParticleHandle
 
 	template <typename T, int d>
-	class TPBDRigidParticle;
+	class TGenericParticleHandle;
+
+	// TParticleIterator
+
+	template <typename TSOA>
+	class TParticleIterator;
+
+	template <typename TSOA>
+	class TConstParticleIterator;
 }
diff --git a/Engine/Source/Runtime/Experimental/ChaosSolvers/Private/PBDRigidsSolver.cpp b/Engine/Source/Runtime/Experimental/ChaosSolvers/Private/PBDRigidsSolver.cpp
--- a/Engine/Source/Runtime/Experimental/ChaosSolvers/Private/PBDRigidsSolver.cpp
+++ b/Engine/Source/Runtime/Experimental/ChaosSolvers/Private/PBDRigidsSolver.cpp
@@ -156,22 +156,19 @@ namespace Chaos
 		{
 			GTParticle->SetSpatialIdx(FSpatialAccelerationIdx{ 1,0 });
 		}
-
-		// Grab the particle's type
-		const EParticleType InParticleType = GTParticle->ObjectType();
-
-		IPhysicsProxyBase* ProxyBase = nullptr;
-		// NOTE: Do we really need this list of proxies if we can just
-		// access them through the GTParticles list?
-		
-		Chaos::FParticleData* ProxyData;
-
 		if (!ensure(GTParticle->IsParticleValid()))
 		{
 			return;
 		}
 
+		// NOTE: Do we really need these lists of proxies if we can just
+		// access them through the GTParticles list?
+		
+		Chaos::FParticleData* ProxyData;
+		IPhysicsProxyBase* ProxyBase;
+
 		// Make a physics proxy, giving it our particle and particle handle
+		const EParticleType InParticleType = GTParticle->ObjectType();
 		if (InParticleType == EParticleType::Dynamic)
 		{
 			auto Proxy = new FRigidParticlePhysicsProxy(GTParticle->AsDynamic(), nullptr);
@@ -237,7 +234,6 @@ namespace Chaos
 			}
 
 			Handle->GTGeometryParticle() = GTParticle;
-
 			Solver->GetEvolution()->CreateParticle(Handle);
 
 			delete ProxyData;
@@ -270,6 +266,10 @@ namespace Chaos
 		{
 			KinematicGeometryParticlePhysicsProxies.Remove((FKinematicGeometryParticlePhysicsProxy*)InProxy);
 		}
+		else if (InParticleType == EParticleType::GeometryCollection)
+		{
+			check(false); // This shouldn't happen.
+		}
 		else
 		{
 			GeometryParticlePhysicsProxies.Remove((FGeometryParticlePhysicsProxy*)InProxy);
@@ -311,6 +311,32 @@ namespace Chaos
 		});
 	}
 
+	void FPBDRigidsSolver::RegisterObject(FGeometryCollectionPhysicsProxy* InProxy)
+	{
+		UE_LOG(LogPBDRigidsSolverSolver, Verbose, TEXT("FPBDRigidsSolver::RegisterObject(FGeometryCollectionPhysicsProxy*)"));
+		GeometryCollectionPhysicsProxies.AddUnique(InProxy);
+		InProxy->SetSolver(this);
+		InProxy->Initialize();
+		InProxy->InitializeBodiesGT();
+		InProxy->NewData(); // Buffers data on the proxy.
+
+		// Finish registration on the physics thread...
+		FChaosSolversModule::GetModule()->GetDispatcher()->EnqueueCommandImmediate(this,
+			[this, InProxy](FPBDRigidsSolver* Solver)
+		{
+			UE_LOG(LogPBDRigidsSolverSolver, Verbose, 
+				TEXT("FPBDRigidsSolver::RegisterObject(FGeometryCollectionPhysicsProxy*)"));
+			InProxy->InitializeBodiesPT(Solver, Particles);
+			InProxy->PushToPhysicsState(nullptr); // Uses buffer on proxy for transport, not ProxyData.
+		});
+	}
+
+	bool FPBDRigidsSolver::UnregisterObject(FGeometryCollectionPhysicsProxy* InProxy)
+	{
+		InProxy->SetSolver(static_cast<FPBDRigidsSolver*>(nullptr));
+		return GeometryCollectionPhysicsProxies.Remove(InProxy) != 0;
+	}
+
 	bool FPBDRigidsSolver::IsSimulating() const
 	{
 		for (FGeometryParticlePhysicsProxy* Obj : GeometryParticlePhysicsProxies)
@@ -337,11 +363,8 @@ namespace Chaos
 		return false;
 	}
 
-
-
 	void FPBDRigidsSolver::Reset()
 	{
-
 		UE_LOG(LogPBDRigidsSolverSolver, Verbose, TEXT("PBDRigidsSolver::Reset()"));
 
 		MTime = 0;
@@ -383,6 +406,12 @@ namespace Chaos
 		GetEventManager()->DispatchEvents();
 	}
 
+	/**
+	 * ProxyType is FSingleParticlePhysicsProxy<T>, where T is:
+	 *    Chaos::TPBDRigidParticle<float,3>
+	 *    Chaos::TKinematicGeometryParticle<float,3>
+	 *    Chaos::TGeometryParticle<float,3>
+	 */
 	template<typename ProxyType>
 	void PushPhysicsStateExec(FPBDRigidsSolver * Solver, ProxyType* Proxy, Chaos::IDispatcher* Dispatcher)
 	{
@@ -411,11 +440,32 @@ namespace Chaos
 				}
 				Solver->RemoveDirtyProxy(Proxy);
 			}
-
 			Proxy->ClearAccumulatedData();
 		}
 	}
 
+	void PushPhysicsStateExec(FPBDRigidsSolver* Solver, FGeometryCollectionPhysicsProxy* Proxy, Chaos::IDispatcher* Dispatcher)
+	{
+		Proxy->NewData();
+		auto Cmd = [Proxy, Solver](Chaos::FPersistentPhysicsTask* PhysThread)
+		{
+			auto* Evolution = Solver->GetEvolution();
+			TManagedArray<Chaos::TPBDGeometryCollectionParticleHandle<float,3>*>& Handles = Proxy->GetSolverParticleHandles();
+			for (auto* Handle : Handles)
+			{
+				if (Handle)
+					Evolution->DirtyParticle(*Handle);
+			}
+			Proxy->PushToPhysicsState(nullptr);
+		};
+
+		if (Dispatcher)
+			Dispatcher->EnqueueCommandImmediate(Cmd);
+		else
+			Cmd(nullptr);
+		Solver->RemoveDirtyProxy(Proxy);
+		Proxy->ClearAccumulatedData();
+	}
 
 	void FPBDRigidsSolver::PushPhysicsState(IDispatcher* Dispatcher)
 	{
@@ -424,13 +474,20 @@ namespace Chaos
 		{
 			switch (Proxy->GetType())
 			{
-			case EPhysicsProxyType::SingleRigidParticleType:
+			//case EPhysicsProxyType::NoneType: // 0
+			//case EPhysicsProxyType::StaticMeshType: // 1
+			case EPhysicsProxyType::GeometryCollectionType: // 2
+				PushPhysicsStateExec(this, static_cast<FGeometryCollectionPhysicsProxy*>(Proxy), Dispatcher);
+				break;
+			//case EPhysicsProxyType::FieldType: // 3
+			//case EPhysicsProxyType::SkeletalMeshType: // 4
+			case EPhysicsProxyType::SingleRigidParticleType: // 7
 				PushPhysicsStateExec(this, static_cast<FSingleParticlePhysicsProxy< Chaos::TPBDRigidParticle<float, 3> >*>(Proxy), Dispatcher);
 				break;
-			case EPhysicsProxyType::SingleKinematicParticleType:
+			case EPhysicsProxyType::SingleKinematicParticleType: // 6
 				PushPhysicsStateExec(this, static_cast<FSingleParticlePhysicsProxy< Chaos::TKinematicGeometryParticle<float, 3> >*>(Proxy), Dispatcher);
 				break;
-			case EPhysicsProxyType::SingleGeometryParticleType:
+			case EPhysicsProxyType::SingleGeometryParticleType: // 5
 				PushPhysicsStateExec(this, static_cast<FSingleParticlePhysicsProxy< Chaos::TGeometryParticle<float, 3> >*>(Proxy), Dispatcher);
 				break;
 			default:
@@ -439,12 +496,13 @@ namespace Chaos
 		}
 	}
 
-
 	void FPBDRigidsSolver::BufferPhysicsResults()
 	{
 		//ensure(IsInPhysicsThread());
+		TArray<FGeometryCollectionPhysicsProxy*> ActiveGC;
+		ActiveGC.Reserve(GeometryCollectionPhysicsProxies.Num());
 
-		for (Chaos::TGeometryParticleHandleImp < float, 3, false>& ActiveObject : GetParticles().GetActiveParticlesView())
+		for (Chaos::TGeometryParticleHandleImp<float, 3, false>& ActiveObject : GetParticles().GetActiveParticlesView())
 		{
 			switch (ActiveObject.Type)
 			{
@@ -457,15 +515,25 @@ namespace Chaos
 			case Chaos::EParticleType::Static:
 				((FGeometryParticlePhysicsProxy*)(ActiveObject.GTGeometryParticle()->Proxy))->BufferPhysicsResults();
 				break;
+			case Chaos::EParticleType::GeometryCollection:
+				ActiveGC.AddUnique((FGeometryCollectionPhysicsProxy*)(ActiveObject.GTGeometryParticle()->Proxy));
+				break;
 			default:
 				check(false);
 			}
 		}
+
+		for (auto* GCProxy : ActiveGC)
+		{
+			GCProxy->BufferPhysicsResults();
+		}
 	}
 
 	void FPBDRigidsSolver::FlipBuffers()
 	{
 		//ensure(IsInPhysicsThread());
+		TArray<FGeometryCollectionPhysicsProxy*> ActiveGC;
+		ActiveGC.Reserve(GeometryCollectionPhysicsProxies.Num());
 
 		for (Chaos::TGeometryParticleHandleImp < float, 3, false>& ActiveObject : GetParticles().GetActiveParticlesView())
 		{
@@ -480,15 +548,25 @@ namespace Chaos
 			case Chaos::EParticleType::Static:
 				((FGeometryParticlePhysicsProxy*)(ActiveObject.GTGeometryParticle()->Proxy))->FlipBuffer();
 				break;
+			case Chaos::EParticleType::GeometryCollection:
+				ActiveGC.AddUnique((FGeometryCollectionPhysicsProxy*)(ActiveObject.GTGeometryParticle()->Proxy));
+				break;
 			default:
 				check(false);
 			}
 		}
+
+		for (auto* GCProxy : ActiveGC)
+		{
+			GCProxy->FlipBuffer();
+		}
 	}
 
 	void FPBDRigidsSolver::UpdateGameThreadStructures()
 	{
 		//ensure(IsInGameThread());
+		TArray<FGeometryCollectionPhysicsProxy*> ActiveGC;
+		ActiveGC.Reserve(GeometryCollectionPhysicsProxies.Num());
 
 		for (Chaos::TGeometryParticleHandleImp < float, 3, false>& ActiveObject : GetParticles().GetActiveParticlesView())
 		{
@@ -503,10 +581,18 @@ namespace Chaos
 			case Chaos::EParticleType::Static:
 				((FGeometryParticlePhysicsProxy*)(ActiveObject.GTGeometryParticle()->Proxy))->PullFromPhysicsState();
 				break;
+			case Chaos::EParticleType::GeometryCollection:
+				ActiveGC.AddUnique((FGeometryCollectionPhysicsProxy*)(ActiveObject.GTGeometryParticle()->Proxy));
+				break;
 			default:
 				check(false);
 			}
 		}
+
+		for (auto* GCProxy : ActiveGC)
+		{
+			GCProxy->PullFromPhysicsState();
+		}
 	}
 
 
diff --git a/Engine/Source/Runtime/Experimental/ChaosSolvers/Private/PhysicsProxy/GeometryCollectionPhysicsProxy.cpp b/Engine/Source/Runtime/Experimental/ChaosSolvers/Private/PhysicsProxy/GeometryCollectionPhysicsProxy.cpp
--- a/Engine/Source/Runtime/Experimental/ChaosSolvers/Private/PhysicsProxy/GeometryCollectionPhysicsProxy.cpp
+++ b/Engine/Source/Runtime/Experimental/ChaosSolvers/Private/PhysicsProxy/GeometryCollectionPhysicsProxy.cpp
@@ -6,6 +6,7 @@
 #include "PhysicsSolver.h"
 #include "ChaosStats.h"
 #include "Chaos/ArrayCollectionArray.h"
+#include "Chaos/BVHParticles.h"
 #include "Chaos/Transform.h"
 #include "Chaos/ParallelFor.h"
 #include "Chaos/Particles.h"
@@ -20,37 +21,81 @@
 #include "GeometryCollection/GeometryCollectionAlgo.h"
 #include "GeometryCollection/GeometryCollectionUtility.h"
 #include "GeometryCollection/GeometryCollectionSimulationTypes.h"
-#include "GeometryCollection/GeometryDynamicCollection.h"
+#include "GeometryCollection/ManagedArrayCollection.h"
 #include "Modules/ModuleManager.h"
 
+#ifndef TODO_REIMPLEMENT_INIT_COMMANDS
+#define TODO_REIMPLEMENT_INIT_COMMANDS 0
+#endif
+
+#ifndef TODO_REIMPLEMENT_KINEMATICS
+#define TODO_REIMPLEMENT_KINEMATICS 0
+#endif
+
+#ifndef TODO_REIMPLEMENT_FRACTURE
+#define TODO_REIMPLEMENT_FRACTURE 0
+#endif
+
+#ifndef TODO_REIMPLEMENT_RIGID_CACHING
+#define TODO_REIMPLEMENT_RIGID_CACHING 0
+#endif
+
 float CollisionParticlesPerObjectFractionDefault = 1.0f;
 FAutoConsoleVariableRef CVarCollisionParticlesPerObjectFractionDefault(TEXT("p.CollisionParticlesPerObjectFractionDefault"), CollisionParticlesPerObjectFractionDefault, TEXT("Fraction of verts"));
 
 FName FGeometryCollectionPhysicsProxy::SimplicialsAttribute("CollisionParticles");
 FName FGeometryCollectionPhysicsProxy::ImplicitsAttribute("Implicits");
+FName FGeometryCollectionPhysicsProxy::SolverParticleHandlesAttribute("SolverParticleHandles");
+FName FGeometryCollectionPhysicsProxy::GTGeometryParticleAttribute("GTGeometryParticles");
 
 DEFINE_LOG_CATEGORY_STATIC(UGCC_LOG, Error, All);
 
+//==============================================================================
+// FGeometryCollectionResults
+//==============================================================================
 
-bool IsMultithreaded()
+FGeometryCollectionResults::FGeometryCollectionResults()
+	: BaseIndex(0)
+	, NumParticlesAdded(0)
+	, IsObjectDynamic(false)
+	, IsObjectLoading(false)
+	, WorldBounds(ForceInit)
+{}
+
+void FGeometryCollectionResults::Reset()
 {
-	FChaosSolversModule* Module = FModuleManager::Get().GetModulePtr<FChaosSolversModule>("ChaosSolvers");
+	BaseIndex = 0;
+	NumParticlesAdded = 0;
+	DisabledStates.SetNum(0);
+	//Transforms.Resize(0);
+	GlobalTransforms.SetNum(0);
+	//RigidBodyIds.Resize(0);
+	ParticleToWorldTransforms.SetNum(0);
+	//Level.Resize(0);
+	//Parent.Resize(0);
+	//Children.Resize(0);
+	//SimulationType.Resize(0);
+	//StatusFlags.Resize(0);
+	IsObjectDynamic = false;
+	IsObjectLoading = false;
+	WorldBounds = FBoxSphereBounds(ForceInit);
+}
+
+//==============================================================================
+// FGeometryCollectionPhysicsProxy helper functions
+//==============================================================================
 
-	if(Module)
+bool IsMultithreaded()
+{
+	if(FChaosSolversModule* Module = FModuleManager::Get().GetModulePtr<FChaosSolversModule>("ChaosSolvers"))
 	{
-		return Module->GetDispatcher() && Module->GetDispatcher()->GetMode() == Chaos::EThreadingMode::DedicatedThread && Module->IsPersistentTaskRunning();
+		return Module->GetDispatcher() && 
+			Module->GetDispatcher()->GetMode() == Chaos::EThreadingMode::DedicatedThread && 
+			Module->IsPersistentTaskRunning();
 	}
-
 	return false;
 }
 
-FGeometryCollectionResults::FGeometryCollectionResults()
-	: BaseIndex(INDEX_NONE)
-	, NumParticlesAdded(0)
-	, WorldBounds(ForceInit)
-{
-}
-
 Chaos::TTriangleMesh<float>* CreateTriangleMesh(int32 FaceCount, int32 VertexOffset, int32 StartIndex, const TManagedArray<FVector>& Vertex, const TManagedArray<bool>& Visible, const TManagedArray<FIntVector>& Indices, TSet<int32>& VertsAdded)
 {
 	TArray<Chaos::TVector<int32, 3>> Faces;
@@ -170,135 +215,711 @@ TArray<int32> ComputeRecursiveOrder(const FGeometryCollection& Collection)
 }
 
 DECLARE_CYCLE_STAT(TEXT("FGeometryCollectionPhysicsProxy::PopulateSimulatedParticle"), STAT_PopulateSimulatedParticle, STATGROUP_Chaos);
-void PopulateSimulatedParticle(FGeometryCollectionPhysicsProxy::FParticlesType& Particles,
+void PopulateSimulatedParticle(
+	Chaos::TPBDRigidParticleHandle<float,3>* Handle,
+	//FGeometryCollectionPhysicsProxy::FParticlesType& Particles,
 	const FSharedSimulationParameters& SharedParams,
 	const FCollisionStructureManager::FSimplicial* Simplicial,
 	Chaos::TSerializablePtr<Chaos::FImplicitObject> Implicit,
-	float MassIn,
+	const float MassIn,
 	const FVector& InertiaTensorVec,
-	int32 RigidBodyIndex,
-	const FTransform& WorldTransform, uint8 DynamicState, 
-	int16 CollisionGroup)
+	//int32 RigidBodyIndex,
+	const FTransform& WorldTransform, 
+	const uint8 DynamicState, 
+	const int16 CollisionGroup)
 {
 	SCOPE_CYCLE_COUNTER(STAT_PopulateSimulatedParticle);
 
-	Particles.SetDisabledLowLevel(RigidBodyIndex, false);
-
-	Particles.X(RigidBodyIndex) = WorldTransform.GetTranslation();
-	Particles.V(RigidBodyIndex) = Chaos::TVector<float, 3>( FVector(0) );
-	Particles.R(RigidBodyIndex) = WorldTransform.GetRotation().GetNormalized();
-	Particles.W(RigidBodyIndex) = Chaos::TVector<float, 3>( FVector(0) );
-	Particles.P(RigidBodyIndex) = Particles.X(RigidBodyIndex);
-	Particles.Q(RigidBodyIndex) = Particles.R(RigidBodyIndex);
-	Particles.Island(RigidBodyIndex) = INDEX_NONE;
-
-	//todo: if mass to small use the right inertia
-	ensureMsgf(MassIn >= SharedParams.MinimumMassClamp, TEXT("Mass smaller than minimum mass clamp. Too late to change"));
-	Particles.M(RigidBodyIndex) = MassIn;
+	Handle->SetDisabledLowLevel(false);
+	//Particles.SetDisabledLowLevel(RigidBodyIndex, false);
+
+	Handle->SetX(WorldTransform.GetTranslation());
+	Handle->SetV(Chaos::TVector<float, 3>(0.f));
+	Handle->SetR(WorldTransform.GetRotation().GetNormalized());
+	Handle->SetW(Chaos::TVector<float, 3>(0.f));
+	Handle->SetP(Handle->X());
+	Handle->SetQ(Handle->R());
+	Handle->SetIsland(INDEX_NONE);
+
+	//Particles.X(RigidBodyIndex) = WorldTransform.GetTranslation();
+	//Particles.V(RigidBodyIndex) = Chaos::TVector<float, 3>( FVector(0) );
+	//Particles.R(RigidBodyIndex) = WorldTransform.GetRotation().GetNormalized();
+	//Particles.W(RigidBodyIndex) = Chaos::TVector<float, 3>( FVector(0) );
+	//Particles.P(RigidBodyIndex) = Particles.X(RigidBodyIndex);
+	//Particles.Q(RigidBodyIndex) = Particles.R(RigidBodyIndex);
+	//Particles.Island(RigidBodyIndex) = INDEX_NONE;
+
+	//todo: If mass is too small use the right inertia
+	ensureMsgf(MassIn >= SharedParams.MinimumMassClamp, 
+		TEXT("Mass smaller than minimum mass clamp. Too late to change"));
+	ensureMsgf(Handle->InvM() > SMALL_NUMBER, 
+		TEXT("Object mass is too large. Too late to change"));
+	ensureMsgf(InertiaTensorVec[0] > SMALL_NUMBER && InertiaTensorVec[1] > SMALL_NUMBER && InertiaTensorVec[2] > SMALL_NUMBER, 
+		TEXT("Inertia tensor is too small. Too late to change"));
+
+	Handle->SetM(MassIn);
+	//Particles.M(RigidBodyIndex) = MassIn;
 	if (FMath::IsNaN(InertiaTensorVec[0]) || FMath::IsNaN(InertiaTensorVec[1]) || FMath::IsNaN(InertiaTensorVec[2]) ||
 		InertiaTensorVec[0] < SMALL_NUMBER || InertiaTensorVec[1] < SMALL_NUMBER || InertiaTensorVec[2] < SMALL_NUMBER)
 	{
-		Particles.I(RigidBodyIndex) = Chaos::PMatrix<float, 3, 3>(1.f, 1.f, 1.f);
+		Handle->SetI(Chaos::PMatrix<float, 3, 3>(1.f, 1.f, 1.f));
+		//Particles.I(RigidBodyIndex) = Chaos::PMatrix<float, 3, 3>(1.f, 1.f, 1.f);
 	}
 	else
 	{
-		Particles.I(RigidBodyIndex) = Chaos::PMatrix<float, 3, 3>(InertiaTensorVec[0], InertiaTensorVec[1], InertiaTensorVec[2]);
+		Handle->SetI(Chaos::PMatrix<float, 3, 3>(InertiaTensorVec[0], InertiaTensorVec[1], InertiaTensorVec[2]));
+		//Particles.I(RigidBodyIndex) = Chaos::PMatrix<float, 3, 3>(InertiaTensorVec[0], InertiaTensorVec[1], InertiaTensorVec[2]);
 	}
 
 	// for validation set the body to dynamic and check the inverse masses. 
-	Particles.SetObjectState(RigidBodyIndex, Chaos::EObjectStateType::Dynamic);
-	ensureMsgf(Particles.InvM(RigidBodyIndex) > SMALL_NUMBER, TEXT("Object mass is too large. Too late to change"));
-	ensureMsgf(InertiaTensorVec[0] > SMALL_NUMBER && InertiaTensorVec[1] > SMALL_NUMBER && InertiaTensorVec[2] > SMALL_NUMBER, TEXT("Inertia tensor is too small. Too late to change"));
+	Handle->SetObjectStateLowLevel(Chaos::EObjectStateType::Dynamic);
+	//Particles.SetObjectState(RigidBodyIndex, Chaos::EObjectStateType::Dynamic);
 
 	//ensureMsgf(Particles.InvI(RigidBodyIndex).M[0][0] > SMALL_NUMBER && Particles.InvI(RigidBodyIndex).M[1][1] > SMALL_NUMBER && Particles.InvI(RigidBodyIndex).M[2][2] > SMALL_NUMBER, TEXT("Inertia tensor is too large. Too late to change"));
-	if (false == (Particles.InvI(RigidBodyIndex).M[0][0] > SMALL_NUMBER && Particles.InvI(RigidBodyIndex).M[1][1] > SMALL_NUMBER && Particles.InvI(RigidBodyIndex).M[2][2] > SMALL_NUMBER))
+	const Chaos::PMatrix<float,3,3>& InvI = Handle->InvI();
+	if (!(InvI.M[0][0] > SMALL_NUMBER && 
+		  InvI.M[1][1] > SMALL_NUMBER && 
+		  InvI.M[2][2] > SMALL_NUMBER))
 	{
 		UE_LOG(LogChaos, Warning, TEXT("Inertia tensor is too large. Too late to change"));
 	}
 
-	Particles.CollisionGroup(RigidBodyIndex) = CollisionGroup;
-	{
+	Handle->SetCollisionGroup(CollisionGroup);
+	//Particles.CollisionGroup(RigidBodyIndex) = CollisionGroup;
 
-#if TODO_REIMPLEMENT_RIGID_CLUSTERING
+	{
+//#if TODO_REIMPLEMENT_RIGID_CLUSTERING
 		// @important (The solver can not free this memory)
 		if (Implicit)	//todo(ocohen): this is only needed for cases where clusters have no proxy. Kind of gross though, should refactor
 		{
-			Particles.SetGeometry(RigidBodyIndex, Implicit);
+			Handle->SetGeometry(Implicit);
+			//Particles.SetGeometry(RigidBodyIndex, Implicit);
 		}
-#endif
+//#endif
 
-		Particles.CollisionParticlesInitIfNeeded(RigidBodyIndex);
+		Handle->CollisionParticlesInitIfNeeded();
+		//Particles.CollisionParticlesInitIfNeeded(RigidBodyIndex);
 		if (Simplicial)
 		{
+			TUniquePtr<Chaos::TBVHParticles<float, 3>>& CollisionParticles = Handle->CollisionParticles();
 			if (Simplicial->Size())
 			{
-				Particles.CollisionParticles(RigidBodyIndex)->Resize(0);
-				Particles.CollisionParticles(RigidBodyIndex)->AddParticles(Simplicial->Size());
-				for (int32 VertexIndex = 0; VertexIndex < (int32)Simplicial->Size(); VertexIndex++)
-					Particles.CollisionParticles(RigidBodyIndex)->X(VertexIndex) = Simplicial->X(VertexIndex);
+				CollisionParticles->Resize(0);
+				//Particles.CollisionParticles(RigidBodyIndex)->Resize(0);
+				CollisionParticles->AddParticles(Simplicial->Size());
+				//Particles.CollisionParticles(RigidBodyIndex)->AddParticles(Simplicial->Size());
+				for (int32 VertexIndex = 0; VertexIndex < (int32)Simplicial->Size(); ++VertexIndex)
+				{
+					CollisionParticles->X(VertexIndex) = Simplicial->X(VertexIndex);
+					//Particles.CollisionParticles(RigidBodyIndex)->X(VertexIndex) = Simplicial->X(VertexIndex);
+				}
 			}
 
 			// @todo(remove): IF there is no simplicial we should not be forcing one. 
-			if (!Particles.CollisionParticles(RigidBodyIndex)->Size())
+			if (CollisionParticles->Size())
+			//if (!Particles.CollisionParticles(RigidBodyIndex)->Size())
+			{
+				CollisionParticles->AddParticles(1);
+				//Particles.CollisionParticles(RigidBodyIndex)->AddParticles(1);
+				CollisionParticles->X(0) = Chaos::TVector<float, 3>(0);
+				///Particles.CollisionParticles(RigidBodyIndex)->X(0) = Chaos::TVector<float, 3>(0);
+			}
+
+			if (CollisionParticles->Size())
+			//if (Particles.CollisionParticles(RigidBodyIndex)->Size())
+			{
+				CollisionParticles->UpdateAccelerationStructures();
+				//Particles.CollisionParticles(RigidBodyIndex)->UpdateAccelerationStructures();
+			}
+		}
+	}
+
+	//
+	//  Manage Object State
+	//
+
+	// Only sleep if we're not replaying a simulation
+	// #BG TODO If this becomes an issue, recorded tracks should track awake state as well as transforms
+	if (DynamicState == (uint8)EObjectStateTypeEnum::Chaos_Object_Sleeping)
+	{
+		Handle->SetObjectStateLowLevel(Chaos::EObjectStateType::Sleeping);
+		//Particles.SetObjectState(RigidBodyIndex,Chaos::EObjectStateType::Sleeping);
+	}
+	else if (DynamicState == (uint8)EObjectStateTypeEnum::Chaos_Object_Kinematic)
+	{
+		Handle->SetObjectStateLowLevel(Chaos::EObjectStateType::Kinematic);
+		//Particles.SetObjectState(RigidBodyIndex, Chaos::EObjectStateType::Kinematic);
+	}
+	else if (DynamicState == (uint8)EObjectStateTypeEnum::Chaos_Object_Static)
+	{
+		Handle->SetObjectStateLowLevel(Chaos::EObjectStateType::Static);
+		//Particles.SetObjectState(RigidBodyIndex, Chaos::EObjectStateType::Static);
+	}
+	else
+	{
+		Handle->SetObjectStateLowLevel(Chaos::EObjectStateType::Dynamic);
+		//Particles.SetObjectState(RigidBodyIndex,Chaos::EObjectStateType::Dynamic);
+	}
+}
+
+//==============================================================================
+// FGeometryCollectionPhysicsProxy
+//==============================================================================
+
+FGeometryCollectionPhysicsProxy::FGeometryCollectionPhysicsProxy(
+	UObject* InOwner, 
+	FGeometryDynamicCollection* InDynamicCollection, 
+	FInitFunc InInitFunc, 
+	FCacheSyncFunc InCacheSyncFunc, 
+	FFinalSyncFunc InFinalSyncFunc,
+	const Chaos::EMultiBufferMode BufferMode)
+	: Base(InOwner)
+	//, InitializedState(ESimulationInitializationState::Unintialized)
+	, InitializedState(ESimulationInitializationState::Created)
+	, IsObjectDynamic(false)
+	, IsObjectLoading(true)
+
+	, GTDynamicCollection(InDynamicCollection)
+	//, PTDynamicCollection(nullptr)
+
+	, BaseParticleIndex(INDEX_NONE)
+	, NumParticles(INDEX_NONE)
+	, ProxySimDuration(0.0f)
+
+	, InitFunc(InInitFunc)
+	, CacheSyncFunc(InCacheSyncFunc)
+	, FinalSyncFunc(InFinalSyncFunc)
+
+	, LastSyncCountGT(MAX_uint32)
+	, CollisionParticlesPerObjectFraction(CollisionParticlesPerObjectFractionDefault)
+
+	, PhysToGameInterchange(Chaos::FMultiBufferFactory<FGeometryCollectionResults>::CreateBuffer(BufferMode))
+	, GameToPhysInterchange(Chaos::FMultiBufferFactory<FGeometryCollectionResults>::CreateBuffer(BufferMode))
+{
+	// We rely on a guarded buffer.
+	check(BufferMode == Chaos::EMultiBufferMode::TripleGuarded);
+}
+
+FGeometryCollectionPhysicsProxy::~FGeometryCollectionPhysicsProxy()
+{}
+
+float ReportHighParticleFraction = -1.f;
+FAutoConsoleVariableRef CVarReportHighParticleFraction(TEXT("p.gc.ReportHighParticleFraction"), ReportHighParticleFraction, TEXT("Report any objects with particle fraction above this threshold"));
+
+void FGeometryCollectionPhysicsProxy::Initialize()
+{
+	if (InitializedState != ESimulationInitializationState::Created)
+		return;
+
+	// Old proxy init
+	check(IsInGameThread());
+
+	PTDynamicCollection.CopyAttribute(*GTDynamicCollection, FTransformCollection::TransformAttribute, FTransformCollection::TransformGroup);
+	PTDynamicCollection.CopyAttribute(*GTDynamicCollection, FTransformCollection::ParentAttribute, FTransformCollection::TransformGroup);
+	PTDynamicCollection.CopyAttribute(*GTDynamicCollection, FTransformCollection::ChildrenAttribute, FTransformCollection::TransformGroup);
+	PTDynamicCollection.CopyAttribute(*GTDynamicCollection, FGeometryCollection::SimulationTypeAttribute, FTransformCollection::TransformGroup);
+	PTDynamicCollection.CopyAttribute(*GTDynamicCollection, FGeometryCollection::StatusFlagsAttribute, FTransformCollection::TransformGroup);
+	PTDynamicCollection.CopyAttribute(*GTDynamicCollection, FGeometryDynamicCollection::DynamicStateAttribute, FTransformCollection::TransformGroup);
+	PTDynamicCollection.CopyAttribute(*GTDynamicCollection, FGeometryDynamicCollection::CollisionGroupAttribute, FTransformCollection::TransformGroup);
+	PTDynamicCollection.CopyAttribute(*GTDynamicCollection, FGeometryDynamicCollection::ActiveAttribute, FTransformCollection::TransformGroup);
+
+	//todo(benn.g): having to spell out which attribute to copy kind of sucks. We don't want a full copy since GT can have a lot of gameplay attributes
+	//question: should we copy dynamic state? 
+	//response: Seems like we will just end up copying everything, otherwise the simulation collection will be missing attributes.
+
+	/*
+	//todo(ocohen): avoid copies when not multi threaded
+	if(IsMultithreaded())
+	{
+		PTDynamicCollection.CopyAttribute(*GTDynamicCollection, FTransformCollection::TransformAttribute, FGeometryCollection::TransformGroup);
+	}*/
+
+	// Replace with normal funcs in this class
+	//Callbacks->SetUpdateRecordedStateFunction([this](float SolverTime, const TManagedArray<int32> & RigidBodyID, const TManagedArray<FGeometryCollectionBoneNode>& Hierarchy, const FSolverCallbacks::FParticlesType& Particles, const FSolverCallbacks::FCollisionConstraintsType& CollisionRule)
+	//{
+	//	UpdateRecordedState(SolverTime, RigidBodyID, Hierarchy, Particles, CollisionRule);
+	//});
+
+	//Callbacks->SetCommitRecordedStateFunction([this](FRecordedTransformTrack& InTrack)
+	//{
+	//	InTrack = FRecordedTransformTrack::ProcessRawRecordedData(RecordedTracks);
+	//});
+
+	// Back to engine for setup from components
+	InitFunc(Parameters);
+	PTDynamicCollection.SyncAllGroups(*Parameters.RestCollection);
+	CollisionParticlesPerObjectFraction = 
+		Parameters.CollisionSampleFraction * CollisionParticlesPerObjectFractionDefault;
+
+	if (ReportHighParticleFraction > 0)
+	{
+		for (const FSharedSimulationSizeSpecificData& Data : Parameters.Shared.SizeSpecificData)
+		{
+			if (Data.CollisionParticlesFraction >= ReportHighParticleFraction)
+			{
+				ensureMsgf(false, TEXT("Collection with small particle fraction"));
+				UE_LOG(LogChaos, Warning, TEXT("Collection with small particle fraction(%f):%s"), Data.CollisionParticlesFraction, *Parameters.Name);
+			}
+		}
+	}
+
+	// Setup proxy parameters
+
+	// Ryan - I don't know why this happens here.  Previously, Parameters.DynamicCollection
+	// was set to the client level GTDynamicCollection, at least in the case of RigidBodiesFallingUnderGravity().
+	// So why override it to the physics thread dynamic collection???
+	//check(!Parameters.DynamicCollection);
+	Parameters.DynamicCollection = &PTDynamicCollection;
+
+	// Old callbacks init
+
+	//if(Parameters.bClearCache)
+	//{
+	//	if(ResetAnimationCacheCallback)
+	//	{
+	//		ResetAnimationCacheCallback();
+	//	}
+	//}
+
+	CreateDynamicAttributes();
+
+	ProxySimDuration = 0.0f;
+//	InitializedState = Parameters.InitializationState;
+
+	// Old proxy init
+	RecordedTracks.Records.Reset();
+	
+	// Setup double buffer data
+	Results.Get(0).Transforms.Init(PTDynamicCollection.Transform);
+	//Results.Get(0).RigidBodyIds.Init(RigidBodyID);
+	Results.Get(1).Transforms.Init(PTDynamicCollection.Transform);
+	//Results.Get(1).RigidBodyIds.Init(RigidBodyID);
+	
+	LastSyncCountGT = 0;
+
+	// Initialize global transforms
+
+	TArray<FMatrix> &TmpGlobalTransforms = Results.Get(0).GlobalTransforms;
+	GeometryCollectionAlgo::GlobalMatrices(PTDynamicCollection.Transform, PTDynamicCollection.Parent, TmpGlobalTransforms);
+	Results.Get(1).GlobalTransforms = TmpGlobalTransforms;
+
+
+	// Initialize data for faster bound calculations
+	// precompute data used for bounds calculation
+	{
+		const TManagedArray<FBox>& BoundingBoxes = Parameters.RestCollection->BoundingBox;
+		const TManagedArray<int32>& TransformIndices = Parameters.RestCollection->TransformIndex;
+
+		const int32 NumBoxes = BoundingBoxes.Num();
+
+		ValidGeometryBoundingBoxes.Reset();
+		ValidGeometryTransformIndices.Reset();
+		for (int32 BoxIdx = 0; BoxIdx < NumBoxes; ++BoxIdx)
+		{
+			const int32 CurrTransformIndex = TransformIndices[BoxIdx];
+
+			if (Parameters.RestCollection->IsGeometry(CurrTransformIndex))
+			{
+				ValidGeometryBoundingBoxes.Add(BoundingBoxes[BoxIdx]);
+				ValidGeometryTransformIndices.Add(CurrTransformIndex);
+			}
+		}
+
+		FBox BoundingBox(ForceInit);
+		const FMatrix& ActorToWorld = Parameters.WorldTransform.ToMatrixWithScale();
+		for (int i = 0; i < ValidGeometryBoundingBoxes.Num(); ++i)
+		{
+			BoundingBox += ValidGeometryBoundingBoxes[i].TransformBy(
+				TmpGlobalTransforms[ValidGeometryTransformIndices[i]] * ActorToWorld);
+		}
+
+		Results.Get(0).WorldBounds = FBoxSphereBounds(BoundingBox);
+		Results.Get(1).WorldBounds = FBoxSphereBounds(BoundingBox);
+	}
+
+	InitializedState = ESimulationInitializationState::Unintialized;
+}
+
+void FGeometryCollectionPhysicsProxy::InitializeBodiesGT()
+{
+	const FGeometryCollection* RestCollection = Parameters.RestCollection;
+	FGeometryDynamicCollection* DynamicCollection = Parameters.DynamicCollection;
+
+	if (Parameters.Simulating &&
+		((InitializedState == ESimulationInitializationState::Unintialized) ||
+		(InitializedState == ESimulationInitializationState::Activated)))
+	{
+		const TManagedArray<int32>& TransformIndex = RestCollection->TransformIndex;
+		const TManagedArray<int32>& BoneMap = RestCollection->BoneMap;
+		const TManagedArray<int32>& Parent = RestCollection->Parent;
+		const TManagedArray<TSet<int32>>& Children = RestCollection->Children;
+		const TManagedArray<int32>& SimulationType = RestCollection->SimulationType;
+		const TManagedArray<FVector>& Vertex = RestCollection->Vertex;
+		const TManagedArray<int32>& DynamicState = DynamicCollection->DynamicState;
+		const TManagedArray<int32>& CollisionGroup = DynamicCollection->CollisionGroup;
+		const TManagedArray<float>& Mass = RestCollection->GetAttribute<float>("Mass", FTransformCollection::TransformGroup);
+		const TManagedArray<FVector>& InertiaTensor = RestCollection->GetAttribute<FVector>("InertiaTensor", FTransformCollection::TransformGroup);
+
+		TArray<FTransform> Transform;
+		GeometryCollectionAlgo::GlobalMatrices(DynamicCollection->Transform, DynamicCollection->Parent, Transform);
+		check(DynamicCollection->Transform.Num() == Transform.Num());
+
+		NumParticles = SimulatableParticles.Count(true);
+		BaseParticleIndex = 0; // Are we always zero indexed now?
+
+		// Create game thread particles...
+		for (int32 Idx = 0; Idx < SimulatableParticles.Num(); ++Idx)
+		{
+			if (SimulatableParticles[Idx])
+			{
+				//GTGeometryParticles[Idx] = Chaos::TGeometryParticle<float, 3>::CreateParticle();
+				//GTGeometryParticles[Idx] = Chaos::TKinematicGeometryParticle<float, 3>::CreateParticle();
+				//GTGeometryParticles[Idx] = Chaos::TPBDRigidParticle<float, 3>::CreateParticle();
+				GTGeometryParticles[Idx] = Chaos::TPBDGeometryCollectionParticle<float, 3>::CreateParticle();
+				Chaos::TPBDRigidParticle<float, 3>* GTParticle = GTGeometryParticles[Idx]->AsDynamic();
+
+//				Chaos::TSerializablePtr<Chaos::TImplicitObject<float, 3>> ImplicitShape = Implicits[Idx];
+//				GTParticle->SetGeometry(ImplicitShape);
+//				if (!ensure(GTParticle->IsParticleValid()))
+//				{
+//					check(false); // continue?
+//				}
+				GTParticle->SetX(Transform[Idx].GetTranslation());
+				GTParticle->SetR(Transform[Idx].GetRotation());
+				//...
+				GTParticle->Proxy = this; // I feel dirty.
+			}
+		}
+
+		InitializedState = Parameters.InitializationState;
+	}
+}
+
+void FGeometryCollectionPhysicsProxy::InitializeBodiesPT(
+	Chaos::FPBDRigidsSolver* RigidsSolver,
+	Chaos::FPBDRigidsSolver::FParticlesType& Particles)
+{
+	const FGeometryCollection* RestCollection = Parameters.RestCollection;
+	FGeometryDynamicCollection* DynamicCollection = Parameters.DynamicCollection;
+
+	if (Parameters.Simulating && 
+		((InitializedState == ESimulationInitializationState::Unintialized) || 
+		 (InitializedState == ESimulationInitializationState::Activated)))
+	{
+		const TManagedArray<int32>& TransformIndex = RestCollection->TransformIndex;
+		const TManagedArray<int32>& BoneMap = RestCollection->BoneMap;
+		const TManagedArray<int32>& Parent = RestCollection->Parent;
+		const TManagedArray<TSet<int32>>& Children = RestCollection->Children;
+		const TManagedArray<int32>& SimulationType = RestCollection->SimulationType;
+		const TManagedArray<FVector>& Vertex = RestCollection->Vertex;
+		const TManagedArray<int32>& DynamicState = DynamicCollection->DynamicState;
+		const TManagedArray<int32>& CollisionGroup = DynamicCollection->CollisionGroup;
+		const TManagedArray<float>& Mass = RestCollection->GetAttribute<float>("Mass", FTransformCollection::TransformGroup);
+		const TManagedArray<FVector>& InertiaTensor = RestCollection->GetAttribute<FVector>("InertiaTensor", FTransformCollection::TransformGroup);
+
+		TArray<FTransform> Transform;
+		GeometryCollectionAlgo::GlobalMatrices(DynamicCollection->Transform, DynamicCollection->Parent, Transform);
+		check(DynamicCollection->Transform.Num() == Transform.Num());
+
+		// Count particles to add
+		const int NumSimulatedParticles = SimulatableParticles.Count(true);
+
+		//Chaos::FPhysicsSolver* FPSolver = GetSolver();
+		//const Chaos::FPhysicsSolver::FParticlesType& Particles = FPSolver->GetParticles();
+
+		const int NumRigids = 0; // ryan - Since we're doing SOA, we start at zero?
+		BaseParticleIndex = NumRigids;
+
+		// Add entries into simulation array
+		TArray<Chaos::TPBDGeometryCollectionParticleHandle<float, 3>*> Handles = 
+			Particles.CreateGeometryCollectionParticles(NumSimulatedParticles);
+		int32 NextIdx = 0;
+		for(int32 Idx=0; Idx < SimulatableParticles.Num(); ++Idx)
+		{
+			if (SimulatableParticles[Idx])
 			{
-				Particles.CollisionParticles(RigidBodyIndex)->AddParticles(1);
-				Particles.CollisionParticles(RigidBodyIndex)->X(0) = Chaos::TVector<float, 3>(0);
+				check(GTGeometryParticles[Idx]);
+				Chaos::TPBDRigidParticle<float, 3>* GTParticle = GTGeometryParticles[Idx]->AsDynamic();
+
+				Chaos::TPBDGeometryCollectionParticleHandle<float, 3>* Handle = Handles[NextIdx++];
+				Handle->GTGeometryParticle() = GTParticle;
+
+				SolverParticleHandles[Idx] = Handle;
+				RigidsSolver->GetEvolution()->CreateParticle(Handle);
+			}
+		}
+
+#if TODO_REIMPLEMENT_RIGID_CLUSTERING
+		Chaos::TArrayCollectionArray<int32>& ClusterGroupIndex = ThisSolver->GetRigidClustering().GetClusterGroupIndexArray();
+		Chaos::TArrayCollectionArray<float>& StrainArray = ThisSolver->GetRigidClustering().GetStrainArray();
+		const float StrainDefault = Parameters.DamageThreshold.Num() ? Parameters.DamageThreshold[0] : 0;
+#endif // TODO_REIMPLEMENT_RIGID_CLUSTERING
+		// Add the rigid bodies
+		const int32 NumGeometries = DynamicCollection->NumElements(FGeometryCollection::GeometryGroup);
+		ParallelFor(NumGeometries, [&](int32 GeometryIndex)
+		{
+			const int32 TransformGroupIndex = TransformIndex[GeometryIndex];
+			//if (SimulatableParticles[TransformGroupIndex] == true)
+			if (Chaos::TPBDRigidParticleHandle<float, 3>* Handle = SolverParticleHandles[TransformGroupIndex])
+			{
+				//const int32 RigidBodyIndex = RigidBodyID[TransformGroupIndex];
+				const FTransform WorldTransform = 
+					MassToLocal[TransformGroupIndex] * Transform[TransformGroupIndex] * Parameters.WorldTransform;
+
+				PopulateSimulatedParticle(
+					Handle,
+					//Particles, 
+					Parameters.Shared, 
+					nullptr,// Simplicials[TransformGroupIndex].Get(),
+					Implicits[TransformGroupIndex], 
+#if TODO_REIMPLEMENT_SOLVER_SETTINGS_ACCESSORS
+					RigidsSolver->GetMassScale() * Mass[TransformGroupIndex], 
+					RigidsSolver->GetMassScale() * InertiaTensor[TransformGroupIndex],
+#else // TODO_REIMPLEMENT_SOLVER_SETTINGS_ACCESSORS
+					Mass[TransformGroupIndex], 
+					InertiaTensor[TransformGroupIndex],
+#endif // TODO_REIMPLEMENT_SOLVER_SETTINGS_ACCESSORS
+					//RigidBodyIndex, 
+					WorldTransform, 
+					(uint8)EObjectStateTypeEnum::Chaos_Object_Dynamic, 
+					CollisionGroup[TransformGroupIndex]);
+
+#if TODO_REIMPLEMENT_RIGID_CLUSTERING
+				ClusterGroupIndex[RigidBodyIndex] = Parameters.ClusterGroupIndex;
+				StrainArray[RigidBodyIndex] = StrainDefault;
+#endif //TODO_REIMPLEMENT_RIGID_CLUSTERING
+
+				TUniquePtr<Chaos::TBVHParticles<float, 3>>& CollisionParticles = Handle->CollisionParticles();
+				CollisionParticles.Reset(Simplicials[TransformGroupIndex].Release()); // Steals!
+				//Particles.CollisionParticles(RigidBodyIndex).Reset(Simplicials[TransformGroupIndex].Release());
+				if (CollisionParticles)
+				//if (Particles.CollisionParticles(RigidBodyIndex))
+				{
+					int32 NumCollisionParticles = CollisionParticles->Size();
+					//int32 NumCollisionParticles = Particles.CollisionParticles(RigidBodyIndex)->Size();
+					int32 CollisionParticlesSize = FMath::Max(0, FMath::Min(int(NumCollisionParticles*CollisionParticlesPerObjectFraction), NumCollisionParticles));
+					CollisionParticles->Resize(CollisionParticlesSize); // Truncates!
+					//Particles.CollisionParticles(RigidBodyIndex)->Resize(CollisionParticlesSize);
+				}
+
+				RigidsSolver->GetEvolution()->SetPhysicsMaterial(Handle, Parameters.PhysicalMaterial);
+				//GetSolver()->SetPhysicsMaterial(RigidBodyIndex, Parameters.PhysicalMaterial);
+			}
+		});
+
+#if TODO_REIMPLEMENT_INIT_COMMANDS
+		for (FFieldSystemCommand& Cmd : Parameters.InitializationCommands)
+		{
+			if (Cmd.MetaData.Contains(FFieldSystemMetaData::EMetaType::ECommandData_ProcessingResolution)) 
+				Cmd.MetaData.Remove(FFieldSystemMetaData::EMetaType::ECommandData_ProcessingResolution);
+			FFieldSystemMetaDataProcessingResolution* ResolutionData = 
+				new FFieldSystemMetaDataProcessingResolution(EFieldResolutionType::Field_Resolution_Maximum);
+			Cmd.MetaData.Add(
+				FFieldSystemMetaData::EMetaType::ECommandData_ProcessingResolution, 
+				TUniquePtr< FFieldSystemMetaDataProcessingResolution >(ResolutionData));
+			Commands.Add(Cmd);
+		}
+		Parameters.InitializationCommands.Empty();
+		ProcessCommands(Particles, GetSolver()->GetSolverTime());
+#endif // TODO_REIMPLEMENT_INIT_COMMANDS
+
+		if (Parameters.InitialVelocityType == EInitialVelocityTypeEnum::Chaos_Initial_Velocity_User_Defined)
+		{
+			for (int32 GeometryIndex = 0; GeometryIndex < NumGeometries; ++GeometryIndex)
+			{
+				const int32 TransformGroupIndex = TransformIndex[GeometryIndex];
+				if (Chaos::TPBDRigidParticleHandle<float, 3>* Handle = SolverParticleHandles[TransformGroupIndex])
+				{
+					InitialLinearVelocity[TransformGroupIndex] = Parameters.InitialLinearVelocity;
+					InitialAngularVelocity[TransformGroupIndex] = Parameters.InitialAngularVelocity;
+					if (DynamicState[TransformGroupIndex] == (int32)EObjectStateTypeEnum::Chaos_Object_Dynamic)
+					{
+						Handle->SetV(Parameters.InitialLinearVelocity);
+						Handle->SetW(Parameters.InitialAngularVelocity);
+					}
+				}
+			}
+		}
+/*
+		ParallelFor(NumGeometries, [&](int32 GeometryIndex)
+		{
+			const int32 TransformGroupIndex = TransformIndex[GeometryIndex];
+			//if (SimulatableParticles[TransformGroupIndex] == true)
+			if (Chaos::TPBDRigidParticleHandle<float, 3>* Handle = SolverParticleHandles[TransformGroupIndex])
+			{
+				//const int32 RigidBodyIndex = RigidBodyID[TransformGroupIndex];
+				if (Parameters.InitialVelocityType == EInitialVelocityTypeEnum::Chaos_Initial_Velocity_User_Defined)
+				{
+					InitialLinearVelocity[TransformGroupIndex] = Parameters.InitialLinearVelocity;
+					InitialAngularVelocity[TransformGroupIndex] = Parameters.InitialAngularVelocity;
+					if (DynamicState[TransformGroupIndex] == (int32)EObjectStateTypeEnum::Chaos_Object_Dynamic)
+					{
+						Particles.V(RigidBodyIndex) = Parameters.InitialLinearVelocity;
+						Particles.W(RigidBodyIndex) = Parameters.InitialAngularVelocity;
+					}
+				}
+			}
+		});
+*/
+#if TODO_REIMPLEMENT_KINEMATICS
+		InitializeKinematics(Particles, DynamicState);
+#endif // TODO_REIMPLEMENT_KINEMATICS
+#if TODO_REIMPLEMENT_FRACTURE
+		InitializeRemoveOnFracture(Particles, DynamicState);
+#endif // TODO_REIMPLEMENT_FRACTURE
+
+#if TODO_REIMPLEMENT_RIGID_CLUSTERING
+		// #BG Temporary - don't cluster when playing back. Needs to be changed when kinematics are per-proxy to support
+		// kinematic to dynamic transition for clusters.
+		if (Parameters.EnableClustering)// && Parameters.CacheType != EGeometryCollectionCacheType::Play)
+		{
+			const TArray<int32> RecursiveOrder = ComputeRecursiveOrder(*RestCollection);
+
+			// num clusters
+			uint32 NumClusters=0;
+			TArray<bool> SubTreeContainsSimulatableParticle;
+			SubTreeContainsSimulatableParticle.SetNum(RecursiveOrder.Num());
+			for (const int32 TransformGroupIndex : RecursiveOrder)
+			{
+				if (Children[TransformGroupIndex].Num() > 0)
+				{
+					SubTreeContainsSimulatableParticle[TransformGroupIndex] = false;
+
+					TArray<uint32> RigidChildren, CollectionChildren;
+					for (const int32 ChildIndex : Children[TransformGroupIndex])
+					{
+						if(SubTreeContainsSimulatableParticle[ChildIndex])
+						{
+							NumClusters++;
+							SubTreeContainsSimulatableParticle[TransformGroupIndex] = true;
+							break;
+						}
+					}
+				}
+				else
+				{
+					SubTreeContainsSimulatableParticle[TransformGroupIndex] = (RigidBodyID[TransformGroupIndex] != INDEX_NONE);
+				}
+			}
+
+			const int32 ClusterStartIndex = Particles.Size();
+			Particles.AddParticles(NumClusters);
+
+			int32 ClusterRigidBodyId = ClusterStartIndex;
+			for (const int32 TransformGroupIndex : RecursiveOrder)
+			{
+				TArray<uint32> RigidChildren, CollectionChildren;
+				for (const int32 ChildIndex : Children[TransformGroupIndex])
+				{
+					if (RigidBodyID[ChildIndex] != INDEX_NONE)
+					{
+						RigidChildren.Add(RigidBodyID[ChildIndex]);
+						CollectionChildren.Add(ChildIndex);
+					}
+				}
+				if (RigidChildren.Num())
+				{
+					if (ReportTooManyChildrenNum >= 0 && RigidChildren.Num() > ReportTooManyChildrenNum)
+					{
+						UE_LOG(LogChaos, Warning, TEXT("Too many children (%d) in a single cluster:%s"), RigidChildren.Num(), *Parameters.Name);
+					}
+
+					RigidBodyID[TransformGroupIndex] = ClusterRigidBodyId;
+					Chaos::FClusterCreationParameters<float> CreationParameters;
+					CreationParameters.RigidBodyIndex = RigidBodyID[TransformGroupIndex];
+					BuildClusters(TransformGroupIndex, CollectionChildren, RigidChildren, CreationParameters);
+					ClusterRigidBodyId++;
+				}
+			}
+		}
+#endif // TODO_REIMPLEMENT_RIGID_CLUSTERING
+ 
+		const int32 NumTransforms = GetTransformGroupSize();
+#if TODO_REIMPLEMENT_RIGID_CLUSTERING
+		// Set Connectivity
+		ParallelFor(NumTransforms, [&](int32 TransformGroupIndex)
+		{
+			if (Children[TransformGroupIndex].Num() > 0)
+			{
+				if (RigidBodyID[TransformGroupIndex] != INDEX_NONE)
+				{
+					Solver->GetEvolution()->GetRigidClustering().GenerateConnectionGraph(RigidBodyID[TransformGroupIndex]);
+					//GetSolver()->GetRigidClustering().GenerateConnectionGraph(RigidBodyID[TransformGroupIndex]);
+				}
+			}
+		});
+#endif // TODO_REIMPLEMENT_RIGID_CLUSTERING
+
+#if TODO_REIMPLEMENT_RIGID_CACHING
+		// If we're recording and want to start immediately caching then we should cache the rest state
+		if (Parameters.IsCacheRecording() && Parameters.CacheBeginTime == 0.0f)
+		{
+			if (UpdateRecordedStateCallback)
+			{
+				UpdateRecordedStateCallback(0.0f, RigidBodyID, Particles, GetSolver()->GetCollisionConstraints());
+			}
+		}
+#endif // TODO_REIMPLEMENT_RIGID_CACHING
+
+		if (InitializedState == ESimulationInitializationState::Activated)
+		{
+			//
+			//  Activated bodies has already been called so we are good to go.
+			//
+			InitializedState = ESimulationInitializationState::Initialized;
+#if TODO_REIMPLEMENT_RIGID_CLUSTERING
+			if (Parameters.EnableClustering && Parameters.ClusterGroupIndex)
+			{
+				GetSolver()->GetRigidClustering().IncrementPendingClusterCounter(Parameters.ClusterGroupIndex);
+				GetSolver()->GetRigidClustering().DecrementPendingClusterCounter(Parameters.ClusterGroupIndex);
+			}
+#endif // TODO_REIMPLEMENT_RIGID_CLUSTERING
+		}
+		else if (InitializedState == ESimulationInitializationState::Unintialized)
+		{
+			//
+			//  Activated bodies has not been called, so we are waiting 
+			//  to become active. Deactivate all bodies, and wait for
+			//  ActivatedBodies to be called, and defer the cluster initialization
+			//
+			InitializedState = ESimulationInitializationState::Created;
+			for (int32 TransformGroupIndex = 0; TransformGroupIndex < NumTransforms; TransformGroupIndex++)
+			{
+				if(Chaos::TPBDRigidParticleHandle<float, 3>* Handle = SolverParticleHandles[TransformGroupIndex])
+				//const int32 RigidBodyIndex = RigidBodyID[TransformGroupIndex];
+				//if (RigidBodyIndex != INDEX_NONE)
+				{
+#define TODO_REIMPLEMENT_DEFERRED_ACTIVATION 0
+#if TODO_REIMPLEMENT_DEFERRED_ACTIVATION
+					if (!Handle->Disabled())
+					//if (!Particles.Disabled(RigidBodyIndex))
+					{
+						PendingActivationList.Add(TransformGroupIndex);
+						RigidsSolver->GetEvolution()->DisableParticle(Handle);
+						//GetSolver()->GetEvolution()->DisableParticle(RigidBodyIndex);
+					}
+#endif
+				}
+			}
+
+#if TODO_REIMPLEMENT_RIGID_CLUSTERING
+			//
+			//  Clustering needs to advertise its group id to the cluster so
+			//  that the group is not initialized before all the bodies are
+			//  loaded and created. 
+			//
+			if (Parameters.EnableClustering && Parameters.ClusterGroupIndex )
+			{
+				GetSolver()->GetRigidClustering().IncrementPendingClusterCounter(Parameters.ClusterGroupIndex);
 			}
-
-			if (Particles.CollisionParticles(RigidBodyIndex)->Size())
-				Particles.CollisionParticles(RigidBodyIndex)->UpdateAccelerationStructures();
+#endif // TODO_REIMPLEMENT_RIGID_CLUSTERING
 		}
-	}
-
-	//
-	//  Manage Object State
-	//
-
-	// Only sleep if we're not replaying a simulation
-	// #BG TODO If this becomes an issue, recorded tracks should track awake state as well as transforms
-	if (DynamicState == (uint8)EObjectStateTypeEnum::Chaos_Object_Sleeping)
-	{
-		Particles.SetObjectState(RigidBodyIndex,Chaos::EObjectStateType::Sleeping);
-	}
-	else if (DynamicState == (uint8)EObjectStateTypeEnum::Chaos_Object_Kinematic)
-	{
-		Particles.SetObjectState(RigidBodyIndex, Chaos::EObjectStateType::Kinematic);
-	}
-	else if (DynamicState == (uint8)EObjectStateTypeEnum::Chaos_Object_Static)
-	{
-		Particles.SetObjectState(RigidBodyIndex, Chaos::EObjectStateType::Static);
-	}
-	else
-	{
-		Particles.SetObjectState(RigidBodyIndex,Chaos::EObjectStateType::Dynamic);
-	}
-}
-
-FGeometryCollectionPhysicsProxy::FGeometryCollectionPhysicsProxy(UObject* InOwner, FGeometryDynamicCollection* InDynamicCollection, FInitFunc InInitFunc, FCacheSyncFunc InCacheSyncFunc, FFinalSyncFunc InFinalSyncFunc)
-	: Base(InOwner)
-	, InitializedState(ESimulationInitializationState::Unintialized)
-	, BaseParticleIndex(INDEX_NONE)
-	, NumParticles(0)
-	, ProxySimDuration(0.0f)
-	, SimulationCollection(nullptr)
-	, GTDynamicCollection(InDynamicCollection)
-	, InitFunc(InInitFunc)
-	, CacheSyncFunc(InCacheSyncFunc)
-	, FinalSyncFunc(InFinalSyncFunc)
-	, LastSyncCountGT(MAX_uint32)
-	, IsObjectDynamic(false)
-	, IsObjectLoading(true)
-	, CollisionParticlesPerObjectFraction(CollisionParticlesPerObjectFractionDefault)
-{
-}
-
-FGeometryCollectionPhysicsProxy::~FGeometryCollectionPhysicsProxy()
-{
-}
-
-bool FGeometryCollectionPhysicsProxy::IsSimulating() const
-{
-	return Parameters.Simulating;
+		else
+		{
+			// unknown initialization state in creation callback
+			ensure(false);
+		}
+	} // end if simulating...
 }
 
 void FGeometryCollectionPhysicsProxy::UpdateKinematicBodiesCallback(const FParticlesType& Particles, const float Dt, const float Time, FKinematicProxy& Proxy)
@@ -877,7 +1498,7 @@ void FGeometryCollectionPhysicsProxy::EndFrameCallback(const float EndFrame)
 					const int32 RigidBodyIndex = RigidBodyID[TransformGroupIndex];
 					if (RigidBodyIndex != INDEX_NONE)
 					{
-						bool HasBeenRemoved = (SimulationCollection->StatusFlags[TransformGroupIndex] & FGeometryCollection::FS_RemoveOnFracture) && Particles.Disabled(RigidBodyIndex) && ClusterID[RigidBodyIndex].Id == INDEX_NONE;
+						bool HasBeenRemoved = (PTDynamicCollection.StatusFlags[TransformGroupIndex] & FGeometryCollection::FS_RemoveOnFracture) && Particles.Disabled(RigidBodyIndex) && ClusterID[RigidBodyIndex].Id == INDEX_NONE;
 
 						if (HasBeenRemoved || !Particles.Disabled(RigidBodyIndex))
 						{
@@ -910,10 +1531,10 @@ void FGeometryCollectionPhysicsProxy::EndFrameCallback(const float EndFrame)
 					const int32 RigidBodyIndex = RigidBodyID[TransformGroupIndex];
 					if (RigidBodyIndex != INDEX_NONE)
 					{
-						if ((SimulationCollection->StatusFlags[TransformGroupIndex] & FGeometryCollection::FS_RemoveOnFracture)
+						if ((PTDynamicCollection.StatusFlags[TransformGroupIndex] & FGeometryCollection::FS_RemoveOnFracture)
 							&& Particles.Disabled(RigidBodyIndex) && ClusterID[RigidBodyIndex].Id == INDEX_NONE)
 						{
-							SimulationCollection->Transform[TransformGroupIndex].SetScale3D(FVector::ZeroVector);
+							PTDynamicCollection.Transform[TransformGroupIndex].SetScale3D(FVector::ZeroVector);
 						}
 					}
 				}
@@ -1205,19 +1826,20 @@ void FGeometryCollectionPhysicsProxy::ActivateBodies()
 				Clustering.DecrementPendingClusterCounter(Parameters.ClusterGroupIndex);
 				ParentIndex = Parameters.ClusterGroupIndex;
 			}
-
+#endif
+#if TODO_REIMPLEMENT_DEFERRED_ACTIVATION
 			Chaos::FPhysicsSolver::FParticlesType & Particles = GetSolver()->GetRigidParticles();
 			for (uint32 TransformGroupIndex : PendingActivationList)
 			{
-				int32 RigidBodyIndex = RigidBodyID[TransformGroupIndex];
-				checkSlow(RigidBodyIndex != INDEX_NONE);
-				if (Particles.Disabled(RigidBodyIndex))
+//				int32 RigidBodyIndex = RigidBodyID[TransformGroupIndex];
+//				checkSlow(RigidBodyIndex != INDEX_NONE);
+				if(Chaos::TPBDRigidParticleHandle<float, 3>* Handle = SolverParticleHandles[TransformGroupIndex])
+//				if (Particles.Disabled(RigidBodyIndex))
 				{
 					GetSolver()->GetEvolution()->EnableParticle(RigidBodyIndex, ParentIndex);
 				}
 			}
 #endif
-
 			PendingActivationList.Reset(0);
 
 			InitializedState = ESimulationInitializationState::Initialized;
@@ -1234,8 +1856,6 @@ void FGeometryCollectionPhysicsProxy::ActivateBodies()
 	}
 }
 
-
-
 void FGeometryCollectionPhysicsProxy::BindParticleCallbackMapping(Chaos::TArrayCollectionArray<PhysicsProxyWrapper> & PhysicsProxyReverseMap, Chaos::TArrayCollectionArray<int32> & ParticleIDReverseMap)
 {
 	if (InitializedState== ESimulationInitializationState::Initialized)
@@ -1306,7 +1926,6 @@ void FGeometryCollectionPhysicsProxy::ContiguousIndices(TArray<ContextIndex>& Ar
 	}
 }
 
-
 void FGeometryCollectionPhysicsProxy::ProcessCommands(FParticlesType& Particles, const float Time)
 {
 	FGeometryDynamicCollection* Collection = Parameters.DynamicCollection;
@@ -1420,7 +2039,6 @@ void FGeometryCollectionPhysicsProxy::ProcessCommands(FParticlesType& Particles,
 		}
 	}
 
-
 	// Process Particle-Particle commands
 	if (Commands.Num())
 	{
@@ -1609,7 +2227,6 @@ void FGeometryCollectionPhysicsProxy::PushKinematicStateToSolver(FParticlesType&
 	}
 }
 
-
 void FGeometryCollectionPhysicsProxy::ParameterUpdateCallback(FParticlesType& Particles, const float Time)
 {
 	if (InitializedState == ESimulationInitializationState::Initialized)
@@ -1724,144 +2341,6 @@ void FGeometryCollectionPhysicsProxy::FieldForcesUpdateCallback(Chaos::FPhysicsS
 	}
 }
 
-float ReportHighParticleFraction = -1.f;
-FAutoConsoleVariableRef CVarReportHighParticleFraction(TEXT("p.gc.ReportHighParticleFraction"), ReportHighParticleFraction, TEXT("Report any objects with particle fraction above this threshold"));
-
-void FGeometryCollectionPhysicsProxy::Initialize()
-{
-	// Old proxy init
-	check(IsInGameThread());
-
-	SimulationCollection = MakeUnique<FGeometryDynamicCollection>();
-	SimulationCollection->CopyAttribute(*GTDynamicCollection, FTransformCollection::TransformAttribute, FTransformCollection::TransformGroup);
-	SimulationCollection->CopyAttribute(*GTDynamicCollection, FTransformCollection::ParentAttribute, FTransformCollection::TransformGroup);
-	SimulationCollection->CopyAttribute(*GTDynamicCollection, FTransformCollection::ChildrenAttribute, FTransformCollection::TransformGroup);
-	SimulationCollection->CopyAttribute(*GTDynamicCollection, FGeometryCollection::SimulationTypeAttribute, FTransformCollection::TransformGroup);
-	SimulationCollection->CopyAttribute(*GTDynamicCollection, FGeometryCollection::StatusFlagsAttribute, FTransformCollection::TransformGroup);
-	SimulationCollection->CopyAttribute(*GTDynamicCollection, FGeometryDynamicCollection::DynamicStateAttribute, FTransformCollection::TransformGroup);
-	SimulationCollection->CopyAttribute(*GTDynamicCollection, FGeometryDynamicCollection::CollisionGroupAttribute, FTransformCollection::TransformGroup);
-	SimulationCollection->CopyAttribute(*GTDynamicCollection, FGeometryDynamicCollection::ActiveAttribute, FTransformCollection::TransformGroup);
-
-	//todo(benn.g): having to spell out which attribute to copy kind of sucks. We don't want a full copy since GT can have a lot of gameplay attributes
-	//question: should we copy dynamic state? 
-	//response: Seems like we will just end up copying everything, otherwise the simulation collection will be missing attributes.
-
-	/*
-	//todo(ocohen): avoid copies when not multi threaded
-	if(IsMultithreaded())
-	{
-		SimulationCollection->CopyAttribute(*GTDynamicCollection, FTransformCollection::TransformAttribute, FGeometryCollection::TransformGroup);
-	}*/
-
-	// Replace with normal funcs in this class
-	//Callbacks->SetUpdateRecordedStateFunction([this](float SolverTime, const TManagedArray<int32> & RigidBodyID, const TManagedArray<FGeometryCollectionBoneNode>& Hierarchy, const FSolverCallbacks::FParticlesType& Particles, const FSolverCallbacks::FCollisionConstraintsType& CollisionRule)
-	//{
-	//	UpdateRecordedState(SolverTime, RigidBodyID, Hierarchy, Particles, CollisionRule);
-	//});
-
-	//Callbacks->SetCommitRecordedStateFunction([this](FRecordedTransformTrack& InTrack)
-	//{
-	//	InTrack = FRecordedTransformTrack::ProcessRawRecordedData(RecordedTracks);
-	//});
-
-	// Back to engine for setup from components
-	InitFunc(Parameters);
-	SimulationCollection->SyncAllGroups(*Parameters.RestCollection);
-	CollisionParticlesPerObjectFraction = Parameters.CollisionSampleFraction*CollisionParticlesPerObjectFractionDefault;
-
-	if (ReportHighParticleFraction > 0)
-	{
-		for (const FSharedSimulationSizeSpecificData& Data : Parameters.Shared.SizeSpecificData)
-		{
-			if (Data.CollisionParticlesFraction >= ReportHighParticleFraction)
-			{
-				ensureMsgf(false, TEXT("Collection with small particle fraction"));
-				UE_LOG(LogChaos, Warning, TEXT("Collection with small particle fraction(%f):%s"), Data.CollisionParticlesFraction, *Parameters.Name);
-			}
-		}
-	}
-
-	// Setup proxy parameters
-	Parameters.DynamicCollection = SimulationCollection.Get();
-
-	check(Parameters.DynamicCollection);
-
-	// Old callbacks init
-
-	//if(Parameters.bClearCache)
-	//{
-	//	if(ResetAnimationCacheCallback)
-	//	{
-	//		ResetAnimationCacheCallback();
-	//	}
-	//}
-
-
-	CreateDynamicAttributes();
-
-	ProxySimDuration = 0.0f;
-	InitializedState = Parameters.InitializationState;
-
-	// Old proxy init
-	RecordedTracks.Records.Reset();
-	
-	// Setup double buffer data
-	Results.Get(0).Transforms.Init(SimulationCollection->Transform);
-	Results.Get(0).RigidBodyIds.Init(RigidBodyID);
-	Results.Get(1).Transforms.Init(SimulationCollection->Transform);
-	Results.Get(1).RigidBodyIds.Init(RigidBodyID);
-	
-	LastSyncCountGT = 0;
-
-	// Initialize global transforms
-
-	TArray<FMatrix> TmpGlobalTransforms;
-	GeometryCollectionAlgo::GlobalMatrices(SimulationCollection->Transform, SimulationCollection->Parent, TmpGlobalTransforms);
-	Results.Get(0).GlobalTransforms = TmpGlobalTransforms;
-	Results.Get(1).GlobalTransforms = TmpGlobalTransforms;
-
-
-	// Initialize data for faster bound calculations
-	// precompute data used for bounds calculation
-	{
-		const TManagedArray<FBox>& BoundingBoxes = Parameters.RestCollection->BoundingBox;
-		const TManagedArray<int32>& TransformIndices = Parameters.RestCollection->TransformIndex;
-
-		const int32 NumBoxes = BoundingBoxes.Num();
-
-		ValidGeometryBoundingBoxes.Reset();
-		ValidGeometryTransformIndices.Reset();
-		for (int32 BoxIdx = 0; BoxIdx < NumBoxes; ++BoxIdx)
-		{
-			const int32 CurrTransformIndex = TransformIndices[BoxIdx];
-
-			if (Parameters.RestCollection->IsGeometry(CurrTransformIndex))
-			{
-				ValidGeometryBoundingBoxes.Add(BoundingBoxes[BoxIdx]);
-				ValidGeometryTransformIndices.Add(CurrTransformIndex);
-			}
-		}
-
-		FBox BoundingBox(ForceInit);
-		const FMatrix& ActorToWorld = Parameters.WorldTransform.ToMatrixWithScale();
-
-		for (int i = 0; i < ValidGeometryBoundingBoxes.Num(); ++i)
-		{
-			BoundingBox += ValidGeometryBoundingBoxes[i].TransformBy(TmpGlobalTransforms[ValidGeometryTransformIndices[i]] * ActorToWorld);
-		}
-
-		Results.Get(0).WorldBounds = FBoxSphereBounds(BoundingBox);
-		Results.Get(1).WorldBounds = FBoxSphereBounds(BoundingBox);
-	}
-}
-
-void FGeometryCollectionPhysicsProxy::Reset()
-{
-	InitializedState = ESimulationInitializationState::Unintialized;
-}
-
-
-
 int32 ReportNoLevelsetCluster = 0;
 FAutoConsoleVariableRef CVarReportNoLevelsetCluster(TEXT("p.gc.ReportNoLevelsetCluster"), ReportNoLevelsetCluster, TEXT("Report any cluster objects without levelsets"));
 
@@ -2017,6 +2496,8 @@ void FGeometryCollectionPhysicsProxy::CreateDynamicAttributes()
 	DynamicCollection->AddExternalAttribute("CollisionStructureID", FTransformCollection::TransformGroup, CollisionStructureID);
 	DynamicCollection->AddExternalAttribute(SimplicialsAttribute, FTransformCollection::TransformGroup, Simplicials);
 	DynamicCollection->AddExternalAttribute(ImplicitsAttribute, FTransformCollection::TransformGroup, Implicits);
+	DynamicCollection->AddExternalAttribute(SolverParticleHandlesAttribute, FTransformCollection::TransformGroup, SolverParticleHandles);
+	DynamicCollection->AddExternalAttribute(GTGeometryParticleAttribute, FTransformCollection::TransformGroup, GTGeometryParticles);
 
 	if (Parameters.InitialVelocityType == EInitialVelocityTypeEnum::Chaos_Initial_Velocity_User_Defined)
 	{
@@ -2101,7 +2582,10 @@ int32 FindSizeSpecificIdx(const TArray<FSharedSimulationSizeSpecificData>& SizeS
 	return UseIdx;
 }
 
-void FGeometryCollectionPhysicsProxy::InitializeSharedCollisionStructures(Chaos::FErrorReporter& ErrorReporter, FGeometryCollection& RestCollection, const FSharedSimulationParameters& SharedParams)
+void FGeometryCollectionPhysicsProxy::InitializeSharedCollisionStructures(
+	Chaos::FErrorReporter& ErrorReporter, 
+	FGeometryCollection& RestCollection, 
+	const FSharedSimulationParameters& SharedParams)
 {
 	FString BaseErrorPrefix = ErrorReporter.GetPrefix();
 
@@ -2121,18 +2605,30 @@ void FGeometryCollectionPhysicsProxy::InitializeSharedCollisionStructures(Chaos:
 	const TManagedArray<int32>& BoneMap = RestCollection.BoneMap;
 	const TManagedArray<int32>& Parent = RestCollection.Parent;
 	const TManagedArray<TSet<int32>>& Children = RestCollection.Children;
-	TManagedArray<bool>& CollectionSimulatableParticles = RestCollection.GetAttribute<bool>(FGeometryCollection::SimulatableParticlesAttribute, FTransformCollection::TransformGroup);
-
-	TManagedArray<FVector>& CollectionInertiaTensor = RestCollection.AddAttribute<FVector>(TEXT("InertiaTensor"), FTransformCollection::TransformGroup);
-	TManagedArray<FTransform>& CollectionMassToLocal = RestCollection.AddAttribute<FTransform>(TEXT("MassToLocal"), FTransformCollection::TransformGroup);
-	TManagedArray<float>& CollectionMass = RestCollection.AddAttribute<float>(TEXT("Mass"), FTransformCollection::TransformGroup);
-	TManagedArray<TUniquePtr<FSimplicial>>& CollectionSimplicials = RestCollection.AddAttribute<TUniquePtr<FSimplicial>>(SimplicialsAttribute, FTransformCollection::TransformGroup);
-	TManagedArray<TUniquePtr<FImplicitObject>>& CollectionImplicits = RestCollection.AddAttribute<TUniquePtr<FImplicitObject>>(ImplicitsAttribute, FTransformCollection::TransformGroup);
-
+	TManagedArray<bool>& CollectionSimulatableParticles = 
+		RestCollection.GetAttribute<bool>(
+			FGeometryCollection::SimulatableParticlesAttribute, FTransformCollection::TransformGroup);
+	TManagedArray<FVector>& CollectionInertiaTensor = 
+		RestCollection.AddAttribute<FVector>(
+			TEXT("InertiaTensor"), FTransformCollection::TransformGroup);
+	TManagedArray<FTransform>& CollectionMassToLocal = 
+		RestCollection.AddAttribute<FTransform>(
+			TEXT("MassToLocal"), FTransformCollection::TransformGroup);
+	TManagedArray<float>& CollectionMass = 
+		RestCollection.AddAttribute<float>(
+			TEXT("Mass"), FTransformCollection::TransformGroup);
+	TManagedArray<TUniquePtr<FSimplicial>>& CollectionSimplicials = 
+		RestCollection.AddAttribute<TUniquePtr<FSimplicial>>(
+			SimplicialsAttribute, FTransformCollection::TransformGroup);
+	TManagedArray<TUniquePtr<FImplicitObject>>& CollectionImplicits = 
+		RestCollection.AddAttribute<TUniquePtr<FImplicitObject>>(
+			ImplicitsAttribute, FTransformCollection::TransformGroup);
+
+	FTransform IdentityXf(FQuat::Identity, FVector(0));
+	IdentityXf.NormalizeRotation();
 	for (int32 Index = 0; Index < CollectionMassToLocal.Num(); ++Index)
 	{
-		CollectionMassToLocal[Index] = FTransform(FQuat::Identity, FVector(0));
-		CollectionMassToLocal[Index].NormalizeRotation();
+		CollectionMassToLocal[Index] = IdentityXf;
 	}
 
 	// VerticesGroup
@@ -3235,30 +3731,151 @@ void FGeometryCollectionPhysicsProxy::SyncBeforeDestroy()
 	}
 }
 
+Chaos::FParticleData* FGeometryCollectionPhysicsProxy::NewData() 
+{
+	/**
+	 * CONTEXT: GAMETHREAD
+	 * Returns a new unmanaged allocation of the data saved on the handle, otherwise nullptr
+	 */
+
+	// The other physics proxies allocate memory and returns a copy of game thread
+	// data as a TGeometryParticleData<T,d>, TKinematicGeometryParticleData<T,d>, 
+	// or a TPBDRigidParticleData<T,d>, which are all derived from FParticleData.
+	// That then gets passed to the physics thread via PushToPhysicsState().
+	//
+	// So all we're doing is pushing game thread state to the physics thread.
+	// I don't think we need to allocate memory to do this.  I think the other 
+	// proxies get away with that because they're dealing with a low number of 
+	// bodies; which is not the case here. Rather, let's just use a triple buffer.
+
+	check(NumParticles != INDEX_NONE); // Make sure InitBodiesGT() has been called!
+
+	FGeometryCollectionResults& Buffer = *GameToPhysInterchange->AccessProducerBuffer();
+	if (Buffer.NumParticlesAdded != NumParticles)
+	{
+		// Initialize the buffer
+		Buffer.NumParticlesAdded = NumParticles;
+		Buffer.BaseIndex = BaseParticleIndex;
+
+		TArray<FTransform> Transforms; Transforms.SetNum(NumParticles);
+		Buffer.Transforms = TManagedArray<FTransform>(MoveTemp(Transforms));
+		Buffer.DisabledStates.SetNum(NumParticles);
+	}
+
+	const FGeometryDynamicCollection* DynamicCollection = Parameters.DynamicCollection;
+	const FGeometryCollection* RestCollection = Parameters.RestCollection;
+
+	Buffer.TransformIndex.Init(RestCollection->TransformIndex);
+	Buffer.BoneMap.Init(RestCollection->BoneMap);
+	Buffer.Parent.Init(RestCollection->Parent);
+	Buffer.Children.Init(RestCollection->Children);
+	Buffer.SimulationType.Init(RestCollection->SimulationType);
+
+	Buffer.DynamicState.Init(DynamicCollection->DynamicState);
+	Buffer.Mass.Init(RestCollection->GetAttribute<float>("Mass", FTransformCollection::TransformGroup));
+	Buffer.InertiaTensor.Init(RestCollection->GetAttribute<FVector>("InertiaTensor", FTransformCollection::TransformGroup));
+	Buffer.Transforms.Init(DynamicCollection->Transform);
+
+	GeometryCollectionAlgo::GlobalMatrices(DynamicCollection->Transform, DynamicCollection->Parent, Buffer.GlobalTransforms);
+
+/*	for (int32 HandleIdx = 0; HandleIdx < NumParticles; ++HandleIdx)
+	{
+		if (Chaos::TPBDRigidParticleHandle<float, 3>* Handle = SolverParticleHandles[HandleIdx])
+		{
+			const int32 Idx = Buffer.BaseIndex + HandleIdx;
+			Buffer.Transforms[Idx].SetTranslation(Handle->X());
+			Buffer.Transforms[Idx].SetRotation(Handle->R());
+			Buffer.DisabledStates[Idx] = Handle->Disabled();
+			//...
+		}
+	}
+*/
+	// Make the buffer available to the consumer.
+	GameToPhysInterchange->FlipProducer(); 
+	
+	// If we wanted to emulate the other proxies, we could augment the 
+	// FStubGeometryCollectionData class to carry along the consumer buffer:
+	//
+	// return new FStubGeometryCollectionData(GametoPhysInterchange.GetConsumerBuffer());
+	//
+	// However, at this point, it seems pretty unnecessary to do the additional
+	// allocation and deallocation of the transient FStubGeometryCollectionData
+	// class, as I believe all these operations are happening with the same 
+	// instance of this class.  So skip all that, and just return null.
+	return nullptr;
+}
+
+void FGeometryCollectionPhysicsProxy::PushToPhysicsState(const Chaos::FParticleData *InData)
+{
+	/**
+	 * CONTEXT: GAMETHREAD -> to -> PHYSICSTHREAD
+	 * Called on the game thread when the solver is about to advance forward. This
+	 * callback should Enqueue commands on the PhysicsThread to update the state of
+	 * the solver
+	 */
+
+	// The other proxies do a deep copy into dynamic transient memory, passed into
+	// this function via InData.  Instead of that, we're using persistent memory
+	// in a buffer to facilitate communicating between threads.
+
+	const FGeometryCollectionResults* GState = GameToPhysInterchange->GetConsumerBuffer();
+	if (!GState)
+	{
+		// This will happen if GameToPhysInterchange hasn't been updated since
+		// the last time GetConsumerBuffer() was called.  At this point, that's 
+		// probably a bug, as the interface code is assuming it needs to lock
+		// to prevent stomping on a double buffer.  We should look into doing 
+		// away with that lock, but for the time being, we'll assert:
+		check(false);
+		return;
+	}
+	
+	//const FStubGeometryCollectionData* GCData = static_cast<const FStubGeometryCollectionData*>(InData);
+	//const FGeometryCollectionResults* GCState = GCData->GetStateData();
+
+	// Move the copied game thread data into the handles.
+	for (int32 HandleIdx = 0; HandleIdx < SolverParticleHandles.Num(); ++HandleIdx)
+	{
+		if (Chaos::TPBDRigidParticleHandle<float, 3>* Handle = SolverParticleHandles[HandleIdx])
+		{
+			const int32 Idx = GState->BaseIndex + HandleIdx;
+			Handle->SetX(GState->Transforms[Idx].GetTranslation());
+			Handle->SetR(GState->Transforms[Idx].GetRotation());
+			Handle->SetDisabled(GState->DisabledStates[Idx]);
+			//...
+		}
+	}
+}
+
 void FGeometryCollectionPhysicsProxy::BufferPhysicsResults()
 {
+	/**
+	 * CONTEXT: PHYSICSTHREAD
+	 * Called per-tick after the simulation has completed. The proxy should cache the results of their
+	 * simulation into the local buffer. 
+	 */
 	using namespace Chaos;
 	SCOPE_CYCLE_COUNTER(STAT_CacheResultGeomCollection);
 
-	FGeometryCollectionResults& TargetResults = Results.GetPhysicsDataForWrite();
+	//FGeometryCollectionResults& TargetResults = Results.GetPhysicsDataForWrite();
+	FGeometryCollectionResults& TargetResults = *PhysToGameInterchange->AccessProducerBuffer();
 
 	TManagedArray<FTransform>& TransformCache = TargetResults.Transforms;
 	TArray<FMatrix>& GlobalTransformCache = TargetResults.GlobalTransforms;	
-
-	TManagedArray<int32>& IdCache = TargetResults.RigidBodyIds;
+	//TManagedArray<int32>& IdCache = TargetResults.RigidBodyIds;
 	TManagedArray<int32>& ParentCache = TargetResults.Parent;
-	TManagedArray<TSet<int32>>& ChildrenCache = TargetResults.Children;
+	//TManagedArray<TSet<int32>>& ChildrenCache = TargetResults.Children;
 	TManagedArray<int32>& SimulationTypeCache = TargetResults.SimulationType;
-	TManagedArray<int32>& StatusFlagsCache = TargetResults.StatusFlags;
-#if TODO_REIMPLEMENT_GET_RIGID_PARTICLES
-	Chaos::FPhysicsSolver::FParticlesType& Particles = GetSolver()->GetRigidParticles();
+	//TManagedArray<int32>& StatusFlagsCache = TargetResults.StatusFlags;
+
+//#if TODO_REIMPLEMENT_GET_RIGID_PARTICLES
 
-	TransformCache.Init(SimulationCollection->Transform);
-	IdCache.Init(RigidBodyID);
-	ParentCache.Init(SimulationCollection->Parent);
-	ChildrenCache.Init(SimulationCollection->Children);
-	SimulationTypeCache.Init(SimulationCollection->SimulationType);
-	StatusFlagsCache.Init(SimulationCollection->StatusFlags);
+	TransformCache.Init(PTDynamicCollection.Transform);
+	//IdCache.Init(RigidBodyID);
+	ParentCache.Init(PTDynamicCollection.Parent);
+	//ChildrenCache.Init(PTDynamicCollection.Children);
+	SimulationTypeCache.Init(PTDynamicCollection.SimulationType);
+	//StatusFlagsCache.Init(PTDynamicCollection.StatusFlags);
 
 	// Base particle index to calculate index from a global particle index on the game thread
 	TargetResults.BaseIndex = BaseParticleIndex;
@@ -3272,12 +3889,12 @@ void FGeometryCollectionPhysicsProxy::BufferPhysicsResults()
 
 	// Advertise to game thread
 	TargetResults.IsObjectLoading = IsObjectLoading;
-
 	
-	// if object is dynamic, compute global matrices	
+	// If object is dynamic, compute global matrices	
 	if (IsObjectDynamic || GlobalTransformCache.Num() == 0)
 	{
 		SCOPE_CYCLE_COUNTER(STAT_CalcGlobalGCMatrices);
+		check(TransformCache.Num() == ParentCache.Num());
 		GeometryCollectionAlgo::GlobalMatrices(TransformCache, ParentCache, GlobalTransformCache);
 	}
 
@@ -3288,53 +3905,68 @@ void FGeometryCollectionPhysicsProxy::BufferPhysicsResults()
 		SCOPE_CYCLE_COUNTER(STAT_CalcGlobalGCBounds);
 		FBox BoundingBox(ForceInit);
 		const FMatrix& ActorToWorld = Parameters.WorldTransform.ToMatrixWithScale();
-
 		for (int i = 0; i < ValidGeometryBoundingBoxes.Num(); ++i)
 		{
-			BoundingBox += ValidGeometryBoundingBoxes[i].TransformBy(GlobalTransformCache[ValidGeometryTransformIndices[i]] * ActorToWorld);
+			BoundingBox += ValidGeometryBoundingBoxes[i].TransformBy(
+				GlobalTransformCache[ValidGeometryTransformIndices[i]] * ActorToWorld);
 		}
-
 		TargetResults.WorldBounds = FBoxSphereBounds(BoundingBox);		
 	}
 
-	if(NumParticles > 0)
+	//Chaos::FPhysicsSolver::FParticlesType& Particles = GetSolver()->GetRigidParticles();
+	//if(NumParticles > 0)
 	{
 		SCOPE_CYCLE_COUNTER(STAT_CaptureDisabledState);
-		TargetResults.DisabledStates.Append(&Particles.DisabledRef(BaseParticleIndex), NumParticles);
+		TargetResults.DisabledStates.SetNumUninitialized(NumParticles);
+		for (int32 Idx = 0; Idx < NumParticles; Idx++)
+		{
+			TargetResults.DisabledStates[Idx] = SolverParticleHandles[Idx]->Disabled();
+		}
+		//TargetResults.DisabledStates.Append(&Particles.DisabledRef(BaseParticleIndex), NumParticles);
 	}
-
 	{
 		SCOPE_CYCLE_COUNTER(STAT_CalcParticleToWorld);
 		// Fill particle to world transforms
 		TargetResults.ParticleToWorldTransforms.SetNum(NumParticles);
+/*
 		for (int32 TransformIndex = 0; TransformIndex < NumParticles; ++TransformIndex)
 		{
 			//only update roots and first children
 			if (ParentCache[TransformIndex] == INDEX_NONE)
 			{
-				const int32 ParticleIndex = BaseParticleIndex + TransformIndex;
-				TargetResults.ParticleToWorldTransforms[TransformIndex] = FTransform(Particles.R(ParticleIndex), Particles.X(ParticleIndex));
+				Chaos::TPBDGeometryCollectionParticleHandle<float, 3>* Handle = SolverParticleHandles[TransformIndex];
+				const FTransform Xf(Handle->R(), Handle->X());
+				//const int32 ParticleIndex = BaseParticleIndex + TransformIndex;
+				TargetResults.ParticleToWorldTransforms[TransformIndex] = Xf;
 			}
 		}
-
-		const TArrayCollectionArray<Chaos::ClusterId>& ClusterID = GetSolver()->GetRigidClustering().GetClusterIdsArray();
-		const TArrayCollectionArray<Chaos::TRigidTransform<float, 3>>& ClusterChildToParentMap = GetSolver()->GetRigidClustering().GetChildToParentMap();
-		const TArrayCollectionArray<FMultiChildProxyId>& MultiChildProxyIdArray = GetSolver()->GetRigidClustering().GetMultiChildProxyIdArray();
-		const TArrayCollectionArray<TUniquePtr<TMultiChildProxyData<float, 3>>>& MultiChildProxyDataArray = GetSolver()->GetRigidClustering().GetMultiChildProxyDataArray();
+*/
+		//const TArrayCollectionArray<Chaos::ClusterId>& ClusterID = GetSolver()->GetRigidClustering().GetClusterIdsArray();
+		//const TArrayCollectionArray<Chaos::TRigidTransform<float, 3>>& ClusterChildToParentMap = GetSolver()->GetRigidClustering().GetChildToParentMap();
+		//const TArrayCollectionArray<FMultiChildProxyId>& MultiChildProxyIdArray = GetSolver()->GetRigidClustering().GetMultiChildProxyIdArray();
+		//const TArrayCollectionArray<TUniquePtr<TMultiChildProxyData<float, 3>>>& MultiChildProxyDataArray = GetSolver()->GetRigidClustering().GetMultiChildProxyDataArray();
 
 		for (int32 TransformIndex = 0; TransformIndex < NumParticles; ++TransformIndex)
 		{
 			//only update roots and first children
 			const int32 ParticleIndex = BaseParticleIndex + TransformIndex;
-			const int32 ParentIndex = ClusterID[ParticleIndex].Id;
-			if (ParentIndex == INDEX_NONE)
+//			const int32 ParentIndex = ClusterID[ParticleIndex].Id;
+
+			Chaos::TPBDRigidParticleHandle<float, 3>& Handle = *SolverParticleHandles[ParticleIndex];
+
+			if (/*ParentIndex == INDEX_NONE*/true)
 			{
-				if (!Particles.Disabled(ParticleIndex))	//No need to copy disabled. If we don't guard against disabled, proxy particles will override their transform
+				// No need to copy disabled. If we don't guard against disabled, 
+				// proxy particles will override their transform.
+				if (!Handle.Disabled())	
 				{
-					TargetResults.ParticleToWorldTransforms[TransformIndex] = FTransform(Particles.R(ParticleIndex), Particles.X(ParticleIndex));
+					FTransform Xf(Handle.R(), Handle.X());
+
+					TargetResults.ParticleToWorldTransforms[TransformIndex] = Xf;
+					TargetResults.Transforms[TransformIndex] = Xf;
 				}
 			}
-			else if(ClusterID[ParentIndex].Id == INDEX_NONE)
+/*			else if(ClusterID[ParentIndex].Id == INDEX_NONE)
 			{
 				const int32 MultiChildProxyId = MultiChildProxyIdArray[ParticleIndex].Id;
 				const TMultiChildProxyData<float, 3>* ProxyData = MultiChildProxyId == INDEX_NONE ? nullptr : MultiChildProxyDataArray[MultiChildProxyId].Get();
@@ -3351,44 +3983,192 @@ void FGeometryCollectionPhysicsProxy::BufferPhysicsResults()
 					TargetResults.ParticleToWorldTransforms[TransformIndex] = ClusterChildToParentMap[ParticleIndex] * FTransform(Particles.R(ParentIndex), Particles.X(ParentIndex));
 				}
 			}
+*/
 		}
 	}
-#endif
+	// We've updated a buffer.  There's no reason to wait for FGeometryCollectionPhysicsProxy::FlipBuffer()
+	// to be called.  Just flip.
+	PhysToGameInterchange->FlipProducer();
+//#endif // TODO_REIMPLEMENT_GET_RIGID_PARTICLES
 }
 
 void FGeometryCollectionPhysicsProxy::FlipBuffer()
 {
-	Results.Flip();
+	/**
+	 * CONTEXT: PHYSICSTHREAD (Write Locked)
+	 * Called by the physics thread to signal that it is safe to perform any double-buffer flips here.
+	 * The physics thread has pre-locked an RW lock for this operation so the game thread won't be reading
+	 * the data
+	 */
+	//Results.Flip();
+	//PhysToGameInterchange->FlipProducer();
 }
 
 void FGeometryCollectionPhysicsProxy::PullFromPhysicsState()
 {
-	uint32 LastSyncCountFromPhysics = Results.GetGameDataSyncCount();
-	if(LastSyncCountFromPhysics != LastSyncCountGT)
+	/**
+	 * CONTEXT: GAMETHREAD (Read Locked)
+	 * Perform a similar operation to Sync, but take the data from a gamethread-safe buffer. This will be called
+	 * from the game thread when it cannot sync to the physics thread. The simulation is very likely to be running
+	 * when this happens so never read any physics thread data here!
+	 *
+	 * Note: A read lock will have been acquired for this - so the physics thread won't force a buffer flip while this
+	 * sync is ongoing
+	 */
+/* 
+	// We don't need to do this if we're using a guarded buffer, as it'll return 
+	// null if it doesn't have new data for us.
+	const uint32 LastSyncCountFromPhysics = Results.GetGameDataSyncCount();
+	if (LastSyncCountFromPhysics == LastSyncCountGT)
+	{
+		return;
+	}
+	LastSyncCountGT = LastSyncCountFromPhysics;
+*/
+	//FGeometryCollectionResults& TargetResult = Results.GetGameDataForWrite();
+	const FGeometryCollectionResults* TargetResultPtr = PhysToGameInterchange->GetConsumerBuffer();
+	if (!TargetResultPtr)
+		return;
+	// Remove const-ness as we're going to do the ExchangeArrays() thing.
+	FGeometryCollectionResults& TR = *const_cast<FGeometryCollectionResults*>(TargetResultPtr);
+/**/
+	// We should never be changing the number of entries, this would break other 
+	// attributes in the transform group.
+	if (ensure(GTDynamicCollection->Transform.Num() == TR.Transforms.Num()))	
 	{
-		LastSyncCountGT = LastSyncCountFromPhysics;
+		GTDynamicCollection->Transform.ExchangeArrays(TR.Transforms);
+//		GTDynamicCollection->Parent.ExchangeArrays(TR.Parent);
+//		GTDynamicCollection->Children.ExchangeArrays(TR.Children);
+//		GTDynamicCollection->SimulationType.ExchangeArrays(TR.SimulationType);
+//		GTDynamicCollection->StatusFlags.ExchangeArrays(TR.StatusFlags);
 
-		FGeometryCollectionResults& TargetResult = Results.GetGameDataForWrite();
+		//question: why do we need this? Sleeping objects will always have to update GPU
+		GTDynamicCollection->MakeDirty();
 
-		if(ensure(GTDynamicCollection->Transform.Num() == TargetResult.Transforms.Num()))	//we should never be changing the number of entries, this would break other attributes in the transform group
+		if (CacheSyncFunc)
+		{
+			CacheSyncFunc(TR);
+		}
+	}
+/**/
+/*
+	for (int32 TransformIndex = 0; TransformIndex < NumParticles; ++TransformIndex)
+	{
+		const int32 ParticleIndex = BaseParticleIndex + TransformIndex;
+		//Chaos::TPBDRigidParticleHandle<float, 3>& Handle = *SolverParticleHandles[ParticleIndex];
+		if (Chaos::TGeometryParticle<float, 3>* Particle = GTGeometryParticles[ParticleIndex].Get())
 		{
-			GTDynamicCollection->Transform.ExchangeArrays(TargetResult.Transforms);
-			GTDynamicCollection->Parent.ExchangeArrays(TargetResult.Parent);
-			GTDynamicCollection->Children.ExchangeArrays(TargetResult.Children);
-			GTDynamicCollection->SimulationType.ExchangeArrays(TargetResult.SimulationType);
-			GTDynamicCollection->StatusFlags.ExchangeArrays(TargetResult.StatusFlags);
+			const FTransform& Xf = TR.Transforms[ParticleIndex];
+			Particle->SetX(Xf.GetTranslation());
+			Particle->SetR(Xf.GetRotation());
+
+			GTDynamicCollection->Transform[TransformIndex] = Xf;
+		}
+	}
+*/
+//	UpdateGeometryCollection(TR);
+}
+
+void FGeometryCollectionPhysicsProxy::UpdateGeometryCollection(FGeometryCollectionResults& GCResults)
+{
+	if (InitializedState != ESimulationInitializationState::Initialized)
+	{
+		return;
+	}
 
-			GTDynamicCollection->MakeDirty();	//question: why do we need this? Sleeping objects will always have to update GPU
+	FGeometryDynamicCollection* Collection = Parameters.DynamicCollection;
+	check(Collection);
+	TManagedArray<FTransform>& Transform = Collection->Transform;
+	TManagedArray<int32>& Parent = Collection->Parent;
+	TManagedArray<TSet<int32>>& Children = Collection->Children;
+	TManagedArray<int32>& SimulationType = Collection->SimulationType;
+	TManagedArray<int32>& DynamicState = Collection->GetAttribute<int32>("DynamicState", FGeometryCollection::TransformGroup);
+
+	// Particles X and R are aligned with center of mass and inertia principal axes.
+	// Renderer doesn't know about this and simply does ActorToWorld * GeomToActor * LocalSpaceVerts
+	// In proper math multiplication order:
+	//		ParticleToWorld = ActorToWorld * GeomToActor * MassToLocal
+	//		GeomToWorld = ActorToWorld * GeomToActor
+	//		=> GeomToWorld = ParticleToWorld * MassToLocal.Inv()
+	//		=> GeomToActor = ActorToWorld.Inv() * ParticleToWorld * MassToLocal.Inv()
+
+	const FTransform& ActorToWorld = Parameters.WorldTransform;
+	const int32 TransformSize = GetTransformGroupSize();
+	for (int32 TransformGroupIndex = 0; TransformGroupIndex < TransformSize; ++TransformGroupIndex)
+	{
+		if (!GTGeometryParticles[TransformGroupIndex])
+			continue;
+
+		// Update the transform and parent hierarchy of the active rigid bodies. Active bodies can be either
+		// rigid geometry defined from the leaf nodes of the collection, or cluster bodies that drive an entire
+		// branch of the hierarchy within the GeometryCollection.
+		// - Active bodies are directly driven from the global position of the corresponding
+		//   rigid bodies within the solver ( cases where RigidBodyID[TransformGroupIndex] is not disabled ). 
+		// - Deactivated bodies are driven from the transforms of their active parents. However the solver can
+		//   take ownership of the parents during the simulation, so it might be necessary to force deactivated
+		//   bodies out of the collections hierarchy during the simulation.  
+
+		if (/*!Particles.Disabled(RigidBodyID[TransformGroupIndex])*/true)
+		{
+			// Update the transform of the active body. The active body can be either a single rigid
+			// or a collection of rigidly attached geometries (Clustering). The cluster is represented as a
+			// single transform in the GeometryCollection, and all children are stored in the local space
+			// of the parent cluster.
+			// ... When setting cluster transforms it is expected that the MassToLocal is identity.
+			//     Cluster initialization will set the vertices in the MassSpace of the rigid body.
+			// ... When setting individual rigid bodies that are not clustered, the MassToLocal will be 
+			//     non-Identity, and will reflect the difference between the geometric center of the geometry
+			//     and that corresponding rigid bodies center of mass. 
+
+			//const FTransform ParticleToWorld(Particles.R(RigidBodyIndex), Particles.X(RigidBodyIndex));
+			const FTransform& ParticleToWorld = GCResults.ParticleToWorldTransforms[TransformGroupIndex];
+				//Collection->Transform[TransformGroupIndex];
+
+			// GeomToActor = ActorToWorld.Inv() * ParticleToWorld * MassToLocal.Inv();
+			//Transform[TransformGroupIndex] = MassToLocal[TransformGroupIndex].GetRelativeTransformReverse(ParticleToWorld).GetRelativeTransform(ActorToWorld);
+			Transform[TransformGroupIndex] = ParticleToWorld;
+			Transform[TransformGroupIndex].NormalizeRotation();
+
+			// dynamic state is also updated by the solver during field interaction. 
+			if (/*!Particles.Sleeping(RigidBodyIndex)*/true)
+			{
+/*
+				if (Particles.ObjectState(RigidBodyIndex) == Chaos::EObjectStateType::Kinematic)
+				{
+					DynamicState[TransformGroupIndex] = (int)EObjectStateTypeEnum::Chaos_Object_Kinematic;
+				}
+				else if (Particles.ObjectState(RigidBodyIndex) == Chaos::EObjectStateType::Static)
+				{
+					DynamicState[TransformGroupIndex] = (int)EObjectStateTypeEnum::Chaos_Object_Static;
+				}
+				else
+				{
+*/
+					DynamicState[TransformGroupIndex] = (int)EObjectStateTypeEnum::Chaos_Object_Dynamic;
+//				}
+			}
 
-			if(CacheSyncFunc)
+			// Force all enabled rigid bodies out of the transform hierarchy
+			if (Parent[TransformGroupIndex] != INDEX_NONE)
 			{
-				CacheSyncFunc(TargetResult);
+				const int32 ParentIndex = Parent[TransformGroupIndex];
+				// Children in the hierarchy are stored in a TSet, which is not thread safe.  So we retain
+				// indices to remove afterwards.
+				EndFrameUnparentingBuffer[TransformGroupIndex] = ParentIndex;
 			}
+
+			// When a leaf node rigid body is removed from a cluster the rigid
+			// body will become active and needs its clusterID updated. This just
+			// syncs the clusterID all the time. 
+//			CollectionClusterID[TransformGroupIndex] = ClusterID[RigidBodyIndex].Id;
+		}
+		else //if (Particles.Disabled(RigidBodyIndex))
+		{
+			//...
 		}
 	}
 }
 
-
 void IdentifySimulatableElements(Chaos::FErrorReporter& ErrorReporter, FGeometryCollection& GeometryCollection)
 {
 	// Determine which collection particles to simulate
@@ -3431,7 +4211,9 @@ void IdentifySimulatableElements(Chaos::FErrorReporter& ErrorReporter, FGeometry
 
 
 	//For now all simulation data is a non compiled attribute. Not clear what we want for simulated vs kinematic collections
-	TManagedArray<bool>& SimulatableParticles = GeometryCollection.AddAttribute<bool>(FGeometryCollection::SimulatableParticlesAttribute, FTransformCollection::TransformGroup);
+	TManagedArray<bool>& SimulatableParticles = 
+		GeometryCollection.AddAttribute<bool>(
+			FGeometryCollection::SimulatableParticlesAttribute, FTransformCollection::TransformGroup);
 
 	for(int32 TransformIdx = 0; TransformIdx < NumTransforms; TransformIdx++)
 	{
diff --git a/Engine/Source/Runtime/Experimental/ChaosSolvers/Public/PBDRigidsSolver.h b/Engine/Source/Runtime/Experimental/ChaosSolvers/Public/PBDRigidsSolver.h
--- a/Engine/Source/Runtime/Experimental/ChaosSolvers/Public/PBDRigidsSolver.h
+++ b/Engine/Source/Runtime/Experimental/ChaosSolvers/Public/PBDRigidsSolver.h
@@ -112,8 +112,8 @@ namespace Chaos
 		void UnregisterObject(Chaos::TGeometryParticle<float, 3>* GTParticle);
 
 		// TODO: Set up an interface for registering fields and geometry collections
-		//void RegisterObject(FGeometryCollectionPhysicsProxy* InProxy);
-		//int UnregisterObject(FGeometryCollectionPhysicsProxy* InProxy);
+		void RegisterObject(FGeometryCollectionPhysicsProxy* InProxy);
+		bool UnregisterObject(FGeometryCollectionPhysicsProxy* InProxy);
 		//void RegisterObject(FFieldSystemPhysicsProxy* InProxy);
 		//int UnregisterObject(FFieldSystemPhysicsProxy* InProxy);
 
diff --git a/Engine/Source/Runtime/Experimental/ChaosSolvers/Public/PhysicsProxy/GeometryCollectionPhysicsProxy.h b/Engine/Source/Runtime/Experimental/ChaosSolvers/Public/PhysicsProxy/GeometryCollectionPhysicsProxy.h
--- a/Engine/Source/Runtime/Experimental/ChaosSolvers/Public/PhysicsProxy/GeometryCollectionPhysicsProxy.h
+++ b/Engine/Source/Runtime/Experimental/ChaosSolvers/Public/PhysicsProxy/GeometryCollectionPhysicsProxy.h
@@ -3,28 +3,51 @@
 #pragma once
 
 #include "GeometryCollection/GeometryCollection.h"
+#include "GeometryCollection/GeometryDynamicCollection.h"
 
 #include "Chaos/Framework/PhysicsProxy.h"
 #include "GeometryCollection/ManagedArray.h"
 #include "GeometryCollection/GeometryCollectionCollisionStructureManager.h"
 #include "Chaos/Framework/BufferedData.h"
+#include "Chaos/ParticleHandle.h"
+#include "Chaos/ParticleHandleFwd.h"
+#include "PBDRigidsSolver.h"
 
-struct FGeometryCollectionResults
+namespace Chaos
 {
+	template <typename T> class TSerializablePtr;
+	class FErrorReporter;
+	template <typename T> struct FClusterCreationParameters;
+}
+
+class FGeometryCollectionResults
+{
+public:
 	FGeometryCollectionResults();
 
+	void Reset();
+
 	int32 BaseIndex;
 	int32 NumParticlesAdded;
 	TArray<bool> DisabledStates;
-	TManagedArray<FTransform> Transforms;
+
 	TArray<FMatrix> GlobalTransforms;
-	TManagedArray<int32> RigidBodyIds;
-	TArray<FTransform> ParticleToWorldTransforms;
-	TManagedArray<int32> Level;
+
+	TManagedArray<int32> TransformIndex;
+
+	TManagedArray<FTransform> Transforms;
+	TManagedArray<int32> BoneMap;
 	TManagedArray<int32> Parent;
 	TManagedArray<TSet<int32>> Children;
 	TManagedArray<int32> SimulationType;
-	TManagedArray<int32> StatusFlags;
+	TManagedArray<int32> DynamicState;
+	TManagedArray<float> Mass;
+	TManagedArray<FVector> InertiaTensor;
+
+	//TManagedArray<int32> RigidBodyIds;
+	TArray<FTransform> ParticleToWorldTransforms;
+	//TManagedArray<int32> Level;
+	//TManagedArray<int32> StatusFlags;
 
 	bool IsObjectDynamic;
 	bool IsObjectLoading;
@@ -32,19 +55,32 @@ struct FGeometryCollectionResults
 	FBoxSphereBounds WorldBounds;
 };
 
-namespace Chaos
+
+class FStubGeometryCollectionData : public Chaos::FParticleData 
 {
-template <typename T>
-class TSerializablePtr;
-class FErrorReporter;
-template <typename T>
-struct FClusterCreationParameters;
-}
+public:
+	typedef Chaos::FParticleData Base;
+
+	FStubGeometryCollectionData(const FGeometryCollectionResults* DataIn=nullptr)
+		: Base(Chaos::EParticleType::GeometryCollection)
+		, Data(DataIn)
+	{}
+
+	void Reset() 
+	{
+		Base::Reset(); // Sets Type to EParticleType::Static
+	}
 
-class FStubGeometryCollectionData : public Chaos::FParticleData {
-	void Reset() {};
+	const FGeometryCollectionResults* GetStateData() const { return Data; }
+
+private:
+	const FGeometryCollectionResults* Data;
 };
 
+/**
+ * Class to manage sharing data between the game thread and the simulation thread 
+ * (which may not be different than the game thread) for a \c FGeometryDynamicCollection.
+ */
 class CHAOSSOLVERS_API FGeometryCollectionPhysicsProxy : public TPhysicsProxy<FGeometryCollectionPhysicsProxy, FStubGeometryCollectionData>
 {
 	typedef TPhysicsProxy<FGeometryCollectionPhysicsProxy, FStubGeometryCollectionData> Base;
@@ -53,6 +89,8 @@ public:
 	// collection attributes
 	static FName SimplicialsAttribute;
 	static FName ImplicitsAttribute;
+	static FName SolverParticleHandlesAttribute;
+	static FName GTGeometryParticleAttribute;
 
 	typedef FCollisionStructureManager::FSimplicial FSimplicial;
 
@@ -63,14 +101,40 @@ public:
 	/** ----------------------- */
 
 	FGeometryCollectionPhysicsProxy() = delete;
-	FGeometryCollectionPhysicsProxy(UObject* InOwner, FGeometryDynamicCollection* InDynamicCollection, FInitFunc InInitFunc, FCacheSyncFunc InCacheSyncFunc, FFinalSyncFunc InFinalSyncFunc);
+	/**
+	 * \p InOwner
+	 * \p InDynamicCollection game thread owned geometry collection.
+	 * \p InInitFunc callback invoked from \c Initialize().
+	 * \p InCacheSyncFunc callback invoked from \c PullFromPhysicsState().
+	 * \p InFinalSyncFunc callback invoked from \c SyncBeforeDestory().
+	 */
+	FGeometryCollectionPhysicsProxy(
+		UObject* InOwner, 
+		FGeometryDynamicCollection* InDynamicCollection, 
+		FInitFunc InInitFunc, 
+		FCacheSyncFunc InCacheSyncFunc, 
+		FFinalSyncFunc InFinalSyncFunc  ,
+		const Chaos::EMultiBufferMode BufferMode=Chaos::EMultiBufferMode::TripleGuarded);
 	virtual ~FGeometryCollectionPhysicsProxy();
 
+	/**
+	 * Construct \c PTDynamicCollection, copying attributes from the game thread, 
+	 * and prepare for simulation.
+	 */
 	void Initialize();
-	void Reset();
+
+	/** Initialize rigid body particles on the game thread. */
+	void InitializeBodiesGT();
+
+	/** Called from the physics thread. */
+	void InitializeBodiesPT(
+		Chaos::FPBDRigidsSolver* RigidsSolver,
+		Chaos::FPBDRigidsSolver::FParticlesType& Particles);
+
+	void Reset() { InitializedState = ESimulationInitializationState::Unintialized; }
 
 	/** Solver Object interface */
-	bool IsSimulating() const;
+	bool IsSimulating() const { return Parameters.Simulating; }
 	void UpdateKinematicBodiesCallback(const FParticlesType& InParticles, const float InDt, const float InTime, FKinematicProxy& InKinematicProxy);
 	void StartFrameCallback(const float InDt, const float InTime);
 	void EndFrameCallback(const float InDt);
@@ -82,13 +146,16 @@ public:
 	void FieldForcesUpdateCallback(Chaos::FPhysicsSolver* InSolver, FParticlesType& Particles, Chaos::TArrayCollectionArray<FVector> & Force, Chaos::TArrayCollectionArray<FVector> & Torque, const float Time);
 	void ActivateBodies();
 
-	Chaos::FParticleData* NewData() { return nullptr; }
 	void SyncBeforeDestroy();
 	void OnRemoveFromScene();
-	void PushToPhysicsState(const Chaos::FParticleData*) {};
+
+	Chaos::FParticleData* NewData();
+	/** Invoked by \c FPBDRigidsSolver::RegisterObject(). */
+	void PushToPhysicsState(const Chaos::FParticleData*);
 	void ClearAccumulatedData() {}
 	void BufferPhysicsResults();
 	void FlipBuffer();
+	/** Updates \c GTDynamicCollection with curr state of the \c Results buffer. */
 	void PullFromPhysicsState();
 	bool IsDirty() { return false; }
 	EPhysicsProxyType ConcreteType() { return EPhysicsProxyType::GeometryCollectionType; }
@@ -105,13 +172,17 @@ public:
 	const Chaos::TBufferedData<FGeometryCollectionResults>& GetPhysicsResults() const { return Results; }
 
 	// Game thread and physics thread access to the dynamic collection. 
-	const FGeometryDynamicCollection* GetGeometryDynamicCollection_PhysicsThread() { return SimulationCollection.Get(); }
+	const FGeometryDynamicCollection* GetGeometryDynamicCollection_PhysicsThread() { return &PTDynamicCollection; }
 	const FGeometryDynamicCollection* GetGeometryDynamicCollection_GameThread() { return GTDynamicCollection; }
 
 	const FSimulationParameters& GetSimulationParameters() { return Parameters; }
 
 	int32 GetBaseParticleIndex() const { return BaseParticleIndex; }
 	int32 GetNumAddedParticles() const { return NumParticles; }
+	int32 GetTransformGroupSize() const { return Parameters.DynamicCollection->NumElements(FGeometryCollection::TransformGroup); }
+
+	TManagedArray<Chaos::TPBDGeometryCollectionParticleHandle<float, 3>*>& GetSolverParticleHandles() { return SolverParticleHandles; }
+	void ClearSolverParticleHandles() { SolverParticleHandles.Fill(nullptr); }
 
 	/*
 	* ContiguousIndices
@@ -130,6 +201,9 @@ public:
 	const TManagedArray<int32>& RigidBodyIDArray_TestingAccess() const { return RigidBodyID; }
 	const TManagedArray<FTransform>& MassToLocal_GameThreadAccess() const { return MassToLocal; }
 
+protected:
+	void UpdateGeometryCollection(FGeometryCollectionResults& Results);
+
 protected:
 	int32 CalculateHierarchyLevel(FGeometryDynamicCollection* GeometryCollection, int32 TransformIndex);
 	void CreateDynamicAttributes();
@@ -138,11 +212,56 @@ protected:
 	void ProcessCommands(FParticlesType& Particles, const float Time);
 	void PushKinematicStateToSolver(FParticlesType& Particles);
 
+private:
+	void UpdateRecordedState(
+		float SolverTime, 
+		const TManagedArray<int32>& RigidBodyID, 
+		const TManagedArray<int32>& CollectionClusterID, 
+		const Chaos::TArrayCollectionArray<bool>& InternalCluster, 
+		const FParticlesType& Particles, 
+		const FCollisionConstraintsType& CollisionRule);
+
+	void AddCollisionToCollisionData(
+		FRecordedFrame* ExistingFrame, 
+		const FParticlesType& Particles, 
+		const Chaos::TPBDCollisionConstraint<float, 3>::FPointContactConstraint& Constraint);
+
+	void UpdateCollisionData(
+		const FParticlesType& Particles, 
+		const FCollisionConstraintsType& CollisionRule, 
+		FRecordedFrame* ExistingFrame);
+
+	void AddBreakingToBreakingData(
+		FRecordedFrame* ExistingFrame, 
+		const FParticlesType& Particles, 
+		const Chaos::TBreakingData<float, 3>& Breaking);
+
+	void UpdateBreakingData(
+		const FParticlesType& Particles, 
+		FRecordedFrame* ExistingFrame);
+
+	void UpdateTrailingData(
+		const FParticlesType& Particles, 
+		FRecordedFrame* ExistingFrame);
 
-
-	
 private:
+	FSimulationParameters Parameters;
+
 	ESimulationInitializationState InitializedState;
+	// Records current dynamic state
+	bool IsObjectDynamic;
+	// Indicate when loaded
+	bool IsObjectLoading;
+
+	// Dynamic collection on the game thread - used to populate the simulated collection
+	FGeometryDynamicCollection* GTDynamicCollection;
+	// Duplicated dynamic collection for use on the physics thread, copied to the game thread on sync
+	//TUniquePtr<FGeometryDynamicCollection> PTDynamicCollection;
+	FGeometryDynamicCollection PTDynamicCollection;
+
+	//TManagedArray<Chaos::TPBDRigidParticleHandle<float, 3>*> SolverParticleHandles;
+	TManagedArray<Chaos::TPBDGeometryCollectionParticleHandle<float, 3>*> SolverParticleHandles;
+	TManagedArray<TUniquePtr<Chaos::TGeometryParticle<float, 3>>> GTGeometryParticles; // unique arrays for geo, kin, dyn?
 
 	TManagedArray<FTransform> MassToLocal;
 	TManagedArray<int32> CollisionMask;
@@ -152,7 +271,7 @@ private:
 	TManagedArray<FVector> InitialAngularVelocity;
 	TManagedArray<FVector> InitialLinearVelocity;
 	TManagedArray<bool> SimulatableParticles;
-	TManagedArray<TUniquePtr<FSimplicial> > Simplicials;
+	TManagedArray<TUniquePtr<FSimplicial> > Simplicials; // FSimplicial = Chaos::TBVHParticles<float,3>
 	TManagedArray<Chaos::TSerializablePtr<Chaos::FImplicitObject>> Implicits;
 	TArray<int32> EndFrameUnparentingBuffer;
 
@@ -160,7 +279,6 @@ private:
 	TArray<FBox> ValidGeometryBoundingBoxes;
 	TArray<int32> ValidGeometryTransformIndices;
 
-	FSimulationParameters Parameters;
 	TArray<FFieldSystemCommand> Commands;
 
 	TFunction<void(void)> ResetAnimationCacheCallback;
@@ -171,31 +289,12 @@ private:
 
 	// Index of the first particles for this collection in the larger particle array
 	int32 BaseParticleIndex;
-
 	// Number of particles added by this collection
 	int32 NumParticles;
 
 	// Time since this object started simulating
 	float ProxySimDuration;
 
-	void UpdateRecordedState(float SolverTime, const TManagedArray<int32>& RigidBodyID, const TManagedArray<int32>& CollectionClusterID, const Chaos::TArrayCollectionArray<bool>& InternalCluster, const FParticlesType& Particles, const FCollisionConstraintsType& CollisionRule);
-
-	void AddCollisionToCollisionData(FRecordedFrame* ExistingFrame, const FParticlesType& Particles, const Chaos::TPBDCollisionConstraint<float, 3>::FPointContactConstraint& Constraint);
-
-	void UpdateCollisionData(const FParticlesType& Particles, const FCollisionConstraintsType& CollisionRule, FRecordedFrame* ExistingFrame);
-
-	void AddBreakingToBreakingData(FRecordedFrame* ExistingFrame, const FParticlesType& Particles, const Chaos::TBreakingData<float, 3>& Breaking);
-
-	void UpdateBreakingData(const FParticlesType& Particles, FRecordedFrame* ExistingFrame);
-
-	void UpdateTrailingData(const FParticlesType& Particles, FRecordedFrame* ExistingFrame);
-
-	// Duplicated dynamic collection for use on the physics thread, copied to the game thread on sync
-	TUniquePtr<FGeometryDynamicCollection> SimulationCollection;
-
-	// Dynamic collection on the game thread - used to populate the simulated collection
-	FGeometryDynamicCollection* GTDynamicCollection;
-
 	// bodies waiting to activate on initialization
 	TArray<uint32> PendingActivationList;
 
@@ -211,20 +310,23 @@ private:
 	// Sync frame numbers so we don't do many syncs when physics is running behind
 	uint32 LastSyncCountGT;
 
-	// Records current dynamic state
-	bool IsObjectDynamic;
-
-	// Indicate when loaded
-	bool IsObjectLoading;
-
 	// Per object collision fraction.
 	float CollisionParticlesPerObjectFraction;
 
-	// Double buffer of geom collection result data
+	// Double buffer of geom collection result data - TODO (Ryan): deprecated?
 	Chaos::TBufferedData<FGeometryCollectionResults> Results;
 
-	/** ----------------------- */
-
+	// TODO (Ryan) - Currently this is using triple buffers for game-physics and 
+	// physics-game thread communication, but not for any reason other than this 
+	// is the only implementation we currently have of a guarded buffer - a buffer 
+	// that tracks it's own state, rather than having other mechanisms determine 
+	// whether or not the contents of the buffer have been updated.  A double 
+	// buffer would probably be fine, as that seems to be the assumption the logic
+	// currently managing the exchange is built upon.  However, I believe that 
+	// logic locks, and the triple buffer would enable a decoupled lock-free 
+	// paradigm, at least for this component of the handshake.
+	TUniquePtr<Chaos::IBufferResource<FGeometryCollectionResults>> PhysToGameInterchange;
+	TUniquePtr<Chaos::IBufferResource<FGeometryCollectionResults>> GameToPhysInterchange;
 };
 
 CHAOSSOLVERS_API void BuildSimulationData(Chaos::FErrorReporter& ErrorReporter, FGeometryCollection& GeometryCollection, const FSharedSimulationParameters& SharedParams);
diff --git a/Engine/Source/Runtime/Experimental/ChaosSolvers/Public/PhysicsProxy/SingleParticlePhysicsProxy.h b/Engine/Source/Runtime/Experimental/ChaosSolvers/Public/PhysicsProxy/SingleParticlePhysicsProxy.h
--- a/Engine/Source/Runtime/Experimental/ChaosSolvers/Public/PhysicsProxy/SingleParticlePhysicsProxy.h
+++ b/Engine/Source/Runtime/Experimental/ChaosSolvers/Public/PhysicsProxy/SingleParticlePhysicsProxy.h
@@ -40,6 +40,12 @@ private:
 	FVector InertiaTensor;
 };
 
+/**
+ * \p PARTICLE_TYPE is one of:
+ *		\c Chaos::TGeometryParticle<float,3>
+ *		\c Chaos::TKinematicGeometryParticle<float,3>
+ *		\c Chaos::TPBDRigidParticle<float,3>
+ */
 template<class PARTICLE_TYPE>
 class FSingleParticlePhysicsProxy : public TPhysicsProxy<FSingleParticlePhysicsProxy<PARTICLE_TYPE>, typename PARTICLE_TYPE::FData>
 {
@@ -99,7 +105,10 @@ public:
 		return Chaos::TRigidTransform<float, 3>(Handle->X(), Handle->R());
 	}
 
-	/**/
+	/** 
+	 * Creates a copy of sim state and returns it as \c TGeometryParticleData, 
+	 * \c TKinemticGeometr yParticleData, or \c TPBDRigidParticleData. 
+	 */
 	Chaos::FParticleData* NewData()
 	{
 		if (Particle)
diff --git a/Engine/Source/Runtime/Experimental/GeometryCollectionCore/Public/GeometryCollection/ManagedArrayTypeValues.inl b/Engine/Source/Runtime/Experimental/GeometryCollectionCore/Public/GeometryCollection/ManagedArrayTypeValues.inl
--- a/Engine/Source/Runtime/Experimental/GeometryCollectionCore/Public/GeometryCollection/ManagedArrayTypeValues.inl
+++ b/Engine/Source/Runtime/Experimental/GeometryCollectionCore/Public/GeometryCollection/ManagedArrayTypeValues.inl
@@ -33,6 +33,9 @@ MANAGED_ARRAY_TYPE(TUniquePtr<Chaos::FImplicitObject3>, FImplicitObject3UniquePo
 MANAGED_ARRAY_TYPE(Chaos::TSerializablePtr<Chaos::FImplicitObject3>, FImplicitObject3SerializablePtr)
 MANAGED_ARRAY_TYPE(Chaos::FBVHParticlesFloat3, FBVHParticlesFloat3Pointer)
 MANAGED_ARRAY_TYPE(TUniquePtr<Chaos::FBVHParticlesFloat3>, FBVHParticlesFloat3UniquePointer)
+MANAGED_ARRAY_TYPE(Chaos::TPBDRigidParticleHandleFloat3*, TPBDRigidParticleHandle3fPtr)
+MANAGED_ARRAY_TYPE(Chaos::TPBDGeometryCollectionParticleHandleFloat3*, TPBDGeometryCollectionParticleHandle3fPtr)
+MANAGED_ARRAY_TYPE(TUniquePtr<Chaos::TGeometryParticleFloat3>, TGeometryParticle3fUniquePtr)
 
 // NOTE: new types must be added at the bottom to keep serialization from breaking
 
diff --git a/Engine/Source/Runtime/Experimental/GeometryCollectionCore/Public/GeometryCollection/ManagedArrayTypes.h b/Engine/Source/Runtime/Experimental/GeometryCollectionCore/Public/GeometryCollection/ManagedArrayTypes.h
--- a/Engine/Source/Runtime/Experimental/GeometryCollectionCore/Public/GeometryCollection/ManagedArrayTypes.h
+++ b/Engine/Source/Runtime/Experimental/GeometryCollectionCore/Public/GeometryCollection/ManagedArrayTypes.h
@@ -12,6 +12,8 @@
 #include "Math/Box.h"
 #include "Chaos/ImplicitObject.h"
 #include "Chaos/BVHParticles.h"
+#include "Chaos/ParticleHandle.h"
+#include "Chaos/ParticleHandleFwd.h"
 
 inline FArchive& operator<<(FArchive& Ar, TArray<FVector>*& ValueIn)
 {
@@ -40,6 +42,14 @@ inline FArchive& operator<<(FArchive& Ar, Chaos::FImplicitObject*& ValueIn)
 	check(false);	//We don't serialize raw pointers to implicit objects. Use unique ptr
 	return Ar;
 }
+
+template <typename T, int d, bool bPersistent>
+inline FArchive& operator<<(FArchive& Ar, Chaos::TPBDRigidParticleHandleImp<T, d, bPersistent>*& Particle)
+{
+	verifyf(false, TEXT("TPBDRigidParticleHandleImp* should never be serialized!  Use unique ptr."));
+	return Ar;
+}
+
 // ---------------------------------------------------------
 //
 // General purpose EManagedArrayType definition. 
