Uncomment use of FChaosArchive for now.

diff --git a/Engine/Source/Runtime/Experimental/Chaos/Public/Chaos/Levelset.h b/Engine/Source/Runtime/Experimental/Chaos/Public/Chaos/Levelset.h
--- a/Engine/Source/Runtime/Experimental/Chaos/Public/Chaos/Levelset.h
+++ b/Engine/Source/Runtime/Experimental/Chaos/Public/Chaos/Levelset.h
@@ -70,10 +70,10 @@ class CHAOS_API TLevelSet final : public FImplicitObject
 		Ar << MBandWidth;
 	}
 
-	//virtual void Serialize(FChaosArchive& Ar) override
-	//{
-	//	SerializeImp(Ar);
-	//}
+	virtual void Serialize(FChaosArchive& Ar) override
+	{
+		SerializeImp(Ar);
+	}
 
 	virtual void Serialize(FArchive& Ar) override
 	{
