Add safe guard against bad hierarchy.

diff --git a/Engine/Source/Runtime/Experimental/GeometryCollectionCore/Private/GeometryCollection/GeometryCollectionClusteringUtility.cpp b/Engine/Source/Runtime/Experimental/GeometryCollectionCore/Private/GeometryCollection/GeometryCollectionClusteringUtility.cpp
--- a/Engine/Source/Runtime/Experimental/GeometryCollectionCore/Private/GeometryCollection/GeometryCollectionClusteringUtility.cpp
+++ b/Engine/Source/Runtime/Experimental/GeometryCollectionCore/Private/GeometryCollection/GeometryCollectionClusteringUtility.cpp
@@ -746,6 +746,11 @@ void FGeometryCollectionClusteringUtility::ContextBasedClusterSelection(
 
 void FGeometryCollectionClusteringUtility::GetLeafBones(FGeometryCollection* GeometryCollection, int BoneIndex, TArray<int32>& LeafBonesOut)
 {
+	if (!ensure(BoneIndex >= 0))
+	{
+		return;
+	}
+
 	const TManagedArray<TSet<int32>>& Children = GeometryCollection->Children;
 
 	if (Children[BoneIndex].Num() > 0)
