Chaos : Collision

Class name and file name cleanup.

diff --git a/Engine/Plugins/Experimental/ChaosNiagara/Source/ChaosNiagara/Private/NiagaraDataInterfaceChaosDestruction.cpp b/Engine/Plugins/Experimental/ChaosNiagara/Source/ChaosNiagara/Private/NiagaraDataInterfaceChaosDestruction.cpp
--- a/Engine/Plugins/Experimental/ChaosNiagara/Source/ChaosNiagara/Private/NiagaraDataInterfaceChaosDestruction.cpp
+++ b/Engine/Plugins/Experimental/ChaosNiagara/Source/ChaosNiagara/Private/NiagaraDataInterfaceChaosDestruction.cpp
@@ -7,7 +7,7 @@
 #include "ShaderParameterUtils.h"
 #include "PhysicsSolver.h"
 #include "Niagara/Private/NiagaraStats.h"
-#include "Chaos/PBDCollisionConstraintUtil.h"
+#include "Chaos/PBDCollisionConstraintsUtil.h"
 #include "PhysicsProxy/GeometryCollectionPhysicsProxy.h"
 #include "Chaos/CollisionResolutionTypes.h"
 #include "PhysicalMaterials/PhysicalMaterial.h"
diff --git a/Engine/Source/Runtime/Engine/Private/PhysicsEngine/ImmediatePhysics/ImmediatePhysicsChaos/ImmediatePhysicsSimulation_Chaos.cpp b/Engine/Source/Runtime/Engine/Private/PhysicsEngine/ImmediatePhysics/ImmediatePhysicsChaos/ImmediatePhysicsSimulation_Chaos.cpp
--- a/Engine/Source/Runtime/Engine/Private/PhysicsEngine/ImmediatePhysics/ImmediatePhysicsChaos/ImmediatePhysicsSimulation_Chaos.cpp
+++ b/Engine/Source/Runtime/Engine/Private/PhysicsEngine/ImmediatePhysics/ImmediatePhysicsChaos/ImmediatePhysicsSimulation_Chaos.cpp
@@ -8,7 +8,7 @@
 #include "Chaos/ChaosDebugDraw.h"
 #include "Chaos/DebugDrawQueue.h"
 #include "Chaos/ParticleHandle.h"
-#include "Chaos/PBDCollisionConstraint.h"
+#include "Chaos/PBDCollisionConstraints.h"
 #include "Chaos/PBDJointConstraints.h"
 #include "Chaos/PBDRigidParticles.h"
 #include "Chaos/PBDRigidsEvolutionGBF.h"
@@ -155,7 +155,7 @@ namespace ImmediatePhysics_Chaos
 		Joints = MakeUnique<FPBDJointConstraints>(FPBDJointSolverSettings());
 		JointsRule = MakeUnique<TPBDConstraintIslandRule<FPBDJointConstraints>>(*Joints);
 		Evolution = MakeUnique<TPBDRigidsEvolutionGBF<FReal, Dimensions>>(*Particles.Get(), 1);
-		TPBDCollisionConstraint<FReal, Dimensions>& Collisions = Evolution->GetCollisionConstraints();
+		TPBDCollisionConstraints<FReal, Dimensions>& Collisions = Evolution->GetCollisionConstraints();
 
 		Evolution->AddConstraintRule(JointsRule.Get());
 
@@ -169,7 +169,7 @@ namespace ImmediatePhysics_Chaos
 			[this]()
 			{
 				Evolution->GetCollisionConstraints().ApplyCollisionModifier(
-					[this](const TPBDCollisionConstraint<FReal, Dimensions>::FConstraintContainerHandle* Handle)
+					[this](const TPBDCollisionConstraints<FReal, Dimensions>::FConstraintContainerHandle* Handle)
 					{
 						TVector<const TGeometryParticleHandle<FReal, Dimensions>*, 2> ConstrainedParticles = Handle->GetConstrainedParticles();
 						if (ShouldIgnoreCollisionConstraint(ConstrainedParticles[0], ConstrainedParticles[1], IgnoreCollisionParticlePairTable))
diff --git a/Engine/Source/Runtime/Experimental/Chaos/Private/Chaos/ChaosDebugDraw.cpp b/Engine/Source/Runtime/Experimental/Chaos/Private/Chaos/ChaosDebugDraw.cpp
--- a/Engine/Source/Runtime/Experimental/Chaos/Private/Chaos/ChaosDebugDraw.cpp
+++ b/Engine/Source/Runtime/Experimental/Chaos/Private/Chaos/ChaosDebugDraw.cpp
@@ -8,7 +8,7 @@
 #include "Chaos/ImplicitObjectUnion.h"
 #include "Chaos/ParticleHandle.h"
 #include "Chaos/PBD6DJointConstraints.h"
-#include "Chaos/PBDCollisionConstraint.h"
+#include "Chaos/PBDCollisionConstraints.h"
 #include "Chaos/PBDJointConstraints.h"
 #include "Chaos/PBDRigidParticles.h"
 #include "Chaos/Sphere.h"
@@ -297,7 +297,7 @@ namespace Chaos
 #endif
 		}
 
-		void DrawParticleCollisions(const FRigidTransform3& SpaceTransform, const TGeometryParticleHandle<float, 3>* Particle, const TPBDCollisionConstraint<float, 3>& Collisions)
+		void DrawParticleCollisions(const FRigidTransform3& SpaceTransform, const TGeometryParticleHandle<float, 3>* Particle, const TPBDCollisionConstraints<float, 3>& Collisions)
 		{
 #if CHAOS_DEBUG_DRAW
 			if (FDebugDrawQueue::IsDebugDrawingEnabled())
@@ -314,7 +314,7 @@ namespace Chaos
 #endif
 		}
 
-		void DrawCollisions(const FRigidTransform3& SpaceTransform, const TPBDCollisionConstraint<float, 3>& Collisions, float ColorScale)
+		void DrawCollisions(const FRigidTransform3& SpaceTransform, const TPBDCollisionConstraints<float, 3>& Collisions, float ColorScale)
 		{
 #if CHAOS_DEBUG_DRAW
 			if (FDebugDrawQueue::IsDebugDrawingEnabled())
diff --git a/Engine/Source/Runtime/Experimental/Chaos/Private/Chaos/CollisionResolutionAlgo.cpp b/Engine/Source/Runtime/Experimental/Chaos/Private/Chaos/CollisionResolutionAlgo.cpp
--- a/Engine/Source/Runtime/Experimental/Chaos/Private/Chaos/CollisionResolutionAlgo.cpp
+++ b/Engine/Source/Runtime/Experimental/Chaos/Private/Chaos/CollisionResolutionAlgo.cpp
@@ -18,7 +18,7 @@
 #include "Stats/Stats.h" 
 
 #if INTEL_ISPC
-#include "PBDCollisionConstraint.ispc.generated.h"
+#include "PBDCollisionConstraints.ispc.generated.h"
 #endif
 
 namespace Chaos
@@ -108,7 +108,7 @@ namespace Chaos
 			Constraint.AddManifold(Implicit0, Implicit1);
 		}
 
-		DECLARE_CYCLE_STAT(TEXT("TPBDCollisionConstraint::UpdateUnionUnionConstraint"), STAT_UpdateUnionUnionConstraint, STATGROUP_ChaosWide);
+		DECLARE_CYCLE_STAT(TEXT("TPBDCollisionConstraints::UpdateUnionUnionConstraint"), STAT_UpdateUnionUnionConstraint, STATGROUP_ChaosWide);
 
 		template<typename T, int d>
 		void ConstructUnionUnionConstraints(TGeometryParticleHandle<T, d>* Particle0, TGeometryParticleHandle<T, d>* Particle1, const FImplicitObject* Implicit0, const FImplicitObject* Implicit1, const T Thickness, TRigidBodyPointContactConstraint<T, d>& Constraint)
@@ -271,12 +271,12 @@ namespace Chaos
 			return false;
 		}
 
-		DECLARE_CYCLE_STAT(TEXT("TPBDCollisionConstraint::UpdateLevelsetPartial"), STAT_UpdateLevelsetPartial, STATGROUP_ChaosWide);
-		DECLARE_CYCLE_STAT(TEXT("TPBDCollisionConstraint::UpdateLevelsetFindParticles"), STAT_UpdateLevelsetFindParticles, STATGROUP_ChaosWide);
-		DECLARE_CYCLE_STAT(TEXT("TPBDCollisionConstraint::UpdateLevelsetBVHTraversal"), STAT_UpdateLevelsetBVHTraversal, STATGROUP_ChaosWide);
-		DECLARE_CYCLE_STAT(TEXT("TPBDCollisionConstraint::UpdateLevelsetSignedDistance"), STAT_UpdateLevelsetSignedDistance, STATGROUP_ChaosWide);
-		DECLARE_CYCLE_STAT(TEXT("TPBDCollisionConstraint::UpdateLevelsetAll"), STAT_UpdateLevelsetAll, STATGROUP_ChaosWide);
-		DECLARE_CYCLE_STAT(TEXT("TPBDCollisionConstraint::SampleObject"), STAT_SampleObject, STATGROUP_ChaosWide);
+		DECLARE_CYCLE_STAT(TEXT("TPBDCollisionConstraints::UpdateLevelsetPartial"), STAT_UpdateLevelsetPartial, STATGROUP_ChaosWide);
+		DECLARE_CYCLE_STAT(TEXT("TPBDCollisionConstraints::UpdateLevelsetFindParticles"), STAT_UpdateLevelsetFindParticles, STATGROUP_ChaosWide);
+		DECLARE_CYCLE_STAT(TEXT("TPBDCollisionConstraints::UpdateLevelsetBVHTraversal"), STAT_UpdateLevelsetBVHTraversal, STATGROUP_ChaosWide);
+		DECLARE_CYCLE_STAT(TEXT("TPBDCollisionConstraints::UpdateLevelsetSignedDistance"), STAT_UpdateLevelsetSignedDistance, STATGROUP_ChaosWide);
+		DECLARE_CYCLE_STAT(TEXT("TPBDCollisionConstraints::UpdateLevelsetAll"), STAT_UpdateLevelsetAll, STATGROUP_ChaosWide);
+		DECLARE_CYCLE_STAT(TEXT("TPBDCollisionConstraints::SampleObject"), STAT_SampleObject, STATGROUP_ChaosWide);
 
 		int32 NormalAveraging2 = 1;
 		FAutoConsoleVariableRef CVarNormalAveraging2(TEXT("p.NormalAveraging2"), NormalAveraging2, TEXT(""));
@@ -1006,7 +1006,7 @@ namespace Chaos
 
 		}
 
-		DECLARE_CYCLE_STAT(TEXT("TPBDCollisionConstraint::FindRelevantShapes"), STAT_FindRelevantShapes, STATGROUP_ChaosWide);
+		DECLARE_CYCLE_STAT(TEXT("TPBDCollisionConstraints::FindRelevantShapes"), STAT_FindRelevantShapes, STATGROUP_ChaosWide);
 		template <typename T, int d>
 		TArray<Pair<const FImplicitObject*, TRigidTransform<T, d>>> FindRelevantShapes2(const FImplicitObject* ParticleObj, const TRigidTransform<T, d>& ParticlesTM, const FImplicitObject& LevelsetObj, const TRigidTransform<T, d>& LevelsetTM, const T Thickness)
 		{
@@ -1071,7 +1071,7 @@ namespace Chaos
 			}
 		}
 
-		DECLARE_CYCLE_STAT(TEXT("TPBDCollisionConstraint::UpdateSingleUnionConstraint"), STAT_UpdateSingleUnionConstraint, STATGROUP_ChaosWide);
+		DECLARE_CYCLE_STAT(TEXT("TPBDCollisionConstraints::UpdateSingleUnionConstraint"), STAT_UpdateSingleUnionConstraint, STATGROUP_ChaosWide);
 		template<ECollisionUpdateType UpdateType, typename T, int d>
 		void UpdateSingleUnionConstraint(const T Thickness, TRigidBodyPointContactConstraint<T, d>& Constraint)
 		{
@@ -1097,7 +1097,7 @@ namespace Chaos
 
 
 
-		DECLARE_CYCLE_STAT(TEXT("TPBDCollisionConstraint::UpdateUnionLevelsetConstraint"), STAT_UpdateUnionLevelsetConstraint, STATGROUP_ChaosWide);
+		DECLARE_CYCLE_STAT(TEXT("TPBDCollisionConstraints::UpdateUnionLevelsetConstraint"), STAT_UpdateUnionLevelsetConstraint, STATGROUP_ChaosWide);
 		template<ECollisionUpdateType UpdateType, typename T, int d>
 		void UpdateUnionLevelsetConstraint(const T Thickness, TRigidBodyPointContactConstraint<T, d>& Constraint)
 		{
@@ -1173,7 +1173,7 @@ namespace Chaos
 		}
 
 
-		DECLARE_CYCLE_STAT(TEXT("TPBDCollisionConstraint::UpdateLevelsetUnionConstraint"), STAT_UpdateLevelsetUnionConstraint, STATGROUP_ChaosWide);
+		DECLARE_CYCLE_STAT(TEXT("TPBDCollisionConstraints::UpdateLevelsetUnionConstraint"), STAT_UpdateLevelsetUnionConstraint, STATGROUP_ChaosWide);
 		template<ECollisionUpdateType UpdateType, typename T, int d>
 		void UpdateLevelsetUnionConstraint(const T Thickness, TRigidBodyPointContactConstraint<T, d>& Constraint)
 		{
@@ -1274,7 +1274,7 @@ namespace Chaos
 			}
 		}
 
-		DECLARE_CYCLE_STAT(TEXT("TPBDCollisionConstraint::UpdateLevelsetConstraint"), STAT_UpdateLevelsetConstraint, STATGROUP_ChaosWide);
+		DECLARE_CYCLE_STAT(TEXT("TPBDCollisionConstraints::UpdateLevelsetConstraint"), STAT_UpdateLevelsetConstraint, STATGROUP_ChaosWide);
 		template<ECollisionUpdateType UpdateType, typename T, int d>
 		void UpdateLevelsetConstraint(const T Thickness, TRigidBodyPointContactConstraint<float, 3>& Constraint)
 		{
diff --git a/Engine/Source/Runtime/Experimental/Chaos/Private/Chaos/CollisionResolutionTypes.cpp b/Engine/Source/Runtime/Experimental/Chaos/Private/Chaos/CollisionResolutionTypes.cpp
--- a/Engine/Source/Runtime/Experimental/Chaos/Private/Chaos/CollisionResolutionTypes.cpp
+++ b/Engine/Source/Runtime/Experimental/Chaos/Private/Chaos/CollisionResolutionTypes.cpp
@@ -1,7 +1,7 @@
 // Copyright 1998-2019 Epic Games, Inc. All Rights Reserved.
 
 #include "Chaos/CollisionResolutionTypes.h"
-#include "Chaos/PBDCollisionConstraint.h"
+#include "Chaos/PBDCollisionConstraints.h"
 
 namespace Chaos
 {
diff --git a/Engine/Source/Runtime/Experimental/Chaos/Private/Chaos/PBDCollisionConstraint.cpp b/Engine/Source/Runtime/Experimental/Chaos/Private/Chaos/PBDCollisionConstraint.cpp
deleted file mode 100644
diff --git a/Engine/Source/Runtime/Experimental/Chaos/Private/Chaos/PBDCollisionConstraintPGS.cpp b/Engine/Source/Runtime/Experimental/Chaos/Private/Chaos/PBDCollisionConstraintPGS.cpp
deleted file mode 100644
diff --git a/Engine/Source/Runtime/Experimental/Chaos/Private/Chaos/PBDCollisionConstraintPlaneContactUtil.cpp b/Engine/Source/Runtime/Experimental/Chaos/Private/Chaos/PBDCollisionConstraintPlaneContactUtil.cpp
deleted file mode 100644
diff --git a/Engine/Source/Runtime/Experimental/Chaos/Private/Chaos/PBDCollisionConstraintPointContactUtil.cpp b/Engine/Source/Runtime/Experimental/Chaos/Private/Chaos/PBDCollisionConstraintPointContactUtil.cpp
deleted file mode 100644
diff --git a/Engine/Source/Runtime/Experimental/Chaos/Private/Chaos/PBDCollisionConstraintUtil.cpp b/Engine/Source/Runtime/Experimental/Chaos/Private/Chaos/PBDCollisionConstraintUtil.cpp
--- a/Engine/Source/Runtime/Experimental/Chaos/Private/Chaos/PBDCollisionConstraintUtil.cpp
+++ b/Engine/Source/Runtime/Experimental/Chaos/Private/Chaos/PBDCollisionConstraintUtil.cpp
@@ -1,6 +1,6 @@
 // Copyright 1998-2019 Epic Games, Inc. All Rights Reserved.
 
-#include "Chaos/PBDCollisionConstraintUtil.h"
+#include "Chaos/PBDCollisionConstraintsUtil.h"
 
 #include "Chaos/BoundingVolumeHierarchy.h"
 #include "Chaos/Defines.h"
@@ -15,7 +15,7 @@
 
 namespace Chaos
 {
-	void ComputeHashTable(const TArray<Chaos::TPBDCollisionConstraint<float, 3>::FPointContactConstraint>& ConstraintsArray,
+	void ComputeHashTable(const TArray<Chaos::TPBDCollisionConstraints<float, 3>::FPointContactConstraint>& ConstraintsArray,
 						  const FBox& BoundingBox, TMultiMap<int32, int32>& HashTableMap, const float SpatialHashRadius)
 	{
 		float CellSize = 2.f * SpatialHashRadius;
diff --git a/Engine/Source/Runtime/Experimental/Chaos/Private/Chaos/PBDCollisionConstraints.cpp b/Engine/Source/Runtime/Experimental/Chaos/Private/Chaos/PBDCollisionConstraints.cpp
new file mode 100644
--- /dev/null
+++ b/Engine/Source/Runtime/Experimental/Chaos/Private/Chaos/PBDCollisionConstraints.cpp
@@ -0,0 +1,384 @@
+// Copyright 1998-2019 Epic Games, Inc. All Rights Reserved.
+
+#include "Chaos/PBDCollisionConstraints.h"
+#include "Chaos/PBDCollisionConstraintsImp.h"
+
+#include "Chaos/Capsule.h"
+#include "Chaos/ChaosPerfTest.h"
+#include "Chaos/ChaosDebugDraw.h"
+#include "Chaos/CollisionResolutionAlgo.h"
+#include "Chaos/CollisionResolutionConvexConvex.h"
+#include "Chaos/Defines.h"
+#include "Chaos/GeometryQueries.h"
+#include "Chaos/ImplicitObjectUnion.h"
+#include "Chaos/ImplicitObjectScaled.h"
+#include "Chaos/ISpatialAccelerationCollection.h"
+#include "Chaos/Levelset.h"
+#include "Chaos/Pair.h"
+#include "Chaos/PBDRigidsSOAs.h"
+#include "Chaos/Sphere.h"
+#include "Chaos/Transform.h"
+#include "ChaosLog.h"
+#include "ChaosStats.h"
+#include "Containers/Queue.h"
+#include "ProfilingDebugging/ScopedTimers.h"
+#include "Chaos/PBDCollisionConstraintsPointContactUtil.h"
+
+#if INTEL_ISPC
+#include "PBDCollisionConstraints.ispc.generated.h"
+#endif
+
+//#pragma optimize("", off)
+
+int32 CollisionParticlesBVHDepth = 4;
+FAutoConsoleVariableRef CVarCollisionParticlesBVHDepth(TEXT("p.CollisionParticlesBVHDepth"), CollisionParticlesBVHDepth, TEXT("The maximum depth for collision particles bvh"));
+
+int32 ConstraintBPBVHDepth = 2;
+FAutoConsoleVariableRef CVarConstraintBPBVHDepth(TEXT("p.ConstraintBPBVHDepth"), ConstraintBPBVHDepth, TEXT("The maximum depth for constraint bvh"));
+
+int32 BPTreeOfGrids = 1;
+FAutoConsoleVariableRef CVarBPTreeOfGrids(TEXT("p.BPTreeOfGrids"), BPTreeOfGrids, TEXT("Whether to use a seperate tree of grids for bp"));
+
+float CollisionVelocityInflationCVar = 2.0f;
+FAutoConsoleVariableRef CVarCollisionVelocityInflation(TEXT("p.CollisionVelocityInflation"), CollisionVelocityInflationCVar, TEXT("Collision velocity inflation.[def:2.0]"));
+
+float CollisionFrictionOverride = -1.0f;
+FAutoConsoleVariableRef CVarCollisionFrictionOverride(TEXT("p.CollisionFriction"), CollisionFrictionOverride, TEXT("Collision friction for all contacts if >= 0"));
+
+
+extern int32 UseLevelsetCollision;
+
+#if !UE_BUILD_SHIPPING
+namespace Chaos
+{
+	int32 CHAOS_API PendingHierarchyDump = 0;
+}
+#endif
+
+namespace Chaos
+{
+	CHAOS_API int32 EnableCollisions = 1;
+	FAutoConsoleVariableRef CVarEnableCollisions(TEXT("p.EnableCollisions"), EnableCollisions, TEXT("Enable/Disable collisions on the Chaos solver."));
+
+
+
+	//
+	// Collision Constraint Container
+	//
+
+	template<typename T, int d>
+	TPBDCollisionConstraints<T, d>::TPBDCollisionConstraints(
+		const TPBDRigidsSOAs<T, d>& InParticles,
+		TArrayCollectionArray<bool>& Collided,
+		const TArrayCollectionArray<TSerializablePtr<FChaosPhysicsMaterial>>& InPerParticleMaterials,
+		const int32 InApplyPairIterations /*= 1*/,
+		const int32 InApplyPushOutPairIterations /*= 1*/,
+		const T Thickness /*= (T)0*/)
+		: Particles(InParticles)
+		, MCollided(Collided)
+		, MPhysicsMaterials(InPerParticleMaterials)
+		, MApplyPairIterations(InApplyPairIterations)
+		, MApplyPushOutPairIterations(InApplyPushOutPairIterations)
+		, MThickness(Thickness)
+		, MAngularFriction(0)
+		, bUseCCD(false)
+		, bEnableCollisions(true)
+		, LifespanCounter(0)
+		, CollisionVelocityInflation(CollisionVelocityInflationCVar)
+		, PostComputeCallback(nullptr)
+		, PostApplyCallback(nullptr)
+		, PostApplyPushOutCallback(nullptr)
+		, SpatialAcceleration(nullptr)
+	{
+	}
+
+	template<typename T, int d>
+	void TPBDCollisionConstraints<T, d>::SetPostComputeCallback(const TRigidBodyContactConstraintsPostComputeCallback<T, d>& Callback)
+	{
+		PostComputeCallback = Callback;
+	}
+
+	template<typename T, int d>
+	void TPBDCollisionConstraints<T, d>::ClearPostComputeCallback()
+	{
+		PostComputeCallback = nullptr;
+	}
+
+	template<typename T, int d>
+	void TPBDCollisionConstraints<T, d>::SetPostApplyCallback(const TRigidBodyContactConstraintsPostApplyCallback<T, d>& Callback)
+	{
+		PostApplyCallback = Callback;
+	}
+
+	template<typename T, int d>
+	void TPBDCollisionConstraints<T, d>::ClearPostApplyCallback()
+	{
+		PostApplyCallback = nullptr;
+	}
+
+	template<typename T, int d>
+	void TPBDCollisionConstraints<T, d>::SetPostApplyPushOutCallback(const TRigidBodyContactConstraintsPostApplyPushOutCallback<T, d>& Callback)
+	{
+		PostApplyPushOutCallback = Callback;
+	}
+
+	template<typename T, int d>
+	void TPBDCollisionConstraints<T, d>::ClearPostApplyPushOutCallback()
+	{
+		PostApplyPushOutCallback = nullptr;
+	}
+
+	DEFINE_STAT(STAT_ComputeConstraints);
+	DEFINE_STAT(STAT_ComputeConstraintsNP);
+	DEFINE_STAT(STAT_ComputeConstraintsBP);
+	DEFINE_STAT(STAT_ComputeConstraintsSU);
+
+
+	template<typename T, int d>
+	void TPBDCollisionConstraints<T, d>::UpdatePositionBasedState(/*const TPBDRigidParticles<T, d>& InParticles, const TArray<int32>& InIndices,*/ const T Dt)
+	{
+		Reset();
+
+		if (bEnableCollisions)
+		{
+#if !UE_BUILD_SHIPPING
+			if (PendingHierarchyDump)
+			{
+				ComputeConstraints<true>(*SpatialAcceleration, Dt);
+			}
+			else
+#endif
+			{
+				ComputeConstraints(*SpatialAcceleration, Dt);
+			}
+		}
+	}
+
+	template<typename T, int d>
+	void TPBDCollisionConstraints<T, d>::ConstructConstraints(TGeometryParticleHandle<T, d>* Particle0, TGeometryParticleHandle<T, d>* Particle1, const T Thickness, TRigidBodyPointContactConstraint<T, d> & Constraint)
+	{
+		if (ensure(Particle0 && Particle1))
+		{
+			Collisions::ConstructConstraintsImpl<T, d>(Particle0, Particle1, Particle0->Geometry().Get(), Particle1->Geometry().Get(), Thickness, Constraint);
+		}
+	}
+
+	DECLARE_CYCLE_STAT(TEXT("TPBDCollisionConstraints::Reset"), STAT_CollisionConstraintsReset, STATGROUP_Chaos);
+	template<typename T, int d>
+	void TPBDCollisionConstraints<T, d>::Reset()
+	{
+		SCOPE_CYCLE_COUNTER(STAT_CollisionConstraintsReset);
+
+		TArray<FConstraintContainerHandle*> CopyOfHandles = Handles;
+
+		for (FConstraintContainerHandle* ContactHandle : CopyOfHandles)
+		{
+			//if (!bEnableCollisions)
+			{
+				RemoveConstraint(ContactHandle);
+			}
+		}
+
+
+		MAngularFriction = 0;
+		bUseCCD = false;
+	}
+
+	template<typename T, int d>
+	void TPBDCollisionConstraints<T, d>::ApplyCollisionModifier(const TFunction<ECollisionModifierResult(const FConstraintContainerHandle* Handle)>& CollisionModifier)
+	{
+		TArray<FConstraintContainerHandle*> CopyOfHandles = Handles;
+
+		for (FConstraintContainerHandle* ContactHandle : CopyOfHandles)
+		{
+			ECollisionModifierResult Result = CollisionModifier(ContactHandle);
+			if (Result == ECollisionModifierResult::Disabled)
+			{
+				RemoveConstraint(ContactHandle);
+			}
+		}
+	}
+
+	template<typename T, int d>
+	void TPBDCollisionConstraints<T, d>::RemoveConstraints(const TSet<TGeometryParticleHandle<T, d>*>&  InHandleSet)
+	{
+		const TArray<TGeometryParticleHandle<T, d>*> HandleArray = InHandleSet.Array();
+		for (auto ParticleHandle : HandleArray)
+		{
+			TArray<FConstraintContainerHandle*> CopyOfHandles = Handles;
+
+			for (FConstraintContainerHandle* ContactHandle : CopyOfHandles)
+			{
+				TVector<TGeometryParticleHandle<T, d>*, 2> ConstraintParticles = ContactHandle->GetConstrainedParticles();
+				if (ConstraintParticles[1] == ParticleHandle || ConstraintParticles[0] == ParticleHandle)
+				{
+					RemoveConstraint(ContactHandle);
+				}
+			}
+		}
+	}
+
+	template<typename T, int d>
+	void TPBDCollisionConstraints<T, d>::RemoveConstraint(FConstraintContainerHandle* Handle)
+	{
+		int32 Idx = Handle->GetConstraintIndex();
+		typename TCollisionConstraintBase<T, d>::FType ConstraintType = Handle->GetType();
+
+		Handles.RemoveAtSwap(Idx);
+		PointConstraints.RemoveAtSwap(Idx);
+		if (Idx < PointConstraints.Num())
+		{
+			Handles[Idx]->SetConstraintIndex(Idx, ConstraintType);
+		}
+
+		ensure(Handles.Num() == PointConstraints.Num());
+		HandleAllocator.FreeHandle(Handle);
+	}
+
+
+
+	template<typename T, int d>
+	template <bool bGatherStats>
+	void TPBDCollisionConstraints<T, d>::ComputeConstraints(const FAccelerationStructure& AccelerationStructure, T Dt)
+	{
+		if (const auto AABBTree = AccelerationStructure.template As<TAABBTree<TAccelerationStructureHandle<T, d>, TAABBTreeLeafArray<TAccelerationStructureHandle<T, d>, T>, T>>())
+		{
+			ComputeConstraintsHelperLowLevel<bGatherStats>(*AABBTree, Dt);
+		}
+		else if (const auto BV = AccelerationStructure.template As<TBoundingVolume<TAccelerationStructureHandle<T, d>, T, d>>())
+		{
+			ComputeConstraintsHelperLowLevel<bGatherStats>(*BV, Dt);
+		}
+		else if (const auto AABBTreeBV = AccelerationStructure.template As<TAABBTree<TAccelerationStructureHandle<T, d>, TBoundingVolume<TAccelerationStructureHandle<T, d>, T, d>, T>>())
+		{
+			ComputeConstraintsHelperLowLevel<bGatherStats>(*AABBTreeBV, Dt);
+		}
+		else if (const auto Collection = AccelerationStructure.template As<ISpatialAccelerationCollection<TAccelerationStructureHandle<T, d>, T, d>>())
+		{
+			if (bGatherStats)
+			{
+				Collection->PBDComputeConstraintsLowLevel_GatherStats(*this, Dt);
+			}
+			else
+			{
+				Collection->PBDComputeConstraintsLowLevel(*this, Dt);
+			}
+		}
+		else
+		{
+			check(false);  //question: do we want to support a dynamic dispatch version?
+		}
+
+		if (PostComputeCallback != nullptr)
+		{
+			PostComputeCallback();
+		}
+	}
+
+
+	template<typename T, int d>
+	void TPBDCollisionConstraints<T, d>::UpdateConstraints(T Dt, const TSet<TGeometryParticleHandle<T, d>*>& ParticlesSet)
+	{
+		// Clustering uses update constraints to force a re-evaluation. 
+	}
+
+	DECLARE_CYCLE_STAT(TEXT("TPBDCollisionConstraints::UpdateConstraint"), STAT_UpdateConstraint, STATGROUP_ChaosWide);
+	template<typename T, int d>
+	template<ECollisionUpdateType UpdateType>
+	void TPBDCollisionConstraints<T, d>::UpdateConstraint(const T Thickness, TRigidBodyPointContactConstraint<T, d> & Constraint)
+	{
+		SCOPE_CYCLE_COUNTER(STAT_UpdateConstraint);
+		//if (ConstraintHandle->GetType() == FPointContactConstraint::StaticType())
+		{
+			Collisions::UpdateConstraint<UpdateType>(Thickness, Constraint);
+		}
+	}
+
+
+
+
+	DECLARE_CYCLE_STAT(TEXT("TPBDCollisionConstraints::Apply"), STAT_Apply, STATGROUP_Chaos);
+	template<typename T, int d>
+	void TPBDCollisionConstraints<T, d>::Apply(const T Dt, const TArray<FConstraintContainerHandle*>& InConstraintHandles, const int32 Iterations, const int32 NumIterations)
+	{
+		SCOPE_CYCLE_COUNTER(STAT_Apply);
+		if (MApplyPairIterations > 0)
+		{
+			PhysicsParallelFor(InConstraintHandles.Num(), [&](int32 ConstraintHandleIndex) {
+				FConstraintContainerHandle* ConstraintHandle = InConstraintHandles[ConstraintHandleIndex];
+				check(ConstraintHandle != nullptr);
+
+				if (ConstraintHandle->GetType() == FPointContactConstraint::StaticType())
+				{
+					Collisions::TPointContactParticleParameters<T> ParticleParameters = { &MCollided, &MPhysicsMaterials, CollisionFrictionOverride, MAngularFriction };
+					Collisions::TPointContactIterationParameters<T> IterationParameters = { Dt, Iterations, NumIterations, MApplyPairIterations, nullptr };
+					Collisions::Apply(ConstraintHandle->GetPointContact(), MThickness, IterationParameters, ParticleParameters);
+				}
+				else if (ConstraintHandle->GetType() == FPlaneContactConstraint::StaticType())
+				{
+
+				}
+				else
+				{
+					ensureMsgf(false, TEXT("Invalid constraint type"));
+				}
+
+			}, bDisableCollisionParallelFor);
+		}
+
+		if (PostApplyCallback != nullptr)
+		{
+			PostApplyCallback(Dt, InConstraintHandles);
+		}
+	}
+
+
+	DECLARE_CYCLE_STAT(TEXT("TPBDCollisionConstraints::ApplyPushOut"), STAT_ApplyPushOut, STATGROUP_Chaos);
+	template<typename T, int d>
+	bool TPBDCollisionConstraints<T, d>::ApplyPushOut(const T Dt, const TArray<FConstraintContainerHandle*>& InConstraintHandles, const TSet< const TGeometryParticleHandle<T, d>*>& IsTemporarilyStatic, int32 Iteration, int32 NumIterations)
+	{
+		SCOPE_CYCLE_COUNTER(STAT_ApplyPushOut);
+
+		bool NeedsAnotherIteration = false;
+		if (MApplyPushOutPairIterations > 0)
+		{
+			PhysicsParallelFor(InConstraintHandles.Num(), [&](int32 ConstraintHandleIndex)
+			{
+				FConstraintContainerHandle* ConstraintHandle = InConstraintHandles[ConstraintHandleIndex];
+				check(ConstraintHandle != nullptr);
+
+				if (ConstraintHandle->GetType() == FPointContactConstraint::StaticType())
+				{
+					Collisions::TPointContactParticleParameters<T> ParticleParameters = { &MCollided, &MPhysicsMaterials, CollisionFrictionOverride, MAngularFriction };
+					Collisions::TPointContactIterationParameters<T> IterationParameters = { Dt, Iteration, NumIterations, MApplyPairIterations, &NeedsAnotherIteration };
+					Collisions::ApplyPushOut(ConstraintHandle->GetPointContact(), MThickness, IsTemporarilyStatic, IterationParameters, ParticleParameters);
+				}
+				else if (ConstraintHandle->GetType() == FPlaneContactConstraint::StaticType())
+				{
+
+				}
+				else
+				{
+					ensureMsgf(false, TEXT("Invalid constraint type"));
+				}
+
+
+			}, bDisableCollisionParallelFor);
+		}
+
+		if (PostApplyPushOutCallback != nullptr)
+		{
+			PostApplyPushOutCallback(Dt, InConstraintHandles, NeedsAnotherIteration);
+		}
+		return NeedsAnotherIteration;
+	}
+
+
+
+	template class TAccelerationStructureHandle<float, 3>;
+	template class CHAOS_API TPBDCollisionConstraints<float, 3>;
+	template void TPBDCollisionConstraints<float, 3>::UpdateConstraint<ECollisionUpdateType::Any>(const float Thickness, TRigidBodyPointContactConstraint<float,3> & Constraint);
+	template void TPBDCollisionConstraints<float, 3>::UpdateConstraint<ECollisionUpdateType::Deepest>(const float Thickness, TRigidBodyPointContactConstraint<float,3> & Constraint);
+	template void TPBDCollisionConstraints<float, 3>::ComputeConstraints<false>(const TPBDCollisionConstraints<float, 3>::FAccelerationStructure&, float Dt);
+	template void TPBDCollisionConstraints<float, 3>::ComputeConstraints<true>(const TPBDCollisionConstraints<float, 3>::FAccelerationStructure&, float Dt);
+}
diff --git a/Engine/Source/Runtime/Experimental/Chaos/Private/Chaos/PBDCollisionConstraintsPGS.cpp b/Engine/Source/Runtime/Experimental/Chaos/Private/Chaos/PBDCollisionConstraintsPGS.cpp
new file mode 100644
--- /dev/null
+++ b/Engine/Source/Runtime/Experimental/Chaos/Private/Chaos/PBDCollisionConstraintsPGS.cpp
@@ -0,0 +1,1333 @@
+// Copyright 1998-2019 Epic Games, Inc. All Rights Reserved.
+
+#include "Chaos/PBDCollisionConstraintsPGS.h"
+
+#include "Chaos/BoundingVolume.h"
+#include "Chaos/BoundingVolumeHierarchy.h"
+#include "Chaos/Defines.h"
+#include "Chaos/Pair.h"
+#include "Chaos/Sphere.h"
+#include "Chaos/Transform.h"
+
+#include "ProfilingDebugging/ScopedTimers.h"
+
+#if CHAOS_PARTICLEHANDLE_TODO
+
+using namespace Chaos;
+
+template<class T_PARTICLES, class T, int d>
+TVector<T, d> GetTranslationPGS(const T_PARTICLES& InParticles)
+{
+	check(false);
+	return TRigidTransform<T, d>();
+}
+
+template<class T, int d>
+TVector<T, d> GetTranslationPGS(const TRigidParticles<T, d>& InParticles, const int32 Index)
+{
+	return InParticles.X(Index);
+}
+
+template<class T, int d>
+TVector<T, d> GetTranslationPGS(const TPBDRigidParticles<T, d>& InParticles, const int32 Index)
+{
+	return InParticles.P(Index);
+}
+
+template<class T_PARTICLES, class T, int d>
+TRotation<T, d> GetRotationPGS(const T_PARTICLES& InParticles)
+{
+	check(false);
+	return TRotation<T, d>();
+}
+
+template<class T, int d>
+TRotation<T, d> GetRotationPGS(const TRigidParticles<T, d>& InParticles, const int32 Index)
+{
+	return InParticles.R(Index);
+}
+
+template<class T, int d>
+TRotation<T, d> GetRotationPGS(const TPBDRigidParticles<T, d>& InParticles, const int32 Index)
+{
+	return InParticles.Q(Index);
+}
+
+template<class T_PARTICLES, class T, int d>
+TRigidTransform<T, d> GetTransformPGS(const T_PARTICLES& InParticles)
+{
+	check(false);
+	return TRigidTransform<T, d>();
+}
+
+template<class T, int d>
+TRigidTransform<T, d> GetTransformPGS(const TRigidParticles<T, d>& InParticles, const int32 Index)
+{
+	return TRigidTransform<T, d>(InParticles.X(Index), InParticles.R(Index));
+}
+
+template<class T, int d>
+TRigidTransform<T, d> GetTransformPGS(const TPBDRigidParticles<T, d>& InParticles, const int32 Index)
+{
+	return TRigidTransform<T, d>(InParticles.P(Index), InParticles.Q(Index));
+}
+
+template<class T, int d>
+TPBDCollisionConstraintPGS<T, d>::TPBDCollisionConstraintPGS(TPBDRigidParticles<T, d>& InParticles, const TArray<int32>& InIndices, TArrayCollectionArray<bool>& Collided, TArrayCollectionArray<TSerializablePtr<FChaosPhysicsMaterial>>& InPhysicsMaterials, const T Thickness /*= (T)0*/)
+    : MCollided(Collided), MPhysicsMaterials(InPhysicsMaterials), MThickness(Thickness), Tolerance(KINDA_SMALL_NUMBER), MaxIterations(10), bUseCCD(false)
+{
+}
+
+template<class T, int d>
+void TPBDCollisionConstraintPGS<T, d>::ComputeConstraints(const TPBDRigidParticles<T, d>& InParticles, const TArray<int32>& InIndices, const T Dt)
+{
+	double Time = 0;
+	FDurationTimer Timer(Time);
+	// Broad phase
+	//TBoundingVolumeHierarchy<TPBDRigidParticles<T, d>, T, d> Hierarchy(InParticles);
+	TBoundingVolume<TPBDRigidParticles<T, d>, T, d> Hierarchy(InParticles, true, Dt);
+	Timer.Stop();
+	UE_LOG(LogChaos, Verbose, TEXT("\tPBDCollisionConstraint Construct Hierarchy %f"), Time);
+	// Narrow phase
+	FCriticalSection CriticalSection;
+	Time = 0;
+	Timer.Start();
+	PhysicsParallelFor(InIndices.Num(), [&](int32 Index) {
+		const int32 ParticleIndex = InIndices[Index];
+
+		TArray<int32> PotentialIntersections;
+		TBox<T, d> Box1;
+		if (InParticles.Geometry(ParticleIndex)->HasBoundingBox())
+		{
+			Box1 = Hierarchy.GetWorldSpaceBoundingBox(InParticles, ParticleIndex);
+			PotentialIntersections = Hierarchy.FindAllIntersections(Box1);
+		}
+		else
+		{
+			PotentialIntersections = Hierarchy.GlobalObjects();
+		}
+
+		for (int32 i = 0; i < PotentialIntersections.Num(); ++i)
+		{
+			int32 Body2Index = PotentialIntersections[i];
+
+			// Collision group culling...
+			// CollisionGroup == 0 : Collide_With_Everything
+			// CollisionGroup == INDEX_NONE : Disabled collisions
+			// CollisionGroup_A != CollisionGroup_B : Skip Check
+			if(InParticles.Disabled(Body2Index))
+			{
+				continue;
+			}
+				
+			if (InParticles.CollisionGroup(ParticleIndex) == INDEX_NONE || InParticles.CollisionGroup(Body2Index) == INDEX_NONE)
+			{
+				continue;
+			}
+			if (InParticles.CollisionGroup(ParticleIndex) && InParticles.CollisionGroup(Body2Index) && InParticles.CollisionGroup(ParticleIndex) != InParticles.CollisionGroup(Body2Index))
+			{
+				continue;
+			}
+
+			if (InParticles.InvM(ParticleIndex) < FLT_MIN && InParticles.InvM(Body2Index) < FLT_MIN)
+			{
+				continue;
+			}
+			if (ParticleIndex == Body2Index || ((InParticles.Geometry(ParticleIndex)->HasBoundingBox() == InParticles.Geometry(Body2Index)->HasBoundingBox()) && Body2Index > ParticleIndex))
+			{
+				continue;
+			}
+			if (InParticles.Geometry(ParticleIndex)->HasBoundingBox() && InParticles.Geometry(Body2Index)->HasBoundingBox())
+			{
+				const auto& Box2 = Hierarchy.GetWorldSpaceBoundingBox(InParticles, Body2Index);
+				if (!Box1.Intersects(Box2))
+				{
+					continue;
+				}
+			}
+			auto Constraint = ComputeConstraint(InParticles, ParticleIndex, Body2Index, MThickness);
+			CriticalSection.Lock();
+			Constraints.Add(Constraint);
+			CriticalSection.Unlock();
+		}
+	});
+
+	Timer.Stop();
+
+	UE_LOG(LogChaos, Verbose, TEXT("\tPBDCollisionConstraint Construct %d Constraints with Potential Collisions %f"), Constraints.Num(), Time);
+}
+
+template<class T, int d>
+void TPBDCollisionConstraintPGS<T, d>::RemoveConstraints(const TSet<uint32>& RemovedParticles)
+{
+	for (int32 i = 0; i < Constraints.Num(); ++i)
+	{
+		const auto& Constraint = Constraints[i];
+		if (RemovedParticles.Contains(Constraint.ParticleIndex) || RemovedParticles.Contains(Constraint.LevelsetIndex))
+		{
+			Constraints.RemoveAtSwap(i);
+			i--;
+		}
+	}
+}
+
+template<class T, int d>
+void TPBDCollisionConstraintPGS<T, d>::UpdateConstraints(const TPBDRigidParticles<T, d>& InParticles, T Dt, const TSet<uint32>& AddedParticles, const TArray<uint32>& ActiveParticles)
+{
+	double Time = 0;
+	FDurationTimer Timer(Time);
+
+	//
+	// Broad phase
+	//
+
+	// @todo(mlentine): We only need to construct the hierarchy for the islands we care about
+	TBoundingVolume<TPBDRigidParticles<T, d>, T, d> Hierarchy(InParticles, ActiveParticles, true, Dt);
+	Timer.Stop();
+	UE_LOG(LogChaos, Verbose, TEXT("\tPBDCollisionConstraint Construct Hierarchy %f"), Time);
+
+	//
+	// Narrow phase
+	//
+
+	FCriticalSection CriticalSection;
+	Time = 0;
+	Timer.Start();
+	TArray<uint32> AddedParticlesArray = AddedParticles.Array();
+	PhysicsParallelFor(AddedParticlesArray.Num(), [&](int32 Index) {
+		int32 Body1Index = AddedParticlesArray[Index];
+		if (InParticles.Disabled(Body1Index))
+			return;
+		TArray<int32> PotentialIntersections;
+		TBox<T, d> Box1;
+		if (InParticles.Geometry(Body1Index)->HasBoundingBox())
+		{
+			Box1 = Hierarchy.GetWorldSpaceBoundingBox(InParticles, Body1Index);
+			PotentialIntersections = Hierarchy.FindAllIntersections(Box1);
+		}
+		else
+		{
+			PotentialIntersections = Hierarchy.GlobalObjects();
+		}
+		for (int32 i = 0; i < PotentialIntersections.Num(); ++i)
+		{
+			int32 Body2Index = PotentialIntersections[i];
+
+			if (Body1Index == Body2Index || ((InParticles.Geometry(Body1Index)->HasBoundingBox() == InParticles.Geometry(Body2Index)->HasBoundingBox()) && AddedParticles.Contains(Body2Index) && AddedParticles.Contains(Body1Index) && Body2Index > Body1Index))
+			{
+				continue;
+			}
+			if (InParticles.Geometry(Body1Index)->HasBoundingBox() && InParticles.Geometry(Body2Index)->HasBoundingBox())
+			{
+				const auto& Box2 = Hierarchy.GetWorldSpaceBoundingBox(InParticles, Body2Index);
+				if (!Box1.Intersects(Box2))
+				{
+					continue;
+				}
+			}
+			
+			if (InParticles.InvM(Body1Index) && InParticles.InvM(Body2Index) && (InParticles.Island(Body1Index) != InParticles.Island(Body2Index)))	//todo(ocohen): this is a hack - we should not even consider dynamics from other islands
+			{
+				continue;
+			}
+			auto Constraint = ComputeConstraint(InParticles, Body1Index, Body2Index, MThickness);
+			CriticalSection.Lock();
+			Constraints.Add(Constraint);
+			CriticalSection.Unlock();
+		}
+	});
+	Timer.Stop();
+	UE_LOG(LogChaos, Verbose, TEXT("\tPBDCollisionConstraint Update %d Constraints with Potential Collisions %f"), Constraints.Num(), Time);
+}
+
+template<class T_PARTICLES, class T, int d>
+void ComputePGSPropeties(const T_PARTICLES& InParticles, const TRigidBodyContactConstraintPGS<T, d>& Constraint, const int32 PointIndex, const int32 FlattenedIndex,
+	const PMatrix<T, d, d>& WorldSpaceInvI1, const PMatrix<T, d, d>& WorldSpaceInvI2, const TVector<T, d> Direction,
+	TArray<TVector<TVector<T, d>, 2>>& Angulars, TArray<TVector<TVector<T, d>, 2>>& MassWeightedAngulars, TArray<T>& Multipliers)
+{
+	TVector<T, d> VectorToPoint1 = Constraint.Location[PointIndex] - GetTranslationPGS(InParticles, Constraint.ParticleIndex);
+	TVector<T, d> VectorToPoint2 = Constraint.Location[PointIndex] - GetTranslationPGS(InParticles, Constraint.LevelsetIndex);
+	Angulars[FlattenedIndex][0] = -TVector<T, d>::CrossProduct(VectorToPoint1, Direction);
+	Angulars[FlattenedIndex][1] = TVector<T, d>::CrossProduct(VectorToPoint2, Direction);
+	MassWeightedAngulars[FlattenedIndex][0] = WorldSpaceInvI1 * Angulars[FlattenedIndex][0];
+	MassWeightedAngulars[FlattenedIndex][1] = WorldSpaceInvI2 * Angulars[FlattenedIndex][1];
+	if (InParticles.InvM(Constraint.ParticleIndex))
+	{
+		Multipliers[FlattenedIndex] += InParticles.InvM(Constraint.ParticleIndex) + TVector<T, d>::DotProduct(Angulars[FlattenedIndex][0], MassWeightedAngulars[FlattenedIndex][0]);
+	}
+	if (InParticles.InvM(Constraint.LevelsetIndex))
+	{
+		Multipliers[FlattenedIndex] += InParticles.InvM(Constraint.LevelsetIndex) + TVector<T, d>::DotProduct(Angulars[FlattenedIndex][1], MassWeightedAngulars[FlattenedIndex][1]);
+	}
+}
+
+template<class T, int d>
+template<class T_PARTICLES>
+void TPBDCollisionConstraintPGS<T, d>::Solve(T_PARTICLES& InParticles, const T Dt, const TArray<int32>& InConstraintIndices)
+{
+	TArray<T> Normals;
+	TArray<TVector<T, d - 1>> Tangents;
+	TArray<T> Multipliers;
+	TArray<TVector<TVector<T, d>, 2>> Angulars, MassWeightedAngulars;
+	TArray<TVector<TVector<T, d>, d - 1>> ConstraintTangents;
+	TVector<TArray<T>, d - 1> TangentMultipliers;
+	TVector<TArray<TVector<TVector<T, d>, 2>>, d - 1> TangentAngulars, TangentMassWeightedAngulars;
+
+	int32 NumConstraints = 0;
+	for (int32 ConstraintIndex : InConstraintIndices)
+	{
+		const FRigidBodyContactConstraint& Constraint = Constraints[ConstraintIndex];
+		NumConstraints += Constraint.Phi.Num();
+	}
+
+	Normals.SetNumZeroed(NumConstraints);
+	Tangents.SetNumZeroed(NumConstraints);
+	Multipliers.SetNumZeroed(NumConstraints);
+	Angulars.SetNum(NumConstraints);
+	MassWeightedAngulars.SetNum(NumConstraints);
+	ConstraintTangents.SetNum(NumConstraints);
+	for (int32 Dimension = 0; Dimension < (d - 1); ++Dimension)
+	{
+		TangentMultipliers[Dimension].SetNumZeroed(NumConstraints);
+		TangentAngulars[Dimension].SetNum(NumConstraints);
+		TangentMassWeightedAngulars[Dimension].SetNum(NumConstraints);
+	}
+
+	int32 FlattenedIndex = 0;
+	for (int32 ConstraintIndex : InConstraintIndices)
+	{
+		FRigidBodyContactConstraint& Constraint = Constraints[ConstraintIndex];
+		PMatrix<T, d, d> WorldSpaceInvI1 = (GetRotationPGS(InParticles, Constraint.ParticleIndex) * FMatrix::Identity).GetTransposed() * InParticles.InvI(Constraint.ParticleIndex) * (GetRotationPGS(InParticles, Constraint.ParticleIndex) * FMatrix::Identity);
+		PMatrix<T, d, d> WorldSpaceInvI2 = (GetRotationPGS(InParticles, Constraint.LevelsetIndex) * FMatrix::Identity).GetTransposed() * InParticles.InvI(Constraint.LevelsetIndex) * (GetRotationPGS(InParticles, Constraint.LevelsetIndex) * FMatrix::Identity);
+		for (int32 PointIndex = 0; PointIndex < Constraint.Phi.Num(); ++PointIndex)
+		{
+			ComputePGSPropeties(InParticles, Constraint, PointIndex, FlattenedIndex, WorldSpaceInvI1, WorldSpaceInvI2, -Constraint.Normal[PointIndex], Angulars, MassWeightedAngulars, Multipliers);
+			// Tangents
+			{
+				T AbsX = FMath::Abs(Constraint.Normal[PointIndex][0]);
+				T AbsY = FMath::Abs(Constraint.Normal[PointIndex][1]);
+				T AbsZ = FMath::Abs(Constraint.Normal[PointIndex][2]);
+				if (AbsX < AbsY)
+				{
+					if (AbsX < AbsZ)
+					{
+						ConstraintTangents[FlattenedIndex][0] = TVector<T, d>(0, Constraint.Normal[PointIndex][2], -Constraint.Normal[PointIndex][1]);
+					}
+					else
+					{
+						ConstraintTangents[FlattenedIndex][0] = TVector<T, d>(Constraint.Normal[PointIndex][1], -Constraint.Normal[PointIndex][0], 0);
+					}
+				}
+				else
+				{
+					if (AbsY < AbsZ)
+					{
+						ConstraintTangents[FlattenedIndex][0] = TVector<T, d>(-Constraint.Normal[PointIndex][2], 0, Constraint.Normal[PointIndex][0]);
+					}
+					else
+					{
+						ConstraintTangents[FlattenedIndex][0] = TVector<T, d>(Constraint.Normal[PointIndex][1], -Constraint.Normal[PointIndex][0], 0);
+					}
+				}
+			}
+			ConstraintTangents[FlattenedIndex][0] = ConstraintTangents[FlattenedIndex][0].GetSafeNormal();
+			ConstraintTangents[FlattenedIndex][1] = TVector<T, d>::CrossProduct(-ConstraintTangents[FlattenedIndex][0], Constraint.Normal[PointIndex]);
+			ComputePGSPropeties(InParticles, Constraint, PointIndex, FlattenedIndex, WorldSpaceInvI1, WorldSpaceInvI2, -ConstraintTangents[FlattenedIndex][0], TangentAngulars[0], TangentMassWeightedAngulars[0], TangentMultipliers[0]);
+			ComputePGSPropeties(InParticles, Constraint, PointIndex, FlattenedIndex, WorldSpaceInvI1, WorldSpaceInvI2, -ConstraintTangents[FlattenedIndex][1], TangentAngulars[1], TangentMassWeightedAngulars[1], TangentMultipliers[1]);
+			FlattenedIndex++;
+		}
+	}
+
+	T Residual = 0;
+	for (int32 Iteration = 0; Iteration < MaxIterations; ++Iteration)
+	{
+		Residual = 0;
+		FlattenedIndex = 0;
+		for (int32 ConstraintIndex : InConstraintIndices)
+		{
+			FRigidBodyContactConstraint& Constraint = Constraints[ConstraintIndex];
+			for (int32 PointIndex = 0; PointIndex < Constraint.Phi.Num(); ++PointIndex)
+			{
+				T Body1NormalVelocity = TVector<T, d>::DotProduct(InParticles.V(Constraint.ParticleIndex), Constraint.Normal[PointIndex]) +
+					TVector<T, d>::DotProduct(InParticles.W(Constraint.ParticleIndex), Angulars[FlattenedIndex][0]);
+				T Body2NormalVelocity = TVector<T, d>::DotProduct(InParticles.V(Constraint.LevelsetIndex), -Constraint.Normal[PointIndex]) +
+					TVector<T, d>::DotProduct(InParticles.W(Constraint.LevelsetIndex), Angulars[FlattenedIndex][1]);
+				T RelativeNormalVelocity = Body1NormalVelocity + Body2NormalVelocity + Constraint.Phi[PointIndex] / Dt;
+				T NewResidual = FMath::Max(-RelativeNormalVelocity, RelativeNormalVelocity * Normals[FlattenedIndex]);
+				if (NewResidual > Residual)
+				{
+					Residual = NewResidual;
+				}
+				T NormalDelta = -RelativeNormalVelocity / Multipliers[FlattenedIndex];
+				// Update Normals
+				T NewNormal = Normals[FlattenedIndex] + NormalDelta;
+				if (NewNormal < 0)
+				{
+					NewNormal = 0;
+					NormalDelta = -Normals[FlattenedIndex];
+				}
+				check(RelativeNormalVelocity < 0 || NormalDelta == 0 || Iteration > 0);
+				// Velocity update
+				InParticles.V(Constraint.ParticleIndex) += NormalDelta * InParticles.InvM(Constraint.ParticleIndex) * Constraint.Normal[PointIndex];
+				InParticles.V(Constraint.LevelsetIndex) += NormalDelta * InParticles.InvM(Constraint.LevelsetIndex) * -Constraint.Normal[PointIndex];
+				InParticles.W(Constraint.ParticleIndex) += NormalDelta * MassWeightedAngulars[FlattenedIndex][0];
+				InParticles.W(Constraint.LevelsetIndex) += NormalDelta * MassWeightedAngulars[FlattenedIndex][1];
+				// Normal update
+				Normals[FlattenedIndex] = NewNormal;
+				T Friction = MPhysicsMaterials[Constraint.LevelsetIndex]
+					? FMath::Max(MPhysicsMaterials[Constraint.ParticleIndex]->Friction, MPhysicsMaterials[Constraint.LevelsetIndex]->Friction)
+					: MPhysicsMaterials[Constraint.ParticleIndex]->Friction;
+				if (Friction)
+				{
+					for (int32 Dimension = 0; Dimension < (d - 1); Dimension++)
+					{
+						T Body1TangentVelocity = TVector<T, d>::DotProduct(InParticles.V(Constraint.ParticleIndex), ConstraintTangents[PointIndex][Dimension]) +
+							TVector<T, d>::DotProduct(InParticles.W(Constraint.ParticleIndex), TangentAngulars[Dimension][FlattenedIndex][0]);
+						T Body2TangentVelocity = TVector<T, d>::DotProduct(InParticles.V(Constraint.LevelsetIndex), -ConstraintTangents[PointIndex][Dimension]) +
+							TVector<T, d>::DotProduct(InParticles.W(Constraint.LevelsetIndex), TangentAngulars[Dimension][FlattenedIndex][1]);
+						T RelativeTangentVelocity = Body1TangentVelocity + Body2TangentVelocity;
+						T TangentDelta = -RelativeTangentVelocity / TangentMultipliers[Dimension][FlattenedIndex];
+						T NewTangent = Tangents[FlattenedIndex][Dimension] + TangentDelta;
+						if (FMath::Abs(NewTangent) > Friction * NewNormal)
+						{
+							NewTangent = Friction * NewNormal;
+							if (NewTangent < 0)
+							{
+								NewTangent *= -1;
+							}
+						}
+						// Velocity update
+						InParticles.V(Constraint.ParticleIndex) += TangentDelta * InParticles.InvM(Constraint.ParticleIndex) * ConstraintTangents[PointIndex][Dimension];
+						InParticles.V(Constraint.LevelsetIndex) += TangentDelta * InParticles.InvM(Constraint.LevelsetIndex) * -ConstraintTangents[PointIndex][Dimension];
+						InParticles.W(Constraint.ParticleIndex) += TangentDelta * TangentMassWeightedAngulars[Dimension][FlattenedIndex][0];
+						InParticles.W(Constraint.LevelsetIndex) += TangentDelta * TangentMassWeightedAngulars[Dimension][FlattenedIndex][1];
+						Tangents[FlattenedIndex][Dimension] = NewTangent;
+					}
+				}
+				FlattenedIndex++;
+			}
+		}
+		UE_LOG(LogChaos, Verbose, TEXT("\tPBDCollisionConstraint Solve with Residual %f"), Residual);
+		if (Residual < Tolerance)
+		{
+			break;
+		}
+	}
+}
+
+template<class T, int d>
+void TPBDCollisionConstraintPGS<T, d>::PrintParticles(const TPBDRigidParticles<T, d>& InParticles, const TArray<int32>& InConstraintIndices)
+{
+	TSet<int32> ConstraintParticles;
+	for (int32 ConstraintIndex : InConstraintIndices)
+	{
+		const FRigidBodyContactConstraint& Constraint = Constraints[ConstraintIndex];
+		if (!ConstraintParticles.Contains(Constraint.ParticleIndex))
+		{
+			ConstraintParticles.Add(Constraint.ParticleIndex);
+		}
+		if (!ConstraintParticles.Contains(Constraint.LevelsetIndex))
+		{
+			ConstraintParticles.Add(Constraint.LevelsetIndex);
+		}
+	}
+	for (const auto& i : ConstraintParticles)
+	{
+		UE_LOG(LogChaos, Verbose, TEXT("Particle %d has X=(%f, %f, %f) and V=(%f, %f, %f)"), i, InParticles.X(i)[0], InParticles.X(i)[1], InParticles.X(i)[2], InParticles.V(i)[0], InParticles.V(i)[1], InParticles.V(i)[2]);
+	}
+}
+
+template<class T, int d>
+void TPBDCollisionConstraintPGS<T, d>::PrintConstraints(const TPBDRigidParticles<T, d>& InParticles, const TArray<int32>& InConstraintIndices)
+{
+	for (int32 ConstraintIndex : InConstraintIndices)
+	{
+		const FRigidBodyContactConstraint& Constraint = Constraints[ConstraintIndex];
+		UE_LOG(LogChaos, Verbose, TEXT("Constraint between %d and %d has %d contacts"), Constraint.ParticleIndex, Constraint.LevelsetIndex, Constraint.Phi.Num());
+		for (int32 j = 0; j < Constraint.Phi.Num(); ++j)
+		{
+			UE_LOG(LogChaos, Verbose, TEXT("Constraint has loction (%f, %f, %f) and phi %f"), Constraint.Location[j][0], Constraint.Location[j][1], Constraint.Location[j][2], Constraint.Phi[j]);
+		}
+	}
+}
+
+template<class T, int d>
+void FindPointsOnHull(TRigidBodyContactConstraintPGS<T, d>& Constraint,
+	const TVector<T, d>& X0, const TVector<T, d>& X1, const TVector<T, d>& X2,
+	const TPlane<T, d>& SplitPlane, const TArray<int32>& Indices, TSet<int32>& HullPoints)
+{
+	int32 MaxD = 0; //This doesn't need to be initialized but we need to avoid the compiler warning
+	T MaxDistance = 0;
+	for (int32 i = 0; i < Indices.Num(); ++i)
+	{
+		T Distance = SplitPlane.SignedDistance(Constraint.Location[Indices[i]]);
+		check(Distance >= 0);
+		if (Distance > MaxDistance)
+		{
+			MaxDistance = Distance;
+			MaxD = Indices[i];
+		}
+	}
+	if (MaxDistance > 0)
+	{
+		if (!HullPoints.Contains(MaxD))
+		{
+			HullPoints.Add(MaxD);
+		}
+		const TVector<T, d>& NewX = Constraint.Location[MaxD];
+		const TVector<T, d> V1 = (NewX - X0).GetSafeNormal();
+		const TVector<T, d> V2 = (NewX - X1).GetSafeNormal();
+		const TVector<T, d> V3 = (NewX - X2).GetSafeNormal();
+		TVector<T, d> Normal1 = TVector<T, d>::CrossProduct(V1, V2).GetSafeNormal();
+		if (TVector<T, 3>::DotProduct(Normal1, X2 - X0) > 0)
+		{
+			Normal1 *= -1;
+		}
+		TVector<T, d> Normal2 = TVector<T, d>::CrossProduct(V1, V3).GetSafeNormal();
+		if (TVector<T, 3>::DotProduct(Normal2, X1 - X0) > 0)
+		{
+			Normal2 *= -1;
+		}
+		TVector<T, d> Normal3 = TVector<T, d>::CrossProduct(V2, V3).GetSafeNormal();
+		if (TVector<T, 3>::DotProduct(Normal3, X0 - X1) > 0)
+		{
+			Normal3 *= -1;
+		}
+		TPlane<T, d> NewPlane1(NewX, Normal1);
+		TPlane<T, d> NewPlane2(NewX, Normal2);
+		TPlane<T, d> NewPlane3(NewX, Normal3);
+		TArray<int32> NewIndices1;
+		TArray<int32> NewIndices2;
+		TArray<int32> NewIndices3;
+		for (int32 i = 0; i < Indices.Num(); ++i)
+		{
+			if (MaxD == Indices[i])
+			{
+				continue;
+			}
+			if (NewPlane1.SignedDistance(Constraint.Location[Indices[i]]) > 0)
+			{
+				NewIndices1.Add(Indices[i]);
+			}
+			if (NewPlane2.SignedDistance(Constraint.Location[Indices[i]]) > 0)
+			{
+				NewIndices2.Add(Indices[i]);
+			}
+			if (NewPlane3.SignedDistance(Constraint.Location[Indices[i]]) > 0)
+			{
+				NewIndices3.Add(Indices[i]);
+			}
+		}
+		FindPointsOnHull(Constraint, X0, X1, NewX, NewPlane1, NewIndices1, HullPoints);
+		FindPointsOnHull(Constraint, X0, X2, NewX, NewPlane2, NewIndices2, HullPoints);
+		FindPointsOnHull(Constraint, X1, X2, NewX, NewPlane3, NewIndices3, HullPoints);
+	}
+}
+
+template<class T, int d>
+void RemovePointsInsideHull(TRigidBodyContactConstraintPGS<T, d>& Constraint)
+{
+	if (Constraint.Location.Num() <= 2)
+	{
+		return;
+	}
+	// Find max and min x points
+	int32 MinX = 0;
+	int32 MaxX = 0;
+	int32 MinY = 0;
+	int32 MaxY = 0;
+	int32 Index1 = 0;
+	int32 Index2 = 0;
+	for (int32 i = 1; i < Constraint.Location.Num(); ++i)
+	{
+		if (Constraint.Location[i][0] > Constraint.Location[MaxX][0])
+		{
+			MaxX = i;
+		}
+		if (Constraint.Location[i][0] < Constraint.Location[MinX][0])
+		{
+			MinX = i;
+		}
+		if (Constraint.Location[i][1] > Constraint.Location[MaxY][1])
+		{
+			MaxY = i;
+		}
+		if (Constraint.Location[i][1] < Constraint.Location[MinY][1])
+		{
+			MinY = i;
+		}
+	}
+	if (MaxX == MinX && MinY == MaxY && MinX == MinY)
+	{
+		// Points are colinear so need to sort but for now do nothing
+		return;
+	}
+	// Find max distance
+	T DistanceY = (Constraint.Location[MaxY] - Constraint.Location[MinY]).Size();
+	T DistanceX = (Constraint.Location[MaxX] - Constraint.Location[MinX]).Size();
+	if (DistanceX > DistanceY)
+	{
+		Index1 = MaxX;
+		Index2 = MinX;
+	}
+	else
+	{
+		Index1 = MaxY;
+		Index2 = MinY;
+	}
+	TSet<int32> HullPoints;
+	HullPoints.Add(Index1);
+	HullPoints.Add(Index2);
+	const TVector<T, d>& X1 = Constraint.Location[Index1];
+	const TVector<T, d>& X2 = Constraint.Location[Index2];
+	T MaxDist = 0;
+	int32 MaxD = -1;
+	for (int32 i = 0; i < Constraint.Location.Num(); ++i)
+	{
+		if (i == Index1 || i == Index2)
+		{
+			continue;
+		}
+		const TVector<T, d>& X0 = Constraint.Location[i];
+		T Distance = TVector<T, d>::CrossProduct(X0 - X1, X0 - X2).Size() / (X2 - X1).Size();
+		if (Distance > MaxDist)
+		{
+			MaxDist = Distance;
+			MaxD = i;
+		}
+	}
+	if (MaxD != -1)
+	{
+		HullPoints.Add(MaxD);
+		const TVector<T, d>& X0 = Constraint.Location[MaxD];
+		TVector<T, d> Normal = TVector<T, d>::CrossProduct((X0 - X1).GetSafeNormal(), (X0 - X2).GetSafeNormal());
+		TPlane<T, d> SplitPlane(X0, Normal);
+		TPlane<T, d> SplitPlaneNeg(X0, -Normal);
+		TArray<int32> Left;
+		TArray<int32> Right;
+		for (int32 i = 0; i < Constraint.Location.Num(); ++i)
+		{
+			if (i == Index1 || i == Index2 || i == MaxD)
+			{
+				continue;
+			}
+			if (SplitPlane.SignedDistance(Constraint.Location[i]) >= 0)
+			{
+				Left.Add(i);
+			}
+			else
+			{
+				Right.Add(i);
+			}
+		}
+		FindPointsOnHull(Constraint, X0, X1, X2, SplitPlane, Left, HullPoints);
+		FindPointsOnHull(Constraint, X0, X1, X2, SplitPlaneNeg, Right, HullPoints);
+	}
+	TArray<TVector<T, d>> Locations;
+	TArray<TVector<T, d>> Normals;
+	TArray<T> Distances;
+	for (const auto& Index : HullPoints)
+	{
+		Locations.Add(Constraint.Location[Index]);
+		Normals.Add(Constraint.Normal[Index]);
+		Distances.Add(Constraint.Phi[Index]);
+	}
+	Constraint.Location = Locations;
+	Constraint.Normal = Normals;
+	Constraint.Phi = Distances;
+}
+
+template<class T, int d>
+void TPBDCollisionConstraintPGS<T, d>::Apply(TPBDRigidParticles<T, d>& InParticles, const T Dt, const TArray<int32>& InConstraintIndices)
+{
+	PhysicsParallelFor(InConstraintIndices.Num(), [&](int32 ConstraintIndex) {
+		FRigidBodyContactConstraint& Constraint = Constraints[ConstraintIndex];
+		if (InParticles.Sleeping(Constraint.ParticleIndex))
+		{
+			check(InParticles.Sleeping(Constraint.LevelsetIndex) || InParticles.InvM(Constraint.LevelsetIndex) == 0);
+			return;
+		}
+		// @todo(mlentine): This is a really coarse approximation. Prune points that we know are not relevant.
+		const T Threshold = (InParticles.V(Constraint.ParticleIndex).Size() - InParticles.V(Constraint.LevelsetIndex).Size()) * Dt;
+		// Guessing Max is a decent approximation as with restitution 0 the difference is X between predicted and actual is Vdt
+		const T Thickness = MThickness + FMath::Max(InParticles.V(Constraint.ParticleIndex).Size(), InParticles.V(Constraint.LevelsetIndex).Size()) * Dt;
+		const_cast<TPBDCollisionConstraintPGS<T, d>*>(this)->UpdateConstraint(static_cast<TRigidParticles<T, d>&>(InParticles), Thickness + Threshold, Constraint);
+		// @todo(mlentine): Prune contact points based on convex hull
+		RemovePointsInsideHull(Constraint);
+	});
+	PrintParticles(InParticles, InConstraintIndices);
+	PrintConstraints(InParticles, InConstraintIndices);
+	Solve(static_cast<TRigidParticles<T, d>&>(InParticles), Dt, InConstraintIndices);
+	PrintParticles(InParticles, InConstraintIndices);
+}
+
+template<class T, int d>
+void TPBDCollisionConstraintPGS<T, d>::ApplyPushOut(TPBDRigidParticles<T, d>& InParticles, const T Dt, const TArray<int32>& InConstraintIndices)
+{
+	PhysicsParallelFor(InConstraintIndices.Num(), [&](int32 ConstraintIndex) {
+		FRigidBodyContactConstraint& Constraint = Constraints[ConstraintIndex];
+		if (InParticles.Sleeping(Constraint.ParticleIndex))
+		{
+			check(InParticles.Sleeping(Constraint.LevelsetIndex) || InParticles.InvM(Constraint.LevelsetIndex) == 0);
+			return;
+		}
+		const_cast<TPBDCollisionConstraintPGS<T, d>*>(this)->UpdateConstraint(InParticles, MThickness, Constraint);
+		// @todo(mlentine): Prune contact points based on convex hull
+	});
+
+	TArray<bool> Saved;
+	TArray<TVector<T, d>> SavedV, SavedW;
+	Saved.SetNum(InParticles.Size());
+	SavedV.SetNum(InParticles.Size());
+	SavedW.SetNum(InParticles.Size());
+	auto SaveParticle = [&](int32 ParticleIndex)
+	{
+		if (!Saved[ParticleIndex])
+		{
+			SavedV[ParticleIndex] = InParticles.V(ParticleIndex);
+			SavedW[ParticleIndex] = InParticles.W(ParticleIndex);
+			InParticles.V(ParticleIndex) = TVector<T, d>(0);
+			InParticles.W(ParticleIndex) = TVector<T, d>(0);
+			Saved[ParticleIndex] = true;
+		}
+	};
+	auto RestoreParticle = [&](int32 ParticleIndex)
+	{
+		if (Saved[ParticleIndex])
+		{
+			if (InParticles.InvM(ParticleIndex))
+			{
+				InParticles.P(ParticleIndex) += InParticles.V(ParticleIndex) * Dt;
+				InParticles.Q(ParticleIndex) += TRotation<T, d>(InParticles.W(ParticleIndex), 0.f) * InParticles.Q(ParticleIndex) * Dt * T(0.5);
+				InParticles.Q(ParticleIndex).Normalize();
+			}
+			InParticles.V(ParticleIndex) = SavedV[ParticleIndex];
+			InParticles.W(ParticleIndex) = SavedW[ParticleIndex];
+			Saved[ParticleIndex] = false;
+		}
+	};
+
+	PhysicsParallelFor(InConstraintIndices.Num(), [&](int32 ConstraintIndex) {
+		SaveParticle(Constraints[ConstraintIndex].ParticleIndex);
+		SaveParticle(Constraints[ConstraintIndex].LevelsetIndex);
+	});
+	PrintParticles(InParticles, InConstraintIndices);
+	PrintConstraints(InParticles, InConstraintIndices);
+
+	Solve(InParticles, Dt, InConstraintIndices);
+	
+	PrintParticles(InParticles, InConstraintIndices);
+	PhysicsParallelFor(InConstraintIndices.Num(), [&](int32 ConstraintIndex) {
+		RestoreParticle(Constraints[ConstraintIndex].ParticleIndex);
+		RestoreParticle(Constraints[ConstraintIndex].LevelsetIndex);
+	});
+}
+
+template<class T, int d>
+bool TPBDCollisionConstraintPGS<T, d>::NearestPoint(TArray<Pair<TVector<T, d>, TVector<T, d>>>& Points, TVector<T, d>& Direction)
+{
+	check(Points.Num() > 1 && Points.Num() <= 4);
+	if (Points.Num() == 2)
+	{
+		TPlane<T, d> LocalPlane(Points[1].First, Points[0].First - Points[1].First);
+		TVector<T, d> Normal;
+		const auto& Phi = LocalPlane.PhiWithNormal(TVector<T, d>(0), Normal);
+		if ((TVector<T, d>::DotProduct(-Points[1].First, Normal.GetSafeNormal()) - Points[1].First.Size()) < SMALL_NUMBER)
+		{
+			T Alpha = Points[0].First.Size() / (Points[1].First - Points[0].First).Size();
+			return true;
+		}
+		if (Phi > 0)
+		{
+			check(Points.Num() == 2);
+			Direction = TVector<T, d>::CrossProduct(TVector<T, d>::CrossProduct(Normal, -Points[1].First), Normal);
+		}
+		else
+		{
+			Direction = -Points[1].First;
+			Points.RemoveAtSwap(0);
+			check(Points.Num() == 1);
+		}
+		check(Points.Num() > 1 && Points.Num() < 4);
+		return false;
+	}
+	if (Points.Num() == 3)
+	{
+		TVector<T, d> TriangleNormal = TVector<T, d>::CrossProduct(Points[0].First - Points[2].First, Points[0].First - Points[1].First);
+		TPlane<T, d> LocalPlane1(Points[2].First, TVector<T, d>::CrossProduct(Points[0].First - Points[2].First, TriangleNormal));
+		TPlane<T, d> LocalPlane2(Points[2].First, TVector<T, d>::CrossProduct(Points[1].First - Points[2].First, TriangleNormal));
+		TVector<T, d> Normal;
+		T Phi = LocalPlane1.PhiWithNormal(TVector<T, d>(0), Normal);
+		if (Phi > 0)
+		{
+			TVector<T, d> Delta = Points[0].First - Points[2].First;
+			if (TVector<T, d>::DotProduct(-Points[2].First, Delta) > 0)
+			{
+				Direction = TVector<T, d>::CrossProduct(TVector<T, d>::CrossProduct(Delta, -Points[2].First), Delta);
+				Points.RemoveAtSwap(1);
+				check(Points.Num() == 2);
+			}
+			else
+			{
+				Delta = Points[1].First - Points[2].First;
+				if (TVector<T, d>::DotProduct(-Points[2].First, Delta) > 0)
+				{
+					Direction = TVector<T, d>::CrossProduct(TVector<T, d>::CrossProduct(Delta, -Points[2].First), Delta);
+					Points.RemoveAtSwap(0);
+					check(Points.Num() == 2);
+				}
+				else
+				{
+					Direction = -Points[2].First;
+					Points = {Points[2]};
+					check(Points.Num() == 1);
+				}
+			}
+		}
+		else
+		{
+			Phi = LocalPlane2.PhiWithNormal(TVector<T, d>(0), Normal);
+			if (Phi > 0)
+			{
+				TVector<T, d> Delta = Points[1].First - Points[2].First;
+				if (TVector<T, d>::DotProduct(-Points[2].First, Delta) > 0)
+				{
+					Direction = TVector<T, d>::CrossProduct(TVector<T, d>::CrossProduct(Delta, -Points[2].First), Delta);
+					Points.RemoveAtSwap(0);
+					check(Points.Num() == 2);
+				}
+				else
+				{
+					Direction = -Points[2].First;
+					Points = {Points[2]};
+					check(Points.Num() == 1);
+				}
+			}
+			else
+			{
+				const auto DotResult = TVector<T, d>::DotProduct(TriangleNormal, -Points[2].First);
+				// We are inside the triangle
+				if (DotResult < SMALL_NUMBER)
+				{
+					TVector<T, 3> Bary;
+					TVector<T, d> P10 = Points[1].First - Points[0].First;
+					TVector<T, d> P20 = Points[2].First - Points[0].First;
+					TVector<T, d> PP0 = -Points[0].First;
+					T Size10 = P10.SizeSquared();
+					T Size20 = P20.SizeSquared();
+					T ProjSides = TVector<T, d>::DotProduct(P10, P20);
+					T ProjP1 = TVector<T, d>::DotProduct(PP0, P10);
+					T ProjP2 = TVector<T, d>::DotProduct(PP0, P20);
+					T Denom = Size10 * Size20 - ProjSides * ProjSides;
+					Bary.Y = (Size20 * ProjP1 - ProjSides * ProjP2) / Denom;
+					Bary.Z = (Size10 * ProjP2 - ProjSides * ProjP1) / Denom;
+					Bary.X = 1.0f - Bary.Z - Bary.Y;
+					return true;
+				}
+				if (DotResult > 0)
+				{
+					Direction = TriangleNormal;
+				}
+				else
+				{
+					Direction = -TriangleNormal;
+					Points.Swap(0, 1);
+					check(Points.Num() == 3);
+				}
+			}
+		}
+		check(Points.Num() > 0 && Points.Num() < 4);
+		return false;
+	}
+	if (Points.Num() == 4)
+	{
+		TVector<T, d> TriangleNormal = TVector<T, d>::CrossProduct(Points[1].First - Points[3].First, Points[1].First - Points[2].First);
+		if (TVector<T, d>::DotProduct(TriangleNormal, Points[0].First - Points[3].First) > 0)
+		{
+			TriangleNormal *= -1;
+		}
+		T DotResult = TVector<T, d>::DotProduct(TriangleNormal, -Points[3].First);
+		if (DotResult > 0)
+		{
+			Points = {Points[1], Points[2], Points[3]};
+			check(Points.Num() == 3);
+			return NearestPoint(Points, Direction);
+		}
+		TriangleNormal = TVector<T, d>::CrossProduct(Points[2].First - Points[0].First, Points[2].First - Points[3].First);
+		if (TVector<T, d>::DotProduct(TriangleNormal, Points[1].First - Points[3].First) > 0)
+		{
+			TriangleNormal *= -1;
+		}
+		DotResult = TVector<T, d>::DotProduct(TriangleNormal, -Points[3].First);
+		if (DotResult > 0)
+		{
+			Points = {Points[0], Points[2], Points[3]};
+			check(Points.Num() == 3);
+			return NearestPoint(Points, Direction);
+		}
+		TriangleNormal = TVector<T, d>::CrossProduct(Points[3].First - Points[1].First, Points[3].First - Points[0].First);
+		if (TVector<T, d>::DotProduct(TriangleNormal, Points[2].First - Points[3].First) > 0)
+		{
+			TriangleNormal *= -1;
+		}
+		DotResult = TVector<T, d>::DotProduct(TriangleNormal, -Points[3].First);
+		if (DotResult > 0)
+		{
+			Points = {Points[0], Points[1], Points[3]};
+			check(Points.Num() == 3);
+			return NearestPoint(Points, Direction);
+		}
+		TVector<T, 4> Bary;
+		TVector<T, d> PP0 = -Points[0].First;
+		TVector<T, d> PP1 = -Points[1].First;
+		TVector<T, d> P10 = Points[1].First - Points[0].First;
+		TVector<T, d> P20 = Points[2].First - Points[0].First;
+		TVector<T, d> P30 = Points[3].First - Points[0].First;
+		TVector<T, d> P21 = Points[2].First - Points[1].First;
+		TVector<T, d> P31 = Points[3].First - Points[1].First;
+		Bary[0] = TVector<T, d>::DotProduct(PP1, TVector<T, d>::CrossProduct(P31, P21));
+		Bary[1] = TVector<T, d>::DotProduct(PP0, TVector<T, d>::CrossProduct(P20, P30));
+		Bary[2] = TVector<T, d>::DotProduct(PP0, TVector<T, d>::CrossProduct(P30, P10));
+		Bary[3] = TVector<T, d>::DotProduct(PP0, TVector<T, d>::CrossProduct(P10, P20));
+		T Denom = TVector<T, d>::DotProduct(P10, TVector<T, d>::CrossProduct(P20, P30));
+		return true;
+	}
+	check(Points.Num() > 1 && Points.Num() < 4);
+	return false;
+}
+
+template<class T, int d>
+void UpdateLevelsetConstraintHelperCCD(const TRigidParticles<T, d>& InParticles, const int32 j, const TRigidTransform<T, d>& LocalToWorld1, const TRigidTransform<T, d>& LocalToWorld2, const T Thickness, TRigidBodyContactConstraintPGS<T, d>& Constraint)
+{
+	if(InParticles.CollisionParticles(Constraint.ParticleIndex))
+	{
+		const TRigidTransform<T, d> PreviousLocalToWorld1 = GetTransformPGS(InParticles, Constraint.ParticleIndex);
+		TVector<T, d> WorldSpacePointStart = PreviousLocalToWorld1.TransformPosition(InParticles.CollisionParticles(Constraint.ParticleIndex)->X(j));
+		TVector<T, d> WorldSpacePointEnd = LocalToWorld1.TransformPosition(InParticles.CollisionParticles(Constraint.ParticleIndex)->X(j));
+		TVector<T, d> Body2SpacePointStart = LocalToWorld2.InverseTransformPosition(WorldSpacePointStart);
+		TVector<T, d> Body2SpacePointEnd = LocalToWorld2.InverseTransformPosition(WorldSpacePointEnd);
+		Pair<TVector<T, d>, bool> PointPair = InParticles.Geometry(Constraint.LevelsetIndex)->FindClosestIntersection(Body2SpacePointStart, Body2SpacePointEnd, Thickness);
+		if (PointPair.Second)
+		{
+			const TVector<T, d> WorldSpaceDelta = WorldSpacePointEnd - TVector<T, d>(LocalToWorld2.TransformPosition(PointPair.First));
+			Constraint.Phi.Add(-WorldSpaceDelta.Size());
+			Constraint.Normal.Add(LocalToWorld2.TransformVector(InParticles.Geometry(Constraint.LevelsetIndex)->Normal(PointPair.First)));
+			// @todo(mlentine): Should we be using the actual collision point or that point evolved to the current time step?
+			Constraint.Location.Add(WorldSpacePointEnd);
+		}
+	}
+}
+
+template<class T, int d>
+void UpdateLevelsetConstraintHelper(const TRigidParticles<T, d>& InParticles, const int32 j, const TRigidTransform<T, d>& LocalToWorld1, const TRigidTransform<T, d>& LocalToWorld2, const T Thickness, TRigidBodyContactConstraintPGS<T, d>& Constraint)
+{
+	if(InParticles.CollisionParticles(Constraint.ParticleIndex))
+	{
+		TVector<T, d> WorldSpacePoint = LocalToWorld1.TransformPosition(InParticles.CollisionParticles(Constraint.ParticleIndex)->X(j));
+		TVector<T, d> Body2SpacePoint = LocalToWorld2.InverseTransformPosition(WorldSpacePoint);
+		TVector<T, d> LocalNormal;
+		T LocalPhi = InParticles.Geometry(Constraint.LevelsetIndex)->PhiWithNormal(Body2SpacePoint, LocalNormal);
+		if (LocalPhi < Thickness)
+		{
+			Constraint.Phi.Add(LocalPhi);
+			Constraint.Normal.Add(LocalToWorld2.TransformVector(LocalNormal));
+			Constraint.Location.Add(WorldSpacePoint);
+		}
+	}
+}
+
+template<class T, int d>
+template<class T_PARTICLES>
+void TPBDCollisionConstraintPGS<T, d>::UpdateLevelsetConstraint(const T_PARTICLES& InParticles, const T Thickness, FRigidBodyContactConstraint& Constraint)
+{
+	Constraint.Normal.Reset();
+	Constraint.Phi.Reset();
+	Constraint.Location.Reset();
+	const TRigidTransform<T, d> LocalToWorld1 = GetTransformPGS(InParticles, Constraint.ParticleIndex);
+	const TRigidTransform<T, d> LocalToWorld2 = GetTransformPGS(InParticles, Constraint.LevelsetIndex);
+	if (InParticles.Geometry(Constraint.LevelsetIndex)->HasBoundingBox())
+	{
+		TBox<T, d> ImplicitBox = InParticles.Geometry(Constraint.LevelsetIndex)->BoundingBox().TransformedBox(LocalToWorld2 * LocalToWorld1.Inverse());
+		if (InParticles.CollisionParticles(Constraint.ParticleIndex))
+		{
+			TArray<int32> PotentialParticles = InParticles.CollisionParticles(Constraint.ParticleIndex)->FindAllIntersections(ImplicitBox);
+			for (int32 j = 0; j < PotentialParticles.Num(); ++j)
+			{
+				if (bUseCCD)
+				{
+					UpdateLevelsetConstraintHelperCCD(InParticles, PotentialParticles[j], LocalToWorld1, LocalToWorld2, Thickness, Constraint);
+				}
+				else
+				{
+					UpdateLevelsetConstraintHelper(InParticles, PotentialParticles[j], LocalToWorld1, LocalToWorld2, Thickness, Constraint);
+				}
+			}
+		}
+	}
+	else
+	{
+		if (InParticles.CollisionParticles(Constraint.ParticleIndex))
+		{
+			for (uint32 j = 0; j < InParticles.CollisionParticles(Constraint.ParticleIndex)->Size(); ++j)
+			{
+				UpdateLevelsetConstraintHelper(InParticles, j, LocalToWorld1, LocalToWorld2, Thickness, Constraint);
+			}
+		}
+	}
+}
+
+template<class T, int d>
+template<class T_PARTICLES>
+void TPBDCollisionConstraintPGS<T, d>::UpdateLevelsetConstraintGJK(const T_PARTICLES& InParticles, const T Thickness, FRigidBodyContactConstraint& Constraint)
+{
+	static int32 MaxIterationsGJK = 100;
+	Constraint.Normal.Reset();
+	Constraint.Phi.Reset();
+	Constraint.Location.Reset();
+	const TRigidTransform<T, d> LocalToWorld1 = GetTransformPGS(InParticles, Constraint.ParticleIndex);
+	const TRigidTransform<T, d> LocalToWorld2 = GetTransformPGS(InParticles, Constraint.LevelsetIndex);
+	TVector<T, d> Direction = LocalToWorld1.GetTranslation() - LocalToWorld2.GetTranslation();
+	TVector<T, d> SupportA = LocalToWorld1.TransformPosition(InParticles.Geometry(Constraint.ParticleIndex)->Support(LocalToWorld1.InverseTransformVector(-Direction), Thickness));
+	TVector<T, d> SupportB = LocalToWorld2.TransformPosition(InParticles.Geometry(Constraint.LevelsetIndex)->Support(LocalToWorld2.InverseTransformVector(Direction), Thickness));
+	TVector<T, d> Point = SupportB - SupportA;
+	TArray<Pair<TVector<T, d>, TVector<T, d>>> Points = {MakePair(Point, SupportA)};
+	Direction = -Point;
+	for (int32 i = 0; i < MaxIterationsGJK; ++i)
+	{
+		SupportA = LocalToWorld1.TransformPosition(InParticles.Geometry(Constraint.ParticleIndex)->Support(LocalToWorld1.InverseTransformVector(-Direction), Thickness));
+		SupportB = LocalToWorld2.TransformPosition(InParticles.Geometry(Constraint.LevelsetIndex)->Support(LocalToWorld2.InverseTransformVector(Direction), Thickness));
+		Point = SupportB - SupportA;
+		if (TVector<T, d>::DotProduct(Point, Direction) < 0)
+		{
+			break;
+		}
+		Points.Add(MakePair(Point, SupportA));
+		TVector<T, d> ClosestPoint;
+		if (NearestPoint(Points, Direction))
+		{
+			for (const auto& SinglePoint : Points)
+			{
+				TVector<T, d> Body1Location = LocalToWorld1.InverseTransformPosition(SinglePoint.Second);
+				TVector<T, d> Normal;
+				T Phi = InParticles.Geometry(Constraint.ParticleIndex)->PhiWithNormal(Body1Location, Normal);
+				Normal = LocalToWorld1.TransformVector(Normal);
+				TVector<T, d> SurfacePoint = SinglePoint.Second - Phi * Normal;
+				Constraint.Location.Add(SurfacePoint);
+				TVector<T, d> Body2Location = LocalToWorld2.InverseTransformPosition(SurfacePoint);
+				Constraint.Phi.Add(InParticles.Geometry(Constraint.LevelsetIndex)->PhiWithNormal(Body2Location, Normal));
+				Constraint.Normal.Add(LocalToWorld2.TransformVector(Normal));
+			}
+			break;
+		}
+	}
+}
+
+template<class T, int d>
+template<class T_PARTICLES>
+void TPBDCollisionConstraintPGS<T, d>::UpdateBoxConstraint(const T_PARTICLES& InParticles, const T Thickness, FRigidBodyContactConstraint& Constraint)
+{
+	Constraint.Normal.Reset();
+	Constraint.Phi.Reset();
+	Constraint.Location.Reset();
+	const TRigidTransform<T, d> Box1Transform = GetTransformPGS(InParticles, Constraint.ParticleIndex);
+	const TRigidTransform<T, d> Box2Transform = GetTransformPGS(InParticles, Constraint.LevelsetIndex);
+	const auto& Box1 = *InParticles.Geometry(Constraint.ParticleIndex)->template GetObject<TBox<T, d>>();
+	const auto& Box2 = *InParticles.Geometry(Constraint.LevelsetIndex)->template GetObject<TBox<T, d>>();
+	auto Box2SpaceBox1 = Box1.TransformedBox(Box1Transform * Box2Transform.Inverse());
+	auto Box1SpaceBox2 = Box2.TransformedBox(Box2Transform * Box1Transform.Inverse());
+	Box2SpaceBox1.Thicken(Thickness);
+	Box1SpaceBox2.Thicken(Thickness);
+	if (Box1SpaceBox2.Intersects(Box1) && Box2SpaceBox1.Intersects(Box2))
+	{
+		const TVector<T, d> Box1Center = (Box1Transform * Box2Transform.Inverse()).TransformPosition(Box1.Center());
+		if (Box2.SignedDistance(Box1Center) < 0)
+		{
+			TSphere<T, d> Sphere1(Box1Transform.TransformPosition(Box1.Center()), Box1.Extents().Min() / 2);
+			TSphere<T, d> Sphere2(Box2Transform.TransformPosition(Box2.Center()), Box2.Extents().Min() / 2);
+			const TVector<T, d> Direction = Sphere1.GetCenter() - Sphere2.GetCenter();
+			T Size = Direction.Size();
+			if (Size < (Sphere1.GetRadius() + Sphere2.GetRadius()))
+			{
+				TVector<T, d> Normal = Size > SMALL_NUMBER ? Direction / Size : TVector<T, d>(0, 0, 1);
+				Constraint.Normal.Add(Normal);
+				Constraint.Phi.Add(Size - (Sphere1.GetRadius() + Sphere2.GetRadius()));
+				Constraint.Location.Add(Sphere1.GetCenter() - Sphere1.GetRadius() * Normal);
+			}
+		}
+		if (!Constraint.Phi.Num())
+		{
+			//UpdateLevelsetConstraintGJK(InParticles, Thickness, Constraint);
+			//check(Constraint.Phi < MThickness);
+			// For now revert to doing all points vs lsv check until we can figure out a good way to get the deepest point without needing this
+			UpdateLevelsetConstraint(InParticles, Thickness, Constraint);
+		}
+	}
+}
+
+template<class T, int d>
+template<class T_PARTICLES>
+void TPBDCollisionConstraintPGS<T, d>::UpdateBoxPlaneConstraint(const T_PARTICLES& InParticles, const T Thickness, FRigidBodyContactConstraint& Constraint)
+{
+	Constraint.Normal.Reset();
+	Constraint.Phi.Reset();
+	Constraint.Location.Reset();
+	const TRigidTransform<T, d> BoxTransform = GetTransformPGS(InParticles, Constraint.ParticleIndex);
+	const TRigidTransform<T, d> PlaneTransform = GetTransformPGS(InParticles, Constraint.LevelsetIndex);
+	const auto& ObjectBox = *InParticles.Geometry(Constraint.ParticleIndex)->template GetObject<TBox<T, d>>();
+	const auto& ObjectPlane = *InParticles.Geometry(Constraint.LevelsetIndex)->template GetObject<TPlane<T, d>>();
+	const TRigidTransform<T, d> BoxToPlaneTransform(BoxTransform * PlaneTransform.Inverse());
+	const TVector<T, d> Extents = ObjectBox.Extents();
+	TArray<TVector<T, d>> Corners;
+	Corners.Add(BoxToPlaneTransform.TransformPosition(ObjectBox.Max()));
+	Corners.Add(BoxToPlaneTransform.TransformPosition(ObjectBox.Min()));
+	for (int32 j = 0; j < d; ++j)
+	{
+		Corners.Add(BoxToPlaneTransform.TransformPosition(ObjectBox.Min() + TVector<T, d>::AxisVector(j) * Extents));
+		Corners.Add(BoxToPlaneTransform.TransformPosition(ObjectBox.Max() - TVector<T, d>::AxisVector(j) * Extents));
+	}
+	TArray<TVector<T, d>> PotentialConstraints;
+	for (int32 i = 0; i < Corners.Num(); ++i)
+	{
+		TVector<T, d> Normal;
+		const T NewPhi = ObjectPlane.PhiWithNormal(Corners[i], Normal);
+		if (NewPhi < Thickness)
+		{
+			Constraint.Phi.Add(NewPhi);
+			Constraint.Normal.Add(PlaneTransform.TransformVector(Normal));
+			Constraint.Location.Add(PlaneTransform.TransformPosition(Corners[i]));
+		}
+	}
+}
+
+template<class T, int d>
+template<class T_PARTICLES>
+void TPBDCollisionConstraintPGS<T, d>::UpdateSphereConstraint(const T_PARTICLES& InParticles, const T Thickness, FRigidBodyContactConstraint& Constraint)
+{
+	Constraint.Normal.Reset();
+	Constraint.Phi.Reset();
+	Constraint.Location.Reset();
+	const TRigidTransform<T, d> Sphere1Transform = GetTransformPGS(InParticles, Constraint.ParticleIndex);
+	const TRigidTransform<T, d> Sphere2Transform = GetTransformPGS(InParticles, Constraint.LevelsetIndex);
+	const auto& Sphere1 = *InParticles.Geometry(Constraint.ParticleIndex)->template GetObject<TSphere<T, d>>();
+	const auto& Sphere2 = *InParticles.Geometry(Constraint.LevelsetIndex)->template GetObject<TSphere<T, d>>();
+	const TVector<T, d> Center1 = Sphere1Transform.TransformPosition(Sphere1.GetCenter());
+	const TVector<T, d> Center2 = Sphere2Transform.TransformPosition(Sphere2.GetCenter());
+	const TVector<T, d> Direction = Center1 - Center2;
+	const T Size = Direction.Size();
+	if (Size < (Sphere1.GetRadius() + Sphere2.GetRadius() + Thickness))
+	{
+		TVector<T, d> Normal = Size > SMALL_NUMBER ? Direction / Size : TVector<T, d>(0, 0, 1);
+		Constraint.Normal.Add(Normal);
+		Constraint.Phi.Add(Size - (Sphere1.GetRadius() + Sphere2.GetRadius()));
+		Constraint.Location.Add(Center1 - Sphere1.GetRadius() * Normal);
+	}
+}
+
+template<class T, int d>
+template<class T_PARTICLES>
+void TPBDCollisionConstraintPGS<T, d>::UpdateSpherePlaneConstraint(const T_PARTICLES& InParticles, const T Thickness, FRigidBodyContactConstraint& Constraint)
+{
+	Constraint.Normal.Reset();
+	Constraint.Phi.Reset();
+	Constraint.Location.Reset();
+	const TRigidTransform<T, d> SphereTransform = GetTransformPGS(InParticles, Constraint.ParticleIndex);
+	const TRigidTransform<T, d> PlaneTransform = GetTransformPGS(InParticles, Constraint.LevelsetIndex);
+	const auto& ObjectSphere = *InParticles.Geometry(Constraint.ParticleIndex)->template GetObject<TSphere<T, d>>();
+	const auto& ObjectPlane = *InParticles.Geometry(Constraint.LevelsetIndex)->template GetObject<TPlane<T, d>>();
+	const TRigidTransform<T, d> SphereToPlaneTransform(PlaneTransform.Inverse() * SphereTransform);
+	const TVector<T, d> SphereCenter = SphereToPlaneTransform.TransformPosition(ObjectSphere.GetCenter());
+	Constraint.Normal.SetNum(1);
+	Constraint.Phi.Add(ObjectPlane.PhiWithNormal(SphereCenter, Constraint.Normal[0]));
+	Constraint.Phi[0] -= ObjectSphere.GetRadius();
+	Constraint.Location.Add(SphereCenter - Constraint.Normal[0] * ObjectSphere.GetRadius());
+}
+
+template<class T, int d>
+template<class T_PARTICLES>
+void TPBDCollisionConstraintPGS<T, d>::UpdateSphereBoxConstraint(const T_PARTICLES& InParticles, const T Thickness, FRigidBodyContactConstraint& Constraint)
+{
+	Constraint.Normal.Reset();
+	Constraint.Phi.Reset();
+	Constraint.Location.Reset();
+	const TRigidTransform<T, d> SphereTransform = GetTransformPGS(InParticles, Constraint.ParticleIndex);
+	const TRigidTransform<T, d> BoxTransform = GetTransformPGS(InParticles, Constraint.LevelsetIndex);
+	const auto& ObjectSphere = *InParticles.Geometry(Constraint.ParticleIndex)->template GetObject<TSphere<T, d>>();
+	const auto& ObjectBox = *InParticles.Geometry(Constraint.LevelsetIndex)->template GetObject<TBox<T, d>>();
+	const TRigidTransform<T, d> SphereToBoxTransform(SphereTransform * BoxTransform.Inverse());
+	const TVector<T, d> SphereCenter = SphereToBoxTransform.TransformPosition(ObjectSphere.GetCenter());
+	Constraint.Normal.SetNum(1);
+	Constraint.Phi.Add(ObjectBox.PhiWithNormal(SphereCenter, Constraint.Normal[0]));
+	Constraint.Phi[0] -= ObjectSphere.GetRadius();
+	Constraint.Location.Add(SphereCenter - Constraint.Normal[0] * ObjectSphere.GetRadius());
+}
+
+template<class T, int d>
+typename TPBDCollisionConstraintPGS<T, d>::FRigidBodyContactConstraint TPBDCollisionConstraintPGS<T, d>::ComputeLevelsetConstraint(const TPBDRigidParticles<T, d>& InParticles, int32 ParticleIndex, int32 LevelsetIndex, const T Thickness)
+{
+	if (!InParticles.CollisionParticlesSize(ParticleIndex))
+	{
+		int32 TmpIndex = ParticleIndex;
+		ParticleIndex = LevelsetIndex;
+		LevelsetIndex = TmpIndex;
+	}
+	// Find Deepest Point
+	FRigidBodyContactConstraint Constraint;
+	Constraint.ParticleIndex = ParticleIndex;
+	Constraint.LevelsetIndex = LevelsetIndex;
+	return Constraint;
+}
+
+template<class T, int d>
+typename TPBDCollisionConstraintPGS<T, d>::FRigidBodyContactConstraint TPBDCollisionConstraintPGS<T, d>::ComputeLevelsetConstraintGJK(const TPBDRigidParticles<T, d>& InParticles, int32 ParticleIndex, int32 LevelsetIndex, const T Thickness)
+{
+	FRigidBodyContactConstraint Constraint;
+	Constraint.ParticleIndex = ParticleIndex;
+	Constraint.LevelsetIndex = LevelsetIndex;
+	return Constraint;
+}
+
+template<class T, int d>
+typename TPBDCollisionConstraintPGS<T, d>::FRigidBodyContactConstraint TPBDCollisionConstraintPGS<T, d>::ComputeBoxConstraint(const TPBDRigidParticles<T, d>& InParticles, int32 Box1Index, int32 Box2Index, const T Thickness)
+{
+	FRigidBodyContactConstraint Constraint;
+	Constraint.ParticleIndex = Box1Index;
+	Constraint.LevelsetIndex = Box2Index;
+	return Constraint;
+}
+
+template<class T, int d>
+typename TPBDCollisionConstraintPGS<T, d>::FRigidBodyContactConstraint TPBDCollisionConstraintPGS<T, d>::ComputeBoxPlaneConstraint(const TPBDRigidParticles<T, d>& InParticles, int32 BoxIndex, int32 PlaneIndex, const T Thickness)
+{
+	FRigidBodyContactConstraint Constraint;
+	Constraint.ParticleIndex = BoxIndex;
+	Constraint.LevelsetIndex = PlaneIndex;
+	return Constraint;
+}
+
+template<class T, int d>
+typename TPBDCollisionConstraintPGS<T, d>::FRigidBodyContactConstraint TPBDCollisionConstraintPGS<T, d>::ComputeSphereConstraint(const TPBDRigidParticles<T, d>& InParticles, int32 Sphere1Index, int32 Sphere2Index, const T Thickness)
+{
+	FRigidBodyContactConstraint Constraint;
+	Constraint.ParticleIndex = Sphere1Index;
+	Constraint.LevelsetIndex = Sphere2Index;
+	return Constraint;
+}
+
+template<class T, int d>
+typename TPBDCollisionConstraintPGS<T, d>::FRigidBodyContactConstraint TPBDCollisionConstraintPGS<T, d>::ComputeSpherePlaneConstraint(const TPBDRigidParticles<T, d>& InParticles, int32 SphereIndex, int32 PlaneIndex, const T Thickness)
+{
+	FRigidBodyContactConstraint Constraint;
+	Constraint.ParticleIndex = SphereIndex;
+	Constraint.LevelsetIndex = PlaneIndex;
+	return Constraint;
+}
+
+template<class T, int d>
+typename TPBDCollisionConstraintPGS<T, d>::FRigidBodyContactConstraint TPBDCollisionConstraintPGS<T, d>::ComputeSphereBoxConstraint(const TPBDRigidParticles<T, d>& InParticles, int32 SphereIndex, int32 BoxIndex, const T Thickness)
+{
+	FRigidBodyContactConstraint Constraint;
+	Constraint.ParticleIndex = SphereIndex;
+	Constraint.LevelsetIndex = BoxIndex;
+	return Constraint;
+}
+
+template<class T, int d>
+typename TPBDCollisionConstraintPGS<T, d>::FRigidBodyContactConstraint TPBDCollisionConstraintPGS<T, d>::ComputeConstraint(const TPBDRigidParticles<T, d>& InParticles, int32 Body1Index, int32 Body2Index, const T Thickness)
+{
+	if (InParticles.Geometry(Body1Index)->GetType() == TBox<T, d>::GetType() && InParticles.Geometry(Body2Index)->GetType() == TBox<T, d>::GetType())
+	{
+		return ComputeBoxConstraint(InParticles, Body1Index, Body2Index, Thickness);
+	}
+	else if (InParticles.Geometry(Body1Index)->GetType() == TSphere<T, d>::GetType() && InParticles.Geometry(Body2Index)->GetType() == TSphere<T, d>::GetType())
+	{
+		return ComputeSphereConstraint(InParticles, Body1Index, Body2Index, Thickness);
+	}
+	else if (InParticles.Geometry(Body1Index)->GetType() == TBox<T, d>::GetType() && InParticles.Geometry(Body2Index)->GetType() == TPlane<T, d>::GetType())
+	{
+		return ComputeBoxPlaneConstraint(InParticles, Body1Index, Body2Index, Thickness);
+	}
+	else if (InParticles.Geometry(Body2Index)->GetType() == TPlane<T, d>::GetType() && InParticles.Geometry(Body1Index)->GetType() == TBox<T, d>::GetType())
+	{
+		return ComputeBoxPlaneConstraint(InParticles, Body2Index, Body1Index, Thickness);
+	}
+	else if (InParticles.Geometry(Body1Index)->GetType() == TSphere<T, d>::GetType() && InParticles.Geometry(Body2Index)->GetType() == TPlane<T, d>::GetType())
+	{
+		return ComputeSpherePlaneConstraint(InParticles, Body1Index, Body2Index, Thickness);
+	}
+	else if (InParticles.Geometry(Body2Index)->GetType() == TPlane<T, d>::GetType() && InParticles.Geometry(Body1Index)->GetType() == TSphere<T, d>::GetType())
+	{
+		return ComputeSpherePlaneConstraint(InParticles, Body2Index, Body1Index, Thickness);
+	}
+	else if (InParticles.Geometry(Body1Index)->GetType() == TSphere<T, d>::GetType() && InParticles.Geometry(Body2Index)->GetType() == TBox<T, d>::GetType())
+	{
+		return ComputeSphereBoxConstraint(InParticles, Body1Index, Body2Index, Thickness);
+	}
+	else if (InParticles.Geometry(Body2Index)->GetType() == TBox<T, d>::GetType() && InParticles.Geometry(Body1Index)->GetType() == TSphere<T, d>::GetType())
+	{
+		return ComputeSphereBoxConstraint(InParticles, Body2Index, Body1Index, Thickness);
+	}
+	else if (InParticles.Geometry(Body1Index)->IsConvex() && InParticles.Geometry(Body2Index)->IsConvex())
+	{
+		return ComputeLevelsetConstraintGJK(InParticles, Body1Index, Body2Index, Thickness);
+	}
+	return ComputeLevelsetConstraint(InParticles, Body1Index, Body2Index, Thickness);
+}
+
+template<class T, int d>
+template<class T_PARTICLES>
+void TPBDCollisionConstraintPGS<T, d>::UpdateConstraint(const T_PARTICLES& InParticles, const T Thickness, FRigidBodyContactConstraint& Constraint)
+{
+	if (InParticles.Geometry(Constraint.ParticleIndex)->GetType() == TBox<T, d>::GetType() && InParticles.Geometry(Constraint.LevelsetIndex)->GetType() == TBox<T, d>::GetType())
+	{
+		UpdateBoxConstraint(InParticles, Thickness, Constraint);
+	}
+	else if (InParticles.Geometry(Constraint.ParticleIndex)->GetType() == TSphere<T, d>::GetType() && InParticles.Geometry(Constraint.LevelsetIndex)->GetType() == TSphere<T, d>::GetType())
+	{
+		UpdateSphereConstraint(InParticles, Thickness, Constraint);
+	}
+	else if (InParticles.Geometry(Constraint.ParticleIndex)->GetType() == TBox<T, d>::GetType() && InParticles.Geometry(Constraint.LevelsetIndex)->GetType() == TPlane<T, d>::GetType())
+	{
+		UpdateBoxPlaneConstraint(InParticles, Thickness, Constraint);
+	}
+	else if (InParticles.Geometry(Constraint.ParticleIndex)->GetType() == TSphere<T, d>::GetType() && InParticles.Geometry(Constraint.LevelsetIndex)->GetType() == TPlane<T, d>::GetType())
+	{
+		UpdateSpherePlaneConstraint(InParticles, Thickness, Constraint);
+	}
+	else if (InParticles.Geometry(Constraint.ParticleIndex)->GetType() == TSphere<T, d>::GetType() && InParticles.Geometry(Constraint.LevelsetIndex)->GetType() == TBox<T, d>::GetType())
+	{
+		UpdateSphereBoxConstraint(InParticles, Thickness, Constraint);
+	}
+	else if (InParticles.Geometry(Constraint.ParticleIndex)->GetType() == TPlane<T, d>::GetType() && InParticles.Geometry(Constraint.LevelsetIndex)->GetType() == TBox<T, d>::GetType())
+	{
+		int32 Tmp = Constraint.ParticleIndex;
+		Constraint.ParticleIndex = Constraint.LevelsetIndex;
+		Constraint.LevelsetIndex = Tmp;
+		UpdateBoxPlaneConstraint(InParticles, Thickness, Constraint);
+	}
+	else if (InParticles.Geometry(Constraint.ParticleIndex)->GetType() == TPlane<T, d>::GetType() && InParticles.Geometry(Constraint.LevelsetIndex)->GetType() == TSphere<T, d>::GetType())
+	{
+		int32 Tmp = Constraint.ParticleIndex;
+		Constraint.ParticleIndex = Constraint.LevelsetIndex;
+		Constraint.LevelsetIndex = Tmp;
+		UpdateSpherePlaneConstraint(InParticles, Thickness, Constraint);
+	}
+	else if (InParticles.Geometry(Constraint.ParticleIndex)->GetType() == TBox<T, d>::GetType() && InParticles.Geometry(Constraint.LevelsetIndex)->GetType() == TSphere<T, d>::GetType())
+	{
+		int32 Tmp = Constraint.ParticleIndex;
+		Constraint.ParticleIndex = Constraint.LevelsetIndex;
+		Constraint.LevelsetIndex = Tmp;
+		UpdateSphereBoxConstraint(InParticles, Thickness, Constraint);
+	}
+	else if (InParticles.Geometry(Constraint.ParticleIndex)->IsConvex() && InParticles.Geometry(Constraint.LevelsetIndex)->IsConvex())
+	{
+		UpdateLevelsetConstraintGJK(InParticles, Thickness, Constraint);
+	}
+	else
+	{
+		UpdateLevelsetConstraint(InParticles, Thickness, Constraint);
+	}
+}
+
+template class Chaos::TPBDCollisionConstraintPGS<float, 3>;
+
+#endif
\ No newline at end of file
diff --git a/Engine/Source/Runtime/Experimental/Chaos/Private/Chaos/PBDCollisionConstraintsPlaneContactUtil.cpp b/Engine/Source/Runtime/Experimental/Chaos/Private/Chaos/PBDCollisionConstraintsPlaneContactUtil.cpp
new file mode 100644
--- /dev/null
+++ b/Engine/Source/Runtime/Experimental/Chaos/Private/Chaos/PBDCollisionConstraintsPlaneContactUtil.cpp
@@ -0,0 +1,45 @@
+// Copyright 1998-2019 Epic Games, Inc. All Rights Reserved.
+
+#include "Chaos/PBDCollisionConstraintsPlaneContactUtil.h"
+#include "Chaos/CollisionResolutionAlgo.h"
+#include "Chaos/Defines.h"
+
+namespace Chaos
+{
+	namespace Collisions
+	{
+		/*
+		template<ECollisionUpdateType UpdateType, typename T, int d>
+		void Update(const T Thickness, TRigidBodyPlaneContactConstraint<T, d>& Constraint)
+		{
+		}
+		template void Update<ECollisionUpdateType::Any, float, 3>(const float, TRigidBodyPlaneContactConstraint<float, 3>&);
+		template void Update<ECollisionUpdateType::Deepest, float, 3>(const float, TRigidBodyPlaneContactConstraint<float, 3>&);
+
+
+		template<ECollisionUpdateType UpdateType, typename T, int d>
+		void UpdateConstraint(const T Thickness, TRigidBodyPlaneContactConstraint<T, d>& Constraint)
+		{
+		}
+		template void UpdateConstraint<ECollisionUpdateType::Any, float, 3>(const float, TRigidBodyPointContactConstraint<float,3>&);
+		template void UpdateConstraint<ECollisionUpdateType::Deepest, float, 3>(const float, TRigidBodyPointContactConstraint<float,3>&);
+
+
+		template<typename T, int d>
+		void Apply(TRigidBodyPlaneContactConstraint<T, d>& Constraint, T Thickness, TPointContactIterationParameters<T> & IterationParameters, TPlainContactParticleParameters<T> & ParticleParameters)
+		{
+		}
+		template void Apply<float, 3>(TRigidBodyPlaneContactConstraint<float, 3>&, float, TPointContactIterationParameters<float> &, TPointContactParticleParameters<float> &);
+
+		template<typename T, int d>
+		void ApplyPushOut(TRigidBodyPointContactConstraint<T, d>& Constraint, T Thickness, const TSet<const TGeometryParticleHandle<T, d>*>& IsTemporarilyStatic,
+			TPointContactIterationParameters<T> & IterationParameters, TPointContactParticleParameters<T> & ParticleParameters)
+		{
+		}
+		template void ApplyPushOut<float, 3>(TRigidBodyPointContactConstraint<float,3>&, float , const TSet<const TGeometryParticleHandle<float,3>*>&,
+			TPointContactIterationParameters<float> & IterationParameters, TPointContactParticleParameters<float> & ParticleParameters);
+		*/
+	} // Collisions
+
+}// Chaos
+
diff --git a/Engine/Source/Runtime/Experimental/Chaos/Private/Chaos/PBDCollisionConstraintsPointContactUtil.cpp b/Engine/Source/Runtime/Experimental/Chaos/Private/Chaos/PBDCollisionConstraintsPointContactUtil.cpp
new file mode 100644
--- /dev/null
+++ b/Engine/Source/Runtime/Experimental/Chaos/Private/Chaos/PBDCollisionConstraintsPointContactUtil.cpp
@@ -0,0 +1,467 @@
+// Copyright 1998-2019 Epic Games, Inc. All Rights Reserved.
+
+#include "Chaos/PBDCollisionConstraintsPointContactUtil.h"
+#include "Chaos/CollisionResolutionAlgo.h"
+#include "Chaos/Defines.h"
+
+namespace Chaos
+{
+	namespace Collisions
+	{
+		template<typename T=float, int d=3>
+		TRigidTransform<T, d> GetTransform(const TGeometryParticleHandle<T, d>* Particle)
+		{
+			TGenericParticleHandle<T, d> Generic = const_cast<TGeometryParticleHandle<T, d>*>(Particle);
+			return TRigidTransform<T, d>(Generic->P(), Generic->Q());
+		}
+
+		// @todo(ccaulfield): This is duplicated in JointConstraints - move to a utility file
+		template<typename T>
+		PMatrix<T, 3, 3> ComputeFactorMatrix3(const TVector<T, 3>& V, const PMatrix<T, 3, 3>& M, const T& Im)
+		{
+			// Rigid objects rotational contribution to the impulse.
+			// Vx*M*VxT+Im
+			check(Im > FLT_MIN)
+				return PMatrix<T, 3, 3>(
+					-V[2] * (-V[2] * M.M[1][1] + V[1] * M.M[2][1]) + V[1] * (-V[2] * M.M[2][1] + V[1] * M.M[2][2]) + Im,
+					V[2] * (-V[2] * M.M[1][0] + V[1] * M.M[2][0]) - V[0] * (-V[2] * M.M[2][1] + V[1] * M.M[2][2]),
+					-V[1] * (-V[2] * M.M[1][0] + V[1] * M.M[2][0]) + V[0] * (-V[2] * M.M[1][1] + V[1] * M.M[2][1]),
+					V[2] * (V[2] * M.M[0][0] - V[0] * M.M[2][0]) - V[0] * (V[2] * M.M[2][0] - V[0] * M.M[2][2]) + Im,
+					-V[1] * (V[2] * M.M[0][0] - V[0] * M.M[2][0]) + V[0] * (V[2] * M.M[1][0] - V[0] * M.M[2][1]),
+					-V[1] * (-V[1] * M.M[0][0] + V[0] * M.M[1][0]) + V[0] * (-V[1] * M.M[1][0] + V[0] * M.M[1][1]) + Im);
+		}
+
+		template<typename T, int d>
+		TVector<T, d> GetEnergyClampedImpulse(const TRigidBodyPointContactConstraint<T, d>& Constraint, const TVector<T, d>& Impulse, const TVector<T, d>& VectorToPoint1, const TVector<T, d>& VectorToPoint2, const TVector<T, d>& Velocity1, const TVector<T, d>& Velocity2)
+		{
+			TPBDRigidParticleHandle<T, d>* PBDRigid0 = Constraint.Particle[0]->AsDynamic();
+			TPBDRigidParticleHandle<T, d>* PBDRigid1 = Constraint.Particle[1]->AsDynamic();
+
+			TVector<T, d> Jr0, Jr1, IInvJr0, IInvJr1;
+			T ImpulseRatioNumerator0 = 0, ImpulseRatioNumerator1 = 0, ImpulseRatioDenom0 = 0, ImpulseRatioDenom1 = 0;
+			T ImpulseSize = Impulse.SizeSquared();
+			TVector<T, d> KinematicVelocity = !PBDRigid0 ? Velocity1 : !PBDRigid1 ? Velocity2 : TVector<T, d>(0);
+			if (PBDRigid0)
+			{
+				Jr0 = TVector<T, d>::CrossProduct(VectorToPoint1, Impulse);
+				IInvJr0 = PBDRigid0->Q().RotateVector(PBDRigid0->InvI() * PBDRigid0->Q().UnrotateVector(Jr0));
+				ImpulseRatioNumerator0 = TVector<T, d>::DotProduct(Impulse, PBDRigid0->V() - KinematicVelocity) + TVector<T, d>::DotProduct(IInvJr0, PBDRigid0->W());
+				ImpulseRatioDenom0 = ImpulseSize / PBDRigid0->M() + TVector<T, d>::DotProduct(Jr0, IInvJr0);
+			}
+			if (PBDRigid1)
+			{
+				Jr1 = TVector<T, d>::CrossProduct(VectorToPoint2, Impulse);
+				IInvJr1 = PBDRigid1->Q().RotateVector(PBDRigid1->InvI() * PBDRigid1->Q().UnrotateVector(Jr1));
+				ImpulseRatioNumerator1 = TVector<T, d>::DotProduct(Impulse, PBDRigid1->V() - KinematicVelocity) + TVector<T, d>::DotProduct(IInvJr1, PBDRigid1->W());
+				ImpulseRatioDenom1 = ImpulseSize / PBDRigid1->M() + TVector<T, d>::DotProduct(Jr1, IInvJr1);
+			}
+			T Numerator = -2 * (ImpulseRatioNumerator0 - ImpulseRatioNumerator1);
+			if (Numerator < 0)
+			{
+				return TVector<T, d>(0);
+			}
+			check(Numerator >= 0);
+			T Denominator = ImpulseRatioDenom0 + ImpulseRatioDenom1;
+			return Numerator < Denominator ? (Impulse * Numerator / Denominator) : Impulse;
+		}
+
+
+		template<ECollisionUpdateType UpdateType, typename T, int d>
+		void Update(const T Thickness, TRigidBodyPointContactConstraint<T, d>& Constraint)
+		{
+			Constraint.ResetPhi(Thickness);
+			const TRigidTransform<T, d> ParticleTM = GetTransform(Constraint.Particle[0]);
+			const TRigidTransform<T, d> LevelsetTM = GetTransform(Constraint.Particle[1]);
+
+			if (!Constraint.Particle[0]->Geometry())
+			{
+				if (Constraint.Particle[1]->Geometry())
+				{
+					if (!Constraint.Particle[1]->Geometry()->IsUnderlyingUnion())
+					{
+						UpdateLevelsetConstraint<UpdateType>(Thickness, Constraint);
+					}
+					else
+					{
+						UpdateUnionLevelsetConstraint<UpdateType>(Thickness, Constraint);
+					}
+				}
+			}
+			else
+			{
+				UpdateConstraintImp<UpdateType>(*Constraint.Particle[0]->Geometry(), ParticleTM, *Constraint.Particle[1]->Geometry(), LevelsetTM, Thickness, Constraint);
+			}
+		}
+		template void Update<ECollisionUpdateType::Any, float, 3>(const float, TRigidBodyPointContactConstraint<float, 3>&);
+		template void Update<ECollisionUpdateType::Deepest, float, 3>(const float, TRigidBodyPointContactConstraint<float, 3>&);
+
+
+		template<ECollisionUpdateType UpdateType, typename T, int d>
+		void UpdateConstraint(const T Thickness, TRigidBodyPointContactConstraint<T, d>& Constraint)
+		{
+			Constraint.ResetPhi(Thickness);
+			const TRigidTransform<T, d> ParticleTM = GetTransform(Constraint.Particle[0]);
+			const TRigidTransform<T, d> LevelsetTM = GetTransform(Constraint.Particle[1]);
+
+			if (!Constraint.Particle[0]->Geometry())
+			{
+				if (Constraint.Particle[1]->Geometry())
+				{
+					if (!Constraint.Particle[1]->Geometry()->IsUnderlyingUnion())
+					{
+						UpdateLevelsetConstraint<UpdateType>(Thickness, Constraint);
+					}
+					else
+					{
+						UpdateUnionLevelsetConstraint<UpdateType>(Thickness, Constraint);
+					}
+				}
+			}
+			else
+			{
+				UpdateConstraintImp<UpdateType>(*Constraint.Particle[0]->Geometry(), ParticleTM, *Constraint.Particle[1]->Geometry(), LevelsetTM, Thickness, Constraint);
+			}
+		}
+		template void UpdateConstraint<ECollisionUpdateType::Any, float, 3>(const float, TRigidBodyPointContactConstraint<float,3>&);
+		template void UpdateConstraint<ECollisionUpdateType::Deepest, float, 3>(const float, TRigidBodyPointContactConstraint<float,3>&);
+
+
+		template<typename T, int d>
+		void Apply(TRigidBodyPointContactConstraint<T, d>& Constraint, T Thickness, TPointContactIterationParameters<T> & IterationParameters, TPointContactParticleParameters<T> & ParticleParameters)
+		{
+			TGenericParticleHandle<T, d> Particle0 = TGenericParticleHandle<T, d>(Constraint.Particle[0]);
+			TGenericParticleHandle<T, d> Particle1 = TGenericParticleHandle<T, d>(Constraint.Particle[1]);
+			TPBDRigidParticleHandle<T, d>* PBDRigid0 = Particle0->AsDynamic();
+			TPBDRigidParticleHandle<T, d>* PBDRigid1 = Particle1->AsDynamic();
+
+			for (int32 PairIt = 0; PairIt < IterationParameters.NumPairIterations; ++PairIt)
+			{
+				Collisions::UpdateConstraint<ECollisionUpdateType::Deepest>(Thickness, Constraint);
+				if (Constraint.GetPhi() >= Thickness)
+				{
+					return;
+				}
+
+				// @todo(ccaulfield): CHAOS_PARTICLEHANDLE_TODO what's the best way to manage external per-particle data?
+				if (ParticleParameters.Collided)
+				{
+					Particle0->AuxilaryValue(*ParticleParameters.Collided) = true;
+					Particle1->AuxilaryValue(*ParticleParameters.Collided) = true;
+				}
+
+				TSerializablePtr<FChaosPhysicsMaterial> PhysicsMaterial0, PhysicsMaterial1;
+				if (ParticleParameters.PhysicsMaterials)
+				{
+					PhysicsMaterial0 = Particle0->AuxilaryValue(*ParticleParameters.PhysicsMaterials);
+					PhysicsMaterial1 = Particle1->AuxilaryValue(*ParticleParameters.PhysicsMaterials);
+				}
+
+				// @todo(ccaulfield): CHAOS_PARTICLEHANDLE_TODO split function to avoid ifs
+				const TVector<T, d> ZeroVector = TVector<T, d>(0);
+				const TRotation<T, d>& Q0 = Particle0->Q();
+				const TRotation<T, d>& Q1 = Particle1->Q();
+				const TVector<T, d>& P0 = Particle0->P();
+				const TVector<T, d>& P1 = Particle1->P();
+				const TVector<T, d>& V0 = Particle0->V();
+				const TVector<T, d>& V1 = Particle1->V();
+				const TVector<T, d>& W0 = Particle0->W();
+				const TVector<T, d>& W1 = Particle1->W();
+
+				typename TRigidBodyPointContactConstraint<T, d>::FManifold & Contact = Constraint.Manifold;
+
+				TVector<T, d> VectorToPoint1 = Contact.Location - P0;
+				TVector<T, d> VectorToPoint2 = Contact.Location - P1;
+				TVector<T, d> Body1Velocity = V0 + TVector<T, d>::CrossProduct(W0, VectorToPoint1);
+				TVector<T, d> Body2Velocity = V1 + TVector<T, d>::CrossProduct(W1, VectorToPoint2);
+				TVector<T, d> RelativeVelocity = Body1Velocity - Body2Velocity;
+				T RelativeNormalVelocity = TVector<T, d>::DotProduct(RelativeVelocity, Contact.Normal);
+
+				if (RelativeNormalVelocity < 0) // ignore separating constraints
+				{
+					PMatrix<T, d, d> WorldSpaceInvI1 = PBDRigid0 ? (Q0 * FMatrix::Identity).GetTransposed() * PBDRigid0->InvI() * (Q0 * FMatrix::Identity) : PMatrix<T, d, d>(0);
+					PMatrix<T, d, d> WorldSpaceInvI2 = PBDRigid1 ? (Q1 * FMatrix::Identity).GetTransposed() * PBDRigid1->InvI() * (Q1 * FMatrix::Identity) : PMatrix<T, d, d>(0);
+					PMatrix<T, d, d> Factor =
+						(PBDRigid0 ? ComputeFactorMatrix3(VectorToPoint1, WorldSpaceInvI1, PBDRigid0->InvM()) : PMatrix<T, d, d>(0)) +
+						(PBDRigid1 ? ComputeFactorMatrix3(VectorToPoint2, WorldSpaceInvI2, PBDRigid1->InvM()) : PMatrix<T, d, d>(0));
+					TVector<T, d> Impulse;
+					TVector<T, d> AngularImpulse(0);
+
+					// Resting contact if very close to the surface
+					T Restitution = (T)0;
+					T Friction = (T)0;
+					T AngularFriction = (T)0;
+					bool bApplyRestitution = (RelativeVelocity.Size() > (2 * 980 * IterationParameters.Dt));
+					if (PhysicsMaterial0 && PhysicsMaterial1)
+					{
+						if (bApplyRestitution)
+						{
+							Restitution = FMath::Min(PhysicsMaterial0->Restitution, PhysicsMaterial1->Restitution);
+						}
+						Friction = FMath::Max(PhysicsMaterial0->Friction, PhysicsMaterial1->Friction);
+					}
+					else if (PhysicsMaterial0)
+					{
+						if (bApplyRestitution)
+						{
+							Restitution = PhysicsMaterial0->Restitution;
+						}
+						Friction = PhysicsMaterial0->Friction;
+					}
+					else if (PhysicsMaterial1)
+					{
+						if (bApplyRestitution)
+						{
+							Restitution = PhysicsMaterial1->Restitution;
+						}
+						Friction = PhysicsMaterial1->Friction;
+					}
+
+					if (ParticleParameters.FrictionOverride >= 0)
+					{
+						Friction = ParticleParameters.FrictionOverride;
+					}
+					if (ParticleParameters.AngularFrictionOverride >= 0)
+					{
+						AngularFriction = ParticleParameters.AngularFrictionOverride;
+					}
+
+					if (Friction)
+					{
+						if (RelativeNormalVelocity > 0)
+						{
+							RelativeNormalVelocity = 0;
+						}
+						TVector<T, d> VelocityChange = -(Restitution * RelativeNormalVelocity * Contact.Normal + RelativeVelocity);
+						T NormalVelocityChange = TVector<T, d>::DotProduct(VelocityChange, Contact.Normal);
+						PMatrix<T, d, d> FactorInverse = Factor.Inverse();
+						TVector<T, d> MinimalImpulse = FactorInverse * VelocityChange;
+						const T MinimalImpulseDotNormal = TVector<T, d>::DotProduct(MinimalImpulse, Contact.Normal);
+						const T TangentialSize = (MinimalImpulse - MinimalImpulseDotNormal * Contact.Normal).Size();
+						if (TangentialSize <= Friction * MinimalImpulseDotNormal)
+						{
+							//within friction cone so just solve for static friction stopping the object
+							Impulse = MinimalImpulse;
+							if (AngularFriction)
+							{
+								TVector<T, d> RelativeAngularVelocity = W0 - W1;
+								T AngularNormal = TVector<T, d>::DotProduct(RelativeAngularVelocity, Contact.Normal);
+								TVector<T, d> AngularTangent = RelativeAngularVelocity - AngularNormal * Contact.Normal;
+								TVector<T, d> FinalAngularVelocity = FMath::Sign(AngularNormal) * FMath::Max((T)0, FMath::Abs(AngularNormal) - AngularFriction * NormalVelocityChange) * Contact.Normal + FMath::Max((T)0, AngularTangent.Size() - AngularFriction * NormalVelocityChange) * AngularTangent.GetSafeNormal();
+								TVector<T, d> Delta = FinalAngularVelocity - RelativeAngularVelocity;
+								if (!PBDRigid0 && PBDRigid1)
+								{
+									PMatrix<T, d, d> WorldSpaceI2 = (Q1 * FMatrix::Identity) * PBDRigid1->I() * (Q1 * FMatrix::Identity).GetTransposed();
+									TVector<T, d> ImpulseDelta = PBDRigid1->M() * TVector<T, d>::CrossProduct(VectorToPoint2, Delta);
+									Impulse += ImpulseDelta;
+									AngularImpulse += WorldSpaceI2 * Delta - TVector<T, d>::CrossProduct(VectorToPoint2, ImpulseDelta);
+								}
+								else if (PBDRigid0 && !PBDRigid1)
+								{
+									PMatrix<T, d, d> WorldSpaceI1 = (Q0 * FMatrix::Identity) * PBDRigid0->I() * (Q0 * FMatrix::Identity).GetTransposed();
+									TVector<T, d> ImpulseDelta = PBDRigid0->M() * TVector<T, d>::CrossProduct(VectorToPoint1, Delta);
+									Impulse += ImpulseDelta;
+									AngularImpulse += WorldSpaceI1 * Delta - TVector<T, d>::CrossProduct(VectorToPoint1, ImpulseDelta);
+								}
+								else if (PBDRigid0 && PBDRigid1)
+								{
+									PMatrix<T, d, d> Cross1(0, VectorToPoint1.Z, -VectorToPoint1.Y, -VectorToPoint1.Z, 0, VectorToPoint1.X, VectorToPoint1.Y, -VectorToPoint1.X, 0);
+									PMatrix<T, d, d> Cross2(0, VectorToPoint2.Z, -VectorToPoint2.Y, -VectorToPoint2.Z, 0, VectorToPoint2.X, VectorToPoint2.Y, -VectorToPoint2.X, 0);
+									PMatrix<T, d, d> CrossI1 = Cross1 * WorldSpaceInvI1;
+									PMatrix<T, d, d> CrossI2 = Cross2 * WorldSpaceInvI2;
+									PMatrix<T, d, d> Diag1 = CrossI1 * Cross1.GetTransposed() + CrossI2 * Cross2.GetTransposed();
+									Diag1.M[0][0] += PBDRigid0->InvM() + PBDRigid1->InvM();
+									Diag1.M[1][1] += PBDRigid0->InvM() + PBDRigid1->InvM();
+									Diag1.M[2][2] += PBDRigid0->InvM() + PBDRigid1->InvM();
+									PMatrix<T, d, d> OffDiag1 = (CrossI1 + CrossI2) * -1;
+									PMatrix<T, d, d> Diag2 = (WorldSpaceInvI1 + WorldSpaceInvI2).Inverse();
+									PMatrix<T, d, d> OffDiag1Diag2 = OffDiag1 * Diag2;
+									TVector<T, d> ImpulseDelta = PMatrix<T, d, d>((Diag1 - OffDiag1Diag2 * OffDiag1.GetTransposed()).Inverse())* ((OffDiag1Diag2 * -1) * Delta);
+									Impulse += ImpulseDelta;
+									AngularImpulse += Diag2 * (Delta - PMatrix<T, d, d>(OffDiag1.GetTransposed()) * ImpulseDelta);
+								}
+							}
+						}
+						else
+						{
+							//outside friction cone, solve for normal relative velocity and keep tangent at cone edge
+							TVector<T, d> Tangent = (RelativeVelocity - TVector<T, d>::DotProduct(RelativeVelocity, Contact.Normal) * Contact.Normal).GetSafeNormal();
+							TVector<T, d> DirectionalFactor = Factor * (Contact.Normal - Friction * Tangent);
+							T ImpulseDenominator = TVector<T, d>::DotProduct(Contact.Normal, DirectionalFactor);
+							if (!ensureMsgf(FMath::Abs(ImpulseDenominator) > SMALL_NUMBER, TEXT("Constraint:%s\n\nParticle:%s\n\nLevelset:%s\n\nDirectionalFactor:%s, ImpulseDenominator:%f"),
+								*Constraint.ToString(),
+								*Particle0->ToString(),
+								*Particle1->ToString(),
+								*DirectionalFactor.ToString(), ImpulseDenominator))
+							{
+								ImpulseDenominator = (T)1;
+							}
+
+							const T ImpulseMag = -(1 + Restitution) * RelativeNormalVelocity / ImpulseDenominator;
+							Impulse = ImpulseMag * (Contact.Normal - Friction * Tangent);
+						}
+					}
+					else
+					{
+						T ImpulseDenominator = TVector<T, d>::DotProduct(Contact.Normal, Factor * Contact.Normal);
+						TVector<T, d> ImpulseNumerator = -(1 + Restitution) * TVector<T, d>::DotProduct(RelativeVelocity, Contact.Normal)* Contact.Normal;
+						if (!ensureMsgf(FMath::Abs(ImpulseDenominator) > SMALL_NUMBER, TEXT("Constraint:%s\n\nParticle:%s\n\nLevelset:%s\n\nFactor*Constraint.Normal:%s, ImpulseDenominator:%f"),
+							*Constraint.ToString(),
+							*Particle0->ToString(),
+							*Particle1->ToString(),
+							*(Factor * Contact.Normal).ToString(), ImpulseDenominator))
+						{
+							ImpulseDenominator = (T)1;
+						}
+						Impulse = ImpulseNumerator / ImpulseDenominator;
+					}
+					Impulse = GetEnergyClampedImpulse(Constraint, Impulse, VectorToPoint1, VectorToPoint2, Body1Velocity, Body2Velocity);
+					Constraint.AccumulatedImpulse += Impulse;
+					if (PBDRigid0)
+					{
+						// Velocity update for next step
+						FVec3 NetAngularImpulse = FVec3::CrossProduct(VectorToPoint1, Impulse) + AngularImpulse;
+						FVec3 DV = PBDRigid0->InvM() * Impulse;
+						FVec3 DW = WorldSpaceInvI1 * NetAngularImpulse;
+						PBDRigid0->V() += DV;
+						PBDRigid0->W() += DW;
+						// Position update as part of pbd
+						PBDRigid0->P() += DV * IterationParameters.Dt;
+						PBDRigid0->Q() += FRotation3::FromElements(DW, 0.f) * Q0 * IterationParameters.Dt * T(0.5);
+						PBDRigid0->Q().Normalize();
+					}
+					if (PBDRigid1)
+					{
+						// Velocity update for next step
+						FVec3 NetAngularImpulse = FVec3::CrossProduct(VectorToPoint2, -Impulse) - AngularImpulse;
+						FVec3 DV = -PBDRigid1->InvM() * Impulse;
+						FVec3 DW = WorldSpaceInvI2 * NetAngularImpulse;
+						PBDRigid1->V() += DV;
+						PBDRigid1->W() += DW;
+						// Position update as part of pbd
+						PBDRigid1->P() += DV * IterationParameters.Dt;
+						PBDRigid1->Q() += FRotation3::FromElements(DW, 0.f) * Q1 * IterationParameters.Dt * T(0.5);
+						PBDRigid1->Q().Normalize();
+					}
+				}
+			}
+		}
+		template void Apply<float, 3>(TRigidBodyPointContactConstraint<float, 3>&, float, TPointContactIterationParameters<float> &, TPointContactParticleParameters<float> &);
+
+		template<typename T, int d>
+		void ApplyPushOut(TRigidBodyPointContactConstraint<T, d>& Constraint, T Thickness, const TSet<const TGeometryParticleHandle<T, d>*>& IsTemporarilyStatic,
+			TPointContactIterationParameters<T> & IterationParameters, TPointContactParticleParameters<T> & ParticleParameters)
+		{
+			TGeometryParticleHandle<T, d>* Particle0 = Constraint.Particle[0];
+			TGeometryParticleHandle<T, d>* Particle1 = Constraint.Particle[1];
+			TPBDRigidParticleHandle<T, d>* PBDRigid0 = Particle0->AsDynamic();
+			TPBDRigidParticleHandle<T, d>* PBDRigid1 = Particle1->AsDynamic();
+
+			const TVector<T, d> ZeroVector = TVector<T, d>(0);
+			const TRotation<T, d>& Q0 = PBDRigid0 ? PBDRigid0->Q() : Particle0->R();
+			const TRotation<T, d>& Q1 = PBDRigid1 ? PBDRigid1->Q() : Particle1->R();
+			const TVector<T, d>& P0 = PBDRigid0 ? PBDRigid0->P() : Particle0->X();
+			const TVector<T, d>& P1 = PBDRigid1 ? PBDRigid1->P() : Particle1->X();
+			const TVector<T, d>& V0 = PBDRigid0 ? PBDRigid0->V() : ZeroVector;
+			const TVector<T, d>& V1 = PBDRigid1 ? PBDRigid1->V() : ZeroVector;
+			const TVector<T, d>& W0 = PBDRigid0 ? PBDRigid0->W() : ZeroVector;
+			const TVector<T, d>& W1 = PBDRigid1 ? PBDRigid1->W() : ZeroVector;
+			const bool IsTemporarilyStatic0 = IsTemporarilyStatic.Contains(Particle0);
+			const bool IsTemporarilyStatic1 = IsTemporarilyStatic.Contains(Particle1);
+	
+			TSerializablePtr<FChaosPhysicsMaterial> PhysicsMaterial0, PhysicsMaterial1;
+			if (ParticleParameters.PhysicsMaterials)
+			{
+				PhysicsMaterial0 = Particle0->AuxilaryValue(*ParticleParameters.PhysicsMaterials);
+				PhysicsMaterial1 = Particle1->AuxilaryValue(*ParticleParameters.PhysicsMaterials);
+			}
+
+			for (int32 PairIteration = 0; PairIteration < IterationParameters.NumPairIterations; ++PairIteration)
+			{
+				UpdateConstraint<ECollisionUpdateType::Deepest>(Thickness, Constraint);
+
+				const typename TRigidBodyPointContactConstraint<T, d>::FManifold & Contact = Constraint.Manifold;
+
+				if (Contact.Phi >= Thickness)
+				{
+					break;
+				}
+
+				if ((!PBDRigid0 || IsTemporarilyStatic0) && (!PBDRigid1 || IsTemporarilyStatic1))
+				{
+					break;
+				}
+
+				*IterationParameters.NeedsAnotherIteration = true;
+				PMatrix<T, d, d> WorldSpaceInvI1 = PBDRigid0 ? (Q0 * FMatrix::Identity).GetTransposed() * PBDRigid0->InvI() * (Q0 * FMatrix::Identity) : PMatrix<T, d, d>(0);
+				PMatrix<T, d, d> WorldSpaceInvI2 = PBDRigid1 ? (Q1 * FMatrix::Identity).GetTransposed() * PBDRigid1->InvI() * (Q1 * FMatrix::Identity) : PMatrix<T, d, d>(0);
+				TVector<T, d> VectorToPoint1 = Contact.Location - P0;
+				TVector<T, d> VectorToPoint2 = Contact.Location - P1;
+				PMatrix<T, d, d> Factor =
+					(PBDRigid0 ? ComputeFactorMatrix3(VectorToPoint1, WorldSpaceInvI1, PBDRigid0->InvM()) : PMatrix<T, d, d>(0)) +
+					(PBDRigid1 ? ComputeFactorMatrix3(VectorToPoint2, WorldSpaceInvI2, PBDRigid1->InvM()) : PMatrix<T, d, d>(0));
+				T Numerator = FMath::Min((T)(IterationParameters.Iteration + 2), (T)IterationParameters.NumIterations);
+				T ScalingFactor = Numerator / (T)IterationParameters.NumIterations;
+
+				//if pushout is needed we better fix relative velocity along normal. Treat it as if 0 restitution
+				TVector<T, d> Body1Velocity = V0 + TVector<T, d>::CrossProduct(W0, VectorToPoint1);
+				TVector<T, d> Body2Velocity = V1 + TVector<T, d>::CrossProduct(W1, VectorToPoint2);
+				TVector<T, d> RelativeVelocity = Body1Velocity - Body2Velocity;
+				const T RelativeVelocityDotNormal = TVector<T, d>::DotProduct(RelativeVelocity, Contact.Normal);
+				if (RelativeVelocityDotNormal < 0)
+				{
+					T ImpulseDenominator = TVector<T, d>::DotProduct(Contact.Normal, Factor * Contact.Normal);
+					TVector<T, d> ImpulseNumerator = -TVector<T, d>::DotProduct(RelativeVelocity, Contact.Normal) * Contact.Normal * ScalingFactor;
+					if (!ensureMsgf(FMath::Abs(ImpulseDenominator) > SMALL_NUMBER, TEXT("ApplyPushout Constraint:%s\n\nParticle:%s\n\nLevelset:%s\n\nFactor*Contact.Normal:%s, ImpulseDenominator:%f"),
+						*Constraint.ToString(),
+						*Particle0->ToString(),
+						*Particle1->ToString(),
+						*(Factor*Contact.Normal).ToString(), ImpulseDenominator))
+					{
+						ImpulseDenominator = (T)1;
+					}
+
+					TVector<T, d> VelocityFixImpulse = ImpulseNumerator / ImpulseDenominator;
+					VelocityFixImpulse = GetEnergyClampedImpulse(Constraint, VelocityFixImpulse, VectorToPoint1, VectorToPoint2, Body1Velocity, Body2Velocity);
+					Constraint.AccumulatedImpulse += VelocityFixImpulse;	//question: should we track this?
+					if (!IsTemporarilyStatic0 && PBDRigid0)
+					{
+						TVector<T, d> AngularImpulse = TVector<T, d>::CrossProduct(VectorToPoint1, VelocityFixImpulse);
+						PBDRigid0->V() += PBDRigid0->InvM() * VelocityFixImpulse;
+						PBDRigid0->W() += WorldSpaceInvI1 * AngularImpulse;
+
+					}
+
+					if (!IsTemporarilyStatic1 && PBDRigid1)
+					{
+						TVector<T, d> AngularImpulse = TVector<T, d>::CrossProduct(VectorToPoint2, -VelocityFixImpulse);
+						PBDRigid1->V() -= PBDRigid1->InvM() * VelocityFixImpulse;
+						PBDRigid1->W() += WorldSpaceInvI2 * AngularImpulse;
+					}
+
+				}
+
+
+				TVector<T, d> Impulse = PMatrix<T, d, d>(Factor.Inverse()) * ((-Contact.Phi + Thickness) * ScalingFactor * Contact.Normal);
+				TVector<T, d> AngularImpulse1 = TVector<T, d>::CrossProduct(VectorToPoint1, Impulse);
+				TVector<T, d> AngularImpulse2 = TVector<T, d>::CrossProduct(VectorToPoint2, -Impulse);
+				if (!IsTemporarilyStatic0 && PBDRigid0)
+				{
+					PBDRigid0->P() += PBDRigid0->InvM() * Impulse;
+					PBDRigid0->Q() = TRotation<T, d>::FromVector(WorldSpaceInvI1 * AngularImpulse1) * Q0;
+					PBDRigid0->Q().Normalize();
+				}
+				if (!IsTemporarilyStatic1 && PBDRigid1)
+				{
+					PBDRigid1->P() -= PBDRigid1->InvM() * Impulse;
+					PBDRigid1->Q() = TRotation<T, d>::FromVector(WorldSpaceInvI2 * AngularImpulse2) * Q1;
+					PBDRigid1->Q().Normalize();
+				}
+			}
+		}
+		template void ApplyPushOut<float, 3>(TRigidBodyPointContactConstraint<float,3>&, float , const TSet<const TGeometryParticleHandle<float,3>*>&,
+			TPointContactIterationParameters<float> & IterationParameters, TPointContactParticleParameters<float> & ParticleParameters);
+
+	} // Collisions
+
+}// Chaos
+
diff --git a/Engine/Source/Runtime/Experimental/Chaos/Private/Chaos/PBDConstraintRule.cpp b/Engine/Source/Runtime/Experimental/Chaos/Private/Chaos/PBDConstraintRule.cpp
--- a/Engine/Source/Runtime/Experimental/Chaos/Private/Chaos/PBDConstraintRule.cpp
+++ b/Engine/Source/Runtime/Experimental/Chaos/Private/Chaos/PBDConstraintRule.cpp
@@ -3,7 +3,7 @@
 #include "Chaos/PBDConstraintRule.h"
 
 #include "Chaos/PBD6DJointConstraints.h"
-#include "Chaos/PBDCollisionConstraint.h"
+#include "Chaos/PBDCollisionConstraints.h"
 #include "Chaos/PBDJointConstraints.h"
 #include "Chaos/PBDPositionConstraints.h"
 #include "Chaos/PBDRigidDynamicSpringConstraints.h"
@@ -53,14 +53,14 @@ namespace Chaos
 		return Constraints.NumConstraints(); 
 	}
 
-	template class TPBDConstraintGraphRuleImpl< TPBDCollisionConstraint<float, 3> >;
+	template class TPBDConstraintGraphRuleImpl< TPBDCollisionConstraints<float, 3> >;
 	template class TPBDConstraintGraphRuleImpl< FPBD6DJointConstraints >;
 	template class TPBDConstraintGraphRuleImpl< FPBDJointConstraints >;
 	template class TPBDConstraintGraphRuleImpl< TPBDPositionConstraints<float, 3> >;
 	template class TPBDConstraintGraphRuleImpl< TPBDRigidDynamicSpringConstraints<float, 3> >;
 	template class TPBDConstraintGraphRuleImpl< TPBDRigidSpringConstraints<float, 3> >;
 
-	template class TPBDConstraintColorRule<TPBDCollisionConstraint<float, 3>>;
+	template class TPBDConstraintColorRule<TPBDCollisionConstraints<float, 3>>;
 	template class TPBDConstraintIslandRule<FPBD6DJointConstraints>;
 	template class TPBDConstraintIslandRule<FPBDJointConstraints>;
 	template class TPBDConstraintIslandRule<TPBDPositionConstraints<float, 3>>;
diff --git a/Engine/Source/Runtime/Experimental/Chaos/Private/Chaos/PBDRigidClustering.cpp b/Engine/Source/Runtime/Experimental/Chaos/Private/Chaos/PBDRigidClustering.cpp
--- a/Engine/Source/Runtime/Experimental/Chaos/Private/Chaos/PBDRigidClustering.cpp
+++ b/Engine/Source/Runtime/Experimental/Chaos/Private/Chaos/PBDRigidClustering.cpp
@@ -9,8 +9,8 @@
 #include "Chaos/MassProperties.h"
 #include "Chaos/PBDRigidsEvolution.h"
 #include "Chaos/PBDRigidsEvolutionPGS.h"
-#include "Chaos/PBDCollisionConstraint.h"
-#include "Chaos/PBDCollisionConstraintPGS.h"
+#include "Chaos/PBDCollisionConstraints.h"
+#include "Chaos/PBDCollisionConstraintsPGS.h"
 #include "Chaos/Sphere.h"
 #include "Chaos/UniformGrid.h"
 #include "ChaosStats.h"
@@ -702,7 +702,7 @@ namespace Chaos
 
 	DECLARE_CYCLE_STAT(TEXT("TPBDRigidClustering<>::RewindAndEvolve<BGF>()"), STAT_RewindAndEvolve_BGF, STATGROUP_Chaos);
 	template<class T, int d>
-	void RewindAndEvolve(TPBDRigidsEvolutionGBF<T, d>& Evolution, TPBDRigidClusteredParticles<T, d>& InParticles, const TSet<int32>& IslandsToRecollide, const TSet<uint32>& AllActivatedChildren, const T Dt, TPBDCollisionConstraint<T, d>& CollisionRule)
+	void RewindAndEvolve(TPBDRigidsEvolutionGBF<T, d>& Evolution, TPBDRigidClusteredParticles<T, d>& InParticles, const TSet<int32>& IslandsToRecollide, const TSet<uint32>& AllActivatedChildren, const T Dt, TPBDCollisionConstraints<T, d>& CollisionRule)
 	{
 		SCOPE_CYCLE_COUNTER(STAT_RewindAndEvolve_BGF);
 		// Rewind active particles
@@ -1913,7 +1913,8 @@ namespace Chaos
 }
 
 using namespace Chaos;
-template class CHAOS_API Chaos::TPBDRigidClustering<TPBDRigidsEvolutionGBF<float, 3>, TPBDCollisionConstraint<float, 3>, float, 3>;
+template class Chaos::TPBDRigidClustering<TPBDRigidsEvolutionGBF<float, 3>, TPBDCollisionConstraints<float, 3>, float, 3>;
+template class CHAOS_API Chaos::TPBDRigidClustering<TPBDRigidsEvolutionGBF<float, 3>, TPBDCollisionConstraints<float, 3>, float, 3>;
 #if CHAOS_PARTICLEHANDLE_TODO
 template class CHAOS_API Chaos::TPBDRigidClustering<TPBDRigidsEvolutionPGS<float, 3>, TPBDCollisionConstraintPGS<float, 3>, float, 3>;
 #endif
diff --git a/Engine/Source/Runtime/Experimental/Chaos/Private/Chaos/PBDRigidsEvolution.cpp b/Engine/Source/Runtime/Experimental/Chaos/Private/Chaos/PBDRigidsEvolution.cpp
--- a/Engine/Source/Runtime/Experimental/Chaos/Private/Chaos/PBDRigidsEvolution.cpp
+++ b/Engine/Source/Runtime/Experimental/Chaos/Private/Chaos/PBDRigidsEvolution.cpp
@@ -598,8 +598,9 @@ namespace Chaos
 	}
 }
 
+template class Chaos::TPBDRigidsEvolutionBase<Chaos::TPBDRigidsEvolutionGBF<float, 3>, Chaos::TPBDCollisionConstraints<float, 3>, float, 3>;
 #ifdef __clang__
-template class CHAOS_API Chaos::TPBDRigidsEvolutionBase<Chaos::TPBDRigidsEvolutionGBF<float, 3>, Chaos::TPBDCollisionConstraint<float,3>, float, 3>;
+template class CHAOS_API Chaos::TPBDRigidsEvolutionBase<Chaos::TPBDRigidsEvolutionGBF<float, 3>, Chaos::TPBDCollisionConstraints<float,3>, float, 3>;
 #else
-template class Chaos::TPBDRigidsEvolutionBase<Chaos::TPBDRigidsEvolutionGBF<float, 3>, Chaos::TPBDCollisionConstraint<float,3>, float, 3>;
+template class Chaos::TPBDRigidsEvolutionBase<Chaos::TPBDRigidsEvolutionGBF<float, 3>, Chaos::TPBDCollisionConstraints<float,3>, float, 3>;
 #endif
diff --git a/Engine/Source/Runtime/Experimental/Chaos/Private/Chaos/PBDRigidsEvolutionGBF.cpp b/Engine/Source/Runtime/Experimental/Chaos/Private/Chaos/PBDRigidsEvolutionGBF.cpp
--- a/Engine/Source/Runtime/Experimental/Chaos/Private/Chaos/PBDRigidsEvolutionGBF.cpp
+++ b/Engine/Source/Runtime/Experimental/Chaos/Private/Chaos/PBDRigidsEvolutionGBF.cpp
@@ -4,7 +4,7 @@
 #include "Chaos/Framework/Parallel.h"
 #include "Chaos/ImplicitObjectTransformed.h"
 #include "Chaos/ImplicitObjectUnion.h"
-#include "Chaos/PBDCollisionConstraint.h"
+#include "Chaos/PBDCollisionConstraints.h"
 #include "Chaos/PBDCollisionSpringConstraints.h"
 #include "Chaos/PerParticleEtherDrag.h"
 #include "Chaos/PerParticleEulerStepVelocity.h"
diff --git a/Engine/Source/Runtime/Experimental/Chaos/Private/Chaos/PBDRigidsEvolutionPGS.cpp b/Engine/Source/Runtime/Experimental/Chaos/Private/Chaos/PBDRigidsEvolutionPGS.cpp
--- a/Engine/Source/Runtime/Experimental/Chaos/Private/Chaos/PBDRigidsEvolutionPGS.cpp
+++ b/Engine/Source/Runtime/Experimental/Chaos/Private/Chaos/PBDRigidsEvolutionPGS.cpp
@@ -6,8 +6,8 @@
 #include "Chaos/Framework/Parallel.h"
 #include "Chaos/ImplicitObjectTransformed.h"
 #include "Chaos/ImplicitObjectUnion.h"
-#include "Chaos/PBDCollisionConstraint.h"
-#include "Chaos/PBDCollisionConstraintPGS.h"
+#include "Chaos/PBDCollisionConstraints.h"
+#include "Chaos/PBDCollisionConstraintsPGS.h"
 #include "Chaos/PBDCollisionSpringConstraints.h"
 #include "Chaos/PerParticleEtherDrag.h"
 #include "Chaos/PerParticleEulerStepVelocity.h"
diff --git a/Engine/Source/Runtime/Experimental/Chaos/Public/Chaos/ChaosDebugDraw.h b/Engine/Source/Runtime/Experimental/Chaos/Public/Chaos/ChaosDebugDraw.h
--- a/Engine/Source/Runtime/Experimental/Chaos/Public/Chaos/ChaosDebugDraw.h
+++ b/Engine/Source/Runtime/Experimental/Chaos/Public/Chaos/ChaosDebugDraw.h
@@ -26,8 +26,8 @@ namespace Chaos
 		CHAOS_API void DrawParticleShapes(const FRigidTransform3& SpaceTransform, const TArray<TGeometryParticleHandle<float, 3>*>& Particles, float ColorScale, bool bDrawKinematic = true, bool bDrawDynamic = true);
 		CHAOS_API void DrawParticleTransforms(const FRigidTransform3& SpaceTransform, const TParticleView<TGeometryParticles<float, 3>>& ParticlesView, float ColorScale, bool bDrawKinematic = true, bool bDrawDynamic = true);
 		CHAOS_API void DrawParticleTransforms(const FRigidTransform3& SpaceTransform, const TArray<TGeometryParticleHandle<float, 3>*>& Particles, float ColorScale, bool bDrawKinematic = true, bool bDrawDynamic = true);
-		CHAOS_API void DrawParticleCollisions(const FRigidTransform3& SpaceTransform, const TGeometryParticleHandle<float, 3>* Particle, const TPBDCollisionConstraint<float, 3>& Collisions);
-		CHAOS_API void DrawCollisions(const FRigidTransform3& SpaceTransform, const TPBDCollisionConstraint<float, 3>& Collisions, float ColorScale);
+		CHAOS_API void DrawParticleCollisions(const FRigidTransform3& SpaceTransform, const TGeometryParticleHandle<float, 3>* Particle, const TPBDCollisionConstraints<float, 3>& Collisions);
+		CHAOS_API void DrawCollisions(const FRigidTransform3& SpaceTransform, const TPBDCollisionConstraints<float, 3>& Collisions, float ColorScale);
 		CHAOS_API void DrawCollisions(const FRigidTransform3& SpaceTransform, const TArray<TPBDCollisionConstraintHandle<float, 3>*>& ConstraintHandles, float ColorScale);
 		CHAOS_API void DrawJointConstraints(const FRigidTransform3& SpaceTransform, const TArray<FPBDJointConstraintHandle*>& ConstraintHandles, float ColorScale, uint32 FeatureMask = (uint32)EDebugDrawJointFeature::Default);
 		CHAOS_API void DrawJointConstraints(const FRigidTransform3& SpaceTransform, const FPBDJointConstraints& Constraints, float ColorScale, uint32 FeatureMask = (uint32)EDebugDrawJointFeature::Default);
diff --git a/Engine/Source/Runtime/Experimental/Chaos/Public/Chaos/ChaosDebugDrawDeclares.h b/Engine/Source/Runtime/Experimental/Chaos/Public/Chaos/ChaosDebugDrawDeclares.h
--- a/Engine/Source/Runtime/Experimental/Chaos/Public/Chaos/ChaosDebugDrawDeclares.h
+++ b/Engine/Source/Runtime/Experimental/Chaos/Public/Chaos/ChaosDebugDrawDeclares.h
@@ -19,7 +19,7 @@ namespace Chaos
 	class FPBD6DJointConstraints;
 
 	template<typename T, int d>
-	class TPBDCollisionConstraint;
+	class TPBDCollisionConstraints;
 
 	template<typename T, int d>
 	class TPBDCollisionConstraintHandle;
diff --git a/Engine/Source/Runtime/Experimental/Chaos/Public/Chaos/CollisionResolutionConvexConvex.h b/Engine/Source/Runtime/Experimental/Chaos/Public/Chaos/CollisionResolutionConvexConvex.h
--- a/Engine/Source/Runtime/Experimental/Chaos/Public/Chaos/CollisionResolutionConvexConvex.h
+++ b/Engine/Source/Runtime/Experimental/Chaos/Public/Chaos/CollisionResolutionConvexConvex.h
@@ -4,7 +4,7 @@
 #include "Chaos/Convex.h"
 #include "Chaos/ParticleHandle.h"
 #include "Chaos/CollisionResolutionTypes.h"
-#include "Chaos/PBDCollisionConstraint.h"
+#include "Chaos/PBDCollisionConstraints.h"
 #include "Chaos/Transform.h"
 
 
diff --git a/Engine/Source/Runtime/Experimental/Chaos/Public/Chaos/CollisionResolutionTypes.h b/Engine/Source/Runtime/Experimental/Chaos/Public/Chaos/CollisionResolutionTypes.h
--- a/Engine/Source/Runtime/Experimental/Chaos/Public/Chaos/CollisionResolutionTypes.h
+++ b/Engine/Source/Runtime/Experimental/Chaos/Public/Chaos/CollisionResolutionTypes.h
@@ -12,7 +12,7 @@ namespace Chaos
 {
 
 	template<typename T, int d>
-	class TPBDCollisionConstraint;
+	class TPBDCollisionConstraints;
 
 	template<typename  T, int d>
 	class TPBDCollisionConstraintHandle;
@@ -214,11 +214,11 @@ namespace Chaos
 
 
 	template<class T, int d>
-	class CHAOS_API TPBDCollisionConstraintHandle : public TContainerConstraintHandle<TPBDCollisionConstraint<T, d>>
+	class CHAOS_API TPBDCollisionConstraintHandle : public TContainerConstraintHandle<TPBDCollisionConstraints<T, d>>
 	{
 	public:
-		using Base = TContainerConstraintHandle<TPBDCollisionConstraint<T, d>>;
-		using FConstraintContainer = TPBDCollisionConstraint<T, d>;
+		using Base = TContainerConstraintHandle<TPBDCollisionConstraints<T, d>>;
+		using FConstraintContainer = TPBDCollisionConstraints<T, d>;
 		using FConstraintBase = TCollisionConstraintBase<T, d>;
 
 
@@ -227,7 +227,7 @@ namespace Chaos
 		{}
 
 		TPBDCollisionConstraintHandle(FConstraintContainer* InConstraintContainer, int32 InConstraintIndex, typename FConstraintBase::FType InType)
-			: TContainerConstraintHandle<TPBDCollisionConstraint<T, d>>(InConstraintContainer, InConstraintIndex)
+			: TContainerConstraintHandle<TPBDCollisionConstraints<T, d>>(InConstraintContainer, InConstraintIndex)
 			, ConstraintType(InType)
 		{
 		}
diff --git a/Engine/Source/Runtime/Experimental/Chaos/Public/Chaos/Framework/PhysicsProxy.h b/Engine/Source/Runtime/Experimental/Chaos/Public/Chaos/Framework/PhysicsProxy.h
--- a/Engine/Source/Runtime/Experimental/Chaos/Public/Chaos/Framework/PhysicsProxy.h
+++ b/Engine/Source/Runtime/Experimental/Chaos/Public/Chaos/Framework/PhysicsProxy.h
@@ -6,7 +6,7 @@
 #include "Chaos/Declares.h"
 #include "Chaos/Framework/PhysicsSolverBase.h"
 #include "Chaos/Framework/PhysicsProxyBase.h"
-#include "Chaos/PBDCollisionConstraint.h"
+#include "Chaos/PBDCollisionConstraints.h"
 #include "Chaos/PBDRigidParticles.h"
 #include "UObject/GCObject.h"
 
@@ -41,7 +41,7 @@ public:
 	using FParticleData = ConcreteData;
 
 	using FParticlesType = Chaos::TPBDRigidParticles<float, 3>;
-	using FCollisionConstraintsType = Chaos::TPBDCollisionConstraint<float, 3>;
+	using FCollisionConstraintsType = Chaos::TPBDCollisionConstraints<float, 3>;
 	using FIntArray = Chaos::TArrayCollectionArray<int32>;
 
 	TPhysicsProxy()
diff --git a/Engine/Source/Runtime/Experimental/Chaos/Public/Chaos/ISpatialAccelerationCollection.h b/Engine/Source/Runtime/Experimental/Chaos/Public/Chaos/ISpatialAccelerationCollection.h
--- a/Engine/Source/Runtime/Experimental/Chaos/Public/Chaos/ISpatialAccelerationCollection.h
+++ b/Engine/Source/Runtime/Experimental/Chaos/Public/Chaos/ISpatialAccelerationCollection.h
@@ -4,7 +4,7 @@
 #include "Chaos/Box.h"
 #include "GeometryParticlesfwd.h"
 #include <tuple>
-#include "PBDCollisionConstraintImp.h"
+#include "PBDCollisionConstraintsImp.h"
 
 namespace Chaos
 {
@@ -24,8 +24,8 @@ public:
 	virtual ISpatialAcceleration<TPayloadType, T, d>* GetSubstructure(FSpatialAccelerationIdx Idx) = 0;
 
 	/** This is kind of a hack to avoid virtuals. We simply route calls into templated functions */
-	virtual void PBDComputeConstraintsLowLevel_GatherStats(TPBDCollisionConstraint<T, d>& CollisionConstraint, T Dt) const = 0;
-	virtual void PBDComputeConstraintsLowLevel(TPBDCollisionConstraint<T, d>& CollisionConstraint, T Dt) const = 0;
+	virtual void PBDComputeConstraintsLowLevel_GatherStats(TPBDCollisionConstraints<T, d>& CollisionConstraint, T Dt) const = 0;
+	virtual void PBDComputeConstraintsLowLevel(TPBDCollisionConstraints<T, d>& CollisionConstraint, T Dt) const = 0;
 	virtual TArray<FSpatialAccelerationIdx> GetAllSpatialIndices() const = 0;
 
 	bool IsBucketActive(uint8 BucketIdx) const
@@ -326,13 +326,13 @@ struct TSpatialAccelerationCollectionHelper
 };
 
 template <bool bGatherStats, typename SpatialAccelerationCollection, typename T, int d>
-typename TEnableIf<TIsSame<typename SpatialAccelerationCollection::TPayloadType, TAccelerationStructureHandle<T,d>>::Value, void>::Type PBDComputeConstraintsLowLevel_Helper(TPBDCollisionConstraint<T, d>& CollisionConstraint, const SpatialAccelerationCollection& Accel, T Dt)
+typename TEnableIf<TIsSame<typename SpatialAccelerationCollection::TPayloadType, TAccelerationStructureHandle<T,d>>::Value, void>::Type PBDComputeConstraintsLowLevel_Helper(TPBDCollisionConstraints<T, d>& CollisionConstraint, const SpatialAccelerationCollection& Accel, T Dt)
 {
 	CollisionConstraint.template ComputeConstraintsHelperLowLevel<bGatherStats>(Accel, Dt);
 }
 
 template <bool bGatherStats, typename SpatialAccelerationCollection, typename T, int d>
-typename TEnableIf<!TIsSame<typename SpatialAccelerationCollection::TPayloadType, TAccelerationStructureHandle<T, d>>::Value, void>::Type PBDComputeConstraintsLowLevel_Helper(TPBDCollisionConstraint<T, d>& CollisionConstraint, const SpatialAccelerationCollection& Accel, T Dt)
+typename TEnableIf<!TIsSame<typename SpatialAccelerationCollection::TPayloadType, TAccelerationStructureHandle<T, d>>::Value, void>::Type PBDComputeConstraintsLowLevel_Helper(TPBDCollisionConstraints<T, d>& CollisionConstraint, const SpatialAccelerationCollection& Accel, T Dt)
 {
 }
 
@@ -501,12 +501,12 @@ public:
 		return TUniquePtr<ISpatialAcceleration<TPayloadType, T, d>>(new TSpatialAccelerationCollection<TSpatialAccelerationTypes...>(*this));
 	}
 
-	virtual void PBDComputeConstraintsLowLevel_GatherStats(TPBDCollisionConstraint<T, d>& CollisionConstraint, T Dt) const override
+	virtual void PBDComputeConstraintsLowLevel_GatherStats(TPBDCollisionConstraints<T, d>& CollisionConstraint, T Dt) const override
 	{
 		PBDComputeConstraintsLowLevel_Helper<true>(CollisionConstraint, *this, Dt);
 	}
 
-	virtual void PBDComputeConstraintsLowLevel(TPBDCollisionConstraint<T,d>& CollisionConstraint, T Dt) const override
+	virtual void PBDComputeConstraintsLowLevel(TPBDCollisionConstraints<T,d>& CollisionConstraint, T Dt) const override
 	{
 		PBDComputeConstraintsLowLevel_Helper<false>(CollisionConstraint, *this, Dt);
 	}
diff --git a/Engine/Source/Runtime/Experimental/Chaos/Public/Chaos/PBDCollisionConstraint.h b/Engine/Source/Runtime/Experimental/Chaos/Public/Chaos/PBDCollisionConstraint.h
deleted file mode 100644
diff --git a/Engine/Source/Runtime/Experimental/Chaos/Public/Chaos/PBDCollisionConstraintImp.h b/Engine/Source/Runtime/Experimental/Chaos/Public/Chaos/PBDCollisionConstraintImp.h
deleted file mode 100644
diff --git a/Engine/Source/Runtime/Experimental/Chaos/Public/Chaos/PBDCollisionConstraintPGS.h b/Engine/Source/Runtime/Experimental/Chaos/Public/Chaos/PBDCollisionConstraintPGS.h
deleted file mode 100644
diff --git a/Engine/Source/Runtime/Experimental/Chaos/Public/Chaos/PBDCollisionConstraintPlaneContactUtil.h b/Engine/Source/Runtime/Experimental/Chaos/Public/Chaos/PBDCollisionConstraintPlaneContactUtil.h
deleted file mode 100644
diff --git a/Engine/Source/Runtime/Experimental/Chaos/Public/Chaos/PBDCollisionConstraintPointContactUtil.h b/Engine/Source/Runtime/Experimental/Chaos/Public/Chaos/PBDCollisionConstraintPointContactUtil.h
deleted file mode 100644
diff --git a/Engine/Source/Runtime/Experimental/Chaos/Public/Chaos/PBDCollisionConstraintUtil.h b/Engine/Source/Runtime/Experimental/Chaos/Public/Chaos/PBDCollisionConstraintUtil.h
deleted file mode 100644
diff --git a/Engine/Source/Runtime/Experimental/Chaos/Public/Chaos/PBDCollisionConstraints.h b/Engine/Source/Runtime/Experimental/Chaos/Public/Chaos/PBDCollisionConstraints.h
new file mode 100644
--- /dev/null
+++ b/Engine/Source/Runtime/Experimental/Chaos/Public/Chaos/PBDCollisionConstraints.h
@@ -0,0 +1,277 @@
+// Copyright 1998-2019 Epic Games, Inc. All Rights Reserved.
+#pragma once
+
+#include "Chaos/ConstraintHandle.h"
+#include "Chaos/CollisionResolutionTypes.h"
+#include "Chaos/PBDConstraintContainer.h"
+#include "Framework/BufferedData.h"
+
+#include <memory>
+#include <queue>
+#include <sstream>
+#include "BoundingVolume.h"
+#include "AABBTree.h"
+
+namespace Chaos
+{
+template<typename T, int d>
+class TPBDCollisionConstraintAccessor;
+
+template<typename T, int d>
+class TPBDCollisionConstraints;
+
+template <typename T, int d>
+class TRigidTransform;
+
+class FImplicitObject;
+
+template <typename T, int d>
+class TBVHParticles;
+
+template <typename T, int d>
+class TBox;
+
+template<class T>
+class TChaosPhysicsMaterial;
+
+template <typename T, int d>
+class TPBDRigidsSOAs;
+
+template<typename T, int d>
+using TRigidBodyContactConstraintsPostComputeCallback = TFunction<void()>;
+
+template<typename T, int d>
+using TRigidBodyContactConstraintsPostApplyCallback = TFunction<void(const T Dt, const TArray<TPBDCollisionConstraintHandle<T, d>*>&)>;
+
+template<typename T, int d>
+using TRigidBodyContactConstraintsPostApplyPushOutCallback = TFunction<void(const T Dt, const TArray<TPBDCollisionConstraintHandle<T, d>*>&, bool)>;
+
+/**
+ * Manages a set of contact constraints:
+ *	- Performs collision detection to generate constraints.
+ *	- Responsible for applying corrections to particles affected by the constraints.
+ * @todo(ccaulfield): remove TPBDCollisionConstraintAccessor
+ */
+template<typename T, int d>
+class CHAOS_API TPBDCollisionConstraints : public FPBDConstraintContainer
+{
+public:
+	friend class TPBDCollisionConstraintAccessor<T, d>;
+	friend class TPBDCollisionConstraintHandle<T, d>;
+
+	using Base = FPBDConstraintContainer;
+	using FHandles = TArray<TPBDCollisionConstraintHandle<T, d>*>;
+	using FConstraintContainerHandle = TPBDCollisionConstraintHandle<T, d>;
+ 	using FPointContactConstraint = TRigidBodyPointContactConstraint<T, d>;
+	using FPlaneContactConstraint = TRigidBodyPlaneContactConstraint<T, d>;
+	using FConstraintHandleAllocator = TConstraintHandleAllocator<TPBDCollisionConstraints<T, d>>;
+	using FConstraintHandleID = TPair<const TGeometryParticleHandle<T, d>*, const TGeometryParticleHandle<T, d>*>;
+
+
+	TPBDCollisionConstraints(const TPBDRigidsSOAs<T,d>& InParticles, 
+		TArrayCollectionArray<bool>& Collided, 
+		const TArrayCollectionArray<TSerializablePtr<FChaosPhysicsMaterial>>& PerParticleMaterials, 
+		const int32 ApplyPairIterations = 1, const int32 ApplyPushOutPairIterations = 1, const T Thickness = (T)0);
+
+	virtual ~TPBDCollisionConstraints() {}
+
+
+	/**
+	 * Generate all contact constraints.
+	 */
+	void UpdatePositionBasedState(const T Dt);
+
+
+	/**
+	*  Reset the constraint frame. 
+	*/
+	void Reset();
+
+	/**
+	 * Apply a modifier to the constraints and specify which constraints should be disabled.
+	 * You would probably call this in the PostComputeCallback. Prefer this to calling RemoveConstraints in a loop,
+	 * so you don't have to worry about constraint iterator/indices changing.
+	 */
+	void ApplyCollisionModifier(const TFunction<ECollisionModifierResult(const FConstraintContainerHandle* Handle)>& CollisionModifier);
+
+
+	/**
+	* Remove the constraints associated with the ParticleHandle.
+	*/
+	void RemoveConstraints(const TSet<TGeometryParticleHandle<T, d>*>&  ParticleHandle);
+
+
+	/**
+	* Remove the constraint, update the handle, and any maps. 
+	*/
+	void RemoveConstraint(FConstraintContainerHandle*);
+
+
+	/**
+	*  Process constraints acceleration structures. 
+	*/
+	//using FAccelerationStructure = TBoundingVolume<TAccelerationStructureHandle<T, d>, T, d>;
+	//using FAccelerationStructure = TAABBTree<TAccelerationStructureHandle<T, d>, TAABBTreeLeafArray<TAccelerationStructureHandle<T, d>, T>, T>;
+	using FAccelerationStructure = ISpatialAcceleration< TAccelerationStructureHandle<T, d>, T, d>;
+	void SetSpatialAcceleration(const FAccelerationStructure* InSpatialAcceleration) { SpatialAcceleration = InSpatialAcceleration; }
+
+	template <bool bGatherStats = false>
+	void ComputeConstraints(const FAccelerationStructure& AccelerationStructure, T Dt);
+
+	//NOTE: this should not be called by anyone other than ISpatialAccelerationCollection and CollisionConstraints - todo: make private with friends?
+	template <bool bGatherStats, typename SPATIAL_ACCELERATION>
+	void ComputeConstraintsHelperLowLevel(const SPATIAL_ACCELERATION& SpatialAcceleration, T Dt);
+
+
+	/**
+	 * Update all constraint values
+	 */
+	void UpdateConstraints(T Dt, const TSet<TGeometryParticleHandle<T, d>*>& AddedParticles);
+
+
+
+	/**
+	 * Apply corrections for the specified list of constraints. (Runs Wide!)
+	 *
+	 * @todo(ccaulfield): this runs wide. The serial/parallel decision should be in the ConstraintRule
+	 *
+	 */
+	void Apply(const T Dt, const TArray<FConstraintContainerHandle*>& InConstraintHandles, const int32 It, const int32 NumIts);
+
+
+	/**
+	 * Apply push out for the specified list of constraints.
+	 *  Return true if we need another iteration 
+	 *
+	 * @todo(ccaulfield): this runs wide. The serial/parallel decision should be in the ConstraintRule
+	 *
+	 */
+	bool ApplyPushOut(const T Dt, const TArray<FConstraintContainerHandle*>& InConstraintHandles, 
+		const TSet<const TGeometryParticleHandle<T,d>*>& IsTemporarilyStatic, int32 Iteration, int32 NumIterations);
+
+
+	/**
+	*
+	*  Callbacks
+	*
+	 * Set the callback used just after contacts are generated at the start of a frame tick.
+	 * This can be used to modify or disable constraints (via ApplyCollisionModifier).
+	 */
+	void SetPostComputeCallback(const TRigidBodyContactConstraintsPostComputeCallback<T, d>& Callback);
+	void ClearPostComputeCallback();
+
+	void SetPostApplyCallback(const TRigidBodyContactConstraintsPostApplyCallback<T, d>& Callback);
+	void ClearPostApplyCallback();
+
+	void SetPostApplyPushOutCallback(const TRigidBodyContactConstraintsPostApplyPushOutCallback<T, d>& Callback);
+	void ClearPostApplyPushOutCallback();
+
+
+	//
+	// Member Access
+	//
+
+	const TArray<FConstraintContainerHandle*>& GetAllConstraintHandles() const 
+	{ 
+		return Handles; 
+	}
+
+	void SetThickness(T InThickness)
+	{
+		MThickness = InThickness;
+	}
+
+	void SetVelocityInflation(float ScaleFactor)
+	{
+		CollisionVelocityInflation = ScaleFactor;
+	}
+
+	void SetPairIterations(int32 InPairIterations)
+	{
+		MApplyPairIterations = InPairIterations;
+	}
+
+	void SetPushOutPairIterations(int32 InPairIterations)
+	{
+		MApplyPushOutPairIterations = InPairIterations;
+	}
+
+	void SetCollisionsEnabled(bool bInEnableCollisions)
+	{
+		bEnableCollisions = bInEnableCollisions;
+	}
+
+	int32 NumConstraints() const
+	{
+		return Handles.Num();
+	}
+
+	FConstraintHandleID GetConstraintHandleID(const FPointContactConstraint & Constraint) const
+	{
+		return  FConstraintHandleID(Constraint.Particle[0], Constraint.Particle[1]);
+	}
+
+	FHandles& GetConstraintHandles()
+	{
+		return Handles;
+	}
+	const FHandles& GetConstConstraintHandles() const
+	{
+		return Handles;
+	}
+
+
+protected:
+		using Base::GetConstraintIndex;
+		using Base::SetConstraintIndex;
+
+private:
+
+	/**
+	* Update the individual constraints.
+	*/
+	template<ECollisionUpdateType>
+	void UpdateConstraint(const T Thickness, TRigidBodyPointContactConstraint<T, d> & Constraint);
+
+	/**
+	*  Build a constraint based on the two particle handles.
+	*/
+	void ConstructConstraints(TGeometryParticleHandle<T, d>* Particle0, TGeometryParticleHandle<T, d>* Particle1, 
+		const T Thickness, TRigidBodyPointContactConstraint<T, d> & Constraint);
+
+
+
+	const TPBDRigidsSOAs<T,d>& Particles;
+
+	TArray<FPointContactConstraint> PointConstraints;
+	TArray<FPlaneContactConstraint> PlaneConstraints;
+
+	TMap< FConstraintHandleID, FConstraintContainerHandle* > Manifolds;
+	TArray<FConstraintContainerHandle*> Handles;
+	FConstraintHandleAllocator HandleAllocator;
+
+	TArrayCollectionArray<bool>& MCollided;
+	const TArrayCollectionArray<TSerializablePtr<FChaosPhysicsMaterial>>& MPhysicsMaterials;
+	int32 MApplyPairIterations;
+	int32 MApplyPushOutPairIterations;
+	T MThickness;
+	T MAngularFriction;
+	bool bUseCCD;
+	bool bEnableCollisions;
+
+	int32 LifespanCounter;
+	float CollisionVelocityInflation;
+
+	TRigidBodyContactConstraintsPostComputeCallback<T, d> PostComputeCallback;
+	TRigidBodyContactConstraintsPostApplyCallback<T, d> PostApplyCallback;
+	TRigidBodyContactConstraintsPostApplyPushOutCallback<T, d> PostApplyPushOutCallback;
+
+	// @todo(ccaulfield): move spatial acceleration out of constraint container and make shareable
+	const FAccelerationStructure* SpatialAcceleration;
+};
+
+extern template void TPBDCollisionConstraints<float, 3>::UpdateConstraint<ECollisionUpdateType::Any>(const float Thickness, TRigidBodyPointContactConstraint<float,3> & Constraint);
+extern template void TPBDCollisionConstraints<float, 3>::UpdateConstraint<ECollisionUpdateType::Deepest>(const float Thickness, TRigidBodyPointContactConstraint<float, 3> & Constraint);
+extern template void TPBDCollisionConstraints<float, 3>::ComputeConstraints<false>(const TPBDCollisionConstraints<float, 3>::FAccelerationStructure&, float Dt);
+extern template void TPBDCollisionConstraints<float, 3>::ComputeConstraints<true>(const TPBDCollisionConstraints<float, 3>::FAccelerationStructure&, float Dt);
+}
diff --git a/Engine/Source/Runtime/Experimental/Chaos/Public/Chaos/PBDCollisionConstraintsImp.h b/Engine/Source/Runtime/Experimental/Chaos/Public/Chaos/PBDCollisionConstraintsImp.h
new file mode 100644
--- /dev/null
+++ b/Engine/Source/Runtime/Experimental/Chaos/Public/Chaos/PBDCollisionConstraintsImp.h
@@ -0,0 +1,358 @@
+// Copyright 1998-2019 Epic Games, Inc. All Rights Reserved.
+#pragma once
+
+#include "PBDCollisionConstraints.h"
+#include "Chaos/ChaosPerfTest.h"
+
+namespace Chaos
+{
+#if !UE_BUILD_SHIPPING
+
+	static constexpr int32 BucketSizes2[] = { 0,1,4,8,16,32,64,128,512,MAX_int32 };
+
+
+	template <bool bGatherStats>
+	struct FStatHelper2
+	{
+		int32 MaxCount;
+
+		FStatHelper2() {}
+		void Record(int32 Count) {}
+		FString ToString() const { return TEXT(""); }
+	};
+
+	template <>
+	struct FStatHelper2<true>
+	{
+		int32 BucketCount[UE_ARRAY_COUNT(BucketSizes2)];
+		int32 MaxCount;
+
+		FStatHelper2()
+		{
+			FMemory::Memset(BucketCount, 0, sizeof(BucketCount));
+			MaxCount = 0;
+		}
+
+		void Record(int32 Count)
+		{
+			for (int32 BucketIdx = 1; BucketIdx < UE_ARRAY_COUNT(BucketSizes2); ++BucketIdx)
+			{
+				if (Count >= BucketSizes2[BucketIdx - 1] && Count < BucketSizes2[BucketIdx])
+				{
+					++BucketCount[BucketIdx];
+				}
+			}
+
+			if (Count > MaxCount)
+			{
+				MaxCount = Count;
+			}
+		}
+
+		FString ToString() const
+		{
+			FString OutLog;
+			int32 MaxBucketCount = 0;
+			for (int32 Count : BucketCount)
+			{
+				if (Count > MaxBucketCount)
+				{
+					MaxBucketCount = Count;
+				}
+			}
+
+			const float CountPerChar = MaxBucketCount / 20.f;
+			for (int32 Idx = 1; Idx < UE_ARRAY_COUNT(BucketSizes2); ++Idx)
+			{
+				int32 NumChars = BucketCount[Idx] / CountPerChar;
+				if (Idx < UE_ARRAY_COUNT(BucketSizes2) - 1)
+				{
+					OutLog += FString::Printf(TEXT("\t[%4d - %4d) (%4d) |"), BucketSizes2[Idx - 1], BucketSizes2[Idx], BucketCount[Idx]);
+				}
+				else
+				{
+					OutLog += FString::Printf(TEXT("\t[%4d -  inf) (%4d) |"), BucketSizes2[Idx - 1], BucketCount[Idx]);
+				}
+				for (int32 Count = 0; Count < NumChars; ++Count)
+				{
+					OutLog += TEXT("-");
+				}
+				OutLog += TEXT("\n");
+			}
+
+			return OutLog;
+		}
+	};
+
+#endif
+
+#if UE_BUILD_SHIPPING
+	const bool bDisableCollisionParallelFor = false;
+#else
+	CHAOS_API extern bool bDisableCollisionParallelFor;
+#endif
+
+	template <typename T, int d>
+	struct TSimOverlapVisitor
+	{
+		TArray<TAccelerationStructureHandle<T, d>>& Intersections;
+		TSimOverlapVisitor(TArray<TAccelerationStructureHandle<T, d>>& InIntersections)
+			: Intersections(InIntersections)
+		{
+		}
+
+		bool VisitOverlap(const TSpatialVisitorData<TAccelerationStructureHandle<T, d>>& Instance)
+		{
+			Intersections.Add(Instance.Payload);
+			return true;
+		}
+
+		bool VisitSweep(TSpatialVisitorData<TAccelerationStructureHandle<T, d>>, FQueryFastData& CurData)
+		{
+			check(false);
+			return false;
+		}
+
+		bool VisitRaycast(TSpatialVisitorData<TAccelerationStructureHandle<T, d>>, FQueryFastData& CurData)
+		{
+			check(false);
+			return false;
+		}
+	};
+
+	extern CHAOS_API int32 EnableCollisions;
+	extern CHAOS_API int32 CollisionConstraintsForceSingleThreaded;
+
+	DECLARE_CYCLE_STAT_EXTERN(TEXT("ComputeConstraints"), STAT_ComputeConstraints, STATGROUP_Chaos, CHAOS_API);
+	DECLARE_CYCLE_STAT_EXTERN(TEXT("ComputeConstraintsNP"), STAT_ComputeConstraintsNP, STATGROUP_Chaos, CHAOS_API);
+	DECLARE_CYCLE_STAT_EXTERN(TEXT("ComputeConstraintsBP"), STAT_ComputeConstraintsBP, STATGROUP_Chaos, CHAOS_API);
+	DECLARE_CYCLE_STAT_EXTERN(TEXT("ComputeConstraintsSU"), STAT_ComputeConstraintsSU, STATGROUP_Chaos, CHAOS_API);
+
+	template <typename T, int d>
+	template <bool bGatherStats, typename SPATIAL_ACCELERATION>
+	void TPBDCollisionConstraints<T, d>::ComputeConstraintsHelperLowLevel(const SPATIAL_ACCELERATION& InSpatialAcceleration, T Dt)
+	{
+		SCOPE_CYCLE_COUNTER(STAT_ComputeConstraints);
+		CHAOS_SCOPED_TIMER(ComputeConstraints);
+		if (!EnableCollisions) return;
+		// Broad phase
+
+#if !UE_BUILD_SHIPPING
+		FStatHelper2<bGatherStats> BroadphasePotentials;
+		FStatHelper2<bGatherStats> NarrowPhaseSkipped;
+		FStatHelper2<bGatherStats> NarrowPhasePerformed;
+		FStatHelper2<bGatherStats> NarrowPhaseRejected;
+		FStatHelper2<bGatherStats> BoundsDistribution;
+		int32 SimulatedParticles = 0;
+#endif
+		{
+			SCOPE_CYCLE_COUNTER(STAT_ComputeConstraintsNP);
+
+
+			// Narrow phase
+			CHAOS_SCOPED_TIMER(ComputeConstraints_NP);
+
+			TQueue<TRigidBodyPointContactConstraint<T, d>, EQueueMode::Mpsc> Queue;	//todo(ocohen): use per thread buffer instead, need better support than ParallelFor for this
+			Particles.GetNonDisabledDynamicView().ParallelFor([&](auto& Particle1, int32 ActiveIdxIdx)
+			{
+#if !UE_BUILD_SHIPPING
+				if (bGatherStats)
+				{
+					++SimulatedParticles;
+				}
+#endif
+
+				TArray<TAccelerationStructureHandle<T, d>> PotentialIntersections;
+
+				const bool bBody1Bounded = HasBoundingBox(Particle1);
+
+				const T Box1Thickness = ComputeThickness(Particle1, CollisionVelocityInflation * Dt).Size();
+				if (bBody1Bounded)
+				{
+#if CHAOS_PARTICLEHANDLE_TODO
+					const TBox <T, d> Box1 = InSpatialAcceleration.GetWorldSpaceBoundingBox(Particle1);
+#else
+					const TBox <T, d> Box1 = ComputeWorldSpaceBoundingBox(Particle1);	//NOTE: this ignores the velocity expansion which is wrong
+#endif
+
+#if !UE_BUILD_SHIPPING
+					if (bGatherStats)
+					{
+						//const TBox<T, d> Box1 = SpatialAcceleration.GetWorldSpaceBoundingBox(InParticles, Body1Index);
+
+						BoundsDistribution.Record(Box1.Extents().GetMax());
+					}
+#endif
+					TSimOverlapVisitor<T, d> OverlapVisitor(PotentialIntersections);
+					InSpatialAcceleration.Overlap(Box1, OverlapVisitor);
+				}
+				else
+				{
+					const auto& GlobalElems = InSpatialAcceleration.GlobalObjects();
+					PotentialIntersections.Reserve(GlobalElems.Num());
+
+					for (auto& Elem : GlobalElems)
+					{
+						PotentialIntersections.Add(Elem.Payload);
+					}
+				}
+
+#if !UE_BUILD_SHIPPING
+				BroadphasePotentials.Record(PotentialIntersections.Num());
+
+				int32 CountNP = 0;
+				int32 RejectedNP = 0;
+#endif
+				const int32 NumPotentials = PotentialIntersections.Num();
+				for (int32 i = 0; i < NumPotentials; ++i)
+				{
+					auto& Particle2 = *PotentialIntersections[i].GetGeometryParticleHandle_PhysicsThread();
+					const TGenericParticleHandle<T, d> Particle2Generic(&Particle2);
+					// Collision group culling...
+					// CollisionGroup == 0 : Collide_With_Everything
+					// CollisionGroup == INDEX_NONE : Disabled collisions
+					// CollisionGroup_A != CollisionGroup_B : Skip Check
+
+					if (Particle1.CollisionGroup() == INDEX_NONE || Particle2Generic->CollisionGroup() == INDEX_NONE)
+					{
+						continue;
+					}
+					if (Particle1.CollisionGroup() && Particle2Generic->CollisionGroup() && Particle1.CollisionGroup() != Particle2Generic->CollisionGroup())
+					{
+						continue;
+					}
+
+					if (!Particle1.Geometry() && !Particle2.Geometry())
+					{
+						continue;
+					}
+
+					const bool bBody2Bounded = HasBoundingBox(Particle2);
+
+					if (Particle1.Handle() == Particle2.Handle())
+					{
+						continue;
+					}
+
+					if (bBody1Bounded == bBody2Bounded && Particle2.AsDynamic())
+					{
+						//if both are dynamic, assume index order matters
+#if CHAOS_DETERMINISTIC
+						if (Particle2.ParticleID() > Particle1.ParticleID())
+#else
+					//not deterministic just use memory address to avoid pair duplication
+						if (Particle2.Handle() > Particle1.Handle())
+#endif
+						{
+							continue;
+						}
+					}
+					//				++TotalTested;
+
+					const TVector<T, d> Box2Thickness = Particle2.AsDynamic() ? ComputeThickness(*Particle2.AsDynamic(), Dt) : TVector<T, d>(0);
+					const T UseThickness = FMath::Max(Box1Thickness, Box2Thickness.Size());// + MThickness
+
+
+					TRigidBodyPointContactConstraint<T, d> Constraint;
+					ConstructConstraints(Particle1.Handle(), Particle2.Handle(), UseThickness, Constraint);
+
+#if !UE_BUILD_SHIPPING
+					if (bGatherStats)
+					{
+						++CountNP;
+					}
+#endif
+					if (Constraint.ContainsManifold())
+					{
+						UpdateConstraint<ECollisionUpdateType::Any>(UseThickness, Constraint);
+					}
+
+					if (Constraint.GetPhi() < UseThickness)
+					{
+						Queue.Enqueue(Constraint);
+					}
+					else
+					{
+#if !UE_BUILD_SHIPPING
+						++RejectedNP;
+#endif
+					}
+				}
+
+#if !UE_BUILD_SHIPPING
+				NarrowPhasePerformed.Record(CountNP);
+				const int32 NPSkipped = NumPotentials - CountNP;
+				NarrowPhaseSkipped.Record(NPSkipped);
+				NarrowPhaseRejected.Record(RejectedNP);
+#endif
+			}, bGatherStats || bDisableCollisionParallelFor);
+			
+			{
+				SCOPE_CYCLE_COUNTER(STAT_ComputeConstraintsSU);
+
+				while (!Queue.IsEmpty())
+				{
+					int32 Idx = PointConstraints.AddUninitialized(1);
+					FConstraintContainerHandle* Handle = HandleAllocator.template AllocHandle< TRigidBodyPointContactConstraint<T,d> >(this, Idx );
+
+					Handles.Add(Handle);
+					Queue.Dequeue(PointConstraints[Idx]);
+				}
+				LifespanCounter++;
+			}
+		}
+
+#if CHAOS_DETERMINISTIC
+		//todo: sort constraints
+#endif
+
+#if !UE_BUILD_SHIPPING
+		FString OutLog;
+		if (bGatherStats)
+		{
+#if CHAOS_PARTICLEHANDLE_TODO
+			const float NumParticles = InParticles.Size();
+			OutLog = FString::Printf(TEXT("ComputeConstraints stats:\n"
+				"Total Particles:%d\nSimulated Particles:%d (%.2f%%)\n"
+				"Max candidates per instance:%d (%.2f%%)\n"
+				"Max candidates skipped per instance (NP skipped):%d (%.2f%%)\n"
+				"Max narrow phase tests per instance:%d (%.2f%%)\n"
+				"Max narrow phase rejected per instance (NP rejected):%d (%.2f%%)\n"
+				"Constraints generated:%d\n"
+			),
+				InParticles.Size(),
+				SimulatedParticles, SimulatedParticles / NumParticles * 100.f,
+				BroadphasePotentials.MaxCount, BroadphasePotentials.MaxCount / NumParticles * 100.f,
+				NarrowPhaseSkipped.MaxCount, NarrowPhaseSkipped.MaxCount / NumParticles * 100.f,
+				NarrowPhasePerformed.MaxCount, NarrowPhasePerformed.MaxCount / NumParticles * 100.f,
+				NarrowPhaseRejected.MaxCount, NarrowPhaseRejected.MaxCount / NumParticles * 100.f,
+				Constraints.Num()
+			);
+
+			OutLog += FString::Printf(TEXT("Potentials per instance distribution:\n"));
+			OutLog += BroadphasePotentials.ToString();
+
+
+			OutLog += FString::Printf(TEXT("\nCandidates skipped per instance (NP skipped) distribution:\n"));
+			OutLog += NarrowPhaseSkipped.ToString();
+
+			OutLog += FString::Printf(TEXT("\nNarrow phase performed per instance distribution:\n"));
+			OutLog += NarrowPhasePerformed.ToString();
+
+			OutLog += FString::Printf(TEXT("\nNarrow phase candidates rejected per instance distribution:\n"));
+			OutLog += NarrowPhaseRejected.ToString();
+
+			OutLog += FString::Printf(TEXT("\nBounds distribution:\n"));
+			OutLog += BoundsDistribution.ToString();
+
+			UE_LOG(LogChaos, Warning, TEXT("%s"), *OutLog);
+			PendingHierarchyDump = 0;
+#endif
+		}
+#endif
+		//	if (TotalTested > 0)
+		//	{
+				//UE_LOG(LogChaos, Warning, TEXT("ComputeConstraints: rejected:%f out of rejected:%f = %.3f"), TotalRejected, TotalTested, TotalRejected / TotalTested);
+		//	}
+	}
+}
\ No newline at end of file
diff --git a/Engine/Source/Runtime/Experimental/Chaos/Public/Chaos/PBDCollisionConstraintsPGS.h b/Engine/Source/Runtime/Experimental/Chaos/Public/Chaos/PBDCollisionConstraintsPGS.h
new file mode 100644
--- /dev/null
+++ b/Engine/Source/Runtime/Experimental/Chaos/Public/Chaos/PBDCollisionConstraintsPGS.h
@@ -0,0 +1,92 @@
+// Copyright 1998-2019 Epic Games, Inc. All Rights Reserved.
+#pragma once
+
+#include "Chaos/CollisionResolutionTypes.h"
+#include "Chaos/PBDConstraintGraph.h"
+
+#include <memory>
+#include <queue>
+#include <sstream>
+
+#if CHAOS_PARTICLEHANDLE_TODO
+
+namespace ChaosTest
+{
+	template<class T> void CollisionPGS();
+	template<class T> void CollisionPGS2();
+}
+
+namespace Chaos
+{
+template<class T, int d>
+class TPBDCollisionConstraintPGS
+{
+  public:
+	friend void ChaosTest::CollisionPGS<T>();
+	friend void ChaosTest::CollisionPGS2<T>();
+
+	typedef TRigidBodyContactConstraintPGS<T, d> FRigidBodyContactConstraint;
+	typedef TPBDConstraintGraph<T, d> FConstraintGraph;
+
+	TPBDCollisionConstraintPGS(TPBDRigidParticles<T, d>& InParticles, const TArray<int32>& InIndices, TArrayCollectionArray<bool>& Collided, TArrayCollectionArray<TSerializablePtr<FChaosPhysicsMaterial>>& InPhysicsMaterials, const T Thickness = (T)0);
+	virtual ~TPBDCollisionConstraintPGS() {}
+
+	void ComputeConstraints(const TPBDRigidParticles<T, d>& InParticles, const TArray<int32>& InIndices, const T Dt);
+
+	void Apply(TPBDRigidParticles<T, d>& InParticles, const T Dt, const TArray<int32>& InConstraintIndices);
+	bool ApplyPushOut(TPBDRigidParticles<T, d>& InParticles, const T Dt, const TArray<int32>& InConstraintIndices);
+	template<class T_PARTICLES> void Solve(T_PARTICLES& InParticles, const T Dt, const TArray<int32>& InConstraintIndices);
+
+	void RemoveConstraints(const TSet<uint32>& RemovedParticles);
+	void UpdateConstraints(const TPBDRigidParticles<T, d>& InParticles, T Dt, const TSet<uint32>& AddedParticles, const TArray<uint32>& ActiveParticles);
+
+
+	static bool NearestPoint(TArray<Pair<TVector<T, d>, TVector<T, d>>>& Points, TVector<T, d>& Direction);
+
+	const TArray<FRigidBodyContactConstraint>& GetAllConstraints() const { return Constraints; }
+
+	int32 NumConstraints() const { return Constraints.Num(); }
+	TVector<int32, 2> ConstraintParticleIndices(int32 ContactIndex) const { return { Constraints[ContactIndex].ParticleIndex, Constraints[ContactIndex].LevelsetIndex }; }
+	void UpdatePositionBasedState(const TPBDRigidParticles<T, d>& InParticles, const TArray<int32>& InIndices, const T Dt) { ComputeConstraints(InParticles, InIndices, Dt); }
+
+  private:
+	void PrintParticles(const TPBDRigidParticles<T, d>& InParticles, const TArray<int32>& InConstraintIndices);
+	void PrintConstraints(const TPBDRigidParticles<T, d>& InParticles, const TArray<int32>& InConstraintIndices);
+
+	template<class T_PARTICLES>
+	void UpdateConstraint(const T_PARTICLES& InParticles, const T Thickness, FRigidBodyContactConstraint& Constraint);
+	template<class T_PARTICLES>
+	void UpdateLevelsetConstraint(const T_PARTICLES& InParticles, const T Thickness, FRigidBodyContactConstraint& Constraint);
+	template<class T_PARTICLES>
+	void UpdateLevelsetConstraintGJK(const T_PARTICLES& InParticles, const T Thickness, FRigidBodyContactConstraint& Constraint);
+	template<class T_PARTICLES>
+	void UpdateBoxConstraint(const T_PARTICLES& InParticles, const T Thickness, FRigidBodyContactConstraint& Constraint);
+	template<class T_PARTICLES>
+	void UpdateBoxPlaneConstraint(const T_PARTICLES& InParticles, const T Thickness, FRigidBodyContactConstraint& Constraint);
+	template<class T_PARTICLES>
+	void UpdateSphereConstraint(const T_PARTICLES& InParticles, const T Thickness, FRigidBodyContactConstraint& Constraint);
+	template<class T_PARTICLES>
+	void UpdateSpherePlaneConstraint(const T_PARTICLES& InParticles, const T Thickness, FRigidBodyContactConstraint& Constraint);
+	template<class T_PARTICLES>
+	void UpdateSphereBoxConstraint(const T_PARTICLES& InParticles, const T Thickness, FRigidBodyContactConstraint& Constraint);
+
+	FRigidBodyContactConstraint ComputeLevelsetConstraint(const TPBDRigidParticles<T, d>& InParticles, int32 ParticleIndex, int32 LevelsetIndex, const T Thickness);
+	FRigidBodyContactConstraint ComputeLevelsetConstraintGJK(const TPBDRigidParticles<T, d>& InParticles, int32 ParticleIndex, int32 LevelsetIndex, const T Thickness);
+	FRigidBodyContactConstraint ComputeBoxConstraint(const TPBDRigidParticles<T, d>& InParticles, int32 Box1Index, int32 Box2Index, const T Thickness);
+	FRigidBodyContactConstraint ComputeBoxPlaneConstraint(const TPBDRigidParticles<T, d>& InParticles, int32 BoxIndex, int32 PlaneIndex, const T Thickness);
+	FRigidBodyContactConstraint ComputeSphereConstraint(const TPBDRigidParticles<T, d>& InParticles, int32 Sphere1Index, int32 Sphere2Index, const T Thickness);
+	FRigidBodyContactConstraint ComputeSpherePlaneConstraint(const TPBDRigidParticles<T, d>& InParticles, int32 SphereIndex, int32 PlaneIndex, const T Thickness);
+	FRigidBodyContactConstraint ComputeSphereBoxConstraint(const TPBDRigidParticles<T, d>& InParticles, int32 SphereIndex, int32 BoxIndex, const T Thickness);
+	FRigidBodyContactConstraint ComputeConstraint(const TPBDRigidParticles<T, d>& InParticles, int32 Body1Index, int32 Body2Index, const T Thickness);
+
+	TArray<FRigidBodyContactConstraint> Constraints;
+	TArrayCollectionArray<bool>& MCollided;
+	const TArrayCollectionArray<TSerializablePtr<FChaosPhysicsMaterial>>& MPhysicsMaterials;
+	const T MThickness;
+	const T Tolerance;
+	const T MaxIterations;
+	const bool bUseCCD;
+};
+}
+
+#endif
\ No newline at end of file
diff --git a/Engine/Source/Runtime/Experimental/Chaos/Public/Chaos/PBDCollisionConstraintsPlaneContactUtil.h b/Engine/Source/Runtime/Experimental/Chaos/Public/Chaos/PBDCollisionConstraintsPlaneContactUtil.h
new file mode 100644
--- /dev/null
+++ b/Engine/Source/Runtime/Experimental/Chaos/Public/Chaos/PBDCollisionConstraintsPlaneContactUtil.h
@@ -0,0 +1,48 @@
+// Copyright 1998-2019 Epic Games, Inc. All Rights Reserved.
+#pragma once
+
+#include "Chaos/ExternalCollisionData.h"
+#include "Chaos/PBDCollisionConstraints.h"
+#include "Chaos/CollisionResolutionTypes.h"
+
+namespace Chaos
+{
+	namespace Collisions
+	{
+		template<class T = float>
+		struct TPlainContactParticleParameters {
+			TArrayCollectionArray<bool>* Collided;
+			const TArrayCollectionArray<TSerializablePtr<FChaosPhysicsMaterial>>* PhysicsMaterials;
+			T FrictionOverride;
+			T AngularFrictionOverride;
+		};
+
+		template<class T = float>
+		struct TPlaneContactIterationParameters {
+			const T Dt; 
+			const int32 Iteration; 
+			const int32 NumIterations; 
+			const int32 NumPairIterations;
+			bool* NeedsAnotherIteration;
+		};
+
+		/*
+		template<ECollisionUpdateType UpdateType, typename T = float, int d = 3>
+		void Update(const T Thickness, TRigidBodyPointContactConstraint<T, d>& Constraint);
+
+		template<ECollisionUpdateType UpdateType, typename T=float, int d=3>
+		void UpdateConstraint(const T Thickness, TRigidBodyPointContactConstraint<T, d>& Constraint);
+			
+		template<typename T = float, int d = 3>
+		void Apply(TRigidBodyPointContactConstraint<T, d>& Constraint, T Thickness, 
+			TPointContactIterationParameters<T> & IterationParameters,
+			TPointContactParticleParameters<T> & ParticleParameters);
+			
+		template<typename T=float, int d=3>
+		void ApplyPushOut(TRigidBodyPointContactConstraint<T, d>& Constraint, T Thickness, const TSet<const TGeometryParticleHandle<T, d>*>& IsTemporarilyStatic,
+			TPointContactIterationParameters<T> & IterationParameters,
+			TPointContactParticleParameters<T> & ParticleParameters);
+	*/
+	}
+
+}
diff --git a/Engine/Source/Runtime/Experimental/Chaos/Public/Chaos/PBDCollisionConstraintsPointContactUtil.h b/Engine/Source/Runtime/Experimental/Chaos/Public/Chaos/PBDCollisionConstraintsPointContactUtil.h
new file mode 100644
--- /dev/null
+++ b/Engine/Source/Runtime/Experimental/Chaos/Public/Chaos/PBDCollisionConstraintsPointContactUtil.h
@@ -0,0 +1,46 @@
+// Copyright 1998-2019 Epic Games, Inc. All Rights Reserved.
+#pragma once
+
+#include "Chaos/ExternalCollisionData.h"
+#include "Chaos/PBDCollisionConstraints.h"
+#include "Chaos/CollisionResolutionTypes.h"
+
+namespace Chaos
+{
+	namespace Collisions
+	{
+		template<class T = float>
+		struct TPointContactParticleParameters {
+			TArrayCollectionArray<bool>* Collided;
+			const TArrayCollectionArray<TSerializablePtr<FChaosPhysicsMaterial>>* PhysicsMaterials;
+			T FrictionOverride;
+			T AngularFrictionOverride;
+		};
+
+		template<class T = float>
+		struct TPointContactIterationParameters {
+			const T Dt; 
+			const int32 Iteration; 
+			const int32 NumIterations; 
+			const int32 NumPairIterations;
+			bool* NeedsAnotherIteration;
+		};
+
+		template<ECollisionUpdateType UpdateType, typename T = float, int d = 3>
+		void Update(const T Thickness, TRigidBodyPointContactConstraint<T, d>& Constraint);
+
+		template<ECollisionUpdateType UpdateType, typename T=float, int d=3>
+		void UpdateConstraint(const T Thickness, TRigidBodyPointContactConstraint<T, d>& Constraint);
+			
+		template<typename T = float, int d = 3>
+		void Apply(TRigidBodyPointContactConstraint<T, d>& Constraint, T Thickness, 
+			TPointContactIterationParameters<T> & IterationParameters,
+			TPointContactParticleParameters<T> & ParticleParameters);
+			
+		template<typename T=float, int d=3>
+		void ApplyPushOut(TRigidBodyPointContactConstraint<T, d>& Constraint, T Thickness, const TSet<const TGeometryParticleHandle<T, d>*>& IsTemporarilyStatic,
+			TPointContactIterationParameters<T> & IterationParameters,
+			TPointContactParticleParameters<T> & ParticleParameters);
+	}
+
+}
diff --git a/Engine/Source/Runtime/Experimental/Chaos/Public/Chaos/PBDCollisionConstraintsUtil.h b/Engine/Source/Runtime/Experimental/Chaos/Public/Chaos/PBDCollisionConstraintsUtil.h
new file mode 100644
--- /dev/null
+++ b/Engine/Source/Runtime/Experimental/Chaos/Public/Chaos/PBDCollisionConstraintsUtil.h
@@ -0,0 +1,27 @@
+// Copyright 1998-2019 Epic Games, Inc. All Rights Reserved.
+#pragma once
+
+#include "Chaos/ExternalCollisionData.h"
+#include "Chaos/PBDCollisionConstraints.h"
+#include "Chaos/CollisionResolutionTypes.h"
+
+namespace Chaos
+{
+	void CHAOS_API ComputeHashTable(const TArray<Chaos::TPBDCollisionConstraints<float, 3>::FPointContactConstraint>& ConstraintsArray,
+									const FBox& BoundingBox, TMultiMap<int32, int32>& HashTableMap, const float SpatialHashRadius);
+
+	void CHAOS_API ComputeHashTable(const TArray<TCollisionData<float, 3>>& CollisionsArray,
+									const FBox& BoundingBox, TMultiMap<int32, int32>& HashTableMap, const float SpatialHashRadius);
+
+	void CHAOS_API ComputeHashTable(const TArray<TCollisionDataExt<float, 3>>& CollisionsArray,
+									const FBox& BoundingBox, TMultiMap<int32, int32>& HashTableMap, const float SpatialHashRadius);
+
+	void CHAOS_API ComputeHashTable(const TArray<FVector>& ParticleArray, const FBox& BoundingBox, 
+									TMultiMap<int32, int32>& HashTableMap, const float SpatialHashRadius);
+
+	void CHAOS_API ComputeHashTable(const TArray<TBreakingData<float, 3>>& BreakingsArray,
+									const FBox& BoundingBox, TMultiMap<int32, int32>& HashTableMap, const float SpatialHashRadius);
+
+	void CHAOS_API ComputeHashTable(const TArray<TBreakingDataExt<float, 3>>& BreakingsArray,
+									const FBox& BoundingBox, TMultiMap<int32, int32>& HashTableMap, const float SpatialHashRadius);
+}
diff --git a/Engine/Source/Runtime/Experimental/Chaos/Public/Chaos/PBDRigidsEvolution.h b/Engine/Source/Runtime/Experimental/Chaos/Public/Chaos/PBDRigidsEvolution.h
--- a/Engine/Source/Runtime/Experimental/Chaos/Public/Chaos/PBDRigidsEvolution.h
+++ b/Engine/Source/Runtime/Experimental/Chaos/Public/Chaos/PBDRigidsEvolution.h
@@ -1,8 +1,8 @@
 // Copyright 1998-2019 Epic Games, Inc. All Rights Reserved.
 #pragma once
 
-#include "Chaos/PBDCollisionConstraint.h"
-#include "Chaos/PBDCollisionConstraintPGS.h"
+#include "Chaos/PBDCollisionConstraints.h"
+#include "Chaos/PBDCollisionConstraintsPGS.h"
 #include "Chaos/PBDConstraintGraph.h"
 #include "Chaos/PBDRigidClustering.h"
 #include "Chaos/PBDRigidParticles.h"
@@ -744,7 +744,7 @@ protected:
 
 // Only way to make this compile at the moment due to visibility attribute issues. TODO: Change this once a fix for this problem is applied.
 #if PLATFORM_MAC || PLATFORM_LINUX
-extern template class CHAOS_API Chaos::TPBDRigidsEvolutionBase<Chaos::TPBDRigidsEvolutionGBF<float, 3>, Chaos::TPBDCollisionConstraint<float,3>, float, 3>;
+extern template class CHAOS_API Chaos::TPBDRigidsEvolutionBase<Chaos::TPBDRigidsEvolutionGBF<float, 3>, Chaos::TPBDCollisionConstraints<float,3>, float, 3>;
 #else
-extern template class Chaos::TPBDRigidsEvolutionBase<Chaos::TPBDRigidsEvolutionGBF<float, 3>, Chaos::TPBDCollisionConstraint<float,3>, float, 3>;
+extern template class Chaos::TPBDRigidsEvolutionBase<Chaos::TPBDRigidsEvolutionGBF<float, 3>, Chaos::TPBDCollisionConstraints<float,3>, float, 3>;
 #endif
diff --git a/Engine/Source/Runtime/Experimental/Chaos/Public/Chaos/PBDRigidsEvolutionGBF.h b/Engine/Source/Runtime/Experimental/Chaos/Public/Chaos/PBDRigidsEvolutionGBF.h
--- a/Engine/Source/Runtime/Experimental/Chaos/Public/Chaos/PBDRigidsEvolutionGBF.h
+++ b/Engine/Source/Runtime/Experimental/Chaos/Public/Chaos/PBDRigidsEvolutionGBF.h
@@ -2,7 +2,7 @@
 #pragma once
 
 #include "Chaos/PBDRigidsEvolution.h"
-#include "Chaos/PBDCollisionConstraint.h"
+#include "Chaos/PBDCollisionConstraints.h"
 #include "Chaos/ChaosPerfTest.h"
 #include "Chaos/PerParticleInitForce.h"
 #include "Chaos/PerParticleEulerStepVelocity.h"
@@ -32,10 +32,10 @@ using TPBDRigidsEvolutionIslandCallback = TFunction<void(int32 Island)>;
 
 
 template<typename T, int d>
-class TPBDRigidsEvolutionGBF : public TPBDRigidsEvolutionBase<TPBDRigidsEvolutionGBF<T, d>, TPBDCollisionConstraint<T,d>, T, d>
+class TPBDRigidsEvolutionGBF : public TPBDRigidsEvolutionBase<TPBDRigidsEvolutionGBF<T, d>, TPBDCollisionConstraints<T,d>, T, d>
 {
 public:
-	using Base = TPBDRigidsEvolutionBase<TPBDRigidsEvolutionGBF<T, d>, TPBDCollisionConstraint<T, d>, T, d>;
+	using Base = TPBDRigidsEvolutionBase<TPBDRigidsEvolutionGBF<T, d>, TPBDCollisionConstraints<T, d>, T, d>;
 	using Base::Particles;
 	using Base::ForceRules;
 	using Base::ParticleUpdatePosition;
@@ -47,7 +47,7 @@ public:
 	using Base::Clustering;
 	using typename Base::FForceRule;
 	using FGravityForces = TPerParticleGravity<T, d>;
-	using FCollisionConstraints = TPBDCollisionConstraint<T, d>;
+	using FCollisionConstraints = TPBDCollisionConstraints<T, d>;
 	using FExternalForces = TPerParticleExternalForces<T, d>;
 	using FCollisionConstraintRule = TPBDConstraintColorRule<FCollisionConstraints>;
 
diff --git a/Engine/Source/Runtime/Experimental/ChaosSolvers/Private/EventDefaults.cpp b/Engine/Source/Runtime/Experimental/ChaosSolvers/Private/EventDefaults.cpp
--- a/Engine/Source/Runtime/Experimental/ChaosSolvers/Private/EventDefaults.cpp
+++ b/Engine/Source/Runtime/Experimental/ChaosSolvers/Private/EventDefaults.cpp
@@ -5,7 +5,7 @@
 #include "PhysicsProxy/SingleParticlePhysicsProxy.h"
 #include "Chaos/PBDRigidsEvolution.h"
 #include "Chaos/PBDRigidClustering.h"
-#include "Chaos/PBDCollisionConstraint.h"
+#include "Chaos/PBDCollisionConstraints.h"
 #include "Chaos/CollisionResolutionTypes.h"
 #include "PBDRigidsSolver.h"
 #include "PhysicsProxy/SkeletalMeshPhysicsProxy.h"
@@ -107,7 +107,7 @@ namespace Chaos
 				{
 					for(int32 IdxCollision = 0; IdxCollision < ValidCollisionHandles.Num(); ++IdxCollision)
 					{
-						Chaos::TPBDCollisionConstraint<float, 3>::FPointContactConstraint const& Constraint = ValidCollisionHandles[IdxCollision]->GetPointContact();
+						Chaos::TPBDCollisionConstraints<float, 3>::FPointContactConstraint const& Constraint = ValidCollisionHandles[IdxCollision]->GetPointContact();
 
 						TGeometryParticleHandle<float, 3>* Particle0 = Constraint.Particle[0];
 						TGeometryParticleHandle<float, 3>* Particle1 = Constraint.Particle[1];
diff --git a/Engine/Source/Runtime/Experimental/ChaosSolvers/Private/PBDRigidsSolver.cpp b/Engine/Source/Runtime/Experimental/ChaosSolvers/Private/PBDRigidsSolver.cpp
--- a/Engine/Source/Runtime/Experimental/ChaosSolvers/Private/PBDRigidsSolver.cpp
+++ b/Engine/Source/Runtime/Experimental/ChaosSolvers/Private/PBDRigidsSolver.cpp
@@ -4,7 +4,7 @@
 
 #include "Async/AsyncWork.h"
 #include "Chaos/ChaosArchive.h"
-#include "Chaos/PBDCollisionConstraintUtil.h"
+#include "Chaos/PBDCollisionConstraintsUtil.h"
 #include "Chaos/Utilities.h"
 #include "Chaos/ChaosDebugDraw.h"
 #include "ChaosStats.h"
diff --git a/Engine/Source/Runtime/Experimental/ChaosSolvers/Private/PhysicsProxy/GeometryCollectionPhysicsProxy.cpp b/Engine/Source/Runtime/Experimental/ChaosSolvers/Private/PhysicsProxy/GeometryCollectionPhysicsProxy.cpp
--- a/Engine/Source/Runtime/Experimental/ChaosSolvers/Private/PhysicsProxy/GeometryCollectionPhysicsProxy.cpp
+++ b/Engine/Source/Runtime/Experimental/ChaosSolvers/Private/PhysicsProxy/GeometryCollectionPhysicsProxy.cpp
@@ -13,7 +13,7 @@
 #include "Chaos/TriangleMesh.h"
 #include "Chaos/MassProperties.h"
 #include "ChaosSolversModule.h"
-#include "Chaos/PBDCollisionConstraintUtil.h"
+#include "Chaos/PBDCollisionConstraintsUtil.h"
 #include "Chaos/ImplicitObject.h"
 #include "Chaos/Serializable.h"
 #include "Chaos/ErrorReporter.h"
@@ -1325,7 +1325,7 @@ void FGeometryCollectionPhysicsProxy::EndFrameCallback(const float EndFrame)
 			TManagedArray<int32>& DynamicState = Collection->GetAttribute<int32>("DynamicState", FGeometryCollection::TransformGroup);
 
 			Chaos::TPBDRigidParticles<float, 3>& Particles = GetSolver()->GetRigidParticles();
-			const Chaos::TPBDCollisionConstraint<float, 3>& CollisionRule = GetSolver()->GetCollisionConstraints();
+			const Chaos::TPBDCollisionConstraints<float, 3>& CollisionRule = GetSolver()->GetCollisionConstraints();
 			const Chaos::TArrayCollectionArray<Chaos::ClusterId>& ClusterID = GetSolver()->GetRigidClustering().GetClusterIdsArray();
 			const Chaos::TArrayCollectionArray<Chaos::TRigidTransform<float, 3>>& ClusterChildToParentMap = GetSolver()->GetRigidClustering().GetChildToParentMap();
 			const Chaos::TArrayCollectionArray<bool>& InternalCluster = GetSolver()->GetRigidClustering().GetInternalClusterArray();
@@ -3168,7 +3168,7 @@ FRecordedFrame& FGeometryCollectionPhysicsProxy::InsertRecordedFrame(FRecordedTr
 	return InTrack.Records[NewRecordIndex];
 }
 
-void FGeometryCollectionPhysicsProxy::AddCollisionToCollisionData(FRecordedFrame* ExistingFrame, const FParticlesType& Particles, const Chaos::TPBDCollisionConstraint<float, 3>::FPointContactConstraint& Constraint)
+void FGeometryCollectionPhysicsProxy::AddCollisionToCollisionData(FRecordedFrame* ExistingFrame, const FParticlesType& Particles, const Chaos::TPBDCollisionConstraints<float, 3>::FPointContactConstraint& Constraint)
 {
 #if TODO_REIMPLEMENT_RIGID_CLUSTERING
 	const Chaos::TArrayCollectionArray<Chaos::ClusterId>& ClusterIdsArray = GetSolver()->GetRigidClustering().GetClusterIdsArray();
@@ -3214,12 +3214,12 @@ void FGeometryCollectionPhysicsProxy::UpdateCollisionData(const FParticlesType&
 	if (Parameters.CollisionData.SaveCollisionData && ExistingFrame->Timestamp > 0.f && Parameters.CollisionData.CollisionDataSizeMax > 0)
 	{
 #if TODO_REIMPLEMENT_PHYSICS_PROXY_REVERSE_MAPPING
-		const TArray<Chaos::TPBDCollisionConstraint<float, 3>::FRigidBodyContactConstraint>& AllConstraintsArray = CollisionRule.GetAllConstraints();
+		const TArray<Chaos::TPBDCollisionConstraints<float, 3>::FRigidBodyContactConstraint>& AllConstraintsArray = CollisionRule.GetAllConstraints();
 		if (AllConstraintsArray.Num() > 0)
 		{
 			const Chaos::TArrayCollectionArray<PhysicsProxyWrapper>& PhysicsProxyReverseMapping = GetSolver()->GetPhysicsProxyReverseMapping();
 
-			TArray<Chaos::TPBDCollisionConstraint<float, 3>::FRigidBodyContactConstraint> ConstraintsArray;
+			TArray<Chaos::TPBDCollisionConstraints<float, 3>::FRigidBodyContactConstraint> ConstraintsArray;
 			ConstraintsArray.SetNumUninitialized(AllConstraintsArray.Num());
 
 			FBox BoundingBox(ForceInitToZero);
diff --git a/Engine/Source/Runtime/Experimental/ChaosSolvers/Public/PBDRigidsSolver.h b/Engine/Source/Runtime/Experimental/ChaosSolvers/Public/PBDRigidsSolver.h
--- a/Engine/Source/Runtime/Experimental/ChaosSolvers/Public/PBDRigidsSolver.h
+++ b/Engine/Source/Runtime/Experimental/ChaosSolvers/Public/PBDRigidsSolver.h
@@ -7,7 +7,7 @@
 #include "Chaos/Framework/PhysicsSolverBase.h"
 #include "Chaos/PBDRigidParticles.h"
 #include "Chaos/PBDRigidsEvolutionGBF.h"
-#include "Chaos/PBDCollisionConstraint.h"
+#include "Chaos/PBDCollisionConstraints.h"
 #include "Chaos/PBDRigidDynamicSpringConstraints.h"
 #include "Chaos/PBDPositionConstraints.h"
 #include "Chaos/PBDJointConstraints.h"
@@ -86,7 +86,7 @@ namespace Chaos
 		typedef Chaos::TGeometryParticle<float, 3> FParticle;
 		typedef Chaos::TGeometryParticleHandle<float, 3> FHandle;
 		typedef Chaos::TPBDRigidsEvolutionGBF<float, 3> FPBDRigidsEvolution;
-		typedef Chaos::TPBDCollisionConstraint<float, 3> FPBDCollisionConstraints;
+		typedef Chaos::TPBDCollisionConstraints<float, 3> FPBDCollisionConstraints;
 
 		typedef FPBDCollisionConstraints FCollisionConstraints;
 		typedef TPBDRigidDynamicSpringConstraints<float, 3> FRigidDynamicSpringConstraints;
diff --git a/Engine/Source/Runtime/Experimental/ChaosSolvers/Public/PhysicsProxy/GeometryCollectionPhysicsProxy.h b/Engine/Source/Runtime/Experimental/ChaosSolvers/Public/PhysicsProxy/GeometryCollectionPhysicsProxy.h
--- a/Engine/Source/Runtime/Experimental/ChaosSolvers/Public/PhysicsProxy/GeometryCollectionPhysicsProxy.h
+++ b/Engine/Source/Runtime/Experimental/ChaosSolvers/Public/PhysicsProxy/GeometryCollectionPhysicsProxy.h
@@ -224,7 +224,7 @@ private:
 	void AddCollisionToCollisionData(
 		FRecordedFrame* ExistingFrame, 
 		const FParticlesType& Particles, 
-		const Chaos::TPBDCollisionConstraint<float, 3>::FPointContactConstraint& Constraint);
+		const Chaos::TPBDCollisionConstraints<float, 3>::FPointContactConstraint& Constraint);
 
 	void UpdateCollisionData(
 		const FParticlesType& Particles, 
diff --git a/Engine/Source/Runtime/PhysicsSQ/Private/SQAccelerator.cpp b/Engine/Source/Runtime/PhysicsSQ/Private/SQAccelerator.cpp
--- a/Engine/Source/Runtime/PhysicsSQ/Private/SQAccelerator.cpp
+++ b/Engine/Source/Runtime/PhysicsSQ/Private/SQAccelerator.cpp
@@ -14,7 +14,7 @@
 
 #include "ChaosInterfaceWrapperCore.h"
 #include "Chaos/ISpatialAcceleration.h"
-#include "Chaos/PBDCollisionConstraint.h"
+#include "Chaos/PBDCollisionConstraints.h"
 #include "Chaos/GeometryQueries.h"
 #include "Chaos/DebugDrawQueue.h"
 
