Chaos - joint matrix solver

diff --git a/Engine/Source/Runtime/Engine/Private/PhysicsEngine/ImmediatePhysics/ImmediatePhysicsChaos/ImmediatePhysicsSimulation_Chaos.cpp b/Engine/Source/Runtime/Engine/Private/PhysicsEngine/ImmediatePhysics/ImmediatePhysicsChaos/ImmediatePhysicsSimulation_Chaos.cpp
--- a/Engine/Source/Runtime/Engine/Private/PhysicsEngine/ImmediatePhysics/ImmediatePhysicsChaos/ImmediatePhysicsSimulation_Chaos.cpp
+++ b/Engine/Source/Runtime/Engine/Private/PhysicsEngine/ImmediatePhysics/ImmediatePhysicsChaos/ImmediatePhysicsSimulation_Chaos.cpp
@@ -45,7 +45,8 @@ int32 ChaosImmediate_Joint_EnableVelocitySolve = 0;
 int32 ChaosImmediate_Joint_EnableTwistLimits = 1;
 int32 ChaosImmediate_Joint_EnableSwingLimits = 1;
 int32 ChaosImmediate_Joint_EnableDrives = 1;
-int32 ChaosImmediate_Joint_ProjectionPhase = (int32)Chaos::EJointProjectionPhase::ApplyPushOut;
+int32 ChaosImmediate_Joint_DrivesPhase = (int32)Chaos::EJointSolverPhase::Apply;
+int32 ChaosImmediate_Joint_ProjectionPhase = (int32)Chaos::EJointSolverPhase::ApplyPushOut;
 float ChaosImmediate_Joint_LinearProjection = 0.0f;
 float ChaosImmediate_Joint_AngularProjection = 0.0f;
 float ChaosImmediate_Joint_Stiffness = 1.0f;
@@ -61,6 +62,7 @@ FAutoConsoleVariableRef CVarChaosImmPhysEnableVelocitySolve(TEXT("p.Chaos.ImmPhy
 FAutoConsoleVariableRef CVarChaosImmPhysEnableTwistLimits(TEXT("p.Chaos.ImmPhys.Joint.EnableTwistLimits"), ChaosImmediate_Joint_EnableTwistLimits, TEXT("EnableTwistLimits."));
 FAutoConsoleVariableRef CVarChaosImmPhysEnableSwingLimits(TEXT("p.Chaos.ImmPhys.Joint.EnableSwingLimits"), ChaosImmediate_Joint_EnableSwingLimits, TEXT("EnableSwingLimits."));
 FAutoConsoleVariableRef CVarChaosImmPhysEnableDrives(TEXT("p.Chaos.ImmPhys.Joint.EnableDrives"), ChaosImmediate_Joint_EnableDrives, TEXT("EnableDrives."));
+FAutoConsoleVariableRef CVarChaosImmPhysDrivesPhase(TEXT("p.Chaos.ImmPhys.Joint.DrivesPhase"), ChaosImmediate_Joint_DrivesPhase, TEXT("Drives Phase (0=Never; 1=Apply; 2=ApplyPushOut)."));
 FAutoConsoleVariableRef CVarChaosImmPhysProjectionPhase(TEXT("p.Chaos.ImmPhys.Joint.ProjectionPhase"), ChaosImmediate_Joint_ProjectionPhase, TEXT("Projection Phase (0=Never; 1=Apply; 2=ApplyPushOut)."));
 FAutoConsoleVariableRef CVarChaosImmPhysLinearProjection(TEXT("p.Chaos.ImmPhys.Joint.LinearProjection"), ChaosImmediate_Joint_LinearProjection, TEXT("6Dof joint projection amount override (if > 0)."));
 FAutoConsoleVariableRef CVarChaosImmPhysAngularProjection(TEXT("p.Chaos.ImmPhys.Joint.AngularProjection"), ChaosImmediate_Joint_AngularProjection, TEXT("6Dof joint projection amount override (if > 0)."));
@@ -103,19 +105,19 @@ FAutoConsoleVariableRef CVarChaosImmPhysDebugDrawJointFeatures(TEXT("p.Chaos.Imm
 
 namespace ImmediatePhysics_Chaos
 {
-	Chaos::EJointProjectionPhase ToProjectionPhase(const int32 Index)
+	Chaos::EJointSolverPhase ToJointSolverPhase(const int32 Index)
 	{
 		using namespace Chaos;
 
 		if (Index == 1)
 		{
-			return EJointProjectionPhase::Apply;
+			return EJointSolverPhase::Apply;
 		}
 		else if (Index == 2)
 		{
-			return EJointProjectionPhase::ApplyPushOut;
+			return EJointSolverPhase::ApplyPushOut;
 		}
-		return EJointProjectionPhase::None;
+		return EJointSolverPhase::None;
 	}
 
 	template<typename T, int d>
@@ -450,7 +452,8 @@ namespace ImmediatePhysics_Chaos
 			JointsSettings.bEnableTwistLimits = ChaosImmediate_Joint_EnableTwistLimits != 0;
 			JointsSettings.bEnableSwingLimits = ChaosImmediate_Joint_EnableSwingLimits != 0;
 			JointsSettings.bEnableDrives = ChaosImmediate_Joint_EnableDrives != 0;
-			JointsSettings.ProjectionPhase = ToProjectionPhase(ChaosImmediate_Joint_ProjectionPhase);
+			JointsSettings.DrivesPhase = ToJointSolverPhase(ChaosImmediate_Joint_DrivesPhase);
+			JointsSettings.ProjectionPhase = ToJointSolverPhase(ChaosImmediate_Joint_ProjectionPhase);
 			JointsSettings.LinearProjection = ChaosImmediate_Joint_LinearProjection;
 			JointsSettings.AngularProjection = ChaosImmediate_Joint_AngularProjection;
 			JointsSettings.Stiffness = ChaosImmediate_Joint_Stiffness;
diff --git a/Engine/Source/Runtime/Experimental/Chaos/Private/Chaos/PBDCollisionConstraintsPointContactUtil.cpp b/Engine/Source/Runtime/Experimental/Chaos/Private/Chaos/PBDCollisionConstraintsPointContactUtil.cpp
--- a/Engine/Source/Runtime/Experimental/Chaos/Private/Chaos/PBDCollisionConstraintsPointContactUtil.cpp
+++ b/Engine/Source/Runtime/Experimental/Chaos/Private/Chaos/PBDCollisionConstraintsPointContactUtil.cpp
@@ -3,6 +3,7 @@
 #include "Chaos/PBDCollisionConstraintsPointContactUtil.h"
 #include "Chaos/CollisionResolutionAlgo.h"
 #include "Chaos/Defines.h"
+#include "Chaos/Utilities.h"
 
 namespace Chaos
 {
@@ -178,8 +179,8 @@ namespace Chaos
 
 				if (RelativeNormalVelocity < 0) // ignore separating constraints
 				{
-					PMatrix<T, d, d> WorldSpaceInvI1 = PBDRigid0 ? (Q0 * FMatrix::Identity).GetTransposed() * PBDRigid0->InvI() * (Q0 * FMatrix::Identity) : PMatrix<T, d, d>(0);
-					PMatrix<T, d, d> WorldSpaceInvI2 = PBDRigid1 ? (Q1 * FMatrix::Identity).GetTransposed() * PBDRigid1->InvI() * (Q1 * FMatrix::Identity) : PMatrix<T, d, d>(0);
+					PMatrix<T, d, d> WorldSpaceInvI1 = PBDRigid0 ? Utilities::ComputeWorldSpaceInertia(Q0, PBDRigid0->InvI()) : PMatrix<T, d, d>(0);
+					PMatrix<T, d, d> WorldSpaceInvI2 = PBDRigid1 ? Utilities::ComputeWorldSpaceInertia(Q1, PBDRigid1->InvI()) : PMatrix<T, d, d>(0);
 					PMatrix<T, d, d> Factor =
 						(PBDRigid0 ? ComputeFactorMatrix3(VectorToPoint1, WorldSpaceInvI1, PBDRigid0->InvM()) : PMatrix<T, d, d>(0)) +
 						(PBDRigid1 ? ComputeFactorMatrix3(VectorToPoint2, WorldSpaceInvI2, PBDRigid1->InvM()) : PMatrix<T, d, d>(0));
@@ -392,8 +393,8 @@ namespace Chaos
 				}
 
 				*IterationParameters.NeedsAnotherIteration = true;
-				PMatrix<T, d, d> WorldSpaceInvI1 = PBDRigid0 ? (Q0 * FMatrix::Identity).GetTransposed() * PBDRigid0->InvI() * (Q0 * FMatrix::Identity) : PMatrix<T, d, d>(0);
-				PMatrix<T, d, d> WorldSpaceInvI2 = PBDRigid1 ? (Q1 * FMatrix::Identity).GetTransposed() * PBDRigid1->InvI() * (Q1 * FMatrix::Identity) : PMatrix<T, d, d>(0);
+				PMatrix<T, d, d> WorldSpaceInvI1 = PBDRigid0 ? Utilities::ComputeWorldSpaceInertia(Q0, PBDRigid0->InvI()) : PMatrix<T, d, d>(0);
+				PMatrix<T, d, d> WorldSpaceInvI2 = PBDRigid1 ? Utilities::ComputeWorldSpaceInertia(Q1, PBDRigid1->InvI()) : PMatrix<T, d, d>(0);
 				TVector<T, d> VectorToPoint1 = Contact.Location - P0;
 				TVector<T, d> VectorToPoint2 = Contact.Location - P1;
 				PMatrix<T, d, d> Factor =
diff --git a/Engine/Source/Runtime/Experimental/Chaos/Private/Chaos/PBDJointConstraintUtilities.cpp b/Engine/Source/Runtime/Experimental/Chaos/Private/Chaos/PBDJointConstraintUtilities.cpp
--- a/Engine/Source/Runtime/Experimental/Chaos/Private/Chaos/PBDJointConstraintUtilities.cpp
+++ b/Engine/Source/Runtime/Experimental/Chaos/Private/Chaos/PBDJointConstraintUtilities.cpp
@@ -1,12 +1,15 @@
 // Copyright 1998-2019 Epic Games, Inc. All Rights Reserved.
 #include "Chaos/PBDJointConstraintUtilities.h"
+#include "Chaos/DenseMatrix.h"
 #include "Chaos/ParticleHandle.h"
+#include "Chaos/PBDJointConstraintSolver.h"
 #include "Chaos/Utilities.h"
 
 //#pragma optimize("", off)
 
 namespace Chaos
 {
+
 	FVec3 FPBDJointUtilities::ConditionInertia(const FVec3& InI, const FReal MaxRatio)
 	{
 		FReal IMin = InI.Min();
@@ -466,8 +469,8 @@ namespace Chaos
 		const FVec3& Axis1,
 		const FReal WC)
 	{
-		const FMatrix33 InvI0 = Utilities::Multiply(Q0.ToMatrix(), Utilities::Multiply(InvIL0, Q0.ToMatrix().GetTransposed()));
-		const FMatrix33 InvI1 = Utilities::Multiply(Q1.ToMatrix(), Utilities::Multiply(InvIL1, Q1.ToMatrix().GetTransposed()));
+		const FMatrix33 InvI0 = Utilities::ComputeWorldSpaceInertia(Q0, InvIL0);
+		const FMatrix33 InvI1 = Utilities::ComputeWorldSpaceInertia(Q1, InvIL1);
 		//const FReal I0 = FVec3::DotProduct(Axis0, Utilities::Multiply(InvI0, Axis0));
 		//const FReal I1 = FVec3::DotProduct(Axis1, Utilities::Multiply(InvI1, Axis1));
 		//const FVec3 DW0 = Axis0 * WC * I0 / (I0 + I1);
@@ -500,8 +503,8 @@ namespace Chaos
 		const FVec3& Axis1,
 		const FReal Angle1)
 	{
-		const FMatrix33 InvI0 = Utilities::Multiply(Q0.ToMatrix(), Utilities::Multiply(InvIL0, Q0.ToMatrix().GetTransposed()));
-		const FMatrix33 InvI1 = Utilities::Multiply(Q1.ToMatrix(), Utilities::Multiply(InvIL1, Q1.ToMatrix().GetTransposed()));
+		const FMatrix33 InvI0 = Utilities::ComputeWorldSpaceInertia(Q0, InvIL0);
+		const FMatrix33 InvI1 = Utilities::ComputeWorldSpaceInertia(Q1, InvIL1);
 
 		//const FReal I0 = FVec3::DotProduct(Axis0, Utilities::Multiply(InvI0, Axis0));
 		//const FReal I1 = FVec3::DotProduct(Axis1, Utilities::Multiply(InvI1, Axis1));
@@ -681,13 +684,14 @@ namespace Chaos
 		FReal InvM1,
 		const FMatrix33& InvIL1)
 	{
+
 		const FRigidTransform3& XL0 = JointSettings.ConstraintFrames[Index0];
 		const FRigidTransform3& XL1 = JointSettings.ConstraintFrames[Index1];
 		const FVec3 X0 = P0 + Q0 * XL0.GetTranslation();
 		const FVec3 X1 = P1 + Q1 * XL1.GetTranslation();
 		const FRotation3 R0 = Q0 * XL0.GetRotation();
-		const FMatrix33 InvI0 = Utilities::Multiply(Q0.ToMatrix(), Utilities::Multiply(InvIL0, Q0.ToMatrix().GetTransposed()));
-		const FMatrix33 InvI1 = Utilities::Multiply(Q1.ToMatrix(), Utilities::Multiply(InvIL1, Q1.ToMatrix().GetTransposed()));
+		const FMatrix33 InvI0 = Utilities::ComputeWorldSpaceInertia(Q0, InvIL0);
+		const FMatrix33 InvI1 = Utilities::ComputeWorldSpaceInertia(Q1, InvIL1);
 
 		// Calculate constraint error
 		const FVec3 CX = GetLimitedPositionError(JointSettings, R0, X1 - X0);
@@ -711,6 +715,7 @@ namespace Chaos
 		const FVec3 DP1 = -InvM1 * DX;
 		const FVec3 DR0 = Utilities::Multiply(InvI0, FVec3::CrossProduct(X0 - P0, DX));
 		const FVec3 DR1 = Utilities::Multiply(InvI1, FVec3::CrossProduct(X1 - P1, -DX));
+
 		ApplyPositionDelta(Dt, SolverSettings, JointSettings, Stiffness, Index0, Index1, P0, V0, P1, V1, DP0, DP1);
 		ApplyRotationDelta(Dt, SolverSettings, JointSettings, Stiffness, Index0, Index1, Q0, W0, Q1, W1, DR0, DR1);
 	}
@@ -741,8 +746,8 @@ namespace Chaos
 		const FVec3 XC0 = Q0 * XL0.GetTranslation();
 		const FVec3 XC1 = Q1 * XL1.GetTranslation();
 		const FRotation3 R0 = Q0 * XL0.GetRotation();
-		const FMatrix33 InvI0 = Utilities::Multiply(Q0.ToMatrix(), Utilities::Multiply(InvIL0, Q0.ToMatrix().GetTransposed()));
-		const FMatrix33 InvI1 = Utilities::Multiply(Q1.ToMatrix(), Utilities::Multiply(InvIL1, Q1.ToMatrix().GetTransposed()));
+		const FMatrix33 InvI0 = Utilities::ComputeWorldSpaceInertia(Q0, InvIL0);
+		const FMatrix33 InvI1 = Utilities::ComputeWorldSpaceInertia(Q1, InvIL1);
 
 		const FVec3 VC0 = V0 + FVec3::CrossProduct(W0, XC0);
 		const FVec3 VC1 = V1 + FVec3::CrossProduct(W1, XC1);
@@ -1449,8 +1454,8 @@ namespace Chaos
 		FReal SLerpAngle;
 		if (DR1.ToAxisAndAngleSafe(SLerpAxis, SLerpAngle, FVec3(1, 0, 0)))
 		{
-			const FMatrix33 InvI0 = Utilities::Multiply(Q0.ToMatrix(), Utilities::Multiply(InvIL0, Q0.ToMatrix().GetTransposed()));
-			const FMatrix33 InvI1 = Utilities::Multiply(Q1.ToMatrix(), Utilities::Multiply(InvIL1, Q1.ToMatrix().GetTransposed()));
+			const FMatrix33 InvI0 = Utilities::ComputeWorldSpaceInertia(Q0, InvIL0);
+			const FMatrix33 InvI1 = Utilities::ComputeWorldSpaceInertia(Q1, InvIL1);
 			const FReal I0 = FVec3::DotProduct(SLerpAxis, Utilities::Multiply(InvI0, SLerpAxis));
 			const FReal I1 = FVec3::DotProduct(SLerpAxis, Utilities::Multiply(InvI1, SLerpAxis));
 			const FReal F0 = DriveStiffness * I0 / (I0 + I1);
@@ -1566,4 +1571,7 @@ namespace Chaos
 		FVec3 W1(0, 0, 0);
 		ApplyJointSwingConstraint(Dt, SolverSettings, JointSettings, Stiffness, Index0, Index1, SwingConstraintIndex, SwingAxisIndex, P0, Q0, V0, W0, P1, Q1, V1, W1, InvM0, InvIL0, InvM1, InvIL1);
 	}
+
+
+
 }
diff --git a/Engine/Source/Runtime/Experimental/Chaos/Private/Chaos/PBDJointConstraints.cpp b/Engine/Source/Runtime/Experimental/Chaos/Private/Chaos/PBDJointConstraints.cpp
--- a/Engine/Source/Runtime/Experimental/Chaos/Private/Chaos/PBDJointConstraints.cpp
+++ b/Engine/Source/Runtime/Experimental/Chaos/Private/Chaos/PBDJointConstraints.cpp
@@ -3,6 +3,7 @@
 #include "Chaos/ChaosDebugDraw.h"
 #include "Chaos/DebugDrawQueue.h"
 #include "Chaos/ParticleHandle.h"
+#include "Chaos/PBDJointConstraintSolver.h"
 #include "Chaos/PBDJointConstraintUtilities.h"
 #include "Chaos/Utilities.h"
 #include "ChaosLog.h"
@@ -15,7 +16,10 @@
 namespace Chaos
 {
 	DECLARE_CYCLE_STAT(TEXT("TPBDJointConstraints::Apply"), STAT_ApplyJointConstraints, STATGROUP_Chaos);
-	DECLARE_CYCLE_STAT(TEXT("TPBDJointConstraints::Apply"), STAT_ApplyPushOutJointConstraints, STATGROUP_Chaos);
+	DECLARE_CYCLE_STAT(TEXT("TPBDJointConstraints::ApplyPushOut"), STAT_ApplyPushOutJointConstraints, STATGROUP_Chaos);
+
+	bool ChaosJoint_UseCholeskySolver = true;
+	FAutoConsoleVariableRef CVarChaosImmPhysDeltaTime(TEXT("p.Chaos.Joint.UseCholeskySolver"), ChaosJoint_UseCholeskySolver, TEXT("Whether to use the new solver"));
 
 	FReal GetLinearStiffness(
 		const FPBDJointSolverSettings& SolverSettings,
@@ -176,7 +180,8 @@ namespace Chaos
 		, bEnableTwistLimits(true)
 		, bEnableSwingLimits(true)
 		, bEnableDrives(true)
-		, ProjectionPhase(EJointProjectionPhase::None)
+		, DrivesPhase(EJointSolverPhase::None)
+		, ProjectionPhase(EJointSolverPhase::None)
 		, LinearProjection((FReal)0)
 		, AngularProjection((FReal)0)
 		, Stiffness((FReal)0)
@@ -394,8 +399,8 @@ namespace Chaos
 		TArray<FConstraintContainerHandle*> SortedConstraintHandles = InConstraintHandles;
 		SortedConstraintHandles.Sort([](const FConstraintContainerHandle& L, const FConstraintContainerHandle& R)
 			{
-				// Sort bodies from leaf to root
-				return L.GetConstraintLevel() > R.GetConstraintLevel();
+				// Sort bodies from root to leaf
+				return L.GetConstraintLevel() < R.GetConstraintLevel();
 			});
 
 		if (PreApplyCallback != nullptr)
@@ -403,19 +408,36 @@ namespace Chaos
 			PreApplyCallback(Dt, SortedConstraintHandles);
 		}
 
-		for (FConstraintContainerHandle* ConstraintHandle : SortedConstraintHandles)
+		// Apply joint drives
+		if (Settings.DrivesPhase == EJointSolverPhase::Apply)
 		{
-			if (Settings.bEnableVelocitySolve)
+			if (It == 0)
 			{
-				SolveVelocity(Dt, ConstraintHandle->GetConstraintIndex(), Settings.ApplyPairIterations, It, NumIts);
+				for (FConstraintContainerHandle* ConstraintHandle : SortedConstraintHandles)
+				{
+					ApplyDrives(Dt, ConstraintHandle->GetConstraintIndex(), 1, 0, 1);
+				}
 			}
-			else
+		}
+
+		// Solve for joint position or velocity, depending on settings
+		if (Settings.ApplyPairIterations > 0)
+		{
+			for (FConstraintContainerHandle* ConstraintHandle : SortedConstraintHandles)
 			{
-				SolvePosition(Dt, ConstraintHandle->GetConstraintIndex(), Settings.ApplyPairIterations, It, NumIts);
+				if (Settings.bEnableVelocitySolve)
+				{
+					SolveVelocity(Dt, ConstraintHandle->GetConstraintIndex(), Settings.ApplyPairIterations, It, NumIts);
+				}
+				else
+				{
+					SolvePosition(Dt, ConstraintHandle->GetConstraintIndex(), Settings.ApplyPairIterations, It, NumIts);
+				}
 			}
 		}
 
-		if (Settings.ProjectionPhase == EJointProjectionPhase::Apply)
+		// Correct remaining errors after last call to Solve if enabled in this phase
+		if (Settings.ProjectionPhase == EJointSolverPhase::Apply)
 		{
 			int32 ProjectionIt = NumIts - 1;
 			if (It == ProjectionIt)
@@ -438,22 +460,36 @@ namespace Chaos
 		// @todo(ccaulfield): track whether we are sufficiently solved
 		bool bNeedsAnotherIteration = true;
 
-		if (Settings.ApplyPushOutPairIterations > 0)
+		TArray<FConstraintContainerHandle*> SortedConstraintHandles = InConstraintHandles;
+		SortedConstraintHandles.Sort([](const FConstraintContainerHandle& L, const FConstraintContainerHandle& R)
+			{
+				// Sort bodies from root to leaf
+				return L.GetConstraintLevel() < R.GetConstraintLevel();
+			});
+
+		// Apply joint drives
+		if (Settings.DrivesPhase == EJointSolverPhase::ApplyPushOut)
 		{
-			TArray<FConstraintContainerHandle*> SortedConstraintHandles = InConstraintHandles;
-			SortedConstraintHandles.Sort([](const FConstraintContainerHandle& L, const FConstraintContainerHandle& R)
+			if (It == 0)
+			{
+				for (FConstraintContainerHandle* ConstraintHandle : SortedConstraintHandles)
 				{
-					// Sort bodies from root to leaf
-					return L.GetConstraintLevel() < R.GetConstraintLevel();
-				});
+					ApplyDrives(Dt, ConstraintHandle->GetConstraintIndex(), 1, 0, 1);
+				}
+			}
+		}
 
+		// Solve for positions
+		if (Settings.ApplyPushOutPairIterations > 0)
+		{
 			for (FConstraintContainerHandle* ConstraintHandle : SortedConstraintHandles)
 			{
 				SolvePosition(Dt, ConstraintHandle->GetConstraintIndex(), Settings.ApplyPushOutPairIterations, It, NumIts);
 			}
 		}
 
-		if (Settings.ProjectionPhase == EJointProjectionPhase::ApplyPushOut)
+		// Correct remaining errors after the last call to Solve (which depends on if PositionSolve is enabled in ApplyPushOut)
+		if (Settings.ProjectionPhase == EJointSolverPhase::ApplyPushOut)
 		{
 			int32 ProjectionIt = (Settings.ApplyPushOutPairIterations > 0) ? NumIts - 1 : 0;
 			if (It == ProjectionIt)
@@ -481,8 +517,145 @@ namespace Chaos
 		}
 	}
 
-	
-	void FPBDJointConstraints::SolveVelocity(const FReal Dt, const int32 ConstraintIndex, const int32 NumPairIts, const int32 It, const int32 NumIts)
+
+	void FPBDJointConstraints::UpdateParticleState(TPBDRigidParticleHandle<FReal, 3>* Rigid, const FReal Dt, const FVec3& P, const FRotation3& Q)
+	{
+		if (Rigid)
+		{
+			if (Dt > SMALL_NUMBER)
+			{
+				FVec3 DV = FVec3::CalculateVelocity(Rigid->P(), P, Dt);
+				FVec3 DW = FRotation3::CalculateAngularVelocity(Rigid->Q(), Q, Dt);
+				Rigid->SetV(Rigid->V() + DV);
+				Rigid->SetW(Rigid->W() + DW);
+			}
+			Rigid->SetP(P);
+			Rigid->SetQ(Q);
+		}
+	}
+
+
+	void FPBDJointConstraints::ApplyDrives(const FReal Dt, const int32 ConstraintIndex, const int32 NumPairIts, const int32 It, const int32 NumIts)
+	{
+		const FPBDJointSettings& JointSettings = ConstraintSettings[ConstraintIndex];
+
+		const int32 Index0 = 1;
+		const int32 Index1 = 0;
+		TGenericParticleHandle<FReal, 3> Particle0 = TGenericParticleHandle<FReal, 3>(ConstraintParticles[ConstraintIndex][Index0]);
+		TGenericParticleHandle<FReal, 3> Particle1 = TGenericParticleHandle<FReal, 3>(ConstraintParticles[ConstraintIndex][Index1]);
+
+		FVec3 P0 = Particle0->P();
+		FRotation3 Q0 = Particle0->Q();
+		FVec3 V0 = Particle0->V();
+		FVec3 W0 = Particle0->W();
+		FVec3 P1 = Particle1->P();
+		FRotation3 Q1 = Particle1->Q();
+		FVec3 V1 = Particle1->V();
+		FVec3 W1 = Particle1->W();
+		float InvM0 = Particle0->InvM();
+		float InvM1 = Particle1->InvM();
+		FMatrix33 InvIL0 = Particle0->InvI();
+		FMatrix33 InvIL1 = Particle1->InvI();
+
+		EJointMotionType TwistMotion = JointSettings.Motion.AngularMotionTypes[(int32)EJointAngularConstraintIndex::Twist];
+		EJointMotionType Swing1Motion = JointSettings.Motion.AngularMotionTypes[(int32)EJointAngularConstraintIndex::Swing1];
+		EJointMotionType Swing2Motion = JointSettings.Motion.AngularMotionTypes[(int32)EJointAngularConstraintIndex::Swing2];
+
+		// Apply angular drives (NOTE: modifies position, not velocity)
+		if (Settings.bEnableDrives)
+		{
+			bool bTwistLocked = TwistMotion == EJointMotionType::Locked;
+			bool bSwing1Locked = Swing1Motion == EJointMotionType::Locked;
+			bool bSwing2Locked = Swing2Motion == EJointMotionType::Locked;
+
+			// No SLerp drive if we have a locked rotation (it will be grayed out in the editor in this case, but could still have been set before the rotation was locked)
+			if (JointSettings.Motion.bAngularSLerpDriveEnabled && !bTwistLocked && !bSwing1Locked && !bSwing2Locked)
+			{
+				FPBDJointUtilities::ApplyJointSLerpDrive(Dt, Settings, JointSettings, Index0, Index1, P0, Q0, V0, W0, P1, Q1, V1, W1, InvM0, InvIL0, InvM1, InvIL1);
+			}
+
+			if (JointSettings.Motion.bAngularTwistDriveEnabled && !bTwistLocked)
+			{
+				FPBDJointUtilities::ApplyJointTwistDrive(Dt, Settings, JointSettings, Index0, Index1, P0, Q0, V0, W0, P1, Q1, V1, W1, InvM0, InvIL0, InvM1, InvIL1);
+			}
+
+			if (JointSettings.Motion.bAngularSwingDriveEnabled && !bSwing1Locked && !bSwing2Locked)
+			{
+				FPBDJointUtilities::ApplyJointConeDrive(Dt, Settings, JointSettings, Index0, Index1, P0, Q0, V0, W0, P1, Q1, V1, W1, InvM0, InvIL0, InvM1, InvIL1);
+			}
+			else if (JointSettings.Motion.bAngularSwingDriveEnabled && !bSwing1Locked)
+			{
+				//FPBDJointUtilities::ApplyJointSwingDrive(Dt, Settings, JointSettings, Index0, Index1, EJointAngularConstraintIndex::Swing1, P0, Q0, P1, Q1, InvM0, InvIL0, InvM1, InvIL1);
+			}
+			else if (JointSettings.Motion.bAngularSwingDriveEnabled && !bSwing2Locked)
+			{
+				//FPBDJointUtilities::ApplyJointSwingDrive(Dt, Settings, JointSettings, Index0, Index1, EJointAngularConstraintIndex::Swing2, P0, Q0, P1, Q1, InvM0, InvIL0, InvM1, InvIL1);
+			}
+		}
+
+		UpdateParticleState(Particle0->AsDynamic(), Dt, P0, Q0);
+		UpdateParticleState(Particle1->AsDynamic(), Dt, P1, Q1);
+	}
+
+
+	void FPBDJointConstraints::SolvePosition(const FReal Dt, const int32 ConstraintIndex, const int32 NumPairIts, const int32 It, const int32 NumIts)
+	{
+		if (ChaosJoint_UseCholeskySolver)
+		{
+			SolvePosition_Cholesky(Dt, ConstraintIndex, NumPairIts, It, NumIts);
+		}
+		else
+		{
+			SolvePosition_GaussSiedel(Dt, ConstraintIndex, NumPairIts, It, NumIts);
+		}
+	}
+
+
+	// This position solver solves all (active) inner position and angular constraints simultaneously by building the Jacobian and solving [JMJt].DX = C
+	// where DX(6x1) are the unknown position and rotation corrections, C(Nx1) is the current constraint error, J(Nx6) the Jacobian, M(6x6) the inverse mass matrix, 
+	// and N the number of active constraints. "Active constraints" are all bilateral constraints plus any violated unilateral constraints.
+	void FPBDJointConstraints::SolvePosition_Cholesky(const FReal Dt, const int32 ConstraintIndex, const int32 NumPairIts, const int32 It, const int32 NumIts)
+	{
+		const TVector<TGeometryParticleHandle<FReal, 3>*, 2>& Constraint = ConstraintParticles[ConstraintIndex];
+		UE_LOG(LogChaosJoint, Verbose, TEXT("Solve Joint Constraint %d %s %s (dt = %f; it = %d / %d)"), ConstraintIndex, *Constraint[0]->ToString(), *Constraint[1]->ToString(), Dt, It, NumIts);
+
+		const FPBDJointSettings& JointSettings = ConstraintSettings[ConstraintIndex];
+
+		// Switch particles - internally we assume the first body is the parent (i.e., the space in which constraint limits are specified)
+		const int32 Index0 = 1;
+		const int32 Index1 = 0;
+		TGenericParticleHandle<FReal, 3> Particle0 = TGenericParticleHandle<FReal, 3>(ConstraintParticles[ConstraintIndex][Index0]);
+		TGenericParticleHandle<FReal, 3> Particle1 = TGenericParticleHandle<FReal, 3>(ConstraintParticles[ConstraintIndex][Index1]);
+
+		FJointConstraintSolver Solver;
+		Solver.InitConstraints(
+			Dt, 
+			Settings, 
+			JointSettings, 
+			Particle0->P(), 
+			Particle0->Q(), 
+			Particle1->P(), 
+			Particle1->Q(), 
+			Particle0->InvM(), 
+			Particle0->InvI(), 
+			Particle1->InvM(), 
+			Particle1->InvI(), 
+			JointSettings.ConstraintFrames[Index0], 
+			JointSettings.ConstraintFrames[Index1]);
+		
+		for (int32 PairIt = 0; PairIt < NumPairIts; ++PairIt)
+		{
+			Solver.ApplyConstraints(Dt, Settings, JointSettings, GetLinearStiffness(Settings, JointSettings));
+		}
+
+		UpdateParticleState(Particle0->AsDynamic(), Dt, Solver.GetP(0), Solver.GetQ(0));
+		UpdateParticleState(Particle1->AsDynamic(), Dt, Solver.GetP(1), Solver.GetQ(1));
+	}
+
+
+	// This position solver iterates over each of the inner constraints (position, twist, swing) and solves them independently.
+	// This will converge slowly in some cases, particularly where resolving angular constraints violates position constraints and vice versa.
+	void FPBDJointConstraints::SolvePosition_GaussSiedel(const FReal Dt, const int32 ConstraintIndex, const int32 NumPairIts, const int32 It, const int32 NumIts)
 	{
 		const TVector<TGeometryParticleHandle<FReal, 3>*, 2>& Constraint = ConstraintParticles[ConstraintIndex];
 		UE_LOG(LogChaosJoint, Verbose, TEXT("Solve Joint Constraint %d %s %s (dt = %f; it = %d / %d)"), ConstraintIndex, *Constraint[0]->ToString(), *Constraint[1]->ToString(), Dt, It, NumIts);
@@ -539,96 +712,52 @@ namespace Chaos
 
 		for (int32 PairIt = 0; PairIt < NumPairIts; ++PairIt)
 		{
-			// Apply angular drives (NOTE: modifies position, not velocity)
-			if (Settings.bEnableDrives)
-			{
-				bool bTwistLocked = TwistMotion == EJointMotionType::Locked;
-				bool bSwing1Locked = Swing1Motion == EJointMotionType::Locked;
-				bool bSwing2Locked = Swing2Motion == EJointMotionType::Locked;
-
-				// No SLerp drive if we have a locked rotation (it will be grayed out in the editor in this case, but could still have been set before the rotation was locked)
-				if (JointSettings.Motion.bAngularSLerpDriveEnabled && !bTwistLocked && !bSwing1Locked && !bSwing2Locked)
-				{
-					FPBDJointUtilities::ApplyJointSLerpDrive(Dt, Settings, JointSettings, Index0, Index1, P0, Q0, V0, W0, P1, Q1, V1, W1, InvM0, InvIL0, InvM1, InvIL1);
-				}
-
-				if (JointSettings.Motion.bAngularTwistDriveEnabled && !bTwistLocked)
-				{
-					FPBDJointUtilities::ApplyJointTwistDrive(Dt, Settings, JointSettings, Index0, Index1, P0, Q0, V0, W0, P1, Q1, V1, W1, InvM0, InvIL0, InvM1, InvIL1);
-				}
-
-				if (JointSettings.Motion.bAngularSwingDriveEnabled && !bSwing1Locked && !bSwing2Locked)
-				{
-					FPBDJointUtilities::ApplyJointConeDrive(Dt, Settings, JointSettings, Index0, Index1, P0, Q0, V0, W0, P1, Q1, V1, W1, InvM0, InvIL0, InvM1, InvIL1);
-				}
-				else if (JointSettings.Motion.bAngularSwingDriveEnabled && !bSwing1Locked)
-				{
-					//FPBDJointUtilities::ApplyJointSwingDrive(Dt, Settings, JointSettings, Index0, Index1, EJointAngularConstraintIndex::Swing1, P0, Q0, P1, Q1, InvM0, InvIL0, InvM1, InvIL1);
-				}
-				else if (JointSettings.Motion.bAngularSwingDriveEnabled && !bSwing2Locked)
-				{
-					//FPBDJointUtilities::ApplyJointSwingDrive(Dt, Settings, JointSettings, Index0, Index1, EJointAngularConstraintIndex::Swing2, P0, Q0, P1, Q1, InvM0, InvIL0, InvM1, InvIL1);
-				}
-			}
-
-			// Apply twist velocity constraint
+			// Apply twist constraint
 			if (Settings.bEnableTwistLimits)
 			{
 				if (TwistMotion != EJointMotionType::Free)
 				{
-					FPBDJointUtilities::ApplyJointTwistVelocityConstraint(Dt, Settings, JointSettings, TwistStiffness, Index0, Index1, P0, Q0, V0, W0, P1, Q1, V1, W1, InvM0, InvIL0, InvM1, InvIL1);
+					FPBDJointUtilities::ApplyJointTwistConstraint(Dt, Settings, JointSettings, TwistStiffness, Index0, Index1, P0, Q0, V0, W0, P1, Q1, V1, W1, InvM0, InvIL0, InvM1, InvIL1);
 				}
 			}
 
-			// Apply swing velocity constraints
+			// Apply swing constraints
 			if (Settings.bEnableSwingLimits)
 			{
 				if ((Swing1Motion == EJointMotionType::Limited) && (Swing2Motion == EJointMotionType::Limited))
 				{
 					// Swing Cone
-					FPBDJointUtilities::ApplyJointConeVelocityConstraint(Dt, Settings, JointSettings, SwingStiffness, Index0, Index1, P0, Q0, V0, W0, P1, Q1, V1, W1, InvM0, InvIL0, InvM1, InvIL1);
+					FPBDJointUtilities::ApplyJointConeConstraint(Dt, Settings, JointSettings, SwingStiffness, Index0, Index1, P0, Q0, V0, W0, P1, Q1, V1, W1, InvM0, InvIL0, InvM1, InvIL1);
 				}
 				else
 				{
 					if (Swing1Motion != EJointMotionType::Free)
 					{
 						// Swing Arc/Lock
-						FPBDJointUtilities::ApplyJointSwingVelocityConstraint(Dt, Settings, JointSettings, SwingStiffness, Index0, Index1, EJointAngularConstraintIndex::Swing1, EJointAngularAxisIndex::Swing1, P0, Q0, V0, W0, P1, Q1, V1, W1, InvM0, InvIL0, InvM1, InvIL1);
+						FPBDJointUtilities::ApplyJointSwingConstraint(Dt, Settings, JointSettings, SwingStiffness, Index0, Index1, EJointAngularConstraintIndex::Swing1, EJointAngularAxisIndex::Swing1, P0, Q0, V0, W0, P1, Q1, V1, W1, InvM0, InvIL0, InvM1, InvIL1);
 					}
 					if (Swing2Motion != EJointMotionType::Free)
 					{
 						// Swing Arc/Lock
-						FPBDJointUtilities::ApplyJointSwingVelocityConstraint(Dt, Settings, JointSettings, SwingStiffness, Index0, Index1, EJointAngularConstraintIndex::Swing2, EJointAngularAxisIndex::Swing2, P0, Q0, V0, W0, P1, Q1, V1, W1, InvM0, InvIL0, InvM1, InvIL1);
+						FPBDJointUtilities::ApplyJointSwingConstraint(Dt, Settings, JointSettings, SwingStiffness, Index0, Index1, EJointAngularConstraintIndex::Swing2, EJointAngularAxisIndex::Swing2, P0, Q0, V0, W0, P1, Q1, V1, W1, InvM0, InvIL0, InvM1, InvIL1);
 					}
 				}
 			}
 
-			// Apply linear velocity  constraints
+			// Apply linear constraints
 			if ((LinearMotion[0] != EJointMotionType::Free) || (LinearMotion[1] != EJointMotionType::Free) || (LinearMotion[2] != EJointMotionType::Free))
 			{
-				FPBDJointUtilities::ApplyJointVelocityConstraint(Dt, Settings, JointSettings, LinearStiffness, Index0, Index1, P0, Q0, V0, W0, P1, Q1, V1, W1, InvM0, InvIL0, InvM1, InvIL1);
+				FPBDJointUtilities::ApplyJointPositionConstraint(Dt, Settings, JointSettings, LinearStiffness, Index0, Index1, P0, Q0, V0, W0, P1, Q1, V1, W1, InvM0, InvIL0, InvM1, InvIL1);
 			}
 		}
 
 		// Update the particles
-		if (Rigid0)
-		{
-			Rigid0->SetP(P0);
-			Rigid0->SetQ(Q0);
-			Rigid0->SetV(V0);
-			Rigid0->SetW(W0);
-		}
-		if (Rigid1)
-		{
-			Rigid1->SetP(P1);
-			Rigid1->SetQ(Q1);
-			Rigid1->SetV(V1);
-			Rigid1->SetW(W1);
-		}
+		UpdateParticleState(Rigid0, Dt, P0, Q0);
+		UpdateParticleState(Rigid1, Dt, P1, Q1);
 	}
 
 	
-	void FPBDJointConstraints::SolvePosition(const FReal Dt, const int32 ConstraintIndex, const int32 NumPairIts, const int32 It, const int32 NumIts)
+	void FPBDJointConstraints::SolveVelocity(const FReal Dt, const int32 ConstraintIndex, const int32 NumPairIts, const int32 It, const int32 NumIts)
 	{
 		const TVector<TGeometryParticleHandle<FReal, 3>*, 2>& Constraint = ConstraintParticles[ConstraintIndex];
 		UE_LOG(LogChaosJoint, Verbose, TEXT("Solve Joint Constraint %d %s %s (dt = %f; it = %d / %d)"), ConstraintIndex, *Constraint[0]->ToString(), *Constraint[1]->ToString(), Dt, It, NumIts);
@@ -685,95 +814,51 @@ namespace Chaos
 
 		for (int32 PairIt = 0; PairIt < NumPairIts; ++PairIt)
 		{
-			// Apply angular drives (NOTE: modifies position, not velocity)
-			if (!Settings.bEnableVelocitySolve && Settings.bEnableDrives)
-			{
-				bool bTwistLocked = TwistMotion == EJointMotionType::Locked;
-				bool bSwing1Locked = Swing1Motion == EJointMotionType::Locked;
-				bool bSwing2Locked = Swing2Motion == EJointMotionType::Locked;
-
-				// No SLerp drive if we have a locked rotation (it will be grayed out in the editor in this case, but could still have been set before the rotation was locked)
-				if (JointSettings.Motion.bAngularSLerpDriveEnabled && !bTwistLocked && !bSwing1Locked && !bSwing2Locked)
-				{
-					FPBDJointUtilities::ApplyJointSLerpDrive(Dt, Settings, JointSettings, Index0, Index1, P0, Q0, V0, W0, P1, Q1, V1, W1, InvM0, InvIL0, InvM1, InvIL1);
-				}
-
-				if (JointSettings.Motion.bAngularTwistDriveEnabled && !bTwistLocked)
-				{
-					FPBDJointUtilities::ApplyJointTwistDrive(Dt, Settings, JointSettings, Index0, Index1, P0, Q0, V0, W0, P1, Q1, V1, W1, InvM0, InvIL0, InvM1, InvIL1);
-				}
-
-				if (JointSettings.Motion.bAngularSwingDriveEnabled && !bSwing1Locked && !bSwing2Locked)
-				{
-					FPBDJointUtilities::ApplyJointConeDrive(Dt, Settings, JointSettings, Index0, Index1, P0, Q0, V0, W0, P1, Q1, V1, W1, InvM0, InvIL0, InvM1, InvIL1);
-				}
-				else if (JointSettings.Motion.bAngularSwingDriveEnabled && !bSwing1Locked)
-				{
-					//FPBDJointUtilities::ApplyJointSwingDrive(Dt, Settings, JointSettings, Index0, Index1, EJointAngularConstraintIndex::Swing1, P0, Q0, P1, Q1, InvM0, InvIL0, InvM1, InvIL1);
-				}
-				else if (JointSettings.Motion.bAngularSwingDriveEnabled && !bSwing2Locked)
-				{
-					//FPBDJointUtilities::ApplyJointSwingDrive(Dt, Settings, JointSettings, Index0, Index1, EJointAngularConstraintIndex::Swing2, P0, Q0, P1, Q1, InvM0, InvIL0, InvM1, InvIL1);
-				}
-			}
-
-			// Apply twist constraint
+			// Apply twist velocity constraint
 			if (Settings.bEnableTwistLimits)
 			{
 				if (TwistMotion != EJointMotionType::Free)
 				{
-					FPBDJointUtilities::ApplyJointTwistConstraint(Dt, Settings, JointSettings, TwistStiffness, Index0, Index1, P0, Q0, V0, W0, P1, Q1, V1, W1, InvM0, InvIL0, InvM1, InvIL1);
+					FPBDJointUtilities::ApplyJointTwistVelocityConstraint(Dt, Settings, JointSettings, TwistStiffness, Index0, Index1, P0, Q0, V0, W0, P1, Q1, V1, W1, InvM0, InvIL0, InvM1, InvIL1);
 				}
 			}
 
-			// Apply swing constraints
+			// Apply swing velocity constraints
 			if (Settings.bEnableSwingLimits)
 			{
 				if ((Swing1Motion == EJointMotionType::Limited) && (Swing2Motion == EJointMotionType::Limited))
 				{
 					// Swing Cone
-					FPBDJointUtilities::ApplyJointConeConstraint(Dt, Settings, JointSettings, SwingStiffness, Index0, Index1, P0, Q0, V0, W0, P1, Q1, V1, W1, InvM0, InvIL0, InvM1, InvIL1);
+					FPBDJointUtilities::ApplyJointConeVelocityConstraint(Dt, Settings, JointSettings, SwingStiffness, Index0, Index1, P0, Q0, V0, W0, P1, Q1, V1, W1, InvM0, InvIL0, InvM1, InvIL1);
 				}
 				else
 				{
 					if (Swing1Motion != EJointMotionType::Free)
 					{
 						// Swing Arc/Lock
-						FPBDJointUtilities::ApplyJointSwingConstraint(Dt, Settings, JointSettings, SwingStiffness, Index0, Index1, EJointAngularConstraintIndex::Swing1, EJointAngularAxisIndex::Swing1, P0, Q0, V0, W0, P1, Q1, V1, W1, InvM0, InvIL0, InvM1, InvIL1);
+						FPBDJointUtilities::ApplyJointSwingVelocityConstraint(Dt, Settings, JointSettings, SwingStiffness, Index0, Index1, EJointAngularConstraintIndex::Swing1, EJointAngularAxisIndex::Swing1, P0, Q0, V0, W0, P1, Q1, V1, W1, InvM0, InvIL0, InvM1, InvIL1);
 					}
 					if (Swing2Motion != EJointMotionType::Free)
 					{
 						// Swing Arc/Lock
-						FPBDJointUtilities::ApplyJointSwingConstraint(Dt, Settings, JointSettings, SwingStiffness, Index0, Index1, EJointAngularConstraintIndex::Swing2, EJointAngularAxisIndex::Swing2, P0, Q0, V0, W0, P1, Q1, V1, W1, InvM0, InvIL0, InvM1, InvIL1);
+						FPBDJointUtilities::ApplyJointSwingVelocityConstraint(Dt, Settings, JointSettings, SwingStiffness, Index0, Index1, EJointAngularConstraintIndex::Swing2, EJointAngularAxisIndex::Swing2, P0, Q0, V0, W0, P1, Q1, V1, W1, InvM0, InvIL0, InvM1, InvIL1);
 					}
 				}
 			}
 
-			// Apply linear constraints
+			// Apply linear velocity  constraints
 			if ((LinearMotion[0] != EJointMotionType::Free) || (LinearMotion[1] != EJointMotionType::Free) || (LinearMotion[2] != EJointMotionType::Free))
 			{
-				FPBDJointUtilities::ApplyJointPositionConstraint(Dt, Settings, JointSettings, LinearStiffness, Index0, Index1, P0, Q0, V0, W0, P1, Q1, V1, W1, InvM0, InvIL0, InvM1, InvIL1);
+				FPBDJointUtilities::ApplyJointVelocityConstraint(Dt, Settings, JointSettings, LinearStiffness, Index0, Index1, P0, Q0, V0, W0, P1, Q1, V1, W1, InvM0, InvIL0, InvM1, InvIL1);
 			}
 		}
 
 		// Update the particles
-		if (Rigid0)
-		{
-			Rigid0->SetP(P0);
-			Rigid0->SetQ(Q0);
-			Rigid0->SetV(V0);
-			Rigid0->SetW(W0);
-		}
-		if (Rigid1)
-		{
-			Rigid1->SetP(P1);
-			Rigid1->SetQ(Q1);
-			Rigid1->SetV(V1);
-			Rigid1->SetW(W1);
-		}
+		UpdateParticleState(Rigid0, Dt, P0, Q0);
+		UpdateParticleState(Rigid1, Dt, P1, Q1);
 	}
 
-	
+
 	void FPBDJointConstraints::ProjectPosition(const FReal Dt, const int32 ConstraintIndex, const int32 It, const int32 NumIts)
 	{
 		const FPBDJointSettings& JointSettings = ConstraintSettings[ConstraintIndex];
diff --git a/Engine/Source/Runtime/Experimental/Chaos/Private/Chaos/PBDRigidDynamicSpringConstraints.cpp b/Engine/Source/Runtime/Experimental/Chaos/Private/Chaos/PBDRigidDynamicSpringConstraints.cpp
--- a/Engine/Source/Runtime/Experimental/Chaos/Private/Chaos/PBDRigidDynamicSpringConstraints.cpp
+++ b/Engine/Source/Runtime/Experimental/Chaos/Private/Chaos/PBDRigidDynamicSpringConstraints.cpp
@@ -1,6 +1,7 @@
 // Copyright 1998-2019 Epic Games, Inc. All Rights Reserved.
 
 #include "Chaos/PBDRigidDynamicSpringConstraints.h"
+#include "Chaos/Utilities.h"
 
 using namespace Chaos;
 
@@ -129,8 +130,8 @@ void TPBDRigidDynamicSpringConstraints<T, d>::ApplySingle(const T Dt, int32 Cons
 	TVector<T, d>& P1 = PBDRigid1 ? PBDRigid1->P() : Static1->X();
 
 	const int32 NumSprings = SpringDistances[ConstraintIndex].Num();
-	const PMatrix<T, d, d> WorldSpaceInvI1 = PBDRigid0? (Q0 * FMatrix::Identity) * PBDRigid0->InvI() * (Q0 * FMatrix::Identity).GetTransposed() : PMatrix<T, d, d>(0);
-	const PMatrix<T, d, d> WorldSpaceInvI2 = PBDRigid1 ? (Q1 * FMatrix::Identity) * PBDRigid1->InvI() * (Q1 * FMatrix::Identity).GetTransposed() : PMatrix<T, d, d>(0);;
+	const PMatrix<T, d, d> WorldSpaceInvI1 = PBDRigid0? Utilities::ComputeWorldSpaceInertia(Q0, PBDRigid0->InvI()) : PMatrix<T, d, d>(0);
+	const PMatrix<T, d, d> WorldSpaceInvI2 = PBDRigid1 ? Utilities::ComputeWorldSpaceInertia(Q1, PBDRigid1->InvI()) : PMatrix<T, d, d>(0);;
 	for (int32 SpringIndex = 0; SpringIndex < NumSprings; ++SpringIndex)
 	{
 		const TVector<T, d>& Distance0 = Distances[ConstraintIndex][SpringIndex][0];
diff --git a/Engine/Source/Runtime/Experimental/Chaos/Private/Chaos/PBDRigidSpringConstraints.cpp b/Engine/Source/Runtime/Experimental/Chaos/Private/Chaos/PBDRigidSpringConstraints.cpp
--- a/Engine/Source/Runtime/Experimental/Chaos/Private/Chaos/PBDRigidSpringConstraints.cpp
+++ b/Engine/Source/Runtime/Experimental/Chaos/Private/Chaos/PBDRigidSpringConstraints.cpp
@@ -1,7 +1,7 @@
 // Copyright 1998-2019 Epic Games, Inc. All Rights Reserved.
 
 #include "Chaos/PBDRigidSpringConstraints.h"
-
+#include "Chaos/Utilities.h"
 
 namespace Chaos
 {
@@ -71,8 +71,8 @@ namespace Chaos
 
 		const TVector<T, d> WorldSpaceX1 = Q0.RotateVector(Distances[ConstraintIndex][0]) + P0;
 		const TVector<T, d> WorldSpaceX2 = Q1.RotateVector(Distances[ConstraintIndex][1]) + P1;
-		const PMatrix<T, d, d> WorldSpaceInvI1 = PBDRigid0 ? (PBDRigid0->Q() * FMatrix::Identity) * PBDRigid0->InvI() * (PBDRigid0->Q() * FMatrix::Identity).GetTransposed() : PMatrix<T, d, d>(0);
-		const PMatrix<T, d, d> WorldSpaceInvI2 = PBDRigid1 ? (PBDRigid1->Q() * FMatrix::Identity) * PBDRigid1->InvI() * (PBDRigid1->Q() * FMatrix::Identity).GetTransposed() : PMatrix<T, d, d>(0);
+		const PMatrix<T, d, d> WorldSpaceInvI1 = PBDRigid0 ? Utilities::ComputeWorldSpaceInertia(Q0, PBDRigid0->InvI()) : PMatrix<T, d, d>(0);
+		const PMatrix<T, d, d> WorldSpaceInvI2 = PBDRigid1 ? Utilities::ComputeWorldSpaceInertia(Q1, PBDRigid1->InvI()) : PMatrix<T, d, d>(0);
 		const TVector<T, d> Delta = GetDelta(ConstraintIndex, WorldSpaceX1, WorldSpaceX2);
 
 		if (PBDRigid0)
diff --git a/Engine/Source/Runtime/Experimental/Chaos/Public/Chaos/DenseMatrix.h b/Engine/Source/Runtime/Experimental/Chaos/Public/Chaos/DenseMatrix.h
--- a/Engine/Source/Runtime/Experimental/Chaos/Public/Chaos/DenseMatrix.h
+++ b/Engine/Source/Runtime/Experimental/Chaos/Public/Chaos/DenseMatrix.h
@@ -52,7 +52,7 @@ namespace Chaos
 		/**
 		 * The number of rows in the matrix.
 		 */
-		int32 NumRows() const
+		FORCEINLINE int32 NumRows() const
 		{
 			return NRows;
 		}
@@ -60,7 +60,7 @@ namespace Chaos
 		/**
 		 * The number of columns in the matrix.
 		 */
-		int32 NumColumns() const
+		FORCEINLINE int32 NumColumns() const
 		{
 			return NCols;
 		}
@@ -68,7 +68,7 @@ namespace Chaos
 		/**
 		 * The number of elements in the matrix.
 		 */
-		int32 NumElements() const
+		FORCEINLINE int32 NumElements() const
 		{
 			return NRows * NCols;
 		}
@@ -76,40 +76,46 @@ namespace Chaos
 		/**
 		 * Set the dimensions of the matrix, but do not initialize any values.
 		 */
-		void SetDimensions(const int32 InNumRows, const int32 InNumColumns)
+		FORCEINLINE void SetDimensions(const int32 InNumRows, const int32 InNumColumns)
 		{
 			check(InNumRows * InNumColumns <= MaxElements);
 			NRows = InNumRows;
 			NCols = InNumColumns;
 		}
 
-		FORCEINLINE int32 ElementIndex(const int32 RowIndex, const int32 ColumnIndex) const
+		/**
+		 * Return a writable reference to the element at the specified row and column.
+		 */
+		FORCEINLINE FReal& At(const int32 RowIndex, const int32 ColumnIndex)
 		{
 			checkSlow(RowIndex < NumRows());
 			checkSlow(ColumnIndex < NumColumns());
-			return RowIndex * NCols + ColumnIndex;
+			return M[ElementIndex(RowIndex, ColumnIndex)];
 		}
 
 		/**
-		 * Return a writable reference to the element at the specified row and column.
+		 * Return a read-only reference to the element at the specified row and column.
 		 */
-		FReal& At(const int32 RowIndex, const int32 ColumnIndex)
+		FORCEINLINE const FReal& At(const int32 RowIndex, const int32 ColumnIndex) const
 		{
+			checkSlow(RowIndex < NumRows());
+			checkSlow(ColumnIndex < NumColumns());
 			return M[ElementIndex(RowIndex, ColumnIndex)];
 		}
 
 		/**
-		 * Return a read-only reference to the element at the specified row and column.
+		 * Set the dimensions and initial values of the matrix.
 		 */
-		const FReal& At(const int32 RowIndex, const int32 ColumnIndex) const
+		void Init(const int32 InNRows, const int32 InNCols, FReal V)
 		{
-			return M[ElementIndex(RowIndex, ColumnIndex)];
+			SetDimensions(InNRows, InNCols);
+			Set(V);
 		}
 
 		/**
 		 * Set all elements to 'V'.
 		 */
-		void SetAll(FReal V)
+		void Set(FReal V)
 		{
 			for (int32 II = 0; II < NumElements(); ++II)
 			{
@@ -123,13 +129,103 @@ namespace Chaos
 		 */
 		void SetDiagonal(FReal V)
 		{
-			int32 Num = FMath::Min(NRows, NCols);
-			for (int32 II = 0; II < Num; ++II)
+			check(NumRows() == NumColumns());
+			for (int32 II = 0; II < NRows; ++II)
 			{
 				M[ElementIndex(II, II)] = V;
 			}
 		}
 
+		/**
+		 * Set the "Num" diagonal elements starting from ("Start", "Start") to "V". Does not set off-diagonal elements.
+		 */
+		void SetDiagonalAt(int32 Start, int32 Num, FReal V)
+		{
+			check(Start + Num <= NumRows());
+			check(Start + Num <= NumColumns());
+			for (int32 II = 0; II < Num; ++II)
+			{
+				int32 JJ = Start + II;
+				M[ElementIndex(JJ, JJ)] = V;
+			}
+		}
+
+		/**
+		 * Starting from element ("RowIndex", "ColumnIndex"), set the next "NumV" elements in the row to the values in "V".
+		 */
+		void SetRowAt(const int32 RowIndex, const int32 ColumnIndex, const FReal* V, const int32 NumV)
+		{
+			check(RowIndex + NumV < NumRows());
+			check(ColumnIndex < NumColumns());
+			for (int32 II = 0; II < NumV; ++II)
+			{
+				M[ElementIndex(RowIndex, ColumnIndex + II)] = V[II];
+			}
+		}
+
+		void SetRowAt(const int32 RowIndex, const int32 ColumnIndex, const FVec3& V)
+		{
+			check(RowIndex + 1 <= NumRows());
+			check(ColumnIndex + 3 <= NumColumns());
+			M[ElementIndex(RowIndex, ColumnIndex + 0)] = V[0];
+			M[ElementIndex(RowIndex, ColumnIndex + 1)] = V[1];
+			M[ElementIndex(RowIndex, ColumnIndex + 2)] = V[2];
+		}
+
+		/**
+		 * Starting from element ("RowIndex", "ColumnIndex"), set the next "NumV" elements in the column to the values in "V".
+		 */
+		void SetColumnAt(const int32 RowIndex, const int32 ColumnIndex, const FReal* V, const int32 NumV)
+		{
+			check(RowIndex + NumV <= NumRows());
+			check(ColumnIndex + 1 <= NumColumns());
+			for (int32 II = 0; II < NumV; ++II)
+			{
+				M[ElementIndex(RowIndex + II, ColumnIndex)] = V[II];
+			}
+		}
+
+		void SetColumnAt(const int32 RowIndex, const int32 ColumnIndex, const FVec3& V)
+		{
+			check(RowIndex + 3 <= NumRows());
+			check(ColumnIndex + 1 <= NumColumns());
+			M[ElementIndex(RowIndex + 0, ColumnIndex)] = V[0];
+			M[ElementIndex(RowIndex + 1, ColumnIndex)] = V[1];
+			M[ElementIndex(RowIndex + 2, ColumnIndex)] = V[2];
+		}
+
+		/**
+		 * Set the block starting at ("RowOffset", "ColumnOffset") from the specified matrix "V"
+		 */
+		template<int32 T_EA>
+		void SetBlockAt(const int32 RowOffset, const int32 ColumnOffset, const TDenseMatrix<T_EA>& V)
+		{
+			check(RowOffset + V.NumRows() <= NumRows());
+			check(ColumnOffset + V.NumColumns() <= NumColumns());
+			for (int32 II = 0; II < V.NumRows(); ++II)
+			{
+				for (int32 JJ = 0; JJ < V.NumColumns(); ++JJ)
+				{
+					M[ElementIndex(II + RowOffset, JJ + ColumnOffset)] = V.At(II, JJ);
+				}
+			}
+		}
+
+		/**
+		 * Set the 3x3 block starting at ("RowOffset", "ColumnOffset") from the specified 3x3 matrix "V" (note: assumes the input UE matrix is column-major order)
+		 */
+		void SetBlockAt(const int32 RowOffset, const int32 ColumnOffset, const FMatrix33& V)
+		{
+			check(RowOffset + 3 <= NumRows());
+			check(ColumnOffset + 3 <= NumColumns());
+			for (int32 II = 0; II < 3; ++II)
+			{
+				for (int32 JJ = 0; JJ < 3; ++JJ)
+				{
+					M[ElementIndex(II + RowOffset, JJ + ColumnOffset)] = V.M[JJ][II];
+				}
+			}
+		}
 
 		//
 		// Factory methods
@@ -143,6 +239,13 @@ namespace Chaos
 			return TDenseMatrix<MaxElements>(InNumRows, InNumCols);
 		}
 
+		/**
+		 * Create a matrix with the specified dimensions, and initialize all elements with "V".
+		 */
+		static TDenseMatrix<MaxElements> Make(const int32 InNumRows, const int32 InNumCols, const FReal V)
+		{
+			return TDenseMatrix<MaxElements>(InNumRows, InNumCols, V);
+		}
 
 		/**
 		 * Create a matrix with the specified elements supplied as an array in row-major order 
@@ -371,9 +474,20 @@ namespace Chaos
 		}
 
 		/**
-		 * Return C = A x V, where A is an MxN matrix, and V a real number.
+		 * Return C = A x B x Transpose(A).
 		 */
 		template<int32 T_EA, int32 T_EB>
+		static TDenseMatrix<MaxElements> MultiplyABAt(const TDenseMatrix<T_EA>& A, const TDenseMatrix<T_EB>& B)
+		{
+			// @todo(ccaulfield): optimize
+			return TDenseMatrix<MaxElements>::Multiply(A, TDenseMatrix<MaxElements>::MultiplyABt(B, A));
+		}
+
+
+		/**
+		 * Return C = A x V, where A is an MxN matrix, and V a real number.
+		 */
+		template<int32 T_EA>
 		static TDenseMatrix<MaxElements> Multiply(const TDenseMatrix<T_EA>& A, const FReal V)
 		{
 			// @todo(ccaulfield): optimize
@@ -382,7 +496,7 @@ namespace Chaos
 			{
 				for (int32 ICol = 0; ICol < Result.NumColumns(); ++ICol)
 				{
-					Result.At(IRow, ICol) = Result.At(IRow, ICol) / V;
+					Result.At(IRow, ICol) = A.At(IRow, ICol) * V;
 				}
 			}
 			return Result;
@@ -391,7 +505,7 @@ namespace Chaos
 		/**
 		 * Return C = A x V, where A is an MxN matrix, and V a real number.
 		 */
-		template<int32 T_EA, int32 T_EB>
+		template<int32 T_EA>
 		static TDenseMatrix<MaxElements> Multiply(const FReal V, const TDenseMatrix<T_EA>& A)
 		{
 			return Multiply(A, V);
@@ -400,7 +514,7 @@ namespace Chaos
 		/**
 		 * Return C = A / V, where A is an MxN matrix, and V a real number.
 		 */
-		template<int32 T_EA, int32 T_EB>
+		template<int32 T_EA>
 		static TDenseMatrix<MaxElements> Divide(const TDenseMatrix<T_EA>& A, const FReal V)
 		{
 			// @todo(ccaulfield): optimize
@@ -409,7 +523,7 @@ namespace Chaos
 			{
 				for (int32 ICol = 0; ICol < Result.NumColumns(); ++ICol)
 				{
-					Result.At(IRow, ICol) = Result.At(IRow, ICol) / V;
+					Result.At(IRow, ICol) = A.At(IRow, ICol) / V;
 				}
 			}
 			return Result;
@@ -425,6 +539,16 @@ namespace Chaos
 		}
 
 	private:
+		TDenseMatrix(const int32 InNRows, const int32 InNCols, const FReal V)
+			: NRows(InNRows)
+			, NCols(InNCols)
+		{
+			for (int32 I = 0; I < NumElements(); ++I)
+			{
+				M[I] = V;
+			}
+		}
+
 		TDenseMatrix(const int32 InNRows, const int32 InNCols, const FReal* V, const int32 N)
 			: NRows(InNRows)
 			, NCols(InNCols)
@@ -436,6 +560,11 @@ namespace Chaos
 			}
 		}
 
+		FORCEINLINE int32 ElementIndex(const int32 RowIndex, const int32 ColumnIndex) const
+		{
+			return RowIndex * NCols + ColumnIndex;
+		}
+
 		FReal M[MaxElements];
 		int32 NRows;
 		int32 NCols;
diff --git a/Engine/Source/Runtime/Experimental/Chaos/Public/Chaos/PBDJointConstraintSolver.h b/Engine/Source/Runtime/Experimental/Chaos/Public/Chaos/PBDJointConstraintSolver.h
new file mode 100644
--- /dev/null
+++ b/Engine/Source/Runtime/Experimental/Chaos/Public/Chaos/PBDJointConstraintSolver.h
@@ -0,0 +1,570 @@
+// Copyright 1998-2019 Epic Games, Inc. All Rights Reserved.
+#pragma once
+
+#include "CoreMinimal.h"
+
+#include "Chaos/DenseMatrix.h"
+#include "Chaos/ParticleHandleFwd.h"
+#include "Chaos/PBDJointConstraintTypes.h"
+#include "Chaos/PBDJointConstraintUtilities.h"
+#include "Chaos/Utilities.h"
+
+namespace Chaos
+{
+	class FJointConstraintSolver
+	{
+	public:
+		using FDenseMatrix66 = TDenseMatrix<6 * 6>;
+		using FDenseMatrix61 = TDenseMatrix<6 * 1>;
+
+		void InitConstraints(
+			const FReal Dt,
+			const FPBDJointSolverSettings& SolverSettings,
+			const FPBDJointSettings& JointSettings,
+			const FVec3& P0,
+			const FRotation3& Q0,
+			const FVec3& P1,
+			const FRotation3& Q1,
+			const FReal InvM0,
+			const FMatrix33& InvIL0,
+			const FReal InvM1,
+			const FMatrix33& InvIL1,
+			const FRigidTransform3& XL0,
+			const FRigidTransform3& XL1)
+		{
+			XLs[0] = XL0;
+			XLs[1] = XL1;
+			InvILs[0] = InvIL0;
+			InvILs[1] = InvIL1;
+			InvMs[0] = InvM0;
+			InvMs[1] = InvM1;
+
+			UpdateConstraints(Dt, SolverSettings, JointSettings, P0, Q0, P1, Q1);
+		}
+
+		void ApplyConstraints(
+			const FReal Dt,
+			const FPBDJointSolverSettings& SolverSettings,
+			const FPBDJointSettings& JointSettings,
+			const FReal Stiffness)
+		{
+			// @todo(ccaulfield): Custom mass-matrix multiplication methods
+
+			// For N constraints
+			// Constraint error: C(Nx1)
+			// Jacobian : J(Nx6)
+			FDenseMatrix61 C;
+			FDenseMatrix66 J0, J1;
+			BuildJacobianAndResidual(SolverSettings, JointSettings, J0, J1, C);
+
+			// InvM(6x6) = inverse mass matrix
+			FDenseMatrix66 InvM0 = FDenseMatrix66::Make(6, 6, 0);
+			InvM0.SetDiagonalAt(0, 3, InvMs[0]);
+			InvM0.SetBlockAt(3, 3, Utilities::ComputeWorldSpaceInertia(Qs[0], InvILs[0]));
+			FDenseMatrix66 InvM1 = FDenseMatrix66::Make(6, 6, 0);
+			InvM1.SetDiagonalAt(0, 3, InvMs[1]);
+			InvM1.SetBlockAt(3, 3, Utilities::ComputeWorldSpaceInertia(Qs[1], InvILs[1]));
+
+			// Joint-space mass: F(NxN) = J.IM.Jt
+			// Joint-space correction: L(Nx1) = (1/F).C
+			FDenseMatrix66 F0 = FDenseMatrix66::MultiplyAB(J0, FDenseMatrix66::MultiplyABt(InvM0, J0));
+			FDenseMatrix66 F1 = FDenseMatrix66::MultiplyAB(J1, FDenseMatrix66::MultiplyABt(InvM1, J1));
+			FDenseMatrix66 F = FDenseMatrix66::Add(F0, F1);
+			FDenseMatrix61 L;
+			if (FDenseMatrixSolver::SolvePositiveDefinite(F, C, L))
+			{
+				// World-space correction: D(6x1) = IM.Jt.L
+				FDenseMatrix61 D0 = FDenseMatrix61::MultiplyAB(InvM0, FDenseMatrix61::MultiplyAtB(J0, L));
+				FDenseMatrix61 D1 = FDenseMatrix61::MultiplyAB(InvM1, FDenseMatrix61::MultiplyAtB(J1, L));
+
+				// Extract world-space position correction
+				const FVec3 DP0 = FVec3(D0.At(0, 0), D0.At(1, 0), D0.At(2, 0));
+				const FVec3 DP1 = FVec3(D1.At(0, 0), D1.At(1, 0), D1.At(2, 0));
+				const FVec3 P0 = Ps[0] + Stiffness * DP0;
+				const FVec3 P1 = Ps[1] + Stiffness * DP1;
+
+				// Extract world-space rotation correction
+				const FVec3 DR0 = FVec3(D0.At(3, 0), D0.At(4, 0), D0.At(5, 0));
+				const FVec3 DR1 = FVec3(D1.At(3, 0), D1.At(4, 0), D1.At(5, 0));
+				const FRotation3 DQ0 = (FRotation3::FromElements(Stiffness * DR0, 0) * Qs[0]) * (FReal)0.5;
+				const FRotation3 DQ1 = (FRotation3::FromElements(Stiffness * DR1, 0) * Qs[1]) * (FReal)0.5;
+				const FRotation3 Q0 = (Qs[0] + DQ0).GetNormalized();
+				const FRotation3 Q1 = (Qs[1] + DQ1).GetNormalized();
+
+				// Apply corrections
+				UpdateConstraints(Dt, SolverSettings, JointSettings, P0, Q0, P1, Q1);
+			}
+		}
+
+		FORCEINLINE const FVec3& GetP(const int32 Index) const
+		{
+			checkSlow(Index < 2);
+			return Ps[Index];
+		}
+
+		FORCEINLINE const FRotation3& GetQ(const int32 Index) const
+		{
+			checkSlow(Index < 2);
+			return Qs[Index];
+		}
+
+	private:
+
+		void UpdateConstraints(
+			const FReal Dt,
+			const FPBDJointSolverSettings& SolverSettings,
+			const FPBDJointSettings& JointSettings,
+			const FVec3& P0,
+			const FRotation3& Q0,
+			const FVec3& P1,
+			const FRotation3& Q1)
+		{
+			Ps[0] = P0;
+			Ps[1] = P1;
+			Qs[0] = Q0;
+			Qs[1] = Q1;
+			Qs[1].EnforceShortestArcWith(Q0);
+
+			Xs[0] = Ps[0] + Qs[0] * XLs[0].GetTranslation();
+			Xs[1] = Ps[1] + Qs[1] * XLs[1].GetTranslation();
+			Rs[0] = Qs[0] * XLs[0].GetRotation();
+			Rs[1] = Qs[1] * XLs[1].GetRotation();
+		}
+
+		void AddLinearConstraints_Locked(
+			const FPBDJointSolverSettings& SolverSettings,
+			const FPBDJointSettings& JointSettings)
+		{
+			const FVec3 ConstraintSeparation = Xs[1] - Xs[0];
+
+			NumLinearConstraints = 3;
+			LinearConstraintAxes[0] = FVec3(1, 0, 0);
+			LinearConstraintAxes[1] = FVec3(0, 1, 0);
+			LinearConstraintAxes[2] = FVec3(0, 0, 1);
+			LinearConstraintDistances[0] = ConstraintSeparation[0];
+			LinearConstraintDistances[1] = ConstraintSeparation[1];
+			LinearConstraintDistances[2] = ConstraintSeparation[2];
+			LinearConstraintErrors[0] = ConstraintSeparation[0];
+			LinearConstraintErrors[1] = ConstraintSeparation[1];
+			LinearConstraintErrors[2] = ConstraintSeparation[2];
+		}
+
+		void AddLinearConstraints_Sphere(
+			const FPBDJointSolverSettings& SolverSettings,
+			const FPBDJointSettings& JointSettings)
+		{
+			NumLinearConstraints = 0;
+
+			const FReal Limit = JointSettings.Motion.LinearLimit;
+			const FVec3 ConstraintSeparation = Xs[1] - Xs[0];
+			const FReal ConstraintSeparationLen = ConstraintSeparation.Size();
+
+			bool bConstraintActive = (ConstraintSeparationLen >= FMath::Max(Limit, KINDA_SMALL_NUMBER));
+			if (bConstraintActive)
+			{
+				const FVec3 Axis = ConstraintSeparation / ConstraintSeparationLen;
+
+				NumLinearConstraints = 1;
+				LinearConstraintAxes[0] = Axis;
+				LinearConstraintDistances[0] = ConstraintSeparationLen;
+				LinearConstraintErrors[0] = ConstraintSeparationLen - Limit;
+			}
+		}
+
+		void AddLinearConstraints_Cylinder(
+			const FPBDJointSolverSettings& SolverSettings,
+			const FPBDJointSettings& JointSettings,
+			const EJointMotionType AxisMotion,
+			const FVec3& Axis)
+		{
+			const FVec3 ConstraintSeparation = Xs[1] - Xs[0];
+
+			// Axial Constraint
+			const FReal ConstraintDistanceAxial = FVec3::DotProduct(ConstraintSeparation, Axis);
+			bool bAxisConstraintActive = (AxisMotion != EJointMotionType::Free);
+			if (bAxisConstraintActive)
+			{
+				int32 ConstraintIndex = NumLinearConstraints++;
+				LinearConstraintAxes[ConstraintIndex] = Axis;
+				LinearConstraintDistances[ConstraintIndex] = ConstraintDistanceAxial;
+				LinearConstraintErrors[ConstraintIndex] = ConstraintDistanceAxial;
+			}
+
+			// Radial Constraint
+			const FVec3 ConstraintSeparationRadial = ConstraintSeparation - ConstraintDistanceAxial * Axis;
+			const FReal ConstraintDistanceRadial = ConstraintSeparationRadial.Size();
+			const FReal RadialLimit = JointSettings.Motion.LinearLimit;
+			bool bRadialConstraintActive = (ConstraintDistanceRadial >= RadialLimit);
+			if (bRadialConstraintActive)
+			{
+				int32 ConstraintIndex = NumLinearConstraints++;
+				LinearConstraintAxes[ConstraintIndex] = ConstraintSeparationRadial / ConstraintDistanceRadial;
+				LinearConstraintDistances[ConstraintIndex] = ConstraintDistanceRadial;
+				LinearConstraintErrors[ConstraintIndex] = ConstraintDistanceRadial - RadialLimit;
+			}
+		}
+
+		void AddLinearConstraints_Plane(
+			const FPBDJointSolverSettings& SolverSettings,
+			const FPBDJointSettings& JointSettings,
+			const EJointMotionType AxisMotion,
+			const FVec3& Axis)
+		{
+			const FReal Limit = (AxisMotion == EJointMotionType::Limited) ? JointSettings.Motion.LinearLimit : (FReal)0;
+			const FVec3 ConstraintSeparation = Xs[1] - Xs[0];
+
+			// Planar Constraint
+			const FReal ConstraintDistanceAxial = FVec3::DotProduct(ConstraintSeparation, Axis);
+			bool bAxisConstraintActive = (ConstraintDistanceAxial <= -Limit) || (ConstraintDistanceAxial >= Limit);
+			if (bAxisConstraintActive)
+			{
+				int32 ConstraintIndex = NumLinearConstraints++;
+				LinearConstraintAxes[ConstraintIndex] = Axis;
+				LinearConstraintDistances[ConstraintIndex] = ConstraintDistanceAxial;
+				LinearConstraintErrors[ConstraintIndex] = (ConstraintDistanceAxial >= Limit) ? ConstraintDistanceAxial - Limit : ConstraintDistanceAxial + Limit;
+			}
+		}
+
+		void AddAngularConstraints_Twist(
+			const FPBDJointSolverSettings& SolverSettings,
+			const FPBDJointSettings& JointSettings)
+		{
+			// Calculate the Twist Axis and Angle for each body
+			const FRotation3 R01 = Rs[0].Inverse() * Rs[1];
+			FRotation3 R01Twist, R01Swing;
+			R01.ToSwingTwist(FJointConstants::TwistAxis(), R01Swing, R01Twist);
+			R01Swing = R01Swing.GetNormalized();
+			R01Twist = R01Twist.GetNormalized();
+
+			FVec3 TwistAxis01;
+			FReal TwistAngle;
+			R01Twist.ToAxisAndAngleSafe(TwistAxis01, TwistAngle, FJointConstants::TwistAxis(), SolverSettings.SwingTwistAngleTolerance);
+			if (TwistAngle > PI)
+			{
+				TwistAngle = TwistAngle - (FReal)2 * PI;
+			}
+			if (FVec3::DotProduct(TwistAxis01, FJointConstants::TwistAxis()) < 0)
+			{
+				TwistAxis01 = -TwistAxis01;
+				TwistAngle = -TwistAngle;
+			}
+
+			FReal TwistAngleMax = FLT_MAX;
+			if (JointSettings.Motion.AngularMotionTypes[(int32)EJointAngularConstraintIndex::Twist] == EJointMotionType::Limited)
+			{
+				TwistAngleMax = JointSettings.Motion.AngularLimits[(int32)EJointAngularConstraintIndex::Twist];
+			}
+			else if (JointSettings.Motion.AngularMotionTypes[(int32)EJointAngularConstraintIndex::Twist] == EJointMotionType::Locked)
+			{
+				TwistAngleMax = 0;
+			}
+
+			bool bConstraintActive = (TwistAngle <= -TwistAngleMax) || (TwistAngle >= TwistAngleMax);
+			if (bConstraintActive)
+			{
+				int32 ConstraintIndex = NumAngularConstraints++;
+				AngularConstraintAxes[ConstraintIndex][0] = Rs[0] * TwistAxis01;
+				AngularConstraintAxes[ConstraintIndex][1] = Rs[1] * TwistAxis01;
+				AngularConstraintDistances[ConstraintIndex] = TwistAngle;
+				AngularConstraintErrors[ConstraintIndex] = (TwistAngle >= TwistAngleMax) ? TwistAngle - TwistAngleMax : TwistAngle + TwistAngleMax;
+			}
+		}
+
+		void AddAngularConstraints_Cone(
+			const FPBDJointSolverSettings& SolverSettings,
+			const FPBDJointSettings& JointSettings)
+		{
+			// Calculate Swing axis for each body
+			const FRotation3 R01 = Rs[0].Inverse() * Rs[1];
+			FRotation3 R01Twist, R01Swing;
+			R01.ToSwingTwist(FJointConstants::TwistAxis(), R01Swing, R01Twist);
+			R01Swing = R01Swing.GetNormalized();
+			R01Twist = R01Twist.GetNormalized();
+
+			FVec3 SwingAxis01;
+			FReal SwingAngle;
+			R01Swing.ToAxisAndAngleSafe(SwingAxis01, SwingAngle, FJointConstants::Swing1Axis(), SolverSettings.SwingTwistAngleTolerance);
+			if (SwingAngle > PI)
+			{
+				SwingAngle = SwingAngle - (FReal)2 * PI;
+			}
+
+			// Calculate swing limit for the current swing axis
+			FReal SwingAngleMax = FLT_MAX;
+			const FReal Swing1Limit = JointSettings.Motion.AngularLimits[(int32)EJointAngularConstraintIndex::Swing1];
+			const FReal Swing2Limit = JointSettings.Motion.AngularLimits[(int32)EJointAngularConstraintIndex::Swing2];
+
+			// Circular swing limit
+			SwingAngleMax = Swing1Limit;
+
+			// Elliptical swing limit
+			if (!FMath::IsNearlyEqual(Swing1Limit, Swing2Limit, KINDA_SMALL_NUMBER))
+			{
+				// Map swing axis to ellipse and calculate limit for this swing axis
+				const FReal DotSwing1 = FMath::Abs(FVec3::DotProduct(SwingAxis01, FJointConstants::Swing1Axis()));
+				const FReal DotSwing2 = FMath::Abs(FVec3::DotProduct(SwingAxis01, FJointConstants::Swing2Axis()));
+				SwingAngleMax = FMath::Sqrt(Swing1Limit * DotSwing2 * Swing1Limit * DotSwing2 + Swing2Limit * DotSwing1 * Swing2Limit * DotSwing1);
+			}
+
+			bool bConstraintActive = (SwingAngle <= -SwingAngleMax) || (SwingAngle >= SwingAngleMax);
+			if (bConstraintActive)
+			{
+				const int32 ConstraintIndex = NumAngularConstraints++;
+				AngularConstraintAxes[ConstraintIndex][0] = Rs[0] * SwingAxis01;
+				AngularConstraintAxes[ConstraintIndex][1] = AngularConstraintAxes[ConstraintIndex][0];
+				AngularConstraintDistances[ConstraintIndex] = SwingAngle;
+				AngularConstraintErrors[ConstraintIndex] = (SwingAngle >= SwingAngleMax) ? SwingAngle - SwingAngleMax : SwingAngle + SwingAngleMax;
+			}
+		}
+
+		void AddAngularConstraints_Swing(
+			const FPBDJointSolverSettings& SolverSettings,
+			const FPBDJointSettings& JointSettings,
+			const EJointAngularConstraintIndex SwingConstraintIndex,
+			const EJointAngularAxisIndex SwingAxisIndex)
+		{
+			// Calculate the swing axis for each body
+			const FRotation3 R01 = Rs[0].Inverse() * Rs[1];
+			FRotation3 R01Twist, R01Swing;
+			R01.ToSwingTwist(FJointConstants::TwistAxis(), R01Swing, R01Twist);
+			R01Swing = R01Swing.GetNormalized();
+			R01Twist = R01Twist.GetNormalized();
+
+			FVec3 TwistAxis01;
+			FReal TwistAngle;
+			R01Twist.ToAxisAndAngleSafe(TwistAxis01, TwistAngle, FJointConstants::TwistAxis(), SolverSettings.SwingTwistAngleTolerance);
+			if (TwistAngle > PI)
+			{
+				TwistAngle = TwistAngle - (FReal)2 * PI;
+			}
+			if (FVec3::DotProduct(TwistAxis01, FJointConstants::TwistAxis()) < 0)
+			{
+				TwistAxis01 = -TwistAxis01;
+				TwistAngle = -TwistAngle;
+			}
+			const FVec3 TwistAxis = Rs[0] * TwistAxis01;
+
+			const FRotation3 R1NoTwist = Rs[1] * R01Twist.Inverse();
+			const FMatrix33 Axes0 = Rs[0].ToMatrix();
+			const FMatrix33 Axes1 = R1NoTwist.ToMatrix();
+			FVec3 SwingCross = FVec3::CrossProduct(Axes0.GetAxis((int32)SwingAxisIndex), Axes1.GetAxis((int32)SwingAxisIndex));
+			SwingCross = SwingCross - FVec3::DotProduct(TwistAxis, SwingCross) * TwistAxis;
+			const FReal SwingCrossLen = SwingCross.Size();
+			if (SwingCrossLen > KINDA_SMALL_NUMBER)
+			{
+				FReal SwingAngle = FMath::Asin(FMath::Clamp(SwingCrossLen, (FReal)0, (FReal)1));
+				const FReal SwingDot = FVec3::DotProduct(Axes0.GetAxis((int32)SwingAxisIndex), Axes1.GetAxis((int32)SwingAxisIndex));
+				if (SwingDot < (FReal)0)
+				{
+					SwingAngle = (FReal)PI - SwingAngle;
+				}
+
+				FReal SwingAngleMax = FLT_MAX;
+				if (JointSettings.Motion.AngularMotionTypes[(int32)SwingConstraintIndex] == EJointMotionType::Limited)
+				{
+					FReal SwingLimit = JointSettings.Motion.AngularLimits[(int32)SwingConstraintIndex];
+					SwingAngleMax = SwingLimit;
+				}
+				else if (JointSettings.Motion.AngularMotionTypes[(int32)SwingConstraintIndex] == EJointMotionType::Locked)
+				{
+					SwingAngleMax = 0;
+				}
+
+				bool bConstraintActive = (SwingAngle <= -SwingAngleMax) || (SwingAngle >= SwingAngleMax);
+				if (bConstraintActive)
+				{
+					int32 ConstraintIndex = NumAngularConstraints++;
+
+					const FVec3 SwingAxis = SwingCross / SwingCrossLen;
+					AngularConstraintAxes[ConstraintIndex][0] = SwingAxis;
+					AngularConstraintAxes[ConstraintIndex][1] = SwingAxis;
+					AngularConstraintDistances[ConstraintIndex] = SwingAngle;
+					AngularConstraintErrors[ConstraintIndex] = (SwingAngle >= SwingAngleMax) ? SwingAngle - SwingAngleMax : SwingAngle + SwingAngleMax;
+				}
+			}
+		}
+
+		void AddLinearConstraints(
+			const FPBDJointSolverSettings& SolverSettings,
+			const FPBDJointSettings& JointSettings)
+		{
+			NumLinearConstraints = 0;
+
+			const TVector<EJointMotionType, 3>& Motion = JointSettings.Motion.LinearMotionTypes;
+			if ((Motion[0] == EJointMotionType::Locked) && (Motion[1] == EJointMotionType::Locked) && (Motion[2] == EJointMotionType::Locked))
+			{
+				AddLinearConstraints_Locked(SolverSettings, JointSettings);
+			}
+			else if ((Motion[0] == EJointMotionType::Limited) && (Motion[1] == EJointMotionType::Limited) && (Motion[2] == EJointMotionType::Limited))
+			{
+				AddLinearConstraints_Sphere(SolverSettings, JointSettings);
+			}
+			else if ((Motion[1] == EJointMotionType::Limited) && (Motion[2] == EJointMotionType::Limited))
+			{
+				// Circular Limit (X Axis)
+				AddLinearConstraints_Cylinder(SolverSettings, JointSettings, Motion[0], Rs[0] * FVec3(1, 0, 0));
+			}
+			else if ((Motion[0] == EJointMotionType::Limited) && (Motion[2] == EJointMotionType::Limited))
+			{
+				// Circular Limit (Y Axis)
+				AddLinearConstraints_Cylinder(SolverSettings, JointSettings, Motion[1], Rs[0] * FVec3(0, 1, 0));
+			}
+			else if ((Motion[0] == EJointMotionType::Limited) && (Motion[1] == EJointMotionType::Limited))
+			{
+				// Circular Limit (Z Axis)
+				AddLinearConstraints_Cylinder(SolverSettings, JointSettings, Motion[2], Rs[0] * FVec3(0, 0, 1));
+			}
+			else
+			{
+				// Plane/Square/Cube Limits (no way to author square or cube limits, but would work if we wanted it)
+				if (Motion[0] != EJointMotionType::Free)
+				{
+					AddLinearConstraints_Plane(SolverSettings, JointSettings, Motion[0], Rs[0] * FVec3(1, 0, 0));
+				}
+				if (Motion[1] != EJointMotionType::Free)
+				{
+					AddLinearConstraints_Plane(SolverSettings, JointSettings, Motion[1], Rs[0] * FVec3(0, 1, 0));
+				}
+				if (Motion[2] != EJointMotionType::Free)
+				{
+					AddLinearConstraints_Plane(SolverSettings, JointSettings, Motion[2], Rs[0] * FVec3(0, 0, 1));
+				}
+			}
+		}
+
+		void AddAngularConstraints(
+			const FPBDJointSolverSettings& SolverSettings,
+			const FPBDJointSettings& JointSettings)
+		{
+			NumAngularConstraints = 0;
+
+			EJointMotionType TwistMotion = JointSettings.Motion.AngularMotionTypes[(int32)EJointAngularConstraintIndex::Twist];
+			EJointMotionType Swing1Motion = JointSettings.Motion.AngularMotionTypes[(int32)EJointAngularConstraintIndex::Swing1];
+			EJointMotionType Swing2Motion = JointSettings.Motion.AngularMotionTypes[(int32)EJointAngularConstraintIndex::Swing2];
+
+			// Apply twist constraint
+			if (SolverSettings.bEnableTwistLimits)
+			{
+				if (TwistMotion != EJointMotionType::Free)
+				{
+					AddAngularConstraints_Twist(SolverSettings, JointSettings);
+				}
+			}
+
+			// Apply swing constraints
+			if (SolverSettings.bEnableSwingLimits)
+			{
+				if ((Swing1Motion == EJointMotionType::Limited) && (Swing2Motion == EJointMotionType::Limited))
+				{
+					AddAngularConstraints_Cone(SolverSettings, JointSettings);
+				}
+				else
+				{
+					if (Swing1Motion != EJointMotionType::Free)
+					{
+						AddAngularConstraints_Swing(SolverSettings, JointSettings, EJointAngularConstraintIndex::Swing1, EJointAngularAxisIndex::Swing1);
+					}
+					if (Swing2Motion != EJointMotionType::Free)
+					{
+						AddAngularConstraints_Swing(SolverSettings, JointSettings, EJointAngularConstraintIndex::Swing2, EJointAngularAxisIndex::Swing2);
+					}
+				}
+			}
+		}
+
+		void BuildJacobianAndResidual(
+			const FPBDJointSolverSettings& SolverSettings,
+			const FPBDJointSettings& JointSettings,
+			FDenseMatrix66& J0,
+			FDenseMatrix66& J1,
+			FDenseMatrix61& C)
+		{
+			// Calculate constraint axes and errors
+			AddLinearConstraints(SolverSettings, JointSettings);
+			AddAngularConstraints(SolverSettings, JointSettings);
+
+			// The Jacobians will be some sub-set of the following rows, depending on which constraints are enabled
+			// and whether the limits are currently being violated.
+			//
+			// J0(Nx6) = | XAxis          -XAxis x Connector0 |
+			//           | YAxis          -YAxis x Connector0 |
+			//           | ZAxis          -ZAxis x Connector0 |
+			//           | 0              TwistAxis           |
+			//           | 0              Swing1Axis          |
+			//           | 0              Swing2Axis          |
+			//
+			// J0(Nx6) = | -XAxis         XAxis x Connector1  |
+			//           | -YAxis         YAxis x Connector1  |
+			//           | -ZAxis         ZAxis x Connector1  |
+			//           | 0              -TwistAxis          |
+			//           | 0              -Swing1Axis         |
+			//           | 0              -Swing2Axis         |
+			//
+			const int32 NumConstraints = NumLinearConstraints + NumAngularConstraints;
+			J0.SetDimensions(NumConstraints, 6);
+			J1.SetDimensions(NumConstraints, 6);
+			C.SetDimensions(NumConstraints, 1);
+
+			const FVec3 XP0 = Xs[0] - Ps[0];
+			const FVec3 XP1 = Xs[1] - Ps[1];
+
+			int32 ConstraintIndex = 0;
+			for (int32 LinearConstraintIndex = 0; LinearConstraintIndex < NumLinearConstraints; ++LinearConstraintIndex, ++ConstraintIndex)
+			{
+				const FVec3& ConstraintAxis = LinearConstraintAxes[LinearConstraintIndex];
+
+				J0.SetRowAt(ConstraintIndex, 0, ConstraintAxis);
+				J0.SetRowAt(ConstraintIndex, 3, -FVec3::CrossProduct(ConstraintAxis, XP0));
+
+				J1.SetRowAt(ConstraintIndex, 0, -ConstraintAxis);
+				J1.SetRowAt(ConstraintIndex, 3, FVec3::CrossProduct(ConstraintAxis, XP1));
+
+				C.At(ConstraintIndex, 0) = LinearConstraintErrors[LinearConstraintIndex];
+			}
+
+			for (int32 AngularConstraintIndex = 0; AngularConstraintIndex < NumAngularConstraints; ++AngularConstraintIndex, ++ConstraintIndex)
+			{
+				const FVec3& ConstraintAxis0 = AngularConstraintAxes[AngularConstraintIndex][0];
+				const FVec3& ConstraintAxis1 = AngularConstraintAxes[AngularConstraintIndex][1];
+
+				J0.SetRowAt(ConstraintIndex, 0, FVec3::ZeroVector);
+				J0.SetRowAt(ConstraintIndex, 3, ConstraintAxis0);
+
+				J1.SetRowAt(ConstraintIndex, 0, FVec3::ZeroVector);
+				J1.SetRowAt(ConstraintIndex, 3, -ConstraintAxis1);
+
+				C.At(ConstraintIndex, 0) = AngularConstraintErrors[AngularConstraintIndex];
+			}
+		}
+
+		static const int32 MaxConstraintedBodies = 2;
+		static const int32 MaxLinearConstraints = 3;
+		static const int32 MaxAngularConstraints = 3;
+		static const int32 MaxConstraints = MaxLinearConstraints + MaxAngularConstraints;
+
+		// Joint solver state
+		// @todo(ccaulfield): maybe just store the Jacobians and errors here rather than the data to construct them
+		FVec3 LinearConstraintAxes[MaxConstraints];
+		FReal LinearConstraintDistances[MaxConstraints];
+		FReal LinearConstraintErrors[MaxConstraints];
+		FVec3 AngularConstraintAxes[MaxConstraints][MaxConstraintedBodies];	// @todo(ccaulfield): per body for twist - should probably switch back to shared twist axis
+		FReal AngularConstraintDistances[MaxConstraints];
+		FReal AngularConstraintErrors[MaxConstraints];
+		int32 NumLinearConstraints;
+		int32 NumAngularConstraints;
+
+		// Local-space constraint settings
+		FRigidTransform3 XLs[MaxConstraintedBodies];	// Local-space joint connector transforms
+		FMatrix33 InvILs[MaxConstraintedBodies];		// Local-space inverse inertias
+		FReal InvMs[MaxConstraintedBodies];				// Inverse masses
+
+		// World-space constraint state
+		FVec3 Xs[MaxConstraintedBodies];				// World-space joint connector positions
+		FRotation3 Rs[MaxConstraintedBodies];			// World-space joint connector rotations
+
+		// World-space body state
+		FVec3 Ps[MaxConstraintedBodies];				// World-space particle CoM positions
+		FRotation3 Qs[MaxConstraintedBodies];			// World-space particle CoM rotations
+	};
+
+}
\ No newline at end of file
diff --git a/Engine/Source/Runtime/Experimental/Chaos/Public/Chaos/PBDJointConstraintTypes.h b/Engine/Source/Runtime/Experimental/Chaos/Public/Chaos/PBDJointConstraintTypes.h
--- a/Engine/Source/Runtime/Experimental/Chaos/Public/Chaos/PBDJointConstraintTypes.h
+++ b/Engine/Source/Runtime/Experimental/Chaos/Public/Chaos/PBDJointConstraintTypes.h
@@ -44,7 +44,7 @@ namespace Chaos
 		Swing1 = 2,		// Swing1 Axis = Z
 	};
 
-	enum class EJointProjectionPhase
+	enum class EJointSolverPhase
 	{
 		None,
 		Apply,
@@ -132,7 +132,8 @@ namespace Chaos
 		bool bEnableTwistLimits;
 		bool bEnableSwingLimits;
 		bool bEnableDrives;
-		EJointProjectionPhase ProjectionPhase;
+		EJointSolverPhase DrivesPhase;
+		EJointSolverPhase ProjectionPhase;
 		FReal LinearProjection;
 		FReal AngularProjection;
 		FReal Stiffness;
diff --git a/Engine/Source/Runtime/Experimental/Chaos/Public/Chaos/PBDJointConstraints.h b/Engine/Source/Runtime/Experimental/Chaos/Public/Chaos/PBDJointConstraints.h
--- a/Engine/Source/Runtime/Experimental/Chaos/Public/Chaos/PBDJointConstraints.h
+++ b/Engine/Source/Runtime/Experimental/Chaos/Public/Chaos/PBDJointConstraints.h
@@ -143,9 +143,13 @@ namespace Chaos
 		friend class FPBDJointConstraintHandle;
 
 		void CalculateConstraintSpace(int32 ConstraintIndex, FVec3& OutX0, FMatrix33& OutR0, FVec3& OutX1, FMatrix33& OutR1, FVec3& OutAngles) const;
+		void UpdateParticleState(TPBDRigidParticleHandle<FReal, 3>* Rigid, const FReal Dt, const FVec3& P, const FRotation3& Q);
 
+		void ApplyDrives(const FReal Dt, const int32 ConstraintIndex, const int32 NumPairIts, const int32 It, const int32 NumIts);
 		void SolveVelocity(const FReal Dt, const int32 ConstraintIndex, const int32 NumPairIts, const int32 It, const int32 NumIts);
 		void SolvePosition(const FReal Dt, const int32 ConstraintIndex, const int32 NumPairIts, const int32 It, const int32 NumIts);
+		void SolvePosition_GaussSiedel(const FReal Dt, const int32 ConstraintIndex, const int32 NumPairIts, const int32 It, const int32 NumIts);
+		void SolvePosition_Cholesky(const FReal Dt, const int32 ConstraintIndex, const int32 NumPairIts, const int32 It, const int32 NumIts);
 		void ProjectPosition(const FReal Dt, const int32 ConstraintIndex, const int32 It, const int32 NumIts);
 
 		FPBDJointSolverSettings Settings;
