Use ensures instead of checks for recoverable errors.

diff --git a/Engine/Source/Runtime/Engine/Private/PhysicsEngine/Experimental/PhysScene_Chaos.cpp b/Engine/Source/Runtime/Engine/Private/PhysicsEngine/Experimental/PhysScene_Chaos.cpp
--- a/Engine/Source/Runtime/Engine/Private/PhysicsEngine/Experimental/PhysScene_Chaos.cpp
+++ b/Engine/Source/Runtime/Engine/Private/PhysicsEngine/Experimental/PhysScene_Chaos.cpp
@@ -546,7 +546,7 @@ void FPhysScene_Chaos::AddObject(UPrimitiveComponent* Component, FFieldSystemPhy
 {
 	AddToComponentMaps(Component, InObject);
 
-	check(false);
+	ensure(false);
 #if 0
 	Chaos::FPhysicsSolver* CurrSceneSolver = GetSolver();
 
diff --git a/Engine/Source/Runtime/Experimental/ChaosSolvers/Private/PhysicsProxy/GeometryCollectionPhysicsProxy.cpp b/Engine/Source/Runtime/Experimental/ChaosSolvers/Private/PhysicsProxy/GeometryCollectionPhysicsProxy.cpp
--- a/Engine/Source/Runtime/Experimental/ChaosSolvers/Private/PhysicsProxy/GeometryCollectionPhysicsProxy.cpp
+++ b/Engine/Source/Runtime/Experimental/ChaosSolvers/Private/PhysicsProxy/GeometryCollectionPhysicsProxy.cpp
@@ -3748,7 +3748,10 @@ Chaos::FParticleData* FGeometryCollectionPhysicsProxy::NewData()
 	// proxies get away with that because they're dealing with a low number of 
 	// bodies; which is not the case here. Rather, let's just use a triple buffer.
 
-	check(NumParticles != INDEX_NONE); // Make sure InitBodiesGT() has been called!
+	if (!ensure(NumParticles != INDEX_NONE)) // Make sure InitBodiesGT() has been called!
+	{
+		return nullptr;
+	}
 
 	FGeometryCollectionResults& Buffer = *GameToPhysInterchange->AccessProducerBuffer();
 	if (Buffer.NumParticlesAdded != NumParticles)
