Chaos : Collisions

- code clean up.

diff --git a/Engine/Source/Runtime/Experimental/Chaos/Private/Chaos/PBDCollisionConstraints.cpp b/Engine/Source/Runtime/Experimental/Chaos/Private/Chaos/PBDCollisionConstraints.cpp
--- a/Engine/Source/Runtime/Experimental/Chaos/Private/Chaos/PBDCollisionConstraints.cpp
+++ b/Engine/Source/Runtime/Experimental/Chaos/Private/Chaos/PBDCollisionConstraints.cpp
@@ -30,37 +30,34 @@
 
 //#pragma optimize("", off)
 
-int32 CollisionParticlesBVHDepth = 4;
-FAutoConsoleVariableRef CVarCollisionParticlesBVHDepth(TEXT("p.CollisionParticlesBVHDepth"), CollisionParticlesBVHDepth, TEXT("The maximum depth for collision particles bvh"));
+namespace Chaos
+{
+	extern int32 UseLevelsetCollision;
 
-int32 ConstraintBPBVHDepth = 2;
-FAutoConsoleVariableRef CVarConstraintBPBVHDepth(TEXT("p.ConstraintBPBVHDepth"), ConstraintBPBVHDepth, TEXT("The maximum depth for constraint bvh"));
+	int32 CollisionParticlesBVHDepth = 4;
+	FAutoConsoleVariableRef CVarCollisionParticlesBVHDepth(TEXT("p.CollisionParticlesBVHDepth"), CollisionParticlesBVHDepth, TEXT("The maximum depth for collision particles bvh"));
 
-int32 BPTreeOfGrids = 1;
-FAutoConsoleVariableRef CVarBPTreeOfGrids(TEXT("p.BPTreeOfGrids"), BPTreeOfGrids, TEXT("Whether to use a seperate tree of grids for bp"));
+	int32 ConstraintBPBVHDepth = 2;
+	FAutoConsoleVariableRef CVarConstraintBPBVHDepth(TEXT("p.ConstraintBPBVHDepth"), ConstraintBPBVHDepth, TEXT("The maximum depth for constraint bvh"));
 
-float CollisionVelocityInflationCVar = 2.0f;
-FAutoConsoleVariableRef CVarCollisionVelocityInflation(TEXT("p.CollisionVelocityInflation"), CollisionVelocityInflationCVar, TEXT("Collision velocity inflation.[def:2.0]"));
+	int32 BPTreeOfGrids = 1;
+	FAutoConsoleVariableRef CVarBPTreeOfGrids(TEXT("p.BPTreeOfGrids"), BPTreeOfGrids, TEXT("Whether to use a seperate tree of grids for bp"));
 
-float CollisionFrictionOverride = -1.0f;
-FAutoConsoleVariableRef CVarCollisionFrictionOverride(TEXT("p.CollisionFriction"), CollisionFrictionOverride, TEXT("Collision friction for all contacts if >= 0"));
+	float CollisionVelocityInflationCVar = 2.0f;
+	FAutoConsoleVariableRef CVarCollisionVelocityInflation(TEXT("p.CollisionVelocityInflation"), CollisionVelocityInflationCVar, TEXT("Collision velocity inflation.[def:2.0]"));
 
+	float CollisionFrictionOverride = -1.0f;
+	FAutoConsoleVariableRef CVarCollisionFrictionOverride(TEXT("p.CollisionFriction"), CollisionFrictionOverride, TEXT("Collision friction for all contacts if >= 0"));
 
-extern int32 UseLevelsetCollision;
+	CHAOS_API int32 EnableCollisions = 1;
+	FAutoConsoleVariableRef CVarEnableCollisions(TEXT("p.EnableCollisions"), EnableCollisions, TEXT("Enable/Disable collisions on the Chaos solver."));
 
 #if !UE_BUILD_SHIPPING
-namespace Chaos
-{
 	int32 CHAOS_API PendingHierarchyDump = 0;
-}
 #endif
 
-namespace Chaos
-{
-	CHAOS_API int32 EnableCollisions = 1;
-	FAutoConsoleVariableRef CVarEnableCollisions(TEXT("p.EnableCollisions"), EnableCollisions, TEXT("Enable/Disable collisions on the Chaos solver."));
-
-
+	DEFINE_STAT(STAT_ComputeConstraints);
+	DEFINE_STAT(STAT_ComputeConstraintsSU);
 
 	//
 	// Collision Constraint Container
@@ -128,11 +125,6 @@ namespace Chaos
 		PostApplyPushOutCallback = nullptr;
 	}
 
-	DEFINE_STAT(STAT_ComputeConstraints);
-	DEFINE_STAT(STAT_ComputeConstraintsNP);
-	DEFINE_STAT(STAT_ComputeConstraintsBP);
-	DEFINE_STAT(STAT_ComputeConstraintsSU);
-
 
 	template<typename T, int d>
 	void TPBDCollisionConstraints<T, d>::UpdatePositionBasedState(/*const TPBDRigidParticles<T, d>& InParticles, const TArray<int32>& InIndices,*/ const T Dt)
diff --git a/Engine/Source/Runtime/Experimental/Chaos/Public/Chaos/PBDCollisionConstraintsImp.h b/Engine/Source/Runtime/Experimental/Chaos/Public/Chaos/PBDCollisionConstraintsImp.h
--- a/Engine/Source/Runtime/Experimental/Chaos/Public/Chaos/PBDCollisionConstraintsImp.h
+++ b/Engine/Source/Runtime/Experimental/Chaos/Public/Chaos/PBDCollisionConstraintsImp.h
@@ -6,92 +6,213 @@
 
 namespace Chaos
 {
-#if !UE_BUILD_SHIPPING
-
-	static constexpr int32 BucketSizes2[] = { 0,1,4,8,16,32,64,128,512,MAX_int32 };
-
-
-	template <bool bGatherStats>
-	struct FStatHelper2
+	namespace CollisionStats
 	{
-		int32 MaxCount;
 
-		FStatHelper2() {}
-		void Record(int32 Count) {}
-		FString ToString() const { return TEXT(""); }
-	};
-
-	template <>
-	struct FStatHelper2<true>
-	{
-		int32 BucketCount[UE_ARRAY_COUNT(BucketSizes2)];
-		int32 MaxCount;
+		static constexpr int32 BucketSizes2[] = { 0,1,4,8,16,32,64,128,512,MAX_int32 };
 
-		FStatHelper2()
+		template <bool bGatherStats>
+		struct FStatHelper
 		{
-			FMemory::Memset(BucketCount, 0, sizeof(BucketCount));
-			MaxCount = 0;
-		}
+			int32 MaxCount;
+
+			FStatHelper() {}
+			void Record(int32 Count) {}
+			FString ToString() const { return TEXT(""); }
+		};
 
-		void Record(int32 Count)
+		template <>
+		struct FStatHelper<true>
 		{
-			for (int32 BucketIdx = 1; BucketIdx < UE_ARRAY_COUNT(BucketSizes2); ++BucketIdx)
-			{
-				if (Count >= BucketSizes2[BucketIdx - 1] && Count < BucketSizes2[BucketIdx])
-				{
-					++BucketCount[BucketIdx];
-				}
-			}
+			int32 BucketCount[UE_ARRAY_COUNT(BucketSizes2)];
+			int32 MaxCount;
 
-			if (Count > MaxCount)
+			FStatHelper()
 			{
-				MaxCount = Count;
+				FMemory::Memset(BucketCount, 0, sizeof(BucketCount));
+				MaxCount = 0;
 			}
-		}
 
-		FString ToString() const
-		{
-			FString OutLog;
-			int32 MaxBucketCount = 0;
-			for (int32 Count : BucketCount)
+			void Record(int32 Count)
 			{
-				if (Count > MaxBucketCount)
+				for (int32 BucketIdx = 1; BucketIdx < UE_ARRAY_COUNT(BucketSizes2); ++BucketIdx)
 				{
-					MaxBucketCount = Count;
+					if (Count >= BucketSizes2[BucketIdx - 1] && Count < BucketSizes2[BucketIdx])
+					{
+						++BucketCount[BucketIdx];
+					}
 				}
-			}
 
-			const float CountPerChar = MaxBucketCount / 20.f;
-			for (int32 Idx = 1; Idx < UE_ARRAY_COUNT(BucketSizes2); ++Idx)
-			{
-				int32 NumChars = BucketCount[Idx] / CountPerChar;
-				if (Idx < UE_ARRAY_COUNT(BucketSizes2) - 1)
+				if (Count > MaxCount)
 				{
-					OutLog += FString::Printf(TEXT("\t[%4d - %4d) (%4d) |"), BucketSizes2[Idx - 1], BucketSizes2[Idx], BucketCount[Idx]);
+					MaxCount = Count;
 				}
-				else
+			}
+
+			FString ToString() const
+			{
+				FString OutLog;
+				int32 MaxBucketCount = 0;
+				for (int32 Count : BucketCount)
 				{
-					OutLog += FString::Printf(TEXT("\t[%4d -  inf) (%4d) |"), BucketSizes2[Idx - 1], BucketCount[Idx]);
+					if (Count > MaxBucketCount)
+					{
+						MaxBucketCount = Count;
+					}
 				}
-				for (int32 Count = 0; Count < NumChars; ++Count)
+
+				const float CountPerChar = MaxBucketCount / 20.f;
+				for (int32 Idx = 1; Idx < UE_ARRAY_COUNT(BucketSizes2); ++Idx)
 				{
-					OutLog += TEXT("-");
+					int32 NumChars = BucketCount[Idx] / CountPerChar;
+					if (Idx < UE_ARRAY_COUNT(BucketSizes2) - 1)
+					{
+						OutLog += FString::Printf(TEXT("\t[%4d - %4d) (%4d) |"), BucketSizes2[Idx - 1], BucketSizes2[Idx], BucketCount[Idx]);
+					}
+					else
+					{
+						OutLog += FString::Printf(TEXT("\t[%4d -  inf) (%4d) |"), BucketSizes2[Idx - 1], BucketCount[Idx]);
+					}
+					for (int32 Count = 0; Count < NumChars; ++Count)
+					{
+						OutLog += TEXT("-");
+					}
+					OutLog += TEXT("\n");
 				}
-				OutLog += TEXT("\n");
+
+				return OutLog;
 			}
+		};
+
+		template <bool bGatherStats, class T = float, int32 d = 3>
+		struct FStatData
+		{
+			FStatData() {}
+
+			void IncrementSimulatedParticles() {}
+			void RecordBoundsData(const TBox <T, d>& Box1) {}
+			void RecordBroadphasePotentials(int32 Num) {}
+			void IncrementCountNP() {}
+			void IncrementRejectedNP() {}
+			void FinalizeData() {}
+			void Print() {}
+		};
+
+
+		template <>
+		struct FStatData<true, float, 3>
+		{
+			FStatData() : CountNP(0), RejectedNP(0), SimulatedParticles(0), NumPotentials(0) {}
 
-			return OutLog;
-		}
-	};
 
+			void IncrementSimulatedParticles()
+			{
+				++SimulatedParticles;
+			}
+
+			void RecordBoundsData(const TBox <float, 3>& Box1)
+			{
+				BoundsDistribution.Record(Box1.Extents().GetMax());
+			}
+
+			void RecordBroadphasePotentials(int32 Num)
+			{
+				NumPotentials = Num;
+				BroadphasePotentials.Record(Num);
+			}
+
+			void IncrementCountNP()
+			{
+				++CountNP;
+			}
+
+			void IncrementRejectedNP()
+			{
+				++RejectedNP;
+			}
+
+			void FinalizeData()
+			{
+				NarrowPhasePerformed.Record(CountNP);
+				const int32 NPSkipped = NumPotentials - CountNP;
+				NarrowPhaseSkipped.Record(NPSkipped);
+				NarrowPhaseRejected.Record(RejectedNP);
+			}
+
+			void Print()
+			{
+				FString OutLog;
+#if CHAOS_PARTICLEHANDLE_TODO
+				const float NumParticles = InParticles.Size();
+				OutLog = FString::Printf(TEXT("ComputeConstraints stats:\n"
+					"Total Particles:%d\nSimulated Particles:%d (%.2f%%)\n"
+					"Max candidates per instance:%d (%.2f%%)\n"
+					"Max candidates skipped per instance (NP skipped):%d (%.2f%%)\n"
+					"Max narrow phase tests per instance:%d (%.2f%%)\n"
+					"Max narrow phase rejected per instance (NP rejected):%d (%.2f%%)\n"
+					"Constraints generated:%d\n"
+				),
+					InParticles.Size(),
+					SimulatedParticles, SimulatedParticles / NumParticles * 100.f,
+					BroadphasePotentials.MaxCount, BroadphasePotentials.MaxCount / NumParticles * 100.f,
+					NarrowPhaseSkipped.MaxCount, NarrowPhaseSkipped.MaxCount / NumParticles * 100.f,
+					NarrowPhasePerformed.MaxCount, NarrowPhasePerformed.MaxCount / NumParticles * 100.f,
+					NarrowPhaseRejected.MaxCount, NarrowPhaseRejected.MaxCount / NumParticles * 100.f,
+					Constraints.Num()
+				);
+
+				OutLog += FString::Printf(TEXT("Potentials per instance distribution:\n"));
+				OutLog += BroadphasePotentials.ToString();
+
+
+				OutLog += FString::Printf(TEXT("\nCandidates skipped per instance (NP skipped) distribution:\n"));
+				OutLog += NarrowPhaseSkipped.ToString();
+
+				OutLog += FString::Printf(TEXT("\nNarrow phase performed per instance distribution:\n"));
+				OutLog += NarrowPhasePerformed.ToString();
+
+				OutLog += FString::Printf(TEXT("\nNarrow phase candidates rejected per instance distribution:\n"));
+				OutLog += NarrowPhaseRejected.ToString();
+
+				OutLog += FString::Printf(TEXT("\nBounds distribution:\n"));
+				OutLog += BoundsDistribution.ToString();
+
+				UE_LOG(LogChaos, Warning, TEXT("%s"), *OutLog);
+				PendingHierarchyDump = 0;
 #endif
+			}
+
+			int32 CountNP;
+			int32 RejectedNP;
+			int32 SimulatedParticles;
+			int32 NumPotentials;
+			FStatHelper<true> BroadphasePotentials;
+			FStatHelper<true> NarrowPhaseSkipped;
+			FStatHelper<true> NarrowPhasePerformed;
+			FStatHelper<true> NarrowPhaseRejected;
+			FStatHelper<true> BoundsDistribution;
+		};
+	}
+
+
+	extern CHAOS_API int32 EnableCollisions;
+	extern CHAOS_API int32 CollisionConstraintsForceSingleThreaded;
 
 #if UE_BUILD_SHIPPING
+#define COLLISION_OPT_OUT(X)
 	const bool bDisableCollisionParallelFor = false;
 #else
+#define COLLISION_OPT_OUT(X) X;
 	CHAOS_API extern bool bDisableCollisionParallelFor;
 #endif
 
+
+
+
+
+	/**
+	*
+	*/
 	template <typename T, int d>
 	struct TSimOverlapVisitor
 	{
@@ -120,239 +241,154 @@ namespace Chaos
 		}
 	};
 
-	extern CHAOS_API int32 EnableCollisions;
-	extern CHAOS_API int32 CollisionConstraintsForceSingleThreaded;
+
+
+	/**
+	*
+	*/
 
 	DECLARE_CYCLE_STAT_EXTERN(TEXT("ComputeConstraints"), STAT_ComputeConstraints, STATGROUP_Chaos, CHAOS_API);
-	DECLARE_CYCLE_STAT_EXTERN(TEXT("ComputeConstraintsNP"), STAT_ComputeConstraintsNP, STATGROUP_Chaos, CHAOS_API);
-	DECLARE_CYCLE_STAT_EXTERN(TEXT("ComputeConstraintsBP"), STAT_ComputeConstraintsBP, STATGROUP_Chaos, CHAOS_API);
 	DECLARE_CYCLE_STAT_EXTERN(TEXT("ComputeConstraintsSU"), STAT_ComputeConstraintsSU, STATGROUP_Chaos, CHAOS_API);
 
 	template <typename T, int d>
 	template <bool bGatherStats, typename SPATIAL_ACCELERATION>
 	void TPBDCollisionConstraints<T, d>::ComputeConstraintsHelperLowLevel(const SPATIAL_ACCELERATION& InSpatialAcceleration, T Dt)
 	{
+		COLLISION_OPT_OUT(CollisionStats::FStatData<bGatherStats> StatData);
 		SCOPE_CYCLE_COUNTER(STAT_ComputeConstraints);
 		CHAOS_SCOPED_TIMER(ComputeConstraints);
 		if (!EnableCollisions) return;
-		// Broad phase
-
-#if !UE_BUILD_SHIPPING
-		FStatHelper2<bGatherStats> BroadphasePotentials;
-		FStatHelper2<bGatherStats> NarrowPhaseSkipped;
-		FStatHelper2<bGatherStats> NarrowPhasePerformed;
-		FStatHelper2<bGatherStats> NarrowPhaseRejected;
-		FStatHelper2<bGatherStats> BoundsDistribution;
-		int32 SimulatedParticles = 0;
-#endif
-		{
-			SCOPE_CYCLE_COUNTER(STAT_ComputeConstraintsNP);
 
+		//todo(ocohen): use per thread buffer instead, need better support than ParallelFor for this
+		TQueue<TRigidBodyPointContactConstraint<T, d>, EQueueMode::Mpsc> Queue;
 
-			// Narrow phase
-			CHAOS_SCOPED_TIMER(ComputeConstraints_NP);
-
-			TQueue<TRigidBodyPointContactConstraint<T, d>, EQueueMode::Mpsc> Queue;	//todo(ocohen): use per thread buffer instead, need better support than ParallelFor for this
-			Particles.GetNonDisabledDynamicView().ParallelFor([&](auto& Particle1, int32 ActiveIdxIdx)
-			{
-#if !UE_BUILD_SHIPPING
-				if (bGatherStats)
-				{
-					++SimulatedParticles;
-				}
-#endif
+		Particles.GetNonDisabledDynamicView().ParallelFor([&](auto& Particle1, int32 ActiveIdxIdx)
+		{
+			COLLISION_OPT_OUT(StatData.IncrementSimulatedParticles());
 
-				TArray<TAccelerationStructureHandle<T, d>> PotentialIntersections;
+			TArray<TAccelerationStructureHandle<T, d>> PotentialIntersections;
 
-				const bool bBody1Bounded = HasBoundingBox(Particle1);
+			const bool bBody1Bounded = HasBoundingBox(Particle1);
 
-				const T Box1Thickness = ComputeThickness(Particle1, CollisionVelocityInflation * Dt).Size();
-				if (bBody1Bounded)
-				{
+			const T Box1Thickness = ComputeThickness(Particle1, CollisionVelocityInflation * Dt).Size();
+			if (bBody1Bounded)
+			{
 #if CHAOS_PARTICLEHANDLE_TODO
-					const TBox <T, d> Box1 = InSpatialAcceleration.GetWorldSpaceBoundingBox(Particle1);
+				const TBox <T, d> Box1 = InSpatialAcceleration.GetWorldSpaceBoundingBox(Particle1);
 #else
-					const TBox <T, d> Box1 = ComputeWorldSpaceBoundingBox(Particle1);	//NOTE: this ignores the velocity expansion which is wrong
+				const TBox <T, d> Box1 = ComputeWorldSpaceBoundingBox(Particle1);	//NOTE: this ignores the velocity expansion which is wrong
 #endif
 
-#if !UE_BUILD_SHIPPING
-					if (bGatherStats)
-					{
-						//const TBox<T, d> Box1 = SpatialAcceleration.GetWorldSpaceBoundingBox(InParticles, Body1Index);
+				COLLISION_OPT_OUT(StatData.RecordBoundsData(Box1));
 
-						BoundsDistribution.Record(Box1.Extents().GetMax());
-					}
-#endif
-					TSimOverlapVisitor<T, d> OverlapVisitor(PotentialIntersections);
-					InSpatialAcceleration.Overlap(Box1, OverlapVisitor);
-				}
-				else
-				{
-					const auto& GlobalElems = InSpatialAcceleration.GlobalObjects();
-					PotentialIntersections.Reserve(GlobalElems.Num());
+				TSimOverlapVisitor<T, d> OverlapVisitor(PotentialIntersections);
+				InSpatialAcceleration.Overlap(Box1, OverlapVisitor);
+			}
+			else
+			{
+				const auto& GlobalElems = InSpatialAcceleration.GlobalObjects();
+				PotentialIntersections.Reserve(GlobalElems.Num());
 
-					for (auto& Elem : GlobalElems)
-					{
-						PotentialIntersections.Add(Elem.Payload);
-					}
+				for (auto& Elem : GlobalElems)
+				{
+					PotentialIntersections.Add(Elem.Payload);
 				}
+			}
 
-#if !UE_BUILD_SHIPPING
-				BroadphasePotentials.Record(PotentialIntersections.Num());
+			COLLISION_OPT_OUT(StatData.RecordBroadphasePotentials(PotentialIntersections.Num()))
 
-				int32 CountNP = 0;
-				int32 RejectedNP = 0;
-#endif
 				const int32 NumPotentials = PotentialIntersections.Num();
-				for (int32 i = 0; i < NumPotentials; ++i)
+			for (int32 i = 0; i < NumPotentials; ++i)
+			{
+				auto& Particle2 = *PotentialIntersections[i].GetGeometryParticleHandle_PhysicsThread();
+				const TGenericParticleHandle<T, d> Particle2Generic(&Particle2);
+				// Collision group culling...
+				// CollisionGroup == 0 : Collide_With_Everything
+				// CollisionGroup == INDEX_NONE : Disabled collisions
+				// CollisionGroup_A != CollisionGroup_B : Skip Check
+
+				if (Particle1.CollisionGroup() == INDEX_NONE || Particle2Generic->CollisionGroup() == INDEX_NONE)
 				{
-					auto& Particle2 = *PotentialIntersections[i].GetGeometryParticleHandle_PhysicsThread();
-					const TGenericParticleHandle<T, d> Particle2Generic(&Particle2);
-					// Collision group culling...
-					// CollisionGroup == 0 : Collide_With_Everything
-					// CollisionGroup == INDEX_NONE : Disabled collisions
-					// CollisionGroup_A != CollisionGroup_B : Skip Check
-
-					if (Particle1.CollisionGroup() == INDEX_NONE || Particle2Generic->CollisionGroup() == INDEX_NONE)
-					{
-						continue;
-					}
-					if (Particle1.CollisionGroup() && Particle2Generic->CollisionGroup() && Particle1.CollisionGroup() != Particle2Generic->CollisionGroup())
-					{
-						continue;
-					}
+					continue;
+				}
+				if (Particle1.CollisionGroup() && Particle2Generic->CollisionGroup() && Particle1.CollisionGroup() != Particle2Generic->CollisionGroup())
+				{
+					continue;
+				}
 
-					if (!Particle1.Geometry() && !Particle2.Geometry())
-					{
-						continue;
-					}
+				if (!Particle1.Geometry() && !Particle2.Geometry())
+				{
+					continue;
+				}
 
-					const bool bBody2Bounded = HasBoundingBox(Particle2);
+				const bool bBody2Bounded = HasBoundingBox(Particle2);
 
-					if (Particle1.Handle() == Particle2.Handle())
-					{
-						continue;
-					}
+				if (Particle1.Handle() == Particle2.Handle())
+				{
+					continue;
+				}
 
-					if (bBody1Bounded == bBody2Bounded && Particle2.AsDynamic())
-					{
-						//if both are dynamic, assume index order matters
+				if (bBody1Bounded == bBody2Bounded && Particle2.AsDynamic())
+				{
 #if CHAOS_DETERMINISTIC
-						if (Particle2.ParticleID() > Particle1.ParticleID())
+					//if both are dynamic, assume index order matters
+					if (Particle2.ParticleID() > Particle1.ParticleID())
 #else
 					//not deterministic just use memory address to avoid pair duplication
-						if (Particle2.Handle() > Particle1.Handle())
+					if (Particle2.Handle() > Particle1.Handle())
 #endif
-						{
-							continue;
-						}
+					{
+						continue;
 					}
-					//				++TotalTested;
+				}
+				//				++TotalTested;
 
-					const TVector<T, d> Box2Thickness = Particle2.AsDynamic() ? ComputeThickness(*Particle2.AsDynamic(), Dt) : TVector<T, d>(0);
-					const T UseThickness = FMath::Max(Box1Thickness, Box2Thickness.Size());// + MThickness
+				const TVector<T, d> Box2Thickness = Particle2.AsDynamic() ? ComputeThickness(*Particle2.AsDynamic(), Dt) : TVector<T, d>(0);
+				const T UseThickness = FMath::Max(Box1Thickness, Box2Thickness.Size());// + MThickness
 
 
-					TRigidBodyPointContactConstraint<T, d> Constraint;
-					ConstructConstraints(Particle1.Handle(), Particle2.Handle(), UseThickness, Constraint);
+				TRigidBodyPointContactConstraint<T, d> Constraint;
+				ConstructConstraints(Particle1.Handle(), Particle2.Handle(), UseThickness, Constraint);
 
-#if !UE_BUILD_SHIPPING
-					if (bGatherStats)
-					{
-						++CountNP;
-					}
-#endif
-					if (Constraint.ContainsManifold())
-					{
-						UpdateConstraint<ECollisionUpdateType::Any>(UseThickness, Constraint);
-					}
+				COLLISION_OPT_OUT(StatData.IncrementCountNP());
 
-					if (Constraint.GetPhi() < UseThickness)
-					{
-						Queue.Enqueue(Constraint);
-					}
-					else
-					{
-#if !UE_BUILD_SHIPPING
-						++RejectedNP;
-#endif
-					}
+				if (Constraint.ContainsManifold())
+				{
+					UpdateConstraint<ECollisionUpdateType::Any>(UseThickness, Constraint);
 				}
 
-#if !UE_BUILD_SHIPPING
-				NarrowPhasePerformed.Record(CountNP);
-				const int32 NPSkipped = NumPotentials - CountNP;
-				NarrowPhaseSkipped.Record(NPSkipped);
-				NarrowPhaseRejected.Record(RejectedNP);
-#endif
-			}, bGatherStats || bDisableCollisionParallelFor);
-			
-			{
-				SCOPE_CYCLE_COUNTER(STAT_ComputeConstraintsSU);
-
-				while (!Queue.IsEmpty())
+				if (Constraint.GetPhi() < UseThickness)
 				{
-					int32 Idx = PointConstraints.AddUninitialized(1);
-					FConstraintContainerHandle* Handle = HandleAllocator.template AllocHandle< TRigidBodyPointContactConstraint<T,d> >(this, Idx );
-
-					Handles.Add(Handle);
-					Queue.Dequeue(PointConstraints[Idx]);
+					Queue.Enqueue(Constraint);
+				}
+				else
+				{
+					COLLISION_OPT_OUT(StatData.IncrementRejectedNP());
 				}
-				LifespanCounter++;
 			}
+
+			COLLISION_OPT_OUT(StatData.FinalizeData());
+
+		}, bGatherStats || bDisableCollisionParallelFor);
+
+		{
+			SCOPE_CYCLE_COUNTER(STAT_ComputeConstraintsSU);
+
+			while (!Queue.IsEmpty())
+			{
+				int32 Idx = PointConstraints.AddUninitialized(1);
+				FConstraintContainerHandle* Handle = HandleAllocator.template AllocHandle< TRigidBodyPointContactConstraint<T, d> >(this, Idx);
+
+				Handles.Add(Handle);
+				Queue.Dequeue(PointConstraints[Idx]);
+			}
+			LifespanCounter++;
 		}
 
 #if CHAOS_DETERMINISTIC
 		//todo: sort constraints
 #endif
 
-#if !UE_BUILD_SHIPPING
-		FString OutLog;
-		if (bGatherStats)
-		{
-#if CHAOS_PARTICLEHANDLE_TODO
-			const float NumParticles = InParticles.Size();
-			OutLog = FString::Printf(TEXT("ComputeConstraints stats:\n"
-				"Total Particles:%d\nSimulated Particles:%d (%.2f%%)\n"
-				"Max candidates per instance:%d (%.2f%%)\n"
-				"Max candidates skipped per instance (NP skipped):%d (%.2f%%)\n"
-				"Max narrow phase tests per instance:%d (%.2f%%)\n"
-				"Max narrow phase rejected per instance (NP rejected):%d (%.2f%%)\n"
-				"Constraints generated:%d\n"
-			),
-				InParticles.Size(),
-				SimulatedParticles, SimulatedParticles / NumParticles * 100.f,
-				BroadphasePotentials.MaxCount, BroadphasePotentials.MaxCount / NumParticles * 100.f,
-				NarrowPhaseSkipped.MaxCount, NarrowPhaseSkipped.MaxCount / NumParticles * 100.f,
-				NarrowPhasePerformed.MaxCount, NarrowPhasePerformed.MaxCount / NumParticles * 100.f,
-				NarrowPhaseRejected.MaxCount, NarrowPhaseRejected.MaxCount / NumParticles * 100.f,
-				Constraints.Num()
-			);
-
-			OutLog += FString::Printf(TEXT("Potentials per instance distribution:\n"));
-			OutLog += BroadphasePotentials.ToString();
-
-
-			OutLog += FString::Printf(TEXT("\nCandidates skipped per instance (NP skipped) distribution:\n"));
-			OutLog += NarrowPhaseSkipped.ToString();
-
-			OutLog += FString::Printf(TEXT("\nNarrow phase performed per instance distribution:\n"));
-			OutLog += NarrowPhasePerformed.ToString();
-
-			OutLog += FString::Printf(TEXT("\nNarrow phase candidates rejected per instance distribution:\n"));
-			OutLog += NarrowPhaseRejected.ToString();
-
-			OutLog += FString::Printf(TEXT("\nBounds distribution:\n"));
-			OutLog += BoundsDistribution.ToString();
-
-			UE_LOG(LogChaos, Warning, TEXT("%s"), *OutLog);
-			PendingHierarchyDump = 0;
-#endif
-		}
-#endif
-		//	if (TotalTested > 0)
-		//	{
-				//UE_LOG(LogChaos, Warning, TEXT("ComputeConstraints: rejected:%f out of rejected:%f = %.3f"), TotalRejected, TotalTested, TotalRejected / TotalTested);
-		//	}
+		COLLISION_OPT_OUT(StatData.Print());
 	}
 }
\ No newline at end of file
