Don't delete invalid pointer.

diff --git a/Engine/Source/Runtime/Experimental/Chaos/Private/Chaos/BVHParticles.cpp b/Engine/Source/Runtime/Experimental/Chaos/Private/Chaos/BVHParticles.cpp
--- a/Engine/Source/Runtime/Experimental/Chaos/Private/Chaos/BVHParticles.cpp
+++ b/Engine/Source/Runtime/Experimental/Chaos/Private/Chaos/BVHParticles.cpp
@@ -56,7 +56,6 @@ TBVHParticles<T,d>::TBVHParticles(const TBVHParticles<T, d>& Other)
 	{
 		X(i) = Other.X(i);
 	}
-    delete MBVH;
 	MBVH = new TBoundingVolumeHierarchy<TParticles<T, d>, TArray<int32>, T, d>(*this, CollisionParticlesBVHDepth);
 }
 
