Fixed crashes due to array length mismatches in geom collection proxy

diff --git a/Engine/Source/Runtime/Experimental/ChaosSolvers/Private/PhysicsProxy/GeometryCollectionPhysicsProxy.cpp b/Engine/Source/Runtime/Experimental/ChaosSolvers/Private/PhysicsProxy/GeometryCollectionPhysicsProxy.cpp
--- a/Engine/Source/Runtime/Experimental/ChaosSolvers/Private/PhysicsProxy/GeometryCollectionPhysicsProxy.cpp
+++ b/Engine/Source/Runtime/Experimental/ChaosSolvers/Private/PhysicsProxy/GeometryCollectionPhysicsProxy.cpp
@@ -564,7 +564,7 @@ void FGeometryCollectionPhysicsProxy::InitializeBodiesGT()
 		GeometryCollectionAlgo::GlobalMatrices(DynamicCollection->Transform, DynamicCollection->Parent, Transform);
 		check(DynamicCollection->Transform.Num() == Transform.Num());
 
-		NumParticles = SimulatableParticles.Count(true);
+		NumParticles = SimulatableParticles.Num();
 		BaseParticleIndex = 0; // Are we always zero indexed now?
 
 		// Create game thread particles...
@@ -3923,6 +3923,12 @@ void FGeometryCollectionPhysicsProxy::BufferPhysicsResults()
 		TargetResults.DisabledStates.SetNumUninitialized(NumParticles);
 		for (int32 Idx = 0; Idx < NumParticles; Idx++)
 		{
+			if(!SolverParticleHandles[Idx])
+			{
+				TargetResults.DisabledStates[Idx] = false;
+				continue;
+			}
+
 			TargetResults.DisabledStates[Idx] = SolverParticleHandles[Idx]->Disabled();
 		}
 		//TargetResults.DisabledStates.Append(&Particles.DisabledRef(BaseParticleIndex), NumParticles);
@@ -3955,6 +3961,11 @@ void FGeometryCollectionPhysicsProxy::BufferPhysicsResults()
 			const int32 ParticleIndex = BaseParticleIndex + TransformIndex;
 //			const int32 ParentIndex = ClusterID[ParticleIndex].Id;
 
+			if(!SolverParticleHandles[ParticleIndex])
+			{
+				continue;
+			}
+
 			Chaos::TPBDRigidParticleHandle<float, 3>& Handle = *SolverParticleHandles[ParticleIndex];
 
 			if (/*ParentIndex == INDEX_NONE*/true)
