Chaos : Fix CIS

diff --git a/Engine/Source/Runtime/Experimental/Chaos/Private/Chaos/PBDRigidClustering.cpp b/Engine/Source/Runtime/Experimental/Chaos/Private/Chaos/PBDRigidClustering.cpp
--- a/Engine/Source/Runtime/Experimental/Chaos/Private/Chaos/PBDRigidClustering.cpp
+++ b/Engine/Source/Runtime/Experimental/Chaos/Private/Chaos/PBDRigidClustering.cpp
@@ -1913,7 +1913,6 @@ namespace Chaos
 }
 
 using namespace Chaos;
-template class Chaos::TPBDRigidClustering<TPBDRigidsEvolutionGBF<float, 3>, TPBDCollisionConstraints<float, 3>, float, 3>;
 template class CHAOS_API Chaos::TPBDRigidClustering<TPBDRigidsEvolutionGBF<float, 3>, TPBDCollisionConstraints<float, 3>, float, 3>;
 #if CHAOS_PARTICLEHANDLE_TODO
 template class CHAOS_API Chaos::TPBDRigidClustering<TPBDRigidsEvolutionPGS<float, 3>, TPBDCollisionConstraintPGS<float, 3>, float, 3>;
