Remove duplicate definition
#rnx

diff --git a/Engine/Source/Runtime/Experimental/Chaos/Private/Chaos/PBDRigidsEvolution.cpp b/Engine/Source/Runtime/Experimental/Chaos/Private/Chaos/PBDRigidsEvolution.cpp
--- a/Engine/Source/Runtime/Experimental/Chaos/Private/Chaos/PBDRigidsEvolution.cpp
+++ b/Engine/Source/Runtime/Experimental/Chaos/Private/Chaos/PBDRigidsEvolution.cpp
@@ -598,7 +598,6 @@ namespace Chaos
 	}
 }
 
-template class Chaos::TPBDRigidsEvolutionBase<Chaos::TPBDRigidsEvolutionGBF<float, 3>, Chaos::TPBDCollisionConstraints<float, 3>, float, 3>;
 #ifdef __clang__
 template class CHAOS_API Chaos::TPBDRigidsEvolutionBase<Chaos::TPBDRigidsEvolutionGBF<float, 3>, Chaos::TPBDCollisionConstraints<float,3>, float, 3>;
 #else
