Reimplemented detecting when an object should go dynamic for geometry collections
Push the correct results out to the scene proxy
Fixed crash handling game->phys data transition

diff --git a/Engine/Source/Runtime/Engine/Private/PhysicsEngine/Experimental/PhysScene_Chaos.cpp b/Engine/Source/Runtime/Engine/Private/PhysicsEngine/Experimental/PhysScene_Chaos.cpp
--- a/Engine/Source/Runtime/Engine/Private/PhysicsEngine/Experimental/PhysScene_Chaos.cpp
+++ b/Engine/Source/Runtime/Engine/Private/PhysicsEngine/Experimental/PhysScene_Chaos.cpp
@@ -1488,6 +1488,11 @@ void FPhysScene_ChaosInterface::SyncBodies(Chaos::FPhysicsSolver* Solver)
 					}
 				}
 			}
+			else if(ProxyBase->GetType() == EPhysicsProxyType::GeometryCollectionType)
+			{
+				FGeometryCollectionPhysicsProxy* Proxy = static_cast<FGeometryCollectionPhysicsProxy*>(ProxyBase);
+				Proxy->PullFromPhysicsState();
+			}
 		}
 	}
 }
diff --git a/Engine/Source/Runtime/Experimental/ChaosSolvers/Private/PhysicsProxy/GeometryCollectionPhysicsProxy.cpp b/Engine/Source/Runtime/Experimental/ChaosSolvers/Private/PhysicsProxy/GeometryCollectionPhysicsProxy.cpp
--- a/Engine/Source/Runtime/Experimental/ChaosSolvers/Private/PhysicsProxy/GeometryCollectionPhysicsProxy.cpp
+++ b/Engine/Source/Runtime/Experimental/ChaosSolvers/Private/PhysicsProxy/GeometryCollectionPhysicsProxy.cpp
@@ -3829,7 +3829,7 @@ void FGeometryCollectionPhysicsProxy::PushToPhysicsState(const Chaos::FParticleD
 		// probably a bug, as the interface code is assuming it needs to lock
 		// to prevent stomping on a double buffer.  We should look into doing 
 		// away with that lock, but for the time being, we'll assert:
-		check(false);
+		ensure(false);
 		return;
 	}
 	
@@ -3925,11 +3925,19 @@ void FGeometryCollectionPhysicsProxy::BufferPhysicsResults()
 		{
 			if(!SolverParticleHandles[Idx])
 			{
-				TargetResults.DisabledStates[Idx] = false;
+				TargetResults.DisabledStates[Idx] = true;
 				continue;
 			}
 
-			TargetResults.DisabledStates[Idx] = SolverParticleHandles[Idx]->Disabled();
+			if(SolverParticleHandles[Idx]->Disabled())
+			{
+				TargetResults.DisabledStates[Idx] = true;
+			}
+			else
+			{
+				TargetResults.DisabledStates[Idx] = false;
+				IsObjectDynamic = true;
+			}
 		}
 		//TargetResults.DisabledStates.Append(&Particles.DisabledRef(BaseParticleIndex), NumParticles);
 	}
diff --git a/Engine/Source/Runtime/Experimental/ChaosSolvers/Public/PhysicsProxy/GeometryCollectionPhysicsProxy.h b/Engine/Source/Runtime/Experimental/ChaosSolvers/Public/PhysicsProxy/GeometryCollectionPhysicsProxy.h
--- a/Engine/Source/Runtime/Experimental/ChaosSolvers/Public/PhysicsProxy/GeometryCollectionPhysicsProxy.h
+++ b/Engine/Source/Runtime/Experimental/ChaosSolvers/Public/PhysicsProxy/GeometryCollectionPhysicsProxy.h
@@ -195,6 +195,8 @@ public:
 	//
 	void SetCollisionParticlesPerObjectFraction(float CollisionParticlesPerObjectFractionIn) {CollisionParticlesPerObjectFraction = CollisionParticlesPerObjectFractionIn;}
 
+	const FGeometryCollectionResults* GetConsumerResultsGT() { return PhysToGameInterchange->GetConsumerBuffer(); }
+
 	/*
 	* For Testing ONLY!
 	*/
diff --git a/Engine/Source/Runtime/Experimental/GeometryCollectionEngine/Private/GeometryCollection/GeometryCollectionComponent.cpp b/Engine/Source/Runtime/Experimental/GeometryCollectionEngine/Private/GeometryCollection/GeometryCollectionComponent.cpp
--- a/Engine/Source/Runtime/Experimental/GeometryCollectionEngine/Private/GeometryCollection/GeometryCollectionComponent.cpp
+++ b/Engine/Source/Runtime/Experimental/GeometryCollectionEngine/Private/GeometryCollection/GeometryCollectionComponent.cpp
@@ -1993,8 +1993,8 @@ void UGeometryCollectionComponent::CalculateGlobalMatrices()
 	SCOPE_CYCLE_COUNTER(STAT_GCCUGlobalMatrices);
 	FChaosSolversModule* Module = FChaosSolversModule::GetModule();
 	Module->LockResultsRead();
-
-	const FGeometryCollectionResults* Results = PhysicsProxy ? &PhysicsProxy->GetPhysicsResults().GetGameDataForRead() : nullptr;
+	
+	const FGeometryCollectionResults* Results = PhysicsProxy ? PhysicsProxy->GetConsumerResultsGT() : nullptr;
 
 	const int32 NumTransforms = Results ? Results->GlobalTransforms.Num() : 0;
 	if(NumTransforms > 0)
