Particle debug names

diff --git a/Engine/Source/Runtime/Engine/Private/PhysicsEngine/Experimental/PhysInterface_Chaos.cpp b/Engine/Source/Runtime/Engine/Private/PhysicsEngine/Experimental/PhysInterface_Chaos.cpp
--- a/Engine/Source/Runtime/Engine/Private/PhysicsEngine/Experimental/PhysInterface_Chaos.cpp
+++ b/Engine/Source/Runtime/Engine/Private/PhysicsEngine/Experimental/PhysInterface_Chaos.cpp
@@ -159,6 +159,9 @@ void FPhysInterface_Chaos::CreateActor(const FActorCreationParams& InParams, FPh
 	// the particle is added to the scene later.
 	Handle->SetX(InParams.InitialTM.GetLocation());
 	Handle->SetR(InParams.InitialTM.GetRotation());
+#if CHAOS_CHECKED
+	Handle->SetDebugName(InParams.DebugName);
+#endif
 }
 
 
diff --git a/Engine/Source/Runtime/Experimental/Chaos/Public/Chaos/GeometryParticles.h b/Engine/Source/Runtime/Experimental/Chaos/Public/Chaos/GeometryParticles.h
--- a/Engine/Source/Runtime/Experimental/Chaos/Public/Chaos/GeometryParticles.h
+++ b/Engine/Source/Runtime/Experimental/Chaos/Public/Chaos/GeometryParticles.h
@@ -102,6 +102,9 @@ namespace Chaos
 			TArrayCollection::AddArray(&MHasBounds);
 			TArrayCollection::AddArray(&MSpatialIdx);
 			TArrayCollection::AddArray(&MHashResult);
+#if CHAOS_CHECKED
+			TArrayCollection::AddArray(&MDebugName);
+#endif
 
 			if (IsRigidBodySim())
 			{
@@ -143,6 +146,9 @@ namespace Chaos
 #if CHAOS_DETERMINISTIC
 			TArrayCollection::AddArray(&MParticleIDs);
 #endif
+#if CHAOS_CHECKED
+			TArrayCollection::AddArray(&MDebugName);
+#endif
 
 			if (IsRigidBodySim())
 			{
@@ -170,6 +176,9 @@ namespace Chaos
 #if CHAOS_DETERMINISTIC
 			TArrayCollection::AddArray(&MParticleIDs);
 #endif
+#if CHAOS_CHECKED
+			TArrayCollection::AddArray(&MDebugName);
+#endif
 
 			if (IsRigidBodySim())
 			{
@@ -261,6 +270,18 @@ namespace Chaos
 			return MHashResult[Index];
 		}
 
+#if CHAOS_CHECKED
+		const FName& DebugName(const int32 Index) const
+		{
+			return MDebugName[Index];
+		}
+
+		FName& DebugName(const int32 Index)
+		{
+			return MDebugName[Index];
+		}
+#endif
+
 		CHAOS_API const TBox<T, d>& WorldSpaceInflatedBounds(const int32 Index) const
 		{
 			return MWorldSpaceInflatedBounds[Index];
@@ -408,7 +429,9 @@ namespace Chaos
 #if CHAOS_DETERMINISTIC
 		TArrayCollectionArray<FParticleID> MParticleIDs;
 #endif
-
+#if CHAOS_CHECKED
+		TArrayCollectionArray<FName> MDebugName;
+#endif
 	};
 
 	template <typename T, int d, EGeometryParticlesSimType SimType>
diff --git a/Engine/Source/Runtime/Experimental/Chaos/Public/Chaos/ParticleDirtyFlags.h b/Engine/Source/Runtime/Experimental/Chaos/Public/Chaos/ParticleDirtyFlags.h
--- a/Engine/Source/Runtime/Experimental/Chaos/Public/Chaos/ParticleDirtyFlags.h
+++ b/Engine/Source/Runtime/Experimental/Chaos/Public/Chaos/ParticleDirtyFlags.h
@@ -32,6 +32,9 @@ namespace Chaos
 		GravityEnabled = 1 << 20,
 		SpatialIdx     = 1 << 21,
 		HashResult     = 1 << 22
+#if CHAOS_CHECKED
+		, DebugName    = 1 << 23
+#endif
 	};
 
 	class FParticleDirtyFlags
diff --git a/Engine/Source/Runtime/Experimental/Chaos/Public/Chaos/ParticleHandle.h b/Engine/Source/Runtime/Experimental/Chaos/Public/Chaos/ParticleHandle.h
--- a/Engine/Source/Runtime/Experimental/Chaos/Public/Chaos/ParticleHandle.h
+++ b/Engine/Source/Runtime/Experimental/Chaos/Public/Chaos/ParticleHandle.h
@@ -394,6 +394,11 @@ public:
 
 	void SetHashResultLowLevel(uint32 Value) { GeometryParticles->HashResultLowLevel(ParticleIdx) = Value; }
 	uint32 GetHashResultLowLevel() const { return GeometryParticles->HashResultLowLevel(ParticleIdx); }
+
+#if CHAOS_CHECKED
+	const FName& DebugName() const { return GeometryParticles->DebugName(ParticleIdx); }
+	void SetDebugName(const FName& InDebugName) { GeometryParticles->DebugName(ParticleIdx) = InDebugName; }
+#endif
 	
 	EObjectStateType ObjectState() const;
 
@@ -1129,6 +1134,15 @@ public:
 		this->MUserData = InUserData;
 	}
 
+#if CHAOS_CHECKED
+	const FName& DebugName() const { return this->MDebugName; }
+	void SetDebugName(const FName& InDebugName)
+	{
+		this->MarkDirty(EParticleFlags::DebugName);
+		this->MDebugName = InDebugName;
+	}
+#endif
+
 	void SetHashResultLowLevel(uint32 Value)
 	{
 		MarkDirty(EParticleFlags::HashResult);
@@ -1226,6 +1240,10 @@ private:
 		return MGeometry;
 	}
 
+#if CHAOS_CHECKED
+	FName MDebugName;
+#endif
+
 protected:
 	EParticleType Type;
 	FParticleDirtyFlags MDirtyFlags;
@@ -1258,6 +1276,9 @@ public:
 		, SpatialIdx(FSpatialAccelerationIdx{ 0,0 })
 		, HashResult(0)
 		, DirtyFlags()
+#if CHAOS_CHECKED
+		, DebugName(NAME_None)
+#endif
 	{}
 
 	TGeometryParticleData(const TGeometryParticle<T, d>& InParticle)
@@ -1268,6 +1289,9 @@ public:
 		, SpatialIdx(InParticle.SpatialIdx())
 		, HashResult(InParticle.GetHashResultLowLevel())
 		, DirtyFlags(InParticle.DirtyFlags())
+#if CHAOS_CHECKED
+		, DebugName(InParticle.DebugName())
+#endif
 	{}
 
 	void Reset() 
@@ -1279,6 +1303,9 @@ public:
 		SpatialIdx = FSpatialAccelerationIdx{ 0,0 };
 		HashResult = 0;
 		DirtyFlags.Clear();
+#if CHAOS_CHECKED
+		DebugName = NAME_None;
+#endif
 	}
 
 	TVector<T, d> X;
@@ -1287,6 +1314,9 @@ public:
 	FSpatialAccelerationIdx SpatialIdx;
 	uint32 HashResult;
 	FParticleDirtyFlags DirtyFlags;
+#if CHAOS_CHECKED
+	FName DebugName;
+#endif
 };
 
 
diff --git a/Engine/Source/Runtime/Experimental/ChaosSolvers/Private/PhysicsProxy/SingleParticlePhysicsProxy.cpp b/Engine/Source/Runtime/Experimental/ChaosSolvers/Private/PhysicsProxy/SingleParticlePhysicsProxy.cpp
--- a/Engine/Source/Runtime/Experimental/ChaosSolvers/Private/PhysicsProxy/SingleParticlePhysicsProxy.cpp
+++ b/Engine/Source/Runtime/Experimental/ChaosSolvers/Private/PhysicsProxy/SingleParticlePhysicsProxy.cpp
@@ -61,6 +61,9 @@ void FSingleParticlePhysicsProxy<Chaos::TGeometryParticle<float, 3>>::PushToPhys
 		}
 		RigidHandle->SetSpatialIdx(Data->SpatialIdx);	//todo: this needs to only happen once during initialization
 		RigidHandle->SetHashResultLowLevel(Data->HashResult);
+#if CHAOS_CHECKED
+		RigidHandle->SetDebugName(Data->DebugName);
+#endif
 	}
 }
 
@@ -119,6 +122,9 @@ void FSingleParticlePhysicsProxy<Chaos::TKinematicGeometryParticle<float, 3>>::P
 
 		RigidHandle->SetSpatialIdx(Data->SpatialIdx);	//todo: this needs to only happen once during initialization
 		RigidHandle->SetHashResultLowLevel(Data->HashResult);	//todo: this needs to only happen once during initialization
+#if CHAOS_CHECKED
+		RigidHandle->SetDebugName(Data->DebugName);
+#endif
 		
 		if (Data->Geometry && Data->Geometry->HasBoundingBox())
 		{
@@ -187,6 +193,9 @@ void FSingleParticlePhysicsProxy<Chaos::TPBDRigidParticle<float, 3>>::PushToPhys
 		// This should be moved to initialization.
 		RigidHandle->SetSpatialIdx(Data->SpatialIdx);	//todo: this needs to only happen once during initialization
 		RigidHandle->SetHashResultLowLevel(Data->HashResult);	//todo: this needs to only happen once during initialization
+#if CHAOS_CHECKED
+		RigidHandle->SetDebugName(Data->DebugName);
+#endif
 
 		if (Data->DirtyFlags.IsDirty(Chaos::EParticleFlags::X))
 		{
