Speculative mac/linux fix.

diff --git a/Engine/Source/Runtime/Experimental/Chaos/Public/Chaos/PBDRigidParticles.h b/Engine/Source/Runtime/Experimental/Chaos/Public/Chaos/PBDRigidParticles.h
--- a/Engine/Source/Runtime/Experimental/Chaos/Public/Chaos/PBDRigidParticles.h
+++ b/Engine/Source/Runtime/Experimental/Chaos/Public/Chaos/PBDRigidParticles.h
@@ -162,7 +162,11 @@ class TPBDRigidParticles : public TRigidParticles<T, d>
 	TArrayCollectionArray<TVector<T, d>> MPreW;
 };
 
+#if PLATFORM_MAC || PLATFORM_LINUX
+extern template class CHAOS_API TPBDRigidParticles<float,3>;
+#else
 extern template class TPBDRigidParticles<float,3>;
+#endif
 
 template <typename T, int d>
 FChaosArchive& operator<<(FChaosArchive& Ar, TPBDRigidParticles<T, d>& Particles)
