Chaos - Fix kinematic RigidParticles ending up in ActiveParticles Array.

diff --git a/Engine/Source/Runtime/Experimental/Chaos/Public/Chaos/PBDRigidsSOAs.h b/Engine/Source/Runtime/Experimental/Chaos/Public/Chaos/PBDRigidsSOAs.h
--- a/Engine/Source/Runtime/Experimental/Chaos/Public/Chaos/PBDRigidsSOAs.h
+++ b/Engine/Source/Runtime/Experimental/Chaos/Public/Chaos/PBDRigidsSOAs.h
@@ -106,7 +106,8 @@ public:
 	{
 		check(Particle->AsClustered() == nullptr);	//not supported
 
-		if (auto PBDRigid = Particle->AsDynamic())
+		auto PBDRigid = Particle->AsDynamic();
+		if(PBDRigid && PBDRigid->ObjectState() == EObjectStateType::Dynamic)
 		{
 			RemoveFromMapAndArray(PBDRigid, ActiveParticlesToIndex, ActiveParticlesArray);
 		}
@@ -132,7 +133,11 @@ public:
 			{
 				Particle->MoveToSOA(*DynamicDisabledParticles);
 			}
-			RemoveFromMapAndArray(PBDRigid, ActiveParticlesToIndex, ActiveParticlesArray);
+
+			if (Particle->ObjectState() == EObjectStateType::Dynamic)
+			{
+				RemoveFromMapAndArray(PBDRigid, ActiveParticlesToIndex, ActiveParticlesArray);
+			}
 		}
 		else if (Particle->AsKinematic())
 		{
@@ -152,7 +157,7 @@ public:
 			if (auto PBDRigidClustered = Particle->AsClustered())
 			{
 				InsertToMapAndArray(PBDRigidClustered, NonDisabledClusteredToIndex, NonDisabledClusteredArray);
-				if (!PBDRigid->Sleeping())
+				if (!PBDRigid->Sleeping() && Particle->ObjectState() == EObjectStateType::Dynamic)
 				{
 					InsertToMapAndArray(PBDRigidClustered, ActiveClusteredToIndex, ActiveClusteredArray);
 				}
@@ -162,7 +167,7 @@ public:
 				SetDynamicParticleSOA(Particle->AsDynamic());
 			}
 
-			if (!PBDRigid->Sleeping())
+			if (!PBDRigid->Sleeping() && Particle->ObjectState() == EObjectStateType::Dynamic)
 			{
 				InsertToMapAndArray(PBDRigid, ActiveParticlesToIndex, ActiveParticlesArray);
 			}
@@ -182,7 +187,8 @@ public:
 
 	void ActivateParticle(TGeometryParticleHandle<T, d>* Particle)
 	{
-		if (auto PBDRigid = Particle->AsDynamic())
+		auto PBDRigid = Particle->AsDynamic();
+		if(PBDRigid && PBDRigid->ObjectState() == EObjectStateType::Dynamic)
 		{
 			check(!PBDRigid->Disabled());
 			if (auto PBDRigidClustered = Particle->AsClustered())
@@ -198,7 +204,8 @@ public:
 
 	void DeactivateParticle(TGeometryParticleHandle<T, d>* Particle)
 	{
-		if (auto PBDRigid = Particle->AsDynamic())
+		auto PBDRigid = Particle->AsDynamic();
+		if(PBDRigid && PBDRigid->ObjectState() == EObjectStateType::Dynamic)
 		{
 			check(!PBDRigid->Disabled());
 			if (auto PBDRigidClustered = Particle->AsClustered())
@@ -224,6 +231,11 @@ public:
 	{
 		const EObjectStateType State = Particle->ObjectState();
 
+		if (State != EObjectStateType::Dynamic)
+		{
+			RemoveFromMapAndArray(Particle->AsDynamic(), ActiveParticlesToIndex, ActiveParticlesArray);
+		}
+
 		// Move to appropriate dynamic SOA
 		switch (State)
 		{
