Chaos - Fix SetDynamicParticleSOA not removing particles switched off of dynamic.
Fix HandleObjectStateTest from expecting two active particles when only one is dynamic.

diff --git a/Engine/Source/Runtime/Experimental/Chaos/Public/Chaos/PBDRigidsSOAs.h b/Engine/Source/Runtime/Experimental/Chaos/Public/Chaos/PBDRigidsSOAs.h
--- a/Engine/Source/Runtime/Experimental/Chaos/Public/Chaos/PBDRigidsSOAs.h
+++ b/Engine/Source/Runtime/Experimental/Chaos/Public/Chaos/PBDRigidsSOAs.h
@@ -104,9 +104,9 @@ public:
 
 	void DestroyParticle(TGeometryParticleHandle<T, d>* Particle)
 	{
-		check(Particle->AsClustered() == nullptr);	//not supported
+		check(Particle->CastToClustered() == nullptr);	//not supported
 
-		auto PBDRigid = Particle->AsDynamic();
+		auto PBDRigid = Particle->CastToRigidParticle();
 		if(PBDRigid && PBDRigid->ObjectState() == EObjectStateType::Dynamic)
 		{
 			RemoveFromMapAndArray(PBDRigid, ActiveParticlesToIndex, ActiveParticlesArray);
@@ -118,13 +118,13 @@ public:
 
 	void DisableParticle(TGeometryParticleHandle<T, d>* Particle)
 	{
-		if (auto PBDRigid = Particle->AsDynamic())
+		if (auto PBDRigid = Particle->CastToRigidParticle())
 		{
 			PBDRigid->Disabled() = true;
 			PBDRigid->V() = TVector<T, d>(0);
 			PBDRigid->W() = TVector<T, d>(0);
 
-			if (auto PBDRigidClustered = Particle->AsClustered())
+			if (auto PBDRigidClustered = Particle->CastToClustered())
 			{
 				RemoveFromMapAndArray(PBDRigidClustered, NonDisabledClusteredToIndex, NonDisabledClusteredArray);
 				RemoveFromMapAndArray(PBDRigidClustered, ActiveClusteredToIndex, ActiveClusteredArray);
@@ -139,7 +139,7 @@ public:
 				RemoveFromMapAndArray(PBDRigid, ActiveParticlesToIndex, ActiveParticlesArray);
 			}
 		}
-		else if (Particle->AsKinematic())
+		else if (Particle->CastToKinematicParticle())
 		{
 			Particle->MoveToSOA(*KinematicDisabledParticles);
 		}
@@ -152,9 +152,9 @@ public:
 
 	void EnableParticle(TGeometryParticleHandle<T, d>* Particle)
 	{
-		if (auto PBDRigid = Particle->AsDynamic())
+		if (auto PBDRigid = Particle->CastToRigidParticle())
 		{
-			if (auto PBDRigidClustered = Particle->AsClustered())
+			if (auto PBDRigidClustered = Particle->CastToClustered())
 			{
 				InsertToMapAndArray(PBDRigidClustered, NonDisabledClusteredToIndex, NonDisabledClusteredArray);
 				if (!PBDRigid->Sleeping() && Particle->ObjectState() == EObjectStateType::Dynamic)
@@ -164,7 +164,7 @@ public:
 			}
 			else
 			{
-				SetDynamicParticleSOA(Particle->AsDynamic());
+				SetDynamicParticleSOA(Particle->CastToRigidParticle());
 			}
 
 			if (!PBDRigid->Sleeping() && Particle->ObjectState() == EObjectStateType::Dynamic)
@@ -174,7 +174,7 @@ public:
 
 			PBDRigid->Disabled() = false;
 		}
-		else if (Particle->AsKinematic())
+		else if (Particle->CastToKinematicParticle())
 		{
 			Particle->MoveToSOA(*KinematicParticles);
 		}
@@ -187,11 +187,11 @@ public:
 
 	void ActivateParticle(TGeometryParticleHandle<T, d>* Particle)
 	{
-		auto PBDRigid = Particle->AsDynamic();
+		auto PBDRigid = Particle->CastToRigidParticle();
 		if(PBDRigid && PBDRigid->ObjectState() == EObjectStateType::Dynamic)
 		{
 			check(!PBDRigid->Disabled());
-			if (auto PBDRigidClustered = Particle->AsClustered())
+			if (auto PBDRigidClustered = Particle->CastToClustered())
 			{
 				InsertToMapAndArray(PBDRigidClustered, ActiveClusteredToIndex, ActiveClusteredArray);
 			}
@@ -204,11 +204,11 @@ public:
 
 	void DeactivateParticle(TGeometryParticleHandle<T, d>* Particle)
 	{
-		auto PBDRigid = Particle->AsDynamic();
+		auto PBDRigid = Particle->CastToRigidParticle();
 		if(PBDRigid && PBDRigid->ObjectState() == EObjectStateType::Dynamic)
 		{
 			check(!PBDRigid->Disabled());
-			if (auto PBDRigidClustered = Particle->AsClustered())
+			if (auto PBDRigidClustered = Particle->CastToClustered())
 			{
 				RemoveFromMapAndArray(PBDRigidClustered, ActiveClusteredToIndex, ActiveClusteredArray);
 			}
@@ -231,9 +231,14 @@ public:
 	{
 		const EObjectStateType State = Particle->ObjectState();
 
-		if (State != EObjectStateType::Dynamic)
+
+		if (Particle->ObjectState() != EObjectStateType::Dynamic)
+		{
+			RemoveFromMapAndArray(Particle->CastToRigidParticle(), ActiveParticlesToIndex, ActiveParticlesArray);
+		}
+		else
 		{
-			RemoveFromMapAndArray(Particle->AsDynamic(), ActiveParticlesToIndex, ActiveParticlesArray);
+			InsertToMapAndArray(Particle->CastToRigidParticle(), ActiveParticlesToIndex, ActiveParticlesArray);
 		}
 
 		// Move to appropriate dynamic SOA
