Chaos -(Resubmitting without errors)  Fix SetDynamicParticleSOA not removing particles switched off of dynamic.
Fix HandleObjectStateTest from expecting two active particles when only one is dynamic.

diff --git a/Engine/Source/Runtime/Experimental/Chaos/Public/Chaos/PBDRigidsSOAs.h b/Engine/Source/Runtime/Experimental/Chaos/Public/Chaos/PBDRigidsSOAs.h
--- a/Engine/Source/Runtime/Experimental/Chaos/Public/Chaos/PBDRigidsSOAs.h
+++ b/Engine/Source/Runtime/Experimental/Chaos/Public/Chaos/PBDRigidsSOAs.h
@@ -231,10 +231,15 @@ public:
 	{
 		const EObjectStateType State = Particle->ObjectState();
 
-		if (State != EObjectStateType::Dynamic)
+
+		if (Particle->ObjectState() != EObjectStateType::Dynamic)
 		{
 			RemoveFromMapAndArray(Particle->AsDynamic(), ActiveParticlesToIndex, ActiveParticlesArray);
 		}
+		else
+		{
+			InsertToMapAndArray(Particle->AsDynamic(), ActiveParticlesToIndex, ActiveParticlesArray);
+		}
 
 		// Move to appropriate dynamic SOA
 		switch (State)
