Fix DeactivateParticle not allowing Sleeping RigidParticles to be deactivated. (Fixes Sleep GraphTest)

diff --git a/Engine/Source/Runtime/Experimental/Chaos/Public/Chaos/PBDRigidsSOAs.h b/Engine/Source/Runtime/Experimental/Chaos/Public/Chaos/PBDRigidsSOAs.h
--- a/Engine/Source/Runtime/Experimental/Chaos/Public/Chaos/PBDRigidsSOAs.h
+++ b/Engine/Source/Runtime/Experimental/Chaos/Public/Chaos/PBDRigidsSOAs.h
@@ -205,15 +205,19 @@ public:
 	void DeactivateParticle(TGeometryParticleHandle<T, d>* Particle)
 	{
 		auto PBDRigid = Particle->AsDynamic();
-		if(PBDRigid && PBDRigid->ObjectState() == EObjectStateType::Dynamic)
+		if(PBDRigid)
 		{
-			check(!PBDRigid->Disabled());
-			if (auto PBDRigidClustered = Particle->AsClustered())
+			if (   PBDRigid->ObjectState() == EObjectStateType::Dynamic
+				|| PBDRigid->ObjectState() == EObjectStateType::Sleeping)
 			{
-				RemoveFromMapAndArray(PBDRigidClustered, ActiveClusteredToIndex, ActiveClusteredArray);
-			}
+				check(!PBDRigid->Disabled());
+				if (auto PBDRigidClustered = Particle->AsClustered())
+				{
+					RemoveFromMapAndArray(PBDRigidClustered, ActiveClusteredToIndex, ActiveClusteredArray);
+				}
 
-			RemoveFromMapAndArray(PBDRigid, ActiveParticlesToIndex, ActiveParticlesArray);
+				RemoveFromMapAndArray(PBDRigid, ActiveParticlesToIndex, ActiveParticlesArray);
+			}
 		}
 
 		UpdateViews();
