Chaos - delete redundant PBD6DConstraint

diff --git a/Engine/Source/Runtime/Experimental/Chaos/Private/Chaos/ChaosDebugDraw.cpp b/Engine/Source/Runtime/Experimental/Chaos/Private/Chaos/ChaosDebugDraw.cpp
--- a/Engine/Source/Runtime/Experimental/Chaos/Private/Chaos/ChaosDebugDraw.cpp
+++ b/Engine/Source/Runtime/Experimental/Chaos/Private/Chaos/ChaosDebugDraw.cpp
@@ -7,7 +7,6 @@
 #include "Chaos/ImplicitObjectTransformed.h"
 #include "Chaos/ImplicitObjectUnion.h"
 #include "Chaos/ParticleHandle.h"
-#include "Chaos/PBD6DJointConstraints.h"
 #include "Chaos/PBDCollisionConstraints.h"
 #include "Chaos/PBDJointConstraints.h"
 #include "Chaos/PBDRigidParticles.h"
@@ -221,16 +220,6 @@ namespace Chaos
 			}
 		}
 
-		void Draw6DofConstraintImpl(const FRigidTransform3& SpaceTransform, const FPBD6DJointConstraintHandle* ConstraintHandle, FReal ColorScale)
-		{
-			FVec3 Pa = TGenericParticleHandle<FReal, 3>(ConstraintHandle->GetConstrainedParticles()[1])->P();
-			FVec3 Pb = TGenericParticleHandle<FReal, 3>(ConstraintHandle->GetConstrainedParticles()[0])->P();
-			FVec3 Xa, Xb, CR;
-			FMatrix33 Ra, Rb;
-			ConstraintHandle->CalculateConstraintSpace(Xa, Ra, Xb, Rb, CR);
-			DrawJointConstraintImpl(SpaceTransform, Pa, Xa, Ra, Pb, Xb, Rb, CR, INDEX_NONE, ColorScale, (uint32)EDebugDrawJointFeature::Default);
-		}
-
 #endif
 
 		void DrawParticleShapes(const FRigidTransform3& SpaceTransform, const TParticleView<TGeometryParticles<float, 3>>& ParticlesView, float ColorScale, bool bDrawKinemtatic, bool bDrawDynamic)
@@ -367,31 +356,5 @@ namespace Chaos
 #endif
 		}
 
-		void Draw6DofConstraints(const FRigidTransform3& SpaceTransform, const TArray<FPBD6DJointConstraintHandle*>& ConstraintHandles, float ColorScale)
-		{
-#if CHAOS_DEBUG_DRAW
-			if (FDebugDrawQueue::IsDebugDrawingEnabled())
-			{
-				for (const FPBD6DJointConstraintHandle* ConstraintHandle : ConstraintHandles)
-				{
-					Draw6DofConstraintImpl(SpaceTransform, ConstraintHandle, ColorScale);
-				}
-			}
-#endif
-		}
-
-		void Draw6DofConstraints(const FRigidTransform3& SpaceTransform, const FPBD6DJointConstraints& Constraints, float ColorScale)
-		{
-#if CHAOS_DEBUG_DRAW
-			if (FDebugDrawQueue::IsDebugDrawingEnabled())
-			{
-				for (int32 ConstraintIndex = 0; ConstraintIndex < Constraints.NumConstraints(); ++ConstraintIndex)
-				{
-					Draw6DofConstraintImpl(SpaceTransform, Constraints.GetConstraintHandle(ConstraintIndex), ColorScale);
-				}
-			}
-#endif
-		}
-
 	}
 }
diff --git a/Engine/Source/Runtime/Experimental/Chaos/Private/Chaos/PBD6DJointConstraints.cpp b/Engine/Source/Runtime/Experimental/Chaos/Private/Chaos/PBD6DJointConstraints.cpp
deleted file mode 100644
diff --git a/Engine/Source/Runtime/Experimental/Chaos/Private/Chaos/PBDConstraintRule.cpp b/Engine/Source/Runtime/Experimental/Chaos/Private/Chaos/PBDConstraintRule.cpp
--- a/Engine/Source/Runtime/Experimental/Chaos/Private/Chaos/PBDConstraintRule.cpp
+++ b/Engine/Source/Runtime/Experimental/Chaos/Private/Chaos/PBDConstraintRule.cpp
@@ -2,7 +2,6 @@
 
 #include "Chaos/PBDConstraintRule.h"
 
-#include "Chaos/PBD6DJointConstraints.h"
 #include "Chaos/PBDCollisionConstraints.h"
 #include "Chaos/PBDJointConstraints.h"
 #include "Chaos/PBDPositionConstraints.h"
@@ -53,15 +52,13 @@ namespace Chaos
 		return Constraints.NumConstraints(); 
 	}
 
-	template class TPBDConstraintGraphRuleImpl< TPBDCollisionConstraints<float, 3> >;
-	template class TPBDConstraintGraphRuleImpl< FPBD6DJointConstraints >;
-	template class TPBDConstraintGraphRuleImpl< FPBDJointConstraints >;
-	template class TPBDConstraintGraphRuleImpl< TPBDPositionConstraints<float, 3> >;
-	template class TPBDConstraintGraphRuleImpl< TPBDRigidDynamicSpringConstraints<float, 3> >;
-	template class TPBDConstraintGraphRuleImpl< TPBDRigidSpringConstraints<float, 3> >;
+	template class TPBDConstraintGraphRuleImpl<TPBDCollisionConstraints<float, 3>>;
+	template class TPBDConstraintGraphRuleImpl<FPBDJointConstraints>;
+	template class TPBDConstraintGraphRuleImpl<TPBDPositionConstraints<float, 3>>;
+	template class TPBDConstraintGraphRuleImpl<TPBDRigidDynamicSpringConstraints<float, 3>>;
+	template class TPBDConstraintGraphRuleImpl<TPBDRigidSpringConstraints<float, 3>>;
 
 	template class TPBDConstraintColorRule<TPBDCollisionConstraints<float, 3>>;
-	template class TPBDConstraintIslandRule<FPBD6DJointConstraints>;
 	template class TPBDConstraintIslandRule<FPBDJointConstraints>;
 	template class TPBDConstraintIslandRule<TPBDPositionConstraints<float, 3>>;
 	template class TPBDConstraintIslandRule<TPBDRigidDynamicSpringConstraints<float, 3>>;
diff --git a/Engine/Source/Runtime/Experimental/Chaos/Public/Chaos/ChaosDebugDraw.h b/Engine/Source/Runtime/Experimental/Chaos/Public/Chaos/ChaosDebugDraw.h
--- a/Engine/Source/Runtime/Experimental/Chaos/Public/Chaos/ChaosDebugDraw.h
+++ b/Engine/Source/Runtime/Experimental/Chaos/Public/Chaos/ChaosDebugDraw.h
@@ -31,8 +31,6 @@ namespace Chaos
 		CHAOS_API void DrawCollisions(const FRigidTransform3& SpaceTransform, const TArray<TPBDCollisionConstraintHandle<float, 3>*>& ConstraintHandles, float ColorScale);
 		CHAOS_API void DrawJointConstraints(const FRigidTransform3& SpaceTransform, const TArray<FPBDJointConstraintHandle*>& ConstraintHandles, float ColorScale, uint32 FeatureMask = (uint32)EDebugDrawJointFeature::Default);
 		CHAOS_API void DrawJointConstraints(const FRigidTransform3& SpaceTransform, const FPBDJointConstraints& Constraints, float ColorScale, uint32 FeatureMask = (uint32)EDebugDrawJointFeature::Default);
-		CHAOS_API void Draw6DofConstraints(const FRigidTransform3& SpaceTransform, const TArray<FPBD6DJointConstraintHandle*>& ConstraintHandles, float ColorScale);
-		CHAOS_API void Draw6DofConstraints(const FRigidTransform3& SpaceTransform, const FPBD6DJointConstraints& Constraints, float ColorScale);
 
 		extern CHAOS_API float ConstraintAxisLen;
 		extern CHAOS_API float BodyAxisLen;
diff --git a/Engine/Source/Runtime/Experimental/Chaos/Public/Chaos/ChaosDebugDrawDeclares.h b/Engine/Source/Runtime/Experimental/Chaos/Public/Chaos/ChaosDebugDrawDeclares.h
--- a/Engine/Source/Runtime/Experimental/Chaos/Public/Chaos/ChaosDebugDrawDeclares.h
+++ b/Engine/Source/Runtime/Experimental/Chaos/Public/Chaos/ChaosDebugDrawDeclares.h
@@ -14,10 +14,6 @@ namespace Chaos
 {
 	class FConstraintHandle;
 
-	class FPBD6DJointConstraintHandle;
-
-	class FPBD6DJointConstraints;
-
 	template<typename T, int d>
 	class TPBDCollisionConstraints;
 
diff --git a/Engine/Source/Runtime/Experimental/Chaos/Public/Chaos/PBD6DJointConstraints.h b/Engine/Source/Runtime/Experimental/Chaos/Public/Chaos/PBD6DJointConstraints.h
deleted file mode 100644
diff --git a/Engine/Source/Runtime/Experimental/Chaos/Public/Chaos/PBDJointConstraintSolver.h b/Engine/Source/Runtime/Experimental/Chaos/Public/Chaos/PBDJointConstraintSolver.h
--- a/Engine/Source/Runtime/Experimental/Chaos/Public/Chaos/PBDJointConstraintSolver.h
+++ b/Engine/Source/Runtime/Experimental/Chaos/Public/Chaos/PBDJointConstraintSolver.h
@@ -131,7 +131,8 @@ namespace Chaos
 			Rs[1] = Qs[1] * XLs[1].GetRotation();
 		}
 
-		void AddLinearConstraints_Locked(
+		// 3 constraints along principle axes
+		void AddLinearConstraints_Point(
 			const FPBDJointSolverSettings& SolverSettings,
 			const FPBDJointSettings& JointSettings)
 		{
@@ -149,6 +150,7 @@ namespace Chaos
 			LinearConstraintErrors[2] = ConstraintSeparation[2];
 		}
 
+		// up to 1 constraint limiting distance
 		void AddLinearConstraints_Sphere(
 			const FPBDJointSolverSettings& SolverSettings,
 			const FPBDJointSettings& JointSettings)
@@ -171,6 +173,7 @@ namespace Chaos
 			}
 		}
 
+		// up to 2 constraint: 1 limiting distance along the axis and another limiting lateral distance from the axis
 		void AddLinearConstraints_Cylinder(
 			const FPBDJointSolverSettings& SolverSettings,
 			const FPBDJointSettings& JointSettings,
@@ -204,6 +207,7 @@ namespace Chaos
 			}
 		}
 
+		// up to 1 constraint limiting distance along the axis (lateral motion unrestricted)
 		void AddLinearConstraints_Plane(
 			const FPBDJointSolverSettings& SolverSettings,
 			const FPBDJointSettings& JointSettings,
@@ -225,6 +229,7 @@ namespace Chaos
 			}
 		}
 
+		// up to 1 constraint limiting rotation about twist axes
 		void AddAngularConstraints_Twist(
 			const FPBDJointSolverSettings& SolverSettings,
 			const FPBDJointSettings& JointSettings)
@@ -270,6 +275,7 @@ namespace Chaos
 			}
 		}
 
+		// up to 1 constraint limiting angle between twist axes
 		void AddAngularConstraints_Cone(
 			const FPBDJointSolverSettings& SolverSettings,
 			const FPBDJointSettings& JointSettings)
@@ -317,6 +323,7 @@ namespace Chaos
 			}
 		}
 
+		// up to 1 constraint limiting rotation about swing axis (relative to body 0)
 		void AddAngularConstraints_Swing(
 			const FPBDJointSolverSettings& SolverSettings,
 			const FPBDJointSettings& JointSettings,
@@ -384,6 +391,7 @@ namespace Chaos
 			}
 		}
 
+		// Add linear constraints to solver
 		void AddLinearConstraints(
 			const FPBDJointSolverSettings& SolverSettings,
 			const FPBDJointSettings& JointSettings)
@@ -393,7 +401,7 @@ namespace Chaos
 			const TVector<EJointMotionType, 3>& Motion = JointSettings.Motion.LinearMotionTypes;
 			if ((Motion[0] == EJointMotionType::Locked) && (Motion[1] == EJointMotionType::Locked) && (Motion[2] == EJointMotionType::Locked))
 			{
-				AddLinearConstraints_Locked(SolverSettings, JointSettings);
+				AddLinearConstraints_Point(SolverSettings, JointSettings);
 			}
 			else if ((Motion[0] == EJointMotionType::Limited) && (Motion[1] == EJointMotionType::Limited) && (Motion[2] == EJointMotionType::Limited))
 			{
@@ -432,6 +440,7 @@ namespace Chaos
 			}
 		}
 
+		// Add angular constraints to solver
 		void AddAngularConstraints(
 			const FPBDJointSolverSettings& SolverSettings,
 			const FPBDJointSettings& JointSettings)
