Chaos : Collisions

- Collision creation pipeline api updates

diff --git a/Engine/Source/Runtime/Experimental/Chaos/Private/Chaos/CollisionResolution.cpp b/Engine/Source/Runtime/Experimental/Chaos/Private/Chaos/CollisionResolution.cpp
new file mode 100644
--- /dev/null
+++ b/Engine/Source/Runtime/Experimental/Chaos/Private/Chaos/CollisionResolution.cpp
@@ -0,0 +1,1129 @@
+// Copyright 1998-2019 Epic Games, Inc. All Rights Reserved.
+
+#include "Chaos/CollisionResolution.h"
+
+#include "Chaos/ChaosPerfTest.h"
+#include "Chaos/Capsule.h"
+#include "Chaos/CollisionResolutionTypes.h"
+#include "Chaos/CollisionResolutionUtil.h"
+#include "Chaos/Convex.h"
+#include "Chaos/Defines.h"
+#include "Chaos/ImplicitObjectScaled.h"
+#include "Chaos/ImplicitObjectTransformed.h"
+#include "Chaos/Levelset.h"
+#include "Chaos/Pair.h"
+#include "Chaos/ParticleHandle.h"
+#include "Chaos/PBDCollisionConstraints.h"
+#include "Chaos/Sphere.h"
+#include "Chaos/Transform.h"
+
+
+namespace Chaos
+{
+	namespace Collisions
+	{
+		//
+		// Box - Box
+		//
+
+		template <typename T, int d>
+		void UpdateBoxBoxConstraint(const TBox<T, d>& Box1, const TRigidTransform<T, d>& Box1Transform, const TBox<T, d>& Box2, const TRigidTransform<T, d>& Box2Transform, T Thickness, TRigidBodyPointContactConstraint<T, d>& Constraint)
+		{
+			TPointContactManifold<T, d> & Contact = Constraint.Manifold;
+
+			TBox<T, d> Box2SpaceBox1 = Box1.TransformedBox(Box1Transform * Box2Transform.Inverse());
+			TBox<T, d> Box1SpaceBox2 = Box2.TransformedBox(Box2Transform * Box1Transform.Inverse());
+			Box2SpaceBox1.Thicken(Thickness);
+			Box1SpaceBox2.Thicken(Thickness);
+			if (Box1SpaceBox2.Intersects(Box1) && Box2SpaceBox1.Intersects(Box2))
+			{
+				const TVector<T, d> Box1Center = (Box1Transform * Box2Transform.Inverse()).TransformPosition(Box1.Center());
+				bool bDeepOverlap = false;
+				if (Box2.SignedDistance(Box1Center) < 0)
+				{
+					//If Box1 is overlapping Box2 by this much the signed distance approach will fail (box1 gets sucked into box2). In this case just use two spheres
+					TSphere<T, d> Sphere1(Box1Transform.TransformPosition(Box1.Center()), Box1.Extents().Min() / 2);
+					TSphere<T, d> Sphere2(Box2Transform.TransformPosition(Box2.Center()), Box2.Extents().Min() / 2);
+					const TVector<T, d> Direction = Sphere1.GetCenter() - Sphere2.GetCenter();
+					T Size = Direction.Size();
+					if (Size < (Sphere1.GetRadius() + Sphere2.GetRadius()))
+					{
+						const T NewPhi = Size - (Sphere1.GetRadius() + Sphere2.GetRadius());;
+						if (NewPhi < Contact.Phi)
+						{
+							bDeepOverlap = true;
+							Contact.Normal = Size > SMALL_NUMBER ? Direction / Size : TVector<T, d>(0, 0, 1);
+							Contact.Phi = NewPhi;
+							Contact.Location = Sphere1.GetCenter() - Sphere1.GetRadius() * Contact.Normal;
+						}
+					}
+				}
+				if (!bDeepOverlap || Contact.Phi >= 0)
+				{
+					//if we didn't have deep penetration use signed distance per particle. If we did have deep penetration but the spheres did not overlap use signed distance per particle
+
+					//UpdateLevelsetConstraintGJK(InParticles, Thickness, Constraint);
+					//check(Contact.Phi < MThickness);
+					// For now revert to doing all points vs lsv check until we can figure out a good way to get the deepest point without needing this
+					{
+						const TArray<TVector<T, d>> SampleParticles = Box1.ComputeSamplePoints();
+						const TRigidTransform<T, d> Box1ToBox2Transform = Box1Transform.GetRelativeTransform(Box2Transform);
+						int32 NumParticles = SampleParticles.Num();
+						for (int32 i = 0; i < NumParticles; ++i)
+						{
+							SampleObjectHelper(Box2, Box2Transform, Box1ToBox2Transform, SampleParticles[i], Thickness, Constraint);
+						}
+					}
+				}
+			}
+		}
+
+
+		template<typename T, int d>
+		void ConstructBoxBoxConstraints(TGeometryParticleHandle<T, d>* Particle0, TGeometryParticleHandle<T, d>* Particle1, const FImplicitObject* Implicit0, const FImplicitObject* Implicit1, const T Thickness, FCollisionConstraintsArray& NewConstraints)
+		{
+			const TBox<T, d> * Object0 = Implicit0->template GetObject<const TBox<T, d> >();
+			const TBox<T, d> * Object1 = Implicit1->template GetObject<const TBox<T, d> >();
+			if (ensure(Object0 && Object1))
+			{
+				FRigidBodyPointContactConstraint* Constraint = new FRigidBodyPointContactConstraint;
+				Constraint->Particle[0] = Particle0;
+				Constraint->Particle[1] = Particle1;
+				Constraint->AddManifold(Implicit0, Implicit1);
+
+				UpdateBoxBoxConstraint(*Object0, GetTransform(Particle0), *Object1, GetTransform(Particle1), Thickness, *Constraint);
+
+				if (Constraint->GetPhi() < Thickness)
+				{
+					NewConstraints.Add(Constraint);
+				}
+				else
+				{
+					delete Constraint;
+				}
+			}
+		}
+
+		//
+		// Box-Plane
+		//
+
+		template <typename T, int d>
+		bool UpdateBoxPlaneConstraint(const TBox<T, d>& Box, const TRigidTransform<T, d>& BoxTransform, const TPlane<T, d>& Plane, const TRigidTransform<T, d>& PlaneTransform, T Thickness, TRigidBodyPointContactConstraint<T, d>& Constraint)
+		{
+			TPointContactManifold<T, d> & Contact = Constraint.Manifold;
+
+#if USING_CODE_ANALYSIS
+			MSVC_PRAGMA(warning(push))
+				MSVC_PRAGMA(warning(disable : ALL_CODE_ANALYSIS_WARNINGS))
+#endif	// USING_CODE_ANALYSIS
+
+				bool bApplied = false;
+			const TRigidTransform<T, d> BoxToPlaneTransform(BoxTransform.GetRelativeTransform(PlaneTransform));
+			const TVector<T, d> Extents = Box.Extents();
+			constexpr int32 NumCorners = 2 + 2 * d;
+			constexpr T Epsilon = KINDA_SMALL_NUMBER;
+
+			TVector<T, d> Corners[NumCorners];
+			int32 CornerIdx = 0;
+			Corners[CornerIdx++] = BoxToPlaneTransform.TransformPosition(Box.Max());
+			Corners[CornerIdx++] = BoxToPlaneTransform.TransformPosition(Box.Min());
+			for (int32 j = 0; j < d; ++j)
+			{
+				Corners[CornerIdx++] = BoxToPlaneTransform.TransformPosition(Box.Min() + TVector<T, d>::AxisVector(j) * Extents);
+				Corners[CornerIdx++] = BoxToPlaneTransform.TransformPosition(Box.Max() - TVector<T, d>::AxisVector(j) * Extents);
+			}
+
+#if USING_CODE_ANALYSIS
+			MSVC_PRAGMA(warning(pop))
+#endif	// USING_CODE_ANALYSIS
+
+				TVector<T, d> PotentialConstraints[NumCorners];
+			int32 NumConstraints = 0;
+			for (int32 i = 0; i < NumCorners; ++i)
+			{
+				TVector<T, d> Normal;
+				const T NewPhi = Plane.PhiWithNormal(Corners[i], Normal);
+				if (NewPhi < Contact.Phi + Epsilon)
+				{
+					if (NewPhi <= Contact.Phi - Epsilon)
+					{
+						NumConstraints = 0;
+					}
+					Contact.Phi = NewPhi;
+					Contact.Normal = PlaneTransform.TransformVector(Normal);
+					Contact.Location = PlaneTransform.TransformPosition(Corners[i]);
+					PotentialConstraints[NumConstraints++] = Contact.Location;
+					bApplied = true;
+				}
+			}
+			if (NumConstraints > 1)
+			{
+				TVector<T, d> AverageLocation(0);
+				for (int32 ConstraintIdx = 0; ConstraintIdx < NumConstraints; ++ConstraintIdx)
+				{
+					AverageLocation += PotentialConstraints[ConstraintIdx];
+				}
+				Contact.Location = AverageLocation / NumConstraints;
+			}
+
+			return bApplied;
+		}
+
+
+		template<typename T, int d>
+		void ConstructBoxPlaneConstraints(TGeometryParticleHandle<T, d>* Particle0, TGeometryParticleHandle<T, d>* Particle1, const FImplicitObject* Implicit0, const FImplicitObject* Implicit1, const T Thickness, FCollisionConstraintsArray& NewConstraints)
+		{
+
+			const TBox<T, d> * Object0 = Implicit0->template GetObject<const TBox<T, d> >();
+			const TPlane<T, d> * Object1 = Implicit1->template GetObject<const TPlane<T, d> >();
+			if (ensure(Object0 && Object1))
+			{
+				FRigidBodyPointContactConstraint* Constraint = new FRigidBodyPointContactConstraint;
+				Constraint->Particle[0] = Particle0;
+				Constraint->Particle[1] = Particle1;
+				Constraint->AddManifold(Implicit0, Implicit1);
+
+				UpdateBoxPlaneConstraint(*Object0, GetTransform(Particle0), *Object1, GetTransform(Particle1), Thickness, *Constraint);
+
+				if (Constraint->GetPhi() < Thickness)
+				{
+					NewConstraints.Add(Constraint);
+				}
+				else
+				{
+					delete Constraint;
+				}
+			}
+		}
+
+
+		//
+		// Sphere - Sphere
+		//
+		template <typename T, int d>
+		void UpdateSphereSphereConstraint(const TSphere<T, d>& Sphere1, const TRigidTransform<T, d>& Sphere1Transform, const TSphere<T, d>& Sphere2, const TRigidTransform<T, d>& Sphere2Transform, T Thickness, TRigidBodyPointContactConstraint<T, d>& Constraint)
+		{
+			TPointContactManifold<T, d> & Contact = Constraint.Manifold;
+
+			const TVector<T, d> Center1 = Sphere1Transform.TransformPosition(Sphere1.GetCenter());
+			const TVector<T, d> Center2 = Sphere2Transform.TransformPosition(Sphere2.GetCenter());
+			const TVector<T, d> Direction = Center1 - Center2;
+			const T Size = Direction.Size();
+			const T NewPhi = Size - (Sphere1.GetRadius() + Sphere2.GetRadius());
+			if (NewPhi < Contact.Phi)
+			{
+				Contact.Normal = Size > SMALL_NUMBER ? Direction / Size : TVector<T, d>(0, 0, 1);
+				Contact.Phi = NewPhi;
+				Contact.Location = Center1 - Sphere1.GetRadius() * Contact.Normal;
+			}
+		}
+
+		template<typename T, int d>
+		void ConstructSphereSphereConstraints(TGeometryParticleHandle<T, d>* Particle0, TGeometryParticleHandle<T, d>* Particle1, const FImplicitObject* Implicit0, const FImplicitObject* Implicit1, const T Thickness, FCollisionConstraintsArray& NewConstraints)
+		{
+
+			const TSphere<T, d> * Object0 = Implicit0->template GetObject<const TSphere<T, d> >();
+			const TSphere<T, d> * Object1 = Implicit1->template GetObject<const TSphere<T, d> >();
+			if (ensure(Object0 && Object1))
+			{
+				FRigidBodyPointContactConstraint* Constraint = new FRigidBodyPointContactConstraint;
+				Constraint->Particle[0] = Particle0;
+				Constraint->Particle[1] = Particle1;
+				Constraint->AddManifold(Implicit0, Implicit1);
+
+				UpdateSphereSphereConstraint(*Object0, GetTransform(Particle0), *Object1, GetTransform(Particle1), Thickness, *Constraint);
+
+				if (Constraint->GetPhi() < Thickness)
+				{
+					NewConstraints.Add(Constraint);
+				}
+				else
+				{
+					delete Constraint;
+				}
+			}
+		}
+
+		//
+		//  Sphere-Plane
+		//
+
+		template <typename T, int d>
+		void UpdateSpherePlaneConstraint(const TSphere<T, d>& Sphere, const TRigidTransform<T, d>& SphereTransform, const TPlane<T, d>& Plane, const TRigidTransform<T, d>& PlaneTransform, T Thickness, TRigidBodyPointContactConstraint<T, d>& Constraint)
+		{
+			TPointContactManifold<T, d> & Contact = Constraint.Manifold;
+
+			const TRigidTransform<T, d> SphereToPlaneTransform(PlaneTransform.Inverse() * SphereTransform);
+			const TVector<T, d> SphereCenter = SphereToPlaneTransform.TransformPosition(Sphere.GetCenter());
+
+			TVector<T, d> NewNormal;
+			T NewPhi = Plane.PhiWithNormal(SphereCenter, NewNormal);
+			NewPhi -= Sphere.GetRadius();
+
+			if (NewPhi < Contact.Phi)
+			{
+				Contact.Phi = NewPhi;
+				Contact.Normal = PlaneTransform.TransformVectorNoScale(NewNormal);
+				Contact.Location = SphereCenter - Contact.Normal * Sphere.GetRadius();
+			}
+		}
+
+		template<typename T, int d>
+		void ConstructSpherePlaneConstraints(TGeometryParticleHandle<T, d>* Particle0, TGeometryParticleHandle<T, d>* Particle1, const FImplicitObject* Implicit0, const FImplicitObject* Implicit1, const T Thickness, FCollisionConstraintsArray& NewConstraints)
+		{
+
+			const TSphere<T, d> * Object0 = Implicit0->template GetObject<const TSphere<T, d> >();
+			const TPlane<T, d> * Object1 = Implicit1->template GetObject<const TPlane<T, d> >();
+			if (ensure(Object0 && Object1))
+			{
+
+				FRigidBodyPointContactConstraint* Constraint = new FRigidBodyPointContactConstraint;
+				Constraint->Particle[0] = Particle0;
+				Constraint->Particle[1] = Particle1;
+				Constraint->AddManifold(Implicit0, Implicit1);
+
+				UpdateSpherePlaneConstraint(*Object0, GetTransform(Particle0), *Object1, GetTransform(Particle1), Thickness, *Constraint);
+
+				if (Constraint->GetPhi() < Thickness)
+				{
+					NewConstraints.Add(Constraint);
+				}
+				else
+				{
+					delete Constraint;
+				}
+			}
+		}
+
+		//
+		// Sphere - Box
+		//
+
+		template <typename T, int d>
+		void UpdateSphereBoxConstraint(const TSphere<T, d>& Sphere, const TRigidTransform<T, d>& SphereTransform, const TBox<T, d>& Box, const TRigidTransform<T, d>& BoxTransform, T Thickness, TRigidBodyPointContactConstraint<T, d>& Constraint)
+		{
+			TPointContactManifold<T, d> & Contact = Constraint.Manifold;
+
+			const TRigidTransform<T, d> SphereToBoxTransform(SphereTransform * BoxTransform.Inverse());
+			const TVector<T, d> SphereCenterInBox = SphereToBoxTransform.TransformPosition(Sphere.GetCenter());
+
+			TVector<T, d> NewNormal;
+			T NewPhi = Box.PhiWithNormal(SphereCenterInBox, NewNormal);
+			NewPhi -= Sphere.GetRadius();
+
+			if (NewPhi < Contact.Phi)
+			{
+				Contact.Phi = NewPhi;
+				Contact.Normal = BoxTransform.TransformVectorNoScale(NewNormal);
+				Contact.Location = SphereTransform.TransformPosition(Sphere.GetCenter()) - Contact.Normal * Sphere.GetRadius();
+			}
+		}
+
+		template<typename T, int d>
+		void ConstructSphereBoxConstraints(TGeometryParticleHandle<T, d>* Particle0, TGeometryParticleHandle<T, d>* Particle1, const FImplicitObject* Implicit0, const FImplicitObject* Implicit1, const T Thickness, FCollisionConstraintsArray& NewConstraints)
+		{
+
+			const TSphere<T, d> * Object0 = Implicit0->template GetObject<const TSphere<T, d> >();
+			const TBox<T, d> * Object1 = Implicit1->template GetObject<const TBox<T, d> >();
+			if (ensure(Object0 && Object1))
+			{
+				FRigidBodyPointContactConstraint* Constraint = new FRigidBodyPointContactConstraint;
+				Constraint->Particle[0] = Particle0;
+				Constraint->Particle[1] = Particle1;
+				Constraint->AddManifold(Implicit0, Implicit1);
+
+				UpdateSphereBoxConstraint(*Object0, GetTransform(Particle0), *Object1, GetTransform(Particle1), Thickness, *Constraint);
+
+				if (Constraint->GetPhi() < Thickness)
+				{
+					NewConstraints.Add(Constraint);
+				}
+				else
+				{
+					delete Constraint;
+				}
+			}
+		}
+
+
+		//
+		// Capsule-Capsule
+		//
+
+		template <typename T, int d>
+		void UpdateCapsuleCapsuleConstraint(const TCapsule<T>& A, const TRigidTransform<T, d>& ATransform, const TCapsule<T>& B, const TRigidTransform<T, d>& BTransform, T Thickness, TRigidBodyPointContactConstraint<T, d>& Constraint)
+		{
+			TPointContactManifold<T, d> & Contact = Constraint.Manifold;
+
+			FVector A1 = ATransform.TransformPosition(A.GetX1());
+			FVector A2 = ATransform.TransformPosition(A.GetX2());
+			FVector B1 = BTransform.TransformPosition(B.GetX1());
+			FVector B2 = BTransform.TransformPosition(B.GetX2());
+			FVector P1, P2;
+			FMath::SegmentDistToSegmentSafe(A1, A2, B1, B2, P1, P2);
+
+			TVector<T, d> Delta = P2 - P1;
+			T DeltaLen = Delta.Size();
+			if (DeltaLen > KINDA_SMALL_NUMBER)
+			{
+				T NewPhi = DeltaLen - (A.GetRadius() + B.GetRadius());
+				if (NewPhi < Contact.Phi)
+				{
+					TVector<T, d> Dir = Delta / DeltaLen;
+					Contact.Phi = NewPhi;
+					Contact.Normal = -Dir;
+					Contact.Location = P1 + Dir * A.GetRadius();
+				}
+			}
+		}
+
+		template<typename T, int d>
+		void ConstructCapsuleCapsuleConstraints(TGeometryParticleHandle<T, d>* Particle0, TGeometryParticleHandle<T, d>* Particle1, const FImplicitObject* Implicit0, const FImplicitObject* Implicit1, const T Thickness, FCollisionConstraintsArray& NewConstraints)
+		{
+
+			const TCapsule<T> * Object0 = Implicit0->template GetObject<const TCapsule<T> >();
+			const TCapsule<T> * Object1 = Implicit1->template GetObject<const TCapsule<T> >();
+			if (ensure(Object0 && Object1))
+			{
+				FRigidBodyPointContactConstraint* Constraint = new FRigidBodyPointContactConstraint;
+				Constraint->Particle[0] = Particle0;
+				Constraint->Particle[1] = Particle1;
+				Constraint->AddManifold(Implicit0, Implicit1);
+
+				UpdateCapsuleCapsuleConstraint(*Object0, GetTransform(Particle0), *Object1, GetTransform(Particle1), Thickness, *Constraint);
+
+				if (Constraint->GetPhi() < Thickness)
+				{
+					NewConstraints.Add(Constraint);
+				}
+				else
+				{
+					delete Constraint;
+				}
+			}
+		}
+
+		//
+		// Capsule - Box
+		//
+
+		template <typename T, int d>
+		void UpdateCapsuleBoxConstraint(const TCapsule<T>& A, const TRigidTransform<T, d>& ATransform, const TBox<T, d>& B, const TRigidTransform<T, d>& BTransform, T Thickness, TRigidBodyPointContactConstraint<T, d>& Constraint)
+		{
+			// @todo(ccaulfield): Add custom capsule-box collision
+			TPointContactManifold<T, d> & Contact = Constraint.Manifold;
+
+			TRigidTransform<T, d> BToATransform = BTransform.GetRelativeTransform(ATransform);
+
+			// Use GJK to track closest points (not strictly necessary yet)
+			TVector<T, d> NearPointALocal, NearPointBLocal;
+			T NearPointDistance;
+			if (GJKDistance<T>(A, B, BToATransform, NearPointDistance, NearPointALocal, NearPointBLocal))
+			{
+				TVector<T, d> NearPointAWorld = ATransform.TransformPosition(NearPointALocal);
+				TVector<T, d> NearPointBWorld = BTransform.TransformPosition(NearPointBLocal);
+				TVector<T, d> NearPointBtoAWorld = NearPointAWorld - NearPointBWorld;
+				Contact.Phi = NearPointDistance;
+				Contact.Normal = NearPointBtoAWorld.GetSafeNormal();
+				Contact.Location = NearPointAWorld;
+			}
+			else
+			{
+				// Use box particle samples against the implicit capsule
+				const TArray<TVector<T, d>> BParticles = B.ComputeSamplePoints();
+				const int32 NumParticles = BParticles.Num();
+				for (int32 ParticleIndex = 0; ParticleIndex < NumParticles; ++ParticleIndex)
+				{
+					if (SampleObjectHelper(A, ATransform, BToATransform, BParticles[ParticleIndex], Thickness, Constraint))
+					{
+						// SampleObjectHelper expects A to be the box, so reverse the results
+						Contact.Normal = -Contact.Normal;
+					}
+				}
+			}
+
+		}
+
+		template<typename T, int d>
+		void ConstructCapsuleBoxConstraints(TGeometryParticleHandle<T, d>* Particle0, TGeometryParticleHandle<T, d>* Particle1, const FImplicitObject* Implicit0, const FImplicitObject* Implicit1, const T Thickness, FCollisionConstraintsArray& NewConstraints)
+		{
+
+			const TCapsule<T> * Object0 = Implicit0->template GetObject<const TCapsule<T> >();
+			const TBox<T, d> * Object1 = Implicit1->template GetObject<const TBox<T, d> >();
+			if (ensure(Object0 && Object1))
+			{
+				FRigidBodyPointContactConstraint* Constraint = new FRigidBodyPointContactConstraint;
+				Constraint->Particle[0] = Particle0;
+				Constraint->Particle[1] = Particle1;
+				Constraint->AddManifold(Implicit0, Implicit1);
+
+				UpdateCapsuleBoxConstraint(*Object0, GetTransform(Particle0), *Object1, GetTransform(Particle1), Thickness, *Constraint);
+
+				if (Constraint->GetPhi() < Thickness)
+				{
+					NewConstraints.Add(Constraint);
+				}
+				else
+				{
+					delete Constraint;
+				}
+			}
+		}
+
+		//
+		// Convex - Convex
+		//
+
+		template<class T, int d>
+		void UpdateConvexConvexConstraint(const FImplicitObject& A, const TRigidTransform<T, d>& ATM, const FImplicitObject& B, const TRigidTransform<T, d>& BTM, T Thickness, TRigidBodyPointContactConstraint<T, d>& Constraint)
+		{
+			TRigidTransform<T, d> BToATM = BTM.GetRelativeTransform(ATM);
+
+			if (ensure(GetInnerType(A.GetType()) == ImplicitObjectType::Convex && GetInnerType(B.GetType()) == ImplicitObjectType::Convex))
+			{
+				const TConvex<T, d>* AObject = nullptr;
+				const TConvex<T, d>* BObject = nullptr;
+
+				if (IsScaled(A.GetType()) && IsScaled(B.GetType()))
+				{
+					AObject = static_cast<const TConvex<T, d>*>(static_cast<const TImplicitObjectScaled<TConvex<T, d>>&>(A).Object().Get());
+					BObject = static_cast<const TConvex<T, d>*>(static_cast<const TImplicitObjectScaled<TConvex<T, d>>&>(B).Object().Get());
+				}
+				else if (A.GetType() == ImplicitObjectType::Transformed && B.GetType() == ImplicitObjectType::Transformed)
+				{
+					AObject = static_cast<const TConvex<T, d>*>(static_cast<const TImplicitObjectTransformed<T, d>&>(A).Object().Get());
+					BObject = static_cast<const TConvex<T, d>*>(static_cast<const TImplicitObjectTransformed<T, d>&>(B).Object().Get());
+				}
+
+				if (AObject && BObject)
+				{
+					const TParticles<T, d>& SurfaceParticles = AObject->GetSurfaceParticles();
+
+					Constraint.SetDisabled(true);
+
+					if (GJKIntersection(*AObject, *BObject, BToATM, Thickness))
+					{
+						TRigidTransform<T, d> AToBTM = ATM.GetRelativeTransform(BTM);
+
+						T Phi;
+						TVector<T, d> Normal;
+						for (int32 Idx = 0; Idx < (int32)SurfaceParticles.Size(); Idx++)
+						{
+							Phi = BObject->PhiWithNormal(AToBTM.TransformPosition(SurfaceParticles.X(Idx)), Normal);
+							if (Phi < Constraint.GetPhi())
+							{
+								Constraint.SetPhi(Phi);
+								Constraint.SetLocation(ATM.TransformPosition(SurfaceParticles.X(Idx)));
+								Constraint.SetNormal(BTM.TransformVector(Normal));
+							}
+						}
+						return;
+					}
+					return;
+				}
+			}
+
+			ensureMsgf(false, TEXT("Unsupported convex to convex constraint."));
+		}
+
+		template<class T, int d>
+		void ConstructConvexConvexConstraints(TGeometryParticleHandle<T, d>* Particle0, TGeometryParticleHandle<T, d>* Particle1, const FImplicitObject* Implicit0, const FImplicitObject* Implicit1, const T Thickness, FCollisionConstraintsArray& NewConstraints)
+		{
+			if (ensure(GetInnerType(Implicit0->GetType()) == ImplicitObjectType::Convex && GetInnerType(Implicit1->GetType()) == ImplicitObjectType::Convex))
+			{
+				FRigidBodyPointContactConstraint* Constraint = new FRigidBodyPointContactConstraint;
+
+				Constraint->Particle[0] = Particle0;
+				Constraint->Particle[1] = Particle1;
+				Constraint->AddManifold(Implicit0, Implicit1);
+
+				UpdateConvexConvexConstraint(*Implicit0, GetTransform(Particle0), *Implicit1, GetTransform(Particle1), Thickness, *Constraint);
+
+				if (Constraint->GetPhi() < Thickness)
+				{
+					NewConstraints.Add(Constraint);
+				}
+				else
+				{
+					delete Constraint;
+				}
+			}
+		}
+
+
+		//
+		// Levelset-Levelset
+		//
+
+		DECLARE_CYCLE_STAT(TEXT("TPBDCollisionConstraints::UpdateLevelsetLevelsetConstraint"), STAT_UpdateLevelsetLevelsetConstraint, STATGROUP_ChaosWide);
+		template<ECollisionUpdateType UpdateType, typename T, int d>
+		void UpdateLevelsetLevelsetConstraint(const T Thickness, TRigidBodyPointContactConstraint<float, 3>& Constraint)
+		{
+			SCOPE_CYCLE_COUNTER(STAT_UpdateLevelsetLevelsetConstraint);
+
+			TGenericParticleHandle<T, d> Particle0 = Constraint.Particle[0];
+			TRigidTransform<T, d> ParticlesTM = TRigidTransform<T, d>(Particle0->P(), Particle0->Q());
+			if (!(ensure(!FMath::IsNaN(ParticlesTM.GetTranslation().X)) && ensure(!FMath::IsNaN(ParticlesTM.GetTranslation().Y)) && ensure(!FMath::IsNaN(ParticlesTM.GetTranslation().Z))))
+			{
+				return;
+			}
+
+			TGenericParticleHandle<T, d> Particle1 = Constraint.Particle[1];
+			TRigidTransform<T, d> LevelsetTM = TRigidTransform<T, d>(Particle1->P(), Particle1->Q());
+			if (!(ensure(!FMath::IsNaN(LevelsetTM.GetTranslation().X)) && ensure(!FMath::IsNaN(LevelsetTM.GetTranslation().Y)) && ensure(!FMath::IsNaN(LevelsetTM.GetTranslation().Z))))
+			{
+				return;
+			}
+
+			const TBVHParticles<T, d>* SampleParticles = nullptr;
+			SampleParticles = Particle0->CollisionParticles().Get();
+
+			if (SampleParticles)
+			{
+				SampleObject<UpdateType>(*Particle1->Geometry(), LevelsetTM, *SampleParticles, ParticlesTM, Thickness, Constraint);
+			}
+		}
+
+
+		template<typename T, int d>
+		void ConstructLevelsetLevelsetConstraints(TGeometryParticleHandle<T, d>* Particle0, TGeometryParticleHandle<T, d>* Particle1, const FImplicitObject* Implicit0, const FImplicitObject* Implicit1, const T Thickness, FCollisionConstraintsArray& NewConstraints)
+		{
+			FRigidBodyPointContactConstraint* Constraint = new FRigidBodyPointContactConstraint;
+			if (!Particle1->Geometry() || (Particle0->AsDynamic() && !Particle0->AsDynamic()->CollisionParticlesSize() && Particle0->Geometry() && !Particle0->Geometry()->IsUnderlyingUnion()))
+			{
+				Constraint->Particle[0] = Particle1;
+				Constraint->Particle[1] = Particle0;
+				Constraint->AddManifold(Implicit1, Implicit0);
+			}
+			else
+			{
+				Constraint->Particle[0] = Particle0;
+				Constraint->Particle[1] = Particle1;
+				Constraint->AddManifold(Implicit0, Implicit1);
+			}
+
+			UpdateLevelsetLevelsetConstraint<ECollisionUpdateType::Any>(Thickness, *Constraint);
+
+			if (Constraint->GetPhi() < Thickness)
+			{
+				NewConstraints.Add(Constraint);
+			}
+			else
+			{
+				delete Constraint;
+			}
+		}
+
+		//
+		// Levelset-Union
+		//
+
+		DECLARE_CYCLE_STAT(TEXT("TPBDCollisionConstraints::UpdateLevelsetUnionConstraint"), STAT_UpdateLevelsetUnionConstraint, STATGROUP_ChaosWide);
+		template<ECollisionUpdateType UpdateType, typename T, int d>
+		void UpdateLevelsetUnionConstraint(const T Thickness, TRigidBodyPointContactConstraint<T, d>& Constraint)
+		{
+			SCOPE_CYCLE_COUNTER(STAT_UpdateLevelsetUnionConstraint);
+
+			TGenericParticleHandle<T, d> Particle0 = Constraint.Particle[0];
+			TGenericParticleHandle<T, d> Particle1 = Constraint.Particle[1];
+
+			TRigidTransform<T, d> ParticlesTM = TRigidTransform<T, d>(Particle0->P(), Particle0->Q());
+			TRigidTransform<T, d> LevelsetTM = TRigidTransform<T, d>(Particle1->P(), Particle1->Q());
+
+			const FImplicitObject* ParticleObj = Particle0->Geometry().Get();
+			const FImplicitObject* LevelsetObj = Particle1->Geometry().Get();
+
+			if (!(ensure(!FMath::IsNaN(ParticlesTM.GetTranslation().X)) && ensure(!FMath::IsNaN(ParticlesTM.GetTranslation().Y)) && ensure(!FMath::IsNaN(ParticlesTM.GetTranslation().Z))))
+			{
+				return;
+			}
+
+			if (!(ensure(!FMath::IsNaN(LevelsetTM.GetTranslation().X)) && ensure(!FMath::IsNaN(LevelsetTM.GetTranslation().Y)) && ensure(!FMath::IsNaN(LevelsetTM.GetTranslation().Z))))
+			{
+				return;
+			}
+
+#if CHAOS_PARTICLEHANDLE_TODO
+			TArray<Pair<const FImplicitObject*, TRigidTransform<T, d>>> ParticleShapes = FindRelevantShapes(LevelsetObj, LevelsetTM, *ParticleObj, ParticlesTM, Thickness);
+			check(ParticleObj->IsUnderlyingUnion());
+			const TImplicitObjectUnion<T, d>* UnionObj = static_cast<const TImplicitObjectUnion<T, d>*>(ParticleObj);
+			for (const Pair<const FImplicitObject*, TRigidTransform<T, d>>& ParticleObjPair : ParticleShapes)
+			{
+				const FImplicitObject* Object = ParticleObjPair.First;
+
+				if (const int32* OriginalIdx = UnionObj->MCollisionParticleLookupHack.Find(Object))
+				{
+					const TBVHParticles<T, d>& SampleParticles = *InParticles.CollisionParticles(*OriginalIdx).Get();
+					const TRigidTransform<T, d> ObjectTM = ParticleObjPair.Second * ParticlesTM;
+
+					SampleObject2<UpdateType>(*LevelsetObj, LevelsetTM, SampleParticles, ObjectTM, Thickness, Constraint);
+					if (UpdateType == ECollisionUpdateType::Any && Constraint.Phi < Thickness)
+					{
+						return;
+					}
+				}
+			}
+#endif
+		}
+
+
+		DECLARE_CYCLE_STAT(TEXT("TPBDCollisionConstraints::UpdateUnionLevelsetConstraint"), STAT_UpdateUnionLevelsetConstraint, STATGROUP_ChaosWide);
+		template<ECollisionUpdateType UpdateType, typename T, int d>
+		void UpdateUnionLevelsetConstraint(const T Thickness, TRigidBodyPointContactConstraint<T, d>& Constraint)
+		{
+			SCOPE_CYCLE_COUNTER(STAT_UpdateUnionLevelsetConstraint);
+
+			TGenericParticleHandle<T, d> Particle0 = Constraint.Particle[0];
+			TGenericParticleHandle<T, d> Particle1 = Constraint.Particle[1];
+
+			TRigidTransform<T, d> ParticlesTM = TRigidTransform<T, d>(Particle0->P(), Particle0->Q());
+			TRigidTransform<T, d> LevelsetTM = TRigidTransform<T, d>(Particle1->P(), Particle1->Q());
+
+			if (!(ensure(!FMath::IsNaN(ParticlesTM.GetTranslation().X)) && ensure(!FMath::IsNaN(ParticlesTM.GetTranslation().Y)) && ensure(!FMath::IsNaN(ParticlesTM.GetTranslation().Z))))
+			{
+				return;
+			}
+
+			if (!(ensure(!FMath::IsNaN(LevelsetTM.GetTranslation().X)) && ensure(!FMath::IsNaN(LevelsetTM.GetTranslation().Y)) && ensure(!FMath::IsNaN(LevelsetTM.GetTranslation().Z))))
+			{
+				return;
+			}
+
+			const FImplicitObject* ParticleObj = Particle0->Geometry().Get();
+			const FImplicitObject* LevelsetObj = Particle1->Geometry().Get();
+			TArray<Pair<const FImplicitObject*, TRigidTransform<T, d>>> LevelsetShapes = FindRelevantShapes(ParticleObj, ParticlesTM, *LevelsetObj, LevelsetTM, Thickness);
+
+			if (LevelsetShapes.Num() && Particle0->CollisionParticles().Get())
+			{
+				const TBVHParticles<T, d>& SampleParticles = *Particle0->CollisionParticles().Get();
+				if (SampleParticles.Size())
+				{
+					for (const Pair<const FImplicitObject*, TRigidTransform<T, d>>& LevelsetObjPair : LevelsetShapes)
+					{
+						const FImplicitObject* Object = LevelsetObjPair.First;
+						const TRigidTransform<T, d> ObjectTM = LevelsetObjPair.Second * LevelsetTM;
+						// todo(brice) : SampleObject<UpdateType>(*Object, ObjectTM, SampleParticles, ParticlesTM, Thickness, Constraint);
+						if (UpdateType == ECollisionUpdateType::Any && Constraint.GetPhi() < Thickness)
+						{
+							return;
+						}
+					}
+				}
+#if CHAOS_PARTICLEHANDLE_TODO
+				else if (ParticleObj && ParticleObj->IsUnderlyingUnion())
+				{
+					const TImplicitObjectUnion<T, d>* UnionObj = static_cast<const TImplicitObjectUnion<T, d>*>(ParticleObj);
+					//need to traverse shapes to get their collision particles
+					for (const Pair<const FImplicitObject*, TRigidTransform<T, d>>& LevelsetObjPair : LevelsetShapes)
+					{
+						const FImplicitObject* LevelsetInnerObject = LevelsetObjPair.First;
+						const TRigidTransform<T, d> LevelsetInnerObjectTM = LevelsetObjPair.Second * LevelsetTM;
+
+						TArray<Pair<const FImplicitObject*, TRigidTransform<T, d>>> ParticleShapes = FindRelevantShapes(LevelsetInnerObject, LevelsetInnerObjectTM, *ParticleObj, ParticlesTM, Thickness);
+						for (const Pair<const FImplicitObject*, TRigidTransform<T, d>>& ParticleObjPair : ParticleShapes)
+						{
+							const FImplicitObject* ParticleInnerObject = ParticleObjPair.First;
+							const TRigidTransform<T, d> ParticleInnerObjectTM = ParticleObjPair.Second * ParticlesTM;
+
+							if (const int32* OriginalIdx = UnionObj->MCollisionParticleLookupHack.Find(ParticleInnerObject))
+							{
+								const TBVHParticles<T, d>& InnerSampleParticles = *InParticles.CollisionParticles(*OriginalIdx).Get();
+								SampleObject2<UpdateType>(*LevelsetInnerObject, LevelsetInnerObjectTM, InnerSampleParticles, ParticleInnerObjectTM, Thickness, Constraint);
+								if (UpdateType == ECollisionUpdateType::Any && Constraint.Phi < Thickness)
+								{
+									return;
+								}
+							}
+						}
+
+					}
+				}
+#endif
+			}
+		}
+
+		//
+		//  Union-Union
+		//
+		DECLARE_CYCLE_STAT(TEXT("TPBDCollisionConstraints::UpdateUnionUnionConstraint"), STAT_UpdateUnionUnionConstraint, STATGROUP_ChaosWide);
+
+		template<ECollisionUpdateType UpdateType, typename T, int d>
+		void UpdateUnionUnionConstraint(const T Thickness, TRigidBodyPointContactConstraint<T, d>& Constraint)
+		{
+			SCOPE_CYCLE_COUNTER(STAT_UpdateUnionUnionConstraint);
+
+			TGenericParticleHandle<T, d> Particle0 = Constraint.Particle[0];
+			TGenericParticleHandle<T, d> Particle1 = Constraint.Particle[1];
+
+			TRigidTransform<T, d> ParticlesTM = TRigidTransform<T, d>(Particle0->P(), Particle0->Q());
+			TRigidTransform<T, d> LevelsetTM = TRigidTransform<T, d>(Particle1->P(), Particle1->Q());
+
+			const FImplicitObject* ParticleObj = Particle0->Geometry().Get();
+			const FImplicitObject* LevelsetObj = Particle1->Geometry().Get();
+			const TArray<Pair<const FImplicitObject*, TRigidTransform<T, d>>> LevelsetShapes = FindRelevantShapes(ParticleObj, ParticlesTM, *LevelsetObj, LevelsetTM, Thickness);
+
+			for (const Pair<const FImplicitObject*, TRigidTransform<T, d>>& LevelsetObjPair : LevelsetShapes)
+			{
+				const FImplicitObject& LevelsetInnerObj = *LevelsetObjPair.First;
+				const TRigidTransform<T, d>& LevelsetInnerObjTM = LevelsetObjPair.Second * LevelsetTM;
+
+				//now find all particle inner objects
+				const TArray<Pair<const FImplicitObject*, TRigidTransform<T, d>>> ParticleShapes = FindRelevantShapes(&LevelsetInnerObj, LevelsetInnerObjTM, *ParticleObj, ParticlesTM, Thickness);
+
+				//for each inner obj pair, update constraint
+				for (const Pair<const FImplicitObject*, TRigidTransform<T, d>>& ParticlePair : ParticleShapes)
+				{
+					const FImplicitObject& ParticleInnerObj = *ParticlePair.First;
+					const TRigidTransform<T, d> ParticleInnerObjTM = ParticlePair.Second * ParticlesTM;
+					UpdateConstraintImp<UpdateType>(ParticleInnerObj, ParticleInnerObjTM, LevelsetInnerObj, LevelsetInnerObjTM, Thickness, Constraint);
+				}
+			}
+		}
+
+
+
+		template<typename T, int d>
+		void ConstructUnionUnionConstraints(TGeometryParticleHandle<T, d>* Particle0, TGeometryParticleHandle<T, d>* Particle1, const FImplicitObject* Implicit0, const FImplicitObject* Implicit1, const T Thickness, FCollisionConstraintsArray& NewConstraints)
+		{
+			SCOPE_CYCLE_COUNTER(STAT_UpdateUnionUnionConstraint);
+
+			const FImplicitObject* ParticleObj = Particle0->Geometry().Get();
+			const FImplicitObject* LevelsetObj = Particle1->Geometry().Get();
+
+			TRigidTransform<T, d> ParticlesTM = TRigidTransform<T, d>(Particle0->X(), Particle0->R());
+			TRigidTransform<T, d> LevelsetTM = TRigidTransform<T, d>(Particle1->X(), Particle1->R());
+
+			const TArray<Pair<const FImplicitObject*, TRigidTransform<T, d>>> LevelsetShapes = FindRelevantShapes(ParticleObj, ParticlesTM, *LevelsetObj, LevelsetTM, Thickness);
+
+			for (const Pair<const FImplicitObject*, TRigidTransform<T, d>>& LevelsetObjPair : LevelsetShapes)
+			{
+				const FImplicitObject* LevelsetInnerObj = LevelsetObjPair.First;
+				const TRigidTransform<T, d>& LevelsetInnerObjTM = LevelsetObjPair.Second * LevelsetTM;
+
+				//now find all particle inner objects
+				const TArray<Pair<const FImplicitObject*, TRigidTransform<T, d>>> ParticleShapes = FindRelevantShapes(LevelsetInnerObj, LevelsetInnerObjTM, *ParticleObj, ParticlesTM, Thickness);
+
+				//for each inner obj pair, update constraint
+				for (const Pair<const FImplicitObject*, TRigidTransform<T, d>>& ParticlePair : ParticleShapes)
+				{
+					const FImplicitObject* ParticleInnerObj = ParticlePair.First;
+					ConstructConstraintsImpl<T, d>(Particle0, Particle1, ParticleInnerObj, LevelsetInnerObj, Thickness, NewConstraints);
+				}
+			}
+		}
+
+		//
+		// Single-Union
+		//
+
+		DECLARE_CYCLE_STAT(TEXT("TPBDCollisionConstraints::UpdateSingleUnionConstraint"), STAT_UpdateSingleUnionConstraint, STATGROUP_ChaosWide);
+		template<ECollisionUpdateType UpdateType, typename T, int d>
+		void UpdateSingleUnionConstraint(const T Thickness, TRigidBodyPointContactConstraint<T, d>& Constraint)
+		{
+			SCOPE_CYCLE_COUNTER(STAT_UpdateSingleUnionConstraint);
+
+			TGenericParticleHandle<T, d> Particle0 = Constraint.Particle[0];
+			TGenericParticleHandle<T, d> Particle1 = Constraint.Particle[1];
+
+			TRigidTransform<T, d> ParticlesTM = TRigidTransform<T, d>(Particle0->P(), Particle0->Q());
+			TRigidTransform<T, d> LevelsetTM = TRigidTransform<T, d>(Particle1->P(), Particle1->Q());
+
+			const FImplicitObject* ParticleObj = Particle0->Geometry().Get();
+			const FImplicitObject* LevelsetObj = Particle1->Geometry().Get();
+			const TArray<Pair<const FImplicitObject*, TRigidTransform<T, d>>> LevelsetShapes = FindRelevantShapes(ParticleObj, ParticlesTM, *LevelsetObj, LevelsetTM, Thickness);
+
+			for (const Pair<const FImplicitObject*, TRigidTransform<T, d>>& LevelsetObjPair : LevelsetShapes)
+			{
+				const FImplicitObject& LevelsetInnerObj = *LevelsetObjPair.First;
+				const TRigidTransform<T, d> LevelsetInnerObjTM = LevelsetObjPair.Second * LevelsetTM;
+				UpdateConstraintImp<UpdateType>(*ParticleObj, ParticlesTM, LevelsetInnerObj, LevelsetInnerObjTM, Thickness, Constraint);
+			}
+		}
+
+
+
+		template<typename T, int d>
+		void ConstructSingleUnionConstraints(TGeometryParticleHandle<T, d>* Particle0, TGeometryParticleHandle<T, d>* Particle1, const FImplicitObject* Implicit0, const FImplicitObject* Implicit1, const T Thickness, FCollisionConstraintsArray& NewConstraints)
+		{
+			FRigidBodyPointContactConstraint* Constraint = new FRigidBodyPointContactConstraint;
+			Constraint->Particle[0] = Particle0;
+			Constraint->Particle[1] = Particle1;
+			Constraint->AddManifold(Implicit0, Implicit1);
+
+			UpdateSingleUnionConstraint<ECollisionUpdateType::Any>(Thickness, *Constraint);
+
+			if (Constraint->GetPhi() < Thickness)
+			{
+				NewConstraints.Add(Constraint);
+			}
+			else
+			{
+				delete Constraint;
+			}
+		}
+
+		//
+		// Constraint API
+		//
+
+		template<ECollisionUpdateType UpdateType, typename T, int d>
+		void UpdateConstraintImp(const FImplicitObject& ParticleObject, const TRigidTransform<T, d>& ParticleTM, const FImplicitObject& LevelsetObject, const TRigidTransform<T, d>& LevelsetTM, const T Thickness, TRigidBodyPointContactConstraint<T, d>& Constraint)
+		{
+			if (ParticleObject.GetType() == TBox<T, d>::StaticType() && LevelsetObject.GetType() == TBox<T, d>::StaticType())
+			{
+				UpdateBoxBoxConstraint(*ParticleObject.template GetObject<TBox<T, d>>(), ParticleTM, *LevelsetObject.template GetObject<TBox<T, d>>(), LevelsetTM, Thickness, Constraint);
+			}
+			else if (ParticleObject.GetType() == TSphere<T, d>::StaticType() && LevelsetObject.GetType() == TSphere<T, d>::StaticType())
+			{
+				UpdateSphereSphereConstraint(*ParticleObject.template GetObject<TSphere<T, d>>(), ParticleTM, *LevelsetObject.template GetObject<TSphere<T, d>>(), LevelsetTM, Thickness, Constraint);
+			}
+			else if (ParticleObject.GetType() == TBox<T, d>::StaticType() && LevelsetObject.GetType() == TPlane<T, d>::StaticType())
+			{
+				UpdateBoxPlaneConstraint(*ParticleObject.template GetObject<TBox<T, d>>(), ParticleTM, *LevelsetObject.template GetObject<TPlane<T, d>>(), LevelsetTM, Thickness, Constraint);
+			}
+			else if (ParticleObject.GetType() == TSphere<T, d>::StaticType() && LevelsetObject.GetType() == TPlane<T, d>::StaticType())
+			{
+				UpdateSpherePlaneConstraint(*ParticleObject.template GetObject<TSphere<T, d>>(), ParticleTM, *LevelsetObject.template GetObject<TPlane<T, d>>(), LevelsetTM, Thickness, Constraint);
+			}
+			else if (ParticleObject.GetType() == TSphere<T, d>::StaticType() && LevelsetObject.GetType() == TBox<T, d>::StaticType())
+			{
+				UpdateSphereBoxConstraint(*ParticleObject.template GetObject<TSphere<T, d>>(), ParticleTM, *LevelsetObject.template GetObject<TBox<T, d>>(), LevelsetTM, Thickness, Constraint);
+			}
+			else if (ParticleObject.GetType() == TCapsule<T>::StaticType() && LevelsetObject.GetType() == TCapsule<T>::StaticType())
+			{
+				UpdateCapsuleCapsuleConstraint(*ParticleObject.template GetObject<TCapsule<T>>(), ParticleTM, *LevelsetObject.template GetObject<TCapsule<T>>(), LevelsetTM, Thickness, Constraint);
+			}
+			else if (ParticleObject.GetType() == TCapsule<T>::StaticType() && LevelsetObject.GetType() == TBox<T, d>::StaticType())
+			{
+				UpdateCapsuleBoxConstraint(*ParticleObject.template GetObject<TCapsule<T>>(), ParticleTM, *LevelsetObject.template GetObject<TBox<T, d>>(), LevelsetTM, Thickness, Constraint);
+			}
+			else if (ParticleObject.GetType() == TPlane<T, d>::StaticType() && LevelsetObject.GetType() == TBox<T, d>::StaticType())
+			{
+				TRigidBodyPointContactConstraint<T, d> TmpConstraint = Constraint;
+				UpdateBoxPlaneConstraint(*LevelsetObject.template GetObject<TBox<T, d>>(), LevelsetTM, *ParticleObject.template GetObject<TPlane<T, d>>(), ParticleTM, Thickness, TmpConstraint);
+				if (TmpConstraint.GetPhi() < Constraint.GetPhi())
+				{
+					Constraint = TmpConstraint;
+					Constraint.SetNormal(-Constraint.GetNormal());
+				}
+			}
+			else if (ParticleObject.GetType() == TPlane<T, d>::StaticType() && LevelsetObject.GetType() == TSphere<T, d>::StaticType())
+			{
+				TRigidBodyPointContactConstraint<T, d> TmpConstraint = Constraint;
+				UpdateSpherePlaneConstraint(*LevelsetObject.template GetObject<TSphere<T, d>>(), LevelsetTM, *ParticleObject.template GetObject<TPlane<T, d>>(), ParticleTM, Thickness, TmpConstraint);
+				if (TmpConstraint.GetPhi() < Constraint.GetPhi())
+				{
+					Constraint = TmpConstraint;
+					Constraint.SetNormal(-Constraint.GetNormal());
+				}
+			}
+			else if (ParticleObject.GetType() == TBox<T, d>::StaticType() && LevelsetObject.GetType() == TSphere<T, d>::StaticType())
+			{
+				TRigidBodyPointContactConstraint<T, d> TmpConstraint = Constraint;
+				UpdateSphereBoxConstraint(*LevelsetObject.template GetObject<TSphere<T, d>>(), LevelsetTM, *ParticleObject.template GetObject<TBox<T, d>>(), ParticleTM, Thickness, TmpConstraint);
+				if (TmpConstraint.GetPhi() < Constraint.GetPhi())
+				{
+					Constraint = TmpConstraint;
+					Constraint.SetNormal(-Constraint.GetNormal());
+				}
+			}
+			else if (ParticleObject.GetType() == TBox<T, d>::StaticType() && LevelsetObject.GetType() == TCapsule<T>::StaticType())
+			{
+				TRigidBodyPointContactConstraint<T, d> TmpConstraint = Constraint;
+				UpdateCapsuleBoxConstraint(*LevelsetObject.template GetObject<TCapsule<T>>(), LevelsetTM, *ParticleObject.template GetObject<TBox<T, d>>(), ParticleTM, Thickness, TmpConstraint);
+				if (TmpConstraint.GetPhi() < Constraint.GetPhi())
+				{
+					Constraint = TmpConstraint;
+					Constraint.SetNormal(-Constraint.GetNormal());
+				}
+			}
+			else if (ParticleObject.GetType() < TImplicitObjectUnion<T, d>::StaticType() && LevelsetObject.GetType() == TImplicitObjectUnion<T, d>::StaticType())
+			{
+				return UpdateSingleUnionConstraint<UpdateType>(Thickness, Constraint);
+			}
+			else if (ParticleObject.GetType() == TImplicitObjectUnion<T, d>::StaticType() && LevelsetObject.GetType() < TImplicitObjectUnion<T, d>::StaticType())
+			{
+				check(false);	//should not be possible to get this ordering (see ComputeConstraint)
+			}
+			else if (ParticleObject.GetType() == TImplicitObjectUnion<T, d>::StaticType() && LevelsetObject.GetType() == TImplicitObjectUnion<T, d>::StaticType())
+			{
+				return UpdateUnionUnionConstraint<UpdateType>(Thickness, Constraint);
+			}
+			else if (ParticleObject.IsConvex() && LevelsetObject.IsConvex())
+			{
+				UpdateConvexConvexConstraint(ParticleObject, ParticleTM, LevelsetObject, LevelsetTM, Thickness, Constraint);
+			}
+			else if (LevelsetObject.IsUnderlyingUnion())
+			{
+				UpdateUnionLevelsetConstraint<UpdateType>(Thickness, Constraint);
+			}
+			else if (ParticleObject.IsUnderlyingUnion())
+			{
+				UpdateLevelsetUnionConstraint<UpdateType>(Thickness, Constraint);
+			}
+			else
+			{
+				UpdateLevelsetLevelsetConstraint<UpdateType>(Thickness, Constraint);
+			}
+		}
+
+
+
+		template<typename T, int d>
+		void ConstructPairConstraintImpl(TGeometryParticleHandle<T, d>* Particle0, TGeometryParticleHandle<T, d>* Particle1, const FImplicitObject* Implicit0, const FImplicitObject* Implicit1, const T Thickness, FCollisionConstraintsArray& NewConstraints)
+		{
+			// See if we already have a constraint for this shape pair
+			for (int32 i = 0; i < NewConstraints.Num(); i++)
+			{
+				if (NewConstraints[i]->GetType() == TRigidBodyPointContactConstraint<T, d>::StaticType())
+				{
+					if (NewConstraints[i]->As< TRigidBodyPointContactConstraint<T, d> >()->ContainsManifold(Implicit0, Implicit1))
+					{
+						return;
+					}
+				}
+			}
+
+			if (!Implicit0 || !Implicit1)
+			{
+				ConstructLevelsetLevelsetConstraints(Particle0, Particle1, Implicit0, Implicit1, Thickness, NewConstraints);
+			}
+			if (Implicit0->GetType() == TBox<T, d>::StaticType() && Implicit1->GetType() == TBox<T, d>::StaticType())
+			{
+				ConstructBoxBoxConstraints(Particle0, Particle1, Implicit0, Implicit1, Thickness, NewConstraints);
+			}
+			else if (Implicit0->GetType() == TSphere<T, d>::StaticType() && Implicit1->GetType() == TSphere<T, d>::StaticType())
+			{
+				ConstructSphereSphereConstraints(Particle0, Particle1, Implicit0, Implicit1, Thickness, NewConstraints);
+			}
+			else if (Implicit0->GetType() == TBox<T, d>::StaticType() && Implicit1->GetType() == TPlane<T, d>::StaticType())
+			{
+				ConstructBoxPlaneConstraints(Particle0, Particle1, Implicit0, Implicit1, Thickness, NewConstraints);
+			}
+			else if (Implicit0->GetType() == TPlane<T, d>::StaticType() && Implicit0->GetType() == TBox<T, d>::StaticType())
+			{
+				ConstructBoxPlaneConstraints(Particle1, Particle0, Implicit1, Implicit0, Thickness, NewConstraints);
+			}
+			else if (Implicit0->GetType() == TSphere<T, d>::StaticType() && Implicit1->GetType() == TPlane<T, d>::StaticType())
+			{
+				ConstructSpherePlaneConstraints(Particle0, Particle1, Implicit0, Implicit1, Thickness, NewConstraints);
+			}
+			else if (Implicit0->GetType() == TPlane<T, d>::StaticType() && Implicit1->GetType() == TSphere<T, d>::StaticType())
+			{
+				ConstructSpherePlaneConstraints(Particle1, Particle0, Implicit1, Implicit0, Thickness, NewConstraints);
+			}
+			else if (Implicit0->GetType() == TSphere<T, d>::StaticType() && Implicit1->GetType() == TBox<T, d>::StaticType())
+			{
+				ConstructSphereBoxConstraints(Particle0, Particle1, Implicit0, Implicit1, Thickness, NewConstraints);
+			}
+			else if (Implicit0->GetType() == TBox<T, d>::StaticType() && Implicit1->GetType() == TSphere<T, d>::StaticType())
+			{
+				ConstructSphereBoxConstraints(Particle1, Particle0, Implicit1, Implicit0, Thickness, NewConstraints);
+			}
+			else if (Implicit0->GetType() == TCapsule<T>::StaticType() && Implicit1->GetType() == TCapsule<T>::StaticType())
+			{
+				ConstructCapsuleCapsuleConstraints(Particle0, Particle1, Implicit0, Implicit1, Thickness, NewConstraints);
+			}
+			else if (Implicit0->GetType() == TCapsule<T>::StaticType() && Implicit1->GetType() == TBox<T, d>::StaticType())
+			{
+				ConstructCapsuleBoxConstraints(Particle0, Particle1, Implicit0, Implicit1, Thickness, NewConstraints);
+			}
+			else if (Implicit0->GetType() == TBox<T, d>::StaticType() && Implicit1->GetType() == TCapsule<T>::StaticType())
+			{
+				ConstructCapsuleBoxConstraints(Particle1, Particle0, Implicit1, Implicit0, Thickness, NewConstraints);
+			}
+			else if (Implicit0->GetType() < TImplicitObjectUnion<T, d>::StaticType() && Implicit1->GetType() == TImplicitObjectUnion<T, d>::StaticType())
+			{
+				ConstructSingleUnionConstraints(Particle0, Particle1, Implicit0, Implicit1, Thickness, NewConstraints);
+			}
+			else if (Implicit0->GetType() == TImplicitObjectUnion<T, d>::StaticType() && Implicit1->GetType() < TImplicitObjectUnion<T, d>::StaticType())
+			{
+				ConstructSingleUnionConstraints(Particle1, Particle0, Implicit1, Implicit0, Thickness, NewConstraints);
+			}
+			else if (Implicit0->GetType() == TImplicitObjectUnion<T, d>::StaticType() && Implicit1->GetType() == TImplicitObjectUnion<T, d>::StaticType())
+			{
+				// Union-union creates multiple manifolds - we should never get here for this pair type. See ConstructConstraintsImpl and ConstructUnionUnionConstraints
+				ensure(false);
+			}
+			else if (Implicit0->IsConvex() && Implicit1->IsConvex())
+			{
+				ConstructConvexConvexConstraints(Particle0, Particle1, Implicit0, Implicit1, Thickness, NewConstraints);
+			}
+			else
+			{
+				ConstructLevelsetLevelsetConstraints(Particle0, Particle1, Implicit0, Implicit1, Thickness, NewConstraints);
+			}
+		}
+
+		template<typename T, int d>
+		void ConstructConstraintsImpl(TGeometryParticleHandle<T, d>* Particle0, TGeometryParticleHandle<T, d>* Particle1, const FImplicitObject* Implicit0, const FImplicitObject* Implicit1, const T Thickness, FCollisionConstraintsArray& NewConstraints)
+		{
+			// TriangleMesh implicits are for scene query only.
+			if (Implicit0 && GetInnerType(Implicit0->GetType()) == ImplicitObjectType::TriangleMesh) return;
+			if (Implicit1 && GetInnerType(Implicit1->GetType()) == ImplicitObjectType::TriangleMesh) return;
+
+			if (Implicit0 && Implicit0->GetType() == TImplicitObjectUnion<T, d>::StaticType() &&
+				Implicit1 && Implicit1->GetType() == TImplicitObjectUnion<T, d>::StaticType())
+			{
+				return ConstructUnionUnionConstraints(Particle0, Particle1, Implicit0, Implicit1, Thickness, NewConstraints);
+			}
+			else
+			{
+				return ConstructPairConstraintImpl(Particle0, Particle1, Implicit0, Implicit1, Thickness, NewConstraints);
+			}
+		}
+
+
+
+		template void UpdateBoxBoxConstraint<float, 3>(const TBox<float, 3>& Box1, const TRigidTransform<float, 3>& Box1Transform, const TBox<float, 3>& Box2, const TRigidTransform<float, 3>& Box2Transform, float Thickness, TRigidBodyPointContactConstraint<float, 3>& Constraint);
+		template void ConstructBoxBoxConstraints<float, 3>(TGeometryParticleHandle<float, 3>* Particle0, TGeometryParticleHandle<float, 3>* Particle1, const FImplicitObject* Implicit0, const FImplicitObject* Implicit1, const float Thickness, FCollisionConstraintsArray& NewConstraints);
+
+		template bool UpdateBoxPlaneConstraint<float, 3>(const TBox<float, 3>& Box, const TRigidTransform<float, 3>& BoxTransform, const TPlane<float, 3>& Plane, const TRigidTransform<float, 3>& PlaneTransform, float Thickness, TRigidBodyPointContactConstraint<float, 3>& Constraint);
+		template void ConstructBoxPlaneConstraints<float, 3>(TGeometryParticleHandle<float, 3>* Particle0, TGeometryParticleHandle<float, 3>* Particle1, const FImplicitObject* Implicit0, const FImplicitObject* Implicit1, const float Thickness, FCollisionConstraintsArray& NewConstraints);
+
+		template void UpdateSphereSphereConstraint<float, 3>(const TSphere<float, 3>& Sphere1, const TRigidTransform<float, 3>& Sphere1Transform, const TSphere<float, 3>& Sphere2, const TRigidTransform<float, 3>& Sphere2Transform, float Thickness, TRigidBodyPointContactConstraint<float, 3>& Constraint);
+		template void ConstructSphereSphereConstraints<float, 3>(TGeometryParticleHandle<float, 3>* Particle0, TGeometryParticleHandle<float, 3>* Particle1, const FImplicitObject* Implicit0, const FImplicitObject* Implicit1, const float Thickness, FCollisionConstraintsArray& NewConstraints);
+
+		template void UpdateSpherePlaneConstraint<float, 3>(const TSphere<float, 3>& Sphere, const TRigidTransform<float, 3>& SphereTransform, const TPlane<float, 3>& Plane, const TRigidTransform<float, 3>& PlaneTransform, float Thickness, TRigidBodyPointContactConstraint<float, 3>& Constraint);
+		template void ConstructSpherePlaneConstraints<float, 3>(TGeometryParticleHandle<float, 3>* Particle0, TGeometryParticleHandle<float, 3>* Particle1, const FImplicitObject* Implicit0, const FImplicitObject* Implicit1, const float Thickness, FCollisionConstraintsArray& NewConstraints);
+
+		template void UpdateSphereBoxConstraint<float, 3>(const TSphere<float, 3>& Sphere, const TRigidTransform<float, 3>& SphereTransform, const TBox<float, 3>& Box, const TRigidTransform<float, 3>& BoxTransform, float Thickness, TRigidBodyPointContactConstraint<float, 3>& Constraint);
+		template void ConstructSphereBoxConstraints<float, 3>(TGeometryParticleHandle<float, 3>* Particle0, TGeometryParticleHandle<float, 3>* Particle1, const FImplicitObject* Implicit0, const FImplicitObject* Implicit1, const float Thickness, FCollisionConstraintsArray& NewConstraints);
+
+		template void UpdateCapsuleCapsuleConstraint<float, 3>(const TCapsule<float>& A, const TRigidTransform<float, 3>& ATransform, const TCapsule<float>& B, const TRigidTransform<float, 3>& BTransform, float Thickness, TRigidBodyPointContactConstraint<float, 3>& Constraint);
+		template void ConstructCapsuleCapsuleConstraints(TGeometryParticleHandle<float, 3>* Particle0, TGeometryParticleHandle<float, 3>* Particle1, const FImplicitObject* Implicit0, const FImplicitObject* Implicit1, const float Thickness, FCollisionConstraintsArray& NewConstraints);
+
+		template void UpdateCapsuleBoxConstraint<float, 3>(const TCapsule<float>& A, const TRigidTransform<float, 3>& ATransform, const TBox<float, 3>& B, const TRigidTransform<float, 3>& BTransform, float Thickness, TRigidBodyPointContactConstraint<float, 3>& Constraint);
+		template void ConstructCapsuleBoxConstraints<float, 3>(TGeometryParticleHandle<float, 3>* Particle0, TGeometryParticleHandle<float, 3>* Particle1, const FImplicitObject* Implicit0, const FImplicitObject* Implicit1, const float Thickness, FCollisionConstraintsArray& NewConstraints);
+
+		template void UpdateConvexConvexConstraint<float, 3>(const FImplicitObject& A, const TRigidTransform<float, 3>& ATM, const FImplicitObject& B, const TRigidTransform<float, 3>& BTM, float Thickness, TRigidBodyPointContactConstraint<float, 3>& Constraint);
+		template void ConstructConvexConvexConstraints<float, 3>(TGeometryParticleHandle<float, 3>* Particle0, TGeometryParticleHandle<float, 3>* Particle1, const FImplicitObject* Implicit0, const FImplicitObject* Implicit1, const float Thickness, FCollisionConstraintsArray& NewConstraints);
+
+		template void UpdateLevelsetLevelsetConstraint<ECollisionUpdateType::Any, float, 3>(const float Thickness, TRigidBodyPointContactConstraint<float, 3>& Constraint);
+		template void UpdateLevelsetLevelsetConstraint<ECollisionUpdateType::Deepest, float, 3>(const float Thickness, TRigidBodyPointContactConstraint<float, 3>& Constraint);
+		template void ConstructLevelsetLevelsetConstraints<float, 3>(TGeometryParticleHandle<float, 3>* Particle0, TGeometryParticleHandle<float, 3>* Particle1, const FImplicitObject* Implicit0, const FImplicitObject* Implicit1, const float Thickness, FCollisionConstraintsArray& NewConstraints);
+
+		template void UpdateLevelsetUnionConstraint<ECollisionUpdateType::Any, float, 3>(const float Thickness, TRigidBodyPointContactConstraint<float, 3>& Constraint);
+		template void UpdateLevelsetUnionConstraint<ECollisionUpdateType::Deepest, float, 3>(const float Thickness, TRigidBodyPointContactConstraint<float, 3>& Constraint);
+
+		template void UpdateUnionLevelsetConstraint<ECollisionUpdateType::Any, float, 3>(const float Thickness, TRigidBodyPointContactConstraint<float, 3>& Constraint);
+		template void UpdateUnionLevelsetConstraint<ECollisionUpdateType::Deepest, float, 3>(const float Thickness, TRigidBodyPointContactConstraint<float, 3>& Constraint);
+
+		template void UpdateUnionUnionConstraint< ECollisionUpdateType::Any, float, 3>(const float Thickness, TRigidBodyPointContactConstraint<float, 3>& Constraint);
+		template void UpdateUnionUnionConstraint< ECollisionUpdateType::Deepest, float, 3>(const float Thickness, TRigidBodyPointContactConstraint<float, 3>& Constraint);
+		template void ConstructUnionUnionConstraints<float, 3>(TGeometryParticleHandle<float, 3>* Particle0, TGeometryParticleHandle<float, 3>* Particle1, const FImplicitObject* Implicit0, const FImplicitObject* Implicit1, const float Thickness, FCollisionConstraintsArray& NewConstraints);
+
+		template void UpdateSingleUnionConstraint< ECollisionUpdateType::Any, float, 3>(const float Thickness, TRigidBodyPointContactConstraint<float, 3>& Constraint);
+		template void UpdateSingleUnionConstraint< ECollisionUpdateType::Deepest, float, 3>(const float Thickness, TRigidBodyPointContactConstraint<float, 3>& Constraint);
+		template void ConstructSingleUnionConstraints<float, 3>(TGeometryParticleHandle<float, 3>* Particle0, TGeometryParticleHandle<float, 3>* Particle1, const FImplicitObject* Implicit0, const FImplicitObject* Implicit1, const float Thickness, FCollisionConstraintsArray& NewConstraints);
+
+		template void UpdateConstraintImp<ECollisionUpdateType::Any, float, 3>(const FImplicitObject& ParticleObject, const TRigidTransform<float, 3>& ParticleTM, const FImplicitObject& LevelsetObject, const TRigidTransform<float, 3>& LevelsetTM, const float Thickness, TRigidBodyPointContactConstraint<float, 3>& Constraints);
+		template void UpdateConstraintImp<ECollisionUpdateType::Deepest, float, 3>(const FImplicitObject& ParticleObject, const TRigidTransform<float, 3>& ParticleTM, const FImplicitObject& LevelsetObject, const TRigidTransform<float, 3>& LevelsetTM, const float Thickness, TRigidBodyPointContactConstraint<float, 3>& Constraint);
+
+		template void ConstructPairConstraintImpl<float, 3>(TGeometryParticleHandle<float, 3>* Particle0, TGeometryParticleHandle<float, 3>* Particle1, const FImplicitObject* Implicit0, const FImplicitObject* Implicit1, const float Thickness, FCollisionConstraintsArray& NewConstraints);
+		template void ConstructConstraintsImpl<float, 3>(TGeometryParticleHandle<float, 3>* Particle0, TGeometryParticleHandle<float, 3>* Particle1, const FImplicitObject* Implicit0, const FImplicitObject* Implicit1, const float Thickness, FCollisionConstraintsArray& NewConstraints);
+
+	} // Collisions
+
+} // Chaos
diff --git a/Engine/Source/Runtime/Experimental/Chaos/Private/Chaos/CollisionResolutionAlgo.cpp b/Engine/Source/Runtime/Experimental/Chaos/Private/Chaos/CollisionResolutionAlgo.cpp
deleted file mode 100644
diff --git a/Engine/Source/Runtime/Experimental/Chaos/Private/Chaos/CollisionResolutionConvexConvex.cpp b/Engine/Source/Runtime/Experimental/Chaos/Private/Chaos/CollisionResolutionConvexConvex.cpp
deleted file mode 100644
diff --git a/Engine/Source/Runtime/Experimental/Chaos/Private/Chaos/CollisionResolutionUtil.cpp b/Engine/Source/Runtime/Experimental/Chaos/Private/Chaos/CollisionResolutionUtil.cpp
new file mode 100644
--- /dev/null
+++ b/Engine/Source/Runtime/Experimental/Chaos/Private/Chaos/CollisionResolutionUtil.cpp
@@ -0,0 +1,614 @@
+// Copyright 1998-2019 Epic Games, Inc. All Rights Reserved.
+
+#include "Chaos/CollisionResolutionUtil.h"
+
+#include "Chaos/ChaosPerfTest.h"
+#include "Chaos/Defines.h"
+#include "Chaos/ImplicitObjectUnion.h"
+#include "Chaos/ImplicitObjectScaled.h"
+#include "Chaos/ISpatialAccelerationCollection.h"
+#include "Chaos/Levelset.h"
+#include "ChaosLog.h"
+#include "ChaosStats.h"
+#include "Stats/Stats.h" 
+
+#if INTEL_ISPC
+#include "PBDCollisionConstraints.ispc.generated.h"
+#endif
+
+namespace Chaos
+{
+	namespace Collisions
+	{
+
+		template<typename T, int d>
+		TRigidTransform<T, d> GetTransform(const TGeometryParticleHandle<T, d>* Particle)
+		{
+			TGenericParticleHandle<T, d> Generic = const_cast<TGeometryParticleHandle<T, d>*>(Particle);
+			return TRigidTransform<T, d>(Generic->P(), Generic->Q());
+		}
+
+		// @todo(ccaulfield): This is duplicated in JointConstraints - move to a utility file
+		template<typename T>
+		PMatrix<T, 3, 3> ComputeFactorMatrix3(const TVector<T, 3>& V, const PMatrix<T, 3, 3>& M, const T& Im)
+		{
+			// Rigid objects rotational contribution to the impulse.
+			// Vx*M*VxT+Im
+			check(Im > FLT_MIN)
+				return PMatrix<T, 3, 3>(
+					-V[2] * (-V[2] * M.M[1][1] + V[1] * M.M[2][1]) + V[1] * (-V[2] * M.M[2][1] + V[1] * M.M[2][2]) + Im,
+					V[2] * (-V[2] * M.M[1][0] + V[1] * M.M[2][0]) - V[0] * (-V[2] * M.M[2][1] + V[1] * M.M[2][2]),
+					-V[1] * (-V[2] * M.M[1][0] + V[1] * M.M[2][0]) + V[0] * (-V[2] * M.M[1][1] + V[1] * M.M[2][1]),
+					V[2] * (V[2] * M.M[0][0] - V[0] * M.M[2][0]) - V[0] * (V[2] * M.M[2][0] - V[0] * M.M[2][2]) + Im,
+					-V[1] * (V[2] * M.M[0][0] - V[0] * M.M[2][0]) + V[0] * (V[2] * M.M[1][0] - V[0] * M.M[2][1]),
+					-V[1] * (-V[1] * M.M[0][0] + V[0] * M.M[1][0]) + V[0] * (-V[1] * M.M[1][0] + V[0] * M.M[1][1]) + Im);
+		}
+
+		template<typename T, int d>
+		TVector<T, d> GetEnergyClampedImpulse(const TRigidBodyPointContactConstraint<T, d>& Constraint, const TVector<T, d>& Impulse, const TVector<T, d>& VectorToPoint1, const TVector<T, d>& VectorToPoint2, const TVector<T, d>& Velocity1, const TVector<T, d>& Velocity2)
+		{
+			TPBDRigidParticleHandle<T, d>* PBDRigid0 = Constraint.Particle[0]->AsDynamic();
+			TPBDRigidParticleHandle<T, d>* PBDRigid1 = Constraint.Particle[1]->AsDynamic();
+
+			TVector<T, d> Jr0, Jr1, IInvJr0, IInvJr1;
+			T ImpulseRatioNumerator0 = 0, ImpulseRatioNumerator1 = 0, ImpulseRatioDenom0 = 0, ImpulseRatioDenom1 = 0;
+			T ImpulseSize = Impulse.SizeSquared();
+			TVector<T, d> KinematicVelocity = !PBDRigid0 ? Velocity1 : !PBDRigid1 ? Velocity2 : TVector<T, d>(0);
+			if (PBDRigid0)
+			{
+				Jr0 = TVector<T, d>::CrossProduct(VectorToPoint1, Impulse);
+				IInvJr0 = PBDRigid0->Q().RotateVector(PBDRigid0->InvI() * PBDRigid0->Q().UnrotateVector(Jr0));
+				ImpulseRatioNumerator0 = TVector<T, d>::DotProduct(Impulse, PBDRigid0->V() - KinematicVelocity) + TVector<T, d>::DotProduct(IInvJr0, PBDRigid0->W());
+				ImpulseRatioDenom0 = ImpulseSize / PBDRigid0->M() + TVector<T, d>::DotProduct(Jr0, IInvJr0);
+			}
+			if (PBDRigid1)
+			{
+				Jr1 = TVector<T, d>::CrossProduct(VectorToPoint2, Impulse);
+				IInvJr1 = PBDRigid1->Q().RotateVector(PBDRigid1->InvI() * PBDRigid1->Q().UnrotateVector(Jr1));
+				ImpulseRatioNumerator1 = TVector<T, d>::DotProduct(Impulse, PBDRigid1->V() - KinematicVelocity) + TVector<T, d>::DotProduct(IInvJr1, PBDRigid1->W());
+				ImpulseRatioDenom1 = ImpulseSize / PBDRigid1->M() + TVector<T, d>::DotProduct(Jr1, IInvJr1);
+			}
+			T Numerator = -2 * (ImpulseRatioNumerator0 - ImpulseRatioNumerator1);
+			if (Numerator < 0)
+			{
+				return TVector<T, d>(0);
+			}
+			check(Numerator >= 0);
+			T Denominator = ImpulseRatioDenom0 + ImpulseRatioDenom1;
+			return Numerator < Denominator ? (Impulse * Numerator / Denominator) : Impulse;
+		}
+
+		template <typename T, int d>
+		bool SampleObjectHelper(const FImplicitObject& Object, const TRigidTransform<T, d>& ObjectTransform, const TRigidTransform<T, d>& SampleToObjectTransform, const TVector<T, d>& SampleParticle, T Thickness, TRigidBodyPointContactConstraint<float, 3>& Constraint)
+		{
+			TVector<T, d> LocalPoint = SampleToObjectTransform.TransformPositionNoScale(SampleParticle);
+			TVector<T, d> LocalNormal;
+			T LocalPhi = Object.PhiWithNormal(LocalPoint, LocalNormal);
+
+			TPointContactManifold<T, d> & Contact = Constraint.Manifold;
+			if (LocalPhi < Contact.Phi)
+			{
+				Contact.Phi = LocalPhi;
+				Contact.Normal = ObjectTransform.TransformVectorNoScale(LocalNormal);
+				Contact.Location = ObjectTransform.TransformPositionNoScale(LocalPoint);
+				return true;
+			}
+			return false;
+		}
+
+		template <typename T, int d>
+		bool SampleObjectNoNormal(const FImplicitObject& Object, const TRigidTransform<T, d>& ObjectTransform, const TRigidTransform<T, d>& SampleToObjectTransform, const TVector<T, d>& SampleParticle, T Thickness, TRigidBodyPointContactConstraint<float, 3>& Constraint)
+		{
+			TVector<T, d> LocalPoint = SampleToObjectTransform.TransformPositionNoScale(SampleParticle);
+			TVector<T, d> LocalNormal;
+			T LocalPhi = Object.PhiWithNormal(LocalPoint, LocalNormal);
+
+			TPointContactManifold<T, d> & Contact = Constraint.Manifold;
+			if (LocalPhi < Contact.Phi)
+			{
+				Contact.Phi = LocalPhi;
+				return true;
+			}
+			return false;
+		}
+
+		template <typename T, int d>
+		bool SampleObjectNormalAverageHelper(const FImplicitObject& Object, const TRigidTransform<T, d>& ObjectTransform, const TRigidTransform<T, d>& SampleToObjectTransform, const TVector<T, d>& SampleParticle, T Thickness, T& TotalThickness, TRigidBodyPointContactConstraint<float, 3>& Constraint)
+		{
+			TVector<T, d> LocalPoint = SampleToObjectTransform.TransformPositionNoScale(SampleParticle);
+			TVector<T, d> LocalNormal;
+			T LocalPhi = Object.PhiWithNormal(LocalPoint, LocalNormal);
+			T LocalThickness = LocalPhi - Thickness;
+
+			TPointContactManifold<T, d> & Contact = Constraint.Manifold;
+			if (LocalThickness < -KINDA_SMALL_NUMBER)
+			{
+				Contact.Location += LocalPoint * LocalThickness;
+				TotalThickness += LocalThickness;
+				return true;
+			}
+			return false;
+		}
+
+		DECLARE_CYCLE_STAT(TEXT("TPBDCollisionConstraints::UpdateLevelsetPartial"), STAT_UpdateLevelsetPartial, STATGROUP_ChaosWide);
+		DECLARE_CYCLE_STAT(TEXT("TPBDCollisionConstraints::UpdateLevelsetFindParticles"), STAT_UpdateLevelsetFindParticles, STATGROUP_ChaosWide);
+		DECLARE_CYCLE_STAT(TEXT("TPBDCollisionConstraints::UpdateLevelsetBVHTraversal"), STAT_UpdateLevelsetBVHTraversal, STATGROUP_ChaosWide);
+		DECLARE_CYCLE_STAT(TEXT("TPBDCollisionConstraints::UpdateLevelsetSignedDistance"), STAT_UpdateLevelsetSignedDistance, STATGROUP_ChaosWide);
+		DECLARE_CYCLE_STAT(TEXT("TPBDCollisionConstraints::UpdateLevelsetAll"), STAT_UpdateLevelsetAll, STATGROUP_ChaosWide);
+		DECLARE_CYCLE_STAT(TEXT("TPBDCollisionConstraints::SampleObject"), STAT_SampleObject, STATGROUP_ChaosWide);
+
+		int32 NormalAveraging = 1;
+		FAutoConsoleVariableRef CVarNormalAveraging(TEXT("p.NormalAveraging2"), NormalAveraging, TEXT(""));
+
+		int32 SampleMinParticlesForAcceleration = 2048;
+		FAutoConsoleVariableRef CVarSampleMinParticlesForAcceleration(TEXT("p.SampleMinParticlesForAcceleration"), SampleMinParticlesForAcceleration, TEXT("The minimum number of particles needed before using an acceleration structure when sampling"));
+
+
+#if INTEL_ISPC
+		template<ECollisionUpdateType UpdateType>
+		void SampleObject(const FImplicitObject& Object, const TRigidTransform<float, 3>& ObjectTransform, const TBVHParticles<float, 3>& SampleParticles, const TRigidTransform<float, 3>& SampleParticlesTransform, float Thickness, TRigidBodyPointContactConstraint<float, 3>& Constraint)
+		{
+			SCOPE_CYCLE_COUNTER(STAT_SampleObject);
+			TRigidBodyPointContactConstraint<float, 3> AvgConstraint = Constraint;
+
+			TPointContactManifold<float, 3> & Contact = Constraint.Manifold;
+			TPointContactManifold<float, 3> & AvgContact = AvgConstraint.Manifold;
+
+			AvgConstraint.Particle[0] = Constraint.Particle[0];
+			AvgConstraint.Particle[1] = Constraint.Particle[1];
+			AvgContact.Location = TVector<float, 3>::ZeroVector;
+			AvgContact.Normal = TVector<float, 3>::ZeroVector;
+			AvgContact.Phi = Thickness;
+			float TotalThickness = float(0);
+
+			int32 DeepestParticle = -1;
+
+			const TRigidTransform<float, 3>& SampleToObjectTM = SampleParticlesTransform.GetRelativeTransform(ObjectTransform);
+			int32 NumParticles = SampleParticles.Size();
+
+			if (NumParticles > SampleMinParticlesForAcceleration && Object.HasBoundingBox())
+			{
+				SCOPE_CYCLE_COUNTER(STAT_UpdateLevelsetPartial);
+				TBox<float, 3> ImplicitBox = Object.BoundingBox().TransformedBox(ObjectTransform.GetRelativeTransform(SampleParticlesTransform));
+				ImplicitBox.Thicken(Thickness);
+				TArray<int32> PotentialParticles;
+				{
+					SCOPE_CYCLE_COUNTER(STAT_UpdateLevelsetFindParticles);
+					PotentialParticles = SampleParticles.FindAllIntersections(ImplicitBox);
+				}
+				{
+					SCOPE_CYCLE_COUNTER(STAT_UpdateLevelsetSignedDistance);
+
+					if (Object.GetType(true) == ImplicitObjectType::LevelSet && PotentialParticles.Num() > 0)
+					{
+						//QUICK_SCOPE_CYCLE_COUNTER(STAT_LevelSet);
+						const TLevelSet<float, 3>* LevelSet = Object.GetObject<TLevelSet<float, 3>>();
+						const TUniformGrid<float, 3>& Grid = LevelSet->GetGrid();
+
+						if (NormalAveraging && UpdateType != ECollisionUpdateType::Any)
+						{
+							ispc::SampleLevelSetNormalAverage(
+								(ispc::FVector&)Grid.MinCorner(),
+								(ispc::FVector&)Grid.MaxCorner(),
+								(ispc::FVector&)Grid.Dx(),
+								(ispc::FIntVector&)Grid.Counts(),
+								(ispc::TArrayND*)&LevelSet->GetPhiArray(),
+								(ispc::FTransform&)SampleToObjectTM,
+								(ispc::FVector*)&SampleParticles.XArray()[0],
+								&PotentialParticles[0],
+								Thickness,
+								TotalThickness,
+								(ispc::FVector&)AvgContact.Location,
+								PotentialParticles.Num());
+						}
+						else
+						{
+							ispc::SampleLevelSetNoNormal(
+								(ispc::FVector&)Grid.MinCorner(),
+								(ispc::FVector&)Grid.MaxCorner(),
+								(ispc::FVector&)Grid.Dx(),
+								(ispc::FIntVector&)Grid.Counts(),
+								(ispc::TArrayND*)&LevelSet->GetPhiArray(),
+								(ispc::FTransform&)SampleToObjectTM,
+								(ispc::FVector*)&SampleParticles.XArray()[0],
+								&PotentialParticles[0],
+								DeepestParticle,
+								AvgContact.Phi,
+								PotentialParticles.Num());
+
+							if (UpdateType == ECollisionUpdateType::Any)
+							{
+								Contact.Phi = AvgContact.Phi;
+								return;
+							}
+						}
+					}
+					else if (Object.GetType(true) == ImplicitObjectType::Box && PotentialParticles.Num() > 0)
+					{
+						//QUICK_SCOPE_CYCLE_COUNTER(STAT_Box);
+						const TBox<float, 3>* Box = Object.GetObject<Chaos::TBox<float, 3>>();
+
+						if (NormalAveraging && UpdateType != ECollisionUpdateType::Any)
+						{
+							ispc::SampleBoxNormalAverage(
+								(ispc::FVector&)Box->Min(),
+								(ispc::FVector&)Box->Max(),
+								(ispc::FTransform&)SampleToObjectTM,
+								(ispc::FVector*)&SampleParticles.XArray()[0],
+								&PotentialParticles[0],
+								Thickness,
+								TotalThickness,
+								(ispc::FVector&)AvgContact.Location,
+								PotentialParticles.Num());
+						}
+						else
+						{
+							ispc::SampleBoxNoNormal(
+								(ispc::FVector&)Box->Min(),
+								(ispc::FVector&)Box->Max(),
+								(ispc::FTransform&)SampleToObjectTM,
+								(ispc::FVector*)&SampleParticles.XArray()[0],
+								&PotentialParticles[0],
+								DeepestParticle,
+								AvgContact.Phi,
+								PotentialParticles.Num());
+
+							if (UpdateType == ECollisionUpdateType::Any)
+							{
+								Contact.Phi = AvgContact.Phi;
+								return;
+							}
+						}
+					}
+					else
+					{
+						//QUICK_SCOPE_CYCLE_COUNTER(STAT_Other);
+						for (int32 i : PotentialParticles)
+						{
+							if (NormalAveraging && UpdateType != ECollisionUpdateType::Any)
+							{
+								SampleObjectNormalAverageHelper(Object, ObjectTransform, SampleToObjectTM, SampleParticles.X(i), Thickness, TotalThickness, AvgConstraint);
+							}
+							else
+							{
+								if (SampleObjectNoNormal(Object, ObjectTransform, SampleToObjectTM, SampleParticles.X(i), Thickness, AvgConstraint))
+								{
+									DeepestParticle = i;
+									if (UpdateType == ECollisionUpdateType::Any)
+									{
+										Contact.Phi = AvgContact.Phi;
+										return;
+									}
+								}
+							}
+						}
+					}
+				}
+			}
+			else
+			{
+				SCOPE_CYCLE_COUNTER(STAT_UpdateLevelsetAll);
+				if (Object.GetType(true) == ImplicitObjectType::LevelSet && NumParticles > 0)
+				{
+					const TLevelSet<float, 3>* LevelSet = Object.GetObject<Chaos::TLevelSet<float, 3>>();
+					const TUniformGrid<float, 3>& Grid = LevelSet->GetGrid();
+
+					if (NormalAveraging && UpdateType != ECollisionUpdateType::Any)
+					{
+						ispc::SampleLevelSetNormalAverageAll(
+							(ispc::FVector&)Grid.MinCorner(),
+							(ispc::FVector&)Grid.MaxCorner(),
+							(ispc::FVector&)Grid.Dx(),
+							(ispc::FIntVector&)Grid.Counts(),
+							(ispc::TArrayND*)&LevelSet->GetPhiArray(),
+							(ispc::FTransform&)SampleToObjectTM,
+							(ispc::FVector*)&SampleParticles.XArray()[0],
+							Thickness,
+							TotalThickness,
+							(ispc::FVector&)AvgContact.Location,
+							NumParticles);
+					}
+					else
+					{
+						ispc::SampleLevelSetNoNormalAll(
+							(ispc::FVector&)Grid.MinCorner(),
+							(ispc::FVector&)Grid.MaxCorner(),
+							(ispc::FVector&)Grid.Dx(),
+							(ispc::FIntVector&)Grid.Counts(),
+							(ispc::TArrayND*)&LevelSet->GetPhiArray(),
+							(ispc::FTransform&)SampleToObjectTM,
+							(ispc::FVector*)&SampleParticles.XArray()[0],
+							DeepestParticle,
+							AvgContact.Phi,
+							NumParticles);
+
+						if (UpdateType == ECollisionUpdateType::Any)
+						{
+							Contact.Phi = AvgContact.Phi;
+							return;
+						}
+					}
+				}
+				else if (Object.GetType(true) == ImplicitObjectType::Plane && NumParticles > 0)
+				{
+					const TPlane<float, 3>* Plane = Object.GetObject<Chaos::TPlane<float, 3>>();
+
+					if (NormalAveraging && UpdateType != ECollisionUpdateType::Any)
+					{
+						ispc::SamplePlaneNormalAverageAll(
+							(ispc::FVector&)Plane->Normal(),
+							(ispc::FVector&)Plane->X(),
+							(ispc::FTransform&)SampleToObjectTM,
+							(ispc::FVector*)&SampleParticles.XArray()[0],
+							Thickness,
+							TotalThickness,
+							(ispc::FVector&)AvgContact.Location,
+							NumParticles);
+					}
+					else
+					{
+						ispc::SamplePlaneNoNormalAll(
+							(ispc::FVector&)Plane->Normal(),
+							(ispc::FVector&)Plane->X(),
+							(ispc::FTransform&)SampleToObjectTM,
+							(ispc::FVector*)&SampleParticles.XArray()[0],
+							DeepestParticle,
+							AvgContact.Phi,
+							NumParticles);
+
+						if (UpdateType == ECollisionUpdateType::Any)
+						{
+							Contact.Phi = AvgContact.Phi;
+							return;
+						}
+					}
+				}
+				else if (Object.GetType(true) == ImplicitObjectType::Box && NumParticles > 0)
+				{
+					const TBox<float, 3>* Box = Object.GetObject<Chaos::TBox<float, 3>>();
+
+					if (NormalAveraging && UpdateType != ECollisionUpdateType::Any)
+					{
+						ispc::SampleBoxNormalAverageAll(
+							(ispc::FVector&)Box->Min(),
+							(ispc::FVector&)Box->Max(),
+							(ispc::FTransform&)SampleToObjectTM,
+							(ispc::FVector*)&SampleParticles.XArray()[0],
+							Thickness,
+							TotalThickness,
+							(ispc::FVector&)AvgContact.Location,
+							NumParticles);
+					}
+					else
+					{
+						ispc::SampleBoxNoNormalAll(
+							(ispc::FVector&)Box->Min(),
+							(ispc::FVector&)Box->Max(),
+							(ispc::FTransform&)SampleToObjectTM,
+							(ispc::FVector*)&SampleParticles.XArray()[0],
+							DeepestParticle,
+							AvgContact.Phi,
+							NumParticles);
+
+						if (UpdateType == ECollisionUpdateType::Any)
+						{
+							Contact.Phi = AvgContact.Phi;
+							return;
+						}
+					}
+				}
+				else
+				{
+					//QUICK_SCOPE_CYCLE_COUNTER(STAT_Other);
+					for (int32 i = 0; i < NumParticles; ++i)
+					{
+						if (NormalAveraging && UpdateType != ECollisionUpdateType::Any)
+						{
+							SampleObjectNormalAverageHelper(Object, ObjectTransform, SampleToObjectTM, SampleParticles.X(i), Thickness, TotalThickness, AvgConstraint);
+						}
+						else
+						{
+							if (SampleObjectNoNormal(Object, ObjectTransform, SampleToObjectTM, SampleParticles.X(i), Thickness, AvgConstraint))
+							{
+								DeepestParticle = i;
+								if (UpdateType == ECollisionUpdateType::Any)
+								{
+									Contact.Phi = AvgContact.Phi;
+									return;
+								}
+							}
+						}
+					}
+				}
+			}
+
+			if (NormalAveraging)
+			{
+				if (TotalThickness < -KINDA_SMALL_NUMBER)
+				{
+					TVector<float, 3> LocalPoint = AvgContact.Location / TotalThickness;
+					TVector<float, 3> LocalNormal;
+					const float NewPhi = Object.PhiWithNormal(LocalPoint, LocalNormal);
+					if (NewPhi < Contact.Phi)
+					{
+						Contact.Phi = NewPhi;
+						Contact.Location = ObjectTransform.TransformPositionNoScale(LocalPoint);
+						Contact.Normal = ObjectTransform.TransformVectorNoScale(LocalNormal);
+					}
+				}
+				else
+				{
+					check(AvgContact.Phi >= Thickness);
+				}
+			}
+			else if (AvgContact.Phi < Contact.Phi)
+			{
+				check(DeepestParticle >= 0);
+				TVector<float, 3> LocalPoint = SampleToObjectTM.TransformPositionNoScale(SampleParticles.X(DeepestParticle));
+				TVector<float, 3> LocalNormal;
+				Contact.Phi = Object.PhiWithNormal(LocalPoint, LocalNormal);
+				Contact.Location = ObjectTransform.TransformPositionNoScale(LocalPoint);
+				Contact.Normal = ObjectTransform.TransformVectorNoScale(LocalNormal);
+			}
+		}
+#else		
+		template <ECollisionUpdateType UpdateType, typename T, int d>
+		void SampleObject(const FImplicitObject& Object, const TRigidTransform<T, d>& ObjectTransform, const TBVHParticles<T, d>& SampleParticles, const TRigidTransform<T, d>& SampleParticlesTransform, T Thickness, TRigidBodyPointContactConstraint<float, 3>& Constraint)
+		{
+			SCOPE_CYCLE_COUNTER(STAT_SampleObject);
+			TRigidBodyPointContactConstraint<float, 3> AvgConstraint = Constraint;
+
+			TPointContactManifold<float, 3> & Contact = Constraint.Manifold;
+			TPointContactManifold<float, 3> & AvgContact = AvgConstraint.Manifold;
+
+			AvgConstraint.Particle[0] = Constraint.Particle[0];
+			AvgConstraint.Particle[1] = Constraint.Particle[1];
+			AvgContact.Location = TVector<float, 3>::ZeroVector;
+			AvgContact.Normal = TVector<float, 3>::ZeroVector;
+			AvgContact.Phi = Thickness;
+			float TotalThickness = float(0);
+
+			int32 DeepestParticle = -1;
+			const int32 NumParticles = SampleParticles.Size();
+
+			const TRigidTransform<T, d> & SampleToObjectTM = SampleParticlesTransform.GetRelativeTransform(ObjectTransform);
+			if (NumParticles > SampleMinParticlesForAcceleration && Object.HasBoundingBox())
+			{
+				SCOPE_CYCLE_COUNTER(STAT_UpdateLevelsetPartial);
+				TBox<T, d> ImplicitBox = Object.BoundingBox().TransformedBox(ObjectTransform.GetRelativeTransform(SampleParticlesTransform));
+				ImplicitBox.Thicken(Thickness);
+				TArray<int32> PotentialParticles;
+				{
+					SCOPE_CYCLE_COUNTER(STAT_UpdateLevelsetFindParticles);
+					PotentialParticles = SampleParticles.FindAllIntersections(ImplicitBox);
+				}
+				{
+					SCOPE_CYCLE_COUNTER(STAT_UpdateLevelsetSignedDistance);
+					for (int32 i : PotentialParticles)
+					{
+						if (NormalAveraging && UpdateType != ECollisionUpdateType::Any)	//if we just want one don't bother with normal
+						{
+							SampleObjectNormalAverageHelper(Object, ObjectTransform, SampleToObjectTM, SampleParticles.X(i), Thickness, TotalThickness, AvgConstraint);
+						}
+						else
+						{
+							if (SampleObjectNoNormal(Object, ObjectTransform, SampleToObjectTM, SampleParticles.X(i), Thickness, AvgConstraint))
+							{
+								DeepestParticle = i;
+								if (UpdateType == ECollisionUpdateType::Any)
+								{
+									Contact.Phi = AvgContact.Phi;
+									return;
+								}
+							}
+						}
+					}
+				}
+			}
+			else
+			{
+				SCOPE_CYCLE_COUNTER(STAT_UpdateLevelsetAll);
+				for (int32 i = 0; i < NumParticles; ++i)
+				{
+					if (NormalAveraging && UpdateType != ECollisionUpdateType::Any)	//if we just want one don't bother with normal
+					{
+						const bool bInside = SampleObjectNormalAverageHelper(Object, ObjectTransform, SampleToObjectTM, SampleParticles.X(i), Thickness, TotalThickness, AvgConstraint);
+					}
+					else
+					{
+						if (SampleObjectNoNormal(Object, ObjectTransform, SampleToObjectTM, SampleParticles.X(i), Thickness, AvgConstraint))
+						{
+							DeepestParticle = i;
+							if (UpdateType == ECollisionUpdateType::Any)
+							{
+								Contact.Phi = AvgContact.Phi;
+								return;
+							}
+						}
+					}
+				}
+			}
+
+			if (NormalAveraging)
+			{
+				if (TotalThickness < -KINDA_SMALL_NUMBER)
+				{
+					TVector<T, d> LocalPoint = AvgContact.Location / TotalThickness;
+					TVector<T, d> LocalNormal;
+					const T NewPhi = Object.PhiWithNormal(LocalPoint, LocalNormal);
+					if (NewPhi < Contact.Phi)
+					{
+						Contact.Phi = NewPhi;
+						Contact.Location = ObjectTransform.TransformPositionNoScale(LocalPoint);
+						Contact.Normal = ObjectTransform.TransformVectorNoScale(LocalNormal);
+					}
+				}
+				else
+				{
+					check(AvgContact.Phi >= Thickness);
+				}
+			}
+			else if (AvgContact.Phi < Contact.Phi)
+			{
+				check(DeepestParticle >= 0);
+				TVector<T, d> LocalPoint = SampleToObjectTM.TransformPositionNoScale(SampleParticles.X(DeepestParticle));
+				TVector<T, d> LocalNormal;
+				Contact.Phi = Object.PhiWithNormal(LocalPoint, LocalNormal);
+				Contact.Location = ObjectTransform.TransformPositionNoScale(LocalPoint);
+				Contact.Normal = ObjectTransform.TransformVectorNoScale(LocalNormal);
+
+			}
+		}
+#endif
+
+
+		DECLARE_CYCLE_STAT(TEXT("TPBDCollisionConstraints::FindRelevantShapes"), STAT_FindRelevantShapes, STATGROUP_ChaosWide);
+		template <typename T, int d>
+		TArray<Pair<const FImplicitObject*, TRigidTransform<T, d>>> FindRelevantShapes(const FImplicitObject* ParticleObj, const TRigidTransform<T, d>& ParticlesTM, const FImplicitObject& LevelsetObj, const TRigidTransform<T, d>& LevelsetTM, const T Thickness)
+		{
+			SCOPE_CYCLE_COUNTER(STAT_FindRelevantShapes);
+			TArray<Pair<const FImplicitObject*, TRigidTransform<T, d>>> RelevantShapes;
+			//find all levelset inner objects
+			if (ParticleObj)
+			{
+				if (ParticleObj->HasBoundingBox())
+				{
+					const TRigidTransform<T, d> ParticlesToLevelsetTM = ParticlesTM.GetRelativeTransform(LevelsetTM);
+					TBox<T, d> ParticleBoundsInLevelset = ParticleObj->BoundingBox().TransformedBox(ParticlesToLevelsetTM);
+					ParticleBoundsInLevelset.Thicken(Thickness);
+					{
+						LevelsetObj.FindAllIntersectingObjects(RelevantShapes, ParticleBoundsInLevelset);
+					}
+				}
+				else
+				{
+					LevelsetObj.AccumulateAllImplicitObjects(RelevantShapes, TRigidTransform<T, d>::Identity);
+				}
+			}
+			else
+			{
+				//todo:compute bounds
+				LevelsetObj.AccumulateAllImplicitObjects(RelevantShapes, TRigidTransform<T, d>::Identity);
+			}
+
+			return RelevantShapes;
+		}
+
+		template TRigidTransform<float, 3> GetTransform(const TGeometryParticleHandle<float, 3>* Particle);
+		template PMatrix<float, 3, 3> ComputeFactorMatrix3(const TVector<float, 3>& V, const PMatrix<float, 3, 3>& M, const float& Im);
+		template TVector<float, 3> GetEnergyClampedImpulse(const TRigidBodyPointContactConstraint<float, 3>& Constraint, const TVector<float, 3>& Impulse, const TVector<float, 3>& VectorToPoint1, const TVector<float, 3>& VectorToPoint2, const TVector<float, 3>& Velocity1, const TVector<float, 3>& Velocity2);
+		template bool SampleObjectHelper<float, 3>(const FImplicitObject& Object, const TRigidTransform<float, 3>& ObjectTransform, const TRigidTransform<float, 3>& SampleToObjectTransform, const TVector<float, 3>& SampleParticle, float Thickness, TRigidBodyPointContactConstraint<float, 3>& Constraint);
+		template bool SampleObjectNoNormal(const FImplicitObject& Object, const TRigidTransform<float, 3>& ObjectTransform, const TRigidTransform<float, 3>& SampleToObjectTransform, const TVector<float, 3>& SampleParticle, float Thickness, TRigidBodyPointContactConstraint<float, 3>& Constraint);
+		template bool SampleObjectNormalAverageHelper(const FImplicitObject& Object, const TRigidTransform<float, 3>& ObjectTransform, const TRigidTransform<float, 3>& SampleToObjectTransform, const TVector<float, 3>& SampleParticle, float Thickness, float& TotalThickness, TRigidBodyPointContactConstraint<float, 3>& Constraint);
+#if INTEL_ISPC
+		template void SampleObject<ECollisionUpdateType::Any>(const FImplicitObject& Object, const TRigidTransform<float, 3>& ObjectTransform, const TBVHParticles<float, 3>& SampleParticles, const TRigidTransform<float, 3>& SampleParticlesTransform, float Thickness, TRigidBodyPointContactConstraint<float, 3>& Constraint);
+		template void SampleObject<ECollisionUpdateType::Deepest>(const FImplicitObject& Object, const TRigidTransform<float, 3>& ObjectTransform, const TBVHParticles<float, 3>& SampleParticles, const TRigidTransform<float, 3>& SampleParticlesTransform, float Thickness, TRigidBodyPointContactConstraint<float, 3>& Constraint);
+#else
+		template void SampleObject<ECollisionUpdateType::Any, float, 3>(const FImplicitObject& Object, const TRigidTransform<float, 3>& ObjectTransform, const TBVHParticles<float, 3>& SampleParticles, const TRigidTransform<float, 3>& SampleParticlesTransform, float Thickness, TRigidBodyPointContactConstraint<float, 3>& Constraint);
+		template void SampleObject<ECollisionUpdateType::Deepest, float, 3>(const FImplicitObject& Object, const TRigidTransform<float, 3>& ObjectTransform, const TBVHParticles<float, 3>& SampleParticles, const TRigidTransform<float, 3>& SampleParticlesTransform, float Thickness, TRigidBodyPointContactConstraint<float, 3>& Constraint);
+#endif
+		template TArray<Pair<const FImplicitObject*, TRigidTransform<float, 3>>> FindRelevantShapes<float, 3>(const FImplicitObject* ParticleObj, const TRigidTransform<float, 3>& ParticlesTM, const FImplicitObject& LevelsetObj, const TRigidTransform<float, 3>& LevelsetTM, const float Thickness);
+
+	} // Collisions
+
+} // Chaos
diff --git a/Engine/Source/Runtime/Experimental/Chaos/Private/Chaos/PBDCollisionConstraints.cpp b/Engine/Source/Runtime/Experimental/Chaos/Private/Chaos/PBDCollisionConstraints.cpp
--- a/Engine/Source/Runtime/Experimental/Chaos/Private/Chaos/PBDCollisionConstraints.cpp
+++ b/Engine/Source/Runtime/Experimental/Chaos/Private/Chaos/PBDCollisionConstraints.cpp
@@ -6,8 +6,8 @@
 #include "Chaos/Capsule.h"
 #include "Chaos/ChaosPerfTest.h"
 #include "Chaos/ChaosDebugDraw.h"
-#include "Chaos/CollisionResolutionAlgo.h"
-#include "Chaos/CollisionResolutionConvexConvex.h"
+#include "Chaos/CollisionResolutionUtil.h"
+#include "Chaos/CollisionResolution.h"
 #include "Chaos/Defines.h"
 #include "Chaos/GeometryQueries.h"
 #include "Chaos/ImplicitObjectUnion.h"
@@ -147,12 +147,13 @@ namespace Chaos
 	}
 
 	template<typename T, int d>
-	void TPBDCollisionConstraints<T, d>::ConstructConstraints(TGeometryParticleHandle<T, d>* Particle0, TGeometryParticleHandle<T, d>* Particle1, const T Thickness, TRigidBodyPointContactConstraint<T, d> & Constraint)
+	void TPBDCollisionConstraints<T, d>::ConstructConstraints(TGeometryParticleHandle<T, d>* Particle0, TGeometryParticleHandle<T, d>* Particle1, const T Thickness, FCollisionConstraintsArray& NewConstraints)
 	{
 		if (ensure(Particle0 && Particle1))
 		{
-			Collisions::ConstructConstraintsImpl<T, d>(Particle0, Particle1, Particle0->Geometry().Get(), Particle1->Geometry().Get(), Thickness, Constraint);
+			Collisions::ConstructConstraintsImpl<T, d>(Particle0, Particle1, Particle0->Geometry().Get(), Particle1->Geometry().Get(), Thickness, NewConstraints);
 		}
+
 	}
 
 	DECLARE_CYCLE_STAT(TEXT("TPBDCollisionConstraints::Reset"), STAT_CollisionConstraintsReset, STATGROUP_Chaos);
@@ -282,7 +283,7 @@ namespace Chaos
 		SCOPE_CYCLE_COUNTER(STAT_UpdateConstraint);
 		//if (ConstraintHandle->GetType() == FPointContactConstraint::StaticType())
 		{
-			Collisions::UpdateConstraint<UpdateType>(Thickness, Constraint);
+			Collisions::Update<UpdateType>(Thickness, Constraint);
 		}
 	}
 
diff --git a/Engine/Source/Runtime/Experimental/Chaos/Private/Chaos/PBDCollisionConstraintsPlaneContactUtil.cpp b/Engine/Source/Runtime/Experimental/Chaos/Private/Chaos/PBDCollisionConstraintsPlaneContactUtil.cpp
--- a/Engine/Source/Runtime/Experimental/Chaos/Private/Chaos/PBDCollisionConstraintsPlaneContactUtil.cpp
+++ b/Engine/Source/Runtime/Experimental/Chaos/Private/Chaos/PBDCollisionConstraintsPlaneContactUtil.cpp
@@ -1,7 +1,7 @@
 // Copyright 1998-2019 Epic Games, Inc. All Rights Reserved.
 
 #include "Chaos/PBDCollisionConstraintsPlaneContactUtil.h"
-#include "Chaos/CollisionResolutionAlgo.h"
+#include "Chaos/CollisionResolutionUtil.h"
 #include "Chaos/Defines.h"
 
 namespace Chaos
@@ -17,14 +17,6 @@ namespace Chaos
 		template void Update<ECollisionUpdateType::Deepest, float, 3>(const float, TRigidBodyPlaneContactConstraint<float, 3>&);
 
 
-		template<ECollisionUpdateType UpdateType, typename T, int d>
-		void UpdateConstraint(const T Thickness, TRigidBodyPlaneContactConstraint<T, d>& Constraint)
-		{
-		}
-		template void UpdateConstraint<ECollisionUpdateType::Any, float, 3>(const float, TRigidBodyPointContactConstraint<float,3>&);
-		template void UpdateConstraint<ECollisionUpdateType::Deepest, float, 3>(const float, TRigidBodyPointContactConstraint<float,3>&);
-
-
 		template<typename T, int d>
 		void Apply(TRigidBodyPlaneContactConstraint<T, d>& Constraint, T Thickness, TPointContactIterationParameters<T> & IterationParameters, TPlainContactParticleParameters<T> & ParticleParameters)
 		{
@@ -38,7 +30,7 @@ namespace Chaos
 		}
 		template void ApplyPushOut<float, 3>(TRigidBodyPointContactConstraint<float,3>&, float , const TSet<const TGeometryParticleHandle<float,3>*>&,
 			TPointContactIterationParameters<float> & IterationParameters, TPointContactParticleParameters<float> & ParticleParameters);
-		*/
+	*/
 	} // Collisions
 
 }// Chaos
diff --git a/Engine/Source/Runtime/Experimental/Chaos/Private/Chaos/PBDCollisionConstraintsPointContactUtil.cpp b/Engine/Source/Runtime/Experimental/Chaos/Private/Chaos/PBDCollisionConstraintsPointContactUtil.cpp
--- a/Engine/Source/Runtime/Experimental/Chaos/Private/Chaos/PBDCollisionConstraintsPointContactUtil.cpp
+++ b/Engine/Source/Runtime/Experimental/Chaos/Private/Chaos/PBDCollisionConstraintsPointContactUtil.cpp
@@ -1,7 +1,7 @@
 // Copyright 1998-2019 Epic Games, Inc. All Rights Reserved.
 
 #include "Chaos/PBDCollisionConstraintsPointContactUtil.h"
-#include "Chaos/CollisionResolutionAlgo.h"
+#include "Chaos/CollisionResolution.h"
 #include "Chaos/Defines.h"
 #include "Chaos/Utilities.h"
 
@@ -9,64 +9,6 @@ namespace Chaos
 {
 	namespace Collisions
 	{
-		template<typename T=float, int d=3>
-		TRigidTransform<T, d> GetTransform(const TGeometryParticleHandle<T, d>* Particle)
-		{
-			TGenericParticleHandle<T, d> Generic = const_cast<TGeometryParticleHandle<T, d>*>(Particle);
-			return TRigidTransform<T, d>(Generic->P(), Generic->Q());
-		}
-
-		// @todo(ccaulfield): This is duplicated in JointConstraints - move to a utility file
-		template<typename T>
-		PMatrix<T, 3, 3> ComputeFactorMatrix3(const TVector<T, 3>& V, const PMatrix<T, 3, 3>& M, const T& Im)
-		{
-			// Rigid objects rotational contribution to the impulse.
-			// Vx*M*VxT+Im
-			check(Im > FLT_MIN)
-				return PMatrix<T, 3, 3>(
-					-V[2] * (-V[2] * M.M[1][1] + V[1] * M.M[2][1]) + V[1] * (-V[2] * M.M[2][1] + V[1] * M.M[2][2]) + Im,
-					V[2] * (-V[2] * M.M[1][0] + V[1] * M.M[2][0]) - V[0] * (-V[2] * M.M[2][1] + V[1] * M.M[2][2]),
-					-V[1] * (-V[2] * M.M[1][0] + V[1] * M.M[2][0]) + V[0] * (-V[2] * M.M[1][1] + V[1] * M.M[2][1]),
-					V[2] * (V[2] * M.M[0][0] - V[0] * M.M[2][0]) - V[0] * (V[2] * M.M[2][0] - V[0] * M.M[2][2]) + Im,
-					-V[1] * (V[2] * M.M[0][0] - V[0] * M.M[2][0]) + V[0] * (V[2] * M.M[1][0] - V[0] * M.M[2][1]),
-					-V[1] * (-V[1] * M.M[0][0] + V[0] * M.M[1][0]) + V[0] * (-V[1] * M.M[1][0] + V[0] * M.M[1][1]) + Im);
-		}
-
-		template<typename T, int d>
-		TVector<T, d> GetEnergyClampedImpulse(const TRigidBodyPointContactConstraint<T, d>& Constraint, const TVector<T, d>& Impulse, const TVector<T, d>& VectorToPoint1, const TVector<T, d>& VectorToPoint2, const TVector<T, d>& Velocity1, const TVector<T, d>& Velocity2)
-		{
-			TPBDRigidParticleHandle<T, d>* PBDRigid0 = Constraint.Particle[0]->AsDynamic();
-			TPBDRigidParticleHandle<T, d>* PBDRigid1 = Constraint.Particle[1]->AsDynamic();
-
-			TVector<T, d> Jr0, Jr1, IInvJr0, IInvJr1;
-			T ImpulseRatioNumerator0 = 0, ImpulseRatioNumerator1 = 0, ImpulseRatioDenom0 = 0, ImpulseRatioDenom1 = 0;
-			T ImpulseSize = Impulse.SizeSquared();
-			TVector<T, d> KinematicVelocity = !PBDRigid0 ? Velocity1 : !PBDRigid1 ? Velocity2 : TVector<T, d>(0);
-			if (PBDRigid0)
-			{
-				Jr0 = TVector<T, d>::CrossProduct(VectorToPoint1, Impulse);
-				IInvJr0 = PBDRigid0->Q().RotateVector(PBDRigid0->InvI() * PBDRigid0->Q().UnrotateVector(Jr0));
-				ImpulseRatioNumerator0 = TVector<T, d>::DotProduct(Impulse, PBDRigid0->V() - KinematicVelocity) + TVector<T, d>::DotProduct(IInvJr0, PBDRigid0->W());
-				ImpulseRatioDenom0 = ImpulseSize / PBDRigid0->M() + TVector<T, d>::DotProduct(Jr0, IInvJr0);
-			}
-			if (PBDRigid1)
-			{
-				Jr1 = TVector<T, d>::CrossProduct(VectorToPoint2, Impulse);
-				IInvJr1 = PBDRigid1->Q().RotateVector(PBDRigid1->InvI() * PBDRigid1->Q().UnrotateVector(Jr1));
-				ImpulseRatioNumerator1 = TVector<T, d>::DotProduct(Impulse, PBDRigid1->V() - KinematicVelocity) + TVector<T, d>::DotProduct(IInvJr1, PBDRigid1->W());
-				ImpulseRatioDenom1 = ImpulseSize / PBDRigid1->M() + TVector<T, d>::DotProduct(Jr1, IInvJr1);
-			}
-			T Numerator = -2 * (ImpulseRatioNumerator0 - ImpulseRatioNumerator1);
-			if (Numerator < 0)
-			{
-				return TVector<T, d>(0);
-			}
-			check(Numerator >= 0);
-			T Denominator = ImpulseRatioDenom0 + ImpulseRatioDenom1;
-			return Numerator < Denominator ? (Impulse * Numerator / Denominator) : Impulse;
-		}
-
-
 		template<ECollisionUpdateType UpdateType, typename T, int d>
 		void Update(const T Thickness, TRigidBodyPointContactConstraint<T, d>& Constraint)
 		{
@@ -80,7 +22,7 @@ namespace Chaos
 				{
 					if (!Constraint.Particle[1]->Geometry()->IsUnderlyingUnion())
 					{
-						UpdateLevelsetConstraint<UpdateType>(Thickness, Constraint);
+						UpdateLevelsetLevelsetConstraint<UpdateType>(Thickness, Constraint);
 					}
 					else
 					{
@@ -93,39 +35,6 @@ namespace Chaos
 				UpdateConstraintImp<UpdateType>(*Constraint.Particle[0]->Geometry(), ParticleTM, *Constraint.Particle[1]->Geometry(), LevelsetTM, Thickness, Constraint);
 			}
 		}
-		template void Update<ECollisionUpdateType::Any, float, 3>(const float, TRigidBodyPointContactConstraint<float, 3>&);
-		template void Update<ECollisionUpdateType::Deepest, float, 3>(const float, TRigidBodyPointContactConstraint<float, 3>&);
-
-
-		template<ECollisionUpdateType UpdateType, typename T, int d>
-		void UpdateConstraint(const T Thickness, TRigidBodyPointContactConstraint<T, d>& Constraint)
-		{
-			Constraint.ResetPhi(Thickness);
-			const TRigidTransform<T, d> ParticleTM = GetTransform(Constraint.Particle[0]);
-			const TRigidTransform<T, d> LevelsetTM = GetTransform(Constraint.Particle[1]);
-
-			if (!Constraint.Particle[0]->Geometry())
-			{
-				if (Constraint.Particle[1]->Geometry())
-				{
-					if (!Constraint.Particle[1]->Geometry()->IsUnderlyingUnion())
-					{
-						UpdateLevelsetConstraint<UpdateType>(Thickness, Constraint);
-					}
-					else
-					{
-						UpdateUnionLevelsetConstraint<UpdateType>(Thickness, Constraint);
-					}
-				}
-			}
-			else
-			{
-				UpdateConstraintImp<UpdateType>(*Constraint.Particle[0]->Geometry(), ParticleTM, *Constraint.Particle[1]->Geometry(), LevelsetTM, Thickness, Constraint);
-			}
-		}
-		template void UpdateConstraint<ECollisionUpdateType::Any, float, 3>(const float, TRigidBodyPointContactConstraint<float,3>&);
-		template void UpdateConstraint<ECollisionUpdateType::Deepest, float, 3>(const float, TRigidBodyPointContactConstraint<float,3>&);
-
 
 		template<typename T, int d>
 		void Apply(TRigidBodyPointContactConstraint<T, d>& Constraint, T Thickness, TPointContactIterationParameters<T> & IterationParameters, TPointContactParticleParameters<T> & ParticleParameters)
@@ -137,7 +46,7 @@ namespace Chaos
 
 			for (int32 PairIt = 0; PairIt < IterationParameters.NumPairIterations; ++PairIt)
 			{
-				Collisions::UpdateConstraint<ECollisionUpdateType::Deepest>(Thickness, Constraint);
+				Collisions::Update<ECollisionUpdateType::Deepest>(Thickness, Constraint);
 				if (Constraint.GetPhi() >= Thickness)
 				{
 					return;
@@ -346,7 +255,6 @@ namespace Chaos
 				}
 			}
 		}
-		template void Apply<float, 3>(TRigidBodyPointContactConstraint<float, 3>&, float, TPointContactIterationParameters<float> &, TPointContactParticleParameters<float> &);
 
 		template<typename T, int d>
 		void ApplyPushOut(TRigidBodyPointContactConstraint<T, d>& Constraint, T Thickness, const TSet<const TGeometryParticleHandle<T, d>*>& IsTemporarilyStatic,
@@ -378,7 +286,7 @@ namespace Chaos
 
 			for (int32 PairIteration = 0; PairIteration < IterationParameters.NumPairIterations; ++PairIteration)
 			{
-				UpdateConstraint<ECollisionUpdateType::Deepest>(Thickness, Constraint);
+				Update<ECollisionUpdateType::Deepest>(Thickness, Constraint);
 
 				const typename TRigidBodyPointContactConstraint<T, d>::FManifold & Contact = Constraint.Manifold;
 
@@ -459,6 +367,10 @@ namespace Chaos
 				}
 			}
 		}
+
+		template void Update<ECollisionUpdateType::Any, float, 3>(const float, TRigidBodyPointContactConstraint<float, 3>&);
+		template void Update<ECollisionUpdateType::Deepest, float, 3>(const float, TRigidBodyPointContactConstraint<float, 3>&);
+		template void Apply<float, 3>(TRigidBodyPointContactConstraint<float, 3>&, float, TPointContactIterationParameters<float> &, TPointContactParticleParameters<float> &);
 		template void ApplyPushOut<float, 3>(TRigidBodyPointContactConstraint<float,3>&, float , const TSet<const TGeometryParticleHandle<float,3>*>&,
 			TPointContactIterationParameters<float> & IterationParameters, TPointContactParticleParameters<float> & ParticleParameters);
 
diff --git a/Engine/Source/Runtime/Experimental/Chaos/Public/Chaos/CollisionResolution.h b/Engine/Source/Runtime/Experimental/Chaos/Public/Chaos/CollisionResolution.h
new file mode 100644
--- /dev/null
+++ b/Engine/Source/Runtime/Experimental/Chaos/Public/Chaos/CollisionResolution.h
@@ -0,0 +1,168 @@
+// Copyright 1998-2019 Epic Games, Inc. All Rights Reserved.
+#pragma once
+
+#include "Chaos/CollisionResolutionTypes.h"
+#include "Chaos/CollisionResolutionUtil.h"
+
+namespace Chaos
+{
+	template <typename T, int d>
+	class TSphere;
+
+	template <typename T, int d>
+	class TBox;
+
+	template <typename T>
+	class TCapsule;
+
+	template <typename T, int d>
+	class TPlane;
+
+	class FImplicitObject;
+
+	template<typename T, int d>
+	class TRigidBodyPointContactConstraint;
+
+	template <typename T, int d>
+	class TRigidTransform;
+
+	namespace Collisions
+	{
+
+		//
+		// Box-Box
+		//
+
+		template <typename T, int d>
+		void CHAOS_API UpdateBoxBoxConstraint(const TBox<T, d>& Box1, const TRigidTransform<T, d>& Box1Transform, const TBox<T, d>& Box2, const TRigidTransform<T, d>& Box2Transform, T Thickness, TRigidBodyPointContactConstraint<T, d>& Constraint);
+
+		template<typename T, int d>
+		void CHAOS_API ConstructBoxBoxConstraints(TGeometryParticleHandle<T, d>* Particle0, TGeometryParticleHandle<T, d>* Particle1, const FImplicitObject* Implicit0, const FImplicitObject* Implicit1, const T Thickness, FCollisionConstraintsArray& NewConstraints);
+
+
+		//
+		// Box-Plane
+		//
+
+		template <typename T, int d>
+		bool CHAOS_API UpdateBoxPlaneConstraint(const TBox<T, d>& Box, const TRigidTransform<T, d>& BoxTransform, const TPlane<T, d>& Plane, const TRigidTransform<T, d>& PlaneTransform, T Thickness, TRigidBodyPointContactConstraint<T, d>& Constraint);
+
+		template<typename T, int d>
+		void CHAOS_API ConstructBoxPlaneConstraints(TGeometryParticleHandle<T, d>* Particle0, TGeometryParticleHandle<T, d>* Particle1, const FImplicitObject* Implicit0, const FImplicitObject* Implicit1, const T Thickness, FCollisionConstraintsArray& NewConstraints);
+
+
+		//
+		// Sphere-Sphere
+		//
+
+		template <typename T, int d>
+		void CHAOS_API UpdateSphereSphereConstraint(const TSphere<T, d>& Sphere1, const TRigidTransform<T, d>& Sphere1Transform, const TSphere<T, d>& Sphere2, const TRigidTransform<T, d>& Sphere2Transform, T Thickness, TRigidBodyPointContactConstraint<T, d>& Constraint);
+
+		template<typename T, int d>
+		void CHAOS_API ConstructSphereSphereConstraints(TGeometryParticleHandle<T, d>* Particle0, TGeometryParticleHandle<T, d>* Particle1, const FImplicitObject* Implicit0, const FImplicitObject* Implicit1, const T Thickness, FCollisionConstraintsArray& NewConstraints);
+
+		//
+		// Sphere-Plane
+		//
+
+		template <typename T, int d>
+		void CHAOS_API UpdateSpherePlaneConstraint(const TSphere<T, d>& Sphere, const TRigidTransform<T, d>& SphereTransform, const TPlane<T, d>& Plane, const TRigidTransform<T, d>& PlaneTransform, T Thickness, TRigidBodyPointContactConstraint<T, d>& Constraint);
+
+		template<typename T, int d>
+		void CHAOS_API ConstructSpherePlaneConstraints(TGeometryParticleHandle<T, d>* Particle0, TGeometryParticleHandle<T, d>* Particle1, const FImplicitObject* Implicit0, const FImplicitObject* Implicit1, const T Thickness, FCollisionConstraintsArray& NewConstraints);
+
+		//
+		// Sphere-Box
+		//
+
+		template <typename T, int d>
+		void CHAOS_API UpdateSphereBoxConstraint(const TSphere<T, d>& Sphere, const TRigidTransform<T, d>& SphereTransform, const TBox<T, d>& Box, const TRigidTransform<T, d>& BoxTransform, T Thickness, TRigidBodyPointContactConstraint<T, d>& Constraint);
+
+		template<typename T, int d>
+		void CHAOS_API ConstructSphereBoxConstraints(TGeometryParticleHandle<T, d>* Particle0, TGeometryParticleHandle<T, d>* Particle1, const FImplicitObject* Implicit0, const FImplicitObject* Implicit1, const T Thickness, FCollisionConstraintsArray& NewConstraints);
+
+		//
+		// Capsule-Capsule
+		//
+
+		template <typename T, int d>
+		void CHAOS_API UpdateCapsuleCapsuleConstraint(const TCapsule<T>& A, const TRigidTransform<T, d>& ATransform, const TCapsule<T>& B, const TRigidTransform<T, d>& BTransform, T Thickness, TRigidBodyPointContactConstraint<T, d>& Constraint);
+
+		template<typename T, int d>
+		void CHAOS_API ConstructCapsuleCapsuleConstraints(TGeometryParticleHandle<T, d>* Particle0, TGeometryParticleHandle<T, d>* Particle1, const FImplicitObject* Implicit0, const FImplicitObject* Implicit1, const T Thickness, FCollisionConstraintsArray& NewConstraints);
+
+		//
+		// Capsule-Box
+		//
+
+		template <typename T, int d>
+		void CHAOS_API UpdateCapsuleBoxConstraint(const TCapsule<T>& A, const TRigidTransform<T, d>& ATransform, const TBox<T, d>& B, const TRigidTransform<T, d>& BTransform, T Thickness, TRigidBodyPointContactConstraint<T, d>& Constraint);
+
+		template<typename T, int d>
+		void CHAOS_API ConstructCapsuleBoxConstraints(TGeometryParticleHandle<T, d>* Particle0, TGeometryParticleHandle<T, d>* Particle1, const FImplicitObject* Implicit0, const FImplicitObject* Implicit1, const T Thickness, FCollisionConstraintsArray& NewConstraints);
+
+		//
+		// Convex-Convex
+		//
+		template <typename T, int d>
+		void CHAOS_API UpdateConvexConvexConstraint(const FImplicitObject& A, const TRigidTransform<T, d>& ATM, const FImplicitObject& B, const TRigidTransform<T, d>& BTM, T Thickness, TRigidBodyPointContactConstraint<T, d>& Constraint);
+
+		template <typename T, int d>
+		void CHAOS_API ConstructConvexConvexConstraints(TGeometryParticleHandle<T, d>* Particle0, TGeometryParticleHandle<T, d>* Particle1, const FImplicitObject* Implicit0, const FImplicitObject* Implicit1, const T Thickness, FCollisionConstraintsArray& NewConstraints);
+
+		//
+		// Levelset-Levelset
+		//
+		template<ECollisionUpdateType UpdateType, typename T = float, int d = 3>
+		void CHAOS_API UpdateLevelsetLevelsetConstraint(const T Thickness, TRigidBodyPointContactConstraint<float, 3>& Constraint);
+
+		template<typename T, int d>
+		void CHAOS_API ConstructLevelsetLevelsetConstraints(TGeometryParticleHandle<T, d>* Particle0, TGeometryParticleHandle<T, d>* Particle1, const FImplicitObject* Implicit0, const FImplicitObject* Implicit1, const T Thickness, FCollisionConstraintsArray& NewConstraints);
+
+
+		//
+		//  Levelset-Union and Union-Levelset
+		//
+
+		template<ECollisionUpdateType UpdateType, typename T, int d>
+		void CHAOS_API UpdateLevelsetUnionConstraint(const T Thickness, TRigidBodyPointContactConstraint<T, d>& Constraint);
+
+		template<ECollisionUpdateType UpdateType, typename T, int d>
+		void CHAOS_API UpdateUnionLevelsetConstraint(const T Thickness, TRigidBodyPointContactConstraint<T, d>& Constraint);
+
+
+		//
+		// Union-Union
+		//
+
+		template<ECollisionUpdateType UpdateType, typename T, int d>
+		void CHAOS_API UpdateUnionUnionConstraint(const T Thickness, TRigidBodyPointContactConstraint<T, d>& Constraint);
+
+		template<typename T, int d>
+		void CHAOS_API ConstructUnionUnionConstraints(TGeometryParticleHandle<T, d>* Particle0, TGeometryParticleHandle<T, d>* Particle1, const FImplicitObject* Implicit0, const FImplicitObject* Implicit1, const T Thickness, FCollisionConstraintsArray& NewConstraints);
+
+		//
+		// Single-Union
+		//
+
+		template<ECollisionUpdateType UpdateType, typename T, int d>
+		void CHAOS_API UpdateSingleUnionConstraint(const T Thickness, TRigidBodyPointContactConstraint<T, d>& Constraint);
+
+		template<typename T, int d>
+		void CHAOS_API ConstructSingleUnionConstraints(TGeometryParticleHandle<T, d>* Particle0, TGeometryParticleHandle<T, d>* Particle1, const FImplicitObject* Implicit0, const FImplicitObject* Implicit1, const T Thickness, FCollisionConstraintsArray& NewConstraints);
+
+
+		//
+		// Constraint API
+		//
+
+		template<ECollisionUpdateType UpdateType, typename T, int d>
+		void CHAOS_API UpdateConstraintImp(const FImplicitObject& ParticleObject, const TRigidTransform<T, d>& ParticleTM, const FImplicitObject& LevelsetObject, const TRigidTransform<T, d>& LevelsetTM, const T Thickness, TRigidBodyPointContactConstraint<T, d>& Constraint);
+
+		template<typename T, int d>
+		void CHAOS_API ConstructPairConstraintImpl(TGeometryParticleHandle<T, d>* Particle0, TGeometryParticleHandle<T, d>* Particle1, const FImplicitObject* Implicit0, const FImplicitObject* Implicit1, const T Thickness, FCollisionConstraintsArray& NewConstraints);
+
+		template<typename T, int d>
+		void CHAOS_API ConstructConstraintsImpl(TGeometryParticleHandle<T, d>* Particle0, TGeometryParticleHandle<T, d>* Particle1, const FImplicitObject* Implicit0, const FImplicitObject* Implicit1, const T Thickness, FCollisionConstraintsArray& NewConstraints);
+	}
+}
diff --git a/Engine/Source/Runtime/Experimental/Chaos/Public/Chaos/CollisionResolutionAlgo.h b/Engine/Source/Runtime/Experimental/Chaos/Public/Chaos/CollisionResolutionAlgo.h
deleted file mode 100644
diff --git a/Engine/Source/Runtime/Experimental/Chaos/Public/Chaos/CollisionResolutionConvexConvex.h b/Engine/Source/Runtime/Experimental/Chaos/Public/Chaos/CollisionResolutionConvexConvex.h
deleted file mode 100644
diff --git a/Engine/Source/Runtime/Experimental/Chaos/Public/Chaos/CollisionResolutionTypes.h b/Engine/Source/Runtime/Experimental/Chaos/Public/Chaos/CollisionResolutionTypes.h
--- a/Engine/Source/Runtime/Experimental/Chaos/Public/Chaos/CollisionResolutionTypes.h
+++ b/Engine/Source/Runtime/Experimental/Chaos/Public/Chaos/CollisionResolutionTypes.h
@@ -20,6 +20,8 @@ namespace Chaos
 	template<typename T, int d>
 	class TRigidBodyPointContactConstraint;
 
+
+
 	/** Specifies the type of work we should do*/
 	enum class CHAOS_API ECollisionUpdateType
 	{
@@ -40,7 +42,7 @@ namespace Chaos
 	/*
 	*
 	*/
-	template<class T, int d>
+	template<class T=float, int d=3>
 	class CHAOS_API TCollisionConstraintBase
 	{
 	public:
@@ -56,12 +58,14 @@ namespace Chaos
 		TCollisionConstraintBase(FType InType) : Type(InType) {}
 		FType GetType() const { return Type; }
 
-		template<class AS_T> AS_T * As() { return static_cast<AS_T>(this); }
-		template<class AS_T> const AS_T * As() const { return static_cast<AS_T>(this); }
+		template<class AS_T> AS_T * As() { return static_cast<AS_T*>(this); }
+		template<class AS_T> const AS_T * As() const { return static_cast<AS_T*>(this); }
 
 	private:
 		FType Type;
 	};
+	typedef TCollisionConstraintBase<float, 3> FCollisionConstraintBase;
+	typedef TArray<TCollisionConstraintBase<float, 3>*> FCollisionConstraintsArray;
 
 
 	/*
@@ -98,7 +102,7 @@ namespace Chaos
 		using FGeometryParticleHandle = TGeometryParticleHandle<T, d>;
 		using FManifold = TPointContactManifold<T, d>;
 
-		TRigidBodyPointContactConstraint() : AccumulatedImpulse(0) {}
+		TRigidBodyPointContactConstraint() : TCollisionConstraintBase(Base::FType::SinglePoint), AccumulatedImpulse(0) {}
 		static typename Base::FType StaticType() { return Base::FType::SinglePoint; };
 
 
@@ -137,6 +141,7 @@ namespace Chaos
 		TVector<T, d> AccumulatedImpulse;
 
 	};
+	typedef TRigidBodyPointContactConstraint<float, 3> FRigidBodyPointContactConstraint;
 
 
 	/*
@@ -178,7 +183,7 @@ namespace Chaos
 		using FGeometryParticleHandle = TGeometryParticleHandle<T, d>;
 		using FManifold = TPlaneContactManifold<T, d>;
 
-		TRigidBodyPlaneContactConstraint() : AccumulatedImpulse(0) {}
+		TRigidBodyPlaneContactConstraint() : TCollisionConstraintBase(Base::FType::Plane), AccumulatedImpulse(0) {}
 		static typename Base::FType StaticType() { return Base::FType::Plane; };
 
 	
@@ -197,7 +202,7 @@ namespace Chaos
 		TVector<T, d> AccumulatedImpulse;
 
 	};
-
+	typedef TRigidBodyPlaneContactConstraint<float, 3> FRigidBodyPlaneContactConstraint;
 
 
 	template<class T, int d>
diff --git a/Engine/Source/Runtime/Experimental/Chaos/Public/Chaos/CollisionResolutionUtil.h b/Engine/Source/Runtime/Experimental/Chaos/Public/Chaos/CollisionResolutionUtil.h
new file mode 100644
--- /dev/null
+++ b/Engine/Source/Runtime/Experimental/Chaos/Public/Chaos/CollisionResolutionUtil.h
@@ -0,0 +1,63 @@
+// Copyright 1998-2019 Epic Games, Inc. All Rights Reserved.
+
+#pragma once
+#include "Chaos/CollisionResolutionTypes.h"
+#include "Chaos/Pair.h"
+
+namespace Chaos
+{
+	class FImplicitObject;
+
+	template<typename T, int d>
+	class TRigidBodyPointContactConstraint;
+
+	template <typename T, int d>
+	class TRigidTransform;
+
+	namespace Collisions
+	{
+		//
+		// Utility
+		//
+
+		template<typename T = float, int d = 3>
+		TRigidTransform<T, d>
+			GetTransform(const TGeometryParticleHandle<T, d>* Particle);
+
+		template<typename T>
+		PMatrix<T, 3, 3>
+			ComputeFactorMatrix3(const TVector<T, 3>& V, const PMatrix<T, 3, 3>& M, const T& Im);
+
+		template<typename T, int d>
+		TVector<T, d>
+			GetEnergyClampedImpulse(const TRigidBodyPointContactConstraint<T, d>& Constraint, const TVector<T, d>& Impulse, const TVector<T, d>& VectorToPoint1, const TVector<T, d>& VectorToPoint2, const TVector<T, d>& Velocity1, const TVector<T, d>& Velocity2);
+
+		template <typename T, int d>
+		bool
+			SampleObjectHelper(const FImplicitObject& Object, const TRigidTransform<T, d>& ObjectTransform, const TRigidTransform<T, d>& SampleToObjectTransform, const TVector<T, d>& SampleParticle, T Thickness, TRigidBodyPointContactConstraint<float, 3>& Constraint);
+
+		template <typename T, int d>
+		bool
+			SampleObjectNoNormal(const FImplicitObject& Object, const TRigidTransform<T, d>& ObjectTransform, const TRigidTransform<T, d>& SampleToObjectTransform, const TVector<T, d>& SampleParticle, T Thickness, TRigidBodyPointContactConstraint<float, 3>& Constraint);
+
+		template <typename T, int d>
+		bool
+			SampleObjectNormalAverageHelper(const FImplicitObject& Object, const TRigidTransform<T, d>& ObjectTransform, const TRigidTransform<T, d>& SampleToObjectTransform, const TVector<T, d>& SampleParticle, T Thickness, T& TotalThickness, TRigidBodyPointContactConstraint<float, 3>& Constraint);
+
+#if INTEL_ISPC
+		template<ECollisionUpdateType UpdateType>
+		void
+			SampleObject(const FImplicitObject& Object, const TRigidTransform<float, 3>& ObjectTransform, const TBVHParticles<float, 3>& SampleParticles, const TRigidTransform<float, 3>& SampleParticlesTransform, float Thickness, TRigidBodyPointContactConstraint<float, 3>& Constraint);
+#else
+		template <ECollisionUpdateType UpdateType, typename T, int d>
+		void
+			SampleObject(const FImplicitObject& Object, const TRigidTransform<T, d>& ObjectTransform, const TBVHParticles<T, d>& SampleParticles, const TRigidTransform<T, d>& SampleParticlesTransform, T Thickness, TRigidBodyPointContactConstraint<float, 3>& Constraint);
+#endif
+
+		template <typename T, int d>
+		TArray<Pair<const FImplicitObject*, TRigidTransform<T, d>>>
+			FindRelevantShapes(const FImplicitObject* ParticleObj, const TRigidTransform<T, d>& ParticlesTM, const FImplicitObject& LevelsetObj, const TRigidTransform<T, d>& LevelsetTM, const T Thickness);
+
+	}// Collisions
+
+} // Chaos
diff --git a/Engine/Source/Runtime/Experimental/Chaos/Public/Chaos/PBDCollisionConstraints.h b/Engine/Source/Runtime/Experimental/Chaos/Public/Chaos/PBDCollisionConstraints.h
--- a/Engine/Source/Runtime/Experimental/Chaos/Public/Chaos/PBDCollisionConstraints.h
+++ b/Engine/Source/Runtime/Experimental/Chaos/Public/Chaos/PBDCollisionConstraints.h
@@ -62,6 +62,7 @@ public:
 	using Base = FPBDConstraintContainer;
 	using FHandles = TArray<TPBDCollisionConstraintHandle<T, d>*>;
 	using FConstraintContainerHandle = TPBDCollisionConstraintHandle<T, d>;
+	using FConstraintBase = TCollisionConstraintBase<T, d>;
  	using FPointContactConstraint = TRigidBodyPointContactConstraint<T, d>;
 	using FPlaneContactConstraint = TRigidBodyPlaneContactConstraint<T, d>;
 	using FConstraintHandleAllocator = TConstraintHandleAllocator<TPBDCollisionConstraints<T, d>>;
@@ -236,8 +237,7 @@ private:
 	/**
 	*  Build a constraint based on the two particle handles.
 	*/
-	void ConstructConstraints(TGeometryParticleHandle<T, d>* Particle0, TGeometryParticleHandle<T, d>* Particle1, 
-		const T Thickness, TRigidBodyPointContactConstraint<T, d> & Constraint);
+	void ConstructConstraints(TGeometryParticleHandle<T, d>* Particle0, TGeometryParticleHandle<T, d>* Particle1, const T Thickness, FCollisionConstraintsArray&);
 
 
 
diff --git a/Engine/Source/Runtime/Experimental/Chaos/Public/Chaos/PBDCollisionConstraintsImp.h b/Engine/Source/Runtime/Experimental/Chaos/Public/Chaos/PBDCollisionConstraintsImp.h
--- a/Engine/Source/Runtime/Experimental/Chaos/Public/Chaos/PBDCollisionConstraintsImp.h
+++ b/Engine/Source/Runtime/Experimental/Chaos/Public/Chaos/PBDCollisionConstraintsImp.h
@@ -92,7 +92,7 @@ namespace Chaos
 			void IncrementSimulatedParticles() {}
 			void RecordBoundsData(const TBox <T, d>& Box1) {}
 			void RecordBroadphasePotentials(int32 Num) {}
-			void IncrementCountNP() {}
+			void IncrementCountNP(int32 Count = 1) {}
 			void IncrementRejectedNP() {}
 			void FinalizeData() {}
 			void Print() {}
@@ -121,9 +121,9 @@ namespace Chaos
 				BroadphasePotentials.Record(Num);
 			}
 
-			void IncrementCountNP()
+			void IncrementCountNP(int32 Count = 1)
 			{
-				++CountNP;
+				CountNP += Count;
 			}
 
 			void IncrementRejectedNP()
@@ -260,7 +260,7 @@ namespace Chaos
 		if (!EnableCollisions) return;
 
 		//todo(ocohen): use per thread buffer instead, need better support than ParallelFor for this
-		TQueue<TRigidBodyPointContactConstraint<T, d>, EQueueMode::Mpsc> Queue;
+		TQueue<FCollisionConstraintBase*, EQueueMode::Mpsc> Queue;
 
 		Particles.GetNonDisabledDynamicView().ParallelFor([&](auto& Particle1, int32 ActiveIdxIdx)
 		{
@@ -302,7 +302,8 @@ namespace Chaos
 			{
 				auto& Particle2 = *PotentialIntersections[i].GetGeometryParticleHandle_PhysicsThread();
 				const TGenericParticleHandle<T, d> Particle2Generic(&Particle2);
-				// Collision group culling...
+
+				// Broad Phase Culling 
 				// CollisionGroup == 0 : Collide_With_Everything
 				// CollisionGroup == INDEX_NONE : Disabled collisions
 				// CollisionGroup_A != CollisionGroup_B : Skip Check
@@ -330,57 +331,48 @@ namespace Chaos
 
 				if (bBody1Bounded == bBody2Bounded && Particle2.AsDynamic())
 				{
-#if CHAOS_DETERMINISTIC
-					//if both are dynamic, assume index order matters
+					//no bidirectional constraints. 
 					if (Particle2.ParticleID() > Particle1.ParticleID())
-#else
-					//not deterministic just use memory address to avoid pair duplication
-					if (Particle2.Handle() > Particle1.Handle())
-#endif
 					{
 						continue;
 					}
 				}
-				//				++TotalTested;
 
 				const TVector<T, d> Box2Thickness = Particle2.AsDynamic() ? ComputeThickness(*Particle2.AsDynamic(), Dt) : TVector<T, d>(0);
 				const T UseThickness = FMath::Max(Box1Thickness, Box2Thickness.Size());// + MThickness
 
-
-				TRigidBodyPointContactConstraint<T, d> Constraint;
-				ConstructConstraints(Particle1.Handle(), Particle2.Handle(), UseThickness, Constraint);
-
-				COLLISION_OPT_OUT(StatData.IncrementCountNP());
-
-				if (Constraint.ContainsManifold())
+				FCollisionConstraintsArray Constraints;
+				ConstructConstraints(Particle1.Handle(), Particle2.Handle(), UseThickness, Constraints);
+				for (TCollisionConstraintBase<T, d>* ConstraintBase : Constraints)
 				{
-					UpdateConstraint<ECollisionUpdateType::Any>(UseThickness, Constraint);
+					Queue.Enqueue(ConstraintBase);
 				}
 
-				if (Constraint.GetPhi() < UseThickness)
-				{
-					Queue.Enqueue(Constraint);
-				}
-				else
-				{
-					COLLISION_OPT_OUT(StatData.IncrementRejectedNP());
-				}
+				COLLISION_OPT_OUT(if (Constraints.Num()) { StatData.IncrementCountNP(Constraints.Num()); });
+				COLLISION_OPT_OUT(if (!Constraints.Num()) { StatData.IncrementRejectedNP(); });
 			}
 
 			COLLISION_OPT_OUT(StatData.FinalizeData());
 
-		}, bGatherStats || bDisableCollisionParallelFor);
+		}, bGatherStats || bDisableCollisionParallelFor || true); // todo(brice) remove true.
 
 		{
 			SCOPE_CYCLE_COUNTER(STAT_ComputeConstraintsSU);
 
-			while (!Queue.IsEmpty())
+			
+			FCollisionConstraintBase* ConstraintBase = nullptr;
+			while ( Queue.Dequeue(ConstraintBase) )
 			{
-				int32 Idx = PointConstraints.AddUninitialized(1);
-				FConstraintContainerHandle* Handle = HandleAllocator.template AllocHandle< TRigidBodyPointContactConstraint<T, d> >(this, Idx);
+				if (ConstraintBase->GetType() == TRigidBodyPointContactConstraint<T, 3>::StaticType())
+				{
+					TRigidBodyPointContactConstraint<T, d>* PointConstraint = ConstraintBase->As< TRigidBodyPointContactConstraint<T, d> >();
+
+					int32 Idx = PointConstraints.AddUninitialized(1);
+					PointConstraints[Idx] = *PointConstraint;
+					Handles.Add(HandleAllocator.template AllocHandle< TRigidBodyPointContactConstraint<T, d> >(this, Idx));
 
-				Handles.Add(Handle);
-				Queue.Dequeue(PointConstraints[Idx]);
+					delete PointConstraint;
+				}
 			}
 			LifespanCounter++;
 		}
diff --git a/Engine/Source/Runtime/Experimental/Chaos/Public/Chaos/PBDCollisionConstraintsPlaneContactUtil.h b/Engine/Source/Runtime/Experimental/Chaos/Public/Chaos/PBDCollisionConstraintsPlaneContactUtil.h
--- a/Engine/Source/Runtime/Experimental/Chaos/Public/Chaos/PBDCollisionConstraintsPlaneContactUtil.h
+++ b/Engine/Source/Runtime/Experimental/Chaos/Public/Chaos/PBDCollisionConstraintsPlaneContactUtil.h
@@ -30,9 +30,6 @@ namespace Chaos
 		template<ECollisionUpdateType UpdateType, typename T = float, int d = 3>
 		void Update(const T Thickness, TRigidBodyPointContactConstraint<T, d>& Constraint);
 
-		template<ECollisionUpdateType UpdateType, typename T=float, int d=3>
-		void UpdateConstraint(const T Thickness, TRigidBodyPointContactConstraint<T, d>& Constraint);
-			
 		template<typename T = float, int d = 3>
 		void Apply(TRigidBodyPointContactConstraint<T, d>& Constraint, T Thickness, 
 			TPointContactIterationParameters<T> & IterationParameters,
diff --git a/Engine/Source/Runtime/Experimental/Chaos/Public/Chaos/PBDCollisionConstraintsPointContactUtil.h b/Engine/Source/Runtime/Experimental/Chaos/Public/Chaos/PBDCollisionConstraintsPointContactUtil.h
--- a/Engine/Source/Runtime/Experimental/Chaos/Public/Chaos/PBDCollisionConstraintsPointContactUtil.h
+++ b/Engine/Source/Runtime/Experimental/Chaos/Public/Chaos/PBDCollisionConstraintsPointContactUtil.h
@@ -29,9 +29,6 @@ namespace Chaos
 		template<ECollisionUpdateType UpdateType, typename T = float, int d = 3>
 		void Update(const T Thickness, TRigidBodyPointContactConstraint<T, d>& Constraint);
 
-		template<ECollisionUpdateType UpdateType, typename T=float, int d=3>
-		void UpdateConstraint(const T Thickness, TRigidBodyPointContactConstraint<T, d>& Constraint);
-			
 		template<typename T = float, int d = 3>
 		void Apply(TRigidBodyPointContactConstraint<T, d>& Constraint, T Thickness, 
 			TPointContactIterationParameters<T> & IterationParameters,
