Fix CIS

diff --git a/Engine/Source/Runtime/Experimental/Chaos/Public/Chaos/CollisionResolutionTypes.h b/Engine/Source/Runtime/Experimental/Chaos/Public/Chaos/CollisionResolutionTypes.h
--- a/Engine/Source/Runtime/Experimental/Chaos/Public/Chaos/CollisionResolutionTypes.h
+++ b/Engine/Source/Runtime/Experimental/Chaos/Public/Chaos/CollisionResolutionTypes.h
@@ -102,7 +102,7 @@ namespace Chaos
 		using FGeometryParticleHandle = TGeometryParticleHandle<T, d>;
 		using FManifold = TPointContactManifold<T, d>;
 
-		TRigidBodyPointContactConstraint() : TCollisionConstraintBase(Base::FType::SinglePoint), AccumulatedImpulse(0) {}
+		TRigidBodyPointContactConstraint() : Base(Base::FType::SinglePoint), AccumulatedImpulse(0) {}
 		static typename Base::FType StaticType() { return Base::FType::SinglePoint; };
 
 
@@ -183,7 +183,7 @@ namespace Chaos
 		using FGeometryParticleHandle = TGeometryParticleHandle<T, d>;
 		using FManifold = TPlaneContactManifold<T, d>;
 
-		TRigidBodyPlaneContactConstraint() : TCollisionConstraintBase(Base::FType::Plane), AccumulatedImpulse(0) {}
+		TRigidBodyPlaneContactConstraint() : Base(Base::FType::Plane), AccumulatedImpulse(0) {}
 		static typename Base::FType StaticType() { return Base::FType::Plane; };
 
 	
