Add ObjectState comparison anywhere AsDynamic was called.
Rename Particle downcasting functions.
Fix PhysicsInterface not doing things it should on kinematic RigidParticles.

diff --git a/Engine/Source/Runtime/Engine/Private/PhysicsEngine/Experimental/PhysInterface_Chaos.cpp b/Engine/Source/Runtime/Engine/Private/PhysicsEngine/Experimental/PhysInterface_Chaos.cpp
--- a/Engine/Source/Runtime/Engine/Private/PhysicsEngine/Experimental/PhysInterface_Chaos.cpp
+++ b/Engine/Source/Runtime/Engine/Private/PhysicsEngine/Experimental/PhysInterface_Chaos.cpp
@@ -341,7 +341,7 @@ bool FPhysInterface_Chaos::CanSimulate_AssumesLocked(const FPhysicsActorHandle&
 
 float FPhysInterface_Chaos::GetMass_AssumesLocked(const FPhysicsActorHandle& InActorReference)
 {
-	if (const Chaos::TPBDRigidParticle<float,3>* RigidParticle = InActorReference->AsDynamic())
+	if (const Chaos::TPBDRigidParticle<float,3>* RigidParticle = InActorReference->CastToRigidParticle())
 	{
 		return RigidParticle->M();
 	}
@@ -356,7 +356,8 @@ void FPhysInterface_Chaos::SetSendsSleepNotifies_AssumesLocked(const FPhysicsAct
 
 void FPhysInterface_Chaos::PutToSleep_AssumesLocked(const FPhysicsActorHandle& InActorReference)
 {
-	if (Chaos::TPBDRigidParticle<float, 3>* Particle = InActorReference->AsDynamic())
+	Chaos::TPBDRigidParticle<float, 3>* Particle = InActorReference->CastToRigidParticle();
+	if(Particle && Particle->ObjectState() == Chaos::EObjectStateType::Dynamic)
 	{
 		Particle->SetObjectState(Chaos::EObjectStateType::Sleeping);
 	}
@@ -365,7 +366,8 @@ void FPhysInterface_Chaos::PutToSleep_AssumesLocked(const FPhysicsActorHandle& I
 
 void FPhysInterface_Chaos::WakeUp_AssumesLocked(const FPhysicsActorHandle& InActorReference)
 {
-	if (Chaos::TPBDRigidParticle<float, 3>* Particle = InActorReference->AsDynamic())
+	Chaos::TPBDRigidParticle<float, 3>* Particle = InActorReference->CastToRigidParticle();
+	if(Particle && Particle->ObjectState() == Chaos::EObjectStateType::Sleeping)
 	{
 		Particle->SetObjectState(Chaos::EObjectStateType::Dynamic);
 	}
@@ -373,7 +375,7 @@ void FPhysInterface_Chaos::WakeUp_AssumesLocked(const FPhysicsActorHandle& InAct
 
 void FPhysInterface_Chaos::SetIsKinematic_AssumesLocked(const FPhysicsActorHandle& InActorReference, bool bIsKinematic)
 {
-	if (Chaos::TPBDRigidParticle<float, 3>* Particle = InActorReference->AsDynamic())
+	if (Chaos::TPBDRigidParticle<float, 3>* Particle = InActorReference->CastToRigidParticle())
 	{
 		//#todo: what if object state has been previously set to sleeping?
 		const Chaos::EObjectStateType NewState
@@ -444,7 +446,7 @@ FVector FPhysInterface_Chaos::GetLinearVelocity_AssumesLocked(const FPhysicsActo
 {
 	if (ensure(FPhysicsInterface::IsValid(InActorReference)))
 	{
-		Chaos::TKinematicGeometryParticle<float, 3>* Kinematic = InActorReference->AsKinematic();
+		Chaos::TKinematicGeometryParticle<float, 3>* Kinematic = InActorReference->CastToKinematicParticle();
 		if (ensure(Kinematic))
 		{
 			return Kinematic->V();
@@ -465,7 +467,7 @@ void FPhysInterface_Chaos::SetLinearVelocity_AssumesLocked(const FPhysicsActorHa
 
 	if (ensure(FPhysicsInterface::IsValid(InActorReference)))
 	{
-		Chaos::TKinematicGeometryParticle<float, 3>* Kinematic = InActorReference->AsKinematic();
+		Chaos::TKinematicGeometryParticle<float, 3>* Kinematic = InActorReference->CastToKinematicParticle();
 		if (ensure(Kinematic))
 		{
 			return Kinematic->SetV(InNewVelocity);
@@ -477,7 +479,7 @@ FVector FPhysInterface_Chaos::GetAngularVelocity_AssumesLocked(const FPhysicsAct
 {
 	if (ensure(FPhysicsInterface::IsValid(InActorReference)))
 	{
-		Chaos::TKinematicGeometryParticle<float, 3>* Kinematic = InActorReference->AsKinematic();
+		Chaos::TKinematicGeometryParticle<float, 3>* Kinematic = InActorReference->CastToKinematicParticle();
 		if (ensure(Kinematic))
 		{
 			return Kinematic->W();
@@ -494,7 +496,7 @@ void FPhysInterface_Chaos::SetAngularVelocity_AssumesLocked(const FPhysicsActorH
 
 	if (ensure(FPhysicsInterface::IsValid(InActorReference)))
 	{
-		Chaos::TKinematicGeometryParticle<float, 3>* Kinematic = InActorReference->AsKinematic();
+		Chaos::TKinematicGeometryParticle<float, 3>* Kinematic = InActorReference->CastToKinematicParticle();
 		if (ensure(Kinematic))
 		{
 			return Kinematic->SetW(InNewAngularVelocity);
@@ -528,7 +530,7 @@ FVector FPhysInterface_Chaos::GetWorldVelocityAtPoint_AssumesLocked(const FPhysi
 {
 	if (ensure(FPhysicsInterface::IsValid(InActorReference)))
 	{
-		Chaos::TKinematicGeometryParticle<float, 3>* Kinematic = InActorReference->AsKinematic();
+		Chaos::TKinematicGeometryParticle<float, 3>* Kinematic = InActorReference->CastToKinematicParticle();
 		if (ensure(Kinematic))
 		{
 			const Chaos::FVec3 COM = Kinematic->X(); // TODO: Fix once we have separate COM
@@ -557,7 +559,7 @@ FTransform FPhysInterface_Chaos::GetComTransformLocal_AssumesLocked(const FPhysi
 
 FVector FPhysInterface_Chaos::GetLocalInertiaTensor_AssumesLocked(const FPhysicsActorHandle& InActorReference)
 {
-	if (Chaos::TPBDRigidParticle<float, 3 >* RigidParticle = InActorReference->AsDynamic())
+	if (Chaos::TPBDRigidParticle<float, 3 >* RigidParticle = InActorReference->CastToRigidParticle())
 	{
 		const Chaos::PMatrix<float, 3, 3> & Tensor = RigidParticle->I();
 		return FVector(Tensor.M[0][0], Tensor.M[1][1], Tensor.M[2][2]) ;
@@ -624,7 +626,7 @@ void FPhysInterface_Chaos::AddRadialImpulse_AssumesLocked(const FPhysicsActorHan
 
 bool FPhysInterface_Chaos::IsGravityEnabled_AssumesLocked(const FPhysicsActorHandle& InActorReference)
 {
-	if (Chaos::TPBDRigidParticle<float, 3 >* RigidParticle = InActorReference->AsDynamic())
+	if (Chaos::TPBDRigidParticle<float, 3 >* RigidParticle = InActorReference->CastToRigidParticle())
 	{
 		return RigidParticle->IsGravityEnabled();
 	}
@@ -632,7 +634,7 @@ bool FPhysInterface_Chaos::IsGravityEnabled_AssumesLocked(const FPhysicsActorHan
 }
 void FPhysInterface_Chaos::SetGravityEnabled_AssumesLocked(const FPhysicsActorHandle& InActorReference, bool bEnabled)
 {
-	if (Chaos::TPBDRigidParticle<float, 3 >* RigidParticle = InActorReference->AsDynamic())
+	if (Chaos::TPBDRigidParticle<float, 3 >* RigidParticle = InActorReference->CastToRigidParticle())
 	{
 		RigidParticle->SetGravityEnabled(bEnabled);
 		FPhysicsCommand::ExecuteWrite(InActorReference, [&](const FPhysicsActorHandle& Actor)
@@ -655,7 +657,7 @@ void FPhysInterface_Chaos::SetSleepEnergyThreshold_AssumesLocked(const FPhysicsA
 
 void FPhysInterface_Chaos::SetMass_AssumesLocked(FPhysicsActorHandle& InActorReference, float InMass)
 {
-	if (Chaos::TPBDRigidParticle<float, 3 >* RigidParticle = InActorReference->AsDynamic())
+	if (Chaos::TPBDRigidParticle<float, 3 >* RigidParticle = InActorReference->CastToRigidParticle())
 	{
 		RigidParticle->SetM(InMass);
 		if (CHAOS_ENSURE(!FMath::IsNearlyZero(InMass)))
@@ -671,7 +673,7 @@ void FPhysInterface_Chaos::SetMass_AssumesLocked(FPhysicsActorHandle& InActorRef
 
 void FPhysInterface_Chaos::SetMassSpaceInertiaTensor_AssumesLocked(FPhysicsActorHandle& InActorReference, const FVector& InTensor)
 {
-	if (Chaos::TPBDRigidParticle<float, 3 >* RigidParticle = InActorReference->AsDynamic())
+	if (Chaos::TPBDRigidParticle<float, 3 >* RigidParticle = InActorReference->CastToRigidParticle())
 	{
 		if(CHAOS_ENSURE(!FMath::IsNearlyZero(InTensor.X)) && CHAOS_ENSURE(!FMath::IsNearlyZero(InTensor.Y)) && CHAOS_ENSURE(!FMath::IsNearlyZero(InTensor.Z)) )
 		{
diff --git a/Engine/Source/Runtime/Engine/Private/PhysicsEngine/Experimental/PhysScene_Chaos.cpp b/Engine/Source/Runtime/Engine/Private/PhysicsEngine/Experimental/PhysScene_Chaos.cpp
--- a/Engine/Source/Runtime/Engine/Private/PhysicsEngine/Experimental/PhysScene_Chaos.cpp
+++ b/Engine/Source/Runtime/Engine/Private/PhysicsEngine/Experimental/PhysScene_Chaos.cpp
@@ -1058,21 +1058,28 @@ void FPhysScene_ChaosInterface::AddCustomPhysics_AssumesLocked(FBodyInstance* Bo
 
 void FPhysScene_ChaosInterface::AddForce_AssumesLocked(FBodyInstance* BodyInstance, const FVector& Force, bool bAllowSubstepping, bool bAccelChange)
 {
+	using namespace Chaos;
+
 	FPhysicsActorHandle& Handle = BodyInstance->GetPhysicsActorHandle();
 	if (FPhysicsInterface::IsValid(Handle))
 	{
-		if (Chaos::TPBDRigidParticle<float, 3>* Rigid = Handle->AsDynamic())
+		Chaos::TPBDRigidParticle<float, 3>* Rigid = Handle->CastToRigidParticle();
+		if(Rigid)
 		{
-			const Chaos::TVector<float, 3> CurrentForce = Rigid->ExternalForce();
-			if (bAccelChange)
-			{
-				const float Mass = Rigid->M();
-				const Chaos::TVector<float, 3> TotalAcceleration = CurrentForce + (Force * Mass);
-				Rigid->SetExternalForce(TotalAcceleration);
-			}
-			else
+			EObjectStateType ObjectState = Rigid->ObjectState();
+			if (CHAOS_ENSURE(ObjectState == EObjectStateType::Dynamic || ObjectState == EObjectStateType::Sleeping))
 			{
-				Rigid->SetExternalForce(CurrentForce + Force);
+				const Chaos::TVector<float, 3> CurrentForce = Rigid->ExternalForce();
+				if (bAccelChange)
+				{
+					const float Mass = Rigid->M();
+					const Chaos::TVector<float, 3> TotalAcceleration = CurrentForce + (Force * Mass);
+					Rigid->SetExternalForce(TotalAcceleration);
+				}
+				else
+				{
+					Rigid->SetExternalForce(CurrentForce + Force);
+				}
 			}
 		}
 	}
@@ -1080,31 +1087,38 @@ void FPhysScene_ChaosInterface::AddForce_AssumesLocked(FBodyInstance* BodyInstan
 
 void FPhysScene_ChaosInterface::AddForceAtPosition_AssumesLocked(FBodyInstance* BodyInstance, const FVector& Force, const FVector& Position, bool bAllowSubstepping, bool bIsLocalForce /*= false*/)
 {
+	using namespace Chaos;
+
 	FPhysicsActorHandle& Handle = BodyInstance->GetPhysicsActorHandle();
 	if (ensure(FPhysicsInterface::IsValid(Handle)))
 	{
-		Chaos::TPBDRigidParticle<float, 3>* Rigid = Handle->AsDynamic();
+		Chaos::TPBDRigidParticle<float, 3>* Rigid = Handle->CastToRigidParticle();
+		
 		if (ensure(Rigid))
 		{
-			const Chaos::FVec3& CurrentForce = Rigid->ExternalForce();
-			const Chaos::FVec3& CurrentTorque = Rigid->ExternalTorque();
-			const Chaos::FVec3& CurrentPosition = Rigid->X();
-			const Chaos::FVec3& CurrentCOM = CurrentPosition; // TODO: Remember to fix this once we add COM to particles!
-			const Chaos::FVec3 Torque = Chaos::FVec3::CrossProduct(Position - CurrentCOM, Force);
-			if (!bIsLocalForce)
-			{
-				Rigid->SetExternalForce(CurrentForce + Force);
-				Rigid->SetExternalTorque(CurrentTorque + Torque);
-			}
-			else
+			EObjectStateType ObjectState = Rigid->ObjectState();
+			if (CHAOS_ENSURE(ObjectState == EObjectStateType::Dynamic || ObjectState == EObjectStateType::Sleeping))
 			{
-				const Chaos::FRotation3& CurrentRotation = Rigid->R();
-				const Chaos::FRigidTransform3 CurrentTransform(CurrentPosition, CurrentRotation);
-				const Chaos::FVec3 WorldPosition = CurrentTransform.TransformPosition(Position);
-				const Chaos::FVec3 WorldForce = CurrentRotation.RotateVector(Force);
-				const Chaos::FVec3 WorldTorque = Chaos::FVec3::CrossProduct(WorldPosition - CurrentCOM, WorldForce);
-				Rigid->SetExternalForce(CurrentForce + WorldForce);
-				Rigid->SetExternalTorque(CurrentTorque + WorldTorque);
+				const Chaos::FVec3& CurrentForce = Rigid->ExternalForce();
+				const Chaos::FVec3& CurrentTorque = Rigid->ExternalTorque();
+				const Chaos::FVec3& CurrentPosition = Rigid->X();
+				const Chaos::FVec3& CurrentCOM = CurrentPosition; // TODO: Remember to fix this once we add COM to particles!
+				const Chaos::FVec3 Torque = Chaos::FVec3::CrossProduct(Position - CurrentCOM, Force);
+				if (!bIsLocalForce)
+				{
+					Rigid->SetExternalForce(CurrentForce + Force);
+					Rigid->SetExternalTorque(CurrentTorque + Torque);
+				}
+				else
+				{
+					const Chaos::FRotation3& CurrentRotation = Rigid->R();
+					const Chaos::FRigidTransform3 CurrentTransform(CurrentPosition, CurrentRotation);
+					const Chaos::FVec3 WorldPosition = CurrentTransform.TransformPosition(Position);
+					const Chaos::FVec3 WorldForce = CurrentRotation.RotateVector(Force);
+					const Chaos::FVec3 WorldTorque = Chaos::FVec3::CrossProduct(WorldPosition - CurrentCOM, WorldForce);
+					Rigid->SetExternalForce(CurrentForce + WorldForce);
+					Rigid->SetExternalTorque(CurrentTorque + WorldTorque);
+				}
 			}
 		}
 	}
@@ -1147,20 +1161,27 @@ void FPhysScene_ChaosInterface::ClearForces_AssumesLocked(FBodyInstance* BodyIns
 
 void FPhysScene_ChaosInterface::AddTorque_AssumesLocked(FBodyInstance* BodyInstance, const FVector& Torque, bool bAllowSubstepping, bool bAccelChange)
 {
+	using namespace Chaos;
+
 	FPhysicsActorHandle& Handle = BodyInstance->GetPhysicsActorHandle();
 	if (ensure(FPhysicsInterface::IsValid(Handle)))
 	{
-		Chaos::TPBDRigidParticle<float, 3>* Rigid = Handle->AsDynamic();
+		Chaos::TPBDRigidParticle<float, 3>* Rigid = Handle->CastToRigidParticle();
+		
 		if (ensure(Rigid))
 		{
-			const Chaos::TVector<float, 3> CurrentTorque = Rigid->ExternalTorque();
-			if (bAccelChange)
+			EObjectStateType ObjectState = Rigid->ObjectState();
+			if (CHAOS_ENSURE(ObjectState == EObjectStateType::Dynamic || ObjectState == EObjectStateType::Sleeping))
 			{
-				Rigid->SetExternalTorque(CurrentTorque + (Rigid->I() * Torque));
-			}
-			else
-			{
-				Rigid->SetExternalTorque(CurrentTorque + Torque);
+				const Chaos::TVector<float, 3> CurrentTorque = Rigid->ExternalTorque();
+				if (bAccelChange)
+				{
+					Rigid->SetExternalTorque(CurrentTorque + (Rigid->I() * Torque));
+				}
+				else
+				{
+					Rigid->SetExternalTorque(CurrentTorque + Torque);
+				}
 			}
 		}
 	}
diff --git a/Engine/Source/Runtime/Engine/Private/PhysicsEngine/ImmediatePhysics/ImmediatePhysicsChaos/ImmediatePhysicsActorHandle_Chaos.cpp b/Engine/Source/Runtime/Engine/Private/PhysicsEngine/ImmediatePhysics/ImmediatePhysicsChaos/ImmediatePhysicsActorHandle_Chaos.cpp
--- a/Engine/Source/Runtime/Engine/Private/PhysicsEngine/ImmediatePhysics/ImmediatePhysicsChaos/ImmediatePhysicsActorHandle_Chaos.cpp
+++ b/Engine/Source/Runtime/Engine/Private/PhysicsEngine/ImmediatePhysics/ImmediatePhysicsChaos/ImmediatePhysicsActorHandle_Chaos.cpp
@@ -239,13 +239,14 @@ namespace ImmediatePhysics_Chaos
 
 				ParticleHandle->SetGeometry(MakeSerializable(Geometry));
 
-				if (auto* Kinematic = ParticleHandle->AsKinematic())
+				if (auto* Kinematic = ParticleHandle->CastToKinematicParticle())
 				{
 					Kinematic->SetV(FVector::ZeroVector);
 					Kinematic->SetW(FVector::ZeroVector);
 				}
 
-				if (auto* Dynamic = ParticleHandle->AsDynamic())
+				auto* Dynamic = ParticleHandle->CastToRigidParticle();
+				if (Dynamic && Dynamic->ObjectState() == EObjectStateType::Dynamic)
 				{
 					float MassInv = (Mass > 0.0f) ? 1.0f / Mass : 0.0f;
 					FVector InertiaInv = (Mass > 0.0f) ? Inertia.Reciprocal() : FVector::ZeroVector;
@@ -286,7 +287,8 @@ namespace ImmediatePhysics_Chaos
 
 	void FActorHandle::SetEnabled(bool bEnabled)
 	{
-		if (auto* Dynamic = ParticleHandle->AsDynamic())
+		auto* Dynamic = ParticleHandle->CastToRigidParticle();
+		if(Dynamic && Dynamic->ObjectState() == Chaos::EObjectStateType::Dynamic)
 		{
 			Dynamic->Disabled() = !bEnabled;
 		}
@@ -299,7 +301,8 @@ namespace ImmediatePhysics_Chaos
 		ParticleHandle->SetX(ParticleTransform.GetTranslation());
 		ParticleHandle->SetR(ParticleTransform.GetRotation());
 
-		if (auto* Dynamic = ParticleHandle->AsDynamic())
+		auto* Dynamic = ParticleHandle->CastToRigidParticle();
+		if(Dynamic && Dynamic->ObjectState() == Chaos::EObjectStateType::Dynamic)
 		{
 			Dynamic->SetP(Dynamic->X());
 			Dynamic->SetQ(Dynamic->R());
@@ -320,14 +323,14 @@ namespace ImmediatePhysics_Chaos
 
 	const FKinematicTarget& FActorHandle::GetKinematicTarget() const
 	{
-		check(ParticleHandle->AsKinematic());
-		return ParticleHandle->AsKinematic()->KinematicTarget();
+		check(ParticleHandle->CastToKinematicParticle());
+		return ParticleHandle->CastToKinematicParticle()->KinematicTarget();
 	}
 
 	FKinematicTarget& FActorHandle::GetKinematicTarget()
 	{
-		check(ParticleHandle->AsKinematic());
-		return ParticleHandle->AsKinematic()->KinematicTarget();
+		check(ParticleHandle->CastToKinematicParticle());
+		return ParticleHandle->CastToKinematicParticle()->KinematicTarget();
 	}
 
 	void FActorHandle::SetKinematicTarget(const FTransform& WorldTM)
@@ -351,7 +354,7 @@ namespace ImmediatePhysics_Chaos
 
 	bool FActorHandle::IsSimulated() const
 	{
-		return ParticleHandle->AsDynamic() != nullptr;
+		return ParticleHandle->CastToRigidParticle() != nullptr && ParticleHandle->ObjectState() == Chaos::EObjectStateType::Dynamic;
 	}
 
 	FTransform FActorHandle::GetWorldTransform() const
@@ -364,7 +367,7 @@ namespace ImmediatePhysics_Chaos
 	{
 		using namespace Chaos;
 
-		if (TKinematicGeometryParticleHandle<FReal, Dimensions>* KinematicParticleHandle = ParticleHandle->AsKinematic())
+		if (TKinematicGeometryParticleHandle<FReal, Dimensions>* KinematicParticleHandle = ParticleHandle->CastToKinematicParticle())
 		{
 			KinematicParticleHandle->SetV(NewLinearVelocity);
 		}
@@ -379,7 +382,7 @@ namespace ImmediatePhysics_Chaos
 	{
 		using namespace Chaos;
 
-		if (TKinematicGeometryParticleHandle<FReal, Dimensions>* KinematicParticleHandle = ParticleHandle->AsKinematic())
+		if (TKinematicGeometryParticleHandle<FReal, Dimensions>* KinematicParticleHandle = ParticleHandle->CastToKinematicParticle())
 		{
 			KinematicParticleHandle->SetW(NewAngularVelocity);
 		}
@@ -458,7 +461,8 @@ namespace ImmediatePhysics_Chaos
 	{
 		using namespace Chaos;
 
-		if (TPBDRigidParticleHandle<FReal, Dimensions>* Dynamic = ParticleHandle->AsDynamic())
+		TPBDRigidParticleHandle<FReal, Dimensions>* Dynamic = ParticleHandle->CastToRigidParticle();
+		if(Dynamic && Dynamic->ObjectState() == EObjectStateType::Dynamic)
 		{
 			float NewMass = (NewInverseMass > SMALL_NUMBER) ? 1.0f / NewInverseMass : 0.0f;
 			Dynamic->SetM(NewMass);
@@ -475,7 +479,8 @@ namespace ImmediatePhysics_Chaos
 	{
 		using namespace Chaos;
 
-		if (TPBDRigidParticleHandle<FReal, Dimensions>* Dynamic = ParticleHandle->AsDynamic())
+		TPBDRigidParticleHandle<FReal, Dimensions>* Dynamic = ParticleHandle->CastToRigidParticle();
+		if(Dynamic && Dynamic->ObjectState() == EObjectStateType::Dynamic)
 		{
 			FVector NewInertia = FVector::ZeroVector;
 			if ((NewInverseInertia.X > SMALL_NUMBER) && (NewInverseInertia.Y > SMALL_NUMBER) && (NewInverseInertia.Z > SMALL_NUMBER))
diff --git a/Engine/Source/Runtime/Engine/Private/PhysicsEngine/ImmediatePhysics/ImmediatePhysicsChaos/ImmediatePhysicsSimulation_Chaos.cpp b/Engine/Source/Runtime/Engine/Private/PhysicsEngine/ImmediatePhysics/ImmediatePhysicsChaos/ImmediatePhysicsSimulation_Chaos.cpp
--- a/Engine/Source/Runtime/Engine/Private/PhysicsEngine/ImmediatePhysics/ImmediatePhysicsChaos/ImmediatePhysicsSimulation_Chaos.cpp
+++ b/Engine/Source/Runtime/Engine/Private/PhysicsEngine/ImmediatePhysics/ImmediatePhysicsChaos/ImmediatePhysicsSimulation_Chaos.cpp
@@ -357,8 +357,8 @@ namespace ImmediatePhysics_Chaos
 
 		for (FActorHandle* ActorHandle : InIgnoreCollisionActors)
 		{
-			TPBDRigidParticleHandle<FReal, Dimensions>* Particle = ActorHandle->GetParticle()->AsDynamic();
-			if (Particle != nullptr)
+			TPBDRigidParticleHandle<FReal, Dimensions>* Particle = ActorHandle->GetParticle()->CastToRigidParticle();
+			if (Particle != nullptr && Particle->ObjectState() == EObjectStateType::Dynamic)
 			{
 				Particle->SetCollisionGroup(INDEX_NONE);
 			}
@@ -382,7 +382,7 @@ namespace ImmediatePhysics_Chaos
 			ActorHandle->SetLevel(INDEX_NONE);
 			ActorJoints.Emplace(ActorHandle);
 
-			if (!ActorHandle->GetParticle()->AsDynamic())
+			if (ActorHandle->GetParticle()->ObjectState() != EObjectStateType::Dynamic)
 			{
 				ActorHandle->SetLevel(0);
 				ActorQueue.Add(ActorHandle);
diff --git a/Engine/Source/Runtime/Experimental/Chaos/Private/Chaos/ChaosDebugDraw.cpp b/Engine/Source/Runtime/Experimental/Chaos/Private/Chaos/ChaosDebugDraw.cpp
--- a/Engine/Source/Runtime/Experimental/Chaos/Private/Chaos/ChaosDebugDraw.cpp
+++ b/Engine/Source/Runtime/Experimental/Chaos/Private/Chaos/ChaosDebugDraw.cpp
@@ -112,10 +112,10 @@ namespace Chaos
 
 		void DrawParticleShapesImpl(const FRigidTransform3& SpaceTransform, const TGeometryParticleHandle<FReal, 3>* Particle, FReal ColorScale)
 		{
-			FColor ShapeColor = Particle->AsDynamic() ? FColor::Yellow : FColor::Orange;
+			FColor ShapeColor = Particle->ObjectState() == EObjectStateType::Dynamic ? FColor::Yellow : FColor::Orange;
 			FColor Color = (((FReal)0.5 * ColorScale) * ShapeColor).ToFColor(false);
-			FVec3 P = SpaceTransform.TransformPosition(Particle->AsDynamic() ? Particle->AsDynamic()->P() : Particle->X());
-			FRotation3 Q = SpaceTransform.GetRotation() * (Particle->AsDynamic() ? Particle->AsDynamic()->Q() : Particle->R());
+			FVec3 P = SpaceTransform.TransformPosition(Particle->ObjectState() == EObjectStateType::Dynamic ? Particle->CastToRigidParticle()->P() : Particle->X());
+			FRotation3 Q = SpaceTransform.GetRotation() * (Particle->ObjectState() == EObjectStateType::Dynamic ? Particle->CastToRigidParticle()->Q() : Particle->R());
 
 			DrawShapesImpl(FRigidTransform3(P, Q), Particle->Geometry().Get(), Color);
 		}
@@ -126,8 +126,8 @@ namespace Chaos
 			FColor G = (ColorScale * FColor::Green).ToFColor(false);
 			FColor B = (ColorScale * FColor::Blue).ToFColor(false);
 
-			FVec3 P = SpaceTransform.TransformPosition((Particle->AsDynamic()) ? Particle->AsDynamic()->P() : Particle->X());
-			FRotation3 Q = SpaceTransform.GetRotation() * ((Particle->AsDynamic()) ? Particle->AsDynamic()->Q() : Particle->R());
+			FVec3 P = SpaceTransform.TransformPosition((Particle->ObjectState() == EObjectStateType::Dynamic) ? Particle->CastToRigidParticle()->P() : Particle->X());
+			FRotation3 Q = SpaceTransform.GetRotation() * ((Particle->ObjectState() == EObjectStateType::Dynamic) ? Particle->CastToRigidParticle()->Q() : Particle->R());
 			FMatrix33 Qm = Q.ToMatrix();
 			FDebugDrawQueue::GetInstance().DrawDebugDirectionalArrow(P, P + DrawScale * BodyAxisLen * Qm.GetAxis(0), DrawScale * ArrowSize, R, false, KINDA_SMALL_NUMBER, DrawPriority, LineThickness);
 			FDebugDrawQueue::GetInstance().DrawDebugDirectionalArrow(P, P + DrawScale * BodyAxisLen * Qm.GetAxis(1), DrawScale * ArrowSize, G, false, KINDA_SMALL_NUMBER, DrawPriority, LineThickness);
@@ -209,7 +209,9 @@ namespace Chaos
 		void DrawJointConstraintImpl(const FRigidTransform3& SpaceTransform, const FPBDJointConstraintHandle* ConstraintHandle, FReal ColorScale, uint32 FeatureMask)
 		{
 			TVector<TGeometryParticleHandle<FReal, 3>*, 2> ConstrainedParticles = ConstraintHandle->GetConstrainedParticles();
-			if (ConstrainedParticles[0]->AsDynamic() || ConstrainedParticles[1]->AsDynamic())
+			auto RigidParticle0 = ConstrainedParticles[0]->CastToRigidParticle();
+			auto RigidParticle1 = ConstrainedParticles[1]->CastToRigidParticle();
+			if ((RigidParticle0 && RigidParticle0->ObjectState() == EObjectStateType::Dynamic) || (RigidParticle1 && RigidParticle1->ObjectState() == EObjectStateType::Dynamic))
 			{
 				FVec3 Pa = TGenericParticleHandle<FReal, 3>(ConstraintHandle->GetConstrainedParticles()[1])->P();
 				FVec3 Pb = TGenericParticleHandle<FReal, 3>(ConstraintHandle->GetConstrainedParticles()[0])->P();
@@ -229,7 +231,7 @@ namespace Chaos
 			{
 				for (auto& Particle : ParticlesView)
 				{
-					if ((bDrawKinemtatic && !Particle.AsDynamic()) || (bDrawDynamic && Particle.AsDynamic()))
+					if ((bDrawKinemtatic && Particle.ObjectState() != EObjectStateType::Dynamic) || (bDrawDynamic && Particle.ObjectState() == EObjectStateType::Dynamic))
 					{
 						DrawParticleShapesImpl(SpaceTransform, GetHandleHelper(&Particle), ColorScale);
 					}
@@ -245,7 +247,7 @@ namespace Chaos
 			{
 				for (auto& Particle : Particles)
 				{
-					if ((bDrawKinemtatic && !Particle->AsDynamic()) || (bDrawDynamic && Particle->AsDynamic()))
+					if ((bDrawKinemtatic && Particle->ObjectState() != EObjectStateType::Dynamic) || (bDrawDynamic && Particle->ObjectState() == EObjectStateType::Dynamic))
 					{
 						DrawParticleShapesImpl(SpaceTransform, Particle, ColorScale);
 					}
@@ -261,7 +263,7 @@ namespace Chaos
 			{
 				for (auto& Particle : ParticlesView)
 				{
-					if ((bDrawKinemtatic && !Particle.AsDynamic()) || (bDrawDynamic && Particle.AsDynamic()))
+					if ((bDrawKinemtatic && Particle.ObjectState() != EObjectStateType::Dynamic) || (bDrawDynamic && Particle.ObjectState() == EObjectStateType::Dynamic))
 					{
 						DrawParticleTransformImpl(SpaceTransform, GetHandleHelper(&Particle), ColorScale);
 					}
@@ -277,7 +279,7 @@ namespace Chaos
 			{
 				for (auto& Particle : Particles)
 				{
-					if ((bDrawKinemtatic && !Particle->AsDynamic()) || (bDrawDynamic && Particle->AsDynamic()))
+					if ((bDrawKinemtatic && Particle->ObjectState() != EObjectStateType::Dynamic) || (bDrawDynamic && Particle->ObjectState() == EObjectStateType::Dynamic))
 					{
 						DrawParticleTransformImpl(SpaceTransform, Particle, ColorScale);
 					}
diff --git a/Engine/Source/Runtime/Experimental/Chaos/Private/Chaos/CollisionResolution.cpp b/Engine/Source/Runtime/Experimental/Chaos/Private/Chaos/CollisionResolution.cpp
--- a/Engine/Source/Runtime/Experimental/Chaos/Private/Chaos/CollisionResolution.cpp
+++ b/Engine/Source/Runtime/Experimental/Chaos/Private/Chaos/CollisionResolution.cpp
@@ -590,7 +590,9 @@ namespace Chaos
 		void ConstructLevelsetLevelsetConstraints(TGeometryParticleHandle<T, d>* Particle0, TGeometryParticleHandle<T, d>* Particle1, const FImplicitObject* Implicit0, const FImplicitObject* Implicit1, const T Thickness, FCollisionConstraintsArray& NewConstraints)
 		{
 			FRigidBodyPointContactConstraint* Constraint = new FRigidBodyPointContactConstraint;
-			if (!Particle1->Geometry() || (Particle0->AsDynamic() && !Particle0->AsDynamic()->CollisionParticlesSize() && Particle0->Geometry() && !Particle0->Geometry()->IsUnderlyingUnion()))
+
+			bool bIsParticleDynamic0 = Particle0->CastToRigidParticle() && Particle0->ObjectState() == EObjectStateType::Dynamic;
+			if (!Particle1->Geometry() || (bIsParticleDynamic0 && !Particle0->CastToRigidParticle()->CollisionParticlesSize() && Particle0->Geometry() && !Particle0->Geometry()->IsUnderlyingUnion()))
 			{
 				Constraint->Particle[0] = Particle1;
 				Constraint->Particle[1] = Particle0;
diff --git a/Engine/Source/Runtime/Experimental/Chaos/Private/Chaos/CollisionResolutionUtil.cpp b/Engine/Source/Runtime/Experimental/Chaos/Private/Chaos/CollisionResolutionUtil.cpp
--- a/Engine/Source/Runtime/Experimental/Chaos/Private/Chaos/CollisionResolutionUtil.cpp
+++ b/Engine/Source/Runtime/Experimental/Chaos/Private/Chaos/CollisionResolutionUtil.cpp
@@ -47,21 +47,23 @@ namespace Chaos
 		template<typename T, int d>
 		TVector<T, d> GetEnergyClampedImpulse(const TRigidBodyPointContactConstraint<T, d>& Constraint, const TVector<T, d>& Impulse, const TVector<T, d>& VectorToPoint1, const TVector<T, d>& VectorToPoint2, const TVector<T, d>& Velocity1, const TVector<T, d>& Velocity2)
 		{
-			TPBDRigidParticleHandle<T, d>* PBDRigid0 = Constraint.Particle[0]->AsDynamic();
-			TPBDRigidParticleHandle<T, d>* PBDRigid1 = Constraint.Particle[1]->AsDynamic();
+			TPBDRigidParticleHandle<T, d>* PBDRigid0 = Constraint.Particle[0]->CastToRigidParticle();
+			TPBDRigidParticleHandle<T, d>* PBDRigid1 = Constraint.Particle[1]->CastToRigidParticle();
+			const bool bIsRigidDynamic0 = PBDRigid0 && PBDRigid0->ObjectState() == EObjectStateType::Dynamic;
+			const bool bIsRigidDynamic1 = PBDRigid1 && PBDRigid1->ObjectState() == EObjectStateType::Dynamic;
 
 			TVector<T, d> Jr0, Jr1, IInvJr0, IInvJr1;
 			T ImpulseRatioNumerator0 = 0, ImpulseRatioNumerator1 = 0, ImpulseRatioDenom0 = 0, ImpulseRatioDenom1 = 0;
 			T ImpulseSize = Impulse.SizeSquared();
-			TVector<T, d> KinematicVelocity = !PBDRigid0 ? Velocity1 : !PBDRigid1 ? Velocity2 : TVector<T, d>(0);
-			if (PBDRigid0)
+			TVector<T, d> KinematicVelocity = !bIsRigidDynamic0 ? Velocity1 : !bIsRigidDynamic1 ? Velocity2 : TVector<T, d>(0);
+			if (bIsRigidDynamic0)
 			{
 				Jr0 = TVector<T, d>::CrossProduct(VectorToPoint1, Impulse);
 				IInvJr0 = PBDRigid0->Q().RotateVector(PBDRigid0->InvI() * PBDRigid0->Q().UnrotateVector(Jr0));
 				ImpulseRatioNumerator0 = TVector<T, d>::DotProduct(Impulse, PBDRigid0->V() - KinematicVelocity) + TVector<T, d>::DotProduct(IInvJr0, PBDRigid0->W());
 				ImpulseRatioDenom0 = ImpulseSize / PBDRigid0->M() + TVector<T, d>::DotProduct(Jr0, IInvJr0);
 			}
-			if (PBDRigid1)
+			if (bIsRigidDynamic1)
 			{
 				Jr1 = TVector<T, d>::CrossProduct(VectorToPoint2, Impulse);
 				IInvJr1 = PBDRigid1->Q().RotateVector(PBDRigid1->InvI() * PBDRigid1->Q().UnrotateVector(Jr1));
diff --git a/Engine/Source/Runtime/Experimental/Chaos/Private/Chaos/GeometryParticles.cpp b/Engine/Source/Runtime/Experimental/Chaos/Private/Chaos/GeometryParticles.cpp
--- a/Engine/Source/Runtime/Experimental/Chaos/Private/Chaos/GeometryParticles.cpp
+++ b/Engine/Source/Runtime/Experimental/Chaos/Private/Chaos/GeometryParticles.cpp
@@ -124,7 +124,7 @@ namespace Chaos
 		{
 		case EParticleType::Static: return Ar.IsLoading() ? new TGeometryParticlesImp<T, d, SimType>() : nullptr;
 		case EParticleType::Kinematic: return Ar.IsLoading() ? new TKinematicGeometryParticlesImp<T, d, SimType>() : nullptr;
-		case EParticleType::Dynamic: return Ar.IsLoading() ? new TPBDRigidParticles<T, d>() : nullptr;
+		case EParticleType::Rigid: return Ar.IsLoading() ? new TPBDRigidParticles<T, d>() : nullptr;
 		case EParticleType::Clustered: return Ar.IsLoading() ? new TPBDRigidClusteredParticles<T, d>() : nullptr;
 		default:
 			check(false); return nullptr;
diff --git a/Engine/Source/Runtime/Experimental/Chaos/Private/Chaos/PBDCollisionConstraintsPointContactUtil.cpp b/Engine/Source/Runtime/Experimental/Chaos/Private/Chaos/PBDCollisionConstraintsPointContactUtil.cpp
--- a/Engine/Source/Runtime/Experimental/Chaos/Private/Chaos/PBDCollisionConstraintsPointContactUtil.cpp
+++ b/Engine/Source/Runtime/Experimental/Chaos/Private/Chaos/PBDCollisionConstraintsPointContactUtil.cpp
@@ -41,8 +41,10 @@ namespace Chaos
 		{
 			TGenericParticleHandle<T, d> Particle0 = TGenericParticleHandle<T, d>(Constraint.Particle[0]);
 			TGenericParticleHandle<T, d> Particle1 = TGenericParticleHandle<T, d>(Constraint.Particle[1]);
-			TPBDRigidParticleHandle<T, d>* PBDRigid0 = Particle0->AsDynamic();
-			TPBDRigidParticleHandle<T, d>* PBDRigid1 = Particle1->AsDynamic();
+			TPBDRigidParticleHandle<T, d>* PBDRigid0 = Particle0->CastToRigidParticle();
+			TPBDRigidParticleHandle<T, d>* PBDRigid1 = Particle1->CastToRigidParticle();
+			bool bIsRigidDynamic0 = PBDRigid0 && PBDRigid0->ObjectState() == EObjectStateType::Dynamic;
+			bool bIsRigidDynamic1 = PBDRigid1 && PBDRigid1->ObjectState() == EObjectStateType::Dynamic;
 
 			for (int32 PairIt = 0; PairIt < IterationParameters.NumPairIterations; ++PairIt)
 			{
@@ -88,11 +90,11 @@ namespace Chaos
 
 				if (RelativeNormalVelocity < 0) // ignore separating constraints
 				{
-					PMatrix<T, d, d> WorldSpaceInvI1 = PBDRigid0 ? Utilities::ComputeWorldSpaceInertia(Q0, PBDRigid0->InvI()) : PMatrix<T, d, d>(0);
-					PMatrix<T, d, d> WorldSpaceInvI2 = PBDRigid1 ? Utilities::ComputeWorldSpaceInertia(Q1, PBDRigid1->InvI()) : PMatrix<T, d, d>(0);
+					PMatrix<T, d, d> WorldSpaceInvI1 = bIsRigidDynamic0 ? Utilities::ComputeWorldSpaceInertia(Q0, PBDRigid0->InvI()) : PMatrix<T, d, d>(0);
+					PMatrix<T, d, d> WorldSpaceInvI2 = bIsRigidDynamic1 ? Utilities::ComputeWorldSpaceInertia(Q1, PBDRigid1->InvI()) : PMatrix<T, d, d>(0);
 					PMatrix<T, d, d> Factor =
-						(PBDRigid0 ? ComputeFactorMatrix3(VectorToPoint1, WorldSpaceInvI1, PBDRigid0->InvM()) : PMatrix<T, d, d>(0)) +
-						(PBDRigid1 ? ComputeFactorMatrix3(VectorToPoint2, WorldSpaceInvI2, PBDRigid1->InvM()) : PMatrix<T, d, d>(0));
+						(bIsRigidDynamic0 ? ComputeFactorMatrix3(VectorToPoint1, WorldSpaceInvI1, PBDRigid0->InvM()) : PMatrix<T, d, d>(0)) +
+						(bIsRigidDynamic1 ? ComputeFactorMatrix3(VectorToPoint2, WorldSpaceInvI2, PBDRigid1->InvM()) : PMatrix<T, d, d>(0));
 					TVector<T, d> Impulse;
 					TVector<T, d> AngularImpulse(0);
 
@@ -158,21 +160,21 @@ namespace Chaos
 								TVector<T, d> AngularTangent = RelativeAngularVelocity - AngularNormal * Contact.Normal;
 								TVector<T, d> FinalAngularVelocity = FMath::Sign(AngularNormal) * FMath::Max((T)0, FMath::Abs(AngularNormal) - AngularFriction * NormalVelocityChange) * Contact.Normal + FMath::Max((T)0, AngularTangent.Size() - AngularFriction * NormalVelocityChange) * AngularTangent.GetSafeNormal();
 								TVector<T, d> Delta = FinalAngularVelocity - RelativeAngularVelocity;
-								if (!PBDRigid0 && PBDRigid1)
+								if (!bIsRigidDynamic0 && bIsRigidDynamic1)
 								{
 									PMatrix<T, d, d> WorldSpaceI2 = (Q1 * FMatrix::Identity) * PBDRigid1->I() * (Q1 * FMatrix::Identity).GetTransposed();
 									TVector<T, d> ImpulseDelta = PBDRigid1->M() * TVector<T, d>::CrossProduct(VectorToPoint2, Delta);
 									Impulse += ImpulseDelta;
 									AngularImpulse += WorldSpaceI2 * Delta - TVector<T, d>::CrossProduct(VectorToPoint2, ImpulseDelta);
 								}
-								else if (PBDRigid0 && !PBDRigid1)
+								else if (bIsRigidDynamic0 && !bIsRigidDynamic1)
 								{
 									PMatrix<T, d, d> WorldSpaceI1 = (Q0 * FMatrix::Identity) * PBDRigid0->I() * (Q0 * FMatrix::Identity).GetTransposed();
 									TVector<T, d> ImpulseDelta = PBDRigid0->M() * TVector<T, d>::CrossProduct(VectorToPoint1, Delta);
 									Impulse += ImpulseDelta;
 									AngularImpulse += WorldSpaceI1 * Delta - TVector<T, d>::CrossProduct(VectorToPoint1, ImpulseDelta);
 								}
-								else if (PBDRigid0 && PBDRigid1)
+								else if (bIsRigidDynamic0 && bIsRigidDynamic1)
 								{
 									PMatrix<T, d, d> Cross1(0, VectorToPoint1.Z, -VectorToPoint1.Y, -VectorToPoint1.Z, 0, VectorToPoint1.X, VectorToPoint1.Y, -VectorToPoint1.X, 0);
 									PMatrix<T, d, d> Cross2(0, VectorToPoint2.Z, -VectorToPoint2.Y, -VectorToPoint2.Z, 0, VectorToPoint2.X, VectorToPoint2.Y, -VectorToPoint2.X, 0);
@@ -226,7 +228,7 @@ namespace Chaos
 					}
 					Impulse = GetEnergyClampedImpulse(Constraint, Impulse, VectorToPoint1, VectorToPoint2, Body1Velocity, Body2Velocity);
 					Constraint.AccumulatedImpulse += Impulse;
-					if (PBDRigid0)
+					if (bIsRigidDynamic0)
 					{
 						// Velocity update for next step
 						FVec3 NetAngularImpulse = FVec3::CrossProduct(VectorToPoint1, Impulse) + AngularImpulse;
@@ -239,7 +241,7 @@ namespace Chaos
 						PBDRigid0->Q() += FRotation3::FromElements(DW, 0.f) * Q0 * IterationParameters.Dt * T(0.5);
 						PBDRigid0->Q().Normalize();
 					}
-					if (PBDRigid1)
+					if (bIsRigidDynamic1)
 					{
 						// Velocity update for next step
 						FVec3 NetAngularImpulse = FVec3::CrossProduct(VectorToPoint2, -Impulse) - AngularImpulse;
@@ -262,18 +264,20 @@ namespace Chaos
 		{
 			TGeometryParticleHandle<T, d>* Particle0 = Constraint.Particle[0];
 			TGeometryParticleHandle<T, d>* Particle1 = Constraint.Particle[1];
-			TPBDRigidParticleHandle<T, d>* PBDRigid0 = Particle0->AsDynamic();
-			TPBDRigidParticleHandle<T, d>* PBDRigid1 = Particle1->AsDynamic();
+			TPBDRigidParticleHandle<T, d>* PBDRigid0 = Particle0->CastToRigidParticle();
+			TPBDRigidParticleHandle<T, d>* PBDRigid1 = Particle1->CastToRigidParticle();
+			const bool bIsRigidDynamic0 = PBDRigid0 && PBDRigid0->ObjectState() == EObjectStateType::Dynamic;
+			const bool bIsRigidDynamic1 = PBDRigid1 && PBDRigid1->ObjectState() == EObjectStateType::Dynamic;
 
 			const TVector<T, d> ZeroVector = TVector<T, d>(0);
-			const TRotation<T, d>& Q0 = PBDRigid0 ? PBDRigid0->Q() : Particle0->R();
-			const TRotation<T, d>& Q1 = PBDRigid1 ? PBDRigid1->Q() : Particle1->R();
-			const TVector<T, d>& P0 = PBDRigid0 ? PBDRigid0->P() : Particle0->X();
-			const TVector<T, d>& P1 = PBDRigid1 ? PBDRigid1->P() : Particle1->X();
-			const TVector<T, d>& V0 = PBDRigid0 ? PBDRigid0->V() : ZeroVector;
-			const TVector<T, d>& V1 = PBDRigid1 ? PBDRigid1->V() : ZeroVector;
-			const TVector<T, d>& W0 = PBDRigid0 ? PBDRigid0->W() : ZeroVector;
-			const TVector<T, d>& W1 = PBDRigid1 ? PBDRigid1->W() : ZeroVector;
+			const TRotation<T, d>& Q0 = bIsRigidDynamic0 ? PBDRigid0->Q() : Particle0->R();
+			const TRotation<T, d>& Q1 = bIsRigidDynamic1 ? PBDRigid1->Q() : Particle1->R();
+			const TVector<T, d>& P0 = bIsRigidDynamic0 ? PBDRigid0->P() : Particle0->X();
+			const TVector<T, d>& P1 = bIsRigidDynamic1 ? PBDRigid1->P() : Particle1->X();
+			const TVector<T, d>& V0 = bIsRigidDynamic0 ? PBDRigid0->V() : ZeroVector;
+			const TVector<T, d>& V1 = bIsRigidDynamic1 ? PBDRigid1->V() : ZeroVector;
+			const TVector<T, d>& W0 = bIsRigidDynamic0 ? PBDRigid0->W() : ZeroVector;
+			const TVector<T, d>& W1 = bIsRigidDynamic1 ? PBDRigid1->W() : ZeroVector;
 			const bool IsTemporarilyStatic0 = IsTemporarilyStatic.Contains(Particle0);
 			const bool IsTemporarilyStatic1 = IsTemporarilyStatic.Contains(Particle1);
 	
@@ -295,19 +299,19 @@ namespace Chaos
 					break;
 				}
 
-				if ((!PBDRigid0 || IsTemporarilyStatic0) && (!PBDRigid1 || IsTemporarilyStatic1))
-				{
+				if ((!bIsRigidDynamic0 || IsTemporarilyStatic0) && (!bIsRigidDynamic1 || IsTemporarilyStatic1))
+				{	
 					break;
 				}
 
 				*IterationParameters.NeedsAnotherIteration = true;
-				PMatrix<T, d, d> WorldSpaceInvI1 = PBDRigid0 ? Utilities::ComputeWorldSpaceInertia(Q0, PBDRigid0->InvI()) : PMatrix<T, d, d>(0);
-				PMatrix<T, d, d> WorldSpaceInvI2 = PBDRigid1 ? Utilities::ComputeWorldSpaceInertia(Q1, PBDRigid1->InvI()) : PMatrix<T, d, d>(0);
+				PMatrix<T, d, d> WorldSpaceInvI1 = bIsRigidDynamic0 ? Utilities::ComputeWorldSpaceInertia(Q0, PBDRigid0->InvI()) : PMatrix<T, d, d>(0);
+				PMatrix<T, d, d> WorldSpaceInvI2 = bIsRigidDynamic1 ? Utilities::ComputeWorldSpaceInertia(Q1, PBDRigid1->InvI()) : PMatrix<T, d, d>(0);
 				TVector<T, d> VectorToPoint1 = Contact.Location - P0;
 				TVector<T, d> VectorToPoint2 = Contact.Location - P1;
 				PMatrix<T, d, d> Factor =
-					(PBDRigid0 ? ComputeFactorMatrix3(VectorToPoint1, WorldSpaceInvI1, PBDRigid0->InvM()) : PMatrix<T, d, d>(0)) +
-					(PBDRigid1 ? ComputeFactorMatrix3(VectorToPoint2, WorldSpaceInvI2, PBDRigid1->InvM()) : PMatrix<T, d, d>(0));
+					(bIsRigidDynamic0 ? ComputeFactorMatrix3(VectorToPoint1, WorldSpaceInvI1, PBDRigid0->InvM()) : PMatrix<T, d, d>(0)) +
+					(bIsRigidDynamic1 ? ComputeFactorMatrix3(VectorToPoint2, WorldSpaceInvI2, PBDRigid1->InvM()) : PMatrix<T, d, d>(0));
 				T Numerator = FMath::Min((T)(IterationParameters.Iteration + 2), (T)IterationParameters.NumIterations);
 				T ScalingFactor = Numerator / (T)IterationParameters.NumIterations;
 
@@ -332,7 +336,7 @@ namespace Chaos
 					TVector<T, d> VelocityFixImpulse = ImpulseNumerator / ImpulseDenominator;
 					VelocityFixImpulse = GetEnergyClampedImpulse(Constraint, VelocityFixImpulse, VectorToPoint1, VectorToPoint2, Body1Velocity, Body2Velocity);
 					Constraint.AccumulatedImpulse += VelocityFixImpulse;	//question: should we track this?
-					if (!IsTemporarilyStatic0 && PBDRigid0)
+					if (!IsTemporarilyStatic0 && bIsRigidDynamic0)
 					{
 						TVector<T, d> AngularImpulse = TVector<T, d>::CrossProduct(VectorToPoint1, VelocityFixImpulse);
 						PBDRigid0->V() += PBDRigid0->InvM() * VelocityFixImpulse;
@@ -340,8 +344,8 @@ namespace Chaos
 
 					}
 
-					if (!IsTemporarilyStatic1 && PBDRigid1)
-					{
+					if (!IsTemporarilyStatic1 && bIsRigidDynamic1)
+					{	
 						TVector<T, d> AngularImpulse = TVector<T, d>::CrossProduct(VectorToPoint2, -VelocityFixImpulse);
 						PBDRigid1->V() -= PBDRigid1->InvM() * VelocityFixImpulse;
 						PBDRigid1->W() += WorldSpaceInvI2 * AngularImpulse;
@@ -353,13 +357,13 @@ namespace Chaos
 				TVector<T, d> Impulse = PMatrix<T, d, d>(Factor.Inverse()) * ((-Contact.Phi + Thickness) * ScalingFactor * Contact.Normal);
 				TVector<T, d> AngularImpulse1 = TVector<T, d>::CrossProduct(VectorToPoint1, Impulse);
 				TVector<T, d> AngularImpulse2 = TVector<T, d>::CrossProduct(VectorToPoint2, -Impulse);
-				if (!IsTemporarilyStatic0 && PBDRigid0)
+				if (!IsTemporarilyStatic0 && bIsRigidDynamic0)
 				{
 					PBDRigid0->P() += PBDRigid0->InvM() * Impulse;
 					PBDRigid0->Q() = TRotation<T, d>::FromVector(WorldSpaceInvI1 * AngularImpulse1) * Q0;
 					PBDRigid0->Q().Normalize();
 				}
-				if (!IsTemporarilyStatic1 && PBDRigid1)
+				if (!IsTemporarilyStatic1 && bIsRigidDynamic1)
 				{
 					PBDRigid1->P() -= PBDRigid1->InvM() * Impulse;
 					PBDRigid1->Q() = TRotation<T, d>::FromVector(WorldSpaceInvI2 * AngularImpulse2) * Q1;
diff --git a/Engine/Source/Runtime/Experimental/Chaos/Private/Chaos/PBDConstraintColor.cpp b/Engine/Source/Runtime/Experimental/Chaos/Private/Chaos/PBDConstraintColor.cpp
--- a/Engine/Source/Runtime/Experimental/Chaos/Private/Chaos/PBDConstraintColor.cpp
+++ b/Engine/Source/Runtime/Experimental/Chaos/Private/Chaos/PBDConstraintColor.cpp
@@ -48,7 +48,8 @@ void FPBDConstraintColor::ComputeIslandColoring(const int32 Island, const FPBDCo
 
 		const int32 ParticleNodeIndex = ConstraintGraph.ParticleToNodeIndex[Particle];
 
-		if (ProcessedNodes.Contains(ParticleNodeIndex) || Particle->AsDynamic() == nullptr)
+		const bool bIsParticleDynamic = Particle->CastToRigidParticle() && Particle->ObjectState() == EObjectStateType::Dynamic;
+		if (ProcessedNodes.Contains(ParticleNodeIndex) || bIsParticleDynamic == false)
 		{
 			continue;
 		}
@@ -103,7 +104,10 @@ void FPBDConstraintColor::ComputeIslandColoring(const int32 Island, const FPBDCo
 				if (OtherNodeIndex != INDEX_NONE)
 				{
 					FGraphNodeColor& OtherColorNode = Nodes[OtherNodeIndex];
-					if (ConstraintGraph.Nodes[OtherNodeIndex].Particle->AsDynamic() != nullptr)
+
+					const typename FPBDConstraintGraph::FGraphNode& OtherGraphNode = ConstraintGraph.Nodes[OtherNodeIndex];
+					const bool bIsOtherGraphNodeDynamic = OtherGraphNode.Particle->CastToRigidParticle() && OtherGraphNode.Particle->ObjectState() == EObjectStateType::Dynamic;
+					if (bIsOtherGraphNodeDynamic)
 					{
 						while (OtherColorNode.UsedColors.Contains(ColorToUse) || ColorNode.UsedColors.Contains(ColorToUse))
 						{
@@ -118,7 +122,7 @@ void FPBDConstraintColor::ComputeIslandColoring(const int32 Island, const FPBDCo
 				ColorEdge.Color = ColorToUse;
 
 				// Bump color to use next time, but only if we weren't forced to use a different color by the other node
-				if ((ColorToUse == ColorNode.NextColor) && Particle->AsDynamic() != nullptr)
+				if ((ColorToUse == ColorNode.NextColor) && bIsParticleDynamic == true)
 				{
 					ColorNode.NextColor++;
 				}
@@ -145,12 +149,13 @@ void FPBDConstraintColor::ComputeIslandColoring(const int32 Island, const FPBDCo
 				if (OtherNodeIndex != INDEX_NONE)
 				{
 					const typename FPBDConstraintGraph::FGraphNode& OtherGraphNode = ConstraintGraph.Nodes[OtherNodeIndex];
-					if (OtherGraphNode.Particle->AsDynamic() != nullptr)
+					const bool bIsOtherGraphNodeDynamic = OtherGraphNode.Particle->CastToRigidParticle() && OtherGraphNode.Particle->ObjectState() == EObjectStateType::Dynamic;
+					if (bIsOtherGraphNodeDynamic)
 					{
 						FGraphNodeColor& OtherColorNode = Nodes[OtherNodeIndex];
 
 						// Mark other node as not allowing use of this color
-						if (Particle->AsDynamic() != nullptr)
+						if (bIsParticleDynamic)
 						{
 							OtherColorNode.UsedColors.Add(ColorEdge.Color);
 						}
@@ -181,7 +186,8 @@ void FPBDConstraintColor::ComputeContactGraph(const int32 Island, const FPBDCons
 	for (const TGeometryParticleHandle<FReal, 3>* Particle : ConstraintGraph.GetIslandParticles(Island))
 	{
 		const int32* NodeIndexPtr = ConstraintGraph.ParticleToNodeIndex.Find(Particle);
-		if (Particle->AsDynamic() == nullptr && NodeIndexPtr)
+		const bool bIsParticleDynamic = Particle->CastToRigidParticle() && Particle->ObjectState() == EObjectStateType::Dynamic;
+		if (bIsParticleDynamic == false && NodeIndexPtr)
 		{
 			const int32 NodeIndex = *NodeIndexPtr;
 			QueueToProcess.push(std::make_pair(0, NodeIndex));
diff --git a/Engine/Source/Runtime/Experimental/Chaos/Private/Chaos/PBDConstraintGraph.cpp b/Engine/Source/Runtime/Experimental/Chaos/Private/Chaos/PBDConstraintGraph.cpp
--- a/Engine/Source/Runtime/Experimental/Chaos/Private/Chaos/PBDConstraintGraph.cpp
+++ b/Engine/Source/Runtime/Experimental/Chaos/Private/Chaos/PBDConstraintGraph.cpp
@@ -14,6 +14,7 @@
 using namespace Chaos;
 
 
+
 FPBDConstraintGraph::FPBDConstraintGraph() : VisitToken(0)
 {
 }
@@ -134,7 +135,9 @@ void FPBDConstraintGraph::InitializeGraph(const TParticleView<TGeometryParticles
 					auto* Particle = Nodes[UpdatedNodes[Index]].Particle;
 				    if (CHAOS_ENSURE(Particle))
 				    {
-					    if (TPBDRigidParticleHandle<FReal, 3>* PBDRigid = Particle->AsDynamic())
+						// This does NOT check if particle is actually dynamic. TODO confirm that this should happen to all nonDisabledParticles, not just dynamics.
+						TPBDRigidParticleHandle<FReal, 3>* PBDRigid = Particle->CastToRigidParticle();
+						if(PBDRigid && PBDRigid->ObjectState() == EObjectStateType::Dynamic)
 					    {
 						    PBDRigid->Island() = INDEX_NONE;
 					    }
@@ -325,7 +328,8 @@ void FPBDConstraintGraph::ComputeIslands(const TParticleView<TPBDRigidParticles<
 
 			for (TGeometryParticleHandle<FReal, 3>* Particle : NewIslandParticles[Island])
 			{
-				if (TPBDRigidParticleHandle<FReal, 3>* PBDRigid = Particle->AsDynamic())
+				TPBDRigidParticleHandle<FReal, 3>* PBDRigid = Particle->CastToRigidParticle();
+				if(PBDRigid && PBDRigid->ObjectState() == EObjectStateType::Dynamic)
 				{
 					PBDRigid->Island() = Island;
 				}
@@ -356,7 +360,8 @@ void FPBDConstraintGraph::ComputeIslands(const TParticleView<TPBDRigidParticles<
 						Particles.ActivateParticle(Particle); 	//todo: record state change for array reorder
 					}
 
-					if (TPBDRigidParticleHandle<FReal, 3>* PBDRigid = Particle->AsDynamic())
+					TPBDRigidParticleHandle<FReal, 3>* PBDRigid = Particle->CastToRigidParticle();
+					if(PBDRigid && PBDRigid->ObjectState() == EObjectStateType::Dynamic)
 					{
 						if (!Particle->Sleeping() && bSleepState)
 						{
@@ -385,8 +390,9 @@ void FPBDConstraintGraph::ComputeIslands(const TParticleView<TPBDRigidParticles<
 
 			for (TGeometryParticleHandle<FReal, 3>* Particle : IslandToParticles[Island])
 			{
-				TPBDRigidParticleHandle<FReal, 3>* PBDRigid = Particle->AsDynamic();
-				int32 TmpIsland = PBDRigid ? PBDRigid->Island() : INDEX_NONE; //question: should we even store non dynamics in this array?
+				TPBDRigidParticleHandle<FReal, 3>* PBDRigid = Particle->CastToRigidParticle();
+				const bool bIsDynamic = PBDRigid && PBDRigid->ObjectState() == EObjectStateType::Dynamic;
+				int32 TmpIsland = bIsDynamic ? PBDRigid->Island() : INDEX_NONE; //question: should we even store non dynamics in this array?
 
 				if (OtherIsland == INDEX_NONE && TmpIsland >= 0)
 				{
@@ -417,7 +423,8 @@ void FPBDConstraintGraph::ComputeIslands(const TParticleView<TPBDRigidParticles<
 			{
 				for (TGeometryParticleHandle<FReal, 3>* Particle : IslandToParticles[Island])
 				{
-					if (TPBDRigidParticleHandle<FReal, 3>* PBDRigid = Particle->AsDynamic())
+					TPBDRigidParticleHandle<FReal, 3>* PBDRigid = Particle->CastToRigidParticle();
+					if(PBDRigid && PBDRigid->ObjectState() == EObjectStateType::Dynamic)
 					{
 						if (!PBDRigid->Disabled())	// todo: why is this needed? [we aren't handling enable/disable state changes properly so disabled particles end up in the graph.]
 						{
@@ -473,9 +480,10 @@ void FPBDConstraintGraph::ComputeIsland(const int32 InNode, const int32 Island,
 			continue;
 		}
 
-		TPBDRigidParticleHandle<FReal, 3>* RigidHandle = Node.Particle->AsDynamic();
+		TPBDRigidParticleHandle<FReal, 3>* RigidHandle = Node.Particle->CastToRigidParticle();
+		const bool isRigidDynamic = RigidHandle && RigidHandle->ObjectState() == EObjectStateType::Dynamic;
 
-		if (RigidHandle == nullptr)
+		if (isRigidDynamic == false)
 		{
 			if (!StaticParticlesInIsland.Contains(Node.Particle))
 			{
@@ -538,7 +546,8 @@ bool FPBDConstraintGraph::SleepInactive(const int32 Island, const TArrayCollecti
 
 	for (const TGeometryParticleHandle<FReal, 3>* Particle : IslandToParticles[Island])
 	{
-		if (const TPBDRigidParticleHandle<FReal, 3>* PBDRigid = Particle->AsDynamic())
+		const TPBDRigidParticleHandle<FReal, 3>* PBDRigid = Particle->CastToRigidParticle();
+		if(PBDRigid && (PBDRigid->ObjectState() == EObjectStateType::Dynamic))
 		{
 			X += PBDRigid->X() * PBDRigid->M();
 			M += PBDRigid->M();
@@ -562,7 +571,8 @@ bool FPBDConstraintGraph::SleepInactive(const int32 Island, const TArrayCollecti
 
 	for (const TGeometryParticleHandle<FReal, 3>* Particle: IslandParticles)
 	{
-		if (const TPBDRigidParticleHandle<FReal, 3>* PBDRigid = Particle->AsDynamic())
+		const TPBDRigidParticleHandle<FReal, 3>* PBDRigid = Particle->CastToRigidParticle();
+		if(PBDRigid && PBDRigid->ObjectState() == EObjectStateType::Dynamic)
 		{
 			W += /*TVector<FReal, 3>::CrossProduct(PBDRigid->X() - X, PBDRigid->M() * PBDRigid->V()/ +*/ PBDRigid->W() * PBDRigid->M();
 		}
@@ -579,7 +589,8 @@ bool FPBDConstraintGraph::SleepInactive(const int32 Island, const TArrayCollecti
 		{
 			for (TGeometryParticleHandle<FReal, 3>* Particle : IslandParticles)
 			{
-				if (TPBDRigidParticleHandle<FReal, 3>* PBDRigid = Particle->AsDynamic())
+				TPBDRigidParticleHandle<FReal, 3>* PBDRigid = Particle->CastToRigidParticle();
+				if(PBDRigid && PBDRigid->ObjectState() == EObjectStateType::Dynamic)
 				{
 					PBDRigid->SetSleeping(true);
 					PBDRigid->V() = TVector<FReal, 3>(0);
@@ -603,7 +614,8 @@ void FPBDConstraintGraph::WakeIsland(const int32 Island)
 {
 	for (TGeometryParticleHandle<FReal, 3>* Particle : IslandToParticles[Island])
 	{
-		if (TPBDRigidParticleHandle<FReal, 3>* PBDRigid = Particle->AsDynamic())
+		TPBDRigidParticleHandle<FReal, 3>* PBDRigid = Particle->CastToRigidParticle();
+		if(PBDRigid && PBDRigid->ObjectState() == EObjectStateType::Dynamic)
 		{
 			if (PBDRigid->Sleeping())
 			{
@@ -647,11 +659,13 @@ void FPBDConstraintGraph::EnableParticle(TGeometryParticleHandle<FReal, 3>* Part
 {
 	if (ParentParticle)
 	{
-		if (const TPBDRigidParticleHandle<FReal, 3>* ParentPBDRigid = ParentParticle->AsDynamic())
+		const TPBDRigidParticleHandle<FReal, 3>* ParentPBDRigid = ParentParticle->CastToRigidParticle();
+		if(ParentPBDRigid && ParentPBDRigid->ObjectState() == EObjectStateType::Dynamic)
 		{
 			ParticleAdd(Particle);
 
-			if (TPBDRigidParticleHandle<FReal, 3>* ChildPBDRigid = Particle->AsDynamic())
+			TPBDRigidParticleHandle<FReal, 3>* ChildPBDRigid = Particle->CastToRigidParticle();
+			if(ChildPBDRigid && ChildPBDRigid->ObjectState() == EObjectStateType::Dynamic)
 			{
 				const int32 Island = ParentPBDRigid->Island();
 				ChildPBDRigid->Island() = Island;
@@ -674,7 +688,8 @@ void FPBDConstraintGraph::EnableParticle(TGeometryParticleHandle<FReal, 3>* Part
 
 void FPBDConstraintGraph::DisableParticle(TGeometryParticleHandle<FReal, 3>* Particle)
 {
-	if (TPBDRigidParticleHandle<FReal, 3>* PBDRigid = Particle->AsDynamic())
+	TPBDRigidParticleHandle<FReal, 3>* PBDRigid = Particle->CastToRigidParticle();
+	if(PBDRigid && PBDRigid->ObjectState() == EObjectStateType::Dynamic)
 	{
 		const int32 Island = PBDRigid->Island();
 		if (Island != INDEX_NONE)
@@ -722,7 +737,7 @@ bool FPBDConstraintGraph::CheckIslands(const TArray<TGeometryParticleHandle<FRea
 			// This is ok only if those particles are static
 			for (TGeometryParticleHandle<FReal, 3>* Particle : IslandParticlesIntersectSet)
 			{
-				if (Particle->AsDynamic())
+				if (Particle->CastToRigidParticle() && Particle->ObjectState() == EObjectStateType::Dynamic)
 				{
 					UE_LOG(LogChaos, Error, TEXT("Island %d contains non-static particle that is also in another Island"), Island);	//todo: add better logging for bad particle
 					bIsValid = false;
diff --git a/Engine/Source/Runtime/Experimental/Chaos/Private/Chaos/PBDJointConstraints.cpp b/Engine/Source/Runtime/Experimental/Chaos/Private/Chaos/PBDJointConstraints.cpp
--- a/Engine/Source/Runtime/Experimental/Chaos/Private/Chaos/PBDJointConstraints.cpp
+++ b/Engine/Source/Runtime/Experimental/Chaos/Private/Chaos/PBDJointConstraints.cpp
@@ -593,8 +593,16 @@ namespace Chaos
 			}
 		}
 
-		UpdateParticleState(Particle0->AsDynamic(), Dt, P0, Q0);
-		UpdateParticleState(Particle1->AsDynamic(), Dt, P1, Q1);
+		TPBDRigidParticleHandle<FReal, 3>* Rigid0 = Particle0->CastToRigidParticle();
+		TPBDRigidParticleHandle<FReal, 3>* Rigid1 = Particle1->CastToRigidParticle();
+		if (Rigid0->ObjectState() == EObjectStateType::Dynamic)
+		{
+			UpdateParticleState(Rigid0, Dt, P0, Q0);
+		}
+		if (Rigid1->ObjectState() == EObjectStateType::Dynamic)
+		{
+			UpdateParticleState(Rigid1, Dt, P1, Q1);
+		}
 	}
 
 
@@ -648,8 +656,16 @@ namespace Chaos
 			Solver.ApplyConstraints(Dt, Settings, JointSettings, GetLinearStiffness(Settings, JointSettings));
 		}
 
-		UpdateParticleState(Particle0->AsDynamic(), Dt, Solver.GetP(0), Solver.GetQ(0));
-		UpdateParticleState(Particle1->AsDynamic(), Dt, Solver.GetP(1), Solver.GetQ(1));
+		TPBDRigidParticleHandle<FReal, 3>* Rigid0 = Particle0->CastToRigidParticle();
+		TPBDRigidParticleHandle<FReal, 3>* Rigid1 = Particle1->CastToRigidParticle();
+		if (Rigid0 && Rigid0->ObjectState() == EObjectStateType::Dynamic)
+		{
+			UpdateParticleState(Rigid0, Dt, Solver.GetP(0), Solver.GetQ(0));
+		}
+		if (Rigid1 && Rigid1->ObjectState() == EObjectStateType::Dynamic)
+		{
+			UpdateParticleState(Rigid1, Dt, Solver.GetP(1), Solver.GetQ(1));
+		}
 	}
 
 
@@ -667,8 +683,6 @@ namespace Chaos
 		const int32 Index1 = 0;
 		TGenericParticleHandle<FReal, 3> Particle0 = TGenericParticleHandle<FReal, 3>(ConstraintParticles[ConstraintIndex][Index0]);
 		TGenericParticleHandle<FReal, 3> Particle1 = TGenericParticleHandle<FReal, 3>(ConstraintParticles[ConstraintIndex][Index1]);
-		TPBDRigidParticleHandle<FReal, 3>* Rigid0 = ConstraintParticles[ConstraintIndex][Index0]->AsDynamic();
-		TPBDRigidParticleHandle<FReal, 3>* Rigid1 = ConstraintParticles[ConstraintIndex][Index1]->AsDynamic();
 
 		FVec3 P0 = Particle0->P();
 		FRotation3 Q0 = Particle0->Q();
@@ -752,8 +766,17 @@ namespace Chaos
 		}
 
 		// Update the particles
-		UpdateParticleState(Rigid0, Dt, P0, Q0);
-		UpdateParticleState(Rigid1, Dt, P1, Q1);
+		TPBDRigidParticleHandle<FReal, 3>* Rigid0 = ConstraintParticles[ConstraintIndex][Index0]->CastToRigidParticle();
+		TPBDRigidParticleHandle<FReal, 3>* Rigid1 = ConstraintParticles[ConstraintIndex][Index1]->CastToRigidParticle();
+
+		if (Rigid0->ObjectState() == EObjectStateType::Dynamic)
+		{
+			UpdateParticleState(Rigid0, Dt, P0, Q0);
+		}
+		if (Rigid1->ObjectState() == EObjectStateType::Dynamic)
+		{
+			UpdateParticleState(Rigid1, Dt, P1, Q1);
+		}
 	}
 
 	
@@ -769,8 +792,6 @@ namespace Chaos
 		const int32 Index1 = 0;
 		TGenericParticleHandle<FReal, 3> Particle0 = TGenericParticleHandle<FReal, 3>(ConstraintParticles[ConstraintIndex][Index0]);
 		TGenericParticleHandle<FReal, 3> Particle1 = TGenericParticleHandle<FReal, 3>(ConstraintParticles[ConstraintIndex][Index1]);
-		TPBDRigidParticleHandle<FReal, 3>* Rigid0 = ConstraintParticles[ConstraintIndex][Index0]->AsDynamic();
-		TPBDRigidParticleHandle<FReal, 3>* Rigid1 = ConstraintParticles[ConstraintIndex][Index1]->AsDynamic();
 
 		FVec3 P0 = Particle0->P();
 		FRotation3 Q0 = Particle0->Q();
@@ -854,8 +875,18 @@ namespace Chaos
 		}
 
 		// Update the particles
-		UpdateParticleState(Rigid0, Dt, P0, Q0);
-		UpdateParticleState(Rigid1, Dt, P1, Q1);
+		TPBDRigidParticleHandle<FReal, 3>* Rigid0 = ConstraintParticles[ConstraintIndex][Index0]->CastToRigidParticle();
+		TPBDRigidParticleHandle<FReal, 3>* Rigid1 = ConstraintParticles[ConstraintIndex][Index1]->CastToRigidParticle();
+
+		if (Rigid0->ObjectState() == EObjectStateType::Dynamic)
+		{
+			UpdateParticleState(Rigid0, Dt, P0, Q0);
+		}
+		if (Rigid1->ObjectState() == EObjectStateType::Dynamic)
+		{
+			UpdateParticleState(Rigid1, Dt, P1, Q1);
+		}
+
 	}
 
 
@@ -947,14 +978,14 @@ namespace Chaos
 		}
 
 		// Update the particles
-		TPBDRigidParticleHandle<FReal, 3>* Rigid0 = ConstraintParticles[ConstraintIndex][Index0]->AsDynamic();
-		TPBDRigidParticleHandle<FReal, 3>* Rigid1 = ConstraintParticles[ConstraintIndex][Index1]->AsDynamic();
-		if (Rigid0)
+		TPBDRigidParticleHandle<FReal, 3>* Rigid0 = ConstraintParticles[ConstraintIndex][Index0]->CastToRigidParticle();
+		TPBDRigidParticleHandle<FReal, 3>* Rigid1 = ConstraintParticles[ConstraintIndex][Index1]->CastToRigidParticle();
+		if (Rigid0 && Rigid0->ObjectState() == EObjectStateType::Dynamic)
 		{
 			Rigid0->SetP(P0);
 			Rigid0->SetQ(Q0);
 		}
-		if (Rigid1)
+		if (Rigid1 && Rigid1->ObjectState() == EObjectStateType::Dynamic)
 		{
 			Rigid1->SetP(P1);
 			Rigid1->SetQ(Q1);
diff --git a/Engine/Source/Runtime/Experimental/Chaos/Private/Chaos/PBDRigidClustering.cpp b/Engine/Source/Runtime/Experimental/Chaos/Private/Chaos/PBDRigidClustering.cpp
--- a/Engine/Source/Runtime/Experimental/Chaos/Private/Chaos/PBDRigidClustering.cpp
+++ b/Engine/Source/Runtime/Experimental/Chaos/Private/Chaos/PBDRigidClustering.cpp
@@ -676,12 +676,12 @@ namespace Chaos
 			TVector<const TGeometryParticleHandle<T, d>*, 2> ConstrainedParticles = ContactHandle->GetConstrainedParticles();
 			if (auto ChildrenPtr = MParentToChildren.Find(ConstrainedParticles[0]))
 			{
-				ComputeStrainLambda(ConstrainedParticles[0]->AsClustered(), *ChildrenPtr);
+				ComputeStrainLambda(ConstrainedParticles[0]->CastToClustered(), *ChildrenPtr);
 			}
 
 			if (auto ChildrenPtr = MParentToChildren.Find(ConstrainedParticles[1]))
 			{
-				ComputeStrainLambda(ConstrainedParticles[1]->AsClustered(), *ChildrenPtr);
+				ComputeStrainLambda(ConstrainedParticles[1]->CastToClustered(), *ChildrenPtr);
 			}
 
 			MCollisionImpulseArrayDirty = true;
@@ -712,7 +712,8 @@ namespace Chaos
 			auto Particles = Evolution.GetIslandParticles(Island); // copy
 			for (int32 ArrayIdx = Particles.Num() - 1; ArrayIdx >= 0; --ArrayIdx)
 			{
-				if (auto PBDRigid = Particles[ArrayIdx]->AsDynamic())
+				auto PBDRigid = Particles[ArrayIdx]->CastToRigidParticle();
+				if(PBDRigid && PBDRigid->ObjectState() == EObjectStateType::Dynamic)
 				{
 					if (!PBDRigid->Sleeping() && !PBDRigid->Disabled())
 					{
@@ -734,7 +735,8 @@ namespace Chaos
 			const auto& ParticleIndices = Evolution.GetIslandParticles(Island);
 			for (const auto Particle : ParticleIndices)
 			{
-				if (auto PBDRigid = Particle->AsDynamic())
+				auto PBDRigid = Particle->CastToRigidParticle();
+				if(PBDRigid && PBDRigid->ObjectState() == EObjectStateType::Dynamic)
 				{
 					bool bDisabled = PBDRigid->Disabled();
 
diff --git a/Engine/Source/Runtime/Experimental/Chaos/Private/Chaos/PBDRigidDynamicSpringConstraints.cpp b/Engine/Source/Runtime/Experimental/Chaos/Private/Chaos/PBDRigidDynamicSpringConstraints.cpp
--- a/Engine/Source/Runtime/Experimental/Chaos/Private/Chaos/PBDRigidDynamicSpringConstraints.cpp
+++ b/Engine/Source/Runtime/Experimental/Chaos/Private/Chaos/PBDRigidDynamicSpringConstraints.cpp
@@ -19,8 +19,10 @@ void TPBDRigidDynamicSpringConstraints<T, d>::UpdatePositionBasedState(const T D
 	{
 		TGeometryParticleHandle<T, d>* Static0 = Constraints[ConstraintIndex][0];
 		TGeometryParticleHandle<T, d>* Static1 = Constraints[ConstraintIndex][1];
-		TPBDRigidParticleHandle<T, d>* PBDRigid0 = Static0->AsDynamic();
-		TPBDRigidParticleHandle<T, d>* PBDRigid1 = Static1->AsDynamic();
+		TPBDRigidParticleHandle<T, d>* PBDRigid0 = Static0->CastToRigidParticle();
+		TPBDRigidParticleHandle<T, d>* PBDRigid1 = Static1->CastToRigidParticle();
+		const bool bIsRigidDynamic0 = PBDRigid0 && PBDRigid0->ObjectState() == EObjectStateType::Dynamic;
+		const bool bIsRigidDynamic1 = PBDRigid1 && PBDRigid1->ObjectState() == EObjectStateType::Dynamic;
 
 		// Do not create springs between objects with no geometry
 		if (!Static0->Geometry() || !Static1->Geometry())
@@ -28,10 +30,10 @@ void TPBDRigidDynamicSpringConstraints<T, d>::UpdatePositionBasedState(const T D
 			continue;
 		}
 
-		const TRotation<T, d>& Q0 = PBDRigid0 ? PBDRigid0->Q() : Static0->R();
-		const TRotation<T, d>& Q1 = PBDRigid1 ? PBDRigid1->Q() : Static1->R();
-		const TVector<T, d>& P0 = PBDRigid0 ? PBDRigid0->P() : Static0->X();
-		const TVector<T, d>& P1 = PBDRigid1 ? PBDRigid1->P() : Static1->X();
+		const TRotation<T, d>& Q0 = bIsRigidDynamic0 ? PBDRigid0->Q() : Static0->R();
+		const TRotation<T, d>& Q1 = bIsRigidDynamic1 ? PBDRigid1->Q() : Static1->R();
+		const TVector<T, d>& P0 = bIsRigidDynamic0 ? PBDRigid0->P() : Static0->X();
+		const TVector<T, d>& P1 = bIsRigidDynamic1 ? PBDRigid1->P() : Static1->X();
 
 		// Delete constraints
 		const int32 NumSprings = SpringDistances[ConstraintIndex].Num();
@@ -95,10 +97,12 @@ void TPBDRigidDynamicSpringConstraints<T, d>::UpdatePositionBasedState(const T D
 template<class T, int d>
 TVector<T, d> TPBDRigidDynamicSpringConstraints<T, d>::GetDelta(const TVector<T, d>& WorldSpaceX1, const TVector<T, d>& WorldSpaceX2, const int32 ConstraintIndex, const int32 SpringIndex) const
 {
-	TPBDRigidParticleHandle<T, d>* PBDRigid0 = Constraints[ConstraintIndex][0]->AsDynamic();
-	TPBDRigidParticleHandle<T, d>* PBDRigid1 = Constraints[ConstraintIndex][1]->AsDynamic();
+	TPBDRigidParticleHandle<T, d>* PBDRigid0 = Constraints[ConstraintIndex][0]->CastToRigidParticle();
+	TPBDRigidParticleHandle<T, d>* PBDRigid1 = Constraints[ConstraintIndex][1]->CastToRigidParticle();
+	const bool bIsRigidDynamic0 = PBDRigid0 && PBDRigid0->ObjectState() == EObjectStateType::Dynamic;
+	const bool bIsRigidDynamic1 = PBDRigid1 && PBDRigid1->ObjectState() == EObjectStateType::Dynamic;
 
-	if (!PBDRigid0 && !PBDRigid1)
+	if (!bIsRigidDynamic0 && !bIsRigidDynamic1)
 	{
 		return TVector<T, d>(0);
 	}
@@ -107,8 +111,8 @@ TVector<T, d> TPBDRigidDynamicSpringConstraints<T, d>::GetDelta(const TVector<T,
 	float Distance = Difference.Size();
 	check(Distance > 1e-7);
 
-	const T InvM0 = PBDRigid0 ? PBDRigid0->InvM() : (T)0;
-	const T InvM1 = PBDRigid1 ? PBDRigid1->InvM() : (T)0;
+	const T InvM0 = bIsRigidDynamic0 ? PBDRigid0->InvM() : (T)0;
+	const T InvM1 = bIsRigidDynamic1 ? PBDRigid1->InvM() : (T)0;
 	const TVector<T, d> Direction = Difference / Distance;
 	const TVector<T, d> Delta = (Distance - SpringDistances[ConstraintIndex][SpringIndex]) * Direction;
 	return Stiffness * Delta / (InvM0 + InvM1);
@@ -119,19 +123,22 @@ void TPBDRigidDynamicSpringConstraints<T, d>::ApplySingle(const T Dt, int32 Cons
 {
 	TGeometryParticleHandle<T, d>* Static0 = Constraints[ConstraintIndex][0];
 	TGeometryParticleHandle<T, d>* Static1 = Constraints[ConstraintIndex][1];
-	TPBDRigidParticleHandle<T, d>* PBDRigid0 = Static0->AsDynamic();
-	TPBDRigidParticleHandle<T, d>* PBDRigid1 = Static1->AsDynamic();
-	check(PBDRigid0 || PBDRigid1);
-	check(!PBDRigid0 || !PBDRigid1 || (PBDRigid0->Island() == PBDRigid1->Island()));
+	TPBDRigidParticleHandle<T, d>* PBDRigid0 = Static0->CastToRigidParticle();
+	TPBDRigidParticleHandle<T, d>* PBDRigid1 = Static1->CastToRigidParticle();
+	const bool bIsRigidDynamic0 = PBDRigid0 && PBDRigid0->ObjectState() == EObjectStateType::Dynamic;
+	const bool bIsRigidDynamic1 = PBDRigid1 && PBDRigid1->ObjectState() == EObjectStateType::Dynamic;
 
-	TRotation<T, d>& Q0 = PBDRigid0 ? PBDRigid0->Q() : Static0->R();
-	TRotation<T, d>& Q1 = PBDRigid1 ? PBDRigid1->Q() : Static1->R();
-	TVector<T, d>& P0 = PBDRigid0 ? PBDRigid0->P() : Static0->X();
-	TVector<T, d>& P1 = PBDRigid1 ? PBDRigid1->P() : Static1->X();
+	check(bIsRigidDynamic0 || bIsRigidDynamic1);
+	check(!bIsRigidDynamic0|| !bIsRigidDynamic1|| (PBDRigid0->Island() == PBDRigid1->Island()));
+
+	TRotation<T, d>& Q0 = bIsRigidDynamic0 ? PBDRigid0->Q() : Static0->R();
+	TRotation<T, d>& Q1 = bIsRigidDynamic1 ? PBDRigid1->Q() : Static1->R();
+	TVector<T, d>& P0 = bIsRigidDynamic0 ? PBDRigid0->P() : Static0->X();
+	TVector<T, d>& P1 = bIsRigidDynamic1 ? PBDRigid1->P() : Static1->X();
 
 	const int32 NumSprings = SpringDistances[ConstraintIndex].Num();
-	const PMatrix<T, d, d> WorldSpaceInvI1 = PBDRigid0? Utilities::ComputeWorldSpaceInertia(Q0, PBDRigid0->InvI()) : PMatrix<T, d, d>(0);
-	const PMatrix<T, d, d> WorldSpaceInvI2 = PBDRigid1 ? Utilities::ComputeWorldSpaceInertia(Q1, PBDRigid1->InvI()) : PMatrix<T, d, d>(0);;
+	const PMatrix<T, d, d> WorldSpaceInvI1 = bIsRigidDynamic0 ? Utilities::ComputeWorldSpaceInertia(Q0, PBDRigid0->InvI()) : PMatrix<T, d, d>(0);
+	const PMatrix<T, d, d> WorldSpaceInvI2 = bIsRigidDynamic1 ? Utilities::ComputeWorldSpaceInertia(Q1, PBDRigid1->InvI()) : PMatrix<T, d, d>(0);;
 	for (int32 SpringIndex = 0; SpringIndex < NumSprings; ++SpringIndex)
 	{
 		const TVector<T, d>& Distance0 = Distances[ConstraintIndex][SpringIndex][0];
@@ -140,7 +147,7 @@ void TPBDRigidDynamicSpringConstraints<T, d>::ApplySingle(const T Dt, int32 Cons
 		const TVector<T, d> WorldSpaceX2 = Q1.RotateVector(Distance1) + P1;
 		const TVector<T, d> Delta = GetDelta(WorldSpaceX1, WorldSpaceX2, ConstraintIndex, SpringIndex);
 
-		if (PBDRigid0)
+		if (bIsRigidDynamic0)
 		{
 			const TVector<T, d> Radius = WorldSpaceX1 - P0;
 			P0 += PBDRigid0->InvM() * Delta;
@@ -148,7 +155,7 @@ void TPBDRigidDynamicSpringConstraints<T, d>::ApplySingle(const T Dt, int32 Cons
 			Q0.Normalize();
 		}
 
-		if (PBDRigid1)
+		if (bIsRigidDynamic1)
 		{
 			const TVector<T, d> Radius = WorldSpaceX2 - P1;
 			P1 -= PBDRigid1->InvM() * Delta;
diff --git a/Engine/Source/Runtime/Experimental/Chaos/Private/Chaos/PBDRigidSpringConstraints.cpp b/Engine/Source/Runtime/Experimental/Chaos/Private/Chaos/PBDRigidSpringConstraints.cpp
--- a/Engine/Source/Runtime/Experimental/Chaos/Private/Chaos/PBDRigidSpringConstraints.cpp
+++ b/Engine/Source/Runtime/Experimental/Chaos/Private/Chaos/PBDRigidSpringConstraints.cpp
@@ -22,10 +22,12 @@ namespace Chaos
 	TVector<T, d> TPBDRigidSpringConstraints<T, d>::GetDelta(int32 ConstraintIndex, const TVector<T, d>& WorldSpaceX1, const TVector<T, d>& WorldSpaceX2) const
 	{
 		const TVector<TGeometryParticleHandle<T, d>*, 2>& Constraint = Constraints[ConstraintIndex];
-		const TPBDRigidParticleHandle<T, d>* PBDRigid0 = Constraint[0]->AsDynamic();
-		const TPBDRigidParticleHandle<T, d>* PBDRigid1 = Constraint[1]->AsDynamic();
+		const TPBDRigidParticleHandle<T, d>* PBDRigid0 = Constraint[0]->CastToRigidParticle();
+		const TPBDRigidParticleHandle<T, d>* PBDRigid1 = Constraint[1]->CastToRigidParticle();
+		const bool bIsRigidDynamic0 = PBDRigid0 && PBDRigid0->ObjectState() == EObjectStateType::Dynamic;
+		const bool bIsRigidDynamic1 = PBDRigid1 && PBDRigid1->ObjectState() == EObjectStateType::Dynamic;
 
-		if (!PBDRigid0 && !PBDRigid1)
+		if (!bIsRigidDynamic0 && !bIsRigidDynamic1)
 		{
 			return TVector<T, d>(0);
 		}
@@ -37,8 +39,8 @@ namespace Chaos
 
 		const TVector<T, d> Direction = Difference / Distance;
 		const TVector<T, d> Delta = (Distance - SpringDistances[ConstraintIndex]) * Direction;
-		const T InvM0 = (PBDRigid0) ? PBDRigid0->InvM() : (T)0;
-		const T InvM1 = (PBDRigid1) ? PBDRigid1->InvM() : (T)0;
+		const T InvM0 = (bIsRigidDynamic0) ? PBDRigid0->InvM() : (T)0;
+		const T InvM1 = (bIsRigidDynamic1) ? PBDRigid1->InvM() : (T)0;
 		const T CombinedMass = InvM0 + InvM1;
 
 		return Stiffness * Delta / CombinedMass;
@@ -59,23 +61,26 @@ namespace Chaos
 	{
 		const FConstrainedParticlePair& Constraint = Constraints[ConstraintIndex];
 
-		TPBDRigidParticleHandle<T, d>* PBDRigid0 = Constraint[0]->AsDynamic();
-		TPBDRigidParticleHandle<T, d>* PBDRigid1 = Constraint[1]->AsDynamic();
-		check((PBDRigid0 && PBDRigid1 && PBDRigid0->Island() == PBDRigid1->Island()) || (!PBDRigid0 && PBDRigid1) || (PBDRigid0 && PBDRigid1));
+		TPBDRigidParticleHandle<T, d>* PBDRigid0 = Constraint[0]->CastToRigidParticle();
+		TPBDRigidParticleHandle<T, d>* PBDRigid1 = Constraint[1]->CastToRigidParticle();
+		const bool bIsRigidDynamic0 = PBDRigid0 && PBDRigid0->ObjectState() == EObjectStateType::Dynamic;
+		const bool bIsRigidDynamic1 = PBDRigid1 && PBDRigid1->ObjectState() == EObjectStateType::Dynamic;
+
+		check((bIsRigidDynamic0 && bIsRigidDynamic1 && PBDRigid0->Island() == PBDRigid1->Island()) || (!bIsRigidDynamic0 && bIsRigidDynamic1) || (bIsRigidDynamic0 && bIsRigidDynamic1));
 
 		// @todo(ccaulfield): see if we can eliminate the need for all these ifs
-		const TRotation<T, d> & Q0 = PBDRigid0 ? PBDRigid0->Q() : Constraint[0]->R();
-		const TRotation<T, d> & Q1 = PBDRigid1 ? PBDRigid1->Q() : Constraint[1]->R();
-		const TVector<T, d> & P0 = PBDRigid0 ? PBDRigid0->P() : Constraint[0]->X();
-		const TVector<T, d> & P1 = PBDRigid1 ? PBDRigid1->P() : Constraint[1]->X();
+		const TRotation<T, d> & Q0 = bIsRigidDynamic0 ? PBDRigid0->Q() : Constraint[0]->R();
+		const TRotation<T, d> & Q1 = bIsRigidDynamic1 ? PBDRigid1->Q() : Constraint[1]->R();
+		const TVector<T, d> & P0 = bIsRigidDynamic0 ? PBDRigid0->P() : Constraint[0]->X();
+		const TVector<T, d> & P1 = bIsRigidDynamic1 ? PBDRigid1->P() : Constraint[1]->X();
 
 		const TVector<T, d> WorldSpaceX1 = Q0.RotateVector(Distances[ConstraintIndex][0]) + P0;
 		const TVector<T, d> WorldSpaceX2 = Q1.RotateVector(Distances[ConstraintIndex][1]) + P1;
-		const PMatrix<T, d, d> WorldSpaceInvI1 = PBDRigid0 ? Utilities::ComputeWorldSpaceInertia(Q0, PBDRigid0->InvI()) : PMatrix<T, d, d>(0);
-		const PMatrix<T, d, d> WorldSpaceInvI2 = PBDRigid1 ? Utilities::ComputeWorldSpaceInertia(Q1, PBDRigid1->InvI()) : PMatrix<T, d, d>(0);
+		const PMatrix<T, d, d> WorldSpaceInvI1 = bIsRigidDynamic0 ? Utilities::ComputeWorldSpaceInertia(Q0, PBDRigid0->InvI()) : PMatrix<T, d, d>(0);
+		const PMatrix<T, d, d> WorldSpaceInvI2 = bIsRigidDynamic1 ? Utilities::ComputeWorldSpaceInertia(Q1, PBDRigid1->InvI()) : PMatrix<T, d, d>(0);
 		const TVector<T, d> Delta = GetDelta(ConstraintIndex, WorldSpaceX1, WorldSpaceX2);
 
-		if (PBDRigid0)
+		if (bIsRigidDynamic0)
 		{
 			const TVector<T, d> Radius = WorldSpaceX1 - PBDRigid0->P();
 			PBDRigid0->P() += PBDRigid0->InvM() * Delta;
@@ -83,7 +88,7 @@ namespace Chaos
 			PBDRigid0->Q().Normalize();
 		}
 
-		if (PBDRigid1)
+		if (bIsRigidDynamic1)
 		{
 			const TVector<T, d> Radius = WorldSpaceX2 - PBDRigid1->P();
 			PBDRigid1->P() -= PBDRigid1->InvM() * Delta;
diff --git a/Engine/Source/Runtime/Experimental/Chaos/Private/Chaos/PBDRigidsEvolutionGBF.cpp b/Engine/Source/Runtime/Experimental/Chaos/Private/Chaos/PBDRigidsEvolutionGBF.cpp
--- a/Engine/Source/Runtime/Experimental/Chaos/Private/Chaos/PBDRigidsEvolutionGBF.cpp
+++ b/Engine/Source/Runtime/Experimental/Chaos/Private/Chaos/PBDRigidsEvolutionGBF.cpp
@@ -205,7 +205,8 @@ void TPBDRigidsEvolutionGBF<T, d>::AdvanceOneTimeStep(const T Dt, const T StepFr
 				// @todo(mlentine): Find a good way of not doing this when we aren't using this functionality
 
 				// increment the disable count for the particle
-				if (auto PBDRigid = Particle->AsDynamic())
+				auto PBDRigid = Particle->CastToRigidParticle();
+				if(PBDRigid && PBDRigid->ObjectState() == EObjectStateType::Dynamic)
 				{
 					if (PBDRigid->AuxilaryValue(PhysicsMaterials) && PBDRigid->V().SizeSquared() < PBDRigid->AuxilaryValue(PhysicsMaterials)->DisabledLinearThreshold &&
 						PBDRigid->W().SizeSquared() < PBDRigid->AuxilaryValue(PhysicsMaterials)->DisabledAngularThreshold)
@@ -267,9 +268,9 @@ TPBDRigidsEvolutionGBF<T, d>::TPBDRigidsEvolutionGBF(TPBDRigidsSOAs<T, d>& InPar
 {
 	SetParticleUpdateVelocityFunction([PBDUpdateRule = TPerParticlePBDUpdateFromDeltaPosition<float, 3>(), this](const TArray<TGeometryParticleHandle<T, d>*>& ParticlesInput, const T Dt) {
 		ParticlesParallelFor(ParticlesInput, [&](auto& Particle, int32 Index) {
-			if (Particle->AsDynamic())
+			if (Particle->CastToRigidParticle() && Particle->ObjectState() == EObjectStateType::Dynamic)
 			{
-				PBDUpdateRule.Apply(Particle->AsDynamic(), Dt);
+				PBDUpdateRule.Apply(Particle->CastToRigidParticle(), Dt);
 			}
 		});
 	});
diff --git a/Engine/Source/Runtime/Experimental/Chaos/Public/Chaos/BoundingVolumeUtilities.h b/Engine/Source/Runtime/Experimental/Chaos/Public/Chaos/BoundingVolumeUtilities.h
--- a/Engine/Source/Runtime/Experimental/Chaos/Public/Chaos/BoundingVolumeUtilities.h
+++ b/Engine/Source/Runtime/Experimental/Chaos/Public/Chaos/BoundingVolumeUtilities.h
@@ -161,8 +161,10 @@ TBox<T, d> ComputeWorldSpaceBoundingBox(const TPBDRigidParticles<T, d>& Objects,
 template<typename THandle, typename T, int d, bool bPersistent>
 TBox<T, d> ComputeWorldSpaceBoundingBoxForHandle(const THandle& Handle)
 {
-	const auto PBDRigid = Handle.AsDynamic();
-	TRigidTransform<T, d> LocalToWorld = PBDRigid ? TRigidTransform<T, d>(PBDRigid->P(), PBDRigid->Q()) : TRigidTransform<T, d>(Handle.X(), Handle.R());
+	const auto PBDRigid = Handle.CastToRigidParticle();
+	const bool bIsRigidDynamic = PBDRigid && PBDRigid->ObjectState() == EObjectStateType::Dynamic;
+
+	TRigidTransform<T, d> LocalToWorld = bIsRigidDynamic ? TRigidTransform<T, d>(PBDRigid->P(), PBDRigid->Q()) : TRigidTransform<T, d>(Handle.X(), Handle.R());
 	if(Handle.Geometry())
 	{
 		const auto& LocalBoundingBox = Handle.Geometry()->BoundingBox();
diff --git a/Engine/Source/Runtime/Experimental/Chaos/Public/Chaos/GeometryParticles.h b/Engine/Source/Runtime/Experimental/Chaos/Public/Chaos/GeometryParticles.h
--- a/Engine/Source/Runtime/Experimental/Chaos/Public/Chaos/GeometryParticles.h
+++ b/Engine/Source/Runtime/Experimental/Chaos/Public/Chaos/GeometryParticles.h
@@ -69,7 +69,7 @@ namespace Chaos
 	{
 		Static,
 		Kinematic,
-		Dynamic,
+		Rigid,
 		Clustered,	//only applicable on physics thread side
 		StaticMesh,
 		SkeletalMesh,
diff --git a/Engine/Source/Runtime/Experimental/Chaos/Public/Chaos/PBDCollisionConstraintsImp.h b/Engine/Source/Runtime/Experimental/Chaos/Public/Chaos/PBDCollisionConstraintsImp.h
--- a/Engine/Source/Runtime/Experimental/Chaos/Public/Chaos/PBDCollisionConstraintsImp.h
+++ b/Engine/Source/Runtime/Experimental/Chaos/Public/Chaos/PBDCollisionConstraintsImp.h
@@ -329,7 +329,8 @@ namespace Chaos
 					continue;
 				}
 
-				if (bBody1Bounded == bBody2Bounded && Particle2.AsDynamic())
+				const bool bIsParticle2Dynamic = Particle2.CastToRigidParticle() && Particle2.ObjectState() == EObjectStateType::Dynamic;
+				if (bBody1Bounded == bBody2Bounded && bIsParticle2Dynamic)
 				{
 					//no bidirectional constraints. 
 					if (Particle2.ParticleID() > Particle1.ParticleID())
@@ -338,7 +339,7 @@ namespace Chaos
 					}
 				}
 
-				const TVector<T, d> Box2Thickness = Particle2.AsDynamic() ? ComputeThickness(*Particle2.AsDynamic(), Dt) : TVector<T, d>(0);
+				const TVector<T, d> Box2Thickness = bIsParticle2Dynamic ? ComputeThickness(*Particle2.CastToRigidParticle(), Dt) : TVector<T, d>(0);
 				const T UseThickness = FMath::Max(Box1Thickness, Box2Thickness.Size());// + MThickness
 
 				FCollisionConstraintsArray Constraints;
diff --git a/Engine/Source/Runtime/Experimental/Chaos/Public/Chaos/PBDConstraintRule.h b/Engine/Source/Runtime/Experimental/Chaos/Public/Chaos/PBDConstraintRule.h
--- a/Engine/Source/Runtime/Experimental/Chaos/Public/Chaos/PBDConstraintRule.h
+++ b/Engine/Source/Runtime/Experimental/Chaos/Public/Chaos/PBDConstraintRule.h
@@ -270,11 +270,14 @@ namespace Chaos
 							TVector<const TGeometryParticleHandle<FReal, 3>*, 2> Particles = Handle->GetConstrainedParticles();
 							if (It == NumIts - 1)
 							{
-								if (Particles[0]->AsDynamic() == nullptr || IsTemporarilyStatic.Contains(Particles[0]))
+								const bool bIsParticleDynamic0 = Particles[0]->CastToRigidParticle() && Particles[0]->ObjectState() == EObjectStateType::Dynamic;
+								const bool bIsParticleDynamic1 = Particles[1]->CastToRigidParticle() && Particles[1]->ObjectState() == EObjectStateType::Dynamic;
+
+								if (bIsParticleDynamic0 == false || IsTemporarilyStatic.Contains(Particles[0]))
 								{
 									IsTemporarilyStatic.Add(Particles[1]);
 								}
-								else if (Particles[1]->AsDynamic() == nullptr || IsTemporarilyStatic.Contains(Particles[1]))
+								else if (bIsParticleDynamic1 == false || IsTemporarilyStatic.Contains(Particles[1]))
 								{
 									IsTemporarilyStatic.Add(Particles[0]);
 								}
diff --git a/Engine/Source/Runtime/Experimental/Chaos/Public/Chaos/PBDRigidParticles.h b/Engine/Source/Runtime/Experimental/Chaos/Public/Chaos/PBDRigidParticles.h
--- a/Engine/Source/Runtime/Experimental/Chaos/Public/Chaos/PBDRigidParticles.h
+++ b/Engine/Source/Runtime/Experimental/Chaos/Public/Chaos/PBDRigidParticles.h
@@ -28,7 +28,7 @@ class TPBDRigidParticles : public TRigidParticles<T, d>
 	CHAOS_API TPBDRigidParticles()
 	    : TRigidParticles<T, d>()
 	{
-		this->MParticleType = EParticleType::Dynamic;
+		this->MParticleType = EParticleType::Rigid;
 		TArrayCollection::AddArray(&MP);
 		TArrayCollection::AddArray(&MQ);
 		TArrayCollection::AddArray(&MPreV);
@@ -42,7 +42,7 @@ class TPBDRigidParticles : public TRigidParticles<T, d>
 		, MPreV(MoveTemp(Other.MPreV))
 		, MPreW(MoveTemp(Other.MPreW))
 	{
-		this->MParticleType = EParticleType::Dynamic;
+		this->MParticleType = EParticleType::Rigid;
 		TArrayCollection::AddArray(&MP);
 		TArrayCollection::AddArray(&MQ);
 		TArrayCollection::AddArray(&MPreV);
@@ -125,6 +125,9 @@ class TPBDRigidParticles : public TRigidParticles<T, d>
 				1.f / this->I(Index).M[0][0], 0.f, 0.f,
 				0.f, 1.f / this->I(Index).M[1][1], 0.f,
 				0.f, 0.f, 1.f / this->I(Index).M[2][2]);
+
+			this->P(Index) = this->X(Index);
+			this->Q(Index) = this->R(Index);
 		}
 		else if (InObjectState == EObjectStateType::Sleeping)
 		{
diff --git a/Engine/Source/Runtime/Experimental/Chaos/Public/Chaos/PBDRigidsEvolutionGBF.h b/Engine/Source/Runtime/Experimental/Chaos/Public/Chaos/PBDRigidsEvolutionGBF.h
--- a/Engine/Source/Runtime/Experimental/Chaos/Public/Chaos/PBDRigidsEvolutionGBF.h
+++ b/Engine/Source/Runtime/Experimental/Chaos/Public/Chaos/PBDRigidsEvolutionGBF.h
@@ -121,7 +121,8 @@ public:
 		InParticles.ParallelFor([&](auto& GeomParticle, int32 Index)
 		{
 			//question: can we enforce this at the API layer? Right now islands contain non dynamic which makes this hard
-			if (auto PBDParticle = GeomParticle.AsDynamic())
+			auto PBDParticle = GeomParticle.CastToRigidParticle();
+			if(PBDParticle && PBDParticle->ObjectState() == EObjectStateType::Dynamic)
 			{
 				auto& Particle = *PBDParticle;
 
diff --git a/Engine/Source/Runtime/Experimental/Chaos/Public/Chaos/PBDRigidsSOAs.h b/Engine/Source/Runtime/Experimental/Chaos/Public/Chaos/PBDRigidsSOAs.h
--- a/Engine/Source/Runtime/Experimental/Chaos/Public/Chaos/PBDRigidsSOAs.h
+++ b/Engine/Source/Runtime/Experimental/Chaos/Public/Chaos/PBDRigidsSOAs.h
@@ -104,9 +104,9 @@ public:
 
 	void DestroyParticle(TGeometryParticleHandle<T, d>* Particle)
 	{
-		check(Particle->AsClustered() == nullptr);	//not supported
+		check(Particle->CastToClustered() == nullptr);	//not supported
 
-		auto PBDRigid = Particle->AsDynamic();
+		auto PBDRigid = Particle->CastToRigidParticle();
 		if(PBDRigid && PBDRigid->ObjectState() == EObjectStateType::Dynamic)
 		{
 			RemoveFromMapAndArray(PBDRigid, ActiveParticlesToIndex, ActiveParticlesArray);
@@ -118,13 +118,13 @@ public:
 
 	void DisableParticle(TGeometryParticleHandle<T, d>* Particle)
 	{
-		if (auto PBDRigid = Particle->AsDynamic())
+		if (auto PBDRigid = Particle->CastToRigidParticle())
 		{
 			PBDRigid->Disabled() = true;
 			PBDRigid->V() = TVector<T, d>(0);
 			PBDRigid->W() = TVector<T, d>(0);
 
-			if (auto PBDRigidClustered = Particle->AsClustered())
+			if (auto PBDRigidClustered = Particle->CastToClustered())
 			{
 				RemoveFromMapAndArray(PBDRigidClustered, NonDisabledClusteredToIndex, NonDisabledClusteredArray);
 				RemoveFromMapAndArray(PBDRigidClustered, ActiveClusteredToIndex, ActiveClusteredArray);
@@ -139,7 +139,7 @@ public:
 				RemoveFromMapAndArray(PBDRigid, ActiveParticlesToIndex, ActiveParticlesArray);
 			}
 		}
-		else if (Particle->AsKinematic())
+		else if (Particle->CastToKinematicParticle())
 		{
 			Particle->MoveToSOA(*KinematicDisabledParticles);
 		}
@@ -152,9 +152,9 @@ public:
 
 	void EnableParticle(TGeometryParticleHandle<T, d>* Particle)
 	{
-		if (auto PBDRigid = Particle->AsDynamic())
+		if (auto PBDRigid = Particle->CastToRigidParticle())
 		{
-			if (auto PBDRigidClustered = Particle->AsClustered())
+			if (auto PBDRigidClustered = Particle->CastToClustered())
 			{
 				InsertToMapAndArray(PBDRigidClustered, NonDisabledClusteredToIndex, NonDisabledClusteredArray);
 				if (!PBDRigid->Sleeping() && Particle->ObjectState() == EObjectStateType::Dynamic)
@@ -164,7 +164,7 @@ public:
 			}
 			else
 			{
-				SetDynamicParticleSOA(Particle->AsDynamic());
+				SetDynamicParticleSOA(Particle->CastToRigidParticle());
 			}
 
 			if (!PBDRigid->Sleeping() && Particle->ObjectState() == EObjectStateType::Dynamic)
@@ -174,7 +174,7 @@ public:
 
 			PBDRigid->Disabled() = false;
 		}
-		else if (Particle->AsKinematic())
+		else if (Particle->CastToKinematicParticle())
 		{
 			Particle->MoveToSOA(*KinematicParticles);
 		}
@@ -187,11 +187,11 @@ public:
 
 	void ActivateParticle(TGeometryParticleHandle<T, d>* Particle)
 	{
-		auto PBDRigid = Particle->AsDynamic();
+		auto PBDRigid = Particle->CastToRigidParticle();
 		if(PBDRigid && PBDRigid->ObjectState() == EObjectStateType::Dynamic)
 		{
 			check(!PBDRigid->Disabled());
-			if (auto PBDRigidClustered = Particle->AsClustered())
+			if (auto PBDRigidClustered = Particle->CastToClustered())
 			{
 				InsertToMapAndArray(PBDRigidClustered, ActiveClusteredToIndex, ActiveClusteredArray);
 			}
@@ -204,14 +204,14 @@ public:
 
 	void DeactivateParticle(TGeometryParticleHandle<T, d>* Particle)
 	{
-		auto PBDRigid = Particle->AsDynamic();
+		auto PBDRigid = Particle->CastToRigidParticle();
 		if(PBDRigid)
 		{
 			if (   PBDRigid->ObjectState() == EObjectStateType::Dynamic
 				|| PBDRigid->ObjectState() == EObjectStateType::Sleeping)
 			{
 				check(!PBDRigid->Disabled());
-				if (auto PBDRigidClustered = Particle->AsClustered())
+				if (auto PBDRigidClustered = Particle->CastToClustered())
 				{
 					RemoveFromMapAndArray(PBDRigidClustered, ActiveClusteredToIndex, ActiveClusteredArray);
 				}
@@ -238,11 +238,11 @@ public:
 
 		if (Particle->ObjectState() != EObjectStateType::Dynamic)
 		{
-			RemoveFromMapAndArray(Particle->AsDynamic(), ActiveParticlesToIndex, ActiveParticlesArray);
+			RemoveFromMapAndArray(Particle->CastToRigidParticle(), ActiveParticlesToIndex, ActiveParticlesArray);
 		}
 		else
 		{
-			InsertToMapAndArray(Particle->AsDynamic(), ActiveParticlesToIndex, ActiveParticlesArray);
+			InsertToMapAndArray(Particle->CastToRigidParticle(), ActiveParticlesToIndex, ActiveParticlesArray);
 		}
 
 		// Move to appropriate dynamic SOA
diff --git a/Engine/Source/Runtime/Experimental/Chaos/Public/Chaos/ParticleHandle.h b/Engine/Source/Runtime/Experimental/Chaos/Public/Chaos/ParticleHandle.h
--- a/Engine/Source/Runtime/Experimental/Chaos/Public/Chaos/ParticleHandle.h
+++ b/Engine/Source/Runtime/Experimental/Chaos/Public/Chaos/ParticleHandle.h
@@ -405,14 +405,14 @@ public:
 	TGeometryParticle<T, d>* GTGeometryParticle() const { return GeometryParticles->GTGeometryParticle(ParticleIdx); }
 	TGeometryParticle<T, d>*& GTGeometryParticle() { return GeometryParticles->GTGeometryParticle(ParticleIdx); }
 
-	const TKinematicGeometryParticleHandleImp<T, d, bPersistent>* AsKinematic() const;
-	TKinematicGeometryParticleHandleImp<T, d, bPersistent>* AsKinematic();
+	const TKinematicGeometryParticleHandleImp<T, d, bPersistent>* CastToKinematicParticle() const;
+	TKinematicGeometryParticleHandleImp<T, d, bPersistent>* CastToKinematicParticle();
 
-	const TPBDRigidParticleHandleImp<T, d, bPersistent>* AsDynamic() const;
-	TPBDRigidParticleHandleImp<T, d, bPersistent>* AsDynamic();
+	const TPBDRigidParticleHandleImp<T, d, bPersistent>* CastToRigidParticle() const;
+	TPBDRigidParticleHandleImp<T, d, bPersistent>* CastToRigidParticle();
 
-	const TPBDRigidClusteredParticleHandleImp<T, d, bPersistent>* AsClustered() const;
-	TPBDRigidClusteredParticleHandleImp<T, d, bPersistent>* AsClustered();
+	const TPBDRigidClusteredParticleHandleImp<T, d, bPersistent>* CastToClustered() const;
+	TPBDRigidClusteredParticleHandleImp<T, d, bPersistent>* CastToClustered();
 
 	const TGeometryParticleHandle<T, d>* Handle() const { return GetHandleHelper(this);}
 	TGeometryParticleHandle<T, d>* Handle() { return GetHandleHelper(this); }
@@ -500,7 +500,7 @@ public:
 	using TGeometryParticleHandleImp<T, d, bPersistent>::ParticleIdx;
 	using TGeometryParticleHandleImp<T, d, bPersistent>::KinematicGeometryParticles;
 	using TGeometryParticleHandleImp<T, d, bPersistent>::Type;
-	using TGeometryParticleHandleImp<T, d, bPersistent>::AsDynamic;
+	using TGeometryParticleHandleImp<T, d, bPersistent>::CastToRigidParticle;
 	using TTransientHandle = TTransientKinematicGeometryParticleHandle<T, d>;
 	using TSOAType = TKinematicGeometryParticles<T, d>;
 	
@@ -674,7 +674,7 @@ public:
 	const TPBDRigidParticleHandleImp<T, d, true>* Handle() const { return PBDRigidParticles->Handle(ParticleIdx); }
 	TPBDRigidParticleHandleImp<T, d, true>* Handle() { return PBDRigidParticles->Handle(ParticleIdx); }
 
-	static constexpr EParticleType StaticType() { return EParticleType::Dynamic; }
+	static constexpr EParticleType StaticType() { return EParticleType::Rigid; }
 };
 
 template <typename T, int d, bool bPersistent>
@@ -744,7 +744,7 @@ public:
 	const TPBDRigidClusteredParticleHandleImp<T, d, true>* Handle() const { return PBDRigidClusteredParticles->Handle(ParticleIdx); }
 	TPBDRigidClusteredParticleHandleImp<T, d, true>* Handle() { return PBDRigidClusteredParticles->Handle(ParticleIdx); }
 
-	static constexpr EParticleType StaticType() { return EParticleType::Dynamic; }
+	static constexpr EParticleType StaticType() { return EParticleType::Rigid; }
 
 	int32 TransientParticleIndex() const { return ParticleIdx; }
 };
@@ -801,31 +801,34 @@ public:
 };
 
 template <typename T, int d, bool bPersistent>
-const TKinematicGeometryParticleHandleImp<T, d, bPersistent>* TGeometryParticleHandleImp<T,d, bPersistent>::AsKinematic() const { return Type >= EParticleType::Kinematic ? static_cast<const TKinematicGeometryParticleHandleImp<T, d, bPersistent>*>(this) : nullptr; }
+TKinematicGeometryParticleHandleImp<T, d, bPersistent>* TGeometryParticleHandleImp<T, d, bPersistent>::CastToKinematicParticle() { return Type >= EParticleType::Kinematic ? static_cast<TKinematicGeometryParticleHandleImp<T, d, bPersistent>*>(this) : nullptr; }
+
 template <typename T, int d, bool bPersistent>
-TKinematicGeometryParticleHandleImp<T, d, bPersistent>* TGeometryParticleHandleImp<T, d, bPersistent>::AsKinematic() { return Type >= EParticleType::Kinematic ? static_cast<TKinematicGeometryParticleHandleImp<T, d, bPersistent>*>(this) : nullptr; }
+const TKinematicGeometryParticleHandleImp<T, d, bPersistent>* TGeometryParticleHandleImp<T,d, bPersistent>::CastToKinematicParticle() const { return Type >= EParticleType::Kinematic ? static_cast<const TKinematicGeometryParticleHandleImp<T, d, bPersistent>*>(this) : nullptr; }
 
 template <typename T, int d, bool bPersistent>
-const TPBDRigidParticleHandleImp<T, d, bPersistent>* TGeometryParticleHandleImp<T, d, bPersistent>::AsDynamic() const { return Type >= EParticleType::Dynamic ? static_cast<const TPBDRigidParticleHandleImp<T, d, bPersistent>*>(this) : nullptr; }
+const TPBDRigidParticleHandleImp<T, d, bPersistent>* TGeometryParticleHandleImp<T, d, bPersistent>::CastToRigidParticle() const { return Type >= EParticleType::Rigid ? static_cast<const TPBDRigidParticleHandleImp<T, d, bPersistent>*>(this) : nullptr; }
+
 template <typename T, int d, bool bPersistent>
-TPBDRigidParticleHandleImp<T, d, bPersistent>* TGeometryParticleHandleImp<T, d, bPersistent>::AsDynamic() { return Type >= EParticleType::Dynamic ? static_cast<TPBDRigidParticleHandleImp<T, d, bPersistent>*>(this) : nullptr; }
+TPBDRigidParticleHandleImp<T, d, bPersistent>* TGeometryParticleHandleImp<T, d, bPersistent>::CastToRigidParticle() { return Type >= EParticleType::Rigid ? static_cast<TPBDRigidParticleHandleImp<T, d, bPersistent>*>(this) : nullptr; }
 
 template <typename T, int d, bool bPersistent>
-const TPBDRigidClusteredParticleHandleImp<T, d, bPersistent>* TGeometryParticleHandleImp<T, d, bPersistent>::AsClustered() const { return Type >= EParticleType::Clustered ? static_cast<const TPBDRigidClusteredParticleHandleImp<T, d, bPersistent>*>(this) : nullptr; }
+const TPBDRigidClusteredParticleHandleImp<T, d, bPersistent>* TGeometryParticleHandleImp<T, d, bPersistent>::CastToClustered() const { return Type >= EParticleType::Clustered ? static_cast<const TPBDRigidClusteredParticleHandleImp<T, d, bPersistent>*>(this) : nullptr; }
+
 template <typename T, int d, bool bPersistent>
-TPBDRigidClusteredParticleHandleImp<T, d, bPersistent>* TGeometryParticleHandleImp<T, d, bPersistent>::AsClustered() { return Type >= EParticleType::Clustered ? static_cast<TPBDRigidClusteredParticleHandleImp<T, d, bPersistent>*>(this) : nullptr; }
+TPBDRigidClusteredParticleHandleImp<T, d, bPersistent>* TGeometryParticleHandleImp<T, d, bPersistent>::CastToClustered() { return Type >= EParticleType::Clustered ? static_cast<TPBDRigidClusteredParticleHandleImp<T, d, bPersistent>*>(this) : nullptr; }
 
 template <typename T, int d, bool bPersistent>
 EObjectStateType TGeometryParticleHandleImp<T,d, bPersistent>::ObjectState() const
 {
-	const TKinematicGeometryParticleHandleImp<T, d, bPersistent>* Kin = AsKinematic();
+	const TKinematicGeometryParticleHandleImp<T, d, bPersistent>* Kin = CastToKinematicParticle();
 	return Kin ? Kin->ObjectState() : EObjectStateType::Static;
 }
 
 template <typename T, int d, bool bPersistent>
 EObjectStateType TKinematicGeometryParticleHandleImp<T, d, bPersistent>::ObjectState() const
 {
-	const TPBDRigidParticleHandleImp<T, d, bPersistent>* Dyn = AsDynamic();
+	const TPBDRigidParticleHandleImp<T, d, bPersistent>* Dyn = CastToRigidParticle();
 	return Dyn ? Dyn->ObjectState() : EObjectStateType::Kinematic;
 }
 
@@ -840,9 +843,8 @@ FString TGeometryParticleHandleImp<T, d, bPersistent>::ToString() const
 	case EParticleType::Kinematic:
 		return FString::Printf(TEXT("Kinemmatic[%d]"), ParticleIdx);
 		break;
-	case EParticleType::Dynamic:
+	case EParticleType::Rigid:
 		return FString::Printf(TEXT("Dynamic[%d]"), ParticleIdx);
-		break;
 	case EParticleType::GeometryCollection:
 		return FString::Printf(TEXT("GeometryCollection[%d]"), ParticleIdx);
 		break;
@@ -882,10 +884,10 @@ public:
 		bool IsKinematic() const { return (Handle->ObjectState() == EObjectStateType::Kinematic); }
 		bool IsDynamic() const { return (Handle->ObjectState() == EObjectStateType::Dynamic) || (Handle->ObjectState() == EObjectStateType::Sleeping); }
 
-		const TKinematicGeometryParticleHandle<T, d>* AsKinematic() const { return Handle->AsKinematic(); }
-		TKinematicGeometryParticleHandle<T, d>* AsKinematic() { return Handle->AsKinematic(); }
-		const TPBDRigidParticleHandle<T, d>* AsDynamic() const { return Handle->AsDynamic(); }
-		TPBDRigidParticleHandle<T, d>* AsDynamic() { return Handle->AsDynamic(); }
+		const TKinematicGeometryParticleHandle<T, d>* CastToKinematicParticle() const { return Handle->CastToKinematicParticle(); }
+		TKinematicGeometryParticleHandle<T, d>* CastToKinematicParticle() { return Handle->CastToKinematicParticle(); }
+		const TPBDRigidParticleHandle<T, d>* CastToRigidParticle() const { return Handle->CastToRigidParticle(); }
+		TPBDRigidParticleHandle<T, d>* CastToRigidParticle() { return Handle->CastToRigidParticle(); }
 
 		// Static Particles
 		const TVector<T, d>& X() const { return Handle->X(); }
@@ -901,26 +903,172 @@ public:
 		auto& AuxilaryValue(Container& AuxContainer) { return Handle->AuxilaryValue(AuxContainer); }
 
 		// Kinematic Particles
-		const TVector<T, d>& V() const { return (Handle->AsKinematic()) ? Handle->AsKinematic()->V() : ZeroVector; }
-		const TVector<T, d>& W() const { return (Handle->AsKinematic()) ? Handle->AsKinematic()->W() : ZeroVector; }
+		const TVector<T, d>& V() const { return (Handle->CastToKinematicParticle()) ? Handle->CastToKinematicParticle()->V() : ZeroVector; }
+		const TVector<T, d>& W() const { return (Handle->CastToKinematicParticle()) ? Handle->CastToKinematicParticle()->W() : ZeroVector; }
 
 		// Dynamic Particles
-		int32 CollisionParticlesSize() const { return (Handle->AsDynamic()) ? Handle->AsDynamic()->CollisionParticlesSize() : 0; }
-		const TUniquePtr<TBVHParticles<T, d>>& CollisionParticles() const { return (Handle->AsDynamic()) ? Handle->AsDynamic()->CollisionParticles() : NullBVHParticles; }
-		int32 CollisionGroup() const { return (Handle->AsDynamic()) ? Handle->AsDynamic()->CollisionGroup() : 0; }
-		bool Disabled() const { return (Handle->AsDynamic()) ? Handle->AsDynamic()->Disabled() : false; }				// @todo(ccaulfield): should be available on all types?
-		const TVector<T, d>& PreV() const { return (Handle->AsDynamic()) ? Handle->AsDynamic()->PreV() : ZeroVector; }	// @todo(ccaulfield): should be available on kinematics?
-		const TVector<T, d>& PreW() const { return (Handle->AsDynamic()) ? Handle->AsDynamic()->PreW() : ZeroVector; }	// @todo(ccaulfield): should be available on kinematics?
-		const TVector<T, d>& P() const { return (Handle->AsDynamic()) ? Handle->AsDynamic()->P() : X(); }
-		const TRotation<T, d>& Q() const { return (Handle->AsDynamic()) ? Handle->AsDynamic()->Q() : R(); }
-		const TVector<T, d>& F() const { return (Handle->AsDynamic()) ? Handle->AsDynamic()->F() : ZeroVector; }
-		const TVector<T, d>& Torque() const { return (Handle->AsDynamic()) ? Handle->AsDynamic()->Torque() : ZeroVector; }
-		const PMatrix<T, d, d>& I() const { return (Handle->AsDynamic()) ? Handle->AsDynamic()->I() : ZeroMatrix; }
-		const PMatrix<T, d, d>& InvI() const { return (Handle->AsDynamic()) ? Handle->AsDynamic()->InvI() : ZeroMatrix; }
-		T M() const { return (Handle->AsDynamic()) ? Handle->AsDynamic()->M() : (T)0; }
-		T InvM() const { return (Handle->AsDynamic()) ? Handle->AsDynamic()->InvM() : (T)0; }
-		int32 Island() const { return (Handle->AsDynamic()) ? Handle->AsDynamic()->Island() : INDEX_NONE; }
-		bool ToBeRemovedOnFracture() const { return (Handle->AsDynamic()) ? Handle->AsDynamic()->ToBeRemovedOnFracture() : false; }
+
+		// TODO: Make all of these check ObjectState to maintain current functionality
+		int32 CollisionParticlesSize() const
+		{
+			if (Handle->CastToRigidParticle() && Handle->ObjectState() == EObjectStateType::Dynamic)
+			{
+				return Handle->CastToRigidParticle()->CollisionParticlesSize();
+			}
+
+			return 0;
+		}
+
+		const TUniquePtr<TBVHParticles<T, d>>& CollisionParticles() const
+		{
+			if (Handle->CastToRigidParticle() && Handle->ObjectState() == EObjectStateType::Dynamic)
+			{
+				return Handle->CastToRigidParticle()->CollisionParticles();
+			}
+
+			return NullBVHParticles;
+		}
+
+		int32 CollisionGroup() const
+		{
+			if (Handle->CastToRigidParticle() && Handle->ObjectState() == EObjectStateType::Dynamic)
+			{
+				return Handle->CastToRigidParticle()->CollisionGroup();
+			}
+
+			return 0;
+		}
+
+		// @todo(ccaulfield): should be available on all types?
+		bool Disabled() const
+		{
+			if (Handle->CastToRigidParticle() && Handle->ObjectState() == EObjectStateType::Dynamic)
+			{
+				return Handle->CastToRigidParticle()->Disabled();
+			}
+
+			return false;
+		}
+
+		// @todo(ccaulfield): should be available on kinematics?
+		const TVector<T, d>& PreV() const
+		{
+			if (Handle->CastToRigidParticle() && Handle->ObjectState() == EObjectStateType::Dynamic)
+			{
+				return Handle->CastToRigidParticle()->PreV();
+			}
+
+			return ZeroVector;
+		}
+
+		// @todo(ccaulfield): should be available on kinematics?
+		const TVector<T, d>& PreW() const
+		{
+			if (Handle->CastToRigidParticle() && Handle->ObjectState() == EObjectStateType::Dynamic)
+			{
+				return Handle->CastToRigidParticle()->PreW();
+			}
+			return ZeroVector;
+		}
+
+		const TVector<T, d>& P() const
+		{
+			if (Handle->CastToRigidParticle() && Handle->ObjectState() == EObjectStateType::Dynamic)
+			{
+				return Handle->CastToRigidParticle()->P();
+			}
+
+			return X();
+		}
+
+		const TRotation<T, d>& Q() const
+		{
+			if (Handle->CastToRigidParticle() && Handle->ObjectState() == EObjectStateType::Dynamic)
+			{
+				return Handle->CastToRigidParticle()->Q();
+			}
+
+			return R();
+		}
+
+		const TVector<T, d>& F() const 
+		{ 
+			if (Handle->CastToRigidParticle() && Handle->ObjectState() == EObjectStateType::Dynamic)
+			{
+				return Handle->CastToRigidParticle()->F();
+			}
+
+			return ZeroVector;
+		}
+		const TVector<T, d>& Torque() const 
+		{ 
+			if (Handle->CastToRigidParticle() && Handle->ObjectState() == EObjectStateType::Dynamic)
+			{
+				return Handle->CastToRigidParticle()->Torque();
+			}
+
+			return ZeroVector;
+		}
+
+		const PMatrix<T, d, d>& I() const 
+		{ 
+			if (Handle->CastToRigidParticle() && Handle->ObjectState() == EObjectStateType::Dynamic)
+			{
+				return Handle->CastToRigidParticle()->I();
+			}
+
+			return ZeroMatrix;
+		}
+
+		const PMatrix<T, d, d>& InvI() const 
+		{ 
+			if (Handle->CastToRigidParticle() && Handle->ObjectState() == EObjectStateType::Dynamic)
+			{
+				return Handle->CastToRigidParticle()->InvI();
+			}
+
+			return ZeroMatrix;
+		}
+
+		T M() const
+		{
+			if (Handle->CastToRigidParticle() && Handle->ObjectState() == EObjectStateType::Dynamic)
+			{
+				return Handle->CastToRigidParticle()->M();
+			}
+
+			return (T)0;
+		}
+
+		T InvM() const
+		{
+			if (Handle->CastToRigidParticle() && Handle->ObjectState() == EObjectStateType::Dynamic)
+			{
+				return Handle->CastToRigidParticle()->InvM();
+			}
+
+			return (T)0;
+		}
+
+		int32 Island() const
+		{
+			if (Handle->CastToRigidParticle() && Handle->ObjectState() == EObjectStateType::Dynamic)
+			{
+				return Handle->CastToRigidParticle()->Island();
+			}
+
+			return INDEX_NONE;
+		}
+
+		bool ToBeRemovedOnFracture() const 
+		{
+			if (Handle->CastToRigidParticle() && Handle->ObjectState() == EObjectStateType::Dynamic)
+			{
+				return Handle->CastToRigidParticle()->ToBeRemovedOnFracture();
+			}
+
+			return false;
+		}
 
 	private:
 		TGeometryParticleHandle<T, d>* Handle;
@@ -1169,11 +1317,12 @@ public:
 		return Type;
 	}
 
-	const TKinematicGeometryParticle<T, d>* AsKinematic() const;
-	TKinematicGeometryParticle<T, d>* AsKinematic();
 
-	const TPBDRigidParticle<T, d>* AsDynamic() const;
-	TPBDRigidParticle<T, d>* AsDynamic();
+	const TKinematicGeometryParticle<T, d>* CastToKinematicParticle() const;
+	TKinematicGeometryParticle<T, d>* CastToKinematicParticle();
+
+	const TPBDRigidParticle<T, d>* CastToRigidParticle() const;
+	TPBDRigidParticle<T, d>* CastToRigidParticle();
 
 	FSpatialAccelerationIdx SpatialIdx() const { return MSpatialIdx; }
 	void SetSpatialIdx(FSpatialAccelerationIdx Idx)
@@ -1328,7 +1477,7 @@ public:
 	typedef TKinematicGeometryParticleHandle<T, d> FHandle;
 
 	using TGeometryParticle<T, d>::Type;
-	using TGeometryParticle<T, d>::AsDynamic;
+	using TGeometryParticle<T, d>::CastToRigidParticle;
 
 protected:
 	friend TGeometryParticle<T,d>* TGeometryParticle<T, d>::SerializationFactory(FChaosArchive& Ar, TGeometryParticle<T, d>* Serializable);
@@ -1424,7 +1573,7 @@ protected:
 	TPBDRigidParticle<T, d>(const TPBDRigidParticleParameters<T, d>& DynamicParams = TPBDRigidParticleParameters<T, d>())
 		: TKinematicGeometryParticle<T, d>(DynamicParams)
 	{
-		Type = EParticleType::Dynamic;
+		Type = EParticleType::Rigid;
 		MIsland = INDEX_NONE;
 		MToBeRemovedOnFracture = false;
 		PBDRigidParticleDefaultConstruct<T, d>(*this, DynamicParams);
@@ -1650,7 +1799,7 @@ public:
 
 	using TKinematicGeometryParticleData<T, d>::Type;
 
-	TPBDRigidParticleData(EParticleType InType = EParticleType::Dynamic)
+	TPBDRigidParticleData(EParticleType InType = EParticleType::Rigid)
 		: Base(InType)
 		, MQ(TRotation<T, d>())
 		, MPreV(TVector<T, d>(0))
@@ -1691,7 +1840,7 @@ public:
 		, MToBeRemovedOnFracture(InParticle.ToBeRemovedOnFracture())
 		, MGravityEnabled(InParticle.IsGravityEnabled())
 	{
-		Type = EParticleType::Dynamic;
+		Type = EParticleType::Rigid;
 	}
 
 
@@ -1715,7 +1864,7 @@ public:
 
 	void Reset() {
 		TKinematicGeometryParticleData<T, d>::Reset();
-		Type = EParticleType::Dynamic;
+		Type = EParticleType::Rigid;
 		MQ = TRotation<T, d>();
 		MPreV = TVector<T, d>(0);
 		MPreW = TVector<T, d>(0);
@@ -1738,26 +1887,62 @@ public:
 
 
 template <typename T, int d>
-const TKinematicGeometryParticle<T, d>* TGeometryParticle<T, d>::AsKinematic() const { return Type >= EParticleType::Kinematic ? static_cast<const TKinematicGeometryParticle<T, d>*>(this) : nullptr; }
+const TKinematicGeometryParticle<T, d>* TGeometryParticle<T, d>::CastToKinematicParticle() const
+{
+	if (Type >= EParticleType::Kinematic)
+	{
+		return static_cast<const TKinematicGeometryParticle<T, d>*>(this);
+	}
+
+	return nullptr;
+}
+
 template <typename T, int d>
-TKinematicGeometryParticle<T, d>* TGeometryParticle<T, d>::AsKinematic() { return Type >= EParticleType::Kinematic ? static_cast<TKinematicGeometryParticle<T, d>*>(this) : nullptr; }
+TKinematicGeometryParticle<T, d>* TGeometryParticle<T, d>::CastToKinematicParticle()
+{
+	if (Type >= EParticleType::Kinematic)
+	{
+		return static_cast<TKinematicGeometryParticle<T, d>*>(this);
+	}
+
+	return nullptr;
+}
 
 template <typename T, int d>
-const TPBDRigidParticle<T, d>* TGeometryParticle<T, d>::AsDynamic() const { return Type >= EParticleType::Dynamic ? static_cast<const TPBDRigidParticle<T, d>*>(this) : nullptr; }
+TPBDRigidParticle<T, d>* TGeometryParticle<T, d>::CastToRigidParticle() 
+{
+	if (Type >= EParticleType::Rigid)
+	{
+		return static_cast<TPBDRigidParticle<T, d>*>(this);
+	}
+
+	return nullptr;
+}
+
+
 template <typename T, int d>
-TPBDRigidParticle<T, d>* TGeometryParticle<T, d>::AsDynamic() { return Type >= EParticleType::Dynamic ? static_cast<TPBDRigidParticle<T, d>*>(this) : nullptr; }
+const TPBDRigidParticle<T, d>* TGeometryParticle<T, d>::CastToRigidParticle()  const
+{
+	if (Type >= EParticleType::Rigid)
+	{
+		return static_cast<const TPBDRigidParticle<T, d>*>(this);
+	}
+
+	return nullptr;
+}
+
 
 template <typename T, int d>
 EObjectStateType TGeometryParticle<T, d>::ObjectState() const
 {
-	const TKinematicGeometryParticle<T, d>* Kin = AsKinematic();
+	const TKinematicGeometryParticle<T, d>* Kin = CastToKinematicParticle();
 	return Kin ? Kin->ObjectState() : EObjectStateType::Static;
 }
 
 template <typename T, int d>
 EObjectStateType TKinematicGeometryParticle<T, d>::ObjectState() const
 {
-	const TPBDRigidParticle<T, d>* Dyn = AsDynamic();
+	const TPBDRigidParticle<T, d>* Dyn = CastToRigidParticle();
 	return Dyn ? Dyn->ObjectState() : EObjectStateType::Kinematic;
 }
 
@@ -1770,7 +1955,7 @@ TGeometryParticle<T, d>* TGeometryParticle<T, d>::SerializationFactory(FChaosArc
 	{
 	case EParticleType::Static: if (Ar.IsLoading()) { return new TGeometryParticle<T, d>(); } break;
 	case EParticleType::Kinematic: if (Ar.IsLoading()) { return new TKinematicGeometryParticle<T, d>(); } break;
-	case EParticleType::Dynamic: if (Ar.IsLoading()) { return new TPBDRigidParticle<T, d>(); } break;
+	case EParticleType::Rigid: if (Ar.IsLoading()) { return new TPBDRigidParticle<T, d>(); } break;
 	case EParticleType::GeometryCollection: if (Ar.IsLoading()) { return new TPBDGeometryCollectionParticle<T, d>(); } break;
 	default:
 		check(false);
diff --git a/Engine/Source/Runtime/Experimental/ChaosSolvers/Private/EventDefaults.cpp b/Engine/Source/Runtime/Experimental/ChaosSolvers/Private/EventDefaults.cpp
--- a/Engine/Source/Runtime/Experimental/ChaosSolvers/Private/EventDefaults.cpp
+++ b/Engine/Source/Runtime/Experimental/ChaosSolvers/Private/EventDefaults.cpp
@@ -79,10 +79,10 @@ namespace Chaos
 							{
 								TGeometryParticleHandle<float, 3>* Particle0 = Constraint.Particle[0];
 								TGeometryParticleHandle<float, 3>* Particle1 = Constraint.Particle[1];
-								TKinematicGeometryParticleHandle<float, 3>* Body0 = Particle0->AsKinematic();
+								TKinematicGeometryParticleHandle<float, 3>* Body0 = Particle0->CastToKinematicParticle();
 
 								// presently when a rigidbody or kinematic hits static geometry then Body1 is null
-								TKinematicGeometryParticleHandle<float, 3>* Body1 = Particle1->AsKinematic();
+								TKinematicGeometryParticleHandle<float, 3>* Body1 = Particle1->CastToKinematicParticle();
 
 								if (!Constraint.AccumulatedImpulse.IsZero() && Body0)
 								{
@@ -122,14 +122,16 @@ namespace Chaos
 
 						// todo: do we need these anymore now we are storing the particles you can access all of this stuff from there
 						// do we still need these now we have pointers to particles returned?
-						if(TPBDRigidParticleHandle<float, 3>* PBDRigid0 = Particle0->AsDynamic())
+						TPBDRigidParticleHandle<float, 3>* PBDRigid0 = Particle0->CastToRigidParticle();
+						if(PBDRigid0 && PBDRigid0->ObjectState() == EObjectStateType::Dynamic)
 						{
 							Data.Velocity1 = PBDRigid0->V();
 							Data.AngularVelocity1 = PBDRigid0->W();
 							Data.Mass1 = PBDRigid0->M();
 						}
 
-						if(TPBDRigidParticleHandle<float, 3>* PBDRigid1 = Particle1->AsDynamic())
+						TPBDRigidParticleHandle<float, 3>* PBDRigid1 = Particle1->CastToRigidParticle();
+						if(PBDRigid1 && PBDRigid1->ObjectState() == EObjectStateType::Dynamic)
 						{
 							Data.Velocity2 = PBDRigid1->V();
 							Data.AngularVelocity2 = PBDRigid1->W();
@@ -290,7 +292,7 @@ namespace Chaos
 
 			for (auto& ActiveParticle : Evolution->GetParticles().GetActiveParticlesView())
 			{
-				TPBDRigidParticle<float, 3>* GTParticle = ActiveParticle.Handle()->GTGeometryParticle()->AsDynamic();
+				TPBDRigidParticle<float, 3>* GTParticle = ActiveParticle.Handle()->GTGeometryParticle()->CastToRigidParticle();
 				// Since Clustered GCs can be unioned the particleIndex representing the union 
 				// is not associated with a PhysicsProxy
 				if (GTParticle && GTParticle->Proxy != nullptr)
diff --git a/Engine/Source/Runtime/Experimental/ChaosSolvers/Private/PBDRigidsSolver.cpp b/Engine/Source/Runtime/Experimental/ChaosSolvers/Private/PBDRigidsSolver.cpp
--- a/Engine/Source/Runtime/Experimental/ChaosSolvers/Private/PBDRigidsSolver.cpp
+++ b/Engine/Source/Runtime/Experimental/ChaosSolvers/Private/PBDRigidsSolver.cpp
@@ -169,16 +169,16 @@ namespace Chaos
 
 		// Make a physics proxy, giving it our particle and particle handle
 		const EParticleType InParticleType = GTParticle->ObjectType();
-		if (InParticleType == EParticleType::Dynamic)
+		if (InParticleType == EParticleType::Rigid)
 		{
-			auto Proxy = new FRigidParticlePhysicsProxy(GTParticle->AsDynamic(), nullptr);
+			auto Proxy = new FRigidParticlePhysicsProxy(GTParticle->CastToRigidParticle(), nullptr);
 			RigidParticlePhysicsProxies.Add((FRigidParticlePhysicsProxy*)Proxy);
 			ProxyData = Proxy->NewData();
 			ProxyBase = Proxy;
 		}
 		else if (InParticleType == EParticleType::Kinematic)
 		{
-			auto Proxy = new FKinematicGeometryParticlePhysicsProxy(GTParticle->AsKinematic(), nullptr);
+			auto Proxy = new FKinematicGeometryParticlePhysicsProxy(GTParticle->CastToKinematicParticle(), nullptr);
 			KinematicGeometryParticlePhysicsProxies.Add((FKinematicGeometryParticlePhysicsProxy*)Proxy);
 			ProxyData = Proxy->NewData();
 			ProxyBase = Proxy;
@@ -210,18 +210,18 @@ namespace Chaos
 			//       Proxy should be able to map back to Particle.
 			//
 
-			if (InParticleType == EParticleType::Dynamic)
+			if (InParticleType == EParticleType::Rigid)
 			{
 				Handle = Solver->Particles.CreateDynamicParticles(1)[0];
 				auto Proxy = static_cast<FRigidParticlePhysicsProxy*>(ProxyBase);
-				Proxy->SetHandle(Handle->AsDynamic());
+				Proxy->SetHandle(Handle->CastToRigidParticle());
 				Proxy->PushToPhysicsState(ProxyData);
 			}
 			else if (InParticleType == EParticleType::Kinematic)
 			{
 				Handle = Solver->Particles.CreateKinematicParticles(1)[0];
 				auto Proxy = static_cast<FKinematicGeometryParticlePhysicsProxy*>(ProxyBase);
-				Proxy->SetHandle(Handle->AsKinematic());
+				Proxy->SetHandle(Handle->CastToKinematicParticle());
 				Proxy->PushToPhysicsState(ProxyData);
 			}
 			else // Assume it's a static (geometry) if it's not dynamic or kinematic
@@ -258,7 +258,7 @@ namespace Chaos
 		GTParticle->Proxy = nullptr;
 
 		// Remove the proxy from the GT proxy map
-		if (InParticleType == EParticleType::Dynamic)
+		if (InParticleType == EParticleType::Rigid)
 		{
 			RigidParticlePhysicsProxies.Remove((FRigidParticlePhysicsProxy*)InProxy);
 		}
@@ -287,7 +287,7 @@ namespace Chaos
 			// maybe that extra safety is not needed if we continue to contain all
 			// references to proxy instances entirely in Chaos?
 			TGeometryParticleHandle<float, 3>* Handle;
-			if (InParticleType == Chaos::EParticleType::Dynamic)
+			if (InParticleType == Chaos::EParticleType::Rigid)
 			{
 				auto Proxy = (FRigidParticlePhysicsProxy*)InProxy;
 				Handle = Proxy->GetHandle();
@@ -506,7 +506,7 @@ namespace Chaos
 		{
 			switch (ActiveObject.Type)
 			{
-			case Chaos::EParticleType::Dynamic:
+			case Chaos::EParticleType::Rigid:
 				((FRigidParticlePhysicsProxy*)(ActiveObject.GTGeometryParticle()->Proxy))->BufferPhysicsResults();
 				break;
 			case Chaos::EParticleType::Kinematic:
@@ -539,7 +539,7 @@ namespace Chaos
 		{
 			switch (ActiveObject.Type)
 			{
-			case Chaos::EParticleType::Dynamic:
+			case Chaos::EParticleType::Rigid:
 				((FRigidParticlePhysicsProxy*)(ActiveObject.GTGeometryParticle()->Proxy))->FlipBuffer();
 				break;
 			case Chaos::EParticleType::Kinematic:
@@ -572,7 +572,7 @@ namespace Chaos
 		{
 			switch (ActiveObject.Type)
 			{
-			case Chaos::EParticleType::Dynamic:
+			case Chaos::EParticleType::Rigid:
 				((FRigidParticlePhysicsProxy*)(ActiveObject.GTGeometryParticle()->Proxy))->PullFromPhysicsState();
 				break;
 			case Chaos::EParticleType::Kinematic:
diff --git a/Engine/Source/Runtime/PhysicsCore/Private/PhysTestSerializer.cpp b/Engine/Source/Runtime/PhysicsCore/Private/PhysTestSerializer.cpp
--- a/Engine/Source/Runtime/PhysicsCore/Private/PhysTestSerializer.cpp
+++ b/Engine/Source/Runtime/PhysicsCore/Private/PhysTestSerializer.cpp
@@ -269,13 +269,14 @@ void FPhysTestSerializer::CreateChaosData()
 			Particle->GTGeometryParticle()->SetX(Particle->X());
 			Particle->GTGeometryParticle()->SetR(Particle->R());
 
-			if (auto PBDRigid = Particle->AsDynamic())
+			auto PBDRigid = Particle->CastToRigidParticle();
+			if(PBDRigid && PBDRigid->ObjectState() == EObjectStateType::Dynamic)
 			{
 				PBDRigid->P() = Particle->X();
 				PBDRigid->Q() = Particle->R();
 
-				PBDRigid->GTGeometryParticle()->AsDynamic()->SetP(PBDRigid->P());
-				PBDRigid->GTGeometryParticle()->AsDynamic()->SetQ(PBDRigid->R());
+				PBDRigid->GTGeometryParticle()->CastToRigidParticle()->SetP(PBDRigid->P());
+				PBDRigid->GTGeometryParticle()->CastToRigidParticle()->SetQ(PBDRigid->R());
 			}
 
 			PxActorToChaosHandle.Add(Act, Particle.Get());
