Chaos - Change AsDynamic in GeometryCollectionPhysicsProxy.cpp to CastToRigidParticle.

diff --git a/Engine/Source/Runtime/Experimental/ChaosSolvers/Private/PhysicsProxy/GeometryCollectionPhysicsProxy.cpp b/Engine/Source/Runtime/Experimental/ChaosSolvers/Private/PhysicsProxy/GeometryCollectionPhysicsProxy.cpp
--- a/Engine/Source/Runtime/Experimental/ChaosSolvers/Private/PhysicsProxy/GeometryCollectionPhysicsProxy.cpp
+++ b/Engine/Source/Runtime/Experimental/ChaosSolvers/Private/PhysicsProxy/GeometryCollectionPhysicsProxy.cpp
@@ -576,7 +576,7 @@ void FGeometryCollectionPhysicsProxy::InitializeBodiesGT()
 				//GTGeometryParticles[Idx] = Chaos::TKinematicGeometryParticle<float, 3>::CreateParticle();
 				//GTGeometryParticles[Idx] = Chaos::TPBDRigidParticle<float, 3>::CreateParticle();
 				GTGeometryParticles[Idx] = Chaos::TPBDGeometryCollectionParticle<float, 3>::CreateParticle();
-				Chaos::TPBDRigidParticle<float, 3>* GTParticle = GTGeometryParticles[Idx]->AsDynamic();
+				Chaos::TPBDRigidParticle<float, 3>* GTParticle = GTGeometryParticles[Idx]->CastToRigidParticle();
 
 //				Chaos::TSerializablePtr<Chaos::TImplicitObject<float, 3>> ImplicitShape = Implicits[Idx];
 //				GTParticle->SetGeometry(ImplicitShape);
@@ -584,10 +584,14 @@ void FGeometryCollectionPhysicsProxy::InitializeBodiesGT()
 //				{
 //					check(false); // continue?
 //				}
-				GTParticle->SetX(Transform[Idx].GetTranslation());
-				GTParticle->SetR(Transform[Idx].GetRotation());
-				//...
-				GTParticle->Proxy = this; // I feel dirty.
+
+				if (GTParticle->ObjectState() == Chaos::EObjectStateType::Dynamic)
+				{
+					GTParticle->SetX(Transform[Idx].GetTranslation());
+					GTParticle->SetR(Transform[Idx].GetRotation());
+					//...
+					GTParticle->Proxy = this; // I feel dirty.
+				}
 			}
 		}
 
@@ -639,10 +643,14 @@ void FGeometryCollectionPhysicsProxy::InitializeBodiesPT(
 			if (SimulatableParticles[Idx])
 			{
 				check(GTGeometryParticles[Idx]);
-				Chaos::TPBDRigidParticle<float, 3>* GTParticle = GTGeometryParticles[Idx]->AsDynamic();
+				Chaos::TPBDRigidParticle<float, 3>* GTParticle = GTGeometryParticles[Idx]->CastToRigidParticle();
 
 				Chaos::TPBDGeometryCollectionParticleHandle<float, 3>* Handle = Handles[NextIdx++];
-				Handle->GTGeometryParticle() = GTParticle;
+
+				if (GTParticle->ObjectState() == Chaos::EObjectStateType::Dynamic)
+				{
+					Handle->GTGeometryParticle() = GTParticle;
+				}
 
 				SolverParticleHandles[Idx] = Handle;
 				RigidsSolver->GetEvolution()->CreateParticle(Handle);
