Chaos : Collisions

- Plane stub functions.

diff --git a/Engine/Source/Runtime/Experimental/Chaos/Private/Chaos/CollisionResolution.cpp b/Engine/Source/Runtime/Experimental/Chaos/Private/Chaos/CollisionResolution.cpp
--- a/Engine/Source/Runtime/Experimental/Chaos/Private/Chaos/CollisionResolution.cpp
+++ b/Engine/Source/Runtime/Experimental/Chaos/Private/Chaos/CollisionResolution.cpp
@@ -976,6 +976,7 @@ namespace Chaos
 		void ConstructPairConstraintImpl(TGeometryParticleHandle<T, d>* Particle0, TGeometryParticleHandle<T, d>* Particle1, const FImplicitObject* Implicit0, const FImplicitObject* Implicit1, const T Thickness, FCollisionConstraintsArray& NewConstraints)
 		{
 			// See if we already have a constraint for this shape pair
+			// todo(brice) : Not sure this actually prevents duplicate constraints.
 			for (int32 i = 0; i < NewConstraints.Num(); i++)
 			{
 				if (NewConstraints[i]->GetType() == TRigidBodyPointContactConstraint<T, d>::StaticType())
@@ -985,6 +986,13 @@ namespace Chaos
 						return;
 					}
 				}
+				else if (NewConstraints[i]->GetType() == TRigidBodyPlaneContactConstraint<T, d>::StaticType())
+				{
+					if (NewConstraints[i]->As< TRigidBodyPlaneContactConstraint<T, d> >()->ContainsManifold(Implicit0, Implicit1))
+					{
+						return;
+					}
+				}
 			}
 
 			if (!Implicit0 || !Implicit1)
diff --git a/Engine/Source/Runtime/Experimental/Chaos/Private/Chaos/PBDCollisionConstraints.cpp b/Engine/Source/Runtime/Experimental/Chaos/Private/Chaos/PBDCollisionConstraints.cpp
--- a/Engine/Source/Runtime/Experimental/Chaos/Private/Chaos/PBDCollisionConstraints.cpp
+++ b/Engine/Source/Runtime/Experimental/Chaos/Private/Chaos/PBDCollisionConstraints.cpp
@@ -15,6 +15,8 @@
 #include "Chaos/ISpatialAccelerationCollection.h"
 #include "Chaos/Levelset.h"
 #include "Chaos/Pair.h"
+#include "Chaos/PBDCollisionConstraintsPointContactUtil.h"
+#include "Chaos/PBDCollisionConstraintsPlaneContactUtil.h"
 #include "Chaos/PBDRigidsSOAs.h"
 #include "Chaos/Sphere.h"
 #include "Chaos/Transform.h"
@@ -22,7 +24,6 @@
 #include "ChaosStats.h"
 #include "Containers/Queue.h"
 #include "ProfilingDebugging/ScopedTimers.h"
-#include "Chaos/PBDCollisionConstraintsPointContactUtil.h"
 
 #if INTEL_ISPC
 #include "PBDCollisionConstraints.ispc.generated.h"
@@ -275,20 +276,6 @@ namespace Chaos
 		// Clustering uses update constraints to force a re-evaluation. 
 	}
 
-	DECLARE_CYCLE_STAT(TEXT("TPBDCollisionConstraints::UpdateConstraint"), STAT_UpdateConstraint, STATGROUP_ChaosWide);
-	template<typename T, int d>
-	template<ECollisionUpdateType UpdateType>
-	void TPBDCollisionConstraints<T, d>::UpdateConstraint(const T Thickness, TRigidBodyPointContactConstraint<T, d> & Constraint)
-	{
-		SCOPE_CYCLE_COUNTER(STAT_UpdateConstraint);
-		//if (ConstraintHandle->GetType() == FPointContactConstraint::StaticType())
-		{
-			Collisions::Update<UpdateType>(Thickness, Constraint);
-		}
-	}
-
-
-
 
 	DECLARE_CYCLE_STAT(TEXT("TPBDCollisionConstraints::Apply"), STAT_Apply, STATGROUP_Chaos);
 	template<typename T, int d>
@@ -309,7 +296,9 @@ namespace Chaos
 				}
 				else if (ConstraintHandle->GetType() == FPlaneContactConstraint::StaticType())
 				{
-
+					Collisions::TPlaneContactParticleParameters<T> ParticleParameters = { &MCollided, &MPhysicsMaterials, CollisionFrictionOverride, MAngularFriction };
+					Collisions::TPlaneContactIterationParameters<T> IterationParameters = { Dt, Iterations, NumIterations, MApplyPairIterations, nullptr };
+					Collisions::Apply(ConstraintHandle->GetPlaneContact(), MThickness, IterationParameters, ParticleParameters);
 				}
 				else
 				{
@@ -348,7 +337,9 @@ namespace Chaos
 				}
 				else if (ConstraintHandle->GetType() == FPlaneContactConstraint::StaticType())
 				{
-
+					Collisions::TPlaneContactParticleParameters<T> ParticleParameters = { &MCollided, &MPhysicsMaterials, CollisionFrictionOverride, MAngularFriction };
+					Collisions::TPlaneContactIterationParameters<T> IterationParameters = { Dt, Iteration, NumIterations, MApplyPushOutPairIterations, &NeedsAnotherIteration };
+					Collisions::ApplyPushOut(ConstraintHandle->GetPlaneContact(), MThickness, IsTemporarilyStatic, IterationParameters, ParticleParameters);
 				}
 				else
 				{
@@ -370,8 +361,6 @@ namespace Chaos
 
 	template class TAccelerationStructureHandle<float, 3>;
 	template class CHAOS_API TPBDCollisionConstraints<float, 3>;
-	template void TPBDCollisionConstraints<float, 3>::UpdateConstraint<ECollisionUpdateType::Any>(const float Thickness, TRigidBodyPointContactConstraint<float,3> & Constraint);
-	template void TPBDCollisionConstraints<float, 3>::UpdateConstraint<ECollisionUpdateType::Deepest>(const float Thickness, TRigidBodyPointContactConstraint<float,3> & Constraint);
 	template void TPBDCollisionConstraints<float, 3>::ComputeConstraints<false>(const TPBDCollisionConstraints<float, 3>::FAccelerationStructure&, float Dt);
 	template void TPBDCollisionConstraints<float, 3>::ComputeConstraints<true>(const TPBDCollisionConstraints<float, 3>::FAccelerationStructure&, float Dt);
 }
diff --git a/Engine/Source/Runtime/Experimental/Chaos/Private/Chaos/PBDCollisionConstraintsPlaneContactUtil.cpp b/Engine/Source/Runtime/Experimental/Chaos/Private/Chaos/PBDCollisionConstraintsPlaneContactUtil.cpp
--- a/Engine/Source/Runtime/Experimental/Chaos/Private/Chaos/PBDCollisionConstraintsPlaneContactUtil.cpp
+++ b/Engine/Source/Runtime/Experimental/Chaos/Private/Chaos/PBDCollisionConstraintsPlaneContactUtil.cpp
@@ -8,29 +8,30 @@ namespace Chaos
 {
 	namespace Collisions
 	{
-		/*
 		template<ECollisionUpdateType UpdateType, typename T, int d>
 		void Update(const T Thickness, TRigidBodyPlaneContactConstraint<T, d>& Constraint)
 		{
 		}
-		template void Update<ECollisionUpdateType::Any, float, 3>(const float, TRigidBodyPlaneContactConstraint<float, 3>&);
-		template void Update<ECollisionUpdateType::Deepest, float, 3>(const float, TRigidBodyPlaneContactConstraint<float, 3>&);
 
 
 		template<typename T, int d>
-		void Apply(TRigidBodyPlaneContactConstraint<T, d>& Constraint, T Thickness, TPointContactIterationParameters<T> & IterationParameters, TPlainContactParticleParameters<T> & ParticleParameters)
+		void Apply(TRigidBodyPlaneContactConstraint<T, d>& Constraint, T Thickness, TPlaneContactIterationParameters<T> & IterationParameters, TPlaneContactParticleParameters<T> & ParticleParameters)
 		{
 		}
-		template void Apply<float, 3>(TRigidBodyPlaneContactConstraint<float, 3>&, float, TPointContactIterationParameters<float> &, TPointContactParticleParameters<float> &);
 
 		template<typename T, int d>
-		void ApplyPushOut(TRigidBodyPointContactConstraint<T, d>& Constraint, T Thickness, const TSet<const TGeometryParticleHandle<T, d>*>& IsTemporarilyStatic,
-			TPointContactIterationParameters<T> & IterationParameters, TPointContactParticleParameters<T> & ParticleParameters)
+		void ApplyPushOut(TRigidBodyPlaneContactConstraint<T, d>& Constraint, T Thickness, const TSet<const TGeometryParticleHandle<T, d>*>& IsTemporarilyStatic,
+			TPlaneContactIterationParameters<T> & IterationParameters, TPlaneContactParticleParameters<T> & ParticleParameters)
 		{
 		}
-		template void ApplyPushOut<float, 3>(TRigidBodyPointContactConstraint<float,3>&, float , const TSet<const TGeometryParticleHandle<float,3>*>&,
-			TPointContactIterationParameters<float> & IterationParameters, TPointContactParticleParameters<float> & ParticleParameters);
-	*/
+
+
+		template void Update<ECollisionUpdateType::Any, float, 3>(const float, TRigidBodyPlaneContactConstraint<float, 3>&);
+		template void Update<ECollisionUpdateType::Deepest, float, 3>(const float, TRigidBodyPlaneContactConstraint<float, 3>&);
+		template void Apply<float, 3>(TRigidBodyPlaneContactConstraint<float, 3>&, float, TPlaneContactIterationParameters<float> &, TPlaneContactParticleParameters<float> &);
+		template void ApplyPushOut<float, 3>(TRigidBodyPlaneContactConstraint<float,3>&, float , const TSet<const TGeometryParticleHandle<float,3>*>&,
+			TPlaneContactIterationParameters<float> & IterationParameters, TPlaneContactParticleParameters<float> & ParticleParameters);
+
 	} // Collisions
 
 }// Chaos
diff --git a/Engine/Source/Runtime/Experimental/Chaos/Public/Chaos/CollisionResolutionTypes.h b/Engine/Source/Runtime/Experimental/Chaos/Public/Chaos/CollisionResolutionTypes.h
--- a/Engine/Source/Runtime/Experimental/Chaos/Public/Chaos/CollisionResolutionTypes.h
+++ b/Engine/Source/Runtime/Experimental/Chaos/Public/Chaos/CollisionResolutionTypes.h
@@ -161,7 +161,8 @@ namespace Chaos
 		bool bDisabled;
 		int32 Timestamp;
 
-		TPlane<T, d> Plane;
+		TVector<T, d> PlaneNormal;
+		TVector<T, d> PlanePosition;
 		TVector<T, d> ManifoldSamples[4]; // iterative plane samples
 
 		TVector<T, d> Location;
diff --git a/Engine/Source/Runtime/Experimental/Chaos/Public/Chaos/PBDCollisionConstraints.h b/Engine/Source/Runtime/Experimental/Chaos/Public/Chaos/PBDCollisionConstraints.h
--- a/Engine/Source/Runtime/Experimental/Chaos/Public/Chaos/PBDCollisionConstraints.h
+++ b/Engine/Source/Runtime/Experimental/Chaos/Public/Chaos/PBDCollisionConstraints.h
@@ -228,12 +228,6 @@ protected:
 
 private:
 
-	/**
-	* Update the individual constraints.
-	*/
-	template<ECollisionUpdateType>
-	void UpdateConstraint(const T Thickness, TRigidBodyPointContactConstraint<T, d> & Constraint);
-
 	/**
 	*  Build a constraint based on the two particle handles.
 	*/
@@ -270,8 +264,6 @@ private:
 	const FAccelerationStructure* SpatialAcceleration;
 };
 
-extern template void TPBDCollisionConstraints<float, 3>::UpdateConstraint<ECollisionUpdateType::Any>(const float Thickness, TRigidBodyPointContactConstraint<float,3> & Constraint);
-extern template void TPBDCollisionConstraints<float, 3>::UpdateConstraint<ECollisionUpdateType::Deepest>(const float Thickness, TRigidBodyPointContactConstraint<float, 3> & Constraint);
 extern template void TPBDCollisionConstraints<float, 3>::ComputeConstraints<false>(const TPBDCollisionConstraints<float, 3>::FAccelerationStructure&, float Dt);
 extern template void TPBDCollisionConstraints<float, 3>::ComputeConstraints<true>(const TPBDCollisionConstraints<float, 3>::FAccelerationStructure&, float Dt);
 }
diff --git a/Engine/Source/Runtime/Experimental/Chaos/Public/Chaos/PBDCollisionConstraintsImp.h b/Engine/Source/Runtime/Experimental/Chaos/Public/Chaos/PBDCollisionConstraintsImp.h
--- a/Engine/Source/Runtime/Experimental/Chaos/Public/Chaos/PBDCollisionConstraintsImp.h
+++ b/Engine/Source/Runtime/Experimental/Chaos/Public/Chaos/PBDCollisionConstraintsImp.h
@@ -374,6 +374,16 @@ namespace Chaos
 
 					delete PointConstraint;
 				}
+				else if (ConstraintBase->GetType() == TRigidBodyPlaneContactConstraint<T, 3>::StaticType())
+				{
+					TRigidBodyPlaneContactConstraint<T, d>* PlaneConstraint = ConstraintBase->As< TRigidBodyPlaneContactConstraint<T, d> >();
+
+					int32 Idx = PlaneConstraints.AddUninitialized(1);
+					PlaneConstraints[Idx] = *PlaneConstraint;
+					Handles.Add(HandleAllocator.template AllocHandle< TRigidBodyPlaneContactConstraint<T, d> >(this, Idx));
+
+					delete PlaneConstraint;
+				}
 			}
 			LifespanCounter++;
 		}
diff --git a/Engine/Source/Runtime/Experimental/Chaos/Public/Chaos/PBDCollisionConstraintsPlaneContactUtil.h b/Engine/Source/Runtime/Experimental/Chaos/Public/Chaos/PBDCollisionConstraintsPlaneContactUtil.h
--- a/Engine/Source/Runtime/Experimental/Chaos/Public/Chaos/PBDCollisionConstraintsPlaneContactUtil.h
+++ b/Engine/Source/Runtime/Experimental/Chaos/Public/Chaos/PBDCollisionConstraintsPlaneContactUtil.h
@@ -10,7 +10,7 @@ namespace Chaos
 	namespace Collisions
 	{
 		template<class T = float>
-		struct TPlainContactParticleParameters {
+		struct TPlaneContactParticleParameters {
 			TArrayCollectionArray<bool>* Collided;
 			const TArrayCollectionArray<TSerializablePtr<FChaosPhysicsMaterial>>* PhysicsMaterials;
 			T FrictionOverride;
@@ -19,27 +19,24 @@ namespace Chaos
 
 		template<class T = float>
 		struct TPlaneContactIterationParameters {
-			const T Dt; 
-			const int32 Iteration; 
-			const int32 NumIterations; 
+			const T Dt;
+			const int32 Iteration;
+			const int32 NumIterations;
 			const int32 NumPairIterations;
 			bool* NeedsAnotherIteration;
 		};
 
-		/*
-		template<ECollisionUpdateType UpdateType, typename T = float, int d = 3>
-		void Update(const T Thickness, TRigidBodyPointContactConstraint<T, d>& Constraint);
-
-		template<typename T = float, int d = 3>
-		void Apply(TRigidBodyPointContactConstraint<T, d>& Constraint, T Thickness, 
-			TPointContactIterationParameters<T> & IterationParameters,
-			TPointContactParticleParameters<T> & ParticleParameters);
-			
-		template<typename T=float, int d=3>
-		void ApplyPushOut(TRigidBodyPointContactConstraint<T, d>& Constraint, T Thickness, const TSet<const TGeometryParticleHandle<T, d>*>& IsTemporarilyStatic,
-			TPointContactIterationParameters<T> & IterationParameters,
-			TPointContactParticleParameters<T> & ParticleParameters);
-	*/
+		template<ECollisionUpdateType UpdateType, typename T, int d>
+		void Update(const T Thickness, TRigidBodyPlaneContactConstraint<T, d>& Constraint);
+
+		template<typename T, int d>
+		void Apply(TRigidBodyPlaneContactConstraint<T, d>& Constraint, T Thickness, TPlaneContactIterationParameters<T> & IterationParameters, TPlaneContactParticleParameters<T> & ParticleParameters);
+
+		template<typename T, int d>
+		void ApplyPushOut(TRigidBodyPlaneContactConstraint<T, d>& Constraint, T Thickness, const TSet<const TGeometryParticleHandle<T, d>*>& IsTemporarilyStatic,
+			TPlaneContactIterationParameters<T> & IterationParameters, TPlaneContactParticleParameters<T> & ParticleParameters);
+
+
 	}
 
 }
