Chaos : Collisions

More plumbing for collision types.

diff --git a/Engine/Source/Runtime/Experimental/Chaos/Private/Chaos/CollisionResolution.cpp b/Engine/Source/Runtime/Experimental/Chaos/Private/Chaos/CollisionResolution.cpp
--- a/Engine/Source/Runtime/Experimental/Chaos/Private/Chaos/CollisionResolution.cpp
+++ b/Engine/Source/Runtime/Experimental/Chaos/Private/Chaos/CollisionResolution.cpp
@@ -751,7 +751,7 @@ namespace Chaos
 		DECLARE_CYCLE_STAT(TEXT("TPBDCollisionConstraints::UpdateUnionUnionConstraint"), STAT_UpdateUnionUnionConstraint, STATGROUP_ChaosWide);
 
 		template<ECollisionUpdateType UpdateType, typename T, int d>
-		void UpdateUnionUnionConstraint(const T Thickness, TRigidBodyPointContactConstraint<T, d>& Constraint)
+		void UpdateUnionUnionConstraint(const T Thickness, TCollisionConstraintBase<T,d>& Constraint)
 		{
 			SCOPE_CYCLE_COUNTER(STAT_UpdateUnionUnionConstraint);
 
@@ -821,7 +821,7 @@ namespace Chaos
 
 		DECLARE_CYCLE_STAT(TEXT("TPBDCollisionConstraints::UpdateSingleUnionConstraint"), STAT_UpdateSingleUnionConstraint, STATGROUP_ChaosWide);
 		template<ECollisionUpdateType UpdateType, typename T, int d>
-		void UpdateSingleUnionConstraint(const T Thickness, TRigidBodyPointContactConstraint<T, d>& Constraint)
+		void UpdateSingleUnionConstraint(const T Thickness, TCollisionConstraintBase<T,d>& Constraint)
 		{
 			SCOPE_CYCLE_COUNTER(STAT_UpdateSingleUnionConstraint);
 
@@ -870,38 +870,39 @@ namespace Chaos
 		//
 
 		template<ECollisionUpdateType UpdateType, typename T, int d>
-		void UpdateConstraintImp(const FImplicitObject& ParticleObject, const TRigidTransform<T, d>& ParticleTM, const FImplicitObject& LevelsetObject, const TRigidTransform<T, d>& LevelsetTM, const T Thickness, TRigidBodyPointContactConstraint<T, d>& Constraint)
+		void UpdateConstraintImp(const FImplicitObject& ParticleObject, const TRigidTransform<T, d>& ParticleTM, const FImplicitObject& LevelsetObject, const TRigidTransform<T, d>& LevelsetTM, const T Thickness, TCollisionConstraintBase<T, d>& ConstraintBase)
 		{
 			if (ParticleObject.GetType() == TBox<T, d>::StaticType() && LevelsetObject.GetType() == TBox<T, d>::StaticType())
 			{
-				UpdateBoxBoxConstraint(*ParticleObject.template GetObject<TBox<T, d>>(), ParticleTM, *LevelsetObject.template GetObject<TBox<T, d>>(), LevelsetTM, Thickness, Constraint);
+				UpdateBoxBoxConstraint(*ParticleObject.template GetObject<TBox<T, d>>(), ParticleTM, *LevelsetObject.template GetObject<TBox<T, d>>(), LevelsetTM, Thickness, *ConstraintBase.As<TRigidBodyPointContactConstraint<T,d> >());
 			}
 			else if (ParticleObject.GetType() == TSphere<T, d>::StaticType() && LevelsetObject.GetType() == TSphere<T, d>::StaticType())
 			{
-				UpdateSphereSphereConstraint(*ParticleObject.template GetObject<TSphere<T, d>>(), ParticleTM, *LevelsetObject.template GetObject<TSphere<T, d>>(), LevelsetTM, Thickness, Constraint);
+				UpdateSphereSphereConstraint(*ParticleObject.template GetObject<TSphere<T, d>>(), ParticleTM, *LevelsetObject.template GetObject<TSphere<T, d>>(), LevelsetTM, Thickness, *ConstraintBase.As<TRigidBodyPointContactConstraint<T, d> >());
 			}
 			else if (ParticleObject.GetType() == TBox<T, d>::StaticType() && LevelsetObject.GetType() == TPlane<T, d>::StaticType())
 			{
-				UpdateBoxPlaneConstraint(*ParticleObject.template GetObject<TBox<T, d>>(), ParticleTM, *LevelsetObject.template GetObject<TPlane<T, d>>(), LevelsetTM, Thickness, Constraint);
+				UpdateBoxPlaneConstraint(*ParticleObject.template GetObject<TBox<T, d>>(), ParticleTM, *LevelsetObject.template GetObject<TPlane<T, d>>(), LevelsetTM, Thickness, *ConstraintBase.As<TRigidBodyPointContactConstraint<T, d> >());
 			}
 			else if (ParticleObject.GetType() == TSphere<T, d>::StaticType() && LevelsetObject.GetType() == TPlane<T, d>::StaticType())
 			{
-				UpdateSpherePlaneConstraint(*ParticleObject.template GetObject<TSphere<T, d>>(), ParticleTM, *LevelsetObject.template GetObject<TPlane<T, d>>(), LevelsetTM, Thickness, Constraint);
+				UpdateSpherePlaneConstraint(*ParticleObject.template GetObject<TSphere<T, d>>(), ParticleTM, *LevelsetObject.template GetObject<TPlane<T, d>>(), LevelsetTM, Thickness, *ConstraintBase.As<TRigidBodyPointContactConstraint<T,d> >());
 			}
 			else if (ParticleObject.GetType() == TSphere<T, d>::StaticType() && LevelsetObject.GetType() == TBox<T, d>::StaticType())
 			{
-				UpdateSphereBoxConstraint(*ParticleObject.template GetObject<TSphere<T, d>>(), ParticleTM, *LevelsetObject.template GetObject<TBox<T, d>>(), LevelsetTM, Thickness, Constraint);
+				UpdateSphereBoxConstraint(*ParticleObject.template GetObject<TSphere<T, d>>(), ParticleTM, *LevelsetObject.template GetObject<TBox<T, d>>(), LevelsetTM, Thickness, *ConstraintBase.As<TRigidBodyPointContactConstraint<T, d> >());
 			}
 			else if (ParticleObject.GetType() == TCapsule<T>::StaticType() && LevelsetObject.GetType() == TCapsule<T>::StaticType())
 			{
-				UpdateCapsuleCapsuleConstraint(*ParticleObject.template GetObject<TCapsule<T>>(), ParticleTM, *LevelsetObject.template GetObject<TCapsule<T>>(), LevelsetTM, Thickness, Constraint);
+				UpdateCapsuleCapsuleConstraint(*ParticleObject.template GetObject<TCapsule<T>>(), ParticleTM, *LevelsetObject.template GetObject<TCapsule<T>>(), LevelsetTM, Thickness, *ConstraintBase.As<TRigidBodyPointContactConstraint<T, d> >());
 			}
 			else if (ParticleObject.GetType() == TCapsule<T>::StaticType() && LevelsetObject.GetType() == TBox<T, d>::StaticType())
 			{
-				UpdateCapsuleBoxConstraint(*ParticleObject.template GetObject<TCapsule<T>>(), ParticleTM, *LevelsetObject.template GetObject<TBox<T, d>>(), LevelsetTM, Thickness, Constraint);
+				UpdateCapsuleBoxConstraint(*ParticleObject.template GetObject<TCapsule<T>>(), ParticleTM, *LevelsetObject.template GetObject<TBox<T, d>>(), LevelsetTM, Thickness, *ConstraintBase.As<TRigidBodyPointContactConstraint<T, d> >());
 			}
 			else if (ParticleObject.GetType() == TPlane<T, d>::StaticType() && LevelsetObject.GetType() == TBox<T, d>::StaticType())
 			{
+				TRigidBodyPointContactConstraint<T, d>& Constraint = *ConstraintBase.As<TRigidBodyPointContactConstraint<T, d> >();
 				TRigidBodyPointContactConstraint<T, d> TmpConstraint = Constraint;
 				UpdateBoxPlaneConstraint(*LevelsetObject.template GetObject<TBox<T, d>>(), LevelsetTM, *ParticleObject.template GetObject<TPlane<T, d>>(), ParticleTM, Thickness, TmpConstraint);
 				if (TmpConstraint.GetPhi() < Constraint.GetPhi())
@@ -912,6 +913,7 @@ namespace Chaos
 			}
 			else if (ParticleObject.GetType() == TPlane<T, d>::StaticType() && LevelsetObject.GetType() == TSphere<T, d>::StaticType())
 			{
+				TRigidBodyPointContactConstraint<T, d>& Constraint = *ConstraintBase.As<TRigidBodyPointContactConstraint<T, d> >();
 				TRigidBodyPointContactConstraint<T, d> TmpConstraint = Constraint;
 				UpdateSpherePlaneConstraint(*LevelsetObject.template GetObject<TSphere<T, d>>(), LevelsetTM, *ParticleObject.template GetObject<TPlane<T, d>>(), ParticleTM, Thickness, TmpConstraint);
 				if (TmpConstraint.GetPhi() < Constraint.GetPhi())
@@ -922,6 +924,7 @@ namespace Chaos
 			}
 			else if (ParticleObject.GetType() == TBox<T, d>::StaticType() && LevelsetObject.GetType() == TSphere<T, d>::StaticType())
 			{
+				TRigidBodyPointContactConstraint<T, d>& Constraint = *ConstraintBase.As<TRigidBodyPointContactConstraint<T, d> >();
 				TRigidBodyPointContactConstraint<T, d> TmpConstraint = Constraint;
 				UpdateSphereBoxConstraint(*LevelsetObject.template GetObject<TSphere<T, d>>(), LevelsetTM, *ParticleObject.template GetObject<TBox<T, d>>(), ParticleTM, Thickness, TmpConstraint);
 				if (TmpConstraint.GetPhi() < Constraint.GetPhi())
@@ -932,6 +935,7 @@ namespace Chaos
 			}
 			else if (ParticleObject.GetType() == TBox<T, d>::StaticType() && LevelsetObject.GetType() == TCapsule<T>::StaticType())
 			{
+				TRigidBodyPointContactConstraint<T, d>& Constraint = *ConstraintBase.As<TRigidBodyPointContactConstraint<T, d> >();
 				TRigidBodyPointContactConstraint<T, d> TmpConstraint = Constraint;
 				UpdateCapsuleBoxConstraint(*LevelsetObject.template GetObject<TCapsule<T>>(), LevelsetTM, *ParticleObject.template GetObject<TBox<T, d>>(), ParticleTM, Thickness, TmpConstraint);
 				if (TmpConstraint.GetPhi() < Constraint.GetPhi())
@@ -942,7 +946,7 @@ namespace Chaos
 			}
 			else if (ParticleObject.GetType() < TImplicitObjectUnion<T, d>::StaticType() && LevelsetObject.GetType() == TImplicitObjectUnion<T, d>::StaticType())
 			{
-				return UpdateSingleUnionConstraint<UpdateType>(Thickness, Constraint);
+				return UpdateSingleUnionConstraint<UpdateType>(Thickness, ConstraintBase);
 			}
 			else if (ParticleObject.GetType() == TImplicitObjectUnion<T, d>::StaticType() && LevelsetObject.GetType() < TImplicitObjectUnion<T, d>::StaticType())
 			{
@@ -950,23 +954,23 @@ namespace Chaos
 			}
 			else if (ParticleObject.GetType() == TImplicitObjectUnion<T, d>::StaticType() && LevelsetObject.GetType() == TImplicitObjectUnion<T, d>::StaticType())
 			{
-				return UpdateUnionUnionConstraint<UpdateType>(Thickness, Constraint);
+				return UpdateUnionUnionConstraint<UpdateType>(Thickness, ConstraintBase);
 			}
 			else if (ParticleObject.IsConvex() && LevelsetObject.IsConvex())
 			{
-				UpdateConvexConvexConstraint(ParticleObject, ParticleTM, LevelsetObject, LevelsetTM, Thickness, Constraint);
+				UpdateConvexConvexConstraint(ParticleObject, ParticleTM, LevelsetObject, LevelsetTM, Thickness, *ConstraintBase.As<TRigidBodyPointContactConstraint<T, d> >());
 			}
 			else if (LevelsetObject.IsUnderlyingUnion())
 			{
-				UpdateUnionLevelsetConstraint<UpdateType>(Thickness, Constraint);
+				UpdateUnionLevelsetConstraint<UpdateType>(Thickness, *ConstraintBase.As<TRigidBodyPointContactConstraint<T, d> >());
 			}
 			else if (ParticleObject.IsUnderlyingUnion())
 			{
-				UpdateLevelsetUnionConstraint<UpdateType>(Thickness, Constraint);
+				UpdateLevelsetUnionConstraint<UpdateType>(Thickness, *ConstraintBase.As<TRigidBodyPointContactConstraint<T, d> >());
 			}
 			else
 			{
-				UpdateLevelsetLevelsetConstraint<UpdateType>(Thickness, Constraint);
+				UpdateLevelsetLevelsetConstraint<UpdateType>(Thickness, *ConstraintBase.As<TRigidBodyPointContactConstraint<T, d> >());
 			}
 		}
 
@@ -1120,16 +1124,16 @@ namespace Chaos
 		template void UpdateUnionLevelsetConstraint<ECollisionUpdateType::Any, float, 3>(const float Thickness, TRigidBodyPointContactConstraint<float, 3>& Constraint);
 		template void UpdateUnionLevelsetConstraint<ECollisionUpdateType::Deepest, float, 3>(const float Thickness, TRigidBodyPointContactConstraint<float, 3>& Constraint);
 
-		template void UpdateUnionUnionConstraint< ECollisionUpdateType::Any, float, 3>(const float Thickness, TRigidBodyPointContactConstraint<float, 3>& Constraint);
-		template void UpdateUnionUnionConstraint< ECollisionUpdateType::Deepest, float, 3>(const float Thickness, TRigidBodyPointContactConstraint<float, 3>& Constraint);
+		template void UpdateUnionUnionConstraint< ECollisionUpdateType::Any, float, 3>(const float Thickness, TCollisionConstraintBase<float, 3>& Constraint);
+		template void UpdateUnionUnionConstraint< ECollisionUpdateType::Deepest, float, 3>(const float Thickness, TCollisionConstraintBase<float, 3>& Constraint);
 		template void ConstructUnionUnionConstraints<float, 3>(TGeometryParticleHandle<float, 3>* Particle0, TGeometryParticleHandle<float, 3>* Particle1, const FImplicitObject* Implicit0, const FImplicitObject* Implicit1, const float Thickness, FCollisionConstraintsArray& NewConstraints);
 
-		template void UpdateSingleUnionConstraint< ECollisionUpdateType::Any, float, 3>(const float Thickness, TRigidBodyPointContactConstraint<float, 3>& Constraint);
-		template void UpdateSingleUnionConstraint< ECollisionUpdateType::Deepest, float, 3>(const float Thickness, TRigidBodyPointContactConstraint<float, 3>& Constraint);
+		template void UpdateSingleUnionConstraint< ECollisionUpdateType::Any, float, 3>(const float Thickness, TCollisionConstraintBase<float, 3>& Constraint);
+		template void UpdateSingleUnionConstraint< ECollisionUpdateType::Deepest, float, 3>(const float Thickness, TCollisionConstraintBase<float, 3>& Constraint);
 		template void ConstructSingleUnionConstraints<float, 3>(TGeometryParticleHandle<float, 3>* Particle0, TGeometryParticleHandle<float, 3>* Particle1, const FImplicitObject* Implicit0, const FImplicitObject* Implicit1, const float Thickness, FCollisionConstraintsArray& NewConstraints);
 
-		template void UpdateConstraintImp<ECollisionUpdateType::Any, float, 3>(const FImplicitObject& ParticleObject, const TRigidTransform<float, 3>& ParticleTM, const FImplicitObject& LevelsetObject, const TRigidTransform<float, 3>& LevelsetTM, const float Thickness, TRigidBodyPointContactConstraint<float, 3>& Constraints);
-		template void UpdateConstraintImp<ECollisionUpdateType::Deepest, float, 3>(const FImplicitObject& ParticleObject, const TRigidTransform<float, 3>& ParticleTM, const FImplicitObject& LevelsetObject, const TRigidTransform<float, 3>& LevelsetTM, const float Thickness, TRigidBodyPointContactConstraint<float, 3>& Constraint);
+		template void UpdateConstraintImp<ECollisionUpdateType::Any, float, 3>(const FImplicitObject& ParticleObject, const TRigidTransform<float, 3>& ParticleTM, const FImplicitObject& LevelsetObject, const TRigidTransform<float, 3>& LevelsetTM, const float Thickness, TCollisionConstraintBase<float, 3>& Constraints);
+		template void UpdateConstraintImp<ECollisionUpdateType::Deepest, float, 3>(const FImplicitObject& ParticleObject, const TRigidTransform<float, 3>& ParticleTM, const FImplicitObject& LevelsetObject, const TRigidTransform<float, 3>& LevelsetTM, const float Thickness, TCollisionConstraintBase<float, 3>& Constraint);
 
 		template void ConstructPairConstraintImpl<float, 3>(TGeometryParticleHandle<float, 3>* Particle0, TGeometryParticleHandle<float, 3>* Particle1, const FImplicitObject* Implicit0, const FImplicitObject* Implicit1, const float Thickness, FCollisionConstraintsArray& NewConstraints);
 		template void ConstructConstraintsImpl<float, 3>(TGeometryParticleHandle<float, 3>* Particle0, TGeometryParticleHandle<float, 3>* Particle1, const FImplicitObject* Implicit0, const FImplicitObject* Implicit1, const float Thickness, FCollisionConstraintsArray& NewConstraints);
diff --git a/Engine/Source/Runtime/Experimental/Chaos/Public/Chaos/CollisionResolution.h b/Engine/Source/Runtime/Experimental/Chaos/Public/Chaos/CollisionResolution.h
--- a/Engine/Source/Runtime/Experimental/Chaos/Public/Chaos/CollisionResolution.h
+++ b/Engine/Source/Runtime/Experimental/Chaos/Public/Chaos/CollisionResolution.h
@@ -136,7 +136,7 @@ namespace Chaos
 			//
 
 			template<ECollisionUpdateType UpdateType, typename T, int d>
-			void CHAOS_API UpdateUnionUnionConstraint(const T Thickness, TRigidBodyPointContactConstraint<T, d>& Constraint);
+			void CHAOS_API UpdateUnionUnionConstraint(const T Thickness, TCollisionConstraintBase<T, d>& Constraint);
 
 			template<typename T, int d>
 			void CHAOS_API ConstructUnionUnionConstraints(TGeometryParticleHandle<T, d>* Particle0, TGeometryParticleHandle<T, d>* Particle1, const FImplicitObject* Implicit0, const FImplicitObject* Implicit1, const T Thickness, FCollisionConstraintsArray& NewConstraints);
@@ -146,7 +146,7 @@ namespace Chaos
 			//
 
 			template<ECollisionUpdateType UpdateType, typename T, int d>
-			void CHAOS_API UpdateSingleUnionConstraint(const T Thickness, TRigidBodyPointContactConstraint<T, d>& Constraint);
+			void CHAOS_API UpdateSingleUnionConstraint(const T Thickness, TCollisionConstraintBase<T, d>& Constraint);
 
 			template<typename T, int d>
 			void CHAOS_API ConstructSingleUnionConstraints(TGeometryParticleHandle<T, d>* Particle0, TGeometryParticleHandle<T, d>* Particle1, const FImplicitObject* Implicit0, const FImplicitObject* Implicit1, const T Thickness, FCollisionConstraintsArray& NewConstraints);
@@ -157,7 +157,7 @@ namespace Chaos
 			//
 
 			template<ECollisionUpdateType UpdateType, typename T, int d>
-			void CHAOS_API UpdateConstraintImp(const FImplicitObject& ParticleObject, const TRigidTransform<T, d>& ParticleTM, const FImplicitObject& LevelsetObject, const TRigidTransform<T, d>& LevelsetTM, const T Thickness, TRigidBodyPointContactConstraint<T, d>& Constraint);
+			void CHAOS_API UpdateConstraintImp(const FImplicitObject& ParticleObject, const TRigidTransform<T, d>& ParticleTM, const FImplicitObject& LevelsetObject, const TRigidTransform<T, d>& LevelsetTM, const T Thickness, TCollisionConstraintBase<T, d>& Constraint);
 
 			template<typename T, int d>
 			void CHAOS_API ConstructPairConstraintImpl(TGeometryParticleHandle<T, d>* Particle0, TGeometryParticleHandle<T, d>* Particle1, const FImplicitObject* Implicit0, const FImplicitObject* Implicit1, const T Thickness, FCollisionConstraintsArray& NewConstraints);
diff --git a/Engine/Source/Runtime/Experimental/Chaos/Public/Chaos/CollisionResolutionTypes.h b/Engine/Source/Runtime/Experimental/Chaos/Public/Chaos/CollisionResolutionTypes.h
--- a/Engine/Source/Runtime/Experimental/Chaos/Public/Chaos/CollisionResolutionTypes.h
+++ b/Engine/Source/Runtime/Experimental/Chaos/Public/Chaos/CollisionResolutionTypes.h
@@ -46,6 +46,7 @@ namespace Chaos
 	class CHAOS_API TCollisionConstraintBase
 	{
 	public:
+		using FGeometryParticleHandle = TGeometryParticleHandle<T, d>;
 
 		enum class FType
 		{
@@ -61,6 +62,8 @@ namespace Chaos
 		template<class AS_T> AS_T * As() { return static_cast<AS_T*>(this); }
 		template<class AS_T> const AS_T * As() const { return static_cast<AS_T*>(this); }
 
+		FGeometryParticleHandle* Particle[2]; // { Point, Volume } 
+
 	private:
 		FType Type;
 	};
@@ -137,7 +140,6 @@ namespace Chaos
 
 
 		FManifold Manifold;
-		FGeometryParticleHandle* Particle[2]; // { Point, Volume } 
 		TVector<T, d> AccumulatedImpulse;
 
 	};
@@ -199,7 +201,6 @@ namespace Chaos
 
 
 		FManifold Manifold;
-		FGeometryParticleHandle* Particle[2]; // { Point, Volume } 
 		TVector<T, d> AccumulatedImpulse;
 
 	};
