Chaos - fix joints after particle handle changes

diff --git a/Engine/Source/Runtime/Experimental/Chaos/Private/Chaos/ChaosDebugDraw.cpp b/Engine/Source/Runtime/Experimental/Chaos/Private/Chaos/ChaosDebugDraw.cpp
--- a/Engine/Source/Runtime/Experimental/Chaos/Private/Chaos/ChaosDebugDraw.cpp
+++ b/Engine/Source/Runtime/Experimental/Chaos/Private/Chaos/ChaosDebugDraw.cpp
@@ -23,13 +23,13 @@ namespace Chaos
 		float ContactWidth = 2.0f;
 		float ContactPhiWidth = 1.5f;
 		float ConstraintAxisLen = 5.0f;
-		float LineThickness = 0.25f;
-		int DrawPriority = 10.0f;
+		float LineThickness = 0.15f;
 		float DrawScale = 1.0f;
 		float FontHeight = 10.0f;
 		float FontScale = 1.0f;
 		float ShapeThicknesScale = 2.0f;
 		float PointSize = 2.0f;
+		int DrawPriority = 10.0f;
 
 		FAutoConsoleVariableRef CVarArrowSize(TEXT("p.Chaos.DebugDrawArrowSize"), ArrowSize, TEXT("ArrowSize."));
 		FAutoConsoleVariableRef CVarBodyAxisLen(TEXT("p.Chaos.DebugDrawBodyAxisLen"), BodyAxisLen, TEXT("BodyAxisLen."));
diff --git a/Engine/Source/Runtime/Experimental/Chaos/Private/Chaos/PBDJointConstraints.cpp b/Engine/Source/Runtime/Experimental/Chaos/Private/Chaos/PBDJointConstraints.cpp
--- a/Engine/Source/Runtime/Experimental/Chaos/Private/Chaos/PBDJointConstraints.cpp
+++ b/Engine/Source/Runtime/Experimental/Chaos/Private/Chaos/PBDJointConstraints.cpp
@@ -520,7 +520,7 @@ namespace Chaos
 
 	void FPBDJointConstraints::UpdateParticleState(TPBDRigidParticleHandle<FReal, 3>* Rigid, const FReal Dt, const FVec3& P, const FRotation3& Q)
 	{
-		if (Rigid)
+		if ((Rigid != nullptr) && (Rigid->ObjectState() == EObjectStateType::Dynamic))
 		{
 			if (Dt > SMALL_NUMBER)
 			{
@@ -593,16 +593,8 @@ namespace Chaos
 			}
 		}
 
-		TPBDRigidParticleHandle<FReal, 3>* Rigid0 = Particle0->CastToRigidParticle();
-		TPBDRigidParticleHandle<FReal, 3>* Rigid1 = Particle1->CastToRigidParticle();
-		if (Rigid0->ObjectState() == EObjectStateType::Dynamic)
-		{
-			UpdateParticleState(Rigid0, Dt, P0, Q0);
-		}
-		if (Rigid1->ObjectState() == EObjectStateType::Dynamic)
-		{
-			UpdateParticleState(Rigid1, Dt, P1, Q1);
-		}
+		UpdateParticleState(Particle0->CastToRigidParticle(), Dt, P0, Q0);
+		UpdateParticleState(Particle1->CastToRigidParticle(), Dt, P1, Q1);
 	}
 
 
@@ -656,16 +648,8 @@ namespace Chaos
 			Solver.ApplyConstraints(Dt, Settings, JointSettings, GetLinearStiffness(Settings, JointSettings));
 		}
 
-		TPBDRigidParticleHandle<FReal, 3>* Rigid0 = Particle0->CastToRigidParticle();
-		TPBDRigidParticleHandle<FReal, 3>* Rigid1 = Particle1->CastToRigidParticle();
-		if (Rigid0 && Rigid0->ObjectState() == EObjectStateType::Dynamic)
-		{
-			UpdateParticleState(Rigid0, Dt, Solver.GetP(0), Solver.GetQ(0));
-		}
-		if (Rigid1 && Rigid1->ObjectState() == EObjectStateType::Dynamic)
-		{
-			UpdateParticleState(Rigid1, Dt, Solver.GetP(1), Solver.GetQ(1));
-		}
+		UpdateParticleState(Particle0->CastToRigidParticle(), Dt, Solver.GetP(0), Solver.GetQ(0));
+		UpdateParticleState(Particle1->CastToRigidParticle(), Dt, Solver.GetP(1), Solver.GetQ(1));
 	}
 
 
@@ -766,17 +750,8 @@ namespace Chaos
 		}
 
 		// Update the particles
-		TPBDRigidParticleHandle<FReal, 3>* Rigid0 = ConstraintParticles[ConstraintIndex][Index0]->CastToRigidParticle();
-		TPBDRigidParticleHandle<FReal, 3>* Rigid1 = ConstraintParticles[ConstraintIndex][Index1]->CastToRigidParticle();
-
-		if (Rigid0->ObjectState() == EObjectStateType::Dynamic)
-		{
-			UpdateParticleState(Rigid0, Dt, P0, Q0);
-		}
-		if (Rigid1->ObjectState() == EObjectStateType::Dynamic)
-		{
-			UpdateParticleState(Rigid1, Dt, P1, Q1);
-		}
+		UpdateParticleState(Particle0->CastToRigidParticle(), Dt, P0, Q0);
+		UpdateParticleState(Particle1->CastToRigidParticle(), Dt, P1, Q1);
 	}
 
 	
@@ -875,18 +850,8 @@ namespace Chaos
 		}
 
 		// Update the particles
-		TPBDRigidParticleHandle<FReal, 3>* Rigid0 = ConstraintParticles[ConstraintIndex][Index0]->CastToRigidParticle();
-		TPBDRigidParticleHandle<FReal, 3>* Rigid1 = ConstraintParticles[ConstraintIndex][Index1]->CastToRigidParticle();
-
-		if (Rigid0->ObjectState() == EObjectStateType::Dynamic)
-		{
-			UpdateParticleState(Rigid0, Dt, P0, Q0);
-		}
-		if (Rigid1->ObjectState() == EObjectStateType::Dynamic)
-		{
-			UpdateParticleState(Rigid1, Dt, P1, Q1);
-		}
-
+		UpdateParticleState(Particle0->CastToRigidParticle(), Dt, P0, Q0);
+		UpdateParticleState(Particle1->CastToRigidParticle(), Dt, P1, Q1);
 	}
 
 
