Chaos - mass matrix x dense matrix

diff --git a/Engine/Source/Runtime/Experimental/Chaos/Public/Chaos/DenseMatrix.h b/Engine/Source/Runtime/Experimental/Chaos/Public/Chaos/DenseMatrix.h
--- a/Engine/Source/Runtime/Experimental/Chaos/Public/Chaos/DenseMatrix.h
+++ b/Engine/Source/Runtime/Experimental/Chaos/Public/Chaos/DenseMatrix.h
@@ -8,6 +8,51 @@
 
 namespace Chaos
 {
+	/**
+	 * A block-diagonal matrix specifically for use with Mass/Inertia (or more usually inverse mass and inertia)
+	 * and TDenseMatrix math used by the constraint solvers.
+	 */
+	class FMassMatrix
+	{
+	public:
+		FReal M() const
+		{
+			return Mass;
+		}
+		
+		const FReal& I(int32 RowIndex, int32 ColIndex) const
+		{
+			// Symmetric - could reduce element count at increased indexing cost...
+			checkSlow(RowIndex < 3);
+			checkSlow(ColIndex < 3);
+			return Inertia[RowIndex * 3 + ColIndex];
+		}
+
+		static FMassMatrix Make(const FReal InM, const FMatrix33& InI)
+		{
+			return FMassMatrix(InM, InI);
+		}
+
+	private:
+		FMassMatrix(const FReal InM, const FMatrix33& InI)
+		{
+			Mass = InM;
+
+			Inertia[0] = InI.M[0][0];
+			Inertia[1] = InI.M[0][1];
+			Inertia[2] = InI.M[0][2];
+			Inertia[3] = InI.M[1][0];
+			Inertia[4] = InI.M[1][1];
+			Inertia[5] = InI.M[1][2];
+			Inertia[6] = InI.M[2][0];
+			Inertia[7] = InI.M[2][1];
+			Inertia[8] = InI.M[2][2];
+		}
+
+		FReal Mass;
+		FReal Inertia[9];
+	};
+
 	/**
 	 * A matrix with run-time variable dimensions, up to an element limit defined at compile-time.
 	 *
@@ -315,6 +360,23 @@ namespace Chaos
 		// Math operations
 		//
 
+		/**
+		 * Return the transpose of 'A'.
+		 */
+		template<int32 T_EA>
+		static TDenseMatrix<MaxElements> Transpose(const TDenseMatrix<T_EA>& A)
+		{
+			TDenseMatrix<T_MAXELEMENTS> Result(A.NumColumns(), A.NumRows());
+			for (int32 IRow = 0; IRow < Result.NumRows(); ++IRow)
+			{
+				for (int32 ICol = 0; ICol < Result.NumColumns(); ++ICol)
+				{
+					Result.At(IRow, ICol) = A.At(ICol, IRow);
+				}
+			}
+			return Result;
+		}
+
 		/**
 		 * Copy a matrix and set each element to its negative.
 		 */
@@ -474,13 +536,80 @@ namespace Chaos
 		}
 
 		/**
-		 * Return C = A x B x Transpose(A).
+		 * Return C = A x B, where A is an Nx6 matrix, and B is a 6x6 mass matrix (Mass in upper left 3x3 diagonals, Inertia in lower right 3x3).
+		 * C = |A0 A1| * |M 0| = |A0.M A1.I|
+		 *               |0 I|
 		 */
-		template<int32 T_EA, int32 T_EB>
-		static TDenseMatrix<MaxElements> MultiplyABAt(const TDenseMatrix<T_EA>& A, const TDenseMatrix<T_EB>& B)
+		template<int32 T_EA>
+		static TDenseMatrix<MaxElements> MultiplyAB(const TDenseMatrix<T_EA>& A, const FMassMatrix& B)
 		{
-			// @todo(ccaulfield): optimize
-			return TDenseMatrix<MaxElements>::Multiply(A, TDenseMatrix<MaxElements>::MultiplyABt(B, A));
+			check(A.NumColumns() == 6);
+
+			TDenseMatrix<T_MAXELEMENTS> Result = TDenseMatrix<T_MAXELEMENTS>::Make(A.NumRows(), 6);
+
+			// Calculate columns 0-2
+			for (int32 IRow = 0; IRow < Result.NumRows(); ++IRow)
+			{
+				for (int32 ICol = 0; ICol < 3; ++ICol)
+				{
+					Result.At(IRow, ICol) = A.At(IRow, ICol) * B.M();
+				}
+			}
+
+			// Calculate columns 3-5
+			for (int32 IRow = 0; IRow < Result.NumRows(); ++IRow)
+			{
+				for (int32 ICol = 3; ICol < 6; ++ICol)
+				{
+					FReal V = 0;
+					for (int32 KK = 3; KK < 6; ++KK)
+					{
+						V += A.At(IRow, KK) * B.I(KK - 3, ICol - 3);
+					}
+					Result.At(IRow, ICol) = V;
+				}
+			}
+
+			return Result;
+		}
+
+
+		/**
+		 * Return C = A x Transpose(B), where B is an Nx6 matrix, and A is a 6x6 mass matrix (Mass in upper left 3x3 diagonals, Inertia in lower right 3x3).
+		 * C = |M 0| * |B0 B1|(T) = |M.B0(T)|
+		 *     |0 I|                |I.B1(T)|
+		 */
+		template<int32 T_EB>
+		static TDenseMatrix<MaxElements> MultiplyABt(const FMassMatrix& A, const TDenseMatrix<T_EB>& B)
+		{
+			check(B.NumColumns() == 6);
+
+			TDenseMatrix<T_MAXELEMENTS> Result = TDenseMatrix<T_MAXELEMENTS>::Make(6, B.NumRows());
+
+			// Calculate rows 0-2
+			for (int32 IRow = 0; IRow < 3; ++IRow)
+			{
+				for (int32 ICol = 0; ICol < B.NumRows(); ++ICol)
+				{
+					Result.At(IRow, ICol) = A.M() * B.At(ICol, IRow);
+				}
+			}
+
+			// Calculate rows 3-5
+			for (int32 IRow = 3; IRow < 6; ++IRow)
+			{
+				for (int32 ICol = 0; ICol < B.NumRows(); ++ICol)
+				{
+					FReal V = 0;
+					for (int32 KK = 3; KK < 6; ++KK)
+					{
+						V += A.I(IRow - 3, KK - 3) * B.At(ICol, KK);
+					}
+					Result.At(IRow, ICol) = V;
+				}
+			}
+
+			return Result;
 		}
 
 
@@ -697,4 +826,5 @@ namespace Chaos
 		}
 	};
 
+
 }
