Chaos - Fix for CIS

diff --git a/Engine/Source/Runtime/Experimental/Chaos/Public/Chaos/CollisionResolutionTypes.h b/Engine/Source/Runtime/Experimental/Chaos/Public/Chaos/CollisionResolutionTypes.h
--- a/Engine/Source/Runtime/Experimental/Chaos/Public/Chaos/CollisionResolutionTypes.h
+++ b/Engine/Source/Runtime/Experimental/Chaos/Public/Chaos/CollisionResolutionTypes.h
@@ -104,6 +104,7 @@ namespace Chaos
 		using Base = TCollisionConstraintBase<T, d>;
 		using FGeometryParticleHandle = TGeometryParticleHandle<T, d>;
 		using FManifold = TPointContactManifold<T, d>;
+		using Base::Particle;
 
 		TRigidBodyPointContactConstraint() : Base(Base::FType::SinglePoint), AccumulatedImpulse(0) {}
 		static typename Base::FType StaticType() { return Base::FType::SinglePoint; };
@@ -185,6 +186,7 @@ namespace Chaos
 		using Base = TCollisionConstraintBase<T, d>;
 		using FGeometryParticleHandle = TGeometryParticleHandle<T, d>;
 		using FManifold = TPlaneContactManifold<T, d>;
+		using Base::Particle;
 
 		TRigidBodyPlaneContactConstraint() : Base(Base::FType::Plane), AccumulatedImpulse(0) {}
 		static typename Base::FType StaticType() { return Base::FType::Plane; };
