Compute volume and center of mass of TConvex

diff --git a/Engine/Source/Runtime/Core/Public/UObject/ExternalPhysicsCustomObjectVersion.h b/Engine/Source/Runtime/Core/Public/UObject/ExternalPhysicsCustomObjectVersion.h
--- a/Engine/Source/Runtime/Core/Public/UObject/ExternalPhysicsCustomObjectVersion.h
+++ b/Engine/Source/Runtime/Core/Public/UObject/ExternalPhysicsCustomObjectVersion.h
@@ -68,6 +68,9 @@ struct CORE_API FExternalPhysicsCustomObjectVersion
 
 		// Added material indices to trimesh collision data
 		AddTrimeshMaterialIndices,
+
+		// Add center of mass and volume cached calculations to TConvex
+		AddConvexCenterOfMassAndVolume,
 		
 		// -----<new versions can be added above this line>-------------------------------------------------
 		VersionPlusOne,
diff --git a/Engine/Source/Runtime/Engine/Private/PhysicsEngine/Experimental/ChaosDerivedData.cpp b/Engine/Source/Runtime/Engine/Private/PhysicsEngine/Experimental/ChaosDerivedData.cpp
--- a/Engine/Source/Runtime/Engine/Private/PhysicsEngine/Experimental/ChaosDerivedData.cpp
+++ b/Engine/Source/Runtime/Engine/Private/PhysicsEngine/Experimental/ChaosDerivedData.cpp
@@ -21,7 +21,7 @@ const TCHAR* FChaosDerivedDataCooker::GetPluginName() const
 
 const TCHAR* FChaosDerivedDataCooker::GetVersionString() const
 {
-	return TEXT("9F79F10AF0354268BB194AA050225CF4");
+	return TEXT("508E8FABC48D413632CD35B8FCA46D29");
 }
 
 FString FChaosDerivedDataCooker::GetPluginSpecificCacheKeySuffix() const
diff --git a/Engine/Source/Runtime/Experimental/Chaos/Private/Chaos/MassProperties.cpp b/Engine/Source/Runtime/Experimental/Chaos/Private/Chaos/MassProperties.cpp
--- a/Engine/Source/Runtime/Experimental/Chaos/Private/Chaos/MassProperties.cpp
+++ b/Engine/Source/Runtime/Experimental/Chaos/Private/Chaos/MassProperties.cpp
@@ -71,10 +71,10 @@ namespace Chaos
 	}
 	template TRotation<float, 3> TransformToLocalSpace(PMatrix<float, 3, 3>& Inertia);
 
-	template <typename T, int d>
-	void CalculateVolumeAndCenterOfMass(const TParticles<T, d>& Vertices, const TTriangleMesh<T>& Surface, T& OutVolume, TVector<T, d>& OutCenterOfMass)
+	template<typename T, int d, typename TSurfaces>
+	void CalculateVolumeAndCenterOfMass(const TParticles<T, d>& Vertices, const TSurfaces& Surfaces, T& OutVolume, TVector<T, d>& OutCenterOfMass)
 	{
-		if (!Surface.GetSurfaceElements().Num())
+		if (!Surfaces.Num())
 		{
 			OutVolume = 0;
 			return;
@@ -82,9 +82,12 @@ namespace Chaos
 		
 		T Volume = 0;
 		TVector<T, d> VolumeTimesSum(0);
-		TVector<T, d> Center = Vertices.X(Surface.GetSurfaceElements()[0][0]);
-		for (const auto& Element : Surface.GetSurfaceElements())
+		TVector<T, d> Center = Vertices.X(Surfaces[0][0]);
+		for (const auto& Element : Surfaces)
 		{
+			// For now we only support triangular elements
+			ensure(Element.Num() == 3);
+
 			PMatrix<T, d, d> DeltaMatrix;
 			TVector<T, d> PerElementSize;
 			for (int32 i = 0; i < Element.Num(); ++i)
@@ -113,15 +116,15 @@ namespace Chaos
 		OutVolume = Volume / 6;
 	}
 
-	template <typename T, int d>
-	void CalculateInertiaAndRotationOfMass(const TParticles<T, d>& Vertices, const TTriangleMesh<T>& Surface, const T Density, const TVector<T,d>& CenterOfMass,
+	template <typename T, int d, typename TSurfaces>
+	void CalculateInertiaAndRotationOfMass(const TParticles<T, d>& Vertices, const TSurfaces& Surfaces, const T Density, const TVector<T,d>& CenterOfMass,
 	PMatrix<T,d,d>& OutInertiaTensor, TRotation<T,d>& OutRotationOfMass)
 	{
 		check(Density > 0);
 
 		static const PMatrix<T, d, d> Standard(2, 1, 1, 2, 1, 2);
 		PMatrix<T, d, d> Covariance(0);
-		for (const auto& Element : Surface.GetSurfaceElements())
+		for (const auto& Element : Surfaces)
 		{
 			PMatrix<T, d, d> DeltaMatrix(0);
 			for (int32 i = 0; i < Element.Num(); ++i)
@@ -144,17 +147,18 @@ namespace Chaos
 	}
 
 
-	template<class T, int d>
+	template<class T, int d, typename TSurfaces>
 	TMassProperties<T, d> CalculateMassProperties(
 		const TParticles<T, d> & Vertices,
-		const TTriangleMesh<T>& Surface, const T Mass)
+		const TSurfaces& Surfaces,
+		const T Mass)
 	{
 		TMassProperties<T, d> MassProperties;
-		CalculateVolumeAndCenterOfMass(Vertices, Surface, MassProperties.Volume, MassProperties.CenterOfMass);
+		CalculateVolumeAndCenterOfMass(Vertices, Surfaces, MassProperties.Volume, MassProperties.CenterOfMass);
 
 		check(Mass > 0);
 		check(MassProperties.Volume > SMALL_NUMBER);
-		CalculateInertiaAndRotationOfMass(Vertices, Surface, Mass / MassProperties.Volume, MassProperties.CenterOfMass, MassProperties.InertiaTensor, MassProperties.RotationOfMass);
+		CalculateInertiaAndRotationOfMass(Vertices, Surfaces, Mass / MassProperties.Volume, MassProperties.CenterOfMass, MassProperties.InertiaTensor, MassProperties.RotationOfMass);
 		
 		return MassProperties;
 	}
@@ -193,12 +197,21 @@ namespace Chaos
 	}
 
 	template CHAOS_API TMassProperties<float, 3> CalculateMassProperties(const TParticles<float, 3> & Vertices,
-		const TTriangleMesh<float>& Surface, const float Mass);
+		const TArray<TVector<int32, 3>>& Surfaces, const float Mass);
+
+	template CHAOS_API TMassProperties<float, 3> CalculateMassProperties(const TParticles<float, 3> & Vertices,
+		const TArray<TArray<int32>>& Surfaces, const float Mass);
 
 	template CHAOS_API void CalculateVolumeAndCenterOfMass(const TParticles<float, 3>& Vertices,
-		const TTriangleMesh<float>& Surface, float& OutVolume, TVector<float, 3>& OutCenterOfMass);
+		const TArray<TVector<int32, 3>>& Surfaces, float& OutVolume, TVector<float, 3>& OutCenterOfMass);
+
+	template CHAOS_API void CalculateVolumeAndCenterOfMass(const TParticles<float, 3>& Vertices,
+		const TArray<TArray<int32>>& Surfaces, float& OutVolume, TVector<float, 3>& OutCenterOfMass);
+
+	template CHAOS_API void CalculateInertiaAndRotationOfMass(const TParticles<float, 3>& Vertices, const TArray<TVector<int32, 3>>& Surface, const float Density,
+		const TVector<float, 3>& CenterOfMass, PMatrix<float, 3, 3>& OutInertiaTensor, TRotation<float, 3>& OutRotationOfMass);
 
-	template CHAOS_API void CalculateInertiaAndRotationOfMass(const TParticles<float, 3>& Vertices, const TTriangleMesh<float>& Surface, const float Density,
+	template CHAOS_API void CalculateInertiaAndRotationOfMass(const TParticles<float, 3>& Vertices, const TArray<TArray<int32>>& Surface, const float Density,
 		const TVector<float, 3>& CenterOfMass, PMatrix<float, 3, 3>& OutInertiaTensor, TRotation<float, 3>& OutRotationOfMass);
 
 	template CHAOS_API TMassProperties<float, 3> Combine(const TArray<TMassProperties<float, 3>>& MPArray);
diff --git a/Engine/Source/Runtime/Experimental/Chaos/Public/Chaos/CollisionConvexMesh.h b/Engine/Source/Runtime/Experimental/Chaos/Public/Chaos/CollisionConvexMesh.h
--- a/Engine/Source/Runtime/Experimental/Chaos/Public/Chaos/CollisionConvexMesh.h
+++ b/Engine/Source/Runtime/Experimental/Chaos/Public/Chaos/CollisionConvexMesh.h
@@ -21,7 +21,7 @@ namespace Chaos
 			return Cross.Size() > 1e-4;
 		}
 
-		static void Build(const TParticles<T, 3>& InParticles, TArray <TPlane<T, 3>>& OutPlanes, TParticles<T, 3>& OutSurfaceParticles, TBox<T, 3>& OutLocalBounds)
+		static void Build(const TParticles<T, 3>& InParticles, TArray <TPlane<T, 3>>& OutPlanes, TArray<TArray<int32>>& OutFaceIndices, TParticles<T, 3>& OutSurfaceParticles, TBox<T, 3>& OutLocalBounds)
 		{
 			OutPlanes.Reset();
 			OutSurfaceParticles.Resize(0);
@@ -44,23 +44,36 @@ namespace Chaos
 				TArray<TVector<int32, 3>> Indices;
 				BuildConvexHull(InParticles, Indices);
 				OutPlanes.Reserve(Indices.Num());
-				TSet<int32> AllIndices;
+				TMap<int32, int32> IndexMap; // maps original particle indices to output particle indices
+				int32 NewIdx = 0;
+
+				const auto AddIndex = [&IndexMap, &NewIdx](const int32 OriginalIdx)
+				{
+					if (int32* Idx = IndexMap.Find(OriginalIdx))
+					{
+						return *Idx;
+					}
+					IndexMap.Add(OriginalIdx, NewIdx);
+					return NewIdx++;
+				};
 
 				for(const TVector<int32, 3>& Idx : Indices)
 				{
 					TVector<T, 3> Vs[3] = {InParticles.X(Idx[0]), InParticles.X(Idx[1]), InParticles.X(Idx[2])};
 					const TVector<T, 3> Normal = TVector<T, 3>::CrossProduct(Vs[1] - Vs[0], Vs[2] - Vs[0]).GetUnsafeNormal();
 					OutPlanes.Add(TPlane<T, 3>(Vs[0], Normal));
-					AllIndices.Add(Idx[0]);
-					AllIndices.Add(Idx[1]);
-					AllIndices.Add(Idx[2]);
+					TArray<int32> FaceIndices;
+					FaceIndices.SetNum(3);
+					FaceIndices[0] = AddIndex(Idx[0]);
+					FaceIndices[1] = AddIndex(Idx[1]);
+					FaceIndices[2] = AddIndex(Idx[2]);
+					OutFaceIndices.Add(FaceIndices);
 				}
 
-				OutSurfaceParticles.AddParticles(AllIndices.Num());
-				int32 NewIdx = 0;
-				for(int32 Idx : AllIndices)
+				OutSurfaceParticles.AddParticles(IndexMap.Num());
+				for(const auto& Elem : IndexMap)
 				{
-					OutSurfaceParticles.X(NewIdx++) = InParticles.X(Idx);
+					OutSurfaceParticles.X(Elem.Value) = InParticles.X(Elem.Key);
 				}
 			}
 			else if(NumParticles == 3)
@@ -77,6 +90,13 @@ namespace Chaos
 					OutSurfaceParticles.X(0) = InParticles.X(0);
 					OutSurfaceParticles.X(1) = InParticles.X(1);
 					OutSurfaceParticles.X(2) = InParticles.X(2);
+
+					TArray<int32> FaceIndices;
+					FaceIndices.SetNum(3);
+					FaceIndices[0] = 0; 
+					FaceIndices[1] = 1;
+					FaceIndices[2] = 2;
+					OutFaceIndices.Add(FaceIndices);
 				}
 			}
 		}
@@ -142,7 +162,7 @@ namespace Chaos
 			return FString::Printf(TEXT("Planes %d, SurfaceParticles %d"), NumPlanes, NumParticles);
 		}
 
-		static CHAOS_API void Simplify(TArray <TPlane<T, 3>>& InOutPlanes, TParticles<T, 3>& InOutParticles, TBox<T, 3>& InOutLocalBounds)
+		static CHAOS_API void Simplify(TArray <TPlane<T, 3>>& InOutPlanes, TArray<TArray<int32>>& InOutFaces, TParticles<T, 3>& InOutParticles, TBox<T, 3>& InOutLocalBounds)
 		{
 			struct TPair
 			{
@@ -215,7 +235,7 @@ namespace Chaos
 				}
 			}
 
-			Build(TmpParticles, InOutPlanes, InOutParticles, InOutLocalBounds);
+			Build(TmpParticles, InOutPlanes, InOutFaces, InOutParticles, InOutLocalBounds);
 			check(InOutParticles.Size() > 3);
 		}
 
diff --git a/Engine/Source/Runtime/Experimental/Chaos/Public/Chaos/Convex.h b/Engine/Source/Runtime/Experimental/Chaos/Public/Chaos/Convex.h
--- a/Engine/Source/Runtime/Experimental/Chaos/Public/Chaos/Convex.h
+++ b/Engine/Source/Runtime/Experimental/Chaos/Public/Chaos/Convex.h
@@ -3,6 +3,7 @@
 
 #include "Chaos/ImplicitObject.h"
 #include "Chaos/Box.h"
+#include "Chaos/MassProperties.h"
 #include "CollisionConvexMesh.h"
 #include "ChaosArchive.h"
 #include "GJK.h"
@@ -17,6 +18,8 @@ namespace Chaos
 
 		TConvex()
 		    : FImplicitObject(EImplicitObject::IsConvex | EImplicitObject::HasBoundingBox, ImplicitObjectType::Convex)
+			, Volume(0.f)
+			, CenterOfMass(FVec3(0.f))
 		{}
 		TConvex(const TConvex&) = delete;
 		TConvex(TConvex&& Other)
@@ -24,17 +27,26 @@ namespace Chaos
 			, Planes(MoveTemp(Other.Planes))
 		    , SurfaceParticles(MoveTemp(Other.SurfaceParticles))
 		    , LocalBoundingBox(MoveTemp(Other.LocalBoundingBox))
-		{}
-		TConvex(TArray<TPlane<T, d>>&& InPlanes, TParticles<T, d>&& InSurfaceParticles)
+			, Volume(MoveTemp(Other.MoveTemp))
+			, CenterOfMass(MoveTemp(Other.CenterOfMass))
+		{
+			CHAOS_ENSURE(Planes.Num() == FaceIndices.Num());
+		}
+		TConvex(TArray<TPlane<T, d>>&& InPlanes, TArray<TArray<int32>>&& InFaceIndices, TParticles<T, d>&& InSurfaceParticles)
 		    : FImplicitObject(EImplicitObject::IsConvex | EImplicitObject::HasBoundingBox, ImplicitObjectType::Convex)
 			, Planes(MoveTemp(InPlanes))
 		    , SurfaceParticles(MoveTemp(InSurfaceParticles))
 		    , LocalBoundingBox(TBox<T, d>::EmptyBox())
 		{
+			CHAOS_ENSURE(Planes.Num() == FaceIndices.Num());
+
 			for (uint32 ParticleIndex = 0; ParticleIndex < SurfaceParticles.Size(); ++ParticleIndex)
 			{
 				LocalBoundingBox.GrowToInclude(SurfaceParticles.X(ParticleIndex));
 			}
+
+			TArray<TArray<int32>> FaceIndices;
+			CalculateVolumeAndCenterOfMass(SurfaceParticles, FaceIndices, Volume, CenterOfMass);
 		}
 		TConvex(const TParticles<T, 3>& InParticles)
 		    : FImplicitObject(EImplicitObject::IsConvex | EImplicitObject::HasBoundingBox, ImplicitObjectType::Convex)
@@ -45,7 +57,10 @@ namespace Chaos
 				return;
 			}
 
-			TConvexBuilder<T>::Build(InParticles, Planes, SurfaceParticles, LocalBoundingBox);
+			TArray<TArray<int32>> FaceIndices;
+			TConvexBuilder<T>::Build(InParticles, Planes, FaceIndices, SurfaceParticles, LocalBoundingBox);
+			CHAOS_ENSURE(Planes.Num() == FaceIndices.Num());
+			CalculateVolumeAndCenterOfMass(SurfaceParticles, FaceIndices, Volume, CenterOfMass);
 		}
 
 		static constexpr EImplicitObjectType StaticType()
@@ -202,8 +217,7 @@ namespace Chaos
 
 		const FReal GetVolume() const
 		{
-			// TODO: More precise volume!
-			return LocalBoundingBox.GetVolume();
+			return Volume;
 		}
 
 		const FMatrix33 GetInertiaTensor(const FReal Mass) const
@@ -214,8 +228,7 @@ namespace Chaos
 
 		const FVec3 GetCenterOfMass() const
 		{
-			// TODO: Actually compute this!
-			return FVec3(0.f);
+			return CenterOfMass;
 		}
 
 		virtual uint32 GetTypeHash() const override
@@ -236,6 +249,22 @@ namespace Chaos
 		{
 			FImplicitObject::SerializeImp(Ar);
 			Ar << Planes << SurfaceParticles << LocalBoundingBox;
+
+			Ar.UsingCustomVersion(FExternalPhysicsCustomObjectVersion::GUID);
+			if (Ar.CustomVer(FExternalPhysicsCustomObjectVersion::GUID) >= FExternalPhysicsCustomObjectVersion::AddConvexCenterOfMassAndVolume)
+			{
+				Ar << Volume;
+				Ar << CenterOfMass;
+			}
+			else if (Ar.IsLoading())
+			{
+				// Rebuild convex in order to extract face indices.
+				// TODO: Make it so it can take SurfaceParticles as both input and output without breaking...
+				TArray<TArray<int32>> FaceIndices;
+				TParticles<T, d> TempSurfaceParticles;
+				TConvexBuilder<T>::Build(SurfaceParticles, Planes, FaceIndices, TempSurfaceParticles, LocalBoundingBox);
+				CalculateVolumeAndCenterOfMass(SurfaceParticles, FaceIndices, Volume, CenterOfMass);
+			}
 		}
 
 		virtual void Serialize(FChaosArchive& Ar) override
@@ -274,7 +303,8 @@ namespace Chaos
 
 		void SimplifyGeometry()
 		{
-			TConvexBuilder<T>::Simplify(Planes, SurfaceParticles, LocalBoundingBox);
+			TArray<TArray<int32>> FaceIndices;
+			TConvexBuilder<T>::Simplify(Planes, FaceIndices, SurfaceParticles, LocalBoundingBox);
 		}
 
 		TVector<T,d> GetCenter() const
@@ -286,5 +316,7 @@ namespace Chaos
 		TArray<TPlane<T, d>> Planes;
 		TParticles<T, d> SurfaceParticles;	//copy of the vertices that are just on the convex hull boundary
 		TBox<T, d> LocalBoundingBox;
+		float Volume;
+		FVec3 CenterOfMass;
 	};
 }
diff --git a/Engine/Source/Runtime/Experimental/Chaos/Public/Chaos/MassProperties.h b/Engine/Source/Runtime/Experimental/Chaos/Public/Chaos/MassProperties.h
--- a/Engine/Source/Runtime/Experimental/Chaos/Public/Chaos/MassProperties.h
+++ b/Engine/Source/Runtime/Experimental/Chaos/Public/Chaos/MassProperties.h
@@ -34,17 +34,17 @@ namespace Chaos
 	template<class T, int d>
 	TRotation<T, d> TransformToLocalSpace(PMatrix<T, d, d>& Inertia);
 
-	template<typename T, int d>
-	void CalculateVolumeAndCenterOfMass(const TParticles<T, d>& Vertices, const TTriangleMesh<T>& Surface, T& OutVolume, TVector<T, d>& OutCenterOfMass);
+	template<typename T, int d, typename TSurfaces>
+	void CalculateVolumeAndCenterOfMass(const TParticles<T, d>& Vertices, const TSurfaces& Surfaces, T& OutVolume, TVector<T, d>& OutCenterOfMass);
 
-	template<class T, int d>
+	template<class T, int d, typename TSurfaces>
 	TMassProperties<T, d> CalculateMassProperties(
 	    const TParticles<T, d>& Vertices,
-	    const TTriangleMesh<T>& Surface,
+		const TSurfaces& Surfaces,
 	    const T Mass);
 
-	template<typename T, int d>
-	void CalculateInertiaAndRotationOfMass(const TParticles<T, d>& Vertices, const TTriangleMesh<T>& Surface, const T Density, const TVector<T, d>& CenterOfMass,
+	template<typename T, int d, typename TSurfaces>
+	void CalculateInertiaAndRotationOfMass(const TParticles<T, d>& Vertices, const TSurfaces& Surfaces, const T Density, const TVector<T, d>& CenterOfMass,
 	    PMatrix<T, d, d>& OutInertiaTensor, TRotation<T, d>& OutRotationOfMass);
 
 	template<class T, int d>
diff --git a/Engine/Source/Runtime/Experimental/ChaosSolvers/Private/PhysicsProxy/GeometryCollectionPhysicsProxy.cpp b/Engine/Source/Runtime/Experimental/ChaosSolvers/Private/PhysicsProxy/GeometryCollectionPhysicsProxy.cpp
--- a/Engine/Source/Runtime/Experimental/ChaosSolvers/Private/PhysicsProxy/GeometryCollectionPhysicsProxy.cpp
+++ b/Engine/Source/Runtime/Experimental/ChaosSolvers/Private/PhysicsProxy/GeometryCollectionPhysicsProxy.cpp
@@ -2784,7 +2784,7 @@ void FGeometryCollectionPhysicsProxy::InitializeSharedCollisionStructures(
 
 			if (InertiaComputationNeeded[GeometryIndex])
 			{
-				CalculateInertiaAndRotationOfMass(MassSpaceParticles, *TriMesh, Density_i, MassProperties.CenterOfMass, MassProperties.InertiaTensor, MassProperties.RotationOfMass);
+				CalculateInertiaAndRotationOfMass(MassSpaceParticles, TriMesh->GetSurfaceElements(), Density_i, MassProperties.CenterOfMass, MassProperties.InertiaTensor, MassProperties.RotationOfMass);
 				CollectionMassToLocal[TransformGroupIndex] = FTransform(MassProperties.RotationOfMass, MassProperties.CenterOfMass);
 			}
 
