Remove extra CHAOS_ENSUREs for FaceIndices in Convex.h

diff --git a/Engine/Source/Runtime/Experimental/Chaos/Public/Chaos/Convex.h b/Engine/Source/Runtime/Experimental/Chaos/Public/Chaos/Convex.h
--- a/Engine/Source/Runtime/Experimental/Chaos/Public/Chaos/Convex.h
+++ b/Engine/Source/Runtime/Experimental/Chaos/Public/Chaos/Convex.h
@@ -29,17 +29,13 @@ namespace Chaos
 		    , LocalBoundingBox(MoveTemp(Other.LocalBoundingBox))
 			, Volume(MoveTemp(Other.MoveTemp))
 			, CenterOfMass(MoveTemp(Other.CenterOfMass))
-		{
-			CHAOS_ENSURE(Planes.Num() == FaceIndices.Num());
-		}
+		{}
 		TConvex(TArray<TPlane<T, d>>&& InPlanes, TArray<TArray<int32>>&& InFaceIndices, TParticles<T, d>&& InSurfaceParticles)
 		    : FImplicitObject(EImplicitObject::IsConvex | EImplicitObject::HasBoundingBox, ImplicitObjectType::Convex)
 			, Planes(MoveTemp(InPlanes))
 		    , SurfaceParticles(MoveTemp(InSurfaceParticles))
 		    , LocalBoundingBox(TBox<T, d>::EmptyBox())
 		{
-			CHAOS_ENSURE(Planes.Num() == FaceIndices.Num());
-
 			for (uint32 ParticleIndex = 0; ParticleIndex < SurfaceParticles.Size(); ++ParticleIndex)
 			{
 				LocalBoundingBox.GrowToInclude(SurfaceParticles.X(ParticleIndex));
