Add CenterOfMass data to particles

diff --git a/Engine/Source/Runtime/Core/Public/UObject/ExternalPhysicsCustomObjectVersion.h b/Engine/Source/Runtime/Core/Public/UObject/ExternalPhysicsCustomObjectVersion.h
--- a/Engine/Source/Runtime/Core/Public/UObject/ExternalPhysicsCustomObjectVersion.h
+++ b/Engine/Source/Runtime/Core/Public/UObject/ExternalPhysicsCustomObjectVersion.h
@@ -71,6 +71,9 @@ struct CORE_API FExternalPhysicsCustomObjectVersion
 
 		// Add center of mass and volume cached calculations to TConvex
 		AddConvexCenterOfMassAndVolume,
+
+		// Add center of mass data to kinematic particle
+		KinematicCentersOfMass,
 		
 		// -----<new versions can be added above this line>-------------------------------------------------
 		VersionPlusOne,
diff --git a/Engine/Source/Runtime/Experimental/Chaos/Public/Chaos/KinematicGeometryParticles.h b/Engine/Source/Runtime/Experimental/Chaos/Public/Chaos/KinematicGeometryParticles.h
--- a/Engine/Source/Runtime/Experimental/Chaos/Public/Chaos/KinematicGeometryParticles.h
+++ b/Engine/Source/Runtime/Experimental/Chaos/Public/Chaos/KinematicGeometryParticles.h
@@ -70,6 +70,7 @@ class TKinematicGeometryParticlesImp : public TGeometryParticlesImp<T, d, SimTyp
 		TArrayCollection::AddArray(&MV);
 		TArrayCollection::AddArray(&MW);
 		TArrayCollection::AddArray(&KinematicTargets);
+		TArrayCollection::AddArray(&MCenterOfMass);
 	}
 	TKinematicGeometryParticlesImp(const TKinematicGeometryParticlesImp<T, d, SimType>& Other) = delete;
 	CHAOS_API TKinematicGeometryParticlesImp(TKinematicGeometryParticlesImp<T, d, SimType>&& Other)
@@ -79,6 +80,7 @@ class TKinematicGeometryParticlesImp : public TGeometryParticlesImp<T, d, SimTyp
 		TArrayCollection::AddArray(&MV);
 		TArrayCollection::AddArray(&MW);
 		TArrayCollection::AddArray(&KinematicTargets);
+		TArrayCollection::AddArray(&MCenterOfMass);
 	}
 	CHAOS_API virtual ~TKinematicGeometryParticlesImp();
 
@@ -91,6 +93,9 @@ class TKinematicGeometryParticlesImp : public TGeometryParticlesImp<T, d, SimTyp
 	const TKinematicTarget<T, d>& KinematicTarget(const int32 Index) const { return KinematicTargets[Index]; }
 	TKinematicTarget<T, d>& KinematicTarget(const int32 Index) { return KinematicTargets[Index]; }
 
+	const TVector<T, d>& CenterOfMass(const int32 Index) const { return MCenterOfMass[Index]; }
+	TVector<T, d>& CenterOfMass(const int32 Index) { return MCenterOfMass[Index]; }
+
 	FString ToString(int32 index) const
 	{
 		FString BaseString = TGeometryParticlesImp<T, d, SimType>::ToString(index);
@@ -113,12 +118,17 @@ class TKinematicGeometryParticlesImp : public TGeometryParticlesImp<T, d, SimTyp
 		{
 			Ar << KinematicTargets;
 		}
+		if (Ar.CustomVer(FExternalPhysicsCustomObjectVersion::GUID) >= FExternalPhysicsCustomObjectVersion::KinematicCentersOfMass)
+		{
+			Ar << MCenterOfMass;
+		}
 	}
 
   private:
 	TArrayCollectionArray<TVector<T, d>> MV;
 	TArrayCollectionArray<TVector<T, d>> MW;
 	TArrayCollectionArray<TKinematicTarget<T, d>> KinematicTargets;
+	TArrayCollectionArray<TVector<T, d>> MCenterOfMass;
 };
 
 template <typename T, int d, EGeometryParticlesSimType SimType>
diff --git a/Engine/Source/Runtime/Experimental/Chaos/Public/Chaos/ParticleDirtyFlags.h b/Engine/Source/Runtime/Experimental/Chaos/Public/Chaos/ParticleDirtyFlags.h
--- a/Engine/Source/Runtime/Experimental/Chaos/Public/Chaos/ParticleDirtyFlags.h
+++ b/Engine/Source/Runtime/Experimental/Chaos/Public/Chaos/ParticleDirtyFlags.h
@@ -13,27 +13,28 @@ namespace Chaos
 		R =				 1 << 1,
 		V =				 1 << 2,
 		W =				 1 << 3,
-		CollisionGroup = 1 << 4,
-		Disabled =		 1 << 5,
-		PreV =			 1 << 6,
-		PreW =			 1 << 7,
-		P =				 1 << 8,
-		Q =				 1 << 9,
-		F =				 1 << 10,
-		Torque =		 1 << 11,
-		I =				 1 << 12,
-		InvI =			 1 << 13,
-		M =				 1 << 14,
-		InvM =			 1 << 15,
-		ObjectState =	 1 << 16,
-		Geometry =		 1 << 17,
-		ExternalForce =  1 << 18,
-		ExternalTorque = 1 << 19,
-		GravityEnabled = 1 << 20,
-		SpatialIdx     = 1 << 21,
-		HashResult     = 1 << 22
+		CenterOfMass =   1 << 4,
+		CollisionGroup = 1 << 5,
+		Disabled =		 1 << 6,
+		PreV =			 1 << 7,
+		PreW =			 1 << 8,
+		P =				 1 << 9,
+		Q =				 1 << 10,
+		F =				 1 << 11,
+		Torque =		 1 << 12,
+		I =				 1 << 13,
+		InvI =			 1 << 14,
+		M =				 1 << 15,
+		InvM =			 1 << 16,
+		ObjectState =	 1 << 17,
+		Geometry =		 1 << 18,
+		ExternalForce =  1 << 19,
+		ExternalTorque = 1 << 20,
+		GravityEnabled = 1 << 21,
+		SpatialIdx     = 1 << 22,
+		HashResult     = 1 << 23
 #if CHAOS_CHECKED
-		, DebugName    = 1 << 23
+		, DebugName    = 1 << 24
 #endif
 	};
 
diff --git a/Engine/Source/Runtime/Experimental/Chaos/Public/Chaos/ParticleHandle.h b/Engine/Source/Runtime/Experimental/Chaos/Public/Chaos/ParticleHandle.h
--- a/Engine/Source/Runtime/Experimental/Chaos/Public/Chaos/ParticleHandle.h
+++ b/Engine/Source/Runtime/Experimental/Chaos/Public/Chaos/ParticleHandle.h
@@ -58,6 +58,7 @@ void KinematicGeometryParticleDefaultConstruct(FConcrete& Concrete, const TKinem
 {
 	Concrete.SetV(TVector<T, d>(0));
 	Concrete.SetW(TVector<T, d>(0));
+	Concrete.SetCenterOfMass(TVector<T, d>(0));
 }
 
 template <typename T, int d, bool bPersistent>
@@ -542,6 +543,9 @@ public:
 	const TKinematicTarget<T, d>& KinematicTarget() const { return KinematicGeometryParticles->KinematicTarget(ParticleIdx); }
 	TKinematicTarget<T, d>& KinematicTarget() { return KinematicGeometryParticles->KinematicTarget(ParticleIdx); }
 
+	const TVector<T, d>& CenterOfMass() const { return KinematicGeometryParticles->CenterOfMass(ParticleIdx); }
+	void SetCenterOfMass(const TVector<T, d>& InCenterOfMass) { KinematicGeometryParticles->CenterOfMass(ParticleIdx) = InCenterOfMass; }
+
 	//Really only useful when using a transient handle
 	const TKinematicGeometryParticleHandleImp<T, d, true>* Handle() const { return KinematicGeometryParticles->Handle(ParticleIdx); }
 	TKinematicGeometryParticleHandleImp<T, d, true>* Handle() { return KinematicGeometryParticles->Handle(ParticleIdx); }
@@ -1514,6 +1518,13 @@ public:
 		this->MW = InW;
 	}
 
+	const TVector<T, d>& CenterOfMass() const { return MCenterOfMass; }
+	void SetCenterOfMass(const TVector<T, d>& InCenterOfMass, bool bInvalidate = true)
+	{
+		this->MarkDirty(EParticleFlags::CenterOfMass, bInvalidate);
+		this->MCenterOfMass = InCenterOfMass;
+	}
+
 	EObjectStateType ObjectState() const;
 
 	FParticleData* NewData() const
@@ -1524,6 +1535,7 @@ public:
 private:
 	TVector<T, d> MV;
 	TVector<T, d> MW;
+	TVector<T, d> MCenterOfMass;
 };
 
 template <typename T, int d>
@@ -1537,12 +1549,14 @@ public:
 	TKinematicGeometryParticleData(EParticleType InType = EParticleType::Kinematic)
 		: Base(InType)
 		, MV(TVector<T, d>(0))
-		, MW(TVector<T, d>(0)) {}
+		, MW(TVector<T, d>(0))
+		, MCenterOfMass(TVector<T, d>(0)) {}
 
 	TKinematicGeometryParticleData(const TKinematicGeometryParticle<T, d>& InParticle)
 		: Base(InParticle)
 		, MV(InParticle.V())
 		, MW(InParticle.W()) 
+		, MCenterOfMass(InParticle.CenterOfMass()) 
 	{
 		Type = EParticleType::Kinematic;
 	}
@@ -1551,11 +1565,14 @@ public:
 	void Reset() {
 		TGeometryParticleData<T, d>::Reset();
 		Type = EParticleType::Kinematic;
-		MV = TVector<T, d>(0); MW = TVector<T, d>();
+		MV = TVector<T, d>(0);
+		MW = TVector<T, d>(0);
+		MCenterOfMass = TVector<T, d>(0);
 	}
 
 	TVector<T, d> MV;
 	TVector<T, d> MW;
+	TVector<T, d> MCenterOfMass;
 };
 
 
diff --git a/Engine/Source/Runtime/Experimental/ChaosSolvers/Private/PhysicsProxy/SingleParticlePhysicsProxy.cpp b/Engine/Source/Runtime/Experimental/ChaosSolvers/Private/PhysicsProxy/SingleParticlePhysicsProxy.cpp
--- a/Engine/Source/Runtime/Experimental/ChaosSolvers/Private/PhysicsProxy/SingleParticlePhysicsProxy.cpp
+++ b/Engine/Source/Runtime/Experimental/ChaosSolvers/Private/PhysicsProxy/SingleParticlePhysicsProxy.cpp
@@ -119,7 +119,7 @@ void FSingleParticlePhysicsProxy<Chaos::TKinematicGeometryParticle<float, 3>>::P
 		RigidHandle->SetSharedGeometry(Data->Geometry);
 		RigidHandle->SetV(Data->MV);
 		RigidHandle->SetW(Data->MW);
-
+		RigidHandle->SetCenterOfMass(Data->MCenterOfMass);
 		RigidHandle->SetSpatialIdx(Data->SpatialIdx);	//todo: this needs to only happen once during initialization
 		RigidHandle->SetHashResultLowLevel(Data->HashResult);	//todo: this needs to only happen once during initialization
 #if CHAOS_CHECKED
