Fixed collisions on geometry collections and make them exist in the actual world solver so they can interact.

diff --git a/Engine/Source/Runtime/Experimental/Chaos/Public/Chaos/TriangleMeshImplicitObject.h b/Engine/Source/Runtime/Experimental/Chaos/Public/Chaos/TriangleMeshImplicitObject.h
--- a/Engine/Source/Runtime/Experimental/Chaos/Public/Chaos/TriangleMeshImplicitObject.h
+++ b/Engine/Source/Runtime/Experimental/Chaos/Public/Chaos/TriangleMeshImplicitObject.h
@@ -37,7 +37,7 @@ namespace Chaos
 		virtual T PhiWithNormal(const TVector<T, 3>& x, TVector<T, 3>& Normal) const
 		{
 			ensure(false);	//not supported yet - might support it in the future or we may change the interface
-			return (T)0;
+			return TNumericLimits<T>::Max();
 		}
 
 		virtual bool Raycast(const TVector<T, 3>& StartPoint, const TVector<T, 3>& Dir, const T Length, const T Thickness, T& OutTime, TVector<T,3>& OutPosition, TVector<T,3>& OutNormal, int32& OutFaceIndex) const override;
diff --git a/Engine/Source/Runtime/Experimental/ChaosSolvers/Private/PhysicsProxy/GeometryCollectionPhysicsProxy.cpp b/Engine/Source/Runtime/Experimental/ChaosSolvers/Private/PhysicsProxy/GeometryCollectionPhysicsProxy.cpp
--- a/Engine/Source/Runtime/Experimental/ChaosSolvers/Private/PhysicsProxy/GeometryCollectionPhysicsProxy.cpp
+++ b/Engine/Source/Runtime/Experimental/ChaosSolvers/Private/PhysicsProxy/GeometryCollectionPhysicsProxy.cpp
@@ -584,13 +584,12 @@ void FGeometryCollectionPhysicsProxy::InitializeBodiesGT()
 //				{
 //					check(false); // continue?
 //				}
-
+				FTransform ParticleTransform = Transform[Idx] * Parameters.WorldTransform;
 				if (GTParticle->ObjectState() == Chaos::EObjectStateType::Dynamic)
 				{
-					GTParticle->SetX(Transform[Idx].GetTranslation());
-					GTParticle->SetR(Transform[Idx].GetRotation());
-					//...
-					GTParticle->Proxy = this; // I feel dirty.
+					GTParticle->SetX(ParticleTransform.GetTranslation());
+					GTParticle->SetR(ParticleTransform.GetRotation());
+					GTParticle->Proxy = this;
 				}
 			}
 		}
@@ -3850,10 +3849,19 @@ void FGeometryCollectionPhysicsProxy::PushToPhysicsState(const Chaos::FParticleD
 		if (Chaos::TPBDRigidParticleHandle<float, 3>* Handle = SolverParticleHandles[HandleIdx])
 		{
 			const int32 Idx = GState->BaseIndex + HandleIdx;
-			Handle->SetX(GState->Transforms[Idx].GetTranslation());
-			Handle->SetR(GState->Transforms[Idx].GetRotation());
+
+			FTransform ParticleTransform = GState->Transforms[Idx] * Parameters.WorldTransform;
+			Handle->SetX(ParticleTransform.GetTranslation());
+			Handle->SetR(ParticleTransform.GetRotation());
 			Handle->SetDisabled(GState->DisabledStates[Idx]);
 			//...
+
+			if(Handle->Geometry() && Handle->Geometry()->HasBoundingBox())
+			{
+				Handle->SetHasBounds(true);
+				Handle->SetLocalBounds(Handle->Geometry()->BoundingBox());
+				Handle->SetWorldSpaceInflatedBounds(Handle->Geometry()->BoundingBox().TransformedBox(Chaos::TRigidTransform<float, 3>(ParticleTransform)));
+			}
 		}
 	}
 }
@@ -4064,13 +4072,17 @@ void FGeometryCollectionPhysicsProxy::PullFromPhysicsState()
 /**/
 	// We should never be changing the number of entries, this would break other 
 	// attributes in the transform group.
-	if (ensure(GTDynamicCollection->Transform.Num() == TR.Transforms.Num()))	
+	const int32 NumTransforms = GTDynamicCollection->Transform.Num();
+	if (ensure(NumTransforms == TR.Transforms.Num()))	
 	{
-		GTDynamicCollection->Transform.ExchangeArrays(TR.Transforms);
-//		GTDynamicCollection->Parent.ExchangeArrays(TR.Parent);
-//		GTDynamicCollection->Children.ExchangeArrays(TR.Children);
-//		GTDynamicCollection->SimulationType.ExchangeArrays(TR.SimulationType);
-//		GTDynamicCollection->StatusFlags.ExchangeArrays(TR.StatusFlags);
+		for(int32 TmIndex = 0; TmIndex < NumTransforms; ++TmIndex)
+		{
+			if(!TR.DisabledStates[TmIndex])
+			{
+				GTDynamicCollection->Transform[TmIndex] = TR.Transforms[TmIndex].GetRelativeTransform(Parameters.WorldTransform);
+				GTDynamicCollection->Transform[TmIndex].NormalizeRotation();
+			}
+		}
 
 		//question: why do we need this? Sleeping objects will always have to update GPU
 		GTDynamicCollection->MakeDirty();
diff --git a/Engine/Source/Runtime/Experimental/GeometryCollectionEngine/Private/ChaosBlueprint.cpp b/Engine/Source/Runtime/Experimental/GeometryCollectionEngine/Private/ChaosBlueprint.cpp
--- a/Engine/Source/Runtime/Experimental/GeometryCollectionEngine/Private/ChaosBlueprint.cpp
+++ b/Engine/Source/Runtime/Experimental/GeometryCollectionEngine/Private/ChaosBlueprint.cpp
@@ -72,7 +72,7 @@ void UChaosDestructionListener::UpdateEvents()
 				{
 					if (const FGeometryCollectionPhysicsProxy* GeometryCollectionPhysicsObject = GeometryCollectionComponent->GetPhysicsProxy())
 					{
-						RegisterChaosEvents(GeometryCollectionComponent->GetPhysicsScene());
+						RegisterChaosEvents(GeometryCollectionComponent->GetWorld()->GetPhysicsScene());
 					}
 				}
 			}
@@ -310,7 +310,7 @@ void UChaosDestructionListener::AddGeometryCollectionActor(AGeometryCollectionAc
 		{
 			if (const FGeometryCollectionPhysicsProxy* GeometryCollectionPhysicsObject = GeometryCollectionComponent->GetPhysicsProxy())
 			{
-				RegisterChaosEvents(GeometryCollectionComponent->GetPhysicsScene());
+				RegisterChaosEvents(GeometryCollectionComponent->GetWorld()->GetPhysicsScene());
 			}
 		}
 	}
diff --git a/Engine/Source/Runtime/Experimental/GeometryCollectionEngine/Private/GeometryCollection/GeometryCollectionActor.cpp b/Engine/Source/Runtime/Experimental/GeometryCollectionEngine/Private/GeometryCollection/GeometryCollectionActor.cpp
--- a/Engine/Source/Runtime/Experimental/GeometryCollectionEngine/Private/GeometryCollection/GeometryCollectionActor.cpp
+++ b/Engine/Source/Runtime/Experimental/GeometryCollectionEngine/Private/GeometryCollection/GeometryCollectionActor.cpp
@@ -69,7 +69,7 @@ bool LowLevelRaycastImp(const Chaos::TVector<float, 3>& Start, const Chaos::TVec
 	//todo(ocohen): need to add thread safety / lock semantics
 	//const TManagedArray<int32>& RigidBodyIdArray = GeomCollectionActor.GetGeometryCollectionComponent()->GetRigidBodyIdArray();
 	const TManagedArray<FGuid>& RigidBodyIdArray = GeomCollectionActor.GetGeometryCollectionComponent()->GetRigidBodyGuidArray();
-	const TSharedPtr<FPhysScene_Chaos> Scene = GeomCollectionActor.GetGeometryCollectionComponent()->GetPhysicsScene();
+	FPhysScene_Chaos* Scene = GeomCollectionActor.GetGeometryCollectionComponent()->GetPhysicsScene();
 	ensure(Scene);
 
 	const Chaos::FPhysicsSolver* Solver = GetSolver(GeomCollectionActor);
diff --git a/Engine/Source/Runtime/Experimental/GeometryCollectionEngine/Private/GeometryCollection/GeometryCollectionComponent.cpp b/Engine/Source/Runtime/Experimental/GeometryCollectionEngine/Private/GeometryCollection/GeometryCollectionComponent.cpp
--- a/Engine/Source/Runtime/Experimental/GeometryCollectionEngine/Private/GeometryCollection/GeometryCollectionComponent.cpp
+++ b/Engine/Source/Runtime/Experimental/GeometryCollectionEngine/Private/GeometryCollection/GeometryCollectionComponent.cpp
@@ -1279,7 +1279,7 @@ void UGeometryCollectionComponent::OnCreatePhysicsState()
 					if (CacheParameters.TargetCache)
 					{
 						// Queue this up to be dirtied after PIE ends
-						TSharedPtr<FPhysScene_Chaos> Scene = GetPhysicsScene();
+						FPhysScene_Chaos* Scene = GetPhysicsScene();
 
 						CacheParameters.TargetCache->PreEditChange(nullptr);
 						CacheParameters.TargetCache->Modify();
@@ -1341,7 +1341,7 @@ void UGeometryCollectionComponent::OnCreatePhysicsState()
 			// end temporary 
 
 			PhysicsProxy = new FGeometryCollectionPhysicsProxy(this, DynamicCollection.Get(), InitFunc, CacheSyncFunc, FinalSyncFunc);
-			TSharedPtr<FPhysScene_Chaos> Scene = GetPhysicsScene();
+			FPhysScene_Chaos* Scene = GetPhysicsScene();
 			Scene->AddObject(this, PhysicsProxy);
 
 			RegisterForEvents();
@@ -1373,7 +1373,7 @@ void UGeometryCollectionComponent::OnDestroyPhysicsState()
 
 	if(PhysicsProxy)
 	{
-		TSharedPtr<FPhysScene_Chaos> Scene = GetPhysicsScene();
+		FPhysScene_Chaos* Scene = GetPhysicsScene();
 		Scene->RemoveObject(PhysicsProxy);
 		InitializationState = ESimulationInitializationState::Unintialized;
 
@@ -1933,21 +1933,22 @@ void UGeometryCollectionComponent::GetInitializationCommands(TArray<FFieldSystem
 }
 
 
-const TSharedPtr<FPhysScene_Chaos> UGeometryCollectionComponent::GetPhysicsScene() const
+FPhysScene_Chaos* UGeometryCollectionComponent::GetPhysicsScene() const
 {
 	if (ChaosSolverActor)
 	{
-		return ChaosSolverActor->GetPhysicsScene();
+		return ChaosSolverActor->GetPhysicsScene().Get();
 	}
 	else
 	{
 #if INCLUDE_CHAOS
 		if (ensure(GetOwner()) && ensure(GetOwner()->GetWorld()))
 		{
-			return GetOwner()->GetWorld()->PhysicsScene_Chaos;
+			FPhysScene_ChaosInterface* WorldPhysScene = GetOwner()->GetWorld()->GetPhysicsScene();
+			return &WorldPhysScene->GetScene();
 		}
 		check(GWorld);
-		return GWorld->PhysicsScene_Chaos;
+		return &GWorld->GetPhysicsScene()->GetScene();
 #else
 		return nullptr;
 #endif
@@ -1962,7 +1963,7 @@ AChaosSolverActor* UGeometryCollectionComponent::GetPhysicsSolverActor() const
 	}
 	else
 	{
-		FPhysScene_Chaos const* const Scene = GetPhysicsScene().Get();
+		FPhysScene_Chaos const* const Scene = GetPhysicsScene();
 		return Scene ? Cast<AChaosSolverActor>(Scene->GetSolverActor()) : nullptr;
 	}
 
diff --git a/Engine/Source/Runtime/Experimental/GeometryCollectionEngine/Public/GeometryCollection/GeometryCollectionComponent.h b/Engine/Source/Runtime/Experimental/GeometryCollectionEngine/Public/GeometryCollection/GeometryCollectionComponent.h
--- a/Engine/Source/Runtime/Experimental/GeometryCollectionEngine/Public/GeometryCollection/GeometryCollectionComponent.h
+++ b/Engine/Source/Runtime/Experimental/GeometryCollectionEngine/Public/GeometryCollection/GeometryCollectionComponent.h
@@ -488,7 +488,7 @@ public:
 	FORCEINLINE const TArray<int32>& GetHighlightedBones() const { return HighlightedBones; }
 #endif
 
-	const TSharedPtr<FPhysScene_Chaos> GetPhysicsScene() const;
+	FPhysScene_Chaos* GetPhysicsScene() const;
 	AChaosSolverActor* GetPhysicsSolverActor() const;
 	const FGeometryCollectionPhysicsProxy* GetPhysicsProxy() const { return PhysicsProxy; }
 
