Fix the chaos cloth plugin build with the new TConvex constructor code that now requires face indices.
Merging CL10371659 to //UE4/Main/Engine/...

diff --git a/Engine/Plugins/Experimental/ChaosCloth/Source/Chaos/Private/ChaosCloth/ChaosClothingSimulation.cpp b/Engine/Plugins/Experimental/ChaosCloth/Source/Chaos/Private/ChaosCloth/ChaosClothingSimulation.cpp
--- a/Engine/Plugins/Experimental/ChaosCloth/Source/Chaos/Private/ChaosCloth/ChaosClothingSimulation.cpp
+++ b/Engine/Plugins/Experimental/ChaosCloth/Source/Chaos/Private/ChaosCloth/ChaosClothingSimulation.cpp
@@ -928,7 +928,8 @@ void ClothingSimulation::AddCollisions(const FClothCollisionData& ClothCollision
 					}
 
 					// Setup the collision particle geometry
-					CollisionParticles.SetDynamicGeometry(i, MakeUnique<Chaos::TConvex<float, 3>>(MoveTemp(Planes), MoveTemp(SurfaceParticles)));
+					const TArray<TArray<int32>> FaceIndices;  // TODO(Kriss.Gossart): This is broken. Find a way to pass the face indices.
+					CollisionParticles.SetDynamicGeometry(i, MakeUnique<Chaos::TConvex<float, 3>>(MoveTemp(Planes), MoveTemp(SurfaceParticles), FaceIndices));
 				}
 			}
 
diff --git a/Engine/Source/Runtime/Experimental/Chaos/Public/Chaos/Convex.h b/Engine/Source/Runtime/Experimental/Chaos/Public/Chaos/Convex.h
--- a/Engine/Source/Runtime/Experimental/Chaos/Public/Chaos/Convex.h
+++ b/Engine/Source/Runtime/Experimental/Chaos/Public/Chaos/Convex.h
@@ -30,7 +30,7 @@ namespace Chaos
 			, Volume(MoveTemp(Other.MoveTemp))
 			, CenterOfMass(MoveTemp(Other.CenterOfMass))
 		{}
-		TConvex(TArray<TPlane<T, d>>&& InPlanes, TArray<TArray<int32>>&& InFaceIndices, TParticles<T, d>&& InSurfaceParticles)
+		TConvex(TArray<TPlane<T, d>>&& InPlanes, TParticles<T, d>&& InSurfaceParticles, const TArray<TArray<int32>>& InFaceIndices)
 		    : FImplicitObject(EImplicitObject::IsConvex | EImplicitObject::HasBoundingBox, ImplicitObjectType::Convex)
 			, Planes(MoveTemp(InPlanes))
 		    , SurfaceParticles(MoveTemp(InSurfaceParticles))
@@ -41,8 +41,7 @@ namespace Chaos
 				LocalBoundingBox.GrowToInclude(SurfaceParticles.X(ParticleIndex));
 			}
 
-			TArray<TArray<int32>> FaceIndices;
-			CalculateVolumeAndCenterOfMass(SurfaceParticles, FaceIndices, Volume, CenterOfMass);
+			CalculateVolumeAndCenterOfMass(SurfaceParticles, InFaceIndices, Volume, CenterOfMass);
 		}
 		TConvex(const TParticles<T, 3>& InParticles)
 		    : FImplicitObject(EImplicitObject::IsConvex | EImplicitObject::HasBoundingBox, ImplicitObjectType::Convex)
