Add TConvex constructor without FaceIndices

diff --git a/Engine/Plugins/Experimental/ChaosCloth/Source/Chaos/Private/ChaosCloth/ChaosClothingSimulation.cpp b/Engine/Plugins/Experimental/ChaosCloth/Source/Chaos/Private/ChaosCloth/ChaosClothingSimulation.cpp
--- a/Engine/Plugins/Experimental/ChaosCloth/Source/Chaos/Private/ChaosCloth/ChaosClothingSimulation.cpp
+++ b/Engine/Plugins/Experimental/ChaosCloth/Source/Chaos/Private/ChaosCloth/ChaosClothingSimulation.cpp
@@ -928,8 +928,7 @@ void ClothingSimulation::AddCollisions(const FClothCollisionData& ClothCollision
 					}
 
 					// Setup the collision particle geometry
-					const TArray<TArray<int32>> FaceIndices;  // TODO(Kriss.Gossart): This is broken. Find a way to pass the face indices.
-					CollisionParticles.SetDynamicGeometry(i, MakeUnique<Chaos::TConvex<float, 3>>(MoveTemp(Planes), MoveTemp(SurfaceParticles), FaceIndices));
+					CollisionParticles.SetDynamicGeometry(i, MakeUnique<Chaos::TConvex<float, 3>>(MoveTemp(Planes), MoveTemp(SurfaceParticles)));
 				}
 			}
 
diff --git a/Engine/Source/Runtime/Experimental/Chaos/Public/Chaos/Convex.h b/Engine/Source/Runtime/Experimental/Chaos/Public/Chaos/Convex.h
--- a/Engine/Source/Runtime/Experimental/Chaos/Public/Chaos/Convex.h
+++ b/Engine/Source/Runtime/Experimental/Chaos/Public/Chaos/Convex.h
@@ -30,7 +30,11 @@ namespace Chaos
 			, Volume(MoveTemp(Other.MoveTemp))
 			, CenterOfMass(MoveTemp(Other.CenterOfMass))
 		{}
-		TConvex(TArray<TPlane<T, d>>&& InPlanes, TParticles<T, d>&& InSurfaceParticles, const TArray<TArray<int32>>& InFaceIndices)
+
+		// NOTE: This constructor will result in approximate COM and volume calculations, since it does
+		// not have face indices for surface particles.
+		// TODO: Keep track of invalid state and ensure on volume or COM access?
+		TConvex(TArray<TPlane<T, d>>&& InPlanes, TParticles<T, d>&& InSurfaceParticles)
 		    : FImplicitObject(EImplicitObject::IsConvex | EImplicitObject::HasBoundingBox, ImplicitObjectType::Convex)
 			, Planes(MoveTemp(InPlanes))
 		    , SurfaceParticles(MoveTemp(InSurfaceParticles))
@@ -41,8 +45,11 @@ namespace Chaos
 				LocalBoundingBox.GrowToInclude(SurfaceParticles.X(ParticleIndex));
 			}
 
-			CalculateVolumeAndCenterOfMass(SurfaceParticles, InFaceIndices, Volume, CenterOfMass);
+			// For now we approximate COM and volume with the bounding box
+			CenterOfMass = LocalBoundingBox.GetCenterOfMass();
+			Volume = LocalBoundingBox.GetVolume();
 		}
+
 		TConvex(const TParticles<T, 3>& InParticles)
 		    : FImplicitObject(EImplicitObject::IsConvex | EImplicitObject::HasBoundingBox, ImplicitObjectType::Convex)
 		{
