Chaos  :

Remove parallel flag.

diff --git a/Engine/Source/Runtime/Experimental/Chaos/Public/Chaos/PBDCollisionConstraintsImp.h b/Engine/Source/Runtime/Experimental/Chaos/Public/Chaos/PBDCollisionConstraintsImp.h
--- a/Engine/Source/Runtime/Experimental/Chaos/Public/Chaos/PBDCollisionConstraintsImp.h
+++ b/Engine/Source/Runtime/Experimental/Chaos/Public/Chaos/PBDCollisionConstraintsImp.h
@@ -355,7 +355,7 @@ namespace Chaos
 
 			COLLISION_OPT_OUT(StatData.FinalizeData());
 
-		}, bGatherStats || bDisableCollisionParallelFor || true); // todo(brice) remove true.
+		}, bGatherStats || bDisableCollisionParallelFor);
 
 		{
 			SCOPE_CYCLE_COUNTER(STAT_ComputeConstraintsSU);
