Chaos : Collisions

Fixed typo

diff --git a/Engine/Source/Runtime/Experimental/Chaos/Private/Chaos/CollisionResolution.cpp b/Engine/Source/Runtime/Experimental/Chaos/Private/Chaos/CollisionResolution.cpp
--- a/Engine/Source/Runtime/Experimental/Chaos/Private/Chaos/CollisionResolution.cpp
+++ b/Engine/Source/Runtime/Experimental/Chaos/Private/Chaos/CollisionResolution.cpp
@@ -1015,7 +1015,7 @@ namespace Chaos
 			{
 				ConstructBoxPlaneConstraints(Particle0, Particle1, Implicit0, Implicit1, Thickness, NewConstraints);
 			}
-			else if (Implicit0->GetType() == TPlane<T, d>::StaticType() && Implicit0->GetType() == TBox<T, d>::StaticType())
+			else if (Implicit0->GetType() == TPlane<T, d>::StaticType() && Implicit1->GetType() == TBox<T, d>::StaticType())
 			{
 				ConstructBoxPlaneConstraints(Particle1, Particle0, Implicit1, Implicit0, Thickness, NewConstraints);
 			}
