Add COM data to KinematicParticle

diff --git a/Engine/Source/Runtime/Core/Public/UObject/ExternalPhysicsCustomObjectVersion.h b/Engine/Source/Runtime/Core/Public/UObject/ExternalPhysicsCustomObjectVersion.h
--- a/Engine/Source/Runtime/Core/Public/UObject/ExternalPhysicsCustomObjectVersion.h
+++ b/Engine/Source/Runtime/Core/Public/UObject/ExternalPhysicsCustomObjectVersion.h
@@ -71,6 +71,9 @@ struct CORE_API FExternalPhysicsCustomObjectVersion
 
 		// Add center of mass and volume cached calculations to TConvex
 		AddConvexCenterOfMassAndVolume,
+
+		// Add mass transform data to kinematic particle
+		KinematicCentersOfMass,
 		
 		// -----<new versions can be added above this line>-------------------------------------------------
 		VersionPlusOne,
diff --git a/Engine/Source/Runtime/Experimental/Chaos/Public/Chaos/KinematicGeometryParticles.h b/Engine/Source/Runtime/Experimental/Chaos/Public/Chaos/KinematicGeometryParticles.h
--- a/Engine/Source/Runtime/Experimental/Chaos/Public/Chaos/KinematicGeometryParticles.h
+++ b/Engine/Source/Runtime/Experimental/Chaos/Public/Chaos/KinematicGeometryParticles.h
@@ -70,6 +70,8 @@ class TKinematicGeometryParticlesImp : public TGeometryParticlesImp<T, d, SimTyp
 		TArrayCollection::AddArray(&MV);
 		TArrayCollection::AddArray(&MW);
 		TArrayCollection::AddArray(&KinematicTargets);
+		TArrayCollection::AddArray(&MCenterOfMass);
+		TArrayCollection::AddArray(&MRotationOfMass);
 	}
 	TKinematicGeometryParticlesImp(const TKinematicGeometryParticlesImp<T, d, SimType>& Other) = delete;
 	CHAOS_API TKinematicGeometryParticlesImp(TKinematicGeometryParticlesImp<T, d, SimType>&& Other)
@@ -79,6 +81,8 @@ class TKinematicGeometryParticlesImp : public TGeometryParticlesImp<T, d, SimTyp
 		TArrayCollection::AddArray(&MV);
 		TArrayCollection::AddArray(&MW);
 		TArrayCollection::AddArray(&KinematicTargets);
+		TArrayCollection::AddArray(&MCenterOfMass);
+		TArrayCollection::AddArray(&MRotationOfMass);
 	}
 	CHAOS_API virtual ~TKinematicGeometryParticlesImp();
 
@@ -91,6 +95,12 @@ class TKinematicGeometryParticlesImp : public TGeometryParticlesImp<T, d, SimTyp
 	const TKinematicTarget<T, d>& KinematicTarget(const int32 Index) const { return KinematicTargets[Index]; }
 	TKinematicTarget<T, d>& KinematicTarget(const int32 Index) { return KinematicTargets[Index]; }
 
+	const TVector<T, d>& CenterOfMass(const int32 Index) const { return MCenterOfMass[Index]; }
+	TVector<T, d>& CenterOfMass(const int32 Index) { return MCenterOfMass[Index]; }
+
+	const TRotation<T, d>& RotationOfMass(const int32 Index) const { return MRotationOfMass[Index]; }
+	TRotation<T, d>& RotationOfMass(const int32 Index) { return MRotationOfMass[Index]; }
+
 	FString ToString(int32 index) const
 	{
 		FString BaseString = TGeometryParticlesImp<T, d, SimType>::ToString(index);
@@ -113,12 +123,19 @@ class TKinematicGeometryParticlesImp : public TGeometryParticlesImp<T, d, SimTyp
 		{
 			Ar << KinematicTargets;
 		}
+		if (Ar.CustomVer(FExternalPhysicsCustomObjectVersion::GUID) >= FExternalPhysicsCustomObjectVersion::KinematicCentersOfMass)
+		{
+			Ar << MCenterOfMass;
+			Ar << MRotationOfMass;
+		}
 	}
 
   private:
 	TArrayCollectionArray<TVector<T, d>> MV;
 	TArrayCollectionArray<TVector<T, d>> MW;
 	TArrayCollectionArray<TKinematicTarget<T, d>> KinematicTargets;
+	TArrayCollectionArray<TVector<T, d>> MCenterOfMass;
+	TArrayCollectionArray<TRotation<T, d>> MRotationOfMass;
 };
 
 template <typename T, int d, EGeometryParticlesSimType SimType>
diff --git a/Engine/Source/Runtime/Experimental/Chaos/Public/Chaos/ParticleDirtyFlags.h b/Engine/Source/Runtime/Experimental/Chaos/Public/Chaos/ParticleDirtyFlags.h
--- a/Engine/Source/Runtime/Experimental/Chaos/Public/Chaos/ParticleDirtyFlags.h
+++ b/Engine/Source/Runtime/Experimental/Chaos/Public/Chaos/ParticleDirtyFlags.h
@@ -13,27 +13,29 @@ namespace Chaos
 		R =				 1 << 1,
 		V =				 1 << 2,
 		W =				 1 << 3,
-		CollisionGroup = 1 << 4,
-		Disabled =		 1 << 5,
-		PreV =			 1 << 6,
-		PreW =			 1 << 7,
-		P =				 1 << 8,
-		Q =				 1 << 9,
-		F =				 1 << 10,
-		Torque =		 1 << 11,
-		I =				 1 << 12,
-		InvI =			 1 << 13,
-		M =				 1 << 14,
-		InvM =			 1 << 15,
-		ObjectState =	 1 << 16,
-		Geometry =		 1 << 17,
-		ExternalForce =  1 << 18,
-		ExternalTorque = 1 << 19,
-		GravityEnabled = 1 << 20,
-		SpatialIdx     = 1 << 21,
-		HashResult     = 1 << 22
+		CenterOfMass =   1 << 4,
+		RotationOfMass = 1 << 5,
+		CollisionGroup = 1 << 6,
+		Disabled =		 1 << 7,
+		PreV =			 1 << 8,
+		PreW =			 1 << 9,
+		P =				 1 << 10,
+		Q =				 1 << 11,
+		F =				 1 << 12,
+		Torque =		 1 << 13,
+		I =				 1 << 14,
+		InvI =			 1 << 15,
+		M =				 1 << 16,
+		InvM =			 1 << 17,
+		ObjectState =	 1 << 18,
+		Geometry =		 1 << 19,
+		ExternalForce =  1 << 20,
+		ExternalTorque = 1 << 20,
+		GravityEnabled = 1 << 22,
+		SpatialIdx     = 1 << 23,
+		HashResult     = 1 << 24
 #if CHAOS_CHECKED
-		, DebugName    = 1 << 23
+		, DebugName    = 1 << 25
 #endif
 	};
 
diff --git a/Engine/Source/Runtime/Experimental/Chaos/Public/Chaos/ParticleHandle.h b/Engine/Source/Runtime/Experimental/Chaos/Public/Chaos/ParticleHandle.h
--- a/Engine/Source/Runtime/Experimental/Chaos/Public/Chaos/ParticleHandle.h
+++ b/Engine/Source/Runtime/Experimental/Chaos/Public/Chaos/ParticleHandle.h
@@ -58,6 +58,8 @@ void KinematicGeometryParticleDefaultConstruct(FConcrete& Concrete, const TKinem
 {
 	Concrete.SetV(TVector<T, d>(0));
 	Concrete.SetW(TVector<T, d>(0));
+	Concrete.SetCenterOfMass(TVector<T, d>(0));
+	Concrete.SetRotationOfMass(TRotation<T, d>(FQuat(EForceInit::ForceInit)));
 }
 
 template <typename T, int d, bool bPersistent>
@@ -542,6 +544,12 @@ public:
 	const TKinematicTarget<T, d>& KinematicTarget() const { return KinematicGeometryParticles->KinematicTarget(ParticleIdx); }
 	TKinematicTarget<T, d>& KinematicTarget() { return KinematicGeometryParticles->KinematicTarget(ParticleIdx); }
 
+	const TVector<T, d>& CenterOfMass() const { return KinematicGeometryParticles->CenterOfMass(ParticleIdx); }
+	void SetCenterOfMass(const TVector<T, d>& InCenterOfMass) { KinematicGeometryParticles->CenterOfMass(ParticleIdx) = InCenterOfMass; }
+	
+	const TRotation<T, d>& RotationOfMass() const { return KinematicGeometryParticles->RotationOfMass(ParticleIdx); }
+	void SetRotationOfMass(const TRotation<T, d>& InRotationOfMass) { KinematicGeometryParticles->RotationOfMass(ParticleIdx) = InRotationOfMass; }
+
 	//Really only useful when using a transient handle
 	const TKinematicGeometryParticleHandleImp<T, d, true>* Handle() const { return KinematicGeometryParticles->Handle(ParticleIdx); }
 	TKinematicGeometryParticleHandleImp<T, d, true>* Handle() { return KinematicGeometryParticles->Handle(ParticleIdx); }
@@ -1514,6 +1522,20 @@ public:
 		this->MW = InW;
 	}
 
+	const TVector<T, d>& CenterOfMass() const { return MCenterOfMass; }
+	void SetCenterOfMass(const TVector<T, d>& InCenterOfMass, bool bInvalidate = true)
+	{
+		this->MarkDirty(EParticleFlags::CenterOfMass, bInvalidate);
+		this->MCenterOfMass = InCenterOfMass;
+	}
+	
+	const TRotation<T, d>& RotationOfMass() const { return MRotationOfMass; }
+	void SetRotationOfMass(const TRotation<T, d>& InRotationOfMass, bool bInvalidate = true)
+	{
+		this->MarkDirty(EParticleFlags::RotationOfMass, bInvalidate);
+		this->MRotationOfMass = InRotationOfMass;
+	}
+
 	EObjectStateType ObjectState() const;
 
 	FParticleData* NewData() const
@@ -1524,6 +1546,8 @@ public:
 private:
 	TVector<T, d> MV;
 	TVector<T, d> MW;
+	TVector<T, d> MCenterOfMass;
+	TRotation<T, d> MRotationOfMass;
 };
 
 template <typename T, int d>
@@ -1537,12 +1561,16 @@ public:
 	TKinematicGeometryParticleData(EParticleType InType = EParticleType::Kinematic)
 		: Base(InType)
 		, MV(TVector<T, d>(0))
-		, MW(TVector<T, d>(0)) {}
+		, MW(TVector<T, d>(0))
+		, MCenterOfMass(TVector<T, d>(0))
+		, MRotationOfMass(TRotation<T, d>(FQuat(EForceInit::ForceInit))) {}
 
 	TKinematicGeometryParticleData(const TKinematicGeometryParticle<T, d>& InParticle)
 		: Base(InParticle)
 		, MV(InParticle.V())
 		, MW(InParticle.W()) 
+		, MCenterOfMass(InParticle.CenterOfMass()) 
+		, MRotationOfMass(InParticle.RotationOfMass()) 
 	{
 		Type = EParticleType::Kinematic;
 	}
@@ -1551,11 +1579,16 @@ public:
 	void Reset() {
 		TGeometryParticleData<T, d>::Reset();
 		Type = EParticleType::Kinematic;
-		MV = TVector<T, d>(0); MW = TVector<T, d>();
+		MV = TVector<T, d>(0);
+		MW = TVector<T, d>(0);
+		MCenterOfMass = TVector<T, d>(0);
+		MRotationOfMass = TRotation<T, d>(FQuat(EForceInit::ForceInit));
 	}
 
 	TVector<T, d> MV;
 	TVector<T, d> MW;
+	TVector<T, d> MCenterOfMass;
+	TRotation<T, d> MRotationOfMass;
 };
 
 
diff --git a/Engine/Source/Runtime/Experimental/ChaosSolvers/Private/PhysicsProxy/SingleParticlePhysicsProxy.cpp b/Engine/Source/Runtime/Experimental/ChaosSolvers/Private/PhysicsProxy/SingleParticlePhysicsProxy.cpp
--- a/Engine/Source/Runtime/Experimental/ChaosSolvers/Private/PhysicsProxy/SingleParticlePhysicsProxy.cpp
+++ b/Engine/Source/Runtime/Experimental/ChaosSolvers/Private/PhysicsProxy/SingleParticlePhysicsProxy.cpp
@@ -119,7 +119,8 @@ void FSingleParticlePhysicsProxy<Chaos::TKinematicGeometryParticle<float, 3>>::P
 		RigidHandle->SetSharedGeometry(Data->Geometry);
 		RigidHandle->SetV(Data->MV);
 		RigidHandle->SetW(Data->MW);
-
+		RigidHandle->SetCenterOfMass(Data->MCenterOfMass);
+		RigidHandle->SetRotationOfMass(Data->MRotationOfMass);
 		RigidHandle->SetSpatialIdx(Data->SpatialIdx);	//todo: this needs to only happen once during initialization
 		RigidHandle->SetHashResultLowLevel(Data->HashResult);	//todo: this needs to only happen once during initialization
 #if CHAOS_CHECKED
@@ -220,6 +221,8 @@ void FSingleParticlePhysicsProxy<Chaos::TPBDRigidParticle<float, 3>>::PushToPhys
 			RigidHandle->SetW(Data->MW);
 		}
 
+		RigidHandle->SetCenterOfMass(Data->MCenterOfMass);
+		RigidHandle->SetRotationOfMass(Data->MRotationOfMass);
 		RigidHandle->SetM(Data->MM);
 		RigidHandle->SetInvM(Data->MInvM);
 		RigidHandle->SetI(Data->MI);
