Chaos - support transformed shapes in collision detection

diff --git a/Engine/Source/Runtime/Experimental/Chaos/Private/Chaos/CollisionResolution.cpp b/Engine/Source/Runtime/Experimental/Chaos/Private/Chaos/CollisionResolution.cpp
--- a/Engine/Source/Runtime/Experimental/Chaos/Private/Chaos/CollisionResolution.cpp
+++ b/Engine/Source/Runtime/Experimental/Chaos/Private/Chaos/CollisionResolution.cpp
@@ -27,7 +27,7 @@ namespace Chaos
 		//
 
 		template <typename T, int d>
-		void UpdateBoxBoxConstraint(const TBox<T, d>& Box1, const TRigidTransform<T, d>& Box1Transform, const TBox<T, d>& Box2, const TRigidTransform<T, d>& Box2Transform, T Thickness, TRigidBodyPointContactConstraint<T, d>& Constraint)
+		void UpdateBoxBoxConstraint(const TBox<T, d>& Box1, const TRigidTransform<T, d>& Box1Transform, const TBox<T, d>& Box2, const TRigidTransform<T, d>& Box2Transform, const T Thickness, TRigidBodyPointContactConstraint<T, d>& Constraint)
 		{
 			TPointContactManifold<T, d> & Contact = Constraint.Manifold;
 
@@ -80,7 +80,7 @@ namespace Chaos
 
 
 		template<typename T, int d>
-		void ConstructBoxBoxConstraints(TGeometryParticleHandle<T, d>* Particle0, TGeometryParticleHandle<T, d>* Particle1, const FImplicitObject* Implicit0, const FImplicitObject* Implicit1, const T Thickness, FCollisionConstraintsArray& NewConstraints)
+		void ConstructBoxBoxConstraints(TGeometryParticleHandle<T, d>* Particle0, TGeometryParticleHandle<T, d>* Particle1, const FImplicitObject* Implicit0, const FImplicitObject* Implicit1, const TRigidTransform<T, d>& Transform0, const TRigidTransform<T, d>& Transform1, const T Thickness, FCollisionConstraintsArray& NewConstraints)
 		{
 			const TBox<T, d> * Object0 = Implicit0->template GetObject<const TBox<T, d> >();
 			const TBox<T, d> * Object1 = Implicit1->template GetObject<const TBox<T, d> >();
@@ -91,7 +91,7 @@ namespace Chaos
 				Constraint->Particle[1] = Particle1;
 				Constraint->AddManifold(Implicit0, Implicit1);
 
-				UpdateBoxBoxConstraint(*Object0, GetTransform(Particle0), *Object1, GetTransform(Particle1), Thickness, *Constraint);
+				UpdateBoxBoxConstraint(*Object0, Transform0, *Object1, Transform1, Thickness, *Constraint);
 
 				if (Constraint->GetPhi() < Thickness)
 				{
@@ -109,7 +109,7 @@ namespace Chaos
 		//
 
 		template <typename T, int d>
-		bool UpdateBoxPlaneConstraint(const TBox<T, d>& Box, const TRigidTransform<T, d>& BoxTransform, const TPlane<T, d>& Plane, const TRigidTransform<T, d>& PlaneTransform, T Thickness, TRigidBodyPointContactConstraint<T, d>& Constraint)
+		bool UpdateBoxPlaneConstraint(const TBox<T, d>& Box, const TRigidTransform<T, d>& BoxTransform, const TPlane<T, d>& Plane, const TRigidTransform<T, d>& PlaneTransform, const T Thickness, TRigidBodyPointContactConstraint<T, d>& Constraint)
 		{
 			TPointContactManifold<T, d> & Contact = Constraint.Manifold;
 
@@ -172,7 +172,7 @@ namespace Chaos
 
 
 		template<typename T, int d>
-		void ConstructBoxPlaneConstraints(TGeometryParticleHandle<T, d>* Particle0, TGeometryParticleHandle<T, d>* Particle1, const FImplicitObject* Implicit0, const FImplicitObject* Implicit1, const T Thickness, FCollisionConstraintsArray& NewConstraints)
+		void ConstructBoxPlaneConstraints(TGeometryParticleHandle<T, d>* Particle0, TGeometryParticleHandle<T, d>* Particle1, const FImplicitObject* Implicit0, const FImplicitObject* Implicit1, const TRigidTransform<T, d>& Transform0, const TRigidTransform<T, d>& Transform1, const T Thickness, FCollisionConstraintsArray& NewConstraints)
 		{
 			
 			const TBox<T, d> * Object0 = Implicit0->template GetObject<const TBox<T, d> >();
@@ -184,7 +184,7 @@ namespace Chaos
 				Constraint->Particle[1] = Particle1;
 				Constraint->AddManifold(Implicit0, Implicit1);
 
-				UpdateBoxPlaneConstraint(*Object0, GetTransform(Particle0), *Object1, GetTransform(Particle1), Thickness, *Constraint);
+				UpdateBoxPlaneConstraint(*Object0, Transform0, *Object1, Transform1, Thickness, *Constraint);
 
 				if (Constraint->GetPhi() < Thickness)
 				{
@@ -202,7 +202,7 @@ namespace Chaos
 		// Sphere - Sphere
 		//
 		template <typename T, int d>
-		void UpdateSphereSphereConstraint(const TSphere<T, d>& Sphere1, const TRigidTransform<T, d>& Sphere1Transform, const TSphere<T, d>& Sphere2, const TRigidTransform<T, d>& Sphere2Transform, T Thickness, TRigidBodyPointContactConstraint<T, d>& Constraint)
+		void UpdateSphereSphereConstraint(const TSphere<T, d>& Sphere1, const TRigidTransform<T, d>& Sphere1Transform, const TSphere<T, d>& Sphere2, const TRigidTransform<T, d>& Sphere2Transform, const T Thickness, TRigidBodyPointContactConstraint<T, d>& Constraint)
 		{
 			TPointContactManifold<T, d> & Contact = Constraint.Manifold;
 
@@ -220,7 +220,7 @@ namespace Chaos
 		}
 
 		template<typename T, int d>
-		void ConstructSphereSphereConstraints(TGeometryParticleHandle<T, d>* Particle0, TGeometryParticleHandle<T, d>* Particle1, const FImplicitObject* Implicit0, const FImplicitObject* Implicit1, const T Thickness, FCollisionConstraintsArray& NewConstraints)
+		void ConstructSphereSphereConstraints(TGeometryParticleHandle<T, d>* Particle0, TGeometryParticleHandle<T, d>* Particle1, const FImplicitObject* Implicit0, const FImplicitObject* Implicit1, const TRigidTransform<T, d>& Transform0, const TRigidTransform<T, d>& Transform1, const T Thickness, FCollisionConstraintsArray& NewConstraints)
 		{
 			
 			const TSphere<T, d> * Object0 = Implicit0->template GetObject<const TSphere<T, d> >();
@@ -232,7 +232,7 @@ namespace Chaos
 				Constraint->Particle[1] = Particle1;
 				Constraint->AddManifold(Implicit0, Implicit1);
 
-				UpdateSphereSphereConstraint(*Object0, GetTransform(Particle0), *Object1, GetTransform(Particle1), Thickness, *Constraint);
+				UpdateSphereSphereConstraint(*Object0, Transform0, *Object1, Transform1, Thickness, *Constraint);
 
 				if (Constraint->GetPhi() < Thickness)
 				{
@@ -250,7 +250,7 @@ namespace Chaos
 		//
 
 		template <typename T, int d>
-		void UpdateSpherePlaneConstraint(const TSphere<T, d>& Sphere, const TRigidTransform<T, d>& SphereTransform, const TPlane<T, d>& Plane, const TRigidTransform<T, d>& PlaneTransform, T Thickness, TRigidBodyPointContactConstraint<T, d>& Constraint)
+		void UpdateSpherePlaneConstraint(const TSphere<T, d>& Sphere, const TRigidTransform<T, d>& SphereTransform, const TPlane<T, d>& Plane, const TRigidTransform<T, d>& PlaneTransform, const T Thickness, TRigidBodyPointContactConstraint<T, d>& Constraint)
 		{
 			TPointContactManifold<T, d> & Contact = Constraint.Manifold;
 
@@ -270,7 +270,7 @@ namespace Chaos
 		}
 
 		template<typename T, int d>
-		void ConstructSpherePlaneConstraints(TGeometryParticleHandle<T, d>* Particle0, TGeometryParticleHandle<T, d>* Particle1, const FImplicitObject* Implicit0, const FImplicitObject* Implicit1, const T Thickness, FCollisionConstraintsArray& NewConstraints)
+		void ConstructSpherePlaneConstraints(TGeometryParticleHandle<T, d>* Particle0, TGeometryParticleHandle<T, d>* Particle1, const FImplicitObject* Implicit0, const FImplicitObject* Implicit1, const TRigidTransform<T, d>& Transform0, const TRigidTransform<T, d>& Transform1, const T Thickness, FCollisionConstraintsArray& NewConstraints)
 		{
 			
 			const TSphere<T, d> * Object0 = Implicit0->template GetObject<const TSphere<T, d> >();
@@ -283,7 +283,7 @@ namespace Chaos
 				Constraint->Particle[1] = Particle1;
 				Constraint->AddManifold(Implicit0, Implicit1);
 
-				UpdateSpherePlaneConstraint(*Object0, GetTransform(Particle0), *Object1, GetTransform(Particle1), Thickness, *Constraint);
+				UpdateSpherePlaneConstraint(*Object0, Transform0, *Object1, Transform1, Thickness, *Constraint);
 
 				if (Constraint->GetPhi() < Thickness)
 				{
@@ -301,7 +301,7 @@ namespace Chaos
 		//
 
 		template <typename T, int d>
-		void UpdateSphereBoxConstraint(const TSphere<T, d>& Sphere, const TRigidTransform<T, d>& SphereTransform, const TBox<T, d>& Box, const TRigidTransform<T, d>& BoxTransform, T Thickness, TRigidBodyPointContactConstraint<T, d>& Constraint)
+		void UpdateSphereBoxConstraint(const TSphere<T, d>& Sphere, const TRigidTransform<T, d>& SphereTransform, const TBox<T, d>& Box, const TRigidTransform<T, d>& BoxTransform, const T Thickness, TRigidBodyPointContactConstraint<T, d>& Constraint)
 		{
 			TPointContactManifold<T, d> & Contact = Constraint.Manifold;
 
@@ -321,7 +321,7 @@ namespace Chaos
 		}
 
 		template<typename T, int d>
-		void ConstructSphereBoxConstraints(TGeometryParticleHandle<T, d>* Particle0, TGeometryParticleHandle<T, d>* Particle1, const FImplicitObject* Implicit0, const FImplicitObject* Implicit1, const T Thickness, FCollisionConstraintsArray& NewConstraints)
+		void ConstructSphereBoxConstraints(TGeometryParticleHandle<T, d>* Particle0, TGeometryParticleHandle<T, d>* Particle1, const FImplicitObject* Implicit0, const FImplicitObject* Implicit1, const TRigidTransform<T, d>& Transform0, const TRigidTransform<T, d>& Transform1, const T Thickness, FCollisionConstraintsArray& NewConstraints)
 		{
 			
 			const TSphere<T, d> * Object0 = Implicit0->template GetObject<const TSphere<T, d> >();
@@ -333,7 +333,7 @@ namespace Chaos
 				Constraint->Particle[1] = Particle1;
 				Constraint->AddManifold(Implicit0, Implicit1);
 
-				UpdateSphereBoxConstraint(*Object0, GetTransform(Particle0), *Object1, GetTransform(Particle1), Thickness, *Constraint);
+				UpdateSphereBoxConstraint(*Object0, Transform0, *Object1, Transform1, Thickness, *Constraint);
 
 				if (Constraint->GetPhi() < Thickness)
 				{
@@ -352,7 +352,7 @@ namespace Chaos
 		//
 
 		template <typename T, int d>
-		void UpdateCapsuleCapsuleConstraint(const TCapsule<T>& A, const TRigidTransform<T, d>& ATransform, const TCapsule<T>& B, const TRigidTransform<T, d>& BTransform, T Thickness, TRigidBodyPointContactConstraint<T, d>& Constraint)
+		void UpdateCapsuleCapsuleConstraint(const TCapsule<T>& A, const TRigidTransform<T, d>& ATransform, const TCapsule<T>& B, const TRigidTransform<T, d>& BTransform, const T Thickness, TRigidBodyPointContactConstraint<T, d>& Constraint)
 		{
 			TPointContactManifold<T, d> & Contact = Constraint.Manifold;
 
@@ -379,7 +379,7 @@ namespace Chaos
 		}
 
 		template<typename T, int d>
-		void ConstructCapsuleCapsuleConstraints(TGeometryParticleHandle<T, d>* Particle0, TGeometryParticleHandle<T, d>* Particle1, const FImplicitObject* Implicit0, const FImplicitObject* Implicit1, const T Thickness, FCollisionConstraintsArray& NewConstraints)
+		void ConstructCapsuleCapsuleConstraints(TGeometryParticleHandle<T, d>* Particle0, TGeometryParticleHandle<T, d>* Particle1, const FImplicitObject* Implicit0, const FImplicitObject* Implicit1, const TRigidTransform<T, d>& Transform0, const TRigidTransform<T, d>& Transform1, const T Thickness, FCollisionConstraintsArray& NewConstraints)
 		{
 			
 			const TCapsule<T> * Object0 = Implicit0->template GetObject<const TCapsule<T> >();
@@ -391,7 +391,7 @@ namespace Chaos
 				Constraint->Particle[1] = Particle1;
 				Constraint->AddManifold(Implicit0, Implicit1);
 
-				UpdateCapsuleCapsuleConstraint(*Object0, GetTransform(Particle0), *Object1, GetTransform(Particle1), Thickness, *Constraint);
+				UpdateCapsuleCapsuleConstraint(*Object0, Transform0, *Object1, Transform1, Thickness, *Constraint);
 
 				if (Constraint->GetPhi() < Thickness)
 				{
@@ -409,7 +409,7 @@ namespace Chaos
 		//
 
 		template <typename T, int d>
-		void UpdateCapsuleBoxConstraint(const TCapsule<T>& A, const TRigidTransform<T, d>& ATransform, const TBox<T, d>& B, const TRigidTransform<T, d>& BTransform, T Thickness, TRigidBodyPointContactConstraint<T, d>& Constraint)
+		void UpdateCapsuleBoxConstraint(const TCapsule<T>& A, const TRigidTransform<T, d>& ATransform, const TBox<T, d>& B, const TRigidTransform<T, d>& BTransform, const T Thickness, TRigidBodyPointContactConstraint<T, d>& Constraint)
 		{
 			// @todo(ccaulfield): Add custom capsule-box collision
 			TPointContactManifold<T, d> & Contact = Constraint.Manifold;
@@ -446,7 +446,7 @@ namespace Chaos
 		}
 
 		template<typename T, int d>
-		void ConstructCapsuleBoxConstraints(TGeometryParticleHandle<T, d>* Particle0, TGeometryParticleHandle<T, d>* Particle1, const FImplicitObject* Implicit0, const FImplicitObject* Implicit1, const T Thickness, FCollisionConstraintsArray& NewConstraints)
+		void ConstructCapsuleBoxConstraints(TGeometryParticleHandle<T, d>* Particle0, TGeometryParticleHandle<T, d>* Particle1, const FImplicitObject* Implicit0, const FImplicitObject* Implicit1, const TRigidTransform<T, d>& Transform0, const TRigidTransform<T, d>& Transform1, const T Thickness, FCollisionConstraintsArray& NewConstraints)
 		{
 			
 			const TCapsule<T> * Object0 = Implicit0->template GetObject<const TCapsule<T> >();
@@ -458,7 +458,7 @@ namespace Chaos
 				Constraint->Particle[1] = Particle1;
 				Constraint->AddManifold(Implicit0, Implicit1);
 
-				UpdateCapsuleBoxConstraint(*Object0, GetTransform(Particle0), *Object1, GetTransform(Particle1), Thickness, *Constraint);
+				UpdateCapsuleBoxConstraint(*Object0, Transform0, *Object1, Transform1, Thickness, *Constraint);
 
 				if (Constraint->GetPhi() < Thickness)
 				{
@@ -476,7 +476,7 @@ namespace Chaos
 		//
 
 		template<class T, int d>
-		void UpdateConvexConvexConstraint(const FImplicitObject& A, const TRigidTransform<T, d>& ATM, const FImplicitObject& B, const TRigidTransform<T, d>& BTM, T Thickness, TRigidBodyPointContactConstraint<T, d>& Constraint)
+		void UpdateConvexConvexConstraint(const FImplicitObject& A, const TRigidTransform<T, d>& ATM, const FImplicitObject& B, const TRigidTransform<T, d>& BTM, const T Thickness, TRigidBodyPointContactConstraint<T, d>& Constraint)
 		{
 			TRigidTransform<T, d> BToATM = BTM.GetRelativeTransform(ATM);
 
@@ -528,7 +528,7 @@ namespace Chaos
 		}
 
 		template<class T, int d>
-		void ConstructConvexConvexConstraints(TGeometryParticleHandle<T, d>* Particle0, TGeometryParticleHandle<T, d>* Particle1, const FImplicitObject* Implicit0, const FImplicitObject* Implicit1, const T Thickness, FCollisionConstraintsArray& NewConstraints)
+		void ConstructConvexConvexConstraints(TGeometryParticleHandle<T, d>* Particle0, TGeometryParticleHandle<T, d>* Particle1, const FImplicitObject* Implicit0, const FImplicitObject* Implicit1, const TRigidTransform<T, d>& Transform0, const TRigidTransform<T, d>& Transform1, const T Thickness, FCollisionConstraintsArray& NewConstraints)
 		{
 			if (ensure(GetInnerType(Implicit0->GetType()) == ImplicitObjectType::Convex && GetInnerType(Implicit1->GetType()) == ImplicitObjectType::Convex))
 			{
@@ -538,7 +538,7 @@ namespace Chaos
 				Constraint->Particle[1] = Particle1;
 				Constraint->AddManifold(Implicit0, Implicit1);
 
-				UpdateConvexConvexConstraint(*Implicit0, GetTransform(Particle0), *Implicit1, GetTransform(Particle1), Thickness, *Constraint);
+				UpdateConvexConvexConstraint(*Implicit0, Transform0, *Implicit1, Transform1, Thickness, *Constraint);
 
 				if (Constraint->GetPhi() < Thickness)
 				{
@@ -587,7 +587,7 @@ namespace Chaos
 
 
 		template<typename T, int d>
-		void ConstructLevelsetLevelsetConstraints(TGeometryParticleHandle<T, d>* Particle0, TGeometryParticleHandle<T, d>* Particle1, const FImplicitObject* Implicit0, const FImplicitObject* Implicit1, const T Thickness, FCollisionConstraintsArray& NewConstraints)
+		void ConstructLevelsetLevelsetConstraints(TGeometryParticleHandle<T, d>* Particle0, TGeometryParticleHandle<T, d>* Particle1, const FImplicitObject* Implicit0, const FImplicitObject* Implicit1, const TRigidTransform<T, d>& Transform0, const TRigidTransform<T, d>& Transform1, const T Thickness, FCollisionConstraintsArray& NewConstraints)
 		{
 			FRigidBodyPointContactConstraint* Constraint = new FRigidBodyPointContactConstraint;
 
@@ -751,33 +751,24 @@ namespace Chaos
 		DECLARE_CYCLE_STAT(TEXT("TPBDCollisionConstraints::UpdateUnionUnionConstraint"), STAT_UpdateUnionUnionConstraint, STATGROUP_ChaosWide);
 
 		template<ECollisionUpdateType UpdateType, typename T, int d>
-		void UpdateUnionUnionConstraint(const T Thickness, TCollisionConstraintBase<T, d>& Constraint)
+		void UpdateUnionUnionConstraint(const FImplicitObject& Implicit0, const TRigidTransform<T, d>& Transform0, const FImplicitObject& Implicit1, const TRigidTransform<T, d>& Transform1, const T Thickness, TCollisionConstraintBase<T, d>& Constraint)
 		{
 			SCOPE_CYCLE_COUNTER(STAT_UpdateUnionUnionConstraint);
 
-			TGenericParticleHandle<T, d> Particle0 = Constraint.Particle[0];
-			TGenericParticleHandle<T, d> Particle1 = Constraint.Particle[1];
-
-			TRigidTransform<T, d> ParticlesTM = TRigidTransform<T, d>(Particle0->P(), Particle0->Q());
-			TRigidTransform<T, d> LevelsetTM = TRigidTransform<T, d>(Particle1->P(), Particle1->Q());
-
-			const FImplicitObject* ParticleObj = Particle0->Geometry().Get();
-			const FImplicitObject* LevelsetObj = Particle1->Geometry().Get();
-			const TArray<Pair<const FImplicitObject*, TRigidTransform<T, d>>> LevelsetShapes = FindRelevantShapes(ParticleObj, ParticlesTM, *LevelsetObj, LevelsetTM, Thickness);
-
+			const TArray<Pair<const FImplicitObject*, TRigidTransform<T, d>>> LevelsetShapes = FindRelevantShapes(&Implicit0, Transform0, Implicit1, Transform1, Thickness);
 			for (const Pair<const FImplicitObject*, TRigidTransform<T, d>>& LevelsetObjPair : LevelsetShapes)
 			{
 				const FImplicitObject& LevelsetInnerObj = *LevelsetObjPair.First;
-				const TRigidTransform<T, d>& LevelsetInnerObjTM = LevelsetObjPair.Second * LevelsetTM;
+				const TRigidTransform<T, d>& LevelsetInnerObjTM = LevelsetObjPair.Second * Transform1;
 
 				//now find all particle inner objects
-				const TArray<Pair<const FImplicitObject*, TRigidTransform<T, d>>> ParticleShapes = FindRelevantShapes(&LevelsetInnerObj, LevelsetInnerObjTM, *ParticleObj, ParticlesTM, Thickness);
+				const TArray<Pair<const FImplicitObject*, TRigidTransform<T, d>>> ParticleShapes = FindRelevantShapes(&LevelsetInnerObj, LevelsetInnerObjTM, Implicit0, Transform0, Thickness);
 
 				//for each inner obj pair, update constraint
 				for (const Pair<const FImplicitObject*, TRigidTransform<T, d>>& ParticlePair : ParticleShapes)
 				{
 					const FImplicitObject& ParticleInnerObj = *ParticlePair.First;
-					const TRigidTransform<T, d> ParticleInnerObjTM = ParticlePair.Second * ParticlesTM;
+					const TRigidTransform<T, d> ParticleInnerObjTM = ParticlePair.Second * Transform0;
 					UpdateConstraintImp<UpdateType>(ParticleInnerObj, ParticleInnerObjTM, LevelsetInnerObj, LevelsetInnerObjTM, Thickness, Constraint);
 				}
 			}
@@ -786,31 +777,25 @@ namespace Chaos
 
 
 		template<typename T, int d>
-		void ConstructUnionUnionConstraints(TGeometryParticleHandle<T, d>* Particle0, TGeometryParticleHandle<T, d>* Particle1, const FImplicitObject* Implicit0, const FImplicitObject* Implicit1, const T Thickness, FCollisionConstraintsArray& NewConstraints)
+		void ConstructUnionUnionConstraints(TGeometryParticleHandle<T, d>* Particle0, TGeometryParticleHandle<T, d>* Particle1, const FImplicitObject* Implicit0, const FImplicitObject* Implicit1, const TRigidTransform<T, d>& Transform0, const TRigidTransform<T, d>& Transform1, const T Thickness, FCollisionConstraintsArray& NewConstraints)
 		{
 			SCOPE_CYCLE_COUNTER(STAT_UpdateUnionUnionConstraint);
 
-			const FImplicitObject* ParticleObj = Particle0->Geometry().Get();
-			const FImplicitObject* LevelsetObj = Particle1->Geometry().Get();
-
-			TRigidTransform<T, d> ParticlesTM = TRigidTransform<T, d>(Particle0->X(), Particle0->R());
-			TRigidTransform<T, d> LevelsetTM = TRigidTransform<T, d>(Particle1->X(), Particle1->R());
-
-			const TArray<Pair<const FImplicitObject*, TRigidTransform<T, d>>> LevelsetShapes = FindRelevantShapes(ParticleObj, ParticlesTM, *LevelsetObj, LevelsetTM, Thickness);
+			const TArray<Pair<const FImplicitObject*, TRigidTransform<T, d>>> LevelsetShapes = FindRelevantShapes(Implicit0, Transform0, *Implicit1, Transform1, Thickness);
 
 			for (const Pair<const FImplicitObject*, TRigidTransform<T, d>>& LevelsetObjPair : LevelsetShapes)
 			{
 				const FImplicitObject* LevelsetInnerObj = LevelsetObjPair.First;
-				const TRigidTransform<T, d>& LevelsetInnerObjTM = LevelsetObjPair.Second * LevelsetTM;
+				const TRigidTransform<T, d>& LevelsetInnerObjTM = LevelsetObjPair.Second * Transform1;
 
 				//now find all particle inner objects
-				const TArray<Pair<const FImplicitObject*, TRigidTransform<T, d>>> ParticleShapes = FindRelevantShapes(LevelsetInnerObj, LevelsetInnerObjTM, *ParticleObj, ParticlesTM, Thickness);
+				const TArray<Pair<const FImplicitObject*, TRigidTransform<T, d>>> ParticleShapes = FindRelevantShapes(LevelsetInnerObj, LevelsetInnerObjTM, *Implicit0, Transform0, Thickness);
 
 				//for each inner obj pair, update constraint
 				for (const Pair<const FImplicitObject*, TRigidTransform<T, d>>& ParticlePair : ParticleShapes)
 				{
 					const FImplicitObject* ParticleInnerObj = ParticlePair.First;
-					ConstructConstraintsImpl<T, d>(Particle0, Particle1, ParticleInnerObj, LevelsetInnerObj, Thickness, NewConstraints);
+					ConstructConstraintsImpl<T, d>(Particle0, Particle1, ParticleInnerObj, LevelsetInnerObj, Transform0, Transform1, Thickness, NewConstraints);
 				}
 			}
 		}
@@ -821,39 +806,31 @@ namespace Chaos
 
 		DECLARE_CYCLE_STAT(TEXT("TPBDCollisionConstraints::UpdateSingleUnionConstraint"), STAT_UpdateSingleUnionConstraint, STATGROUP_ChaosWide);
 		template<ECollisionUpdateType UpdateType, typename T, int d>
-		void UpdateSingleUnionConstraint(const T Thickness, TCollisionConstraintBase<T, d>& Constraint)
+		void UpdateSingleUnionConstraint(const FImplicitObject& Implicit0, const TRigidTransform<T, d>& Transform0, const FImplicitObject& Implicit1, const TRigidTransform<T, d>& Transform1, const T Thickness, TCollisionConstraintBase<T, d>& Constraint)
 		{
 			SCOPE_CYCLE_COUNTER(STAT_UpdateSingleUnionConstraint);
 
-			TGenericParticleHandle<T, d> Particle0 = Constraint.Particle[0];
-			TGenericParticleHandle<T, d> Particle1 = Constraint.Particle[1];
-
-			TRigidTransform<T, d> ParticlesTM = TRigidTransform<T, d>(Particle0->P(), Particle0->Q());
-			TRigidTransform<T, d> LevelsetTM = TRigidTransform<T, d>(Particle1->P(), Particle1->Q());
-
-			const FImplicitObject* ParticleObj = Particle0->Geometry().Get();
-			const FImplicitObject* LevelsetObj = Particle1->Geometry().Get();
-			const TArray<Pair<const FImplicitObject*, TRigidTransform<T, d>>> LevelsetShapes = FindRelevantShapes(ParticleObj, ParticlesTM, *LevelsetObj, LevelsetTM, Thickness);
+			const TArray<Pair<const FImplicitObject*, TRigidTransform<T, d>>> LevelsetShapes = FindRelevantShapes(&Implicit0, Transform0, Implicit1, Transform1, Thickness);
 
 			for (const Pair<const FImplicitObject*, TRigidTransform<T, d>>& LevelsetObjPair : LevelsetShapes)
 			{
-				const FImplicitObject& LevelsetInnerObj = *LevelsetObjPair.First;
-				const TRigidTransform<T, d> LevelsetInnerObjTM = LevelsetObjPair.Second * LevelsetTM;
-				UpdateConstraintImp<UpdateType>(*ParticleObj, ParticlesTM, LevelsetInnerObj, LevelsetInnerObjTM, Thickness, Constraint);
+				const FImplicitObject& Implicit1InnerObj = *LevelsetObjPair.First;
+				const TRigidTransform<T, d> Implicit1InnerObjTM = LevelsetObjPair.Second * Transform1;
+				UpdateConstraintImp<UpdateType>(Implicit0, Transform0, Implicit1InnerObj, Implicit1InnerObjTM, Thickness, Constraint);
 			}
 		}
 
 
 
 		template<typename T, int d>
-		void ConstructSingleUnionConstraints(TGeometryParticleHandle<T, d>* Particle0, TGeometryParticleHandle<T, d>* Particle1, const FImplicitObject* Implicit0, const FImplicitObject* Implicit1, const T Thickness, FCollisionConstraintsArray& NewConstraints)
+		void ConstructSingleUnionConstraints(TGeometryParticleHandle<T, d>* Particle0, TGeometryParticleHandle<T, d>* Particle1, const FImplicitObject* Implicit0, const FImplicitObject* Implicit1, const TRigidTransform<T, d>& Transform0, const TRigidTransform<T, d>& Transform1, const T Thickness, FCollisionConstraintsArray& NewConstraints)
 		{
 			FRigidBodyPointContactConstraint* Constraint = new FRigidBodyPointContactConstraint;
 			Constraint->Particle[0] = Particle0;
 			Constraint->Particle[1] = Particle1;
 			Constraint->AddManifold(Implicit0, Implicit1);
 
-			UpdateSingleUnionConstraint<ECollisionUpdateType::Any>(Thickness, *Constraint);
+			UpdateSingleUnionConstraint<ECollisionUpdateType::Any>(*Implicit0, Transform0, *Implicit1, Transform1, Thickness, *Constraint);
 
 			if (Constraint->GetPhi() < Thickness)
 			{
@@ -870,101 +847,116 @@ namespace Chaos
 		//
 
 		template<ECollisionUpdateType UpdateType, typename T, int d>
-		void UpdateConstraintImp(const FImplicitObject& ParticleObject, const TRigidTransform<T, d>& ParticleTM, const FImplicitObject& LevelsetObject, const TRigidTransform<T, d>& LevelsetTM, const T Thickness, TCollisionConstraintBase<T, d>& ConstraintBase)
+		void UpdateConstraintImp(const FImplicitObject& Implicit0, const TRigidTransform<T, d>& Transform0, const FImplicitObject& Implicit1, const TRigidTransform<T, d>& Transform1, const T Thickness, TCollisionConstraintBase<T, d>& ConstraintBase)
 		{
-			if (ParticleObject.GetType() == TBox<T, d>::StaticType() && LevelsetObject.GetType() == TBox<T, d>::StaticType())
+			EImplicitObjectType Implicit0Type = GetInnerType(Implicit0.GetType());
+			EImplicitObjectType Implicit1Type = GetInnerType(Implicit1.GetType());
+
+			if (Implicit0Type == TImplicitObjectTransformed<T, d>::StaticType())
+			{
+				const TImplicitObjectTransformed<FReal, 3>* TransformedImplicit0 = Implicit0.template GetObject<const TImplicitObjectTransformed<FReal, 3>>();
+				TRigidTransform<T, d> TransformedTransform0 = TransformedImplicit0->GetTransform() * Transform0;
+				UpdateConstraintImp<UpdateType>(*TransformedImplicit0->GetTransformedObject(), TransformedTransform0, Implicit1, Transform1, Thickness, ConstraintBase);
+			}
+			else if (Implicit1Type == TImplicitObjectTransformed<T, d>::StaticType())
+			{
+				const TImplicitObjectTransformed<FReal, 3>* TransformedImplicit1 = Implicit1.template GetObject<const TImplicitObjectTransformed<FReal, 3>>();
+				TRigidTransform<T, d> TransformedTransform1 = TransformedImplicit1->GetTransform() * Transform1;
+				UpdateConstraintImp<UpdateType>(Implicit0, Transform0, *TransformedImplicit1->GetTransformedObject(), TransformedTransform1, Thickness, ConstraintBase);
+			}
+			else if (Implicit0Type == TBox<T, d>::StaticType() && Implicit1Type == TBox<T, d>::StaticType())
 			{
-				UpdateBoxBoxConstraint(*ParticleObject.template GetObject<TBox<T, d>>(), ParticleTM, *LevelsetObject.template GetObject<TBox<T, d>>(), LevelsetTM, Thickness, *ConstraintBase.template As<TRigidBodyPointContactConstraint<T,d> >());
+				UpdateBoxBoxConstraint(*Implicit0.template GetObject<TBox<T, d>>(), Transform0, *Implicit1.template GetObject<TBox<T, d>>(), Transform1, Thickness, *ConstraintBase.template As<TRigidBodyPointContactConstraint<T,d>>());
 			}
-			else if (ParticleObject.GetType() == TSphere<T, d>::StaticType() && LevelsetObject.GetType() == TSphere<T, d>::StaticType())
+			else if (Implicit0Type == TSphere<T, d>::StaticType() && Implicit1Type == TSphere<T, d>::StaticType())
 			{
-				UpdateSphereSphereConstraint(*ParticleObject.template GetObject<TSphere<T, d>>(), ParticleTM, *LevelsetObject.template GetObject<TSphere<T, d>>(), LevelsetTM, Thickness, *ConstraintBase.template As<TRigidBodyPointContactConstraint<T,d> >());
+				UpdateSphereSphereConstraint(*Implicit0.template GetObject<TSphere<T, d>>(), Transform0, *Implicit1.template GetObject<TSphere<T, d>>(), Transform1, Thickness, *ConstraintBase.template As<TRigidBodyPointContactConstraint<T,d>>());
 			}
-			else if (ParticleObject.GetType() == TBox<T, d>::StaticType() && LevelsetObject.GetType() == TPlane<T, d>::StaticType())
+			else if (Implicit0Type == TBox<T, d>::StaticType() && Implicit1Type == TPlane<T, d>::StaticType())
 			{
-				UpdateBoxPlaneConstraint(*ParticleObject.template GetObject<TBox<T, d>>(), ParticleTM, *LevelsetObject.template GetObject<TPlane<T, d>>(), LevelsetTM, Thickness, *ConstraintBase.template As<TRigidBodyPointContactConstraint<T,d> >());
+				UpdateBoxPlaneConstraint(*Implicit0.template GetObject<TBox<T, d>>(), Transform0, *Implicit1.template GetObject<TPlane<T, d>>(), Transform1, Thickness, *ConstraintBase.template As<TRigidBodyPointContactConstraint<T,d>>());
 			}
-			else if (ParticleObject.GetType() == TSphere<T, d>::StaticType() && LevelsetObject.GetType() == TPlane<T, d>::StaticType())
+			else if (Implicit0Type == TSphere<T, d>::StaticType() && Implicit1Type == TPlane<T, d>::StaticType())
 			{
-				UpdateSpherePlaneConstraint(*ParticleObject.template GetObject<TSphere<T, d>>(), ParticleTM, *LevelsetObject.template GetObject<TPlane<T, d>>(), LevelsetTM, Thickness, *ConstraintBase.template As<TRigidBodyPointContactConstraint<T,d> >());
+				UpdateSpherePlaneConstraint(*Implicit0.template GetObject<TSphere<T, d>>(), Transform0, *Implicit1.template GetObject<TPlane<T, d>>(), Transform1, Thickness, *ConstraintBase.template As<TRigidBodyPointContactConstraint<T,d>>());
 			}
-			else if (ParticleObject.GetType() == TSphere<T, d>::StaticType() && LevelsetObject.GetType() == TBox<T, d>::StaticType())
+			else if (Implicit0Type == TSphere<T, d>::StaticType() && Implicit1Type == TBox<T, d>::StaticType())
 			{
-				UpdateSphereBoxConstraint(*ParticleObject.template GetObject<TSphere<T, d>>(), ParticleTM, *LevelsetObject.template GetObject<TBox<T, d>>(), LevelsetTM, Thickness, *ConstraintBase.template As<TRigidBodyPointContactConstraint<T,d> >());
+				UpdateSphereBoxConstraint(*Implicit0.template GetObject<TSphere<T, d>>(), Transform0, *Implicit1.template GetObject<TBox<T, d>>(), Transform1, Thickness, *ConstraintBase.template As<TRigidBodyPointContactConstraint<T,d>>());
 			}
-			else if (ParticleObject.GetType() == TCapsule<T>::StaticType() && LevelsetObject.GetType() == TCapsule<T>::StaticType())
+			else if (Implicit0Type == TCapsule<T>::StaticType() && Implicit1Type == TCapsule<T>::StaticType())
 			{
-				UpdateCapsuleCapsuleConstraint(*ParticleObject.template GetObject<TCapsule<T>>(), ParticleTM, *LevelsetObject.template GetObject<TCapsule<T>>(), LevelsetTM, Thickness, *ConstraintBase.template As<TRigidBodyPointContactConstraint<T,d> >());
+				UpdateCapsuleCapsuleConstraint(*Implicit0.template GetObject<TCapsule<T>>(), Transform0, *Implicit1.template GetObject<TCapsule<T>>(), Transform1, Thickness, *ConstraintBase.template As<TRigidBodyPointContactConstraint<T,d>>());
 			}
-			else if (ParticleObject.GetType() == TCapsule<T>::StaticType() && LevelsetObject.GetType() == TBox<T, d>::StaticType())
+			else if (Implicit0Type == TCapsule<T>::StaticType() && Implicit1Type == TBox<T, d>::StaticType())
 			{
-				UpdateCapsuleBoxConstraint(*ParticleObject.template GetObject<TCapsule<T>>(), ParticleTM, *LevelsetObject.template GetObject<TBox<T, d>>(), LevelsetTM, Thickness, *ConstraintBase.template As<TRigidBodyPointContactConstraint<T, d> >());
+				UpdateCapsuleBoxConstraint(*Implicit0.template GetObject<TCapsule<T>>(), Transform0, *Implicit1.template GetObject<TBox<T, d>>(), Transform1, Thickness, *ConstraintBase.template As<TRigidBodyPointContactConstraint<T,d>>());
 			}
-			else if (ParticleObject.GetType() == TPlane<T, d>::StaticType() && LevelsetObject.GetType() == TBox<T, d>::StaticType())
+			else if (Implicit0Type == TPlane<T, d>::StaticType() && Implicit1Type == TBox<T, d>::StaticType())
 			{
 				TRigidBodyPointContactConstraint<T, d>& Constraint = *ConstraintBase.template As<TRigidBodyPointContactConstraint<T, d> >();
 				TRigidBodyPointContactConstraint<T, d> TmpConstraint = Constraint;
-				UpdateBoxPlaneConstraint(*LevelsetObject.template GetObject<TBox<T, d>>(), LevelsetTM, *ParticleObject.template GetObject<TPlane<T, d>>(), ParticleTM, Thickness, TmpConstraint);
+				UpdateBoxPlaneConstraint(*Implicit1.template GetObject<TBox<T, d>>(), Transform1, *Implicit0.template GetObject<TPlane<T, d>>(), Transform0, Thickness, TmpConstraint);
 				if (TmpConstraint.GetPhi() < Constraint.GetPhi())
 				{
 					Constraint = TmpConstraint;
 					Constraint.SetNormal(-Constraint.GetNormal());
 				}
 			}
-			else if (ParticleObject.GetType() == TPlane<T, d>::StaticType() && LevelsetObject.GetType() == TSphere<T, d>::StaticType())
+			else if (Implicit0Type == TPlane<T, d>::StaticType() && Implicit1Type == TSphere<T, d>::StaticType())
 			{
 				TRigidBodyPointContactConstraint<T, d>& Constraint = *ConstraintBase.template As<TRigidBodyPointContactConstraint<T, d> >();
 				TRigidBodyPointContactConstraint<T, d> TmpConstraint = Constraint;
-				UpdateSpherePlaneConstraint(*LevelsetObject.template GetObject<TSphere<T, d>>(), LevelsetTM, *ParticleObject.template GetObject<TPlane<T, d>>(), ParticleTM, Thickness, TmpConstraint);
+				UpdateSpherePlaneConstraint(*Implicit1.template GetObject<TSphere<T, d>>(), Transform1, *Implicit0.template GetObject<TPlane<T, d>>(), Transform0, Thickness, TmpConstraint);
 				if (TmpConstraint.GetPhi() < Constraint.GetPhi())
 				{
 					Constraint = TmpConstraint;
 					Constraint.SetNormal(-Constraint.GetNormal());
 				}
 			}
-			else if (ParticleObject.GetType() == TBox<T, d>::StaticType() && LevelsetObject.GetType() == TSphere<T, d>::StaticType())
+			else if (Implicit0Type == TBox<T, d>::StaticType() && Implicit1Type == TSphere<T, d>::StaticType())
 			{
 				TRigidBodyPointContactConstraint<T, d>& Constraint = *ConstraintBase.template As<TRigidBodyPointContactConstraint<T, d> >();
 				TRigidBodyPointContactConstraint<T, d> TmpConstraint = Constraint;
-				UpdateSphereBoxConstraint(*LevelsetObject.template GetObject<TSphere<T, d>>(), LevelsetTM, *ParticleObject.template GetObject<TBox<T, d>>(), ParticleTM, Thickness, TmpConstraint);
+				UpdateSphereBoxConstraint(*Implicit1.template GetObject<TSphere<T, d>>(), Transform1, *Implicit0.template GetObject<TBox<T, d>>(), Transform0, Thickness, TmpConstraint);
 				if (TmpConstraint.GetPhi() < Constraint.GetPhi())
 				{
 					Constraint = TmpConstraint;
 					Constraint.SetNormal(-Constraint.GetNormal());
 				}
 			}
-			else if (ParticleObject.GetType() == TBox<T, d>::StaticType() && LevelsetObject.GetType() == TCapsule<T>::StaticType())
+			else if (Implicit0Type == TBox<T, d>::StaticType() && Implicit1Type == TCapsule<T>::StaticType())
 			{
 				TRigidBodyPointContactConstraint<T, d>& Constraint = *ConstraintBase.template As<TRigidBodyPointContactConstraint<T, d> >();
 				TRigidBodyPointContactConstraint<T, d> TmpConstraint = Constraint;
-				UpdateCapsuleBoxConstraint(*LevelsetObject.template GetObject<TCapsule<T>>(), LevelsetTM, *ParticleObject.template GetObject<TBox<T, d>>(), ParticleTM, Thickness, TmpConstraint);
+				UpdateCapsuleBoxConstraint(*Implicit1.template GetObject<TCapsule<T>>(), Transform1, *Implicit0.template GetObject<TBox<T, d>>(), Transform0, Thickness, TmpConstraint);
 				if (TmpConstraint.GetPhi() < Constraint.GetPhi())
 				{
 					Constraint = TmpConstraint;
 					Constraint.SetNormal(-Constraint.GetNormal());
 				}
 			}
-			else if (ParticleObject.GetType() < TImplicitObjectUnion<T, d>::StaticType() && LevelsetObject.GetType() == TImplicitObjectUnion<T, d>::StaticType())
+			else if (Implicit0Type < TImplicitObjectUnion<T, d>::StaticType() && Implicit1Type == TImplicitObjectUnion<T, d>::StaticType())
 			{
-				return UpdateSingleUnionConstraint<UpdateType>(Thickness, ConstraintBase);
+				return UpdateSingleUnionConstraint<UpdateType>(Implicit0, Transform0, Implicit1, Transform1, Thickness, ConstraintBase);
 			}
-			else if (ParticleObject.GetType() == TImplicitObjectUnion<T, d>::StaticType() && LevelsetObject.GetType() < TImplicitObjectUnion<T, d>::StaticType())
+			else if (Implicit0Type == TImplicitObjectUnion<T, d>::StaticType() && Implicit1Type < TImplicitObjectUnion<T, d>::StaticType())
 			{
 				check(false);	//should not be possible to get this ordering (see ComputeConstraint)
 			}
-			else if (ParticleObject.GetType() == TImplicitObjectUnion<T, d>::StaticType() && LevelsetObject.GetType() == TImplicitObjectUnion<T, d>::StaticType())
+			else if (Implicit0Type == TImplicitObjectUnion<T, d>::StaticType() && Implicit1Type == TImplicitObjectUnion<T, d>::StaticType())
 			{
-				return UpdateUnionUnionConstraint<UpdateType>(Thickness, ConstraintBase);
+				return UpdateUnionUnionConstraint<UpdateType>(Implicit0, Transform0, Implicit1, Transform1, Thickness, ConstraintBase);
 			}
-			else if (ParticleObject.IsConvex() && LevelsetObject.IsConvex())
+			else if (Implicit0.IsConvex() && Implicit1.IsConvex())
 			{
-				UpdateConvexConvexConstraint(ParticleObject, ParticleTM, LevelsetObject, LevelsetTM, Thickness, *ConstraintBase.template As<TRigidBodyPointContactConstraint<T,d> >());
+				UpdateConvexConvexConstraint(Implicit0, Transform0, Implicit1, Transform1, Thickness, *ConstraintBase.template As<TRigidBodyPointContactConstraint<T, d>>());
 			}
-			else if (LevelsetObject.IsUnderlyingUnion())
+			else if (Implicit1.IsUnderlyingUnion())
 			{
 				UpdateUnionLevelsetConstraint<UpdateType>(Thickness, *ConstraintBase.template As<TRigidBodyPointContactConstraint<T,d> >());
 			}
-			else if (ParticleObject.IsUnderlyingUnion())
+			else if (Implicit0.IsUnderlyingUnion())
 			{
 				UpdateLevelsetUnionConstraint<UpdateType>(Thickness, *ConstraintBase.template As<TRigidBodyPointContactConstraint<T,d> >());
 			}
@@ -974,10 +966,8 @@ namespace Chaos
 			}
 		}
 
-
-
 		template<typename T, int d>
-		void ConstructPairConstraintImpl(TGeometryParticleHandle<T, d>* Particle0, TGeometryParticleHandle<T, d>* Particle1, const FImplicitObject* Implicit0, const FImplicitObject* Implicit1, const T Thickness, FCollisionConstraintsArray& NewConstraints)
+		void ConstructPairConstraintImpl(TGeometryParticleHandle<T, d>* Particle0, TGeometryParticleHandle<T, d>* Particle1, const FImplicitObject* Implicit0, const FImplicitObject* Implicit1, const TRigidTransform<T, d>& Transform0, const TRigidTransform<T, d>& Transform1, const T Thickness, FCollisionConstraintsArray& NewConstraints)
 		{
 			// See if we already have a constraint for this shape pair
 			// todo(brice) : Not sure this actually prevents duplicate constraints.
@@ -1001,77 +991,93 @@ namespace Chaos
 
 			if (!Implicit0 || !Implicit1)
 			{
-				ConstructLevelsetLevelsetConstraints(Particle0, Particle1, Implicit0, Implicit1, Thickness, NewConstraints);
+				ConstructLevelsetLevelsetConstraints(Particle0, Particle1, Implicit0, Implicit1, Transform0, Transform1, Thickness, NewConstraints);
+				return;
+			}
+			
+			EImplicitObjectType Implicit0Type = GetInnerType(Implicit0->GetType());
+			EImplicitObjectType Implicit1Type = GetInnerType(Implicit1->GetType());
+			if (Implicit0Type == TImplicitObjectTransformed<T, d>::StaticType())
+			{
+				const TImplicitObjectTransformed<FReal, 3>* TransformedImplicit0 = Implicit0->template GetObject<const TImplicitObjectTransformed<FReal, 3>>();
+				TRigidTransform<T, d> TransformedTransform0 = TransformedImplicit0->GetTransform() * Transform0;
+				ConstructPairConstraintImpl(Particle0, Particle1, TransformedImplicit0->GetTransformedObject(), Implicit1, TransformedTransform0, Transform1, Thickness, NewConstraints);
+			}
+			else if (Implicit1Type == TImplicitObjectTransformed<T, d>::StaticType())
+			{
+				const TImplicitObjectTransformed<FReal, 3>* TransformedImplicit1 = Implicit1->template GetObject<const TImplicitObjectTransformed<FReal, 3>>();
+				TRigidTransform<T, d> TransformedTransform1 = TransformedImplicit1->GetTransform() * Transform1;
+				ConstructPairConstraintImpl(Particle0, Particle1, Implicit0, TransformedImplicit1->GetTransformedObject(), Transform0, TransformedTransform1, Thickness, NewConstraints);
 			}
-			if (Implicit0->GetType() == TBox<T, d>::StaticType() && Implicit1->GetType() == TBox<T, d>::StaticType())
+			else if (Implicit0Type == TBox<T, d>::StaticType() && Implicit1Type == TBox<T, d>::StaticType())
 			{
-				ConstructBoxBoxConstraints(Particle0, Particle1, Implicit0, Implicit1, Thickness, NewConstraints);
+				ConstructBoxBoxConstraints(Particle0, Particle1, Implicit0, Implicit1, Transform0, Transform1, Thickness, NewConstraints);
 			}
-			else if (Implicit0->GetType() == TSphere<T, d>::StaticType() && Implicit1->GetType() == TSphere<T, d>::StaticType())
+			else if (Implicit0Type == TSphere<T, d>::StaticType() && Implicit1Type == TSphere<T, d>::StaticType())
 			{
-				ConstructSphereSphereConstraints(Particle0, Particle1, Implicit0, Implicit1, Thickness, NewConstraints);
+				ConstructSphereSphereConstraints(Particle0, Particle1, Implicit0, Implicit1, Transform0, Transform1, Thickness, NewConstraints);
 			}
-			else if (Implicit0->GetType() == TBox<T, d>::StaticType() && Implicit1->GetType() == TPlane<T, d>::StaticType())
+			else if (Implicit0Type == TBox<T, d>::StaticType() && Implicit1Type == TPlane<T, d>::StaticType())
 			{
-				ConstructBoxPlaneConstraints(Particle0, Particle1, Implicit0, Implicit1, Thickness, NewConstraints);
+				ConstructBoxPlaneConstraints(Particle0, Particle1, Implicit0, Implicit1, Transform0, Transform1, Thickness, NewConstraints);
 			}
-			else if (Implicit0->GetType() == TPlane<T, d>::StaticType() && Implicit1->GetType() == TBox<T, d>::StaticType())
+			else if (Implicit0Type == TPlane<T, d>::StaticType() && Implicit1Type == TBox<T, d>::StaticType())
 			{
-				ConstructBoxPlaneConstraints(Particle1, Particle0, Implicit1, Implicit0, Thickness, NewConstraints);
+				ConstructBoxPlaneConstraints(Particle1, Particle0, Implicit1, Implicit0, Transform1, Transform0, Thickness, NewConstraints);
 			}
-			else if (Implicit0->GetType() == TSphere<T, d>::StaticType() && Implicit1->GetType() == TPlane<T, d>::StaticType())
+			else if (Implicit0Type == TSphere<T, d>::StaticType() && Implicit1Type == TPlane<T, d>::StaticType())
 			{
-				ConstructSpherePlaneConstraints(Particle0, Particle1, Implicit0, Implicit1, Thickness, NewConstraints);
+				ConstructSpherePlaneConstraints(Particle0, Particle1, Implicit0, Implicit1, Transform0, Transform1, Thickness, NewConstraints);
 			}
-			else if (Implicit0->GetType() == TPlane<T, d>::StaticType() && Implicit1->GetType() == TSphere<T, d>::StaticType())
+			else if (Implicit0Type == TPlane<T, d>::StaticType() && Implicit1Type == TSphere<T, d>::StaticType())
 			{
-				ConstructSpherePlaneConstraints(Particle1, Particle0, Implicit1, Implicit0, Thickness, NewConstraints);
+				ConstructSpherePlaneConstraints(Particle1, Particle0, Implicit1, Implicit0, Transform1, Transform0, Thickness, NewConstraints);
 			}
-			else if (Implicit0->GetType() == TSphere<T, d>::StaticType() && Implicit1->GetType() == TBox<T, d>::StaticType())
+			else if (Implicit0Type == TSphere<T, d>::StaticType() && Implicit1Type == TBox<T, d>::StaticType())
 			{
-				ConstructSphereBoxConstraints(Particle0, Particle1, Implicit0, Implicit1, Thickness, NewConstraints);
+				ConstructSphereBoxConstraints(Particle0, Particle1, Implicit0, Implicit1, Transform0, Transform1, Thickness, NewConstraints);
 			}
-			else if (Implicit0->GetType() == TBox<T, d>::StaticType() && Implicit1->GetType() == TSphere<T, d>::StaticType())
+			else if (Implicit0Type == TBox<T, d>::StaticType() && Implicit1Type == TSphere<T, d>::StaticType())
 			{
-				ConstructSphereBoxConstraints(Particle1, Particle0, Implicit1, Implicit0, Thickness, NewConstraints);
+				ConstructSphereBoxConstraints(Particle1, Particle0, Implicit1, Implicit0, Transform1, Transform0, Thickness, NewConstraints);
 			}
-			else if (Implicit0->GetType() == TCapsule<T>::StaticType() && Implicit1->GetType() == TCapsule<T>::StaticType())
+			else if (Implicit0Type == TCapsule<T>::StaticType() && Implicit1Type == TCapsule<T>::StaticType())
 			{
-				ConstructCapsuleCapsuleConstraints(Particle0, Particle1, Implicit0, Implicit1, Thickness, NewConstraints);
+				ConstructCapsuleCapsuleConstraints(Particle0, Particle1, Implicit0, Implicit1, Transform1, Transform0, Thickness, NewConstraints);
 			}
-			else if (Implicit0->GetType() == TCapsule<T>::StaticType() && Implicit1->GetType() == TBox<T, d>::StaticType())
+			else if (Implicit0Type == TCapsule<T>::StaticType() && Implicit1Type == TBox<T, d>::StaticType())
 			{
-				ConstructCapsuleBoxConstraints(Particle0, Particle1, Implicit0, Implicit1, Thickness, NewConstraints);
+				ConstructCapsuleBoxConstraints(Particle0, Particle1, Implicit0, Implicit1, Transform0, Transform1, Thickness, NewConstraints);
 			}
-			else if (Implicit0->GetType() == TBox<T,d>::StaticType() && Implicit1->GetType() == TCapsule<T>::StaticType())
+			else if (Implicit0Type == TBox<T,d>::StaticType() && Implicit1Type == TCapsule<T>::StaticType())
 			{
-				ConstructCapsuleBoxConstraints(Particle1, Particle0, Implicit1, Implicit0, Thickness, NewConstraints);
+				ConstructCapsuleBoxConstraints(Particle1, Particle0, Implicit1, Implicit0, Transform1, Transform0, Thickness, NewConstraints);
 			}
-			else if (Implicit0->GetType() < TImplicitObjectUnion<T, d>::StaticType() && Implicit1->GetType() == TImplicitObjectUnion<T, d>::StaticType())
+			else if (Implicit0Type < TImplicitObjectUnion<T, d>::StaticType() && Implicit1Type == TImplicitObjectUnion<T, d>::StaticType())
 			{
-				ConstructSingleUnionConstraints(Particle0, Particle1, Implicit0, Implicit1, Thickness, NewConstraints);
+				ConstructSingleUnionConstraints(Particle0, Particle1, Implicit0, Implicit1, Transform0, Transform1, Thickness, NewConstraints);
 			}
-			else if (Implicit0->GetType() == TImplicitObjectUnion<T, d>::StaticType() && Implicit1->GetType() < TImplicitObjectUnion<T, d>::StaticType())
+			else if (Implicit0Type == TImplicitObjectUnion<T, d>::StaticType() && Implicit1Type < TImplicitObjectUnion<T, d>::StaticType())
 			{
-				ConstructSingleUnionConstraints(Particle1, Particle0, Implicit1, Implicit0, Thickness, NewConstraints);
+				ConstructSingleUnionConstraints(Particle1, Particle0, Implicit1, Implicit0, Transform1, Transform0, Thickness, NewConstraints);
 			}
-			else if (Implicit0->GetType() == TImplicitObjectUnion<T, d>::StaticType() && Implicit1->GetType() == TImplicitObjectUnion<T, d>::StaticType())
+			else if (Implicit0Type == TImplicitObjectUnion<T, d>::StaticType() && Implicit1Type == TImplicitObjectUnion<T, d>::StaticType())
 			{
 				// Union-union creates multiple manifolds - we should never get here for this pair type. See ConstructConstraintsImpl and ConstructUnionUnionConstraints
 				ensure(false);
 			}
 			else if (Implicit0->IsConvex() && Implicit1->IsConvex())
 			{
-				ConstructConvexConvexConstraints(Particle0, Particle1, Implicit0, Implicit1, Thickness, NewConstraints);
+				ConstructConvexConvexConstraints(Particle0, Particle1, Implicit0, Implicit1, Transform0, Transform1, Thickness, NewConstraints);
 			}
 			else
 			{
-				ConstructLevelsetLevelsetConstraints(Particle0, Particle1, Implicit0, Implicit1, Thickness, NewConstraints);
+				ConstructLevelsetLevelsetConstraints(Particle0, Particle1, Implicit0, Implicit1, Transform0, Transform1, Thickness, NewConstraints);
 			}
 		}
 
 		template<typename T, int d>
-		void ConstructConstraintsImpl(TGeometryParticleHandle<T, d>* Particle0, TGeometryParticleHandle<T, d>* Particle1, const FImplicitObject* Implicit0, const FImplicitObject* Implicit1, const T Thickness, FCollisionConstraintsArray& NewConstraints)
+		void ConstructConstraintsImpl(TGeometryParticleHandle<T, d>* Particle0, TGeometryParticleHandle<T, d>* Particle1, const FImplicitObject* Implicit0, const FImplicitObject* Implicit1, const TRigidTransform<T, d>& Transform0, const TRigidTransform<T, d>& Transform1, const T Thickness, FCollisionConstraintsArray& NewConstraints)
 		{
 			// TriangleMesh implicits are for scene query only.
 			if (Implicit0 && GetInnerType(Implicit0->GetType()) == ImplicitObjectType::TriangleMesh) return;
@@ -1080,43 +1086,43 @@ namespace Chaos
 			if (Implicit0 && Implicit0->GetType() == TImplicitObjectUnion<T, d>::StaticType() &&
 				Implicit1 && Implicit1->GetType() == TImplicitObjectUnion<T, d>::StaticType())
 			{
-				return ConstructUnionUnionConstraints(Particle0, Particle1, Implicit0, Implicit1, Thickness, NewConstraints);
+				return ConstructUnionUnionConstraints(Particle0, Particle1, Implicit0, Implicit1, Transform0, Transform1, Thickness, NewConstraints);
 			}
 			else
 			{
-				return ConstructPairConstraintImpl(Particle0, Particle1, Implicit0, Implicit1, Thickness, NewConstraints);
+				return ConstructPairConstraintImpl(Particle0, Particle1, Implicit0, Implicit1, Transform0, Transform1, Thickness, NewConstraints);
 			}
 		}
 
 
 
-		template void UpdateBoxBoxConstraint<float, 3>(const TBox<float, 3>& Box1, const TRigidTransform<float, 3>& Box1Transform, const TBox<float, 3>& Box2, const TRigidTransform<float, 3>& Box2Transform, float Thickness, TRigidBodyPointContactConstraint<float, 3>& Constraint);
-		template void ConstructBoxBoxConstraints<float, 3>(TGeometryParticleHandle<float, 3>* Particle0, TGeometryParticleHandle<float, 3>* Particle1, const FImplicitObject* Implicit0, const FImplicitObject* Implicit1, const float Thickness, FCollisionConstraintsArray& NewConstraints);
+		template void UpdateBoxBoxConstraint<float, 3>(const TBox<float, 3>& Box1, const TRigidTransform<float, 3>& Box1Transform, const TBox<float, 3>& Box2, const TRigidTransform<float, 3>& Box2Transform, const float Thickness, TRigidBodyPointContactConstraint<float, 3>& Constraint);
+		template void ConstructBoxBoxConstraints<float, 3>(TGeometryParticleHandle<float, 3>* Particle0, TGeometryParticleHandle<float, 3>* Particle1, const FImplicitObject* Implicit0, const FImplicitObject* Implicit1, const TRigidTransform<float, 3>& Transform0, const TRigidTransform<float, 3>& Transform1, const float Thickness, FCollisionConstraintsArray& NewConstraints);
 
-		template bool UpdateBoxPlaneConstraint<float, 3>(const TBox<float, 3>& Box, const TRigidTransform<float, 3>& BoxTransform, const TPlane<float, 3>& Plane, const TRigidTransform<float, 3>& PlaneTransform, float Thickness, TRigidBodyPointContactConstraint<float, 3>& Constraint);
-		template void ConstructBoxPlaneConstraints<float, 3>(TGeometryParticleHandle<float, 3>* Particle0, TGeometryParticleHandle<float, 3>* Particle1, const FImplicitObject* Implicit0, const FImplicitObject* Implicit1, const float Thickness, FCollisionConstraintsArray& NewConstraints);
+		template bool UpdateBoxPlaneConstraint<float, 3>(const TBox<float, 3>& Box, const TRigidTransform<float, 3>& BoxTransform, const TPlane<float, 3>& Plane, const TRigidTransform<float, 3>& PlaneTransform, const float Thickness, TRigidBodyPointContactConstraint<float, 3>& Constraint);
+		template void ConstructBoxPlaneConstraints<float, 3>(TGeometryParticleHandle<float, 3>* Particle0, TGeometryParticleHandle<float, 3>* Particle1, const FImplicitObject* Implicit0, const FImplicitObject* Implicit1, const TRigidTransform<float, 3>& Transform0, const TRigidTransform<float, 3>& Transform1, const float Thickness, FCollisionConstraintsArray& NewConstraints);
 
-		template void UpdateSphereSphereConstraint<float, 3>(const TSphere<float, 3>& Sphere1, const TRigidTransform<float, 3>& Sphere1Transform, const TSphere<float, 3>& Sphere2, const TRigidTransform<float, 3>& Sphere2Transform, float Thickness, TRigidBodyPointContactConstraint<float, 3>& Constraint);
-		template void ConstructSphereSphereConstraints<float, 3>(TGeometryParticleHandle<float, 3>* Particle0, TGeometryParticleHandle<float, 3>* Particle1, const FImplicitObject* Implicit0, const FImplicitObject* Implicit1, const float Thickness, FCollisionConstraintsArray& NewConstraints);
+		template void UpdateSphereSphereConstraint<float, 3>(const TSphere<float, 3>& Sphere1, const TRigidTransform<float, 3>& Sphere1Transform, const TSphere<float, 3>& Sphere2, const TRigidTransform<float, 3>& Sphere2Transform, const float Thickness, TRigidBodyPointContactConstraint<float, 3>& Constraint);
+		template void ConstructSphereSphereConstraints<float, 3>(TGeometryParticleHandle<float, 3>* Particle0, TGeometryParticleHandle<float, 3>* Particle1, const FImplicitObject* Implicit0, const FImplicitObject* Implicit1, const TRigidTransform<float, 3>& Transform0, const TRigidTransform<float, 3>& Transform1, const float Thickness, FCollisionConstraintsArray& NewConstraints);
 
-		template void UpdateSpherePlaneConstraint<float, 3>(const TSphere<float, 3>& Sphere, const TRigidTransform<float, 3>& SphereTransform, const TPlane<float, 3>& Plane, const TRigidTransform<float, 3>& PlaneTransform, float Thickness, TRigidBodyPointContactConstraint<float, 3>& Constraint);
-		template void ConstructSpherePlaneConstraints<float, 3>(TGeometryParticleHandle<float, 3>* Particle0, TGeometryParticleHandle<float, 3>* Particle1, const FImplicitObject* Implicit0, const FImplicitObject* Implicit1, const float Thickness, FCollisionConstraintsArray& NewConstraints);
+		template void UpdateSpherePlaneConstraint<float, 3>(const TSphere<float, 3>& Sphere, const TRigidTransform<float, 3>& SphereTransform, const TPlane<float, 3>& Plane, const TRigidTransform<float, 3>& PlaneTransform, const float Thickness, TRigidBodyPointContactConstraint<float, 3>& Constraint);
+		template void ConstructSpherePlaneConstraints<float, 3>(TGeometryParticleHandle<float, 3>* Particle0, TGeometryParticleHandle<float, 3>* Particle1, const FImplicitObject* Implicit0, const FImplicitObject* Implicit1, const TRigidTransform<float, 3>& Transform0, const TRigidTransform<float, 3>& Transform1, const float Thickness, FCollisionConstraintsArray& NewConstraints);
 
-		template void UpdateSphereBoxConstraint<float, 3>(const TSphere<float, 3>& Sphere, const TRigidTransform<float, 3>& SphereTransform, const TBox<float, 3>& Box, const TRigidTransform<float, 3>& BoxTransform, float Thickness, TRigidBodyPointContactConstraint<float, 3>& Constraint);
-		template void ConstructSphereBoxConstraints<float, 3>(TGeometryParticleHandle<float, 3>* Particle0, TGeometryParticleHandle<float, 3>* Particle1, const FImplicitObject* Implicit0, const FImplicitObject* Implicit1, const float Thickness, FCollisionConstraintsArray& NewConstraints);
+		template void UpdateSphereBoxConstraint<float, 3>(const TSphere<float, 3>& Sphere, const TRigidTransform<float, 3>& SphereTransform, const TBox<float, 3>& Box, const TRigidTransform<float, 3>& BoxTransform, const float Thickness, TRigidBodyPointContactConstraint<float, 3>& Constraint);
+		template void ConstructSphereBoxConstraints<float, 3>(TGeometryParticleHandle<float, 3>* Particle0, TGeometryParticleHandle<float, 3>* Particle1, const FImplicitObject* Implicit0, const FImplicitObject* Implicit1, const TRigidTransform<float, 3>& Transform0, const TRigidTransform<float, 3>& Transform1, const float Thickness, FCollisionConstraintsArray& NewConstraints);
 
-		template void UpdateCapsuleCapsuleConstraint<float, 3>(const TCapsule<float>& A, const TRigidTransform<float, 3>& ATransform, const TCapsule<float>& B, const TRigidTransform<float, 3>& BTransform, float Thickness, TRigidBodyPointContactConstraint<float, 3>& Constraint);
-		template void ConstructCapsuleCapsuleConstraints(TGeometryParticleHandle<float, 3>* Particle0, TGeometryParticleHandle<float, 3>* Particle1, const FImplicitObject* Implicit0, const FImplicitObject* Implicit1, const float Thickness, FCollisionConstraintsArray& NewConstraints);
+		template void UpdateCapsuleCapsuleConstraint<float, 3>(const TCapsule<float>& A, const TRigidTransform<float, 3>& ATransform, const TCapsule<float>& B, const TRigidTransform<float, 3>& BTransform, const float Thickness, TRigidBodyPointContactConstraint<float, 3>& Constraint);
+		template void ConstructCapsuleCapsuleConstraints(TGeometryParticleHandle<float, 3>* Particle0, TGeometryParticleHandle<float, 3>* Particle1, const FImplicitObject* Implicit0, const FImplicitObject* Implicit1, const TRigidTransform<float, 3>& Transform0, const TRigidTransform<float, 3>& Transform1, const float Thickness, FCollisionConstraintsArray& NewConstraints);
 
-		template void UpdateCapsuleBoxConstraint<float, 3>(const TCapsule<float>& A, const TRigidTransform<float, 3>& ATransform, const TBox<float, 3>& B, const TRigidTransform<float, 3>& BTransform, float Thickness, TRigidBodyPointContactConstraint<float, 3>& Constraint);
-		template void ConstructCapsuleBoxConstraints<float, 3>(TGeometryParticleHandle<float, 3>* Particle0, TGeometryParticleHandle<float, 3>* Particle1, const FImplicitObject* Implicit0, const FImplicitObject* Implicit1, const float Thickness, FCollisionConstraintsArray& NewConstraints);
+		template void UpdateCapsuleBoxConstraint<float, 3>(const TCapsule<float>& A, const TRigidTransform<float, 3>& ATransform, const TBox<float, 3>& B, const TRigidTransform<float, 3>& BTransform, const float Thickness, TRigidBodyPointContactConstraint<float, 3>& Constraint);
+		template void ConstructCapsuleBoxConstraints<float, 3>(TGeometryParticleHandle<float, 3>* Particle0, TGeometryParticleHandle<float, 3>* Particle1, const FImplicitObject* Implicit0, const FImplicitObject* Implicit1, const TRigidTransform<float, 3>& Transform0, const TRigidTransform<float, 3>& Transform1, const float Thickness, FCollisionConstraintsArray& NewConstraints);
 
-		template void UpdateConvexConvexConstraint<float, 3>(const FImplicitObject& A, const TRigidTransform<float, 3>& ATM, const FImplicitObject& B, const TRigidTransform<float, 3>& BTM, float Thickness, TRigidBodyPointContactConstraint<float, 3>& Constraint);
-		template void ConstructConvexConvexConstraints<float, 3>(TGeometryParticleHandle<float, 3>* Particle0, TGeometryParticleHandle<float, 3>* Particle1, const FImplicitObject* Implicit0, const FImplicitObject* Implicit1, const float Thickness, FCollisionConstraintsArray& NewConstraints);
+		template void UpdateConvexConvexConstraint<float, 3>(const FImplicitObject& A, const TRigidTransform<float, 3>& ATM, const FImplicitObject& B, const TRigidTransform<float, 3>& BTM, const float Thickness, TRigidBodyPointContactConstraint<float, 3>& Constraint);
+		template void ConstructConvexConvexConstraints<float, 3>(TGeometryParticleHandle<float, 3>* Particle0, TGeometryParticleHandle<float, 3>* Particle1, const FImplicitObject* Implicit0, const FImplicitObject* Implicit1, const TRigidTransform<float, 3>& Transform0, const TRigidTransform<float, 3>& Transform1, const float Thickness, FCollisionConstraintsArray& NewConstraints);
 
 		template void UpdateLevelsetLevelsetConstraint<ECollisionUpdateType::Any, float, 3>(const float Thickness, TRigidBodyPointContactConstraint<float, 3>& Constraint);
 		template void UpdateLevelsetLevelsetConstraint<ECollisionUpdateType::Deepest, float, 3>(const float Thickness, TRigidBodyPointContactConstraint<float, 3>& Constraint);
-		template void ConstructLevelsetLevelsetConstraints<float, 3>(TGeometryParticleHandle<float, 3>* Particle0, TGeometryParticleHandle<float, 3>* Particle1, const FImplicitObject* Implicit0, const FImplicitObject* Implicit1, const float Thickness, FCollisionConstraintsArray& NewConstraints);
+		template void ConstructLevelsetLevelsetConstraints<float, 3>(TGeometryParticleHandle<float, 3>* Particle0, TGeometryParticleHandle<float, 3>* Particle1, const FImplicitObject* Implicit0, const FImplicitObject* Implicit1, const TRigidTransform<float, 3>& Transform0, const TRigidTransform<float, 3>& Transform1, const float Thickness, FCollisionConstraintsArray& NewConstraints);
 
 		template void UpdateLevelsetUnionConstraint<ECollisionUpdateType::Any, float, 3>(const float Thickness, TRigidBodyPointContactConstraint<float, 3>& Constraint);
 		template void UpdateLevelsetUnionConstraint<ECollisionUpdateType::Deepest, float, 3>(const float Thickness, TRigidBodyPointContactConstraint<float, 3>& Constraint);
@@ -1124,19 +1130,19 @@ namespace Chaos
 		template void UpdateUnionLevelsetConstraint<ECollisionUpdateType::Any, float, 3>(const float Thickness, TRigidBodyPointContactConstraint<float, 3>& Constraint);
 		template void UpdateUnionLevelsetConstraint<ECollisionUpdateType::Deepest, float, 3>(const float Thickness, TRigidBodyPointContactConstraint<float, 3>& Constraint);
 
-		template void UpdateUnionUnionConstraint< ECollisionUpdateType::Any, float, 3>(const float Thickness, TCollisionConstraintBase<float, 3>& Constraint);
-		template void UpdateUnionUnionConstraint< ECollisionUpdateType::Deepest, float, 3>(const float Thickness, TCollisionConstraintBase<float, 3>& Constraint);
-		template void ConstructUnionUnionConstraints<float, 3>(TGeometryParticleHandle<float, 3>* Particle0, TGeometryParticleHandle<float, 3>* Particle1, const FImplicitObject* Implicit0, const FImplicitObject* Implicit1, const float Thickness, FCollisionConstraintsArray& NewConstraints);
+		template void UpdateUnionUnionConstraint< ECollisionUpdateType::Any, float, 3>(const FImplicitObject& Implicit0, const TRigidTransform<float, 3>& Transform0, const FImplicitObject& Implicit1, const TRigidTransform<float, 3>& Transform1, const float Thickness, TCollisionConstraintBase<float, 3>& Constraint);
+		template void UpdateUnionUnionConstraint< ECollisionUpdateType::Deepest, float, 3>(const FImplicitObject& Implicit0, const TRigidTransform<float, 3>& Transform0, const FImplicitObject& Implicit1, const TRigidTransform<float, 3>& Transform1, const float Thickness, TCollisionConstraintBase<float, 3>& Constraint);
+		template void ConstructUnionUnionConstraints<float, 3>(TGeometryParticleHandle<float, 3>* Particle0, TGeometryParticleHandle<float, 3>* Particle1, const FImplicitObject* Implicit0, const FImplicitObject* Implicit1, const TRigidTransform<float, 3>& Transform0, const TRigidTransform<float, 3>& Transform1, const float Thickness, FCollisionConstraintsArray& NewConstraints);
 
-		template void UpdateSingleUnionConstraint< ECollisionUpdateType::Any, float, 3>(const float Thickness, TCollisionConstraintBase<float, 3>& Constraint);
-		template void UpdateSingleUnionConstraint< ECollisionUpdateType::Deepest, float, 3>(const float Thickness, TCollisionConstraintBase<float, 3>& Constraint);
-		template void ConstructSingleUnionConstraints<float, 3>(TGeometryParticleHandle<float, 3>* Particle0, TGeometryParticleHandle<float, 3>* Particle1, const FImplicitObject* Implicit0, const FImplicitObject* Implicit1, const float Thickness, FCollisionConstraintsArray& NewConstraints);
+		template void UpdateSingleUnionConstraint< ECollisionUpdateType::Any, float, 3>(const FImplicitObject& Implicit0, const TRigidTransform<float, 3>& Transform0, const FImplicitObject& Implicit1, const TRigidTransform<float, 3>& Transform1, const float Thickness, TCollisionConstraintBase<float, 3>& Constraint);
+		template void UpdateSingleUnionConstraint< ECollisionUpdateType::Deepest, float, 3>(const FImplicitObject& Implicit0, const TRigidTransform<float, 3>& Transform0, const FImplicitObject& Implicit1, const TRigidTransform<float, 3>& Transform1, const float Thickness, TCollisionConstraintBase<float, 3>& Constraint);
+		template void ConstructSingleUnionConstraints<float, 3>(TGeometryParticleHandle<float, 3>* Particle0, TGeometryParticleHandle<float, 3>* Particle1, const FImplicitObject* Implicit0, const FImplicitObject* Implicit1, const TRigidTransform<float, 3>& Transform0, const TRigidTransform<float, 3>& Transform1, const float Thickness, FCollisionConstraintsArray& NewConstraints);
 
-		template void UpdateConstraintImp<ECollisionUpdateType::Any, float, 3>(const FImplicitObject& ParticleObject, const TRigidTransform<float, 3>& ParticleTM, const FImplicitObject& LevelsetObject, const TRigidTransform<float, 3>& LevelsetTM, const float Thickness, TCollisionConstraintBase<float, 3>& Constraints);
-		template void UpdateConstraintImp<ECollisionUpdateType::Deepest, float, 3>(const FImplicitObject& ParticleObject, const TRigidTransform<float, 3>& ParticleTM, const FImplicitObject& LevelsetObject, const TRigidTransform<float, 3>& LevelsetTM, const float Thickness, TCollisionConstraintBase<float, 3>& Constraint);
+		template void UpdateConstraintImp<ECollisionUpdateType::Any, float, 3>(const FImplicitObject& Implicit0, const TRigidTransform<float, 3>& Transform0, const FImplicitObject& Implicit1, const TRigidTransform<float, 3>& Transform1, const float Thickness, TCollisionConstraintBase<float, 3>& Constraints);
+		template void UpdateConstraintImp<ECollisionUpdateType::Deepest, float, 3>(const FImplicitObject& Implicit0, const TRigidTransform<float, 3>& Transform0, const FImplicitObject& Implicit1, const TRigidTransform<float, 3>& Transform1, const float Thickness, TCollisionConstraintBase<float, 3>& Constraint);
 
-		template void ConstructPairConstraintImpl<float, 3>(TGeometryParticleHandle<float, 3>* Particle0, TGeometryParticleHandle<float, 3>* Particle1, const FImplicitObject* Implicit0, const FImplicitObject* Implicit1, const float Thickness, FCollisionConstraintsArray& NewConstraints);
-		template void ConstructConstraintsImpl<float, 3>(TGeometryParticleHandle<float, 3>* Particle0, TGeometryParticleHandle<float, 3>* Particle1, const FImplicitObject* Implicit0, const FImplicitObject* Implicit1, const float Thickness, FCollisionConstraintsArray& NewConstraints);
+		template void ConstructPairConstraintImpl<float, 3>(TGeometryParticleHandle<float, 3>* Particle0, TGeometryParticleHandle<float, 3>* Particle1, const FImplicitObject* Implicit0, const FImplicitObject* Implicit1, const TRigidTransform<float, 3>& Transform0, const TRigidTransform<float, 3>& Transform1, const float Thickness, FCollisionConstraintsArray& NewConstraints);
+		template void ConstructConstraintsImpl<float, 3>(TGeometryParticleHandle<float, 3>* Particle0, TGeometryParticleHandle<float, 3>* Particle1, const FImplicitObject* Implicit0, const FImplicitObject* Implicit1, const TRigidTransform<float, 3>& Transform0, const TRigidTransform<float, 3>& Transform1, const float Thickness, FCollisionConstraintsArray& NewConstraints);
 
 	} // Collisions
 
diff --git a/Engine/Source/Runtime/Experimental/Chaos/Private/Chaos/PBDCollisionConstraints.cpp b/Engine/Source/Runtime/Experimental/Chaos/Private/Chaos/PBDCollisionConstraints.cpp
--- a/Engine/Source/Runtime/Experimental/Chaos/Private/Chaos/PBDCollisionConstraints.cpp
+++ b/Engine/Source/Runtime/Experimental/Chaos/Private/Chaos/PBDCollisionConstraints.cpp
@@ -152,7 +152,7 @@ namespace Chaos
 	{
 		if (ensure(Particle0 && Particle1))
 		{
-			Collisions::ConstructConstraintsImpl<T, d>(Particle0, Particle1, Particle0->Geometry().Get(), Particle1->Geometry().Get(), Thickness, NewConstraints);
+			Collisions::ConstructConstraintsImpl<T, d>(Particle0, Particle1, Particle0->Geometry().Get(), Particle1->Geometry().Get(), Collisions::GetTransform(Particle0), Collisions::GetTransform(Particle1), Thickness, NewConstraints);
 		}
 
 	}
diff --git a/Engine/Source/Runtime/Experimental/Chaos/Private/Chaos/PBDJointConstraints.cpp b/Engine/Source/Runtime/Experimental/Chaos/Private/Chaos/PBDJointConstraints.cpp
--- a/Engine/Source/Runtime/Experimental/Chaos/Private/Chaos/PBDJointConstraints.cpp
+++ b/Engine/Source/Runtime/Experimental/Chaos/Private/Chaos/PBDJointConstraints.cpp
@@ -518,11 +518,11 @@ namespace Chaos
 	}
 
 
-	void FPBDJointConstraints::UpdateParticleState(TPBDRigidParticleHandle<FReal, 3>* Rigid, const FReal Dt, const FVec3& P, const FRotation3& Q)
+	void FPBDJointConstraints::UpdateParticleState(TPBDRigidParticleHandle<FReal, 3>* Rigid, const FReal Dt, const FVec3& P, const FRotation3& Q, const bool bUpdateVelocity)
 	{
 		if ((Rigid != nullptr) && (Rigid->ObjectState() == EObjectStateType::Dynamic))
 		{
-			if (Dt > SMALL_NUMBER)
+			if (bUpdateVelocity && (Dt > SMALL_NUMBER))
 			{
 				FVec3 DV = FVec3::CalculateVelocity(Rigid->P(), P, Dt);
 				FVec3 DW = FRotation3::CalculateAngularVelocity(Rigid->Q(), Q, Dt);
@@ -943,18 +943,8 @@ namespace Chaos
 		}
 
 		// Update the particles
-		TPBDRigidParticleHandle<FReal, 3>* Rigid0 = ConstraintParticles[ConstraintIndex][Index0]->CastToRigidParticle();
-		TPBDRigidParticleHandle<FReal, 3>* Rigid1 = ConstraintParticles[ConstraintIndex][Index1]->CastToRigidParticle();
-		if (Rigid0 && Rigid0->ObjectState() == EObjectStateType::Dynamic)
-		{
-			Rigid0->SetP(P0);
-			Rigid0->SetQ(Q0);
-		}
-		if (Rigid1 && Rigid1->ObjectState() == EObjectStateType::Dynamic)
-		{
-			Rigid1->SetP(P1);
-			Rigid1->SetQ(Q1);
-		}
+		UpdateParticleState(Particle0->CastToRigidParticle(), Dt, P0, Q0, false);
+		UpdateParticleState(Particle1->CastToRigidParticle(), Dt, P1, Q1, false);
 	}
 }
 
diff --git a/Engine/Source/Runtime/Experimental/Chaos/Public/Chaos/CollisionResolution.h b/Engine/Source/Runtime/Experimental/Chaos/Public/Chaos/CollisionResolution.h
--- a/Engine/Source/Runtime/Experimental/Chaos/Public/Chaos/CollisionResolution.h
+++ b/Engine/Source/Runtime/Experimental/Chaos/Public/Chaos/CollisionResolution.h
@@ -34,10 +34,10 @@ namespace Chaos
 			//
 
 			template <typename T, int d>
-			void CHAOS_API UpdateBoxBoxConstraint(const TBox<T, d>& Box1, const TRigidTransform<T, d>& Box1Transform, const TBox<T, d>& Box2, const TRigidTransform<T, d>& Box2Transform, T Thickness, TRigidBodyPointContactConstraint<T, d>& Constraint);
+			void CHAOS_API UpdateBoxBoxConstraint(const TBox<T, d>& Box1, const TRigidTransform<T, d>& Box1Transform, const TBox<T, d>& Box2, const TRigidTransform<T, d>& Box2Transform, const T Thickness, TRigidBodyPointContactConstraint<T, d>& Constraint);
 
 			template<typename T, int d>
-			void CHAOS_API ConstructBoxBoxConstraints(TGeometryParticleHandle<T, d>* Particle0, TGeometryParticleHandle<T, d>* Particle1, const FImplicitObject* Implicit0, const FImplicitObject* Implicit1, const T Thickness, FCollisionConstraintsArray& NewConstraints);
+			void CHAOS_API ConstructBoxBoxConstraints(TGeometryParticleHandle<T, d>* Particle0, TGeometryParticleHandle<T, d>* Particle1, const FImplicitObject* Implicit0, const FImplicitObject* Implicit1, const TRigidTransform<T, d>& Transform0, const TRigidTransform<T, d>& Transform1, const T Thickness, FCollisionConstraintsArray& NewConstraints);
 
 
 			//
@@ -45,10 +45,10 @@ namespace Chaos
 			//
 
 			template <typename T, int d>
-			bool CHAOS_API UpdateBoxPlaneConstraint(const TBox<T, d>& Box, const TRigidTransform<T, d>& BoxTransform, const TPlane<T, d>& Plane, const TRigidTransform<T, d>& PlaneTransform, T Thickness, TRigidBodyPointContactConstraint<T, d>& Constraint);
+			bool CHAOS_API UpdateBoxPlaneConstraint(const TBox<T, d>& Box, const TRigidTransform<T, d>& BoxTransform, const TPlane<T, d>& Plane, const TRigidTransform<T, d>& PlaneTransform, const T Thickness, TRigidBodyPointContactConstraint<T, d>& Constraint);
 
 			template<typename T, int d>
-			void CHAOS_API ConstructBoxPlaneConstraints(TGeometryParticleHandle<T, d>* Particle0, TGeometryParticleHandle<T, d>* Particle1, const FImplicitObject* Implicit0, const FImplicitObject* Implicit1, const T Thickness, FCollisionConstraintsArray& NewConstraints);
+			void CHAOS_API ConstructBoxPlaneConstraints(TGeometryParticleHandle<T, d>* Particle0, TGeometryParticleHandle<T, d>* Particle1, const FImplicitObject* Implicit0, const FImplicitObject* Implicit1, const TRigidTransform<T, d>& Transform0, const TRigidTransform<T, d>& Transform1, const T Thickness, FCollisionConstraintsArray& NewConstraints);
 
 
 			//
@@ -56,59 +56,59 @@ namespace Chaos
 			//
 
 			template <typename T, int d>
-			void CHAOS_API UpdateSphereSphereConstraint(const TSphere<T, d>& Sphere1, const TRigidTransform<T, d>& Sphere1Transform, const TSphere<T, d>& Sphere2, const TRigidTransform<T, d>& Sphere2Transform, T Thickness, TRigidBodyPointContactConstraint<T, d>& Constraint);
+			void CHAOS_API UpdateSphereSphereConstraint(const TSphere<T, d>& Sphere1, const TRigidTransform<T, d>& Sphere1Transform, const TSphere<T, d>& Sphere2, const TRigidTransform<T, d>& Sphere2Transform, const T Thickness, TRigidBodyPointContactConstraint<T, d>& Constraint);
 
 			template<typename T, int d>
-			void CHAOS_API ConstructSphereSphereConstraints(TGeometryParticleHandle<T, d>* Particle0, TGeometryParticleHandle<T, d>* Particle1, const FImplicitObject* Implicit0, const FImplicitObject* Implicit1, const T Thickness, FCollisionConstraintsArray& NewConstraints);
+			void CHAOS_API ConstructSphereSphereConstraints(TGeometryParticleHandle<T, d>* Particle0, TGeometryParticleHandle<T, d>* Particle1, const FImplicitObject* Implicit0, const FImplicitObject* Implicit1, const TRigidTransform<T, d>& Transform0, const TRigidTransform<T, d>& Transform1, const T Thickness, FCollisionConstraintsArray& NewConstraints);
 
 			//
 			// Sphere-Plane
 			//
 
 			template <typename T, int d>
-			void CHAOS_API UpdateSpherePlaneConstraint(const TSphere<T, d>& Sphere, const TRigidTransform<T, d>& SphereTransform, const TPlane<T, d>& Plane, const TRigidTransform<T, d>& PlaneTransform, T Thickness, TRigidBodyPointContactConstraint<T, d>& Constraint);
+			void CHAOS_API UpdateSpherePlaneConstraint(const TSphere<T, d>& Sphere, const TRigidTransform<T, d>& SphereTransform, const TPlane<T, d>& Plane, const TRigidTransform<T, d>& PlaneTransform, const T Thickness, TRigidBodyPointContactConstraint<T, d>& Constraint);
 
 			template<typename T, int d>
-			void CHAOS_API ConstructSpherePlaneConstraints(TGeometryParticleHandle<T, d>* Particle0, TGeometryParticleHandle<T, d>* Particle1, const FImplicitObject* Implicit0, const FImplicitObject* Implicit1, const T Thickness, FCollisionConstraintsArray& NewConstraints);
+			void CHAOS_API ConstructSpherePlaneConstraints(TGeometryParticleHandle<T, d>* Particle0, TGeometryParticleHandle<T, d>* Particle1, const FImplicitObject* Implicit0, const FImplicitObject* Implicit1, const TRigidTransform<T, d>& Transform0, const TRigidTransform<T, d>& Transform1, const T Thickness, FCollisionConstraintsArray& NewConstraints);
 
 			//
 			// Sphere-Box
 			//
 
 			template <typename T, int d>
-			void CHAOS_API UpdateSphereBoxConstraint(const TSphere<T, d>& Sphere, const TRigidTransform<T, d>& SphereTransform, const TBox<T, d>& Box, const TRigidTransform<T, d>& BoxTransform, T Thickness, TRigidBodyPointContactConstraint<T, d>& Constraint);
+			void CHAOS_API UpdateSphereBoxConstraint(const TSphere<T, d>& Sphere, const TRigidTransform<T, d>& SphereTransform, const TBox<T, d>& Box, const TRigidTransform<T, d>& BoxTransform, const T Thickness, TRigidBodyPointContactConstraint<T, d>& Constraint);
 
 			template<typename T, int d>
-			void CHAOS_API ConstructSphereBoxConstraints(TGeometryParticleHandle<T, d>* Particle0, TGeometryParticleHandle<T, d>* Particle1, const FImplicitObject* Implicit0, const FImplicitObject* Implicit1, const T Thickness, FCollisionConstraintsArray& NewConstraints);
+			void CHAOS_API ConstructSphereBoxConstraints(TGeometryParticleHandle<T, d>* Particle0, TGeometryParticleHandle<T, d>* Particle1, const FImplicitObject* Implicit0, const FImplicitObject* Implicit1, const TRigidTransform<T, d>& Transform0, const TRigidTransform<T, d>& Transform1, const T Thickness, FCollisionConstraintsArray& NewConstraints);
 
 			//
 			// Capsule-Capsule
 			//
 
 			template <typename T, int d>
-			void CHAOS_API UpdateCapsuleCapsuleConstraint(const TCapsule<T>& A, const TRigidTransform<T, d>& ATransform, const TCapsule<T>& B, const TRigidTransform<T, d>& BTransform, T Thickness, TRigidBodyPointContactConstraint<T, d>& Constraint);
+			void CHAOS_API UpdateCapsuleCapsuleConstraint(const TCapsule<T>& A, const TRigidTransform<T, d>& ATransform, const TCapsule<T>& B, const TRigidTransform<T, d>& BTransform, const T Thickness, TRigidBodyPointContactConstraint<T, d>& Constraint);
 
 			template<typename T, int d>
-			void CHAOS_API ConstructCapsuleCapsuleConstraints(TGeometryParticleHandle<T, d>* Particle0, TGeometryParticleHandle<T, d>* Particle1, const FImplicitObject* Implicit0, const FImplicitObject* Implicit1, const T Thickness, FCollisionConstraintsArray& NewConstraints);
+			void CHAOS_API ConstructCapsuleCapsuleConstraints(TGeometryParticleHandle<T, d>* Particle0, TGeometryParticleHandle<T, d>* Particle1, const FImplicitObject* Implicit0, const FImplicitObject* Implicit1, const TRigidTransform<T, d>& Transform0, const TRigidTransform<T, d>& Transform1, const T Thickness, FCollisionConstraintsArray& NewConstraints);
 
 			//
 			// Capsule-Box
 			//
 
 			template <typename T, int d>
-			void CHAOS_API UpdateCapsuleBoxConstraint(const TCapsule<T>& A, const TRigidTransform<T, d>& ATransform, const TBox<T, d>& B, const TRigidTransform<T, d>& BTransform, T Thickness, TRigidBodyPointContactConstraint<T, d>& Constraint);
+			void CHAOS_API UpdateCapsuleBoxConstraint(const TCapsule<T>& A, const TRigidTransform<T, d>& ATransform, const TBox<T, d>& B, const TRigidTransform<T, d>& BTransform, const T Thickness, TRigidBodyPointContactConstraint<T, d>& Constraint);
 
 			template<typename T, int d>
-			void CHAOS_API ConstructCapsuleBoxConstraints(TGeometryParticleHandle<T, d>* Particle0, TGeometryParticleHandle<T, d>* Particle1, const FImplicitObject* Implicit0, const FImplicitObject* Implicit1, const T Thickness, FCollisionConstraintsArray& NewConstraints);
+			void CHAOS_API ConstructCapsuleBoxConstraints(TGeometryParticleHandle<T, d>* Particle0, TGeometryParticleHandle<T, d>* Particle1, const FImplicitObject* Implicit0, const FImplicitObject* Implicit1, const TRigidTransform<T, d>& Transform0, const TRigidTransform<T, d>& Transform1, const T Thickness, FCollisionConstraintsArray& NewConstraints);
 
 			//
 			// Convex-Convex
 			//
 			template <typename T, int d>
-			void CHAOS_API UpdateConvexConvexConstraint(const FImplicitObject& A, const TRigidTransform<T, d>& ATM, const FImplicitObject& B, const TRigidTransform<T, d>& BTM, T Thickness, TRigidBodyPointContactConstraint<T, d>& Constraint);
+			void CHAOS_API UpdateConvexConvexConstraint(const FImplicitObject& A, const TRigidTransform<T, d>& ATM, const FImplicitObject& B, const TRigidTransform<T, d>& BTM, const T Thickness, TRigidBodyPointContactConstraint<T, d>& Constraint);
 
 			template <typename T, int d>
-			void CHAOS_API ConstructConvexConvexConstraints(TGeometryParticleHandle<T, d>* Particle0, TGeometryParticleHandle<T, d>* Particle1, const FImplicitObject* Implicit0, const FImplicitObject* Implicit1, const T Thickness, FCollisionConstraintsArray& NewConstraints);
+			void CHAOS_API ConstructConvexConvexConstraints(TGeometryParticleHandle<T, d>* Particle0, TGeometryParticleHandle<T, d>* Particle1, const FImplicitObject* Implicit0, const FImplicitObject* Implicit1, const TRigidTransform<T, d>& Transform0, const TRigidTransform<T, d>& Transform1, const T Thickness, FCollisionConstraintsArray& NewConstraints);
 
 			//
 			// Levelset-Levelset
@@ -117,7 +117,7 @@ namespace Chaos
 			void CHAOS_API UpdateLevelsetLevelsetConstraint(const T Thickness, TRigidBodyPointContactConstraint<float, 3>& Constraint);
 
 			template<typename T, int d>
-			void CHAOS_API ConstructLevelsetLevelsetConstraints(TGeometryParticleHandle<T, d>* Particle0, TGeometryParticleHandle<T, d>* Particle1, const FImplicitObject* Implicit0, const FImplicitObject* Implicit1, const T Thickness, FCollisionConstraintsArray& NewConstraints);
+			void CHAOS_API ConstructLevelsetLevelsetConstraints(TGeometryParticleHandle<T, d>* Particle0, TGeometryParticleHandle<T, d>* Particle1, const FImplicitObject* Implicit0, const FImplicitObject* Implicit1, const TRigidTransform<T, d>& Transform0, const TRigidTransform<T, d>& Transform1, const T Thickness, FCollisionConstraintsArray& NewConstraints);
 
 
 			//
@@ -136,20 +136,20 @@ namespace Chaos
 			//
 
 			template<ECollisionUpdateType UpdateType, typename T, int d>
-			void CHAOS_API UpdateUnionUnionConstraint(const T Thickness, TCollisionConstraintBase<T, d>& Constraint);
+			void CHAOS_API UpdateUnionUnionConstraint(const FImplicitObject& Implicit0, const TRigidTransform<T, d>& Transform0, const FImplicitObject& Implicit1, const TRigidTransform<T, d>& Transform1, const T Thickness, TCollisionConstraintBase<T, d>& Constraint);
 
 			template<typename T, int d>
-			void CHAOS_API ConstructUnionUnionConstraints(TGeometryParticleHandle<T, d>* Particle0, TGeometryParticleHandle<T, d>* Particle1, const FImplicitObject* Implicit0, const FImplicitObject* Implicit1, const T Thickness, FCollisionConstraintsArray& NewConstraints);
+			void CHAOS_API ConstructUnionUnionConstraints(TGeometryParticleHandle<T, d>* Particle0, TGeometryParticleHandle<T, d>* Particle1, const FImplicitObject* Implicit0, const FImplicitObject* Implicit1, const TRigidTransform<T, d>& Transform0, const TRigidTransform<T, d>& Transform1, const T Thickness, FCollisionConstraintsArray& NewConstraints);
 
 			//
 			// Single-Union
 			//
 
 			template<ECollisionUpdateType UpdateType, typename T, int d>
-			void CHAOS_API UpdateSingleUnionConstraint(const T Thickness, TCollisionConstraintBase<T, d>& Constraint);
+			void CHAOS_API UpdateSingleUnionConstraint(const FImplicitObject& Implicit0, const TRigidTransform<T, d>& Transform0, const FImplicitObject& Implicit1, const TRigidTransform<T, d>& Transform1, const T Thickness, TCollisionConstraintBase<T, d>& Constraint);
 
 			template<typename T, int d>
-			void CHAOS_API ConstructSingleUnionConstraints(TGeometryParticleHandle<T, d>* Particle0, TGeometryParticleHandle<T, d>* Particle1, const FImplicitObject* Implicit0, const FImplicitObject* Implicit1, const T Thickness, FCollisionConstraintsArray& NewConstraints);
+			void CHAOS_API ConstructSingleUnionConstraints(TGeometryParticleHandle<T, d>* Particle0, TGeometryParticleHandle<T, d>* Particle1, const FImplicitObject* Implicit0, const FImplicitObject* Implicit1, const TRigidTransform<T, d>& Transform0, const TRigidTransform<T, d>& Transform1, const  T Thickness, FCollisionConstraintsArray& NewConstraints);
 
 
 			//
@@ -160,9 +160,9 @@ namespace Chaos
 			void CHAOS_API UpdateConstraintImp(const FImplicitObject& ParticleObject, const TRigidTransform<T, d>& ParticleTM, const FImplicitObject& LevelsetObject, const TRigidTransform<T, d>& LevelsetTM, const T Thickness, TCollisionConstraintBase<T, d>& Constraint);
 
 			template<typename T, int d>
-			void CHAOS_API ConstructPairConstraintImpl(TGeometryParticleHandle<T, d>* Particle0, TGeometryParticleHandle<T, d>* Particle1, const FImplicitObject* Implicit0, const FImplicitObject* Implicit1, const T Thickness, FCollisionConstraintsArray& NewConstraints);
+			void CHAOS_API ConstructPairConstraintImpl(TGeometryParticleHandle<T, d>* Particle0, TGeometryParticleHandle<T, d>* Particle1, const FImplicitObject* Implicit0, const FImplicitObject* Implicit1, const TRigidTransform<T, d>& Transform0, const TRigidTransform<T, d>& Transform1, const T Thickness, FCollisionConstraintsArray& NewConstraints);
 
 			template<typename T, int d>
-			void CHAOS_API ConstructConstraintsImpl(TGeometryParticleHandle<T, d>* Particle0, TGeometryParticleHandle<T, d>* Particle1, const FImplicitObject* Implicit0, const FImplicitObject* Implicit1, const T Thickness, FCollisionConstraintsArray& NewConstraints);
+			void CHAOS_API ConstructConstraintsImpl(TGeometryParticleHandle<T, d>* Particle0, TGeometryParticleHandle<T, d>* Particle1, const FImplicitObject* Implicit0, const FImplicitObject* Implicit1, const TRigidTransform<T, d>& Transform0, const TRigidTransform<T, d>& Transform1, const T Thickness, FCollisionConstraintsArray& NewConstraints);
 	}
 }
