Chaos - dense matrix cleanup

diff --git a/Engine/Source/Runtime/Experimental/Chaos/Public/Chaos/DenseMatrix.h b/Engine/Source/Runtime/Experimental/Chaos/Public/Chaos/DenseMatrix.h
--- a/Engine/Source/Runtime/Experimental/Chaos/Public/Chaos/DenseMatrix.h
+++ b/Engine/Source/Runtime/Experimental/Chaos/Public/Chaos/DenseMatrix.h
@@ -177,7 +177,7 @@ namespace Chaos
 			check(NumRows() == NumColumns());
 			for (int32 II = 0; II < NRows; ++II)
 			{
-				M[ElementIndex(II, II)] = V;
+				At(II, II) = V;
 			}
 		}
 
@@ -191,7 +191,7 @@ namespace Chaos
 			for (int32 II = 0; II < Num; ++II)
 			{
 				int32 JJ = Start + II;
-				M[ElementIndex(JJ, JJ)] = V;
+				At(JJ, JJ) = V;
 			}
 		}
 
@@ -204,7 +204,7 @@ namespace Chaos
 			check(ColumnIndex < NumColumns());
 			for (int32 II = 0; II < NumV; ++II)
 			{
-				M[ElementIndex(RowIndex, ColumnIndex + II)] = V[II];
+				At(RowIndex, ColumnIndex + II) = V[II];
 			}
 		}
 
@@ -212,9 +212,9 @@ namespace Chaos
 		{
 			check(RowIndex + 1 <= NumRows());
 			check(ColumnIndex + 3 <= NumColumns());
-			M[ElementIndex(RowIndex, ColumnIndex + 0)] = V[0];
-			M[ElementIndex(RowIndex, ColumnIndex + 1)] = V[1];
-			M[ElementIndex(RowIndex, ColumnIndex + 2)] = V[2];
+			At(RowIndex, ColumnIndex + 0) = V[0];
+			At(RowIndex, ColumnIndex + 1) = V[1];
+			At(RowIndex, ColumnIndex + 2) = V[2];
 		}
 
 		/**
@@ -226,7 +226,7 @@ namespace Chaos
 			check(ColumnIndex + 1 <= NumColumns());
 			for (int32 II = 0; II < NumV; ++II)
 			{
-				M[ElementIndex(RowIndex + II, ColumnIndex)] = V[II];
+				At(RowIndex + II, ColumnIndex) = V[II];
 			}
 		}
 
@@ -234,9 +234,9 @@ namespace Chaos
 		{
 			check(RowIndex + 3 <= NumRows());
 			check(ColumnIndex + 1 <= NumColumns());
-			M[ElementIndex(RowIndex + 0, ColumnIndex)] = V[0];
-			M[ElementIndex(RowIndex + 1, ColumnIndex)] = V[1];
-			M[ElementIndex(RowIndex + 2, ColumnIndex)] = V[2];
+			At(RowIndex + 0, ColumnIndex) = V[0];
+			At(RowIndex + 1, ColumnIndex) = V[1];
+			At(RowIndex + 2, ColumnIndex) = V[2];
 		}
 
 		/**
@@ -251,7 +251,7 @@ namespace Chaos
 			{
 				for (int32 JJ = 0; JJ < V.NumColumns(); ++JJ)
 				{
-					M[ElementIndex(II + RowOffset, JJ + ColumnOffset)] = V.At(II, JJ);
+					At(II + RowOffset, JJ + ColumnOffset) = V.At(II, JJ);
 				}
 			}
 		}
@@ -267,7 +267,7 @@ namespace Chaos
 			{
 				for (int32 JJ = 0; JJ < 3; ++JJ)
 				{
-					M[ElementIndex(II + RowOffset, JJ + ColumnOffset)] = V.M[JJ][II];
+					At(II + RowOffset, JJ + ColumnOffset) = V.M[JJ][II];
 				}
 			}
 		}
@@ -773,7 +773,15 @@ namespace Chaos
 			const int32 N = G.NumRows();
 			X.SetDimensions(N, 1);
 
-			// Solve LY = B (G is lower-triangular)
+			// By definition: 
+			//		A.X = G.Gt.X = B
+			// Rearrange and define Y: 
+			//		Gt.X = G^-1.B = Y
+			// Which gives:
+			//		GY = B
+			//		GtX = Y
+
+			// Solve GY = B (G is lower-triangular) for Y
 			for (int32 I = 0; I < N; ++I)
 			{
 				FReal Sum = B.At(I, 0);
@@ -784,7 +792,7 @@ namespace Chaos
 				X.At(I, 0) = Sum / G.At(I, I);
 			}
 
-			// Solve LtX = Y (Gt is upper-triangular)
+			// Solve GtX = Y (Gt is upper-triangular) for X
 			for (int32 I = N - 1; I >= 0; --I)
 			{
 				FReal Sum = X.At(I, 0);
@@ -799,13 +807,13 @@ namespace Chaos
 		/**
 		 * Solve AX = B, for positive-definite NxN matrix A, and Nx1 column vectors B and X.
 		 *
-		 * For positive definite A, A = GGt, where G is the Cholesky factor and lower trangular.
+		 * For positive definite A, A = GGt, where G is the Cholesky factor and lower triangular.
 		 * We can solve GGtX = B by first solving GY = B, and then GtX = Y.
 		 *
 		 * E.g., this can be used to solve constraint equations of the form
 		 *		J.I.Jt.X = B
 		 * where J is a Jacobian (Jt its transpose), I is an Inverse mas matrix, and B the residual.
-		 * In this case, I is positive definite, and therefore so is JIJt.
+		 * In this case, I is symmetric positive definite, and therefore so is JIJt.
 		 *
 		 */
 		template<int32 T_EA, int32 T_EB, int32 T_EX>
