Chaos: Fix missing code for vehicle sleep state

diff --git a/Engine/Source/Runtime/Experimental/Chaos/Private/Chaos/PBDRigidParticles.cpp b/Engine/Source/Runtime/Experimental/Chaos/Private/Chaos/PBDRigidParticles.cpp
--- a/Engine/Source/Runtime/Experimental/Chaos/Private/Chaos/PBDRigidParticles.cpp
+++ b/Engine/Source/Runtime/Experimental/Chaos/Private/Chaos/PBDRigidParticles.cpp
@@ -4,7 +4,8 @@
 
 void Chaos::EnsureSleepingObjectState(EObjectStateType ObjectState)
 {
-	ensure(ObjectState != EObjectStateType::Kinematic);
+	// feasible for game to set a dynamic body to kinematic state then to sleeping state
+	//ensure(ObjectState != EObjectStateType::Kinematic);
 	ensure(ObjectState != EObjectStateType::Static);
 }
 
diff --git a/Engine/Source/Runtime/Experimental/Chaos/Public/Chaos/PBDRigidParticles.h b/Engine/Source/Runtime/Experimental/Chaos/Public/Chaos/PBDRigidParticles.h
--- a/Engine/Source/Runtime/Experimental/Chaos/Public/Chaos/PBDRigidParticles.h
+++ b/Engine/Source/Runtime/Experimental/Chaos/Public/Chaos/PBDRigidParticles.h
@@ -85,11 +85,18 @@ class TPBDRigidParticles : public TRigidParticles<T, d>
 			this->AddSleepData(Particle);
 		}
 
+		// Dynamic -> Sleeping or Sleeping -> Dynamic
 		if (this->ObjectState(Index) == EObjectStateType::Dynamic || this->ObjectState(Index) == EObjectStateType::Sleeping)
 		{
 			this->ObjectState(Index) = bSleeping ? EObjectStateType::Sleeping : EObjectStateType::Dynamic;
 		}
 
+		// Possible for code to set a Dynamic to Kinematic State then to Sleeping State
+		if (this->ObjectState(Index) == EObjectStateType::Kinematic && bSleeping)
+		{
+			this->ObjectState(Index) =  EObjectStateType::Sleeping;
+		}
+
 		if (bSleeping)
 		{
 			EnsureSleepingObjectState(this->ObjectState(Index));
