Fixed erratic collision regression when geometry collection particles interact with the world or each other

diff --git a/Engine/Source/Runtime/Experimental/ChaosSolvers/Private/PhysicsProxy/GeometryCollectionPhysicsProxy.cpp b/Engine/Source/Runtime/Experimental/ChaosSolvers/Private/PhysicsProxy/GeometryCollectionPhysicsProxy.cpp
--- a/Engine/Source/Runtime/Experimental/ChaosSolvers/Private/PhysicsProxy/GeometryCollectionPhysicsProxy.cpp
+++ b/Engine/Source/Runtime/Experimental/ChaosSolvers/Private/PhysicsProxy/GeometryCollectionPhysicsProxy.cpp
@@ -258,6 +258,7 @@ void PopulateSimulatedParticle(
 		TEXT("Inertia tensor is too small. Too late to change"));
 
 	Handle->SetM(MassIn);
+	Handle->SetInvM(1.0f / MassIn);
 	//Particles.M(RigidBodyIndex) = MassIn;
 	if (FMath::IsNaN(InertiaTensorVec[0]) || FMath::IsNaN(InertiaTensorVec[1]) || FMath::IsNaN(InertiaTensorVec[2]) ||
 		InertiaTensorVec[0] < SMALL_NUMBER || InertiaTensorVec[1] < SMALL_NUMBER || InertiaTensorVec[2] < SMALL_NUMBER)
@@ -268,6 +269,7 @@ void PopulateSimulatedParticle(
 	else
 	{
 		Handle->SetI(Chaos::PMatrix<float, 3, 3>(InertiaTensorVec[0], InertiaTensorVec[1], InertiaTensorVec[2]));
+		Handle->SetInvI(Chaos::PMatrix<float, 3, 3>(1.0f / InertiaTensorVec[0], 1.0f / InertiaTensorVec[1], 1.0f / InertiaTensorVec[2]));
 		//Particles.I(RigidBodyIndex) = Chaos::PMatrix<float, 3, 3>(InertiaTensorVec[0], InertiaTensorVec[1], InertiaTensorVec[2]);
 	}
 
