Fix heightfield most opposing using unscaled points. Also fix the wrong index being saved out. Still need to fix overlap using unscaled points

diff --git a/Engine/Source/Runtime/Experimental/Chaos/Private/Chaos/HeightField.cpp b/Engine/Source/Runtime/Experimental/Chaos/Private/Chaos/HeightField.cpp
--- a/Engine/Source/Runtime/Experimental/Chaos/Private/Chaos/HeightField.cpp
+++ b/Engine/Source/Runtime/Experimental/Chaos/Private/Chaos/HeightField.cpp
@@ -1277,6 +1277,7 @@ namespace Chaos
 			}
 		};
 
+		ensure(PotentialIntersections.Num());
 		for(const TVector<int32, 2>& CellCoord : PotentialIntersections)
 		{
 			const int32 CellIndex = CellCoord[1] * (GeomData.NumCols - 1) + CellCoord[0];
@@ -1287,10 +1288,10 @@ namespace Chaos
 
 			TVector<T, 3> Points[4];
 
-			GeomData.GetPoints(FullIndex, Points);
+			GeomData.GetPointsScaled(FullIndex, Points);
 
-			CheckTriangle(FullIndex * 2, Points[0], Points[1], Points[2]);
-			CheckTriangle(FullIndex * 2 + 1, Points[2], Points[1], Points[3]);
+			CheckTriangle(CellIndex * 2, Points[0], Points[1], Points[2]);
+			CheckTriangle(CellIndex * 2 + 1, Points[2], Points[1], Points[3]);
 		}
 
 		return MostOpposingFace;
