Chaos - Sphere-Capsule collision

diff --git a/Engine/Source/Runtime/Experimental/Chaos/Private/Chaos/CollisionResolution.cpp b/Engine/Source/Runtime/Experimental/Chaos/Private/Chaos/CollisionResolution.cpp
--- a/Engine/Source/Runtime/Experimental/Chaos/Private/Chaos/CollisionResolution.cpp
+++ b/Engine/Source/Runtime/Experimental/Chaos/Private/Chaos/CollisionResolution.cpp
@@ -347,6 +347,61 @@ namespace Chaos
 		}
 
 
+		//
+		// Sphere - Capsule
+		//
+
+		template <typename T, int d>
+		void UpdateSphereCapsuleConstraint(const TSphere<T, d>& A, const TRigidTransform<T, d>& ATransform, const TCapsule<T>& B, const TRigidTransform<T, d>& BTransform, const T Thickness, TRigidBodyPointContactConstraint<T, d>& Constraint)
+		{
+			TPointContactManifold<T, d>& Contact = Constraint.Manifold;
+
+			FVector A1 = ATransform.TransformPosition(A.GetCenter());
+			FVector B1 = BTransform.TransformPosition(B.GetX1());
+			FVector B2 = BTransform.TransformPosition(B.GetX2());
+			FVector P2 = FMath::ClosestPointOnSegment(A1, B1, B2);
+
+			TVector<T, d> Delta = P2 - A1;
+			T DeltaLen = Delta.Size();
+			if (DeltaLen > KINDA_SMALL_NUMBER)
+			{
+				T NewPhi = DeltaLen - (A.GetRadius() + B.GetRadius());
+				if (NewPhi < Contact.Phi)
+				{
+					TVector<T, d> Dir = Delta / DeltaLen;
+					Contact.Phi = NewPhi;
+					Contact.Normal = -Dir;
+					Contact.Location = A1 + Dir * A.GetRadius();
+				}
+			}
+		}
+
+		template<typename T, int d>
+		void ConstructSphereCapsuleConstraints(TGeometryParticleHandle<T, d>* Particle0, TGeometryParticleHandle<T, d>* Particle1, const FImplicitObject* Implicit0, const FImplicitObject* Implicit1, const TRigidTransform<T, d>& Transform0, const TRigidTransform<T, d>& Transform1, const T Thickness, FCollisionConstraintsArray& NewConstraints)
+		{
+			const TSphere<T, d>* Object0 = Implicit0->template GetObject<const TSphere<T, d> >();
+			const TCapsule<T>* Object1 = Implicit1->template GetObject<const TCapsule<T> >();
+			if (ensure(Object0 && Object1))
+			{
+				FRigidBodyPointContactConstraint* Constraint = new FRigidBodyPointContactConstraint;
+				Constraint->Particle[0] = Particle0;
+				Constraint->Particle[1] = Particle1;
+				Constraint->AddManifold(Implicit0, Implicit1);
+
+				UpdateSphereCapsuleConstraint(*Object0, Transform0, *Object1, Transform1, Thickness, *Constraint);
+
+				if (Constraint->GetPhi() < Thickness)
+				{
+					NewConstraints.Add(Constraint);
+				}
+				else
+				{
+					delete Constraint;
+				}
+			}
+		}
+
+
 		//
 		// Capsule-Capsule
 		//
@@ -884,6 +939,10 @@ namespace Chaos
 			{
 				UpdateSphereBoxConstraint(*Implicit0.template GetObject<TSphere<T, d>>(), Transform0, *Implicit1.template GetObject<TBox<T, d>>(), Transform1, Thickness, *ConstraintBase.template As<TRigidBodyPointContactConstraint<T,d>>());
 			}
+			else if (Implicit0Type == TSphere<T, d>::StaticType() && Implicit1Type == TCapsule<T>::StaticType())
+			{
+				UpdateSphereCapsuleConstraint(*Implicit0.template GetObject<TSphere<T, d>>(), Transform0, *Implicit1.template GetObject<TCapsule<T>>(), Transform1, Thickness, *ConstraintBase.template As<TRigidBodyPointContactConstraint<T, d>>());
+			}
 			else if (Implicit0Type == TCapsule<T>::StaticType() && Implicit1Type == TCapsule<T>::StaticType())
 			{
 				UpdateCapsuleCapsuleConstraint(*Implicit0.template GetObject<TCapsule<T>>(), Transform0, *Implicit1.template GetObject<TCapsule<T>>(), Transform1, Thickness, *ConstraintBase.template As<TRigidBodyPointContactConstraint<T,d>>());
@@ -936,6 +995,17 @@ namespace Chaos
 					Constraint.SetNormal(-Constraint.GetNormal());
 				}
 			}
+			else if (Implicit0Type == TCapsule<T>::StaticType() && Implicit1Type == TSphere<T, d>::StaticType())
+			{
+				TRigidBodyPointContactConstraint<T, d>& Constraint = *ConstraintBase.template As<TRigidBodyPointContactConstraint<T, d> >();
+				TRigidBodyPointContactConstraint<T, d> TmpConstraint = Constraint;
+				UpdateSphereCapsuleConstraint(*Implicit1.template GetObject<TSphere<T, d>>(), Transform1, *Implicit0.template GetObject<TCapsule<T>>(), Transform0, Thickness, TmpConstraint);
+				if (TmpConstraint.GetPhi() < Constraint.GetPhi())
+				{
+					Constraint = TmpConstraint;
+					Constraint.SetNormal(-Constraint.GetNormal());
+				}
+			}
 			else if (Implicit0Type < TImplicitObjectUnion<T, d>::StaticType() && Implicit1Type == TImplicitObjectUnion<T, d>::StaticType())
 			{
 				return UpdateSingleUnionConstraint<UpdateType>(Implicit0, Transform0, Implicit1, Transform1, Thickness, ConstraintBase);
@@ -1041,6 +1111,14 @@ namespace Chaos
 			{
 				ConstructSphereBoxConstraints(Particle1, Particle0, Implicit1, Implicit0, Transform1, Transform0, Thickness, NewConstraints);
 			}
+			else if (Implicit0Type == TSphere<T, d>::StaticType() && Implicit1Type == TCapsule<T>::StaticType())
+			{
+				ConstructSphereCapsuleConstraints(Particle0, Particle1, Implicit0, Implicit1, Transform0, Transform1, Thickness, NewConstraints);
+			}
+			else if (Implicit0Type == TCapsule<T>::StaticType() && Implicit1Type == TSphere<T, d>::StaticType())
+			{
+				ConstructSphereCapsuleConstraints(Particle1, Particle0, Implicit1, Implicit0, Transform1, Transform0, Thickness, NewConstraints);
+			}
 			else if (Implicit0Type == TCapsule<T>::StaticType() && Implicit1Type == TCapsule<T>::StaticType())
 			{
 				ConstructCapsuleCapsuleConstraints(Particle0, Particle1, Implicit0, Implicit1, Transform1, Transform0, Thickness, NewConstraints);
@@ -1111,6 +1189,9 @@ namespace Chaos
 		template void UpdateSphereBoxConstraint<float, 3>(const TSphere<float, 3>& Sphere, const TRigidTransform<float, 3>& SphereTransform, const TBox<float, 3>& Box, const TRigidTransform<float, 3>& BoxTransform, const float Thickness, TRigidBodyPointContactConstraint<float, 3>& Constraint);
 		template void ConstructSphereBoxConstraints<float, 3>(TGeometryParticleHandle<float, 3>* Particle0, TGeometryParticleHandle<float, 3>* Particle1, const FImplicitObject* Implicit0, const FImplicitObject* Implicit1, const TRigidTransform<float, 3>& Transform0, const TRigidTransform<float, 3>& Transform1, const float Thickness, FCollisionConstraintsArray& NewConstraints);
 
+		template void UpdateSphereCapsuleConstraint<float, 3>(const TSphere<float, 3>& Sphere, const TRigidTransform<float, 3>& SphereTransform, const TCapsule<float>& Box, const TRigidTransform<float, 3>& BoxTransform, const float Thickness, TRigidBodyPointContactConstraint<float, 3>& Constraint);
+		template void ConstructSphereCapsuleConstraints<float, 3>(TGeometryParticleHandle<float, 3>* Particle0, TGeometryParticleHandle<float, 3>* Particle1, const FImplicitObject* Implicit0, const FImplicitObject* Implicit1, const TRigidTransform<float, 3>& Transform0, const TRigidTransform<float, 3>& Transform1, const float Thickness, FCollisionConstraintsArray& NewConstraints);
+
 		template void UpdateCapsuleCapsuleConstraint<float, 3>(const TCapsule<float>& A, const TRigidTransform<float, 3>& ATransform, const TCapsule<float>& B, const TRigidTransform<float, 3>& BTransform, const float Thickness, TRigidBodyPointContactConstraint<float, 3>& Constraint);
 		template void ConstructCapsuleCapsuleConstraints(TGeometryParticleHandle<float, 3>* Particle0, TGeometryParticleHandle<float, 3>* Particle1, const FImplicitObject* Implicit0, const FImplicitObject* Implicit1, const TRigidTransform<float, 3>& Transform0, const TRigidTransform<float, 3>& Transform1, const float Thickness, FCollisionConstraintsArray& NewConstraints);
 
diff --git a/Engine/Source/Runtime/Experimental/Chaos/Public/Chaos/CollisionResolution.h b/Engine/Source/Runtime/Experimental/Chaos/Public/Chaos/CollisionResolution.h
--- a/Engine/Source/Runtime/Experimental/Chaos/Public/Chaos/CollisionResolution.h
+++ b/Engine/Source/Runtime/Experimental/Chaos/Public/Chaos/CollisionResolution.h
@@ -81,6 +81,16 @@ namespace Chaos
 			template<typename T, int d>
 			void CHAOS_API ConstructSphereBoxConstraints(TGeometryParticleHandle<T, d>* Particle0, TGeometryParticleHandle<T, d>* Particle1, const FImplicitObject* Implicit0, const FImplicitObject* Implicit1, const TRigidTransform<T, d>& Transform0, const TRigidTransform<T, d>& Transform1, const T Thickness, FCollisionConstraintsArray& NewConstraints);
 
+			//
+			// Sphere-Capsule
+			//
+
+			template <typename T, int d>
+			void CHAOS_API UpdateSphereCapsuleConstraint(const TSphere<T, d>& Sphere, const TRigidTransform<T, d>& SphereTransform, const TCapsule<T>& Box, const TRigidTransform<T, d>& BoxTransform, const T Thickness, TRigidBodyPointContactConstraint<T, d>& Constraint);
+
+			template<typename T, int d>
+			void CHAOS_API ConstructSphereCapsuleConstraints(TGeometryParticleHandle<T, d>* Particle0, TGeometryParticleHandle<T, d>* Particle1, const FImplicitObject* Implicit0, const FImplicitObject* Implicit1, const TRigidTransform<T, d>& Transform0, const TRigidTransform<T, d>& Transform1, const T Thickness, FCollisionConstraintsArray& NewConstraints);
+
 			//
 			// Capsule-Capsule
 			//
