Turn TConvex into FConvex

diff --git a/Engine/Plugins/Experimental/ChaosCloth/Source/Chaos/Private/ChaosCloth/ChaosClothingSimulation.cpp b/Engine/Plugins/Experimental/ChaosCloth/Source/Chaos/Private/ChaosCloth/ChaosClothingSimulation.cpp
--- a/Engine/Plugins/Experimental/ChaosCloth/Source/Chaos/Private/ChaosCloth/ChaosClothingSimulation.cpp
+++ b/Engine/Plugins/Experimental/ChaosCloth/Source/Chaos/Private/ChaosCloth/ChaosClothingSimulation.cpp
@@ -699,7 +699,7 @@ void ClothingSimulation::ExtractPhysicsAssetCollisions(UClothingAssetCommon* Ass
 
 #elif WITH_CHAOS  // #if WITH_PHYSX
 				const Chaos::FImplicitObject& ChaosConvexMesh = *ConvexElem.GetChaosConvexMesh();
-				const Chaos::TConvex<float, 3>& ChaosConvex = ChaosConvexMesh.GetObjectChecked<Chaos::TConvex<float, 3>>();
+				const Chaos::FConvex& ChaosConvex = ChaosConvexMesh.GetObjectChecked<Chaos::FConvex>();
 
 				// Copy planes
 				const TArray<TPlane<float, 3>>& Planes = ChaosConvex.GetFaces();
@@ -928,7 +928,7 @@ void ClothingSimulation::AddCollisions(const FClothCollisionData& ClothCollision
 					}
 
 					// Setup the collision particle geometry
-					CollisionParticles.SetDynamicGeometry(i, MakeUnique<Chaos::TConvex<float, 3>>(MoveTemp(Planes), MoveTemp(SurfaceParticles)));
+					CollisionParticles.SetDynamicGeometry(i, MakeUnique<Chaos::FConvex>(MoveTemp(Planes), MoveTemp(SurfaceParticles)));
 				}
 			}
 
@@ -1442,7 +1442,7 @@ void ClothingSimulation::DebugDrawCollision(USkeletalMeshComponent* OwnerCompone
 		}
 	};
 
-	auto DrawConvex = [&PDI](const Chaos::TConvex<float, 3>& Convex, const TRotation<float, 3>& Rotation, const Chaos::TVector<float, 3>& Position, const FLinearColor& Color)
+	auto DrawConvex = [&PDI](const Chaos::FConvex& Convex, const TRotation<float, 3>& Rotation, const Chaos::TVector<float, 3>& Position, const FLinearColor& Color)
 	{
 		const TArray<Chaos::TPlane<float, 3>>& Planes = Convex.GetFaces();
 		for (int32 PlaneIndex1 = 0; PlaneIndex1 < Planes.Num(); ++PlaneIndex1)
@@ -1531,7 +1531,7 @@ void ClothingSimulation::DebugDrawCollision(USkeletalMeshComponent* OwnerCompone
 				break;
 	
 			case Chaos::ImplicitObjectType::Convex:
-				DrawConvex(Object->GetObjectChecked<Chaos::TConvex<float, 3>>(), Rotation, Position, Color);
+				DrawConvex(Object->GetObjectChecked<Chaos::FConvex>(), Rotation, Position, Color);
 				break;
 
 			default:
diff --git a/Engine/Source/Runtime/Engine/Classes/PhysicsEngine/ConvexElem.h b/Engine/Source/Runtime/Engine/Classes/PhysicsEngine/ConvexElem.h
--- a/Engine/Source/Runtime/Engine/Classes/PhysicsEngine/ConvexElem.h
+++ b/Engine/Source/Runtime/Engine/Classes/PhysicsEngine/ConvexElem.h
@@ -54,7 +54,7 @@ private:
 	physx::PxConvexMesh*   ConvexMeshNegX;
 
 #if WITH_CHAOS
-	TUniquePtr<Chaos::TConvex<float,3>> ChaosConvex;
+	TUniquePtr<Chaos::FConvex> ChaosConvex;
 #endif
 
 public:
@@ -106,9 +106,9 @@ public:
 	ENGINE_API void SetMirroredConvexMesh(physx::PxConvexMesh* InMesh);
 
 #if WITH_CHAOS
-	ENGINE_API const TUniquePtr<Chaos::TConvex<float, 3>>& GetChaosConvexMesh() const;
+	ENGINE_API const TUniquePtr<Chaos::FConvex>& GetChaosConvexMesh() const;
 
-	ENGINE_API void SetChaosConvexMesh(TUniquePtr<Chaos::TConvex<float, 3>>&& InChaosConvex);
+	ENGINE_API void SetChaosConvexMesh(TUniquePtr<Chaos::FConvex>&& InChaosConvex);
 
 	ENGINE_API void ResetChaosConvexMesh();
 #endif
diff --git a/Engine/Source/Runtime/Engine/Private/PhysicsEngine/BodySetup.cpp b/Engine/Source/Runtime/Engine/Private/PhysicsEngine/BodySetup.cpp
--- a/Engine/Source/Runtime/Engine/Private/PhysicsEngine/BodySetup.cpp
+++ b/Engine/Source/Runtime/Engine/Private/PhysicsEngine/BodySetup.cpp
@@ -1674,12 +1674,12 @@ float FKConvexElem::GetVolume(const FVector& Scale) const
 }
 
 #if WITH_CHAOS
-const TUniquePtr<Chaos::TConvex<float, 3>>& FKConvexElem::GetChaosConvexMesh() const
+const TUniquePtr<Chaos::FConvex>& FKConvexElem::GetChaosConvexMesh() const
 {
 	return ChaosConvex;
 }
 
-void FKConvexElem::SetChaosConvexMesh(TUniquePtr<Chaos::TConvex<float, 3>>&& InChaosConvex)
+void FKConvexElem::SetChaosConvexMesh(TUniquePtr<Chaos::FConvex>&& InChaosConvex)
 {
 	ChaosConvex = MoveTemp(InChaosConvex);
 }
diff --git a/Engine/Source/Runtime/Engine/Private/PhysicsEngine/Experimental/ChaosDerivedData.cpp b/Engine/Source/Runtime/Engine/Private/PhysicsEngine/Experimental/ChaosDerivedData.cpp
--- a/Engine/Source/Runtime/Engine/Private/PhysicsEngine/Experimental/ChaosDerivedData.cpp
+++ b/Engine/Source/Runtime/Engine/Private/PhysicsEngine/Experimental/ChaosDerivedData.cpp
@@ -134,7 +134,7 @@ void FChaosDerivedDataCooker::BuildTriangleMeshes(TArray<TUniquePtr<Chaos::TTria
 	{
 		// Only add this triangle if it is valid
 		const int32 BaseIndex = TriangleIndex * 3;
-		const bool bIsValidTriangle = Chaos::TConvexBuilder<Precision>::IsValidTriangle(
+		const bool bIsValidTriangle = Chaos::FConvexBuilder::IsValidTriangle(
 			FinalVerts[FinalIndices[BaseIndex]],
 			FinalVerts[FinalIndices[BaseIndex + 1]],
 			FinalVerts[FinalIndices[BaseIndex + 2]]);
@@ -195,7 +195,7 @@ void FChaosDerivedDataCooker::BuildConvexMeshes(TArray<TUniquePtr<Chaos::FImplic
 				ConvexParticles.X(VertIndex) = FVector(bMirrored ? -HullVert.X : HullVert.X, HullVert.Y, HullVert.Z);
 			}
 
-			OutConvexes.Emplace(new Chaos::TConvex<Chaos::FReal, 3>(ConvexParticles));
+			OutConvexes.Emplace(new Chaos::FConvex(ConvexParticles));
 		}
 	};
 
diff --git a/Engine/Source/Runtime/Engine/Private/PhysicsEngine/Experimental/ChaosInterfaceUtils.cpp b/Engine/Source/Runtime/Engine/Private/PhysicsEngine/Experimental/ChaosInterfaceUtils.cpp
--- a/Engine/Source/Runtime/Engine/Private/PhysicsEngine/Experimental/ChaosInterfaceUtils.cpp
+++ b/Engine/Source/Runtime/Engine/Private/PhysicsEngine/Experimental/ChaosInterfaceUtils.cpp
@@ -260,7 +260,7 @@ namespace ChaosInterface
 				//}
 				//else
 				{
-					auto Implicit = MakeUnique<Chaos::TImplicitObjectScaled<Chaos::TConvex<float, 3>>>(MakeSerializable(ConvexImplicit), Scale);
+					auto Implicit = MakeUnique<Chaos::TImplicitObjectScaled<Chaos::FConvex>>(MakeSerializable(ConvexImplicit), Scale);
 					auto NewShape = NewShapeHelper(MakeSerializable(Implicit), (void*)CollisionBody.GetUserData());
 					Shapes.Emplace(MoveTemp(NewShape));
 					Geoms.Add(MoveTemp(Implicit));
diff --git a/Engine/Source/Runtime/Engine/Private/PhysicsEngine/Experimental/PhysInterface_Chaos.cpp b/Engine/Source/Runtime/Engine/Private/PhysicsEngine/Experimental/PhysInterface_Chaos.cpp
--- a/Engine/Source/Runtime/Engine/Private/PhysicsEngine/Experimental/PhysInterface_Chaos.cpp
+++ b/Engine/Source/Runtime/Engine/Private/PhysicsEngine/Experimental/PhysInterface_Chaos.cpp
@@ -1339,9 +1339,9 @@ const Chaos::TCapsule<float>&  FPhysicsGeometryCollection_Chaos::GetCapsuleGeome
 	return Geom.GetObjectChecked<Chaos::TCapsule<float>>();
 }
 
-const Chaos::TConvex<float, 3>& FPhysicsGeometryCollection_Chaos::GetConvexGeometry() const
+const Chaos::FConvex& FPhysicsGeometryCollection_Chaos::GetConvexGeometry() const
 {
-	return Geom.GetObjectChecked<Chaos::TConvex<float, 3>>();
+	return Geom.GetObjectChecked<Chaos::FConvex>();
 }
 
 const Chaos::TTriangleMeshImplicitObject<float>& FPhysicsGeometryCollection_Chaos::GetTriMeshGeometry() const
diff --git a/Engine/Source/Runtime/Engine/Public/Physics/Experimental/ChaosDerivedDataReader.h b/Engine/Source/Runtime/Engine/Public/Physics/Experimental/ChaosDerivedDataReader.h
--- a/Engine/Source/Runtime/Engine/Public/Physics/Experimental/ChaosDerivedDataReader.h
+++ b/Engine/Source/Runtime/Engine/Public/Physics/Experimental/ChaosDerivedDataReader.h
@@ -25,7 +25,7 @@ public:
 	// Only valid use is to explicitly read chaos bulk data
 	explicit FChaosDerivedDataReader(FUntypedBulkData* InBulkData);
 
-	TArray<TUniquePtr<Chaos::TConvex<T, d>>> ConvexImplicitObjects;
+	TArray<TUniquePtr<Chaos::FConvex>> ConvexImplicitObjects;
 	TArray<TUniquePtr<Chaos::TTriangleMeshImplicitObject<T>>> TrimeshImplicitObjects;
 	FBodySetupUVInfo UVInfo;
 
diff --git a/Engine/Source/Runtime/Engine/Public/Physics/Experimental/PhysInterface_Chaos.h b/Engine/Source/Runtime/Engine/Public/Physics/Experimental/PhysInterface_Chaos.h
--- a/Engine/Source/Runtime/Engine/Public/Physics/Experimental/PhysInterface_Chaos.h
+++ b/Engine/Source/Runtime/Engine/Public/Physics/Experimental/PhysInterface_Chaos.h
@@ -49,8 +49,7 @@ namespace Chaos
 	template <typename T, int>
 	class TPBDSpringConstraints;
 
-	template <typename T, int>
-	class TConvex;
+	class FConvex;
 
 	template <typename T>
 	class TCapsule;
@@ -147,7 +146,7 @@ struct ENGINE_API FPhysicsGeometryCollection_Chaos
 	const Chaos::TBox<float, 3>& GetBoxGeometry() const;
 	const Chaos::TSphere<float, 3>&  GetSphereGeometry() const;
 	const Chaos::TCapsule<float>&  GetCapsuleGeometry() const;
-	const Chaos::TConvex<float, 3>& GetConvexGeometry() const;
+	const Chaos::FConvex& GetConvexGeometry() const;
 	const Chaos::TTriangleMeshImplicitObject<float>& GetTriMeshGeometry() const;
 
 private:
diff --git a/Engine/Source/Runtime/Experimental/Chaos/Private/Chaos/CollisionConvexMesh.cpp b/Engine/Source/Runtime/Experimental/Chaos/Private/Chaos/CollisionConvexMesh.cpp
--- a/Engine/Source/Runtime/Experimental/Chaos/Private/Chaos/CollisionConvexMesh.cpp
+++ b/Engine/Source/Runtime/Experimental/Chaos/Private/Chaos/CollisionConvexMesh.cpp
@@ -6,18 +6,15 @@
 namespace Chaos
 {
 	// CVars variables for controlling geometry complexity checking and simplification
-	template <>
-	int32 TConvexBuilder<float>::PerformGeometryCheck = 0;
+	int32 FConvexBuilder::PerformGeometryCheck = 0;
 
-	template <>
-	int32 TConvexBuilder<float>::PerformGeometryReduction = 0;
+	int32 FConvexBuilder::PerformGeometryReduction = 0;
 
-	template <>
-	int32 TConvexBuilder<float>::ParticlesThreshold = 50;
+	int32 FConvexBuilder::ParticlesThreshold = 50;
 
-	FAutoConsoleVariableRef CVarConvexGeometryCheckEnable(TEXT("p.Chaos.ConvexGeometryCheckEnable"), TConvexBuilder<float>::PerformGeometryCheck, TEXT("Perform convex geometry complexity check for Chaos physics."));
+	FAutoConsoleVariableRef CVarConvexGeometryCheckEnable(TEXT("p.Chaos.ConvexGeometryCheckEnable"), FConvexBuilder::PerformGeometryCheck, TEXT("Perform convex geometry complexity check for Chaos physics."));
 
-	FAutoConsoleVariableRef CVarConvexGeometrySimplifyEnable(TEXT("p.Chaos.PerformGeometryReduction"), TConvexBuilder<float>::PerformGeometryReduction, TEXT("Perform convex geometry simplification to increase performance in Chaos physics."));
+	FAutoConsoleVariableRef CVarConvexGeometrySimplifyEnable(TEXT("p.Chaos.PerformGeometryReduction"), FConvexBuilder::PerformGeometryReduction, TEXT("Perform convex geometry simplification to increase performance in Chaos physics."));
 
-	FAutoConsoleVariableRef CVarConvexParticlesWarningThreshold(TEXT("p.Chaos.ConvexParticlesWarningThreshold"), TConvexBuilder<float>::ParticlesThreshold, TEXT("Threshold beyond which we warn about collision geometry complexity."));
+	FAutoConsoleVariableRef CVarConvexParticlesWarningThreshold(TEXT("p.Chaos.ConvexParticlesWarningThreshold"), FConvexBuilder::ParticlesThreshold, TEXT("Threshold beyond which we warn about collision geometry complexity."));
 }
\ No newline at end of file
diff --git a/Engine/Source/Runtime/Experimental/Chaos/Private/Chaos/CollisionResolution.cpp b/Engine/Source/Runtime/Experimental/Chaos/Private/Chaos/CollisionResolution.cpp
--- a/Engine/Source/Runtime/Experimental/Chaos/Private/Chaos/CollisionResolution.cpp
+++ b/Engine/Source/Runtime/Experimental/Chaos/Private/Chaos/CollisionResolution.cpp
@@ -537,18 +537,18 @@ namespace Chaos
 
 			if (ensure(GetInnerType(A.GetType()) == ImplicitObjectType::Convex && GetInnerType(B.GetType()) == ImplicitObjectType::Convex))
 			{
-				const TConvex<T, d>* AObject = nullptr;
-				const TConvex<T, d>* BObject = nullptr;
+				const FConvex* AObject = nullptr;
+				const FConvex* BObject = nullptr;
 
 				if (IsScaled(A.GetType()) && IsScaled(B.GetType()))
 				{
-					AObject = static_cast<const TConvex<T, d>*>(static_cast<const TImplicitObjectScaled<TConvex<T, d>>&>(A).Object().Get());
-					BObject = static_cast<const TConvex<T, d>*>(static_cast<const TImplicitObjectScaled<TConvex<T, d>>&>(B).Object().Get());
+					AObject = static_cast<const FConvex*>(static_cast<const TImplicitObjectScaled<FConvex>&>(A).Object().Get());
+					BObject = static_cast<const FConvex*>(static_cast<const TImplicitObjectScaled<FConvex>&>(B).Object().Get());
 				}
 				else if (A.GetType() == ImplicitObjectType::Transformed && B.GetType() == ImplicitObjectType::Transformed)
 				{
-					AObject = static_cast<const TConvex<T, d>*>(static_cast<const TImplicitObjectTransformed<T, d>&>(A).Object().Get());
-					BObject = static_cast<const TConvex<T, d>*>(static_cast<const TImplicitObjectTransformed<T, d>&>(B).Object().Get());
+					AObject = static_cast<const FConvex*>(static_cast<const TImplicitObjectTransformed<T, d>&>(A).Object().Get());
+					BObject = static_cast<const FConvex*>(static_cast<const TImplicitObjectTransformed<T, d>&>(B).Object().Get());
 				}
 
 				if (AObject && BObject)
diff --git a/Engine/Source/Runtime/Experimental/Chaos/Private/Chaos/HeightField.cpp b/Engine/Source/Runtime/Experimental/Chaos/Private/Chaos/HeightField.cpp
--- a/Engine/Source/Runtime/Experimental/Chaos/Private/Chaos/HeightField.cpp
+++ b/Engine/Source/Runtime/Experimental/Chaos/Private/Chaos/HeightField.cpp
@@ -1137,7 +1137,7 @@ namespace Chaos
 	}
 
 	template <typename T>
-	bool THeightField<T>::OverlapGeom(const TConvex<T, 3>& QueryGeom, const TRigidTransform<T, 3>& QueryTM, const T Thickness) const
+	bool THeightField<T>::OverlapGeom(const FConvex& QueryGeom, const TRigidTransform<T, 3>& QueryTM, const T Thickness) const
 	{
 		return OverlapGeomImp(QueryGeom, QueryTM, Thickness);
 	}
@@ -1161,7 +1161,7 @@ namespace Chaos
 	}
 
 	template <typename T>
-	bool THeightField<T>::OverlapGeom(const TImplicitObjectScaled<TConvex<T, 3>>& QueryGeom, const TRigidTransform<T, 3>& QueryTM, const T Thickness) const
+	bool THeightField<T>::OverlapGeom(const TImplicitObjectScaled<FConvex>& QueryGeom, const TRigidTransform<T, 3>& QueryTM, const T Thickness) const
 	{
 		return OverlapGeomImp(QueryGeom, QueryTM, Thickness);
 	}
@@ -1209,7 +1209,7 @@ namespace Chaos
 	}
 
 	template <typename T>
-	bool THeightField<T>::SweepGeom(const TConvex<T, 3>& QueryGeom, const TRigidTransform<T, 3>& StartTM, const TVector<T, 3>& Dir, const T Length, T& OutTime, TVector<T, 3>& OutPosition, TVector<T, 3>& OutNormal, int32& OutFaceIndex, const T Thickness, bool bComputeMTD) const
+	bool THeightField<T>::SweepGeom(const FConvex& QueryGeom, const TRigidTransform<T, 3>& StartTM, const TVector<T, 3>& Dir, const T Length, T& OutTime, TVector<T, 3>& OutPosition, TVector<T, 3>& OutNormal, int32& OutFaceIndex, const T Thickness, bool bComputeMTD) const
 	{
 		return SweepGeomImp(QueryGeom, StartTM, Dir, Length, OutTime, OutPosition, OutNormal, OutFaceIndex, Thickness, bComputeMTD);
 	}
@@ -1233,7 +1233,7 @@ namespace Chaos
 	}
 
 	template <typename T>
-	bool THeightField<T>::SweepGeom(const TImplicitObjectScaled<TConvex<T, 3>>& QueryGeom, const TRigidTransform<T, 3>& StartTM, const TVector<T, 3>& Dir, const T Length, T& OutTime, TVector<T, 3>& OutPosition, TVector<T, 3>& OutNormal, int32& OutFaceIndex, const T Thickness, bool bComputeMTD) const
+	bool THeightField<T>::SweepGeom(const TImplicitObjectScaled<FConvex>& QueryGeom, const TRigidTransform<T, 3>& StartTM, const TVector<T, 3>& Dir, const T Length, T& OutTime, TVector<T, 3>& OutPosition, TVector<T, 3>& OutNormal, int32& OutFaceIndex, const T Thickness, bool bComputeMTD) const
 	{
 		return SweepGeomImp(QueryGeom, StartTM, Dir, Length, OutTime, OutPosition, OutNormal, OutFaceIndex, Thickness, bComputeMTD);
 	}
diff --git a/Engine/Source/Runtime/Experimental/Chaos/Private/Chaos/ImplicitObject.cpp b/Engine/Source/Runtime/Experimental/Chaos/Private/Chaos/ImplicitObject.cpp
--- a/Engine/Source/Runtime/Experimental/Chaos/Private/Chaos/ImplicitObject.cpp
+++ b/Engine/Source/Runtime/Experimental/Chaos/Private/Chaos/ImplicitObject.cpp
@@ -349,7 +349,7 @@ FImplicitObject* FImplicitObject::SerializationFactory(FChaosArchive& Ar, FImpli
 			EImplicitObjectType InnerType = GetInnerType(ObjectType);
 			switch (InnerType)
 			{
-			case ImplicitObjectType::Convex: if (Ar.IsLoading()) { return new TImplicitObjectScaled<TConvex<FReal, 3>>(); } break;
+			case ImplicitObjectType::Convex: if (Ar.IsLoading()) { return new TImplicitObjectScaled<FConvex>(); } break;
 			case ImplicitObjectType::TriangleMesh: if (Ar.IsLoading()) { return new TImplicitObjectScaled<TTriangleMeshImplicitObject<FReal>>(); } break;
 			default: check(false);
 			}
@@ -367,7 +367,7 @@ FImplicitObject* FImplicitObject::SerializationFactory(FChaosArchive& Ar, FImpli
 	case ImplicitObjectType::Transformed: if (Ar.IsLoading()) { return new TImplicitObjectTransformed<FReal, 3>(); } break;
 	case ImplicitObjectType::Union: if (Ar.IsLoading()) { return new TImplicitObjectUnion<FReal, 3>(); } break;
 	case ImplicitObjectType::LevelSet: if (Ar.IsLoading()) { return new TLevelSet<FReal, 3>(); } break;
-	case ImplicitObjectType::Convex: if (Ar.IsLoading()) { return new TConvex<FReal, 3>(); } break;
+	case ImplicitObjectType::Convex: if (Ar.IsLoading()) { return new FConvex(); } break;
 	case ImplicitObjectType::TaperedCylinder: if (Ar.IsLoading()) { return new TTaperedCylinder<FReal>(); } break;
 	case ImplicitObjectType::TriangleMesh: if (Ar.IsLoading()) { return new TTriangleMeshImplicitObject<FReal>(); } break;
 	case ImplicitObjectType::DEPRECATED_Scaled:
diff --git a/Engine/Source/Runtime/Experimental/Chaos/Private/Chaos/TriangleMeshImplicitObject.cpp b/Engine/Source/Runtime/Experimental/Chaos/Private/Chaos/TriangleMeshImplicitObject.cpp
--- a/Engine/Source/Runtime/Experimental/Chaos/Private/Chaos/TriangleMeshImplicitObject.cpp
+++ b/Engine/Source/Runtime/Experimental/Chaos/Private/Chaos/TriangleMeshImplicitObject.cpp
@@ -336,7 +336,7 @@ bool TTriangleMeshImplicitObject<T>::OverlapGeom(const TCapsule<T>& QueryGeom, c
 }
 
 template <typename T>
-bool TTriangleMeshImplicitObject<T>::OverlapGeom(const TConvex<T, 3>& QueryGeom, const TRigidTransform<T, 3>& QueryTM, const T Thickness) const
+bool TTriangleMeshImplicitObject<T>::OverlapGeom(const FConvex& QueryGeom, const TRigidTransform<T, 3>& QueryTM, const T Thickness) const
 {
 	return OverlapGeomImp(QueryGeom, QueryTM, Thickness);
 }
@@ -360,13 +360,13 @@ bool TTriangleMeshImplicitObject<T>::OverlapGeom(const TImplicitObjectScaled<TCa
 }
 
 template <typename T>
-bool TTriangleMeshImplicitObject<T>::OverlapGeom(const TImplicitObjectScaled<TConvex<T, 3>>& QueryGeom, const TRigidTransform<T, 3>& QueryTM, const T Thickness) const
+bool TTriangleMeshImplicitObject<T>::OverlapGeom(const TImplicitObjectScaled<FConvex>& QueryGeom, const TRigidTransform<T, 3>& QueryTM, const T Thickness) const
 {
 	return OverlapGeomImp(QueryGeom, QueryTM, Thickness);
 }
 
 template <typename T>
-bool TTriangleMeshImplicitObject<T>::OverlapGeom(const TImplicitObjectScaled<TImplicitObjectScaled<TConvex<T, 3>>>& QueryGeom, const TRigidTransform<T, 3>& QueryTM, const T Thickness) const
+bool TTriangleMeshImplicitObject<T>::OverlapGeom(const TImplicitObjectScaled<TImplicitObjectScaled<FConvex>>& QueryGeom, const TRigidTransform<T, 3>& QueryTM, const T Thickness) const
 {
 	return OverlapGeomImp(QueryGeom, QueryTM, Thickness);
 }
@@ -490,7 +490,7 @@ bool TTriangleMeshImplicitObject<T>::SweepGeom(const TCapsule<T>& QueryGeom, con
 }
 
 template <typename T>
-bool TTriangleMeshImplicitObject<T>::SweepGeom(const TConvex<T, 3>& QueryGeom, const TRigidTransform<T, 3>& StartTM, const TVector<T, 3>& Dir, const T Length, T& OutTime, TVector<T, 3>& OutPosition, TVector<T, 3>& OutNormal, int32& OutFaceIndex, const T Thickness, const bool bComputeMTD) const
+bool TTriangleMeshImplicitObject<T>::SweepGeom(const FConvex& QueryGeom, const TRigidTransform<T, 3>& StartTM, const TVector<T, 3>& Dir, const T Length, T& OutTime, TVector<T, 3>& OutPosition, TVector<T, 3>& OutNormal, int32& OutFaceIndex, const T Thickness, const bool bComputeMTD) const
 {
 	return SweepGeomImp(QueryGeom, StartTM, Dir, Length, OutTime, OutPosition, OutNormal, OutFaceIndex, Thickness, bComputeMTD);
 }
@@ -514,7 +514,7 @@ bool TTriangleMeshImplicitObject<T>::SweepGeom(const TImplicitObjectScaled<TCaps
 }
 
 template <typename T>
-bool TTriangleMeshImplicitObject<T>::SweepGeom(const TImplicitObjectScaled<TConvex<T, 3>>& QueryGeom, const TRigidTransform<T, 3>& StartTM, const TVector<T, 3>& Dir, const T Length, T& OutTime, TVector<T, 3>& OutPosition, TVector<T, 3>& OutNormal, int32& OutFaceIndex, const T Thickness, const bool bComputeMTD) const
+bool TTriangleMeshImplicitObject<T>::SweepGeom(const TImplicitObjectScaled<FConvex>& QueryGeom, const TRigidTransform<T, 3>& StartTM, const TVector<T, 3>& Dir, const T Length, T& OutTime, TVector<T, 3>& OutPosition, TVector<T, 3>& OutNormal, int32& OutFaceIndex, const T Thickness, const bool bComputeMTD) const
 {
 	return SweepGeomImp(QueryGeom, StartTM, Dir, Length, OutTime, OutPosition, OutNormal, OutFaceIndex, Thickness, bComputeMTD);
 }
diff --git a/Engine/Source/Runtime/Experimental/Chaos/Public/Chaos/CollisionConvexMesh.h b/Engine/Source/Runtime/Experimental/Chaos/Public/Chaos/CollisionConvexMesh.h
--- a/Engine/Source/Runtime/Experimental/Chaos/Public/Chaos/CollisionConvexMesh.h
+++ b/Engine/Source/Runtime/Experimental/Chaos/Public/Chaos/CollisionConvexMesh.h
@@ -9,23 +9,22 @@
 
 namespace Chaos
 {
-	template <typename T>
-	class TConvexBuilder
+	class FConvexBuilder
 	{
 	public:
-		static bool IsValidTriangle(const TVector<T, 3>& A, const TVector<T, 3>& B, const TVector<T, 3>& C)
+		static bool IsValidTriangle(const FVec3& A, const FVec3& B, const FVec3& C)
 		{
-			const TVector<T, 3> BA = B - A;
-			const TVector<T, 3> CA = C - A;
-			const TVector<T, 3> Cross = TVector<T, 3>::CrossProduct(BA, CA);
+			const FVec3 BA = B - A;
+			const FVec3 CA = C - A;
+			const FVec3 Cross = FVec3::CrossProduct(BA, CA);
 			return Cross.Size() > 1e-4;
 		}
 
-		static void Build(const TParticles<T, 3>& InParticles, TArray <TPlane<T, 3>>& OutPlanes, TArray<TArray<int32>>& OutFaceIndices, TParticles<T, 3>& OutSurfaceParticles, TBox<T, 3>& OutLocalBounds)
+		static void Build(const TParticles<FReal, 3>& InParticles, TArray <TPlane<FReal, 3>>& OutPlanes, TArray<TArray<int32>>& OutFaceIndices, TParticles<FReal, 3>& OutSurfaceParticles, TBox<FReal, 3>& OutLocalBounds)
 		{
 			OutPlanes.Reset();
 			OutSurfaceParticles.Resize(0);
-			OutLocalBounds = TBox<T, 3>::EmptyBox();
+			OutLocalBounds = TBox<FReal, 3>::EmptyBox();
 
 			const uint32 NumParticles = InParticles.Size();
 			if(NumParticles == 0)
@@ -33,7 +32,7 @@ namespace Chaos
 				return;
 			}
 
-			OutLocalBounds = TBox<T, 3>(InParticles.X(0), InParticles.X(0));
+			OutLocalBounds = TBox<FReal, 3>(InParticles.X(0), InParticles.X(0));
 			for(uint32 ParticleIndex = 0; ParticleIndex < NumParticles; ++ParticleIndex)
 			{
 				OutLocalBounds.GrowToInclude(InParticles.X(ParticleIndex));
@@ -59,9 +58,9 @@ namespace Chaos
 
 				for(const TVector<int32, 3>& Idx : Indices)
 				{
-					TVector<T, 3> Vs[3] = {InParticles.X(Idx[0]), InParticles.X(Idx[1]), InParticles.X(Idx[2])};
-					const TVector<T, 3> Normal = TVector<T, 3>::CrossProduct(Vs[1] - Vs[0], Vs[2] - Vs[0]).GetUnsafeNormal();
-					OutPlanes.Add(TPlane<T, 3>(Vs[0], Normal));
+					FVec3 Vs[3] = {InParticles.X(Idx[0]), InParticles.X(Idx[1]), InParticles.X(Idx[2])};
+					const FVec3 Normal = FVec3::CrossProduct(Vs[1] - Vs[0], Vs[2] - Vs[0]).GetUnsafeNormal();
+					OutPlanes.Add(TPlane<FReal, 3>(Vs[0], Normal));
 					TArray<int32> FaceIndices;
 					FaceIndices.SetNum(3);
 					FaceIndices[0] = AddIndex(Idx[0]);
@@ -82,10 +81,10 @@ namespace Chaos
 				const bool bIsValidTriangle = IsValidTriangle(InParticles.X(0), InParticles.X(1), InParticles.X(2));
 
 				//TODO_SQ_IMPLEMENTATION: should do proper cleanup to avoid this
-				if (ensureMsgf(bIsValidTriangle, TEXT("TConvexBuilder::Build(): Generated invalid triangle!")))
+				if (ensureMsgf(bIsValidTriangle, TEXT("FConvexBuilder::Build(): Generated invalid triangle!")))
 				{
-					TVector<T, 3> Normal = TVector<T, 3>::CrossProduct(InParticles.X(1) - InParticles.X(0), InParticles.X(2) - InParticles.X(0)).GetSafeNormal();
-					OutPlanes.Add(TPlane<T, 3>(InParticles.X(0), Normal));
+					FVec3 Normal = FVec3::CrossProduct(InParticles.X(1) - InParticles.X(0), InParticles.X(2) - InParticles.X(0)).GetSafeNormal();
+					OutPlanes.Add(TPlane<FReal, 3>(InParticles.X(0), Normal));
 					OutSurfaceParticles.AddParticles(3);
 					OutSurfaceParticles.X(0) = InParticles.X(0);
 					OutSurfaceParticles.X(1) = InParticles.X(1);
@@ -101,7 +100,7 @@ namespace Chaos
 			}
 		}
 
-		static void BuildConvexHull(const TParticles<T, 3>& InParticles, TArray<TVector<int32, 3>>& OutIndices)
+		static void BuildConvexHull(const TParticles<FReal, 3>& InParticles, TArray<TVector<int32, 3>>& OutIndices)
 		{
 			OutIndices.Reset();
 			FConvexFace* Faces = BuildInitialHull(InParticles);
@@ -135,11 +134,11 @@ namespace Chaos
 			}
 		}
 
-		static TTriangleMesh<T> BuildConvexHullTriMesh(const TParticles<T, 3>& InParticles)
+		static TTriangleMesh<FReal> BuildConvexHullTriMesh(const TParticles<FReal, 3>& InParticles)
 		{
 			TArray<TVector<int32, 3>> Indices;
 			BuildConvexHull(InParticles, Indices);
-			return TTriangleMesh<T>(MoveTemp(Indices));
+			return TTriangleMesh<FReal>(MoveTemp(Indices));
 		}
 
 		static CHAOS_API bool IsPerformanceWarning(int32 NumPlanes, int32 NumParticles)
@@ -162,7 +161,7 @@ namespace Chaos
 			return FString::Printf(TEXT("Planes %d, SurfaceParticles %d"), NumPlanes, NumParticles);
 		}
 
-		static CHAOS_API void Simplify(TArray <TPlane<T, 3>>& InOutPlanes, TArray<TArray<int32>>& InOutFaces, TParticles<T, 3>& InOutParticles, TBox<T, 3>& InOutLocalBounds)
+		static CHAOS_API void Simplify(TArray <TPlane<FReal, 3>>& InOutPlanes, TArray<TArray<int32>>& InOutFaces, TParticles<FReal, 3>& InOutParticles, TBox<FReal, 3>& InOutLocalBounds)
 		{
 			struct TPair
 			{
@@ -176,7 +175,7 @@ namespace Chaos
 			int32 NumToDelete = NumberOfParticlesWeHave - NumberOfParticlesRequired;
 
 			uint32 Size = InOutParticles.Size();
-			TArray<TVector<T, 3>> Particles;
+			TArray<FVec3> Particles;
 			Particles.AddUninitialized(Size);
 			for (uint32 A = 0; A < Size; A++)
 			{
@@ -202,7 +201,7 @@ namespace Chaos
 							{
 								if (!IsDeleted[B])
 								{
-									TVector<T, 3> Vec = Particles[A] - Particles[B];
+									FVec3 Vec = Particles[A] - Particles[B];
 									float LengthSqr = Vec.SizeSquared();
 									if (LengthSqr < ClosestDistSqr)
 									{
@@ -224,7 +223,7 @@ namespace Chaos
 				}
 			}
 
-			TParticles<T, 3> TmpParticles;
+			TParticles<FReal, 3> TmpParticles;
 			for (int Idx = 0; Idx < Particles.Num(); Idx++)
 			{
 				// Only add particles that have not been merged away
@@ -255,7 +254,7 @@ namespace Chaos
 		struct FHalfEdge;
 		struct FConvexFace
 		{
-			FConvexFace(const TPlane<T, 3>& FacePlane)
+			FConvexFace(const TPlane<FReal, 3>& FacePlane)
 				: ConflictList(nullptr)
 				, Plane(FacePlane)
 			{
@@ -263,7 +262,7 @@ namespace Chaos
 
 			FHalfEdge* FirstEdge;
 			FHalfEdge* ConflictList; //Note that these half edges are really just free verts grouped together
-			TPlane<T, 3> Plane;
+			TPlane<FReal, 3> Plane;
 			FConvexFace* Prev;
 			FConvexFace* Next; //these have no geometric meaning, just used for book keeping
 		};
@@ -280,12 +279,12 @@ namespace Chaos
 			FConvexFace* Face;
 		};
 
-		static TVector<T, 3> ComputeFaceNormal(const TVector<T, 3>& A, const TVector<T, 3>& B, const TVector<T, 3>& C)
+		static FVec3 ComputeFaceNormal(const FVec3& A, const FVec3& B, const FVec3& C)
 		{
-			return TVector<T, 3>::CrossProduct((B - A), (C - A));
+			return FVec3::CrossProduct((B - A), (C - A));
 		}
 
-		static FConvexFace* CreateFace(const TParticles<T, 3>& InParticles, FHalfEdge* RS, FHalfEdge* ST, FHalfEdge* TR)
+		static FConvexFace* CreateFace(const TParticles<FReal, 3>& InParticles, FHalfEdge* RS, FHalfEdge* ST, FHalfEdge* TR)
 		{
 			RS->Prev = TR;
 			RS->Next = ST;
@@ -293,11 +292,11 @@ namespace Chaos
 			ST->Next = TR;
 			TR->Prev = ST;
 			TR->Next = RS;
-			TVector<T, 3> RSTNormal = ComputeFaceNormal(InParticles.X(RS->Vertex), InParticles.X(ST->Vertex), InParticles.X(TR->Vertex));
-			const T RSTNormalSize = RSTNormal.Size();
+			FVec3 RSTNormal = ComputeFaceNormal(InParticles.X(RS->Vertex), InParticles.X(ST->Vertex), InParticles.X(TR->Vertex));
+			const FReal RSTNormalSize = RSTNormal.Size();
 			check(RSTNormalSize > 1e-4);
 			RSTNormal = RSTNormal * (1 / RSTNormalSize);
-			FConvexFace* RST = new FConvexFace(TPlane<T, 3>(InParticles.X(RS->Vertex), RSTNormal));
+			FConvexFace* RST = new FConvexFace(TPlane<FReal, 3>(InParticles.X(RS->Vertex), RSTNormal));
 			RST->FirstEdge = RS;
 			RS->Face = RST;
 			ST->Face = RST;
@@ -305,16 +304,16 @@ namespace Chaos
 			return RST;
 		}
 
-		static void StealConflictList(const TParticles<T, 3>& InParticles, FHalfEdge* OldList, FConvexFace** Faces, int32 NumFaces)
+		static void StealConflictList(const TParticles<FReal, 3>& InParticles, FHalfEdge* OldList, FConvexFace** Faces, int32 NumFaces)
 		{
 			FHalfEdge* Cur = OldList;
 			while(Cur)
 			{
-				T MaxD = 1e-4;
+				FReal MaxD = 1e-4;
 				int32 MaxIdx = -1;
 				for(int32 Idx = 0; Idx < NumFaces; ++Idx)
 				{
-					T Distance = Faces[Idx]->Plane.SignedDistance(InParticles.X(Cur->Vertex));
+					FReal Distance = Faces[Idx]->Plane.SignedDistance(InParticles.X(Cur->Vertex));
 					if(Distance > MaxD)
 					{
 						MaxD = Distance;
@@ -327,7 +326,7 @@ namespace Chaos
 				{
 					//let's make sure faces created with this new conflict vertex will be valid. The plane check above is not sufficient because long thin triangles will have a plane with its point at one of these. Combined with normal and precision we can have errors
 					auto PretendNormal = [&InParticles](FHalfEdge* A, FHalfEdge* B, FHalfEdge* C) {
-						return TVector<T, 3>::CrossProduct(InParticles.X(B->Vertex) - InParticles.X(A->Vertex), InParticles.X(C->Vertex) - InParticles.X(A->Vertex)).SizeSquared();
+						return FVec3::CrossProduct(InParticles.X(B->Vertex) - InParticles.X(A->Vertex), InParticles.X(C->Vertex) - InParticles.X(A->Vertex)).SizeSquared();
 					};
 					FHalfEdge* Edge = Faces[MaxIdx]->FirstEdge;
 					do
@@ -364,21 +363,21 @@ namespace Chaos
 			}
 		}
 
-		static FConvexFace* BuildInitialHull(const TParticles<T, 3>& InParticles)
+		static FConvexFace* BuildInitialHull(const TParticles<FReal, 3>& InParticles)
 		{
 			if(InParticles.Size() < 4) //not enough points
 			{
 				return nullptr;
 			}
 
-			constexpr T Epsilon = 1e-4;
+			constexpr FReal Epsilon = 1e-4;
 
 			const int32 NumParticles = InParticles.Size();
 
 			//We store the vertex directly in the half-edge. We use its next to group free vertices by context list
 			//create a starting triangle by finding min/max on X and max on Y
-			T MinX = TNumericLimits<T>::Max();
-			T MaxX = TNumericLimits<T>::Lowest();
+			FReal MinX = TNumericLimits<FReal>::Max();
+			FReal MaxX = TNumericLimits<FReal>::Lowest();
 			FHalfEdge* A = nullptr; //min x
 			FHalfEdge* B = nullptr; //max x
 			FHalfEdge DummyHalfEdge(-1);
@@ -392,7 +391,7 @@ namespace Chaos
 				Prev->Next = VHalf;
 				VHalf->Prev = Prev;
 				VHalf->Next = nullptr;
-				const TVector<T, 3>& V = InParticles.X(i);
+				const FVec3& V = InParticles.X(i);
 
 				if(V[0] < MinX)
 				{
@@ -427,12 +426,12 @@ namespace Chaos
 			}
 
 			//find C so that we get the biggest base triangle
-			T MaxTriSize = Epsilon;
-			const TVector<T, 3> AToB = InParticles.X(B->Vertex) - InParticles.X(A->Vertex);
+			FReal MaxTriSize = Epsilon;
+			const FVec3 AToB = InParticles.X(B->Vertex) - InParticles.X(A->Vertex);
 			FHalfEdge* C = nullptr;
 			for(FHalfEdge* V = DummyHalfEdge.Next; V; V = V->Next)
 			{
-				T TriSize = TVector<T, 3>::CrossProduct(AToB, InParticles.X(V->Vertex) - InParticles.X(A->Vertex)).SizeSquared();
+				FReal TriSize = FVec3::CrossProduct(AToB, InParticles.X(V->Vertex) - InParticles.X(A->Vertex)).SizeSquared();
 				if(TriSize > MaxTriSize)
 				{
 					MaxTriSize = TriSize;
@@ -453,16 +452,16 @@ namespace Chaos
 			}
 
 			//find farthest D along normal
-			const TVector<T, 3> AToC = InParticles.X(C->Vertex) - InParticles.X(A->Vertex);
-			const TVector<T, 3> Normal = TVector<T, 3>::CrossProduct(AToB, AToC);
+			const FVec3 AToC = InParticles.X(C->Vertex) - InParticles.X(A->Vertex);
+			const FVec3 Normal = FVec3::CrossProduct(AToB, AToC);
 
-			T MaxPosDistance = Epsilon;
-			T MaxNegDistance = Epsilon;
+			FReal MaxPosDistance = Epsilon;
+			FReal MaxNegDistance = Epsilon;
 			FHalfEdge* PosD = nullptr;
 			FHalfEdge* NegD = nullptr;
 			for(FHalfEdge* V = DummyHalfEdge.Next; V; V = V->Next)
 			{
-				T Dot = TVector<T, 3>::DotProduct(InParticles.X(V->Vertex) - InParticles.X(A->Vertex), Normal);
+				FReal Dot = FVec3::DotProduct(InParticles.X(V->Vertex) - InParticles.X(A->Vertex), Normal);
 				if(Dot > MaxPosDistance)
 				{
 					MaxPosDistance = Dot;
@@ -534,16 +533,16 @@ namespace Chaos
 			return Faces[0];
 		}
 
-		static FHalfEdge* FindConflictVertex(const TParticles<T, 3>& InParticles, FConvexFace* FaceList)
+		static FHalfEdge* FindConflictVertex(const TParticles<FReal, 3>& InParticles, FConvexFace* FaceList)
 		{
 			for(FConvexFace* CurFace = FaceList; CurFace; CurFace = CurFace->Next)
 			{
-				T MaxD = TNumericLimits<T>::Lowest();
+				FReal MaxD = TNumericLimits<FReal>::Lowest();
 				FHalfEdge* MaxV = nullptr;
 				for(FHalfEdge* CurFaceVertex = CurFace->ConflictList; CurFaceVertex; CurFaceVertex = CurFaceVertex->Next)
 				{
 					//is it faster to cache this from stealing stage?
-					T Dist = TVector<T, 3>::DotProduct(InParticles.X(CurFaceVertex->Vertex), CurFace->Plane.Normal());
+					FReal Dist = FVec3::DotProduct(InParticles.X(CurFaceVertex->Vertex), CurFace->Plane.Normal());
 					if(Dist > MaxD)
 					{
 						MaxD = Dist;
@@ -573,13 +572,13 @@ namespace Chaos
 			return nullptr;
 		}
 
-		static void BuildHorizon(const TParticles<T, 3>& InParticles, FHalfEdge* ConflictV, TArray<FHalfEdge*>& HorizonEdges, TArray<FConvexFace*>& FacesToDelete)
+		static void BuildHorizon(const TParticles<FReal, 3>& InParticles, FHalfEdge* ConflictV, TArray<FHalfEdge*>& HorizonEdges, TArray<FConvexFace*>& FacesToDelete)
 		{
 			//We must flood fill from the initial face and mark edges of faces the conflict vertex cannot see
 			//In order to return a CCW ordering we must traverse each face in CCW order from the edge we crossed over
 			//This should already be the ordering in the half edge
-			const T Epsilon = 1e-1;
-			const TVector<T, 3> V = InParticles.X(ConflictV->Vertex);
+			const FReal Epsilon = 1e-1;
+			const FVec3 V = InParticles.X(ConflictV->Vertex);
 			TSet<FConvexFace*> Processed;
 			TArray<FHalfEdge*> Queue;
 			check(ConflictV->Face);
@@ -597,7 +596,7 @@ namespace Chaos
 				{
 					continue;
 				}
-				const T Distance = NextFace->Plane.SignedDistance(V);
+				const FReal Distance = NextFace->Plane.SignedDistance(V);
 				if(Distance > Epsilon)
 				{
 					Queue.Add(Twin->Prev); //stack pops so reverse order
@@ -611,7 +610,7 @@ namespace Chaos
 			}
 		}
 
-		static void BuildFaces(const TParticles<T, 3>& InParticles, const FHalfEdge* ConflictV, const TArray<FHalfEdge*>& HorizonEdges, const TArray<FConvexFace*> OldFaces, TArray<FConvexFace*>& NewFaces)
+		static void BuildFaces(const TParticles<FReal, 3>& InParticles, const FHalfEdge* ConflictV, const TArray<FHalfEdge*>& HorizonEdges, const TArray<FConvexFace*> OldFaces, TArray<FConvexFace*>& NewFaces)
 		{
 			//The HorizonEdges are in CCW order. We must make new faces and edges to join from ConflictV to these edges
 			check(HorizonEdges.Num() >= 3);
@@ -671,7 +670,7 @@ namespace Chaos
 			StartFace->Prev = OldFace;
 		}
 
-		static void AddVertex(const TParticles<T, 3>& InParticles, FHalfEdge* ConflictV)
+		static void AddVertex(const TParticles<FReal, 3>& InParticles, FHalfEdge* ConflictV)
 		{
 			TArray<FHalfEdge*> HorizonEdges;
 			TArray<FConvexFace*> FacesToDelete;
diff --git a/Engine/Source/Runtime/Experimental/Chaos/Public/Chaos/Convex.h b/Engine/Source/Runtime/Experimental/Chaos/Public/Chaos/Convex.h
--- a/Engine/Source/Runtime/Experimental/Chaos/Public/Chaos/Convex.h
+++ b/Engine/Source/Runtime/Experimental/Chaos/Public/Chaos/Convex.h
@@ -7,38 +7,40 @@
 #include "CollisionConvexMesh.h"
 #include "ChaosArchive.h"
 #include "GJK.h"
+#include "ChaosCheck.h"
 
 namespace Chaos
 {
-	template<class T, int d>
-	class TConvex final : public FImplicitObject
+	class FConvex final : public FImplicitObject
 	{
 	public:
 		using FImplicitObject::GetTypeName;
+		using TType = float;
+		static constexpr unsigned D = 3;
 
-		TConvex()
+		FConvex()
 		    : FImplicitObject(EImplicitObject::IsConvex | EImplicitObject::HasBoundingBox, ImplicitObjectType::Convex)
 			, Volume(0.f)
 			, CenterOfMass(FVec3(0.f))
 		{}
-		TConvex(const TConvex&) = delete;
-		TConvex(TConvex&& Other)
+		FConvex(const FConvex&) = delete;
+		FConvex(FConvex&& Other)
 		    : FImplicitObject(EImplicitObject::IsConvex | EImplicitObject::HasBoundingBox, ImplicitObjectType::Convex)
 			, Planes(MoveTemp(Other.Planes))
 		    , SurfaceParticles(MoveTemp(Other.SurfaceParticles))
 		    , LocalBoundingBox(MoveTemp(Other.LocalBoundingBox))
-			, Volume(MoveTemp(Other.MoveTemp))
+			, Volume(MoveTemp(Other.Volume))
 			, CenterOfMass(MoveTemp(Other.CenterOfMass))
 		{}
 
 		// NOTE: This constructor will result in approximate COM and volume calculations, since it does
 		// not have face indices for surface particles.
 		// TODO: Keep track of invalid state and ensure on volume or COM access?
-		TConvex(TArray<TPlane<T, d>>&& InPlanes, TParticles<T, d>&& InSurfaceParticles)
+		FConvex(TArray<TPlane<FReal, 3>>&& InPlanes, TParticles<FReal, 3>&& InSurfaceParticles)
 		    : FImplicitObject(EImplicitObject::IsConvex | EImplicitObject::HasBoundingBox, ImplicitObjectType::Convex)
 			, Planes(MoveTemp(InPlanes))
 		    , SurfaceParticles(MoveTemp(InSurfaceParticles))
-		    , LocalBoundingBox(TBox<T, d>::EmptyBox())
+		    , LocalBoundingBox(TBox<FReal, 3>::EmptyBox())
 		{
 			for (uint32 ParticleIndex = 0; ParticleIndex < SurfaceParticles.Size(); ++ParticleIndex)
 			{
@@ -50,7 +52,7 @@ namespace Chaos
 			Volume = LocalBoundingBox.GetVolume();
 		}
 
-		TConvex(const TParticles<T, 3>& InParticles)
+		FConvex(const TParticles<FReal, 3>& InParticles)
 		    : FImplicitObject(EImplicitObject::IsConvex | EImplicitObject::HasBoundingBox, ImplicitObjectType::Convex)
 		{
 			const uint32 NumParticles = InParticles.Size();
@@ -60,7 +62,7 @@ namespace Chaos
 			}
 
 			TArray<TArray<int32>> FaceIndices;
-			TConvexBuilder<T>::Build(InParticles, Planes, FaceIndices, SurfaceParticles, LocalBoundingBox);
+			FConvexBuilder::Build(InParticles, Planes, FaceIndices, SurfaceParticles, LocalBoundingBox);
 			CHAOS_ENSURE(Planes.Num() == FaceIndices.Num());
 			CalculateVolumeAndCenterOfMass(SurfaceParticles, FaceIndices, Volume, CenterOfMass);
 		}
@@ -70,12 +72,12 @@ namespace Chaos
 			return ImplicitObjectType::Convex;
 		}
 
-		virtual const TBox<T, d>& BoundingBox() const override
+		virtual const TBox<FReal, 3>& BoundingBox() const override
 		{
 			return LocalBoundingBox;
 		}
 
-		virtual T PhiWithNormal(const TVector<T, d>& x, TVector<T, d>& Normal) const override
+		virtual FReal PhiWithNormal(const FVec3& x, FVec3& Normal) const override
 		{
 			const int32 NumPlanes = Planes.Num();
 			if (NumPlanes == 0)
@@ -84,12 +86,12 @@ namespace Chaos
 			}
 			check(NumPlanes > 0);
 
-			T MaxPhi = TNumericLimits<T>::Lowest();
+			FReal MaxPhi = TNumericLimits<FReal>::Lowest();
 			int32 MaxPlane = 0;
 
 			for (int32 Idx = 0; Idx < NumPlanes; ++Idx)
 			{
-				const T Phi = Planes[Idx].SignedDistance(x);
+				const FReal Phi = Planes[Idx].SignedDistance(x);
 				if (Phi > MaxPhi)
 				{
 					MaxPhi = Phi;
@@ -105,18 +107,18 @@ namespace Chaos
 		 * and \c OutNormal should be.  The burden for detecting this case is deferred to the
 		 * caller. 
 		 */
-		virtual bool Raycast(const TVector<T, d>& StartPoint, const TVector<T, d>& Dir, const T Length, const T Thickness, T& OutTime, TVector<T, d>& OutPosition, TVector<T, d>& OutNormal, int32& OutFaceIndex) const override
+		virtual bool Raycast(const FVec3& StartPoint, const FVec3& Dir, const FReal Length, const FReal Thickness, FReal& OutTime, FVec3& OutPosition, FVec3& OutNormal, int32& OutFaceIndex) const override
 		{
 			OutFaceIndex = INDEX_NONE;	//finding face is expensive, should be called directly by user
-			const TRigidTransform<T, d> StartTM(StartPoint, TRotation<T,d>::FromIdentity());
-			const TSphere<T, d> Sphere(TVector<T, d>(0), Thickness);
+			const FRigidTransform3 StartTM(StartPoint, FRotation3::FromIdentity());
+			const TSphere<FReal, 3> Sphere(FVec3(0), Thickness);
 			return GJKRaycast(*this, Sphere, StartTM, Dir, Length, OutTime, OutPosition, OutNormal);
 		}
 
-		virtual Pair<TVector<T, d>, bool> FindClosestIntersectionImp(const TVector<T, d>& StartPoint, const TVector<T, d>& EndPoint, const T Thickness) const override
+		virtual Pair<FVec3, bool> FindClosestIntersectionImp(const FVec3& StartPoint, const FVec3& EndPoint, const FReal Thickness) const override
 		{
 			const int32 NumPlanes = Planes.Num();
-			TArray<Pair<T, TVector<T, 3>>> Intersections;
+			TArray<Pair<FReal, FVec3>> Intersections;
 			Intersections.Reserve(FMath::Min(static_cast<int32>(NumPlanes*.1), 16)); // Was NumPlanes, which seems excessive.
 			for (int32 Idx = 0; Idx < NumPlanes; ++Idx)
 			{
@@ -126,7 +128,7 @@ namespace Chaos
 					Intersections.Add(MakePair((PlaneIntersection.First - StartPoint).SizeSquared(), PlaneIntersection.First));
 				}
 			}
-			Intersections.Sort([](const Pair<T, TVector<T, 3>>& Elem1, const Pair<T, TVector<T, 3>>& Elem2) { return Elem1.First < Elem2.First; });
+			Intersections.Sort([](const Pair<FReal, FVec3>& Elem1, const Pair<FReal, FVec3>& Elem2) { return Elem1.First < Elem2.First; });
 			for (const auto& Elem : Intersections)
 			{
 				if (this->SignedDistance(Elem.Second) < (Thickness + 1e-4))
@@ -134,22 +136,22 @@ namespace Chaos
 					return MakePair(Elem.Second, true);
 				}
 			}
-			return MakePair(TVector<T, 3>(0), false);
+			return MakePair(FVec3(0), false);
 		}
 
-		virtual int32 FindMostOpposingFace(const TVector<T, 3>& Position, const TVector<T, 3>& UnitDir, int32 HintFaceIndex, T SearchDist) const override
+		virtual int32 FindMostOpposingFace(const FVec3& Position, const FVec3& UnitDir, int32 HintFaceIndex, FReal SearchDist) const override
 		{
 			//todo: use hill climbing
 			int32 MostOpposingIdx = INDEX_NONE;
-			T MostOpposingDot = TNumericLimits<T>::Max();
+			FReal MostOpposingDot = TNumericLimits<FReal>::Max();
 			for(int32 Idx = 0; Idx < Planes.Num(); ++Idx)
 			{
-				const TPlane<T, d>& Plane = Planes[Idx];
-				const T Distance = Plane.SignedDistance(Position);
+				const TPlane<FReal, 3>& Plane = Planes[Idx];
+				const FReal Distance = Plane.SignedDistance(Position);
 				if (FMath::Abs(Distance) < SearchDist)
 				{
 					// TPlane has an override for Normal() that doesn't call PhiWithNormal().
-					const T Dot = TVector<T, d>::DotProduct(Plane.Normal(), UnitDir);
+					const FReal Dot = FVec3::DotProduct(Plane.Normal(), UnitDir);
 					if (Dot < MostOpposingDot)
 					{
 						MostOpposingDot = Dot;
@@ -161,33 +163,33 @@ namespace Chaos
 			return MostOpposingIdx;
 		}
 
-		TVector<T, d> FindGeometryOpposingNormal(const TVector<T, d>& DenormDir, int32 FaceIndex, const TVector<T, d>& OriginalNormal) const
+		FVec3 FindGeometryOpposingNormal(const FVec3& DenormDir, int32 FaceIndex, const FVec3& OriginalNormal) const
 		{
 			// For convexes, this function must be called with a face index.
 			// If this ensure is getting hit, fix the caller so that it
 			// passes in a valid face index.
 			if (ensure(FaceIndex != INDEX_NONE))
 			{
-				const TPlane<float, 3>& OpposingFace = GetFaces()[FaceIndex];
+				const TPlane<FReal, 3>& OpposingFace = GetFaces()[FaceIndex];
 				return OpposingFace.Normal();
 			}
-			return TVector<float, 3>(0.f, 0.f, 1.f);
+			return FVec3(0.f, 0.f, 1.f);
 		}
 
-		FORCEINLINE T GetMargin() const { return 0; }
+		FORCEINLINE FReal GetMargin() const { return 0; }
 
-		FORCEINLINE TVector<T, d> Support2(const TVector<T, d>& Direction) const { return Support(Direction, 0); }
+		FORCEINLINE FVec3 Support2(const FVec3& Direction) const { return Support(Direction, 0); }
 
-		TVector<T, d> Support(const TVector<T, d>& Direction, const T Thickness) const
+		FVec3 Support(const FVec3& Direction, const FReal Thickness) const
 		{
-			T MaxDot = TNumericLimits<T>::Lowest();
+			FReal MaxDot = TNumericLimits<FReal>::Lowest();
 			int32 MaxVIdx = 0;
 			const int32 NumVertices = SurfaceParticles.Size();
 
 			check(NumVertices > 0);
 			for (int32 Idx = 0; Idx < NumVertices; ++Idx)
 			{
-				const T Dot = TVector<T, d>::DotProduct(SurfaceParticles.X(Idx), Direction);
+				const FReal Dot = FVec3::DotProduct(SurfaceParticles.X(Idx), Direction);
 				if (Dot > MaxDot)
 				{
 					MaxDot = Dot;
@@ -207,12 +209,12 @@ namespace Chaos
 			return FString::Printf(TEXT("Convex"));
 		}
 
-		const TParticles<T, d>& GetSurfaceParticles() const
+		const TParticles<FReal, 3>& GetSurfaceParticles() const
 		{
 			return SurfaceParticles;
 		}
 
-		const TArray<TPlane<T, d>>& GetFaces() const
+		const TArray<TPlane<FReal, 3>>& GetFaces() const
 		{
 			return Planes;
 		}
@@ -239,7 +241,7 @@ namespace Chaos
 
 			Result = HashCombine(Result, SurfaceParticles.GetTypeHash());
 
-			for(const TPlane<T, d>& Plane : Planes)
+			for(const TPlane<FReal, 3>& Plane : Planes)
 			{
 				Result = HashCombine(Result, Plane.GetTypeHash());
 			}
@@ -263,8 +265,8 @@ namespace Chaos
 				// Rebuild convex in order to extract face indices.
 				// TODO: Make it so it can take SurfaceParticles as both input and output without breaking...
 				TArray<TArray<int32>> FaceIndices;
-				TParticles<T, d> TempSurfaceParticles;
-				TConvexBuilder<T>::Build(SurfaceParticles, Planes, FaceIndices, TempSurfaceParticles, LocalBoundingBox);
+				TParticles<FReal, 3> TempSurfaceParticles;
+				FConvexBuilder::Build(SurfaceParticles, Planes, FaceIndices, TempSurfaceParticles, LocalBoundingBox);
 				CalculateVolumeAndCenterOfMass(SurfaceParticles, FaceIndices, Volume, CenterOfMass);
 			}
 		}
@@ -287,14 +289,14 @@ namespace Chaos
 
 		virtual bool IsPerformanceWarning() const override
 		{
-			return TConvexBuilder<T>::IsPerformanceWarning(Planes.Num(), SurfaceParticles.Size());
+			return FConvexBuilder::IsPerformanceWarning(Planes.Num(), SurfaceParticles.Size());
 		}
 
 		virtual FString PerformanceWarningAndSimplifaction() override
 		{
 
-			FString PerformanceWarningString = TConvexBuilder<T>::PerformanceWarningString(Planes.Num(), SurfaceParticles.Size());
-			if (TConvexBuilder<T>::IsGeometryReductionEnabled())
+			FString PerformanceWarningString = FConvexBuilder::PerformanceWarningString(Planes.Num(), SurfaceParticles.Size());
+			if (FConvexBuilder::IsGeometryReductionEnabled())
 			{
 				PerformanceWarningString += ", [Simplifying]";
 				SimplifyGeometry();
@@ -306,18 +308,18 @@ namespace Chaos
 		void SimplifyGeometry()
 		{
 			TArray<TArray<int32>> FaceIndices;
-			TConvexBuilder<T>::Simplify(Planes, FaceIndices, SurfaceParticles, LocalBoundingBox);
+			FConvexBuilder::Simplify(Planes, FaceIndices, SurfaceParticles, LocalBoundingBox);
 		}
 
-		TVector<T,d> GetCenter() const
+		FVec3 GetCenter() const
 		{
-			return TVector<T, d>(0);
+			return FVec3(0);
 		}
 
 	private:
-		TArray<TPlane<T, d>> Planes;
-		TParticles<T, d> SurfaceParticles;	//copy of the vertices that are just on the convex hull boundary
-		TBox<T, d> LocalBoundingBox;
+		TArray<TPlane<FReal, 3>> Planes;
+		TParticles<FReal, 3> SurfaceParticles;	//copy of the vertices that are just on the convex hull boundary
+		TBox<FReal, 3> LocalBoundingBox;
 		float Volume;
 		FVec3 CenterOfMass;
 	};
diff --git a/Engine/Source/Runtime/Experimental/Chaos/Public/Chaos/GeometryQueries.h b/Engine/Source/Runtime/Experimental/Chaos/Public/Chaos/GeometryQueries.h
--- a/Engine/Source/Runtime/Experimental/Chaos/Public/Chaos/GeometryQueries.h
+++ b/Engine/Source/Runtime/Experimental/Chaos/Public/Chaos/GeometryQueries.h
@@ -28,11 +28,11 @@ namespace Chaos
 		case ImplicitObjectType::Sphere: return Func(Geom.template GetObjectChecked<TSphere<FReal, 3>>());
 		case ImplicitObjectType::Box: return Func(Geom.template GetObjectChecked<TBox<FReal, 3>>());
 		case ImplicitObjectType::Capsule: return Func(Geom.template GetObjectChecked<TCapsule<FReal>>());
-		case ImplicitObjectType::Convex: return Func(Geom.template GetObjectChecked<TConvex<FReal, 3>>());
+		case ImplicitObjectType::Convex: return Func(Geom.template GetObjectChecked<FConvex>());
 		case ImplicitObjectType::IsScaled | ImplicitObjectType::Sphere: return Func(Geom.template GetObjectChecked<TImplicitObjectScaled<TSphere<FReal, 3>>>());
 		case ImplicitObjectType::IsScaled | ImplicitObjectType::Box: return Func(Geom.template GetObjectChecked< TImplicitObjectScaled<TBox<FReal, 3>>>());
 		case ImplicitObjectType::IsScaled | ImplicitObjectType::Capsule: return Func(Geom.template GetObjectChecked< TImplicitObjectScaled<TCapsule<FReal>>>());
-		case ImplicitObjectType::IsScaled | ImplicitObjectType::Convex: return Func(Geom.template GetObjectChecked< TImplicitObjectScaled<TConvex<FReal, 3>>>());
+		case ImplicitObjectType::IsScaled | ImplicitObjectType::Convex: return Func(Geom.template GetObjectChecked< TImplicitObjectScaled<FConvex>>());
 		case ImplicitObjectType::Transformed:
 			ensure(false); // We are drilling down to concrete implicit inside transformed, this is disregarding transform data. Caller must specially handle transform.
 			// TODO: Refactor Transformed implicit to use same structure as scaled.
@@ -160,7 +160,7 @@ namespace Chaos
 			auto IsValidConvex = [](const FImplicitObject& InObject) -> bool
 			{
 				//todo: move this out of here
-				if (const auto Convex = TImplicitObjectScaled<TConvex<FReal,3>>::AsScaled(InObject))
+				if (const auto Convex = TImplicitObjectScaled<FConvex>::AsScaled(InObject))
 				{
 					return Convex->GetUnscaledObject()->GetSurfaceParticles().Size() > 0;
 				}				
diff --git a/Engine/Source/Runtime/Experimental/Chaos/Public/Chaos/HeightField.h b/Engine/Source/Runtime/Experimental/Chaos/Public/Chaos/HeightField.h
--- a/Engine/Source/Runtime/Experimental/Chaos/Public/Chaos/HeightField.h
+++ b/Engine/Source/Runtime/Experimental/Chaos/Public/Chaos/HeightField.h
@@ -17,6 +17,7 @@ namespace Chaos
 {
 	template<typename T>
 	class THeightfieldRaycastVisitor;
+	class FConvex;
 }
 
 namespace Chaos
@@ -50,20 +51,20 @@ namespace Chaos
 		bool OverlapGeom(const TSphere<T, 3>& QueryGeom, const TRigidTransform<T, 3>& QueryTM, const T Thickness) const;
 		bool OverlapGeom(const TBox<T, 3>& QueryGeom, const TRigidTransform<T, 3>& QueryTM, const T Thickness) const;
 		bool OverlapGeom(const TCapsule<T>& QueryGeom, const TRigidTransform<T, 3>& QueryTM, const T Thickness) const;
-		bool OverlapGeom(const TConvex<T, 3>& QueryGeom, const TRigidTransform<T, 3>& QueryTM, const T Thickness) const;
+		bool OverlapGeom(const FConvex& QueryGeom, const TRigidTransform<T, 3>& QueryTM, const T Thickness) const;
 		bool OverlapGeom(const TImplicitObjectScaled<TSphere<T, 3>>& QueryGeom, const TRigidTransform<T, 3>& QueryTM, const T Thickness) const;
 		bool OverlapGeom(const TImplicitObjectScaled<TBox<T, 3>>& QueryGeom, const TRigidTransform<T, 3>& QueryTM, const T Thickness) const;
 		bool OverlapGeom(const TImplicitObjectScaled<TCapsule<T>>& QueryGeom, const TRigidTransform<T, 3>& QueryTM, const T Thickness) const;
-		bool OverlapGeom(const TImplicitObjectScaled<TConvex<T, 3>>& QueryGeom, const TRigidTransform<T, 3>& QueryTM, const T Thickness) const;
+		bool OverlapGeom(const TImplicitObjectScaled<FConvex>& QueryGeom, const TRigidTransform<T, 3>& QueryTM, const T Thickness) const;
 
 		bool SweepGeom(const TSphere<T, 3>& QueryGeom, const TRigidTransform<T, 3>& StartTM, const TVector<T, 3>& Dir, const T Length, T& OutTime, TVector<T, 3>& OutPosition, TVector<T, 3>& OutNormal, int32& OutFaceIndex, const T Thickness = 0, bool bComputeMTD = false) const;
 		bool SweepGeom(const TBox<T, 3>& QueryGeom, const TRigidTransform<T, 3>& StartTM, const TVector<T, 3>& Dir, const T Length, T& OutTime, TVector<T, 3>& OutPosition, TVector<T, 3>& OutNormal, int32& OutFaceIndex, const T Thickness = 0, bool bComputeMTD = false) const;
 		bool SweepGeom(const TCapsule<T>& QueryGeom, const TRigidTransform<T, 3>& StartTM, const TVector<T, 3>& Dir, const T Length, T& OutTime, TVector<T, 3>& OutPosition, TVector<T, 3>& OutNormal, int32& OutFaceIndex, const T Thickness = 0, bool bComputeMTD = false) const;
-		bool SweepGeom(const TConvex<T, 3>& QueryGeom, const TRigidTransform<T, 3>& StartTM, const TVector<T, 3>& Dir, const T Length, T& OutTime, TVector<T, 3>& OutPosition, TVector<T, 3>& OutNormal, int32& OutFaceIndex, const T Thickness = 0, bool bComputeMTD = false) const;
+		bool SweepGeom(const FConvex& QueryGeom, const TRigidTransform<T, 3>& StartTM, const TVector<T, 3>& Dir, const T Length, T& OutTime, TVector<T, 3>& OutPosition, TVector<T, 3>& OutNormal, int32& OutFaceIndex, const T Thickness = 0, bool bComputeMTD = false) const;
 		bool SweepGeom(const TImplicitObjectScaled<TSphere<T, 3>>& QueryGeom, const TRigidTransform<T, 3>& StartTM, const TVector<T, 3>& Dir, const T Length, T& OutTime, TVector<T, 3>& OutPosition, TVector<T, 3>& OutNormal, int32& OutFaceIndex, const T Thickness = 0, bool bComputeMTD = false) const;
 		bool SweepGeom(const TImplicitObjectScaled<TBox<T, 3>>& QueryGeom, const TRigidTransform<T, 3>& StartTM, const TVector<T, 3>& Dir, const T Length, T& OutTime, TVector<T, 3>& OutPosition, TVector<T, 3>& OutNormal, int32& OutFaceIndex, const T Thickness = 0, bool bComputeMTD = false) const;
 		bool SweepGeom(const TImplicitObjectScaled<TCapsule<T>>& QueryGeom, const TRigidTransform<T, 3>& StartTM, const TVector<T, 3>& Dir, const T Length, T& OutTime, TVector<T, 3>& OutPosition, TVector<T, 3>& OutNormal, int32& OutFaceIndex, const T Thickness = 0, bool bComputeMTD = false) const;
-		bool SweepGeom(const TImplicitObjectScaled<TConvex<T, 3>>& QueryGeom, const TRigidTransform<T, 3>& StartTM, const TVector<T, 3>& Dir, const T Length, T& OutTime, TVector<T, 3>& OutPosition, TVector<T, 3>& OutNormal, int32& OutFaceIndex, const T Thickness = 0, bool bComputeMTD = false) const;
+		bool SweepGeom(const TImplicitObjectScaled<FConvex>& QueryGeom, const TRigidTransform<T, 3>& StartTM, const TVector<T, 3>& Dir, const T Length, T& OutTime, TVector<T, 3>& OutPosition, TVector<T, 3>& OutNormal, int32& OutFaceIndex, const T Thickness = 0, bool bComputeMTD = false) const;
 
 		virtual int32 FindMostOpposingFace(const TVector<T, 3>& Position, const TVector<T, 3>& UnitDir, int32 HintFaceIndex, T SearchDist) const override;
 		virtual TVector<T, 3> FindGeometryOpposingNormal(const TVector<T, 3>& DenormDir, int32 FaceIndex, const TVector<T, 3>& OriginalNormal) const override;
diff --git a/Engine/Source/Runtime/Experimental/Chaos/Public/Chaos/TriangleMeshImplicitObject.h b/Engine/Source/Runtime/Experimental/Chaos/Public/Chaos/TriangleMeshImplicitObject.h
--- a/Engine/Source/Runtime/Experimental/Chaos/Public/Chaos/TriangleMeshImplicitObject.h
+++ b/Engine/Source/Runtime/Experimental/Chaos/Public/Chaos/TriangleMeshImplicitObject.h
@@ -20,8 +20,7 @@ namespace Chaos
 	template <typename T>
 	class TCapsule;
 
-	template <typename T, int d>
-	class TConvex;
+	class FConvex;
 
 	template<typename T>
 	class CHAOS_API TTriangleMeshImplicitObject final : public FImplicitObject
@@ -46,21 +45,21 @@ namespace Chaos
 		bool OverlapGeom(const TSphere<T,3>& QueryGeom, const TRigidTransform<T, 3>& QueryTM, const T Thickness) const;
 		bool OverlapGeom(const TBox<T, 3>& QueryGeom, const TRigidTransform<T, 3>& QueryTM, const T Thickness) const;
 		bool OverlapGeom(const TCapsule<T>& QueryGeom, const TRigidTransform<T, 3>& QueryTM, const T Thickness) const;
-		bool OverlapGeom(const TConvex<T, 3>& QueryGeom, const TRigidTransform<T, 3>& QueryTM, const T Thickness) const;
+		bool OverlapGeom(const FConvex& QueryGeom, const TRigidTransform<T, 3>& QueryTM, const T Thickness) const;
 		bool OverlapGeom(const TImplicitObjectScaled<TSphere<T, 3>>& QueryGeom, const TRigidTransform<T, 3>& QueryTM, const T Thickness) const;
 		bool OverlapGeom(const TImplicitObjectScaled<TBox<T, 3>>& QueryGeom, const TRigidTransform<T, 3>& QueryTM, const T Thickness) const;
 		bool OverlapGeom(const TImplicitObjectScaled<TCapsule<T>>& QueryGeom, const TRigidTransform<T, 3>& QueryTM, const T Thickness) const;
-		bool OverlapGeom(const TImplicitObjectScaled<TConvex<T, 3>>& QueryGeom, const TRigidTransform<T, 3>& QueryTM, const T Thickness) const;
-		bool OverlapGeom(const TImplicitObjectScaled<TImplicitObjectScaled<TConvex<T, 3>>>& QueryGeom, const TRigidTransform<T, 3>& QueryTM, const T Thickness) const;
+		bool OverlapGeom(const TImplicitObjectScaled<FConvex>& QueryGeom, const TRigidTransform<T, 3>& QueryTM, const T Thickness) const;
+		bool OverlapGeom(const TImplicitObjectScaled<TImplicitObjectScaled<FConvex>>& QueryGeom, const TRigidTransform<T, 3>& QueryTM, const T Thickness) const;
 
 		bool SweepGeom(const TSphere<T,3>& QueryGeom, const TRigidTransform<T, 3>& StartTM, const TVector<T,3>& Dir, const T Length, T& OutTime, TVector<T,3>& OutPosition, TVector<T,3>& OutNormal, int32& OutFaceIndex, const T Thickness = 0, const bool bComputeMTD = false) const;
 		bool SweepGeom(const TBox<T, 3>& QueryGeom, const TRigidTransform<T, 3>& StartTM, const TVector<T, 3>& Dir, const T Length, T& OutTime, TVector<T, 3>& OutPosition, TVector<T, 3>& OutNormal, int32& OutFaceIndex, const T Thickness = 0, const bool bComputeMTD = false) const;
 		bool SweepGeom(const TCapsule<T>& QueryGeom, const TRigidTransform<T, 3>& StartTM, const TVector<T, 3>& Dir, const T Length, T& OutTime, TVector<T, 3>& OutPosition, TVector<T, 3>& OutNormal, int32& OutFaceIndex, const T Thickness = 0, const bool bComputeMTD = false) const;
-		bool SweepGeom(const TConvex<T, 3>& QueryGeom, const TRigidTransform<T, 3>& StartTM, const TVector<T, 3>& Dir, const T Length, T& OutTime, TVector<T, 3>& OutPosition, TVector<T, 3>& OutNormal, int32& OutFaceIndex, const T Thickness = 0, const bool bComputeMTD = false) const;
+		bool SweepGeom(const FConvex& QueryGeom, const TRigidTransform<T, 3>& StartTM, const TVector<T, 3>& Dir, const T Length, T& OutTime, TVector<T, 3>& OutPosition, TVector<T, 3>& OutNormal, int32& OutFaceIndex, const T Thickness = 0, const bool bComputeMTD = false) const;
 		bool SweepGeom(const TImplicitObjectScaled<TSphere<T, 3>>& QueryGeom, const TRigidTransform<T, 3>& StartTM, const TVector<T, 3>& Dir, const T Length, T& OutTime, TVector<T, 3>& OutPosition, TVector<T, 3>& OutNormal, int32& OutFaceIndex, const T Thickness = 0, const bool bComputeMTD = false) const;
 		bool SweepGeom(const TImplicitObjectScaled<TBox<T, 3>>& QueryGeom, const TRigidTransform<T, 3>& StartTM, const TVector<T, 3>& Dir, const T Length, T& OutTime, TVector<T, 3>& OutPosition, TVector<T, 3>& OutNormal, int32& OutFaceIndex, const T Thickness = 0, const bool bComputeMTD = false) const;
 		bool SweepGeom(const TImplicitObjectScaled<TCapsule<T>>& QueryGeom, const TRigidTransform<T, 3>& StartTM, const TVector<T, 3>& Dir, const T Length, T& OutTime, TVector<T, 3>& OutPosition, TVector<T, 3>& OutNormal, int32& OutFaceIndex, const T Thickness = 0, const bool bComputeMTD = false) const;
-		bool SweepGeom(const TImplicitObjectScaled<TConvex<T, 3>>& QueryGeom, const TRigidTransform<T, 3>& StartTM, const TVector<T, 3>& Dir, const T Length, T& OutTime, TVector<T, 3>& OutPosition, TVector<T, 3>& OutNormal, int32& OutFaceIndex, const T Thickness = 0, const bool bComputeMTD = false) const;
+		bool SweepGeom(const TImplicitObjectScaled<FConvex>& QueryGeom, const TRigidTransform<T, 3>& StartTM, const TVector<T, 3>& Dir, const T Length, T& OutTime, TVector<T, 3>& OutPosition, TVector<T, 3>& OutNormal, int32& OutFaceIndex, const T Thickness = 0, const bool bComputeMTD = false) const;
 
 		virtual int32 FindMostOpposingFace(const TVector<T, 3>& Position, const TVector<T, 3>& UnitDir, int32 HintFaceIndex, T SearchDistance) const override;
 		virtual TVector<T, 3> FindGeometryOpposingNormal(const TVector<T, 3>& DenormDir, int32 FaceIndex, const TVector<T, 3>& OriginalNormal) const override;
diff --git a/Engine/Source/Runtime/Experimental/ChaosSolvers/Public/PhysicsProxy/AnalyticImplicitGroup.h b/Engine/Source/Runtime/Experimental/ChaosSolvers/Public/PhysicsProxy/AnalyticImplicitGroup.h
--- a/Engine/Source/Runtime/Experimental/ChaosSolvers/Public/PhysicsProxy/AnalyticImplicitGroup.h
+++ b/Engine/Source/Runtime/Experimental/ChaosSolvers/Public/PhysicsProxy/AnalyticImplicitGroup.h
@@ -65,7 +65,7 @@ public:
 		for (Chaos::TBox<float, 3>* Box : Boxes) if (Box) delete Box;
 		for (Chaos::TCapsule<float>* Capsule : Capsules) if (Capsule) delete Capsule;
 		for (Chaos::TTaperedCylinder<float>* TaperedCylinder : TaperedCylinders) if (TaperedCylinder) delete TaperedCylinder;
-		for (Chaos::TConvex<float, 3>* ConvexHull : ConvexHulls) if (ConvexHull) delete ConvexHull;
+		for (Chaos::FConvex* ConvexHull : ConvexHulls) if (ConvexHull) delete ConvexHull;
 		for (Chaos::TLevelSet<float, 3>* LevelSet : LevelSets) if (LevelSet) delete LevelSet;
 	}
 
@@ -100,7 +100,7 @@ public:
 	int32 Add(const FTransform &InitialXf, Chaos::TBox<float, 3> *Box) { Boxes.Add(Box); return Transforms.Insert(InitialXf, Spheres.Num()+Boxes.Num()-1); }
 	int32 Add(const FTransform &InitialXf, Chaos::TCapsule<float> *Capsule) { Capsules.Add(Capsule); return Transforms.Insert(InitialXf, Spheres.Num()+Boxes.Num()+Capsules.Num()-1); }
 	int32 Add(const FTransform &InitialXf, Chaos::TTaperedCylinder<float> *TaperedCylinder) { TaperedCylinders.Add(TaperedCylinder); return Transforms.Insert(InitialXf, Spheres.Num()+Boxes.Num()+Capsules.Num()+TaperedCylinders.Num()-1); }
-	int32 Add(const FTransform &InitialXf, Chaos::TConvex<float, 3> *ConvexHull) { ConvexHulls.Add(ConvexHull); return Transforms.Insert(InitialXf, Spheres.Num()+Boxes.Num()+Capsules.Num()+TaperedCylinders.Num()+ConvexHulls.Num()-1); }
+	int32 Add(const FTransform &InitialXf, Chaos::FConvex *ConvexHull) { ConvexHulls.Add(ConvexHull); return Transforms.Insert(InitialXf, Spheres.Num()+Boxes.Num()+Capsules.Num()+TaperedCylinders.Num()+ConvexHulls.Num()-1); }
 	int32 Add(const FTransform &InitialXf, Chaos::TLevelSet<float, 3> *LevelSet) { LevelSets.Add(LevelSet); return Transforms.Add(InitialXf); }
 
 	void SetCollisionTopology(
@@ -164,7 +164,7 @@ public:
 			MP.CenterOfMass = Xf.TransformPositionNoScale(TaperedCylinder->GetCenterOfMass());
 			MP.RotationOfMass = Xf.TransformRotation(TaperedCylinder->GetRotationOfMass());
 		}
-		for (Chaos::TConvex<float, 3>* Convex : ConvexHulls)
+		for (Chaos::FConvex* Convex : ConvexHulls)
 		{
 			const FTransform& Xf = Transforms[TransformIndex];
 			BBoxes[TransformIndex] = Convex->BoundingBox().TransformedBox(Xf);
@@ -240,7 +240,7 @@ public:
 			TotalMass += Mass;
 			MP.InertiaTensor = TaperedCylinder->GetInertiaTensor(Mass);
 		}
-		for (Chaos::TConvex<float, 3>* Convex : ConvexHulls)
+		for (Chaos::FConvex* Convex : ConvexHulls)
 		{
 			Chaos::TMassProperties<float, 3> &MP = MPArray[TransformIndex++];
 			float Mass = Density * MP.Volume;
@@ -308,7 +308,7 @@ public:
 			CullDeepPoints(Points, TransformIndex);
 			TransformIndex++;
 		}
-		for (Chaos::TConvex<float, 3>* Convex : ConvexHulls)
+		for (Chaos::FConvex* Convex : ConvexHulls)
 		{
 			TArray<Chaos::TVector<float, 3>>& Points = CollisionPoints[TransformIndex];
 			if (!Points.Num())
@@ -544,7 +544,7 @@ protected:
 			}
 			TransformIndex++;
 		}
-		for (Chaos::TConvex<float, 3>* Convex : ConvexHulls)
+		for (Chaos::FConvex* Convex : ConvexHulls)
 		{
 			if (TransformIndex != SkipIndex)
 			{
@@ -636,7 +636,7 @@ protected:
 	// FKTaperedCapsuleElem - Z axis is the capsule axis
 	TArray<Chaos::TTaperedCylinder<float>*> TaperedCylinders;
 	// FKConvexElem
-	TArray<Chaos::TConvex<float, 3>*> ConvexHulls;
+	TArray<Chaos::FConvex*> ConvexHulls;
 	// Chaos::TConvex replacement
 	TArray<Chaos::TLevelSet<float, 3>*> LevelSets;
 
diff --git a/Engine/Source/Runtime/PhysicsCore/Public/PhysXToChaosUtil.h b/Engine/Source/Runtime/PhysicsCore/Public/PhysXToChaosUtil.h
--- a/Engine/Source/Runtime/PhysicsCore/Public/PhysXToChaosUtil.h
+++ b/Engine/Source/Runtime/PhysicsCore/Public/PhysXToChaosUtil.h
@@ -63,7 +63,7 @@ inline TUniquePtr<Chaos::TImplicitObjectTransformed<float, 3>> PxShapeToChaosGeo
 			Particles.X(Idx) = P2UVector(Vertices[Idx]);
 		}
 
-		TUniquePtr<TConvex<float, 3>> ConvexObj = MakeUnique<TConvex<float, 3>>(Particles);
+		TUniquePtr<FConvex> ConvexObj = MakeUnique<FConvex>(Particles);
 		if (ConvexMeshGeom.scale.isIdentity())
 		{
 			InnerObj = MoveTemp(ConvexObj);
@@ -71,7 +71,7 @@ inline TUniquePtr<Chaos::TImplicitObjectTransformed<float, 3>> PxShapeToChaosGeo
 		else
 		{
 			ensure(ConvexMeshGeom.scale.rotation == PxQuat(PxIdentity));
-			InnerObj = MakeUnique<TImplicitObjectScaled<TConvex<float, 3>, /*bInstanced=*/false>>(MoveTemp(ConvexObj), P2UVector(ConvexMeshGeom.scale.scale));	//todo: make instanced
+			InnerObj = MakeUnique<TImplicitObjectScaled<FConvex, /*bInstanced=*/false>>(MoveTemp(ConvexObj), P2UVector(ConvexMeshGeom.scale.scale));	//todo: make instanced
 		}
 		break;
 	}
