Fix GJKRaycast not terminating when non inflated shapes initially overlap

diff --git a/Engine/Source/Runtime/Experimental/Chaos/Public/Chaos/GJK.h b/Engine/Source/Runtime/Experimental/Chaos/Public/Chaos/GJK.h
--- a/Engine/Source/Runtime/Experimental/Chaos/Public/Chaos/GJK.h
+++ b/Engine/Source/Runtime/Experimental/Chaos/Public/Chaos/GJK.h
@@ -516,7 +516,7 @@ namespace Chaos
 				InGJKPreDist2 = NewDist2;
 
 
-				if (bComputeMTD && bCloseEnough && Lambda == 0 && Inflation2 > 0 && SimplexIDs.NumVerts < 4)
+				if (bComputeMTD && bCloseEnough && Lambda == 0 && Inflation2 > 1e-6 && SimplexIDs.NumVerts < 4)
 				{
 					//For mtd of inflated shapes we have to find the closest point, so we have to keep going
 					bCloseEnough = false;
@@ -550,6 +550,7 @@ namespace Chaos
 		{
 			if (InGJKPreDist2 > 1e-6 && InGJKPreDist2 < TNumericLimits<T>::Max())
 			{
+				ensure(Inflation > 0);	//shouldn't end up here if there is no inflation
 				OutNormal = Normal;
 				TVector<T, 3> ClosestA(0);
 				TVector<T, 3> ClosestB(0);
