Chaos: Improve performance of Constraint Color Init

diff --git a/Engine/Source/Runtime/Experimental/Chaos/Private/Chaos/PBDConstraintColor.cpp b/Engine/Source/Runtime/Experimental/Chaos/Private/Chaos/PBDConstraintColor.cpp
--- a/Engine/Source/Runtime/Experimental/Chaos/Private/Chaos/PBDConstraintColor.cpp
+++ b/Engine/Source/Runtime/Experimental/Chaos/Private/Chaos/PBDConstraintColor.cpp
@@ -19,7 +19,6 @@
 
 using namespace Chaos;
 
-
 void FPBDConstraintColor::ComputeIslandColoring(const int32 Island, const FPBDConstraintGraph& ConstraintGraph, uint32 ContainerId)
 {
 	const TArray<TGeometryParticleHandle<FReal, 3>*>& IslandParticles = ConstraintGraph.GetIslandParticles(Island);
@@ -267,13 +266,6 @@ void FPBDConstraintColor::ComputeContactGraph(const int32 Island, const FPBDCons
 
 void FPBDConstraintColor::InitializeColor(const FPBDConstraintGraph& ConstraintGraph)
 {
-	// Reset the existing Nodes - so colors are all reset to zero
-	ParallelFor(Nodes.Num(), [&](int32 Index)
-	{
-		Nodes[Index].NextColor = 0;
-		Nodes[Index].UsedColors.Empty();
-	});
-
 	// The Number of nodes is large and fairly constant so persist rather than resetting every frame
 	if (Nodes.Num() != ConstraintGraph.Nodes.Num())
 	{
@@ -281,14 +273,22 @@ void FPBDConstraintColor::InitializeColor(const FPBDConstraintGraph& ConstraintG
 		Nodes.AddDefaulted(ConstraintGraph.Nodes.Num() - Nodes.Num());
 	}
 	
+	// Reset the existing Nodes - so colors are all reset to zero
+	for (int32 UpdatedNode : UpdatedNodes)
+	{
+		Nodes[UpdatedNode].NextColor = 0;
+		Nodes[UpdatedNode].UsedColors.Empty();
+	}
+
 	// edges are not persistent right now so we still reset them
 	Edges.Reset();
 	IslandData.Reset();
 
 	Edges.SetNum(ConstraintGraph.Edges.Num());
 	IslandData.SetNum(ConstraintGraph.IslandToData.Num());
-}
 
+	UpdatedNodes = ConstraintGraph.GetUpdatedNodes();
+}
 
 void FPBDConstraintColor::ComputeColor(const int32 Island, const FPBDConstraintGraph& ConstraintGraph, uint32 ContainerId)
 {
@@ -298,7 +298,6 @@ void FPBDConstraintColor::ComputeColor(const int32 Island, const FPBDConstraintG
 	ComputeIslandColoring(Island, ConstraintGraph, ContainerId);
 }
 
-
 const typename FPBDConstraintColor::FLevelToColorToConstraintListMap& FPBDConstraintColor::GetIslandLevelToColorToConstraintListMap(int32 Island) const
 {
 	if (Island < IslandData.Num())
@@ -308,7 +307,6 @@ const typename FPBDConstraintColor::FLevelToColorToConstraintListMap& FPBDConstr
 	return EmptyLevelToColorToConstraintListMap;
 }
 
-
 int FPBDConstraintColor::GetIslandMaxColor(int32 Island) const
 {
 	if (Island < IslandData.Num())
@@ -318,7 +316,6 @@ int FPBDConstraintColor::GetIslandMaxColor(int32 Island) const
 	return -1;
 }
 
-
 int FPBDConstraintColor::GetIslandMaxLevel(int32 Island) const
 {
 	if (Island < IslandData.Num())
diff --git a/Engine/Source/Runtime/Experimental/Chaos/Private/Chaos/PBDConstraintGraph.cpp b/Engine/Source/Runtime/Experimental/Chaos/Private/Chaos/PBDConstraintGraph.cpp
--- a/Engine/Source/Runtime/Experimental/Chaos/Private/Chaos/PBDConstraintGraph.cpp
+++ b/Engine/Source/Runtime/Experimental/Chaos/Private/Chaos/PBDConstraintGraph.cpp
@@ -4,6 +4,7 @@
 
 #include "ProfilingDebugging/ScopedTimers.h"
 #include "ChaosStats.h"
+#include "ChaosLog.h"
 #include "Chaos/Framework/Parallel.h"
 #include "Chaos/PBDRigidParticles.h"
 #include "Containers/Queue.h"
diff --git a/Engine/Source/Runtime/Experimental/Chaos/Public/Chaos/PBDConstraintColor.h b/Engine/Source/Runtime/Experimental/Chaos/Public/Chaos/PBDConstraintColor.h
--- a/Engine/Source/Runtime/Experimental/Chaos/Public/Chaos/PBDConstraintColor.h
+++ b/Engine/Source/Runtime/Experimental/Chaos/Public/Chaos/PBDConstraintColor.h
@@ -97,6 +97,7 @@ namespace Chaos
 		TArray<FGraphEdgeColor> Edges;
 		TArray<FIslandColor> IslandData;
 		FLevelToColorToConstraintListMap EmptyLevelToColorToConstraintListMap;
+		TArray<int32> UpdatedNodes;
 	};
 
 }
diff --git a/Engine/Source/Runtime/Experimental/Chaos/Public/Chaos/PBDConstraintGraph.h b/Engine/Source/Runtime/Experimental/Chaos/Public/Chaos/PBDConstraintGraph.h
--- a/Engine/Source/Runtime/Experimental/Chaos/Public/Chaos/PBDConstraintGraph.h
+++ b/Engine/Source/Runtime/Experimental/Chaos/Public/Chaos/PBDConstraintGraph.h
@@ -192,8 +192,8 @@ namespace Chaos
 		struct FGraphEdge
 		{
 			FGraphEdge()
-				: FirstNode(INDEX_NONE)
-				, SecondNode(INDEX_NONE)
+			    : FirstNode(INDEX_NONE)
+			    , SecondNode(INDEX_NONE)
 			{
 			}
 
@@ -205,7 +205,7 @@ namespace Chaos
 		struct FIslandData
 		{
 			FIslandData()
-				: bIsIslandPersistant(false)
+			    : bIsIslandPersistant(false)
 			{
 			}
 
@@ -220,6 +220,7 @@ namespace Chaos
 		void ParticleAdd(TGeometryParticleHandle<FReal, 3>* AddedParticle);
 		void ParticleRemove(TGeometryParticleHandle<FReal, 3>* RemovedParticle);
 		int32 GetNextNodeIndex();
+		const TArray<int32>& GetUpdatedNodes() const { return UpdatedNodes; }
 
 		TArray<FGraphNode> Nodes;
 		TArray<FGraphEdge> Edges;
