Rotate the tesselation on heightfields to match render geometry

diff --git a/Engine/Source/Runtime/Experimental/Chaos/Private/Chaos/HeightField.cpp b/Engine/Source/Runtime/Experimental/Chaos/Private/Chaos/HeightField.cpp
--- a/Engine/Source/Runtime/Experimental/Chaos/Private/Chaos/HeightField.cpp
+++ b/Engine/Source/Runtime/Experimental/Chaos/Private/Chaos/HeightField.cpp
@@ -158,8 +158,8 @@ namespace Chaos
 			GeomData->GetPointsScaled(FullIndex, Points);
 
 			// Test both triangles that are in this cell, as we could hit both in any order
-			TestTriangle(Payload * 2, Points[0], Points[1], Points[2]);
-			TestTriangle(Payload * 2 + 1, Points[2], Points[1], Points[3]);
+			TestTriangle(Payload * 2, Points[0], Points[1], Points[3]);
+			TestTriangle(Payload * 2 + 1, Points[0], Points[3], Points[2]);
 
 			return OutTime > 0;
 		}
@@ -247,10 +247,10 @@ namespace Chaos
 			TVector<T, 3> Points[4];
 			HfData->GetPointsScaled(FullIndex, Points);
 
-			bool bContinue = TestTriangle(Payload * 2, Points[0], Points[1], Points[2]);
+			bool bContinue = TestTriangle(Payload * 2, Points[0], Points[1], Points[3]);
 			if (bContinue)
 			{
-				TestTriangle(Payload * 2 + 1, Points[2], Points[1], Points[3]);
+				TestTriangle(Payload * 2 + 1, Points[0], Points[3], Points[2]);
 			}
 
 			return OutTime > 0;
@@ -1051,14 +1051,14 @@ namespace Chaos
 		for(const TVector<int32, 2>& Cell : Intersections)
 		{
 			const int32 SingleIndex = Cell[1] * (GeomData.NumCols - 1) + Cell[0];
-			GeomData.GetPoints(SingleIndex, Points);
+			GeomData.GetPointsScaled(SingleIndex, Points);
 
-			if(OverlapTriangle(Points[0], Points[1], Points[2]))
+			if(OverlapTriangle(Points[0], Points[1], Points[3]))
 			{
 				return true;
 			}
 
-			if(OverlapTriangle(Points[2], Points[1], Points[3]))
+			if(OverlapTriangle(Points[0], Points[3], Points[2]))
 			{
 				return true;
 			}
@@ -1102,14 +1102,14 @@ namespace Chaos
 		for(const TVector<int32, 2>& Cell : Intersections)
 		{
 			const int32 SingleIndex = Cell[1] * (GeomData.NumCols - 1) + Cell[0];
-			GeomData.GetPoints(SingleIndex, Points);
+			GeomData.GetPointsScaled(SingleIndex, Points);
 
-			if(OverlapTriangle(Points[0], Points[1], Points[2]))
+			if(OverlapTriangle(Points[0], Points[1], Points[3]))
 			{
 				return true;
 			}
 
-			if(OverlapTriangle(Points[2], Points[1], Points[3]))
+			if(OverlapTriangle(Points[0], Points[3], Points[2]))
 			{
 				return true;
 			}
@@ -1290,8 +1290,8 @@ namespace Chaos
 
 			GeomData.GetPointsScaled(FullIndex, Points);
 
-			CheckTriangle(CellIndex * 2, Points[0], Points[1], Points[2]);
-			CheckTriangle(CellIndex * 2 + 1, Points[2], Points[1], Points[3]);
+			CheckTriangle(CellIndex * 2, Points[0], Points[1], Points[3]);
+			CheckTriangle(CellIndex * 2 + 1, Points[0], Points[3], Points[2]);
 		}
 
 		return MostOpposingFace;
@@ -1317,15 +1317,15 @@ namespace Chaos
 
 			if(bSecondFace)
 			{
-				A = Points[2];
-				B = Points[1];
-				C = Points[3];
+				A = Points[0];
+				B = Points[3];
+				C = Points[2];
 			}
 			else
 			{
 				A = Points[0];
 				B = Points[1];
-				C = Points[2];
+				C = Points[3];
 			}
 
 			const TVector<T, 3> AB = B - A;
