Chaos - Fix  TriangleMeshImplicitObject not properly transforming inputs/outputs for Scaled Overlaps/Sweeps.
Added tests to verify correctness.

diff --git a/Engine/Source/Runtime/Experimental/Chaos/Private/Chaos/TriangleMeshImplicitObject.cpp b/Engine/Source/Runtime/Experimental/Chaos/Private/Chaos/TriangleMeshImplicitObject.cpp
--- a/Engine/Source/Runtime/Experimental/Chaos/Private/Chaos/TriangleMeshImplicitObject.cpp
+++ b/Engine/Source/Runtime/Experimental/Chaos/Private/Chaos/TriangleMeshImplicitObject.cpp
@@ -267,24 +267,71 @@ bool TTriangleMeshImplicitObject<T>::Overlap(const TVector<T, 3>& Point, const T
 }
 
 template <typename QueryGeomType, typename T>
-const QueryGeomType& GetVertsAndQueryGeomHelper(const QueryGeomType& QueryGeom, const int32 TriIdx, const TParticles<T, 3>& Particles, const TArray<TVector<int32, 3>>& Elements, TVec3<T>& OutA, TVec3<T>& OutB, TVec3<T>& OutC)
+void TransformVertsHelper(const QueryGeomType& QueryGeom, int32 TriIdx, const TParticles<T, 3>& Particles,
+	const TArray<TVector<int32, 3>>& Elements, TVec3<T>& OutA, TVec3<T>& OutB, TVec3<T>& OutC)
 {
 	OutA = Particles.X(Elements[TriIdx][0]);
 	OutB = Particles.X(Elements[TriIdx][1]);
 	OutC = Particles.X(Elements[TriIdx][2]);
-	return QueryGeom;
 }
 
 template <typename QueryGeomType, typename T>
-const QueryGeomType& GetVertsAndQueryGeomHelper(const TImplicitObjectScaled<QueryGeomType>& QueryGeom, const int32 TriIdx, const TParticles<T, 3>& Particles, const TArray<TVector<int32, 3>>& Elements, TVec3<T>& OutA, TVec3<T>& OutB, TVec3<T>& OutC)
+void TransformVertsHelper(const TImplicitObjectScaled<QueryGeomType>& QueryGeom, int32 TriIdx, const TParticles<T, 3>& Particles,
+	const TArray<TVector<int32, 3>>& Elements, TVec3<T>& OutA, TVec3<T>& OutB, TVec3<T>& OutC)
 {
 	const TVec3<T> InvScale = QueryGeom.GetInvScale();
 	OutA = Particles.X(Elements[TriIdx][0]) * InvScale;
 	OutB = Particles.X(Elements[TriIdx][1]) * InvScale;
 	OutC = Particles.X(Elements[TriIdx][2]) * InvScale;
+}
+
+template <typename QueryGeomType>
+const QueryGeomType& GetGeomHelper(const QueryGeomType& QueryGeom)
+{
+	return QueryGeom;
+}
+
+template <typename QueryGeomType>
+const QueryGeomType& GetGeomHelper(const TImplicitObjectScaled<QueryGeomType>& QueryGeom)
+{
 	return *QueryGeom.GetUnscaledObject();
 }
 
+template <typename QueryGeomType, typename T>
+void TransformSweepOutputsHelper(const QueryGeomType& QueryGeom, const TVec3<T>& HitNormal, const TVec3<T>& HitPosition, const T LengthScale,
+	const T Time,  TVec3<T>& OutNormal, TVec3<T>& OutPosition, T& OutTime)
+{
+	OutNormal = HitNormal;
+	OutPosition = HitPosition;
+	OutTime = Time;
+}
+
+template <typename QueryGeomType, typename T>
+void TransformSweepOutputsHelper(const TImplicitObjectScaled<QueryGeomType>& QueryGeom, const TVec3<T>& HitNormal, const TVec3<T>& HitPosition,  const T LengthScale,
+	const T Time,  TVec3<T>& OutNormal, TVec3<T>& OutPosition, T& OutTime)
+{
+	const TVec3<T> InvScale = QueryGeom.GetInvScale();
+	const TVec3<T> Scale = QueryGeom.GetScale();
+
+	OutTime = Time / LengthScale;
+	OutNormal = (InvScale * HitNormal).GetSafeNormal();
+	OutPosition = Scale * HitPosition;
+}
+
+template <typename QueryGeomType, typename T>
+void TransformOverlapInputsHelper(const QueryGeomType& QueryGeom, const TRigidTransform<T, 3>& QueryTM, TRigidTransform<T, 3>& OutScaledQueryTM)
+{
+	OutScaledQueryTM = QueryTM;
+}
+
+template <typename QueryGeomType, typename T>
+void TransformOverlapInputsHelper(const TImplicitObjectScaled<QueryGeomType>& QueryGeom, const TRigidTransform<T, 3>& QueryTM, TRigidTransform<T, 3>& OutScaledQueryTM)
+{
+	const TVec3<T> InvScale = QueryGeom.GetInvScale();
+	OutScaledQueryTM = TRigidTransform<FReal, 3>(QueryTM.GetLocation() * InvScale, QueryTM.GetRotation());
+}
+
+
 template <typename T>
 template <typename QueryGeomType>
 bool TTriangleMeshImplicitObject<T>::OverlapGeomImp(const QueryGeomType& QueryGeom, const TRigidTransform<T, 3>& QueryTM, const T Thickness) const
@@ -295,11 +342,15 @@ bool TTriangleMeshImplicitObject<T>::OverlapGeomImp(const QueryGeomType& QueryGe
 	QueryBounds = QueryBounds.TransformedBox(QueryTM);
 	const TArray<int32> PotentialIntersections = BVH.FindAllIntersections(QueryBounds);
 
+	const auto& InnerQueryGeom = GetGeomHelper(QueryGeom);
 
+	TRigidTransform<FReal, 3> TransformedQueryTM;
+	TransformOverlapInputsHelper(QueryGeom, QueryTM, TransformedQueryTM);
+	
 	for (int32 TriIdx : PotentialIntersections)
 	{
 		TVec3<T> A, B, C;
-		const auto & InnerQueryGeom = GetVertsAndQueryGeomHelper(QueryGeom, TriIdx, MParticles, MElements, A, B, C);
+		TransformVertsHelper(QueryGeom, TriIdx, MParticles, MElements, A, B, C);
 
 		const TVector<T, 3> AB = B - A;
 		const TVector<T, 3> AC = C - A;
@@ -308,7 +359,7 @@ bool TTriangleMeshImplicitObject<T>::OverlapGeomImp(const QueryGeomType& QueryGe
 		//However, maybe we should check if it's behind the triangle plane. Also, we should enforce this winding in some way
 		const TVector<T, 3> Offset = TVector<T, 3>::CrossProduct(AB, AC);
 
-		if (GJKIntersection(TTriangle<T>(A, B, C), InnerQueryGeom, QueryTM, Thickness, Offset))
+		if (GJKIntersection(TTriangle<T>(A, B, C), InnerQueryGeom, TransformedQueryTM, Thickness, Offset))
 		{
 			return true;
 		}
@@ -376,13 +427,17 @@ bool TTriangleMeshImplicitObject<T>::OverlapGeom(const TImplicitObjectScaled<TIm
 template <typename QueryGeomType, typename T>
 struct TTriangleMeshSweepVisitor
 {
-	TTriangleMeshSweepVisitor(const TTriangleMeshImplicitObject<T>& InTriMesh, const QueryGeomType& InQueryGeom, const TRigidTransform<T,3>& InStartTM, const TVector<T,3>& InDir, const T InThickness, const bool InComputeMTD)
+	TTriangleMeshSweepVisitor(const TTriangleMeshImplicitObject<T>& InTriMesh, const QueryGeomType& InQueryGeom, const TRigidTransform<T,3>& InStartTM, const TVector<T,3>& InDir,
+		const TVector<T, 3>& InScaledDirNormalized, const T InLengthScale, const TRigidTransform<T, 3>& InScaledStartTM, const T InThickness, const bool InComputeMTD)
 	: TriMesh(InTriMesh)
 	, StartTM(InStartTM)
 	, QueryGeom(InQueryGeom)
 	, Dir(InDir)
 	, Thickness(InThickness)
 	, bComputeMTD(InComputeMTD)
+	, ScaledDirNormalized(InScaledDirNormalized)
+	, LengthScale(InLengthScale)
+	, ScaledStartTM(InScaledStartTM)
 	, OutTime(TNumericLimits<T>::Max())
 	{
 	}
@@ -408,17 +463,17 @@ struct TTriangleMeshSweepVisitor
 		TVector<T, 3> HitNormal;
 
 		TVec3<T> A, B, C;
-		const auto & InnerQueryGeom = GetVertsAndQueryGeomHelper(QueryGeom, TriIdx, TriMesh.MParticles, TriMesh.MElements, A, B, C);
-
+		TransformVertsHelper(QueryGeom, TriIdx, TriMesh.MParticles, TriMesh.MElements, A, B, C);
 		TTriangle<T> Tri(A, B, C);
 
-		if(GJKRaycast2<T>(Tri, InnerQueryGeom, StartTM, Dir, CurData.CurrentLength, Time, HitPosition, HitNormal, Thickness, bComputeMTD))
+		const auto& InnerQueryGeom = GetGeomHelper(QueryGeom);
+
+		if(GJKRaycast2<T>(Tri, InnerQueryGeom, ScaledStartTM, ScaledDirNormalized, LengthScale * CurData.CurrentLength, Time, HitPosition, HitNormal, Thickness, bComputeMTD))
 		{
 			if(Time < OutTime)
 			{
-				OutNormal = HitNormal;
-				OutPosition = HitPosition;
-				OutTime = Time;
+				TransformSweepOutputsHelper(QueryGeom, HitNormal, HitPosition, LengthScale, Time, OutNormal, OutPosition, OutTime);
+
 				OutFaceIndex = TriIdx;
 
 				if(Time <= 0)	//MTD or initial overlap
@@ -443,18 +498,61 @@ struct TTriangleMeshSweepVisitor
 	const T Thickness;
 	const bool bComputeMTD;
 
+	// Cache these values for Scaled Triangle Mesh, as they are needed for transformation when sweeping against triangles.
+	TVector<T, 3> ScaledDirNormalized;
+	T LengthScale;
+	TRigidTransform<T, 3> ScaledStartTM;
+
 	T OutTime;
 	TVector<T, 3> OutPosition;
 	TVector<T, 3> OutNormal;
 	int32 OutFaceIndex;
 };
 
+template <typename QueryGeomType, typename T>
+void ComputeScaledSweepInputs(const QueryGeomType& QueryGeom, const TRigidTransform<T, 3>& StartTM, const TVector<T, 3>& Dir, const T Length,
+	TVector<T, 3>& OutScaledDirNormalized, T& OutLengthScale, TRigidTransform<T, 3>& OutScaledStartTM)
+{
+	OutScaledDirNormalized = Dir;
+	OutLengthScale = 1.0f;
+	OutScaledStartTM = StartTM;
+}
+
+template<typename QueryGeomType, typename T>
+void ComputeScaledSweepInputs(const TImplicitObjectScaled<QueryGeomType>& QueryGeom, const TRigidTransform<T, 3>& StartTM, const TVector<T, 3>& Dir, const T Length,
+	TVector<T, 3>& OutScaledDirNormalized, T& OutLengthScale, TRigidTransform<T, 3>& OutScaledStartTM)
+{
+	const TVector<T, 3>& InvScale = QueryGeom.GetInvScale();
+
+	const TVector<T, 3> UnscaledDirDenorm = InvScale * Dir;
+	const T LengthScale = UnscaledDirDenorm.Size();
+	if (CHAOS_ENSURE(LengthScale > TNumericLimits<T>::Min()))
+	{
+		const T LengthScaleInv = 1.f / LengthScale;
+		OutScaledDirNormalized = UnscaledDirDenorm * LengthScaleInv;
+	}
+
+
+	OutLengthScale = LengthScale;
+	OutScaledStartTM = TRigidTransform<T, 3>(StartTM.GetLocation() * InvScale, StartTM.GetRotation());
+}
+
 template <typename T>
 template <typename QueryGeomType>
-bool TTriangleMeshImplicitObject<T>::SweepGeomImp(const QueryGeomType& QueryGeom, const TRigidTransform<T, 3>& StartTM, const TVector<T, 3>& Dir, const T Length, T& OutTime, TVector<T, 3>& OutPosition, TVector<T, 3>& OutNormal, int32& OutFaceIndex, const T Thickness, const bool bComputeMTD) const
+bool TTriangleMeshImplicitObject<T>::SweepGeomImp(const QueryGeomType& QueryGeom, const TRigidTransform<T, 3>& StartTM, const TVector<T, 3>& Dir, const T Length,
+	T& OutTime, TVector<T, 3>& OutPosition, TVector<T, 3>& OutNormal, int32& OutFaceIndex, const T Thickness, const bool bComputeMTD) const
 {
+
+	// Compute scaled sweep inputs to cache in visitor.
+	TVector<T, 3> ScaledDirNormalized;
+	T LengthScale;
+	TRigidTransform<T, 3> ScaledStartTM;
+	ComputeScaledSweepInputs(QueryGeom, StartTM, Dir, Length, ScaledDirNormalized, LengthScale, ScaledStartTM);
+
 	bool bHit = false;
-	TTriangleMeshSweepVisitor<QueryGeomType, T> SQVisitor(*this, QueryGeom, StartTM, Dir, Thickness, bComputeMTD);
+	TTriangleMeshSweepVisitor<QueryGeomType, T> SQVisitor(*this, QueryGeom, StartTM, Dir, ScaledDirNormalized, LengthScale, ScaledStartTM, Thickness, bComputeMTD);
+
+
 	const TBox<T, 3> QueryBounds = QueryGeom.BoundingBox();
 	const TVector<T, 3> StartPoint = StartTM.TransformPositionNoScale(QueryBounds.Center());
 	const TVector<T, 3> Inflation = QueryBounds.Extents() * 0.5 + TVector<T, 3>(Thickness);
