Centers of mass for scaled and transformed implicits

diff --git a/Engine/Source/Runtime/Experimental/Chaos/Public/Chaos/ImplicitObjectScaled.h b/Engine/Source/Runtime/Experimental/Chaos/Public/Chaos/ImplicitObjectScaled.h
--- a/Engine/Source/Runtime/Experimental/Chaos/Public/Chaos/ImplicitObjectScaled.h
+++ b/Engine/Source/Runtime/Experimental/Chaos/Public/Chaos/ImplicitObjectScaled.h
@@ -462,8 +462,8 @@ public:
 
 	const FVec3 GetCenterOfMass() const
 	{
-		// TODO: Actually compute this!
-		return FVec3(0.f);
+		// TODO: I'm not sure this is correct in all cases
+		return MScale * MObject->GetCenterOfMass();
 	}
 
 	const ObjectType Object() const { return MObject; }
diff --git a/Engine/Source/Runtime/Experimental/Chaos/Public/Chaos/ImplicitObjectTransformed.h b/Engine/Source/Runtime/Experimental/Chaos/Public/Chaos/ImplicitObjectTransformed.h
--- a/Engine/Source/Runtime/Experimental/Chaos/Public/Chaos/ImplicitObjectTransformed.h
+++ b/Engine/Source/Runtime/Experimental/Chaos/Public/Chaos/ImplicitObjectTransformed.h
@@ -201,7 +201,7 @@ public:
 	const FVec3 GetCenterOfMass() const
 	{
 		// TODO: Actually compute this!
-		return FVec3(0.f);
+		return BoundingBox().GetCenterOfMass();
 	}
 
 
