Chaos - Joint optimizations

diff --git a/Engine/Source/Runtime/Experimental/Chaos/Private/Chaos/PBDCollisionConstraints.cpp b/Engine/Source/Runtime/Experimental/Chaos/Private/Chaos/PBDCollisionConstraints.cpp
--- a/Engine/Source/Runtime/Experimental/Chaos/Private/Chaos/PBDCollisionConstraints.cpp
+++ b/Engine/Source/Runtime/Experimental/Chaos/Private/Chaos/PBDCollisionConstraints.cpp
@@ -157,11 +157,11 @@ namespace Chaos
 
 	}
 
-	DECLARE_CYCLE_STAT(TEXT("TPBDCollisionConstraints::Reset"), STAT_CollisionConstraintsReset, STATGROUP_Chaos);
+	DECLARE_CYCLE_STAT(TEXT("TPBDCollisionConstraints::Reset"), STAT_Collisions_Reset, STATGROUP_Chaos);
 	template<typename T, int d>
 	void TPBDCollisionConstraints<T, d>::Reset()
 	{
-		SCOPE_CYCLE_COUNTER(STAT_CollisionConstraintsReset);
+		SCOPE_CYCLE_COUNTER(STAT_Collisions_Reset);
 
 		TArray<FConstraintContainerHandle*> CopyOfHandles = Handles;
 
@@ -277,11 +277,11 @@ namespace Chaos
 	}
 
 
-	DECLARE_CYCLE_STAT(TEXT("TPBDCollisionConstraints::Apply"), STAT_Apply, STATGROUP_Chaos);
+	DECLARE_CYCLE_STAT(TEXT("TPBDCollisionConstraints::Apply"), STAT_Collisions_Apply, STATGROUP_Chaos);
 	template<typename T, int d>
 	void TPBDCollisionConstraints<T, d>::Apply(const T Dt, const TArray<FConstraintContainerHandle*>& InConstraintHandles, const int32 Iterations, const int32 NumIterations)
 	{
-		SCOPE_CYCLE_COUNTER(STAT_Apply);
+		SCOPE_CYCLE_COUNTER(STAT_Collisions_Apply);
 		if (MApplyPairIterations > 0)
 		{
 			PhysicsParallelFor(InConstraintHandles.Num(), [&](int32 ConstraintHandleIndex) {
@@ -315,11 +315,11 @@ namespace Chaos
 	}
 
 
-	DECLARE_CYCLE_STAT(TEXT("TPBDCollisionConstraints::ApplyPushOut"), STAT_ApplyPushOut, STATGROUP_Chaos);
+	DECLARE_CYCLE_STAT(TEXT("TPBDCollisionConstraints::ApplyPushOut"), STAT_Collisions_ApplyPushOut, STATGROUP_Chaos);
 	template<typename T, int d>
 	bool TPBDCollisionConstraints<T, d>::ApplyPushOut(const T Dt, const TArray<FConstraintContainerHandle*>& InConstraintHandles, const TSet< const TGeometryParticleHandle<T, d>*>& IsTemporarilyStatic, int32 Iteration, int32 NumIterations)
 	{
-		SCOPE_CYCLE_COUNTER(STAT_ApplyPushOut);
+		SCOPE_CYCLE_COUNTER(STAT_Collisions_ApplyPushOut);
 
 		bool NeedsAnotherIteration = false;
 		if (MApplyPushOutPairIterations > 0)
diff --git a/Engine/Source/Runtime/Experimental/Chaos/Private/Chaos/PBDJointConstraintSolver.cpp b/Engine/Source/Runtime/Experimental/Chaos/Private/Chaos/PBDJointConstraintSolver.cpp
new file mode 100644
--- /dev/null
+++ b/Engine/Source/Runtime/Experimental/Chaos/Private/Chaos/PBDJointConstraintSolver.cpp
@@ -0,0 +1,614 @@
+// Copyright 1998-2019 Epic Games, Inc. All Rights Reserved.
+#include "Chaos/PBDJointConstraintSolver.h"
+#include "Chaos/ParticleHandle.h"
+#include "Chaos/PBDJointConstraintUtilities.h"
+#include "Chaos/Utilities.h"
+#include "ChaosLog.h"
+#include "ChaosStats.h"
+
+//#pragma optimize("", off)
+
+// Set to '1' enable solver stats (very high frequency, so usually disabled)
+#define CHAOS_JOINTSOLVER_STATSENABLED 0
+
+namespace Chaos
+{
+#if CHAOS_JOINTSOLVER_STATSENABLED
+	DECLARE_CYCLE_STAT(TEXT("TPBDJointConstraints::Apply"), STAT_JointSolver_Apply, STATGROUP_Chaos);
+	DECLARE_CYCLE_STAT(TEXT("TPBDJointConstraints::ApplyPushOut"), STAT_JointSolver_Jacobian, STATGROUP_Chaos);
+	#define CHAOS_JOINTSOLVER_SCOPE_CYCLE_STAT SCOPE_CYCLE_COUNTER(X)
+#else
+	#define CHAOS_JOINTSOLVER_SCOPE_CYCLE_STAT(X)
+#endif
+
+	void FJointConstraintSolver::InitConstraints(
+		const FReal Dt,
+		const FPBDJointSolverSettings& SolverSettings,
+		const FPBDJointSettings& JointSettings,
+		const FVec3& P0,
+		const FRotation3& Q0,
+		const FVec3& P1,
+		const FRotation3& Q1,
+		const FReal InvM0,
+		const FMatrix33& InvIL0,
+		const FReal InvM1,
+		const FMatrix33& InvIL1,
+		const FRigidTransform3& XL0,
+		const FRigidTransform3& XL1)
+	{
+		XLs[0] = XL0;
+		XLs[1] = XL1;
+		InvILs[0] = InvIL0;
+		InvILs[1] = InvIL1;
+		InvMs[0] = InvM0;
+		InvMs[1] = InvM1;
+
+		Ps[0] = P0;
+		Ps[1] = P1;
+		Qs[0] = Q0;
+		Qs[1] = Q1;
+		Qs[1].EnforceShortestArcWith(Qs[0]);
+
+		UpdateDerivedState();
+	}
+
+	void FJointConstraintSolver::ApplyConstraints(
+		const FReal Dt,
+		const FPBDJointSolverSettings& SolverSettings,
+		const FPBDJointSettings& JointSettings,
+		const FReal Stiffness)
+	{
+		CHAOS_JOINTSOLVER_SCOPE_CYCLE_STAT(STAT_JointSolver_Apply);
+
+		// Solving for world-space position/rotation corrections D(6x1) where
+		// D = [I.Jt / [J.I.Jt]].C = I.Jt.L
+		// I is the inverse mass matrix, J the Jacobian, C the current constraint violation,
+		// and L = [1 / [J.I.Jt]].C is the Joint-space correction.
+
+		// For N constraints
+		// Constraint error: C(Nx1)
+		// Jacobian : J(Nx6)
+		// The Jacobians will be some sub-set of the following rows, depending on the "active" constraints 
+		// ("active" = enabled and either fixed or with limits exceeded).
+		//
+		// J0(Nx6) = | XAxis          -XAxis x Connector0 |
+		//           | YAxis          -YAxis x Connector0 |
+		//           | ZAxis          -ZAxis x Connector0 |
+		//           | 0              TwistAxis           |
+		//           | 0              Swing1Axis          |
+		//           | 0              Swing2Axis          |
+		//
+		// J1(Nx6) = | -XAxis         XAxis x Connector1  |
+		//           | -YAxis         YAxis x Connector1  |
+		//           | -ZAxis         ZAxis x Connector1  |
+		//           | 0              -TwistAxis          |
+		//           | 0              -Swing1Axis         |
+		//           | 0              -Swing2Axis         |
+		//
+		FDenseMatrix61 C;
+		FDenseMatrix66 J0, J1;
+		BuildJacobianAndResidual(SolverSettings, JointSettings, J0, J1, C);
+
+		// InvM(6x6) = inverse mass matrix
+		const FMassMatrix InvM0 = FMassMatrix::Make(InvMs[0], Utilities::ComputeWorldSpaceInertia(Qs[0], InvILs[0]));
+		const FMassMatrix InvM1 = FMassMatrix::Make(InvMs[1], Utilities::ComputeWorldSpaceInertia(Qs[1], InvILs[1]));
+
+		// Joint-space mass: F(NxN) = J(Nx6).I(6x6).Jt(6xN)
+		const FDenseMatrix66 F0 = FDenseMatrix66::MultiplyAB(J0, FDenseMatrix66::MultiplyABt(InvM0, J0));
+		const FDenseMatrix66 F1 = FDenseMatrix66::MultiplyAB(J1, FDenseMatrix66::MultiplyABt(InvM1, J1));
+		const FDenseMatrix66 F = FDenseMatrix66::Add(F0, F1);
+
+		// Joint-space correction: L(Nx1) = [1/F](NxN).C(Nx1)
+		FDenseMatrix61 L;
+		if (FDenseMatrixSolver::SolvePositiveDefinite(F, C, L))
+		{
+			// World-space correction: D(6x1) = I(6x6).Jt(6xN).L(Nx1) = I(6x6).JtL(6x1)
+			const FDenseMatrix61 D0 = FDenseMatrix61::MultiplyAB(InvM0, FDenseMatrix61::MultiplyAtB(J0, L));
+			const FDenseMatrix61 D1 = FDenseMatrix61::MultiplyAB(InvM1, FDenseMatrix61::MultiplyAtB(J1, L));
+
+			// Extract world-space position correction
+			Ps[0] = Ps[0] + FVec3(Stiffness * D0.At(0, 0), Stiffness * D0.At(1, 0), Stiffness * D0.At(2, 0));
+			Ps[1] = Ps[1] + FVec3(Stiffness * D1.At(0, 0), Stiffness * D1.At(1, 0), Stiffness * D1.At(2, 0));
+
+			// Extract world-space rotation correction
+			const FReal HalfStiffness = (FReal)0.5 * Stiffness;
+			const FRotation3 DQ0 = FRotation3::FromElements(HalfStiffness * D0.At(3, 0), HalfStiffness * D0.At(4, 0), HalfStiffness * D0.At(5, 0), 0) * Qs[0];
+			const FRotation3 DQ1 = FRotation3::FromElements(HalfStiffness * D1.At(3, 0), HalfStiffness * D1.At(4, 0), HalfStiffness * D1.At(5, 0), 0) * Qs[1];
+			Qs[0] = (Qs[0] + DQ0).GetNormalized();
+			Qs[1] = (Qs[1] + DQ1).GetNormalized();
+			Qs[1].EnforceShortestArcWith(Qs[0]);
+
+			UpdateDerivedState();
+		}
+	}
+
+	void FJointConstraintSolver::BuildJacobianAndResidual(
+		const FPBDJointSolverSettings& SolverSettings,
+		const FPBDJointSettings& JointSettings,
+		FDenseMatrix66& J0,
+		FDenseMatrix66& J1,
+		FDenseMatrix61& C)
+	{
+		CHAOS_JOINTSOLVER_SCOPE_CYCLE_STAT(STAT_JointSolver_Jacobian);
+
+		// Initialize with zero constraints
+		J0.SetDimensions(0, 6);
+		J1.SetDimensions(0, 6);
+		C.SetDimensions(0, 1);
+
+		AddLinearConstraints(SolverSettings, JointSettings, J0, J1, C);
+		AddAngularConstraints(SolverSettings, JointSettings, J0, J1, C);
+	}
+
+	void FJointConstraintSolver::UpdateDerivedState()
+	{
+		Xs[0] = Ps[0] + Qs[0] * XLs[0].GetTranslation();
+		Xs[1] = Ps[1] + Qs[1] * XLs[1].GetTranslation();
+		Rs[0] = Qs[0] * XLs[0].GetRotation();
+		Rs[1] = Qs[1] * XLs[1].GetRotation();
+	}
+
+	// 3 constraints along principle axes
+	void FJointConstraintSolver::AddLinearConstraints_Point(
+		const FPBDJointSolverSettings& SolverSettings,
+		const FPBDJointSettings& JointSettings,
+		FDenseMatrix66& J0,
+		FDenseMatrix66& J1,
+		FDenseMatrix61& C)
+	{
+		const int32 RowIndex = J0.NumRows();
+		J0.AddRows(3);
+		J1.AddRows(3);
+		C.AddRows(3);
+
+		// Cross Product
+		//Result[0] = V1[1] * V2[2] - V1[2] * V2[1];
+		//Result[1] = V1[2] * V2[0] - V1[0] * V2[2];
+		//Result[2] = V1[0] * V2[1] - V1[1] * V2[0];
+
+		const FVec3 XP0 = Xs[0] - Ps[0];
+		J0.SetRowAt(RowIndex + 0, 0, FVec3(1, 0, 0));
+		J0.SetRowAt(RowIndex + 1, 0, FVec3(0, 1, 0));
+		J0.SetRowAt(RowIndex + 2, 0, FVec3(0, 0, 1));
+		J0.SetRowAt(RowIndex + 0, 3, FVec3((FReal)0, XP0[2], -XP0[1]));	// -(1,0,0) x XP0
+		J0.SetRowAt(RowIndex + 1, 3, FVec3(-XP0[2], (FReal)0, XP0[0]));	// -(0,1,0) x XP0
+		J0.SetRowAt(RowIndex + 2, 3, FVec3(XP0[1], -XP0[0], (FReal)0));	// -(0,0,1) x XP0
+
+		const FVec3 XP1 = Xs[1] - Ps[1];
+		J1.SetRowAt(RowIndex + 0, 0, -FVec3(1, 0, 0));
+		J1.SetRowAt(RowIndex + 1, 0, -FVec3(0, 1, 0));
+		J1.SetRowAt(RowIndex + 2, 0, -FVec3(0, 0, 1));
+		J1.SetRowAt(RowIndex + 0, 3, FVec3((FReal)0, -XP1[2], XP1[1]));	// (1,0,0) x XP1
+		J1.SetRowAt(RowIndex + 1, 3, FVec3(XP1[2], (FReal)0, -XP1[0]));	// (0,1,0) x XP1
+		J1.SetRowAt(RowIndex + 2, 3, FVec3(-XP1[1], XP1[0], (FReal)0));	// (0,0,1) x XP1
+
+		const FVec3 ConstraintSeparation = Xs[1] - Xs[0];
+		C.SetAt(RowIndex + 0, 0, ConstraintSeparation[0]);
+		C.SetAt(RowIndex + 1, 0, ConstraintSeparation[1]);
+		C.SetAt(RowIndex + 2, 0, ConstraintSeparation[2]);
+	}
+
+	// up to 1 constraint limiting distance
+	void FJointConstraintSolver::AddLinearConstraints_Sphere(
+		const FPBDJointSolverSettings& SolverSettings,
+		const FPBDJointSettings& JointSettings,
+		FDenseMatrix66& J0,
+		FDenseMatrix66& J1,
+		FDenseMatrix61& C)
+	{
+		const FReal Limit = JointSettings.Motion.LinearLimit;
+		const FVec3 ConstraintSeparation = Xs[1] - Xs[0];
+		const FReal ConstraintSeparationLen = ConstraintSeparation.Size();
+
+		bool bConstraintActive = (ConstraintSeparationLen >= FMath::Max(Limit, KINDA_SMALL_NUMBER));
+		if (bConstraintActive)
+		{
+			const FVec3 XP0 = Xs[0] - Ps[0];
+			const FVec3 XP1 = Xs[1] - Ps[1];
+			const FVec3 Axis = ConstraintSeparation / ConstraintSeparationLen;
+
+			const int32 RowIndex = J0.NumRows();
+			J0.AddRows(1);
+			J1.AddRows(1);
+			C.AddRows(1);
+
+			J0.SetRowAt(RowIndex, 0, Axis);
+			J0.SetRowAt(RowIndex, 3, -FVec3::CrossProduct(Axis, XP0));
+
+			J1.SetRowAt(RowIndex, 0, -Axis);
+			J1.SetRowAt(RowIndex, 3, FVec3::CrossProduct(Axis, XP1));
+
+			C.SetAt(RowIndex, 0, ConstraintSeparationLen - Limit);
+		}
+	}
+
+	// up to 2 constraint: 1 limiting distance along the axis and another limiting lateral distance from the axis
+	void FJointConstraintSolver::AddLinearConstraints_Cylinder(
+		const FPBDJointSolverSettings& SolverSettings,
+		const FPBDJointSettings& JointSettings,
+		const EJointMotionType AxisMotion,
+		const FVec3& Axis,
+		FDenseMatrix66& J0,
+		FDenseMatrix66& J1,
+		FDenseMatrix61& C)
+	{
+		const FVec3 ConstraintSeparation = Xs[1] - Xs[0];
+		const FVec3 XP0 = Xs[0] - Ps[0];
+		const FVec3 XP1 = Xs[1] - Ps[1];
+
+		// Axial Constraint
+		const FReal ConstraintDistanceAxial = FVec3::DotProduct(ConstraintSeparation, Axis);
+		bool bAxisConstraintActive = (AxisMotion != EJointMotionType::Free);
+		if (bAxisConstraintActive)
+		{
+			const int32 RowIndex = J0.NumRows();
+			J0.AddRows(1);
+			J1.AddRows(1);
+			C.AddRows(1);
+
+			J0.SetRowAt(RowIndex, 0, Axis);
+			J0.SetRowAt(RowIndex, 3, -FVec3::CrossProduct(Axis, XP0));
+
+			J1.SetRowAt(RowIndex, 0, -Axis);
+			J1.SetRowAt(RowIndex, 3, FVec3::CrossProduct(Axis, XP1));
+
+			C.SetAt(RowIndex, 0, ConstraintDistanceAxial);
+		}
+
+		// Radial Constraint
+		const FVec3 ConstraintSeparationRadial = ConstraintSeparation - ConstraintDistanceAxial * Axis;
+		const FReal ConstraintDistanceRadial = ConstraintSeparationRadial.Size();
+		const FReal RadialLimit = JointSettings.Motion.LinearLimit;
+		bool bRadialConstraintActive = (ConstraintDistanceRadial >= RadialLimit);
+		if (bRadialConstraintActive)
+		{
+			const FVec3 RadialAxis = ConstraintSeparationRadial / ConstraintDistanceRadial;
+
+			const int32 RowIndex = J0.NumRows();
+			J0.AddRows(1);
+			J1.AddRows(1);
+			C.AddRows(1);
+
+			J0.SetRowAt(RowIndex, 0, RadialAxis);
+			J0.SetRowAt(RowIndex, 3, -FVec3::CrossProduct(RadialAxis, XP0));
+
+			J1.SetRowAt(RowIndex, 0, -RadialAxis);
+			J1.SetRowAt(RowIndex, 3, FVec3::CrossProduct(RadialAxis, XP1));
+
+			C.SetAt(RowIndex, 0, ConstraintDistanceRadial - RadialLimit);
+		}
+	}
+
+	// up to 1 constraint limiting distance along the axis (lateral motion unrestricted)
+	void FJointConstraintSolver::AddLinearConstraints_Plane(
+		const FPBDJointSolverSettings& SolverSettings,
+		const FPBDJointSettings& JointSettings,
+		const EJointMotionType AxisMotion,
+		const FVec3& Axis,
+		FDenseMatrix66& J0,
+		FDenseMatrix66& J1,
+		FDenseMatrix61& C)
+	{
+		const FReal Limit = (AxisMotion == EJointMotionType::Limited) ? JointSettings.Motion.LinearLimit : (FReal)0;
+		const FVec3 ConstraintSeparation = Xs[1] - Xs[0];
+
+		// Planar Constraint
+		const FReal ConstraintDistanceAxial = FVec3::DotProduct(ConstraintSeparation, Axis);
+		bool bAxisConstraintActive = (ConstraintDistanceAxial <= -Limit) || (ConstraintDistanceAxial >= Limit);
+		if (bAxisConstraintActive)
+		{
+			const FVec3 XP0 = Xs[0] - Ps[0];
+			const FVec3 XP1 = Xs[1] - Ps[1];
+
+			const int32 RowIndex = J0.NumRows();
+			J0.AddRows(1);
+			J1.AddRows(1);
+			C.AddRows(1);
+
+			J0.SetRowAt(RowIndex, 0, Axis);
+			J0.SetRowAt(RowIndex, 3, -FVec3::CrossProduct(Axis, XP0));
+
+			J1.SetRowAt(RowIndex, 0, -Axis);
+			J1.SetRowAt(RowIndex, 3, FVec3::CrossProduct(Axis, XP1));
+
+			C.SetAt(RowIndex, 0, (ConstraintDistanceAxial >= Limit) ? ConstraintDistanceAxial - Limit : ConstraintDistanceAxial + Limit);
+		}
+	}
+
+	// up to 1 constraint limiting rotation about twist axes
+	void FJointConstraintSolver::AddAngularConstraints_Twist(
+		const FPBDJointSolverSettings& SolverSettings,
+		const FPBDJointSettings& JointSettings,
+		FDenseMatrix66& J0,
+		FDenseMatrix66& J1,
+		FDenseMatrix61& C)
+	{
+		// Calculate the Twist Axis and Angle for each body
+		const FRotation3 R01 = Rs[0].Inverse() * Rs[1];
+		FRotation3 R01Twist, R01Swing;
+		R01.ToSwingTwist(FJointConstants::TwistAxis(), R01Swing, R01Twist);
+		R01Swing = R01Swing.GetNormalized();
+		R01Twist = R01Twist.GetNormalized();
+
+		FVec3 TwistAxis01;
+		FReal TwistAngle;
+		R01Twist.ToAxisAndAngleSafe(TwistAxis01, TwistAngle, FJointConstants::TwistAxis(), SolverSettings.SwingTwistAngleTolerance);
+		if (TwistAngle > PI)
+		{
+			TwistAngle = TwistAngle - (FReal)2 * PI;
+		}
+		if (FVec3::DotProduct(TwistAxis01, FJointConstants::TwistAxis()) < 0)
+		{
+			TwistAxis01 = -TwistAxis01;
+			TwistAngle = -TwistAngle;
+		}
+
+		FReal TwistAngleMax = FLT_MAX;
+		if (JointSettings.Motion.AngularMotionTypes[(int32)EJointAngularConstraintIndex::Twist] == EJointMotionType::Limited)
+		{
+			TwistAngleMax = JointSettings.Motion.AngularLimits[(int32)EJointAngularConstraintIndex::Twist];
+		}
+		else if (JointSettings.Motion.AngularMotionTypes[(int32)EJointAngularConstraintIndex::Twist] == EJointMotionType::Locked)
+		{
+			TwistAngleMax = 0;
+		}
+
+		bool bConstraintActive = (TwistAngle <= -TwistAngleMax) || (TwistAngle >= TwistAngleMax);
+		if (bConstraintActive)
+		{
+			const FVec3 Axis0 = Rs[0] * TwistAxis01;
+			const FVec3 Axis1 = Rs[1] * TwistAxis01;
+			const FVec3 XP0 = Xs[0] - Ps[0];
+			const FVec3 XP1 = Xs[1] - Ps[1];
+
+			const int32 RowIndex = J0.NumRows();
+			J0.AddRows(1);
+			J1.AddRows(1);
+			C.AddRows(1);
+
+			J0.SetRowAt(RowIndex + 0, 0, FVec3(0, 0, 0));
+			J0.SetRowAt(RowIndex + 0, 3, Axis0);
+
+			J1.SetRowAt(RowIndex + 0, 0, FVec3(0, 0, 0));
+			J1.SetRowAt(RowIndex + 0, 3, -Axis1);
+
+			C.SetAt(RowIndex, 0, (TwistAngle >= TwistAngleMax) ? TwistAngle - TwistAngleMax : TwistAngle + TwistAngleMax);
+
+		}
+	}
+
+	// up to 1 constraint limiting angle between twist axes
+	void FJointConstraintSolver::AddAngularConstraints_Cone(
+		const FPBDJointSolverSettings& SolverSettings,
+		const FPBDJointSettings& JointSettings,
+		FDenseMatrix66& J0,
+		FDenseMatrix66& J1,
+		FDenseMatrix61& C)
+	{
+		// Calculate Swing axis for each body
+		const FRotation3 R01 = Rs[0].Inverse() * Rs[1];
+		FRotation3 R01Twist, R01Swing;
+		R01.ToSwingTwist(FJointConstants::TwistAxis(), R01Swing, R01Twist);
+		R01Swing = R01Swing.GetNormalized();
+		R01Twist = R01Twist.GetNormalized();
+
+		FVec3 SwingAxis01;
+		FReal SwingAngle;
+		R01Swing.ToAxisAndAngleSafe(SwingAxis01, SwingAngle, FJointConstants::Swing1Axis(), SolverSettings.SwingTwistAngleTolerance);
+		if (SwingAngle > PI)
+		{
+			SwingAngle = SwingAngle - (FReal)2 * PI;
+		}
+
+		// Calculate swing limit for the current swing axis
+		FReal SwingAngleMax = FLT_MAX;
+		const FReal Swing1Limit = JointSettings.Motion.AngularLimits[(int32)EJointAngularConstraintIndex::Swing1];
+		const FReal Swing2Limit = JointSettings.Motion.AngularLimits[(int32)EJointAngularConstraintIndex::Swing2];
+
+		// Circular swing limit
+		SwingAngleMax = Swing1Limit;
+
+		// Elliptical swing limit
+		if (!FMath::IsNearlyEqual(Swing1Limit, Swing2Limit, KINDA_SMALL_NUMBER))
+		{
+			// Map swing axis to ellipse and calculate limit for this swing axis
+			const FReal DotSwing1 = FMath::Abs(FVec3::DotProduct(SwingAxis01, FJointConstants::Swing1Axis()));
+			const FReal DotSwing2 = FMath::Abs(FVec3::DotProduct(SwingAxis01, FJointConstants::Swing2Axis()));
+			SwingAngleMax = FMath::Sqrt(Swing1Limit * DotSwing2 * Swing1Limit * DotSwing2 + Swing2Limit * DotSwing1 * Swing2Limit * DotSwing1);
+		}
+
+		bool bConstraintActive = (SwingAngle <= -SwingAngleMax) || (SwingAngle >= SwingAngleMax);
+		if (bConstraintActive)
+		{
+			const FVec3 XP0 = Xs[0] - Ps[0];
+			const FVec3 XP1 = Xs[1] - Ps[1];
+			const FVec3 Axis = Rs[0] * SwingAxis01;
+
+			const int32 RowIndex = J0.NumRows();
+			J0.AddRows(1);
+			J1.AddRows(1);
+			C.AddRows(1);
+
+			J0.SetRowAt(RowIndex + 0, 0, FVec3(0, 0, 0));
+			J0.SetRowAt(RowIndex + 0, 3, -FVec3::CrossProduct(Axis, XP0));
+
+			J1.SetRowAt(RowIndex + 0, 0, FVec3(0, 0, 0));
+			J1.SetRowAt(RowIndex + 0, 3, FVec3::CrossProduct(Axis, XP1));
+
+			C.SetAt(RowIndex, 0, (SwingAngle >= SwingAngleMax) ? SwingAngle - SwingAngleMax : SwingAngle + SwingAngleMax);
+		}
+	}
+
+	// up to 1 constraint limiting rotation about swing axis (relative to body 0)
+	void FJointConstraintSolver::AddAngularConstraints_Swing(
+		const FPBDJointSolverSettings& SolverSettings,
+		const FPBDJointSettings& JointSettings,
+		const EJointAngularConstraintIndex SwingConstraintIndex,
+		const EJointAngularAxisIndex SwingAxisIndex,
+		FDenseMatrix66& J0,
+		FDenseMatrix66& J1,
+		FDenseMatrix61& C)
+	{
+		// Calculate the swing axis for each body
+		const FRotation3 R01 = Rs[0].Inverse() * Rs[1];
+		FRotation3 R01Twist, R01Swing;
+		R01.ToSwingTwist(FJointConstants::TwistAxis(), R01Swing, R01Twist);
+		R01Swing = R01Swing.GetNormalized();
+		R01Twist = R01Twist.GetNormalized();
+
+		FVec3 TwistAxis01;
+		FReal TwistAngle;
+		R01Twist.ToAxisAndAngleSafe(TwistAxis01, TwistAngle, FJointConstants::TwistAxis(), SolverSettings.SwingTwistAngleTolerance);
+		if (TwistAngle > PI)
+		{
+			TwistAngle = TwistAngle - (FReal)2 * PI;
+		}
+		if (FVec3::DotProduct(TwistAxis01, FJointConstants::TwistAxis()) < 0)
+		{
+			TwistAxis01 = -TwistAxis01;
+			TwistAngle = -TwistAngle;
+		}
+		const FVec3 TwistAxis = Rs[0] * TwistAxis01;
+
+		const FRotation3 R1NoTwist = Rs[1] * R01Twist.Inverse();
+		const FMatrix33 Axes0 = Rs[0].ToMatrix();
+		const FMatrix33 Axes1 = R1NoTwist.ToMatrix();
+		FVec3 SwingCross = FVec3::CrossProduct(Axes0.GetAxis((int32)SwingAxisIndex), Axes1.GetAxis((int32)SwingAxisIndex));
+		SwingCross = SwingCross - FVec3::DotProduct(TwistAxis, SwingCross) * TwistAxis;
+		const FReal SwingCrossLen = SwingCross.Size();
+		if (SwingCrossLen > KINDA_SMALL_NUMBER)
+		{
+			FReal SwingAngle = FMath::Asin(FMath::Clamp(SwingCrossLen, (FReal)0, (FReal)1));
+			const FReal SwingDot = FVec3::DotProduct(Axes0.GetAxis((int32)SwingAxisIndex), Axes1.GetAxis((int32)SwingAxisIndex));
+			if (SwingDot < (FReal)0)
+			{
+				SwingAngle = (FReal)PI - SwingAngle;
+			}
+
+			FReal SwingAngleMax = FLT_MAX;
+			if (JointSettings.Motion.AngularMotionTypes[(int32)SwingConstraintIndex] == EJointMotionType::Limited)
+			{
+				FReal SwingLimit = JointSettings.Motion.AngularLimits[(int32)SwingConstraintIndex];
+				SwingAngleMax = SwingLimit;
+			}
+			else if (JointSettings.Motion.AngularMotionTypes[(int32)SwingConstraintIndex] == EJointMotionType::Locked)
+			{
+				SwingAngleMax = 0;
+			}
+
+			bool bConstraintActive = (SwingAngle <= -SwingAngleMax) || (SwingAngle >= SwingAngleMax);
+			if (bConstraintActive)
+			{
+				const FVec3 XP0 = Xs[0] - Ps[0];
+				const FVec3 XP1 = Xs[1] - Ps[1];
+				const FVec3 Axis = SwingCross / SwingCrossLen;
+
+				const int32 RowIndex = J0.NumRows();
+				J0.AddRows(1);
+				J1.AddRows(1);
+				C.AddRows(1);
+
+				J0.SetRowAt(RowIndex + 0, 0, FVec3(0, 0, 0));
+				J0.SetRowAt(RowIndex + 0, 3, Axis);
+
+				J1.SetRowAt(RowIndex + 0, 0, FVec3(0, 0, 0));
+				J1.SetRowAt(RowIndex + 0, 3, -Axis);
+
+				C.SetAt(RowIndex, 0, (SwingAngle >= SwingAngleMax) ? SwingAngle - SwingAngleMax : SwingAngle + SwingAngleMax);
+			}
+		}
+	}
+
+	// Add linear constraints to solver
+	void FJointConstraintSolver::AddLinearConstraints(
+		const FPBDJointSolverSettings& SolverSettings,
+		const FPBDJointSettings& JointSettings,
+		FDenseMatrix66& J0,
+		FDenseMatrix66& J1,
+		FDenseMatrix61& C)
+	{
+		const TVector<EJointMotionType, 3>& Motion = JointSettings.Motion.LinearMotionTypes;
+		if ((Motion[0] == EJointMotionType::Locked) && (Motion[1] == EJointMotionType::Locked) && (Motion[2] == EJointMotionType::Locked))
+		{
+			AddLinearConstraints_Point(SolverSettings, JointSettings, J0, J1, C);
+		}
+		else if ((Motion[0] == EJointMotionType::Limited) && (Motion[1] == EJointMotionType::Limited) && (Motion[2] == EJointMotionType::Limited))
+		{
+			AddLinearConstraints_Sphere(SolverSettings, JointSettings, J0, J1, C);
+		}
+		else if ((Motion[1] == EJointMotionType::Limited) && (Motion[2] == EJointMotionType::Limited))
+		{
+			// Circular Limit (X Axis)
+			AddLinearConstraints_Cylinder(SolverSettings, JointSettings, Motion[0], Rs[0] * FVec3(1, 0, 0), J0, J1, C);
+		}
+		else if ((Motion[0] == EJointMotionType::Limited) && (Motion[2] == EJointMotionType::Limited))
+		{
+			// Circular Limit (Y Axis)
+			AddLinearConstraints_Cylinder(SolverSettings, JointSettings, Motion[1], Rs[0] * FVec3(0, 1, 0), J0, J1, C);
+		}
+		else if ((Motion[0] == EJointMotionType::Limited) && (Motion[1] == EJointMotionType::Limited))
+		{
+			// Circular Limit (Z Axis)
+			AddLinearConstraints_Cylinder(SolverSettings, JointSettings, Motion[2], Rs[0] * FVec3(0, 0, 1), J0, J1, C);
+		}
+		else
+		{
+			// Plane/Square/Cube Limits (no way to author square or cube limits, but would work if we wanted it)
+			if (Motion[0] != EJointMotionType::Free)
+			{
+				AddLinearConstraints_Plane(SolverSettings, JointSettings, Motion[0], Rs[0] * FVec3(1, 0, 0), J0, J1, C);
+			}
+			if (Motion[1] != EJointMotionType::Free)
+			{
+				AddLinearConstraints_Plane(SolverSettings, JointSettings, Motion[1], Rs[0] * FVec3(0, 1, 0), J0, J1, C);
+			}
+			if (Motion[2] != EJointMotionType::Free)
+			{
+				AddLinearConstraints_Plane(SolverSettings, JointSettings, Motion[2], Rs[0] * FVec3(0, 0, 1), J0, J1, C);
+			}
+		}
+	}
+
+	// Add angular constraints to solver
+	void FJointConstraintSolver::AddAngularConstraints(
+		const FPBDJointSolverSettings& SolverSettings,
+		const FPBDJointSettings& JointSettings,
+		FDenseMatrix66& J0,
+		FDenseMatrix66& J1,
+		FDenseMatrix61& C)
+	{
+		EJointMotionType TwistMotion = JointSettings.Motion.AngularMotionTypes[(int32)EJointAngularConstraintIndex::Twist];
+		EJointMotionType Swing1Motion = JointSettings.Motion.AngularMotionTypes[(int32)EJointAngularConstraintIndex::Swing1];
+		EJointMotionType Swing2Motion = JointSettings.Motion.AngularMotionTypes[(int32)EJointAngularConstraintIndex::Swing2];
+
+		// Apply twist constraint
+		if (SolverSettings.bEnableTwistLimits)
+		{
+			if (TwistMotion != EJointMotionType::Free)
+			{
+				AddAngularConstraints_Twist(SolverSettings, JointSettings, J0, J1, C);
+			}
+		}
+
+		// Apply swing constraints
+		if (SolverSettings.bEnableSwingLimits)
+		{
+			if ((Swing1Motion == EJointMotionType::Limited) && (Swing2Motion == EJointMotionType::Limited))
+			{
+				AddAngularConstraints_Cone(SolverSettings, JointSettings, J0, J1, C);
+			}
+			else
+			{
+				if (Swing1Motion != EJointMotionType::Free)
+				{
+					AddAngularConstraints_Swing(SolverSettings, JointSettings, EJointAngularConstraintIndex::Swing1, EJointAngularAxisIndex::Swing1, J0, J1, C);
+				}
+				if (Swing2Motion != EJointMotionType::Free)
+				{
+					AddAngularConstraints_Swing(SolverSettings, JointSettings, EJointAngularConstraintIndex::Swing2, EJointAngularAxisIndex::Swing2, J0, J1, C);
+				}
+			}
+		}
+	}
+
+}
diff --git a/Engine/Source/Runtime/Experimental/Chaos/Private/Chaos/PBDJointConstraints.cpp b/Engine/Source/Runtime/Experimental/Chaos/Private/Chaos/PBDJointConstraints.cpp
--- a/Engine/Source/Runtime/Experimental/Chaos/Private/Chaos/PBDJointConstraints.cpp
+++ b/Engine/Source/Runtime/Experimental/Chaos/Private/Chaos/PBDJointConstraints.cpp
@@ -15,8 +15,11 @@
 
 namespace Chaos
 {
-	DECLARE_CYCLE_STAT(TEXT("TPBDJointConstraints::Apply"), STAT_ApplyJointConstraints, STATGROUP_Chaos);
-	DECLARE_CYCLE_STAT(TEXT("TPBDJointConstraints::ApplyPushOut"), STAT_ApplyPushOutJointConstraints, STATGROUP_Chaos);
+	DECLARE_CYCLE_STAT(TEXT("TPBDJointConstraints::Apply"), STAT_Joints_Apply, STATGROUP_Chaos);
+	DECLARE_CYCLE_STAT(TEXT("TPBDJointConstraints::ApplyPushOut"), STAT_Joints_ApplyPushOut, STATGROUP_Chaos);
+	DECLARE_CYCLE_STAT(TEXT("TPBDJointConstraints::Drives"), STAT_Joints_Drives, STATGROUP_Chaos);
+	DECLARE_CYCLE_STAT(TEXT("TPBDJointConstraints::Solve"), STAT_Joints_Solve, STATGROUP_Chaos);
+	DECLARE_CYCLE_STAT(TEXT("TPBDJointConstraints::Project"), STAT_Joints_Project, STATGROUP_Chaos);
 
 	bool ChaosJoint_UseCholeskySolver = true;
 	FAutoConsoleVariableRef CVarChaosImmPhysDeltaTime(TEXT("p.Chaos.Joint.UseCholeskySolver"), ChaosJoint_UseCholeskySolver, TEXT("Whether to use the new solver"));
@@ -388,7 +391,7 @@ namespace Chaos
 	
 	void FPBDJointConstraints::Apply(const FReal Dt, const TArray<FConstraintContainerHandle*>& InConstraintHandles, const int32 It, const int32 NumIts)
 	{
-		SCOPE_CYCLE_COUNTER(STAT_ApplyJointConstraints);
+		SCOPE_CYCLE_COUNTER(STAT_Joints_Apply);
 
 		// @todo(ccaulfield): make sorting optional
 		// @todo(ccaulfield): handles should be sorted by level by the constraint rule/graph
@@ -411,6 +414,7 @@ namespace Chaos
 		// Apply joint drives
 		if (Settings.DrivesPhase == EJointSolverPhase::Apply)
 		{
+			SCOPE_CYCLE_COUNTER(STAT_Joints_Drives);
 			if (It == 0)
 			{
 				for (FConstraintContainerHandle* ConstraintHandle : SortedConstraintHandles)
@@ -423,6 +427,7 @@ namespace Chaos
 		// Solve for joint position or velocity, depending on settings
 		if (Settings.ApplyPairIterations > 0)
 		{
+			SCOPE_CYCLE_COUNTER(STAT_Joints_Solve);
 			for (FConstraintContainerHandle* ConstraintHandle : SortedConstraintHandles)
 			{
 				if (Settings.bEnableVelocitySolve)
@@ -439,6 +444,7 @@ namespace Chaos
 		// Correct remaining errors after last call to Solve if enabled in this phase
 		if (Settings.ProjectionPhase == EJointSolverPhase::Apply)
 		{
+			SCOPE_CYCLE_COUNTER(STAT_Joints_Project);
 			int32 ProjectionIt = NumIts - 1;
 			if (It == ProjectionIt)
 			{
@@ -455,7 +461,7 @@ namespace Chaos
 	
 	bool FPBDJointConstraints::ApplyPushOut(const FReal Dt, const TArray<FConstraintContainerHandle*>& InConstraintHandles, const int32 It, const int32 NumIts)
 	{
-		SCOPE_CYCLE_COUNTER(STAT_ApplyPushOutJointConstraints);
+		SCOPE_CYCLE_COUNTER(STAT_Joints_ApplyPushOut);
 
 		// @todo(ccaulfield): track whether we are sufficiently solved
 		bool bNeedsAnotherIteration = true;
@@ -470,6 +476,7 @@ namespace Chaos
 		// Apply joint drives
 		if (Settings.DrivesPhase == EJointSolverPhase::ApplyPushOut)
 		{
+			SCOPE_CYCLE_COUNTER(STAT_Joints_Drives);
 			if (It == 0)
 			{
 				for (FConstraintContainerHandle* ConstraintHandle : SortedConstraintHandles)
@@ -482,6 +489,8 @@ namespace Chaos
 		// Solve for positions
 		if (Settings.ApplyPushOutPairIterations > 0)
 		{
+			SCOPE_CYCLE_COUNTER(STAT_Joints_Solve);
+
 			for (FConstraintContainerHandle* ConstraintHandle : SortedConstraintHandles)
 			{
 				SolvePosition(Dt, ConstraintHandle->GetConstraintIndex(), Settings.ApplyPushOutPairIterations, It, NumIts);
@@ -491,6 +500,7 @@ namespace Chaos
 		// Correct remaining errors after the last call to Solve (which depends on if PositionSolve is enabled in ApplyPushOut)
 		if (Settings.ProjectionPhase == EJointSolverPhase::ApplyPushOut)
 		{
+			SCOPE_CYCLE_COUNTER(STAT_Joints_Project);
 			int32 ProjectionIt = (Settings.ApplyPushOutPairIterations > 0) ? NumIts - 1 : 0;
 			if (It == ProjectionIt)
 			{
@@ -617,7 +627,7 @@ namespace Chaos
 	void FPBDJointConstraints::SolvePosition_Cholesky(const FReal Dt, const int32 ConstraintIndex, const int32 NumPairIts, const int32 It, const int32 NumIts)
 	{
 		const TVector<TGeometryParticleHandle<FReal, 3>*, 2>& Constraint = ConstraintParticles[ConstraintIndex];
-		UE_LOG(LogChaosJoint, Verbose, TEXT("Solve Joint Constraint %d %s %s (dt = %f; it = %d / %d)"), ConstraintIndex, *Constraint[0]->ToString(), *Constraint[1]->ToString(), Dt, It, NumIts);
+		UE_LOG(LogChaosJoint, VeryVerbose, TEXT("Solve Joint Constraint %d %s %s (dt = %f; it = %d / %d)"), ConstraintIndex, *Constraint[0]->ToString(), *Constraint[1]->ToString(), Dt, It, NumIts);
 
 		const FPBDJointSettings& JointSettings = ConstraintSettings[ConstraintIndex];
 
@@ -658,7 +668,7 @@ namespace Chaos
 	void FPBDJointConstraints::SolvePosition_GaussSiedel(const FReal Dt, const int32 ConstraintIndex, const int32 NumPairIts, const int32 It, const int32 NumIts)
 	{
 		const TVector<TGeometryParticleHandle<FReal, 3>*, 2>& Constraint = ConstraintParticles[ConstraintIndex];
-		UE_LOG(LogChaosJoint, Verbose, TEXT("Solve Joint Constraint %d %s %s (dt = %f; it = %d / %d)"), ConstraintIndex, *Constraint[0]->ToString(), *Constraint[1]->ToString(), Dt, It, NumIts);
+		UE_LOG(LogChaosJoint, VeryVerbose, TEXT("Solve Joint Constraint %d %s %s (dt = %f; it = %d / %d)"), ConstraintIndex, *Constraint[0]->ToString(), *Constraint[1]->ToString(), Dt, It, NumIts);
 
 		const FPBDJointSettings& JointSettings = ConstraintSettings[ConstraintIndex];
 
@@ -758,7 +768,7 @@ namespace Chaos
 	void FPBDJointConstraints::SolveVelocity(const FReal Dt, const int32 ConstraintIndex, const int32 NumPairIts, const int32 It, const int32 NumIts)
 	{
 		const TVector<TGeometryParticleHandle<FReal, 3>*, 2>& Constraint = ConstraintParticles[ConstraintIndex];
-		UE_LOG(LogChaosJoint, Verbose, TEXT("Solve Joint Constraint %d %s %s (dt = %f; it = %d / %d)"), ConstraintIndex, *Constraint[0]->ToString(), *Constraint[1]->ToString(), Dt, It, NumIts);
+		UE_LOG(LogChaosJoint, VeryVerbose, TEXT("Solve Joint Constraint %d %s %s (dt = %f; it = %d / %d)"), ConstraintIndex, *Constraint[0]->ToString(), *Constraint[1]->ToString(), Dt, It, NumIts);
 
 		const FPBDJointSettings& JointSettings = ConstraintSettings[ConstraintIndex];
 
@@ -868,7 +878,7 @@ namespace Chaos
 		}
 
 		const TVector<TGeometryParticleHandle<FReal, 3>*, 2>& Constraint = ConstraintParticles[ConstraintIndex];
-		UE_LOG(LogChaosJoint, Verbose, TEXT("Project Joint Constraint %d %f %f (it = %d / %d)"), ConstraintIndex, LinearProjectionFactor, AngularProjectionFactor, It, NumIts);
+		UE_LOG(LogChaosJoint, VeryVerbose, TEXT("Project Joint Constraint %d %f %f (it = %d / %d)"), ConstraintIndex, LinearProjectionFactor, AngularProjectionFactor, It, NumIts);
 
 		// Switch particles - internally we assume the first body is the parent (i.e., the space in which constraint limits are specified)
 		const int32 Index0 = 1;
diff --git a/Engine/Source/Runtime/Experimental/Chaos/Private/Chaos/PBDRigidsEvolutionGBF.cpp b/Engine/Source/Runtime/Experimental/Chaos/Private/Chaos/PBDRigidsEvolutionGBF.cpp
--- a/Engine/Source/Runtime/Experimental/Chaos/Private/Chaos/PBDRigidsEvolutionGBF.cpp
+++ b/Engine/Source/Runtime/Experimental/Chaos/Private/Chaos/PBDRigidsEvolutionGBF.cpp
@@ -40,17 +40,17 @@ int DisableThreshold = 5;
 FAutoConsoleVariableRef CVarDisableThreshold(TEXT("p.DisableThreshold2"), DisableThreshold, TEXT("Disable threshold frames to transition to sleeping"));
 
 
-DECLARE_CYCLE_STAT(TEXT("TPBDRigidsEvolutionGBF::AdvanceOneTimeStep"), STAT_Evo_AdvanceOneTimeStep, STATGROUP_Chaos);
-DECLARE_CYCLE_STAT(TEXT("TPBDRigidsEvolutionGBF::Integrate"), STAT_Evo_Integrate, STATGROUP_Chaos);
-DECLARE_CYCLE_STAT(TEXT("TPBDRigidsEvolutionGBF::KinematicTargets"), STAT_Evo_KinematicTargets, STATGROUP_Chaos);
-DECLARE_CYCLE_STAT(TEXT("TPBDRigidsEvolutionGBF::ApplyConstraints"), STAT_Evo_ApplyConstraints, STATGROUP_Chaos);
-DECLARE_CYCLE_STAT(TEXT("TPBDRigidsEvolutionGBF::UpdateVelocities"), STAT_Evo_UpdateVelocites, STATGROUP_Chaos);
-DECLARE_CYCLE_STAT(TEXT("TPBDRigidsEvolutionGBF::ApplyPushOut"), STAT_Evo_ApplyPushOut, STATGROUP_Chaos);
-DECLARE_CYCLE_STAT(TEXT("TPBDRigidsEvolutionGBF::UpdateConstraintPositionBasedState"), STAT_Evo_UpdateConstraintPositionBasedState, STATGROUP_Chaos);
-DECLARE_CYCLE_STAT(TEXT("TPBDRigidsEvolutionGBF::CreateConstraintGraph"), STAT_Evo_CreateConstraintGraph, STATGROUP_Chaos);
-DECLARE_CYCLE_STAT(TEXT("TPBDRigidsEvolutionGBF::CreateIslands"), STAT_Evo_CreateIslands, STATGROUP_Chaos);
-DECLARE_CYCLE_STAT(TEXT("TPBDRigidsEvolutionGBF::ParallelSolve"), STAT_Evo_ParallelSolve, STATGROUP_Chaos);
-DECLARE_CYCLE_STAT(TEXT("TPBDRigidsEvolutionGBF::DeactivateSleep"), STAT_Evo_DeactivateSleep, STATGROUP_Chaos);
+DECLARE_CYCLE_STAT(TEXT("TPBDRigidsEvolutionGBF::AdvanceOneTimeStep"), STAT_Evolution_AdvanceOneTimeStep, STATGROUP_Chaos);
+DECLARE_CYCLE_STAT(TEXT("TPBDRigidsEvolutionGBF::Integrate"), STAT_Evolution_Integrate, STATGROUP_Chaos);
+DECLARE_CYCLE_STAT(TEXT("TPBDRigidsEvolutionGBF::KinematicTargets"), STAT_Evolution_KinematicTargets, STATGROUP_Chaos);
+DECLARE_CYCLE_STAT(TEXT("TPBDRigidsEvolutionGBF::ApplyConstraints"), STAT_Evolution_ApplyConstraints, STATGROUP_Chaos);
+DECLARE_CYCLE_STAT(TEXT("TPBDRigidsEvolutionGBF::UpdateVelocities"), STAT_Evolution_UpdateVelocites, STATGROUP_Chaos);
+DECLARE_CYCLE_STAT(TEXT("TPBDRigidsEvolutionGBF::ApplyPushOut"), STAT_Evolution_ApplyPushOut, STATGROUP_Chaos);
+DECLARE_CYCLE_STAT(TEXT("TPBDRigidsEvolutionGBF::UpdateConstraintPositionBasedState"), STAT_Evolution_UpdateConstraintPositionBasedState, STATGROUP_Chaos);
+DECLARE_CYCLE_STAT(TEXT("TPBDRigidsEvolutionGBF::CreateConstraintGraph"), STAT_Evolution_CreateConstraintGraph, STATGROUP_Chaos);
+DECLARE_CYCLE_STAT(TEXT("TPBDRigidsEvolutionGBF::CreateIslands"), STAT_Evolution_CreateIslands, STATGROUP_Chaos);
+DECLARE_CYCLE_STAT(TEXT("TPBDRigidsEvolutionGBF::ParallelSolve"), STAT_Evolution_ParallelSolve, STATGROUP_Chaos);
+DECLARE_CYCLE_STAT(TEXT("TPBDRigidsEvolutionGBF::DeactivateSleep"), STAT_Evolution_DeactivateSleep, STATGROUP_Chaos);
 
 int32 SerializeEvolution = 0;
 FAutoConsoleVariableRef CVarSerializeEvolution(TEXT("p.SerializeEvolution"), SerializeEvolution, TEXT(""));
@@ -108,6 +108,8 @@ void TPBDRigidsEvolutionGBF<T, d>::Advance(const T Dt, const T MaxStepDt, const
 			// E.g., for 4 steps this will be: 1/4, 1/3, 1/2, 1
 			const float StepFraction = (T)1 / (T)(NumSteps - Step);
 		
+			UE_LOG(LogChaos, Verbose, TEXT("Advance dt = %f [%d/%d]"), StepDt, Step + 1, NumSteps);
+
 			AdvanceOneTimeStep(StepDt, StepFraction);
 		}
 	}
@@ -116,7 +118,7 @@ void TPBDRigidsEvolutionGBF<T, d>::Advance(const T Dt, const T MaxStepDt, const
 template <typename T, int d>
 void TPBDRigidsEvolutionGBF<T, d>::AdvanceOneTimeStep(const T Dt, const T StepFraction)
 {
-	SCOPE_CYCLE_COUNTER(STAT_Evo_AdvanceOneTimeStep);
+	SCOPE_CYCLE_COUNTER(STAT_Evolution_AdvanceOneTimeStep);
 
 #if !UE_BUILD_SHIPPING
 	if (SerializeEvolution)
@@ -126,12 +128,12 @@ void TPBDRigidsEvolutionGBF<T, d>::AdvanceOneTimeStep(const T Dt, const T StepFr
 #endif
 
 	{
-		SCOPE_CYCLE_COUNTER(STAT_Evo_Integrate);
+		SCOPE_CYCLE_COUNTER(STAT_Evolution_Integrate);
 		Integrate(Particles.GetNonDisabledDynamicView(), Dt);	//Question: should we use an awake view?
 	}
 
 	{
-		SCOPE_CYCLE_COUNTER(STAT_Evo_KinematicTargets);
+		SCOPE_CYCLE_COUNTER(STAT_Evolution_KinematicTargets);
 		ApplyKinematicTargets(Dt, StepFraction);
 	}
 
@@ -146,15 +148,15 @@ void TPBDRigidsEvolutionGBF<T, d>::AdvanceOneTimeStep(const T Dt, const T StepFr
 	}
 
 	{
-		SCOPE_CYCLE_COUNTER(STAT_Evo_UpdateConstraintPositionBasedState);
+		SCOPE_CYCLE_COUNTER(STAT_Evolution_UpdateConstraintPositionBasedState);
 		UpdateConstraintPositionBasedState(Dt);
 	}
 	{
-		SCOPE_CYCLE_COUNTER(STAT_Evo_CreateConstraintGraph);
+		SCOPE_CYCLE_COUNTER(STAT_Evolution_CreateConstraintGraph);
 		CreateConstraintGraph();
 	}
 	{
-		SCOPE_CYCLE_COUNTER(STAT_Evo_CreateIslands);
+		SCOPE_CYCLE_COUNTER(STAT_Evolution_CreateIslands);
 		CreateIslands();
 	}
 
@@ -170,12 +172,12 @@ void TPBDRigidsEvolutionGBF<T, d>::AdvanceOneTimeStep(const T Dt, const T StepFr
 	SleepedIslands.SetNum(GetConstraintGraph().NumIslands());
 	if(Dt > 0)
 	{
-		SCOPE_CYCLE_COUNTER(STAT_Evo_ParallelSolve);
+		SCOPE_CYCLE_COUNTER(STAT_Evolution_ParallelSolve);
 		PhysicsParallelFor(GetConstraintGraph().NumIslands(), [&](int32 Island) {
 			const TArray<TGeometryParticleHandle<T, d>*>& IslandParticles = GetConstraintGraph().GetIslandParticles(Island);
 
 			{
-				SCOPE_CYCLE_COUNTER(STAT_Evo_ApplyConstraints);
+				SCOPE_CYCLE_COUNTER(STAT_Evolution_ApplyConstraints);
 				ApplyConstraints(Dt, Island);
 			}
 
@@ -185,12 +187,12 @@ void TPBDRigidsEvolutionGBF<T, d>::AdvanceOneTimeStep(const T Dt, const T StepFr
 			}
 
 			{
-				SCOPE_CYCLE_COUNTER(STAT_Evo_UpdateVelocites);
+				SCOPE_CYCLE_COUNTER(STAT_Evolution_UpdateVelocites);
 				UpdateVelocities(Dt, Island);
 			}
 
 			{
-				SCOPE_CYCLE_COUNTER(STAT_Evo_ApplyPushOut);
+				SCOPE_CYCLE_COUNTER(STAT_Evolution_ApplyPushOut);
 				ApplyPushOut(Dt, Island);
 			}
 
@@ -237,7 +239,7 @@ void TPBDRigidsEvolutionGBF<T, d>::AdvanceOneTimeStep(const T Dt, const T StepFr
 		});
 	}
 	{
-		SCOPE_CYCLE_COUNTER(STAT_Evo_DeactivateSleep);
+		SCOPE_CYCLE_COUNTER(STAT_Evolution_DeactivateSleep);
 		for (int32 Island = 0; Island < GetConstraintGraph().NumIslands(); ++Island)
 		{
 			if (SleepedIslands[Island])
diff --git a/Engine/Source/Runtime/Experimental/Chaos/Public/Chaos/DenseMatrix.h b/Engine/Source/Runtime/Experimental/Chaos/Public/Chaos/DenseMatrix.h
--- a/Engine/Source/Runtime/Experimental/Chaos/Public/Chaos/DenseMatrix.h
+++ b/Engine/Source/Runtime/Experimental/Chaos/Public/Chaos/DenseMatrix.h
@@ -22,10 +22,9 @@ namespace Chaos
 		
 		const FReal& I(int32 RowIndex, int32 ColIndex) const
 		{
-			// Symmetric - could reduce element count at increased indexing cost...
 			checkSlow(RowIndex < 3);
 			checkSlow(ColIndex < 3);
-			return Inertia[RowIndex * 3 + ColIndex];
+			return Inertia.M[RowIndex][ColIndex];
 		}
 
 		static FMassMatrix Make(const FReal InM, const FMatrix33& InI)
@@ -33,24 +32,26 @@ namespace Chaos
 			return FMassMatrix(InM, InI);
 		}
 
+		static FMassMatrix Make(const FReal InM, FMatrix33&& InI)
+		{
+			return FMassMatrix(InM, MoveTemp(InI));
+		}
+
 	private:
 		FMassMatrix(const FReal InM, const FMatrix33& InI)
+			: Mass(InM)
+			, Inertia(InI)
 		{
-			Mass = InM;
+		}
 
-			Inertia[0] = InI.M[0][0];
-			Inertia[1] = InI.M[0][1];
-			Inertia[2] = InI.M[0][2];
-			Inertia[3] = InI.M[1][0];
-			Inertia[4] = InI.M[1][1];
-			Inertia[5] = InI.M[1][2];
-			Inertia[6] = InI.M[2][0];
-			Inertia[7] = InI.M[2][1];
-			Inertia[8] = InI.M[2][2];
+		FMassMatrix(const FReal InM, FMatrix33&& InI)
+			: Mass(InM)
+			, Inertia(MoveTemp(InI))
+		{
 		}
 
 		FReal Mass;
-		FReal Inertia[9];
+		FMatrix33 Inertia;
 	};
 
 	/**
@@ -120,6 +121,7 @@ namespace Chaos
 
 		/**
 		 * Set the dimensions of the matrix, but do not initialize any values.
+		 * This will invalidate any existing data.
 		 */
 		FORCEINLINE void SetDimensions(const int32 InNumRows, const int32 InNumColumns)
 		{
@@ -128,6 +130,19 @@ namespace Chaos
 			NCols = InNumColumns;
 		}
 
+		/**
+		 * Add uninitialized rows to the matrix. This will not invalidate data in any previously added rows,
+		 * so it can be used to build NxM matrices where M is known, but N is calculated later.
+		 * /return the index of the first new row.
+		 */
+		FORCEINLINE int32 AddRows(const int32 InNumRows)
+		{
+			check((NumRows() + InNumRows) * NumColumns() <= MaxElements);
+			const int32 NewRowIndex = NRows;
+			NRows = NRows + InNumRows;
+			return NewRowIndex;
+		}
+
 		/**
 		 * Return a writable reference to the element at the specified row and column.
 		 */
@@ -157,6 +172,14 @@ namespace Chaos
 			Set(V);
 		}
 
+		/**
+		 * Set the element
+		 */
+		FORCEINLINE void SetAt(const int32 RowIndex, const int32 ColumnIndex, const FReal V)
+		{
+			At(RowIndex, ColumnIndex) = V;
+		}
+
 		/**
 		 * Set all elements to 'V'.
 		 */
@@ -574,6 +597,44 @@ namespace Chaos
 		}
 
 
+		/**
+		 * Return C = A x B, where B is an 6xN matrix, and A is a 6x6 mass matrix (Mass in upper left 3x3 diagonals, Inertia in lower right 3x3).
+		 * C = |M 0| * |B0| = |M.B0|
+		 *     |0 I|   |B1|   |I.B1|
+		 */
+		template<int32 T_EB>
+		static TDenseMatrix<MaxElements> MultiplyAB(const FMassMatrix& A, const TDenseMatrix<T_EB>& B)
+		{
+			check(B.NumRows() == 6);
+
+			TDenseMatrix<T_MAXELEMENTS> Result = TDenseMatrix<T_MAXELEMENTS>::Make(6, B.NumColumns());
+
+			// Calculate rows 0-2
+			for (int32 IRow = 0; IRow < 3; ++IRow)
+			{
+				for (int32 ICol = 0; ICol < B.NumColumns(); ++ICol)
+				{
+					Result.At(IRow, ICol) = A.M() * B.At(IRow, ICol);
+				}
+			}
+
+			// Calculate rows 3-5
+			for (int32 IRow = 3; IRow < 6; ++IRow)
+			{
+				for (int32 ICol = 0; ICol < B.NumColumns(); ++ICol)
+				{
+					FReal V = 0;
+					for (int32 KK = 3; KK < 6; ++KK)
+					{
+						V += A.I(IRow - 3, KK - 3) * B.At(KK, ICol);
+					}
+					Result.At(IRow, ICol) = V;
+				}
+			}
+
+			return Result;
+		}
+
 		/**
 		 * Return C = A x Transpose(B), where B is an Nx6 matrix, and A is a 6x6 mass matrix (Mass in upper left 3x3 diagonals, Inertia in lower right 3x3).
 		 * C = |M 0| * |B0 B1|(T) = |M.B0(T)|
diff --git a/Engine/Source/Runtime/Experimental/Chaos/Public/Chaos/PBDJointConstraintSolver.h b/Engine/Source/Runtime/Experimental/Chaos/Public/Chaos/PBDJointConstraintSolver.h
--- a/Engine/Source/Runtime/Experimental/Chaos/Public/Chaos/PBDJointConstraintSolver.h
+++ b/Engine/Source/Runtime/Experimental/Chaos/Public/Chaos/PBDJointConstraintSolver.h
@@ -17,6 +17,18 @@ namespace Chaos
 		using FDenseMatrix66 = TDenseMatrix<6 * 6>;
 		using FDenseMatrix61 = TDenseMatrix<6 * 1>;
 
+		FORCEINLINE const FVec3& GetP(const int32 Index) const
+		{
+			checkSlow(Index < 2);
+			return Ps[Index];
+		}
+
+		FORCEINLINE const FRotation3& GetQ(const int32 Index) const
+		{
+			checkSlow(Index < 2);
+			return Qs[Index];
+		}
+
 		void InitConstraints(
 			const FReal Dt,
 			const FPBDJointSolverSettings& SolverSettings,
@@ -30,537 +42,95 @@ namespace Chaos
 			const FReal InvM1,
 			const FMatrix33& InvIL1,
 			const FRigidTransform3& XL0,
-			const FRigidTransform3& XL1)
-		{
-			XLs[0] = XL0;
-			XLs[1] = XL1;
-			InvILs[0] = InvIL0;
-			InvILs[1] = InvIL1;
-			InvMs[0] = InvM0;
-			InvMs[1] = InvM1;
-
-			UpdateConstraints(Dt, SolverSettings, JointSettings, P0, Q0, P1, Q1);
-		}
+			const FRigidTransform3& XL1);
 
 		void ApplyConstraints(
 			const FReal Dt,
 			const FPBDJointSolverSettings& SolverSettings,
 			const FPBDJointSettings& JointSettings,
-			const FReal Stiffness)
-		{
-			// @todo(ccaulfield): Custom mass-matrix multiplication methods
-
-			// For N constraints
-			// Constraint error: C(Nx1)
-			// Jacobian : J(Nx6)
-			FDenseMatrix61 C;
-			FDenseMatrix66 J0, J1;
-			BuildJacobianAndResidual(SolverSettings, JointSettings, J0, J1, C);
-
-			// InvM(6x6) = inverse mass matrix
-			FDenseMatrix66 InvM0 = FDenseMatrix66::Make(6, 6, 0);
-			InvM0.SetDiagonalAt(0, 3, InvMs[0]);
-			InvM0.SetBlockAt(3, 3, Utilities::ComputeWorldSpaceInertia(Qs[0], InvILs[0]));
-			FDenseMatrix66 InvM1 = FDenseMatrix66::Make(6, 6, 0);
-			InvM1.SetDiagonalAt(0, 3, InvMs[1]);
-			InvM1.SetBlockAt(3, 3, Utilities::ComputeWorldSpaceInertia(Qs[1], InvILs[1]));
-
-			// Joint-space mass: F(NxN) = J.IM.Jt
-			// Joint-space correction: L(Nx1) = (1/F).C
-			FDenseMatrix66 F0 = FDenseMatrix66::MultiplyAB(J0, FDenseMatrix66::MultiplyABt(InvM0, J0));
-			FDenseMatrix66 F1 = FDenseMatrix66::MultiplyAB(J1, FDenseMatrix66::MultiplyABt(InvM1, J1));
-			FDenseMatrix66 F = FDenseMatrix66::Add(F0, F1);
-			FDenseMatrix61 L;
-			if (FDenseMatrixSolver::SolvePositiveDefinite(F, C, L))
-			{
-				// World-space correction: D(6x1) = IM.Jt.L
-				FDenseMatrix61 D0 = FDenseMatrix61::MultiplyAB(InvM0, FDenseMatrix61::MultiplyAtB(J0, L));
-				FDenseMatrix61 D1 = FDenseMatrix61::MultiplyAB(InvM1, FDenseMatrix61::MultiplyAtB(J1, L));
-
-				// Extract world-space position correction
-				const FVec3 DP0 = FVec3(D0.At(0, 0), D0.At(1, 0), D0.At(2, 0));
-				const FVec3 DP1 = FVec3(D1.At(0, 0), D1.At(1, 0), D1.At(2, 0));
-				const FVec3 P0 = Ps[0] + Stiffness * DP0;
-				const FVec3 P1 = Ps[1] + Stiffness * DP1;
-
-				// Extract world-space rotation correction
-				const FVec3 DR0 = FVec3(D0.At(3, 0), D0.At(4, 0), D0.At(5, 0));
-				const FVec3 DR1 = FVec3(D1.At(3, 0), D1.At(4, 0), D1.At(5, 0));
-				const FRotation3 DQ0 = (FRotation3::FromElements(Stiffness * DR0, 0) * Qs[0]) * (FReal)0.5;
-				const FRotation3 DQ1 = (FRotation3::FromElements(Stiffness * DR1, 0) * Qs[1]) * (FReal)0.5;
-				const FRotation3 Q0 = (Qs[0] + DQ0).GetNormalized();
-				const FRotation3 Q1 = (Qs[1] + DQ1).GetNormalized();
-
-				// Apply corrections
-				UpdateConstraints(Dt, SolverSettings, JointSettings, P0, Q0, P1, Q1);
-			}
-		}
-
-		FORCEINLINE const FVec3& GetP(const int32 Index) const
-		{
-			checkSlow(Index < 2);
-			return Ps[Index];
-		}
-
-		FORCEINLINE const FRotation3& GetQ(const int32 Index) const
-		{
-			checkSlow(Index < 2);
-			return Qs[Index];
-		}
+			const FReal Stiffness);
 
 	private:
 
-		void UpdateConstraints(
-			const FReal Dt,
-			const FPBDJointSolverSettings& SolverSettings,
-			const FPBDJointSettings& JointSettings,
-			const FVec3& P0,
-			const FRotation3& Q0,
-			const FVec3& P1,
-			const FRotation3& Q1)
-		{
-			Ps[0] = P0;
-			Ps[1] = P1;
-			Qs[0] = Q0;
-			Qs[1] = Q1;
-			Qs[1].EnforceShortestArcWith(Q0);
-
-			Xs[0] = Ps[0] + Qs[0] * XLs[0].GetTranslation();
-			Xs[1] = Ps[1] + Qs[1] * XLs[1].GetTranslation();
-			Rs[0] = Qs[0] * XLs[0].GetRotation();
-			Rs[1] = Qs[1] * XLs[1].GetRotation();
-		}
+		void UpdateDerivedState();
 
-		// 3 constraints along principle axes
 		void AddLinearConstraints_Point(
 			const FPBDJointSolverSettings& SolverSettings,
-			const FPBDJointSettings& JointSettings)
-		{
-			const FVec3 ConstraintSeparation = Xs[1] - Xs[0];
-
-			NumLinearConstraints = 3;
-			LinearConstraintAxes[0] = FVec3(1, 0, 0);
-			LinearConstraintAxes[1] = FVec3(0, 1, 0);
-			LinearConstraintAxes[2] = FVec3(0, 0, 1);
-			LinearConstraintDistances[0] = ConstraintSeparation[0];
-			LinearConstraintDistances[1] = ConstraintSeparation[1];
-			LinearConstraintDistances[2] = ConstraintSeparation[2];
-			LinearConstraintErrors[0] = ConstraintSeparation[0];
-			LinearConstraintErrors[1] = ConstraintSeparation[1];
-			LinearConstraintErrors[2] = ConstraintSeparation[2];
-		}
+			const FPBDJointSettings& JointSettings,
+			FDenseMatrix66& J0,
+			FDenseMatrix66& J1,
+			FDenseMatrix61& C);
 
-		// up to 1 constraint limiting distance
 		void AddLinearConstraints_Sphere(
 			const FPBDJointSolverSettings& SolverSettings,
-			const FPBDJointSettings& JointSettings)
-		{
-			NumLinearConstraints = 0;
-
-			const FReal Limit = JointSettings.Motion.LinearLimit;
-			const FVec3 ConstraintSeparation = Xs[1] - Xs[0];
-			const FReal ConstraintSeparationLen = ConstraintSeparation.Size();
-
-			bool bConstraintActive = (ConstraintSeparationLen >= FMath::Max(Limit, KINDA_SMALL_NUMBER));
-			if (bConstraintActive)
-			{
-				const FVec3 Axis = ConstraintSeparation / ConstraintSeparationLen;
-
-				NumLinearConstraints = 1;
-				LinearConstraintAxes[0] = Axis;
-				LinearConstraintDistances[0] = ConstraintSeparationLen;
-				LinearConstraintErrors[0] = ConstraintSeparationLen - Limit;
-			}
-		}
+			const FPBDJointSettings& JointSettings,
+			FDenseMatrix66& J0,
+			FDenseMatrix66& J1,
+			FDenseMatrix61& C);
 
-		// up to 2 constraint: 1 limiting distance along the axis and another limiting lateral distance from the axis
 		void AddLinearConstraints_Cylinder(
 			const FPBDJointSolverSettings& SolverSettings,
 			const FPBDJointSettings& JointSettings,
 			const EJointMotionType AxisMotion,
-			const FVec3& Axis)
-		{
-			const FVec3 ConstraintSeparation = Xs[1] - Xs[0];
-
-			// Axial Constraint
-			const FReal ConstraintDistanceAxial = FVec3::DotProduct(ConstraintSeparation, Axis);
-			bool bAxisConstraintActive = (AxisMotion != EJointMotionType::Free);
-			if (bAxisConstraintActive)
-			{
-				int32 ConstraintIndex = NumLinearConstraints++;
-				LinearConstraintAxes[ConstraintIndex] = Axis;
-				LinearConstraintDistances[ConstraintIndex] = ConstraintDistanceAxial;
-				LinearConstraintErrors[ConstraintIndex] = ConstraintDistanceAxial;
-			}
-
-			// Radial Constraint
-			const FVec3 ConstraintSeparationRadial = ConstraintSeparation - ConstraintDistanceAxial * Axis;
-			const FReal ConstraintDistanceRadial = ConstraintSeparationRadial.Size();
-			const FReal RadialLimit = JointSettings.Motion.LinearLimit;
-			bool bRadialConstraintActive = (ConstraintDistanceRadial >= RadialLimit);
-			if (bRadialConstraintActive)
-			{
-				int32 ConstraintIndex = NumLinearConstraints++;
-				LinearConstraintAxes[ConstraintIndex] = ConstraintSeparationRadial / ConstraintDistanceRadial;
-				LinearConstraintDistances[ConstraintIndex] = ConstraintDistanceRadial;
-				LinearConstraintErrors[ConstraintIndex] = ConstraintDistanceRadial - RadialLimit;
-			}
-		}
+			const FVec3& Axis,
+			FDenseMatrix66& J0,
+			FDenseMatrix66& J1,
+			FDenseMatrix61& C);
 
-		// up to 1 constraint limiting distance along the axis (lateral motion unrestricted)
 		void AddLinearConstraints_Plane(
 			const FPBDJointSolverSettings& SolverSettings,
 			const FPBDJointSettings& JointSettings,
 			const EJointMotionType AxisMotion,
-			const FVec3& Axis)
-		{
-			const FReal Limit = (AxisMotion == EJointMotionType::Limited) ? JointSettings.Motion.LinearLimit : (FReal)0;
-			const FVec3 ConstraintSeparation = Xs[1] - Xs[0];
-
-			// Planar Constraint
-			const FReal ConstraintDistanceAxial = FVec3::DotProduct(ConstraintSeparation, Axis);
-			bool bAxisConstraintActive = (ConstraintDistanceAxial <= -Limit) || (ConstraintDistanceAxial >= Limit);
-			if (bAxisConstraintActive)
-			{
-				int32 ConstraintIndex = NumLinearConstraints++;
-				LinearConstraintAxes[ConstraintIndex] = Axis;
-				LinearConstraintDistances[ConstraintIndex] = ConstraintDistanceAxial;
-				LinearConstraintErrors[ConstraintIndex] = (ConstraintDistanceAxial >= Limit) ? ConstraintDistanceAxial - Limit : ConstraintDistanceAxial + Limit;
-			}
-		}
+			const FVec3& Axis,
+			FDenseMatrix66& J0,
+			FDenseMatrix66& J1,
+			FDenseMatrix61& C);
 
-		// up to 1 constraint limiting rotation about twist axes
 		void AddAngularConstraints_Twist(
 			const FPBDJointSolverSettings& SolverSettings,
-			const FPBDJointSettings& JointSettings)
-		{
-			// Calculate the Twist Axis and Angle for each body
-			const FRotation3 R01 = Rs[0].Inverse() * Rs[1];
-			FRotation3 R01Twist, R01Swing;
-			R01.ToSwingTwist(FJointConstants::TwistAxis(), R01Swing, R01Twist);
-			R01Swing = R01Swing.GetNormalized();
-			R01Twist = R01Twist.GetNormalized();
-
-			FVec3 TwistAxis01;
-			FReal TwistAngle;
-			R01Twist.ToAxisAndAngleSafe(TwistAxis01, TwistAngle, FJointConstants::TwistAxis(), SolverSettings.SwingTwistAngleTolerance);
-			if (TwistAngle > PI)
-			{
-				TwistAngle = TwistAngle - (FReal)2 * PI;
-			}
-			if (FVec3::DotProduct(TwistAxis01, FJointConstants::TwistAxis()) < 0)
-			{
-				TwistAxis01 = -TwistAxis01;
-				TwistAngle = -TwistAngle;
-			}
-
-			FReal TwistAngleMax = FLT_MAX;
-			if (JointSettings.Motion.AngularMotionTypes[(int32)EJointAngularConstraintIndex::Twist] == EJointMotionType::Limited)
-			{
-				TwistAngleMax = JointSettings.Motion.AngularLimits[(int32)EJointAngularConstraintIndex::Twist];
-			}
-			else if (JointSettings.Motion.AngularMotionTypes[(int32)EJointAngularConstraintIndex::Twist] == EJointMotionType::Locked)
-			{
-				TwistAngleMax = 0;
-			}
-
-			bool bConstraintActive = (TwistAngle <= -TwistAngleMax) || (TwistAngle >= TwistAngleMax);
-			if (bConstraintActive)
-			{
-				int32 ConstraintIndex = NumAngularConstraints++;
-				AngularConstraintAxes[ConstraintIndex][0] = Rs[0] * TwistAxis01;
-				AngularConstraintAxes[ConstraintIndex][1] = Rs[1] * TwistAxis01;
-				AngularConstraintDistances[ConstraintIndex] = TwistAngle;
-				AngularConstraintErrors[ConstraintIndex] = (TwistAngle >= TwistAngleMax) ? TwistAngle - TwistAngleMax : TwistAngle + TwistAngleMax;
-			}
-		}
+			const FPBDJointSettings& JointSettings,
+			FDenseMatrix66& J0,
+			FDenseMatrix66& J1,
+			FDenseMatrix61& C);
 
-		// up to 1 constraint limiting angle between twist axes
 		void AddAngularConstraints_Cone(
 			const FPBDJointSolverSettings& SolverSettings,
-			const FPBDJointSettings& JointSettings)
-		{
-			// Calculate Swing axis for each body
-			const FRotation3 R01 = Rs[0].Inverse() * Rs[1];
-			FRotation3 R01Twist, R01Swing;
-			R01.ToSwingTwist(FJointConstants::TwistAxis(), R01Swing, R01Twist);
-			R01Swing = R01Swing.GetNormalized();
-			R01Twist = R01Twist.GetNormalized();
-
-			FVec3 SwingAxis01;
-			FReal SwingAngle;
-			R01Swing.ToAxisAndAngleSafe(SwingAxis01, SwingAngle, FJointConstants::Swing1Axis(), SolverSettings.SwingTwistAngleTolerance);
-			if (SwingAngle > PI)
-			{
-				SwingAngle = SwingAngle - (FReal)2 * PI;
-			}
-
-			// Calculate swing limit for the current swing axis
-			FReal SwingAngleMax = FLT_MAX;
-			const FReal Swing1Limit = JointSettings.Motion.AngularLimits[(int32)EJointAngularConstraintIndex::Swing1];
-			const FReal Swing2Limit = JointSettings.Motion.AngularLimits[(int32)EJointAngularConstraintIndex::Swing2];
-
-			// Circular swing limit
-			SwingAngleMax = Swing1Limit;
-
-			// Elliptical swing limit
-			if (!FMath::IsNearlyEqual(Swing1Limit, Swing2Limit, KINDA_SMALL_NUMBER))
-			{
-				// Map swing axis to ellipse and calculate limit for this swing axis
-				const FReal DotSwing1 = FMath::Abs(FVec3::DotProduct(SwingAxis01, FJointConstants::Swing1Axis()));
-				const FReal DotSwing2 = FMath::Abs(FVec3::DotProduct(SwingAxis01, FJointConstants::Swing2Axis()));
-				SwingAngleMax = FMath::Sqrt(Swing1Limit * DotSwing2 * Swing1Limit * DotSwing2 + Swing2Limit * DotSwing1 * Swing2Limit * DotSwing1);
-			}
-
-			bool bConstraintActive = (SwingAngle <= -SwingAngleMax) || (SwingAngle >= SwingAngleMax);
-			if (bConstraintActive)
-			{
-				const int32 ConstraintIndex = NumAngularConstraints++;
-				AngularConstraintAxes[ConstraintIndex][0] = Rs[0] * SwingAxis01;
-				AngularConstraintAxes[ConstraintIndex][1] = AngularConstraintAxes[ConstraintIndex][0];
-				AngularConstraintDistances[ConstraintIndex] = SwingAngle;
-				AngularConstraintErrors[ConstraintIndex] = (SwingAngle >= SwingAngleMax) ? SwingAngle - SwingAngleMax : SwingAngle + SwingAngleMax;
-			}
-		}
+			const FPBDJointSettings& JointSettings,
+			FDenseMatrix66& J0,
+			FDenseMatrix66& J1,
+			FDenseMatrix61& C);
 
-		// up to 1 constraint limiting rotation about swing axis (relative to body 0)
 		void AddAngularConstraints_Swing(
 			const FPBDJointSolverSettings& SolverSettings,
 			const FPBDJointSettings& JointSettings,
 			const EJointAngularConstraintIndex SwingConstraintIndex,
-			const EJointAngularAxisIndex SwingAxisIndex)
-		{
-			// Calculate the swing axis for each body
-			const FRotation3 R01 = Rs[0].Inverse() * Rs[1];
-			FRotation3 R01Twist, R01Swing;
-			R01.ToSwingTwist(FJointConstants::TwistAxis(), R01Swing, R01Twist);
-			R01Swing = R01Swing.GetNormalized();
-			R01Twist = R01Twist.GetNormalized();
-
-			FVec3 TwistAxis01;
-			FReal TwistAngle;
-			R01Twist.ToAxisAndAngleSafe(TwistAxis01, TwistAngle, FJointConstants::TwistAxis(), SolverSettings.SwingTwistAngleTolerance);
-			if (TwistAngle > PI)
-			{
-				TwistAngle = TwistAngle - (FReal)2 * PI;
-			}
-			if (FVec3::DotProduct(TwistAxis01, FJointConstants::TwistAxis()) < 0)
-			{
-				TwistAxis01 = -TwistAxis01;
-				TwistAngle = -TwistAngle;
-			}
-			const FVec3 TwistAxis = Rs[0] * TwistAxis01;
-
-			const FRotation3 R1NoTwist = Rs[1] * R01Twist.Inverse();
-			const FMatrix33 Axes0 = Rs[0].ToMatrix();
-			const FMatrix33 Axes1 = R1NoTwist.ToMatrix();
-			FVec3 SwingCross = FVec3::CrossProduct(Axes0.GetAxis((int32)SwingAxisIndex), Axes1.GetAxis((int32)SwingAxisIndex));
-			SwingCross = SwingCross - FVec3::DotProduct(TwistAxis, SwingCross) * TwistAxis;
-			const FReal SwingCrossLen = SwingCross.Size();
-			if (SwingCrossLen > KINDA_SMALL_NUMBER)
-			{
-				FReal SwingAngle = FMath::Asin(FMath::Clamp(SwingCrossLen, (FReal)0, (FReal)1));
-				const FReal SwingDot = FVec3::DotProduct(Axes0.GetAxis((int32)SwingAxisIndex), Axes1.GetAxis((int32)SwingAxisIndex));
-				if (SwingDot < (FReal)0)
-				{
-					SwingAngle = (FReal)PI - SwingAngle;
-				}
-
-				FReal SwingAngleMax = FLT_MAX;
-				if (JointSettings.Motion.AngularMotionTypes[(int32)SwingConstraintIndex] == EJointMotionType::Limited)
-				{
-					FReal SwingLimit = JointSettings.Motion.AngularLimits[(int32)SwingConstraintIndex];
-					SwingAngleMax = SwingLimit;
-				}
-				else if (JointSettings.Motion.AngularMotionTypes[(int32)SwingConstraintIndex] == EJointMotionType::Locked)
-				{
-					SwingAngleMax = 0;
-				}
-
-				bool bConstraintActive = (SwingAngle <= -SwingAngleMax) || (SwingAngle >= SwingAngleMax);
-				if (bConstraintActive)
-				{
-					int32 ConstraintIndex = NumAngularConstraints++;
-
-					const FVec3 SwingAxis = SwingCross / SwingCrossLen;
-					AngularConstraintAxes[ConstraintIndex][0] = SwingAxis;
-					AngularConstraintAxes[ConstraintIndex][1] = SwingAxis;
-					AngularConstraintDistances[ConstraintIndex] = SwingAngle;
-					AngularConstraintErrors[ConstraintIndex] = (SwingAngle >= SwingAngleMax) ? SwingAngle - SwingAngleMax : SwingAngle + SwingAngleMax;
-				}
-			}
-		}
+			const EJointAngularAxisIndex SwingAxisIndex,
+			FDenseMatrix66& J0,
+			FDenseMatrix66& J1,
+			FDenseMatrix61& C);
 
-		// Add linear constraints to solver
 		void AddLinearConstraints(
 			const FPBDJointSolverSettings& SolverSettings,
-			const FPBDJointSettings& JointSettings)
-		{
-			NumLinearConstraints = 0;
-
-			const TVector<EJointMotionType, 3>& Motion = JointSettings.Motion.LinearMotionTypes;
-			if ((Motion[0] == EJointMotionType::Locked) && (Motion[1] == EJointMotionType::Locked) && (Motion[2] == EJointMotionType::Locked))
-			{
-				AddLinearConstraints_Point(SolverSettings, JointSettings);
-			}
-			else if ((Motion[0] == EJointMotionType::Limited) && (Motion[1] == EJointMotionType::Limited) && (Motion[2] == EJointMotionType::Limited))
-			{
-				AddLinearConstraints_Sphere(SolverSettings, JointSettings);
-			}
-			else if ((Motion[1] == EJointMotionType::Limited) && (Motion[2] == EJointMotionType::Limited))
-			{
-				// Circular Limit (X Axis)
-				AddLinearConstraints_Cylinder(SolverSettings, JointSettings, Motion[0], Rs[0] * FVec3(1, 0, 0));
-			}
-			else if ((Motion[0] == EJointMotionType::Limited) && (Motion[2] == EJointMotionType::Limited))
-			{
-				// Circular Limit (Y Axis)
-				AddLinearConstraints_Cylinder(SolverSettings, JointSettings, Motion[1], Rs[0] * FVec3(0, 1, 0));
-			}
-			else if ((Motion[0] == EJointMotionType::Limited) && (Motion[1] == EJointMotionType::Limited))
-			{
-				// Circular Limit (Z Axis)
-				AddLinearConstraints_Cylinder(SolverSettings, JointSettings, Motion[2], Rs[0] * FVec3(0, 0, 1));
-			}
-			else
-			{
-				// Plane/Square/Cube Limits (no way to author square or cube limits, but would work if we wanted it)
-				if (Motion[0] != EJointMotionType::Free)
-				{
-					AddLinearConstraints_Plane(SolverSettings, JointSettings, Motion[0], Rs[0] * FVec3(1, 0, 0));
-				}
-				if (Motion[1] != EJointMotionType::Free)
-				{
-					AddLinearConstraints_Plane(SolverSettings, JointSettings, Motion[1], Rs[0] * FVec3(0, 1, 0));
-				}
-				if (Motion[2] != EJointMotionType::Free)
-				{
-					AddLinearConstraints_Plane(SolverSettings, JointSettings, Motion[2], Rs[0] * FVec3(0, 0, 1));
-				}
-			}
-		}
+			const FPBDJointSettings& JointSettings,
+			FDenseMatrix66& J0,
+			FDenseMatrix66& J1,
+			FDenseMatrix61& C);
 
-		// Add angular constraints to solver
 		void AddAngularConstraints(
 			const FPBDJointSolverSettings& SolverSettings,
-			const FPBDJointSettings& JointSettings)
-		{
-			NumAngularConstraints = 0;
-
-			EJointMotionType TwistMotion = JointSettings.Motion.AngularMotionTypes[(int32)EJointAngularConstraintIndex::Twist];
-			EJointMotionType Swing1Motion = JointSettings.Motion.AngularMotionTypes[(int32)EJointAngularConstraintIndex::Swing1];
-			EJointMotionType Swing2Motion = JointSettings.Motion.AngularMotionTypes[(int32)EJointAngularConstraintIndex::Swing2];
-
-			// Apply twist constraint
-			if (SolverSettings.bEnableTwistLimits)
-			{
-				if (TwistMotion != EJointMotionType::Free)
-				{
-					AddAngularConstraints_Twist(SolverSettings, JointSettings);
-				}
-			}
-
-			// Apply swing constraints
-			if (SolverSettings.bEnableSwingLimits)
-			{
-				if ((Swing1Motion == EJointMotionType::Limited) && (Swing2Motion == EJointMotionType::Limited))
-				{
-					AddAngularConstraints_Cone(SolverSettings, JointSettings);
-				}
-				else
-				{
-					if (Swing1Motion != EJointMotionType::Free)
-					{
-						AddAngularConstraints_Swing(SolverSettings, JointSettings, EJointAngularConstraintIndex::Swing1, EJointAngularAxisIndex::Swing1);
-					}
-					if (Swing2Motion != EJointMotionType::Free)
-					{
-						AddAngularConstraints_Swing(SolverSettings, JointSettings, EJointAngularConstraintIndex::Swing2, EJointAngularAxisIndex::Swing2);
-					}
-				}
-			}
-		}
+			const FPBDJointSettings& JointSettings,
+			FDenseMatrix66& J0,
+			FDenseMatrix66& J1,
+			FDenseMatrix61& C);
 
 		void BuildJacobianAndResidual(
 			const FPBDJointSolverSettings& SolverSettings,
 			const FPBDJointSettings& JointSettings,
 			FDenseMatrix66& J0,
 			FDenseMatrix66& J1,
-			FDenseMatrix61& C)
-		{
-			// Calculate constraint axes and errors
-			AddLinearConstraints(SolverSettings, JointSettings);
-			AddAngularConstraints(SolverSettings, JointSettings);
-
-			// The Jacobians will be some sub-set of the following rows, depending on which constraints are enabled
-			// and whether the limits are currently being violated.
-			//
-			// J0(Nx6) = | XAxis          -XAxis x Connector0 |
-			//           | YAxis          -YAxis x Connector0 |
-			//           | ZAxis          -ZAxis x Connector0 |
-			//           | 0              TwistAxis           |
-			//           | 0              Swing1Axis          |
-			//           | 0              Swing2Axis          |
-			//
-			// J0(Nx6) = | -XAxis         XAxis x Connector1  |
-			//           | -YAxis         YAxis x Connector1  |
-			//           | -ZAxis         ZAxis x Connector1  |
-			//           | 0              -TwistAxis          |
-			//           | 0              -Swing1Axis         |
-			//           | 0              -Swing2Axis         |
-			//
-			const int32 NumConstraints = NumLinearConstraints + NumAngularConstraints;
-			J0.SetDimensions(NumConstraints, 6);
-			J1.SetDimensions(NumConstraints, 6);
-			C.SetDimensions(NumConstraints, 1);
-
-			const FVec3 XP0 = Xs[0] - Ps[0];
-			const FVec3 XP1 = Xs[1] - Ps[1];
-
-			int32 ConstraintIndex = 0;
-			for (int32 LinearConstraintIndex = 0; LinearConstraintIndex < NumLinearConstraints; ++LinearConstraintIndex, ++ConstraintIndex)
-			{
-				const FVec3& ConstraintAxis = LinearConstraintAxes[LinearConstraintIndex];
-
-				J0.SetRowAt(ConstraintIndex, 0, ConstraintAxis);
-				J0.SetRowAt(ConstraintIndex, 3, -FVec3::CrossProduct(ConstraintAxis, XP0));
-
-				J1.SetRowAt(ConstraintIndex, 0, -ConstraintAxis);
-				J1.SetRowAt(ConstraintIndex, 3, FVec3::CrossProduct(ConstraintAxis, XP1));
-
-				C.At(ConstraintIndex, 0) = LinearConstraintErrors[LinearConstraintIndex];
-			}
-
-			for (int32 AngularConstraintIndex = 0; AngularConstraintIndex < NumAngularConstraints; ++AngularConstraintIndex, ++ConstraintIndex)
-			{
-				const FVec3& ConstraintAxis0 = AngularConstraintAxes[AngularConstraintIndex][0];
-				const FVec3& ConstraintAxis1 = AngularConstraintAxes[AngularConstraintIndex][1];
-
-				J0.SetRowAt(ConstraintIndex, 0, FVec3::ZeroVector);
-				J0.SetRowAt(ConstraintIndex, 3, ConstraintAxis0);
-
-				J1.SetRowAt(ConstraintIndex, 0, FVec3::ZeroVector);
-				J1.SetRowAt(ConstraintIndex, 3, -ConstraintAxis1);
-
-				C.At(ConstraintIndex, 0) = AngularConstraintErrors[AngularConstraintIndex];
-			}
-		}
+			FDenseMatrix61& C);
 
 		static const int32 MaxConstraintedBodies = 2;
-		static const int32 MaxLinearConstraints = 3;
-		static const int32 MaxAngularConstraints = 3;
-		static const int32 MaxConstraints = MaxLinearConstraints + MaxAngularConstraints;
-
-		// Joint solver state
-		// @todo(ccaulfield): maybe just store the Jacobians and errors here rather than the data to construct them
-		FVec3 LinearConstraintAxes[MaxConstraints];
-		FReal LinearConstraintDistances[MaxConstraints];
-		FReal LinearConstraintErrors[MaxConstraints];
-		FVec3 AngularConstraintAxes[MaxConstraints][MaxConstraintedBodies];	// @todo(ccaulfield): per body for twist - should probably switch back to shared twist axis
-		FReal AngularConstraintDistances[MaxConstraints];
-		FReal AngularConstraintErrors[MaxConstraints];
-		int32 NumLinearConstraints;
-		int32 NumAngularConstraints;
 
 		// Local-space constraint settings
 		FRigidTransform3 XLs[MaxConstraintedBodies];	// Local-space joint connector transforms
