Added relative error for most opposing face. Added comment about flaws.

diff --git a/Engine/Source/Runtime/Experimental/Chaos/Private/Chaos/Convex.cpp b/Engine/Source/Runtime/Experimental/Chaos/Private/Chaos/Convex.cpp
new file mode 100644
--- /dev/null
+++ b/Engine/Source/Runtime/Experimental/Chaos/Private/Chaos/Convex.cpp
@@ -0,0 +1,32 @@
+// Copyright 1998-2019 Epic Games, Inc. All Rights Reserved.
+#include "Chaos/Convex.h"
+
+namespace Chaos
+{
+	int32 FConvex::FindMostOpposingFace(const FVec3& Position, const FVec3& UnitDir, int32 HintFaceIndex, FReal SearchDist) const
+	{
+		//NOTE: this approach assumes we never have conincident planes, which at the moment is not true.
+		//Need to make sure convex hull is cleaned up so that there are n-gon faces so this is unique
+		SearchDist = FMath::Max(SearchDist,FMath::Abs(BoundingBox().Extents().GetAbsMax()) * 1e-4f);
+		//todo: use hill climbing
+		int32 MostOpposingIdx = INDEX_NONE;
+		FReal MostOpposingDot = TNumericLimits<FReal>::Max();
+		for (int32 Idx = 0; Idx < Planes.Num(); ++Idx)
+		{
+			const TPlane<FReal, 3>& Plane = Planes[Idx];
+			const FReal Distance = Plane.SignedDistance(Position);
+			if (FMath::Abs(Distance) < SearchDist)
+			{
+				// TPlane has an override for Normal() that doesn't call PhiWithNormal().
+				const FReal Dot = FVec3::DotProduct(Plane.Normal(), UnitDir);
+				if (Dot < MostOpposingDot)
+				{
+					MostOpposingDot = Dot;
+					MostOpposingIdx = Idx;
+				}
+			}
+		}
+		ensure(MostOpposingIdx != INDEX_NONE);
+		return MostOpposingIdx;
+	}
+}
\ No newline at end of file
diff --git a/Engine/Source/Runtime/Experimental/Chaos/Public/Chaos/Convex.h b/Engine/Source/Runtime/Experimental/Chaos/Public/Chaos/Convex.h
--- a/Engine/Source/Runtime/Experimental/Chaos/Public/Chaos/Convex.h
+++ b/Engine/Source/Runtime/Experimental/Chaos/Public/Chaos/Convex.h
@@ -11,7 +11,7 @@
 
 namespace Chaos
 {
-	class FConvex final : public FImplicitObject
+	class CHAOS_API FConvex final : public FImplicitObject
 	{
 	public:
 		using FImplicitObject::GetTypeName;
@@ -139,29 +139,8 @@ namespace Chaos
 			return MakePair(FVec3(0), false);
 		}
 
-		virtual int32 FindMostOpposingFace(const FVec3& Position, const FVec3& UnitDir, int32 HintFaceIndex, FReal SearchDist) const override
-		{
-			//todo: use hill climbing
-			int32 MostOpposingIdx = INDEX_NONE;
-			FReal MostOpposingDot = TNumericLimits<FReal>::Max();
-			for(int32 Idx = 0; Idx < Planes.Num(); ++Idx)
-			{
-				const TPlane<FReal, 3>& Plane = Planes[Idx];
-				const FReal Distance = Plane.SignedDistance(Position);
-				if (FMath::Abs(Distance) < SearchDist)
-				{
-					// TPlane has an override for Normal() that doesn't call PhiWithNormal().
-					const FReal Dot = FVec3::DotProduct(Plane.Normal(), UnitDir);
-					if (Dot < MostOpposingDot)
-					{
-						MostOpposingDot = Dot;
-						MostOpposingIdx = Idx;
-					}
-				}
-			}
-			ensure(MostOpposingIdx != INDEX_NONE);
-			return MostOpposingIdx;
-		}
+		virtual int32 FindMostOpposingFace(const FVec3& Position, const FVec3& UnitDir, int32 HintFaceIndex, FReal SearchDist) const override;
+		
 
 		FVec3 FindGeometryOpposingNormal(const FVec3& DenormDir, int32 FaceIndex, const FVec3& OriginalNormal) const
 		{
