Added very simple collision filtering to chaos allowing single shapes to completely ignore collisions.

diff --git a/Engine/Source/Runtime/Core/Public/UObject/ExternalPhysicsCustomObjectVersion.h b/Engine/Source/Runtime/Core/Public/UObject/ExternalPhysicsCustomObjectVersion.h
--- a/Engine/Source/Runtime/Core/Public/UObject/ExternalPhysicsCustomObjectVersion.h
+++ b/Engine/Source/Runtime/Core/Public/UObject/ExternalPhysicsCustomObjectVersion.h
@@ -74,6 +74,9 @@ struct CORE_API FExternalPhysicsCustomObjectVersion
 
 		// Add mass transform data to kinematic particle
 		KinematicCentersOfMass,
+
+		// Added ability to remove shapes from collision resolution (will not construct constraints when one or more shapes removed)
+		AddShapeCollisionDisable,
 		
 		// -----<new versions can be added above this line>-------------------------------------------------
 		VersionPlusOne,
diff --git a/Engine/Source/Runtime/Engine/Private/PhysicsEngine/BodyInstance.cpp b/Engine/Source/Runtime/Engine/Private/PhysicsEngine/BodyInstance.cpp
--- a/Engine/Source/Runtime/Engine/Private/PhysicsEngine/BodyInstance.cpp
+++ b/Engine/Source/Runtime/Engine/Private/PhysicsEngine/BodyInstance.cpp
@@ -1960,7 +1960,7 @@ void FBodyInstance::SetInstanceSimulatePhysics(bool bSimulate, bool bMaintainPhy
 				FMessageLog("PIE").Warning(FText::Format(LOCTEXT("SimPhysNoBody", "Trying to simulate physics on ''{0}'' but no physics body."),
 					FText::FromString(GetPathNameSafe(OwnerComponentInst))));
 			}
-			else if (!IsDynamic())
+			else if (FPhysicsInterface::IsStatic(ActorHandle))
 			{
 				FMessageLog("PIE").Warning(FText::Format(LOCTEXT("SimPhysStatic", "Trying to simulate physics on ''{0}'' but it is static."),
 					FText::FromString(GetPathNameSafe(OwnerComponentInst))));
diff --git a/Engine/Source/Runtime/Engine/Private/PhysicsEngine/Experimental/PhysInterface_Chaos.cpp b/Engine/Source/Runtime/Engine/Private/PhysicsEngine/Experimental/PhysInterface_Chaos.cpp
--- a/Engine/Source/Runtime/Engine/Private/PhysicsEngine/Experimental/PhysInterface_Chaos.cpp
+++ b/Engine/Source/Runtime/Engine/Private/PhysicsEngine/Experimental/PhysInterface_Chaos.cpp
@@ -1872,7 +1872,7 @@ int32 GetAllShapesInternal_AssumedLocked(const FPhysicsActorHandle& InActorHandl
 	OutShapes.Reset();
 	const Chaos::TShapesArray<float, 3>& ShapesArray = InActorHandle->ShapesArray();
 	//todo: can we avoid this construction?
-	for (const auto& Shape : ShapesArray)
+	for (const TUniquePtr<Chaos::TPerShapeData<float, 3>>& Shape : ShapesArray)
 	{
 		OutShapes.Add(FPhysicsShapeReference_Chaos(Shape.Get(), true, true, InActorHandle));
 	}
diff --git a/Engine/Source/Runtime/Experimental/Chaos/Private/Chaos/CollisionResolution.cpp b/Engine/Source/Runtime/Experimental/Chaos/Private/Chaos/CollisionResolution.cpp
--- a/Engine/Source/Runtime/Experimental/Chaos/Private/Chaos/CollisionResolution.cpp
+++ b/Engine/Source/Runtime/Experimental/Chaos/Private/Chaos/CollisionResolution.cpp
@@ -1039,6 +1039,15 @@ namespace Chaos
 		template<typename T, int d>
 		void ConstructPairConstraintImpl(TGeometryParticleHandle<T, d>* Particle0, TGeometryParticleHandle<T, d>* Particle1, const FImplicitObject* Implicit0, const FImplicitObject* Implicit1, const TRigidTransform<T, d>& Transform0, const TRigidTransform<T, d>& Transform1, const T Thickness, FCollisionConstraintsArray& NewConstraints)
 		{
+			const TPerShapeData<T, d>* Shape0 = Particle0->GetImplicitShape(Implicit0);
+			const TPerShapeData<T, d>* Shape1 = Particle1->GetImplicitShape(Implicit1);
+
+			// If either shape is disabled for collision bail without constructing a constraint
+			if((Shape0 && Shape0->bDisable) || (Shape1 && Shape1->bDisable))
+			{
+				return;
+			}
+
 			// See if we already have a constraint for this shape pair
 			// todo(brice) : Not sure this actually prevents duplicate constraints.
 			for (int32 i = 0; i < NewConstraints.Num(); i++)
diff --git a/Engine/Source/Runtime/Experimental/Chaos/Private/Chaos/GeometryParticles.cpp b/Engine/Source/Runtime/Experimental/Chaos/Private/Chaos/GeometryParticles.cpp
--- a/Engine/Source/Runtime/Experimental/Chaos/Private/Chaos/GeometryParticles.cpp
+++ b/Engine/Source/Runtime/Experimental/Chaos/Private/Chaos/GeometryParticles.cpp
@@ -46,7 +46,8 @@ namespace Chaos
 
 	template <typename T, int d>
 	TPerShapeData<T, d>::TPerShapeData()
-	: UserData(nullptr)
+		: UserData(nullptr)
+		, bDisable(false)
 	{
 	}
 
@@ -99,6 +100,11 @@ namespace Chaos
 		{
 			Ar << Materials;
 		}
+
+		if(Ar.CustomVer(FExternalPhysicsCustomObjectVersion::GUID) >= FExternalPhysicsCustomObjectVersion::AddShapeCollisionDisable)
+		{
+			Ar << bDisable;
+		}
 	}
 
 
diff --git a/Engine/Source/Runtime/Experimental/Chaos/Public/Chaos/GeometryParticles.h b/Engine/Source/Runtime/Experimental/Chaos/Public/Chaos/GeometryParticles.h
--- a/Engine/Source/Runtime/Experimental/Chaos/Public/Chaos/GeometryParticles.h
+++ b/Engine/Source/Runtime/Experimental/Chaos/Public/Chaos/GeometryParticles.h
@@ -35,6 +35,7 @@ namespace Chaos
 		TSerializablePtr<FImplicitObject> Geometry;
 		TAABB<FReal, 3> WorldSpaceInflatedShapeBounds;
 		TArray<FMaterialHandle> Materials;
+		bool bDisable;
 
 		void UpdateShapeBounds(const TRigidTransform<T, d>& WorldTM);
 
@@ -97,6 +98,7 @@ namespace Chaos
 			TArrayCollection::AddArray(&MDynamicGeometry);
 			TArrayCollection::AddArray(&MParticleIDs);
 			TArrayCollection::AddArray(&MShapesArray);
+			TArrayCollection::AddArray(&ImplicitShapeMap);
 			TArrayCollection::AddArray(&MLocalBounds);
 			TArrayCollection::AddArray(&MWorldSpaceInflatedBounds);
 			TArrayCollection::AddArray(&MHasBounds);
@@ -123,6 +125,7 @@ namespace Chaos
 			, MGeometryParticleHandle(MoveTemp(Other.MGeometryParticleHandle))
 			, MGeometryParticle(MoveTemp(Other.MGeometryParticle))
 			, MShapesArray(MoveTemp(Other.MShapesArray))
+			, ImplicitShapeMap(MoveTemp(Other.ImplicitShapeMap))
 			, MLocalBounds(MoveTemp(Other.MLocalBounds))
 			, MWorldSpaceInflatedBounds(MoveTemp(Other.MWorldSpaceInflatedBounds))
 			, MHasBounds(MoveTemp(Other.MHasBounds))
@@ -138,6 +141,7 @@ namespace Chaos
 			TArrayCollection::AddArray(&MSharedGeometry);
 			TArrayCollection::AddArray(&MDynamicGeometry);
 			TArrayCollection::AddArray(&MShapesArray);
+			TArrayCollection::AddArray(&ImplicitShapeMap);
 			TArrayCollection::AddArray(&MLocalBounds);
 			TArrayCollection::AddArray(&MWorldSpaceInflatedBounds);
 			TArrayCollection::AddArray(&MHasBounds);
@@ -168,6 +172,7 @@ namespace Chaos
 			TArrayCollection::AddArray(&MSharedGeometry);
 			TArrayCollection::AddArray(&MDynamicGeometry);
 			TArrayCollection::AddArray(&MShapesArray);
+			TArrayCollection::AddArray(&ImplicitShapeMap);
 			TArrayCollection::AddArray(&MLocalBounds);
 			TArrayCollection::AddArray(&MWorldSpaceInflatedBounds);
 			TArrayCollection::AddArray(&MHasBounds);
@@ -212,6 +217,7 @@ namespace Chaos
 			MGeometry[Index] = MakeSerializable(InUnique);
 			MDynamicGeometry[Index] = MoveTemp(InUnique);
 			UpdateShapesArray(Index);
+			MapImplicitShapes(Index);
 		}
 
 		CHAOS_API void SetSharedGeometry(const int32 Index, TSharedPtr<FImplicitObject, ESPMode::ThreadSafe> InShared)
@@ -220,6 +226,7 @@ namespace Chaos
 			MGeometry[Index] = MakeSerializable(InShared);
 			MSharedGeometry[Index] = InShared;
 			UpdateShapesArray(Index);
+			MapImplicitShapes(Index);
 		}
 		
 		CHAOS_API void SetGeometry(const int32 Index, TSerializablePtr<FImplicitObject> InGeometry)
@@ -228,6 +235,7 @@ namespace Chaos
 			check(!SharedGeometry(Index));
 			MGeometry[Index] = InGeometry;
 			UpdateShapesArray(Index);
+			MapImplicitShapes(Index);
 		}
 
 		CHAOS_API const TBox<T,d>& LocalBounds(const int32 Index) const
@@ -326,6 +334,11 @@ namespace Chaos
 			if (Ar.CustomVer(FPhysicsObjectVersion::GUID) >= FPhysicsObjectVersion::PerShapeData)
 			{
 				Ar << MShapesArray;
+
+				if(Ar.IsLoading())
+				{
+					MapImplicitShapes();
+				}
 			}
 
 			if (Ar.CustomVer(FPhysicsObjectVersion::GUID) >= FPhysicsObjectVersion::SerializeGTGeometryParticles)
@@ -387,6 +400,21 @@ namespace Chaos
 
 		CHAOS_API EParticleType ParticleType() const { return MParticleType; }
 
+		const TPerShapeData<T, d>* GetImplicitShape(int32 Index, const FImplicitObject* InObject)
+		{
+			checkSlow(Index >= 0 && Index < ImplicitShapeMap.Num());
+			TMap<const FImplicitObject*, int32>& Mapping = ImplicitShapeMap[Index];
+			TShapesArray<T, d>& ShapeArray = MShapesArray[Index];
+			int32* ShapeIndex = Mapping.Find(InObject);
+
+			if(ShapeIndex && ShapeArray.IsValidIndex(*ShapeIndex))
+			{
+				return ShapeArray[*ShapeIndex].Get();
+			}
+
+			return nullptr;
+		}
+
 	protected:
 		EParticleType MParticleType;
 
@@ -404,6 +432,7 @@ namespace Chaos
 		TArrayCollectionArray<TSerializablePtr<TGeometryParticleHandle<T, d>>> MGeometryParticleHandle;
 		TArrayCollectionArray<TGeometryParticle<T, d>*> MGeometryParticle;
 		TArrayCollectionArray<TShapesArray<T,d>> MShapesArray;
+		TArrayCollectionArray<TMap<const FImplicitObject*, int32>> ImplicitShapeMap;
 		TArrayCollectionArray<TBox<T,d>> MLocalBounds;
 		TArrayCollectionArray<TBox<T, d>> MWorldSpaceInflatedBounds;
 		TArrayCollectionArray<bool> MHasBounds;
@@ -413,6 +442,32 @@ namespace Chaos
 		void UpdateShapesArray(const int32 Index)
 		{
 			UpdateShapesArrayFromGeometry(MShapesArray[Index], MGeometry[Index], FRigidTransform3(X(Index), R(Index)));
+			MapImplicitShapes(Index);
+		}
+
+		void MapImplicitShapes()
+		{
+			int32 NumShapeArrays = MShapesArray.Num();
+			ImplicitShapeMap.Resize(NumShapeArrays);
+			for(int32 Index = 0; Index < NumShapeArrays; ++Index)
+			{
+				MapImplicitShapes(Index);
+			}
+		}
+
+		void MapImplicitShapes(int32 Index)
+		{
+			checkSlow(Index >= 0 && Index < ImplicitShapeMap.Num());
+
+			TMap<const FImplicitObject*, int32>& Mapping = ImplicitShapeMap[Index];
+			TShapesArray<T, d>& ShapeArray = MShapesArray[Index];
+			Mapping.Reset();
+
+			int32 ShapeIndex = 0;
+			for(TUniquePtr<TPerShapeData<T, d>>& Shape : ShapeArray)
+			{
+				Mapping.Add(Shape->Geometry.Get(), ShapeIndex++);
+			}
 		}
 
 		template <typename T2, int d2, EGeometryParticlesSimType SimType2>
diff --git a/Engine/Source/Runtime/Experimental/Chaos/Public/Chaos/ParticleDirtyFlags.h b/Engine/Source/Runtime/Experimental/Chaos/Public/Chaos/ParticleDirtyFlags.h
--- a/Engine/Source/Runtime/Experimental/Chaos/Public/Chaos/ParticleDirtyFlags.h
+++ b/Engine/Source/Runtime/Experimental/Chaos/Public/Chaos/ParticleDirtyFlags.h
@@ -9,33 +9,34 @@ namespace Chaos
 	// but clean property values will get overridden with physics thread results.
 	enum EParticleFlags : int32
 	{
-		X =				 1 << 0,
-		R =				 1 << 1,
-		V =				 1 << 2,
-		W =				 1 << 3,
-		CenterOfMass =   1 << 4,
-		RotationOfMass = 1 << 5,
-		CollisionGroup = 1 << 6,
-		Disabled =		 1 << 7,
-		PreV =			 1 << 8,
-		PreW =			 1 << 9,
-		P =				 1 << 10,
-		Q =				 1 << 11,
-		F =				 1 << 12,
-		Torque =		 1 << 13,
-		I =				 1 << 14,
-		InvI =			 1 << 15,
-		M =				 1 << 16,
-		InvM =			 1 << 17,
-		ObjectState =	 1 << 18,
-		Geometry =		 1 << 19,
-		ExternalForce =  1 << 20,
-		ExternalTorque = 1 << 20,
-		GravityEnabled = 1 << 22,
-		SpatialIdx     = 1 << 23,
-		HashResult     = 1 << 24
+		X						= 1 << 0,
+		R						= 1 << 1,
+		V						= 1 << 2,
+		W						= 1 << 3,
+		CenterOfMass			= 1 << 4,
+		RotationOfMass			= 1 << 5,
+		CollisionGroup			= 1 << 6,
+		Disabled				= 1 << 7,
+		PreV					= 1 << 8,
+		PreW					= 1 << 9,
+		P						= 1 << 10,
+		Q						= 1 << 11,
+		F						= 1 << 12,
+		Torque					= 1 << 13,
+		I						= 1 << 14,
+		InvI					= 1 << 15,
+		M						= 1 << 16,
+		InvM					= 1 << 17,
+		ObjectState				= 1 << 18,
+		Geometry				= 1 << 19,
+		ExternalForce			= 1 << 20,
+		ExternalTorque			= 1 << 20,
+		GravityEnabled			= 1 << 22,
+		SpatialIdx				= 1 << 23,
+		HashResult				= 1 << 24,
+		ShapeDisableCollision	= 1 << 25
 #if CHAOS_CHECKED
-		, DebugName    = 1 << 25
+		, DebugName				= 1 << 26
 #endif
 	};
 
diff --git a/Engine/Source/Runtime/Experimental/Chaos/Public/Chaos/ParticleHandle.h b/Engine/Source/Runtime/Experimental/Chaos/Public/Chaos/ParticleHandle.h
--- a/Engine/Source/Runtime/Experimental/Chaos/Public/Chaos/ParticleHandle.h
+++ b/Engine/Source/Runtime/Experimental/Chaos/Public/Chaos/ParticleHandle.h
@@ -466,6 +466,11 @@ public:
 		GeometryParticles->SetHandle(ParticleIdx, this);
 	}
 
+	const TPerShapeData<T, d>* GetImplicitShape(const FImplicitObject* InObject)
+	{
+		return GeometryParticles->GetImplicitShape(ParticleIdx, InObject);
+	}
+
 protected:
 
 	friend TGeometryParticleHandles<T, d>;
@@ -1234,6 +1239,11 @@ public:
 				}
 			}
 		}
+
+		if(Ar.IsLoading())
+		{
+			MapImplicitShapes();
+		}
 	}
 
 	virtual bool IsParticleValid() const
@@ -1312,6 +1322,7 @@ public:
 	{
 		ensure(InShapesArray.Num() == MShapesArray.Num());
 		MShapesArray = MoveTemp(InShapesArray);
+		MapImplicitShapes();
 	}
 
 	TSerializablePtr<FImplicitObject> Geometry() const { return MakeSerializable(MGeometry); }
@@ -1339,6 +1350,16 @@ public:
 		MSpatialIdx = Idx;
 	}
 
+	void SetShapeCollisionDisable(int32 InShapeIndex, bool bInDisable)
+	{
+		const bool bCurrent = MShapesArray[InShapeIndex]->bDisable;
+		if(bCurrent != bInDisable)
+		{
+			MShapesArray[InShapeIndex]->bDisable = bInDisable;
+			MarkDirty(EParticleFlags::ShapeDisableCollision);
+		}
+	}
+
 	FParticleData* NewData() const { return new TGeometryParticleData<T, d>( *this ); }
 
 	bool IsDirty() const
@@ -1371,6 +1392,17 @@ public:
 		return nullptr;
 	}
 
+	const TPerShapeData<T, d>* GetImplicitShape(const FImplicitObject* InImplicit) const
+	{
+		const int32* ShapeIndex = ImplicitShapeMap.Find(InImplicit);
+		if(ShapeIndex)
+		{
+			return MShapesArray[*ShapeIndex].Get();
+		}
+
+		return nullptr;
+	}
+
 	// Pointer to any data that the solver wants to associate with this particle
 	// TODO: It's important to eventually hide this!
 	// Right now it's exposed to lubricate the creation of the whole proxy system.
@@ -1381,6 +1413,7 @@ private:
 	TRotation<T, d> MR;
 	TSharedPtr<FImplicitObject, ESPMode::ThreadSafe> MGeometry;	//TODO: geometry should live in bodysetup
 	TShapesArray<T,d> MShapesArray;
+	TMap<const FImplicitObject*, int32> ImplicitShapeMap;
 	FSpatialAccelerationIdx MSpatialIdx;
 
 	// This value is generated and used in AccelerationStructureHandle for a hash, as that handle
@@ -1410,6 +1443,17 @@ protected:
 	void UpdateShapesArray()
 	{
 		UpdateShapesArrayFromGeometry<T, d>(MShapesArray, MakeSerializable(MGeometry), FRigidTransform3(X(), R()));
+		MapImplicitShapes();
+	}
+
+	void MapImplicitShapes()
+	{
+		ImplicitShapeMap.Reset();
+		int32 ShapeIndex = 0;
+		for(TUniquePtr<TPerShapeData<T, d>>& ShapeData : MShapesArray)
+		{
+			ImplicitShapeMap.Add(ShapeData->Geometry.Get(), ShapeIndex++);
+		}
 	}
 };
 
@@ -1874,6 +1918,13 @@ public:
 		, MGravityEnabled(InParticle.IsGravityEnabled())
 	{
 		Type = EParticleType::Rigid;
+
+		const TShapesArray<T, d>& Shapes = InParticle.ShapesArray();
+		ShapeCollisionDisableFlags.Empty(Shapes.Num());
+		for(const TUniquePtr<TPerShapeData<T, d>>& ShapePtr : Shapes)
+		{
+			ShapeCollisionDisableFlags.Add(ShapePtr->bDisable);
+		}
 	}
 
 
@@ -1894,6 +1945,7 @@ public:
 	bool MDisabled;
 	bool MToBeRemovedOnFracture;
 	bool MGravityEnabled;
+	TBitArray<> ShapeCollisionDisableFlags;
 
 	void Reset() {
 		TKinematicGeometryParticleData<T, d>::Reset();
@@ -1915,6 +1967,7 @@ public:
 		MDisabled = false;
 		MToBeRemovedOnFracture = false;
 		MGravityEnabled = false;
+		ShapeCollisionDisableFlags.Reset();
 	}
 };
 
diff --git a/Engine/Source/Runtime/Experimental/ChaosSolvers/Private/PhysicsProxy/SingleParticlePhysicsProxy.cpp b/Engine/Source/Runtime/Experimental/ChaosSolvers/Private/PhysicsProxy/SingleParticlePhysicsProxy.cpp
--- a/Engine/Source/Runtime/Experimental/ChaosSolvers/Private/PhysicsProxy/SingleParticlePhysicsProxy.cpp
+++ b/Engine/Source/Runtime/Experimental/ChaosSolvers/Private/PhysicsProxy/SingleParticlePhysicsProxy.cpp
@@ -260,6 +260,15 @@ void FSingleParticlePhysicsProxy<Chaos::TPBDRigidParticle<float, 3>>::PushToPhys
 				//todo: compute bounds based on sample points
 			}
 		}
+
+		if(Data->DirtyFlags.IsDirty(Chaos::EParticleFlags::ShapeDisableCollision))
+		{
+			int32 CurrShape = 0;
+			for(const TUniquePtr<Chaos::TPerShapeData<Chaos::FReal, 3>>& Shape : RigidHandle->ShapesArray())
+			{
+				Shape->bDisable = Data->ShapeCollisionDisableFlags[CurrShape++];
+			}
+		}
 	}
 }
 
