Chaos - separate BroadPhase and NarrowPhase from CollisionConstraints

diff --git a/Engine/Source/Runtime/Engine/Private/PhysicsEngine/ImmediatePhysics/ImmediatePhysicsChaos/ImmediatePhysicsSimulation_Chaos.cpp b/Engine/Source/Runtime/Engine/Private/PhysicsEngine/ImmediatePhysics/ImmediatePhysicsChaos/ImmediatePhysicsSimulation_Chaos.cpp
--- a/Engine/Source/Runtime/Engine/Private/PhysicsEngine/ImmediatePhysics/ImmediatePhysicsChaos/ImmediatePhysicsSimulation_Chaos.cpp
+++ b/Engine/Source/Runtime/Engine/Private/PhysicsEngine/ImmediatePhysics/ImmediatePhysicsChaos/ImmediatePhysicsSimulation_Chaos.cpp
@@ -167,7 +167,7 @@ namespace ImmediatePhysics_Chaos
 
 		// Filter collisions after detection
 		// @todo(ccaulfield): Eventually we will build lists of potentially colliding pairs and won't need this
-		Collisions.SetPostComputeCallback(
+		Evolution->SetPostDetectCollisionsCallback(
 			[this]()
 			{
 				Evolution->GetCollisionConstraints().ApplyCollisionModifier(
@@ -243,7 +243,7 @@ namespace ImmediatePhysics_Chaos
 	{
 		using namespace Chaos;
 
-		Evolution->GetCollisionConstraints().ClearPostComputeCallback();
+		//Evolution->ClearDetectCollisionsCallback();
 		Evolution->GetCollisionConstraints().ClearPostApplyCallback();
 		Evolution->GetCollisionConstraints().ClearPostApplyPushOutCallback();
 
@@ -349,6 +349,28 @@ namespace ImmediatePhysics_Chaos
 			IgnoreCollisionParticlePairTable.FindOrAdd(ParticleA).Add(ParticleB); 
 			IgnoreCollisionParticlePairTable.FindOrAdd(ParticleB).Add(ParticleA);
 		}
+
+		TArray<TVector<const TGeometryParticleHandle<FReal, Dimensions>*, 2>> BroadPhaseParticlePairs;
+		int NumActorHandles = ActorHandles.Num();
+		for (int ActorHandleIndex0 = 0; ActorHandleIndex0 < NumActorHandles; ++ActorHandleIndex0)
+		{
+			FActorHandle* ActorHandle0 = ActorHandles[ActorHandleIndex0];
+			const TGeometryParticleHandle<FReal, Dimensions>* Particle0 = ActorHandle0->GetParticle();
+
+			for (int ActorHandleIndex1 = ActorHandleIndex0 + 1; ActorHandleIndex1 < NumActorHandles; ++ActorHandleIndex1)
+			{
+				FActorHandle* ActorHandle1 = ActorHandles[ActorHandleIndex1];
+				const TGeometryParticleHandle<FReal, Dimensions>* Particle1 = ActorHandle1->GetParticle();
+
+				const TSet<const FParticleHandle*>* Particle0IgnoreSet = IgnoreCollisionParticlePairTable.Find(Particle0);
+				bool bIgnoreActorHandlePair = (Particle0IgnoreSet != nullptr)  && Particle0IgnoreSet->Contains(Particle1);
+				if (!bIgnoreActorHandlePair)
+				{
+					BroadPhaseParticlePairs.Emplace(TVector<const TGeometryParticleHandle<FReal, Dimensions>*, 2>(Particle0, Particle1));
+				}
+			}
+		}
+
 	}
 
 	void FSimulation::SetIgnoreCollisionActors(const TArray<FActorHandle*>& InIgnoreCollisionActors)
diff --git a/Engine/Source/Runtime/Experimental/Chaos/Private/Chaos/Collision/CollisionDetector.cpp b/Engine/Source/Runtime/Experimental/Chaos/Private/Chaos/Collision/CollisionDetector.cpp
new file mode 100644
--- /dev/null
+++ b/Engine/Source/Runtime/Experimental/Chaos/Private/Chaos/Collision/CollisionDetector.cpp
@@ -0,0 +1,9 @@
+// Copyright 1998-2019 Epic Games, Inc. All Rights Reserved.
+#include "Chaos/Collision/CollisionDetector.h"
+#include "ChaosLog.h"
+#include "ChaosStats.h"
+
+namespace Chaos
+{
+	DEFINE_STAT(STAT_DetectCollisions);
+}
\ No newline at end of file
diff --git a/Engine/Source/Runtime/Experimental/Chaos/Private/Chaos/Collision/SpatialAccelerationBroadPhase.cpp b/Engine/Source/Runtime/Experimental/Chaos/Private/Chaos/Collision/SpatialAccelerationBroadPhase.cpp
new file mode 100644
--- /dev/null
+++ b/Engine/Source/Runtime/Experimental/Chaos/Private/Chaos/Collision/SpatialAccelerationBroadPhase.cpp
@@ -0,0 +1,6 @@
+// Copyright 1998-2019 Epic Games, Inc. All Rights Reserved.
+#include "Chaos/Collision/SpatialAccelerationBroadPhase.h"
+
+namespace Chaos
+{
+}
\ No newline at end of file
diff --git a/Engine/Source/Runtime/Experimental/Chaos/Private/Chaos/CollisionResolutionUtil.cpp b/Engine/Source/Runtime/Experimental/Chaos/Private/Chaos/CollisionResolutionUtil.cpp
--- a/Engine/Source/Runtime/Experimental/Chaos/Private/Chaos/CollisionResolutionUtil.cpp
+++ b/Engine/Source/Runtime/Experimental/Chaos/Private/Chaos/CollisionResolutionUtil.cpp
@@ -6,7 +6,7 @@
 #include "Chaos/Defines.h"
 #include "Chaos/ImplicitObjectUnion.h"
 #include "Chaos/ImplicitObjectScaled.h"
-#include "Chaos/ISpatialAccelerationCollection.h"
+#include "Chaos/SpatialAccelerationCollection.h"
 #include "Chaos/Levelset.h"
 #include "ChaosLog.h"
 #include "ChaosStats.h"
diff --git a/Engine/Source/Runtime/Experimental/Chaos/Private/Chaos/PBDCollisionConstraints.cpp b/Engine/Source/Runtime/Experimental/Chaos/Private/Chaos/PBDCollisionConstraints.cpp
--- a/Engine/Source/Runtime/Experimental/Chaos/Private/Chaos/PBDCollisionConstraints.cpp
+++ b/Engine/Source/Runtime/Experimental/Chaos/Private/Chaos/PBDCollisionConstraints.cpp
@@ -1,7 +1,6 @@
 // Copyright 1998-2019 Epic Games, Inc. All Rights Reserved.
 
 #include "Chaos/PBDCollisionConstraints.h"
-#include "Chaos/PBDCollisionConstraintsImp.h"
 
 #include "Chaos/Capsule.h"
 #include "Chaos/ChaosPerfTest.h"
@@ -12,7 +11,7 @@
 #include "Chaos/GeometryQueries.h"
 #include "Chaos/ImplicitObjectUnion.h"
 #include "Chaos/ImplicitObjectScaled.h"
-#include "Chaos/ISpatialAccelerationCollection.h"
+#include "Chaos/SpatialAccelerationCollection.h"
 #include "Chaos/Levelset.h"
 #include "Chaos/Pair.h"
 #include "Chaos/PBDCollisionConstraintsPointContactUtil.h"
@@ -44,22 +43,12 @@ namespace Chaos
 	int32 BPTreeOfGrids = 1;
 	FAutoConsoleVariableRef CVarBPTreeOfGrids(TEXT("p.BPTreeOfGrids"), BPTreeOfGrids, TEXT("Whether to use a seperate tree of grids for bp"));
 
-	float CollisionVelocityInflationCVar = 2.0f;
-	FAutoConsoleVariableRef CVarCollisionVelocityInflation(TEXT("p.CollisionVelocityInflation"), CollisionVelocityInflationCVar, TEXT("Collision velocity inflation.[def:2.0]"));
-
 	float CollisionFrictionOverride = -1.0f;
 	FAutoConsoleVariableRef CVarCollisionFrictionOverride(TEXT("p.CollisionFriction"), CollisionFrictionOverride, TEXT("Collision friction for all contacts if >= 0"));
 
 	CHAOS_API int32 EnableCollisions = 1;
 	FAutoConsoleVariableRef CVarEnableCollisions(TEXT("p.EnableCollisions"), EnableCollisions, TEXT("Enable/Disable collisions on the Chaos solver."));
 
-#if !UE_BUILD_SHIPPING
-	int32 CHAOS_API PendingHierarchyDump = 0;
-#endif
-
-	DEFINE_STAT(STAT_ComputeConstraints);
-	DEFINE_STAT(STAT_ComputeConstraintsSU);
-
 	//
 	// Collision Constraint Container
 	//
@@ -82,24 +71,9 @@ namespace Chaos
 		, bUseCCD(false)
 		, bEnableCollisions(true)
 		, LifespanCounter(0)
-		, CollisionVelocityInflation(CollisionVelocityInflationCVar)
-		, PostComputeCallback(nullptr)
 		, PostApplyCallback(nullptr)
 		, PostApplyPushOutCallback(nullptr)
-		, SpatialAcceleration(nullptr)
-	{
-	}
-
-	template<typename T, int d>
-	void TPBDCollisionConstraints<T, d>::SetPostComputeCallback(const TRigidBodyContactConstraintsPostComputeCallback<T, d>& Callback)
 	{
-		PostComputeCallback = Callback;
-	}
-
-	template<typename T, int d>
-	void TPBDCollisionConstraints<T, d>::ClearPostComputeCallback()
-	{
-		PostComputeCallback = nullptr;
 	}
 
 	template<typename T, int d>
@@ -126,35 +100,37 @@ namespace Chaos
 		PostApplyPushOutCallback = nullptr;
 	}
 
-
 	template<typename T, int d>
-	void TPBDCollisionConstraints<T, d>::UpdatePositionBasedState(/*const TPBDRigidParticles<T, d>& InParticles, const TArray<int32>& InIndices,*/ const T Dt)
+	void TPBDCollisionConstraints<T, d>::AddConstraint(FConstraintBase* ConstraintBase)
 	{
-		Reset();
+		if (ConstraintBase->GetType() == TRigidBodyPointContactConstraint<T, 3>::StaticType())
+		{
+			TRigidBodyPointContactConstraint<T, d>* PointConstraint = ConstraintBase->template As< TRigidBodyPointContactConstraint<T, d> >();
 
-		if (bEnableCollisions)
+			int32 Idx = PointConstraints.AddUninitialized(1);
+			PointConstraints[Idx] = *PointConstraint;
+			Handles.Add(HandleAllocator.template AllocHandle< TRigidBodyPointContactConstraint<T, d> >(this, Idx));
+
+			delete PointConstraint;
+		}
+		else if (ConstraintBase->GetType() == TRigidBodyPlaneContactConstraint<T, 3>::StaticType())
 		{
-#if !UE_BUILD_SHIPPING
-			if (PendingHierarchyDump)
-			{
-				ComputeConstraints<true>(*SpatialAcceleration, Dt);
-			}
-			else
-#endif
-			{
-				ComputeConstraints(*SpatialAcceleration, Dt);
-			}
+			TRigidBodyPlaneContactConstraint<T, d>* PlaneConstraint = ConstraintBase->template As< TRigidBodyPlaneContactConstraint<T, d> >();
+
+			int32 Idx = PlaneConstraints.AddUninitialized(1);
+			PlaneConstraints[Idx] = *PlaneConstraint;
+			Handles.Add(HandleAllocator.template AllocHandle< TRigidBodyPlaneContactConstraint<T, d> >(this, Idx));
+
+			delete PlaneConstraint;
 		}
 	}
 
 	template<typename T, int d>
-	void TPBDCollisionConstraints<T, d>::ConstructConstraints(TGeometryParticleHandle<T, d>* Particle0, TGeometryParticleHandle<T, d>* Particle1, const T Thickness, FCollisionConstraintsArray& NewConstraints)
+	void TPBDCollisionConstraints<T, d>::UpdatePositionBasedState(/*const TPBDRigidParticles<T, d>& InParticles, const TArray<int32>& InIndices,*/ const T Dt)
 	{
-		if (ensure(Particle0 && Particle1))
-		{
-			Collisions::ConstructConstraintsImpl<T, d>(Particle0, Particle1, Particle0->Geometry().Get(), Particle1->Geometry().Get(), Collisions::GetTransform(Particle0), Collisions::GetTransform(Particle1), Thickness, NewConstraints);
-		}
-
+		Reset();
+	
+		LifespanCounter++;
 	}
 
 	DECLARE_CYCLE_STAT(TEXT("TPBDCollisionConstraints::Reset"), STAT_Collisions_Reset, STATGROUP_Chaos);
@@ -230,46 +206,6 @@ namespace Chaos
 	}
 
 
-
-	template<typename T, int d>
-	template <bool bGatherStats>
-	void TPBDCollisionConstraints<T, d>::ComputeConstraints(const FAccelerationStructure& AccelerationStructure, T Dt)
-	{
-		if (const auto AABBTree = AccelerationStructure.template As<TAABBTree<TAccelerationStructureHandle<T, d>, TAABBTreeLeafArray<TAccelerationStructureHandle<T, d>, T>, T>>())
-		{
-			ComputeConstraintsHelperLowLevel<bGatherStats>(*AABBTree, Dt);
-		}
-		else if (const auto BV = AccelerationStructure.template As<TBoundingVolume<TAccelerationStructureHandle<T, d>, T, d>>())
-		{
-			ComputeConstraintsHelperLowLevel<bGatherStats>(*BV, Dt);
-		}
-		else if (const auto AABBTreeBV = AccelerationStructure.template As<TAABBTree<TAccelerationStructureHandle<T, d>, TBoundingVolume<TAccelerationStructureHandle<T, d>, T, d>, T>>())
-		{
-			ComputeConstraintsHelperLowLevel<bGatherStats>(*AABBTreeBV, Dt);
-		}
-		else if (const auto Collection = AccelerationStructure.template As<ISpatialAccelerationCollection<TAccelerationStructureHandle<T, d>, T, d>>())
-		{
-			if (bGatherStats)
-			{
-				Collection->PBDComputeConstraintsLowLevel_GatherStats(*this, Dt);
-			}
-			else
-			{
-				Collection->PBDComputeConstraintsLowLevel(*this, Dt);
-			}
-		}
-		else
-		{
-			check(false);  //question: do we want to support a dynamic dispatch version?
-		}
-
-		if (PostComputeCallback != nullptr)
-		{
-			PostComputeCallback();
-		}
-	}
-
-
 	template<typename T, int d>
 	void TPBDCollisionConstraints<T, d>::UpdateConstraints(T Dt, const TSet<TGeometryParticleHandle<T, d>*>& ParticlesSet)
 	{
@@ -361,6 +297,4 @@ namespace Chaos
 
 	template class TAccelerationStructureHandle<float, 3>;
 	template class CHAOS_API TPBDCollisionConstraints<float, 3>;
-	template void TPBDCollisionConstraints<float, 3>::ComputeConstraints<false>(const TPBDCollisionConstraints<float, 3>::FAccelerationStructure&, float Dt);
-	template void TPBDCollisionConstraints<float, 3>::ComputeConstraints<true>(const TPBDCollisionConstraints<float, 3>::FAccelerationStructure&, float Dt);
 }
diff --git a/Engine/Source/Runtime/Experimental/Chaos/Private/Chaos/PBDRigidsEvolution.cpp b/Engine/Source/Runtime/Experimental/Chaos/Private/Chaos/PBDRigidsEvolution.cpp
--- a/Engine/Source/Runtime/Experimental/Chaos/Private/Chaos/PBDRigidsEvolution.cpp
+++ b/Engine/Source/Runtime/Experimental/Chaos/Private/Chaos/PBDRigidsEvolution.cpp
@@ -5,7 +5,7 @@
 #include "ChaosStats.h"
 #include "Chaos/PBDRigidsEvolutionGBF.h"
 #include "Chaos/ParticleHandle.h"
-#include "Chaos/ISpatialAccelerationCollection.h"
+#include "Chaos/SpatialAccelerationCollection.h"
 
 namespace Chaos
 {
diff --git a/Engine/Source/Runtime/Experimental/Chaos/Private/Chaos/PBDRigidsEvolutionGBF.cpp b/Engine/Source/Runtime/Experimental/Chaos/Private/Chaos/PBDRigidsEvolutionGBF.cpp
--- a/Engine/Source/Runtime/Experimental/Chaos/Private/Chaos/PBDRigidsEvolutionGBF.cpp
+++ b/Engine/Source/Runtime/Experimental/Chaos/Private/Chaos/PBDRigidsEvolutionGBF.cpp
@@ -23,6 +23,10 @@
 
 namespace Chaos
 {
+#if !UE_BUILD_SHIPPING
+	CHAOS_API bool bPendingHierarchyDump = false;
+#endif
+
 float HackMaxAngularVelocity = 1000.f;
 FAutoConsoleVariableRef CVarHackMaxAngularVelocity(TEXT("p.HackMaxAngularVelocity"), HackMaxAngularVelocity, TEXT("Max cap on angular velocity: rad/s. This is only a temp solution and should not be relied on as a feature. -1.f to disable"));
 
@@ -39,6 +43,11 @@ FAutoConsoleVariableRef CVarHackAngularDrag(TEXT("p.HackAngularDrag2"), HackAngu
 int DisableThreshold = 5;
 FAutoConsoleVariableRef CVarDisableThreshold(TEXT("p.DisableThreshold2"), DisableThreshold, TEXT("Disable threshold frames to transition to sleeping"));
 
+float BoundsThickness = 0;
+float BoundsThicknessVelocityMultiplier = 2.0f;
+FAutoConsoleVariableRef CVarBoundsThickness(TEXT("p.CollisionBoundsThickness"), BoundsThickness, TEXT("Collision inflation for speculative contact generation.[def:0.0]"));
+FAutoConsoleVariableRef CVarBoundsThicknessVelocityMultiplier(TEXT("p.CollisionBoundsVelocityInflation"), BoundsThicknessVelocityMultiplier, TEXT("Collision velocity inflation for speculatibe contact generation.[def:2.0]"));
+
 
 DECLARE_CYCLE_STAT(TEXT("TPBDRigidsEvolutionGBF::AdvanceOneTimeStep"), STAT_Evolution_AdvanceOneTimeStep, STATGROUP_Chaos);
 DECLARE_CYCLE_STAT(TEXT("TPBDRigidsEvolutionGBF::Integrate"), STAT_Evolution_Integrate, STATGROUP_Chaos);
@@ -46,6 +55,7 @@ DECLARE_CYCLE_STAT(TEXT("TPBDRigidsEvolutionGBF::KinematicTargets"), STAT_Evolut
 DECLARE_CYCLE_STAT(TEXT("TPBDRigidsEvolutionGBF::ApplyConstraints"), STAT_Evolution_ApplyConstraints, STATGROUP_Chaos);
 DECLARE_CYCLE_STAT(TEXT("TPBDRigidsEvolutionGBF::UpdateVelocities"), STAT_Evolution_UpdateVelocites, STATGROUP_Chaos);
 DECLARE_CYCLE_STAT(TEXT("TPBDRigidsEvolutionGBF::ApplyPushOut"), STAT_Evolution_ApplyPushOut, STATGROUP_Chaos);
+DECLARE_CYCLE_STAT(TEXT("TPBDRigidsEvolutionGBF::DetectCollisions"), STAT_Evolution_DetectCollisions, STATGROUP_Chaos);
 DECLARE_CYCLE_STAT(TEXT("TPBDRigidsEvolutionGBF::UpdateConstraintPositionBasedState"), STAT_Evolution_UpdateConstraintPositionBasedState, STATGROUP_Chaos);
 DECLARE_CYCLE_STAT(TEXT("TPBDRigidsEvolutionGBF::CreateConstraintGraph"), STAT_Evolution_CreateConstraintGraph, STATGROUP_Chaos);
 DECLARE_CYCLE_STAT(TEXT("TPBDRigidsEvolutionGBF::CreateIslands"), STAT_Evolution_CreateIslands, STATGROUP_Chaos);
@@ -142,15 +152,31 @@ void TPBDRigidsEvolutionGBF<T, d>::AdvanceOneTimeStep(const T Dt, const T StepFr
 		PostIntegrateCallback();
 	}
 
+	// @todo(ccaulfield): When we have persistent constraints we will want to run collision detection before
+	// UpdateConstraintPositionBasedState. For now it is after because we just delete all collisions there right now.
+	{
+		SCOPE_CYCLE_COUNTER(STAT_Evolution_UpdateConstraintPositionBasedState);
+		UpdateConstraintPositionBasedState(Dt);
+	}
 	{
 		Base::ComputeIntermediateSpatialAcceleration();
-		CollisionConstraints.SetSpatialAcceleration(InternalAcceleration.Get());
+	}
+	{
+		SCOPE_CYCLE_COUNTER(STAT_Evolution_DetectCollisions);
+		CollisionDetector.GetBroadPhase().SetSpatialAcceleration(InternalAcceleration.Get());
+
+		CollisionStats::FStatData StatData(bPendingHierarchyDump);
+
+		CollisionDetector.DetectCollisions(Dt, StatData);
+
+		CHAOS_COLLISION_STAT(StatData.Print());
 	}
 
+	if (PostDetectCollisionsCallback != nullptr)
 	{
-		SCOPE_CYCLE_COUNTER(STAT_Evolution_UpdateConstraintPositionBasedState);
-		UpdateConstraintPositionBasedState(Dt);
+		PostDetectCollisionsCallback();
 	}
+
 	{
 		SCOPE_CYCLE_COUNTER(STAT_Evolution_CreateConstraintGraph);
 		CreateConstraintGraph();
@@ -261,8 +287,10 @@ void TPBDRigidsEvolutionGBF<T, d>::AdvanceOneTimeStep(const T Dt, const T StepFr
 template <typename T, int d>
 TPBDRigidsEvolutionGBF<T, d>::TPBDRigidsEvolutionGBF(TPBDRigidsSOAs<T, d>& InParticles, int32 InNumIterations, bool InIsSingleThreaded)
 	: Base(InParticles, InNumIterations, 1, InIsSingleThreaded)
-	, CollisionConstraints(InParticles, Collided, PhysicsMaterials, DefaultNumPushOutPairIterations)
+	, CollisionConstraints(InParticles, Collided, PhysicsMaterials, DefaultNumPairIterations, DefaultNumPushOutPairIterations)
 	, CollisionRule(CollisionConstraints, DefaultNumPushOutIterations)
+	, BroadPhase(InParticles, BoundsThickness, BoundsThicknessVelocityMultiplier)
+	, CollisionDetector(BroadPhase, CollisionConstraints)
 	, PostIntegrateCallback(nullptr)
 	, PreApplyCallback(nullptr)
 	, PostApplyCallback(nullptr)
diff --git a/Engine/Source/Runtime/Experimental/Chaos/Public/Chaos/BoundingVolumeUtilities.h b/Engine/Source/Runtime/Experimental/Chaos/Public/Chaos/BoundingVolumeUtilities.h
--- a/Engine/Source/Runtime/Experimental/Chaos/Public/Chaos/BoundingVolumeUtilities.h
+++ b/Engine/Source/Runtime/Experimental/Chaos/Public/Chaos/BoundingVolumeUtilities.h
@@ -17,31 +17,29 @@
 namespace Chaos
 {
 
+// @todo(ccaulfield): COLLISION - get rid of these?
 extern float MinBoundsThickness;
 extern float BoundsThicknessMultiplier;
 
-template <typename T, int d>
-TVector<T, d> ComputeThickness(const TPBDRigidParticles<T, d>& InParticles, T Dt, int32 BodyIndex)
+inline FVec3 ComputeBoundsThickness(FVec3 Vel, FReal Dt, FReal BoundsThickness, FReal BoundsVelocityInflation)
 {
-	TVector<T, d> AbsVelocity = InParticles.V(BodyIndex).GetAbs();
-	for (int i = 0; i < d; ++i)
+	// @todo(ccaulfield): this expands the bounds in negative velocity direction as well...maybe we don't want that?
+	for (int i = 0; i < 3; ++i)
 	{
-		AbsVelocity[i] = FMath::Max(MinBoundsThickness, AbsVelocity[i] * Dt * BoundsThicknessMultiplier);//todo(ocohen): ignoring MThickness for now
+		Vel[i] = FMath::Max(MinBoundsThickness, BoundsThicknessMultiplier * (BoundsThickness + FMath::Abs(Vel[i]) * Dt * BoundsVelocityInflation));
 	}
-
-	return AbsVelocity;
+	return Vel;
 }
 
-template <typename THandle, typename T>
-TVector<T, THandle::D> ComputeThickness(const THandle& PBDRigid, T Dt)
+inline FVec3 ComputeBoundsThickness(const TPBDRigidParticles<FReal, 3>& InParticles, FReal Dt, int32 BodyIndex, FReal BoundsThickness, FReal BoundsVelocityInflation)
 {
-	auto AbsVelocity = PBDRigid.V().GetAbs();
-	for (int i = 0; i < THandle::D; ++i)
-	{
-		AbsVelocity[i] = FMath::Max(MinBoundsThickness, AbsVelocity[i] * Dt * BoundsThicknessMultiplier);//todo(ocohen): ignoring MThickness for now
-	}
+	return ComputeBoundsThickness(InParticles.V(BodyIndex), Dt, BoundsThickness, BoundsVelocityInflation);
+}
 
-	return AbsVelocity;
+template <typename THandle>
+FVec3 ComputeBoundsThickness(const THandle& PBDRigid, FReal Dt, FReal BoundsThickness, FReal BoundsVelocityInflation)
+{
+	return ComputeBoundsThickness(PBDRigid.V(), Dt, BoundsThickness, BoundsVelocityInflation);
 }
 
 template<class OBJECT_ARRAY>
@@ -153,7 +151,8 @@ TBox<T, d> ComputeWorldSpaceBoundingBox(const TPBDRigidParticles<T, d>& Objects,
 	
 	if (bUseVelocity)
 	{
-		WorldSpaceBox.ThickenSymmetrically(ComputeThickness(Objects, Dt, i));
+		// @todo(ccaulfield): COLLISION - thickness
+		WorldSpaceBox.ThickenSymmetrically(ComputeBoundsThickness(Objects, Dt, i, 0, 1));
 	}
 	return WorldSpaceBox;
 }
@@ -344,7 +343,8 @@ typename TEnableIf<TModels<CParticleView, ParticleView>::Value>::Type ComputeAll
 			{
 				if (const auto PBDRigid = Particle.AsDynamic())
 				{
-					WorldSpaceBoxes.Last().ThickenSymmetrically(ComputeThickness(*PBDRigid, Dt));
+					// @todo(ccaulfield): COLLISION - thickness
+					WorldSpaceBoxes.Last().ThickenSymmetrically(ComputeBoundsThickness(*PBDRigid, Dt, 0, 1));
 				}
 			}
 		}
diff --git a/Engine/Source/Runtime/Experimental/Chaos/Public/Chaos/Collision/BroadPhase.h b/Engine/Source/Runtime/Experimental/Chaos/Public/Chaos/Collision/BroadPhase.h
new file mode 100644
--- /dev/null
+++ b/Engine/Source/Runtime/Experimental/Chaos/Public/Chaos/Collision/BroadPhase.h
@@ -0,0 +1,48 @@
+// Copyright 1998-2019 Epic Games, Inc. All Rights Reserved.
+#pragma once
+
+#include "Chaos/Core.h"
+
+namespace Chaos
+{
+	/**
+	 * Base class for collision detection BroadPhases.
+	 *
+	 * The BroadPhase produces a list of potentially overlapping particle pairs.
+	 *
+	 * /see FSpatialAccelerationBroadPhase
+	 */
+	class FBroadPhase
+	{
+	public:
+		FBroadPhase(const FReal InThickness, const FReal InVelocityInflation)
+			: BoundsThickness(InThickness)
+			, BoundsThicknessVelocityInflation(InVelocityInflation)
+		{
+		}
+
+		virtual ~FBroadPhase()
+		{
+		}
+
+		/**
+		 * BoundsThickness is the distance at which we speculatively generate constraints, not a shape padding.
+		 */
+		void SetBoundsThickness(const FReal Thickness)
+		{
+			BoundsThickness = Thickness;
+		}
+
+		/**
+		 * The bounds are expanded by the velocity inflation times the velocity.
+		 */
+		void SetBoundsVelocityInflation(const FReal VelocityInflation)
+		{
+			BoundsThicknessVelocityInflation = VelocityInflation;
+		}
+
+	protected:
+		FReal BoundsThickness;
+		FReal BoundsThicknessVelocityInflation;
+	};
+}
\ No newline at end of file
diff --git a/Engine/Source/Runtime/Experimental/Chaos/Public/Chaos/Collision/CollisionDetector.h b/Engine/Source/Runtime/Experimental/Chaos/Public/Chaos/Collision/CollisionDetector.h
new file mode 100644
--- /dev/null
+++ b/Engine/Source/Runtime/Experimental/Chaos/Public/Chaos/Collision/CollisionDetector.h
@@ -0,0 +1,59 @@
+// Copyright 1998-2019 Epic Games, Inc. All Rights Reserved.
+#pragma once
+
+#include "Chaos/Collision/CollisionReceiver.h"
+#include "Chaos/Collision/NarrowPhase.h"
+#include "Chaos/Collision/ParticlePairBroadPhase.h"
+#include "Chaos/Collision/SpatialAccelerationBroadPhase.h"
+#include "Chaos/ParticleHandle.h"
+#include "Chaos/PBDCollisionConstraints.h"
+#include "Chaos/PBDRigidsSOAs.h"
+
+namespace Chaos
+{
+	DECLARE_CYCLE_STAT_EXTERN(TEXT("DetectCollisions"), STAT_DetectCollisions, STATGROUP_Chaos, CHAOS_API);
+
+
+	template<typename T_BROADPHASE, typename T_NARROWPHASE, typename T_RECEIVER, typename T_CONTAINER>
+	class CHAOS_API TCollisionDetector
+	{
+	public:
+		using FBroadPhase = T_BROADPHASE;
+		using FReceiver = T_RECEIVER;
+		using FNarrowPhase = T_NARROWPHASE;
+		using FContainer = T_CONTAINER;
+
+		TCollisionDetector(FBroadPhase& InBroadPhase, FContainer& InCollisionContainer)
+			: BroadPhase(InBroadPhase)
+			, CollisionContainer(InCollisionContainer)
+		{
+		}
+
+		FBroadPhase& GetBroadPhase() { return BroadPhase; }
+		FContainer& GetCollisionContainer() { return CollisionContainer; }
+
+		void DetectCollisions(const FReal Dt, CollisionStats::FStatData& StatData)
+		{
+			SCOPE_CYCLE_COUNTER(STAT_DetectCollisions);
+			CHAOS_SCOPED_TIMER(DetectCollisions);
+
+			if (!GetCollisionContainer().GetCollisionsEnabled())
+			{
+				return;
+			}
+
+			// Collision detection pipeline: BroadPhase -> NarrowPhase -> Receiver -> Container
+			// Receivers and NarrowPhase are assumed to be stateless atm. If we change that, they need to
+			// be passed into the constructor with the BroadPhase and Container.
+			FReceiver Receiver(CollisionContainer);
+			FNarrowPhase NarrowPhase;
+			BroadPhase.ProduceOverlaps(Dt, NarrowPhase, Receiver, StatData);
+			Receiver.ProcessCollisions();
+		}
+
+	private:
+		FBroadPhase& BroadPhase;
+		FContainer& CollisionContainer;
+	};
+
+}
diff --git a/Engine/Source/Runtime/Experimental/Chaos/Public/Chaos/Collision/CollisionReceiver.h b/Engine/Source/Runtime/Experimental/Chaos/Public/Chaos/Collision/CollisionReceiver.h
new file mode 100644
--- /dev/null
+++ b/Engine/Source/Runtime/Experimental/Chaos/Public/Chaos/Collision/CollisionReceiver.h
@@ -0,0 +1,89 @@
+// Copyright 1998-2019 Epic Games, Inc. All Rights Reserved.
+#pragma once
+
+#include "Chaos/ParticleHandle.h"
+#include "Chaos/PBDCollisionConstraints.h"
+#include "Chaos/PBDRigidsSOAs.h"
+
+namespace Chaos
+{
+	/**
+	 * Consumes lists of contacts generated by an FCollisionDetector (which may be on a separate thread) and
+	 * then adds the contacts to the CollisionConstraints in a thread-safe way.
+	 */
+	class FAsyncCollisionReceiver
+	{
+	public:
+		using FCollisionConstraints = TPBDCollisionConstraints<FReal, 3>;
+
+		FAsyncCollisionReceiver(FCollisionConstraints& InCollisionConstraints)
+			: CollisionConstraints(InCollisionConstraints)
+		{
+		}
+
+		/**
+		 * Called by a CollisionDetector (possibly in a task) when it finds collisions.
+		 */
+		void ReceiveCollisions(const TArray<TCollisionConstraintBase<float, 3>*>& Constraints)
+		{
+			for (TCollisionConstraintBase<FReal, 3>* ConstraintBase : Constraints)
+			{
+				Queue.Enqueue(ConstraintBase);
+			}
+		}
+
+		/**
+		 * Add received collisions to the container (on physics thread)
+		 */
+		void ProcessCollisions()
+		{
+			FCollisionConstraintBase* ConstraintBase = nullptr;
+			while (Queue.Dequeue(ConstraintBase))
+			{
+				CollisionConstraints.AddConstraint(ConstraintBase);
+			}
+		}
+
+	private:
+		//todo(ocohen): use per thread buffer instead, need better support than ParallelFor for this
+		TQueue<FCollisionConstraintBase*, EQueueMode::Mpsc> Queue;
+		FCollisionConstraints& CollisionConstraints;
+	};
+
+
+	/**
+	 * A synchronous collision receiver. This immediately adds collisions to the container as they are received. 
+	 * /see FAsyncCollisionReceiver
+	 */
+	class FSyncCollisionReceiver
+	{
+	public:
+		using FCollisionConstraints = TPBDCollisionConstraints<FReal, 3>;
+
+		FSyncCollisionReceiver(FCollisionConstraints& InCollisionConstraints)
+			: CollisionConstraints(InCollisionConstraints)
+		{
+		}
+
+		/**
+		 * Called by a CollisionDetector (possibly in a task) when it finds collisions.
+		 */
+		void ReceiveCollisions(const TArray<TCollisionConstraintBase<float, 3>*>& Constraints)
+		{
+			for (TCollisionConstraintBase<FReal, 3>* ConstraintBase : Constraints)
+			{
+				CollisionConstraints.AddConstraint(ConstraintBase);
+			}
+		}
+
+		/**
+		 * No-op
+		 */
+		void ProcessCollisions()
+		{
+		}
+
+	private:
+		FCollisionConstraints& CollisionConstraints;
+	};
+}
\ No newline at end of file
diff --git a/Engine/Source/Runtime/Experimental/Chaos/Public/Chaos/Collision/NarrowPhase.h b/Engine/Source/Runtime/Experimental/Chaos/Public/Chaos/Collision/NarrowPhase.h
new file mode 100644
--- /dev/null
+++ b/Engine/Source/Runtime/Experimental/Chaos/Public/Chaos/Collision/NarrowPhase.h
@@ -0,0 +1,49 @@
+// Copyright 1998-2019 Epic Games, Inc. All Rights Reserved.
+#pragma once
+
+#include "Chaos/Collision/CollisionReceiver.h"
+#include "Chaos/Collision/StatsData.h"
+#include "Chaos/CollisionResolution.h"
+#include "Chaos/ParticleHandle.h"
+#include "Chaos/PBDCollisionConstraints.h"
+#include "Chaos/PBDRigidsSOAs.h"
+
+namespace Chaos
+{
+	/**
+	 * Generate contact manifolds for particle pairs and pass them on to the consumer. Can be composed with a
+	 * multi-threaded BroadPhase as long as the collisions receiver type can handle multi-threaded calls to ReceiveCollisions.
+	 *
+	 * /template T_RECEIVER The object type that will take the detected collisions. Generally this will depend on the type of BroadPhase.
+	 * /see FAsyncCollisionReceiver, FSyncCollisionReceiver.
+	 */
+	class CHAOS_API FNarrowPhase
+	{
+	public:
+		using FCollisionConstraintsArray = TArray<TCollisionConstraintBase<FReal, 3>*>;
+
+		// @todo(ccaulfield): COLLISION Transient Handle version
+		template<typename T_RECEIVER>
+		void GenerateCollisions(FReal Dt, T_RECEIVER& Receiver, TGeometryParticleHandle<FReal, 3>* Particle1, TGeometryParticleHandle<FReal, 3>* Particle2, const FReal BoundsThickness, CollisionStats::FStatData& StatData)
+		{
+			// @todo(ccaulfield): COLLISION - Thickness: add shape padding (BoundsThickness is the distance within which we speculatively create constraints)
+
+			FCollisionConstraintsArray NewConstraints;
+			ConstructConstraints(Particle1, Particle2, BoundsThickness, NewConstraints, StatData);
+
+			Receiver.ReceiveCollisions(NewConstraints);
+		}
+
+	private:
+		void ConstructConstraints(TGeometryParticleHandle<FReal, 3>* Particle0, TGeometryParticleHandle<FReal, 3>* Particle1, const FReal Thickness, FCollisionConstraintsArray& NewConstraints, CollisionStats::FStatData& StatData)
+		{
+			if (ensure(Particle0 && Particle1))
+			{
+				Collisions::ConstructConstraintsImpl<FReal, 3>(Particle0, Particle1, Particle0->Geometry().Get(), Particle1->Geometry().Get(), Collisions::GetTransform(Particle0), Collisions::GetTransform(Particle1), Thickness, NewConstraints);
+
+				CHAOS_COLLISION_STAT(if (NewConstraints.Num()) { StatData.IncrementCountNP(NewConstraints.Num()); });
+				CHAOS_COLLISION_STAT(if (!NewConstraints.Num()) { StatData.IncrementRejectedNP(); });
+			}
+		}
+	};
+}
diff --git a/Engine/Source/Runtime/Experimental/Chaos/Public/Chaos/Collision/ParticlePairBroadPhase.h b/Engine/Source/Runtime/Experimental/Chaos/Public/Chaos/Collision/ParticlePairBroadPhase.h
new file mode 100644
--- /dev/null
+++ b/Engine/Source/Runtime/Experimental/Chaos/Public/Chaos/Collision/ParticlePairBroadPhase.h
@@ -0,0 +1,40 @@
+// Copyright 1998-2019 Epic Games, Inc. All Rights Reserved.
+#pragma once
+
+#include "Chaos/Collision/CollisionReceiver.h"
+#include "Chaos/ParticleHandle.h"
+#include "Chaos/PBDCollisionConstraints.h"
+#include "Chaos/PBDRigidsSOAs.h"
+
+namespace Chaos
+{
+	/**
+	 * Run through a list of particle pairs and pass them onto the collision detector if their AABBs overlap.
+	 * No spatial acceleration, and the order is assumed to be already optimized for cache efficiency.
+	 */
+	class CHAOS_API FParticlePairBroadPhase
+	{
+	public:
+		using FParticlePair = TVector<TGeometryParticleHandle<FReal, 3>*, 2>;
+		using FAABB = TAABB<FReal, 3>;
+
+		FParticlePairBroadPhase(const TArray<FParticlePair> InParticlePairs)
+			: ParticlePairs(InParticlePairs)
+		{
+		}
+
+		/**
+		 *
+		 */
+		void ProduceOverlaps(FReal Dt,
+			FNarrowPhase& NarrowPhase,
+			FSyncCollisionReceiver& Receiver,
+			CollisionStats::FStatData& StatData)
+		{
+			// @todo(ccaulfield)
+		}
+
+	private:
+		TArray<FParticlePair> ParticlePairs;
+	};
+}
diff --git a/Engine/Source/Runtime/Experimental/Chaos/Public/Chaos/Collision/SpatialAccelerationBroadPhase.h b/Engine/Source/Runtime/Experimental/Chaos/Public/Chaos/Collision/SpatialAccelerationBroadPhase.h
new file mode 100644
--- /dev/null
+++ b/Engine/Source/Runtime/Experimental/Chaos/Public/Chaos/Collision/SpatialAccelerationBroadPhase.h
@@ -0,0 +1,221 @@
+// Copyright 1998-2019 Epic Games, Inc. All Rights Reserved.
+#pragma once
+
+#include "Chaos/Collision/BroadPhase.h"
+#include "Chaos/Collision/StatsData.h"
+#include "Chaos/Collision/NarrowPhase.h"
+#include "Chaos/ISpatialAccelerationCollection.h"
+#include "Chaos/ParticleHandle.h"
+#include "Chaos/PBDRigidsSOAs.h"
+
+namespace Chaos
+{
+	class FAsyncCollisionReceiver;
+
+	template <typename TPayloadType, typename T, int d>
+	class ISpatialAcceleration;
+
+	/**
+	 *
+	 */
+	template <typename T, int d>
+	struct TSimOverlapVisitor
+	{
+		TArray<TAccelerationStructureHandle<T, d>>& Intersections;
+		TSimOverlapVisitor(TArray<TAccelerationStructureHandle<T, d>>& InIntersections)
+			: Intersections(InIntersections)
+		{
+		}
+
+		bool VisitOverlap(const TSpatialVisitorData<TAccelerationStructureHandle<T, d>>& Instance)
+		{
+			Intersections.Add(Instance.Payload);
+			return true;
+		}
+
+		bool VisitSweep(TSpatialVisitorData<TAccelerationStructureHandle<T, d>>, FQueryFastData& CurData)
+		{
+			check(false);
+			return false;
+		}
+
+		bool VisitRaycast(TSpatialVisitorData<TAccelerationStructureHandle<T, d>>, FQueryFastData& CurData)
+		{
+			check(false);
+			return false;
+		}
+	};
+
+	/**
+	 * A broad phase that iterates over particle and uses a spatial acceleration structure to output
+	 * potentially overlapping SpatialAccelerationHandles.
+	 */
+	class FSpatialAccelerationBroadPhase : public FBroadPhase
+	{
+	public:
+		using FAccelerationStructure = ISpatialAcceleration<TAccelerationStructureHandle<FReal, 3>, FReal, 3>;
+
+		FSpatialAccelerationBroadPhase(const TPBDRigidsSOAs<FReal, 3>& InParticles, const FReal InThickness, const FReal InVelocityInflation)
+			: FBroadPhase(InThickness, InVelocityInflation)
+			, Particles(InParticles)
+			, SpatialAcceleration(nullptr)
+		{
+		}
+
+		void SetSpatialAcceleration(const FAccelerationStructure* InSpatialAcceleration)
+		{
+			SpatialAcceleration = InSpatialAcceleration;
+		}
+
+		/**
+		 * Generate all overlapping pairs and pass them to the narrow phase.
+		 */
+		void ProduceOverlaps(
+			FReal Dt, 
+			FNarrowPhase& NarrowPhase, 
+			FAsyncCollisionReceiver& Receiver,
+			CollisionStats::FStatData& StatData)
+		{
+			if (!ensure(SpatialAcceleration))
+			{
+				// Must call SetSpatialAcceleration
+				return;
+			}
+
+			if (const auto AABBTree = SpatialAcceleration->template As<TAABBTree<TAccelerationStructureHandle<FReal, 3>, TAABBTreeLeafArray<TAccelerationStructureHandle<FReal, 3>, FReal>, FReal>>())
+			{
+				ProduceOverlaps(Dt, *AABBTree, NarrowPhase, Receiver, StatData);
+			}
+			else if (const auto BV = SpatialAcceleration->template As<TBoundingVolume<TAccelerationStructureHandle<FReal, 3>, FReal, 3>>())
+			{
+				ProduceOverlaps(Dt, *BV, NarrowPhase, Receiver, StatData);
+			}
+			else if (const auto AABBTreeBV = SpatialAcceleration->template As<TAABBTree<TAccelerationStructureHandle<FReal, 3>, TBoundingVolume<TAccelerationStructureHandle<FReal, 3>, FReal, 3>, FReal>>())
+			{
+				ProduceOverlaps(Dt, *AABBTreeBV, NarrowPhase, Receiver, StatData);
+			}
+			else if (const auto Collection = SpatialAcceleration->template As<ISpatialAccelerationCollection<TAccelerationStructureHandle<FReal, 3>, FReal, 3>>())
+			{
+				Collection->PBDComputeConstraintsLowLevel(Dt, *this, NarrowPhase, Receiver, StatData);
+			}
+			else
+			{
+				check(false);  //question: do we want to support a dynamic dispatch version?
+			}
+		}
+
+		template<typename T_SPATIALACCELERATION>
+		void ProduceOverlaps(
+			FReal Dt, 
+			const T_SPATIALACCELERATION& InSpatialAcceleration, 
+			FNarrowPhase& NarrowPhase, 
+			FAsyncCollisionReceiver& Receiver,
+			CollisionStats::FStatData& StatData)
+		{
+			const bool bDisableParallelFor = StatData.IsEnabled() || bDisableCollisionParallelFor;
+
+			Particles.GetNonDisabledDynamicView().ParallelFor(
+				[&](auto& Particle1, int32 ActiveIdxIdx)
+				{
+					ProduceParticleOverlaps(Dt, Particle1, InSpatialAcceleration, NarrowPhase, Receiver, StatData);
+				}, bDisableParallelFor);
+		}
+
+	private:
+		template<typename T_SPATIALACCELERATION>
+		void ProduceParticleOverlaps(
+			FReal Dt, 
+			TTransientPBDRigidParticleHandle<FReal, 3>& Particle1, 
+			const T_SPATIALACCELERATION& InSpatialAcceleration, 
+			FNarrowPhase& NarrowPhase, 
+			FAsyncCollisionReceiver& Receiver,
+			CollisionStats::FStatData& StatData)
+		{
+			CHAOS_COLLISION_STAT(StatData.IncrementSimulatedParticles());
+
+			TArray<TAccelerationStructureHandle<FReal, 3>> PotentialIntersections;
+
+			// @todo(ccaulfield): the spatial acceleration scheme needs to know the expanded bboxes for all particles, not just the one doing the test
+			const bool bBody1Bounded = HasBoundingBox(Particle1);
+			const FReal Box1Thickness = ComputeBoundsThickness(Particle1, Dt, BoundsThickness, BoundsThicknessVelocityInflation).Size();
+			if (bBody1Bounded)
+			{
+				// @todo(ccaulfield): COLLISION - see the NOTE below - fix it
+#if CHAOS_PARTICLEHANDLE_TODO
+				const TBox <FReal, 3> Box1 = InSpatialAcceleration.GetWorldSpaceBoundingBox(Particle1);
+#else
+				const TBox <FReal, 3> Box1 = ComputeWorldSpaceBoundingBox(Particle1);	// NOTE: this ignores the velocity expansion which is wrong
+#endif
+
+				CHAOS_COLLISION_STAT(StatData.RecordBoundsData(Box1));
+
+				TSimOverlapVisitor<FReal, 3> OverlapVisitor(PotentialIntersections);
+				InSpatialAcceleration.Overlap(Box1, OverlapVisitor);
+			}
+			else
+			{
+				const auto& GlobalElems = InSpatialAcceleration.GlobalObjects();
+				PotentialIntersections.Reserve(GlobalElems.Num());
+
+				for (auto& Elem : GlobalElems)
+				{
+					PotentialIntersections.Add(Elem.Payload);
+				}
+			}
+
+			CHAOS_COLLISION_STAT(StatData.RecordBroadphasePotentials(PotentialIntersections.Num()))
+
+			const int32 NumPotentials = PotentialIntersections.Num();
+			for (int32 i = 0; i < NumPotentials; ++i)
+			{
+				auto& Particle2 = *PotentialIntersections[i].GetGeometryParticleHandle_PhysicsThread();
+				const TGenericParticleHandle<FReal, 3> Particle2Generic(&Particle2);
+
+				// Broad Phase Culling 
+				// CollisionGroup == 0 : Collide_With_Everything
+				// CollisionGroup == INDEX_NONE : Disabled collisions
+				// CollisionGroup_A != CollisionGroup_B : Skip Check
+
+				if (Particle1.CollisionGroup() == INDEX_NONE || Particle2Generic->CollisionGroup() == INDEX_NONE)
+				{
+					continue;
+				}
+				if (Particle1.CollisionGroup() && Particle2Generic->CollisionGroup() && Particle1.CollisionGroup() != Particle2Generic->CollisionGroup())
+				{
+					continue;
+				}
+
+				if (!Particle1.Geometry() && !Particle2.Geometry())
+				{
+					continue;
+				}
+
+				const bool bBody2Bounded = HasBoundingBox(Particle2);
+
+				if (Particle1.Handle() == Particle2.Handle())
+				{
+					continue;
+				}
+
+				const bool bIsParticle2Dynamic = Particle2.CastToRigidParticle() && Particle2.ObjectState() == EObjectStateType::Dynamic;
+				if (bBody1Bounded == bBody2Bounded && bIsParticle2Dynamic)
+				{
+					//no bidirectional constraints. 
+					if (Particle2.ParticleID() > Particle1.ParticleID())
+					{
+						continue;
+					}
+				}
+
+				const FReal Box2Thickness = bIsParticle2Dynamic ? ComputeBoundsThickness(*Particle2.CastToRigidParticle(), Dt, BoundsThickness, BoundsThicknessVelocityInflation).Size() : (FReal)0;
+
+				NarrowPhase.GenerateCollisions(Dt, Receiver, Particle1.Handle(), Particle2.Handle(), FMath::Max(Box1Thickness, Box2Thickness), StatData);
+			}
+
+			CHAOS_COLLISION_STAT(StatData.FinalizeData());
+		}
+
+		const TPBDRigidsSOAs<FReal, 3>& Particles;
+		const FAccelerationStructure* SpatialAcceleration;
+	};
+}
\ No newline at end of file
diff --git a/Engine/Source/Runtime/Experimental/Chaos/Public/Chaos/Collision/StatsData.h b/Engine/Source/Runtime/Experimental/Chaos/Public/Chaos/Collision/StatsData.h
new file mode 100644
--- /dev/null
+++ b/Engine/Source/Runtime/Experimental/Chaos/Public/Chaos/Collision/StatsData.h
@@ -0,0 +1,255 @@
+// Copyright 1998-2019 Epic Games, Inc. All Rights Reserved.
+#pragma once
+
+#include "Chaos/ChaosPerfTest.h"
+#include "Chaos/Box.h"
+#include "Chaos/ParticleHandle.h"
+
+// Wrap use of FStatData in CHAOS_COLLISION_STAT to remove impact in shipping builds
+#if UE_BUILD_SHIPPING
+#define CHAOS_COLLISION_STAT(X)
+#define CHAOS_COLLISION_STAT_DISABLED(X) X
+#else
+#define CHAOS_COLLISION_STAT(X) X;
+#define CHAOS_COLLISION_STAT_DISABLED(X)
+#endif
+
+namespace Chaos
+{
+	namespace CollisionStats
+	{
+
+		static constexpr int32 BucketSizes2[] = { 0,1,4,8,16,32,64,128,512,MAX_int32 };
+
+		template <bool bGatherStats>
+		struct FStatHelper
+		{
+			int32 MaxCount;
+
+			FStatHelper() {}
+			void Record(int32 Count) {}
+			FString ToString() const { return TEXT(""); }
+		};
+
+		template <>
+		struct FStatHelper<true>
+		{
+			int32 BucketCount[UE_ARRAY_COUNT(BucketSizes2)];
+			int32 MaxCount;
+
+			FStatHelper()
+			{
+				FMemory::Memset(BucketCount, 0, sizeof(BucketCount));
+				MaxCount = 0;
+			}
+
+			void Record(int32 Count)
+			{
+				for (int32 BucketIdx = 1; BucketIdx < UE_ARRAY_COUNT(BucketSizes2); ++BucketIdx)
+				{
+					if (Count >= BucketSizes2[BucketIdx - 1] && Count < BucketSizes2[BucketIdx])
+					{
+						++BucketCount[BucketIdx];
+					}
+				}
+
+				if (Count > MaxCount)
+				{
+					MaxCount = Count;
+				}
+			}
+
+			FString ToString() const
+			{
+				FString OutLog;
+				int32 MaxBucketCount = 0;
+				for (int32 Count : BucketCount)
+				{
+					if (Count > MaxBucketCount)
+					{
+						MaxBucketCount = Count;
+					}
+				}
+
+				const float CountPerChar = MaxBucketCount / 20.f;
+				for (int32 Idx = 1; Idx < UE_ARRAY_COUNT(BucketSizes2); ++Idx)
+				{
+					int32 NumChars = BucketCount[Idx] / CountPerChar;
+					if (Idx < UE_ARRAY_COUNT(BucketSizes2) - 1)
+					{
+						OutLog += FString::Printf(TEXT("\t[%4d - %4d) (%4d) |"), BucketSizes2[Idx - 1], BucketSizes2[Idx], BucketCount[Idx]);
+					}
+					else
+					{
+						OutLog += FString::Printf(TEXT("\t[%4d -  inf) (%4d) |"), BucketSizes2[Idx - 1], BucketCount[Idx]);
+					}
+					for (int32 Count = 0; Count < NumChars; ++Count)
+					{
+						OutLog += TEXT("-");
+					}
+					OutLog += TEXT("\n");
+				}
+
+				return OutLog;
+			}
+		};
+
+		struct FStatDataImp
+		{
+			FStatDataImp() : CountNP(0), RejectedNP(0), SimulatedParticles(0), NumPotentials(0) {}
+
+
+			void IncrementSimulatedParticles()
+			{
+				++SimulatedParticles;
+			}
+
+			void RecordBoundsData(const TBox <float, 3>& Box1)
+			{
+				BoundsDistribution.Record(Box1.Extents().GetMax());
+			}
+
+			void RecordBroadphasePotentials(int32 Num)
+			{
+				NumPotentials = Num;
+				BroadphasePotentials.Record(Num);
+			}
+
+			void IncrementCountNP(int32 Count = 1)
+			{
+				CountNP += Count;
+			}
+
+			void IncrementRejectedNP()
+			{
+				++RejectedNP;
+			}
+
+			void FinalizeData()
+			{
+				NarrowPhasePerformed.Record(CountNP);
+				const int32 NPSkipped = NumPotentials - CountNP;
+				NarrowPhaseSkipped.Record(NPSkipped);
+				NarrowPhaseRejected.Record(RejectedNP);
+			}
+
+			void Print()
+			{
+				FString OutLog;
+#if CHAOS_PARTICLEHANDLE_TODO
+				const float NumParticles = InParticles.Size();
+				OutLog = FString::Printf(TEXT("ComputeConstraints stats:\n"
+					"Total Particles:%d\nSimulated Particles:%d (%.2f%%)\n"
+					"Max candidates per instance:%d (%.2f%%)\n"
+					"Max candidates skipped per instance (NP skipped):%d (%.2f%%)\n"
+					"Max narrow phase tests per instance:%d (%.2f%%)\n"
+					"Max narrow phase rejected per instance (NP rejected):%d (%.2f%%)\n"
+					"Constraints generated:%d\n"
+				),
+					InParticles.Size(),
+					SimulatedParticles, SimulatedParticles / NumParticles * 100.f,
+					BroadphasePotentials.MaxCount, BroadphasePotentials.MaxCount / NumParticles * 100.f,
+					NarrowPhaseSkipped.MaxCount, NarrowPhaseSkipped.MaxCount / NumParticles * 100.f,
+					NarrowPhasePerformed.MaxCount, NarrowPhasePerformed.MaxCount / NumParticles * 100.f,
+					NarrowPhaseRejected.MaxCount, NarrowPhaseRejected.MaxCount / NumParticles * 100.f,
+					Constraints.Num()
+				);
+
+				OutLog += FString::Printf(TEXT("Potentials per instance distribution:\n"));
+				OutLog += BroadphasePotentials.ToString();
+
+
+				OutLog += FString::Printf(TEXT("\nCandidates skipped per instance (NP skipped) distribution:\n"));
+				OutLog += NarrowPhaseSkipped.ToString();
+
+				OutLog += FString::Printf(TEXT("\nNarrow phase performed per instance distribution:\n"));
+				OutLog += NarrowPhasePerformed.ToString();
+
+				OutLog += FString::Printf(TEXT("\nNarrow phase candidates rejected per instance distribution:\n"));
+				OutLog += NarrowPhaseRejected.ToString();
+
+				OutLog += FString::Printf(TEXT("\nBounds distribution:\n"));
+				OutLog += BoundsDistribution.ToString();
+
+				UE_LOG(LogChaos, Warning, TEXT("%s"), *OutLog);
+
+				CHAOS_API extern bool bPendingHierarchyDump;
+				bPendingHierarchyDump = false;
+#endif
+			}
+
+			int32 CountNP;
+			int32 RejectedNP;
+			int32 SimulatedParticles;
+			int32 NumPotentials;
+			FStatHelper<true> BroadphasePotentials;
+			FStatHelper<true> NarrowPhaseSkipped;
+			FStatHelper<true> NarrowPhasePerformed;
+			FStatHelper<true> NarrowPhaseRejected;
+			FStatHelper<true> BoundsDistribution;
+		};
+
+		struct FStatData
+		{
+			FStatData(bool bGatherStats)
+				: StatImp(nullptr)
+			{
+				CHAOS_COLLISION_STAT({ if (bGatherStats) { StatImp = new FStatDataImp(); } });
+			}
+
+			~FStatData()
+			{
+				CHAOS_COLLISION_STAT({ if (IsEnabled()) { delete StatImp; } });
+			}
+
+			FORCEINLINE bool IsEnabled() const
+			{
+				CHAOS_COLLISION_STAT({ return StatImp != nullptr; })
+				CHAOS_COLLISION_STAT_DISABLED({ return false; })
+			}
+
+			FORCEINLINE void IncrementSimulatedParticles()
+			{
+				CHAOS_COLLISION_STAT({ if (IsEnabled()) { StatImp->IncrementSimulatedParticles(); } });
+			}
+
+			FORCEINLINE void RecordBoundsData(const TBox<FReal, 3>& Box1)
+			{
+				CHAOS_COLLISION_STAT({ if (IsEnabled()) { StatImp->RecordBoundsData(Box1); } });
+			}
+
+			FORCEINLINE void RecordBroadphasePotentials(int32 Num)
+			{
+				CHAOS_COLLISION_STAT({ if (IsEnabled()) { StatImp->RecordBroadphasePotentials(Num); } });
+			}
+
+			FORCEINLINE void IncrementCountNP(int32 Count = 1)
+			{
+				CHAOS_COLLISION_STAT({ if (IsEnabled()) { StatImp->IncrementCountNP(Count); } })
+			}
+
+			FORCEINLINE void IncrementRejectedNP()
+			{
+				CHAOS_COLLISION_STAT({ if (IsEnabled()) { StatImp->IncrementRejectedNP(); } });
+			}
+
+			FORCEINLINE void FinalizeData()
+			{
+				CHAOS_COLLISION_STAT({ if (IsEnabled()) { StatImp->FinalizeData(); } });
+			}
+
+			FORCEINLINE void Print()
+			{
+				CHAOS_COLLISION_STAT({ if (IsEnabled()) { StatImp->Print(); } });
+			}
+
+		private:
+			FStatDataImp* StatImp;
+		};
+
+
+	}
+
+}
+
+
diff --git a/Engine/Source/Runtime/Experimental/Chaos/Public/Chaos/Framework/Parallel.h b/Engine/Source/Runtime/Experimental/Chaos/Public/Chaos/Framework/Parallel.h
--- a/Engine/Source/Runtime/Experimental/Chaos/Public/Chaos/Framework/Parallel.h
+++ b/Engine/Source/Runtime/Experimental/Chaos/Public/Chaos/Framework/Parallel.h
@@ -8,4 +8,13 @@ namespace Chaos
 {
 	void CHAOS_API PhysicsParallelFor(int32 InNum, TFunctionRef<void(int32)> InCallable, bool bForceSingleThreaded = false);
 	//void CHAOS_API PhysicsParallelFor_RecursiveDivide(int32 InNum, TFunctionRef<void(int32)> InCallable, bool bForceSingleThreaded = false);
+
+
+#if UE_BUILD_SHIPPING
+	const bool bDisableParticleParallelFor = false;
+	const bool bDisableCollisionParallelFor = false;
+#else
+	CHAOS_API extern bool bDisableParticleParallelFor;
+	CHAOS_API extern bool bDisableCollisionParallelFor;
+#endif
 }
\ No newline at end of file
diff --git a/Engine/Source/Runtime/Experimental/Chaos/Public/Chaos/ISpatialAccelerationCollection.h b/Engine/Source/Runtime/Experimental/Chaos/Public/Chaos/ISpatialAccelerationCollection.h
--- a/Engine/Source/Runtime/Experimental/Chaos/Public/Chaos/ISpatialAccelerationCollection.h
+++ b/Engine/Source/Runtime/Experimental/Chaos/Public/Chaos/ISpatialAccelerationCollection.h
@@ -2,12 +2,16 @@
 #pragma once
 #include "Chaos/ISpatialAcceleration.h"
 #include "Chaos/Box.h"
+#include "Chaos/Collision/StatsData.h"
 #include "GeometryParticlesfwd.h"
+
 #include <tuple>
-#include "PBDCollisionConstraintsImp.h"
 
 namespace Chaos
 {
+	class FAsyncCollisionReceiver;
+	class FNarrowPhase;
+	class FSpatialAccelerationBroadPhase;
 
 template <typename TPayloadType, typename T, int d>
 class CHAOS_API ISpatialAccelerationCollection : public ISpatialAcceleration<TPayloadType, T, d>
@@ -24,8 +28,7 @@ public:
 	virtual ISpatialAcceleration<TPayloadType, T, d>* GetSubstructure(FSpatialAccelerationIdx Idx) = 0;
 
 	/** This is kind of a hack to avoid virtuals. We simply route calls into templated functions */
-	virtual void PBDComputeConstraintsLowLevel_GatherStats(TPBDCollisionConstraints<T, d>& CollisionConstraint, T Dt) const = 0;
-	virtual void PBDComputeConstraintsLowLevel(TPBDCollisionConstraints<T, d>& CollisionConstraint, T Dt) const = 0;
+	virtual void PBDComputeConstraintsLowLevel(T Dt, FSpatialAccelerationBroadPhase& BroadPhase, FNarrowPhase& NarrowPhase, FAsyncCollisionReceiver& Receiver, CollisionStats::FStatData& StatData) const = 0;
 	virtual TArray<FSpatialAccelerationIdx> GetAllSpatialIndices() const = 0;
 
 	bool IsBucketActive(uint8 BucketIdx) const
@@ -41,594 +44,4 @@ protected:
 	bool AllAsyncTrasksComplete;
 };
 
-template <typename T>
-void FreeObjHelper(T*& RawPtr)
-{
-	RawPtr = nullptr;
-}
-
-template <typename TPayloadType, typename T, int d>
-struct CHAOS_API TSpatialAccelerationBucketEntry
-{
-	TUniquePtr<ISpatialAcceleration<TPayloadType, T, d>> Acceleration;
-	uint16 TypeInnerIdx;
-
-	void CopyFrom(TSpatialAccelerationBucketEntry<TPayloadType, T, d>& Src)
-	{
-		Acceleration = Src.Acceleration->Copy();
-		TypeInnerIdx = Src.TypeInnerIdx;
-	}
-
-	TSpatialAccelerationBucketEntry() = default;
-	TSpatialAccelerationBucketEntry(const TSpatialAccelerationBucketEntry<TPayloadType, T, d>& Other) = delete;
-	TSpatialAccelerationBucketEntry(TSpatialAccelerationBucketEntry<TPayloadType, T, d>&& Other) = default;
-	TSpatialAccelerationBucketEntry<TPayloadType, T, d>& operator=(TSpatialAccelerationBucketEntry<TPayloadType, T, d>&& Other) = default;
-	TSpatialAccelerationBucketEntry<TPayloadType, T, d>& operator=(TSpatialAccelerationBucketEntry<TPayloadType, T, d>& Other) = delete;
-};
-
-template <typename TPayloadType, typename T, int d>
-FChaosArchive& operator<<(FChaosArchive& Ar, TSpatialAccelerationBucketEntry<TPayloadType, T, d>& BucketEntry)
-{
-	Ar << BucketEntry.Acceleration;
-	Ar << BucketEntry.TypeInnerIdx;
-	return Ar;
-}
-
-
-template <typename TPayloadType, typename T, int d>
-void FreeObjHelper(TSpatialAccelerationBucketEntry<TPayloadType, T, d >& BucketEntry)
-{
-	BucketEntry.Acceleration = nullptr;
-}
-
-template <typename T>
-T CopyFromHelper(const T& Src)
-{
-	return Src;
-}
-
-template <typename TPayloadType, typename T, int d>
-TSpatialAccelerationBucketEntry<TPayloadType, T, d> CopyFromHelper(const TSpatialAccelerationBucketEntry<TPayloadType, T, d >& Src)
-{
-	TSpatialAccelerationBucketEntry<TPayloadType, T, d> Copy;
-	Copy.Acceleration = Src.Acceleration ? Src.Acceleration->Copy() : nullptr;
-	Copy.TypeInnerIdx = Src.TypeInnerIdx;
-	return Copy;
-}
-
-
-template <typename TObj>
-struct CHAOS_API TSpatialCollectionBucket
-{
-	TSpatialCollectionBucket() = default;
-	TSpatialCollectionBucket(const TSpatialCollectionBucket<TObj>& Other) = delete;
-	TSpatialCollectionBucket(TSpatialCollectionBucket<TObj>&& Other) = default;
-	TSpatialCollectionBucket<TObj>& operator=(const TSpatialCollectionBucket<TObj>& Other) = delete;
-	TSpatialCollectionBucket<TObj>& operator=(TSpatialCollectionBucket<TObj>&& Other) = default;
-
-	TArray<TObj> Objects;
-	TArray<uint16> FreeIndices;
-
-	int32 Add(TObj&& Obj)
-	{
-		uint16 Idx;
-		if (FreeIndices.Num())
-		{
-			Idx = FreeIndices.Pop();
-			Objects[Idx] = MoveTemp(Obj);
-		}
-		else
-		{
-			Idx = Objects.Add(MoveTemp(Obj));
-		}
-
-		return Idx;
-	}
-
-	void Remove(uint16 Idx)
-	{
-		if (Objects.Num() == Idx + 1)
-		{
-			Objects.Pop(false);
-		}
-		else
-		{
-			FreeObjHelper(Objects[Idx]);
-			FreeIndices.Add(Idx);
-		}
-	}
-
-	void CopyFrom(const TSpatialCollectionBucket& Src)
-	{
-		Objects.SetNum(Src.Objects.Num());
-		for (int32 ObjIdx = 0; ObjIdx < Objects.Num(); ++ObjIdx)
-		{
-			Objects[ObjIdx] = CopyFromHelper(Src.Objects[ObjIdx]);
-		}
-		FreeIndices = Src.FreeIndices;
-	}
-
-};
-
-template <typename TObj>
-FChaosArchive& operator<<(FChaosArchive& Ar, TSpatialCollectionBucket<TObj>& Bucket)
-{
-	Ar << Bucket.Objects;
-	Ar << Bucket.FreeIndices;
-	return Ar;
-}
-
-template <typename... TRemaining>
-struct CHAOS_API TSpatialTypeTuple
-{
-};
-
-template <typename TAcceleration, typename... TRemaining>
-struct CHAOS_API TSpatialTypeTuple<TAcceleration, TRemaining...>
-{
-	using FirstType = TAcceleration;
-	TSpatialCollectionBucket<TAcceleration*> First;
-	TSpatialTypeTuple<TRemaining...> Remaining;
-
-	TSpatialTypeTuple() = default;
-
-	TSpatialTypeTuple(const TSpatialTypeTuple<TAcceleration, TRemaining...>& Other) = delete;	//raw pointers need updating because buckets use a deep copy
-};
-
-template<int Idx, typename ... Rest>
-struct CHAOS_API TSpatialTypeTupleGetter
-{
-
-};
-
-template<int Idx, typename First, typename... Rest>
-struct CHAOS_API TSpatialTypeTupleGetter<Idx, First, Rest...>
-{
-	static auto& Get(TSpatialTypeTuple<First, Rest...>& Types) { return TSpatialTypeTupleGetter<Idx - 1, Rest...>::Get(Types.Remaining); }
-	static const auto& Get(const TSpatialTypeTuple<First, Rest...>& Types) { return TSpatialTypeTupleGetter<Idx - 1, Rest...>::Get(Types.Remaining); }
-};
-
-template<typename First, typename... Rest>
-struct CHAOS_API TSpatialTypeTupleGetter<0, First, Rest...>
-{
-	static auto& Get(TSpatialTypeTuple<First, Rest...>& Types) { return Types.First; }
-	static const auto& Get(const TSpatialTypeTuple<First, Rest...>& Types) { return Types.First; }
-};
-
-template <int Idx, typename First, typename... Rest>
-auto& GetAccelerationsPerType(TSpatialTypeTuple<First, Rest...>& Types)
-{
-	return TSpatialTypeTupleGetter<Idx, First, Rest...>::Get(Types);
-}
-
-template <int Idx, typename First, typename... Rest>
-const auto& GetAccelerationsPerType(const TSpatialTypeTuple<First, Rest...>& Types)
-{
-	return TSpatialTypeTupleGetter<Idx, First, Rest...>::Get(Types);
-}
-
-template <int TypeIdx, int NumTypes, typename Tuple, typename TPayloadType, typename T, int d>
-struct TSpatialAccelerationCollectionHelper
-{
-	template <typename SQVisitor>
-	static bool RaycastFast(const Tuple& Types, const TVector<T, d>& Start, FQueryFastData& CurData, SQVisitor& Visitor)
-	{
-		const auto& Accelerations = GetAccelerationsPerType<TypeIdx>(Types).Objects;
-		for (const auto& Accelerator : Accelerations)
-		{
-			if (Accelerator && !Accelerator->RaycastFast(Start, CurData, Visitor))
-			{
-				return false;
-			}
-		}
-		
-		constexpr int NextType = TypeIdx + 1;
-		if (NextType < NumTypes)
-		{
-			return TSpatialAccelerationCollectionHelper<NextType < NumTypes ? NextType : 0, NumTypes, Tuple, TPayloadType, T, d>::RaycastFast(Types, Start, CurData, Visitor);
-		}
-
-		return true;
-	}
-
-	template <typename SQVisitor>
-	static bool SweepFast(const Tuple& Types, const TVector<T, d>& Start, FQueryFastData& CurData, const TVector<T, d> QueryHalfExtents, SQVisitor& Visitor)
-	{
-		const auto& Accelerations = GetAccelerationsPerType<TypeIdx>(Types).Objects;
-		for (const auto& Accelerator : Accelerations)
-		{
-			if (Accelerator && !Accelerator->SweepFast(Start, CurData, QueryHalfExtents, Visitor))
-			{
-				return false;
-			}
-		}
-
-		constexpr int NextType = TypeIdx + 1;
-		if (NextType < NumTypes)
-		{
-			return TSpatialAccelerationCollectionHelper < NextType < NumTypes ? NextType : 0, NumTypes, Tuple, TPayloadType, T, d>::SweepFast(Types, Start, CurData, QueryHalfExtents, Visitor);
-		}
-
-		return true;
-	}
-
-	template <typename SQVisitor>
-	static bool OverlapFast(const Tuple& Types, const TBox<T, d> QueryBounds, SQVisitor& Visitor)
-	{
-		const auto& Accelerations = GetAccelerationsPerType<TypeIdx>(Types).Objects;
-		for (const auto& Accelerator : Accelerations)
-		{
-			if (Accelerator && !Accelerator->OverlapFast(QueryBounds, Visitor))
-			{
-				return false;
-			}
-		}
-
-		constexpr int NextType = TypeIdx + 1;
-		if (NextType < NumTypes)
-		{
-			return TSpatialAccelerationCollectionHelper < NextType < NumTypes ? NextType : 0, NumTypes, Tuple, TPayloadType, T, d>::OverlapFast(Types, QueryBounds, Visitor);
-		}
-
-		return true;
-	}
-
-	static void GlobalObjects(const Tuple& Types, TArray<TPayloadBoundsElement<TPayloadType, T>>& ObjList)
-	{
-		const auto& Accelerations = GetAccelerationsPerType<TypeIdx>(Types).Objects;
-		for (const auto& Accelerator : Accelerations)
-		{
-			if (Accelerator)
-			{
-				//todo: handle early out
-				ObjList.Append(Accelerator->GlobalObjects());
-			}
-		}
-
-		constexpr int NextType = TypeIdx + 1;
-		if (NextType < NumTypes)
-		{
-			TSpatialAccelerationCollectionHelper < NextType < NumTypes ? NextType : 0, NumTypes, Tuple, TPayloadType, T, d>::GlobalObjects(Types, ObjList);
-		}
-	}
-
-	static void SetNumFrom(const Tuple& TypesSrc, Tuple& TypesDest)
-	{
-		const auto& SrcAccelerations = GetAccelerationsPerType<TypeIdx>(TypesSrc).Objects;
-		auto& DestAccelerations = GetAccelerationsPerType<TypeIdx>(TypesDest).Objects;
-		DestAccelerations.SetNum(SrcAccelerations.Num());
-
-		constexpr int NextType = TypeIdx + 1;
-		if (NextType < NumTypes)
-		{
-			TSpatialAccelerationCollectionHelper < NextType < NumTypes ? NextType : 0, NumTypes, Tuple, TPayloadType, T, d>::SetNumFrom(TypesSrc, TypesDest);
-		}
-	}
-
-	static uint16 FindTypeIdx(const Tuple& Types, SpatialAccelerationType Type)
-	{
-		using AccelType = typename std::remove_pointer<typename decltype(GetAccelerationsPerType<TypeIdx>(Types).Objects)::ElementType>::type;
-		if (AccelType::StaticType == Type)
-		{
-			return TypeIdx;
-		}
-
-		constexpr int NextType = TypeIdx + 1;
-		if (NextType < NumTypes)
-		{
-			return TSpatialAccelerationCollectionHelper < NextType < NumTypes ? NextType : 0, NumTypes, Tuple, TPayloadType, T, d>::FindTypeIdx(Types, Type);
-		}
-		else
-		{
-			return NumTypes;
-		}
-	}
-};
-
-template <bool bGatherStats, typename SpatialAccelerationCollection, typename T, int d>
-typename TEnableIf<TIsSame<typename SpatialAccelerationCollection::TPayloadType, TAccelerationStructureHandle<T,d>>::Value, void>::Type PBDComputeConstraintsLowLevel_Helper(TPBDCollisionConstraints<T, d>& CollisionConstraint, const SpatialAccelerationCollection& Accel, T Dt)
-{
-	CollisionConstraint.template ComputeConstraintsHelperLowLevel<bGatherStats>(Accel, Dt);
-}
-
-template <bool bGatherStats, typename SpatialAccelerationCollection, typename T, int d>
-typename TEnableIf<!TIsSame<typename SpatialAccelerationCollection::TPayloadType, TAccelerationStructureHandle<T, d>>::Value, void>::Type PBDComputeConstraintsLowLevel_Helper(TPBDCollisionConstraints<T, d>& CollisionConstraint, const SpatialAccelerationCollection& Accel, T Dt)
-{
-}
-
-template <typename ... TSpatialAccelerationTypes>
-class CHAOS_API TSpatialAccelerationCollection : public
-	ISpatialAccelerationCollection<typename std::tuple_element<0, std::tuple< TSpatialAccelerationTypes...>>::type::PayloadType,
-	typename std::tuple_element<0, std::tuple< TSpatialAccelerationTypes...>>::type::TType,
-	std::tuple_element<0, std::tuple< TSpatialAccelerationTypes...>>::type::D>
-{
-public:
-	using FirstAccelerationType = typename std::tuple_element<0, std::tuple< TSpatialAccelerationTypes...>>::type;
-	using TPayloadType = typename FirstAccelerationType::PayloadType;
-	using T = typename FirstAccelerationType::TType;
-	static constexpr int d = FirstAccelerationType::D;
-
-	TSpatialAccelerationCollection()
-	{
-	}
-
-	virtual FSpatialAccelerationIdx AddSubstructure(TUniquePtr<ISpatialAcceleration<TPayloadType, T, d>>&& Substructure, uint16 BucketIdx) override
-	{
-		check(BucketIdx < MaxBuckets);
-		FSpatialAccelerationIdx Result;
-		Result.Bucket = BucketIdx;
-		
-		ISpatialAcceleration<TPayloadType, T, d>* AccelPtr = Substructure.Get();
-		TSpatialAccelerationBucketEntry<TPayloadType, T, d> BucketEntry;
-		BucketEntry.Acceleration = MoveTemp(Substructure);
-		
-		const int32 TypeIdx = GetTypeIdx(AccelPtr);
-		switch (TypeIdx)
-		{
-		case 0: BucketEntry.TypeInnerIdx = GetAccelerationsPerType<0>(Types).Add(&AccelPtr->template AsChecked<typename std::tuple_element<0, std::tuple<TSpatialAccelerationTypes...>>::type>()); break;
-		case 1: BucketEntry.TypeInnerIdx = GetAccelerationsPerType<ClampedIdx(1)>(Types).Add(&AccelPtr->template AsChecked<typename std::tuple_element<ClampedIdx(1), std::tuple<TSpatialAccelerationTypes...>>::type>()); break;
-		case 2: BucketEntry.TypeInnerIdx = GetAccelerationsPerType<ClampedIdx(2)>(Types).Add(&AccelPtr->template AsChecked<typename std::tuple_element<ClampedIdx(2), std::tuple<TSpatialAccelerationTypes...>>::type>()); break;
-		case 3: BucketEntry.TypeInnerIdx = GetAccelerationsPerType<ClampedIdx(3)>(Types).Add(&AccelPtr->template AsChecked<typename std::tuple_element<ClampedIdx(3), std::tuple<TSpatialAccelerationTypes...>>::type>()); break;
-		case 4: BucketEntry.TypeInnerIdx = GetAccelerationsPerType<ClampedIdx(4)>(Types).Add(&AccelPtr->template AsChecked<typename std::tuple_element<ClampedIdx(4), std::tuple<TSpatialAccelerationTypes...>>::type>()); break;
-		case 5: BucketEntry.TypeInnerIdx = GetAccelerationsPerType<ClampedIdx(5)>(Types).Add(&AccelPtr->template AsChecked<typename std::tuple_element<ClampedIdx(5), std::tuple<TSpatialAccelerationTypes...>>::type>()); break;
-		case 6: BucketEntry.TypeInnerIdx = GetAccelerationsPerType<ClampedIdx(6)>(Types).Add(&AccelPtr->template AsChecked<typename std::tuple_element<ClampedIdx(6), std::tuple<TSpatialAccelerationTypes...>>::type>()); break;
-		case 7: BucketEntry.TypeInnerIdx = GetAccelerationsPerType<ClampedIdx(7)>(Types).Add(&AccelPtr->template AsChecked<typename std::tuple_element<ClampedIdx(7), std::tuple<TSpatialAccelerationTypes...>>::type>()); break;
-		}
-
-		this->ActiveBucketsMask |= (1 << BucketIdx);
-
-		Result.InnerIdx = Buckets[BucketIdx].Add(MoveTemp(BucketEntry));
-		return Result;
-	}
-
-	virtual TUniquePtr <ISpatialAcceleration<TPayloadType, T, d>> RemoveSubstructure(FSpatialAccelerationIdx Idx) override
-	{
-		TUniquePtr <ISpatialAcceleration<TPayloadType, T, d>> Removed;
-		{
-			auto& BucketEntry = Buckets[Idx.Bucket].Objects[Idx.InnerIdx];
-			const int32 TypeIdx = GetTypeIdx(BucketEntry.Acceleration.Get());
-			switch (TypeIdx)
-			{
-			case 0: GetAccelerationsPerType<0>(Types).Remove(BucketEntry.TypeInnerIdx); break;
-			case 1: GetAccelerationsPerType<ClampedIdx(1)>(Types).Remove(BucketEntry.TypeInnerIdx); break;
-			case 2: GetAccelerationsPerType<ClampedIdx(2)>(Types).Remove(BucketEntry.TypeInnerIdx); break;
-			case 3: GetAccelerationsPerType<ClampedIdx(3)>(Types).Remove(BucketEntry.TypeInnerIdx); break;
-			case 4: GetAccelerationsPerType<ClampedIdx(4)>(Types).Remove(BucketEntry.TypeInnerIdx); break;
-			case 5: GetAccelerationsPerType<ClampedIdx(5)>(Types).Remove(BucketEntry.TypeInnerIdx); break;
-			case 6: GetAccelerationsPerType<ClampedIdx(6)>(Types).Remove(BucketEntry.TypeInnerIdx); break;
-			case 7: GetAccelerationsPerType<ClampedIdx(7)>(Types).Remove(BucketEntry.TypeInnerIdx); break;
-			}
-			Removed = MoveTemp(BucketEntry.Acceleration);
-		}
-		Buckets[Idx.Bucket].Remove(Idx.InnerIdx);
-		if (Buckets[Idx.Bucket].Objects.Num() == 0)
-		{
-			this->ActiveBucketsMask &= ~(1 << Idx.Bucket);
-		}
-		return Removed;
-	}
-
-	virtual ISpatialAcceleration<TPayloadType, T, d>* GetSubstructure(FSpatialAccelerationIdx Idx) override
-	{
-		check(Idx.Bucket < MaxBuckets);
-		if (Buckets[Idx.Bucket].Objects.Num() > 0)
-		{
-			return Buckets[Idx.Bucket].Objects[Idx.InnerIdx].Acceleration.Get();
-		}
-
-		return nullptr;
-	}
-
-	virtual void Raycast(const TVector<T, d>& Start, const TVector<T, d>& Dir, const T Length, ISpatialVisitor<TPayloadType, T>& Visitor) const override
-	{
-		TSpatialVisitor<TPayloadType, T> ProxyVisitor(Visitor);
-		Raycast(Start, Dir, Length, ProxyVisitor);
-	}
-
-	template <typename SQVisitor>
-	void Raycast(const TVector<T, d>& Start, const TVector<T, d>& Dir, const T Length, SQVisitor& Visitor) const
-	{
-		FQueryFastData QueryFastData(Dir, Length);
-		TSpatialAccelerationCollectionHelper<0, NumTypes, decltype(Types), TPayloadType, T, d>::RaycastFast(Types, Start, QueryFastData, Visitor);
-	}
-
-	void Sweep(const TVector<T, d>& Start, const TVector<T, d>& Dir, const T Length, const TVector<T, d> QueryHalfExtents, ISpatialVisitor<TPayloadType, T>& Visitor) const override
-	{
-		TSpatialVisitor<TPayloadType, T> ProxyVisitor(Visitor);
-		Sweep(Start, Dir, Length, QueryHalfExtents, ProxyVisitor);
-	}
-
-	template <typename SQVisitor>
-	void Sweep(const TVector<T, d>& Start, const TVector<T, d>& Dir, const T Length, const TVector<T, d> QueryHalfExtents, SQVisitor& Visitor) const
-	{
-		FQueryFastData QueryFastData(Dir, Length);
-		TSpatialAccelerationCollectionHelper<0, NumTypes, decltype(Types), TPayloadType, T, d>::SweepFast(Types, Start, QueryFastData, QueryHalfExtents, Visitor);
-	}
-
-	virtual void Overlap(const TBox<T, d>& QueryBounds, ISpatialVisitor<TPayloadType, T>& Visitor) const override
-	{
-		TSpatialVisitor<TPayloadType, T> ProxyVisitor(Visitor);
-		Overlap(QueryBounds, ProxyVisitor);
-	}
-
-	template <typename SQVisitor>
-	void Overlap(const TBox<T, 3>& QueryBounds, SQVisitor& Visitor) const
-	{
-		TSpatialAccelerationCollectionHelper<0, NumTypes, decltype(Types), TPayloadType, T, d>::OverlapFast(Types, QueryBounds, Visitor);
-	}
-
-	TArray<TPayloadBoundsElement<TPayloadType, T>> GlobalObjects() const
-	{
-		TArray<TPayloadBoundsElement<TPayloadType, T>> ObjList;
-		TSpatialAccelerationCollectionHelper<0, NumTypes, decltype(Types), TPayloadType, T, d>::GlobalObjects(Types, ObjList);
-		return ObjList;
-	}
-
-	virtual TArray<FSpatialAccelerationIdx> GetAllSpatialIndices() const override
-	{
-		uint16 BucketIdx = 0;
-		TArray<FSpatialAccelerationIdx> Indices;
-		for (const auto& Bucket : Buckets)
-		{
-			uint16 InnerIdx = 0;
-			for (const auto& Entry : Bucket.Objects)
-			{
-				if (Entry.Acceleration)
-				{
-					Indices.Add(FSpatialAccelerationIdx{ BucketIdx, InnerIdx });
-				}
-				++InnerIdx;
-			}
-			++BucketIdx;
-		}
-		return Indices;
-	}
-
-	virtual void RemoveElementFrom(const TPayloadType& Payload, FSpatialAccelerationIdx SpatialIdx) override
-	{
-		const uint16 UseBucket = ((1 << SpatialIdx.Bucket) & this->ActiveBucketsMask) ? SpatialIdx.Bucket : 0;
-		Buckets[UseBucket].Objects[SpatialIdx.InnerIdx].Acceleration->RemoveElement(Payload);
-	}
-
-	virtual void UpdateElementIn(const TPayloadType& Payload, const TBox<T, d>& NewBounds, bool bHasBounds, FSpatialAccelerationIdx SpatialIdx)
-	{
-		const uint16 UseBucket = ((1 << SpatialIdx.Bucket) & this->ActiveBucketsMask) ? SpatialIdx.Bucket : 0;
-		Buckets[UseBucket].Objects[SpatialIdx.InnerIdx].Acceleration->UpdateElement(Payload, NewBounds, bHasBounds);
-	}
-
-	virtual TUniquePtr<ISpatialAcceleration<TPayloadType, T, d>> Copy() const
-	{
-		return TUniquePtr<ISpatialAcceleration<TPayloadType, T, d>>(new TSpatialAccelerationCollection<TSpatialAccelerationTypes...>(*this));
-	}
-
-	virtual void PBDComputeConstraintsLowLevel_GatherStats(TPBDCollisionConstraints<T, d>& CollisionConstraint, T Dt) const override
-	{
-		PBDComputeConstraintsLowLevel_Helper<true>(CollisionConstraint, *this, Dt);
-	}
-
-	virtual void PBDComputeConstraintsLowLevel(TPBDCollisionConstraints<T,d>& CollisionConstraint, T Dt) const override
-	{
-		PBDComputeConstraintsLowLevel_Helper<false>(CollisionConstraint, *this, Dt);
-	}
-
-	virtual void Serialize(FChaosArchive& Ar)
-	{
-		//todo: let user serialize out bucket params
-
-		//serialize out sub structures
-		for (int BucketIdx = 0; BucketIdx < MaxBuckets; ++BucketIdx)
-		{
-			Ar << Buckets[BucketIdx];
-
-			if (Ar.IsLoading())
-			{
-				for (const TSpatialAccelerationBucketEntry<TPayloadType, T, d>& Entry : Buckets[BucketIdx].Objects)
-				{
-					ISpatialAcceleration<TPayloadType, T, d>* RawPtr = Entry.Acceleration.Get();
-					const int32 TypeIdx = GetTypeIdx(RawPtr);
-					switch (TypeIdx)
-					{
-					case 0: GetAccelerationsPerType<0>(Types).Objects[Entry.TypeInnerIdx] = &RawPtr->template AsChecked<typename std::tuple_element<0, std::tuple<TSpatialAccelerationTypes...>>::type>(); break;
-					case 1: GetAccelerationsPerType<ClampedIdx(1)>(Types).Objects[Entry.TypeInnerIdx] = &RawPtr->template AsChecked<typename std::tuple_element<ClampedIdx(1), std::tuple<TSpatialAccelerationTypes...>>::type>(); break;
-					case 2: GetAccelerationsPerType<ClampedIdx(2)>(Types).Objects[Entry.TypeInnerIdx] = &RawPtr->template AsChecked<typename std::tuple_element<ClampedIdx(2), std::tuple<TSpatialAccelerationTypes...>>::type>(); break;
-					case 3: GetAccelerationsPerType<ClampedIdx(3)>(Types).Objects[Entry.TypeInnerIdx] = &RawPtr->template AsChecked<typename std::tuple_element<ClampedIdx(3), std::tuple<TSpatialAccelerationTypes...>>::type>(); break;
-					case 4: GetAccelerationsPerType<ClampedIdx(4)>(Types).Objects[Entry.TypeInnerIdx] = &RawPtr->template AsChecked<typename std::tuple_element<ClampedIdx(4), std::tuple<TSpatialAccelerationTypes...>>::type>(); break;
-					case 5: GetAccelerationsPerType<ClampedIdx(5)>(Types).Objects[Entry.TypeInnerIdx] = &RawPtr->template AsChecked<typename std::tuple_element<ClampedIdx(5), std::tuple<TSpatialAccelerationTypes...>>::type>(); break;
-					case 6: GetAccelerationsPerType<ClampedIdx(6)>(Types).Objects[Entry.TypeInnerIdx] = &RawPtr->template AsChecked<typename std::tuple_element<ClampedIdx(6), std::tuple<TSpatialAccelerationTypes...>>::type>(); break;
-					case 7: GetAccelerationsPerType<ClampedIdx(7)>(Types).Objects[Entry.TypeInnerIdx] = &RawPtr->template AsChecked<typename std::tuple_element<ClampedIdx(7), std::tuple<TSpatialAccelerationTypes...>>::type>(); break;
-					}
-				}
-			}
-		}
-	}
-
-private:
-
-	TSpatialAccelerationCollection(const TSpatialAccelerationCollection<TSpatialAccelerationTypes...>& Other)
-	{
-		//need to make a deep copy of the unique ptrs and then update raw ptrs
-		TSpatialAccelerationCollectionHelper<0, NumTypes, decltype(Types), TPayloadType, T, d>::SetNumFrom(Other.Types, Types);
-		this->ActiveBucketsMask = Other.ActiveBucketsMask;
-
-		for (int BucketIdx = 0; BucketIdx < MaxBuckets; ++BucketIdx)
-		{
-			Buckets[BucketIdx].CopyFrom(Other.Buckets[BucketIdx]);
-			bool bFirst = true;
-			for (const TSpatialAccelerationBucketEntry<TPayloadType, T, d>& Entry : Buckets[BucketIdx].Objects)
-			{
-				ISpatialAcceleration<TPayloadType, T, d>* RawPtr = Entry.Acceleration.Get();
-				const int32 TypeIdx = GetTypeIdx(RawPtr);
-				switch (TypeIdx)
-				{
-				case 0: GetAccelerationsPerType<0>(Types).Objects[Entry.TypeInnerIdx] = &RawPtr->template AsChecked<typename std::tuple_element<0, std::tuple<TSpatialAccelerationTypes...>>::type>(); break;
-				case 1: GetAccelerationsPerType<ClampedIdx(1)>(Types).Objects[Entry.TypeInnerIdx] = &RawPtr->template AsChecked<typename std::tuple_element<ClampedIdx(1), std::tuple<TSpatialAccelerationTypes...>>::type>(); break;
-				case 2: GetAccelerationsPerType<ClampedIdx(2)>(Types).Objects[Entry.TypeInnerIdx] = &RawPtr->template AsChecked<typename std::tuple_element<ClampedIdx(2), std::tuple<TSpatialAccelerationTypes...>>::type>(); break;
-				case 3: GetAccelerationsPerType<ClampedIdx(3)>(Types).Objects[Entry.TypeInnerIdx] = &RawPtr->template AsChecked<typename std::tuple_element<ClampedIdx(3), std::tuple<TSpatialAccelerationTypes...>>::type>(); break;
-				case 4: GetAccelerationsPerType<ClampedIdx(4)>(Types).Objects[Entry.TypeInnerIdx] = &RawPtr->template AsChecked<typename std::tuple_element<ClampedIdx(4), std::tuple<TSpatialAccelerationTypes...>>::type>(); break;
-				case 5: GetAccelerationsPerType<ClampedIdx(5)>(Types).Objects[Entry.TypeInnerIdx] = &RawPtr->template AsChecked<typename std::tuple_element<ClampedIdx(5), std::tuple<TSpatialAccelerationTypes...>>::type>(); break;
-				case 6: GetAccelerationsPerType<ClampedIdx(6)>(Types).Objects[Entry.TypeInnerIdx] = &RawPtr->template AsChecked<typename std::tuple_element<ClampedIdx(6), std::tuple<TSpatialAccelerationTypes...>>::type>(); break;
-				case 7: GetAccelerationsPerType<ClampedIdx(7)>(Types).Objects[Entry.TypeInnerIdx] = &RawPtr->template AsChecked<typename std::tuple_element<ClampedIdx(7), std::tuple<TSpatialAccelerationTypes...>>::type>(); break;
-				}
-			}
-		}
-	}
-	
-	static constexpr uint32 ClampedIdx(uint32 Idx)
-	{
-		return Idx < NumTypes ? Idx : 0;
-	}
-
-	static uint8 GetTypeIdx(const ISpatialAcceleration<TPayloadType, T, d>* Accel)
-	{
-		if (Accel->template As<typename std::tuple_element<0, std::tuple< TSpatialAccelerationTypes...>>::type> ())
-		{
-			return 0;
-		}
-		
-		if (Accel->template As<typename std::tuple_element<ClampedIdx(1), std::tuple< TSpatialAccelerationTypes...>>::type> ())
-		{
-			return 1;
-		}
-
-		if (Accel->template As<typename std::tuple_element<ClampedIdx(2), std::tuple< TSpatialAccelerationTypes...>>::type> ())
-		{
-			return 2;
-		}
-
-		if (Accel->template As<typename std::tuple_element<ClampedIdx(3), std::tuple< TSpatialAccelerationTypes...>>::type> ())
-		{
-			return 3;
-		}
-
-		if (Accel->template As<typename std::tuple_element<ClampedIdx(4), std::tuple< TSpatialAccelerationTypes...>>::type> ())
-		{
-			return 4;
-		}
-
-		if (Accel->template As<typename std::tuple_element<ClampedIdx(5), std::tuple< TSpatialAccelerationTypes...>>::type> ())
-		{
-			return 5;
-		}
-
-		if (Accel->template As<typename std::tuple_element<ClampedIdx(6), std::tuple< TSpatialAccelerationTypes...>>::type> ())
-		{
-			return 6;
-		}
-
-		if (Accel->template As<typename std::tuple_element<ClampedIdx(7), std::tuple< TSpatialAccelerationTypes...>>::type> ())
-		{
-			return 7;
-		}
-
-		check(false);	//Did you instantiate the collection with the right template arguments?
-		return 0;
-	}
-
-	static constexpr uint16 MaxBuckets = 8;
-	TSpatialCollectionBucket<TSpatialAccelerationBucketEntry<TPayloadType, T, d>> Buckets[MaxBuckets];
-	TSpatialTypeTuple< TSpatialAccelerationTypes...> Types;
-	static constexpr uint32 NumTypes = sizeof...(TSpatialAccelerationTypes);
-};
-
-
 }
diff --git a/Engine/Source/Runtime/Experimental/Chaos/Public/Chaos/PBDCollisionConstraints.h b/Engine/Source/Runtime/Experimental/Chaos/Public/Chaos/PBDCollisionConstraints.h
--- a/Engine/Source/Runtime/Experimental/Chaos/Public/Chaos/PBDCollisionConstraints.h
+++ b/Engine/Source/Runtime/Experimental/Chaos/Public/Chaos/PBDCollisionConstraints.h
@@ -14,9 +14,6 @@
 
 namespace Chaos
 {
-template<typename T, int d>
-class TPBDCollisionConstraintAccessor;
-
 template<typename T, int d>
 class TPBDCollisionConstraints;
 
@@ -47,16 +44,12 @@ template<typename T, int d>
 using TRigidBodyContactConstraintsPostApplyPushOutCallback = TFunction<void(const T Dt, const TArray<TPBDCollisionConstraintHandle<T, d>*>&, bool)>;
 
 /**
- * Manages a set of contact constraints:
- *	- Performs collision detection to generate constraints.
- *	- Responsible for applying corrections to particles affected by the constraints.
- * @todo(ccaulfield): remove TPBDCollisionConstraintAccessor
+ * A container and solver for collision constraints.
  */
 template<typename T, int d>
 class CHAOS_API TPBDCollisionConstraints : public FPBDConstraintContainer
 {
 public:
-	friend class TPBDCollisionConstraintAccessor<T, d>;
 	friend class TPBDCollisionConstraintHandle<T, d>;
 
 	using Base = FPBDConstraintContainer;
@@ -77,6 +70,8 @@ public:
 	virtual ~TPBDCollisionConstraints() {}
 
 
+	void AddConstraint(FConstraintBase* InConstraint);
+
 	/**
 	 * Generate all contact constraints.
 	 */
@@ -108,22 +103,6 @@ public:
 	void RemoveConstraint(FConstraintContainerHandle*);
 
 
-	/**
-	*  Process constraints acceleration structures. 
-	*/
-	//using FAccelerationStructure = TBoundingVolume<TAccelerationStructureHandle<T, d>, T, d>;
-	//using FAccelerationStructure = TAABBTree<TAccelerationStructureHandle<T, d>, TAABBTreeLeafArray<TAccelerationStructureHandle<T, d>, T>, T>;
-	using FAccelerationStructure = ISpatialAcceleration< TAccelerationStructureHandle<T, d>, T, d>;
-	void SetSpatialAcceleration(const FAccelerationStructure* InSpatialAcceleration) { SpatialAcceleration = InSpatialAcceleration; }
-
-	template <bool bGatherStats = false>
-	void ComputeConstraints(const FAccelerationStructure& AccelerationStructure, T Dt);
-
-	//NOTE: this should not be called by anyone other than ISpatialAccelerationCollection and CollisionConstraints - todo: make private with friends?
-	template <bool bGatherStats, typename SPATIAL_ACCELERATION>
-	void ComputeConstraintsHelperLowLevel(const SPATIAL_ACCELERATION& SpatialAcceleration, T Dt);
-
-
 	/**
 	 * Update all constraint values
 	 */
@@ -135,32 +114,24 @@ public:
 	 * Apply corrections for the specified list of constraints. (Runs Wide!)
 	 *
 	 * @todo(ccaulfield): this runs wide. The serial/parallel decision should be in the ConstraintRule
-	 *
 	 */
 	void Apply(const T Dt, const TArray<FConstraintContainerHandle*>& InConstraintHandles, const int32 It, const int32 NumIts);
 
 
 	/**
 	 * Apply push out for the specified list of constraints.
-	 *  Return true if we need another iteration 
 	 *
-	 * @todo(ccaulfield): this runs wide. The serial/parallel decision should be in the ConstraintRule
+	 * /return true if we need another iteration.
 	 *
+	 * @todo(ccaulfield): this runs wide. The serial/parallel decision should be in the ConstraintRule
 	 */
 	bool ApplyPushOut(const T Dt, const TArray<FConstraintContainerHandle*>& InConstraintHandles, 
 		const TSet<const TGeometryParticleHandle<T,d>*>& IsTemporarilyStatic, int32 Iteration, int32 NumIterations);
 
 
 	/**
-	*
-	*  Callbacks
-	*
-	 * Set the callback used just after contacts are generated at the start of a frame tick.
-	 * This can be used to modify or disable constraints (via ApplyCollisionModifier).
+	 *  Callbacks
 	 */
-	void SetPostComputeCallback(const TRigidBodyContactConstraintsPostComputeCallback<T, d>& Callback);
-	void ClearPostComputeCallback();
-
 	void SetPostApplyCallback(const TRigidBodyContactConstraintsPostApplyCallback<T, d>& Callback);
 	void ClearPostApplyCallback();
 
@@ -182,11 +153,6 @@ public:
 		MThickness = InThickness;
 	}
 
-	void SetVelocityInflation(float ScaleFactor)
-	{
-		CollisionVelocityInflation = ScaleFactor;
-	}
-
 	void SetPairIterations(int32 InPairIterations)
 	{
 		MApplyPairIterations = InPairIterations;
@@ -202,6 +168,11 @@ public:
 		bEnableCollisions = bInEnableCollisions;
 	}
 
+	bool GetCollisionsEnabled() const
+	{
+		return bEnableCollisions;
+	}
+
 	int32 NumConstraints() const
 	{
 		return Handles.Num();
@@ -228,13 +199,6 @@ protected:
 
 private:
 
-	/**
-	*  Build a constraint based on the two particle handles.
-	*/
-	void ConstructConstraints(TGeometryParticleHandle<T, d>* Particle0, TGeometryParticleHandle<T, d>* Particle1, const T Thickness, FCollisionConstraintsArray&);
-
-
-
 	const TPBDRigidsSOAs<T,d>& Particles;
 
 	TArray<FPointContactConstraint> PointConstraints;
@@ -248,22 +212,14 @@ private:
 	const TArrayCollectionArray<TSerializablePtr<FChaosPhysicsMaterial>>& MPhysicsMaterials;
 	int32 MApplyPairIterations;
 	int32 MApplyPushOutPairIterations;
-	T MThickness;
+	T MThickness;	// @todo(ccaulfield) - COLLISION thickness - this should be used as shape padding (as opposed to broad/narrowphase thickness which is for speculative creation of constraints)
 	T MAngularFriction;
 	bool bUseCCD;
 	bool bEnableCollisions;
 
 	int32 LifespanCounter;
-	float CollisionVelocityInflation;
 
-	TRigidBodyContactConstraintsPostComputeCallback<T, d> PostComputeCallback;
 	TRigidBodyContactConstraintsPostApplyCallback<T, d> PostApplyCallback;
 	TRigidBodyContactConstraintsPostApplyPushOutCallback<T, d> PostApplyPushOutCallback;
-
-	// @todo(ccaulfield): move spatial acceleration out of constraint container and make shareable
-	const FAccelerationStructure* SpatialAcceleration;
 };
-
-extern template void TPBDCollisionConstraints<float, 3>::ComputeConstraints<false>(const TPBDCollisionConstraints<float, 3>::FAccelerationStructure&, float Dt);
-extern template void TPBDCollisionConstraints<float, 3>::ComputeConstraints<true>(const TPBDCollisionConstraints<float, 3>::FAccelerationStructure&, float Dt);
 }
diff --git a/Engine/Source/Runtime/Experimental/Chaos/Public/Chaos/PBDCollisionConstraintsImp.h b/Engine/Source/Runtime/Experimental/Chaos/Public/Chaos/PBDCollisionConstraintsImp.h
deleted file mode 100644
diff --git a/Engine/Source/Runtime/Experimental/Chaos/Public/Chaos/PBDRigidsEvolution.h b/Engine/Source/Runtime/Experimental/Chaos/Public/Chaos/PBDRigidsEvolution.h
--- a/Engine/Source/Runtime/Experimental/Chaos/Public/Chaos/PBDRigidsEvolution.h
+++ b/Engine/Source/Runtime/Experimental/Chaos/Public/Chaos/PBDRigidsEvolution.h
@@ -12,7 +12,7 @@
 #include "Chaos/Framework/DebugSubstep.h"
 #include "HAL/Event.h"
 #include "Chaos/PBDRigidsSOAs.h"
-#include "Chaos/ISpatialAccelerationCollection.h"
+#include "Chaos/SpatialAccelerationCollection.h"
 
 
 // Declaring so it can be friended for tests.
diff --git a/Engine/Source/Runtime/Experimental/Chaos/Public/Chaos/PBDRigidsEvolutionGBF.h b/Engine/Source/Runtime/Experimental/Chaos/Public/Chaos/PBDRigidsEvolutionGBF.h
--- a/Engine/Source/Runtime/Experimental/Chaos/Public/Chaos/PBDRigidsEvolutionGBF.h
+++ b/Engine/Source/Runtime/Experimental/Chaos/Public/Chaos/PBDRigidsEvolutionGBF.h
@@ -4,6 +4,10 @@
 #include "Chaos/PBDRigidsEvolution.h"
 #include "Chaos/PBDCollisionConstraints.h"
 #include "Chaos/ChaosPerfTest.h"
+#include "Chaos/Collision/CollisionDetector.h"
+#include "Chaos/Collision/CollisionReceiver.h"
+#include "Chaos/Collision/NarrowPhase.h"
+#include "Chaos/Collision/SpatialAccelerationBroadPhase.h"
 #include "Chaos/PerParticleInitForce.h"
 #include "Chaos/PerParticleEulerStepVelocity.h"
 #include "Chaos/PerParticleEtherDrag.h"
@@ -48,10 +52,12 @@ public:
 	using typename Base::FForceRule;
 	using FGravityForces = TPerParticleGravity<T, d>;
 	using FCollisionConstraints = TPBDCollisionConstraints<T, d>;
-	using FExternalForces = TPerParticleExternalForces<T, d>;
 	using FCollisionConstraintRule = TPBDConstraintColorRule<FCollisionConstraints>;
+	using FCollisionDetector = TCollisionDetector<FSpatialAccelerationBroadPhase, FNarrowPhase, FAsyncCollisionReceiver, FCollisionConstraints>;
+	using FExternalForces = TPerParticleExternalForces<T, d>;
 
 	static constexpr int32 DefaultNumIterations = 1;
+	static constexpr int32 DefaultNumPairIterations = 1;
 	static constexpr int32 DefaultNumPushOutIterations = 5;
 	static constexpr int32 DefaultNumPushOutPairIterations = 2;
 
@@ -63,6 +69,11 @@ public:
 		PostIntegrateCallback = Cb;
 	}
 
+	void SetPostDetectCollisionsCallback(const TPBDRigidsEvolutionCallback<T, d>& Cb)
+	{
+		PostDetectCollisionsCallback = Cb;
+	}
+
 	void SetPreApplyCallback(const TPBDRigidsEvolutionCallback<T, d>& Cb)
 	{
 		PreApplyCallback = Cb;
@@ -90,6 +101,9 @@ public:
 	FCollisionConstraintRule& GetCollisionConstraintsRule() { return CollisionRule; }
 	const FCollisionConstraintRule& GetCollisionConstraintsRule() const { return CollisionRule; }
 
+	FCollisionDetector& GetCollisionDetector() { return CollisionDetector; }
+	const FCollisionDetector& GetCollisionDetector() const { return CollisionDetector; }
+
 	FExternalForces& GetExternalForces() { return ExternalForces; }
 	const FExternalForces& GetExternalForces() const { return ExternalForces; }
 
@@ -192,8 +206,11 @@ protected:
 	FExternalForces ExternalForces;
 	FCollisionConstraints CollisionConstraints;
 	FCollisionConstraintRule CollisionRule;
+	FSpatialAccelerationBroadPhase BroadPhase;
+	FCollisionDetector CollisionDetector;
 
 	TPBDRigidsEvolutionCallback<T, d> PostIntegrateCallback;
+	TPBDRigidsEvolutionCallback<T, d> PostDetectCollisionsCallback;
 	TPBDRigidsEvolutionCallback<T, d> PreApplyCallback;
 	TPBDRigidsEvolutionIslandCallback<T, d> PostApplyCallback;
 	TPBDRigidsEvolutionIslandCallback<T, d> PostApplyPushOutCallback;
diff --git a/Engine/Source/Runtime/Experimental/Chaos/Public/Chaos/ParticleIterator.h b/Engine/Source/Runtime/Experimental/Chaos/Public/Chaos/ParticleIterator.h
--- a/Engine/Source/Runtime/Experimental/Chaos/Public/Chaos/ParticleIterator.h
+++ b/Engine/Source/Runtime/Experimental/Chaos/Public/Chaos/ParticleIterator.h
@@ -1,6 +1,7 @@
 // Copyright 1998-2019 Epic Games, Inc. All Rights Reserved.
 #pragma once
 
+#include "Chaos/Framework/Parallel.h"
 #include "Chaos/ParallelFor.h"
 #include "Chaos/Particles.h"
 #include "ChaosStats.h"
@@ -12,12 +13,6 @@ DECLARE_CYCLE_STAT(TEXT("HandleViewParallelForImp"), STAT_HandleViewParallelForI
 
 namespace Chaos
 {
-#if UE_BUILD_SHIPPING
-	const bool bDisableParticleParallelFor = false;
-#else
-	CHAOS_API extern bool bDisableParticleParallelFor;
-#endif
-
 template <typename TSOA>
 class TConstParticleView;
 
diff --git a/Engine/Source/Runtime/Experimental/Chaos/Public/Chaos/SpatialAccelerationCollection.h b/Engine/Source/Runtime/Experimental/Chaos/Public/Chaos/SpatialAccelerationCollection.h
new file mode 100644
--- /dev/null
+++ b/Engine/Source/Runtime/Experimental/Chaos/Public/Chaos/SpatialAccelerationCollection.h
@@ -0,0 +1,599 @@
+// Copyright 1998-2019 Epic Games, Inc. All Rights Reserved.
+#pragma once
+#include "Chaos/ISpatialAccelerationCollection.h"
+#include "Chaos/Box.h"
+#include "Chaos/Collision/SpatialAccelerationBroadPhase.h"
+#include "Chaos/Collision/StatsData.h"
+#include "GeometryParticlesfwd.h"
+
+#include <tuple>
+
+namespace Chaos
+{
+
+template <typename T>
+void FreeObjHelper(T*& RawPtr)
+{
+	RawPtr = nullptr;
+}
+
+template <typename TPayloadType, typename T, int d>
+struct CHAOS_API TSpatialAccelerationBucketEntry
+{
+	TUniquePtr<ISpatialAcceleration<TPayloadType, T, d>> Acceleration;
+	uint16 TypeInnerIdx;
+
+	void CopyFrom(TSpatialAccelerationBucketEntry<TPayloadType, T, d>& Src)
+	{
+		Acceleration = Src.Acceleration->Copy();
+		TypeInnerIdx = Src.TypeInnerIdx;
+	}
+
+	TSpatialAccelerationBucketEntry() = default;
+	TSpatialAccelerationBucketEntry(const TSpatialAccelerationBucketEntry<TPayloadType, T, d>& Other) = delete;
+	TSpatialAccelerationBucketEntry(TSpatialAccelerationBucketEntry<TPayloadType, T, d>&& Other) = default;
+	TSpatialAccelerationBucketEntry<TPayloadType, T, d>& operator=(TSpatialAccelerationBucketEntry<TPayloadType, T, d>&& Other) = default;
+	TSpatialAccelerationBucketEntry<TPayloadType, T, d>& operator=(TSpatialAccelerationBucketEntry<TPayloadType, T, d>& Other) = delete;
+};
+
+template <typename TPayloadType, typename T, int d>
+FChaosArchive& operator<<(FChaosArchive& Ar, TSpatialAccelerationBucketEntry<TPayloadType, T, d>& BucketEntry)
+{
+	Ar << BucketEntry.Acceleration;
+	Ar << BucketEntry.TypeInnerIdx;
+	return Ar;
+}
+
+
+template <typename TPayloadType, typename T, int d>
+void FreeObjHelper(TSpatialAccelerationBucketEntry<TPayloadType, T, d >& BucketEntry)
+{
+	BucketEntry.Acceleration = nullptr;
+}
+
+template <typename T>
+T CopyFromHelper(const T& Src)
+{
+	return Src;
+}
+
+template <typename TPayloadType, typename T, int d>
+TSpatialAccelerationBucketEntry<TPayloadType, T, d> CopyFromHelper(const TSpatialAccelerationBucketEntry<TPayloadType, T, d >& Src)
+{
+	TSpatialAccelerationBucketEntry<TPayloadType, T, d> Copy;
+	Copy.Acceleration = Src.Acceleration ? Src.Acceleration->Copy() : nullptr;
+	Copy.TypeInnerIdx = Src.TypeInnerIdx;
+	return Copy;
+}
+
+
+template <typename TObj>
+struct CHAOS_API TSpatialCollectionBucket
+{
+	TSpatialCollectionBucket() = default;
+	TSpatialCollectionBucket(const TSpatialCollectionBucket<TObj>& Other) = delete;
+	TSpatialCollectionBucket(TSpatialCollectionBucket<TObj>&& Other) = default;
+	TSpatialCollectionBucket<TObj>& operator=(const TSpatialCollectionBucket<TObj>& Other) = delete;
+	TSpatialCollectionBucket<TObj>& operator=(TSpatialCollectionBucket<TObj>&& Other) = default;
+
+	TArray<TObj> Objects;
+	TArray<uint16> FreeIndices;
+
+	int32 Add(TObj&& Obj)
+	{
+		uint16 Idx;
+		if (FreeIndices.Num())
+		{
+			Idx = FreeIndices.Pop();
+			Objects[Idx] = MoveTemp(Obj);
+		}
+		else
+		{
+			Idx = Objects.Add(MoveTemp(Obj));
+		}
+
+		return Idx;
+	}
+
+	void Remove(uint16 Idx)
+	{
+		if (Objects.Num() == Idx + 1)
+		{
+			Objects.Pop(false);
+		}
+		else
+		{
+			FreeObjHelper(Objects[Idx]);
+			FreeIndices.Add(Idx);
+		}
+	}
+
+	void CopyFrom(const TSpatialCollectionBucket& Src)
+	{
+		Objects.SetNum(Src.Objects.Num());
+		for (int32 ObjIdx = 0; ObjIdx < Objects.Num(); ++ObjIdx)
+		{
+			Objects[ObjIdx] = CopyFromHelper(Src.Objects[ObjIdx]);
+		}
+		FreeIndices = Src.FreeIndices;
+	}
+
+};
+
+template <typename TObj>
+FChaosArchive& operator<<(FChaosArchive& Ar, TSpatialCollectionBucket<TObj>& Bucket)
+{
+	Ar << Bucket.Objects;
+	Ar << Bucket.FreeIndices;
+	return Ar;
+}
+
+template <typename... TRemaining>
+struct CHAOS_API TSpatialTypeTuple
+{
+};
+
+template <typename TAcceleration, typename... TRemaining>
+struct CHAOS_API TSpatialTypeTuple<TAcceleration, TRemaining...>
+{
+	using FirstType = TAcceleration;
+	TSpatialCollectionBucket<TAcceleration*> First;
+	TSpatialTypeTuple<TRemaining...> Remaining;
+
+	TSpatialTypeTuple() = default;
+
+	TSpatialTypeTuple(const TSpatialTypeTuple<TAcceleration, TRemaining...>& Other) = delete;	//raw pointers need updating because buckets use a deep copy
+};
+
+template<int Idx, typename ... Rest>
+struct CHAOS_API TSpatialTypeTupleGetter
+{
+
+};
+
+template<int Idx, typename First, typename... Rest>
+struct CHAOS_API TSpatialTypeTupleGetter<Idx, First, Rest...>
+{
+	static auto& Get(TSpatialTypeTuple<First, Rest...>& Types) { return TSpatialTypeTupleGetter<Idx - 1, Rest...>::Get(Types.Remaining); }
+	static const auto& Get(const TSpatialTypeTuple<First, Rest...>& Types) { return TSpatialTypeTupleGetter<Idx - 1, Rest...>::Get(Types.Remaining); }
+};
+
+template<typename First, typename... Rest>
+struct CHAOS_API TSpatialTypeTupleGetter<0, First, Rest...>
+{
+	static auto& Get(TSpatialTypeTuple<First, Rest...>& Types) { return Types.First; }
+	static const auto& Get(const TSpatialTypeTuple<First, Rest...>& Types) { return Types.First; }
+};
+
+template <int Idx, typename First, typename... Rest>
+auto& GetAccelerationsPerType(TSpatialTypeTuple<First, Rest...>& Types)
+{
+	return TSpatialTypeTupleGetter<Idx, First, Rest...>::Get(Types);
+}
+
+template <int Idx, typename First, typename... Rest>
+const auto& GetAccelerationsPerType(const TSpatialTypeTuple<First, Rest...>& Types)
+{
+	return TSpatialTypeTupleGetter<Idx, First, Rest...>::Get(Types);
+}
+
+template <int TypeIdx, int NumTypes, typename Tuple, typename TPayloadType, typename T, int d>
+struct TSpatialAccelerationCollectionHelper
+{
+	template <typename SQVisitor>
+	static bool RaycastFast(const Tuple& Types, const TVector<T, d>& Start, FQueryFastData& CurData, SQVisitor& Visitor)
+	{
+		const auto& Accelerations = GetAccelerationsPerType<TypeIdx>(Types).Objects;
+		for (const auto& Accelerator : Accelerations)
+		{
+			if (Accelerator && !Accelerator->RaycastFast(Start, CurData, Visitor))
+			{
+				return false;
+			}
+		}
+		
+		constexpr int NextType = TypeIdx + 1;
+		if (NextType < NumTypes)
+		{
+			return TSpatialAccelerationCollectionHelper<NextType < NumTypes ? NextType : 0, NumTypes, Tuple, TPayloadType, T, d>::RaycastFast(Types, Start, CurData, Visitor);
+		}
+
+		return true;
+	}
+
+	template <typename SQVisitor>
+	static bool SweepFast(const Tuple& Types, const TVector<T, d>& Start, FQueryFastData& CurData, const TVector<T, d> QueryHalfExtents, SQVisitor& Visitor)
+	{
+		const auto& Accelerations = GetAccelerationsPerType<TypeIdx>(Types).Objects;
+		for (const auto& Accelerator : Accelerations)
+		{
+			if (Accelerator && !Accelerator->SweepFast(Start, CurData, QueryHalfExtents, Visitor))
+			{
+				return false;
+			}
+		}
+
+		constexpr int NextType = TypeIdx + 1;
+		if (NextType < NumTypes)
+		{
+			return TSpatialAccelerationCollectionHelper < NextType < NumTypes ? NextType : 0, NumTypes, Tuple, TPayloadType, T, d>::SweepFast(Types, Start, CurData, QueryHalfExtents, Visitor);
+		}
+
+		return true;
+	}
+
+	template <typename SQVisitor>
+	static bool OverlapFast(const Tuple& Types, const TBox<T, d> QueryBounds, SQVisitor& Visitor)
+	{
+		const auto& Accelerations = GetAccelerationsPerType<TypeIdx>(Types).Objects;
+		for (const auto& Accelerator : Accelerations)
+		{
+			if (Accelerator && !Accelerator->OverlapFast(QueryBounds, Visitor))
+			{
+				return false;
+			}
+		}
+
+		constexpr int NextType = TypeIdx + 1;
+		if (NextType < NumTypes)
+		{
+			return TSpatialAccelerationCollectionHelper < NextType < NumTypes ? NextType : 0, NumTypes, Tuple, TPayloadType, T, d>::OverlapFast(Types, QueryBounds, Visitor);
+		}
+
+		return true;
+	}
+
+	static void GlobalObjects(const Tuple& Types, TArray<TPayloadBoundsElement<TPayloadType, T>>& ObjList)
+	{
+		const auto& Accelerations = GetAccelerationsPerType<TypeIdx>(Types).Objects;
+		for (const auto& Accelerator : Accelerations)
+		{
+			if (Accelerator)
+			{
+				//todo: handle early out
+				ObjList.Append(Accelerator->GlobalObjects());
+			}
+		}
+
+		constexpr int NextType = TypeIdx + 1;
+		if (NextType < NumTypes)
+		{
+			TSpatialAccelerationCollectionHelper < NextType < NumTypes ? NextType : 0, NumTypes, Tuple, TPayloadType, T, d>::GlobalObjects(Types, ObjList);
+		}
+	}
+
+	static void SetNumFrom(const Tuple& TypesSrc, Tuple& TypesDest)
+	{
+		const auto& SrcAccelerations = GetAccelerationsPerType<TypeIdx>(TypesSrc).Objects;
+		auto& DestAccelerations = GetAccelerationsPerType<TypeIdx>(TypesDest).Objects;
+		DestAccelerations.SetNum(SrcAccelerations.Num());
+
+		constexpr int NextType = TypeIdx + 1;
+		if (NextType < NumTypes)
+		{
+			TSpatialAccelerationCollectionHelper < NextType < NumTypes ? NextType : 0, NumTypes, Tuple, TPayloadType, T, d>::SetNumFrom(TypesSrc, TypesDest);
+		}
+	}
+
+	static uint16 FindTypeIdx(const Tuple& Types, SpatialAccelerationType Type)
+	{
+		using AccelType = typename std::remove_pointer<typename decltype(GetAccelerationsPerType<TypeIdx>(Types).Objects)::ElementType>::type;
+		if (AccelType::StaticType == Type)
+		{
+			return TypeIdx;
+		}
+
+		constexpr int NextType = TypeIdx + 1;
+		if (NextType < NumTypes)
+		{
+			return TSpatialAccelerationCollectionHelper < NextType < NumTypes ? NextType : 0, NumTypes, Tuple, TPayloadType, T, d>::FindTypeIdx(Types, Type);
+		}
+		else
+		{
+			return NumTypes;
+		}
+	}
+};
+
+template <typename SpatialAccelerationCollection>
+typename TEnableIf<TIsSame<typename SpatialAccelerationCollection::TPayloadType, TAccelerationStructureHandle<FReal, 3>>::Value, void>::Type PBDComputeConstraintsLowLevel_Helper(FReal Dt, const SpatialAccelerationCollection& Accel, FSpatialAccelerationBroadPhase& BroadPhase, FNarrowPhase& NarrowPhase, FAsyncCollisionReceiver& Receiver, CollisionStats::FStatData& StatData)
+{
+	BroadPhase.ProduceOverlaps(Dt, Accel, NarrowPhase, Receiver, StatData);
+}
+
+template <typename SpatialAccelerationCollection>
+typename TEnableIf<!TIsSame<typename SpatialAccelerationCollection::TPayloadType, TAccelerationStructureHandle<FReal, 3>>::Value, void>::Type PBDComputeConstraintsLowLevel_Helper(FReal Dt, const SpatialAccelerationCollection& Accel, FSpatialAccelerationBroadPhase& BroadPhase, FNarrowPhase& NarrowPhase, FAsyncCollisionReceiver& Receiver, CollisionStats::FStatData& StatData)
+{
+}
+
+template <typename ... TSpatialAccelerationTypes>
+class CHAOS_API TSpatialAccelerationCollection : public
+	ISpatialAccelerationCollection<typename std::tuple_element<0, std::tuple< TSpatialAccelerationTypes...>>::type::PayloadType,
+	typename std::tuple_element<0, std::tuple< TSpatialAccelerationTypes...>>::type::TType,
+	std::tuple_element<0, std::tuple< TSpatialAccelerationTypes...>>::type::D>
+{
+public:
+	using FirstAccelerationType = typename std::tuple_element<0, std::tuple< TSpatialAccelerationTypes...>>::type;
+	using TPayloadType = typename FirstAccelerationType::PayloadType;
+	using T = typename FirstAccelerationType::TType;
+	static constexpr int d = FirstAccelerationType::D;
+
+	TSpatialAccelerationCollection()
+	{
+	}
+
+	virtual FSpatialAccelerationIdx AddSubstructure(TUniquePtr<ISpatialAcceleration<TPayloadType, T, d>>&& Substructure, uint16 BucketIdx) override
+	{
+		check(BucketIdx < MaxBuckets);
+		FSpatialAccelerationIdx Result;
+		Result.Bucket = BucketIdx;
+		
+		ISpatialAcceleration<TPayloadType, T, d>* AccelPtr = Substructure.Get();
+		TSpatialAccelerationBucketEntry<TPayloadType, T, d> BucketEntry;
+		BucketEntry.Acceleration = MoveTemp(Substructure);
+		
+		const int32 TypeIdx = GetTypeIdx(AccelPtr);
+		switch (TypeIdx)
+		{
+		case 0: BucketEntry.TypeInnerIdx = GetAccelerationsPerType<0>(Types).Add(&AccelPtr->template AsChecked<typename std::tuple_element<0, std::tuple<TSpatialAccelerationTypes...>>::type>()); break;
+		case 1: BucketEntry.TypeInnerIdx = GetAccelerationsPerType<ClampedIdx(1)>(Types).Add(&AccelPtr->template AsChecked<typename std::tuple_element<ClampedIdx(1), std::tuple<TSpatialAccelerationTypes...>>::type>()); break;
+		case 2: BucketEntry.TypeInnerIdx = GetAccelerationsPerType<ClampedIdx(2)>(Types).Add(&AccelPtr->template AsChecked<typename std::tuple_element<ClampedIdx(2), std::tuple<TSpatialAccelerationTypes...>>::type>()); break;
+		case 3: BucketEntry.TypeInnerIdx = GetAccelerationsPerType<ClampedIdx(3)>(Types).Add(&AccelPtr->template AsChecked<typename std::tuple_element<ClampedIdx(3), std::tuple<TSpatialAccelerationTypes...>>::type>()); break;
+		case 4: BucketEntry.TypeInnerIdx = GetAccelerationsPerType<ClampedIdx(4)>(Types).Add(&AccelPtr->template AsChecked<typename std::tuple_element<ClampedIdx(4), std::tuple<TSpatialAccelerationTypes...>>::type>()); break;
+		case 5: BucketEntry.TypeInnerIdx = GetAccelerationsPerType<ClampedIdx(5)>(Types).Add(&AccelPtr->template AsChecked<typename std::tuple_element<ClampedIdx(5), std::tuple<TSpatialAccelerationTypes...>>::type>()); break;
+		case 6: BucketEntry.TypeInnerIdx = GetAccelerationsPerType<ClampedIdx(6)>(Types).Add(&AccelPtr->template AsChecked<typename std::tuple_element<ClampedIdx(6), std::tuple<TSpatialAccelerationTypes...>>::type>()); break;
+		case 7: BucketEntry.TypeInnerIdx = GetAccelerationsPerType<ClampedIdx(7)>(Types).Add(&AccelPtr->template AsChecked<typename std::tuple_element<ClampedIdx(7), std::tuple<TSpatialAccelerationTypes...>>::type>()); break;
+		}
+
+		this->ActiveBucketsMask |= (1 << BucketIdx);
+
+		Result.InnerIdx = Buckets[BucketIdx].Add(MoveTemp(BucketEntry));
+		return Result;
+	}
+
+	virtual TUniquePtr <ISpatialAcceleration<TPayloadType, T, d>> RemoveSubstructure(FSpatialAccelerationIdx Idx) override
+	{
+		TUniquePtr <ISpatialAcceleration<TPayloadType, T, d>> Removed;
+		{
+			auto& BucketEntry = Buckets[Idx.Bucket].Objects[Idx.InnerIdx];
+			const int32 TypeIdx = GetTypeIdx(BucketEntry.Acceleration.Get());
+			switch (TypeIdx)
+			{
+			case 0: GetAccelerationsPerType<0>(Types).Remove(BucketEntry.TypeInnerIdx); break;
+			case 1: GetAccelerationsPerType<ClampedIdx(1)>(Types).Remove(BucketEntry.TypeInnerIdx); break;
+			case 2: GetAccelerationsPerType<ClampedIdx(2)>(Types).Remove(BucketEntry.TypeInnerIdx); break;
+			case 3: GetAccelerationsPerType<ClampedIdx(3)>(Types).Remove(BucketEntry.TypeInnerIdx); break;
+			case 4: GetAccelerationsPerType<ClampedIdx(4)>(Types).Remove(BucketEntry.TypeInnerIdx); break;
+			case 5: GetAccelerationsPerType<ClampedIdx(5)>(Types).Remove(BucketEntry.TypeInnerIdx); break;
+			case 6: GetAccelerationsPerType<ClampedIdx(6)>(Types).Remove(BucketEntry.TypeInnerIdx); break;
+			case 7: GetAccelerationsPerType<ClampedIdx(7)>(Types).Remove(BucketEntry.TypeInnerIdx); break;
+			}
+			Removed = MoveTemp(BucketEntry.Acceleration);
+		}
+		Buckets[Idx.Bucket].Remove(Idx.InnerIdx);
+		if (Buckets[Idx.Bucket].Objects.Num() == 0)
+		{
+			this->ActiveBucketsMask &= ~(1 << Idx.Bucket);
+		}
+		return Removed;
+	}
+
+	virtual ISpatialAcceleration<TPayloadType, T, d>* GetSubstructure(FSpatialAccelerationIdx Idx) override
+	{
+		check(Idx.Bucket < MaxBuckets);
+		if (Buckets[Idx.Bucket].Objects.Num() > 0)
+		{
+			return Buckets[Idx.Bucket].Objects[Idx.InnerIdx].Acceleration.Get();
+		}
+
+		return nullptr;
+	}
+
+	virtual void Raycast(const TVector<T, d>& Start, const TVector<T, d>& Dir, const T Length, ISpatialVisitor<TPayloadType, T>& Visitor) const override
+	{
+		TSpatialVisitor<TPayloadType, T> ProxyVisitor(Visitor);
+		Raycast(Start, Dir, Length, ProxyVisitor);
+	}
+
+	template <typename SQVisitor>
+	void Raycast(const TVector<T, d>& Start, const TVector<T, d>& Dir, const T Length, SQVisitor& Visitor) const
+	{
+		FQueryFastData QueryFastData(Dir, Length);
+		TSpatialAccelerationCollectionHelper<0, NumTypes, decltype(Types), TPayloadType, T, d>::RaycastFast(Types, Start, QueryFastData, Visitor);
+	}
+
+	void Sweep(const TVector<T, d>& Start, const TVector<T, d>& Dir, const T Length, const TVector<T, d> QueryHalfExtents, ISpatialVisitor<TPayloadType, T>& Visitor) const override
+	{
+		TSpatialVisitor<TPayloadType, T> ProxyVisitor(Visitor);
+		Sweep(Start, Dir, Length, QueryHalfExtents, ProxyVisitor);
+	}
+
+	template <typename SQVisitor>
+	void Sweep(const TVector<T, d>& Start, const TVector<T, d>& Dir, const T Length, const TVector<T, d> QueryHalfExtents, SQVisitor& Visitor) const
+	{
+		FQueryFastData QueryFastData(Dir, Length);
+		TSpatialAccelerationCollectionHelper<0, NumTypes, decltype(Types), TPayloadType, T, d>::SweepFast(Types, Start, QueryFastData, QueryHalfExtents, Visitor);
+	}
+
+	virtual void Overlap(const TBox<T, d>& QueryBounds, ISpatialVisitor<TPayloadType, T>& Visitor) const override
+	{
+		TSpatialVisitor<TPayloadType, T> ProxyVisitor(Visitor);
+		Overlap(QueryBounds, ProxyVisitor);
+	}
+
+	template <typename SQVisitor>
+	void Overlap(const TBox<T, 3>& QueryBounds, SQVisitor& Visitor) const
+	{
+		TSpatialAccelerationCollectionHelper<0, NumTypes, decltype(Types), TPayloadType, T, d>::OverlapFast(Types, QueryBounds, Visitor);
+	}
+
+	TArray<TPayloadBoundsElement<TPayloadType, T>> GlobalObjects() const
+	{
+		TArray<TPayloadBoundsElement<TPayloadType, T>> ObjList;
+		TSpatialAccelerationCollectionHelper<0, NumTypes, decltype(Types), TPayloadType, T, d>::GlobalObjects(Types, ObjList);
+		return ObjList;
+	}
+
+	virtual TArray<FSpatialAccelerationIdx> GetAllSpatialIndices() const override
+	{
+		uint16 BucketIdx = 0;
+		TArray<FSpatialAccelerationIdx> Indices;
+		for (const auto& Bucket : Buckets)
+		{
+			uint16 InnerIdx = 0;
+			for (const auto& Entry : Bucket.Objects)
+			{
+				if (Entry.Acceleration)
+				{
+					Indices.Add(FSpatialAccelerationIdx{ BucketIdx, InnerIdx });
+				}
+				++InnerIdx;
+			}
+			++BucketIdx;
+		}
+		return Indices;
+	}
+
+	virtual void RemoveElementFrom(const TPayloadType& Payload, FSpatialAccelerationIdx SpatialIdx) override
+	{
+		const uint16 UseBucket = ((1 << SpatialIdx.Bucket) & this->ActiveBucketsMask) ? SpatialIdx.Bucket : 0;
+		Buckets[UseBucket].Objects[SpatialIdx.InnerIdx].Acceleration->RemoveElement(Payload);
+	}
+
+	virtual void UpdateElementIn(const TPayloadType& Payload, const TBox<T, d>& NewBounds, bool bHasBounds, FSpatialAccelerationIdx SpatialIdx)
+	{
+		const uint16 UseBucket = ((1 << SpatialIdx.Bucket) & this->ActiveBucketsMask) ? SpatialIdx.Bucket : 0;
+		Buckets[UseBucket].Objects[SpatialIdx.InnerIdx].Acceleration->UpdateElement(Payload, NewBounds, bHasBounds);
+	}
+
+	virtual TUniquePtr<ISpatialAcceleration<TPayloadType, T, d>> Copy() const
+	{
+		return TUniquePtr<ISpatialAcceleration<TPayloadType, T, d>>(new TSpatialAccelerationCollection<TSpatialAccelerationTypes...>(*this));
+	}
+
+	virtual void PBDComputeConstraintsLowLevel(T Dt, FSpatialAccelerationBroadPhase& BroadPhase, FNarrowPhase& NarrowPhase, FAsyncCollisionReceiver& Receiver, CollisionStats::FStatData& StatData) const override
+	{
+		PBDComputeConstraintsLowLevel_Helper(Dt, *this, BroadPhase, NarrowPhase, Receiver, StatData);
+	}
+
+	virtual void Serialize(FChaosArchive& Ar)
+	{
+		//todo: let user serialize out bucket params
+
+		//serialize out sub structures
+		for (int BucketIdx = 0; BucketIdx < MaxBuckets; ++BucketIdx)
+		{
+			Ar << Buckets[BucketIdx];
+
+			if (Ar.IsLoading())
+			{
+				for (const TSpatialAccelerationBucketEntry<TPayloadType, T, d>& Entry : Buckets[BucketIdx].Objects)
+				{
+					ISpatialAcceleration<TPayloadType, T, d>* RawPtr = Entry.Acceleration.Get();
+					const int32 TypeIdx = GetTypeIdx(RawPtr);
+					switch (TypeIdx)
+					{
+					case 0: GetAccelerationsPerType<0>(Types).Objects[Entry.TypeInnerIdx] = &RawPtr->template AsChecked<typename std::tuple_element<0, std::tuple<TSpatialAccelerationTypes...>>::type>(); break;
+					case 1: GetAccelerationsPerType<ClampedIdx(1)>(Types).Objects[Entry.TypeInnerIdx] = &RawPtr->template AsChecked<typename std::tuple_element<ClampedIdx(1), std::tuple<TSpatialAccelerationTypes...>>::type>(); break;
+					case 2: GetAccelerationsPerType<ClampedIdx(2)>(Types).Objects[Entry.TypeInnerIdx] = &RawPtr->template AsChecked<typename std::tuple_element<ClampedIdx(2), std::tuple<TSpatialAccelerationTypes...>>::type>(); break;
+					case 3: GetAccelerationsPerType<ClampedIdx(3)>(Types).Objects[Entry.TypeInnerIdx] = &RawPtr->template AsChecked<typename std::tuple_element<ClampedIdx(3), std::tuple<TSpatialAccelerationTypes...>>::type>(); break;
+					case 4: GetAccelerationsPerType<ClampedIdx(4)>(Types).Objects[Entry.TypeInnerIdx] = &RawPtr->template AsChecked<typename std::tuple_element<ClampedIdx(4), std::tuple<TSpatialAccelerationTypes...>>::type>(); break;
+					case 5: GetAccelerationsPerType<ClampedIdx(5)>(Types).Objects[Entry.TypeInnerIdx] = &RawPtr->template AsChecked<typename std::tuple_element<ClampedIdx(5), std::tuple<TSpatialAccelerationTypes...>>::type>(); break;
+					case 6: GetAccelerationsPerType<ClampedIdx(6)>(Types).Objects[Entry.TypeInnerIdx] = &RawPtr->template AsChecked<typename std::tuple_element<ClampedIdx(6), std::tuple<TSpatialAccelerationTypes...>>::type>(); break;
+					case 7: GetAccelerationsPerType<ClampedIdx(7)>(Types).Objects[Entry.TypeInnerIdx] = &RawPtr->template AsChecked<typename std::tuple_element<ClampedIdx(7), std::tuple<TSpatialAccelerationTypes...>>::type>(); break;
+					}
+				}
+			}
+		}
+	}
+
+private:
+
+	TSpatialAccelerationCollection(const TSpatialAccelerationCollection<TSpatialAccelerationTypes...>& Other)
+	{
+		//need to make a deep copy of the unique ptrs and then update raw ptrs
+		TSpatialAccelerationCollectionHelper<0, NumTypes, decltype(Types), TPayloadType, T, d>::SetNumFrom(Other.Types, Types);
+		this->ActiveBucketsMask = Other.ActiveBucketsMask;
+
+		for (int BucketIdx = 0; BucketIdx < MaxBuckets; ++BucketIdx)
+		{
+			Buckets[BucketIdx].CopyFrom(Other.Buckets[BucketIdx]);
+			bool bFirst = true;
+			for (const TSpatialAccelerationBucketEntry<TPayloadType, T, d>& Entry : Buckets[BucketIdx].Objects)
+			{
+				ISpatialAcceleration<TPayloadType, T, d>* RawPtr = Entry.Acceleration.Get();
+				const int32 TypeIdx = GetTypeIdx(RawPtr);
+				switch (TypeIdx)
+				{
+				case 0: GetAccelerationsPerType<0>(Types).Objects[Entry.TypeInnerIdx] = &RawPtr->template AsChecked<typename std::tuple_element<0, std::tuple<TSpatialAccelerationTypes...>>::type>(); break;
+				case 1: GetAccelerationsPerType<ClampedIdx(1)>(Types).Objects[Entry.TypeInnerIdx] = &RawPtr->template AsChecked<typename std::tuple_element<ClampedIdx(1), std::tuple<TSpatialAccelerationTypes...>>::type>(); break;
+				case 2: GetAccelerationsPerType<ClampedIdx(2)>(Types).Objects[Entry.TypeInnerIdx] = &RawPtr->template AsChecked<typename std::tuple_element<ClampedIdx(2), std::tuple<TSpatialAccelerationTypes...>>::type>(); break;
+				case 3: GetAccelerationsPerType<ClampedIdx(3)>(Types).Objects[Entry.TypeInnerIdx] = &RawPtr->template AsChecked<typename std::tuple_element<ClampedIdx(3), std::tuple<TSpatialAccelerationTypes...>>::type>(); break;
+				case 4: GetAccelerationsPerType<ClampedIdx(4)>(Types).Objects[Entry.TypeInnerIdx] = &RawPtr->template AsChecked<typename std::tuple_element<ClampedIdx(4), std::tuple<TSpatialAccelerationTypes...>>::type>(); break;
+				case 5: GetAccelerationsPerType<ClampedIdx(5)>(Types).Objects[Entry.TypeInnerIdx] = &RawPtr->template AsChecked<typename std::tuple_element<ClampedIdx(5), std::tuple<TSpatialAccelerationTypes...>>::type>(); break;
+				case 6: GetAccelerationsPerType<ClampedIdx(6)>(Types).Objects[Entry.TypeInnerIdx] = &RawPtr->template AsChecked<typename std::tuple_element<ClampedIdx(6), std::tuple<TSpatialAccelerationTypes...>>::type>(); break;
+				case 7: GetAccelerationsPerType<ClampedIdx(7)>(Types).Objects[Entry.TypeInnerIdx] = &RawPtr->template AsChecked<typename std::tuple_element<ClampedIdx(7), std::tuple<TSpatialAccelerationTypes...>>::type>(); break;
+				}
+			}
+		}
+	}
+	
+	static constexpr uint32 ClampedIdx(uint32 Idx)
+	{
+		return Idx < NumTypes ? Idx : 0;
+	}
+
+	static uint8 GetTypeIdx(const ISpatialAcceleration<TPayloadType, T, d>* Accel)
+	{
+		if (Accel->template As<typename std::tuple_element<0, std::tuple< TSpatialAccelerationTypes...>>::type> ())
+		{
+			return 0;
+		}
+		
+		if (Accel->template As<typename std::tuple_element<ClampedIdx(1), std::tuple< TSpatialAccelerationTypes...>>::type> ())
+		{
+			return 1;
+		}
+
+		if (Accel->template As<typename std::tuple_element<ClampedIdx(2), std::tuple< TSpatialAccelerationTypes...>>::type> ())
+		{
+			return 2;
+		}
+
+		if (Accel->template As<typename std::tuple_element<ClampedIdx(3), std::tuple< TSpatialAccelerationTypes...>>::type> ())
+		{
+			return 3;
+		}
+
+		if (Accel->template As<typename std::tuple_element<ClampedIdx(4), std::tuple< TSpatialAccelerationTypes...>>::type> ())
+		{
+			return 4;
+		}
+
+		if (Accel->template As<typename std::tuple_element<ClampedIdx(5), std::tuple< TSpatialAccelerationTypes...>>::type> ())
+		{
+			return 5;
+		}
+
+		if (Accel->template As<typename std::tuple_element<ClampedIdx(6), std::tuple< TSpatialAccelerationTypes...>>::type> ())
+		{
+			return 6;
+		}
+
+		if (Accel->template As<typename std::tuple_element<ClampedIdx(7), std::tuple< TSpatialAccelerationTypes...>>::type> ())
+		{
+			return 7;
+		}
+
+		check(false);	//Did you instantiate the collection with the right template arguments?
+		return 0;
+	}
+
+	static constexpr uint16 MaxBuckets = 8;
+	TSpatialCollectionBucket<TSpatialAccelerationBucketEntry<TPayloadType, T, d>> Buckets[MaxBuckets];
+	TSpatialTypeTuple< TSpatialAccelerationTypes...> Types;
+	static constexpr uint32 NumTypes = sizeof...(TSpatialAccelerationTypes);
+};
+
+
+}
diff --git a/Engine/Source/Runtime/Experimental/ChaosSolvers/Private/ChaosSolversModule.cpp b/Engine/Source/Runtime/Experimental/ChaosSolvers/Private/ChaosSolversModule.cpp
--- a/Engine/Source/Runtime/Experimental/ChaosSolvers/Private/ChaosSolversModule.cpp
+++ b/Engine/Source/Runtime/Experimental/ChaosSolvers/Private/ChaosSolversModule.cpp
@@ -41,6 +41,10 @@ TAutoConsoleVariable<int32> CVarDedicatedThreadSyncThreshold(
 
 namespace Chaos
 {
+#if !UE_BUILD_SHIPPING
+	CHAOS_API extern bool bPendingHierarchyDump;
+#endif
+
 	void ChangeThreadingMode(EThreadingMode InMode)
 	{
 		FChaosSolversModule* ChaosModule = FModuleManager::Get().GetModulePtr<FChaosSolversModule>("ChaosSolvers");
@@ -555,13 +559,6 @@ TAutoConsoleVariable<FString> DumpHier_ElementBuckets(
 	TEXT("1,4,8,16,32,64,128,256,512"),
 	TEXT("Distribution buckets for dump hierarchy stats command"));
 
-#if !UE_BUILD_SHIPPING
-namespace Chaos
-{
-CHAOS_API extern int32 PendingHierarchyDump;
-}
-#endif
-
 void FChaosSolversModule::DumpHierarchyStats(int32* OutOptMaxCellElements)
 {
 	TArray<FString> BucketStrings;
@@ -711,7 +708,7 @@ void FChaosSolversModule::DumpHierarchyStats(int32* OutOptMaxCellElements)
 #endif
 
 #if !UE_BUILD_SHIPPING
-		Chaos::PendingHierarchyDump = 1;	//mark solver pending dump to get more info
+		Chaos::bPendingHierarchyDump = true;	//mark solver pending dump to get more info
 #endif
 	}
 }
@@ -933,7 +930,11 @@ void FChaosSolversModule::OnCreateMaterial(Chaos::FMaterialHandle InHandle)
 
 void FChaosSolversModule::OnDestroyMaterial(Chaos::FMaterialHandle InHandle)
 {
-	check(Dispatcher);
+	//check(Dispatcher);
+	if (!Dispatcher)
+	{
+		return;
+	}
 
 	// Grab the material
 	Chaos::FChaosPhysicsMaterial* Material = InHandle.Get();
