Chaos - fix swing joint limit math

diff --git a/Engine/Source/Runtime/Experimental/Chaos/Private/Chaos/PBDJointConstraintSolver.cpp b/Engine/Source/Runtime/Experimental/Chaos/Private/Chaos/PBDJointConstraintSolver.cpp
--- a/Engine/Source/Runtime/Experimental/Chaos/Private/Chaos/PBDJointConstraintSolver.cpp
+++ b/Engine/Source/Runtime/Experimental/Chaos/Private/Chaos/PBDJointConstraintSolver.cpp
@@ -430,10 +430,10 @@ namespace Chaos
 			C.AddRows(1);
 
 			J0.SetRowAt(RowIndex + 0, 0, FVec3(0, 0, 0));
-			J0.SetRowAt(RowIndex + 0, 3, -FVec3::CrossProduct(Axis, XP0));
+			J0.SetRowAt(RowIndex + 0, 3, Axis);
 
 			J1.SetRowAt(RowIndex + 0, 0, FVec3(0, 0, 0));
-			J1.SetRowAt(RowIndex + 0, 3, FVec3::CrossProduct(Axis, XP1));
+			J1.SetRowAt(RowIndex + 0, 3, -Axis);
 
 			C.SetAt(RowIndex, 0, (SwingAngle >= SwingAngleMax) ? SwingAngle - SwingAngleMax : SwingAngle + SwingAngleMax);
 		}
