Fix Shipping CIS (PBDRigidsEvolution.cpp)

diff --git a/Engine/Source/Runtime/Experimental/Chaos/Private/Chaos/PBDRigidsEvolutionGBF.cpp b/Engine/Source/Runtime/Experimental/Chaos/Private/Chaos/PBDRigidsEvolutionGBF.cpp
--- a/Engine/Source/Runtime/Experimental/Chaos/Private/Chaos/PBDRigidsEvolutionGBF.cpp
+++ b/Engine/Source/Runtime/Experimental/Chaos/Private/Chaos/PBDRigidsEvolutionGBF.cpp
@@ -25,6 +25,8 @@ namespace Chaos
 {
 #if !UE_BUILD_SHIPPING
 	CHAOS_API bool bPendingHierarchyDump = false;
+#else
+	const bool bPendingHierarchyDump = false;
 #endif
 
 float HackMaxAngularVelocity = 1000.f;
diff --git a/Engine/Source/Runtime/Experimental/ChaosSolvers/Private/ChaosSolversModule.cpp b/Engine/Source/Runtime/Experimental/ChaosSolvers/Private/ChaosSolversModule.cpp
--- a/Engine/Source/Runtime/Experimental/ChaosSolvers/Private/ChaosSolversModule.cpp
+++ b/Engine/Source/Runtime/Experimental/ChaosSolvers/Private/ChaosSolversModule.cpp
@@ -930,8 +930,9 @@ void FChaosSolversModule::OnCreateMaterial(Chaos::FMaterialHandle InHandle)
 
 void FChaosSolversModule::OnDestroyMaterial(Chaos::FMaterialHandle InHandle)
 {
+	// @todo(bgallagher)
 	//check(Dispatcher);
-	if (!Dispatcher)
+	if (!ensure(Dispatcher))
 	{
 		return;
 	}
