Fix fields for geometry collections.

diff --git a/Engine/Plugins/Experimental/ChaosNiagara/Source/ChaosNiagara/Private/NiagaraDataInterfaceChaosDestruction.cpp b/Engine/Plugins/Experimental/ChaosNiagara/Source/ChaosNiagara/Private/NiagaraDataInterfaceChaosDestruction.cpp
--- a/Engine/Plugins/Experimental/ChaosNiagara/Source/ChaosNiagara/Private/NiagaraDataInterfaceChaosDestruction.cpp
+++ b/Engine/Plugins/Experimental/ChaosNiagara/Source/ChaosNiagara/Private/NiagaraDataInterfaceChaosDestruction.cpp
@@ -175,9 +175,12 @@ void UNiagaraDataInterfaceChaosDestruction::PostLoad()
 	{
 		FPhysScene* Scene = GetWorld()->GetPhysicsScene();
 		Chaos::FEventManager* EventManager = Scene->GetScene().GetSolver()->GetEventManager();
-		EventManager->RegisterHandler<Chaos::FCollisionEventData>(Chaos::EEventType::Collision, this, &UNiagaraDataInterfaceChaosDestruction::HandleCollisionEvents);
-		EventManager->RegisterHandler<Chaos::FBreakingEventData>(Chaos::EEventType::Breaking, this, &UNiagaraDataInterfaceChaosDestruction::HandleBreakingEvents);
-		EventManager->RegisterHandler<Chaos::FTrailingEventData>(Chaos::EEventType::Trailing, this, &UNiagaraDataInterfaceChaosDestruction::HandleTrailingEvents);
+		if (EventManager)
+		{
+			EventManager->RegisterHandler<Chaos::FCollisionEventData>(Chaos::EEventType::Collision, this, &UNiagaraDataInterfaceChaosDestruction::HandleCollisionEvents);
+			EventManager->RegisterHandler<Chaos::FBreakingEventData>(Chaos::EEventType::Breaking, this, &UNiagaraDataInterfaceChaosDestruction::HandleBreakingEvents);
+			EventManager->RegisterHandler<Chaos::FTrailingEventData>(Chaos::EEventType::Trailing, this, &UNiagaraDataInterfaceChaosDestruction::HandleTrailingEvents);
+		}
 	}
 #endif
 
@@ -193,9 +196,12 @@ void UNiagaraDataInterfaceChaosDestruction::BeginDestroy()
 	{
 		FPhysScene* Scene = GetWorld()->GetPhysicsScene();
 		Chaos::FEventManager* EventManager = Scene->GetScene().GetSolver()->GetEventManager();
-		EventManager->UnregisterHandler(Chaos::EEventType::Collision, this);
-		EventManager->UnregisterHandler(Chaos::EEventType::Breaking, this);
-		EventManager->UnregisterHandler(Chaos::EEventType::Trailing, this);
+		if (EventManager)
+		{
+			EventManager->UnregisterHandler(Chaos::EEventType::Collision, this);
+			EventManager->UnregisterHandler(Chaos::EEventType::Breaking, this);
+			EventManager->UnregisterHandler(Chaos::EEventType::Trailing, this);
+		}
 	}
 #endif
 }
diff --git a/Engine/Plugins/Experimental/GeometryCollectionPlugin/Source/GeometryCollectionExampleCore/Private/GeometryCollection/GeometryCollectionExample.cpp b/Engine/Plugins/Experimental/GeometryCollectionPlugin/Source/GeometryCollectionExampleCore/Private/GeometryCollection/GeometryCollectionExample.cpp
--- a/Engine/Plugins/Experimental/GeometryCollectionPlugin/Source/GeometryCollectionExampleCore/Private/GeometryCollection/GeometryCollectionExample.cpp
+++ b/Engine/Plugins/Experimental/GeometryCollectionPlugin/Source/GeometryCollectionExampleCore/Private/GeometryCollection/GeometryCollectionExample.cpp
@@ -19,11 +19,12 @@
 namespace GeometryCollectionExample
 {
 #define RUN_EXAMPLE(X) X<float>();
+#define RUN_EXAMPLE_NO_TEMPLATE(X) X();
 
-#if TODO_REIMPLEMENT_RIGID_CLUSTERING
 	template<class RESPONSE>
 	void ExecuteExamples()
 	{
+#if TODO_REIMPLEMENT_RIGID_CLUSTERING
 		RUN_EXAMPLE(BasicGlobalMatrices);
 		RUN_EXAMPLE(TransformMatrixElement);
 		RUN_EXAMPLE(ReparentingMatrices);
@@ -39,23 +40,25 @@ namespace GeometryCollectionExample
 		RUN_EXAMPLE(DeleteBranch);
 		RUN_EXAMPLE(DeleteRootLeafMiddle);
 		RUN_EXAMPLE(DeleteEverything);
-		RUN_EXAMPLE(Fields_NoiseSample);
-		RUN_EXAMPLE(Fields_RadialIntMask);
-		RUN_EXAMPLE(Fields_RadialFalloff);
-		RUN_EXAMPLE(Fields_PlaneFalloff);
-		RUN_EXAMPLE(Fields_UniformVector);
-		RUN_EXAMPLE(Fields_RaidalVector);
-		RUN_EXAMPLE(Fields_SumVectorFullMult);
-		RUN_EXAMPLE(Fields_SumVectorFullDiv);
-		RUN_EXAMPLE(Fields_SumVectorFullAdd);
-		RUN_EXAMPLE(Fields_SumVectorFullSub);
-		RUN_EXAMPLE(Fields_SumVectorLeftSide);
-		RUN_EXAMPLE(Fields_SumVectorRightSide);
-		RUN_EXAMPLE(Fields_SumScalar);
-		RUN_EXAMPLE(Fields_SumScalarRightSide);
-		RUN_EXAMPLE(Fields_SumScalarLeftSide);
-		RUN_EXAMPLE(Fields_Culling);
-		RUN_EXAMPLE(Fields_SerializeAPI);
+#endif //TODO_REIMPLEMENT_RIGID_CLUSTERING
+		RUN_EXAMPLE_NO_TEMPLATE(Fields_NoiseSample);
+		RUN_EXAMPLE_NO_TEMPLATE(Fields_RadialIntMask);
+		RUN_EXAMPLE_NO_TEMPLATE(Fields_RadialFalloff);
+		RUN_EXAMPLE_NO_TEMPLATE(Fields_PlaneFalloff);
+		RUN_EXAMPLE_NO_TEMPLATE(Fields_UniformVector);
+		RUN_EXAMPLE_NO_TEMPLATE(Fields_RaidalVector);
+		RUN_EXAMPLE_NO_TEMPLATE(Fields_SumVectorFullMult);
+		RUN_EXAMPLE_NO_TEMPLATE(Fields_SumVectorFullDiv);
+		RUN_EXAMPLE_NO_TEMPLATE(Fields_SumVectorFullAdd);
+		RUN_EXAMPLE_NO_TEMPLATE(Fields_SumVectorFullSub);
+		RUN_EXAMPLE_NO_TEMPLATE(Fields_SumVectorLeftSide);
+		RUN_EXAMPLE_NO_TEMPLATE(Fields_SumVectorRightSide);
+		RUN_EXAMPLE_NO_TEMPLATE(Fields_SumScalar);
+		RUN_EXAMPLE_NO_TEMPLATE(Fields_SumScalarRightSide);
+		RUN_EXAMPLE_NO_TEMPLATE(Fields_SumScalarLeftSide);
+		RUN_EXAMPLE_NO_TEMPLATE(Fields_Culling);
+		RUN_EXAMPLE_NO_TEMPLATE(Fields_SerializeAPI);
+#if TODO_REIMPLEMENT_RIGID_CLUSTERING
 		RUN_EXAMPLE(Solver_AdvanceNoObjects);
 		RUN_EXAMPLE(Solver_AdvanceDisabledObjects);
 		RUN_EXAMPLE(Solver_AdvanceDisabledClusteredObjects);
@@ -80,6 +83,7 @@ namespace GeometryCollectionExample
 		RUN_EXAMPLE(RigidBodies_ClusterTest_ReleaseClusterParticle_ClusteredKinematicNode);
 		RUN_EXAMPLE(RigidBodies_ClusterTest_ReleaseClusterParticles_AllLeafNodes);
 		RUN_EXAMPLE(RigidBodies_ClusterTest_ReleaseClusterParticles_ClusterNodeAndSubClusterNode);
+#endif // TODO_REIMPLEMENT_RIGID_CLUSTERING
 		RUN_EXAMPLE(RigidBodies_Field_KinematicActivation);
 		RUN_EXAMPLE(RigidBodies_Field_InitialLinearVelocity);
 		RUN_EXAMPLE(RigidBodies_Field_StayDynamic);
@@ -91,6 +95,7 @@ namespace GeometryCollectionExample
 		RUN_EXAMPLE(RigidBodies_Field_ClusterBreak_StrainModel_Test1);
 		RUN_EXAMPLE(RigidBodies_Field_ClusterBreak_StrainModel_Test2);
 		RUN_EXAMPLE(RigidBodies_Field_ClusterBreak_StrainModel_Test3);
+#if TODO_REIMPLEMENT_RIGID_CLUSTERING
 		RUN_EXAMPLE(RigidBodies_Streaming_StartSolverEmpty);
 		RUN_EXAMPLE(RigidBodies_Streaming_BulkInitialization);
 		RUN_EXAMPLE(RigidBodies_Streaming_DeferedClusteringInitialization);
@@ -115,8 +120,7 @@ namespace GeometryCollectionExample
 		RUN_EXAMPLE(HashTableUpdateTest);
 		RUN_EXAMPLE(HashTablePressureTest);
 		RUN_EXAMPLE(TestGeometryDecimation);
-		
+#endif
 	}
 	template void ExecuteExamples<ExampleResponse>();
-#endif
 }
diff --git a/Engine/Plugins/Experimental/GeometryCollectionPlugin/Source/GeometryCollectionExampleCore/Private/GeometryCollection/GeometryCollectionExampleFields.cpp b/Engine/Plugins/Experimental/GeometryCollectionPlugin/Source/GeometryCollectionExampleCore/Private/GeometryCollection/GeometryCollectionExampleFields.cpp
--- a/Engine/Plugins/Experimental/GeometryCollectionPlugin/Source/GeometryCollectionExampleCore/Private/GeometryCollection/GeometryCollectionExampleFields.cpp
+++ b/Engine/Plugins/Experimental/GeometryCollectionPlugin/Source/GeometryCollectionExampleCore/Private/GeometryCollection/GeometryCollectionExampleFields.cpp
@@ -2,6 +2,7 @@
 
 #include "GeometryCollection/GeometryCollectionExampleFields.h"
 
+#include "Chaos/Utilities.h"
 #include "Field/FieldSystem.h"
 #include "Field/FieldSystemNodes.h"
 #include "Field/FieldSystemCoreAlgo.h"
@@ -55,23 +56,18 @@ namespace GeometryCollectionExample
 		NoiseField->Evaluate(Context, ResultsView);
 
 		float min = FLT_MAX, max = -FLT_MAX;
-
-		for (int32 Index = 0, i = 0; i < Bounds; i++)
-		{
-			for (int32 j = 0; j < Bounds; j++)
-			{
-				min = FMath::Min(min, ResultsView[Index]);
-				max = FMath::Max(max, ResultsView[Index]);
-				//std::cout << ResultsView[Index] << " ";
-				Index++;
-			}
-			//std::cout << std::endl;
-		}
-
+		double avg = 0.0;
+		Chaos::Utilities::GetMinAvgMax(ResultsView, min, avg, max);
 		EXPECT_GE(min, MinDoimain);
 		EXPECT_LE(max, MaxDomain);
 		EXPECT_LT(min, max);
 
+		float variance = Chaos::Utilities::GetVariance(ResultsView, avg);
+		float stdDev = Chaos::Utilities::GetStandardDeviation(variance);
+		EXPECT_GT(variance, 0.0); // If variance is 0, then all values are the same.
+		EXPECT_GT(stdDev, 0.0);
+		EXPECT_LT(stdDev, 0.5);
+
 		delete NoiseField;
 	}
 
@@ -89,8 +85,6 @@ namespace GeometryCollectionExample
 		}
 		TArrayView<FVector> SamplesView(&(SamplesArray.operator[](0)), SamplesArray.Num());
 
-
-
 		FRadialIntMask * RadialMask = new FRadialIntMask(0.0, FVector(), 1.0, 0.0, ESetMaskConditionType::Field_Set_Always);
 		RadialMask->Position = FVector(0.0, 0.0, 0.0);
 		RadialMask->Radius = 5.0;
@@ -127,7 +121,6 @@ namespace GeometryCollectionExample
 		ContextIndex::ContiguousIndices(IndicesArray, 10);
 		TArrayView<ContextIndex> IndexView(&(IndicesArray[0]), IndicesArray.Num());
 
-
 		TArray<FVector> SamplesArray;
 		SamplesArray.Init(FVector(0.f), 10);
 		for (int32 Index = 0; Index < 10; Index++)
@@ -658,8 +651,6 @@ namespace GeometryCollectionExample
 		}
 		delete SumVector;
 		delete UniformVector;
-
-		
 	}
 
 	void Fields_SumVectorRightSide()
@@ -668,7 +659,6 @@ namespace GeometryCollectionExample
 		ContextIndex::ContiguousIndices(IndicesArray, 10);
 		TArrayView<ContextIndex> IndexView(&(IndicesArray[0]), IndicesArray.Num());
 
-
 		float AverageSampleLength = 0.0;
 		TArray<FVector> SamplesArray;
 		SamplesArray.Init(FVector(0.f), 10);
@@ -700,7 +690,6 @@ namespace GeometryCollectionExample
 
 		FSumVector * SumVector = new FSumVector(1.0, RadialFalloff, UniformVector, nullptr, EFieldOperationType::Field_Multiply);
 
-
 		FFieldContext Context{
 			IndexView,
 			SamplesView,
@@ -734,8 +723,6 @@ namespace GeometryCollectionExample
 		}
 		delete SumVector;
 		delete RadialVector;
-
-		
 	}
 
 	void Fields_SumScalar()
@@ -753,8 +740,6 @@ namespace GeometryCollectionExample
 			SamplesArray[Index + 10] = FVector(Index, 0, 0);
 		}
 		TArrayView<FVector> SamplesView(&(SamplesArray.operator[](0)), SamplesArray.Num());
-
-
 		
 		FRadialFalloff * RadialFalloff = new FRadialFalloff();
 		RadialFalloff->Position = FVector(5.0, 0.0, 0.0);
@@ -811,8 +796,6 @@ namespace GeometryCollectionExample
 			EXPECT_LT((ResultsView[Index] - ExpectedVal), KINDA_SMALL_NUMBER);
 		}
 		delete SumScalar;
-
-		
 	}
 
 	void Fields_SumScalarRightSide()
@@ -830,7 +813,6 @@ namespace GeometryCollectionExample
 		}
 		TArrayView<FVector> SamplesView(&(SamplesArray.operator[](0)), SamplesArray.Num());
 
-
 		FRadialFalloff * RadialFalloff = new FRadialFalloff();
 		RadialFalloff->Position = FVector(5.0, 0.0, 0.0);
 		RadialFalloff->Radius = 10;
@@ -868,7 +850,6 @@ namespace GeometryCollectionExample
 			}
 
 			float ScalarLeft = 1.f;
-
 			float ExpectedVal = ScalarLeft * ScalarRight;
 
 			//UE_LOG(GCTF_Log, Error, TEXT("[%d:%3.5f] sample:(%3.5f,%3.5f,%3.5f) %3.5f -> %3.5f"), Index,
@@ -880,8 +861,6 @@ namespace GeometryCollectionExample
 		}
 		delete SumScalar;
 		delete RadialFalloff2;
-
-		
 	}
 
 	void Fields_SumScalarLeftSide()
@@ -899,8 +878,6 @@ namespace GeometryCollectionExample
 		}
 		TArrayView<FVector> SamplesView(&(SamplesArray.operator[](0)), SamplesArray.Num());
 
-
-
 		FRadialFalloff * RadialFalloff = new FRadialFalloff();
 		RadialFalloff->Position = FVector(5.0, 0.0, 0.0);
 		RadialFalloff->Radius = 10;
@@ -950,8 +927,6 @@ namespace GeometryCollectionExample
 		}
 		delete SumScalar;
 		delete RadialFalloff;
-
-		
 	}
 
 	void Fields_Culling()
@@ -969,15 +944,12 @@ namespace GeometryCollectionExample
 		}
 		TArrayView<FVector> SamplesView(&(SamplesArray.operator[](0)), SamplesArray.Num());
 
-
-
 		FRadialFalloff * RadialFalloff = new FRadialFalloff();
 		RadialFalloff->Position = FVector(0.0, 0.0, 0.0);
 		RadialFalloff->Radius = 4.;
 		RadialFalloff->Magnitude = 3.0;
 		float RadialFalloffRadius2 = RadialFalloff->Radius * RadialFalloff->Radius;
 
-
 		FRadialFalloff * RadialFalloff2 = new FRadialFalloff();
 		RadialFalloff2->Position = FVector(5.0, 0.0, 0.0);
 		RadialFalloff2->Radius = 10;
@@ -1076,7 +1048,6 @@ namespace GeometryCollectionExample
 		return false;
 	}
 
-	template<class T>
 	void Fields_SerializeAPI()
 	{
 		{
@@ -1255,7 +1226,5 @@ namespace GeometryCollectionExample
 		}
 
 	}
-	template void Fields_SerializeAPI<float>();
-
-
+	//template void Fields_SerializeAPI<float>();
 }
diff --git a/Engine/Plugins/Experimental/GeometryCollectionPlugin/Source/GeometryCollectionExampleCore/Private/GeometryCollection/GeometryCollectionExampleSimulationField.cpp b/Engine/Plugins/Experimental/GeometryCollectionPlugin/Source/GeometryCollectionExampleCore/Private/GeometryCollection/GeometryCollectionExampleSimulationField.cpp
--- a/Engine/Plugins/Experimental/GeometryCollectionPlugin/Source/GeometryCollectionExampleCore/Private/GeometryCollection/GeometryCollectionExampleSimulationField.cpp
+++ b/Engine/Plugins/Experimental/GeometryCollectionPlugin/Source/GeometryCollectionExampleCore/Private/GeometryCollection/GeometryCollectionExampleSimulationField.cpp
@@ -31,6 +31,8 @@ namespace GeometryCollectionExample
 	template<class T>
 	void RigidBodies_Field_KinematicActivation()
 	{
+		FChaosSolversModule* Module = FChaosSolversModule::GetModule();
+
 		TUniquePtr<Chaos::FChaosPhysicsMaterial> PhysicalMaterial = nullptr;
 		TSharedPtr<FGeometryCollection> RestCollection = nullptr;
 		TSharedPtr<FGeometryDynamicCollection> DynamicCollection = nullptr;
@@ -38,32 +40,28 @@ namespace GeometryCollectionExample
 		//
 		//  Rigid Body Setup
 		//
-		auto RestInitFunc = [](TSharedPtr<FGeometryCollection>& RestCollection)
+		const FVector Translation0(0, 0, 1);
+		auto RestInitFunc = [Translation0](TSharedPtr<FGeometryCollection>& RestCollection)
 		{
-			RestCollection->Transform[0].SetTranslation(FVector(0, 0, 1));
+			RestCollection->Transform[0].SetTranslation(Translation0);
 		};
 
 		InitCollectionsParameters InitParams = { FTransform::Identity, FVector(1.0), RestInitFunc, (int32)EObjectStateTypeEnum::Chaos_Object_Kinematic };
 		InitCollections(PhysicalMaterial, RestCollection, DynamicCollection, InitParams);
 
-		FRadialIntMask * RadialMask = new FRadialIntMask();
-		RadialMask->Position = FVector(0.0, 0.0, 0.0);
-		RadialMask->Radius = 100.0;
-		RadialMask->InteriorValue = (int32)EObjectStateTypeEnum::Chaos_Object_Dynamic;
-		RadialMask->ExteriorValue = (int32)EObjectStateTypeEnum::Chaos_Object_Kinematic;
-		RadialMask->SetMaskCondition = ESetMaskConditionType::Field_Set_IFF_NOT_Interior;
 
 		FGeometryCollectionPhysicsProxy* PhysObject = RigidBodySetup(PhysicalMaterial, RestCollection, DynamicCollection);
 
-		Chaos::FPBDRigidsSolver* Solver = FChaosSolversModule::GetModule()->CreateSolver(true);
-#if CHAOS_PARTICLEHANDLE_TODO
+		Chaos::FPBDRigidsSolver* Solver = Module->CreateSolver(true);
 		Solver->RegisterObject(PhysObject);
-#endif
 		Solver->SetHasFloor(false);
 		Solver->SetIsFloorAnalytic(true);
 		Solver->SetEnabled(true);
 		PhysObject->ActivateBodies();
 
+		Solver->AddDirtyProxy(PhysObject);
+		Solver->PushPhysicsState(Module->GetDispatcher());
+
 		for (int i = 0; i < 100; i++)
 		{
 			Solver->AdvanceSolverBy(1 / 24.);
@@ -72,9 +70,19 @@ namespace GeometryCollectionExample
 		FinalizeSolver(*Solver);
 
 		// simulated
-		TManagedArray<FTransform>& Transform = DynamicCollection->Transform;
+		const TManagedArray<FTransform>& Transform = DynamicCollection->Transform;
 		EXPECT_EQ(Transform.Num(), 1);
+		const FVector Translation1 = Transform[0].GetTranslation();
+		EXPECT_EQ(Translation0, Translation1);
 		EXPECT_EQ(Transform[0].GetTranslation().Z, 1.f);
+
+
+		FRadialIntMask * RadialMask = new FRadialIntMask();
+		RadialMask->Position = FVector(0.0, 0.0, 0.0);
+		RadialMask->Radius = 100.0;
+		RadialMask->InteriorValue = (int32)EObjectStateTypeEnum::Chaos_Object_Dynamic;
+		RadialMask->ExteriorValue = (int32)EObjectStateTypeEnum::Chaos_Object_Kinematic;
+		RadialMask->SetMaskCondition = ESetMaskConditionType::Field_Set_IFF_NOT_Interior;
 		FName TargetName = GetFieldPhysicsName(EFieldPhysicsType::Field_DynamicState);
 		PhysObject->BufferCommand(Solver, { TargetName, RadialMask });
 
@@ -85,9 +93,11 @@ namespace GeometryCollectionExample
 
 		FinalizeSolver(*Solver);
 
+		const FVector Translation2 = Transform[0].GetTranslation();
+		EXPECT_NE(Translation1, Translation2);
 		EXPECT_LE(Transform[0].GetTranslation().Z, 0.f);
 
-		FChaosSolversModule::GetModule()->DestroySolver(Solver);
+		Module->DestroySolver(Solver);
 		delete PhysObject;
 
 	}
@@ -96,6 +106,8 @@ namespace GeometryCollectionExample
 	template<class T>
 	void RigidBodies_Field_InitialLinearVelocity()
 	{
+		FChaosSolversModule* Module = FChaosSolversModule::GetModule();
+
 		TUniquePtr<Chaos::FChaosPhysicsMaterial> PhysicalMaterial = nullptr;
 		TSharedPtr<FGeometryCollection> RestCollection = nullptr;
 		TSharedPtr<FGeometryDynamicCollection> DynamicCollection = nullptr;
@@ -129,15 +141,17 @@ namespace GeometryCollectionExample
 		FGeometryCollectionPhysicsProxy* PhysObject = RigidBodySetup(PhysicalMaterial, RestCollection, DynamicCollection, CustomFunc);
 
 		Chaos::FPBDRigidsSolver* Solver = FChaosSolversModule::GetModule()->CreateSolver(true);
-#if CHAOS_PARTICLEHANDLE_TODO
 		Solver->RegisterObject(PhysObject);
 		Solver->RegisterObject(FieldObject);
-#endif
 		Solver->SetHasFloor(false);
 		Solver->SetIsFloorAnalytic(true);
 		Solver->SetEnabled(true);
 		PhysObject->ActivateBodies();
 
+		Solver->AddDirtyProxy(PhysObject);
+		Solver->AddDirtyProxy(FieldObject);
+		Solver->PushPhysicsState(Module->GetDispatcher());
+
 		TManagedArray<FTransform>& Transform = DynamicCollection->Transform;
 
 		float PreviousY = 0.f;
@@ -172,13 +186,14 @@ namespace GeometryCollectionExample
 		FChaosSolversModule::GetModule()->DestroySolver(Solver);
 		delete FieldObject;
 		delete PhysObject;
-
 	}
 	template void RigidBodies_Field_InitialLinearVelocity<float>();
 
 	template<class T>
 	void RigidBodies_Field_StayDynamic()
 	{
+		FChaosSolversModule* Module = FChaosSolversModule::GetModule();
+
 		TUniquePtr<Chaos::FChaosPhysicsMaterial> PhysicalMaterial = nullptr;
 		TSharedPtr<FGeometryCollection> RestCollection = nullptr;
 		TSharedPtr<FGeometryDynamicCollection> DynamicCollection = nullptr;
@@ -212,14 +227,16 @@ namespace GeometryCollectionExample
 		FFieldSystemPhysicsProxy* FieldObject = new FFieldSystemPhysicsProxy(nullptr);
 
 		Chaos::FPBDRigidsSolver* Solver = FChaosSolversModule::GetModule()->CreateSolver(true);
-#if CHAOS_PARTICLEHANDLE_TODO
 		Solver->RegisterObject(PhysObject);
 		Solver->RegisterObject(FieldObject);
-#endif
 		Solver->SetHasFloor(false);
 		Solver->SetEnabled(true);
 		PhysObject->ActivateBodies();
 
+		Solver->AddDirtyProxy(PhysObject);
+		Solver->AddDirtyProxy(FieldObject);
+		Solver->PushPhysicsState(Module->GetDispatcher());
+
 		TManagedArray<FTransform>& Transform = DynamicCollection->Transform;
 		float PreviousHeight = 5.f;
 		for (int Frame = 0; Frame < 10; Frame++)
@@ -265,6 +282,8 @@ namespace GeometryCollectionExample
 	template<class T>
 	void RigidBodies_Field_LinearForce()
 	{
+		FChaosSolversModule* Module = FChaosSolversModule::GetModule();
+
 		TUniquePtr<Chaos::FChaosPhysicsMaterial> PhysicalMaterial = nullptr;
 		TSharedPtr<FGeometryCollection> RestCollection = nullptr;
 		TSharedPtr<FGeometryDynamicCollection> DynamicCollection = nullptr;
@@ -295,14 +314,16 @@ namespace GeometryCollectionExample
 		FFieldSystemPhysicsProxy* FieldObject = new FFieldSystemPhysicsProxy(nullptr);
 
 		Chaos::FPBDRigidsSolver* Solver = FChaosSolversModule::GetModule()->CreateSolver(true);
-#if CHAOS_PARTICLEHANDLE_TODO
 		Solver->RegisterObject(PhysObject);
 		Solver->RegisterObject(FieldObject);
-#endif
 		Solver->SetHasFloor(false);
 		Solver->SetEnabled(true);
 		PhysObject->ActivateBodies();
 
+		Solver->AddDirtyProxy(PhysObject);
+		Solver->AddDirtyProxy(FieldObject);
+		Solver->PushPhysicsState(Module->GetDispatcher());
+
 		TManagedArray<FTransform>& Transform = DynamicCollection->Transform;
 		float PreviousY = 0.f;
 		for (int Frame = 0; Frame < 10; Frame++)
@@ -342,6 +363,8 @@ namespace GeometryCollectionExample
 	template<class T>
 	void RigidBodies_Field_Torque()
 	{
+		FChaosSolversModule* Module = FChaosSolversModule::GetModule();
+
 		TUniquePtr<Chaos::FChaosPhysicsMaterial> PhysicalMaterial = nullptr;
 		TSharedPtr<FGeometryCollection> RestCollection = nullptr;
 		TSharedPtr<FGeometryDynamicCollection> DynamicCollection = nullptr;
@@ -372,14 +395,16 @@ namespace GeometryCollectionExample
 		FFieldSystemPhysicsProxy* FieldObject = new FFieldSystemPhysicsProxy(nullptr);
 
 		Chaos::FPBDRigidsSolver* Solver = FChaosSolversModule::GetModule()->CreateSolver(true);
-#if CHAOS_PARTICLEHANDLE_TODO
 		Solver->RegisterObject(PhysObject);
 		Solver->RegisterObject(FieldObject);
-#endif
 		Solver->SetHasFloor(false);
 		Solver->SetEnabled(true);
 		PhysObject->ActivateBodies();
 
+		Solver->AddDirtyProxy(PhysObject);
+		Solver->AddDirtyProxy(FieldObject);
+		Solver->PushPhysicsState(Module->GetDispatcher());
+
 		TManagedArray<FTransform>& Transform = DynamicCollection->Transform;
 		//TManagedArray<FVector>& AngularVelocity = DynamicCollection->;
 		float PreviousY = 0.f;
@@ -427,6 +452,8 @@ namespace GeometryCollectionExample
 	template<class T>
 	void RigidBodies_Field_Kill()
 	{
+		FChaosSolversModule* Module = FChaosSolversModule::GetModule();
+
 		TUniquePtr<Chaos::FChaosPhysicsMaterial> PhysicalMaterial = nullptr;
 		TSharedPtr<FGeometryCollection> RestCollection = nullptr;
 		TSharedPtr<FGeometryDynamicCollection> DynamicCollection = nullptr;
@@ -459,14 +486,16 @@ namespace GeometryCollectionExample
 		FFieldSystemPhysicsProxy* FieldObject = new FFieldSystemPhysicsProxy(nullptr);
 
 		Chaos::FPBDRigidsSolver* Solver = FChaosSolversModule::GetModule()->CreateSolver(true);
-#if CHAOS_PARTICLEHANDLE_TODO
 		Solver->RegisterObject(PhysObject);
 		Solver->RegisterObject(FieldObject);
-#endif
 		Solver->SetHasFloor(false);
 		Solver->SetEnabled(true);
 		PhysObject->ActivateBodies();
 
+		Solver->AddDirtyProxy(PhysObject);
+		Solver->AddDirtyProxy(FieldObject);
+		Solver->PushPhysicsState(Module->GetDispatcher());
+
 		TManagedArray<FTransform>& Transform = DynamicCollection->Transform;
 		for (int Frame = 0; Frame < 10; Frame++)
 		{
@@ -495,6 +524,8 @@ namespace GeometryCollectionExample
 	template<class T>
 	void RigidBodies_Field_LinearVelocity()
 	{
+		FChaosSolversModule* Module = FChaosSolversModule::GetModule();
+
 		TUniquePtr<Chaos::FChaosPhysicsMaterial> PhysicalMaterial = nullptr;
 		TSharedPtr<FGeometryCollection> RestCollection = nullptr;
 		TSharedPtr<FGeometryDynamicCollection> DynamicCollection = nullptr;
@@ -525,14 +556,16 @@ namespace GeometryCollectionExample
 		FFieldSystemPhysicsProxy* FieldObject = new FFieldSystemPhysicsProxy(nullptr);
 
 		Chaos::FPBDRigidsSolver* Solver = FChaosSolversModule::GetModule()->CreateSolver(true);
-#if CHAOS_PARTICLEHANDLE_TODO
 		Solver->RegisterObject(PhysObject);
 		Solver->RegisterObject(FieldObject);
-#endif
 		Solver->SetHasFloor(false);
 		Solver->SetEnabled(true);
 		PhysObject->ActivateBodies();
 
+		Solver->AddDirtyProxy(PhysObject);
+		Solver->AddDirtyProxy(FieldObject);
+		Solver->PushPhysicsState(Module->GetDispatcher());
+
 		FName TargetName = GetFieldPhysicsName(EFieldPhysicsType::Field_LinearVelocity);
 		FieldObject->BufferCommand(Solver, { TargetName, VectorField->NewCopy() });
 		Solver->AdvanceSolverBy(1 / 24.);
@@ -568,6 +601,8 @@ namespace GeometryCollectionExample
 	template<class T>
 	void RigidBodies_Field_CollisionGroup()
 	{
+		FChaosSolversModule* Module = FChaosSolversModule::GetModule();
+
 		TUniquePtr<Chaos::FChaosPhysicsMaterial> PhysicalMaterial = nullptr;
 		TSharedPtr<FGeometryCollection> RestCollection = nullptr;
 		TSharedPtr<FGeometryDynamicCollection> DynamicCollection = nullptr;
@@ -603,13 +638,13 @@ namespace GeometryCollectionExample
 		FGeometryCollectionPhysicsProxy* PhysObject = RigidBodySetup(PhysicalMaterial, RestCollection, DynamicCollection);
 
 		Chaos::FPBDRigidsSolver* Solver = FChaosSolversModule::GetModule()->CreateSolver(true);
-#if CHAOS_PARTICLEHANDLE_TODO
 		Solver->RegisterObject(PhysObject);
-#endif
 		Solver->SetHasFloor(true);
 		Solver->SetEnabled(true);
 		PhysObject->ActivateBodies();
 
+		Solver->AddDirtyProxy(PhysObject);
+		Solver->PushPhysicsState(Module->GetDispatcher());
 
 #if TODO_REIMPLEMENT_GET_RIGID_PARTICLES
 		for (int Frame = 0; Frame < 60; Frame++)
@@ -652,6 +687,8 @@ namespace GeometryCollectionExample
 	template<class T>
 	void RigidBodies_Field_ClusterBreak_StrainModel_Test1()
 	{
+		FChaosSolversModule* Module = FChaosSolversModule::GetModule();
+
 		TUniquePtr<Chaos::FChaosPhysicsMaterial> PhysicalMaterial = nullptr;
 		TSharedPtr<FGeometryCollection> RestCollection = CreateClusteredBody_TwoByTwo_ThreeTransform(FVector(0));
 		TSharedPtr<FGeometryDynamicCollection> DynamicCollection = nullptr;
@@ -688,14 +725,16 @@ namespace GeometryCollectionExample
 #if TODO_REIMPLEMENT_GET_RIGID_PARTICLES
 		const Chaos::TPBDRigidParticles<float, 3>& Particles = Solver->GetRigidParticles();
 
-#if CHAOS_PARTICLEHANDLE_TODO
 		Solver->RegisterObject(PhysObject);
 		Solver->RegisterObject(FieldObject);
-#endif
 		Solver->SetHasFloor(false);
 		Solver->SetEnabled(true);
 		PhysObject->ActivateBodies();
 
+		Solver->AddDirtyProxy(PhysObject);
+		Solver->AddDirtyProxy(FieldObject);
+		Solver->PushPhysicsState(Module->GetDispatcher());
+
 		Solver->AdvanceSolverBy(1 / 24.);
 		FinalizeSolver(*Solver);
 
@@ -745,6 +784,8 @@ namespace GeometryCollectionExample
 	template<class T>
 	void RigidBodies_Field_ClusterBreak_StrainModel_Test2()
 	{
+		FChaosSolversModule* Module = FChaosSolversModule::GetModule();
+
 		TUniquePtr<Chaos::FChaosPhysicsMaterial> PhysicalMaterial = nullptr;
 		TSharedPtr<FGeometryCollection> RestCollection = CreateClusteredBody_ThreeByTwo_ThreeTransform(FVector(0));
 		TSharedPtr<FGeometryDynamicCollection> DynamicCollection = nullptr;
@@ -781,14 +822,16 @@ namespace GeometryCollectionExample
 		const Chaos::TPBDRigidParticles<float, 3>& Particles = Solver->GetRigidParticles();
 #endif
 
-#if CHAOS_PARTICLEHANDLE_TODO
 		Solver->RegisterObject(PhysObject);
 		Solver->RegisterObject(FieldObject);
-#endif
 		Solver->SetHasFloor(false);
 		Solver->SetEnabled(true);
 		PhysObject->ActivateBodies();
 
+		Solver->AddDirtyProxy(PhysObject);
+		Solver->AddDirtyProxy(FieldObject);
+		Solver->PushPhysicsState(Module->GetDispatcher());
+
 		Solver->AdvanceSolverBy(1 / 24.);
 		FinalizeSolver(*Solver);
 
@@ -848,6 +891,8 @@ namespace GeometryCollectionExample
 	template<class T>
 	void RigidBodies_Field_ClusterBreak_StrainModel_Test3()
 	{
+		FChaosSolversModule* Module = FChaosSolversModule::GetModule();
+
 		TUniquePtr<Chaos::FChaosPhysicsMaterial> PhysicalMaterial = nullptr;
 		TSharedPtr<FGeometryCollection> RestCollection = CreateClusteredBody_ThreeByTwo_ThreeTransform(FVector(0));
 		TSharedPtr<FGeometryDynamicCollection> DynamicCollection = nullptr;
@@ -883,14 +928,16 @@ namespace GeometryCollectionExample
 #if TODO_REIMPLEMENT_GET_RIGID_PARTICLES
 		const Chaos::TPBDRigidParticles<float, 3>& Particles = Solver->GetRigidParticles();
 #endif
-#if CHAOS_PARTICLEHANDLE_TODO
 		Solver->RegisterObject(PhysObject);
 		Solver->RegisterObject(FieldObject);
-#endif
 		Solver->SetHasFloor(false);
 		Solver->SetEnabled(true);
 		PhysObject->ActivateBodies();
 
+		Solver->AddDirtyProxy(PhysObject);
+		Solver->AddDirtyProxy(FieldObject);
+		Solver->PushPhysicsState(Module->GetDispatcher());
+
 		Solver->AdvanceSolverBy(1 / 24.);
 		FinalizeSolver(*Solver);
 
@@ -953,6 +1000,8 @@ namespace GeometryCollectionExample
 	template<class T>
 	void RigidBodies_Field_ClusterBreak_StrainModel_Test4()
 	{
+		FChaosSolversModule* Module = FChaosSolversModule::GetModule();
+
 		TUniquePtr<Chaos::FChaosPhysicsMaterial> PhysicalMaterial = nullptr;
 		TSharedPtr<FGeometryCollection> RestCollection = CreateClusteredBody_TwoByTwo_ThreeTransform(FVector(0));
 		TSharedPtr<FGeometryDynamicCollection> DynamicCollection = nullptr;
@@ -988,14 +1037,16 @@ namespace GeometryCollectionExample
 #if TODO_REIMPLEMENT_GET_RIGID_PARTICLES
 		const Chaos::TPBDRigidParticles<float, 3>& Particles = Solver->GetRigidParticles();
 #endif
-#if CHAOS_PARTICLEHANDLE_TODO
 		Solver->RegisterObject(PhysObject);
 		Solver->RegisterObject(FieldObject);
-#endif
 		Solver->SetHasFloor(false);
 		Solver->SetEnabled(true);
 		PhysObject->ActivateBodies();
 
+		Solver->AddDirtyProxy(PhysObject);
+		Solver->AddDirtyProxy(FieldObject);
+		Solver->PushPhysicsState(Module->GetDispatcher());
+
 		Solver->AdvanceSolverBy(1 / 24.);
 		FinalizeSolver(*Solver);
 
diff --git a/Engine/Plugins/Experimental/GeometryCollectionPlugin/Source/GeometryCollectionExampleCore/Public/GeometryCollection/GeometryCollectionExampleFields.h b/Engine/Plugins/Experimental/GeometryCollectionPlugin/Source/GeometryCollectionExampleCore/Public/GeometryCollection/GeometryCollectionExampleFields.h
--- a/Engine/Plugins/Experimental/GeometryCollectionPlugin/Source/GeometryCollectionExampleCore/Public/GeometryCollection/GeometryCollectionExampleFields.h
+++ b/Engine/Plugins/Experimental/GeometryCollectionPlugin/Source/GeometryCollectionExampleCore/Public/GeometryCollection/GeometryCollectionExampleFields.h
@@ -5,57 +5,21 @@
 
 namespace GeometryCollectionExample
 {
-	
-	template<class T>
 	void GEOMETRYCOLLECTIONEXAMPLECORE_API Fields_NoiseSample();
-
-	template<class T>
 	void GEOMETRYCOLLECTIONEXAMPLECORE_API Fields_RadialIntMask();
-
-	template<class T>
 	void GEOMETRYCOLLECTIONEXAMPLECORE_API Fields_RadialFalloff();
-
-	template<class T>
 	void GEOMETRYCOLLECTIONEXAMPLECORE_API Fields_PlaneFalloff();
-
-	template<class T>
 	void GEOMETRYCOLLECTIONEXAMPLECORE_API Fields_UniformVector();
-
-	template<class T>
 	void GEOMETRYCOLLECTIONEXAMPLECORE_API Fields_RaidalVector();
-
-	template<class T>
 	void GEOMETRYCOLLECTIONEXAMPLECORE_API Fields_SumVectorFullMult();
-
-	template<class T>
 	void GEOMETRYCOLLECTIONEXAMPLECORE_API Fields_SumVectorFullDiv();
-
-	template<class T>
 	void GEOMETRYCOLLECTIONEXAMPLECORE_API Fields_SumVectorFullAdd();
-
-	template<class T>
 	void GEOMETRYCOLLECTIONEXAMPLECORE_API Fields_SumVectorFullSub();
-
-	template<class T>
 	void GEOMETRYCOLLECTIONEXAMPLECORE_API Fields_SumVectorLeftSide();
-
-	template<class T>
 	void GEOMETRYCOLLECTIONEXAMPLECORE_API Fields_SumVectorRightSide();
-
-	template<class T>
 	void GEOMETRYCOLLECTIONEXAMPLECORE_API Fields_SumScalar();
-
-	template<class T>
 	void GEOMETRYCOLLECTIONEXAMPLECORE_API Fields_SumScalarRightSide();
-
-	template<class T>
 	void GEOMETRYCOLLECTIONEXAMPLECORE_API Fields_SumScalarLeftSide();
-
-	template<class T>
 	void GEOMETRYCOLLECTIONEXAMPLECORE_API Fields_Culling();
-
-	template<class T>
 	void GEOMETRYCOLLECTIONEXAMPLECORE_API Fields_SerializeAPI();
-
-
 }
diff --git a/Engine/Source/Runtime/Engine/Private/PhysicsEngine/Experimental/PhysScene_Chaos.cpp b/Engine/Source/Runtime/Engine/Private/PhysicsEngine/Experimental/PhysScene_Chaos.cpp
--- a/Engine/Source/Runtime/Engine/Private/PhysicsEngine/Experimental/PhysScene_Chaos.cpp
+++ b/Engine/Source/Runtime/Engine/Private/PhysicsEngine/Experimental/PhysScene_Chaos.cpp
@@ -546,37 +546,30 @@ void FPhysScene_Chaos::AddObject(UPrimitiveComponent* Component, FFieldSystemPhy
 {
 	AddToComponentMaps(Component, InObject);
 
-	ensure(false);
-#if 0
 	Chaos::FPhysicsSolver* CurrSceneSolver = GetSolver();
 
-	check(IsInGameThread() && CurrSceneSolver);
-	//const bool bDedicatedThread = ChaosModule->IsPersistentTaskEnabled();
-
 	InObject->SetSolver(CurrSceneSolver);
-	
-	Chaos::IDispatcher* Dispatcher = GetDispatcher();
+	InObject->Initialize();
 
-	const TArray<Chaos::FPhysicsSolver*>& Solvers = ChaosModule->GetSolvers();
-	for(Chaos::FPhysicsSolver* Solver : Solvers)
+	if (Chaos::IDispatcher* Dispatcher = GetDispatcher())
 	{
-		if(true || Solver->HasActiveParticles())
+		for (Chaos::FPhysicsSolver* Solver : ChaosModule->GetSolvers())
 		{
-			Solver->RegisterObject(InObject);
-
-			if(/*bDedicatedThread && */Dispatcher)
+			if (true || Solver->HasActiveParticles())
 			{
-				// Pass the proxy off to the physics thread
-				Dispatcher->EnqueueCommand([InObject, InSolver = Solver](Chaos::FPersistentPhysicsTask* PhysThread)
+				Solver->RegisterObject(InObject);
+
+				if (/*bDedicatedThread && */Dispatcher)
 				{
-#if CHAOS_PARTICLEHANDLE_TODO
-					InSolver->RegisterObject(InObject);
-#endif
-				});
+					// Pass the proxy off to the physics thread
+					Dispatcher->EnqueueCommandImmediate([InObject, InSolver = Solver](Chaos::FPersistentPhysicsTask* PhysThread)
+					{
+						InSolver->RegisterObject(InObject);
+					});
+				}
 			}
 		}
 	}
-#endif
 }
 
 void FPhysScene_Chaos::RemoveActorFromAccelerationStructure(FPhysicsActorHandle& Actor)
@@ -688,48 +681,32 @@ void FPhysScene_Chaos::RemoveObject(FGeometryCollectionPhysicsProxy* InObject)
 
 void FPhysScene_Chaos::RemoveObject(FFieldSystemPhysicsProxy* InObject)
 {
-	ensure(false);
-#if 0
-	if(!InObject)
+	Chaos::FPhysicsSolver* CurrSceneSolver = InObject->GetSolver();
+	if(CurrSceneSolver && !CurrSceneSolver->UnregisterObject(InObject))
 	{
-		return;
+		UE_LOG(LogChaos, Warning, TEXT("Attempted to remove an object that wasn't found in its solver's gamethread storage - it's likely the solver has been mistakenly changed."));
 	}
+	RemoveFromComponentMaps(InObject);
 
-	Chaos::FPhysicsSolver* CurrSceneSolver = GetSolver();
-
-	CurrSceneSolver->UnregisterObject(InObject);
-
-	check(IsInGameThread());
-
-	Chaos::IDispatcher* Dispatcher = GetDispatcher();
-	const bool bDedicatedThread = Dispatcher->GetMode() == Chaos::EThreadingMode::DedicatedThread;
-
-	const TArray<Chaos::FPhysicsSolver*>& Solvers = ChaosModule->GetSolvers();
-	for(Chaos::FPhysicsSolver* Solver : Solvers)
+	if (Chaos::IDispatcher* Dispatcher = GetDispatcher())
 	{
-		if(true || Solver->HasActiveParticles())
+		for (Chaos::FPhysicsSolver* Solver : ChaosModule->GetSolvers())
 		{
-			// @todo(question) : Is the RegisterObject a bug? Seems like it should be removing. [brice->benn]
-			Solver->RegisterObject(InObject);
-
-			if(/*bDedicatedThread && */Dispatcher)
+			if (true || Solver->HasActiveParticles())
 			{
-				// Pass the proxy off to the physics thread
-				Dispatcher->EnqueueCommand([InObject, InSolver = Solver](Chaos::FPersistentPhysicsTask* PhysThread)
+				Solver->RegisterObject(InObject);
+
+				if (/*bDedicatedThread && */Dispatcher)
 				{
-					InSolver->UnregisterObject(InObject);
-				});
+					// Pass the proxy off to the physics thread
+					Dispatcher->EnqueueCommandImmediate([InObject, InSolver = Solver](Chaos::FPersistentPhysicsTask* PhysThread)
+					{
+						InSolver->UnregisterObject(InObject);
+					});
+				}
 			}
 		}
 	}
-
-	RemoveFromComponentMaps(InObject);
-
-	if(!bDedicatedThread)
-	{
-		delete InObject;
-	}
-#endif
 }
 
 #if XGE_FIXED
diff --git a/Engine/Source/Runtime/Experimental/Chaos/Public/Chaos/PBDRigidsSOAs.h b/Engine/Source/Runtime/Experimental/Chaos/Public/Chaos/PBDRigidsSOAs.h
--- a/Engine/Source/Runtime/Experimental/Chaos/Public/Chaos/PBDRigidsSOAs.h
+++ b/Engine/Source/Runtime/Experimental/Chaos/Public/Chaos/PBDRigidsSOAs.h
@@ -30,8 +30,8 @@ public:
 		ClusteredParticles = MakeUnique< TPBDRigidClusteredParticles<T, d>>();
 		ClusteredParticles->RemoveParticleBehavior() = ERemoveParticleBehavior::Remove;	//clustered particles maintain relative ordering
 
-		GeometryCollectionDisabledParticles = MakeUnique<TPBDGeometryCollectionParticles<T, d>>();
 		GeometryCollectionParticles = MakeUnique<TPBDGeometryCollectionParticles<T, d>>();
+		GeometryCollectionParticles->RemoveParticleBehavior() = ERemoveParticleBehavior::Remove;	//clustered particles maintain relative ordering
 		
 		UpdateViews();
 	}
@@ -70,13 +70,13 @@ public:
 	TArray<TPBDGeometryCollectionParticleHandle<T, d>*> CreateGeometryCollectionParticles(int32 NumParticles, const TPBDRigidParticleParameters<T, d>& Params = TPBDRigidParticleParameters<T, d>())
 	{
 		TArray<TPBDGeometryCollectionParticleHandle<T, d>*> Results = CreateParticlesHelper<TPBDGeometryCollectionParticleHandle<T, d>>(
-			NumParticles, Params.bDisabled ? GeometryCollectionDisabledParticles : GeometryCollectionParticles, Params);
+			NumParticles, GeometryCollectionParticles, Params);
 		//TArray<TPBDRigidParticleHandle<T, d>*>& RigidHandles = (TArray<TPBDRigidParticleHandle<T, d>*>*)&Results;//*static_cast<TArray<TPBDRigidParticleHandle<T, d>*>*>(&Results);
 		if (!Params.bStartSleeping)
 		{
-			InsertToMapAndArray(Results, ActiveParticlesToIndex, ActiveParticlesArray);
 			InsertToMapAndArray(Results, ActiveGeometryCollectionToIndex, ActiveGeometryCollectionArray);
 		}
+		UpdateGeometryCollectionViews();
 		UpdateViews();
 		return Results;
 	}
@@ -288,7 +288,6 @@ public:
 		ensure(ClusteredParticles->Size() == 0);	//not supported yet
 		//Ar << ClusteredParticles;
 		Ar << GeometryCollectionParticles;
-		Ar << GeometryCollectionDisabledParticles;
 
 		SerializeMapAndArray(Ar, ActiveParticlesToIndex, ActiveParticlesArray);
 		//SerializeMapAndArray(Ar, ActiveClusteredToIndex, ActiveClusteredArray);
@@ -327,6 +326,115 @@ public:
 	const TParticleView<TPBDGeometryCollectionParticles<T, d>>& GetActiveGeometryCollectionParticlesView() const { return ActiveGeometryCollectionParticlesView; }
 	TParticleView<TPBDGeometryCollectionParticles<T, d>>& GetActiveGeometryCollectionParticlesView() { return ActiveGeometryCollectionParticlesView; }
 
+	/**
+	 * Update which particle arrays geometry collection particles are in based on 
+	 * their object state (static, kinematic, dynamic, sleeping) and their disabled 
+	 * state.
+	 */
+	void UpdateGeometryCollectionViews()
+	{
+		int32 AIdx = 0, SIdx = 0, KIdx = 0, DIdx = 0;
+
+		for (TPBDGeometryCollectionParticleHandle<T, d>* Handle : ActiveGeometryCollectionArray)
+		{
+			// If the particle is disabled we treat it as static, but for no reason 
+			// other than immediate convenience.
+			const Chaos::EObjectStateType State = Handle->Disabled() ? Chaos::EObjectStateType::Static : Handle->ObjectState();
+
+			switch (State)
+			{
+			case Chaos::EObjectStateType::Static:
+				SIdx++;
+				AIdx += (int32)(!Handle->Disabled());
+				break;
+
+			case Chaos::EObjectStateType::Kinematic:
+				KIdx++;
+				AIdx += (int32)(!Handle->Disabled());
+				break;
+
+			case Chaos::EObjectStateType::Sleeping: // Sleeping is a modified dynamic state
+				DIdx++;
+				break;
+
+			case Chaos::EObjectStateType::Dynamic:
+				DIdx++;
+				AIdx += (int32)(!Handle->Disabled());
+				break;
+
+			default:
+				break;
+			};
+		}
+
+		bool Changed = 
+			ActiveGeometryCollectionArray.Num() != AIdx ||
+			StaticGeometryCollectionArray.Num() != SIdx || 
+			KinematicGeometryCollectionArray.Num() != KIdx || 
+			DynamicGeometryCollectionArray.Num() != DIdx;
+
+		if (Changed)
+		{
+			ActiveGeometryCollectionArray.SetNumUninitialized(AIdx);
+			StaticGeometryCollectionArray.SetNumUninitialized(SIdx);
+			KinematicGeometryCollectionArray.SetNumUninitialized(KIdx);
+			DynamicGeometryCollectionArray.SetNumUninitialized(DIdx);
+		}
+
+		AIdx = SIdx = KIdx = DIdx = 0;
+
+		for (TPBDGeometryCollectionParticleHandle<T, d>* Handle : ActiveGeometryCollectionArray)
+		{
+			const Chaos::EObjectStateType State = 
+				Handle->Disabled() ? Chaos::EObjectStateType::Static : Handle->ObjectState();
+			switch (State)
+			{
+			case Chaos::EObjectStateType::Static:
+				Changed |= StaticGeometryCollectionArray[SIdx] != Handle;
+				StaticGeometryCollectionArray[SIdx++] = Handle;
+				if (!Handle->Disabled())
+				{
+					Changed |= ActiveGeometryCollectionArray[AIdx] != Handle;
+					ActiveGeometryCollectionArray[AIdx++] = Handle;
+				}
+				break;
+
+			case Chaos::EObjectStateType::Kinematic:
+				Changed |= KinematicGeometryCollectionArray[KIdx] != Handle;
+				KinematicGeometryCollectionArray[KIdx++] = Handle;
+				if (!Handle->Disabled())
+				{
+					Changed |= ActiveGeometryCollectionArray[AIdx] != Handle;
+					ActiveGeometryCollectionArray[AIdx++] = Handle;
+				}
+				break;
+
+			case Chaos::EObjectStateType::Sleeping: // Sleeping is a modified dynamic state
+				Changed |= DynamicGeometryCollectionArray[DIdx] != Handle;
+				DynamicGeometryCollectionArray[DIdx++] = Handle;
+				break;
+
+			case Chaos::EObjectStateType::Dynamic:
+				Changed |= DynamicGeometryCollectionArray[DIdx] != Handle;
+				DynamicGeometryCollectionArray[DIdx++] = Handle;
+				if (!Handle->Disabled())
+				{
+					Changed |= ActiveGeometryCollectionArray[AIdx] != Handle;
+					ActiveGeometryCollectionArray[AIdx++] = Handle;
+				}
+				break;
+
+			default:
+				break;
+			};
+		}
+
+		if(Changed)
+		{
+			UpdateViews();
+		}
+	}
+
 	//TEMP: only needed while clustering code continues to use direct indices
 	const auto& GetActiveClusteredArray() const { return ActiveClusteredArray; }
 	const auto& GetNonDisabledClusteredArray() const { return NonDisabledClusteredArray; }
@@ -418,15 +526,34 @@ private:
 	{
 		//build various views. Group SOA types together for better branch prediction
 		{
-			TArray<TSOAView<TGeometryParticles<T, d>>> TmpArray = { StaticParticles.Get(), KinematicParticles.Get(), DynamicParticles.Get(), DynamicKinematicParticles.Get(), {&NonDisabledClusteredArray}, GeometryCollectionParticles.Get() };
+			TArray<TSOAView<TGeometryParticles<T, d>>> TmpArray = 
+			{ 
+				StaticParticles.Get(), 
+				KinematicParticles.Get(), 
+				DynamicParticles.Get(),
+				DynamicKinematicParticles.Get(),
+				{&NonDisabledClusteredArray}, 
+				{&StaticGeometryCollectionArray},
+				{&KinematicGeometryCollectionArray},
+				{&DynamicGeometryCollectionArray}
+			};
 			NonDisabledView = MakeParticleView(MoveTemp(TmpArray));
 		}
 		{
-			TArray<TSOAView<TPBDRigidParticles<T, d>>> TmpArray = { DynamicParticles.Get(), {&NonDisabledClusteredArray}, GeometryCollectionParticles.Get() };
+			TArray<TSOAView<TPBDRigidParticles<T, d>>> TmpArray = 
+			{ 
+				DynamicParticles.Get(), 
+				{&NonDisabledClusteredArray}, 
+				{&DynamicGeometryCollectionArray}
+			};
 			NonDisabledDynamicView = MakeParticleView(MoveTemp(TmpArray));
 		}
 		{
-			TArray<TSOAView<TPBDRigidParticles<T, d>>> TmpArray = { {&ActiveParticlesArray} };
+			TArray<TSOAView<TPBDRigidParticles<T, d>>> TmpArray = 
+			{ 
+				{&ActiveParticlesArray},
+				{&ActiveGeometryCollectionArray}
+			};
 			ActiveParticlesView = MakeParticleView(MoveTemp(TmpArray));
 		}
 		{
@@ -435,11 +562,16 @@ private:
 			AllParticlesView = MakeParticleView(MoveTemp(TmpArray));
 		}
 		{
-			TArray<TSOAView<TKinematicGeometryParticles<T, d>>> TmpArray = { KinematicParticles.Get(), DynamicKinematicParticles.Get() };
+			TArray<TSOAView<TKinematicGeometryParticles<T, d>>> TmpArray = 
+			{ 
+				KinematicParticles.Get(),
+				DynamicKinematicParticles.Get(),
+				{&KinematicGeometryCollectionArray}
+			};
 			ActiveKinematicParticlesView = MakeParticleView(MoveTemp(TmpArray));
 		}
 		{
-			TArray<TSOAView<TPBDGeometryCollectionParticles<T, d>>> TmpArray = { GeometryCollectionParticles.Get()};
+			TArray<TSOAView<TPBDGeometryCollectionParticles<T, d>>> TmpArray = { {&ActiveGeometryCollectionArray} };
 			ActiveGeometryCollectionParticlesView = MakeParticleView(MoveTemp(TmpArray));
 		}
 	}
@@ -458,10 +590,14 @@ private:
 	TUniquePtr<TPBDRigidClusteredParticles<T, d>> ClusteredParticles;
 
 	TUniquePtr<TPBDGeometryCollectionParticles<T, d>> GeometryCollectionParticles;
-	TUniquePtr<TPBDGeometryCollectionParticles<T, d>> GeometryCollectionDisabledParticles;
+
 	TMap<TPBDGeometryCollectionParticleHandle<T, d>*, int32> ActiveGeometryCollectionToIndex;
 	TArray<TPBDGeometryCollectionParticleHandle<T, d>*> ActiveGeometryCollectionArray;
 
+	TArray<TPBDGeometryCollectionParticleHandle<T, d>*> StaticGeometryCollectionArray;
+	TArray<TPBDGeometryCollectionParticleHandle<T, d>*> KinematicGeometryCollectionArray;
+	TArray<TPBDGeometryCollectionParticleHandle<T, d>*> DynamicGeometryCollectionArray;
+
 	//Utility structures for maintaining an Active particles view
 	TMap<TPBDRigidParticleHandle<T, d>*, int32> ActiveParticlesToIndex;
 	TArray<TPBDRigidParticleHandle<T, d>*> ActiveParticlesArray;
diff --git a/Engine/Source/Runtime/Experimental/Chaos/Public/Chaos/Utilities.h b/Engine/Source/Runtime/Experimental/Chaos/Public/Chaos/Utilities.h
--- a/Engine/Source/Runtime/Experimental/Chaos/Public/Chaos/Utilities.h
+++ b/Engine/Source/Runtime/Experimental/Chaos/Public/Chaos/Utilities.h
@@ -16,7 +16,7 @@ namespace Chaos
 {
 	namespace Utilities
 	{
-		// Take the factorial of \p Num, which should be of integral type.
+		//! Take the factorial of \p Num, which should be of integral type.
 		template<class TINT = uint64>
 		TINT Factorial(TINT Num)
 		{
@@ -29,7 +29,7 @@ namespace Chaos
 			return Result;
 		}
 
-		// Number of ways to choose of \p R elements from a set of size \p N, with no repetitions.
+		//! Number of ways to choose of \p R elements from a set of size \p N, with no repetitions.
 		template<class TINT = uint64>
 		TINT NChooseR(const TINT N, const TINT R)
 		{
@@ -37,7 +37,7 @@ namespace Chaos
 			return Factorial(N) / (Factorial(R) * Factorial(N - R));
 		}
 
-		// Number of ways to choose of \p R elements from a set of size \p N, with repetitions.
+		//! Number of ways to choose of \p R elements from a set of size \p N, with repetitions.
 		template<class TINT = uint64>
 		TINT NPermuteR(const TINT N, const TINT R)
 		{
@@ -45,8 +45,9 @@ namespace Chaos
 			return Factorial(N) / Factorial(N - R);
 		}
 
-		template<class T>
-		void GetMinAvgMax(const TArray<T>& Values, T& MinV, double& AvgV, T& MaxV)
+		//! Compute the minimum, average, and maximum values of \p Values.
+		template<class T, class TARRAY=TArray<T>>
+		void GetMinAvgMax(const TARRAY& Values, T& MinV, double& AvgV, T& MaxV)
 		{
 			MinV = TNumericLimits<T>::Max();
 			MaxV = TNumericLimits<T>::Lowest();
@@ -67,12 +68,75 @@ namespace Chaos
 			}
 		}
 
+		//! Compute the average value.
+		template<class T, class TARRAY=TArray<T>>
+		T GetAverage(const TARRAY& Values)
+		{
+			double AvgV = 0.0;
+			for (const T& V : Values)
+			{
+				AvgV += V;
+			}
+			if (Values.Num())
+			{
+				AvgV /= Values.Num();
+			}
+			return static_cast<T>(AvgV);
+		}
+
+		//! Compute the variance of \p Values, given the average value of \p Avg.
+		template<class T, class TARRAY=TArray<T>>
+		T GetVariance(const TARRAY& Values, const T Avg)
+		{
+			double Variance = 0.0;
+			for (const T& V : Values)
+			{
+				const T Deviation = V - Avg;
+				Variance += Deviation * Deviation;
+			}
+			if (Values.Num())
+			{
+				Variance /= Values.Num();
+			}
+			return Variance;
+		}
+
+		//! Compute the variance of \p Values (computes their average on the fly).
+		template<class T, class TARRAY=TArray<T>>
+		T GetVariance(const TARRAY& Values)
+		{
+			return GetVariance(Values, GetAverage(Values));
+		}
+
+		//! Compute the standard deviation of \p Values, given the average value of \p Avg.
+		template<class T, class TARRAY=TArray<T>>
+		T GetStandardDeviation(const TARRAY& Values, const T Avg)
+		{
+			const T Variance = GetVariance(Values, Avg);
+			return FMath::Sqrt(Variance);
+		}
+
+		//! Compute the standard deviation of \p Values (computes their average on the fly).
+		template<class T, class TARRAY=TArray<T>>
+		T GetStandardDeviation(const TARRAY& Values)
+		{
+			const T Variance = GetVariance(Values);
+			return FMath::Sqrt(Variance);
+		}
+
+		//! Compute the standard deviation from \p Variance.
+		template<class T>
+		T GetStandardDeviation(const T Variance)
+		{
+			return FMath::Sqrt(Variance);
+		}
+
 		inline static PMatrix<float, 3, 3> CrossProductMatrix(const TVector<float, 3>& V)
 		{
 			return PMatrix<float, 3, 3>(
-				0, -V.Z, V.Y,
-				V.Z, 0, -V.X,
-				-V.Y, V.X, 0);
+			    0, -V.Z, V.Y,
+			    V.Z, 0, -V.X,
+			    -V.Y, V.X, 0);
 		}
 
 		/**
@@ -88,18 +152,19 @@ namespace Chaos
 			// We want L.R (FMatrix operator* actually calculates R.(L)T; i.e., Right is on the left, and the Left is transposed on the right.)
 			// NOTE: PMatrix constructor takes values in column order
 			return PMatrix<float, 3, 3>(
-				L.M[0][0] * R.M[0][0] + L.M[0][1] * R.M[1][0] + L.M[0][2] * R.M[2][0],	// x00
-				L.M[1][0] * R.M[0][0] + L.M[1][1] * R.M[1][0] + L.M[1][2] * R.M[2][0],	// x10
-				L.M[2][0] * R.M[0][0] + L.M[2][1] * R.M[1][0] + L.M[2][2] * R.M[2][0],	// x20
+			    L.M[0][0] * R.M[0][0] + L.M[0][1] * R.M[1][0] + L.M[0][2] * R.M[2][0], // x00
+			    L.M[1][0] * R.M[0][0] + L.M[1][1] * R.M[1][0] + L.M[1][2] * R.M[2][0], // x10
+			    L.M[2][0] * R.M[0][0] + L.M[2][1] * R.M[1][0] + L.M[2][2] * R.M[2][0], // x20
 
-				L.M[0][0] * R.M[0][1] + L.M[0][1] * R.M[1][1] + L.M[0][2] * R.M[2][1],	// x01
-				L.M[1][0] * R.M[0][1] + L.M[1][1] * R.M[1][1] + L.M[1][2] * R.M[2][1],	// x11
-				L.M[2][0] * R.M[0][1] + L.M[2][1] * R.M[1][1] + L.M[2][2] * R.M[2][1],	// x21
+			    L.M[0][0] * R.M[0][1] + L.M[0][1] * R.M[1][1] + L.M[0][2] * R.M[2][1], // x01
+			    L.M[1][0] * R.M[0][1] + L.M[1][1] * R.M[1][1] + L.M[1][2] * R.M[2][1], // x11
+			    L.M[2][0] * R.M[0][1] + L.M[2][1] * R.M[1][1] + L.M[2][2] * R.M[2][1], // x21
 
-				L.M[0][0] * R.M[0][2] + L.M[0][1] * R.M[1][2] + L.M[0][2] * R.M[2][2],	// x02
-				L.M[1][0] * R.M[0][2] + L.M[1][1] * R.M[1][2] + L.M[1][2] * R.M[2][2],	// x12
-				L.M[2][0] * R.M[0][2] + L.M[2][1] * R.M[1][2] + L.M[2][2] * R.M[2][2]	// x22
-				).GetTransposed();
+			    L.M[0][0] * R.M[0][2] + L.M[0][1] * R.M[1][2] + L.M[0][2] * R.M[2][2], // x02
+			    L.M[1][0] * R.M[0][2] + L.M[1][1] * R.M[1][2] + L.M[1][2] * R.M[2][2], // x12
+			    L.M[2][0] * R.M[0][2] + L.M[2][1] * R.M[1][2] + L.M[2][2] * R.M[2][2] // x22
+			    )
+			    .GetTransposed();
 		}
 
 		inline PMatrix<float, 3, 3> MultiplyAB(const PMatrix<float, 3, 3>& LIn, const PMatrix<float, 3, 3>& RIn)
@@ -127,10 +192,9 @@ namespace Chaos
 			PMatrix<float, 3, 3> L = LIn.GetTransposed();
 
 			return TVector<float, 3>(
-				L.M[0][0] * R.X + L.M[0][1] * R.Y + L.M[0][2] * R.Z,
-				L.M[1][0] * R.X + L.M[1][1] * R.Y + L.M[1][2] * R.Z,
-				L.M[2][0] * R.X + L.M[2][1] * R.Y + L.M[2][2] * R.Z
-				);
+			    L.M[0][0] * R.X + L.M[0][1] * R.Y + L.M[0][2] * R.Z,
+			    L.M[1][0] * R.X + L.M[1][1] * R.Y + L.M[1][2] * R.Z,
+			    L.M[2][0] * R.X + L.M[2][1] * R.Y + L.M[2][2] * R.Z);
 		}
 
 		/**
@@ -165,12 +229,12 @@ namespace Chaos
 			// Vx*M*VxT+Im
 			check(Im > FLT_MIN);
 			return PMatrix<T, 3, 3>(
-				-V[2] * (-V[2] * M.M[1][1] + V[1] * M.M[2][1]) + V[1] * (-V[2] * M.M[2][1] + V[1] * M.M[2][2]) + Im,
-				V[2] * (-V[2] * M.M[1][0] + V[1] * M.M[2][0]) - V[0] * (-V[2] * M.M[2][1] + V[1] * M.M[2][2]),
-				-V[1] * (-V[2] * M.M[1][0] + V[1] * M.M[2][0]) + V[0] * (-V[2] * M.M[1][1] + V[1] * M.M[2][1]),
-				V[2] * (V[2] * M.M[0][0] - V[0] * M.M[2][0]) - V[0] * (V[2] * M.M[2][0] - V[0] * M.M[2][2]) + Im,
-				-V[1] * (V[2] * M.M[0][0] - V[0] * M.M[2][0]) + V[0] * (V[2] * M.M[1][0] - V[0] * M.M[2][1]),
-				-V[1] * (-V[1] * M.M[0][0] + V[0] * M.M[1][0]) + V[0] * (-V[1] * M.M[1][0] + V[0] * M.M[1][1]) + Im);
+			    -V[2] * (-V[2] * M.M[1][1] + V[1] * M.M[2][1]) + V[1] * (-V[2] * M.M[2][1] + V[1] * M.M[2][2]) + Im,
+			    V[2] * (-V[2] * M.M[1][0] + V[1] * M.M[2][0]) - V[0] * (-V[2] * M.M[2][1] + V[1] * M.M[2][2]),
+			    -V[1] * (-V[2] * M.M[1][0] + V[1] * M.M[2][0]) + V[0] * (-V[2] * M.M[1][1] + V[1] * M.M[2][1]),
+			    V[2] * (V[2] * M.M[0][0] - V[0] * M.M[2][0]) - V[0] * (V[2] * M.M[2][0] - V[0] * M.M[2][2]) + Im,
+			    -V[1] * (V[2] * M.M[0][0] - V[0] * M.M[2][0]) + V[0] * (V[2] * M.M[1][0] - V[0] * M.M[2][1]),
+			    -V[1] * (-V[1] * M.M[0][0] + V[0] * M.M[1][0]) + V[0] * (-V[1] * M.M[1][0] + V[0] * M.M[1][1]) + Im);
 		}
 
 		/**
@@ -185,7 +249,7 @@ namespace Chaos
 
 			T Divisor = (InEndB[0] - InStartB[0]) * (InStartA[1] - InEndA[1]) - (InStartA[0] - InEndA[0]) * (InEndB[1] - InStartB[1]);
 
-			if(FMath::IsNearlyZero(Divisor))
+			if (FMath::IsNearlyZero(Divisor))
 			{
 				// The line segments are parallel and will never meet for any values of Ta/Tb
 				return false;
diff --git a/Engine/Source/Runtime/Experimental/ChaosSolvers/Private/PBDRigidsSolver.cpp b/Engine/Source/Runtime/Experimental/ChaosSolvers/Private/PBDRigidsSolver.cpp
--- a/Engine/Source/Runtime/Experimental/ChaosSolvers/Private/PBDRigidsSolver.cpp
+++ b/Engine/Source/Runtime/Experimental/ChaosSolvers/Private/PBDRigidsSolver.cpp
@@ -53,8 +53,17 @@ namespace Chaos
 			UE_LOG(LogPBDRigidsSolverSolver, Verbose, TEXT("AdvanceOneTimeStepTask::DoWork()"));
 
 			{
-				//SCOPE_CYCLE_COUNTER(STAT_UpdateParams);
-				//MSolver->ParameterUpdateCallback(MSolver->GetSolverTime());
+				SCOPE_CYCLE_COUNTER(STAT_UpdateParams);
+				Chaos::TPBDPositionConstraints<float, 3> PositionTarget; // Dummy for now
+				TMap<int32, int32> PositionTargetedParticles;
+				//TArray<FKinematicProxy> AnimatedPositions;
+				Chaos::TArrayCollectionArray<float> Strains;
+				for (FFieldSystemPhysicsProxy* Obj : MSolver->GetFieldSystemPhysicsProxies())
+				{
+					auto& GeomCollectionParticles = MSolver->GetEvolution()->GetParticles().GetGeometryCollectionParticles();
+					Obj->FieldParameterUpdateCallback(MSolver, GeomCollectionParticles, Strains, 
+						PositionTarget, PositionTargetedParticles, /*AnimatedPositions,*/ MSolver->GetSolverTime());
+				}
 			}
 
 			{
@@ -66,13 +75,25 @@ namespace Chaos
 				SCOPE_CYCLE_COUNTER(STAT_EvolutionAndKinematicUpdate);
 				while (MDeltaTime > MSolver->GetMaxDeltaTime())
 				{
+					Chaos::TArrayCollectionArray<FVector> Forces, Torques;
 					//MSolver->ForceUpdateCallback(MSolver->GetSolverTime());
+					for (FFieldSystemPhysicsProxy* Obj : MSolver->GetFieldSystemPhysicsProxies())
+					{
+						auto& GeomCollectionParticles = MSolver->GetEvolution()->GetParticles().GetGeometryCollectionParticles();
+						Obj->FieldForcesUpdateCallback(MSolver, GeomCollectionParticles, Forces, Torques, MSolver->GetSolverTime());
+					}
 					//MSolver->GetEvolution()->ReconcileIslands();
 					//MSolver->KinematicUpdateCallback(MSolver->GetMaxDeltaTime(), MSolver->GetSolverTime());
 					MSolver->GetEvolution()->AdvanceOneTimeStep(MSolver->GetMaxDeltaTime());
 					MDeltaTime -= MSolver->GetMaxDeltaTime();
 				}
 				//MSolver->ForceUpdateCallback(MSolver->GetSolverTime());
+				Chaos::TArrayCollectionArray<FVector> Forces, Torques;
+				for (FFieldSystemPhysicsProxy* Obj : MSolver->GetFieldSystemPhysicsProxies())
+				{
+					auto& GeomCollectionParticles = MSolver->GetEvolution()->GetParticles().GetGeometryCollectionParticles();
+					Obj->FieldForcesUpdateCallback(MSolver, GeomCollectionParticles, Forces, Torques, MSolver->GetSolverTime());
+				}
 				//MSolver->GetEvolution()->ReconcileIslands();
 				//MSolver->KinematicUpdateCallback(MDeltaTime, MSolver->GetSolverTime());
 				MSolver->GetEvolution()->AdvanceOneTimeStep(MDeltaTime);
@@ -337,6 +358,31 @@ namespace Chaos
 		return GeometryCollectionPhysicsProxies.Remove(InProxy) != 0;
 	}
 
+	void FPBDRigidsSolver::RegisterObject(FFieldSystemPhysicsProxy* InProxy)
+	{
+		UE_LOG(LogPBDRigidsSolverSolver, Verbose, TEXT("FPBDRigidsSolver::RegisterObject(FFieldSystemPhysicsProxy*)"));
+		FieldSystemPhysicsProxies.AddUnique(InProxy);
+		InProxy->SetSolver(this);
+		InProxy->Initialize();
+		Chaos::FParticleData* ProxyData = InProxy->NewData();
+
+		FChaosSolversModule::GetModule()->GetDispatcher()->EnqueueCommandImmediate(
+			this,
+			[this, InProxy, ProxyData](FPBDRigidsSolver* Solver)
+			{
+				UE_LOG(LogPBDRigidsSolverSolver, Verbose,
+					TEXT("FPBDRigidsSolver::RegisterObject(FFieldSystemPhysicsProxy*)"));
+				//InProxy->ActivateBodies();
+				InProxy->PushToPhysicsState(ProxyData);
+			});
+	}
+
+	bool FPBDRigidsSolver::UnregisterObject(FFieldSystemPhysicsProxy* InProxy)
+	{
+		InProxy->SetSolver(static_cast<FPBDRigidsSolver*>(nullptr));
+		return FieldSystemPhysicsProxies.Remove(InProxy) != 0;
+	}
+
 	bool FPBDRigidsSolver::IsSimulating() const
 	{
 		for (FGeometryParticlePhysicsProxy* Obj : GeometryParticlePhysicsProxies)
@@ -467,9 +513,27 @@ namespace Chaos
 		Proxy->ClearAccumulatedData();
 	}
 
+	void PushPhysicsStateExec(FPBDRigidsSolver* Solver, FFieldSystemPhysicsProxy* Proxy, Chaos::IDispatcher* Dispatcher)
+	{
+		Chaos::FParticleData* ProxyData = Proxy->NewData();
+		auto Cmd = [Proxy, Solver, ProxyData](Chaos::FPersistentPhysicsTask* PhysThread)
+		{
+			Proxy->PushToPhysicsState(nullptr);
+		};
+
+		if (Dispatcher)
+			Dispatcher->EnqueueCommandImmediate(Cmd);
+		else
+			Cmd(nullptr);
+		Solver->RemoveDirtyProxy(Proxy);
+		Proxy->ClearAccumulatedData();
+	}
+
 	void FPBDRigidsSolver::PushPhysicsState(IDispatcher* Dispatcher)
 	{
-		TArray< IPhysicsProxyBase *> DirtyProxiesArray = DirtyProxiesSet.Array();
+		if (DirtyProxiesSet.Num() == 0)
+			return;
+		TArray<IPhysicsProxyBase*> DirtyProxiesArray = DirtyProxiesSet.Array();
 		for (auto & Proxy : DirtyProxiesArray)
 		{
 			switch (Proxy->GetType())
@@ -479,7 +543,9 @@ namespace Chaos
 			case EPhysicsProxyType::GeometryCollectionType: // 2
 				PushPhysicsStateExec(this, static_cast<FGeometryCollectionPhysicsProxy*>(Proxy), Dispatcher);
 				break;
-			//case EPhysicsProxyType::FieldType: // 3
+			case EPhysicsProxyType::FieldType: // 3
+				PushPhysicsStateExec(this, static_cast<FFieldSystemPhysicsProxy*>(Proxy), Dispatcher);
+				break;
 			//case EPhysicsProxyType::SkeletalMeshType: // 4
 			case EPhysicsProxyType::SingleRigidParticleType: // 7
 				PushPhysicsStateExec(this, static_cast<FSingleParticlePhysicsProxy< Chaos::TPBDRigidParticle<float, 3> >*>(Proxy), Dispatcher);
diff --git a/Engine/Source/Runtime/Experimental/ChaosSolvers/Private/PhysicsProxy/FieldSystemPhysicsProxy.cpp b/Engine/Source/Runtime/Experimental/ChaosSolvers/Private/PhysicsProxy/FieldSystemPhysicsProxy.cpp
--- a/Engine/Source/Runtime/Experimental/ChaosSolvers/Private/PhysicsProxy/FieldSystemPhysicsProxy.cpp
+++ b/Engine/Source/Runtime/Experimental/ChaosSolvers/Private/PhysicsProxy/FieldSystemPhysicsProxy.cpp
@@ -20,12 +20,13 @@ void ResetIndicesArray(TArray<int32> & IndicesArray, int32 Size)
 	}
 }
 
+//==============================================================================
+// FFieldSystemPhysicsProxy
+//==============================================================================
 
 FFieldSystemPhysicsProxy::FFieldSystemPhysicsProxy(UObject* InOwner)
 	: Base(InOwner)
-{
-
-}
+{}
 
 FFieldSystemPhysicsProxy::~FFieldSystemPhysicsProxy()
 {
@@ -38,12 +39,22 @@ FFieldSystemPhysicsProxy::~FFieldSystemPhysicsProxy()
 	Commands.Reset();
 }
 
+void FFieldSystemPhysicsProxy::Initialize()
+{}
+
 bool FFieldSystemPhysicsProxy::IsSimulating() const
 {
 	return true; // #todo Actually start gating this?
 }
 
-void FFieldSystemPhysicsProxy::FieldParameterUpdateCallback(Chaos::FPhysicsSolver* InSolver, FParticlesType& Particles, Chaos::TArrayCollectionArray<float>& Strains, Chaos::TPBDPositionConstraints<float, 3>& PositionTarget, TMap<int32, int32>& PositionTargetedParticles, const TArray<FKinematicProxy>& AnimatedPosition, const float InTime)
+void FFieldSystemPhysicsProxy::FieldParameterUpdateCallback(
+	Chaos::FPhysicsSolver* InSolver, 
+	FParticlesType& Particles, 
+	Chaos::TArrayCollectionArray<float>& Strains, 
+	Chaos::TPBDPositionConstraints<float, 3>& PositionTarget, 
+	TMap<int32, int32>& PositionTargetedParticles, 
+	//const TArray<FKinematicProxy>& AnimatedPosition, 
+	const float InTime)
 {
 	using namespace Chaos;
 	SCOPE_CYCLE_COUNTER(STAT_ParamUpdateField_Object);
@@ -54,68 +65,76 @@ void FFieldSystemPhysicsProxy::FieldParameterUpdateCallback(Chaos::FPhysicsSolve
 	{
 #if TODO_REIMPLEMENT_RIGID_CLUSTERING
 		const Chaos::TArrayCollectionArray<Chaos::ClusterId> & ClusterIDs = CurrentSolver->GetRigidClustering().GetClusterIdsArray();
-#endif
-
-		// @todo: This seems like a waste if we just want to get everything
-		TArray<ContextIndex> IndicesArray;
+#endif // TODO_REIMPLEMENT_RIGID_CLUSTERING
 
 		TArray<FFieldSystemCommand>* CommandListPtr = GetSolverCommandList(InSolver);
-		
 		if(!CommandListPtr)
 		{
 			// No command list present for this solver, bail out
 			return;
 		}
-
 		const int32 NumCommands = CommandListPtr->Num();
-
 		TArray<int32> CommandsToRemove;
+		CommandsToRemove.Reserve(NumCommands);
+
+		// @todo: This seems like a waste if we just want to get everything
+		//TArray<ContextIndex> IndicesArray;
+		TArray<Chaos::TGeometryParticleHandle<float,3>*> Handles;
+		TArray<FVector> SamplePoints;
+		TArray<ContextIndex> SampleIndices;
+		TArray<ContextIndex>& IndicesArray = SampleIndices; // Do away with!
+		EFieldResolutionType PrevResolutionType = static_cast<EFieldResolutionType>(0); // none
 		for (int32 CommandIndex = 0; CommandIndex < NumCommands; CommandIndex++)
 		{
 			const FFieldSystemCommand& Command = (*CommandListPtr)[CommandIndex];
-			EFieldResolutionType ResolutionType = EFieldResolutionType::Field_Resolution_Minimal;
-			if (Command.MetaData.Contains(FFieldSystemMetaData::EMetaType::ECommandData_ProcessingResolution))
+			const EFieldResolutionType ResolutionType = 
+				Command.HasMetaData(FFieldSystemMetaData::EMetaType::ECommandData_ProcessingResolution) ?
+					Command.GetMetaDataAs<FFieldSystemMetaDataProcessingResolution>(
+						FFieldSystemMetaData::EMetaType::ECommandData_ProcessingResolution)->ProcessingResolution :
+					EFieldResolutionType::Field_Resolution_Minimal;
+			if (PrevResolutionType != ResolutionType || Handles.Num() == 0)
 			{
-				check(Command.MetaData[FFieldSystemMetaData::EMetaType::ECommandData_ProcessingResolution] != nullptr);
-				ResolutionType = static_cast<FFieldSystemMetaDataProcessingResolution*>(Command.MetaData[FFieldSystemMetaData::EMetaType::ECommandData_ProcessingResolution].Get())->ProcessingResolution;
+				FFieldSystemPhysicsProxy::GetParticleHandles(Handles, CurrentSolver, ResolutionType);
+				PrevResolutionType = ResolutionType;
+
+				SamplePoints.AddUninitialized(Handles.Num());
+				SampleIndices.AddUninitialized(Handles.Num());
+				for (int32 Idx = 0; Idx < Handles.Num(); ++Idx)
+				{
+					SamplePoints[Idx] = Handles[Idx]->X();
+					SampleIndices[Idx] = ContextIndex(Idx, Idx);
+				}
 			}
 
 			if (Command.TargetAttribute == GetFieldPhysicsName(EFieldPhysicsType::Field_DynamicState))
 			{
 				SCOPE_CYCLE_COUNTER(STAT_ParamUpdateField_DynamicState);
 
-				FFieldSystemPhysicsProxy::ContiguousIndices(IndicesArray, CurrentSolver, ResolutionType, IndicesArray.Num() != Particles.Size());
-				if (IndicesArray.Num())
+				if (Handles.Num())
 				{
-					TArrayView<ContextIndex> IndexView(&(IndicesArray[0]), IndicesArray.Num());
-
-					FVector * tptr = &(Particles.X(0));
-					TArrayView<FVector> SamplesView(tptr, int32(Particles.Size()));
+					TArrayView<Chaos::TGeometryParticleHandle<float,3>*> HandlesView(&(Handles[0]), Handles.Num());
+					TArrayView<FVector> SamplePointsView(&(SamplePoints[0]), SamplePoints.Num());
+					TArrayView<ContextIndex> SampleIndicesView(&(SampleIndices[0]), SampleIndices.Num());
 
-					FFieldContext Context{
-						IndexView, // @todo(brice) important: an empty index array should evaluate everything
-						SamplesView,
-						Command.MetaData
-					};
+					FFieldContext Context(
+						SampleIndicesView, // @todo(brice) important: an empty index array should evaluate everything
+						SamplePointsView,
+						Command.MetaData);
 
 					//
 					//  Sample the dynamic state array in the field
 					//
-					TArray<int32> DynamicState;
-					DynamicState.AddUninitialized(Particles.Size());
-					for(const ContextIndex& Index : IndicesArray)
+					TArray<int32> DynamicState; // TODO: 32 bits seems rather excessive!!!
+					DynamicState.AddUninitialized(Handles.Num());					
+					int32 i = 0;
+					for(Chaos::TGeometryParticleHandle<float,3>* Handle : Handles)
 					{
-						DynamicState[Index.Sample] = 0;	//is this needed?
-					}
-					
-					for (const ContextIndex& Index : IndicesArray)
-					{
-						const int32 i = Index.Sample;
-						if (Particles.ObjectState(i) == Chaos::EObjectStateType::Kinematic)
+						const Chaos::EObjectStateType CurrState = Handle->ObjectState();
+						if(CurrState == Chaos::EObjectStateType::Kinematic)
 						{
 							DynamicState[i] = (int)EObjectStateTypeEnum::Chaos_Object_Kinematic;
 						}
-						else if (Particles.ObjectState(i) == Chaos::EObjectStateType::Static)
+						else if (CurrState == Chaos::EObjectStateType::Static)
 						{
 							DynamicState[i] = (int)EObjectStateTypeEnum::Chaos_Object_Static;
 						}
@@ -123,66 +142,88 @@ void FFieldSystemPhysicsProxy::FieldParameterUpdateCallback(Chaos::FPhysicsSolve
 						{
 							DynamicState[i] = (int)EObjectStateTypeEnum::Chaos_Object_Dynamic;
 						}
+						++i;
 					}
-					TArrayView<int32> DynamicStateView(&(DynamicState[0]), DynamicState.Num());
 
+					TArrayView<int32> DynamicStateView(&(DynamicState[0]), DynamicState.Num());
 					if (ensureMsgf(Command.RootNode->Type() == FFieldNode<int32>::StaticType(),
 						TEXT("Field based evaluation of the simulations 'ObjectType' parameter expects int32 field inputs.")))
 					{
-						static_cast<const FFieldNode<int32> *>(Command.RootNode.Get())->Evaluate(Context, DynamicStateView);
+						static_cast<const FFieldNode<int32>*>(
+							Command.RootNode.Get())->Evaluate(Context, DynamicStateView);
 					}
 
-#if TODO_REIMPLEMENT_GETFLOORINDEX
-					// transfer results to rigid system.
-					int32 FloorIndex = CurrentSolver->GetFloorIndex();
-					int32 NumSamples = Context.SampleIndices.Num();
-					for (int32 SampleIndex = 0; SampleIndex < NumSamples; SampleIndex++)
+					bool StateChanged = false;
+					i = 0;
+					for (Chaos::TGeometryParticleHandle<float, 3>* Handle : Handles)
 					{
-						int32 RigidBodyIndex = Context.SampleIndices[SampleIndex].Result;
-						if (RigidBodyIndex != FloorIndex) // ignore the floor
-						{ 
-							if (DynamicStateView[RigidBodyIndex] == (int32)EObjectStateTypeEnum::Chaos_Object_Dynamic
-								&& Particles.ObjectState(RigidBodyIndex) == Chaos::EObjectStateType::Static
-								&& FLT_EPSILON < Particles.M(RigidBodyIndex))
-							{
-								Particles.SetObjectState(RigidBodyIndex, Chaos::EObjectStateType::Dynamic);
-							}
-							// todo: This breaks shooting kinematic characters as dynamics takes over but enables cache interaction
-							else if (DynamicStateView[RigidBodyIndex] == (int32)EObjectStateTypeEnum::Chaos_Object_Dynamic
-								&& Particles.ObjectState(RigidBodyIndex) == Chaos::EObjectStateType::Kinematic
-								&& FLT_EPSILON < Particles.M(RigidBodyIndex))
-							{
-								Particles.SetObjectState(RigidBodyIndex, Chaos::EObjectStateType::Dynamic);
-							}
-							else if ((DynamicStateView[RigidBodyIndex] == (int32)EObjectStateTypeEnum::Chaos_Object_Kinematic)
-								&& Particles.ObjectState(RigidBodyIndex) == Chaos::EObjectStateType::Dynamic)
+						// Lower level particle handles, like TGeometryParticleHandle and 
+						// TKinematicParticleHandle, infer their dynamic state by whether or not
+						// promotion to a derived c++ handle type succeeds or fails.
+						//
+						// THAT IS NOT WHAT WE WANT.
+						//
+						// PBDRigidParticles has an array of EObjectStateType, and the associated
+						// handle has a getter and a setter for that data.  So, at least for now,
+						// we're just going to ignore non-dynamic particles.  This has the added
+						// benefit of not needing to deal with the floor, as it's pretty likely to
+						// not be dynamic.  Har.
+
+						if (Chaos::TPBDRigidParticleHandleImp<float, 3, true>* DynHandle = Handle->AsDynamic())
+						{
+							const int32 FieldState = DynamicStateView[i];
+							const EObjectStateType HandleState = DynHandle->ObjectState();
+							if (FieldState == (int32)EObjectStateTypeEnum::Chaos_Object_Dynamic)
 							{
-								Particles.SetObjectState(RigidBodyIndex, Chaos::EObjectStateType::Kinematic);
-								Particles.V(RigidBodyIndex) = Chaos::TVector<float, 3>(0);
-								Particles.W(RigidBodyIndex) = Chaos::TVector<float, 3>(0);
+								if ((HandleState == Chaos::EObjectStateType::Static ||
+									 HandleState == Chaos::EObjectStateType::Kinematic) &&
+									DynHandle->M() > FLT_EPSILON)
+								{
+									DynHandle->SetObjectState(Chaos::EObjectStateType::Dynamic);
+									StateChanged = true;
+								}
+								else if (HandleState == Chaos::EObjectStateType::Sleeping)
+								{
+									DynHandle->SetObjectState(Chaos::EObjectStateType::Dynamic);
+									StateChanged = true;
+								}
 							}
-							else if ((DynamicStateView[RigidBodyIndex] == (int32)EObjectStateTypeEnum::Chaos_Object_Static)
-								&& Particles.ObjectState(RigidBodyIndex) == Chaos::EObjectStateType::Dynamic)
+							else if (FieldState == (int32)EObjectStateTypeEnum::Chaos_Object_Kinematic)
 							{
-								Particles.SetObjectState(RigidBodyIndex, Chaos::EObjectStateType::Static);
-								Particles.V(RigidBodyIndex) = Chaos::TVector<float, 3>(0);
-								Particles.W(RigidBodyIndex) = Chaos::TVector<float, 3>(0);
+								if (HandleState == Chaos::EObjectStateType::Dynamic)
+								{
+									DynHandle->SetObjectState(Chaos::EObjectStateType::Kinematic);
+									DynHandle->SetV(Chaos::TVector<float, 3>(0));
+									DynHandle->SetW(Chaos::TVector<float, 3>(0));
+									StateChanged = true;
+								}
 							}
-							else if ((DynamicStateView[RigidBodyIndex] == (int32)EObjectStateTypeEnum::Chaos_Object_Dynamic)
-								&& Particles.ObjectState(RigidBodyIndex) == Chaos::EObjectStateType::Sleeping)
+							else if (FieldState == (int32)EObjectStateTypeEnum::Chaos_Object_Static)
 							{
-								Particles.SetObjectState(RigidBodyIndex, Chaos::EObjectStateType::Dynamic);
-								CurrentSolver->ActiveIndices().Add(RigidBodyIndex);
+								if (HandleState == Chaos::EObjectStateType::Dynamic)
+								{
+									DynHandle->SetObjectState(Chaos::EObjectStateType::Static);
+									DynHandle->SetV(Chaos::TVector<float, 3>(0));
+									DynHandle->SetW(Chaos::TVector<float, 3>(0));
+									StateChanged = true;
+								}
 							}
-							else if ((DynamicStateView[RigidBodyIndex] == (int32)EObjectStateTypeEnum::Chaos_Object_Sleeping)
-								&& Particles.ObjectState(RigidBodyIndex) == Chaos::EObjectStateType::Dynamic)
+							else if (FieldState == (int32)EObjectStateTypeEnum::Chaos_Object_Sleeping)
 							{
-								Particles.SetObjectState(RigidBodyIndex, Chaos::EObjectStateType::Sleeping);
-								CurrentSolver->ActiveIndices().Remove(RigidBodyIndex);
+								if (HandleState == Chaos::EObjectStateType::Dynamic)
+								{
+									DynHandle->SetObjectState(Chaos::EObjectStateType::Sleeping);
+									StateChanged = true;
+								}
 							}
-						}
+						} // handle is dynamic
+						++i;
+					} // end for all handles
+					if (StateChanged)
+					{
+						// regenerate views
+						CurrentSolver->GetParticles().UpdateGeometryCollectionViews();
 					}
-#endif
 
 #if TODO_REIMPLEMENT_RIGID_CLUSTERING
 					//  Update all cluster bodies based on the changes in the kinematic state.
@@ -194,7 +235,7 @@ void FFieldSystemPhysicsProxy::FieldParameterUpdateCallback(Chaos::FPhysicsSolve
 							CurrentSolver->GetRigidClustering().UpdateKinematicProperties(ActiveParticleIndex);
 						}
 					}
-#endif
+#endif // TODO_REIMPLEMENT_RIGID_CLUSTERING
 
 				}
 				CommandsToRemove.Add(CommandIndex);
@@ -892,8 +933,12 @@ void FFieldSystemPhysicsProxy::FieldParameterUpdateCallback(Chaos::FPhysicsSolve
 	}
 }
 
-void FFieldSystemPhysicsProxy::FieldForcesUpdateCallback(Chaos::FPhysicsSolver* InSolver, FParticlesType& Particles, Chaos::TArrayCollectionArray<FVector> & Force,
-	Chaos::TArrayCollectionArray<FVector> & Torque, const float Time)
+void FFieldSystemPhysicsProxy::FieldForcesUpdateCallback(
+	Chaos::FPhysicsSolver* InSolver, 
+	FParticlesType& Particles, 
+	Chaos::TArrayCollectionArray<FVector> & Force, 
+	Chaos::TArrayCollectionArray<FVector> & Torque, 
+	const float Time)
 {
 	if (Commands.Num() && InSolver)
 	{
@@ -901,7 +946,6 @@ void FFieldSystemPhysicsProxy::FieldForcesUpdateCallback(Chaos::FPhysicsSolver*
 		TArray<ContextIndex> IndicesArray;
 
 		TArray<FFieldSystemCommand>* CommandListPtr = GetSolverCommandList(InSolver);
-
 		if(!CommandListPtr)
 		{
 			return;
@@ -912,34 +956,55 @@ void FFieldSystemPhysicsProxy::FieldForcesUpdateCallback(Chaos::FPhysicsSolver*
 		for(int32 CommandIndex = 0; CommandIndex < NumCommands; CommandIndex++)
 		{
 			const FFieldSystemCommand & Command = (*CommandListPtr)[CommandIndex];
-			EFieldResolutionType ResolutionType = EFieldResolutionType::Field_Resolution_Minimal;
-			if (Command.MetaData.Contains(FFieldSystemMetaData::EMetaType::ECommandData_ProcessingResolution))
-			{
-				check(Command.MetaData[FFieldSystemMetaData::EMetaType::ECommandData_ProcessingResolution] != nullptr);
-				ResolutionType = static_cast<FFieldSystemMetaDataProcessingResolution*>(Command.MetaData[FFieldSystemMetaData::EMetaType::ECommandData_ProcessingResolution].Get())->ProcessingResolution;
-			}
+			const EFieldResolutionType ResolutionType = 
+				Command.HasMetaData(FFieldSystemMetaData::EMetaType::ECommandData_ProcessingResolution) ?
+					Command.GetMetaDataAs<FFieldSystemMetaDataProcessingResolution>(
+						FFieldSystemMetaData::EMetaType::ECommandData_ProcessingResolution)->ProcessingResolution :
+					EFieldResolutionType::Field_Resolution_Minimal;
 
 			if (Command.TargetAttribute == GetFieldPhysicsName(EFieldPhysicsType::Field_LinearForce))
 			{
 				if (ensureMsgf(Command.RootNode->Type() == FFieldNode<FVector>::StaticType(),
 					TEXT("Field based evaluation of the simulations 'Force' parameter expects FVector field inputs.")))
 				{
-					FFieldSystemPhysicsProxy::ContiguousIndices(IndicesArray, CurrentSolver, ResolutionType, IndicesArray.Num() != Particles.Size());
-					if (IndicesArray.Num())
+					TArray<Chaos::TGeometryParticleHandle<float,3>*> Handles;
+					FFieldSystemPhysicsProxy::GetParticleHandles(Handles, CurrentSolver, ResolutionType);
+					if (Handles.Num())
 					{
-						TArrayView<ContextIndex> IndexView(&(IndicesArray[0]), IndicesArray.Num());
+						TArray<FVector> SamplePoints;
+						TArray<ContextIndex> SampleIndices;
+						SamplePoints.AddUninitialized(Handles.Num());
+						SampleIndices.AddUninitialized(Handles.Num());
+						for (int32 Idx = 0; Idx < Handles.Num(); ++Idx)
+						{
+							SamplePoints[Idx] = Handles[Idx]->X();
+							SampleIndices[Idx] = ContextIndex(Idx, Idx);
+						}
 
-						FVector * tptr = &(Particles.X(0));
-						TArrayView<FVector> SamplesView(tptr, int32(Particles.Size()));
+						TArrayView<FVector> SamplePointsView(&(SamplePoints[0]), SamplePoints.Num());
+						TArrayView<ContextIndex> SampleIndicesView(&(SampleIndices[0]), SampleIndices.Num());
 
-						FFieldContext Context{
-							IndexView,
-							SamplesView,
-							Command.MetaData
-						};
-						TArrayView<FVector> ForceView(&(Force[0]), Force.Num());
+						FFieldContext Context(
+							SampleIndicesView,
+							SamplePointsView,
+							Command.MetaData);
+
+						TArray<FVector> LocalForce;
+						LocalForce.AddUninitialized(Handles.Num());					
+						TArrayView<FVector> ForceView(&(LocalForce[0]), LocalForce.Num());
 						static_cast<const FFieldNode<FVector> *>(Command.RootNode.Get())->Evaluate(Context, ForceView);
+		
+						int32 i = 0;
+						for (Chaos::TGeometryParticleHandle<float, 3>* Handle : Handles)
+						{
+							if (Chaos::TPBDRigidParticleHandleImp<float, 3, true>* DynHandle = Handle->AsDynamic())
+							{
+								DynHandle->ExternalForce() += ForceView[i];
+							}
+							++i;
+						}
 
+#if TODO_REIMPLEMENT_WAKE_ISLANDS
 						// @todo(ccaulfield): encapsulation: add WakeParticles (and therefore islands) functionality to Evolution
 						TSet<int32> IslandsToActivate;
 						for (const ContextIndex& CIndex : IndicesArray)
@@ -950,7 +1015,6 @@ void FFieldSystemPhysicsProxy::FieldForcesUpdateCallback(Chaos::FPhysicsSolver*
 								IslandsToActivate.Add(Particles.Island(i));
 							}
 						}
-#if TODO_REIMPLEMENT_WAKE_ISLANDS
 						InSolver->WakeIslands(IslandsToActivate);
 #endif
 					}
@@ -962,22 +1026,44 @@ void FFieldSystemPhysicsProxy::FieldForcesUpdateCallback(Chaos::FPhysicsSolver*
 				if (ensureMsgf(Command.RootNode->Type() == FFieldNode<FVector>::StaticType(),
 					TEXT("Field based evaluation of the simulations 'Torque' parameter expects FVector field inputs.")))
 				{
-					FFieldSystemPhysicsProxy::ContiguousIndices(IndicesArray, CurrentSolver, ResolutionType, IndicesArray.Num() != Particles.Size());
-					if (IndicesArray.Num())
+					TArray<Chaos::TGeometryParticleHandle<float,3>*> Handles;
+					FFieldSystemPhysicsProxy::GetParticleHandles(Handles, CurrentSolver, ResolutionType);
+					if (Handles.Num())
 					{
-						TArrayView<ContextIndex> IndexView(&(IndicesArray[0]), IndicesArray.Num());
+						TArray<FVector> SamplePoints;
+						TArray<ContextIndex> SampleIndices;
+						SamplePoints.AddUninitialized(Handles.Num());
+						SampleIndices.AddUninitialized(Handles.Num());
+						for (int32 Idx = 0; Idx < Handles.Num(); ++Idx)
+						{
+							SamplePoints[Idx] = Handles[Idx]->X();
+							SampleIndices[Idx] = ContextIndex(Idx, Idx);
+						}
 
-						FVector * tptr = &(Particles.X(0));
-						TArrayView<FVector> SamplesView(tptr, int32(Particles.Size()));
+						TArrayView<FVector> SamplePointsView(&(SamplePoints[0]), SamplePoints.Num());
+						TArrayView<ContextIndex> SampleIndicesView(&(SampleIndices[0]), SampleIndices.Num());
 
-						FFieldContext Context{
-							IndexView,
-							SamplesView,
-							Command.MetaData
-						};
-						TArrayView<FVector> TorqueView(&(Torque[0]), Torque.Num());
+						FFieldContext Context(
+							SampleIndicesView,
+							SamplePointsView,
+							Command.MetaData);
+
+						TArray<FVector> LocalTorque;
+						LocalTorque.AddUninitialized(Handles.Num());					
+						TArrayView<FVector> TorqueView(&(LocalTorque[0]), LocalTorque.Num());
 						static_cast<const FFieldNode<FVector> *>(Command.RootNode.Get())->Evaluate(Context, TorqueView);
+		
+						int32 i = 0;
+						for (Chaos::TGeometryParticleHandle<float, 3>* Handle : Handles)
+						{
+							if (Chaos::TPBDRigidParticleHandleImp<float, 3, true>* DynHandle = Handle->AsDynamic())
+							{
+								DynHandle->ExternalTorque() += TorqueView[i];
+							}
+							++i;
+						}
 
+#if TODO_REIMPLEMENT_WAKE_ISLANDS
 						// @todo(ccaulfield): encapsulation: add WakeParticles (and therefore islands) functionality to Evolution
 						TSet<int32> IslandsToActivate;
 						for (const ContextIndex& CIndex : IndicesArray)
@@ -988,7 +1074,6 @@ void FFieldSystemPhysicsProxy::FieldForcesUpdateCallback(Chaos::FPhysicsSolver*
 								IslandsToActivate.Add(Particles.Island(i));
 							}
 						}
-#if TODO_REIMPLEMENT_WAKE_ISLANDS
 						InSolver->WakeIslands(IslandsToActivate);
 #endif
 					}
@@ -1003,87 +1088,140 @@ void FFieldSystemPhysicsProxy::FieldForcesUpdateCallback(Chaos::FPhysicsSolver*
 	}
 }
 
-void FFieldSystemPhysicsProxy::EndFrameCallback(const float InDt)
-{
-}
-
 void FFieldSystemPhysicsProxy::BufferCommand(Chaos::FPhysicsSolver* InSolver, const FFieldSystemCommand& InCommand)
 {
+	// TODO: Consider inspecting InCommand and bucketing according to which evaluation 
+	// path it requires; FieldParameterUpdateCallback() or FieldForcesUpdateCallback().
+	// TODO: Consider using a lock free triple buffer.
+
 	TArray<FFieldSystemCommand>** ExistingList = nullptr;
 	{
 		FScopeLock Lock(&CommandLock);
 		ExistingList = Commands.Find(InSolver);
-
 		if(!ExistingList)
 		{
-			Commands.Add(InSolver);
-			ExistingList = Commands.Find(InSolver);
+			ExistingList = &Commands.Add(InSolver);
+			//Commands.Add(InSolver);
+			//ExistingList = Commands.Find(InSolver);
 			check(ExistingList);
-
 			(*ExistingList) = new TArray<FFieldSystemCommand>();
 		}
 	}
-
 	(*ExistingList)->Add(InCommand);
 }
 
-void FFieldSystemPhysicsProxy::ContiguousIndices(TArray<ContextIndex>& Array, const Chaos::FPhysicsSolver* RigidSolver, EFieldResolutionType ResolutionType, bool bForce = true)
+void FFieldSystemPhysicsProxy::GetParticleHandles(
+	TArray<Chaos::TGeometryParticleHandle<float,3>*>& Handles,
+	const Chaos::FPhysicsSolver* RigidSolver,
+	const EFieldResolutionType ResolutionType,
+	const bool bForce)
 {
-#if TODO_REIMPLEMENT_RIGID_CLUSTERING
-	if (bForce)
+	Handles.SetNum(0, false);
+	if (!bForce)
+		return;
+
+	const Chaos::TPBDRigidsSOAs<float, 3>& SolverParticles = RigidSolver->GetParticles(); // const?
+
+	if (ResolutionType == EFieldResolutionType::Field_Resolution_Maximum)
 	{
-		const Chaos::FPhysicsSolver::FParticlesType & Particles = RigidSolver->GetRigidParticles();
-		if (ResolutionType == EFieldResolutionType::Field_Resolution_Minimal)
+		// const TParticleView<TGeometryParticles<T, d>>& TPBDRigidSOAs<T,d>::GetAllParticlesView()
+		const Chaos::TParticleView<Chaos::TGeometryParticles<float, 3>> &ParticleView = 
+			SolverParticles.GetAllParticlesView();
+		Handles.Reserve(ParticleView.Num());
+
+		// TParticleIterator<TSOA> Begin() const, TSOA = TGeometryParticles<T, d>
+		for (Chaos::TParticleIterator<Chaos::TGeometryParticles<float, 3>> It = ParticleView.Begin(), ItEnd = ParticleView.End();
+			It != ItEnd; ++It)
+		{
+			const Chaos::TTransientGeometryParticleHandle<float,3> *Handle = &(*It);
+			// PBDRigidsSOAs.h only has a const version of GetAllParticlesView() - is that wrong?
+			Handles.Add(GetHandleHelper(const_cast<Chaos::TTransientGeometryParticleHandle<float,3>*>(Handle)));
+		}
+	}
+	else if (ResolutionType == EFieldResolutionType::Field_Resolution_Minimal)
+	{
+		const Chaos::TParticleView<Chaos::TGeometryParticles<float, 3>> &ParticleView = SolverParticles.GetNonDisabledView();
+		Handles.Reserve(ParticleView.Num());
+		for (Chaos::TParticleIterator<Chaos::TGeometryParticles<float, 3>> It = ParticleView.Begin(), ItEnd = ParticleView.End();
+			It != ItEnd; ++It)
 		{
-			Array.SetNum(0, false);
+			const Chaos::TTransientGeometryParticleHandle<float,3> *Handle = &(*It);
+			Handles.Add(GetHandleHelper(const_cast<Chaos::TTransientGeometryParticleHandle<float,3>*>(Handle)));
+
+		}
+	}
+	else if (ResolutionType == EFieldResolutionType::Field_Resolution_DisabledParents)
+	{
+		check(false); // unimplemented
+	}
+}
+
+void FFieldSystemPhysicsProxy::ContiguousIndices(
+	TArray<ContextIndex>& Array, 
+	const Chaos::FPhysicsSolver* RigidSolver, 
+	const EFieldResolutionType ResolutionType, 
+	const bool bForce)
+{
+#if TODO_REIMPLEMENT_RIGID_CLUSTERING
+	if (!bForce)
+	{
+		return;
+	}
+
+	if (ResolutionType == EFieldResolutionType::Field_Resolution_Minimal)
+	{
+		Array.SetNum(0, false);
 
-			int32 IndexCount = 0;
-			int32 FloorIndex = RigidSolver->GetFloorIndex();
-			const Chaos::FPhysicsSolver::FClusteringType& Clustering = RigidSolver->GetRigidClustering();
-			const Chaos::TArrayCollectionArray<Chaos::ClusterId>& ClusterIdArray = Clustering.GetClusterIdsArray();
-			const Chaos::FPhysicsSolver::FClusteringType::FClusterMap &  ClusterMap = Clustering.GetChildrenMap();
+		int32 IndexCount = 0;
+		int32 FloorIndex = RigidSolver->GetFloorIndex();
 
-			for (int32 ActiveParticleIndex : RigidSolver->NonDisabledIndices())
+#if TODO_REIMPLEMENT_RIGID_CLUSTERING
+		const Chaos::FPhysicsSolver::FClusteringType& Clustering = RigidSolver->GetRigidClustering();
+		const Chaos::TArrayCollectionArray<Chaos::ClusterId>& ClusterIdArray = Clustering.GetClusterIdsArray();
+		const Chaos::FPhysicsSolver::FClusteringType::FClusterMap &  ClusterMap = Clustering.GetChildrenMap();
+#endif // TODO_REIMPLEMENT_RIGID_CLUSTERING
+
+		for (int32 ActiveParticleIndex : RigidSolver->NonDisabledIndices())
+		{
+			if (ClusterIdArray[ActiveParticleIndex].NumChildren)
 			{
-				if (ClusterIdArray[ActiveParticleIndex].NumChildren)
+				for (uint32 ClusterChild : *ClusterMap[ActiveParticleIndex])
 				{
-					for (uint32 ClusterChild : *ClusterMap[ActiveParticleIndex])
-					{
-						Array.Add(ContextIndex(ClusterChild, ClusterChild));
-					}
+					Array.Add(ContextIndex(ClusterChild, ClusterChild));
 				}
+			}
 
-				if( ActiveParticleIndex != FloorIndex )
-				{
-					Array.Add(ContextIndex(ActiveParticleIndex, ActiveParticleIndex));
-				}
+			if( ActiveParticleIndex != FloorIndex )
+			{
+				Array.Add(ContextIndex(ActiveParticleIndex, ActiveParticleIndex));
 			}
 		}
-		if (ResolutionType == EFieldResolutionType::Field_Resolution_DisabledParents)
-		{
-			Array.SetNum(0, false);
+	}
+	else if (ResolutionType == EFieldResolutionType::Field_Resolution_DisabledParents)
+	{
+		Array.SetNum(0, false);
 
-			int32 IndexCount = 0;
-			int32 FloorIndex = RigidSolver->GetFloorIndex();
-			const Chaos::FPhysicsSolver::FClusteringType& Clustering = RigidSolver->GetRigidClustering();
-			const Chaos::TArrayCollectionArray<Chaos::ClusterId>& ClusterIdArray = Clustering.GetClusterIdsArray();
+		int32 IndexCount = 0;
+		int32 FloorIndex = RigidSolver->GetFloorIndex();
+		const Chaos::FPhysicsSolver::FClusteringType& Clustering = RigidSolver->GetRigidClustering();
+		const Chaos::TArrayCollectionArray<Chaos::ClusterId>& ClusterIdArray = Clustering.GetClusterIdsArray();
 
-			for (int32 TopLevelParent : Clustering.GetTopLevelClusterParents())
+		for (int32 TopLevelParent : Clustering.GetTopLevelClusterParents())
+		{
+			if( TopLevelParent != FloorIndex )
 			{
-				if( TopLevelParent != FloorIndex )
-				{
-					Array.Add(ContextIndex(TopLevelParent, TopLevelParent));
-				}
+				Array.Add(ContextIndex(TopLevelParent, TopLevelParent));
 			}
 		}
-		else if (ResolutionType == EFieldResolutionType::Field_Resolution_Maximum)
+	}
+	else if (ResolutionType == EFieldResolutionType::Field_Resolution_Maximum)
+	{
+		const Chaos::FPhysicsSolver::FParticlesType & Particles = RigidSolver->GetRigidParticles();
+		Array.SetNum(Particles.Size());
+		for (int32 i = 0; i < Array.Num(); ++i)
 		{
-			Array.SetNum(Particles.Size());
-			for (int32 i = 0; i < Array.Num(); ++i)
-			{
-				Array[i].Result = i;
-				Array[i].Sample = i;
-			}
+			Array[i].Sample = i;
+			Array[i].Result = i;
 		}
 	}
 #endif
diff --git a/Engine/Source/Runtime/Experimental/ChaosSolvers/Private/PhysicsProxy/GeometryCollectionPhysicsProxy.cpp b/Engine/Source/Runtime/Experimental/ChaosSolvers/Private/PhysicsProxy/GeometryCollectionPhysicsProxy.cpp
--- a/Engine/Source/Runtime/Experimental/ChaosSolvers/Private/PhysicsProxy/GeometryCollectionPhysicsProxy.cpp
+++ b/Engine/Source/Runtime/Experimental/ChaosSolvers/Private/PhysicsProxy/GeometryCollectionPhysicsProxy.cpp
@@ -586,7 +586,7 @@ void FGeometryCollectionPhysicsProxy::InitializeBodiesGT()
 //				{
 //					check(false); // continue?
 //				}
-				FTransform ParticleTransform = Transform[Idx] * Parameters.WorldTransform;
+				FTransform ParticleTransform = MassToLocal[Idx] * Transform[Idx] * Parameters.WorldTransform;
 				if (GTParticle->ObjectState() == Chaos::EObjectStateType::Dynamic)
 				{
 					GTParticle->SetX(ParticleTransform.GetTranslation());
@@ -672,8 +672,7 @@ void FGeometryCollectionPhysicsProxy::InitializeBodiesPT(
 			if (Chaos::TPBDRigidParticleHandle<float, 3>* Handle = SolverParticleHandles[TransformGroupIndex])
 			{
 				//const int32 RigidBodyIndex = RigidBodyID[TransformGroupIndex];
-				const FTransform WorldTransform = 
-					MassToLocal[TransformGroupIndex] * Transform[TransformGroupIndex] * Parameters.WorldTransform;
+				const FTransform WorldTransform = MassToLocal[TransformGroupIndex] * Transform[TransformGroupIndex] * Parameters.WorldTransform;
 
 				PopulateSimulatedParticle(
 					Handle,
@@ -690,7 +689,7 @@ void FGeometryCollectionPhysicsProxy::InitializeBodiesPT(
 #endif // TODO_REIMPLEMENT_SOLVER_SETTINGS_ACCESSORS
 					//RigidBodyIndex, 
 					WorldTransform, 
-					(uint8)EObjectStateTypeEnum::Chaos_Object_Dynamic, 
+					(uint8)Parameters.ObjectType, 
 					CollisionGroup[TransformGroupIndex]);
 
 #if TODO_REIMPLEMENT_RIGID_CLUSTERING
@@ -716,6 +715,9 @@ void FGeometryCollectionPhysicsProxy::InitializeBodiesPT(
 			}
 		});
 
+		// After population, the states of each particle could have changed
+		Particles.UpdateGeometryCollectionViews();
+
 #if TODO_REIMPLEMENT_INIT_COMMANDS
 		for (FFieldSystemCommand& Cmd : Parameters.InitializationCommands)
 		{
@@ -3757,7 +3759,8 @@ Chaos::FParticleData* FGeometryCollectionPhysicsProxy::NewData()
 	// proxies get away with that because they're dealing with a low number of 
 	// bodies; which is not the case here. Rather, let's just use a triple buffer.
 
-	if (!ensure(NumParticles != INDEX_NONE)) // Make sure InitBodiesGT() has been called!
+	if (!Parameters.Simulating || 
+		!ensure(NumParticles != INDEX_NONE)) // Make sure InitBodiesGT() has been called!
 	{
 		return nullptr;
 	}
@@ -4081,7 +4084,7 @@ void FGeometryCollectionPhysicsProxy::PullFromPhysicsState()
 		{
 			if(!TR.DisabledStates[TmIndex])
 			{
-				GTDynamicCollection->Transform[TmIndex] = TR.Transforms[TmIndex].GetRelativeTransform(Parameters.WorldTransform);
+				GTDynamicCollection->Transform[TmIndex] = MassToLocal[TmIndex].GetRelativeTransformReverse(TR.Transforms[TmIndex]).GetRelativeTransform(Parameters.WorldTransform);
 				GTDynamicCollection->Transform[TmIndex].NormalizeRotation();
 			}
 		}
diff --git a/Engine/Source/Runtime/Experimental/ChaosSolvers/Public/PBDRigidsSolver.h b/Engine/Source/Runtime/Experimental/ChaosSolvers/Public/PBDRigidsSolver.h
--- a/Engine/Source/Runtime/Experimental/ChaosSolvers/Public/PBDRigidsSolver.h
+++ b/Engine/Source/Runtime/Experimental/ChaosSolvers/Public/PBDRigidsSolver.h
@@ -114,8 +114,8 @@ namespace Chaos
 		// TODO: Set up an interface for registering fields and geometry collections
 		void RegisterObject(FGeometryCollectionPhysicsProxy* InProxy);
 		bool UnregisterObject(FGeometryCollectionPhysicsProxy* InProxy);
-		//void RegisterObject(FFieldSystemPhysicsProxy* InProxy);
-		//int UnregisterObject(FFieldSystemPhysicsProxy* InProxy);
+		void RegisterObject(FFieldSystemPhysicsProxy* InProxy);
+		bool UnregisterObject(FFieldSystemPhysicsProxy* InProxy);
 
 		bool IsSimulating() const;
 
@@ -261,6 +261,7 @@ namespace Chaos
 		FPBDRigidsEvolution* GetEvolution() { return MEvolution.Get(); }
 		FPBDRigidsEvolution* GetEvolution() const { return MEvolution.Get(); }
 
+		FParticlesType& GetParticles() { return Particles; }
 		const FParticlesType& GetParticles() const { return Particles; }
 
 		/**/
@@ -276,6 +277,11 @@ namespace Chaos
 		/**/
 		void PostTickDebugDraw() const;
 
+		TArray<FFieldSystemPhysicsProxy*>& GetFieldSystemPhysicsProxies()
+		{
+			return FieldSystemPhysicsProxies;
+		}
+
 		/** Events hooked up to the Chaos material manager */
 		void UpdateMaterial(Chaos::FMaterialHandle InHandle, const Chaos::FChaosPhysicsMaterial& InNewData);
 		void CreateMaterial(Chaos::FMaterialHandle InHandle, const Chaos::FChaosPhysicsMaterial& InNewData);
@@ -308,7 +314,6 @@ namespace Chaos
 			((ParticleType*)(Handle->GTGeometryParticle()->Proxy))->BufferPhysicsResults();
 		}
 
-
 		//
 		// Solver Data
 		//
diff --git a/Engine/Source/Runtime/Experimental/ChaosSolvers/Public/PhysicsProxy/FieldSystemPhysicsProxy.h b/Engine/Source/Runtime/Experimental/ChaosSolvers/Public/PhysicsProxy/FieldSystemPhysicsProxy.h
--- a/Engine/Source/Runtime/Experimental/ChaosSolvers/Public/PhysicsProxy/FieldSystemPhysicsProxy.h
+++ b/Engine/Source/Runtime/Experimental/ChaosSolvers/Public/PhysicsProxy/FieldSystemPhysicsProxy.h
@@ -21,11 +21,53 @@ public:
 	FFieldSystemPhysicsProxy(UObject* InOwner);
 	virtual ~FFieldSystemPhysicsProxy();
 
+	// Called by FPBDRigidsSolver::RegisterObject(FFieldSystemPhysicsProxy*)
+	void Initialize();
+
 	// Callbacks
 	bool IsSimulating() const;
-	void FieldParameterUpdateCallback(Chaos::FPhysicsSolver* InSolver, FParticlesType& InParticles, Chaos::TArrayCollectionArray<float>& Strains, Chaos::TPBDPositionConstraints<float, 3>& PositionTarget, TMap<int32, int32>& PositionTargetedParticles, const TArray<FKinematicProxy>& AnimatedPositions, const float InTime);
-	void FieldForcesUpdateCallback(Chaos::FPhysicsSolver* InSolver, FParticlesType& Particles, Chaos::TArrayCollectionArray<FVector> & Force, Chaos::TArrayCollectionArray<FVector> & Torque, const float Time);
-	void EndFrameCallback(const float InDt);
+	/**
+	 * Services queued \c FFieldSystemCommand commands.
+	 *
+	 * Supported fields:
+	 *	* EFieldPhysicsType::Field_DynamicState
+	 *	* EFieldPhysicsType::Field_ActivateDisabled
+	 *	* EFieldPhysicsType::Field_ExternalClusterStrain (clustering)
+	 *	* EFieldPhysicsType::Field_Kill
+	 *	* EFieldPhysicsType::Field_LinearVelocity
+	 *	* EFieldPhysicsType::Field_AngularVelociy
+	 *	* EFieldPhysicsType::Field_SleepingThreshold
+	 *	* EFieldPhysicsType::Field_DisableThreshold
+	 *	* EFieldPhysicsType::Field_InternalClusterStrain (clustering)
+	 *	* EFieldPhysicsType::Field_CollisionGroup
+	 *	* EFieldPhysicsType::Field_PositionStatic
+	 *	* EFieldPhysicsType::Field_PositionTarget
+	 *	* EFieldPhysicsType::Field_PositionAnimated
+	 *	* EFieldPhysicsType::Field_DynamicConstraint
+	 */
+	void FieldParameterUpdateCallback(
+		Chaos::FPhysicsSolver* InSolver, 
+		FParticlesType& InParticles, 
+		Chaos::TArrayCollectionArray<float>& Strains, 
+		Chaos::TPBDPositionConstraints<float, 3>& PositionTarget, 
+		TMap<int32, int32>& PositionTargetedParticles, 
+		//const TArray<FKinematicProxy>& AnimatedPositions, 
+		const float InTime);
+
+	/**
+	 * Services queued \c FFieldSystemCommand commands.
+	 *
+	 * Supported fields:
+	 *	* EFieldPhysicsType::Field_LinearForce
+	 *	* EFieldPhysicsType::Field_AngularTorque
+	 */
+	void FieldForcesUpdateCallback(
+		Chaos::FPhysicsSolver* InSolver, 
+		FParticlesType& Particles, 
+		Chaos::TArrayCollectionArray<FVector> & Force, 
+		Chaos::TArrayCollectionArray<FVector> & Torque, 
+		const float Time);
+
 	void BufferCommand(Chaos::FPhysicsSolver* InSolver, const FFieldSystemCommand& InCommand);
 
 	// Inactive Callbacks
@@ -36,25 +78,44 @@ public:
 	void CreateRigidBodyCallback(FParticlesType& InOutParticles){}
 	void DisableCollisionsCallback(TSet<TTuple<int32, int32>>& InPairs){}
 	void AddForceCallback(FParticlesType& InParticles, const float InDt, const int32 InIndex){}
+	void EndFrameCallback(const float InDt) { check(false); } // never called
 
+	// Called by FPBDRigidsSolver::PushPhysicsState() on game thread.
 	Chaos::FParticleData* NewData() { return nullptr; }
-	void SyncBeforeDestroy(){}
-	void OnRemoveFromScene();
+	// Called by FPBDRigidsSolver::PushPhysicsState() on physics thread.
 	void PushToPhysicsState(const Chaos::FParticleData*) {};
+	// Called by FPBDRigidsSolver::PushPhysicsState() on game thread.
 	void ClearAccumulatedData() {}
-	void BufferPhysicsResults(){}
-	void FlipBuffer(){}
-	void PullFromPhysicsState(){}
+
+	void BufferPhysicsResults(){}	// Not called
+	void FlipBuffer(){}				// Not called
+	void PullFromPhysicsState(){}	// Not called
+
 	bool IsDirty() { return false; }
+	void SyncBeforeDestroy(){}
+	void OnRemoveFromScene();
+
 	EPhysicsProxyType ConcreteType() {return EPhysicsProxyType::FieldType;}
 
-	/*
-	* ContiguousIndices
-	*   Generates a mapping between the Position array and the results array. When EFieldResolutionType is set to Maximum the complete
-	*   particle mapping is provided from the Particles.X to Particles.Attribute, when Minimum is set only the ActiveIndices and the
-	*   direct children of the active clusters are set in the IndicesArray.
-	*/
-	static void ContiguousIndices(TArray<ContextIndex>& IndicesArray, const Chaos::FPhysicsSolver* RigidSolver, EFieldResolutionType ResolutionType, bool bForce);
+	/**
+	 * Generates a mapping between the Position array and the results array. 
+	 *
+	 * When \p ResolutionType is set to \c Maximum the complete particle mapping 
+	 * is provided from the \c Particles.X to \c Particles.Attribute. 
+	 * When \c Minimum is set only the ActiveIndices and the direct children of 
+	 * the active clusters are set in the \p IndicesArray.
+	 */
+	static void ContiguousIndices(
+		TArray<ContextIndex>& IndicesArray, 
+		const Chaos::FPhysicsSolver* RigidSolver, 
+		const EFieldResolutionType ResolutionType, 
+		const bool bForce = true);
+
+	static void GetParticleHandles(
+		TArray<Chaos::TGeometryParticleHandle<float,3>*>& Handles,
+		const Chaos::FPhysicsSolver* RigidSolver,
+		const EFieldResolutionType ResolutionType,
+		const bool bForce = true);
 
 private:
 
diff --git a/Engine/Source/Runtime/Experimental/FieldSystem/Source/FieldSystemCore/Public/Field/FieldSystem.h b/Engine/Source/Runtime/Experimental/FieldSystem/Source/FieldSystemCore/Public/Field/FieldSystem.h
--- a/Engine/Source/Runtime/Experimental/FieldSystem/Source/FieldSystemCore/Public/Field/FieldSystem.h
+++ b/Engine/Source/Runtime/Experimental/FieldSystem/Source/FieldSystemCore/Public/Field/FieldSystem.h
@@ -79,10 +79,13 @@ struct FIELDSYSTEMCORE_API ContextIndex
 {
 	ContextIndex(int32 InSample=INDEX_NONE, int32 InResult=INDEX_NONE)
 		: Sample(InSample)
-		, Result(InResult) {}
-
+		, Result(InResult) 
+	{}
 
-	static void ContiguousIndices(TArray<ContextIndex>& Array, int NumParticles, bool bForce = true)
+	static void ContiguousIndices(
+		TArray<ContextIndex>& Array, 
+		const int NumParticles, 
+		const bool bForce = true)
 	{
 		if (bForce)
 		{
@@ -95,12 +98,10 @@ struct FIELDSYSTEMCORE_API ContextIndex
 		}
 	}
 
-
 	int32 Sample;
 	int32 Result;
 };
 
-
 struct FIELDSYSTEMCORE_API FFieldContext
 {
 	typedef  TMap<FFieldSystemMetaData::EMetaType, TUniquePtr<FFieldSystemMetaData> > UniquePointerMap;
@@ -124,7 +125,15 @@ struct FIELDSYSTEMCORE_API FFieldContext
 		, MetaData(MetaDataIn)
 	{}
 
-	const TArrayView< ContextIndex >& SampleIndices;
+	//
+	// Ryan - TODO: This concept of having discreet sample data needs to change.  
+	// I think we'd be better off supplying lambda accessors which can be specialized 
+	// for each respective use case.  That means the method by which this data is 
+	// traversed also needs to change; possibly to some load balanced threaded iterator 
+	// or task based paradigm.
+
+	//const TArrayView<THandle::TTransientHandle*>& Handles;
+	const TArrayView<ContextIndex>& SampleIndices;
 	const TArrayView<FVector>& Samples;
 	PointerMap MetaData;
 };
@@ -234,19 +243,45 @@ public:
 		, RootNode(RootNodeIn)
 	{}
 
-
 	// Commands are copied when moved from the one thread to 
 	// another. This requires a full copy of all associated data. 
-	FFieldSystemCommand(const FFieldSystemCommand & CommandIn)
-		: TargetAttribute(CommandIn.RootNode ? CommandIn.TargetAttribute:"")
-		, RootNode(CommandIn.RootNode?CommandIn.RootNode->NewCopy():nullptr)
+	FFieldSystemCommand(const FFieldSystemCommand & Other)
+		: TargetAttribute(Other.RootNode ? Other.TargetAttribute:"")
+		, RootNode(Other.RootNode?Other.RootNode->NewCopy():nullptr)
 	{
-		for (const TPair<FFieldSystemMetaData::EMetaType, TUniquePtr<FFieldSystemMetaData>>& Meta : CommandIn.MetaData)
+		for (const TPair<FFieldSystemMetaData::EMetaType, TUniquePtr<FFieldSystemMetaData>>& Meta : Other.MetaData)
 		{
 			MetaData.Add(Meta.Key).Reset(Meta.Value->NewCopy());
 		}
 	}
 
+	bool HasMetaData(const FFieldSystemMetaData::EMetaType Key) const
+	{
+		return MetaData.Contains(Key) && GetMetaData(Key) != nullptr;
+	}
+
+	const TUniquePtr<FFieldSystemMetaData>& GetMetaData(
+		const FFieldSystemMetaData::EMetaType Key) const
+	{
+		return MetaData[Key];
+	}
+
+	template <class TMetaData>
+	const TMetaData* GetMetaDataAs(const FFieldSystemMetaData::EMetaType Key) const
+	{
+		return static_cast<const TMetaData*>(GetMetaData(Key).Get());
+	}
+
+	void SetMetaData(const FFieldSystemMetaData::EMetaType Key, TUniquePtr<FFieldSystemMetaData>&& Value)
+	{
+		MetaData.Add(Key, MoveTemp(Value));
+	}
+
+	void SetMetaData(const FFieldSystemMetaData::EMetaType Key, FFieldSystemMetaData* Value)
+	{
+		MetaData[Key].Reset(Value);
+	}
+
 	void Serialize(FArchive& Ar);
 	bool operator==(const FFieldSystemCommand&);
 	bool operator!=(const FFieldSystemCommand& Other) { return !this->operator==(Other); }
diff --git a/Engine/Source/Runtime/Experimental/FieldSystem/Source/FieldSystemCore/Public/Field/FieldSystemTypes.h b/Engine/Source/Runtime/Experimental/FieldSystem/Source/FieldSystemCore/Public/Field/FieldSystemTypes.h
--- a/Engine/Source/Runtime/Experimental/FieldSystem/Source/FieldSystemCore/Public/Field/FieldSystemTypes.h
+++ b/Engine/Source/Runtime/Experimental/FieldSystem/Source/FieldSystemCore/Public/Field/FieldSystemTypes.h
@@ -83,6 +83,7 @@ enum EFieldFalloffType
 UENUM(BlueprintType)
 enum EFieldPhysicsType
 {
+	Field_None						UMETA(Hidden),
 	Field_DynamicState				UMETA(DisplayName = "DynamicState"),
 	Field_LinearForce				UMETA(DisplayName = "LinearForce"),
 	Field_ExternalClusterStrain		UMETA(DisplayName = "ExternalClusterStrain"),
@@ -107,7 +108,7 @@ enum EFieldPhysicsType
 inline 
 FName FIELDSYSTEMCORE_API GetFieldPhysicsName(EFieldPhysicsType Type)
 {
-	switch (Type)
+	switch(Type)
 	{
 	case Field_DynamicState:
 		return "DynamicState";
@@ -145,6 +146,86 @@ FName FIELDSYSTEMCORE_API GetFieldPhysicsName(EFieldPhysicsType Type)
 	return "None";
 }
 
+inline
+EFieldPhysicsType FIELDSYSTEMCORE_API GetFieldPhysicsType(const FName& Name)
+{
+	if(Name == "DynamicState")
+	{
+		return Field_DynamicState;
+	}
+	else if(Name == "LinearForce")
+	{
+		return Field_LinearForce;
+	}
+	else if(Name == "ExternalClusterStrain")
+	{
+		return Field_ExternalClusterStrain;
+	}
+	else if(Name == "Kill")
+	{
+		return Field_Kill;
+	}
+	else if(Name == "LinearVelocity")
+	{
+		return Field_LinearVelocity;
+	}
+	else if (Name == "AngularVelocity")
+	{
+		return Field_AngularVelociy;
+	}
+	else if(Name == "AngularTorque")
+	{
+		return Field_AngularTorque;
+	}
+	else if(Name == "InternalClusterStrain")
+	{
+		return Field_InternalClusterStrain;
+	}
+	else if(Name == "DisableThreshold")
+	{
+		return Field_DisableThreshold;
+	}
+	else if(Name == "SleepingThreshold")
+	{
+		return Field_SleepingThreshold;
+	}
+	else if(Name == "PositionStatic")
+	{
+		return Field_PositionStatic;
+	}
+	else if(Name == "PositionAnimated")
+	{
+		return Field_PositionAnimated;
+	}
+	else if(Name == "PositionTarget")
+	{
+		return Field_PositionTarget;
+	}
+	else if(Name == "DynamicConstraint")
+	{
+		return Field_DynamicConstraint;
+	}
+	else if(Name == "CollisionGroup")
+	{
+		return Field_CollisionGroup;
+	}
+	else if(Name == "ActivateDisabled")
+	{
+		return Field_ActivateDisabled;
+	}
+	else if(Name == "None")
+	{
+		return Field_None;
+	}
+	else
+	{
+		check(false);
+	}
+
+	return Field_None;
+}
+
+
 /**
 *
 */
diff --git a/Engine/Source/Runtime/Experimental/FieldSystem/Source/FieldSystemEngine/Private/Field/FieldSystemComponent.cpp b/Engine/Source/Runtime/Experimental/FieldSystem/Source/FieldSystemEngine/Private/Field/FieldSystemComponent.cpp
--- a/Engine/Source/Runtime/Experimental/FieldSystem/Source/FieldSystemEngine/Private/Field/FieldSystemComponent.cpp
+++ b/Engine/Source/Runtime/Experimental/FieldSystem/Source/FieldSystemEngine/Private/Field/FieldSystemComponent.cpp
@@ -54,6 +54,14 @@ void UFieldSystemComponent::OnCreatePhysicsState()
 #endif
 
 		bHasPhysicsState = true;
+
+		if(FieldSystem)
+		{
+			for(FFieldSystemCommand& Cmd : FieldSystem->Commands)
+			{
+				DispatchCommand(Cmd);
+			}
+		}
 	}
 }
 
diff --git a/Engine/Source/Runtime/Experimental/GeometryCollectionEngine/Private/GeometryCollection/GeometryCollectionComponent.cpp b/Engine/Source/Runtime/Experimental/GeometryCollectionEngine/Private/GeometryCollection/GeometryCollectionComponent.cpp
--- a/Engine/Source/Runtime/Experimental/GeometryCollectionEngine/Private/GeometryCollection/GeometryCollectionComponent.cpp
+++ b/Engine/Source/Runtime/Experimental/GeometryCollectionEngine/Private/GeometryCollection/GeometryCollectionComponent.cpp
@@ -1155,6 +1155,7 @@ void UGeometryCollectionComponent::OnCreatePhysicsState()
 				InParams.InitialLinearVelocity = InitialLinearVelocity;
 				InParams.InitialAngularVelocity = InitialAngularVelocity;
 				InParams.bClearCache = true;
+				InParams.ObjectType = ObjectType;
 				InParams.PhysicalMaterial = MakeSerializable(ChaosMaterial);
 				InParams.CacheType = CacheParameters.CacheMode;
 				InParams.ReverseCacheBeginTime = CacheParameters.ReverseCacheBeginTime;
diff --git a/Engine/Source/Runtime/Experimental/GeometryCollectionSimulationCore/Public/GeometryCollection/GeometryCollectionSimulationCoreTypes.h b/Engine/Source/Runtime/Experimental/GeometryCollectionSimulationCore/Public/GeometryCollection/GeometryCollectionSimulationCoreTypes.h
--- a/Engine/Source/Runtime/Experimental/GeometryCollectionSimulationCore/Public/GeometryCollection/GeometryCollectionSimulationCoreTypes.h
+++ b/Engine/Source/Runtime/Experimental/GeometryCollectionSimulationCore/Public/GeometryCollection/GeometryCollectionSimulationCoreTypes.h
@@ -291,6 +291,8 @@ struct FSimulationParameters
 	float ReverseCacheBeginTime;
 	bool bClearCache;
 
+	EObjectStateTypeEnum ObjectType;
+
 	Chaos::TSerializablePtr<Chaos::FChaosPhysicsMaterial> PhysicalMaterial;
 
 	FCollisionDataSimulationParameters CollisionData;
