Fixed convex-convex collisions consuming every primitive collsiion that made it that far down the if/else chain and allow them to head into the correct resolution function. Restores functionality of primitive->convex collisions for Chaos

diff --git a/Engine/Source/Runtime/Experimental/Chaos/Private/Chaos/CollisionResolution.cpp b/Engine/Source/Runtime/Experimental/Chaos/Private/Chaos/CollisionResolution.cpp
--- a/Engine/Source/Runtime/Experimental/Chaos/Private/Chaos/CollisionResolution.cpp
+++ b/Engine/Source/Runtime/Experimental/Chaos/Private/Chaos/CollisionResolution.cpp
@@ -1018,7 +1018,7 @@ namespace Chaos
 			{
 				return UpdateUnionUnionConstraint<UpdateType>(Implicit0, Transform0, Implicit1, Transform1, Thickness, ConstraintBase);
 			}
-			else if (Implicit0.IsConvex() && Implicit1.IsConvex())
+			else if(Implicit0Type == FConvex::StaticType() && Implicit1Type == FConvex::StaticType())
 			{
 				UpdateConvexConvexConstraint(Implicit0, Transform0, Implicit1, Transform1, Thickness, *ConstraintBase.template As<TRigidBodyPointContactConstraint<T, d>>());
 			}
@@ -1153,7 +1153,7 @@ namespace Chaos
 				// Union-union creates multiple manifolds - we should never get here for this pair type. See ConstructConstraintsImpl and ConstructUnionUnionConstraints
 				ensure(false);
 			}
-			else if (Implicit0->IsConvex() && Implicit1->IsConvex())
+			else if (Implicit0Type == FConvex::StaticType() && Implicit1Type == FConvex::StaticType())
 			{
 				ConstructConvexConvexConstraints(Particle0, Particle1, Implicit0, Implicit1, Transform0, Transform1, Thickness, NewConstraints);
 			}
