Fixed AsDynamic particle calls in Chaos after an API change conflict

diff --git a/Engine/Source/Runtime/Experimental/ChaosSolvers/Private/PhysicsProxy/FieldSystemPhysicsProxy.cpp b/Engine/Source/Runtime/Experimental/ChaosSolvers/Private/PhysicsProxy/FieldSystemPhysicsProxy.cpp
--- a/Engine/Source/Runtime/Experimental/ChaosSolvers/Private/PhysicsProxy/FieldSystemPhysicsProxy.cpp
+++ b/Engine/Source/Runtime/Experimental/ChaosSolvers/Private/PhysicsProxy/FieldSystemPhysicsProxy.cpp
@@ -168,23 +168,23 @@ void FFieldSystemPhysicsProxy::FieldParameterUpdateCallback(
 						// we're just going to ignore non-dynamic particles.  This has the added
 						// benefit of not needing to deal with the floor, as it's pretty likely to
 						// not be dynamic.  Har.
-
-						if (Chaos::TPBDRigidParticleHandleImp<float, 3, true>* DynHandle = Handle->AsDynamic())
+						Chaos::TPBDRigidParticleHandle<float, 3>* RigidHandle = Handle->CastToRigidParticle();
+						if(RigidHandle && RigidHandle->ObjectState() == Chaos::EObjectStateType::Dynamic)
 						{
 							const int32 FieldState = DynamicStateView[i];
-							const EObjectStateType HandleState = DynHandle->ObjectState();
+							const EObjectStateType HandleState = RigidHandle->ObjectState();
 							if (FieldState == (int32)EObjectStateTypeEnum::Chaos_Object_Dynamic)
 							{
 								if ((HandleState == Chaos::EObjectStateType::Static ||
 									 HandleState == Chaos::EObjectStateType::Kinematic) &&
-									DynHandle->M() > FLT_EPSILON)
+									RigidHandle->M() > FLT_EPSILON)
 								{
-									DynHandle->SetObjectState(Chaos::EObjectStateType::Dynamic);
+									RigidHandle->SetObjectStateLowLevel(Chaos::EObjectStateType::Dynamic);
 									StateChanged = true;
 								}
 								else if (HandleState == Chaos::EObjectStateType::Sleeping)
 								{
-									DynHandle->SetObjectState(Chaos::EObjectStateType::Dynamic);
+									RigidHandle->SetObjectStateLowLevel(Chaos::EObjectStateType::Dynamic);
 									StateChanged = true;
 								}
 							}
@@ -192,9 +192,9 @@ void FFieldSystemPhysicsProxy::FieldParameterUpdateCallback(
 							{
 								if (HandleState == Chaos::EObjectStateType::Dynamic)
 								{
-									DynHandle->SetObjectState(Chaos::EObjectStateType::Kinematic);
-									DynHandle->SetV(Chaos::TVector<float, 3>(0));
-									DynHandle->SetW(Chaos::TVector<float, 3>(0));
+									RigidHandle->SetObjectStateLowLevel(Chaos::EObjectStateType::Kinematic);
+									RigidHandle->SetV(Chaos::TVector<float, 3>(0));
+									RigidHandle->SetW(Chaos::TVector<float, 3>(0));
 									StateChanged = true;
 								}
 							}
@@ -202,9 +202,9 @@ void FFieldSystemPhysicsProxy::FieldParameterUpdateCallback(
 							{
 								if (HandleState == Chaos::EObjectStateType::Dynamic)
 								{
-									DynHandle->SetObjectState(Chaos::EObjectStateType::Static);
-									DynHandle->SetV(Chaos::TVector<float, 3>(0));
-									DynHandle->SetW(Chaos::TVector<float, 3>(0));
+									RigidHandle->SetObjectStateLowLevel(Chaos::EObjectStateType::Static);
+									RigidHandle->SetV(Chaos::TVector<float, 3>(0));
+									RigidHandle->SetW(Chaos::TVector<float, 3>(0));
 									StateChanged = true;
 								}
 							}
@@ -212,7 +212,7 @@ void FFieldSystemPhysicsProxy::FieldParameterUpdateCallback(
 							{
 								if (HandleState == Chaos::EObjectStateType::Dynamic)
 								{
-									DynHandle->SetObjectState(Chaos::EObjectStateType::Sleeping);
+									RigidHandle->SetObjectStateLowLevel(Chaos::EObjectStateType::Sleeping);
 									StateChanged = true;
 								}
 							}
@@ -997,9 +997,10 @@ void FFieldSystemPhysicsProxy::FieldForcesUpdateCallback(
 						int32 i = 0;
 						for (Chaos::TGeometryParticleHandle<float, 3>* Handle : Handles)
 						{
-							if (Chaos::TPBDRigidParticleHandleImp<float, 3, true>* DynHandle = Handle->AsDynamic())
+							Chaos::TPBDRigidParticleHandle<float, 3>* RigidHandle = Handle->CastToRigidParticle();
+							if(RigidHandle && RigidHandle->ObjectState() == Chaos::EObjectStateType::Dynamic)
 							{
-								DynHandle->ExternalForce() += ForceView[i];
+								RigidHandle->ExternalForce() += ForceView[i];
 							}
 							++i;
 						}
@@ -1056,9 +1057,10 @@ void FFieldSystemPhysicsProxy::FieldForcesUpdateCallback(
 						int32 i = 0;
 						for (Chaos::TGeometryParticleHandle<float, 3>* Handle : Handles)
 						{
-							if (Chaos::TPBDRigidParticleHandleImp<float, 3, true>* DynHandle = Handle->AsDynamic())
+							Chaos::TPBDRigidParticleHandle<float, 3>* RigidHandle = Handle->CastToRigidParticle();
+							if(RigidHandle && RigidHandle->ObjectState() == Chaos::EObjectStateType::Dynamic)
 							{
-								DynHandle->ExternalTorque() += TorqueView[i];
+								RigidHandle->ExternalTorque() += TorqueView[i];
 							}
 							++i;
 						}
