Chaos - variable CoM support working (RBAN, integration, collisions and joints) - currently disabled.

diff --git a/Engine/Source/Runtime/Engine/Private/PhysicsEngine/ImmediatePhysics/ImmediatePhysicsChaos/ImmediatePhysicsActorHandle_Chaos.cpp b/Engine/Source/Runtime/Engine/Private/PhysicsEngine/ImmediatePhysics/ImmediatePhysicsChaos/ImmediatePhysicsActorHandle_Chaos.cpp
--- a/Engine/Source/Runtime/Engine/Private/PhysicsEngine/ImmediatePhysics/ImmediatePhysicsChaos/ImmediatePhysicsActorHandle_Chaos.cpp
+++ b/Engine/Source/Runtime/Engine/Private/PhysicsEngine/ImmediatePhysics/ImmediatePhysicsChaos/ImmediatePhysicsActorHandle_Chaos.cpp
@@ -5,6 +5,7 @@
 #include "Chaos/Box.h"
 #include "Chaos/Capsule.h"
 #include "Chaos/MassProperties.h"
+#include "Chaos/Particle/ParticleUtilities.h"
 #include "Chaos/ParticleHandle.h"
 #include "Chaos/PBDRigidParticles.h"
 #include "Chaos/PBDRigidsEvolutionGBF.h"
@@ -166,7 +167,11 @@ namespace ImmediatePhysics_Chaos
 		AddParams.Scale = Scale;
 		//AddParams.SimpleMaterial = SimpleMaterial;
 		//AddParams.ComplexMaterials = TArrayView<UPhysicalMaterial*>(ComplexMaterials);
+#if CHAOS_PARTICLE_ACTORTRANSFORM
+		AddParams.LocalTransform = FTransform::Identity;
+#else
 		AddParams.LocalTransform = Chaos::TRigidTransform<float, 3>(OutCoMTransform.GetRotation().Inverse() * -OutCoMTransform.GetTranslation(), OutCoMTransform.GetRotation().Inverse());
+#endif
 		AddParams.WorldTransform = BodyInstance->GetUnrealWorldTransform();
 		AddParams.Geometry = &BodySetup->AggGeom;
 #if WITH_PHYSX
@@ -218,8 +223,6 @@ namespace ImmediatePhysics_Chaos
 		TRigidTransform<float, 3> CoMTransform = TRigidTransform<float, 3>::Identity;
 		if (CreateGeometry(BodyInstance, FVector::OneVector, Mass, Inertia, CoMTransform, Geometry, Shapes))
 		{
-			ActorToCoMTransform = CoMTransform;
-
 			switch (ActorType)
 			{
 			case EActorType::StaticActor:
@@ -243,6 +246,8 @@ namespace ImmediatePhysics_Chaos
 				{
 					Kinematic->SetV(FVector::ZeroVector);
 					Kinematic->SetW(FVector::ZeroVector);
+					Kinematic->SetCenterOfMass(CoMTransform.GetTranslation());
+					Kinematic->SetRotationOfMass(CoMTransform.GetRotation());
 				}
 
 				auto* Dynamic = ParticleHandle->CastToRigidParticle();
@@ -296,16 +301,15 @@ namespace ImmediatePhysics_Chaos
 
 	void FActorHandle::SetWorldTransform(const FTransform& WorldTM)
 	{
-		FTransform ParticleTransform = ActorToCoMTransform * WorldTM;
+		using namespace Chaos;
 
-		ParticleHandle->SetX(ParticleTransform.GetTranslation());
-		ParticleHandle->SetR(ParticleTransform.GetRotation());
+		FParticleUtilities::SetActorWorldTransform(TGenericParticleHandle<FReal, 3>(ParticleHandle), WorldTM);
 
 		auto* Dynamic = ParticleHandle->CastToRigidParticle();
 		if(Dynamic && Dynamic->ObjectState() == Chaos::EObjectStateType::Dynamic)
 		{
-			Dynamic->SetP(Dynamic->X());
-			Dynamic->SetQ(Dynamic->R());
+			Dynamic->X() = Dynamic->P();
+			Dynamic->R() = Dynamic->Q();
 		}
 	}
 
@@ -335,9 +339,11 @@ namespace ImmediatePhysics_Chaos
 
 	void FActorHandle::SetKinematicTarget(const FTransform& WorldTM)
 	{
+		using namespace Chaos;
+
 		if (ensure(GetIsKinematic()))
 		{
-			FTransform ParticleTransform = ActorToCoMTransform * WorldTM;
+			FTransform ParticleTransform = FParticleUtilities::ActorWorldToParticleWorld(TGenericParticleHandle<FReal, 3>(ParticleHandle), WorldTM);
 			GetKinematicTarget().SetTargetMode(ParticleTransform);
 		}
 
@@ -359,8 +365,9 @@ namespace ImmediatePhysics_Chaos
 
 	FTransform FActorHandle::GetWorldTransform() const
 	{
-		FTransform ParticleTransform = FTransform(Handle()->R(), Handle()->X());
-		return ActorToCoMTransform.GetRelativeTransformReverse(ParticleTransform);
+		using namespace Chaos;
+
+		return FParticleUtilities::GetActorWorldTransform(TGenericParticleHandle<FReal, 3>(ParticleHandle));
 	}
 
 	void FActorHandle::SetLinearVelocity(const FVector& NewLinearVelocity)
@@ -526,9 +533,9 @@ namespace ImmediatePhysics_Chaos
 		return FLT_MAX;
 	}
 
-	const FTransform& FActorHandle::GetLocalCoMTransform() const
+	FTransform FActorHandle::GetLocalCoMTransform() const
 	{
-		return ActorToCoMTransform;
+		return FTransform(Handle()->RotationOfMass(), Handle()->CenterOfMass());
 	}
 
 	int32 FActorHandle::GetLevel() const
diff --git a/Engine/Source/Runtime/Engine/Private/PhysicsEngine/ImmediatePhysics/ImmediatePhysicsChaos/ImmediatePhysicsJointHandle_Chaos.cpp b/Engine/Source/Runtime/Engine/Private/PhysicsEngine/ImmediatePhysics/ImmediatePhysicsChaos/ImmediatePhysicsJointHandle_Chaos.cpp
--- a/Engine/Source/Runtime/Engine/Private/PhysicsEngine/ImmediatePhysics/ImmediatePhysicsChaos/ImmediatePhysicsJointHandle_Chaos.cpp
+++ b/Engine/Source/Runtime/Engine/Private/PhysicsEngine/ImmediatePhysics/ImmediatePhysicsChaos/ImmediatePhysicsJointHandle_Chaos.cpp
@@ -3,6 +3,7 @@
 #include "Physics/ImmediatePhysics/ImmediatePhysicsChaos/ImmediatePhysicsJointHandle_Chaos.h"
 #include "Physics/ImmediatePhysics/ImmediatePhysicsChaos/ImmediatePhysicsActorHandle_Chaos.h"
 
+#include "Chaos/Particle/ParticleUtilities.h"
 #include "Chaos/PBDJointConstraints.h"
 
 #include "PhysicsEngine/ConstraintInstance.h"
@@ -98,8 +99,8 @@ namespace ImmediatePhysics_Chaos
 		FPBDJointSettings ConstraintSettings;
 		ConstraintSettings.ConstraintFrames = 
 		{ 
-			ConstraintFrame1.GetRelativeTransform(Actor1->GetLocalCoMTransform()),
-			ConstraintFrame2.GetRelativeTransform(Actor2->GetLocalCoMTransform())
+			FParticleUtilities::ActorLocalToParticleLocal(TGenericParticleHandle<FReal, 3>(Actor1->GetParticle()), ConstraintFrame1),
+			FParticleUtilities::ActorLocalToParticleLocal(TGenericParticleHandle<FReal, 3>(Actor2->GetParticle()), ConstraintFrame2),
 		};
 		ConstraintSettings.Motion.Stiffness = ChaosImmediate_JointStiffness;
 		ConstraintSettings.Motion.LinearMotionTypes = 
diff --git a/Engine/Source/Runtime/Engine/Public/Physics/ImmediatePhysics/ImmediatePhysicsChaos/ImmediatePhysicsActorHandle_Chaos.h b/Engine/Source/Runtime/Engine/Public/Physics/ImmediatePhysics/ImmediatePhysicsChaos/ImmediatePhysicsActorHandle_Chaos.h
--- a/Engine/Source/Runtime/Engine/Public/Physics/ImmediatePhysics/ImmediatePhysicsChaos/ImmediatePhysicsActorHandle_Chaos.h
+++ b/Engine/Source/Runtime/Engine/Public/Physics/ImmediatePhysics/ImmediatePhysicsChaos/ImmediatePhysicsActorHandle_Chaos.h
@@ -105,7 +105,7 @@ namespace ImmediatePhysics_Chaos
 		float GetMaxContactImpulse() const;
 
 		/** Get the actor-space centre of mass offset */
-		const FTransform& GetLocalCoMTransform() const;
+		FTransform GetLocalCoMTransform() const;
 
 		Chaos::TGeometryParticleHandle<FReal, Dimensions>* GetParticle();
 		const Chaos::TGeometryParticleHandle<FReal, Dimensions>* GetParticle() const;
@@ -127,7 +127,6 @@ namespace ImmediatePhysics_Chaos
 		Chaos::TGeometryParticleHandle<FReal, Dimensions>* ParticleHandle;
 		TUniquePtr<Chaos::FImplicitObject> Geometry;
 		TArray<TUniquePtr<Chaos::TPerShapeData<float, 3>>> Shapes;
-		FTransform ActorToCoMTransform;
 		int32 Level;
 	};
 
diff --git a/Engine/Source/Runtime/Experimental/Chaos/Private/Chaos/ChaosDebugDraw.cpp b/Engine/Source/Runtime/Experimental/Chaos/Private/Chaos/ChaosDebugDraw.cpp
--- a/Engine/Source/Runtime/Experimental/Chaos/Private/Chaos/ChaosDebugDraw.cpp
+++ b/Engine/Source/Runtime/Experimental/Chaos/Private/Chaos/ChaosDebugDraw.cpp
@@ -6,6 +6,7 @@
 #include "Chaos/ImplicitObjectScaled.h"
 #include "Chaos/ImplicitObjectTransformed.h"
 #include "Chaos/ImplicitObjectUnion.h"
+#include "Chaos/Particle/ParticleUtilities.h"
 #include "Chaos/ParticleHandle.h"
 #include "Chaos/PBDCollisionConstraints.h"
 #include "Chaos/PBDJointConstraints.h"
@@ -18,8 +19,8 @@ namespace Chaos
 	namespace DebugDraw
 	{
 		float ArrowSize = 1.5f;
-		float BodyAxisLen = 12.0f;
-		float ContactLen = 6.0f;
+		float BodyAxisLen = 4.0f;
+		float ContactLen = 4.0f;
 		float ContactWidth = 2.0f;
 		float ContactPhiWidth = 1.5f;
 		float ConstraintAxisLen = 5.0f;
@@ -120,18 +121,25 @@ namespace Chaos
 			DrawShapesImpl(FRigidTransform3(P, Q), Particle->Geometry().Get(), Color);
 		}
 
-		void DrawParticleTransformImpl(const FRigidTransform3& SpaceTransform, const TGeometryParticleHandle<FReal, 3>* Particle, FReal ColorScale)
+		void DrawParticleTransformImpl(const FRigidTransform3& SpaceTransform, const TGeometryParticleHandle<FReal, 3>* InParticle, FReal ColorScale)
 		{
-			FColor R = (ColorScale * FColor::Red).ToFColor(false);
-			FColor G = (ColorScale * FColor::Green).ToFColor(false);
-			FColor B = (ColorScale * FColor::Blue).ToFColor(false);
+			FColor Red = (ColorScale * FColor::Red).ToFColor(false);
+			FColor Green = (ColorScale * FColor::Green).ToFColor(false);
+			FColor Blue = (ColorScale * FColor::Blue).ToFColor(false);
+
+			TConstGenericParticleHandle<FReal, 3> Particle(InParticle);
+			FVec3 PCOM = SpaceTransform.TransformPosition(FParticleUtilities::GetCoMWorldPosition(Particle));
+			FRotation3 QCOM = SpaceTransform.GetRotation() * FParticleUtilities::GetCoMWorldRotation(Particle);
+			FMatrix33 QCOMm = QCOM.ToMatrix();
+			FDebugDrawQueue::GetInstance().DrawDebugDirectionalArrow(PCOM, PCOM + DrawScale * BodyAxisLen * QCOMm.GetAxis(0), DrawScale * ArrowSize, Red, false, KINDA_SMALL_NUMBER, DrawPriority, LineThickness);
+			FDebugDrawQueue::GetInstance().DrawDebugDirectionalArrow(PCOM, PCOM + DrawScale * BodyAxisLen * QCOMm.GetAxis(1), DrawScale * ArrowSize, Green, false, KINDA_SMALL_NUMBER, DrawPriority, LineThickness);
+			FDebugDrawQueue::GetInstance().DrawDebugDirectionalArrow(PCOM, PCOM + DrawScale * BodyAxisLen * QCOMm.GetAxis(2), DrawScale * ArrowSize, Blue, false, KINDA_SMALL_NUMBER, DrawPriority, LineThickness);
 
-			FVec3 P = SpaceTransform.TransformPosition((Particle->ObjectState() == EObjectStateType::Dynamic) ? Particle->CastToRigidParticle()->P() : Particle->X());
-			FRotation3 Q = SpaceTransform.GetRotation() * ((Particle->ObjectState() == EObjectStateType::Dynamic) ? Particle->CastToRigidParticle()->Q() : Particle->R());
-			FMatrix33 Qm = Q.ToMatrix();
-			FDebugDrawQueue::GetInstance().DrawDebugDirectionalArrow(P, P + DrawScale * BodyAxisLen * Qm.GetAxis(0), DrawScale * ArrowSize, R, false, KINDA_SMALL_NUMBER, DrawPriority, LineThickness);
-			FDebugDrawQueue::GetInstance().DrawDebugDirectionalArrow(P, P + DrawScale * BodyAxisLen * Qm.GetAxis(1), DrawScale * ArrowSize, G, false, KINDA_SMALL_NUMBER, DrawPriority, LineThickness);
-			FDebugDrawQueue::GetInstance().DrawDebugDirectionalArrow(P, P + DrawScale * BodyAxisLen * Qm.GetAxis(2), DrawScale * ArrowSize, B, false, KINDA_SMALL_NUMBER, DrawPriority, LineThickness);
+			FColor Black = FColor::Black;
+			FColor Grey = (ColorScale * FColor(64, 64, 64)).ToFColor(false);
+			FVec3 PActor = SpaceTransform.TransformPosition(FParticleUtilities::GetActorWorldTransform(Particle).GetTranslation());
+			FDebugDrawQueue::GetInstance().DrawDebugPoint(PActor, Black, false, KINDA_SMALL_NUMBER, DrawPriority, DrawScale * PointSize);
+			FDebugDrawQueue::GetInstance().DrawDebugLine(PCOM, PActor, Grey, false, KINDA_SMALL_NUMBER, DrawPriority, LineThickness);
 		}
 
 		void DrawCollisionImpl(const FRigidTransform3& SpaceTransform, const TPBDCollisionConstraintHandle<float, 3>* ConstraintHandle, float ColorScale)
@@ -213,8 +221,8 @@ namespace Chaos
 			auto RigidParticle1 = ConstrainedParticles[1]->CastToRigidParticle();
 			if ((RigidParticle0 && RigidParticle0->ObjectState() == EObjectStateType::Dynamic) || (RigidParticle1 && RigidParticle1->ObjectState() == EObjectStateType::Dynamic))
 			{
-				FVec3 Pa = TGenericParticleHandle<FReal, 3>(ConstraintHandle->GetConstrainedParticles()[1])->P();
-				FVec3 Pb = TGenericParticleHandle<FReal, 3>(ConstraintHandle->GetConstrainedParticles()[0])->P();
+				FVec3 Pa = FParticleUtilities::GetCoMWorldPosition(TConstGenericParticleHandle<FReal, 3>(ConstraintHandle->GetConstrainedParticles()[1]));
+				FVec3 Pb = FParticleUtilities::GetCoMWorldPosition(TConstGenericParticleHandle<FReal, 3>(ConstraintHandle->GetConstrainedParticles()[0]));
 				FVec3 Xa, Xb, CR;
 				FMatrix33 Ra, Rb;
 				ConstraintHandle->CalculateConstraintSpace(Xa, Ra, Xb, Rb, CR);
diff --git a/Engine/Source/Runtime/Experimental/Chaos/Private/Chaos/PBDCollisionConstraintsPointContactUtil.cpp b/Engine/Source/Runtime/Experimental/Chaos/Private/Chaos/PBDCollisionConstraintsPointContactUtil.cpp
--- a/Engine/Source/Runtime/Experimental/Chaos/Private/Chaos/PBDCollisionConstraintsPointContactUtil.cpp
+++ b/Engine/Source/Runtime/Experimental/Chaos/Private/Chaos/PBDCollisionConstraintsPointContactUtil.cpp
@@ -3,6 +3,7 @@
 #include "Chaos/PBDCollisionConstraintsPointContactUtil.h"
 #include "Chaos/CollisionResolution.h"
 #include "Chaos/Defines.h"
+#include "Chaos/Particle/ParticleUtilities.h"
 #include "Chaos/Utilities.h"
 
 namespace Chaos
@@ -68,23 +69,18 @@ namespace Chaos
 					PhysicsMaterial1 = Particle1->AuxilaryValue(*ParticleParameters.PhysicsMaterials);
 				}
 
-				// @todo(ccaulfield): CHAOS_PARTICLEHANDLE_TODO split function to avoid ifs
 				const TVector<T, d> ZeroVector = TVector<T, d>(0);
-				const TRotation<T, d>& Q0 = Particle0->Q();
-				const TRotation<T, d>& Q1 = Particle1->Q();
-				const TVector<T, d>& P0 = Particle0->P();
-				const TVector<T, d>& P1 = Particle1->P();
-				const TVector<T, d>& V0 = Particle0->V();
-				const TVector<T, d>& V1 = Particle1->V();
-				const TVector<T, d>& W0 = Particle0->W();
-				const TVector<T, d>& W1 = Particle1->W();
+				TVector<T, d> P0 = FParticleUtilities::GetCoMWorldPosition(Particle0);
+				TVector<T, d> P1 = FParticleUtilities::GetCoMWorldPosition(Particle1);
+				TRotation<T, d> Q0 = FParticleUtilities::GetCoMWorldRotation(Particle0);
+				TRotation<T, d> Q1 = FParticleUtilities::GetCoMWorldRotation(Particle1);
 
 				typename TRigidBodyPointContactConstraint<T, d>::FManifold & Contact = Constraint.Manifold;
 
 				TVector<T, d> VectorToPoint1 = Contact.Location - P0;
 				TVector<T, d> VectorToPoint2 = Contact.Location - P1;
-				TVector<T, d> Body1Velocity = V0 + TVector<T, d>::CrossProduct(W0, VectorToPoint1);
-				TVector<T, d> Body2Velocity = V1 + TVector<T, d>::CrossProduct(W1, VectorToPoint2);
+				TVector<T, d> Body1Velocity = FParticleUtilities::GetVelocityAtCoMRelativePosition(Particle0, VectorToPoint1);
+				TVector<T, d> Body2Velocity = FParticleUtilities::GetVelocityAtCoMRelativePosition(Particle1, VectorToPoint2);
 				TVector<T, d> RelativeVelocity = Body1Velocity - Body2Velocity;
 				T RelativeNormalVelocity = TVector<T, d>::DotProduct(RelativeVelocity, Contact.Normal);
 
@@ -155,7 +151,7 @@ namespace Chaos
 							Impulse = MinimalImpulse;
 							if (AngularFriction)
 							{
-								TVector<T, d> RelativeAngularVelocity = W0 - W1;
+								TVector<T, d> RelativeAngularVelocity = Particle0->W() - Particle1->W();
 								T AngularNormal = TVector<T, d>::DotProduct(RelativeAngularVelocity, Contact.Normal);
 								TVector<T, d> AngularTangent = RelativeAngularVelocity - AngularNormal * Contact.Normal;
 								TVector<T, d> FinalAngularVelocity = FMath::Sign(AngularNormal) * FMath::Max((T)0, FMath::Abs(AngularNormal) - AngularFriction * NormalVelocityChange) * Contact.Normal + FMath::Max((T)0, AngularTangent.Size() - AngularFriction * NormalVelocityChange) * AngularTangent.GetSafeNormal();
@@ -237,9 +233,10 @@ namespace Chaos
 						PBDRigid0->V() += DV;
 						PBDRigid0->W() += DW;
 						// Position update as part of pbd
-						PBDRigid0->P() += DV * IterationParameters.Dt;
-						PBDRigid0->Q() += FRotation3::FromElements(DW, 0.f) * Q0 * IterationParameters.Dt * T(0.5);
-						PBDRigid0->Q().Normalize();
+						P0 += (DV * IterationParameters.Dt);
+						Q0 += FRotation3::FromElements(DW, 0.f) * Q0 * IterationParameters.Dt * T(0.5);
+						Q0.Normalize();
+						FParticleUtilities::SetCoMWorldTransform(PBDRigid0, P0, Q0);
 					}
 					if (bIsRigidDynamic1)
 					{
@@ -250,36 +247,34 @@ namespace Chaos
 						PBDRigid1->V() += DV;
 						PBDRigid1->W() += DW;
 						// Position update as part of pbd
-						PBDRigid1->P() += DV * IterationParameters.Dt;
-						PBDRigid1->Q() += FRotation3::FromElements(DW, 0.f) * Q1 * IterationParameters.Dt * T(0.5);
-						PBDRigid1->Q().Normalize();
+						P1 += (DV * IterationParameters.Dt);
+						Q1 += FRotation3::FromElements(DW, 0.f) * Q1 * IterationParameters.Dt * T(0.5);
+						Q1.Normalize();
+						FParticleUtilities::SetCoMWorldTransform(PBDRigid1, P1, Q1);
 					}
 				}
 			}
 		}
 
+		// @todo(ccaulfield): Kinematic-Dynamic optimized version
 		template<typename T, int d>
 		void ApplyPushOut(TRigidBodyPointContactConstraint<T, d>& Constraint, T Thickness, const TSet<const TGeometryParticleHandle<T, d>*>& IsTemporarilyStatic,
 			TPointContactIterationParameters<T> & IterationParameters, TPointContactParticleParameters<T> & ParticleParameters)
 		{
-			TGeometryParticleHandle<T, d>* Particle0 = Constraint.Particle[0];
-			TGeometryParticleHandle<T, d>* Particle1 = Constraint.Particle[1];
+			TGenericParticleHandle<T, d> Particle0 = TGenericParticleHandle<T, d>(Constraint.Particle[0]);
+			TGenericParticleHandle<T, d> Particle1 = TGenericParticleHandle<T, d>(Constraint.Particle[1]);
 			TPBDRigidParticleHandle<T, d>* PBDRigid0 = Particle0->CastToRigidParticle();
 			TPBDRigidParticleHandle<T, d>* PBDRigid1 = Particle1->CastToRigidParticle();
 			const bool bIsRigidDynamic0 = PBDRigid0 && PBDRigid0->ObjectState() == EObjectStateType::Dynamic;
 			const bool bIsRigidDynamic1 = PBDRigid1 && PBDRigid1->ObjectState() == EObjectStateType::Dynamic;
 
 			const TVector<T, d> ZeroVector = TVector<T, d>(0);
-			const TRotation<T, d>& Q0 = bIsRigidDynamic0 ? PBDRigid0->Q() : Particle0->R();
-			const TRotation<T, d>& Q1 = bIsRigidDynamic1 ? PBDRigid1->Q() : Particle1->R();
-			const TVector<T, d>& P0 = bIsRigidDynamic0 ? PBDRigid0->P() : Particle0->X();
-			const TVector<T, d>& P1 = bIsRigidDynamic1 ? PBDRigid1->P() : Particle1->X();
-			const TVector<T, d>& V0 = bIsRigidDynamic0 ? PBDRigid0->V() : ZeroVector;
-			const TVector<T, d>& V1 = bIsRigidDynamic1 ? PBDRigid1->V() : ZeroVector;
-			const TVector<T, d>& W0 = bIsRigidDynamic0 ? PBDRigid0->W() : ZeroVector;
-			const TVector<T, d>& W1 = bIsRigidDynamic1 ? PBDRigid1->W() : ZeroVector;
-			const bool IsTemporarilyStatic0 = IsTemporarilyStatic.Contains(Particle0);
-			const bool IsTemporarilyStatic1 = IsTemporarilyStatic.Contains(Particle1);
+			TVector<T, d> P0 = FParticleUtilities::GetCoMWorldPosition(Particle0);
+			TVector<T, d> P1 = FParticleUtilities::GetCoMWorldPosition(Particle1);
+			TRotation<T, d> Q0 = FParticleUtilities::GetCoMWorldRotation(Particle0);
+			TRotation<T, d> Q1 = FParticleUtilities::GetCoMWorldRotation(Particle1);
+			const bool IsTemporarilyStatic0 = IsTemporarilyStatic.Contains(Constraint.Particle[0]);
+			const bool IsTemporarilyStatic1 = IsTemporarilyStatic.Contains(Constraint.Particle[1]);
 	
 			TSerializablePtr<FChaosPhysicsMaterial> PhysicsMaterial0, PhysicsMaterial1;
 			if (ParticleParameters.PhysicsMaterials)
@@ -316,8 +311,8 @@ namespace Chaos
 				T ScalingFactor = Numerator / (T)IterationParameters.NumIterations;
 
 				//if pushout is needed we better fix relative velocity along normal. Treat it as if 0 restitution
-				TVector<T, d> Body1Velocity = V0 + TVector<T, d>::CrossProduct(W0, VectorToPoint1);
-				TVector<T, d> Body2Velocity = V1 + TVector<T, d>::CrossProduct(W1, VectorToPoint2);
+				TVector<T, d> Body1Velocity = FParticleUtilities::GetVelocityAtCoMRelativePosition(Particle0, VectorToPoint1);
+				TVector<T, d> Body2Velocity = FParticleUtilities::GetVelocityAtCoMRelativePosition(Particle1, VectorToPoint2);
 				TVector<T, d> RelativeVelocity = Body1Velocity - Body2Velocity;
 				const T RelativeVelocityDotNormal = TVector<T, d>::DotProduct(RelativeVelocity, Contact.Normal);
 				if (RelativeVelocityDotNormal < 0)
@@ -359,15 +354,17 @@ namespace Chaos
 				TVector<T, d> AngularImpulse2 = TVector<T, d>::CrossProduct(VectorToPoint2, -Impulse);
 				if (!IsTemporarilyStatic0 && bIsRigidDynamic0)
 				{
-					PBDRigid0->P() += PBDRigid0->InvM() * Impulse;
-					PBDRigid0->Q() = TRotation<T, d>::FromVector(WorldSpaceInvI1 * AngularImpulse1) * Q0;
-					PBDRigid0->Q().Normalize();
+					P0 += PBDRigid0->InvM() * Impulse;
+					Q0 = TRotation<T, d>::FromVector(WorldSpaceInvI1 * AngularImpulse1) * Q0;
+					Q0.Normalize();
+					FParticleUtilities::SetCoMWorldTransform(Particle0, P0, Q0);
 				}
 				if (!IsTemporarilyStatic1 && bIsRigidDynamic1)
 				{
-					PBDRigid1->P() -= PBDRigid1->InvM() * Impulse;
-					PBDRigid1->Q() = TRotation<T, d>::FromVector(WorldSpaceInvI2 * AngularImpulse2) * Q1;
-					PBDRigid1->Q().Normalize();
+					P1 -= PBDRigid1->InvM() * Impulse;
+					Q1 = TRotation<T, d>::FromVector(WorldSpaceInvI2 * AngularImpulse2) * Q1;
+					Q1.Normalize();
+					FParticleUtilities::SetCoMWorldTransform(Particle1, P1, Q1);
 				}
 			}
 		}
diff --git a/Engine/Source/Runtime/Experimental/Chaos/Private/Chaos/PBDJointConstraintUtilities.cpp b/Engine/Source/Runtime/Experimental/Chaos/Private/Chaos/PBDJointConstraintUtilities.cpp
--- a/Engine/Source/Runtime/Experimental/Chaos/Private/Chaos/PBDJointConstraintUtilities.cpp
+++ b/Engine/Source/Runtime/Experimental/Chaos/Private/Chaos/PBDJointConstraintUtilities.cpp
@@ -355,8 +355,6 @@ namespace Chaos
 		const FPBDJointSolverSettings& SolverSettings,
 		const FPBDJointSettings& JointSettings,
 		const FReal Stiffness,
-		const int32 Index0,
-		const int32 Index1,
 		FVec3& P0,
 		FVec3& V0,
 		FVec3& P1,
@@ -379,8 +377,6 @@ namespace Chaos
 		const FPBDJointSolverSettings& SolverSettings,
 		const FPBDJointSettings& JointSettings,
 		const FReal Stiffness,
-		const int32 Index0,
-		const int32 Index1,
 		FVec3& P0,
 		FVec3& V0,
 		FVec3& P1,
@@ -400,8 +396,6 @@ namespace Chaos
 		const FPBDJointSolverSettings& SolverSettings,
 		const FPBDJointSettings& JointSettings,
 		const FReal Stiffness,
-		const int32 Index0,
-		const int32 Index1,
 		FRotation3& Q0,
 		FVec3& W0,
 		FRotation3& Q1,
@@ -429,8 +423,6 @@ namespace Chaos
 		const FPBDJointSolverSettings& SolverSettings,
 		const FPBDJointSettings& JointSettings,
 		const FReal Stiffness,
-		const int32 Index0,
-		const int32 Index1,
 		FRotation3& Q0,
 		FVec3& W0,
 		FRotation3& Q1,
@@ -455,8 +447,6 @@ namespace Chaos
 		const FPBDJointSolverSettings& SolverSettings,
 		const FPBDJointSettings& JointSettings,
 		const FReal Stiffness,
-		const int32 Index0,
-		const int32 Index1,
 		FRotation3& Q0,
 		FVec3& W0,
 		FRotation3& Q1,
@@ -479,7 +469,7 @@ namespace Chaos
 		const FVec3 DW0 = Utilities::Multiply(InvI0, Axis0) * L * WC;
 		const FVec3 DW1 = -Utilities::Multiply(InvI1, Axis1) * L * WC;
 
-		ApplyRotationVelocityDelta(Dt, SolverSettings, JointSettings, Stiffness, Index0, Index1, Q0, W0, Q1, W1, DW0, DW1);
+		ApplyRotationVelocityDelta(Dt, SolverSettings, JointSettings, Stiffness, Q0, W0, Q1, W1, DW0, DW1);
 	}
 
 	
@@ -488,8 +478,6 @@ namespace Chaos
 		const FPBDJointSolverSettings& SolverSettings,
 		const FPBDJointSettings& JointSettings,
 		const FReal Stiffness,
-		const int32 Index0,
-		const int32 Index1,
 		FRotation3& Q0,
 		FVec3& W0,
 		FRotation3& Q1,
@@ -514,27 +502,25 @@ namespace Chaos
 		const FVec3 DR0 = Utilities::Multiply(InvI0, Axis0) * L * Angle0;
 		const FVec3 DR1 = Utilities::Multiply(InvI1, Axis1) * L * Angle1;
 
-		ApplyRotationDelta(Dt, SolverSettings, JointSettings, Stiffness, Index0, Index1, Q0, W0, Q1, W1, DR0, DR1);
+		ApplyRotationDelta(Dt, SolverSettings, JointSettings, Stiffness, Q0, W0, Q1, W1, DR0, DR1);
 	}
 
 	
 	void FPBDJointUtilities::CalculateSwingConstraintSpace(
 		const FPBDJointSolverSettings& SolverSettings,
 		const FPBDJointSettings& JointSettings,
-		const int32 Index0,
-		const int32 Index1,
-		FVec3& P0,
-		FRotation3& Q0,
-		FVec3& P1,
-		FRotation3& Q1,
+		const FRigidTransform3& XL0,
+		const FRigidTransform3& XL1,
+		const FVec3& P0,
+		const FRotation3& Q0,
+		const FVec3& P1,
+		const FRotation3& Q1,
 		FVec3& OutX0,
 		FMatrix33& OutR0,
 		FVec3& OutX1,
 		FMatrix33& OutR1,
 		FVec3& OutCR)
 	{
-		const FRigidTransform3& XL0 = JointSettings.ConstraintFrames[Index0];
-		const FRigidTransform3& XL1 = JointSettings.ConstraintFrames[Index1];
 		const FVec3 X0 = P0 + Q0 * XL0.GetTranslation();
 		const FVec3 X1 = P1 + Q1 * XL1.GetTranslation();
 		const FRotation3 R0 = Q0 * XL0.GetRotation();
@@ -600,20 +586,18 @@ namespace Chaos
 	void FPBDJointUtilities::CalculateConeConstraintSpace(
 		const FPBDJointSolverSettings& SolverSettings,
 		const FPBDJointSettings& JointSettings,
-		const int32 Index0,
-		const int32 Index1,
-		FVec3& P0,
-		FRotation3& Q0,
-		FVec3& P1,
-		FRotation3& Q1,
+		const FRigidTransform3& XL0,
+		const FRigidTransform3& XL1,
+		const FVec3& P0,
+		const FRotation3& Q0,
+		const FVec3& P1,
+		const FRotation3& Q1,
 		FVec3& OutX0,
 		FMatrix33& OutR0, 
 		FVec3& OutX1, 
 		FMatrix33& OutR1, 
 		FVec3& OutCR)
 	{
-		const FRigidTransform3& XL0 = JointSettings.ConstraintFrames[Index0];
-		const FRigidTransform3& XL1 = JointSettings.ConstraintFrames[Index1];
 		const FVec3 X0 = P0 + Q0 * XL0.GetTranslation();
 		const FVec3 X1 = P1 + Q1 * XL1.GetTranslation();
 		const FRotation3 R0 = Q0 * XL0.GetRotation();
@@ -669,8 +653,8 @@ namespace Chaos
 		const FPBDJointSolverSettings& SolverSettings,
 		const FPBDJointSettings& JointSettings,
 		const FReal Stiffness,
-		const int32 Index0,
-		const int32 Index1,
+		const FRigidTransform3& XL0,
+		const FRigidTransform3& XL1,
 		FVec3& P0,
 		FRotation3& Q0,
 		FVec3& V0,
@@ -685,8 +669,6 @@ namespace Chaos
 		const FMatrix33& InvIL1)
 	{
 
-		const FRigidTransform3& XL0 = JointSettings.ConstraintFrames[Index0];
-		const FRigidTransform3& XL1 = JointSettings.ConstraintFrames[Index1];
 		const FVec3 X0 = P0 + Q0 * XL0.GetTranslation();
 		const FVec3 X1 = P1 + Q1 * XL1.GetTranslation();
 		const FRotation3 R0 = Q0 * XL0.GetRotation();
@@ -716,8 +698,8 @@ namespace Chaos
 		const FVec3 DR0 = Utilities::Multiply(InvI0, FVec3::CrossProduct(X0 - P0, DX));
 		const FVec3 DR1 = Utilities::Multiply(InvI1, FVec3::CrossProduct(X1 - P1, -DX));
 
-		ApplyPositionDelta(Dt, SolverSettings, JointSettings, Stiffness, Index0, Index1, P0, V0, P1, V1, DP0, DP1);
-		ApplyRotationDelta(Dt, SolverSettings, JointSettings, Stiffness, Index0, Index1, Q0, W0, Q1, W1, DR0, DR1);
+		ApplyPositionDelta(Dt, SolverSettings, JointSettings, Stiffness, P0, V0, P1, V1, DP0, DP1);
+		ApplyRotationDelta(Dt, SolverSettings, JointSettings, Stiffness, Q0, W0, Q1, W1, DR0, DR1);
 	}
 
 	
@@ -726,8 +708,8 @@ namespace Chaos
 		const FPBDJointSolverSettings& SolverSettings,
 		const FPBDJointSettings& JointSettings,
 		const FReal Stiffness,
-		const int32 Index0,
-		const int32 Index1,
+		const FRigidTransform3& XL0,
+		const FRigidTransform3& XL1,
 		FVec3& P0,
 		FRotation3& Q0,
 		FVec3& V0,
@@ -741,8 +723,6 @@ namespace Chaos
 		FReal InvM1,
 		const FMatrix33& InvIL1)
 	{
-		const FRigidTransform3& XL0 = JointSettings.ConstraintFrames[Index0];
-		const FRigidTransform3& XL1 = JointSettings.ConstraintFrames[Index1];
 		const FVec3 XC0 = Q0 * XL0.GetTranslation();
 		const FVec3 XC1 = Q1 * XL1.GetTranslation();
 		const FRotation3 R0 = Q0 * XL0.GetRotation();
@@ -773,8 +753,8 @@ namespace Chaos
 		const FVec3 DW0 = Utilities::Multiply(InvI0, FVec3::CrossProduct(XC0, DL));
 		const FVec3 DW1 = -Utilities::Multiply(InvI1, FVec3::CrossProduct(XC1, DL));
 
-		ApplyVelocityDelta(Dt, SolverSettings, JointSettings, Stiffness, Index0, Index1, P0, V0, P1, V1, DV0, DV1);
-		ApplyRotationVelocityDelta(Dt, SolverSettings, JointSettings, Stiffness, Index0, Index1, Q0, W0, Q1, W1, DW0, DW1);
+		ApplyVelocityDelta(Dt, SolverSettings, JointSettings, Stiffness, P0, V0, P1, V1, DV0, DV1);
+		ApplyRotationVelocityDelta(Dt, SolverSettings, JointSettings, Stiffness, Q0, W0, Q1, W1, DW0, DW1);
 	}
 
 
@@ -784,8 +764,8 @@ namespace Chaos
 		const FPBDJointSolverSettings& SolverSettings,
 		const FPBDJointSettings& JointSettings,
 		const FReal Stiffness,
-		const int32 Index0,
-		const int32 Index1,
+		const FRigidTransform3& XL0,
+		const FRigidTransform3& XL1,
 		FVec3& P0,
 		FRotation3& Q0,
 		FVec3& V0,
@@ -799,8 +779,6 @@ namespace Chaos
 		FReal InvM1,
 		const FMatrix33& InvIL1)
 	{
-		const FRigidTransform3& XL0 = JointSettings.ConstraintFrames[Index0];
-		const FRigidTransform3& XL1 = JointSettings.ConstraintFrames[Index1];
 		const FVec3 X0 = P0 + Q0 * XL0.GetTranslation();
 		const FVec3 X1 = P1 + Q1 * XL1.GetTranslation();
 		const FRotation3 R0 = Q0 * XL0.GetRotation();
@@ -853,7 +831,7 @@ namespace Chaos
 		const FReal DTwistAngle1 = -DTwistAngle;
 
 		// Apply twist correction
-		ApplyRotationDelta(Dt, SolverSettings, JointSettings, Stiffness, Index0, Index1, Q0, W0, Q1, W1, InvM0, InvIL0, InvM1, InvIL1, TwistAxis0, DTwistAngle0, TwistAxis1, DTwistAngle1);
+		ApplyRotationDelta(Dt, SolverSettings, JointSettings, Stiffness, Q0, W0, Q1, W1, InvM0, InvIL0, InvM1, InvIL1, TwistAxis0, DTwistAngle0, TwistAxis1, DTwistAngle1);
 	}
 
 	
@@ -862,8 +840,8 @@ namespace Chaos
 		const FPBDJointSolverSettings& SolverSettings,
 		const FPBDJointSettings& JointSettings,
 		const FReal Stiffness,
-		const int32 Index0,
-		const int32 Index1,
+		const FRigidTransform3& XL0,
+		const FRigidTransform3& XL1,
 		FVec3& P0,
 		FRotation3& Q0,
 		FVec3& V0,
@@ -877,8 +855,6 @@ namespace Chaos
 		FReal InvM1,
 		const FMatrix33& InvIL1)
 	{
-		const FRigidTransform3& XL0 = JointSettings.ConstraintFrames[Index0];
-		const FRigidTransform3& XL1 = JointSettings.ConstraintFrames[Index1];
 		const FVec3 X0 = P0 + Q0 * XL0.GetTranslation();
 		const FVec3 X1 = P1 + Q1 * XL1.GetTranslation();
 		const FRotation3 R0 = Q0 * XL0.GetRotation();
@@ -930,7 +906,7 @@ namespace Chaos
 			DW = FMath::Min((FReal)0, WC1 - WC0);
 		}
 
-		ApplyRotationVelocityDelta(Dt, SolverSettings, JointSettings, Stiffness, Index0, Index1, Q0, W0, Q1, W1, InvM0, InvIL0, InvM1, InvIL1, TwistAxis0, TwistAxis1, DW);
+		ApplyRotationVelocityDelta(Dt, SolverSettings, JointSettings, Stiffness, Q0, W0, Q1, W1, InvM0, InvIL0, InvM1, InvIL1, TwistAxis0, TwistAxis1, DW);
 	}
 
 
@@ -940,8 +916,8 @@ namespace Chaos
 		const FPBDJointSolverSettings& SolverSettings,
 		const FPBDJointSettings& JointSettings,
 		const FReal Stiffness,
-		const int32 Index0,
-		const int32 Index1,
+		const FRigidTransform3& XL0,
+		const FRigidTransform3& XL1,
 		FVec3& P0,
 		FRotation3& Q0,
 		FVec3& V0,
@@ -955,8 +931,6 @@ namespace Chaos
 		FReal InvM1,
 		const FMatrix33& InvIL1)
 	{
-		const FRigidTransform3& XL0 = JointSettings.ConstraintFrames[Index0];
-		const FRigidTransform3& XL1 = JointSettings.ConstraintFrames[Index1];
 		const FVec3 X0 = P0 + Q0 * XL0.GetTranslation();
 		const FVec3 X1 = P1 + Q1 * XL1.GetTranslation();
 		const FRotation3 R0 = Q0 * XL0.GetRotation();
@@ -1012,7 +986,7 @@ namespace Chaos
 		FReal DSwingAngle1 = -DSwingAngle;
 
 		// Apply swing correction
-		ApplyRotationDelta(Dt, SolverSettings, JointSettings, Stiffness, Index0, Index1, Q0, W0, Q1, W1, InvM0, InvIL0, InvM1, InvIL1, SwingAxis0, DSwingAngle0, SwingAxis1, DSwingAngle1);
+		ApplyRotationDelta(Dt, SolverSettings, JointSettings, Stiffness, Q0, W0, Q1, W1, InvM0, InvIL0, InvM1, InvIL1, SwingAxis0, DSwingAngle0, SwingAxis1, DSwingAngle1);
 	}
 
 	
@@ -1021,8 +995,8 @@ namespace Chaos
 		const FPBDJointSolverSettings& SolverSettings,
 		const FPBDJointSettings& JointSettings,
 		const FReal Stiffness,
-		const int32 Index0,
-		const int32 Index1,
+		const FRigidTransform3& XL0,
+		const FRigidTransform3& XL1,
 		FVec3& P0,
 		FRotation3& Q0,
 		FVec3& V0,
@@ -1036,8 +1010,6 @@ namespace Chaos
 		FReal InvM1,
 		const FMatrix33& InvIL1)
 	{
-		const FRigidTransform3& XL0 = JointSettings.ConstraintFrames[Index0];
-		const FRigidTransform3& XL1 = JointSettings.ConstraintFrames[Index1];
 		const FVec3 X0 = P0 + Q0 * XL0.GetTranslation();
 		const FVec3 X1 = P1 + Q1 * XL1.GetTranslation();
 		const FRotation3 R0 = Q0 * XL0.GetRotation();
@@ -1091,7 +1063,7 @@ namespace Chaos
 			DW = FMath::Min((FReal)0, WC1 - WC0);
 		}
 
-		ApplyRotationVelocityDelta(Dt, SolverSettings, JointSettings, Stiffness, Index0, Index1, Q0, W0, Q1, W1, InvM0, InvIL0, InvM1, InvIL1, SwingAxis0, SwingAxis1, DW);
+		ApplyRotationVelocityDelta(Dt, SolverSettings, JointSettings, Stiffness, Q0, W0, Q1, W1, InvM0, InvIL0, InvM1, InvIL1, SwingAxis0, SwingAxis1, DW);
 	}
 
 
@@ -1101,8 +1073,8 @@ namespace Chaos
 		const FPBDJointSolverSettings& SolverSettings,
 		const FPBDJointSettings& JointSettings,
 		const FReal Stiffness,
-		const int32 Index0,
-		const int32 Index1,
+		const FRigidTransform3& XL0,
+		const FRigidTransform3& XL1,
 		const EJointAngularConstraintIndex SwingConstraintIndex,
 		const EJointAngularAxisIndex SwingAxisIndex,
 		FVec3& P0,
@@ -1118,8 +1090,6 @@ namespace Chaos
 		FReal InvM1,
 		const FMatrix33& InvIL1)
 	{
-		const FRigidTransform3& XL0 = JointSettings.ConstraintFrames[Index0];
-		const FRigidTransform3& XL1 = JointSettings.ConstraintFrames[Index1];
 		const FVec3 X0 = P0 + Q0 * XL0.GetTranslation();
 		const FVec3 X1 = P1 + Q1 * XL1.GetTranslation();
 		const FRotation3 R0 = Q0 * XL0.GetRotation();
@@ -1191,7 +1161,7 @@ namespace Chaos
 			FReal DSwingAngle1 = -DSwingAngle;
 
 			// Apply swing correction
-			ApplyRotationDelta(Dt, SolverSettings, JointSettings, Stiffness, Index0, Index1, Q0, W0, Q1, W1, InvM0, InvIL0, InvM1, InvIL1, SwingAxis0, DSwingAngle0, SwingAxis1, DSwingAngle1);
+			ApplyRotationDelta(Dt, SolverSettings, JointSettings, Stiffness, Q0, W0, Q1, W1, InvM0, InvIL0, InvM1, InvIL1, SwingAxis0, DSwingAngle0, SwingAxis1, DSwingAngle1);
 		}
 	}
 
@@ -1201,8 +1171,8 @@ namespace Chaos
 		const FPBDJointSolverSettings& SolverSettings,
 		const FPBDJointSettings& JointSettings,
 		const FReal Stiffness,
-		const int32 Index0,
-		const int32 Index1,
+		const FRigidTransform3& XL0,
+		const FRigidTransform3& XL1,
 		const EJointAngularConstraintIndex SwingConstraintIndex,
 		const EJointAngularAxisIndex SwingAxisIndex,
 		FVec3& P0,
@@ -1218,8 +1188,6 @@ namespace Chaos
 		FReal InvM1,
 		const FMatrix33& InvIL1)
 	{
-		const FRigidTransform3& XL0 = JointSettings.ConstraintFrames[Index0];
-		const FRigidTransform3& XL1 = JointSettings.ConstraintFrames[Index1];
 		const FVec3 X0 = P0 + Q0 * XL0.GetTranslation();
 		const FVec3 X1 = P1 + Q1 * XL1.GetTranslation();
 		const FRotation3 R0 = Q0 * XL0.GetRotation();
@@ -1289,7 +1257,7 @@ namespace Chaos
 				DW = FMath::Min((FReal)0, WC1 - WC0);
 			}
 
-			ApplyRotationVelocityDelta(Dt, SolverSettings, JointSettings, Stiffness, Index0, Index1, Q0, W0, Q1, W1, InvM0, InvIL0, InvM1, InvIL1, SwingAxis0, SwingAxis1, DW);
+			ApplyRotationVelocityDelta(Dt, SolverSettings, JointSettings, Stiffness, Q0, W0, Q1, W1, InvM0, InvIL0, InvM1, InvIL1, SwingAxis0, SwingAxis1, DW);
 		}
 	}
 
@@ -1298,8 +1266,8 @@ namespace Chaos
 		const FReal Dt,
 		const FPBDJointSolverSettings& SolverSettings,
 		const FPBDJointSettings& JointSettings,
-		const int32 Index0,
-		const int32 Index1,
+		const FRigidTransform3& XL0,
+		const FRigidTransform3& XL1,
 		FVec3& P0,
 		FRotation3& Q0,
 		FVec3& V0,
@@ -1313,8 +1281,6 @@ namespace Chaos
 		FReal InvM1,
 		const FMatrix33& InvIL1)
 	{
-		const FRigidTransform3& XL0 = JointSettings.ConstraintFrames[Index0];
-		const FRigidTransform3& XL1 = JointSettings.ConstraintFrames[Index1];
 		const FVec3 X0 = P0 + Q0 * XL0.GetTranslation();
 		const FVec3 X1 = P1 + Q1 * XL1.GetTranslation();
 		const FRotation3 R0 = Q0 * XL0.GetRotation();
@@ -1349,7 +1315,7 @@ namespace Chaos
 		const FReal DTwistAngle0 = DTwistAngle;
 		const FReal DTwistAngle1 = -DTwistAngle;
 
-		ApplyRotationDelta(Dt, SolverSettings, JointSettings, DriveStiffness, Index0, Index1, Q0, W0, Q1, W1, InvM0, InvIL0, InvM1, InvIL1, TwistAxis0, DTwistAngle0, TwistAxis1, DTwistAngle1);
+		ApplyRotationDelta(Dt, SolverSettings, JointSettings, DriveStiffness, Q0, W0, Q1, W1, InvM0, InvIL0, InvM1, InvIL1, TwistAxis0, DTwistAngle0, TwistAxis1, DTwistAngle1);
 	}
 
 	
@@ -1357,8 +1323,8 @@ namespace Chaos
 		const FReal Dt,
 		const FPBDJointSolverSettings& SolverSettings,
 		const FPBDJointSettings& JointSettings,
-		const int32 Index0,
-		const int32 Index1,
+		const FRigidTransform3& XL0,
+		const FRigidTransform3& XL1,
 		FVec3& P0,
 		FRotation3& Q0,
 		FVec3& V0,
@@ -1372,8 +1338,6 @@ namespace Chaos
 		FReal InvM1,
 		const FMatrix33& InvIL1)
 	{
-		const FRigidTransform3& XL0 = JointSettings.ConstraintFrames[Index0];
-		const FRigidTransform3& XL1 = JointSettings.ConstraintFrames[Index1];
 		const FVec3 X0 = P0 + Q0 * XL0.GetTranslation();
 		const FVec3 X1 = P1 + Q1 * XL1.GetTranslation();
 		const FRotation3 R0 = Q0 * XL0.GetRotation();
@@ -1410,7 +1374,7 @@ namespace Chaos
 		FReal DSwingAngle1 = -DSwingAngle;
 
 		// Apply swing correction
-		ApplyRotationDelta(Dt, SolverSettings, JointSettings, DriveStiffness, Index0, Index1, Q0, W0, Q1, W1, InvM0, InvIL0, InvM1, InvIL1, SwingAxis0, DSwingAngle0, SwingAxis1, DSwingAngle1);
+		ApplyRotationDelta(Dt, SolverSettings, JointSettings, DriveStiffness, Q0, W0, Q1, W1, InvM0, InvIL0, InvM1, InvIL1, SwingAxis0, DSwingAngle0, SwingAxis1, DSwingAngle1);
 	}
 
 	
@@ -1418,8 +1382,8 @@ namespace Chaos
 		const FReal Dt,
 		const FPBDJointSolverSettings& SolverSettings,
 		const FPBDJointSettings& JointSettings,
-		const int32 Index0,
-		const int32 Index1,
+		const FRigidTransform3& XL0,
+		const FRigidTransform3& XL1,
 		FVec3& P0,
 		FRotation3& Q0,
 		FVec3& V0,
@@ -1433,8 +1397,6 @@ namespace Chaos
 		FReal InvM1,
 		const FMatrix33& InvIL1)
 	{
-		const FRigidTransform3& XL0 = JointSettings.ConstraintFrames[Index0];
-		const FRigidTransform3& XL1 = JointSettings.ConstraintFrames[Index1];
 		const FVec3 X0 = P0 + Q0 * XL0.GetTranslation();
 		const FVec3 X1 = P1 + Q1 * XL1.GetTranslation();
 		const FRotation3 R0 = Q0 * XL0.GetRotation();
@@ -1474,8 +1436,8 @@ namespace Chaos
 		const FPBDJointSolverSettings& SolverSettings,
 		const FPBDJointSettings& JointSettings,
 		const FReal Stiffness,
-		const int32 Index0,
-		const int32 Index1,
+		const FRigidTransform3& XL0,
+		const FRigidTransform3& XL1,
 		FVec3& P0,
 		FRotation3& Q0,
 		FVec3& P1,
@@ -1485,8 +1447,6 @@ namespace Chaos
 		FReal InvM1,
 		const FMatrix33& InvIL1)
 	{
-		const FRigidTransform3& XL0 = JointSettings.ConstraintFrames[Index0];
-		const FRigidTransform3& XL1 = JointSettings.ConstraintFrames[Index1];
 		const FVec3 X0 = P0 + Q0 * XL0.GetTranslation();
 		const FVec3 X1 = P1 + Q1 * XL1.GetTranslation();
 		const FRotation3 R0 = Q0 * XL0.GetRotation();
@@ -1504,8 +1464,8 @@ namespace Chaos
 		const FPBDJointSolverSettings& SolverSettings,
 		const FPBDJointSettings& JointSettings,
 		const FReal Stiffness,
-		const int32 Index0,
-		const int32 Index1,
+		const FRigidTransform3& XL0,
+		const FRigidTransform3& XL1,
 		FVec3& P0,
 		FRotation3& Q0,
 		FVec3& P1,
@@ -1519,7 +1479,7 @@ namespace Chaos
 		FVec3 W0(0, 0, 0);
 		FVec3 V1(0, 0, 0);
 		FVec3 W1(0, 0, 0);
-		ApplyJointTwistConstraint(Dt, SolverSettings, JointSettings, Stiffness, Index0, Index1, P0, Q0, V0, W0, P1, Q1, V1, W1, InvM0, InvIL0, InvM1, InvIL1);
+		ApplyJointTwistConstraint(Dt, SolverSettings, JointSettings, Stiffness, XL0, XL1, P0, Q0, V0, W0, P1, Q1, V1, W1, InvM0, InvIL0, InvM1, InvIL1);
 	}
 
 	
@@ -1528,8 +1488,8 @@ namespace Chaos
 		const FPBDJointSolverSettings& SolverSettings,
 		const FPBDJointSettings& JointSettings,
 		const FReal Stiffness,
-		const int32 Index0,
-		const int32 Index1,
+		const FRigidTransform3& XL0,
+		const FRigidTransform3& XL1,
 		FVec3& P0,
 		FRotation3& Q0,
 		FVec3& P1,
@@ -1543,7 +1503,7 @@ namespace Chaos
 		FVec3 W0(0, 0, 0);
 		FVec3 V1(0, 0, 0);
 		FVec3 W1(0, 0, 0);
-		ApplyJointConeConstraint(Dt, SolverSettings, JointSettings, Stiffness, Index0, Index1, P0, Q0, V0, W0, P1, Q1, V1, W1, InvM0, InvIL0, InvM1, InvIL1);
+		ApplyJointConeConstraint(Dt, SolverSettings, JointSettings, Stiffness, XL0, XL1, P0, Q0, V0, W0, P1, Q1, V1, W1, InvM0, InvIL0, InvM1, InvIL1);
 	}
 
 	
@@ -1552,8 +1512,8 @@ namespace Chaos
 		const FPBDJointSolverSettings& SolverSettings,
 		const FPBDJointSettings& JointSettings,
 		const FReal Stiffness,
-		const int32 Index0,
-		const int32 Index1,
+		const FRigidTransform3& XL0,
+		const FRigidTransform3& XL1,
 		const EJointAngularConstraintIndex SwingConstraintIndex,
 		const EJointAngularAxisIndex SwingAxisIndex,
 		FVec3& P0,
@@ -1569,7 +1529,7 @@ namespace Chaos
 		FVec3 W0(0, 0, 0);
 		FVec3 V1(0, 0, 0);
 		FVec3 W1(0, 0, 0);
-		ApplyJointSwingConstraint(Dt, SolverSettings, JointSettings, Stiffness, Index0, Index1, SwingConstraintIndex, SwingAxisIndex, P0, Q0, V0, W0, P1, Q1, V1, W1, InvM0, InvIL0, InvM1, InvIL1);
+		ApplyJointSwingConstraint(Dt, SolverSettings, JointSettings, Stiffness, XL0, XL1, SwingConstraintIndex, SwingAxisIndex, P0, Q0, V0, W0, P1, Q1, V1, W1, InvM0, InvIL0, InvM1, InvIL1);
 	}
 
 
diff --git a/Engine/Source/Runtime/Experimental/Chaos/Private/Chaos/PBDJointConstraints.cpp b/Engine/Source/Runtime/Experimental/Chaos/Private/Chaos/PBDJointConstraints.cpp
--- a/Engine/Source/Runtime/Experimental/Chaos/Private/Chaos/PBDJointConstraints.cpp
+++ b/Engine/Source/Runtime/Experimental/Chaos/Private/Chaos/PBDJointConstraints.cpp
@@ -2,6 +2,7 @@
 #include "Chaos/PBDJointConstraints.h"
 #include "Chaos/ChaosDebugDraw.h"
 #include "Chaos/DebugDrawQueue.h"
+#include "Chaos/Particle/ParticleUtilities.h"
 #include "Chaos/ParticleHandle.h"
 #include "Chaos/PBDJointConstraintSolver.h"
 #include "Chaos/PBDJointConstraintUtilities.h"
@@ -370,21 +371,23 @@ namespace Chaos
 		const int32 Index1 = 0;
 		TGenericParticleHandle<FReal, 3> Particle0 = TGenericParticleHandle<FReal, 3>(ConstraintParticles[ConstraintIndex][Index0]);
 		TGenericParticleHandle<FReal, 3> Particle1 = TGenericParticleHandle<FReal, 3>(ConstraintParticles[ConstraintIndex][Index1]);
-		FVec3 P0 = Particle0->P();
-		FRotation3 Q0 = Particle0->Q();
-		FVec3 P1 = Particle1->P();
-		FRotation3 Q1 = Particle1->Q();
+		const FVec3 P0 = FParticleUtilities::GetCoMWorldPosition(Particle0);
+		const FRotation3 Q0 = FParticleUtilities::GetCoMWorldRotation(Particle0);
+		const FVec3 P1 = FParticleUtilities::GetCoMWorldPosition(Particle1);
+		const FRotation3 Q1 = FParticleUtilities::GetCoMWorldRotation(Particle1);
 
 		const FPBDJointSettings& JointSettings = ConstraintSettings[ConstraintIndex];
+		const FRigidTransform3 XL0 = FParticleUtilities::ParticleLocalToCoMLocal(Particle0, JointSettings.ConstraintFrames[Index0]);
+		const FRigidTransform3 XL1 = FParticleUtilities::ParticleLocalToCoMLocal(Particle1, JointSettings.ConstraintFrames[Index1]);
 		EJointMotionType Swing1Motion = JointSettings.Motion.AngularMotionTypes[(int32)EJointAngularConstraintIndex::Swing1];
 		EJointMotionType Swing2Motion = JointSettings.Motion.AngularMotionTypes[(int32)EJointAngularConstraintIndex::Swing2];
 		if ((Swing1Motion == EJointMotionType::Limited) && (Swing2Motion == EJointMotionType::Limited))
 		{
-			FPBDJointUtilities::CalculateConeConstraintSpace(Settings, ConstraintSettings[ConstraintIndex], Index0, Index1, P0, Q0, P1, Q1, OutX0, OutR0, OutX1, OutR1, OutCR);
+			FPBDJointUtilities::CalculateConeConstraintSpace(Settings, ConstraintSettings[ConstraintIndex], XL0, XL1, P0, Q0, P1, Q1, OutX0, OutR0, OutX1, OutR1, OutCR);
 		}
 		else
 		{
-			FPBDJointUtilities::CalculateSwingConstraintSpace(Settings, ConstraintSettings[ConstraintIndex], Index0, Index1, P0, Q0, P1, Q1, OutX0, OutR0, OutX1, OutR1, OutCR);
+			FPBDJointUtilities::CalculateSwingConstraintSpace(Settings, ConstraintSettings[ConstraintIndex], XL0, XL1, P0, Q0, P1, Q1, OutX0, OutR0, OutX1, OutR1, OutCR);
 		}
 	}
 
@@ -534,13 +537,14 @@ namespace Chaos
 		{
 			if (bUpdateVelocity && (Dt > SMALL_NUMBER))
 			{
-				FVec3 DV = FVec3::CalculateVelocity(Rigid->P(), P, Dt);
-				FVec3 DW = FRotation3::CalculateAngularVelocity(Rigid->Q(), Q, Dt);
+				FVec3 PCoM = FParticleUtilities::GetCoMWorldPosition(Rigid);
+				FRotation3 QCoM = FParticleUtilities::GetCoMWorldRotation(Rigid);
+				FVec3 DV = FVec3::CalculateVelocity(PCoM, P, Dt);
+				FVec3 DW = FRotation3::CalculateAngularVelocity(QCoM, Q, Dt);
 				Rigid->SetV(Rigid->V() + DV);
 				Rigid->SetW(Rigid->W() + DW);
 			}
-			Rigid->SetP(P);
-			Rigid->SetQ(Q);
+			FParticleUtilities::SetCoMWorldTransform(Rigid, P, Q);
 		}
 	}
 
@@ -554,14 +558,16 @@ namespace Chaos
 		TGenericParticleHandle<FReal, 3> Particle0 = TGenericParticleHandle<FReal, 3>(ConstraintParticles[ConstraintIndex][Index0]);
 		TGenericParticleHandle<FReal, 3> Particle1 = TGenericParticleHandle<FReal, 3>(ConstraintParticles[ConstraintIndex][Index1]);
 
-		FVec3 P0 = Particle0->P();
-		FRotation3 Q0 = Particle0->Q();
+		FVec3 P0 = FParticleUtilities::GetCoMWorldPosition(Particle0);
+		FRotation3 Q0 = FParticleUtilities::GetCoMWorldRotation(Particle0);
+		FVec3 P1 = FParticleUtilities::GetCoMWorldPosition(Particle1);
+		FRotation3 Q1 = FParticleUtilities::GetCoMWorldRotation(Particle1);
 		FVec3 V0 = Particle0->V();
 		FVec3 W0 = Particle0->W();
-		FVec3 P1 = Particle1->P();
-		FRotation3 Q1 = Particle1->Q();
 		FVec3 V1 = Particle1->V();
 		FVec3 W1 = Particle1->W();
+		const FRigidTransform3 XL0 = FParticleUtilities::ParticleLocalToCoMLocal(Particle0, JointSettings.ConstraintFrames[Index0]);
+		const FRigidTransform3 XL1 = FParticleUtilities::ParticleLocalToCoMLocal(Particle1, JointSettings.ConstraintFrames[Index1]);
 		float InvM0 = Particle0->InvM();
 		float InvM1 = Particle1->InvM();
 		FMatrix33 InvIL0 = Particle0->InvI();
@@ -581,25 +587,25 @@ namespace Chaos
 			// No SLerp drive if we have a locked rotation (it will be grayed out in the editor in this case, but could still have been set before the rotation was locked)
 			if (JointSettings.Motion.bAngularSLerpDriveEnabled && !bTwistLocked && !bSwing1Locked && !bSwing2Locked)
 			{
-				FPBDJointUtilities::ApplyJointSLerpDrive(Dt, Settings, JointSettings, Index0, Index1, P0, Q0, V0, W0, P1, Q1, V1, W1, InvM0, InvIL0, InvM1, InvIL1);
+				FPBDJointUtilities::ApplyJointSLerpDrive(Dt, Settings, JointSettings, XL0, XL1, P0, Q0, V0, W0, P1, Q1, V1, W1, InvM0, InvIL0, InvM1, InvIL1);
 			}
 
 			if (JointSettings.Motion.bAngularTwistDriveEnabled && !bTwistLocked)
 			{
-				FPBDJointUtilities::ApplyJointTwistDrive(Dt, Settings, JointSettings, Index0, Index1, P0, Q0, V0, W0, P1, Q1, V1, W1, InvM0, InvIL0, InvM1, InvIL1);
+				FPBDJointUtilities::ApplyJointTwistDrive(Dt, Settings, JointSettings, XL0, XL1, P0, Q0, V0, W0, P1, Q1, V1, W1, InvM0, InvIL0, InvM1, InvIL1);
 			}
 
 			if (JointSettings.Motion.bAngularSwingDriveEnabled && !bSwing1Locked && !bSwing2Locked)
 			{
-				FPBDJointUtilities::ApplyJointConeDrive(Dt, Settings, JointSettings, Index0, Index1, P0, Q0, V0, W0, P1, Q1, V1, W1, InvM0, InvIL0, InvM1, InvIL1);
+				FPBDJointUtilities::ApplyJointConeDrive(Dt, Settings, JointSettings, XL0, XL1, P0, Q0, V0, W0, P1, Q1, V1, W1, InvM0, InvIL0, InvM1, InvIL1);
 			}
 			else if (JointSettings.Motion.bAngularSwingDriveEnabled && !bSwing1Locked)
 			{
-				//FPBDJointUtilities::ApplyJointSwingDrive(Dt, Settings, JointSettings, Index0, Index1, EJointAngularConstraintIndex::Swing1, P0, Q0, P1, Q1, InvM0, InvIL0, InvM1, InvIL1);
+				//FPBDJointUtilities::ApplyJointSwingDrive(Dt, Settings, JointSettings, XL0, XL1, EJointAngularConstraintIndex::Swing1, P0, Q0, P1, Q1, InvM0, InvIL0, InvM1, InvIL1);
 			}
 			else if (JointSettings.Motion.bAngularSwingDriveEnabled && !bSwing2Locked)
 			{
-				//FPBDJointUtilities::ApplyJointSwingDrive(Dt, Settings, JointSettings, Index0, Index1, EJointAngularConstraintIndex::Swing2, P0, Q0, P1, Q1, InvM0, InvIL0, InvM1, InvIL1);
+				//FPBDJointUtilities::ApplyJointSwingDrive(Dt, Settings, JointSettings, XL0, XL1, EJointAngularConstraintIndex::Swing2, P0, Q0, P1, Q1, InvM0, InvIL0, InvM1, InvIL1);
 			}
 		}
 
@@ -637,21 +643,28 @@ namespace Chaos
 		TGenericParticleHandle<FReal, 3> Particle0 = TGenericParticleHandle<FReal, 3>(ConstraintParticles[ConstraintIndex][Index0]);
 		TGenericParticleHandle<FReal, 3> Particle1 = TGenericParticleHandle<FReal, 3>(ConstraintParticles[ConstraintIndex][Index1]);
 
+		const FVec3 P0 = FParticleUtilities::GetCoMWorldPosition(Particle0);
+		const FRotation3 Q0 = FParticleUtilities::GetCoMWorldRotation(Particle0);
+		const FVec3 P1 = FParticleUtilities::GetCoMWorldPosition(Particle1);
+		const FRotation3 Q1 = FParticleUtilities::GetCoMWorldRotation(Particle1);
+		const FRigidTransform3 XL0 = FParticleUtilities::ParticleLocalToCoMLocal(Particle0, JointSettings.ConstraintFrames[Index0]);
+		const FRigidTransform3 XL1 = FParticleUtilities::ParticleLocalToCoMLocal(Particle1, JointSettings.ConstraintFrames[Index1]);
+
 		FJointConstraintSolver Solver;
 		Solver.InitConstraints(
 			Dt, 
 			Settings, 
 			JointSettings, 
-			Particle0->P(), 
-			Particle0->Q(), 
-			Particle1->P(), 
-			Particle1->Q(), 
+			P0,
+			Q0,
+			P1,
+			Q1,
 			Particle0->InvM(), 
 			Particle0->InvI(), 
 			Particle1->InvM(), 
 			Particle1->InvI(), 
-			JointSettings.ConstraintFrames[Index0], 
-			JointSettings.ConstraintFrames[Index1]);
+			XL0,
+			XL1);
 		
 		for (int32 PairIt = 0; PairIt < NumPairIts; ++PairIt)
 		{
@@ -678,14 +691,16 @@ namespace Chaos
 		TGenericParticleHandle<FReal, 3> Particle0 = TGenericParticleHandle<FReal, 3>(ConstraintParticles[ConstraintIndex][Index0]);
 		TGenericParticleHandle<FReal, 3> Particle1 = TGenericParticleHandle<FReal, 3>(ConstraintParticles[ConstraintIndex][Index1]);
 
-		FVec3 P0 = Particle0->P();
-		FRotation3 Q0 = Particle0->Q();
+		FVec3 P0 = FParticleUtilities::GetCoMWorldPosition(Particle0);
+		FRotation3 Q0 = FParticleUtilities::GetCoMWorldRotation(Particle0);
+		FVec3 P1 = FParticleUtilities::GetCoMWorldPosition(Particle1);
+		FRotation3 Q1 = FParticleUtilities::GetCoMWorldRotation(Particle1);
 		FVec3 V0 = Particle0->V();
 		FVec3 W0 = Particle0->W();
-		FVec3 P1 = Particle1->P();
-		FRotation3 Q1 = Particle1->Q();
 		FVec3 V1 = Particle1->V();
 		FVec3 W1 = Particle1->W();
+		const FRigidTransform3 XL0 = FParticleUtilities::ParticleLocalToCoMLocal(Particle0, JointSettings.ConstraintFrames[Index0]);
+		const FRigidTransform3 XL1 = FParticleUtilities::ParticleLocalToCoMLocal(Particle1, JointSettings.ConstraintFrames[Index1]);
 		float InvM0 = Particle0->InvM();
 		float InvM1 = Particle1->InvM();
 		FMatrix33 InvIL0 = Particle0->InvI();
@@ -725,7 +740,7 @@ namespace Chaos
 			{
 				if (TwistMotion != EJointMotionType::Free)
 				{
-					FPBDJointUtilities::ApplyJointTwistConstraint(Dt, Settings, JointSettings, TwistStiffness, Index0, Index1, P0, Q0, V0, W0, P1, Q1, V1, W1, InvM0, InvIL0, InvM1, InvIL1);
+					FPBDJointUtilities::ApplyJointTwistConstraint(Dt, Settings, JointSettings, TwistStiffness, XL0, XL1, P0, Q0, V0, W0, P1, Q1, V1, W1, InvM0, InvIL0, InvM1, InvIL1);
 				}
 			}
 
@@ -735,19 +750,19 @@ namespace Chaos
 				if ((Swing1Motion == EJointMotionType::Limited) && (Swing2Motion == EJointMotionType::Limited))
 				{
 					// Swing Cone
-					FPBDJointUtilities::ApplyJointConeConstraint(Dt, Settings, JointSettings, SwingStiffness, Index0, Index1, P0, Q0, V0, W0, P1, Q1, V1, W1, InvM0, InvIL0, InvM1, InvIL1);
+					FPBDJointUtilities::ApplyJointConeConstraint(Dt, Settings, JointSettings, SwingStiffness, XL0, XL1, P0, Q0, V0, W0, P1, Q1, V1, W1, InvM0, InvIL0, InvM1, InvIL1);
 				}
 				else
 				{
 					if (Swing1Motion != EJointMotionType::Free)
 					{
 						// Swing Arc/Lock
-						FPBDJointUtilities::ApplyJointSwingConstraint(Dt, Settings, JointSettings, SwingStiffness, Index0, Index1, EJointAngularConstraintIndex::Swing1, EJointAngularAxisIndex::Swing1, P0, Q0, V0, W0, P1, Q1, V1, W1, InvM0, InvIL0, InvM1, InvIL1);
+						FPBDJointUtilities::ApplyJointSwingConstraint(Dt, Settings, JointSettings, SwingStiffness, XL0, XL1, EJointAngularConstraintIndex::Swing1, EJointAngularAxisIndex::Swing1, P0, Q0, V0, W0, P1, Q1, V1, W1, InvM0, InvIL0, InvM1, InvIL1);
 					}
 					if (Swing2Motion != EJointMotionType::Free)
 					{
 						// Swing Arc/Lock
-						FPBDJointUtilities::ApplyJointSwingConstraint(Dt, Settings, JointSettings, SwingStiffness, Index0, Index1, EJointAngularConstraintIndex::Swing2, EJointAngularAxisIndex::Swing2, P0, Q0, V0, W0, P1, Q1, V1, W1, InvM0, InvIL0, InvM1, InvIL1);
+						FPBDJointUtilities::ApplyJointSwingConstraint(Dt, Settings, JointSettings, SwingStiffness, XL0, XL1, EJointAngularConstraintIndex::Swing2, EJointAngularAxisIndex::Swing2, P0, Q0, V0, W0, P1, Q1, V1, W1, InvM0, InvIL0, InvM1, InvIL1);
 					}
 				}
 			}
@@ -755,7 +770,7 @@ namespace Chaos
 			// Apply linear constraints
 			if ((LinearMotion[0] != EJointMotionType::Free) || (LinearMotion[1] != EJointMotionType::Free) || (LinearMotion[2] != EJointMotionType::Free))
 			{
-				FPBDJointUtilities::ApplyJointPositionConstraint(Dt, Settings, JointSettings, LinearStiffness, Index0, Index1, P0, Q0, V0, W0, P1, Q1, V1, W1, InvM0, InvIL0, InvM1, InvIL1);
+				FPBDJointUtilities::ApplyJointPositionConstraint(Dt, Settings, JointSettings, LinearStiffness, XL0, XL1, P0, Q0, V0, W0, P1, Q1, V1, W1, InvM0, InvIL0, InvM1, InvIL1);
 			}
 		}
 
@@ -778,14 +793,16 @@ namespace Chaos
 		TGenericParticleHandle<FReal, 3> Particle0 = TGenericParticleHandle<FReal, 3>(ConstraintParticles[ConstraintIndex][Index0]);
 		TGenericParticleHandle<FReal, 3> Particle1 = TGenericParticleHandle<FReal, 3>(ConstraintParticles[ConstraintIndex][Index1]);
 
-		FVec3 P0 = Particle0->P();
-		FRotation3 Q0 = Particle0->Q();
+		FVec3 P0 = FParticleUtilities::GetCoMWorldPosition(Particle0);
+		FRotation3 Q0 = FParticleUtilities::GetCoMWorldRotation(Particle0);
+		FVec3 P1 = FParticleUtilities::GetCoMWorldPosition(Particle1);
+		FRotation3 Q1 = FParticleUtilities::GetCoMWorldRotation(Particle1);
 		FVec3 V0 = Particle0->V();
 		FVec3 W0 = Particle0->W();
-		FVec3 P1 = Particle1->P();
-		FRotation3 Q1 = Particle1->Q();
 		FVec3 V1 = Particle1->V();
 		FVec3 W1 = Particle1->W();
+		const FRigidTransform3 XL0 = FParticleUtilities::ParticleLocalToCoMLocal(Particle0, JointSettings.ConstraintFrames[Index0]);
+		const FRigidTransform3 XL1 = FParticleUtilities::ParticleLocalToCoMLocal(Particle1, JointSettings.ConstraintFrames[Index1]);
 		float InvM0 = Particle0->InvM();
 		float InvM1 = Particle1->InvM();
 		FMatrix33 InvIL0 = Particle0->InvI();
@@ -825,7 +842,7 @@ namespace Chaos
 			{
 				if (TwistMotion != EJointMotionType::Free)
 				{
-					FPBDJointUtilities::ApplyJointTwistVelocityConstraint(Dt, Settings, JointSettings, TwistStiffness, Index0, Index1, P0, Q0, V0, W0, P1, Q1, V1, W1, InvM0, InvIL0, InvM1, InvIL1);
+					FPBDJointUtilities::ApplyJointTwistVelocityConstraint(Dt, Settings, JointSettings, TwistStiffness, XL0, XL1, P0, Q0, V0, W0, P1, Q1, V1, W1, InvM0, InvIL0, InvM1, InvIL1);
 				}
 			}
 
@@ -835,19 +852,19 @@ namespace Chaos
 				if ((Swing1Motion == EJointMotionType::Limited) && (Swing2Motion == EJointMotionType::Limited))
 				{
 					// Swing Cone
-					FPBDJointUtilities::ApplyJointConeVelocityConstraint(Dt, Settings, JointSettings, SwingStiffness, Index0, Index1, P0, Q0, V0, W0, P1, Q1, V1, W1, InvM0, InvIL0, InvM1, InvIL1);
+					FPBDJointUtilities::ApplyJointConeVelocityConstraint(Dt, Settings, JointSettings, SwingStiffness, XL0, XL1, P0, Q0, V0, W0, P1, Q1, V1, W1, InvM0, InvIL0, InvM1, InvIL1);
 				}
 				else
 				{
 					if (Swing1Motion != EJointMotionType::Free)
 					{
 						// Swing Arc/Lock
-						FPBDJointUtilities::ApplyJointSwingVelocityConstraint(Dt, Settings, JointSettings, SwingStiffness, Index0, Index1, EJointAngularConstraintIndex::Swing1, EJointAngularAxisIndex::Swing1, P0, Q0, V0, W0, P1, Q1, V1, W1, InvM0, InvIL0, InvM1, InvIL1);
+						FPBDJointUtilities::ApplyJointSwingVelocityConstraint(Dt, Settings, JointSettings, SwingStiffness, XL0, XL1, EJointAngularConstraintIndex::Swing1, EJointAngularAxisIndex::Swing1, P0, Q0, V0, W0, P1, Q1, V1, W1, InvM0, InvIL0, InvM1, InvIL1);
 					}
 					if (Swing2Motion != EJointMotionType::Free)
 					{
 						// Swing Arc/Lock
-						FPBDJointUtilities::ApplyJointSwingVelocityConstraint(Dt, Settings, JointSettings, SwingStiffness, Index0, Index1, EJointAngularConstraintIndex::Swing2, EJointAngularAxisIndex::Swing2, P0, Q0, V0, W0, P1, Q1, V1, W1, InvM0, InvIL0, InvM1, InvIL1);
+						FPBDJointUtilities::ApplyJointSwingVelocityConstraint(Dt, Settings, JointSettings, SwingStiffness, XL0, XL1, EJointAngularConstraintIndex::Swing2, EJointAngularAxisIndex::Swing2, P0, Q0, V0, W0, P1, Q1, V1, W1, InvM0, InvIL0, InvM1, InvIL1);
 					}
 				}
 			}
@@ -855,7 +872,7 @@ namespace Chaos
 			// Apply linear velocity  constraints
 			if ((LinearMotion[0] != EJointMotionType::Free) || (LinearMotion[1] != EJointMotionType::Free) || (LinearMotion[2] != EJointMotionType::Free))
 			{
-				FPBDJointUtilities::ApplyJointVelocityConstraint(Dt, Settings, JointSettings, LinearStiffness, Index0, Index1, P0, Q0, V0, W0, P1, Q1, V1, W1, InvM0, InvIL0, InvM1, InvIL1);
+				FPBDJointUtilities::ApplyJointVelocityConstraint(Dt, Settings, JointSettings, LinearStiffness, XL0, XL1, P0, Q0, V0, W0, P1, Q1, V1, W1, InvM0, InvIL0, InvM1, InvIL1);
 			}
 		}
 
@@ -886,10 +903,12 @@ namespace Chaos
 		TGenericParticleHandle<FReal, 3> Particle0 = TGenericParticleHandle<FReal, 3>(ConstraintParticles[ConstraintIndex][Index0]);
 		TGenericParticleHandle<FReal, 3> Particle1 = TGenericParticleHandle<FReal, 3>(ConstraintParticles[ConstraintIndex][Index1]);
 
-		FVec3 P0 = Particle0->P();
-		FRotation3 Q0 = Particle0->Q();
-		FVec3 P1 = Particle1->P();
-		FRotation3 Q1 = Particle1->Q();
+		FVec3 P0 = FParticleUtilities::GetCoMWorldPosition(Particle0);
+		FRotation3 Q0 = FParticleUtilities::GetCoMWorldRotation(Particle0);
+		FVec3 P1 = FParticleUtilities::GetCoMWorldPosition(Particle1);
+		FRotation3 Q1 = FParticleUtilities::GetCoMWorldRotation(Particle1);
+		const FRigidTransform3 XL0 = FParticleUtilities::ParticleLocalToCoMLocal(Particle0, JointSettings.ConstraintFrames[Index0]);
+		const FRigidTransform3 XL1 = FParticleUtilities::ParticleLocalToCoMLocal(Particle1, JointSettings.ConstraintFrames[Index1]);
 		float InvM0 = Particle0->InvM();
 		float InvM1 = Particle1->InvM();
 		FMatrix33 InvIL0 = Particle0->InvI();
@@ -919,7 +938,7 @@ namespace Chaos
 			if (Settings.bEnableTwistLimits)
 			{
 				// Remove Twist Error
-				FPBDJointUtilities::ApplyJointTwistProjection(Dt, Settings, JointSettings, AngularProjectionFactor, Index0, Index1, P0, Q0, P1, Q1, InvM0, InvIL0, InvM1, InvIL1);
+				FPBDJointUtilities::ApplyJointTwistProjection(Dt, Settings, JointSettings, AngularProjectionFactor, XL0, XL1, P0, Q0, P1, Q1, InvM0, InvIL0, InvM1, InvIL1);
 			}
 
 			if (Settings.bEnableSwingLimits)
@@ -927,17 +946,17 @@ namespace Chaos
 				// Remove Swing Error
 				if ((Swing1Motion == EJointMotionType::Limited) && (Swing2Motion == EJointMotionType::Limited))
 				{
-					FPBDJointUtilities::ApplyJointConeProjection(Dt, Settings, JointSettings, AngularProjectionFactor, Index0, Index1, P0, Q0, P1, Q1, InvM0, InvIL0, InvM1, InvIL1);
+					FPBDJointUtilities::ApplyJointConeProjection(Dt, Settings, JointSettings, AngularProjectionFactor, XL0, XL1, P0, Q0, P1, Q1, InvM0, InvIL0, InvM1, InvIL1);
 				}
 				else
 				{
 					if (Swing1Motion != EJointMotionType::Free)
 					{
-						FPBDJointUtilities::ApplyJointSwingProjection(Dt, Settings, JointSettings, AngularProjectionFactor, Index0, Index1, EJointAngularConstraintIndex::Swing1, EJointAngularAxisIndex::Swing1, P0, Q0, P1, Q1, InvM0, InvIL0, InvM1, InvIL1);
+						FPBDJointUtilities::ApplyJointSwingProjection(Dt, Settings, JointSettings, AngularProjectionFactor, XL0, XL1, EJointAngularConstraintIndex::Swing1, EJointAngularAxisIndex::Swing1, P0, Q0, P1, Q1, InvM0, InvIL0, InvM1, InvIL1);
 					}
 					if (Swing2Motion != EJointMotionType::Free)
 					{
-						FPBDJointUtilities::ApplyJointSwingProjection(Dt, Settings, JointSettings, AngularProjectionFactor, Index0, Index1, EJointAngularConstraintIndex::Swing2, EJointAngularAxisIndex::Swing2, P0, Q0, P1, Q1, InvM0, InvIL0, InvM1, InvIL1);
+						FPBDJointUtilities::ApplyJointSwingProjection(Dt, Settings, JointSettings, AngularProjectionFactor, XL0, XL1, EJointAngularConstraintIndex::Swing2, EJointAngularAxisIndex::Swing2, P0, Q0, P1, Q1, InvM0, InvIL0, InvM1, InvIL1);
 					}
 				}
 			}
@@ -948,7 +967,7 @@ namespace Chaos
 		{
 			if ((LinearMotion[0] != EJointMotionType::Free) || (LinearMotion[1] != EJointMotionType::Free) || (LinearMotion[2] != EJointMotionType::Free))
 			{
-				FPBDJointUtilities::ApplyJointPositionProjection(Dt, Settings, JointSettings, LinearProjectionFactor, Index0, Index1, P0, Q0, P1, Q1, InvM0, InvIL0, InvM1, InvIL1);
+				FPBDJointUtilities::ApplyJointPositionProjection(Dt, Settings, JointSettings, LinearProjectionFactor, XL0, XL1, P0, Q0, P1, Q1, InvM0, InvIL0, InvM1, InvIL1);
 			}
 		}
 
diff --git a/Engine/Source/Runtime/Experimental/Chaos/Public/Chaos/PBDJointConstraintUtilities.h b/Engine/Source/Runtime/Experimental/Chaos/Public/Chaos/PBDJointConstraintUtilities.h
--- a/Engine/Source/Runtime/Experimental/Chaos/Public/Chaos/PBDJointConstraintUtilities.h
+++ b/Engine/Source/Runtime/Experimental/Chaos/Public/Chaos/PBDJointConstraintUtilities.h
@@ -64,12 +64,12 @@ namespace Chaos
 		static CHAOS_API void CalculateSwingConstraintSpace(
 			const FPBDJointSolverSettings& SolverSettings,
 			const FPBDJointSettings& JointSettings,
-			const int32 Index0,
-			const int32 Index1,
-			FVec3& P0,
-			FRotation3& Q0,
-			FVec3& P1,
-			FRotation3& Q1,
+			const FRigidTransform3& XL0,
+			const FRigidTransform3& XL1,
+			const FVec3& P0,
+			const FRotation3& Q0,
+			const FVec3& P1,
+			const FRotation3& Q1,
 			FVec3& OutX0,
 			FMatrix33& OutR0,
 			FVec3& OutX1,
@@ -79,12 +79,12 @@ namespace Chaos
 		static CHAOS_API void CalculateConeConstraintSpace(
 			const FPBDJointSolverSettings& SolverSettings,
 			const FPBDJointSettings& JointSettings,
-			const int32 Index0,
-			const int32 Index1,
-			FVec3& P0,
-			FRotation3& Q0,
-			FVec3& P1,
-			FRotation3& Q1,
+			const FRigidTransform3& XL0,
+			const FRigidTransform3& XL1,
+			const FVec3& P0,
+			const FRotation3& Q0,
+			const FVec3& P1,
+			const FRotation3& Q1,
 			FVec3& OutX0,
 			FMatrix33& OutR0,
 			FVec3& OutX1,
@@ -96,8 +96,8 @@ namespace Chaos
 			const FPBDJointSolverSettings& SolverSettings,
 			const FPBDJointSettings& JointSettings,
 			const FReal Stiffness,
-			const int32 Index0,
-			const int32 Index1,
+			const FRigidTransform3& XL0,
+			const FRigidTransform3& XL1,
 			FVec3& P0,
 			FRotation3& Q0,
 			FVec3& V0,
@@ -116,8 +116,8 @@ namespace Chaos
 			const FPBDJointSolverSettings& SolverSettings,
 			const FPBDJointSettings& JointSettings,
 			const FReal Stiffness,
-			const int32 Index0,
-			const int32 Index1,
+			const FRigidTransform3& XL0,
+			const FRigidTransform3& XL1,
 			FVec3& P0,
 			FRotation3& Q0,
 			FVec3& V0,
@@ -136,8 +136,8 @@ namespace Chaos
 			const FPBDJointSolverSettings& SolverSettings,
 			const FPBDJointSettings& JointSettings,
 			const FReal Stiffness,
-			const int32 Index0,
-			const int32 Index1,
+			const FRigidTransform3& XL0,
+			const FRigidTransform3& XL1,
 			FVec3& P0,
 			FRotation3& Q0,
 			FVec3& V0,
@@ -156,8 +156,8 @@ namespace Chaos
 			const FPBDJointSolverSettings& SolverSettings,
 			const FPBDJointSettings& JointSettings,
 			const FReal Stiffness,
-			const int32 Index0,
-			const int32 Index1,
+			const FRigidTransform3& XL0,
+			const FRigidTransform3& XL1,
 			FVec3& P0,
 			FRotation3& Q0,
 			FVec3& V0,
@@ -176,8 +176,8 @@ namespace Chaos
 			const FPBDJointSolverSettings& SolverSettings,
 			const FPBDJointSettings& JointSettings,
 			const FReal Stiffness,
-			const int32 Index0,
-			const int32 Index1,
+			const FRigidTransform3& XL0,
+			const FRigidTransform3& XL1,
 			FVec3& P0,
 			FRotation3& Q0,
 			FVec3& V0,
@@ -196,8 +196,8 @@ namespace Chaos
 			const FPBDJointSolverSettings& SolverSettings,
 			const FPBDJointSettings& JointSettings,
 			const FReal Stiffness,
-			const int32 Index0,
-			const int32 Index1,
+			const FRigidTransform3& XL0,
+			const FRigidTransform3& XL1,
 			FVec3& P0,
 			FRotation3& Q0,
 			FVec3& V0,
@@ -216,8 +216,8 @@ namespace Chaos
 			const FPBDJointSolverSettings& SolverSettings,
 			const FPBDJointSettings& JointSettings,
 			const FReal Stiffness,
-			const int32 Index0,
-			const int32 Index1,
+			const FRigidTransform3& XL0,
+			const FRigidTransform3& XL1,
 			const EJointAngularConstraintIndex SwingConstraintIndex,
 			const EJointAngularAxisIndex SwingAxisIndex,
 			FVec3& P0,
@@ -238,8 +238,8 @@ namespace Chaos
 			const FPBDJointSolverSettings& SolverSettings,
 			const FPBDJointSettings& JointSettings,
 			const FReal Stiffness,
-			const int32 Index0,
-			const int32 Index1,
+			const FRigidTransform3& XL0,
+			const FRigidTransform3& XL1,
 			const EJointAngularConstraintIndex SwingConstraintIndex,
 			const EJointAngularAxisIndex SwingAxisIndex,
 			FVec3& P0,
@@ -259,8 +259,8 @@ namespace Chaos
 			const FReal Dt,
 			const FPBDJointSolverSettings& SolverSettings,
 			const FPBDJointSettings& JointSettings,
-			const int32 Index0,
-			const int32 Index1,
+			const FRigidTransform3& XL0,
+			const FRigidTransform3& XL1,
 			FVec3& P0,
 			FRotation3& Q0,
 			FVec3& V0,
@@ -278,8 +278,8 @@ namespace Chaos
 			const FReal Dt,
 			const FPBDJointSolverSettings& SolverSettings,
 			const FPBDJointSettings& JointSettings,
-			const int32 Index0,
-			const int32 Index1,
+			const FRigidTransform3& XL0,
+			const FRigidTransform3& XL1,
 			FVec3& P0,
 			FRotation3& Q0,
 			FVec3& V0,
@@ -297,8 +297,8 @@ namespace Chaos
 			const FReal Dt,
 			const FPBDJointSolverSettings& SolverSettings,
 			const FPBDJointSettings& JointSettings,
-			const int32 Index0,
-			const int32 Index1,
+			const FRigidTransform3& XL0,
+			const FRigidTransform3& XL1,
 			FVec3& P0,
 			FRotation3& Q0,
 			FVec3& V0,
@@ -317,8 +317,8 @@ namespace Chaos
 			const FPBDJointSolverSettings& SolverSettings,
 			const FPBDJointSettings& JointSettings,
 			const FReal Stiffness,
-			const int32 Index0,
-			const int32 Index1,
+			const FRigidTransform3& XL0,
+			const FRigidTransform3& XL1,
 			FVec3& P0,
 			FRotation3& Q0,
 			FVec3& P1,
@@ -333,8 +333,8 @@ namespace Chaos
 			const FPBDJointSolverSettings& SolverSettings,
 			const FPBDJointSettings& JointSettings,
 			const FReal Stiffness,
-			const int32 Index0,
-			const int32 Index1,
+			const FRigidTransform3& XL0,
+			const FRigidTransform3& XL1,
 			FVec3& P0,
 			FRotation3& Q0,
 			FVec3& P1,
@@ -349,8 +349,8 @@ namespace Chaos
 			const FPBDJointSolverSettings& SolverSettings,
 			const FPBDJointSettings& JointSettings,
 			const FReal Stiffness,
-			const int32 Index0,
-			const int32 Index1,
+			const FRigidTransform3& XL0,
+			const FRigidTransform3& XL1,
 			FVec3& P0,
 			FRotation3& Q0,
 			FVec3& P1,
@@ -365,8 +365,8 @@ namespace Chaos
 			const FPBDJointSolverSettings& SolverSettings,
 			const FPBDJointSettings& JointSettings,
 			const FReal Stiffness,
-			const int32 Index0,
-			const int32 Index1,
+			const FRigidTransform3& XL0,
+			const FRigidTransform3& XL1,
 			const EJointAngularConstraintIndex SwingConstraintIndex,
 			const EJointAngularAxisIndex SwingAxisIndex,
 			FVec3& P0,
diff --git a/Engine/Source/Runtime/Experimental/Chaos/Public/Chaos/Particle/ParticleUtilities.h b/Engine/Source/Runtime/Experimental/Chaos/Public/Chaos/Particle/ParticleUtilities.h
new file mode 100644
--- /dev/null
+++ b/Engine/Source/Runtime/Experimental/Chaos/Public/Chaos/Particle/ParticleUtilities.h
@@ -0,0 +1,239 @@
+// Copyright 1998-2019 Epic Games, Inc. All Rights Reserved.
+#pragma once
+
+#include "Chaos/ParticleHandle.h"
+
+#define CHAOS_PARTICLE_ACTORTRANSFORM 0
+
+namespace Chaos
+{
+	/**
+	 * Particle Space == Actor Space (Transforms)
+	 * Velocities in CoM Space.
+	 */
+	class FParticleUtilities_ActorSpace
+	{
+	public:
+		template<typename T_PARTICLEHANDLE>
+		static inline FRigidTransform3 GetActorWorldTransform(T_PARTICLEHANDLE Particle)
+		{
+			return FRigidTransform3(Particle->P(), Particle->Q());
+		}
+
+		template<typename T_PARTICLEHANDLE>
+		static inline void SetActorWorldTransform(T_PARTICLEHANDLE Particle, const FRigidTransform3& ActorWorldT)
+		{
+			Particle->P() = ActorWorldT.GetTranslation();
+			Particle->Q() = ActorWorldT.GetRotation();
+		}
+
+		template<typename T_PARTICLEHANDLE>
+		static inline const FRigidTransform3& ActorLocalToParticleLocal(T_PARTICLEHANDLE Particle, const FRigidTransform3& ActorLocalT)
+		{
+			return ActorLocalT;
+		}
+
+		template<typename T_PARTICLEHANDLE>
+		static inline const FRigidTransform3& ActorWorldToParticleWorld(T_PARTICLEHANDLE Particle, const FRigidTransform3& ActorWorldT)
+		{
+			return ActorWorldT;
+		}
+
+		/**
+		 * Convert an particle position into a com-local-space position
+		 */
+		template<typename T_PARTICLEHANDLE>
+		static inline FVec3 ParticleLocalToCoMLocal(T_PARTICLEHANDLE Particle, const FVec3& P)
+		{
+			return Particle->RotationOfMass().UnrotateVector(P - Particle->CenterOfMass());
+		}
+
+		/**
+		 * Convert a particle rotation into a com-local-space rotation
+		 */
+		template<typename T_PARTICLEHANDLE>
+		static inline FRotation3 ParticleLocalToCoMLocal(T_PARTICLEHANDLE Particle, const FRotation3& Q)
+		{
+			return Particle->RotationOfMass().Inverse()* Q;
+		}
+
+		/**
+		 * Convert an particle-local-space transform into a com-local-space transform
+		 */
+		template<typename T_PARTICLEHANDLE>
+		static inline FRigidTransform3 ParticleLocalToCoMLocal(T_PARTICLEHANDLE Particle, const FRigidTransform3& T)
+		{
+			return FRigidTransform3(ParticleLocalToCoMLocal(Particle, T.GetTranslation()), ParticleLocalToCoMLocal(Particle, T.GetRotation()));
+		}
+
+		/**
+		 * Get the velocity at point 'RelPos', where 'RelPos' is a world-space position relative to the Particle's center of mass.
+		 */
+		template<typename T_PARTICLEHANDLE>
+		static inline FVec3 GetVelocityAtCoMRelativePosition(T_PARTICLEHANDLE Particle, const FVec3& RelPos)
+		{
+			return Particle->V() + FVec3::CrossProduct(Particle->W(), RelPos);
+		}
+
+		template<typename T_PARTICLEHANDLE>
+		static inline FVec3 GetCoMWorldPosition(T_PARTICLEHANDLE Particle)
+		{
+			return Particle->P() + Particle->Q().RotateVector(Particle->CenterOfMass());
+		}
+
+		static inline FVec3 GetCoMWorldPosition(TPBDRigidParticles<FReal, 3>& Particles, const int32 Index)
+		{
+			return Particles.P(Index) + Particles.Q(Index).RotateVector(Particles.CenterOfMass(Index));
+		}
+
+		template<typename T_PARTICLEHANDLE>
+		static inline FRotation3 GetCoMWorldRotation(T_PARTICLEHANDLE Particle)
+		{
+			return Particle->Q() * Particle->RotationOfMass();
+		}
+	
+		static inline FRotation3 GetCoMWorldRotation(TPBDRigidParticles<FReal, 3>& Particles, const int32 Index)
+		{
+			return Particles.Q(Index) * Particles.RotationOfMass(Index);
+		}
+
+		template<typename T_PARTICLEHANDLE>
+		static inline FRigidTransform3 GetCoMWorldTransform(T_PARTICLEHANDLE Particle)
+		{
+			return FRigidTransform3(GetCoMWorldPosition(Particle), GetCoMWorldRotation(Particle));
+		}
+
+		/**
+		 * Update the particle's position and rotation by specifying a new center of mass transform.
+		 */
+		template<typename T_PARTICLEHANDLE>
+		static inline void SetCoMWorldTransform(T_PARTICLEHANDLE Particle, const FVec3& PCoM, const FRotation3& QCoM)
+		{
+			const FRotation3 Q = QCoM * Particle->RotationOfMass().Inverse();
+			const FVec3 P = PCoM - Q.RotateVector(Particle->CenterOfMass());
+			
+			Particle->P() = P;
+			Particle->Q() = Q;
+		}
+
+		static inline void SetCoMWorldTransform(TPBDRigidParticles<FReal, 3>& Particles, const int32 Index, const FVec3& PCoM, const FRotation3& QCoM)
+		{
+			const FRotation3 Q = QCoM * Particles.RotationOfMass(Index).Inverse();
+			const FVec3 P = PCoM - Q.RotateVector(Particles.CenterOfMass(Index));
+
+			Particles.P(Index) = P;
+			Particles.Q(Index) = Q;
+		}
+	};
+
+	/**
+	 * Particle Space == CoM Space.
+	 * Velocities in CoM Space.
+	 */
+	class FParticleUtilities_CoMSpace
+	{
+	public:
+		template<typename T_PARTICLEHANDLE>
+		static inline FRigidTransform3 GetActorWorldTransform(T_PARTICLEHANDLE Particle)
+		{
+			FRotation3 ActorQ = Particle->Q() * Particle->RotationOfMass().Inverse();
+			FVec3 ActorP = Particle->P() - ActorQ.RotateVector(Particle->CenterOfMass());
+			return FRigidTransform3(ActorP, ActorQ);
+		}
+
+		template<typename T_PARTICLEHANDLE>
+		static inline void SetActorWorldTransform(T_PARTICLEHANDLE Particle, const FRigidTransform3& ActorWorldT)
+		{
+			FRotation3 CoMQ = ActorWorldT.GetRotation() * Particle->RotationOfMass();
+			FVec3 CoMP = ActorWorldT.GetTranslation() + ActorWorldT.GetRotation().RotateVector(Particle->CenterOfMass());
+			Particle->P() = CoMP;
+			Particle->Q() = CoMQ;
+		}
+
+		template<typename T_PARTICLEHANDLE>
+		static inline FRigidTransform3 ActorLocalToParticleLocal(T_PARTICLEHANDLE Particle, const FRigidTransform3& ActorLocalT)
+		{
+			return FRigidTransform3(Particle->RotationOfMass().UnrotateVector(ActorLocalT.GetTranslation() - Particle->CenterOfMass()), Particle->RotationOfMass().Inverse() * ActorLocalT.GetRotation());
+		}
+
+		template<typename T_PARTICLEHANDLE>
+		static inline FRigidTransform3 ActorWorldToParticleWorld(T_PARTICLEHANDLE Particle, const FRigidTransform3& ActorWorldT)
+		{
+			FRotation3 CoMQ = ActorWorldT.GetRotation() * Particle->RotationOfMass();
+			FVec3 CoMP = ActorWorldT.GetTranslation() + ActorWorldT.GetRotation().RotateVector(Particle->CenterOfMass());
+			return FRigidTransform3(CoMP, CoMQ);
+		}
+
+		template<typename T_PARTICLEHANDLE>
+		static inline const FVec3& ParticleLocalToCoMLocal(T_PARTICLEHANDLE Particle, const FVec3& P)
+		{
+			return P;
+		}
+
+		template<typename T_PARTICLEHANDLE>
+		static inline const FRotation3& ParticleLocalToCoMLocal(T_PARTICLEHANDLE Particle, const FRotation3& Q)
+		{
+			return Q;
+		}
+
+		template<typename T_PARTICLEHANDLE>
+		static inline const FRigidTransform3& ParticleLocalToCoMLocal(T_PARTICLEHANDLE Particle, const FRigidTransform3& T)
+		{
+			return T;
+		}
+
+		template<typename T_PARTICLEHANDLE>
+		static inline FVec3 GetVelocityAtCoMRelativePosition(T_PARTICLEHANDLE Particle, const FVec3& RelPos)
+		{
+			return Particle->V() + FVec3::CrossProduct(Particle->W(), RelPos);
+		}
+
+		template<typename T_PARTICLEHANDLE>
+		static inline const FVec3& GetCoMWorldPosition(T_PARTICLEHANDLE Particle)
+		{
+			return Particle->P();
+		}
+
+		static inline const FVec3& GetCoMWorldPosition(TPBDRigidParticles<FReal, 3>& Particles, const int32 Index)
+		{
+			return Particles.P(Index);
+		}
+
+		template<typename T_PARTICLEHANDLE>
+		static inline const FRotation3& GetCoMWorldRotation(T_PARTICLEHANDLE Particle)
+		{
+			return Particle->Q();
+		}
+
+		template<typename T_PARTICLEHANDLE>
+		static inline FRigidTransform3 GetCoMWorldTransform(T_PARTICLEHANDLE Particle)
+		{
+			return FRigidTransform3(GetCoMWorldRotation(Particle), GetCoMWorldPosition(Particle));
+		}
+
+		static inline const FRotation3& GetCoMWorldRotation(TPBDRigidParticles<FReal, 3>& Particles, const int32 Index)
+		{
+			return Particles.Q(Index);
+		}
+
+		template<typename T_PARTICLEHANDLE>
+		static inline void SetCoMWorldTransform(T_PARTICLEHANDLE Particle, const FVec3& PCoM, const FRotation3& QCoM)
+		{
+			Particle->P() = PCoM;
+			Particle->Q() = QCoM;
+		}
+
+		static inline void SetCoMWorldTransform(TPBDRigidParticles<FReal, 3>& Particles, const int32 Index, const FVec3& PCoM, const FRotation3& QCoM)
+		{
+			Particles.P(Index) = PCoM;
+			Particles.Q(Index) = QCoM;
+		}
+	};
+
+#if CHAOS_PARTICLE_ACTORTRANSFORM
+	using FParticleUtilities = FParticleUtilities_ActorSpace;
+#else
+	using FParticleUtilities = FParticleUtilities_CoMSpace;
+#endif
+
+}
\ No newline at end of file
diff --git a/Engine/Source/Runtime/Experimental/Chaos/Public/Chaos/ParticleHandle.h b/Engine/Source/Runtime/Experimental/Chaos/Public/Chaos/ParticleHandle.h
--- a/Engine/Source/Runtime/Experimental/Chaos/Public/Chaos/ParticleHandle.h
+++ b/Engine/Source/Runtime/Experimental/Chaos/Public/Chaos/ParticleHandle.h
@@ -873,225 +873,247 @@ TGeometryParticleHandleImp<T,d,bPersistent>* TGeometryParticleHandleImp<T,d, bPe
 	return Ar.IsLoading() ? new TGeometryParticleHandleImp<T, d, bPersistent>() : nullptr;
 }
 
-/**
- * A wrapper around any type of particle handle to provide a consistent (read-only) API for all particle types.
- * This can make code simpler because you can write code that is type-agnostic, but it
- * has a cost. Where possible it is better to write code that is specific to the type(s)
- * of particles being operated on. TGenericParticleHandle has pointer symantics, so you can use one wherever
- * you have a particle handle pointer;
- *
- */
 template <typename T, int d>
-class TGenericParticleHandle
+class TGenericParticleHandleHandleImp
 {
 public:
-	TGenericParticleHandle(TGeometryParticleHandle<T, d>* InHandle) : Imp(InHandle) {}
+	TGenericParticleHandleHandleImp(TGeometryParticleHandle<T, d>* InHandle) { Handle = InHandle; }
+
+	// Check for the exact type of particle (see also AsKinematic etc, which will work on derived types)
+	bool IsStatic() const { return (Handle->ObjectState() == EObjectStateType::Static); }
+	bool IsKinematic() const { return (Handle->ObjectState() == EObjectStateType::Kinematic); }
+	bool IsDynamic() const { return (Handle->ObjectState() == EObjectStateType::Dynamic) || (Handle->ObjectState() == EObjectStateType::Sleeping); }
+
+	const TKinematicGeometryParticleHandle<T, d>* CastToKinematicParticle() const { return Handle->CastToKinematicParticle(); }
+	TKinematicGeometryParticleHandle<T, d>* CastToKinematicParticle() { return Handle->CastToKinematicParticle(); }
+	const TPBDRigidParticleHandle<T, d>* CastToRigidParticle() const { return Handle->CastToRigidParticle(); }
+	TPBDRigidParticleHandle<T, d>* CastToRigidParticle() { return Handle->CastToRigidParticle(); }
+
+	// Static Particles
+	TVector<T, d>& X() { return Handle->X(); }
+	const TVector<T, d>& X() const { return Handle->X(); }
+	TRotation<T, d>& R() { return Handle->R(); }
+	const TRotation<T, d>& R() const { return Handle->R(); }
+	TSerializablePtr<FImplicitObject> Geometry() const { return Handle->Geometry(); }
+	const TUniquePtr<FImplicitObject>& DynamicGeometry() const { return Handle->DynamicGeometry(); }
+	bool Sleeping() const { return Handle->Sleeping(); }
+	FString ToString() const { return Handle->ToString(); }
 
-	class HandleImp
-	{
-	public:
-		HandleImp(TGeometryParticleHandle<T, d>* InHandle) { Handle = InHandle; }
+	template <typename Container>
+	const auto& AuxilaryValue(const Container& AuxContainer) const { return Handle->AuxilaryValue(AuxContainer); }
+	template <typename Container>
+	auto& AuxilaryValue(Container& AuxContainer) { return Handle->AuxilaryValue(AuxContainer); }
 
-		// Check for the exact type of particle (see also AsKinematic etc, which will work on derived types)
-		bool IsStatic() const { return (Handle->ObjectState() == EObjectStateType::Static); }
-		bool IsKinematic() const { return (Handle->ObjectState() == EObjectStateType::Kinematic); }
-		bool IsDynamic() const { return (Handle->ObjectState() == EObjectStateType::Dynamic) || (Handle->ObjectState() == EObjectStateType::Sleeping); }
+	// Kinematic Particles
+	const TVector<T, d>& V() const { return (Handle->CastToKinematicParticle()) ? Handle->CastToKinematicParticle()->V() : ZeroVector; }
+	const TVector<T, d>& W() const { return (Handle->CastToKinematicParticle()) ? Handle->CastToKinematicParticle()->W() : ZeroVector; }
 
-		const TKinematicGeometryParticleHandle<T, d>* CastToKinematicParticle() const { return Handle->CastToKinematicParticle(); }
-		TKinematicGeometryParticleHandle<T, d>* CastToKinematicParticle() { return Handle->CastToKinematicParticle(); }
-		const TPBDRigidParticleHandle<T, d>* CastToRigidParticle() const { return Handle->CastToRigidParticle(); }
-		TPBDRigidParticleHandle<T, d>* CastToRigidParticle() { return Handle->CastToRigidParticle(); }
+	// Dynamic Particles
 
-		// Static Particles
-		const TVector<T, d>& X() const { return Handle->X(); }
-		const TRotation<T, d>& R() const { return Handle->R(); }
-		TSerializablePtr<FImplicitObject> Geometry() const { return Handle->Geometry(); }
-		const TUniquePtr<FImplicitObject>& DynamicGeometry() const { return Handle->DynamicGeometry(); }
-		bool Sleeping() const { return Handle->Sleeping(); }
-		FString ToString() const { return Handle->ToString(); }
+	// TODO: Make all of these check ObjectState to maintain current functionality
+	int32 CollisionParticlesSize() const
+	{
+		if (Handle->CastToRigidParticle() && Handle->ObjectState() == EObjectStateType::Dynamic)
+		{
+			return Handle->CastToRigidParticle()->CollisionParticlesSize();
+		}
 
-		template <typename Container>
-		const auto& AuxilaryValue(const Container& AuxContainer) const { return Handle->AuxilaryValue(AuxContainer); }
-		template <typename Container>
-		auto& AuxilaryValue(Container& AuxContainer) { return Handle->AuxilaryValue(AuxContainer); }
+		return 0;
+	}
 
-		// Kinematic Particles
-		const TVector<T, d>& V() const { return (Handle->CastToKinematicParticle()) ? Handle->CastToKinematicParticle()->V() : ZeroVector; }
-		const TVector<T, d>& W() const { return (Handle->CastToKinematicParticle()) ? Handle->CastToKinematicParticle()->W() : ZeroVector; }
+	const TUniquePtr<TBVHParticles<T, d>>& CollisionParticles() const
+	{
+		if (Handle->CastToRigidParticle() && Handle->ObjectState() == EObjectStateType::Dynamic)
+		{
+			return Handle->CastToRigidParticle()->CollisionParticles();
+		}
 
-		// Dynamic Particles
+		return NullBVHParticles;
+	}
 
-		// TODO: Make all of these check ObjectState to maintain current functionality
-		int32 CollisionParticlesSize() const
+	int32 CollisionGroup() const
+	{
+		if (Handle->CastToRigidParticle() && Handle->ObjectState() == EObjectStateType::Dynamic)
 		{
-			if (Handle->CastToRigidParticle() && Handle->ObjectState() == EObjectStateType::Dynamic)
-			{
-				return Handle->CastToRigidParticle()->CollisionParticlesSize();
-			}
-
-			return 0;
+			return Handle->CastToRigidParticle()->CollisionGroup();
 		}
 
-		const TUniquePtr<TBVHParticles<T, d>>& CollisionParticles() const
-		{
-			if (Handle->CastToRigidParticle() && Handle->ObjectState() == EObjectStateType::Dynamic)
-			{
-				return Handle->CastToRigidParticle()->CollisionParticles();
-			}
+		return 0;
+	}
 
-			return NullBVHParticles;
+	// @todo(ccaulfield): should be available on all types?
+	bool Disabled() const
+	{
+		if (Handle->CastToRigidParticle() && Handle->ObjectState() == EObjectStateType::Dynamic)
+		{
+			return Handle->CastToRigidParticle()->Disabled();
 		}
 
-		int32 CollisionGroup() const
-		{
-			if (Handle->CastToRigidParticle() && Handle->ObjectState() == EObjectStateType::Dynamic)
-			{
-				return Handle->CastToRigidParticle()->CollisionGroup();
-			}
+		return false;
+	}
 
-			return 0;
+	// @todo(ccaulfield): should be available on kinematics?
+	const TVector<T, d>& PreV() const
+	{
+		if (Handle->CastToRigidParticle() && Handle->ObjectState() == EObjectStateType::Dynamic)
+		{
+			return Handle->CastToRigidParticle()->PreV();
 		}
 
-		// @todo(ccaulfield): should be available on all types?
-		bool Disabled() const
-		{
-			if (Handle->CastToRigidParticle() && Handle->ObjectState() == EObjectStateType::Dynamic)
-			{
-				return Handle->CastToRigidParticle()->Disabled();
-			}
+		return ZeroVector;
+	}
 
-			return false;
+	// @todo(ccaulfield): should be available on kinematics?
+	const TVector<T, d>& PreW() const
+	{
+		if (Handle->CastToRigidParticle() && Handle->ObjectState() == EObjectStateType::Dynamic)
+		{
+			return Handle->CastToRigidParticle()->PreW();
 		}
+		return ZeroVector;
+	}
 
-		// @todo(ccaulfield): should be available on kinematics?
-		const TVector<T, d>& PreV() const
+	TVector<T, d>& P()
+	{
+		if (Handle->CastToRigidParticle() && Handle->ObjectState() == EObjectStateType::Dynamic)
 		{
-			if (Handle->CastToRigidParticle() && Handle->ObjectState() == EObjectStateType::Dynamic)
-			{
-				return Handle->CastToRigidParticle()->PreV();
-			}
-
-			return ZeroVector;
+			return Handle->CastToRigidParticle()->P();
 		}
 
-		// @todo(ccaulfield): should be available on kinematics?
-		const TVector<T, d>& PreW() const
+		return X();
+	}
+
+	const TVector<T, d>& P() const
+	{
+		if (Handle->CastToRigidParticle() && Handle->ObjectState() == EObjectStateType::Dynamic)
 		{
-			if (Handle->CastToRigidParticle() && Handle->ObjectState() == EObjectStateType::Dynamic)
-			{
-				return Handle->CastToRigidParticle()->PreW();
-			}
-			return ZeroVector;
+			return Handle->CastToRigidParticle()->P();
 		}
 
-		const TVector<T, d>& P() const
-		{
-			if (Handle->CastToRigidParticle() && Handle->ObjectState() == EObjectStateType::Dynamic)
-			{
-				return Handle->CastToRigidParticle()->P();
-			}
+		return X();
+	}
 
-			return X();
+	TRotation<T, d>& Q()
+	{
+		if (Handle->CastToRigidParticle() && Handle->ObjectState() == EObjectStateType::Dynamic)
+		{
+			return Handle->CastToRigidParticle()->Q();
 		}
 
-		const TRotation<T, d>& Q() const
-		{
-			if (Handle->CastToRigidParticle() && Handle->ObjectState() == EObjectStateType::Dynamic)
-			{
-				return Handle->CastToRigidParticle()->Q();
-			}
+		return R();
+	}
 
-			return R();
+	const TRotation<T, d>& Q() const
+	{
+		if (Handle->CastToRigidParticle() && Handle->ObjectState() == EObjectStateType::Dynamic)
+		{
+			return Handle->CastToRigidParticle()->Q();
 		}
 
-		const TVector<T, d>& F() const 
-		{ 
-			if (Handle->CastToRigidParticle() && Handle->ObjectState() == EObjectStateType::Dynamic)
-			{
-				return Handle->CastToRigidParticle()->F();
-			}
+		return R();
+	}
 
-			return ZeroVector;
+	const TVector<T, d>& F() const 
+	{ 
+		if (Handle->CastToRigidParticle() && Handle->ObjectState() == EObjectStateType::Dynamic)
+		{
+			return Handle->CastToRigidParticle()->F();
 		}
-		const TVector<T, d>& Torque() const 
-		{ 
-			if (Handle->CastToRigidParticle() && Handle->ObjectState() == EObjectStateType::Dynamic)
-			{
-				return Handle->CastToRigidParticle()->Torque();
-			}
 
-			return ZeroVector;
+		return ZeroVector;
+	}
+	const TVector<T, d>& Torque() const 
+	{ 
+		if (Handle->CastToRigidParticle() && Handle->ObjectState() == EObjectStateType::Dynamic)
+		{
+			return Handle->CastToRigidParticle()->Torque();
 		}
 
-		const PMatrix<T, d, d>& I() const 
-		{ 
-			if (Handle->CastToRigidParticle() && Handle->ObjectState() == EObjectStateType::Dynamic)
-			{
-				return Handle->CastToRigidParticle()->I();
-			}
+		return ZeroVector;
+	}
 
-			return ZeroMatrix;
+	const PMatrix<T, d, d>& I() const 
+	{ 
+		if (Handle->CastToRigidParticle() && Handle->ObjectState() == EObjectStateType::Dynamic)
+		{
+			return Handle->CastToRigidParticle()->I();
 		}
 
-		const PMatrix<T, d, d>& InvI() const 
-		{ 
-			if (Handle->CastToRigidParticle() && Handle->ObjectState() == EObjectStateType::Dynamic)
-			{
-				return Handle->CastToRigidParticle()->InvI();
-			}
+		return ZeroMatrix;
+	}
 
-			return ZeroMatrix;
+	const PMatrix<T, d, d>& InvI() const 
+	{ 
+		if (Handle->CastToRigidParticle() && Handle->ObjectState() == EObjectStateType::Dynamic)
+		{
+			return Handle->CastToRigidParticle()->InvI();
 		}
 
-		T M() const
-		{
-			if (Handle->CastToRigidParticle() && Handle->ObjectState() == EObjectStateType::Dynamic)
-			{
-				return Handle->CastToRigidParticle()->M();
-			}
+		return ZeroMatrix;
+	}
 
-			return (T)0;
+	T M() const
+	{
+		if (Handle->CastToRigidParticle() && Handle->ObjectState() == EObjectStateType::Dynamic)
+		{
+			return Handle->CastToRigidParticle()->M();
 		}
 
-		T InvM() const
-		{
-			if (Handle->CastToRigidParticle() && Handle->ObjectState() == EObjectStateType::Dynamic)
-			{
-				return Handle->CastToRigidParticle()->InvM();
-			}
+		return (T)0;
+	}
 
-			return (T)0;
+	T InvM() const
+	{
+		if (Handle->CastToRigidParticle() && Handle->ObjectState() == EObjectStateType::Dynamic)
+		{
+			return Handle->CastToRigidParticle()->InvM();
 		}
 
-		int32 Island() const
-		{
-			if (Handle->CastToRigidParticle() && Handle->ObjectState() == EObjectStateType::Dynamic)
-			{
-				return Handle->CastToRigidParticle()->Island();
-			}
+		return (T)0;
+	}
 
-			return INDEX_NONE;
+	TVector<T, d> CenterOfMass() const
+	{
+		if (auto KinematicHandle = Handle->CastToKinematicParticle())
+		{
+			return KinematicHandle->CenterOfMass();
 		}
 
-		bool ToBeRemovedOnFracture() const 
+		return TVector<T, d>(0);
+	}
+
+	TRotation<T, d> RotationOfMass() const
+	{
+		if (auto KinematicHandle = Handle->CastToKinematicParticle())
 		{
-			if (Handle->CastToRigidParticle() && Handle->ObjectState() == EObjectStateType::Dynamic)
-			{
-				return Handle->CastToRigidParticle()->ToBeRemovedOnFracture();
-			}
+			return KinematicHandle->RotationOfMass();
+		}
+
+		return TRotation<T, d>::FromIdentity();
+	}
 
-			return false;
+	int32 Island() const
+	{
+		if (Handle->CastToRigidParticle() && Handle->ObjectState() == EObjectStateType::Dynamic)
+		{
+			return Handle->CastToRigidParticle()->Island();
 		}
 
-	private:
-		TGeometryParticleHandle<T, d>* Handle;
-	};
+		return INDEX_NONE;
+	}
 
-	HandleImp* operator->() { return &Imp; }
-	const HandleImp* operator->() const { return &Imp; }
+	bool ToBeRemovedOnFracture() const 
+	{
+		if (Handle->CastToRigidParticle() && Handle->ObjectState() == EObjectStateType::Dynamic)
+		{
+			return Handle->CastToRigidParticle()->ToBeRemovedOnFracture();
+		}
+
+		return false;
+	}
 
 private:
-	HandleImp Imp;
+	TGeometryParticleHandle<T, d>* Handle;
 
 	static const TVector<T, d> ZeroVector;
 	static const TRotation<T, d> IdentityRotation;
@@ -1099,17 +1121,51 @@ private:
 	static const TUniquePtr<TBVHParticles<T, d>> NullBVHParticles;
 };
 
+/**
+ * A wrapper around any type of particle handle to provide a consistent (read-only) API for all particle types.
+ * This can make code simpler because you can write code that is type-agnostic, but it
+ * has a cost. Where possible it is better to write code that is specific to the type(s)
+ * of particles being operated on. TGenericParticleHandle has pointer semantics, so you can use one wherever
+ * you have a particle handle pointer;
+ *
+ */
+template <typename T, int d>
+class TGenericParticleHandle
+{
+public:
+	TGenericParticleHandle(TGeometryParticleHandle<T, d>* InHandle) : Imp(InHandle) {}
+
+	TGenericParticleHandleHandleImp<T, d>* operator->() const { return const_cast<TGenericParticleHandleHandleImp<T, d>*>(&Imp); }
+
+private:
+	TGenericParticleHandleHandleImp<T, d> Imp;
+};
+
+template <typename T, int d>
+class TConstGenericParticleHandle
+{
+public:
+	TConstGenericParticleHandle(const TGeometryParticleHandle<T, d>* InHandle) : Imp(const_cast<TGeometryParticleHandle<T, d>*>(InHandle)) {}
+
+	const TGenericParticleHandleHandleImp<T, d>* operator->() const { return &Imp; }
+
+private:
+	TGenericParticleHandleHandleImp<T, d> Imp;
+};
+
 template <typename T, int d>
-const TVector<T, d> TGenericParticleHandle<T, d>::ZeroVector = TVector<T, d>(0);
+const TVector<T, d> TGenericParticleHandleHandleImp<T, d>::ZeroVector = TVector<T, d>(0);
 
 template <typename T, int d>
-const TRotation<T, d> TGenericParticleHandle<T, d>::IdentityRotation = TRotation<T, d>(0, 0, 0, 1);
+const TRotation<T, d> TGenericParticleHandleHandleImp<T, d>::IdentityRotation = TRotation<T, d>(0, 0, 0, 1);
 
 template <typename T, int d>
-const PMatrix<T, d, d> TGenericParticleHandle<T, d>::ZeroMatrix = PMatrix<T, d, d>(0);
+const PMatrix<T, d, d> TGenericParticleHandleHandleImp<T, d>::ZeroMatrix = PMatrix<T, d, d>(0);
 
 template <typename T, int d>
-const TUniquePtr<TBVHParticles<T, d>> TGenericParticleHandle<T, d>::NullBVHParticles = TUniquePtr<TBVHParticles<T, d>>();
+const TUniquePtr<TBVHParticles<T, d>> TGenericParticleHandleHandleImp<T, d>::NullBVHParticles = TUniquePtr<TBVHParticles<T, d>>();
+
+
 
 template <typename T, int d>
 class TGeometryParticleHandles : public TArrayCollection
diff --git a/Engine/Source/Runtime/Experimental/Chaos/Public/Chaos/PerParticlePBDEulerStep.h b/Engine/Source/Runtime/Experimental/Chaos/Public/Chaos/PerParticlePBDEulerStep.h
--- a/Engine/Source/Runtime/Experimental/Chaos/Public/Chaos/PerParticlePBDEulerStep.h
+++ b/Engine/Source/Runtime/Experimental/Chaos/Public/Chaos/PerParticlePBDEulerStep.h
@@ -1,11 +1,15 @@
 // Copyright 1998-2019 Epic Games, Inc. All Rights Reserved.
 #pragma once
 
-#include "Chaos/PerParticleRule.h"
+#include "Chaos/Particle/ParticleUtilities.h"
 #include "Chaos/ParticleHandle.h"
+#include "Chaos/PerParticleRule.h"
 
 namespace Chaos
 {
+/**
+ * Update position and rotation based on velocity and angular velocity.
+ */
 template<class T, int d>
 class TPerParticlePBDEulerStep : public TPerParticleRule<T, d>
 {
@@ -26,16 +30,24 @@ class TPerParticlePBDEulerStep : public TPerParticleRule<T, d>
 
 	inline void Apply(TPBDRigidParticles<T, d>& InParticles, const T Dt, const int32 Index) const override //-V762
 	{
-		ApplyHelper(InParticles, Dt, Index);
-		InParticles.Q(Index) = InParticles.R(Index) + TRotation<T, d>::FromElements(InParticles.W(Index), 0.f) * InParticles.R(Index) * Dt * T(0.5);
-		InParticles.Q(Index).Normalize();
+		FVec3 PCoM = FParticleUtilities::GetCoMWorldPosition(InParticles, Index);
+		FRotation3 QCoM = FParticleUtilities::GetCoMWorldRotation(InParticles, Index);
+
+		PCoM = PCoM + InParticles.V(Index) * Dt;
+		QCoM = FRotation3::IntegrateRotationWithAngularVelocity(QCoM, InParticles.W(Index), Dt);
+
+		FParticleUtilities::SetCoMWorldTransform(InParticles, Index, PCoM, QCoM);
 	}
 
 	inline void Apply(TTransientPBDRigidParticleHandle<T, d>& Particle, const T Dt) const override
 	{
-		Particle.P() = Particle.X() + Particle.V() * Dt;
-		Particle.Q() = Particle.R() + TRotation<T, d>::FromElements(Particle.W(), 0.f) * Particle.R() * Dt * T(0.5);
-		Particle.Q().Normalize();
+		FVec3 PCoM = FParticleUtilities::GetCoMWorldPosition(&Particle);
+		FRotation3 QCoM = FParticleUtilities::GetCoMWorldRotation(&Particle);
+
+		PCoM = PCoM + Particle.V() * Dt;
+		QCoM = FRotation3::IntegrateRotationWithAngularVelocity(QCoM, Particle.W(), Dt);
+
+		FParticleUtilities::SetCoMWorldTransform(&Particle, PCoM, QCoM);
 	}
 };
 }
diff --git a/Engine/Source/Runtime/Experimental/Chaos/Public/Chaos/PerParticlePBDUpdateFromDeltaPosition.h b/Engine/Source/Runtime/Experimental/Chaos/Public/Chaos/PerParticlePBDUpdateFromDeltaPosition.h
--- a/Engine/Source/Runtime/Experimental/Chaos/Public/Chaos/PerParticlePBDUpdateFromDeltaPosition.h
+++ b/Engine/Source/Runtime/Experimental/Chaos/Public/Chaos/PerParticlePBDUpdateFromDeltaPosition.h
@@ -1,6 +1,7 @@
 // Copyright 1998-2019 Epic Games, Inc. All Rights Reserved.
 #pragma once
 
+#include "Chaos/Particle/ParticleUtilities.h"
 #include "Chaos/PBDParticles.h"
 #include "Chaos/PBDRigidParticles.h"
 #include "Chaos/PerParticleRule.h"
@@ -35,14 +36,28 @@ class TPerParticlePBDUpdateFromDeltaPosition : public TPerParticleRule<T, d>
 
 	inline void Apply(TPBDRigidParticleHandle<T, d>* Particle, const T Dt) const override //-V762
 	{
-		Particle->V() = (Particle->P() - Particle->X()) / Dt;
-		Particle->W() = TRotation<T, d>::CalculateAngularVelocity(Particle->R(), Particle->Q(), Dt);
+#if CHAOS_PARTICLE_ACTORTRANSFORM
+		const FVec3& CenterOfMass = Particle->CenterOfMass();
+		const FVec3 CenteredX = Particle->X() + Particle->R().RotateVector(CenterOfMass);
+		const FVec3 CenteredP = Particle->P() + Particle->Q().RotateVector(CenterOfMass);
+		Particle->V() = FVec3::CalculateVelocity(CenteredX, CenteredP, Dt);
+#else
+		Particle->V() = FVec3::CalculateVelocity(Particle->X(), Particle->P(), Dt);
+#endif
+		Particle->W() = FRotation3::CalculateAngularVelocity(Particle->R(), Particle->Q(), Dt);
 	}
 
 	inline void Apply(TTransientPBDRigidParticleHandle<T, d>& Particle, const T Dt) const override //-V762
 	{
-		Particle.V() = (Particle.P() - Particle.X()) / Dt;
-		Particle.W() = TRotation<T, d>::CalculateAngularVelocity(Particle.R(), Particle.Q(), Dt);
+#if CHAOS_PARTICLE_ACTORTRANSFORM
+		const FVec3& CenterOfMass = Particle.CenterOfMass();
+		const FVec3 CenteredX = Particle.X() + Particle.R().RotateVector(CenterOfMass);
+		const FVec3 CenteredP = Particle.P() + Particle.Q().RotateVector(CenterOfMass);
+		Particle.V() = FVec3::CalculateVelocity(CenteredX, CenteredP, Dt);
+#else
+		Particle.V() = FVec3::CalculateVelocity(Particle.X(), Particle.P(), Dt);
+#endif
+		Particle.W() = FRotation3::CalculateAngularVelocity(Particle.R(), Particle.Q(), Dt);
 	}
 };
 }
