Change convex convex collision to use GJKPenetration. Fix missing convex convex pairs like box vs convex. Temp disable box box unit test until a known EPA bug is fixed.

diff --git a/Engine/Source/Runtime/Experimental/Chaos/Private/Chaos/CollisionResolution.cpp b/Engine/Source/Runtime/Experimental/Chaos/Private/Chaos/CollisionResolution.cpp
--- a/Engine/Source/Runtime/Experimental/Chaos/Private/Chaos/CollisionResolution.cpp
+++ b/Engine/Source/Runtime/Experimental/Chaos/Private/Chaos/CollisionResolution.cpp
@@ -16,12 +16,68 @@
 #include "Chaos/PBDCollisionConstraints.h"
 #include "Chaos/Sphere.h"
 #include "Chaos/Transform.h"
-
+#include "Chaos/GeometryQueries.h"
 
 namespace Chaos
 {
 	namespace Collisions
 	{
+		template <typename T>
+		struct TContactPoint
+		{
+			TVec3<T> Normal;
+			TVec3<T> Location;
+			T Phi;
+
+			TContactPoint() : Phi(TNumericLimits<T>::Max()) {}
+		};
+
+		template <typename T>
+		void UpdateContactPoint(TPointContactManifold<T, 3>& Manifold, const TContactPoint<T>& NewContactPoint)
+		{
+			//for now just override
+			if (NewContactPoint.Phi < Manifold.Phi)
+			{
+				Manifold.Normal = NewContactPoint.Normal;
+				Manifold.Location = NewContactPoint.Location;
+				Manifold.Phi = NewContactPoint.Phi;
+			}
+		}
+
+		template <typename T, int d, typename GeometryA, typename GeometryB>
+		TContactPoint<T> GJKContactPoint(const GeometryA& A, const TRigidTransform<T, d>& ATM, const GeometryB& B, const TRigidTransform<T, d>& BTM, const T Thickness)
+		{
+			TContactPoint<T> Contact;
+			const TRigidTransform<T, d> BToATM = BTM.GetRelativeTransform(ATM);
+
+			T Penetration;
+			TVec3<T> ClosestA, ClosestB, Normal;
+			if (GJKPenetration(A, B, BToATM, Penetration, ClosestA, ClosestB, Normal, (T)0))
+			{
+				Contact.Location = ATM.TransformPosition(ClosestA);
+				Contact.Normal = -ATM.TransformVector(Normal);
+				Contact.Phi = -Penetration;
+			}
+			else if(ensure(GJKDistance(A, B, BToATM, Penetration, ClosestA, ClosestB)))
+			{
+				//todo: make GJKPenetration support no penetration case
+				TVector<T, d> NearPointAWorld = ATM.TransformPosition(ClosestA);
+				TVector<T, d> NearPointBWorld = BTM.TransformPosition(ClosestB);
+				TVector<T, d> NearPointBtoAWorld = NearPointAWorld - NearPointBWorld;
+				Contact.Phi = Penetration;
+				Contact.Normal = NearPointBtoAWorld.GetSafeNormal();
+				Contact.Location = NearPointAWorld;
+			}
+
+			return Contact;
+		}
+
+
+		template <typename T, int d>
+		TContactPoint<T> BoxBoxContactPoint(const TBox<T, d>& Box1, const TRigidTransform<T, d>& ATM, const TBox<T, d>& Box2, const TRigidTransform<T, d>& BTM, const T Thickness)
+		{
+			return GJKContactPoint(Box1, ATM, Box2, BTM, Thickness);
+		}
 		//
 		// Box - Box
 		//
@@ -29,53 +85,7 @@ namespace Chaos
 		template <typename T, int d>
 		void UpdateBoxBoxConstraint(const TBox<T, d>& Box1, const TRigidTransform<T, d>& Box1Transform, const TBox<T, d>& Box2, const TRigidTransform<T, d>& Box2Transform, const T Thickness, TRigidBodyPointContactConstraint<T, d>& Constraint)
 		{
-			TPointContactManifold<T, d> & Contact = Constraint.Manifold;
-
-			TBox<T, d> Box2SpaceBox1 = Box1.TransformedBox(Box1Transform * Box2Transform.Inverse());
-			TBox<T, d> Box1SpaceBox2 = Box2.TransformedBox(Box2Transform * Box1Transform.Inverse());
-			Box2SpaceBox1.Thicken(Thickness);
-			Box1SpaceBox2.Thicken(Thickness);
-			if (Box1SpaceBox2.Intersects(Box1) && Box2SpaceBox1.Intersects(Box2))
-			{
-				const TVector<T, d> Box1Center = (Box1Transform * Box2Transform.Inverse()).TransformPosition(Box1.Center());
-				bool bDeepOverlap = false;
-				if (Box2.SignedDistance(Box1Center) < 0)
-				{
-					//If Box1 is overlapping Box2 by this much the signed distance approach will fail (box1 gets sucked into box2). In this case just use two spheres
-					TSphere<T, d> Sphere1(Box1Transform.TransformPosition(Box1.Center()), Box1.Extents().Min() / 2);
-					TSphere<T, d> Sphere2(Box2Transform.TransformPosition(Box2.Center()), Box2.Extents().Min() / 2);
-					const TVector<T, d> Direction = Sphere1.GetCenter() - Sphere2.GetCenter();
-					T Size = Direction.Size();
-					if (Size < (Sphere1.GetRadius() + Sphere2.GetRadius()))
-					{
-						const T NewPhi = Size - (Sphere1.GetRadius() + Sphere2.GetRadius());;
-						if (NewPhi < Contact.Phi)
-						{
-							bDeepOverlap = true;
-							Contact.Normal = Size > SMALL_NUMBER ? Direction / Size : TVector<T, d>(0, 0, 1);
-							Contact.Phi = NewPhi;
-							Contact.Location = Sphere1.GetCenter() - Sphere1.GetRadius() * Contact.Normal;
-						}
-					}
-				}
-				if (!bDeepOverlap || Contact.Phi >= 0)
-				{
-					//if we didn't have deep penetration use signed distance per particle. If we did have deep penetration but the spheres did not overlap use signed distance per particle
-
-					//UpdateLevelsetConstraintGJK(InParticles, Thickness, Constraint);
-					//check(Contact.Phi < MThickness);
-					// For now revert to doing all points vs lsv check until we can figure out a good way to get the deepest point without needing this
-					{
-						const TArray<TVector<T, d>> SampleParticles = Box1.ComputeSamplePoints();
-						const TRigidTransform<T, d> Box1ToBox2Transform = Box1Transform.GetRelativeTransform(Box2Transform);
-						int32 NumParticles = SampleParticles.Num();
-						for (int32 i = 0; i < NumParticles; ++i)
-						{
-							SampleObjectHelper(Box2, Box2Transform, Box1ToBox2Transform, SampleParticles[i], Thickness, Constraint);
-						}
-					}
-				}
-			}
+			UpdateContactPoint(Constraint.Manifold, BoxBoxContactPoint(Box1, Box1Transform, Box2, Box2Transform, Thickness));
 		}
 
 
@@ -201,22 +211,28 @@ namespace Chaos
 		//
 		// Sphere - Sphere
 		//
+
 		template <typename T, int d>
-		void UpdateSphereSphereConstraint(const TSphere<T, d>& Sphere1, const TRigidTransform<T, d>& Sphere1Transform, const TSphere<T, d>& Sphere2, const TRigidTransform<T, d>& Sphere2Transform, const T Thickness, TRigidBodyPointContactConstraint<T, d>& Constraint)
+		TContactPoint<T> SphereSphereContactPoint(const TSphere<T, d>& Sphere1, const TRigidTransform<T, d>& Sphere1Transform, const TSphere<T, d>& Sphere2, const TRigidTransform<T, d>& Sphere2Transform, const T Thickness)
 		{
-			TPointContactManifold<T, d> & Contact = Constraint.Manifold;
+			TContactPoint<T> Result;
 
 			const TVector<T, d> Center1 = Sphere1Transform.TransformPosition(Sphere1.GetCenter());
 			const TVector<T, d> Center2 = Sphere2Transform.TransformPosition(Sphere2.GetCenter());
 			const TVector<T, d> Direction = Center1 - Center2;
 			const T Size = Direction.Size();
 			const T NewPhi = Size - (Sphere1.GetRadius() + Sphere2.GetRadius());
-			if (NewPhi < Contact.Phi)
-			{
-				Contact.Normal = Size > SMALL_NUMBER ? Direction / Size : TVector<T, d>(0, 0, 1);
-				Contact.Phi = NewPhi;
-				Contact.Location = Center1 - Sphere1.GetRadius() * Contact.Normal;
-			}
+			Result.Phi = NewPhi;
+			Result.Normal = Size > SMALL_NUMBER ? Direction / Size : TVector<T, d>(0, 0, 1);
+			Result.Location = Center1 - Sphere1.GetRadius() * Result.Normal;
+
+			return Result;
+		}
+
+		template <typename T, int d>
+		void UpdateSphereSphereConstraint(const TSphere<T, d>& Sphere1, const TRigidTransform<T, d>& Sphere1Transform, const TSphere<T, d>& Sphere2, const TRigidTransform<T, d>& Sphere2Transform, const T Thickness, TRigidBodyPointContactConstraint<T, d>& Constraint)
+		{
+			UpdateContactPoint(Constraint.Manifold, SphereSphereContactPoint(Sphere1, Sphere1Transform, Sphere2, Sphere2Transform, Thickness));
 		}
 
 		template<typename T, int d>
@@ -301,23 +317,27 @@ namespace Chaos
 		//
 
 		template <typename T, int d>
-		void UpdateSphereBoxConstraint(const TSphere<T, d>& Sphere, const TRigidTransform<T, d>& SphereTransform, const TBox<T, d>& Box, const TRigidTransform<T, d>& BoxTransform, const T Thickness, TRigidBodyPointContactConstraint<T, d>& Constraint)
+		TContactPoint<T> SphereBoxContactPoint(const TSphere<T, d>& Sphere, const TRigidTransform<T, d>& SphereTransform, const TBox<T, d>& Box, const TRigidTransform<T, d>& BoxTransform, const T Thickness)
 		{
-			TPointContactManifold<T, d> & Contact = Constraint.Manifold;
+			TContactPoint<T> Result;
 
-			const TRigidTransform<T, d> SphereToBoxTransform(SphereTransform * BoxTransform.Inverse());
+			const TRigidTransform<T, d> SphereToBoxTransform(SphereTransform * BoxTransform.Inverse());	//todo: this should use GetRelative
 			const TVector<T, d> SphereCenterInBox = SphereToBoxTransform.TransformPosition(Sphere.GetCenter());
 
 			TVector<T, d> NewNormal;
 			T NewPhi = Box.PhiWithNormal(SphereCenterInBox, NewNormal);
 			NewPhi -= Sphere.GetRadius();
 
-			if (NewPhi < Contact.Phi)
-			{
-				Contact.Phi = NewPhi;
-				Contact.Normal = BoxTransform.TransformVectorNoScale(NewNormal);
-				Contact.Location = SphereTransform.TransformPosition(Sphere.GetCenter()) - Contact.Normal * Sphere.GetRadius();
-			}
+			Result.Phi = NewPhi;
+			Result.Normal = BoxTransform.TransformVectorNoScale(NewNormal);
+			Result.Location = SphereTransform.TransformPosition(Sphere.GetCenter()) - Result.Normal * Sphere.GetRadius();
+			return Result;
+		}
+
+		template <typename T, int d>
+		void UpdateSphereBoxConstraint(const TSphere<T, d>& Sphere, const TRigidTransform<T, d>& SphereTransform, const TBox<T, d>& Box, const TRigidTransform<T, d>& BoxTransform, const T Thickness, TRigidBodyPointContactConstraint<T, d>& Constraint)
+		{
+			UpdateContactPoint(Constraint.Manifold, SphereBoxContactPoint(Sphere, SphereTransform, Box, BoxTransform, Thickness));
 		}
 
 		template<typename T, int d>
@@ -352,9 +372,9 @@ namespace Chaos
 		//
 
 		template <typename T, int d>
-		void UpdateSphereCapsuleConstraint(const TSphere<T, d>& A, const TRigidTransform<T, d>& ATransform, const TCapsule<T>& B, const TRigidTransform<T, d>& BTransform, const T Thickness, TRigidBodyPointContactConstraint<T, d>& Constraint)
+		TContactPoint<T> SphereCapsuleContactPoint(const TSphere<T, d>& A, const TRigidTransform<T, d>& ATransform, const TCapsule<T>& B, const TRigidTransform<T, d>& BTransform, const T Thickness)
 		{
-			TPointContactManifold<T, d>& Contact = Constraint.Manifold;
+			TContactPoint<T> Result;
 
 			FVector A1 = ATransform.TransformPosition(A.GetCenter());
 			FVector B1 = BTransform.TransformPosition(B.GetX1());
@@ -366,14 +386,19 @@ namespace Chaos
 			if (DeltaLen > KINDA_SMALL_NUMBER)
 			{
 				T NewPhi = DeltaLen - (A.GetRadius() + B.GetRadius());
-				if (NewPhi < Contact.Phi)
-				{
-					TVector<T, d> Dir = Delta / DeltaLen;
-					Contact.Phi = NewPhi;
-					Contact.Normal = -Dir;
-					Contact.Location = A1 + Dir * A.GetRadius();
-				}
+				TVector<T, d> Dir = Delta / DeltaLen;
+				Result.Phi = NewPhi;
+				Result.Normal = -Dir;
+				Result.Location = A1 + Dir * A.GetRadius();
 			}
+
+			return Result;
+		}
+
+		template <typename T, int d>
+		void UpdateSphereCapsuleConstraint(const TSphere<T, d>& A, const TRigidTransform<T, d>& ATransform, const TCapsule<T>& B, const TRigidTransform<T, d>& BTransform, const T Thickness, TRigidBodyPointContactConstraint<T, d>& Constraint)
+		{
+			UpdateContactPoint(Constraint.Manifold, SphereCapsuleContactPoint(A, ATransform, B, BTransform, Thickness));
 		}
 
 		template<typename T, int d>
@@ -407,9 +432,9 @@ namespace Chaos
 		//
 
 		template <typename T, int d>
-		void UpdateCapsuleCapsuleConstraint(const TCapsule<T>& A, const TRigidTransform<T, d>& ATransform, const TCapsule<T>& B, const TRigidTransform<T, d>& BTransform, const T Thickness, TRigidBodyPointContactConstraint<T, d>& Constraint)
+		TContactPoint<T> CapsuleCapsuleContactPoint(const TCapsule<T>& A, const TRigidTransform<T, d>& ATransform, const TCapsule<T>& B, const TRigidTransform<T, d>& BTransform, const T Thickness)
 		{
-			TPointContactManifold<T, d> & Contact = Constraint.Manifold;
+			TContactPoint<T> Result;
 
 			FVector A1 = ATransform.TransformPosition(A.GetX1());
 			FVector A2 = ATransform.TransformPosition(A.GetX2());
@@ -423,14 +448,19 @@ namespace Chaos
 			if (DeltaLen > KINDA_SMALL_NUMBER)
 			{
 				T NewPhi = DeltaLen - (A.GetRadius() + B.GetRadius());
-				if (NewPhi < Contact.Phi)
-				{
-					TVector<T, d> Dir = Delta / DeltaLen;
-					Contact.Phi = NewPhi;
-					Contact.Normal = -Dir;
-					Contact.Location = P1 + Dir * A.GetRadius();
-				}
+				TVector<T, d> Dir = Delta / DeltaLen;
+				Result.Phi = NewPhi;
+				Result.Normal = -Dir;
+				Result.Location = P1 + Dir * A.GetRadius();
 			}
+
+			return Result;
+		}
+
+		template <typename T, int d>
+		void UpdateCapsuleCapsuleConstraint(const TCapsule<T>& A, const TRigidTransform<T, d>& ATransform, const TCapsule<T>& B, const TRigidTransform<T, d>& BTransform, const T Thickness, TRigidBodyPointContactConstraint<T, d>& Constraint)
+		{
+			UpdateContactPoint(Constraint.Manifold, CapsuleCapsuleContactPoint(A, ATransform, B, BTransform, Thickness));
 		}
 
 		template<typename T, int d>
@@ -464,40 +494,15 @@ namespace Chaos
 		//
 
 		template <typename T, int d>
-		void UpdateCapsuleBoxConstraint(const TCapsule<T>& A, const TRigidTransform<T, d>& ATransform, const TBox<T, d>& B, const TRigidTransform<T, d>& BTransform, const T Thickness, TRigidBodyPointContactConstraint<T, d>& Constraint)
+		TContactPoint<T> CapsuleBoxContactPoint(const TCapsule<T>& A, const TRigidTransform<T, d>& ATransform, const TBox<T, d>& B, const TRigidTransform<T, d>& BTransform, const T Thickness)
 		{
-			// @todo(ccaulfield): Add custom capsule-box collision
-			TPointContactManifold<T, d> & Contact = Constraint.Manifold;
-
-			TRigidTransform<T, d> BToATransform = BTransform.GetRelativeTransform(ATransform);
-
-			// Use GJK to track closest points (not strictly necessary yet)
-			TVector<T, d> NearPointALocal, NearPointBLocal;
-			T NearPointDistance;
-			if (GJKDistance<T>(A, B, BToATransform, NearPointDistance, NearPointALocal, NearPointBLocal))
-			{
-				TVector<T, d> NearPointAWorld = ATransform.TransformPosition(NearPointALocal);
-				TVector<T, d> NearPointBWorld = BTransform.TransformPosition(NearPointBLocal);
-				TVector<T, d> NearPointBtoAWorld = NearPointAWorld - NearPointBWorld;
-				Contact.Phi = NearPointDistance;
-				Contact.Normal = NearPointBtoAWorld.GetSafeNormal();
-				Contact.Location = NearPointAWorld;
-			}
-			else
-			{
-				// Use box particle samples against the implicit capsule
-				const TArray<TVector<T, d>> BParticles = B.ComputeSamplePoints();
-				const int32 NumParticles = BParticles.Num();
-				for (int32 ParticleIndex = 0; ParticleIndex < NumParticles; ++ParticleIndex)
-				{
-					if (SampleObjectHelper(A, ATransform, BToATransform, BParticles[ParticleIndex], Thickness, Constraint))
-					{
-						// SampleObjectHelper expects A to be the box, so reverse the results
-						Contact.Normal = -Contact.Normal;
-					}
-				}
-			}
+			return GJKContactPoint(A, ATransform, B, BTransform, Thickness);
+		}
 
+		template <typename T, int d>
+		void UpdateCapsuleBoxConstraint(const TCapsule<T>& A, const TRigidTransform<T, d>& ATransform, const TBox<T, d>& B, const TRigidTransform<T, d>& BTransform, const T Thickness, TRigidBodyPointContactConstraint<T, d>& Constraint)
+		{
+			UpdateContactPoint(Constraint.Manifold, CapsuleBoxContactPoint(A, ATransform, B, BTransform, Thickness));
 		}
 
 		template<typename T, int d>
@@ -531,78 +536,41 @@ namespace Chaos
 		//
 
 		template<class T, int d>
-		void UpdateConvexConvexConstraint(const FImplicitObject& A, const TRigidTransform<T, d>& ATM, const FImplicitObject& B, const TRigidTransform<T, d>& BTM, const T Thickness, TRigidBodyPointContactConstraint<T, d>& Constraint)
+		TContactPoint<T> ConvexConvexContactPoint(const FImplicitObject& A, const TRigidTransform<T, d>& ATM, const FImplicitObject& B, const TRigidTransform<T, d>& BTM, const T Thickness)
 		{
-			TRigidTransform<T, d> BToATM = BTM.GetRelativeTransform(ATM);
-
-			if (ensure(GetInnerType(A.GetType()) == ImplicitObjectType::Convex && GetInnerType(B.GetType()) == ImplicitObjectType::Convex))
+			return CastHelper(A, [&](const auto& ADowncast)
 			{
-				const FConvex* AObject = nullptr;
-				const FConvex* BObject = nullptr;
-
-				if (IsScaled(A.GetType()) && IsScaled(B.GetType()))
+				return CastHelper(B, [&](const auto& BDowncast)
 				{
-					AObject = static_cast<const FConvex*>(static_cast<const TImplicitObjectScaled<FConvex>&>(A).Object().Get());
-					BObject = static_cast<const FConvex*>(static_cast<const TImplicitObjectScaled<FConvex>&>(B).Object().Get());
-				}
-				else if (A.GetType() == ImplicitObjectType::Transformed && B.GetType() == ImplicitObjectType::Transformed)
-				{
-					AObject = static_cast<const FConvex*>(static_cast<const TImplicitObjectTransformed<T, d>&>(A).Object().Get());
-					BObject = static_cast<const FConvex*>(static_cast<const TImplicitObjectTransformed<T, d>&>(B).Object().Get());
-				}
-
-				if (AObject && BObject)
-				{
-					const TParticles<T, d>& SurfaceParticles = AObject->GetSurfaceParticles();
-
-					Constraint.SetDisabled(true);
-
-					if (GJKIntersection(*AObject, *BObject, BToATM, Thickness))
-					{
-						TRigidTransform<T, d> AToBTM = ATM.GetRelativeTransform(BTM);
-
-						T Phi;
-						TVector<T, d> Normal;
-						for (int32 Idx = 0; Idx < (int32)SurfaceParticles.Size(); Idx++)
-						{
-							Phi = BObject->PhiWithNormal(AToBTM.TransformPosition(SurfaceParticles.X(Idx)), Normal);
-							if (Phi < Constraint.GetPhi())
-							{
-								Constraint.SetPhi(Phi);
-								Constraint.SetLocation(ATM.TransformPosition(SurfaceParticles.X(Idx)));
-								Constraint.SetNormal(BTM.TransformVector(Normal));
-							}
-						}
-						return;
-					}
-					return;
-				}
-			}
+					return GJKContactPoint(ADowncast, ATM, BDowncast, BTM, Thickness);
+				});
+			});
+		}
 
-			ensureMsgf(false, TEXT("Unsupported convex to convex constraint."));
+		template<class T, int d>
+		void UpdateConvexConvexConstraint(const FImplicitObject& A, const TRigidTransform<T, d>& ATM, const FImplicitObject& B, const TRigidTransform<T, d>& BTM, const T Thickness, TRigidBodyPointContactConstraint<T, d>& Constraint)
+		{
+			UpdateContactPoint(Constraint.Manifold, ConvexConvexContactPoint(A, ATM, B, BTM, Thickness));
 		}
 
 		template<class T, int d>
 		void ConstructConvexConvexConstraints(TGeometryParticleHandle<T, d>* Particle0, TGeometryParticleHandle<T, d>* Particle1, const FImplicitObject* Implicit0, const FImplicitObject* Implicit1, const TRigidTransform<T, d>& Transform0, const TRigidTransform<T, d>& Transform1, const T Thickness, FCollisionConstraintsArray& NewConstraints)
 		{
-			if (ensure(GetInnerType(Implicit0->GetType()) == ImplicitObjectType::Convex && GetInnerType(Implicit1->GetType()) == ImplicitObjectType::Convex))
-			{
-				FRigidBodyPointContactConstraint* Constraint = new FRigidBodyPointContactConstraint;
+			FRigidBodyPointContactConstraint* Constraint = new FRigidBodyPointContactConstraint;
 
-				Constraint->Particle[0] = Particle0;
-				Constraint->Particle[1] = Particle1;
-				Constraint->AddManifold(Implicit0, Implicit1);
+			Constraint->Particle[0] = Particle0;
+			Constraint->Particle[1] = Particle1;
+			Constraint->AddManifold(Implicit0, Implicit1);
 
-				UpdateConvexConvexConstraint(*Implicit0, Transform0, *Implicit1, Transform1, Thickness, *Constraint);
+			UpdateConvexConvexConstraint(*Implicit0, Transform0, *Implicit1, Transform1, Thickness, *Constraint);
 
-				if (Constraint->GetPhi() < Thickness)
-				{
-					NewConstraints.Add(Constraint);
-				}
-				else
-				{
-					delete Constraint;
-				}
+			if (Constraint->GetPhi() < Thickness)
+			{
+				NewConstraints.Add(Constraint);
+			}
+			else
+			{
+				delete Constraint;
 			}
 		}
 
@@ -1018,7 +986,7 @@ namespace Chaos
 			{
 				return UpdateUnionUnionConstraint<UpdateType>(Implicit0, Transform0, Implicit1, Transform1, Thickness, ConstraintBase);
 			}
-			else if(Implicit0Type == FConvex::StaticType() && Implicit1Type == FConvex::StaticType())
+			else if(Implicit0.IsConvex() && Implicit1.IsConvex())
 			{
 				UpdateConvexConvexConstraint(Implicit0, Transform0, Implicit1, Transform1, Thickness, *ConstraintBase.template As<TRigidBodyPointContactConstraint<T, d>>());
 			}
@@ -1153,7 +1121,7 @@ namespace Chaos
 				// Union-union creates multiple manifolds - we should never get here for this pair type. See ConstructConstraintsImpl and ConstructUnionUnionConstraints
 				ensure(false);
 			}
-			else if (Implicit0Type == FConvex::StaticType() && Implicit1Type == FConvex::StaticType())
+			else if (Implicit0->IsConvex() && Implicit1->IsConvex())
 			{
 				ConstructConvexConvexConstraints(Particle0, Particle1, Implicit0, Implicit1, Transform0, Transform1, Thickness, NewConstraints);
 			}
diff --git a/Engine/Source/Runtime/Experimental/Chaos/Public/Chaos/GJK.h b/Engine/Source/Runtime/Experimental/Chaos/Public/Chaos/GJK.h
--- a/Engine/Source/Runtime/Experimental/Chaos/Public/Chaos/GJK.h
+++ b/Engine/Source/Runtime/Experimental/Chaos/Public/Chaos/GJK.h
@@ -203,7 +203,7 @@ namespace Chaos
 			if (EPA(VertsA, VertsB, SupportAFunc, SupportBFunc, Penetration, MTD, ClosestA, ClosestBInA) != EPAResult::BadInitialSimplex)
 			{
 				OutNormal = MTD;
-				OutPenetration = Penetration + Inflation;
+				OutPenetration = Penetration + ThicknessA + ThicknessB;
 				OutClosestA = ClosestA + OutNormal * ThicknessA;
 				OutClosestB = ClosestBInA - OutNormal * ThicknessB;
 			}
@@ -220,7 +220,7 @@ namespace Chaos
 					ClosestBInA += Bs[i] * Barycentric[i];
 				}
 
-				OutPenetration = Inflation;
+				OutPenetration = ThicknessA + ThicknessB;
 				OutNormal = { 0,0,1 };
 				OutClosestA = ClosestA + OutNormal * ThicknessA;
 				OutClosestB = ClosestBInA - OutNormal * ThicknessB;
@@ -642,7 +642,12 @@ namespace Chaos
 	template <typename T, typename TGeometryA, typename TGeometryB>
 	TVector<T, 3> GJKDistanceInitialV(const TGeometryA& A, const TGeometryB& B, const TRigidTransform<T, 3>& BToATM)
 	{
-		return A.GetCenter() - (B.GetCenter() + BToATM.GetTranslation());
+		const TVec3<T> V(1, 0, 0);
+		const TVector<T, 3> SupportA = A.Support(-V, 0);
+		const TVector<T, 3> VInB = BToATM.GetRotation().Inverse() * V;
+		const TVector<T, 3> SupportBLocal = B.Support(VInB, 0);
+		const TVector<T, 3> SupportB = BToATM.TransformPositionNoScale(SupportBLocal);
+		return SupportA - SupportB;
 	}
 
 	/**
