Make sure EPA handles the case where the origin is outside the CSO

diff --git a/Engine/Source/Runtime/Experimental/Chaos/Public/Chaos/EPA.h b/Engine/Source/Runtime/Experimental/Chaos/Public/Chaos/EPA.h
--- a/Engine/Source/Runtime/Experimental/Chaos/Public/Chaos/EPA.h
+++ b/Engine/Source/Runtime/Experimental/Chaos/Public/Chaos/EPA.h
@@ -425,7 +425,8 @@ EPAResult EPA(TArray<TVec3<T>>& VertsABuffer, TArray<TVec3<T>>& VertsBBuffer, co
 
 		LowerBound = Entry.Distance;
 
-		if (FMath::Abs(UpperBound - LowerBound) <= Eps * LowerBound)
+		//It's possible the origin is not contained by the CSO. In this case the upper bound will be negative, at which point we should just exit. Maybe return a different enum value?
+		if ((UpperBound - LowerBound) <= FMath::Abs(Eps * LowerBound))
 		{
 			//UE_LOG(LogChaos, Warning, TEXT("Iteration:%d"), Iteration);
 			ComputeEPAResults(VertsABuffer.GetData(), VertsBBuffer.GetData(), Entry, OutPenetration, OutDir, WitnessA, WitnessB);
