Chaos : Collisions

Combined the Apply / ApplyPushOut / Update forPoint and Plane constraints.

diff --git a/Engine/Source/Runtime/Experimental/Chaos/Private/Chaos/CollisionResolutionUtil.cpp b/Engine/Source/Runtime/Experimental/Chaos/Private/Chaos/CollisionResolutionUtil.cpp
--- a/Engine/Source/Runtime/Experimental/Chaos/Private/Chaos/CollisionResolutionUtil.cpp
+++ b/Engine/Source/Runtime/Experimental/Chaos/Private/Chaos/CollisionResolutionUtil.cpp
@@ -45,7 +45,7 @@ namespace Chaos
 		}
 
 		template<typename T, int d>
-		TVector<T, d> GetEnergyClampedImpulse(const TRigidBodyPointContactConstraint<T, d>& Constraint, const TVector<T, d>& Impulse, const TVector<T, d>& VectorToPoint1, const TVector<T, d>& VectorToPoint2, const TVector<T, d>& Velocity1, const TVector<T, d>& Velocity2)
+		TVector<T, d> GetEnergyClampedImpulse(const TCollisionConstraintBase<T, d>& Constraint, const TVector<T, d>& Impulse, const TVector<T, d>& VectorToPoint1, const TVector<T, d>& VectorToPoint2, const TVector<T, d>& Velocity1, const TVector<T, d>& Velocity2)
 		{
 			TPBDRigidParticleHandle<T, d>* PBDRigid0 = Constraint.Particle[0]->CastToRigidParticle();
 			TPBDRigidParticleHandle<T, d>* PBDRigid1 = Constraint.Particle[1]->CastToRigidParticle();
@@ -598,7 +598,7 @@ namespace Chaos
 
 		template TRigidTransform<float, 3> GetTransform(const TGeometryParticleHandle<float, 3>* Particle);
 		template PMatrix<float, 3, 3> ComputeFactorMatrix3(const TVector<float, 3>& V, const PMatrix<float, 3, 3>& M, const float& Im);
-		template TVector<float, 3> GetEnergyClampedImpulse(const TRigidBodyPointContactConstraint<float, 3>& Constraint, const TVector<float, 3>& Impulse, const TVector<float, 3>& VectorToPoint1, const TVector<float, 3>& VectorToPoint2, const TVector<float, 3>& Velocity1, const TVector<float, 3>& Velocity2);
+		template TVector<float, 3> GetEnergyClampedImpulse(const TCollisionConstraintBase<float, 3>& Constraint, const TVector<float, 3>& Impulse, const TVector<float, 3>& VectorToPoint1, const TVector<float, 3>& VectorToPoint2, const TVector<float, 3>& Velocity1, const TVector<float, 3>& Velocity2);
 		template bool SampleObjectHelper<float, 3>(const FImplicitObject& Object, const TRigidTransform<float, 3>& ObjectTransform, const TRigidTransform<float, 3>& SampleToObjectTransform, const TVector<float, 3>& SampleParticle, float Thickness, TRigidBodyPointContactConstraint<float, 3>& Constraint);
 		template bool SampleObjectNoNormal(const FImplicitObject& Object, const TRigidTransform<float, 3>& ObjectTransform, const TRigidTransform<float, 3>& SampleToObjectTransform, const TVector<float, 3>& SampleParticle, float Thickness, TRigidBodyPointContactConstraint<float, 3>& Constraint);
 		template bool SampleObjectNormalAverageHelper(const FImplicitObject& Object, const TRigidTransform<float, 3>& ObjectTransform, const TRigidTransform<float, 3>& SampleToObjectTransform, const TVector<float, 3>& SampleParticle, float Thickness, float& TotalThickness, TRigidBodyPointContactConstraint<float, 3>& Constraint);
diff --git a/Engine/Source/Runtime/Experimental/Chaos/Private/Chaos/PBDCollisionConstraints.cpp b/Engine/Source/Runtime/Experimental/Chaos/Private/Chaos/PBDCollisionConstraints.cpp
--- a/Engine/Source/Runtime/Experimental/Chaos/Private/Chaos/PBDCollisionConstraints.cpp
+++ b/Engine/Source/Runtime/Experimental/Chaos/Private/Chaos/PBDCollisionConstraints.cpp
@@ -14,8 +14,7 @@
 #include "Chaos/SpatialAccelerationCollection.h"
 #include "Chaos/Levelset.h"
 #include "Chaos/Pair.h"
-#include "Chaos/PBDCollisionConstraintsPointContactUtil.h"
-#include "Chaos/PBDCollisionConstraintsPlaneContactUtil.h"
+#include "Chaos/PBDCollisionConstraintsContact.h"
 #include "Chaos/PBDRigidsSOAs.h"
 #include "Chaos/Sphere.h"
 #include "Chaos/Transform.h"
@@ -224,22 +223,9 @@ namespace Chaos
 				FConstraintContainerHandle* ConstraintHandle = InConstraintHandles[ConstraintHandleIndex];
 				check(ConstraintHandle != nullptr);
 
-				if (ConstraintHandle->GetType() == FPointContactConstraint::StaticType())
-				{
-					Collisions::TPointContactParticleParameters<T> ParticleParameters = { &MCollided, &MPhysicsMaterials, CollisionFrictionOverride, MAngularFriction };
-					Collisions::TPointContactIterationParameters<T> IterationParameters = { Dt, Iterations, NumIterations, MApplyPairIterations, nullptr };
-					Collisions::Apply(ConstraintHandle->GetPointContact(), MThickness, IterationParameters, ParticleParameters);
-				}
-				else if (ConstraintHandle->GetType() == FPlaneContactConstraint::StaticType())
-				{
-					Collisions::TPlaneContactParticleParameters<T> ParticleParameters = { &MCollided, &MPhysicsMaterials, CollisionFrictionOverride, MAngularFriction };
-					Collisions::TPlaneContactIterationParameters<T> IterationParameters = { Dt, Iterations, NumIterations, MApplyPairIterations, nullptr };
-					Collisions::Apply(ConstraintHandle->GetPlaneContact(), MThickness, IterationParameters, ParticleParameters);
-				}
-				else
-				{
-					ensureMsgf(false, TEXT("Invalid constraint type"));
-				}
+				Collisions::TContactParticleParameters<T> ParticleParameters = { &MCollided, &MPhysicsMaterials, CollisionFrictionOverride, MAngularFriction };
+				Collisions::TContactIterationParameters<T> IterationParameters = { Dt, Iterations, NumIterations, MApplyPairIterations, nullptr };
+				Collisions::Apply(ConstraintHandle->GetContact(), MThickness, IterationParameters, ParticleParameters);
 
 			}, bDisableCollisionParallelFor);
 		}
@@ -265,23 +251,9 @@ namespace Chaos
 				FConstraintContainerHandle* ConstraintHandle = InConstraintHandles[ConstraintHandleIndex];
 				check(ConstraintHandle != nullptr);
 
-				if (ConstraintHandle->GetType() == FPointContactConstraint::StaticType())
-				{
-					Collisions::TPointContactParticleParameters<T> ParticleParameters = { &MCollided, &MPhysicsMaterials, CollisionFrictionOverride, MAngularFriction };
-					Collisions::TPointContactIterationParameters<T> IterationParameters = { Dt, Iteration, NumIterations, MApplyPushOutPairIterations, &NeedsAnotherIteration };
-					Collisions::ApplyPushOut(ConstraintHandle->GetPointContact(), MThickness, IsTemporarilyStatic, IterationParameters, ParticleParameters);
-				}
-				else if (ConstraintHandle->GetType() == FPlaneContactConstraint::StaticType())
-				{
-					Collisions::TPlaneContactParticleParameters<T> ParticleParameters = { &MCollided, &MPhysicsMaterials, CollisionFrictionOverride, MAngularFriction };
-					Collisions::TPlaneContactIterationParameters<T> IterationParameters = { Dt, Iteration, NumIterations, MApplyPushOutPairIterations, &NeedsAnotherIteration };
-					Collisions::ApplyPushOut(ConstraintHandle->GetPlaneContact(), MThickness, IsTemporarilyStatic, IterationParameters, ParticleParameters);
-				}
-				else
-				{
-					ensureMsgf(false, TEXT("Invalid constraint type"));
-				}
-
+				Collisions::TContactParticleParameters<T> ParticleParameters = { &MCollided, &MPhysicsMaterials, CollisionFrictionOverride, MAngularFriction };
+				Collisions::TContactIterationParameters<T> IterationParameters = { Dt, Iteration, NumIterations, MApplyPushOutPairIterations, &NeedsAnotherIteration };
+				Collisions::ApplyPushOut(ConstraintHandle->GetContact(), MThickness, IsTemporarilyStatic, IterationParameters, ParticleParameters);
 
 			}, bDisableCollisionParallelFor);
 		}
diff --git a/Engine/Source/Runtime/Experimental/Chaos/Private/Chaos/PBDCollisionConstraintsContact.cpp b/Engine/Source/Runtime/Experimental/Chaos/Private/Chaos/PBDCollisionConstraintsContact.cpp
new file mode 100644
--- /dev/null
+++ b/Engine/Source/Runtime/Experimental/Chaos/Private/Chaos/PBDCollisionConstraintsContact.cpp
@@ -0,0 +1,390 @@
+// Copyright 1998-2019 Epic Games, Inc. All Rights Reserved.
+
+#include "Chaos/PBDCollisionConstraintsContact.h"
+#include "Chaos/CollisionResolution.h"
+#include "Chaos/CollisionResolutionUtil.h"
+#include "Chaos/Defines.h"
+#include "Chaos/Particle/ParticleUtilities.h"
+#include "Chaos/Utilities.h"
+
+namespace Chaos
+{
+	namespace Collisions
+	{
+		template<ECollisionUpdateType UpdateType, typename T, int d>
+		void Update(const T Thickness, TCollisionConstraintBase<T, d>& Constraint)
+		{
+			Constraint.ResetPhi(Thickness);
+			const TRigidTransform<T, d> ParticleTM = GetTransform(Constraint.Particle[0]);
+			const TRigidTransform<T, d> LevelsetTM = GetTransform(Constraint.Particle[1]);
+
+			if (Constraint.GetType() == TCollisionConstraintBase<T, d>::FType::SinglePoint)
+			{
+				if (!Constraint.Particle[0]->Geometry())
+				{
+					if (Constraint.Particle[1]->Geometry())
+					{
+						if (!Constraint.Particle[1]->Geometry()->IsUnderlyingUnion())
+						{
+							UpdateLevelsetLevelsetConstraint<UpdateType>(Thickness, *Constraint.template As<TRigidBodyPointContactConstraint<T, d>>());
+						}
+						else
+						{
+							UpdateUnionLevelsetConstraint<UpdateType>(Thickness, *Constraint.template As<TRigidBodyPointContactConstraint<T, d>>());
+						}
+					}
+				}
+				else
+				{
+					UpdateConstraintImp<UpdateType>(*Constraint.Particle[0]->Geometry(), ParticleTM, *Constraint.Particle[1]->Geometry(), LevelsetTM, Thickness, Constraint);
+				}
+			}
+			else if(Constraint.GetType() == TCollisionConstraintBase<T, d>::FType::Plane)
+			{
+				UpdateConstraintImp<UpdateType>(*Constraint.Particle[0]->Geometry(), ParticleTM, *Constraint.Particle[1]->Geometry(), LevelsetTM, Thickness, Constraint);
+			}
+		}
+
+
+		template<typename T, int d>
+		void Apply(TCollisionConstraintBase<T, d>& Constraint, T Thickness, TContactIterationParameters<T> & IterationParameters, TContactParticleParameters<T> & ParticleParameters)
+		{
+			TGenericParticleHandle<T, d> Particle0 = TGenericParticleHandle<T, d>(Constraint.Particle[0]);
+			TGenericParticleHandle<T, d> Particle1 = TGenericParticleHandle<T, d>(Constraint.Particle[1]);
+			TPBDRigidParticleHandle<T, d>* PBDRigid0 = Particle0->CastToRigidParticle();
+			TPBDRigidParticleHandle<T, d>* PBDRigid1 = Particle1->CastToRigidParticle();
+			bool bIsRigidDynamic0 = PBDRigid0 && PBDRigid0->ObjectState() == EObjectStateType::Dynamic;
+			bool bIsRigidDynamic1 = PBDRigid1 && PBDRigid1->ObjectState() == EObjectStateType::Dynamic;
+
+			for (int32 PairIt = 0; PairIt < IterationParameters.NumPairIterations; ++PairIt)
+			{
+				Collisions::Update<ECollisionUpdateType::Deepest>(Thickness, Constraint);
+				if (Constraint.GetPhi() >= Thickness)
+				{
+					return;
+				}
+
+				// @todo(ccaulfield): CHAOS_PARTICLEHANDLE_TODO what's the best way to manage external per-particle data?
+				if (ParticleParameters.Collided)
+				{
+					Particle0->AuxilaryValue(*ParticleParameters.Collided) = true;
+					Particle1->AuxilaryValue(*ParticleParameters.Collided) = true;
+				}
+
+				TSerializablePtr<FChaosPhysicsMaterial> PhysicsMaterial0, PhysicsMaterial1;
+				if (ParticleParameters.PhysicsMaterials)
+				{
+					PhysicsMaterial0 = Particle0->AuxilaryValue(*ParticleParameters.PhysicsMaterials);
+					PhysicsMaterial1 = Particle1->AuxilaryValue(*ParticleParameters.PhysicsMaterials);
+				}
+
+				const TVector<T, d> ZeroVector = TVector<T, d>(0);
+				TVector<T, d> P0 = FParticleUtilities::GetCoMWorldPosition(Particle0);
+				TVector<T, d> P1 = FParticleUtilities::GetCoMWorldPosition(Particle1);
+				TRotation<T, d> Q0 = FParticleUtilities::GetCoMWorldRotation(Particle0);
+				TRotation<T, d> Q1 = FParticleUtilities::GetCoMWorldRotation(Particle1);
+
+				typename TRigidBodyPointContactConstraint<T, d>::FManifold & Contact = Constraint.Manifold;
+
+				TVector<T, d> VectorToPoint1 = Contact.Location - P0;
+				TVector<T, d> VectorToPoint2 = Contact.Location - P1;
+				TVector<T, d> Body1Velocity = FParticleUtilities::GetVelocityAtCoMRelativePosition(Particle0, VectorToPoint1);
+				TVector<T, d> Body2Velocity = FParticleUtilities::GetVelocityAtCoMRelativePosition(Particle1, VectorToPoint2);
+				TVector<T, d> RelativeVelocity = Body1Velocity - Body2Velocity;
+				T RelativeNormalVelocity = TVector<T, d>::DotProduct(RelativeVelocity, Contact.Normal);
+
+				if (RelativeNormalVelocity < 0) // ignore separating constraints
+				{
+					PMatrix<T, d, d> WorldSpaceInvI1 = bIsRigidDynamic0 ? Utilities::ComputeWorldSpaceInertia(Q0, PBDRigid0->InvI()) : PMatrix<T, d, d>(0);
+					PMatrix<T, d, d> WorldSpaceInvI2 = bIsRigidDynamic1 ? Utilities::ComputeWorldSpaceInertia(Q1, PBDRigid1->InvI()) : PMatrix<T, d, d>(0);
+					PMatrix<T, d, d> Factor =
+						(bIsRigidDynamic0 ? ComputeFactorMatrix3(VectorToPoint1, WorldSpaceInvI1, PBDRigid0->InvM()) : PMatrix<T, d, d>(0)) +
+						(bIsRigidDynamic1 ? ComputeFactorMatrix3(VectorToPoint2, WorldSpaceInvI2, PBDRigid1->InvM()) : PMatrix<T, d, d>(0));
+					TVector<T, d> Impulse;
+					TVector<T, d> AngularImpulse(0);
+
+					// Resting contact if very close to the surface
+					T Restitution = (T)0;
+					T Friction = (T)0;
+					T AngularFriction = (T)0;
+					bool bApplyRestitution = (RelativeVelocity.Size() > (2 * 980 * IterationParameters.Dt));
+					if (PhysicsMaterial0 && PhysicsMaterial1)
+					{
+						if (bApplyRestitution)
+						{
+							Restitution = FMath::Min(PhysicsMaterial0->Restitution, PhysicsMaterial1->Restitution);
+						}
+						Friction = FMath::Max(PhysicsMaterial0->Friction, PhysicsMaterial1->Friction);
+					}
+					else if (PhysicsMaterial0)
+					{
+						if (bApplyRestitution)
+						{
+							Restitution = PhysicsMaterial0->Restitution;
+						}
+						Friction = PhysicsMaterial0->Friction;
+					}
+					else if (PhysicsMaterial1)
+					{
+						if (bApplyRestitution)
+						{
+							Restitution = PhysicsMaterial1->Restitution;
+						}
+						Friction = PhysicsMaterial1->Friction;
+					}
+
+					if (ParticleParameters.FrictionOverride >= 0)
+					{
+						Friction = ParticleParameters.FrictionOverride;
+					}
+					if (ParticleParameters.AngularFrictionOverride >= 0)
+					{
+						AngularFriction = ParticleParameters.AngularFrictionOverride;
+					}
+
+					if (Friction)
+					{
+						if (RelativeNormalVelocity > 0)
+						{
+							RelativeNormalVelocity = 0;
+						}
+						TVector<T, d> VelocityChange = -(Restitution * RelativeNormalVelocity * Contact.Normal + RelativeVelocity);
+						T NormalVelocityChange = TVector<T, d>::DotProduct(VelocityChange, Contact.Normal);
+						PMatrix<T, d, d> FactorInverse = Factor.Inverse();
+						TVector<T, d> MinimalImpulse = FactorInverse * VelocityChange;
+						const T MinimalImpulseDotNormal = TVector<T, d>::DotProduct(MinimalImpulse, Contact.Normal);
+						const T TangentialSize = (MinimalImpulse - MinimalImpulseDotNormal * Contact.Normal).Size();
+						if (TangentialSize <= Friction * MinimalImpulseDotNormal)
+						{
+							//within friction cone so just solve for static friction stopping the object
+							Impulse = MinimalImpulse;
+							if (AngularFriction)
+							{
+								TVector<T, d> RelativeAngularVelocity = Particle0->W() - Particle1->W();
+								T AngularNormal = TVector<T, d>::DotProduct(RelativeAngularVelocity, Contact.Normal);
+								TVector<T, d> AngularTangent = RelativeAngularVelocity - AngularNormal * Contact.Normal;
+								TVector<T, d> FinalAngularVelocity = FMath::Sign(AngularNormal) * FMath::Max((T)0, FMath::Abs(AngularNormal) - AngularFriction * NormalVelocityChange) * Contact.Normal + FMath::Max((T)0, AngularTangent.Size() - AngularFriction * NormalVelocityChange) * AngularTangent.GetSafeNormal();
+								TVector<T, d> Delta = FinalAngularVelocity - RelativeAngularVelocity;
+								if (!bIsRigidDynamic0 && bIsRigidDynamic1)
+								{
+									PMatrix<T, d, d> WorldSpaceI2 = (Q1 * FMatrix::Identity) * PBDRigid1->I() * (Q1 * FMatrix::Identity).GetTransposed();
+									TVector<T, d> ImpulseDelta = PBDRigid1->M() * TVector<T, d>::CrossProduct(VectorToPoint2, Delta);
+									Impulse += ImpulseDelta;
+									AngularImpulse += WorldSpaceI2 * Delta - TVector<T, d>::CrossProduct(VectorToPoint2, ImpulseDelta);
+								}
+								else if (bIsRigidDynamic0 && !bIsRigidDynamic1)
+								{
+									PMatrix<T, d, d> WorldSpaceI1 = (Q0 * FMatrix::Identity) * PBDRigid0->I() * (Q0 * FMatrix::Identity).GetTransposed();
+									TVector<T, d> ImpulseDelta = PBDRigid0->M() * TVector<T, d>::CrossProduct(VectorToPoint1, Delta);
+									Impulse += ImpulseDelta;
+									AngularImpulse += WorldSpaceI1 * Delta - TVector<T, d>::CrossProduct(VectorToPoint1, ImpulseDelta);
+								}
+								else if (bIsRigidDynamic0 && bIsRigidDynamic1)
+								{
+									PMatrix<T, d, d> Cross1(0, VectorToPoint1.Z, -VectorToPoint1.Y, -VectorToPoint1.Z, 0, VectorToPoint1.X, VectorToPoint1.Y, -VectorToPoint1.X, 0);
+									PMatrix<T, d, d> Cross2(0, VectorToPoint2.Z, -VectorToPoint2.Y, -VectorToPoint2.Z, 0, VectorToPoint2.X, VectorToPoint2.Y, -VectorToPoint2.X, 0);
+									PMatrix<T, d, d> CrossI1 = Cross1 * WorldSpaceInvI1;
+									PMatrix<T, d, d> CrossI2 = Cross2 * WorldSpaceInvI2;
+									PMatrix<T, d, d> Diag1 = CrossI1 * Cross1.GetTransposed() + CrossI2 * Cross2.GetTransposed();
+									Diag1.M[0][0] += PBDRigid0->InvM() + PBDRigid1->InvM();
+									Diag1.M[1][1] += PBDRigid0->InvM() + PBDRigid1->InvM();
+									Diag1.M[2][2] += PBDRigid0->InvM() + PBDRigid1->InvM();
+									PMatrix<T, d, d> OffDiag1 = (CrossI1 + CrossI2) * -1;
+									PMatrix<T, d, d> Diag2 = (WorldSpaceInvI1 + WorldSpaceInvI2).Inverse();
+									PMatrix<T, d, d> OffDiag1Diag2 = OffDiag1 * Diag2;
+									TVector<T, d> ImpulseDelta = PMatrix<T, d, d>((Diag1 - OffDiag1Diag2 * OffDiag1.GetTransposed()).Inverse())* ((OffDiag1Diag2 * -1) * Delta);
+									Impulse += ImpulseDelta;
+									AngularImpulse += Diag2 * (Delta - PMatrix<T, d, d>(OffDiag1.GetTransposed()) * ImpulseDelta);
+								}
+							}
+						}
+						else
+						{
+							//outside friction cone, solve for normal relative velocity and keep tangent at cone edge
+							TVector<T, d> Tangent = (RelativeVelocity - TVector<T, d>::DotProduct(RelativeVelocity, Contact.Normal) * Contact.Normal).GetSafeNormal();
+							TVector<T, d> DirectionalFactor = Factor * (Contact.Normal - Friction * Tangent);
+							T ImpulseDenominator = TVector<T, d>::DotProduct(Contact.Normal, DirectionalFactor);
+							if (!ensureMsgf(FMath::Abs(ImpulseDenominator) > SMALL_NUMBER, TEXT("Constraint:%s\n\nParticle:%s\n\nLevelset:%s\n\nDirectionalFactor:%s, ImpulseDenominator:%f"),
+								*Constraint.ToString(),
+								*Particle0->ToString(),
+								*Particle1->ToString(),
+								*DirectionalFactor.ToString(), ImpulseDenominator))
+							{
+								ImpulseDenominator = (T)1;
+							}
+
+							const T ImpulseMag = -(1 + Restitution) * RelativeNormalVelocity / ImpulseDenominator;
+							Impulse = ImpulseMag * (Contact.Normal - Friction * Tangent);
+						}
+					}
+					else
+					{
+						T ImpulseDenominator = TVector<T, d>::DotProduct(Contact.Normal, Factor * Contact.Normal);
+						TVector<T, d> ImpulseNumerator = -(1 + Restitution) * TVector<T, d>::DotProduct(RelativeVelocity, Contact.Normal)* Contact.Normal;
+						if (!ensureMsgf(FMath::Abs(ImpulseDenominator) > SMALL_NUMBER, TEXT("Constraint:%s\n\nParticle:%s\n\nLevelset:%s\n\nFactor*Constraint.Normal:%s, ImpulseDenominator:%f"),
+							*Constraint.ToString(),
+							*Particle0->ToString(),
+							*Particle1->ToString(),
+							*(Factor * Contact.Normal).ToString(), ImpulseDenominator))
+						{
+							ImpulseDenominator = (T)1;
+						}
+						Impulse = ImpulseNumerator / ImpulseDenominator;
+					}
+					Impulse = GetEnergyClampedImpulse(Constraint, Impulse, VectorToPoint1, VectorToPoint2, Body1Velocity, Body2Velocity);
+					Constraint.AccumulatedImpulse += Impulse;
+					if (bIsRigidDynamic0)
+					{
+						// Velocity update for next step
+						FVec3 NetAngularImpulse = FVec3::CrossProduct(VectorToPoint1, Impulse) + AngularImpulse;
+						FVec3 DV = PBDRigid0->InvM() * Impulse;
+						FVec3 DW = WorldSpaceInvI1 * NetAngularImpulse;
+						PBDRigid0->V() += DV;
+						PBDRigid0->W() += DW;
+						// Position update as part of pbd
+						P0 += (DV * IterationParameters.Dt);
+						Q0 += FRotation3::FromElements(DW, 0.f) * Q0 * IterationParameters.Dt * T(0.5);
+						Q0.Normalize();
+						FParticleUtilities::SetCoMWorldTransform(PBDRigid0, P0, Q0);
+					}
+					if (bIsRigidDynamic1)
+					{
+						// Velocity update for next step
+						FVec3 NetAngularImpulse = FVec3::CrossProduct(VectorToPoint2, -Impulse) - AngularImpulse;
+						FVec3 DV = -PBDRigid1->InvM() * Impulse;
+						FVec3 DW = WorldSpaceInvI2 * NetAngularImpulse;
+						PBDRigid1->V() += DV;
+						PBDRigid1->W() += DW;
+						// Position update as part of pbd
+						P1 += (DV * IterationParameters.Dt);
+						Q1 += FRotation3::FromElements(DW, 0.f) * Q1 * IterationParameters.Dt * T(0.5);
+						Q1.Normalize();
+						FParticleUtilities::SetCoMWorldTransform(PBDRigid1, P1, Q1);
+					}
+				}
+			}
+		}
+
+		template<typename T, int d>
+		void ApplyPushOut(TCollisionConstraintBase<T, d>& Constraint, T Thickness, const TSet<const TGeometryParticleHandle<T, d>*>& IsTemporarilyStatic,
+			TContactIterationParameters<T> & IterationParameters, TContactParticleParameters<T> & ParticleParameters)
+		{
+			TGenericParticleHandle<T, d> Particle0 = TGenericParticleHandle<T, d>(Constraint.Particle[0]);
+			TGenericParticleHandle<T, d> Particle1 = TGenericParticleHandle<T, d>(Constraint.Particle[1]);
+			TPBDRigidParticleHandle<T, d>* PBDRigid0 = Particle0->CastToRigidParticle();
+			TPBDRigidParticleHandle<T, d>* PBDRigid1 = Particle1->CastToRigidParticle();
+			const bool bIsRigidDynamic0 = PBDRigid0 && PBDRigid0->ObjectState() == EObjectStateType::Dynamic;
+			const bool bIsRigidDynamic1 = PBDRigid1 && PBDRigid1->ObjectState() == EObjectStateType::Dynamic;
+
+			const TVector<T, d> ZeroVector = TVector<T, d>(0);
+			TVector<T, d> P0 = FParticleUtilities::GetCoMWorldPosition(Particle0);
+			TVector<T, d> P1 = FParticleUtilities::GetCoMWorldPosition(Particle1);
+			TRotation<T, d> Q0 = FParticleUtilities::GetCoMWorldRotation(Particle0);
+			TRotation<T, d> Q1 = FParticleUtilities::GetCoMWorldRotation(Particle1);
+			const bool IsTemporarilyStatic0 = IsTemporarilyStatic.Contains(Constraint.Particle[0]);
+			const bool IsTemporarilyStatic1 = IsTemporarilyStatic.Contains(Constraint.Particle[1]);
+
+			TSerializablePtr<FChaosPhysicsMaterial> PhysicsMaterial0, PhysicsMaterial1;
+			if (ParticleParameters.PhysicsMaterials)
+			{
+				PhysicsMaterial0 = Particle0->AuxilaryValue(*ParticleParameters.PhysicsMaterials);
+				PhysicsMaterial1 = Particle1->AuxilaryValue(*ParticleParameters.PhysicsMaterials);
+			}
+
+			for (int32 PairIteration = 0; PairIteration < IterationParameters.NumPairIterations; ++PairIteration)
+			{
+				Update<ECollisionUpdateType::Deepest>(Thickness, Constraint);
+
+				const typename TRigidBodyPointContactConstraint<T, d>::FManifold & Contact = Constraint.Manifold;
+
+				if (Contact.Phi >= Thickness)
+				{
+					break;
+				}
+
+				if ((!bIsRigidDynamic0 || IsTemporarilyStatic0) && (!bIsRigidDynamic1 || IsTemporarilyStatic1))
+				{
+					break;
+				}
+
+				*IterationParameters.NeedsAnotherIteration = true;
+				PMatrix<T, d, d> WorldSpaceInvI1 = bIsRigidDynamic0 ? Utilities::ComputeWorldSpaceInertia(Q0, PBDRigid0->InvI()) : PMatrix<T, d, d>(0);
+				PMatrix<T, d, d> WorldSpaceInvI2 = bIsRigidDynamic1 ? Utilities::ComputeWorldSpaceInertia(Q1, PBDRigid1->InvI()) : PMatrix<T, d, d>(0);
+				TVector<T, d> VectorToPoint1 = Contact.Location - P0;
+				TVector<T, d> VectorToPoint2 = Contact.Location - P1;
+				PMatrix<T, d, d> Factor =
+					(bIsRigidDynamic0 ? ComputeFactorMatrix3(VectorToPoint1, WorldSpaceInvI1, PBDRigid0->InvM()) : PMatrix<T, d, d>(0)) +
+					(bIsRigidDynamic1 ? ComputeFactorMatrix3(VectorToPoint2, WorldSpaceInvI2, PBDRigid1->InvM()) : PMatrix<T, d, d>(0));
+				T Numerator = FMath::Min((T)(IterationParameters.Iteration + 2), (T)IterationParameters.NumIterations);
+				T ScalingFactor = Numerator / (T)IterationParameters.NumIterations;
+
+				//if pushout is needed we better fix relative velocity along normal. Treat it as if 0 restitution
+				TVector<T, d> Body1Velocity = FParticleUtilities::GetVelocityAtCoMRelativePosition(Particle0, VectorToPoint1);
+				TVector<T, d> Body2Velocity = FParticleUtilities::GetVelocityAtCoMRelativePosition(Particle1, VectorToPoint2);
+				TVector<T, d> RelativeVelocity = Body1Velocity - Body2Velocity;
+				const T RelativeVelocityDotNormal = TVector<T, d>::DotProduct(RelativeVelocity, Contact.Normal);
+				if (RelativeVelocityDotNormal < 0)
+				{
+					T ImpulseDenominator = TVector<T, d>::DotProduct(Contact.Normal, Factor * Contact.Normal);
+					TVector<T, d> ImpulseNumerator = -TVector<T, d>::DotProduct(RelativeVelocity, Contact.Normal) * Contact.Normal * ScalingFactor;
+					if (!ensureMsgf(FMath::Abs(ImpulseDenominator) > SMALL_NUMBER, TEXT("ApplyPushout Constraint:%s\n\nParticle:%s\n\nLevelset:%s\n\nFactor*Contact.Normal:%s, ImpulseDenominator:%f"),
+						*Constraint.ToString(),
+						*Particle0->ToString(),
+						*Particle1->ToString(),
+						*(Factor*Contact.Normal).ToString(), ImpulseDenominator))
+					{
+						ImpulseDenominator = (T)1;
+					}
+
+					TVector<T, d> VelocityFixImpulse = ImpulseNumerator / ImpulseDenominator;
+					VelocityFixImpulse = GetEnergyClampedImpulse(Constraint, VelocityFixImpulse, VectorToPoint1, VectorToPoint2, Body1Velocity, Body2Velocity);
+					Constraint.AccumulatedImpulse += VelocityFixImpulse;	//question: should we track this?
+					if (!IsTemporarilyStatic0 && bIsRigidDynamic0)
+					{
+						TVector<T, d> AngularImpulse = TVector<T, d>::CrossProduct(VectorToPoint1, VelocityFixImpulse);
+						PBDRigid0->V() += PBDRigid0->InvM() * VelocityFixImpulse;
+						PBDRigid0->W() += WorldSpaceInvI1 * AngularImpulse;
+
+					}
+
+					if (!IsTemporarilyStatic1 && bIsRigidDynamic1)
+					{
+						TVector<T, d> AngularImpulse = TVector<T, d>::CrossProduct(VectorToPoint2, -VelocityFixImpulse);
+						PBDRigid1->V() -= PBDRigid1->InvM() * VelocityFixImpulse;
+						PBDRigid1->W() += WorldSpaceInvI2 * AngularImpulse;
+					}
+
+				}
+
+
+				TVector<T, d> Impulse = PMatrix<T, d, d>(Factor.Inverse()) * ((-Contact.Phi + Thickness) * ScalingFactor * Contact.Normal);
+				TVector<T, d> AngularImpulse1 = TVector<T, d>::CrossProduct(VectorToPoint1, Impulse);
+				TVector<T, d> AngularImpulse2 = TVector<T, d>::CrossProduct(VectorToPoint2, -Impulse);
+				if (!IsTemporarilyStatic0 && bIsRigidDynamic0)
+				{
+					P0 += PBDRigid0->InvM() * Impulse;
+					Q0 = TRotation<T, d>::FromVector(WorldSpaceInvI1 * AngularImpulse1) * Q0;
+					Q0.Normalize();
+					FParticleUtilities::SetCoMWorldTransform(Particle0, P0, Q0);
+				}
+				if (!IsTemporarilyStatic1 && bIsRigidDynamic1)
+				{
+					P1 -= PBDRigid1->InvM() * Impulse;
+					Q1 = TRotation<T, d>::FromVector(WorldSpaceInvI2 * AngularImpulse2) * Q1;
+					Q1.Normalize();
+					FParticleUtilities::SetCoMWorldTransform(Particle1, P1, Q1);
+				}
+			}
+		}
+
+
+		template void Update<ECollisionUpdateType::Any, float, 3>(const float, TCollisionConstraintBase<float, 3>&);
+		template void Update<ECollisionUpdateType::Deepest, float, 3>(const float, TCollisionConstraintBase<float, 3>&);
+		template void Apply<float, 3>(TCollisionConstraintBase<float, 3>&, float, TContactIterationParameters<float> &, TContactParticleParameters<float> &);
+		template void ApplyPushOut<float, 3>(TCollisionConstraintBase<float,3>&, float , const TSet<const TGeometryParticleHandle<float,3>*>&,
+			TContactIterationParameters<float> & IterationParameters, TContactParticleParameters<float> & ParticleParameters);
+
+	} // Collisions
+
+}// Chaos
+
diff --git a/Engine/Source/Runtime/Experimental/Chaos/Private/Chaos/PBDCollisionConstraintsPlaneContactUtil.cpp b/Engine/Source/Runtime/Experimental/Chaos/Private/Chaos/PBDCollisionConstraintsPlaneContactUtil.cpp
deleted file mode 100644
diff --git a/Engine/Source/Runtime/Experimental/Chaos/Private/Chaos/PBDCollisionConstraintsPointContactUtil.cpp b/Engine/Source/Runtime/Experimental/Chaos/Private/Chaos/PBDCollisionConstraintsPointContactUtil.cpp
deleted file mode 100644
diff --git a/Engine/Source/Runtime/Experimental/Chaos/Public/Chaos/CollisionResolution.h b/Engine/Source/Runtime/Experimental/Chaos/Public/Chaos/CollisionResolution.h
--- a/Engine/Source/Runtime/Experimental/Chaos/Public/Chaos/CollisionResolution.h
+++ b/Engine/Source/Runtime/Experimental/Chaos/Public/Chaos/CollisionResolution.h
@@ -115,7 +115,7 @@ namespace Chaos
 			// Convex-Convex
 			//
 			template <typename T, int d>
-			void CHAOS_API UpdateConvexConvexConstraint(const FImplicitObject& A, const TRigidTransform<T, d>& ATM, const FImplicitObject& B, const TRigidTransform<T, d>& BTM, const T Thickness, TRigidBodyPointContactConstraint<T, d>& Constraint);
+			void CHAOS_API UpdateConvexConvexConstraint(const FImplicitObject& A, const TRigidTransform<T, d>& ATM, const FImplicitObject& B, const TRigidTransform<T, d>& BTM, const T Thickness, TRigidBodyPlaneContactConstraint<T, d>& Constraint);
 
 			template <typename T, int d>
 			void CHAOS_API ConstructConvexConvexConstraints(TGeometryParticleHandle<T, d>* Particle0, TGeometryParticleHandle<T, d>* Particle1, const FImplicitObject* Implicit0, const FImplicitObject* Implicit1, const TRigidTransform<T, d>& Transform0, const TRigidTransform<T, d>& Transform1, const T Thickness, FCollisionConstraintsArray& NewConstraints);
diff --git a/Engine/Source/Runtime/Experimental/Chaos/Public/Chaos/CollisionResolutionTypes.h b/Engine/Source/Runtime/Experimental/Chaos/Public/Chaos/CollisionResolutionTypes.h
--- a/Engine/Source/Runtime/Experimental/Chaos/Public/Chaos/CollisionResolutionTypes.h
+++ b/Engine/Source/Runtime/Experimental/Chaos/Public/Chaos/CollisionResolutionTypes.h
@@ -38,39 +38,6 @@ namespace Chaos
 	};
 
 
-
-	/*
-	*
-	*/
-	template<class T = float, int d = 3>
-	class CHAOS_API TCollisionConstraintBase
-	{
-	public:
-		using FGeometryParticleHandle = TGeometryParticleHandle<T, d>;
-
-		enum class FType
-		{
-			None = 0,     // default value also indicates a invalid constraint
-			SinglePoint,   // TRigidBodyPointContactConstraint
-			Plane         // TRigidBodyPlaneContactConstraint
-		};
-
-		TCollisionConstraintBase() : Type(FType::None) {}
-		TCollisionConstraintBase(FType InType) : Type(InType) {}
-		FType GetType() const { return Type; }
-
-		template<class AS_T> AS_T * As() { return static_cast<AS_T*>(this); }
-		template<class AS_T> const AS_T * As() const { return static_cast<AS_T*>(this); }
-
-		FGeometryParticleHandle* Particle[2]; // { Point, Volume } 
-
-	private:
-		FType Type;
-	};
-	typedef TCollisionConstraintBase<float, 3> FCollisionConstraintBase;
-	typedef TArray<TCollisionConstraintBase<float, 3>*> FCollisionConstraintsArray;
-
-
 	/*
 	*
 	*/
@@ -94,86 +61,94 @@ namespace Chaos
 		const FImplicitObject* Implicit[2]; // {Of Particle[0], Of Particle[1]}
 	};
 
+
+
+
 	/*
 	*
 	*/
-	template<class T, int d>
-	class CHAOS_API TRigidBodyPointContactConstraint : public TCollisionConstraintBase<T, d>
+	template<class T = float, int d = 3>
+	class CHAOS_API TCollisionConstraintBase
 	{
 	public:
-		using Base = TCollisionConstraintBase<T, d>;
 		using FGeometryParticleHandle = TGeometryParticleHandle<T, d>;
 		using FManifold = TPointContactManifold<T, d>;
-		using Base::Particle;
 
-		TRigidBodyPointContactConstraint() : Base(Base::FType::SinglePoint), AccumulatedImpulse(0) {}
-		static typename Base::FType StaticType() { return Base::FType::SinglePoint; };
+		enum class FType
+		{
+			None = 0,     // default value also indicates a invalid constraint
+			SinglePoint,   // TRigidBodyPointContactConstraint
+			Plane         // TRigidBodyPlaneContactConstraint
+		};
 
+		TCollisionConstraintBase() : AccumulatedImpulse(0), Type(FType::None) {}
+		TCollisionConstraintBase(FType InType) : AccumulatedImpulse(0), Type(InType) {}
+		FType GetType() const { return Type; }
 
-		//
-		// Manidfold management
-		//
-		bool ContainsManifold(const FImplicitObject* B = nullptr, const FImplicitObject* A = nullptr) const { return (A == nullptr&&B == nullptr) ? (Manifold.Implicit[0] != nullptr && Manifold.Implicit[1] != nullptr) : (Manifold.Implicit[0] == A && Manifold.Implicit[1] == B); }
-		void AddManifold(const FImplicitObject* B = nullptr, const FImplicitObject* A = nullptr) { Manifold.Implicit[0] = A; Manifold.Implicit[1] = B; }
+		template<class AS_T> AS_T * As() { return static_cast<AS_T*>(this); }
+		template<class AS_T> const AS_T * As() const { return static_cast<AS_T*>(this); }
 
+		FString ToString() const
+		{
+			return FString::Printf(TEXT("Particle:%s, Levelset:%s, AccumulatedImpulse:%s"), *Particle[0]->ToString(), *Particle[1]->ToString(), *AccumulatedImpulse.ToString());
+		}
 
 		//
 		// API
 		//
-		void ResetPhi(T InPhi) { SetPhi(InPhi); }
-
-		void SetPhi(T InPhi, const FImplicitObject* B = nullptr, const FImplicitObject* A = nullptr, int32 Index = 0) { Manifold.Phi = InPhi; }
-		T GetPhi(const FImplicitObject* B = nullptr, const FImplicitObject* A = nullptr) const { return Manifold.Phi; }
 
-		void SetDisabled(bool bInDisabled, const FImplicitObject* B = nullptr, const FImplicitObject* A = nullptr) { Manifold.bDisabled = bInDisabled; }
-		TVector<T, d> GetDisabled(const FImplicitObject* B = nullptr, const FImplicitObject* A = nullptr) const { return Manifold.bDisabled; }
+		void ResetPhi(T InPhi) { SetPhi(InPhi); }
 
-		void SetNormal(const TVector<T, d> & InNormal, const FImplicitObject* B = nullptr, const FImplicitObject* A = nullptr) { Manifold.Normal = InNormal; }
-		TVector<T, d> GetNormal(const FImplicitObject* B = nullptr, const FImplicitObject* A = nullptr) const { return Manifold.Normal; }
+		void SetPhi(T InPhi) { Manifold.Phi = InPhi; }
+		T GetPhi() const { return Manifold.Phi; }
 
-		void SetLocation(const TVector<T, d> & InLocation, const FImplicitObject* B = nullptr, const FImplicitObject* A = nullptr, int32 Index = 0) { Manifold.Location = InLocation; }
-		TVector<T, d> GetLocation(const FImplicitObject* B = nullptr, const FImplicitObject* A = nullptr) const { return Manifold.Location; }
+		void SetDisabled(bool bInDisabled) { Manifold.bDisabled = bInDisabled; }
+		TVector<T, d> GetDisabled() const { return Manifold.bDisabled; }
 
-		FString ToString() const
-		{
-			return FString::Printf(TEXT("Particle:%s, Levelset:%s, AccumulatedImpulse:%s"), *Particle[0]->ToString(), *Particle[1]->ToString(), *AccumulatedImpulse.ToString());
-		}
+		void SetNormal(const TVector<T, d> & InNormal) { Manifold.Normal = InNormal; }
+		TVector<T, d> GetNormal() const { return Manifold.Normal; }
 
+		void SetLocation(const TVector<T, d> & InLocation) { Manifold.Location = InLocation; }
+		TVector<T, d> GetLocation() const { return Manifold.Location; }
 
-		FManifold Manifold;
+		FGeometryParticleHandle* Particle[2]; // { Point, Volume }
 		TVector<T, d> AccumulatedImpulse;
+		FManifold Manifold;
 
+	private:
+		FType Type;
 	};
-	typedef TRigidBodyPointContactConstraint<float, 3> FRigidBodyPointContactConstraint;
+	typedef TCollisionConstraintBase<float, 3> FCollisionConstraintBase;
+	typedef TArray<TCollisionConstraintBase<float, 3>*> FCollisionConstraintsArray;
 
 
 	/*
 	*
 	*/
 	template<class T, int d>
-	class CHAOS_API TPlaneContactManifold
+	class CHAOS_API TRigidBodyPointContactConstraint : public TCollisionConstraintBase<T, d>
 	{
 	public:
-		TPlaneContactManifold(int32 InTimestamp = -INT_MAX, const FImplicitObject* InImplicit0 = nullptr, const FImplicitObject* InImplicit1 = nullptr)
-			: bDisabled(true), Timestamp(InTimestamp), Phi(FLT_MAX)
-		{
-			Implicit[0] = InImplicit0;
-			Implicit[1] = InImplicit1;
-		}
+		using Base = TCollisionConstraintBase<T, d>;
+		using FGeometryParticleHandle = TGeometryParticleHandle<T, d>;
+		using Base::Particle;
 
-		bool bDisabled;
-		int32 Timestamp;
+		TRigidBodyPointContactConstraint() : Base(Base::FType::SinglePoint) {}
+		static typename Base::FType StaticType() { return Base::FType::SinglePoint; };
 
-		TVector<T, d> PlaneNormal;
-		TVector<T, d> PlanePosition;
-		TVector<T, d> ManifoldSamples[4]; // iterative plane samples
+		//
+		// Manidfold management
+		//
+		bool ContainsManifold(const FImplicitObject* B, const FImplicitObject* A) const { return A == Manifold.Implicit[0] && Manifold.Implicit[1] == B; }
+		void AddManifold(const FImplicitObject* B, const FImplicitObject* A) { Manifold.Implicit[0] = A; Manifold.Implicit[1] = B; }
 
-		TVector<T, d> Location;
-		T Phi;
+		FString ToString() const
+		{
+			return FString::Printf(TEXT("Particle:%s, Levelset:%s, AccumulatedImpulse:%s"), *Particle[0]->ToString(), *Particle[1]->ToString(), *AccumulatedImpulse.ToString());
+		}
 
-		const FImplicitObject* Implicit[2]; // {Of Particle[0], Of Particle[1]}
 	};
-
+	typedef TRigidBodyPointContactConstraint<float, 3> FRigidBodyPointContactConstraint;
 
 
 	/*
@@ -185,30 +160,29 @@ namespace Chaos
 	public:
 		using Base = TCollisionConstraintBase<T, d>;
 		using FGeometryParticleHandle = TGeometryParticleHandle<T, d>;
-		using FManifold = TPlaneContactManifold<T, d>;
 		using Base::Particle;
 
-		TRigidBodyPlaneContactConstraint() : Base(Base::FType::Plane), AccumulatedImpulse(0) {}
+		TRigidBodyPlaneContactConstraint() : Base(Base::FType::Plane) {}
 		static typename Base::FType StaticType() { return Base::FType::Plane; };
 
-
-		bool ContainsManifold(const FImplicitObject* B = nullptr, const FImplicitObject* A = nullptr) const { return (A == nullptr&&B == nullptr) ? (Manifold.Implicit[0] != nullptr && Manifold.Implicit[1] != nullptr) : (Manifold.Implicit[0] == A && Manifold.Implicit[1] == B); }
-		void AddManifold(const FImplicitObject* B = nullptr, const FImplicitObject* A = nullptr) { Manifold.Implicit[0] = A; Manifold.Implicit[1] = B; }
-
+		bool ContainsManifold(const FImplicitObject* B, const FImplicitObject* A) const { return A == Manifold.Implicit[0] && Manifold.Implicit[1] == B; }
 
 		FString ToString() const
 		{
 			return FString::Printf(TEXT("Particle:%s, Levelset:%s, AccumulatedImpulse:%s"), *Particle[0]->ToString(), *Particle[1]->ToString(), *AccumulatedImpulse.ToString());
 		}
 
-
-		FManifold Manifold;
-		TVector<T, d> AccumulatedImpulse;
-
+		TVector<T, d> PlaneNormal;
+		TVector<T, d> PlanePosition;
+		TArray<TVector<T, d>> PlaneSamples; // iterative plane samples
 	};
 	typedef TRigidBodyPlaneContactConstraint<float, 3> FRigidBodyPlaneContactConstraint;
 
 
+	//
+	//
+	//
+
 	template<class T, int d>
 	struct CHAOS_API TRigidBodyContactConstraintPGS
 	{
@@ -243,13 +217,34 @@ namespace Chaos
 
 		// Handle API
 
+		const TCollisionConstraintBase<T, d>& GetContact() const 
+		{ 
+			if (GetType() == FConstraintBase::FType::SinglePoint)
+			{
+				return ConstraintContainer->PointConstraints[ConstraintIndex];
+			}
+			check(GetType() == FConstraintBase::FType::Plane);
+			return ConstraintContainer->PlaneConstraints[ConstraintIndex];
+		}
+
+		TCollisionConstraintBase<T, d>& GetContact()
+		{
+			if (GetType() == FConstraintBase::FType::SinglePoint)
+			{
+				return ConstraintContainer->PointConstraints[ConstraintIndex];
+			}
+			check(GetType() == FConstraintBase::FType::Plane);
+			return ConstraintContainer->PlaneConstraints[ConstraintIndex];
+		}
+
+
+
 		const TRigidBodyPointContactConstraint<T, d>& GetPointContact() const { check(GetType() == FConstraintBase::FType::SinglePoint); return ConstraintContainer->PointConstraints[ConstraintIndex]; }
 		TRigidBodyPointContactConstraint<T, d>& GetPointContact() { check(GetType() == FConstraintBase::FType::SinglePoint); return ConstraintContainer->PointConstraints[ConstraintIndex]; }
 
 		const TRigidBodyPlaneContactConstraint<T, d>& GetPlaneContact() const { check(GetType() == FConstraintBase::FType::Plane); return ConstraintContainer->PlaneConstraints[ConstraintIndex]; }
 		TRigidBodyPlaneContactConstraint<T, d>& GetPlaneContact() { check(GetType() == FConstraintBase::FType::Plane); return ConstraintContainer->PlaneConstraints[ConstraintIndex]; }
 
-
 		typename FConstraintBase::FType GetType() const { return ConstraintType; }
 
 		void SetConstraintIndex(int32 IndexIn, typename FConstraintBase::FType InType)
diff --git a/Engine/Source/Runtime/Experimental/Chaos/Public/Chaos/CollisionResolutionUtil.h b/Engine/Source/Runtime/Experimental/Chaos/Public/Chaos/CollisionResolutionUtil.h
--- a/Engine/Source/Runtime/Experimental/Chaos/Public/Chaos/CollisionResolutionUtil.h
+++ b/Engine/Source/Runtime/Experimental/Chaos/Public/Chaos/CollisionResolutionUtil.h
@@ -30,7 +30,7 @@ namespace Chaos
 
 		template<typename T, int d>
 		TVector<T, d> 
-		GetEnergyClampedImpulse(const TRigidBodyPointContactConstraint<T, d>& Constraint, const TVector<T, d>& Impulse, const TVector<T, d>& VectorToPoint1, const TVector<T, d>& VectorToPoint2, const TVector<T, d>& Velocity1, const TVector<T, d>& Velocity2);
+		GetEnergyClampedImpulse(const TCollisionConstraintBase<T, d>& Constraint, const TVector<T, d>& Impulse, const TVector<T, d>& VectorToPoint1, const TVector<T, d>& VectorToPoint2, const TVector<T, d>& Velocity1, const TVector<T, d>& Velocity2);
 
 		template <typename T, int d>
 		bool 
diff --git a/Engine/Source/Runtime/Experimental/Chaos/Public/Chaos/PBDCollisionConstraintsContact.h b/Engine/Source/Runtime/Experimental/Chaos/Public/Chaos/PBDCollisionConstraintsContact.h
new file mode 100644
--- /dev/null
+++ b/Engine/Source/Runtime/Experimental/Chaos/Public/Chaos/PBDCollisionConstraintsContact.h
@@ -0,0 +1,42 @@
+// Copyright 1998-2019 Epic Games, Inc. All Rights Reserved.
+#pragma once
+
+#include "Chaos/ExternalCollisionData.h"
+#include "Chaos/PBDCollisionConstraints.h"
+#include "Chaos/CollisionResolutionTypes.h"
+
+namespace Chaos
+{
+	namespace Collisions
+	{
+		template<class T = float>
+		struct TContactParticleParameters {
+			TArrayCollectionArray<bool>* Collided;
+			const TArrayCollectionArray<TSerializablePtr<FChaosPhysicsMaterial>>* PhysicsMaterials;
+			T FrictionOverride;
+			T AngularFrictionOverride;
+		};
+
+		template<class T = float>
+		struct TContactIterationParameters {
+			const T Dt;
+			const int32 Iteration;
+			const int32 NumIterations;
+			const int32 NumPairIterations;
+			bool* NeedsAnotherIteration;
+		};
+
+		template<ECollisionUpdateType UpdateType, typename T, int d>
+		void Update(const T Thickness, TCollisionConstraintBase<T, d>& Constraint);
+
+		template<typename T, int d>
+		void Apply(TCollisionConstraintBase<T, d>& Constraint, T Thickness, TContactIterationParameters<T> & IterationParameters, TContactParticleParameters<T> & ParticleParameters);
+
+		template<typename T, int d>
+		void ApplyPushOut(TCollisionConstraintBase<T, d>& Constraint, T Thickness, const TSet<const TGeometryParticleHandle<T, d>*>& IsTemporarilyStatic,
+			TContactIterationParameters<T> & IterationParameters, TContactParticleParameters<T> & ParticleParameters);
+
+
+	}
+
+}
diff --git a/Engine/Source/Runtime/Experimental/Chaos/Public/Chaos/PBDCollisionConstraintsPlaneContactUtil.h b/Engine/Source/Runtime/Experimental/Chaos/Public/Chaos/PBDCollisionConstraintsPlaneContactUtil.h
deleted file mode 100644
diff --git a/Engine/Source/Runtime/Experimental/Chaos/Public/Chaos/PBDCollisionConstraintsPointContactUtil.h b/Engine/Source/Runtime/Experimental/Chaos/Public/Chaos/PBDCollisionConstraintsPointContactUtil.h
deleted file mode 100644
