Chaos : Collisions

Naming cleanup.

diff --git a/Engine/Source/Runtime/Experimental/Chaos/Private/Chaos/CollisionResolution.cpp b/Engine/Source/Runtime/Experimental/Chaos/Private/Chaos/CollisionResolution.cpp
--- a/Engine/Source/Runtime/Experimental/Chaos/Private/Chaos/CollisionResolution.cpp
+++ b/Engine/Source/Runtime/Experimental/Chaos/Private/Chaos/CollisionResolution.cpp
@@ -33,7 +33,7 @@ namespace Chaos
 		};
 
 		template <typename T>
-		void UpdateContactPoint(TPointContactManifold<T, 3>& Manifold, const TContactPoint<T>& NewContactPoint)
+		void UpdateContactPoint(TCollisionContact<T, 3>& Manifold, const TContactPoint<T>& NewContactPoint)
 		{
 			//for now just override
 			if (NewContactPoint.Phi < Manifold.Phi)
@@ -99,7 +99,7 @@ namespace Chaos
 				FRigidBodyPointContactConstraint* Constraint = new FRigidBodyPointContactConstraint;
 				Constraint->Particle[0] = Particle0;
 				Constraint->Particle[1] = Particle1;
-				Constraint->AddManifold(Implicit0, Implicit1);
+				Constraint->SetManifold(Implicit0, Implicit1);
 
 				UpdateBoxBoxConstraint(*Object0, Transform0, *Object1, Transform1, Thickness, *Constraint);
 
@@ -121,7 +121,7 @@ namespace Chaos
 		template <typename T, int d>
 		bool UpdateBoxPlaneConstraint(const TBox<T, d>& Box, const TRigidTransform<T, d>& BoxTransform, const TPlane<T, d>& Plane, const TRigidTransform<T, d>& PlaneTransform, const T Thickness, TRigidBodyPointContactConstraint<T, d>& Constraint)
 		{
-			TPointContactManifold<T, d> & Contact = Constraint.Manifold;
+			TCollisionContact<T, d> & Contact = Constraint.Manifold;
 
 #if USING_CODE_ANALYSIS
 			MSVC_PRAGMA(warning(push))
@@ -192,7 +192,7 @@ namespace Chaos
 				FRigidBodyPointContactConstraint* Constraint = new FRigidBodyPointContactConstraint;
 				Constraint->Particle[0] = Particle0;
 				Constraint->Particle[1] = Particle1;
-				Constraint->AddManifold(Implicit0, Implicit1);
+				Constraint->SetManifold(Implicit0, Implicit1);
 
 				UpdateBoxPlaneConstraint(*Object0, Transform0, *Object1, Transform1, Thickness, *Constraint);
 
@@ -246,7 +246,7 @@ namespace Chaos
 				FRigidBodyPointContactConstraint* Constraint = new FRigidBodyPointContactConstraint;
 				Constraint->Particle[0] = Particle0;
 				Constraint->Particle[1] = Particle1;
-				Constraint->AddManifold(Implicit0, Implicit1);
+				Constraint->SetManifold(Implicit0, Implicit1);
 
 				UpdateSphereSphereConstraint(*Object0, Transform0, *Object1, Transform1, Thickness, *Constraint);
 
@@ -268,7 +268,7 @@ namespace Chaos
 		template <typename T, int d>
 		void UpdateSpherePlaneConstraint(const TSphere<T, d>& Sphere, const TRigidTransform<T, d>& SphereTransform, const TPlane<T, d>& Plane, const TRigidTransform<T, d>& PlaneTransform, const T Thickness, TRigidBodyPointContactConstraint<T, d>& Constraint)
 		{
-			TPointContactManifold<T, d> & Contact = Constraint.Manifold;
+			TCollisionContact<T, d> & Contact = Constraint.Manifold;
 
 			const TRigidTransform<T, d> SphereToPlaneTransform(PlaneTransform.Inverse() * SphereTransform);
 			const TVector<T, d> SphereCenter = SphereToPlaneTransform.TransformPosition(Sphere.GetCenter());
@@ -297,7 +297,7 @@ namespace Chaos
 				FRigidBodyPointContactConstraint* Constraint = new FRigidBodyPointContactConstraint;
 				Constraint->Particle[0] = Particle0;
 				Constraint->Particle[1] = Particle1;
-				Constraint->AddManifold(Implicit0, Implicit1);
+				Constraint->SetManifold(Implicit0, Implicit1);
 
 				UpdateSpherePlaneConstraint(*Object0, Transform0, *Object1, Transform1, Thickness, *Constraint);
 
@@ -351,7 +351,7 @@ namespace Chaos
 				FRigidBodyPointContactConstraint* Constraint = new FRigidBodyPointContactConstraint;
 				Constraint->Particle[0] = Particle0;
 				Constraint->Particle[1] = Particle1;
-				Constraint->AddManifold(Implicit0, Implicit1);
+				Constraint->SetManifold(Implicit0, Implicit1);
 
 				UpdateSphereBoxConstraint(*Object0, Transform0, *Object1, Transform1, Thickness, *Constraint);
 
@@ -411,7 +411,7 @@ namespace Chaos
 				FRigidBodyPointContactConstraint* Constraint = new FRigidBodyPointContactConstraint;
 				Constraint->Particle[0] = Particle0;
 				Constraint->Particle[1] = Particle1;
-				Constraint->AddManifold(Implicit0, Implicit1);
+				Constraint->SetManifold(Implicit0, Implicit1);
 
 				UpdateSphereCapsuleConstraint(*Object0, Transform0, *Object1, Transform1, Thickness, *Constraint);
 
@@ -474,7 +474,7 @@ namespace Chaos
 				FRigidBodyPointContactConstraint* Constraint = new FRigidBodyPointContactConstraint;
 				Constraint->Particle[0] = Particle0;
 				Constraint->Particle[1] = Particle1;
-				Constraint->AddManifold(Implicit0, Implicit1);
+				Constraint->SetManifold(Implicit0, Implicit1);
 
 				UpdateCapsuleCapsuleConstraint(*Object0, Transform0, *Object1, Transform1, Thickness, *Constraint);
 
@@ -516,7 +516,7 @@ namespace Chaos
 				FRigidBodyPointContactConstraint* Constraint = new FRigidBodyPointContactConstraint;
 				Constraint->Particle[0] = Particle0;
 				Constraint->Particle[1] = Particle1;
-				Constraint->AddManifold(Implicit0, Implicit1);
+				Constraint->SetManifold(Implicit0, Implicit1);
 
 				UpdateCapsuleBoxConstraint(*Object0, Transform0, *Object1, Transform1, Thickness, *Constraint);
 
@@ -560,7 +560,7 @@ namespace Chaos
 
 			Constraint->Particle[0] = Particle0;
 			Constraint->Particle[1] = Particle1;
-			Constraint->AddManifold(Implicit0, Implicit1);
+			Constraint->SetManifold(Implicit0, Implicit1);
 
 			UpdateConvexConvexConstraint(*Implicit0, Transform0, *Implicit1, Transform1, Thickness, *Constraint);
 
@@ -619,13 +619,13 @@ namespace Chaos
 			{
 				Constraint->Particle[0] = Particle1;
 				Constraint->Particle[1] = Particle0;
-				Constraint->AddManifold(Implicit1, Implicit0);
+				Constraint->SetManifold(Implicit1, Implicit0);
 			}
 			else
 			{
 				Constraint->Particle[0] = Particle0;
 				Constraint->Particle[1] = Particle1;
-				Constraint->AddManifold(Implicit0, Implicit1);
+				Constraint->SetManifold(Implicit0, Implicit1);
 			}
 
 			UpdateLevelsetLevelsetConstraint<ECollisionUpdateType::Any>(Thickness, *Constraint);
@@ -851,7 +851,7 @@ namespace Chaos
 			FRigidBodyPointContactConstraint* Constraint = new FRigidBodyPointContactConstraint;
 			Constraint->Particle[0] = Particle0;
 			Constraint->Particle[1] = Particle1;
-			Constraint->AddManifold(Implicit0, Implicit1);
+			Constraint->SetManifold(Implicit0, Implicit1);
 
 			UpdateSingleUnionConstraint<ECollisionUpdateType::Any>(*Implicit0, Transform0, *Implicit1, Transform1, Thickness, *Constraint);
 
diff --git a/Engine/Source/Runtime/Experimental/Chaos/Private/Chaos/CollisionResolutionUtil.cpp b/Engine/Source/Runtime/Experimental/Chaos/Private/Chaos/CollisionResolutionUtil.cpp
--- a/Engine/Source/Runtime/Experimental/Chaos/Private/Chaos/CollisionResolutionUtil.cpp
+++ b/Engine/Source/Runtime/Experimental/Chaos/Private/Chaos/CollisionResolutionUtil.cpp
@@ -87,7 +87,7 @@ namespace Chaos
 			TVector<T, d> LocalNormal;
 			T LocalPhi = Object.PhiWithNormal(LocalPoint, LocalNormal);
 
-			TPointContactManifold<T, d> & Contact = Constraint.Manifold;
+			TCollisionContact<T, d> & Contact = Constraint.Manifold;
 			if (LocalPhi < Contact.Phi)
 			{
 				Contact.Phi = LocalPhi;
@@ -105,7 +105,7 @@ namespace Chaos
 			TVector<T, d> LocalNormal;
 			T LocalPhi = Object.PhiWithNormal(LocalPoint, LocalNormal);
 
-			TPointContactManifold<T, d> & Contact = Constraint.Manifold;
+			TCollisionContact<T, d> & Contact = Constraint.Manifold;
 			if (LocalPhi < Contact.Phi)
 			{
 				Contact.Phi = LocalPhi;
@@ -122,7 +122,7 @@ namespace Chaos
 			T LocalPhi = Object.PhiWithNormal(LocalPoint, LocalNormal);
 			T LocalThickness = LocalPhi - Thickness;
 
-			TPointContactManifold<T, d> & Contact = Constraint.Manifold;
+			TCollisionContact<T, d> & Contact = Constraint.Manifold;
 			if (LocalThickness < -KINDA_SMALL_NUMBER)
 			{
 				Contact.Location += LocalPoint * LocalThickness;
@@ -153,8 +153,8 @@ namespace Chaos
 			SCOPE_CYCLE_COUNTER(STAT_SampleObject);
 			TRigidBodyPointContactConstraint<float, 3> AvgConstraint = Constraint;
 
-			TPointContactManifold<float, 3> & Contact = Constraint.Manifold;
-			TPointContactManifold<float, 3> & AvgContact = AvgConstraint.Manifold;
+			TCollisionContact<float, 3> & Contact = Constraint.Manifold;
+			TCollisionContact<float, 3> & AvgContact = AvgConstraint.Manifold;
 
 			AvgConstraint.Particle[0] = Constraint.Particle[0];
 			AvgConstraint.Particle[1] = Constraint.Particle[1];
@@ -460,8 +460,8 @@ namespace Chaos
 			SCOPE_CYCLE_COUNTER(STAT_SampleObject);
 			TRigidBodyPointContactConstraint<float, 3> AvgConstraint = Constraint;
 
-			TPointContactManifold<float, 3> & Contact = Constraint.Manifold;
-			TPointContactManifold<float, 3> & AvgContact = AvgConstraint.Manifold;
+			TCollisionContact<float, 3> & Contact = Constraint.Manifold;
+			TCollisionContact<float, 3> & AvgContact = AvgConstraint.Manifold;
 
 			AvgConstraint.Particle[0] = Constraint.Particle[0];
 			AvgConstraint.Particle[1] = Constraint.Particle[1];
diff --git a/Engine/Source/Runtime/Experimental/Chaos/Public/Chaos/CollisionResolutionTypes.h b/Engine/Source/Runtime/Experimental/Chaos/Public/Chaos/CollisionResolutionTypes.h
--- a/Engine/Source/Runtime/Experimental/Chaos/Public/Chaos/CollisionResolutionTypes.h
+++ b/Engine/Source/Runtime/Experimental/Chaos/Public/Chaos/CollisionResolutionTypes.h
@@ -42,10 +42,10 @@ namespace Chaos
 	*
 	*/
 	template<class T, int d>
-	class CHAOS_API TPointContactManifold
+	class CHAOS_API TCollisionContact
 	{
 	public:
-		TPointContactManifold(int32 InTimestamp = -INT_MAX, const FImplicitObject* InImplicit0 = nullptr, const FImplicitObject* InImplicit1 = nullptr)
+		TCollisionContact(int32 InTimestamp = -INT_MAX, const FImplicitObject* InImplicit0 = nullptr, const FImplicitObject* InImplicit1 = nullptr)
 			: bDisabled(true), Timestamp(InTimestamp), Normal(0), Phi(FLT_MAX)
 		{
 			Implicit[0] = InImplicit0;
@@ -72,7 +72,7 @@ namespace Chaos
 	{
 	public:
 		using FGeometryParticleHandle = TGeometryParticleHandle<T, d>;
-		using FManifold = TPointContactManifold<T, d>;
+		using FManifold = TCollisionContact<T, d>;
 
 		enum class FType
 		{
@@ -88,10 +88,8 @@ namespace Chaos
 		template<class AS_T> AS_T * As() { return static_cast<AS_T*>(this); }
 		template<class AS_T> const AS_T * As() const { return static_cast<AS_T*>(this); }
 
-		FString ToString() const
-		{
-			return FString::Printf(TEXT("Particle:%s, Levelset:%s, AccumulatedImpulse:%s"), *Particle[0]->ToString(), *Particle[1]->ToString(), *AccumulatedImpulse.ToString());
-		}
+		bool ContainsManifold(const FImplicitObject* B, const FImplicitObject* A) const { return A == Manifold.Implicit[0] && B == Manifold.Implicit[1]; }
+		void SetManifold(const FImplicitObject* B, const FImplicitObject* A) { Manifold.Implicit[0] = A; Manifold.Implicit[1] = B; }
 
 		//
 		// API
@@ -111,6 +109,13 @@ namespace Chaos
 		void SetLocation(const TVector<T, d> & InLocation) { Manifold.Location = InLocation; }
 		TVector<T, d> GetLocation() const { return Manifold.Location; }
 
+
+		FString ToString() const
+		{
+			return FString::Printf(TEXT("Particle:%s, Levelset:%s, AccumulatedImpulse:%s"), *Particle[0]->ToString(), *Particle[1]->ToString(), *AccumulatedImpulse.ToString());
+		}
+
+
 		FGeometryParticleHandle* Particle[2]; // { Point, Volume }
 		TVector<T, d> AccumulatedImpulse;
 		FManifold Manifold;
@@ -135,18 +140,6 @@ namespace Chaos
 
 		TRigidBodyPointContactConstraint() : Base(Base::FType::SinglePoint) {}
 		static typename Base::FType StaticType() { return Base::FType::SinglePoint; };
-
-		//
-		// Manidfold management
-		//
-		bool ContainsManifold(const FImplicitObject* B, const FImplicitObject* A) const { return A == Manifold.Implicit[0] && Manifold.Implicit[1] == B; }
-		void AddManifold(const FImplicitObject* B, const FImplicitObject* A) { Manifold.Implicit[0] = A; Manifold.Implicit[1] = B; }
-
-		FString ToString() const
-		{
-			return FString::Printf(TEXT("Particle:%s, Levelset:%s, AccumulatedImpulse:%s"), *Particle[0]->ToString(), *Particle[1]->ToString(), *AccumulatedImpulse.ToString());
-		}
-
 	};
 	typedef TRigidBodyPointContactConstraint<float, 3> FRigidBodyPointContactConstraint;
 
@@ -165,12 +158,6 @@ namespace Chaos
 		TRigidBodyPlaneContactConstraint() : Base(Base::FType::Plane) {}
 		static typename Base::FType StaticType() { return Base::FType::Plane; };
 
-		bool ContainsManifold(const FImplicitObject* B, const FImplicitObject* A) const { return A == Manifold.Implicit[0] && Manifold.Implicit[1] == B; }
-
-		FString ToString() const
-		{
-			return FString::Printf(TEXT("Particle:%s, Levelset:%s, AccumulatedImpulse:%s"), *Particle[0]->ToString(), *Particle[1]->ToString(), *AccumulatedImpulse.ToString());
-		}
 
 		TVector<T, d> PlaneNormal;
 		TVector<T, d> PlanePosition;
