Chaos : Collisions

Updated convex-convex to initialize the PlaneConstraint. Currently being reset every frame, so still behaves as a point constraint. Persistence is next.

diff --git a/Engine/Source/Runtime/Experimental/Chaos/Private/Chaos/CollisionResolution.cpp b/Engine/Source/Runtime/Experimental/Chaos/Private/Chaos/CollisionResolution.cpp
--- a/Engine/Source/Runtime/Experimental/Chaos/Private/Chaos/CollisionResolution.cpp
+++ b/Engine/Source/Runtime/Experimental/Chaos/Private/Chaos/CollisionResolution.cpp
@@ -547,8 +547,19 @@ namespace Chaos
 			});
 		}
 
+		template<class TYPE>
+		TYPE SumArray(TArray<TYPE>& Array)
+		{
+			TYPE ReturnValue(0);
+			for (int i = 0; i < Array.Num(); i++)
+			{
+				ReturnValue += Array[i];
+			}
+			return ReturnValue;
+		}
+
 		template<class T, int d>
-		void UpdateConvexConvexConstraint(const FImplicitObject& A, const TRigidTransform<T, d>& ATM, const FImplicitObject& B, const TRigidTransform<T, d>& BTM, const T Thickness, TRigidBodyPointContactConstraint<T, d>& Constraint)
+		void UpdateConvexConvexConstraint(const FImplicitObject& A, const TRigidTransform<T, d>& ATM, const FImplicitObject& B, const TRigidTransform<T, d>& BTM, const T Thickness, TRigidBodyPlaneContactConstraint<T,d>& Constraint)
 		{
 			UpdateContactPoint(Constraint.Manifold, ConvexConvexContactPoint(A, ATM, B, BTM, Thickness));
 		}
@@ -556,7 +567,7 @@ namespace Chaos
 		template<class T, int d>
 		void ConstructConvexConvexConstraints(TGeometryParticleHandle<T, d>* Particle0, TGeometryParticleHandle<T, d>* Particle1, const FImplicitObject* Implicit0, const FImplicitObject* Implicit1, const TRigidTransform<T, d>& Transform0, const TRigidTransform<T, d>& Transform1, const T Thickness, FCollisionConstraintsArray& NewConstraints)
 		{
-			FRigidBodyPointContactConstraint* Constraint = new FRigidBodyPointContactConstraint;
+			FRigidBodyPlaneContactConstraint* Constraint = new FRigidBodyPlaneContactConstraint;
 
 			Constraint->Particle[0] = Particle0;
 			Constraint->Particle[1] = Particle1;
@@ -988,7 +999,7 @@ namespace Chaos
 			}
 			else if(Implicit0.IsConvex() && Implicit1.IsConvex())
 			{
-				UpdateConvexConvexConstraint(Implicit0, Transform0, Implicit1, Transform1, Thickness, *ConstraintBase.template As<TRigidBodyPointContactConstraint<T, d>>());
+				UpdateConvexConvexConstraint(Implicit0, Transform0, Implicit1, Transform1, Thickness, *ConstraintBase.template As<TRigidBodyPlaneContactConstraint<T, d>>());
 			}
 			else if (Implicit1.IsUnderlyingUnion())
 			{
@@ -1175,7 +1186,7 @@ namespace Chaos
 		template void UpdateCapsuleBoxConstraint<float, 3>(const TCapsule<float>& A, const TRigidTransform<float, 3>& ATransform, const TBox<float, 3>& B, const TRigidTransform<float, 3>& BTransform, const float Thickness, TRigidBodyPointContactConstraint<float, 3>& Constraint);
 		template void ConstructCapsuleBoxConstraints<float, 3>(TGeometryParticleHandle<float, 3>* Particle0, TGeometryParticleHandle<float, 3>* Particle1, const FImplicitObject* Implicit0, const FImplicitObject* Implicit1, const TRigidTransform<float, 3>& Transform0, const TRigidTransform<float, 3>& Transform1, const float Thickness, FCollisionConstraintsArray& NewConstraints);
 
-		template void UpdateConvexConvexConstraint<float, 3>(const FImplicitObject& A, const TRigidTransform<float, 3>& ATM, const FImplicitObject& B, const TRigidTransform<float, 3>& BTM, const float Thickness, TRigidBodyPointContactConstraint<float, 3>& Constraint);
+		template void UpdateConvexConvexConstraint<float, 3>(const FImplicitObject& A, const TRigidTransform<float, 3>& ATM, const FImplicitObject& B, const TRigidTransform<float, 3>& BTM, const float Thickness, TRigidBodyPlaneContactConstraint<float, 3>& Constraint);
 		template void ConstructConvexConvexConstraints<float, 3>(TGeometryParticleHandle<float, 3>* Particle0, TGeometryParticleHandle<float, 3>* Particle1, const FImplicitObject* Implicit0, const FImplicitObject* Implicit1, const TRigidTransform<float, 3>& Transform0, const TRigidTransform<float, 3>& Transform1, const float Thickness, FCollisionConstraintsArray& NewConstraints);
 
 		template void UpdateLevelsetLevelsetConstraint<ECollisionUpdateType::Any, float, 3>(const float Thickness, TRigidBodyPointContactConstraint<float, 3>& Constraint);
diff --git a/Engine/Source/Runtime/Experimental/Chaos/Private/Chaos/PBDCollisionConstraints.cpp b/Engine/Source/Runtime/Experimental/Chaos/Private/Chaos/PBDCollisionConstraints.cpp
--- a/Engine/Source/Runtime/Experimental/Chaos/Private/Chaos/PBDCollisionConstraints.cpp
+++ b/Engine/Source/Runtime/Experimental/Chaos/Private/Chaos/PBDCollisionConstraints.cpp
@@ -106,8 +106,7 @@ namespace Chaos
 		{
 			TRigidBodyPointContactConstraint<T, d>* PointConstraint = ConstraintBase->template As< TRigidBodyPointContactConstraint<T, d> >();
 
-			int32 Idx = PointConstraints.AddUninitialized(1);
-			PointConstraints[Idx] = *PointConstraint;
+			int32 Idx = PointConstraints.Add(*PointConstraint);
 			Handles.Add(HandleAllocator.template AllocHandle< TRigidBodyPointContactConstraint<T, d> >(this, Idx));
 
 			delete PointConstraint;
@@ -116,8 +115,7 @@ namespace Chaos
 		{
 			TRigidBodyPlaneContactConstraint<T, d>* PlaneConstraint = ConstraintBase->template As< TRigidBodyPlaneContactConstraint<T, d> >();
 
-			int32 Idx = PlaneConstraints.AddUninitialized(1);
-			PlaneConstraints[Idx] = *PlaneConstraint;
+			int32 Idx = PlaneConstraints.Add(*PlaneConstraint);
 			Handles.Add(HandleAllocator.template AllocHandle< TRigidBodyPlaneContactConstraint<T, d> >(this, Idx));
 
 			delete PlaneConstraint;
@@ -191,16 +189,33 @@ namespace Chaos
 	void TPBDCollisionConstraints<T, d>::RemoveConstraint(FConstraintContainerHandle* Handle)
 	{
 		int32 Idx = Handle->GetConstraintIndex();
-		typename TCollisionConstraintBase<T, d>::FType ConstraintType = Handle->GetType();
+		typename FCollisionConstraintBase::FType ConstraintType = Handle->GetType();
 
 		Handles.RemoveAtSwap(Idx);
-		PointConstraints.RemoveAtSwap(Idx);
-		if (Idx < PointConstraints.Num())
+
+		if (ConstraintType == FCollisionConstraintBase::FType::SinglePoint)
+		{
+			PointConstraints.RemoveAtSwap(Idx);
+			if (Idx < PointConstraints.Num())
+			{
+				Handles[Idx]->SetConstraintIndex(Idx, ConstraintType);
+			}
+		}
+		else if (ConstraintType == FCollisionConstraintBase::FType::Plane)
+		{
+			PlaneConstraints.RemoveAtSwap(Idx);
+			if (Idx < PlaneConstraints.Num())
+			{
+				Handles[Idx]->SetConstraintIndex(Idx, ConstraintType);
+			}
+		}
+		else 
 		{
-			Handles[Idx]->SetConstraintIndex(Idx, ConstraintType);
+			check(false);
 		}
+		
+		ensure(Handles.Num() == PointConstraints.Num() + PlaneConstraints.Num());
 
-		ensure(Handles.Num() == PointConstraints.Num());
 		HandleAllocator.FreeHandle(Handle);
 	}
 
diff --git a/Engine/Source/Runtime/Experimental/Chaos/Public/Chaos/CollisionResolutionTypes.h b/Engine/Source/Runtime/Experimental/Chaos/Public/Chaos/CollisionResolutionTypes.h
--- a/Engine/Source/Runtime/Experimental/Chaos/Public/Chaos/CollisionResolutionTypes.h
+++ b/Engine/Source/Runtime/Experimental/Chaos/Public/Chaos/CollisionResolutionTypes.h
@@ -46,7 +46,7 @@ namespace Chaos
 	{
 	public:
 		TCollisionContact(int32 InTimestamp = -INT_MAX, const FImplicitObject* InImplicit0 = nullptr, const FImplicitObject* InImplicit1 = nullptr)
-			: bDisabled(true), Timestamp(InTimestamp), Normal(0), Phi(FLT_MAX)
+			: bDisabled(true), Timestamp(InTimestamp), Normal(0), Location(0), Phi(FLT_MAX)
 		{
 			Implicit[0] = InImplicit0;
 			Implicit[1] = InImplicit1;
@@ -81,8 +81,8 @@ namespace Chaos
 			Plane         // TRigidBodyPlaneContactConstraint
 		};
 
-		TCollisionConstraintBase() : AccumulatedImpulse(0), Type(FType::None) {}
-		TCollisionConstraintBase(FType InType) : AccumulatedImpulse(0), Type(InType) {}
+		TCollisionConstraintBase() : AccumulatedImpulse(0), Type(FType::None) { Particle[0] = nullptr; Particle[1] = nullptr; }
+		TCollisionConstraintBase(FType InType) : AccumulatedImpulse(0), Type(InType) { Particle[0] = nullptr; Particle[1] = nullptr; }
 		FType GetType() const { return Type; }
 
 		template<class AS_T> AS_T * As() { return static_cast<AS_T*>(this); }
@@ -155,10 +155,9 @@ namespace Chaos
 		using FGeometryParticleHandle = TGeometryParticleHandle<T, d>;
 		using Base::Particle;
 
-		TRigidBodyPlaneContactConstraint() : Base(Base::FType::Plane) {}
+		TRigidBodyPlaneContactConstraint() : Base(Base::FType::Plane), PlaneNormal(0), PlanePosition(0) {}
 		static typename Base::FType StaticType() { return Base::FType::Plane; };
 
-
 		TVector<T, d> PlaneNormal;
 		TVector<T, d> PlanePosition;
 		TArray<TVector<T, d>> PlaneSamples; // iterative plane samples
@@ -244,47 +243,22 @@ namespace Chaos
 
 		TVector<T, d> GetContactLocation() const
 		{
-			if (ConstraintType == FConstraintBase::FType::SinglePoint)
-			{
-				return GetPointContact().GetLocation();
-			}
-			ensure(false);
-			return TVector<T, d>(0);
+			return GetContact().GetLocation();
 		}
 
 		TVector<T, d> GetAccumulatedImpulse() const
 		{
-			if (ConstraintType == FConstraintBase::FType::SinglePoint)
-			{
-				return GetPointContact().AccumulatedImpulse;
-			}
-			ensure(false);
-
-			return TVector<T, d>(0);
+			return GetContact().AccumulatedImpulse;
 		}
 
-
 		TVector<const TGeometryParticleHandle<T, d>*, 2> GetConstrainedParticles() const
 		{
-			if (ConstraintType == FConstraintBase::FType::SinglePoint)
-			{
-				const TRigidBodyPointContactConstraint<float, 3>& Contact = GetPointContact();
-				return { Contact.Particle[0], Contact.Particle[1] };
-			}
-			ensure(false);
-			return { nullptr,nullptr };
+			return { GetContact().Particle[0], GetContact().Particle[1] };
 		}
 
-
-
 		TVector<TGeometryParticleHandle<T, d>*, 2> GetConstrainedParticles()
 		{
-			if (ConstraintType == FConstraintBase::FType::SinglePoint)
-			{
-				return { GetPointContact().Particle[0], GetPointContact().Particle[1] };
-			}
-			ensure(false);
-			return { nullptr,nullptr };
+			return { GetContact().Particle[0], GetContact().Particle[1] };
 		}
 
 
