Fixing up compile issue for debug builds on XB1
#rnx

diff --git a/Engine/Source/Runtime/Experimental/Chaos/Public/Chaos/ISpatialAcceleration.h b/Engine/Source/Runtime/Experimental/Chaos/Public/Chaos/ISpatialAcceleration.h
--- a/Engine/Source/Runtime/Experimental/Chaos/Public/Chaos/ISpatialAcceleration.h
+++ b/Engine/Source/Runtime/Experimental/Chaos/Public/Chaos/ISpatialAcceleration.h
@@ -25,7 +25,7 @@ struct CHAOS_API FQueryFastData
 	FReal CurrentLength;
 	FReal InvCurrentLength;
 
-#if PLATFORM_WINDOWS
+#if PLATFORM_WINDOWS || PLATFORM_XBOXONE
 	#pragma warning (push)
 	#pragma warning(disable:4723)
 #endif
@@ -42,7 +42,7 @@ struct CHAOS_API FQueryFastData
 		}
 	}
 
-#if PLATFORM_WINDOWS
+#if PLATFORM_WINDOWS || PLATFORM_XBOXONE
 	#pragma warning(pop)
 #endif
 
