Use COM math in chaos interface. Add X/R version of particle utilities.

diff --git a/Engine/Source/Runtime/Engine/Private/PhysicsEngine/Experimental/PhysInterface_Chaos.cpp b/Engine/Source/Runtime/Engine/Private/PhysicsEngine/Experimental/PhysInterface_Chaos.cpp
--- a/Engine/Source/Runtime/Engine/Private/PhysicsEngine/Experimental/PhysInterface_Chaos.cpp
+++ b/Engine/Source/Runtime/Engine/Private/PhysicsEngine/Experimental/PhysInterface_Chaos.cpp
@@ -29,7 +29,7 @@
 #include "Chaos/GeometryQueries.h"
 #include "Chaos/Plane.h"
 #include "ChaosCheck.h"
-
+#include "Chaos/Particle/ParticleUtilities.h"
 #include "Async/ParallelFor.h"
 #include "Components/PrimitiveComponent.h"
 #include "Physics/PhysicsFiltering.h"
@@ -552,7 +552,7 @@ FVector FPhysInterface_Chaos::GetWorldVelocityAtPoint_AssumesLocked(const FPhysi
 		Chaos::TKinematicGeometryParticle<float, 3>* Kinematic = InActorReference->CastToKinematicParticle();
 		if (ensure(Kinematic))
 		{
-			const Chaos::FVec3 COM = Kinematic->X(); // TODO: Fix once we have separate COM
+			const Chaos::FVec3 COM = Chaos::FParticleUtilitiesGT::GetCoMWorldPosition(Kinematic);
 			const Chaos::FVec3 Diff = InPoint - COM;
 			return Kinematic->V() + Chaos::FVec3::CrossProduct(Diff, Kinematic->W());
 		}
@@ -562,14 +562,11 @@ FVector FPhysInterface_Chaos::GetWorldVelocityAtPoint_AssumesLocked(const FPhysi
 
 FTransform FPhysInterface_Chaos::GetComTransform_AssumesLocked(const FPhysicsActorHandle& InActorReference)
 {
-	// NOTE: This might need to change after we add separate COM transforms.
 	if (ensure(FPhysicsInterface::IsValid(InActorReference)))
 	{
-		if (Chaos::TKinematicGeometryParticle<float, 3>* Kinematic = InActorReference->CastToKinematicParticle())
+		if (const Chaos::TKinematicGeometryParticle<float, 3>* Kinematic = InActorReference->CastToKinematicParticle())
 		{
-			const FTransform WorldTransform(Kinematic->R(), Kinematic->X());
-			const FTransform COMTransform = FTransform(Kinematic->RotationOfMass(), Kinematic->CenterOfMass());
-			return WorldTransform * COMTransform;
+			return Chaos::FParticleUtilitiesGT::GetCoMWorldTransform(Kinematic);
 		}
 	}
 	return FTransform();
@@ -577,7 +574,6 @@ FTransform FPhysInterface_Chaos::GetComTransform_AssumesLocked(const FPhysicsAct
 
 FTransform FPhysInterface_Chaos::GetComTransformLocal_AssumesLocked(const FPhysicsActorHandle& InActorReference)
 {
-	// #todo : Implement
 	if (ensure(FPhysicsInterface::IsValid(InActorReference)))
 	{
 		if (Chaos::TKinematicGeometryParticle<float, 3>* Kinematic = InActorReference->CastToKinematicParticle())
diff --git a/Engine/Source/Runtime/Engine/Private/PhysicsEngine/Experimental/PhysScene_Chaos.cpp b/Engine/Source/Runtime/Engine/Private/PhysicsEngine/Experimental/PhysScene_Chaos.cpp
--- a/Engine/Source/Runtime/Engine/Private/PhysicsEngine/Experimental/PhysScene_Chaos.cpp
+++ b/Engine/Source/Runtime/Engine/Private/PhysicsEngine/Experimental/PhysScene_Chaos.cpp
@@ -1110,11 +1110,11 @@ void FPhysScene_ChaosInterface::AddForceAtPosition_AssumesLocked(FBodyInstance*
 			{
 				const Chaos::FVec3& CurrentForce = Rigid->ExternalForce();
 				const Chaos::FVec3& CurrentTorque = Rigid->ExternalTorque();
-				const Chaos::FVec3& CurrentPosition = Rigid->X();
-				const Chaos::FVec3 WorldCOM = Rigid->X(); // TODO: Remember to fix this once we add COM to particles!
+				const Chaos::FVec3 WorldCOM = FParticleUtilitiesGT::GetCoMWorldPosition(Rigid);
+
 				if (bIsLocalForce)
 				{
-					const Chaos::FRigidTransform3 CurrentTransform(Rigid->X(), Rigid->R());
+					const Chaos::FRigidTransform3 CurrentTransform = FParticleUtilitiesGT::GetActorWorldTransform(Rigid);
 					const Chaos::FVec3 WorldPosition = CurrentTransform.TransformPosition(Position);
 					const Chaos::FVec3 WorldForce = CurrentTransform.TransformVector(Force);
 					const Chaos::FVec3 WorldTorque = Chaos::FVec3::CrossProduct(WorldPosition - WorldCOM, WorldForce);
diff --git a/Engine/Source/Runtime/Experimental/Chaos/Public/Chaos/Particle/ParticleUtilities.h b/Engine/Source/Runtime/Experimental/Chaos/Public/Chaos/Particle/ParticleUtilities.h
--- a/Engine/Source/Runtime/Experimental/Chaos/Public/Chaos/Particle/ParticleUtilities.h
+++ b/Engine/Source/Runtime/Experimental/Chaos/Public/Chaos/Particle/ParticleUtilities.h
@@ -7,10 +7,40 @@
 
 namespace Chaos
 {
+	class FParticleSpatialAccessorXR
+	{
+	public:
+		template<typename T_PARTICLEHANDLE>
+		static inline FVec3& Position(T_PARTICLEHANDLE Particle) { return Particle->X(); }
+		template<typename T_PARTICLEHANDLE>
+		static inline const FVec3& GetPosition(T_PARTICLEHANDLE Particle) { return Particle->X(); }
+		template<typename T_PARTICLEHANDLE>
+		static inline FRotation3& Rotation(T_PARTICLEHANDLE Particle) { return Particle->R(); }
+		template<typename T_PARTICLEHANDLE>
+		static inline const FRotation3& GetRotation(T_PARTICLEHANDLE Particle) { return Particle->R(); }
+	};
+
+	class FParticleSpatialAccessorPQ
+	{
+	public:
+		template<typename T_PARTICLEHANDLE>
+		static inline FVec3& Position(T_PARTICLEHANDLE Particle) { return Particle->P(); }
+		template<typename T_PARTICLEHANDLE>
+		static inline const FVec3& GetPosition(T_PARTICLEHANDLE Particle) { return Particle->P(); }
+		template<typename T_PARTICLEHANDLE>
+		static inline FRotation3& Rotation(T_PARTICLEHANDLE Particle) { return Particle->Q(); }
+		template<typename T_PARTICLEHANDLE>
+		static inline const FRotation3& GetRotation(T_PARTICLEHANDLE Particle) { return Particle->Q(); }
+
+		static inline FVec3& Position(TPBDRigidParticles<FReal, 3>& Particles, const int32 Index) { return Particles.P(Index); }
+		static inline FRotation3& Rotation(TPBDRigidParticles<FReal, 3>& Particles, const int32 Index) { return Particles.Q(Index); }
+	};
+
 	/**
 	 * Particle Space == Actor Space (Transforms)
 	 * Velocities in CoM Space.
 	 */
+	template <typename TSpatialAccessor>
 	class FParticleUtilities_ActorSpace
 	{
 	public:
@@ -23,8 +53,8 @@ namespace Chaos
 		template<typename T_PARTICLEHANDLE>
 		static inline void SetActorWorldTransform(T_PARTICLEHANDLE Particle, const FRigidTransform3& ActorWorldT)
 		{
-			Particle->P() = ActorWorldT.GetTranslation();
-			Particle->Q() = ActorWorldT.GetRotation();
+			TSpatialAccessor::Position(Particle) = ActorWorldT.GetTranslation();
+			TSpatialAccessor::Rotation(Particle) = ActorWorldT.GetRotation();
 		}
 
 		template<typename T_PARTICLEHANDLE>
@@ -78,23 +108,23 @@ namespace Chaos
 		template<typename T_PARTICLEHANDLE>
 		static inline FVec3 GetCoMWorldPosition(T_PARTICLEHANDLE Particle)
 		{
-			return Particle->P() + Particle->Q().RotateVector(Particle->CenterOfMass());
+			return TSpatialAccessor::GetPosition(Particle) + TSpatialAccessor::GetRotation(Particle).RotateVector(Particle->CenterOfMass());
 		}
 
 		static inline FVec3 GetCoMWorldPosition(TPBDRigidParticles<FReal, 3>& Particles, const int32 Index)
 		{
-			return Particles.P(Index) + Particles.Q(Index).RotateVector(Particles.CenterOfMass(Index));
+			return TSpatialAccessor::Position(Particles, Index) + TSpatialAccessor::Rotation(Particles, Index).RotateVector(Particles.CenterOfMass(Index));
 		}
 
 		template<typename T_PARTICLEHANDLE>
 		static inline FRotation3 GetCoMWorldRotation(T_PARTICLEHANDLE Particle)
 		{
-			return Particle->Q() * Particle->RotationOfMass();
+			return TSpatialAccessor::GetRotation(Particle) * Particle->RotationOfMass();
 		}
 	
 		static inline FRotation3 GetCoMWorldRotation(TPBDRigidParticles<FReal, 3>& Particles, const int32 Index)
 		{
-			return Particles.Q(Index) * Particles.RotationOfMass(Index);
+			return TSpatialAccessor::Rotation(Particles, Index) * Particles.RotationOfMass(Index);
 		}
 
 		template<typename T_PARTICLEHANDLE>
@@ -112,8 +142,8 @@ namespace Chaos
 			const FRotation3 Q = QCoM * Particle->RotationOfMass().Inverse();
 			const FVec3 P = PCoM - Q.RotateVector(Particle->CenterOfMass());
 			
-			Particle->P() = P;
-			Particle->Q() = Q;
+			TSpatialAccessor::Position(Particle) = P;
+			TSpatialAccessor::Rotation(Particle) = Q;
 		}
 
 		static inline void SetCoMWorldTransform(TPBDRigidParticles<FReal, 3>& Particles, const int32 Index, const FVec3& PCoM, const FRotation3& QCoM)
@@ -121,8 +151,8 @@ namespace Chaos
 			const FRotation3 Q = QCoM * Particles.RotationOfMass(Index).Inverse();
 			const FVec3 P = PCoM - Q.RotateVector(Particles.CenterOfMass(Index));
 
-			Particles.P(Index) = P;
-			Particles.Q(Index) = Q;
+			TSpatialAccessor::Position(Particles, Index) = P;
+			TSpatialAccessor::Rotation(Particles, Index) = Q;
 		}
 	};
 
@@ -130,14 +160,15 @@ namespace Chaos
 	 * Particle Space == CoM Space.
 	 * Velocities in CoM Space.
 	 */
+	template <typename TSpatialAccessor>
 	class FParticleUtilities_CoMSpace
 	{
 	public:
 		template<typename T_PARTICLEHANDLE>
 		static inline FRigidTransform3 GetActorWorldTransform(T_PARTICLEHANDLE Particle)
 		{
-			FRotation3 ActorQ = Particle->Q() * Particle->RotationOfMass().Inverse();
-			FVec3 ActorP = Particle->P() - ActorQ.RotateVector(Particle->CenterOfMass());
+			FRotation3 ActorQ = TSpatialAccessor::GetRotation(Particle) * Particle->RotationOfMass().Inverse();
+			FVec3 ActorP = TSpatialAccessor::GetPosition(Particle) - ActorQ.RotateVector(Particle->CenterOfMass());
 			return FRigidTransform3(ActorP, ActorQ);
 		}
 
@@ -146,8 +177,8 @@ namespace Chaos
 		{
 			FRotation3 CoMQ = ActorWorldT.GetRotation() * Particle->RotationOfMass();
 			FVec3 CoMP = ActorWorldT.GetTranslation() + ActorWorldT.GetRotation().RotateVector(Particle->CenterOfMass());
-			Particle->P() = CoMP;
-			Particle->Q() = CoMQ;
+			TSpatialAccessor::Position(Particle) = CoMP;
+			TSpatialAccessor::Rotation(Particle) = CoMQ;
 		}
 
 		template<typename T_PARTICLEHANDLE>
@@ -191,49 +222,51 @@ namespace Chaos
 		template<typename T_PARTICLEHANDLE>
 		static inline const FVec3& GetCoMWorldPosition(T_PARTICLEHANDLE Particle)
 		{
-			return Particle->P();
+			return TSpatialAccessor::GetPosition(Particle);
 		}
 
 		static inline const FVec3& GetCoMWorldPosition(TPBDRigidParticles<FReal, 3>& Particles, const int32 Index)
 		{
-			return Particles.P(Index);
+			return TSpatialAccessor::Position(Particles, Index);
 		}
 
 		template<typename T_PARTICLEHANDLE>
 		static inline const FRotation3& GetCoMWorldRotation(T_PARTICLEHANDLE Particle)
 		{
-			return Particle->Q();
+			return TSpatialAccessor::GetRotation(Particle);
 		}
 
 		template<typename T_PARTICLEHANDLE>
 		static inline FRigidTransform3 GetCoMWorldTransform(T_PARTICLEHANDLE Particle)
 		{
-			return FRigidTransform3(GetCoMWorldRotation(Particle), GetCoMWorldPosition(Particle));
+			return FRigidTransform3(GetCoMWorldPosition(Particle), GetCoMWorldRotation(Particle));
 		}
 
 		static inline const FRotation3& GetCoMWorldRotation(TPBDRigidParticles<FReal, 3>& Particles, const int32 Index)
 		{
-			return Particles.Q(Index);
+			return TSpatialAccessor::Rotation(Particles, Index);
 		}
 
 		template<typename T_PARTICLEHANDLE>
 		static inline void SetCoMWorldTransform(T_PARTICLEHANDLE Particle, const FVec3& PCoM, const FRotation3& QCoM)
 		{
-			Particle->P() = PCoM;
-			Particle->Q() = QCoM;
+			TSpatialAccessor::Position(Particle) = PCoM;
+			TSpatialAccessor::Rotation(Particle) = QCoM;
 		}
 
 		static inline void SetCoMWorldTransform(TPBDRigidParticles<FReal, 3>& Particles, const int32 Index, const FVec3& PCoM, const FRotation3& QCoM)
 		{
-			Particles.P(Index) = PCoM;
-			Particles.Q(Index) = QCoM;
+			TSpatialAccessor::Position(Particles, Index) = PCoM;
+			TSpatialAccessor::Rotation(Particles, Index) = QCoM;
 		}
 	};
 
 #if CHAOS_PARTICLE_ACTORTRANSFORM
-	using FParticleUtilities = FParticleUtilities_ActorSpace;
+	using FParticleUtilities = FParticleUtilities_ActorSpace<FParticleSpatialAccessorPQ>;
+	using FParticleUtilitiesGT = FParticleUtilities_ActorSpace<FParticleSpatialAccessorXR>;
 #else
-	using FParticleUtilities = FParticleUtilities_CoMSpace;
+	using FParticleUtilities = FParticleUtilities_CoMSpace<FParticleSpatialAccessorPQ>;
+	using FParticleUtilitiesGT = FParticleUtilities_CoMSpace<FParticleSpatialAccessorXR>;
 #endif
 
 }
\ No newline at end of file
