Chaos : Collisions

- Added support for persistent constraints.

diff --git a/Engine/Source/Runtime/Experimental/Chaos/Private/Chaos/PBDCollisionConstraints.cpp b/Engine/Source/Runtime/Experimental/Chaos/Private/Chaos/PBDCollisionConstraints.cpp
--- a/Engine/Source/Runtime/Experimental/Chaos/Private/Chaos/PBDCollisionConstraints.cpp
+++ b/Engine/Source/Runtime/Experimental/Chaos/Private/Chaos/PBDCollisionConstraints.cpp
@@ -5,6 +5,7 @@
 #include "Chaos/Capsule.h"
 #include "Chaos/ChaosPerfTest.h"
 #include "Chaos/ChaosDebugDraw.h"
+#include "Chaos/PBDCollisionConstraintsContact.h"
 #include "Chaos/CollisionResolutionUtil.h"
 #include "Chaos/CollisionResolution.h"
 #include "Chaos/Defines.h"
@@ -102,12 +103,14 @@ namespace Chaos
 	template<typename T, int d>
 	void TPBDCollisionConstraints<T, d>::AddConstraint(FConstraintBase* ConstraintBase)
 	{
+		FConstraintContainerHandle* Handle = nullptr;
+
 		if (ConstraintBase->GetType() == TRigidBodyPointContactConstraint<T, 3>::StaticType())
 		{
 			TRigidBodyPointContactConstraint<T, d>* PointConstraint = ConstraintBase->template As< TRigidBodyPointContactConstraint<T, d> >();
 
 			int32 Idx = PointConstraints.Add(*PointConstraint);
-			Handles.Add(HandleAllocator.template AllocHandle< TRigidBodyPointContactConstraint<T, d> >(this, Idx));
+			Handle = HandleAllocator.template AllocHandle< TRigidBodyPointContactConstraint<T, d> >(this, Idx);
 
 			delete PointConstraint;
 		}
@@ -116,10 +119,14 @@ namespace Chaos
 			TRigidBodyPlaneContactConstraint<T, d>* PlaneConstraint = ConstraintBase->template As< TRigidBodyPlaneContactConstraint<T, d> >();
 
 			int32 Idx = PlaneConstraints.Add(*PlaneConstraint);
-			Handles.Add(HandleAllocator.template AllocHandle< TRigidBodyPlaneContactConstraint<T, d> >(this, Idx));
+			Handle = HandleAllocator.template AllocHandle< TRigidBodyPlaneContactConstraint<T, d> >(this, Idx);
 
 			delete PlaneConstraint;
 		}
+
+		check(Handle != nullptr);
+		Handles.Add(Handle);
+		Manifolds.Add(Handle->GetKey(), Handle);
 	}
 
 	template<typename T, int d>
@@ -137,10 +144,11 @@ namespace Chaos
 		SCOPE_CYCLE_COUNTER(STAT_Collisions_Reset);
 
 		TArray<FConstraintContainerHandle*> CopyOfHandles = Handles;
+		int LifespanWindow = LifespanCounter + 5;
 
 		for (FConstraintContainerHandle* ContactHandle : CopyOfHandles)
 		{
-			//if (!bEnableCollisions)
+			//if (!bEnableCollisions && ContactHandle->GetContact().Timestamp<LifespanWindow)
 			{
 				RemoveConstraint(ContactHandle);
 			}
@@ -192,6 +200,7 @@ namespace Chaos
 		typename FCollisionConstraintBase::FType ConstraintType = Handle->GetType();
 
 		Handles.RemoveAtSwap(Idx);
+		Manifolds.Remove(Handle->GetKey());
 
 		if (ConstraintType == FCollisionConstraintBase::FType::SinglePoint)
 		{
@@ -226,6 +235,17 @@ namespace Chaos
 		// Clustering uses update constraints to force a re-evaluation. 
 	}
 
+	template<typename T, int d>
+	void TPBDCollisionConstraints<T, d>::UpdateConstraints(T Dt)
+	{
+		PhysicsParallelFor(Handles.Num(), [&](int32 ConstraintHandleIndex)
+		{
+			FConstraintContainerHandle* ConstraintHandle = Handles[ConstraintHandleIndex];
+			check(ConstraintHandle != nullptr);
+			Collisions::Update<ECollisionUpdateType::Deepest, float, 3>(MThickness, ConstraintHandle->GetContact());
+
+		}, bDisableCollisionParallelFor);
+	}
 
 	template<typename T, int d>
 	void TPBDCollisionConstraints<T, d>::Apply(const T Dt, const int32 It, const int32 NumIts)
diff --git a/Engine/Source/Runtime/Experimental/Chaos/Public/Chaos/Collision/CollisionDetector.h b/Engine/Source/Runtime/Experimental/Chaos/Public/Chaos/Collision/CollisionDetector.h
--- a/Engine/Source/Runtime/Experimental/Chaos/Public/Chaos/Collision/CollisionDetector.h
+++ b/Engine/Source/Runtime/Experimental/Chaos/Public/Chaos/Collision/CollisionDetector.h
@@ -48,6 +48,8 @@ namespace Chaos
 				return;
 			}
 
+			CollisionContainer.UpdateConstraints(Dt);
+
 			// Collision detection pipeline: BroadPhase -> NarrowPhase -> Receiver -> Container
 			// Receivers and NarrowPhase are assumed to be stateless atm. If we change that, they need to
 			// be passed into the constructor with the BroadPhase and Container.
diff --git a/Engine/Source/Runtime/Experimental/Chaos/Public/Chaos/Collision/CollisionReceiver.h b/Engine/Source/Runtime/Experimental/Chaos/Public/Chaos/Collision/CollisionReceiver.h
--- a/Engine/Source/Runtime/Experimental/Chaos/Public/Chaos/Collision/CollisionReceiver.h
+++ b/Engine/Source/Runtime/Experimental/Chaos/Public/Chaos/Collision/CollisionReceiver.h
@@ -40,7 +40,14 @@ namespace Chaos
 			FCollisionConstraintBase* ConstraintBase = nullptr;
 			while (Queue.Dequeue(ConstraintBase))
 			{
-				CollisionConstraints.AddConstraint(ConstraintBase);
+				// todo(brice) : This needs to be moved within the MidPhase (which does not exist yet). The 
+				//               NarrowPhase is creating temporary constraints for all interacting bodies 
+				//               even though we might already have those bodies within the constraint.
+				//               So there is a bunch of wasted computation when the CreateConstraints is called. 
+				if (!CollisionConstraints.Contains(ConstraintBase))
+				{
+					CollisionConstraints.AddConstraint(ConstraintBase);
+				}
 			}
 		}
 
diff --git a/Engine/Source/Runtime/Experimental/Chaos/Public/Chaos/CollisionResolutionTypes.h b/Engine/Source/Runtime/Experimental/Chaos/Public/Chaos/CollisionResolutionTypes.h
--- a/Engine/Source/Runtime/Experimental/Chaos/Public/Chaos/CollisionResolutionTypes.h
+++ b/Engine/Source/Runtime/Experimental/Chaos/Public/Chaos/CollisionResolutionTypes.h
@@ -45,15 +45,14 @@ namespace Chaos
 	class CHAOS_API TCollisionContact
 	{
 	public:
-		TCollisionContact(int32 InTimestamp = -INT_MAX, const FImplicitObject* InImplicit0 = nullptr, const FImplicitObject* InImplicit1 = nullptr)
-			: bDisabled(true), Timestamp(InTimestamp), Normal(0), Location(0), Phi(FLT_MAX)
+		TCollisionContact(const FImplicitObject* InImplicit0 = nullptr, const FImplicitObject* InImplicit1 = nullptr)
+			: bDisabled(true), Normal(0), Location(0), Phi(FLT_MAX)
 		{
 			Implicit[0] = InImplicit0;
 			Implicit[1] = InImplicit1;
 		}
 
 		bool bDisabled;
-		int32 Timestamp;
 		TVector<T, d> Normal;
 		TVector<T, d> Location;
 		T Phi;
@@ -62,8 +61,6 @@ namespace Chaos
 	};
 
 
-
-
 	/*
 	*
 	*/
@@ -81,8 +78,18 @@ namespace Chaos
 			Plane         // TRigidBodyPlaneContactConstraint
 		};
 
-		TCollisionConstraintBase() : AccumulatedImpulse(0), Type(FType::None) { Particle[0] = nullptr; Particle[1] = nullptr; }
-		TCollisionConstraintBase(FType InType) : AccumulatedImpulse(0), Type(InType) { Particle[0] = nullptr; Particle[1] = nullptr; }
+		TCollisionConstraintBase(int32 InTimestamp = -INT_MAX) 
+			: AccumulatedImpulse(0)
+			, Timestamp(InTimestamp)
+			, Type(FType::None)
+		{ Particle[0] = nullptr; Particle[1] = nullptr; }
+		
+		TCollisionConstraintBase(FType InType, int32 InTimestamp = -INT_MAX) 
+			: AccumulatedImpulse(0)
+			, Timestamp(InTimestamp)
+			, Type(InType)
+		{ Particle[0] = nullptr; Particle[1] = nullptr; }
+
 		FType GetType() const { return Type; }
 
 		template<class AS_T> AS_T * As() { return static_cast<AS_T*>(this); }
@@ -119,8 +126,10 @@ namespace Chaos
 		FGeometryParticleHandle* Particle[2]; // { Point, Volume }
 		TVector<T, d> AccumulatedImpulse;
 		FManifold Manifold;
+		int32 Timestamp;
 
 	private:
+
 		FType Type;
 	};
 	typedef TCollisionConstraintBase<float, 3> FCollisionConstraintBase;
@@ -190,6 +199,10 @@ namespace Chaos
 		using FConstraintContainer = TPBDCollisionConstraints<T, d>;
 		using FConstraintBase = TCollisionConstraintBase<T, d>;
 
+		using FImplicitPair = TPair<const FImplicitObject*, const FImplicitObject*>;
+		using FGeometryPair = TPair<const TGeometryParticleHandle<T, d>*, const TGeometryParticleHandle<T, d>*>;
+		using FHandleKey = TPair<FImplicitPair, FGeometryPair>;
+
 
 		TPBDCollisionConstraintHandle()
 			: ConstraintType(FConstraintBase::FType::None)
@@ -201,7 +214,25 @@ namespace Chaos
 		{
 		}
 
-		// Handle API
+		FHandleKey GetKey()
+		{
+			const TCollisionConstraintBase<T, d>& Contact = GetContact();
+			return FHandleKey(
+				FImplicitPair( Contact.Manifold.Implicit[0],Contact.Manifold.Implicit[1]),
+				FGeometryPair( Contact.Particle[0], Contact.Particle[1] ) );
+		}
+
+		static FHandleKey MakeKey(const TGeometryParticleHandle<T, d>* Particle0, const TGeometryParticleHandle<T, d>* Particle1, 
+			const FImplicitObject* Implicit0, const FImplicitObject* Implicit1)
+		{
+			return FHandleKey(FImplicitPair(Implicit0,Implicit1), FGeometryPair(Particle0, Particle1));
+		}
+
+		static FHandleKey MakeKey(const TCollisionConstraintBase<T,d> * Base)
+		{
+			return FHandleKey(FImplicitPair(Base->Manifold.Implicit[0], Base->Manifold.Implicit[1]), FGeometryPair(Base->Particle[0], Base->Particle[1]));
+		}
+
 
 		const TCollisionConstraintBase<T, d>& GetContact() const 
 		{ 
@@ -239,8 +270,6 @@ namespace Chaos
 			ConstraintType = InType;
 		}
 
-		// Contact API
-
 		TVector<T, d> GetContactLocation() const
 		{
 			return GetContact().GetLocation();
@@ -269,5 +298,5 @@ namespace Chaos
 
 
 	};
-
+	typedef TPBDCollisionConstraintHandle<float, 3> FPBDCollisionConstraintHandle;
 }
diff --git a/Engine/Source/Runtime/Experimental/Chaos/Public/Chaos/PBDCollisionConstraints.h b/Engine/Source/Runtime/Experimental/Chaos/Public/Chaos/PBDCollisionConstraints.h
--- a/Engine/Source/Runtime/Experimental/Chaos/Public/Chaos/PBDCollisionConstraints.h
+++ b/Engine/Source/Runtime/Experimental/Chaos/Public/Chaos/PBDCollisionConstraints.h
@@ -59,8 +59,7 @@ public:
  	using FPointContactConstraint = TRigidBodyPointContactConstraint<T, d>;
 	using FPlaneContactConstraint = TRigidBodyPlaneContactConstraint<T, d>;
 	using FConstraintHandleAllocator = TConstraintHandleAllocator<TPBDCollisionConstraints<T, d>>;
-	using FConstraintHandleID = TPair<const TGeometryParticleHandle<T, d>*, const TGeometryParticleHandle<T, d>*>;
-
+	using FConstraintContainerHandleKey = typename TPBDCollisionConstraintHandle<T, d>::FHandleKey;
 
 	TPBDCollisionConstraints(const TPBDRigidsSOAs<T,d>& InParticles, 
 		TArrayCollectionArray<bool>& Collided, 
@@ -98,10 +97,18 @@ public:
 
 
 	/**
-	 * Update all constraint values
+	 * Update all constraint values within the set
 	 */
 	void UpdateConstraints(T Dt, const TSet<TGeometryParticleHandle<T, d>*>& AddedParticles);
 
+	/**
+	 * Update all constraint values
+	 */
+
+	/**
+	* Update all constraint values
+	*/
+	void UpdateConstraints(T Dt);
 
 	//
 	// General Rule API
@@ -148,6 +155,11 @@ public:
 		return Handles; 
 	}
 
+	bool Contains(const FConstraintBase* Base) const 
+	{
+		return Manifolds.Contains(FConstraintContainerHandle::MakeKey(Base));
+	}
+
 	void SetThickness(T InThickness)
 	{
 		MThickness = InThickness;
@@ -178,11 +190,6 @@ public:
 		return Handles.Num();
 	}
 
-	FConstraintHandleID GetConstraintHandleID(const FPointContactConstraint & Constraint) const
-	{
-		return  FConstraintHandleID(Constraint.Particle[0], Constraint.Particle[1]);
-	}
-
 	FHandles& GetConstraintHandles()
 	{
 		return Handles;
@@ -204,7 +211,7 @@ private:
 	TArray<FPointContactConstraint> PointConstraints;
 	TArray<FPlaneContactConstraint> PlaneConstraints;
 
-	TMap< FConstraintHandleID, FConstraintContainerHandle* > Manifolds;
+	TMap< FConstraintContainerHandleKey, FConstraintContainerHandle* > Manifolds;
 	TArray<FConstraintContainerHandle*> Handles;
 	FConstraintHandleAllocator HandleAllocator;
 
