Chaos - RBAN remove parallel for in constraint solve

diff --git a/Engine/Source/Runtime/Experimental/Chaos/Private/Chaos/PBDCollisionConstraints.cpp b/Engine/Source/Runtime/Experimental/Chaos/Private/Chaos/PBDCollisionConstraints.cpp
--- a/Engine/Source/Runtime/Experimental/Chaos/Private/Chaos/PBDCollisionConstraints.cpp
+++ b/Engine/Source/Runtime/Experimental/Chaos/Private/Chaos/PBDCollisionConstraints.cpp
@@ -248,17 +248,63 @@ namespace Chaos
 	}
 
 	template<typename T, int d>
-	void TPBDCollisionConstraints<T, d>::Apply(const T Dt, const int32 It, const int32 NumIts)
+	void TPBDCollisionConstraints<T, d>::Apply(const T Dt, const int32 Iterations, const int32 NumIterations)
 	{
-		Apply(Dt, Handles, It, NumIts);
-	}
+		SCOPE_CYCLE_COUNTER(STAT_Collisions_Apply);
+
+		if (MApplyPairIterations > 0)
+		{
+			for (FPointContactConstraint& Contact : PointConstraints)
+			{
+				Collisions::TContactParticleParameters<T> ParticleParameters = { &MCollided, &MPhysicsMaterials, CollisionFrictionOverride, MAngularFriction };
+				Collisions::TContactIterationParameters<T> IterationParameters = { Dt, Iterations, NumIterations, MApplyPairIterations, nullptr };
+				Collisions::Apply(Contact, MThickness, IterationParameters, ParticleParameters);
+			}
 
+			for (FPlaneContactConstraint& Contact : PlaneConstraints)
+			{
+				Collisions::TContactParticleParameters<T> ParticleParameters = { &MCollided, &MPhysicsMaterials, CollisionFrictionOverride, MAngularFriction };
+				Collisions::TContactIterationParameters<T> IterationParameters = { Dt, Iterations, NumIterations, MApplyPairIterations, nullptr };
+				Collisions::Apply(Contact, MThickness, IterationParameters, ParticleParameters);
+			}
+		}
+
+		if (PostApplyCallback != nullptr)
+		{
+			PostApplyCallback(Dt, Handles);
+		}
+	}
 
 	template<typename T, int d>
-	bool TPBDCollisionConstraints<T, d>::ApplyPushOut(const T Dt, const int32 It, const int32 NumIts)
+	bool TPBDCollisionConstraints<T, d>::ApplyPushOut(const T Dt, const int32 Iterations, const int32 NumIterations)
 	{
+		SCOPE_CYCLE_COUNTER(STAT_Collisions_ApplyPushOut);
+
 		TSet<const TGeometryParticleHandle<T, d>*> TempStatic;
-		return ApplyPushOut(Dt, Handles, TempStatic, It, NumIts);
+		bool bNeedsAnotherIteration = false;
+		if (MApplyPushOutPairIterations > 0)
+		{
+			for (FPointContactConstraint& Contact : PointConstraints)
+			{
+				Collisions::TContactParticleParameters<T> ParticleParameters = { &MCollided, &MPhysicsMaterials, CollisionFrictionOverride, MAngularFriction };
+				Collisions::TContactIterationParameters<T> IterationParameters = { Dt, Iterations, NumIterations, MApplyPushOutPairIterations, &bNeedsAnotherIteration };
+				Collisions::ApplyPushOut(Contact, MThickness, TempStatic, IterationParameters, ParticleParameters);
+			}
+
+			for (FPlaneContactConstraint& Contact : PlaneConstraints)
+			{
+				Collisions::TContactParticleParameters<T> ParticleParameters = { &MCollided, &MPhysicsMaterials, CollisionFrictionOverride, MAngularFriction };
+				Collisions::TContactIterationParameters<T> IterationParameters = { Dt, Iterations, NumIterations, MApplyPushOutPairIterations, &bNeedsAnotherIteration };
+				Collisions::ApplyPushOut(Contact, MThickness, TempStatic, IterationParameters, ParticleParameters);
+			}
+		}
+
+		if (PostApplyPushOutCallback != nullptr)
+		{
+			PostApplyPushOutCallback(Dt, Handles, bNeedsAnotherIteration);
+		}
+
+		return bNeedsAnotherIteration;
 	}
 
 
@@ -293,7 +339,7 @@ namespace Chaos
 	{
 		SCOPE_CYCLE_COUNTER(STAT_Collisions_ApplyPushOut);
 
-		bool NeedsAnotherIteration = false;
+		bool bNeedsAnotherIteration = false;
 		if (MApplyPushOutPairIterations > 0)
 		{
 			PhysicsParallelFor(InConstraintHandles.Num(), [&](int32 ConstraintHandleIndex)
@@ -302,7 +348,7 @@ namespace Chaos
 				check(ConstraintHandle != nullptr);
 
 				Collisions::TContactParticleParameters<T> ParticleParameters = { &MCollided, &MPhysicsMaterials, CollisionFrictionOverride, MAngularFriction };
-				Collisions::TContactIterationParameters<T> IterationParameters = { Dt, Iteration, NumIterations, MApplyPushOutPairIterations, &NeedsAnotherIteration };
+				Collisions::TContactIterationParameters<T> IterationParameters = { Dt, Iteration, NumIterations, MApplyPushOutPairIterations, &bNeedsAnotherIteration };
 				Collisions::ApplyPushOut(ConstraintHandle->GetContact(), MThickness, IsTemporarilyStatic, IterationParameters, ParticleParameters);
 
 			}, bDisableCollisionParallelFor);
@@ -310,9 +356,10 @@ namespace Chaos
 
 		if (PostApplyPushOutCallback != nullptr)
 		{
-			PostApplyPushOutCallback(Dt, InConstraintHandles, NeedsAnotherIteration);
+			PostApplyPushOutCallback(Dt, InConstraintHandles, bNeedsAnotherIteration);
 		}
-		return NeedsAnotherIteration;
+
+		return bNeedsAnotherIteration;
 	}
 
 
diff --git a/Engine/Source/Runtime/Experimental/Chaos/Private/Chaos/PBDJointConstraints.cpp b/Engine/Source/Runtime/Experimental/Chaos/Private/Chaos/PBDJointConstraints.cpp
--- a/Engine/Source/Runtime/Experimental/Chaos/Private/Chaos/PBDJointConstraints.cpp
+++ b/Engine/Source/Runtime/Experimental/Chaos/Private/Chaos/PBDJointConstraints.cpp
@@ -394,17 +394,111 @@ namespace Chaos
 
 	void FPBDJointConstraints::Apply(const FReal Dt, const int32 It, const int32 NumIts)
 	{
-		Apply(Dt, Handles, It, NumIts);
+		SCOPE_CYCLE_COUNTER(STAT_Joints_Apply);
+
+		if (PreApplyCallback != nullptr)
+		{
+			PreApplyCallback(Dt, Handles);
+		}
+
+		// Apply joint drives
+		if (Settings.DrivesPhase == EJointSolverPhase::Apply)
+		{
+			SCOPE_CYCLE_COUNTER(STAT_Joints_Drives);
+			if (It == 0)
+			{
+				for (int32 ConstraintIndex = 0; ConstraintIndex < NumConstraints(); ++ConstraintIndex)
+				{
+					ApplyDrives(Dt, ConstraintIndex, 1, 0, 1);
+				}
+			}
+		}
+
+		// Solve for joint position or velocity, depending on settings
+		if (Settings.ApplyPairIterations > 0)
+		{
+			SCOPE_CYCLE_COUNTER(STAT_Joints_Solve);
+			for (int32 ConstraintIndex = 0; ConstraintIndex < NumConstraints(); ++ConstraintIndex)
+			{
+				if (Settings.bEnableVelocitySolve)
+				{
+					SolveVelocity(Dt, ConstraintIndex, Settings.ApplyPairIterations, It, NumIts);
+				}
+				else
+				{
+					SolvePosition(Dt, ConstraintIndex, Settings.ApplyPairIterations, It, NumIts);
+				}
+			}
+		}
+
+		// Correct remaining errors after last call to Solve if enabled in this phase
+		if (Settings.ProjectionPhase == EJointSolverPhase::Apply)
+		{
+			SCOPE_CYCLE_COUNTER(STAT_Joints_Project);
+			int32 ProjectionIt = NumIts - 1;
+			if (It == ProjectionIt)
+			{
+				ApplyProjection(Dt);
+			}
+		}
+
+		if (PostApplyCallback != nullptr)
+		{
+			PostApplyCallback(Dt, Handles);
+		}
 	}
 
 	bool FPBDJointConstraints::ApplyPushOut(const FReal Dt, const int32 It, const int32 NumIts)
 	{
-		return ApplyPushOut(Dt, Handles, It, NumIts);
+		SCOPE_CYCLE_COUNTER(STAT_Joints_ApplyPushOut);
+
+		// @todo(ccaulfield): track whether we are sufficiently solved
+		bool bNeedsAnotherIteration = true;
+
+		// Apply joint drives
+		if (Settings.DrivesPhase == EJointSolverPhase::ApplyPushOut)
+		{
+			SCOPE_CYCLE_COUNTER(STAT_Joints_Drives);
+			if (It == 0)
+			{
+				for (int32 ConstraintIndex = 0; ConstraintIndex < NumConstraints(); ++ConstraintIndex)
+				{
+					ApplyDrives(Dt, ConstraintIndex, 1, 0, 1);
+				}
+			}
+		}
+
+		// Solve for positions
+		if (Settings.ApplyPushOutPairIterations > 0)
+		{
+			SCOPE_CYCLE_COUNTER(STAT_Joints_Solve);
+
+			for (int32 ConstraintIndex = 0; ConstraintIndex < NumConstraints(); ++ConstraintIndex)
+			{
+				SolvePosition(Dt, ConstraintIndex, Settings.ApplyPushOutPairIterations, It, NumIts);
+			}
+		}
+
+		// Correct remaining errors after the last call to Solve (which depends on if PositionSolve is enabled in ApplyPushOut)
+		if (Settings.ProjectionPhase == EJointSolverPhase::ApplyPushOut)
+		{
+			SCOPE_CYCLE_COUNTER(STAT_Joints_Project);
+			int32 ProjectionIt = (Settings.ApplyPushOutPairIterations > 0) ? NumIts - 1 : 0;
+			if (It == ProjectionIt)
+			{
+				ApplyProjection(Dt);
+			}
+		}
+
+		return bNeedsAnotherIteration;
 	}
 
 	void FPBDJointConstraints::ApplyProjection(const FReal Dt)
 	{
-		ApplyProjection(Dt, Handles);
+		for (int32 ConstraintIndex = 0; ConstraintIndex < NumConstraints(); ++ConstraintIndex)
+		{
+			ProjectPosition(Dt, ConstraintIndex, 0, 1);
+		}
 	}
 
 
