CHAOS: fix for when start awake is unchecked, bodies now start sleeping
sleeping bodies now get their island set
sleeping transitions now working

diff --git a/Engine/Source/Runtime/Engine/Private/PhysicsEngine/Experimental/PhysInterface_Chaos.cpp b/Engine/Source/Runtime/Engine/Private/PhysicsEngine/Experimental/PhysInterface_Chaos.cpp
--- a/Engine/Source/Runtime/Engine/Private/PhysicsEngine/Experimental/PhysInterface_Chaos.cpp
+++ b/Engine/Source/Runtime/Engine/Private/PhysicsEngine/Experimental/PhysInterface_Chaos.cpp
@@ -147,7 +147,14 @@ void FPhysInterface_Chaos::CreateActor(const FActorCreationParams& InParams, FPh
 		RigidHandle->SetGravityEnabled(InParams.bEnableGravity);
 		if (InParams.BodyInstance && InParams.BodyInstance->ShouldInstanceSimulatingPhysics())
 		{
-			RigidHandle->SetObjectState(Chaos::EObjectStateType::Dynamic);
+			if (InParams.BodyInstance->bStartAwake)
+			{
+				RigidHandle->SetObjectState(Chaos::EObjectStateType::Dynamic);
+			}
+			else
+			{
+				RigidHandle->SetObjectState(Chaos::EObjectStateType::Sleeping);
+			}
 		}
 		else
 		{
@@ -396,12 +403,40 @@ void FPhysInterface_Chaos::SetIsKinematic_AssumesLocked(const FPhysicsActorHandl
 {
 	if (Chaos::TPBDRigidParticle<float, 3>* Particle = InActorReference->CastToRigidParticle())
 	{
-		//#todo: what if object state has been previously set to sleeping?
 		const Chaos::EObjectStateType NewState
 			= bIsKinematic
 			? Chaos::EObjectStateType::Kinematic
 			: Chaos::EObjectStateType::Dynamic;
-		Particle->SetObjectState(NewState);
+
+		bool AllowedToChangeToNewState = false;
+
+		switch (Particle->ObjectState())
+		{
+		case Chaos::EObjectStateType::Kinematic:
+			// from kinematic we can only go dynamic
+			if (NewState == Chaos::EObjectStateType::Dynamic)
+			{
+				AllowedToChangeToNewState = true;
+			}
+			break;
+
+		case Chaos::EObjectStateType::Dynamic:
+			// from dynamic we can go to sleeping or to kinematic
+			if (NewState == Chaos::EObjectStateType::Kinematic)
+			{
+				AllowedToChangeToNewState = true;
+			}
+			break;
+
+		case Chaos::EObjectStateType::Sleeping:
+			// from sleeping we can't change state without waking first
+			break;
+		}
+		
+		if (AllowedToChangeToNewState)
+		{
+			Particle->SetObjectState(NewState);
+		}
 	}
 	else
 	{
diff --git a/Engine/Source/Runtime/Experimental/Chaos/Private/Chaos/PBDConstraintGraph.cpp b/Engine/Source/Runtime/Experimental/Chaos/Private/Chaos/PBDConstraintGraph.cpp
--- a/Engine/Source/Runtime/Experimental/Chaos/Private/Chaos/PBDConstraintGraph.cpp
+++ b/Engine/Source/Runtime/Experimental/Chaos/Private/Chaos/PBDConstraintGraph.cpp
@@ -330,7 +330,7 @@ void FPBDConstraintGraph::ComputeIslands(const TParticleView<TPBDRigidParticles<
 			for (TGeometryParticleHandle<FReal, 3>* Particle : NewIslandParticles[Island])
 			{
 				TPBDRigidParticleHandle<FReal, 3>* PBDRigid = Particle->CastToRigidParticle();
-				if(PBDRigid && PBDRigid->ObjectState() == EObjectStateType::Dynamic)
+				if (PBDRigid && PBDRigid->ObjectState() != EObjectStateType::Kinematic)
 				{
 					PBDRigid->Island() = Island;
 				}
@@ -392,7 +392,7 @@ void FPBDConstraintGraph::ComputeIslands(const TParticleView<TPBDRigidParticles<
 			for (TGeometryParticleHandle<FReal, 3>* Particle : IslandToParticles[Island])
 			{
 				TPBDRigidParticleHandle<FReal, 3>* PBDRigid = Particle->CastToRigidParticle();
-				const bool bIsDynamic = PBDRigid && PBDRigid->ObjectState() == EObjectStateType::Dynamic;
+				const bool bIsDynamic = PBDRigid && PBDRigid->ObjectState() != EObjectStateType::Kinematic;
 				int32 TmpIsland = bIsDynamic ? PBDRigid->Island() : INDEX_NONE; //question: should we even store non dynamics in this array?
 
 				if (OtherIsland == INDEX_NONE && TmpIsland >= 0)
@@ -425,7 +425,7 @@ void FPBDConstraintGraph::ComputeIslands(const TParticleView<TPBDRigidParticles<
 				for (TGeometryParticleHandle<FReal, 3>* Particle : IslandToParticles[Island])
 				{
 					TPBDRigidParticleHandle<FReal, 3>* PBDRigid = Particle->CastToRigidParticle();
-					if(PBDRigid && PBDRigid->ObjectState() == EObjectStateType::Dynamic)
+					if (PBDRigid && PBDRigid->ObjectState() != EObjectStateType::Kinematic)
 					{
 						if (!PBDRigid->Disabled())	// todo: why is this needed? [we aren't handling enable/disable state changes properly so disabled particles end up in the graph.]
 						{
@@ -482,7 +482,7 @@ void FPBDConstraintGraph::ComputeIsland(const int32 InNode, const int32 Island,
 		}
 
 		TPBDRigidParticleHandle<FReal, 3>* RigidHandle = Node.Particle->CastToRigidParticle();
-		const bool isRigidDynamic = RigidHandle && RigidHandle->ObjectState() == EObjectStateType::Dynamic;
+		const bool isRigidDynamic = RigidHandle && RigidHandle->ObjectState() != EObjectStateType::Kinematic;  //??
 
 		if (isRigidDynamic == false)
 		{
@@ -538,19 +538,23 @@ bool FPBDConstraintGraph::SleepInactive(const int32 Island, const TArrayCollecti
 
 	int32& IslandSleepCount = IslandToSleepCount[Island];
 
-	TVector<FReal, 3> X(0);
 	TVector<FReal, 3> V(0);
 	TVector<FReal, 3> W(0);
 	FReal M = 0;
 	FReal LinearSleepingThreshold = FLT_MAX;
 	FReal AngularSleepingThreshold = FLT_MAX;
+	FReal DefaultLinearSleepingThreshold = (FReal)1;
+	FReal DefaultAngularSleepingThreshold = (FReal)1;
+
+	int32 NumDynamicParticles = 0;
 
 	for (const TGeometryParticleHandle<FReal, 3>* Particle : IslandToParticles[Island])
 	{
 		const TPBDRigidParticleHandle<FReal, 3>* PBDRigid = Particle->CastToRigidParticle();
 		if(PBDRigid && (PBDRigid->ObjectState() == EObjectStateType::Dynamic))
 		{
-			X += PBDRigid->X() * PBDRigid->M();
+			NumDynamicParticles++;
+
 			M += PBDRigid->M();
 			V += PBDRigid->V() * PBDRigid->M();
 
@@ -561,13 +565,18 @@ bool FPBDConstraintGraph::SleepInactive(const int32 Island, const TArrayCollecti
 			}
 			else
 			{
-				LinearSleepingThreshold = (FReal)0;
-				LinearSleepingThreshold = (FReal)0;
+				LinearSleepingThreshold = DefaultLinearSleepingThreshold;
+				AngularSleepingThreshold = DefaultAngularSleepingThreshold;
 			}
 		}
 	}
 
-	X /= M;
+	if (NumDynamicParticles == 0)
+	{
+		// prevent divide by zero - all particles much be sleeping/disabled already
+		return false;
+	}
+
 	V /= M;
 
 	for (const TGeometryParticleHandle<FReal, 3>* Particle: IslandParticles)
@@ -583,7 +592,6 @@ bool FPBDConstraintGraph::SleepInactive(const int32 Island, const TArrayCollecti
 
 	const FReal VSize = V.SizeSquared();
 	const FReal WSize = W.SizeSquared();
-
 	if (VSize < LinearSleepingThreshold && WSize < AngularSleepingThreshold)
 	{
 		if (IslandSleepCount > SleepCountThreshold)
