Chaos : Collisions

Convex Convex Manifold Generation. Still testing so the generation is ifdefed out.

diff --git a/Engine/Source/Runtime/Experimental/Chaos/Private/Chaos/CollisionResolution.cpp b/Engine/Source/Runtime/Experimental/Chaos/Private/Chaos/CollisionResolution.cpp
--- a/Engine/Source/Runtime/Experimental/Chaos/Private/Chaos/CollisionResolution.cpp
+++ b/Engine/Source/Runtime/Experimental/Chaos/Private/Chaos/CollisionResolution.cpp
@@ -558,12 +558,55 @@ namespace Chaos
 			return ReturnValue;
 		}
 
+//#define USE_CONTACT_MANIFOLD
+
 		template<class T, int d>
-		void UpdateConvexConvexConstraint(const FImplicitObject& A, const TRigidTransform<T, d>& ATM, const FImplicitObject& B, const TRigidTransform<T, d>& BTM, const T Thickness, TRigidBodyPlaneContactConstraint<T,d>& Constraint)
+		void UpdateConvexConvexConstraint(const FImplicitObject& Implicit0, const TRigidTransform<T, d>& Transform0, const FImplicitObject& Implicit1, const TRigidTransform<T, d>& Transform1, const T Thickness, TRigidBodyPlaneContactConstraint<T,d>& Constraint)
 		{
-			UpdateContactPoint(Constraint.Manifold, ConvexConvexContactPoint(A, ATM, B, BTM, Thickness));
+			TContactPoint<T> ContactPoint = ConvexConvexContactPoint(Implicit0, Transform0, Implicit1, Transform1, Thickness);
+
+#ifdef USE_CONTACT_MANIFOLD
+			if (ContactPoint.Phi < Thickness)
+			{asdf
+				ContactPoint.Location = Constraint.PlanePosition + Constraint.PlaneNormal * ContactPoint.Phi;
+			}
+#endif
+			UpdateContactPoint(Constraint.Manifold, ContactPoint);
 		}
 
+		template<class T, int d>
+		void UpdateConvexConvexManifold(TRigidBodyPlaneContactConstraint<T, d>& Constraint, const TRigidTransform<T, d>& Transform0, const TRigidTransform<T, d>& Transform1, const T Thickness)
+		{
+#ifdef USE_CONTACT_MANIFOLD
+
+			const FImplicitObject& Implicit0 = *Constraint.Manifold.Implicit[0];
+			const FImplicitObject& Implicit1 = *Constraint.Manifold.Implicit[1];
+			TContactPoint<T> ContactPoint = ConvexConvexContactPoint(Implicit0, Transform0, Implicit1, Transform1, Thickness);
+
+			if (ContactPoint.Phi < Thickness)
+			{
+				FRigidBodyPlaneContactConstraint::FManifold& Manifold = Constraint.Manifold;
+
+				if (!ContactPoint.Normal.Equals(Constraint.PlaneNormal) || !Constraint.PlaneSamples.Num())
+				{
+					Constraint.PlaneSamples.Reset();
+					Constraint.PlaneNormal = ContactPoint.Normal;
+				}
+
+				if (Constraint.PlaneSamples.Num() < 4)
+				{
+					TVector<T, d> SurfaceSample = ContactPoint.Location - ContactPoint.Normal * ContactPoint.Phi;
+					Constraint.PlaneSamples.Add(SurfaceSample);
+					Constraint.PlanePosition = SumArray(Constraint.PlaneSamples) / Constraint.PlaneSamples.Num();
+				}
+
+				ContactPoint.Location = Constraint.PlanePosition;
+				ContactPoint.Normal = Constraint.PlaneNormal;
+			}
+#endif
+		}
+
+
 		template<class T, int d>
 		void ConstructConvexConvexConstraints(TGeometryParticleHandle<T, d>* Particle0, TGeometryParticleHandle<T, d>* Particle1, const FImplicitObject* Implicit0, const FImplicitObject* Implicit1, const TRigidTransform<T, d>& Transform0, const TRigidTransform<T, d>& Transform1, const T Thickness, FCollisionConstraintsArray& NewConstraints)
 		{
@@ -573,6 +616,8 @@ namespace Chaos
 			Constraint->Particle[1] = Particle1;
 			Constraint->SetManifold(Implicit0, Implicit1);
 
+			UpdateConvexConvexManifold(*Constraint, Transform0, Transform1, Thickness);
+
 			UpdateConvexConvexConstraint(*Implicit0, Transform0, *Implicit1, Transform1, Thickness, *Constraint);
 
 			if (Constraint->GetPhi() < Thickness)
@@ -583,6 +628,7 @@ namespace Chaos
 			{
 				delete Constraint;
 			}
+
 		}
 
 
@@ -1187,6 +1233,7 @@ namespace Chaos
 		template void ConstructCapsuleBoxConstraints<float, 3>(TGeometryParticleHandle<float, 3>* Particle0, TGeometryParticleHandle<float, 3>* Particle1, const FImplicitObject* Implicit0, const FImplicitObject* Implicit1, const TRigidTransform<float, 3>& Transform0, const TRigidTransform<float, 3>& Transform1, const float Thickness, FCollisionConstraintsArray& NewConstraints);
 
 		template void UpdateConvexConvexConstraint<float, 3>(const FImplicitObject& A, const TRigidTransform<float, 3>& ATM, const FImplicitObject& B, const TRigidTransform<float, 3>& BTM, const float Thickness, TRigidBodyPlaneContactConstraint<float, 3>& Constraint);
+		template void UpdateConvexConvexManifold<float, 3>(TRigidBodyPlaneContactConstraint<float, 3>& Constraint, const TRigidTransform<float, 3>& ATM, const TRigidTransform<float, 3>& BTM, const float Thickness);
 		template void ConstructConvexConvexConstraints<float, 3>(TGeometryParticleHandle<float, 3>* Particle0, TGeometryParticleHandle<float, 3>* Particle1, const FImplicitObject* Implicit0, const FImplicitObject* Implicit1, const TRigidTransform<float, 3>& Transform0, const TRigidTransform<float, 3>& Transform1, const float Thickness, FCollisionConstraintsArray& NewConstraints);
 
 		template void UpdateLevelsetLevelsetConstraint<ECollisionUpdateType::Any, float, 3>(const float Thickness, TRigidBodyPointContactConstraint<float, 3>& Constraint);
diff --git a/Engine/Source/Runtime/Experimental/Chaos/Private/Chaos/PBDCollisionConstraints.cpp b/Engine/Source/Runtime/Experimental/Chaos/Private/Chaos/PBDCollisionConstraints.cpp
--- a/Engine/Source/Runtime/Experimental/Chaos/Private/Chaos/PBDCollisionConstraints.cpp
+++ b/Engine/Source/Runtime/Experimental/Chaos/Private/Chaos/PBDCollisionConstraints.cpp
@@ -250,6 +250,19 @@ namespace Chaos
 		}, bDisableCollisionParallelFor);
 	}
 
+	template<typename T, int d>
+	void TPBDCollisionConstraints<T, d>::UpdateManifolds(T Dt)
+	{
+		PhysicsParallelFor(PlaneConstraints.Num(), [&](int32 ConstraintIndex)
+		{
+			FConstraintBase& ConstraintBase = PlaneConstraints[ConstraintIndex];
+			if (ConstraintBase.GetType() == FCollisionConstraintBase::FType::Plane)
+			{
+				Collisions::UpdateManifold<float, 3>(MThickness, ConstraintBase);
+			}
+		}, bDisableCollisionParallelFor);
+	}
+
 	template<typename T, int d>
 	void TPBDCollisionConstraints<T, d>::Apply(const T Dt, const int32 Iterations, const int32 NumIterations)
 	{
diff --git a/Engine/Source/Runtime/Experimental/Chaos/Private/Chaos/PBDCollisionConstraintsContact.cpp b/Engine/Source/Runtime/Experimental/Chaos/Private/Chaos/PBDCollisionConstraintsContact.cpp
--- a/Engine/Source/Runtime/Experimental/Chaos/Private/Chaos/PBDCollisionConstraintsContact.cpp
+++ b/Engine/Source/Runtime/Experimental/Chaos/Private/Chaos/PBDCollisionConstraintsContact.cpp
@@ -45,6 +45,27 @@ namespace Chaos
 			}
 		}
 
+		template<typename T, int d>
+		void UpdateManifold(const T Thickness, TCollisionConstraintBase<T, d>& Constraint)
+		{
+			if (Constraint.GetType() == TCollisionConstraintBase<T, d>::FType::Plane)
+			{
+				const TRigidTransform<T, d> Transform0 = GetTransform(Constraint.Particle[0]);
+				const TRigidTransform<T, d> Transform1 = GetTransform(Constraint.Particle[1]);
+
+				const FImplicitObject& Implicit0 = *Constraint.Particle[0]->Geometry();
+				const FImplicitObject& Implicit1 = *Constraint.Particle[1]->Geometry();
+
+				EImplicitObjectType Implicit0Type = GetInnerType(Implicit0.GetType());
+				EImplicitObjectType Implicit1Type = GetInnerType(Implicit1.GetType());
+
+				if (Implicit0.IsConvex() && Implicit1.IsConvex())
+				{
+					UpdateConvexConvexManifold( *Constraint.template As<TRigidBodyPlaneContactConstraint<T, d>>(), Transform0, Transform1, Thickness);
+				}
+			}
+		}
+
 
 		template<typename T, int d>
 		void Apply(TCollisionConstraintBase<T, d>& Constraint, T Thickness, TContactIterationParameters<T> & IterationParameters, TContactParticleParameters<T> & ParticleParameters)
@@ -380,6 +401,7 @@ namespace Chaos
 
 		template void Update<ECollisionUpdateType::Any, float, 3>(const float, TCollisionConstraintBase<float, 3>&);
 		template void Update<ECollisionUpdateType::Deepest, float, 3>(const float, TCollisionConstraintBase<float, 3>&);
+		template void UpdateManifold<float, 3>(const float, TCollisionConstraintBase<float, 3>&);
 		template void Apply<float, 3>(TCollisionConstraintBase<float, 3>&, float, TContactIterationParameters<float> &, TContactParticleParameters<float> &);
 		template void ApplyPushOut<float, 3>(TCollisionConstraintBase<float,3>&, float , const TSet<const TGeometryParticleHandle<float,3>*>&,
 			TContactIterationParameters<float> & IterationParameters, TContactParticleParameters<float> & ParticleParameters);
diff --git a/Engine/Source/Runtime/Experimental/Chaos/Public/Chaos/Collision/CollisionDetector.h b/Engine/Source/Runtime/Experimental/Chaos/Public/Chaos/Collision/CollisionDetector.h
--- a/Engine/Source/Runtime/Experimental/Chaos/Public/Chaos/Collision/CollisionDetector.h
+++ b/Engine/Source/Runtime/Experimental/Chaos/Public/Chaos/Collision/CollisionDetector.h
@@ -48,6 +48,7 @@ namespace Chaos
 				return;
 			}
 
+			CollisionContainer.UpdateManifolds(Dt);
 			CollisionContainer.UpdateConstraints(Dt);
 
 			// Collision detection pipeline: BroadPhase -> NarrowPhase -> Receiver -> Container
diff --git a/Engine/Source/Runtime/Experimental/Chaos/Public/Chaos/CollisionResolution.h b/Engine/Source/Runtime/Experimental/Chaos/Public/Chaos/CollisionResolution.h
--- a/Engine/Source/Runtime/Experimental/Chaos/Public/Chaos/CollisionResolution.h
+++ b/Engine/Source/Runtime/Experimental/Chaos/Public/Chaos/CollisionResolution.h
@@ -117,6 +117,9 @@ namespace Chaos
 			template <typename T, int d>
 			void CHAOS_API UpdateConvexConvexConstraint(const FImplicitObject& A, const TRigidTransform<T, d>& ATM, const FImplicitObject& B, const TRigidTransform<T, d>& BTM, const T Thickness, TRigidBodyPlaneContactConstraint<T, d>& Constraint);
 
+			template<class T, int d>
+			void UpdateConvexConvexManifold(TRigidBodyPlaneContactConstraint<T, d>& Constraint, const TRigidTransform<T, d>& ATM, const TRigidTransform<T, d>& BTM, const T Thickness);
+
 			template <typename T, int d>
 			void CHAOS_API ConstructConvexConvexConstraints(TGeometryParticleHandle<T, d>* Particle0, TGeometryParticleHandle<T, d>* Particle1, const FImplicitObject* Implicit0, const FImplicitObject* Implicit1, const TRigidTransform<T, d>& Transform0, const TRigidTransform<T, d>& Transform1, const T Thickness, FCollisionConstraintsArray& NewConstraints);
 
diff --git a/Engine/Source/Runtime/Experimental/Chaos/Public/Chaos/CollisionResolutionTypes.h b/Engine/Source/Runtime/Experimental/Chaos/Public/Chaos/CollisionResolutionTypes.h
--- a/Engine/Source/Runtime/Experimental/Chaos/Public/Chaos/CollisionResolutionTypes.h
+++ b/Engine/Source/Runtime/Experimental/Chaos/Public/Chaos/CollisionResolutionTypes.h
@@ -95,8 +95,8 @@ namespace Chaos
 		template<class AS_T> AS_T * As() { return static_cast<AS_T*>(this); }
 		template<class AS_T> const AS_T * As() const { return static_cast<AS_T*>(this); }
 
-		bool ContainsManifold(const FImplicitObject* B, const FImplicitObject* A) const { return A == Manifold.Implicit[0] && B == Manifold.Implicit[1]; }
-		void SetManifold(const FImplicitObject* B, const FImplicitObject* A) { Manifold.Implicit[0] = A; Manifold.Implicit[1] = B; }
+		bool ContainsManifold(const FImplicitObject* A, const FImplicitObject* B) const { return A == Manifold.Implicit[0] && B == Manifold.Implicit[1]; }
+		void SetManifold(const FImplicitObject* A, const FImplicitObject* B) { Manifold.Implicit[0] = A; Manifold.Implicit[1] = B; }
 
 		//
 		// API
diff --git a/Engine/Source/Runtime/Experimental/Chaos/Public/Chaos/PBDCollisionConstraints.h b/Engine/Source/Runtime/Experimental/Chaos/Public/Chaos/PBDCollisionConstraints.h
--- a/Engine/Source/Runtime/Experimental/Chaos/Public/Chaos/PBDCollisionConstraints.h
+++ b/Engine/Source/Runtime/Experimental/Chaos/Public/Chaos/PBDCollisionConstraints.h
@@ -105,10 +105,15 @@ public:
 	 * Update all constraint values
 	 */
 
+	 /**
+	 * Update all constraint values
+	 */
+	void UpdateConstraints(T Dt);
+
 	/**
-	* Update all constraint values
+	* Update all contact manifolds
 	*/
-	void UpdateConstraints(T Dt);
+	void UpdateManifolds(T Dt);
 
 	//
 	// General Rule API
diff --git a/Engine/Source/Runtime/Experimental/Chaos/Public/Chaos/PBDCollisionConstraintsContact.h b/Engine/Source/Runtime/Experimental/Chaos/Public/Chaos/PBDCollisionConstraintsContact.h
--- a/Engine/Source/Runtime/Experimental/Chaos/Public/Chaos/PBDCollisionConstraintsContact.h
+++ b/Engine/Source/Runtime/Experimental/Chaos/Public/Chaos/PBDCollisionConstraintsContact.h
@@ -29,6 +29,9 @@ namespace Chaos
 		template<ECollisionUpdateType UpdateType, typename T, int d>
 		void Update(const T Thickness, TCollisionConstraintBase<T, d>& Constraint);
 
+		template<typename T, int d>
+		void UpdateManifold(const T Thickness, TCollisionConstraintBase<T, d>& Constraint);
+
 		template<typename T, int d>
 		void Apply(TCollisionConstraintBase<T, d>& Constraint, T Thickness, TContactIterationParameters<T> & IterationParameters, TContactParticleParameters<T> & ParticleParameters);
 
