Chaos : Collisions

- Fix for collision constraint handle crash.

diff --git a/Engine/Source/Runtime/Experimental/Chaos/Private/Chaos/PBDCollisionConstraints.cpp b/Engine/Source/Runtime/Experimental/Chaos/Private/Chaos/PBDCollisionConstraints.cpp
--- a/Engine/Source/Runtime/Experimental/Chaos/Private/Chaos/PBDCollisionConstraints.cpp
+++ b/Engine/Source/Runtime/Experimental/Chaos/Private/Chaos/PBDCollisionConstraints.cpp
@@ -199,33 +199,39 @@ namespace Chaos
 	template<typename T, int d>
 	void TPBDCollisionConstraints<T, d>::RemoveConstraint(FConstraintContainerHandle* Handle)
 	{
-		int32 Idx = Handle->GetConstraintIndex();
+		FConstraintContainerHandleKey KeyToRemove = Handle->GetKey();
+		int32 Idx = Handle->GetConstraintIndex(); // index into specific array
 		typename FCollisionConstraintBase::FType ConstraintType = Handle->GetType();
 
-		Handles.RemoveAtSwap(Idx);
-		Manifolds.Remove(Handle->GetKey());
-
 		if (ConstraintType == FCollisionConstraintBase::FType::SinglePoint)
 		{
-			PointConstraints.RemoveAtSwap(Idx);
-			if (Idx < PointConstraints.Num())
+			if (Idx < PointConstraints.Num() - 1)
 			{
-				Handles[Idx]->SetConstraintIndex(Idx, ConstraintType);
+				// update the handle
+				FConstraintContainerHandleKey Key = FPBDCollisionConstraintHandle::MakeKey(&PointConstraints.Last());
+				Manifolds[Key]->SetConstraintIndex(Idx, ConstraintType);
 			}
+			PointConstraints.RemoveAtSwap(Idx);
+
 		}
 		else if (ConstraintType == FCollisionConstraintBase::FType::Plane)
 		{
-			PlaneConstraints.RemoveAtSwap(Idx);
-			if (Idx < PlaneConstraints.Num())
+			if (Idx < PlaneConstraints.Num() - 1)
 			{
-				Handles[Idx]->SetConstraintIndex(Idx, ConstraintType);
+				// update the handle
+				FConstraintContainerHandleKey Key = FPBDCollisionConstraintHandle::MakeKey(&PlaneConstraints.Last());
+				Manifolds[Key]->SetConstraintIndex(Idx, ConstraintType);
 			}
+			PlaneConstraints.RemoveAtSwap(Idx);
 		}
 		else 
 		{
 			check(false);
 		}
-		
+
+		Manifolds.Remove(KeyToRemove); // todo(brice): Add index to the handle to prevent the search. 
+		Handles.Remove(Handle);
+
 		ensure(Handles.Num() == PointConstraints.Num() + PlaneConstraints.Num());
 
 		HandleAllocator.FreeHandle(Handle);
