Speed up Collisions processing where objects are sleeping

diff --git a/Engine/Source/Runtime/Experimental/Chaos/Public/Chaos/Collision/SpatialAccelerationBroadPhase.h b/Engine/Source/Runtime/Experimental/Chaos/Public/Chaos/Collision/SpatialAccelerationBroadPhase.h
--- a/Engine/Source/Runtime/Experimental/Chaos/Public/Chaos/Collision/SpatialAccelerationBroadPhase.h
+++ b/Engine/Source/Runtime/Experimental/Chaos/Public/Chaos/Collision/SpatialAccelerationBroadPhase.h
@@ -124,92 +124,95 @@ namespace Chaos
 	private:
 		template<typename T_SPATIALACCELERATION>
 		void ProduceParticleOverlaps(
-			FReal Dt, 
-			TTransientPBDRigidParticleHandle<FReal, 3>& Particle1, 
-			const T_SPATIALACCELERATION& InSpatialAcceleration, 
-			FNarrowPhase& NarrowPhase, 
-			FAsyncCollisionReceiver& Receiver,
-			CollisionStats::FStatData& StatData)
+		    FReal Dt,
+		    TTransientPBDRigidParticleHandle<FReal, 3>& Particle1,
+		    const T_SPATIALACCELERATION& InSpatialAcceleration,
+		    FNarrowPhase& NarrowPhase,
+		    FAsyncCollisionReceiver& Receiver,
+		    CollisionStats::FStatData& StatData)
 		{
 			CHAOS_COLLISION_STAT(StatData.IncrementSimulatedParticles());
 
 			TArray<TAccelerationStructureHandle<FReal, 3>> PotentialIntersections;
 
-			// @todo(ccaulfield): the spatial acceleration scheme needs to know the expanded bboxes for all particles, not just the one doing the test
-			const bool bBody1Bounded = HasBoundingBox(Particle1);
-			const FReal Box1Thickness = ComputeBoundsThickness(Particle1, Dt, BoundsThickness, BoundsThicknessVelocityInflation).Size();
-			if (bBody1Bounded)
+			if (Particle1.CastToRigidParticle() && Particle1.ObjectState() == EObjectStateType::Dynamic)
 			{
-				// @todo(ccaulfield): COLLISION - see the NOTE below - fix it
+				// @todo(ccaulfield): the spatial acceleration scheme needs to know the expanded bboxes for all particles, not just the one doing the test
+				const bool bBody1Bounded = HasBoundingBox(Particle1);
+				const FReal Box1Thickness = ComputeBoundsThickness(Particle1, Dt, BoundsThickness, BoundsThicknessVelocityInflation).Size();
+				if (bBody1Bounded)
+				{
+					// @todo(ccaulfield): COLLISION - see the NOTE below - fix it
 #if CHAOS_PARTICLEHANDLE_TODO
-				const TBox <FReal, 3> Box1 = InSpatialAcceleration.GetWorldSpaceBoundingBox(Particle1);
+					const TBox<FReal, 3> Box1 = InSpatialAcceleration.GetWorldSpaceBoundingBox(Particle1);
 #else
-				const TBox <FReal, 3> Box1 = ComputeWorldSpaceBoundingBox(Particle1);	// NOTE: this ignores the velocity expansion which is wrong
+					const TBox<FReal, 3> Box1 = ComputeWorldSpaceBoundingBox(Particle1); // NOTE: this ignores the velocity expansion which is wrong
 #endif
 
-				CHAOS_COLLISION_STAT(StatData.RecordBoundsData(Box1));
+					CHAOS_COLLISION_STAT(StatData.RecordBoundsData(Box1));
 
-				TSimOverlapVisitor<FReal, 3> OverlapVisitor(PotentialIntersections);
-				InSpatialAcceleration.Overlap(Box1, OverlapVisitor);
-			}
-			else
-			{
-				const auto& GlobalElems = InSpatialAcceleration.GlobalObjects();
-				PotentialIntersections.Reserve(GlobalElems.Num());
-
-				for (auto& Elem : GlobalElems)
-				{
-					PotentialIntersections.Add(Elem.Payload);
+					TSimOverlapVisitor<FReal, 3> OverlapVisitor(PotentialIntersections);
+					InSpatialAcceleration.Overlap(Box1, OverlapVisitor);
 				}
-			}
-
-			CHAOS_COLLISION_STAT(StatData.RecordBroadphasePotentials(PotentialIntersections.Num()))
+				else
+				{
+					const auto& GlobalElems = InSpatialAcceleration.GlobalObjects();
+					PotentialIntersections.Reserve(GlobalElems.Num());
 
-			const int32 NumPotentials = PotentialIntersections.Num();
-			for (int32 i = 0; i < NumPotentials; ++i)
-			{
-				auto& Particle2 = *PotentialIntersections[i].GetGeometryParticleHandle_PhysicsThread();
-				const TGenericParticleHandle<FReal, 3> Particle2Generic(&Particle2);
+					for (auto& Elem : GlobalElems)
+					{
+						PotentialIntersections.Add(Elem.Payload);
+					}
+				}
 
-				// Broad Phase Culling 
-				// CollisionGroup == 0 : Collide_With_Everything
-				// CollisionGroup == INDEX_NONE : Disabled collisions
-				// CollisionGroup_A != CollisionGroup_B : Skip Check
+				CHAOS_COLLISION_STAT(StatData.RecordBroadphasePotentials(PotentialIntersections.Num()))
 
-				if (Particle1.CollisionGroup() == INDEX_NONE || Particle2Generic->CollisionGroup() == INDEX_NONE)
+				const int32 NumPotentials = PotentialIntersections.Num();
+				for (int32 i = 0; i < NumPotentials; ++i)
 				{
-					continue;
-				}
-				if (Particle1.CollisionGroup() && Particle2Generic->CollisionGroup() && Particle1.CollisionGroup() != Particle2Generic->CollisionGroup())
-				{
-					continue;
-				}
+					auto& Particle2 = *PotentialIntersections[i].GetGeometryParticleHandle_PhysicsThread();
+					const TGenericParticleHandle<FReal, 3> Particle2Generic(&Particle2);
 
-				if (!Particle1.Geometry() && !Particle2.Geometry())
-				{
-					continue;
-				}
+					// Broad Phase Culling
+					// CollisionGroup == 0 : Collide_With_Everything
+					// CollisionGroup == INDEX_NONE : Disabled collisions
+					// CollisionGroup_A != CollisionGroup_B : Skip Check
 
-				const bool bBody2Bounded = HasBoundingBox(Particle2);
+					if (Particle1.CollisionGroup() == INDEX_NONE || Particle2Generic->CollisionGroup() == INDEX_NONE)
+					{
+						continue;
+					}
+					if (Particle1.CollisionGroup() && Particle2Generic->CollisionGroup() && Particle1.CollisionGroup() != Particle2Generic->CollisionGroup())
+					{
+						continue;
+					}
 
-				if (Particle1.Handle() == Particle2.Handle())
-				{
-					continue;
-				}
+					if (!Particle1.Geometry() && !Particle2.Geometry())
+					{
+						continue;
+					}
 
-				const bool bIsParticle2Dynamic = Particle2.CastToRigidParticle() && Particle2.ObjectState() == EObjectStateType::Dynamic;
-				if (bBody1Bounded == bBody2Bounded && bIsParticle2Dynamic)
-				{
-					//no bidirectional constraints. 
-					if (Particle2.ParticleID() > Particle1.ParticleID())
+					const bool bBody2Bounded = HasBoundingBox(Particle2);
+
+					if (Particle1.Handle() == Particle2.Handle())
 					{
 						continue;
 					}
-				}
 
-				const FReal Box2Thickness = bIsParticle2Dynamic ? ComputeBoundsThickness(*Particle2.CastToRigidParticle(), Dt, BoundsThickness, BoundsThicknessVelocityInflation).Size() : (FReal)0;
+					const bool bIsParticle2Dynamic = Particle2.CastToRigidParticle() && Particle2.ObjectState() == EObjectStateType::Dynamic;
+					if (bBody1Bounded == bBody2Bounded && bIsParticle2Dynamic)
+					{
+						//no bidirectional constraints.
+						if (Particle2.ParticleID() > Particle1.ParticleID())
+						{
+							continue;
+						}
+					}
 
-				NarrowPhase.GenerateCollisions(Dt, Receiver, Particle1.Handle(), Particle2.Handle(), FMath::Max(Box1Thickness, Box2Thickness), StatData);
+					const FReal Box2Thickness = bIsParticle2Dynamic ? ComputeBoundsThickness(*Particle2.CastToRigidParticle(), Dt, BoundsThickness, BoundsThicknessVelocityInflation).Size() : (FReal)0;
+
+					NarrowPhase.GenerateCollisions(Dt, Receiver, Particle1.Handle(), Particle2.Handle(), FMath::Max(Box1Thickness, Box2Thickness), StatData);
+				}
 			}
 
 			CHAOS_COLLISION_STAT(StatData.FinalizeData());
