Chaos - fix integration

diff --git a/Engine/Source/Runtime/Engine/Private/PhysicsEngine/ImmediatePhysics/ImmediatePhysicsChaos/ImmediatePhysicsSimulation_Chaos.cpp b/Engine/Source/Runtime/Engine/Private/PhysicsEngine/ImmediatePhysics/ImmediatePhysicsChaos/ImmediatePhysicsSimulation_Chaos.cpp
--- a/Engine/Source/Runtime/Engine/Private/PhysicsEngine/ImmediatePhysics/ImmediatePhysicsChaos/ImmediatePhysicsSimulation_Chaos.cpp
+++ b/Engine/Source/Runtime/Engine/Private/PhysicsEngine/ImmediatePhysics/ImmediatePhysicsChaos/ImmediatePhysicsSimulation_Chaos.cpp
@@ -151,7 +151,11 @@ namespace ImmediatePhysics_Chaos
 	//
 
 	FSimulation::FSimulation()
-		: Particles()
+		: PotentiallyCollidingPairs()
+		, CollidedParticles()
+		, ParticleMaterials()
+		, PerParticleMaterials()
+		, Particles()
 		, Joints()
 		, Collisions(Particles, CollidedParticles, ParticleMaterials, 0, 0)
 		, BroadPhase(PotentiallyCollidingPairs)
diff --git a/Engine/Source/Runtime/Experimental/Chaos/Public/Chaos/Particle/ParticleUtilities.h b/Engine/Source/Runtime/Experimental/Chaos/Public/Chaos/Particle/ParticleUtilities.h
--- a/Engine/Source/Runtime/Experimental/Chaos/Public/Chaos/Particle/ParticleUtilities.h
+++ b/Engine/Source/Runtime/Experimental/Chaos/Public/Chaos/Particle/ParticleUtilities.h
@@ -18,6 +18,9 @@ namespace Chaos
 		static inline FRotation3& Rotation(T_PARTICLEHANDLE Particle) { return Particle->R(); }
 		template<typename T_PARTICLEHANDLE>
 		static inline const FRotation3& GetRotation(T_PARTICLEHANDLE Particle) { return Particle->R(); }
+
+		static inline FVec3& Position(TPBDRigidParticles<FReal, 3>& Particles, const int32 Index) { return Particles.X(Index); }
+		static inline FRotation3& Rotation(TPBDRigidParticles<FReal, 3>& Particles, const int32 Index) { return Particles.R(Index); }
 	};
 
 	class FParticleSpatialAccessorPQ
@@ -262,9 +265,13 @@ namespace Chaos
 	};
 
 #if CHAOS_PARTICLE_ACTORTRANSFORM
+	using FParticleUtilitiesPQ = FParticleUtilities_ActorSpace<FParticleSpatialAccessorPQ>;
+	using FParticleUtilitiesXR = FParticleUtilities_ActorSpace<FParticleSpatialAccessorXR>;
 	using FParticleUtilities = FParticleUtilities_ActorSpace<FParticleSpatialAccessorPQ>;
 	using FParticleUtilitiesGT = FParticleUtilities_ActorSpace<FParticleSpatialAccessorXR>;
 #else
+	using FParticleUtilitiesPQ = FParticleUtilities_CoMSpace<FParticleSpatialAccessorPQ>;
+	using FParticleUtilitiesXR = FParticleUtilities_CoMSpace<FParticleSpatialAccessorXR>;
 	using FParticleUtilities = FParticleUtilities_CoMSpace<FParticleSpatialAccessorPQ>;
 	using FParticleUtilitiesGT = FParticleUtilities_CoMSpace<FParticleSpatialAccessorXR>;
 #endif
diff --git a/Engine/Source/Runtime/Experimental/Chaos/Public/Chaos/PerParticlePBDEulerStep.h b/Engine/Source/Runtime/Experimental/Chaos/Public/Chaos/PerParticlePBDEulerStep.h
--- a/Engine/Source/Runtime/Experimental/Chaos/Public/Chaos/PerParticlePBDEulerStep.h
+++ b/Engine/Source/Runtime/Experimental/Chaos/Public/Chaos/PerParticlePBDEulerStep.h
@@ -30,24 +30,24 @@ class TPerParticlePBDEulerStep : public TPerParticleRule<T, d>
 
 	inline void Apply(TPBDRigidParticles<T, d>& InParticles, const T Dt, const int32 Index) const override //-V762
 	{
-		FVec3 PCoM = FParticleUtilities::GetCoMWorldPosition(InParticles, Index);
-		FRotation3 QCoM = FParticleUtilities::GetCoMWorldRotation(InParticles, Index);
+		FVec3 PCoM = FParticleUtilitiesXR::GetCoMWorldPosition(InParticles, Index);
+		FRotation3 QCoM = FParticleUtilitiesXR::GetCoMWorldRotation(InParticles, Index);
 
 		PCoM = PCoM + InParticles.V(Index) * Dt;
 		QCoM = FRotation3::IntegrateRotationWithAngularVelocity(QCoM, InParticles.W(Index), Dt);
 
-		FParticleUtilities::SetCoMWorldTransform(InParticles, Index, PCoM, QCoM);
+		FParticleUtilitiesPQ::SetCoMWorldTransform(InParticles, Index, PCoM, QCoM);
 	}
 
 	inline void Apply(TTransientPBDRigidParticleHandle<T, d>& Particle, const T Dt) const override
 	{
-		FVec3 PCoM = FParticleUtilities::GetCoMWorldPosition(&Particle);
-		FRotation3 QCoM = FParticleUtilities::GetCoMWorldRotation(&Particle);
+		FVec3 PCoM = FParticleUtilitiesXR::GetCoMWorldPosition(&Particle);
+		FRotation3 QCoM = FParticleUtilitiesXR::GetCoMWorldRotation(&Particle);
 
 		PCoM = PCoM + Particle.V() * Dt;
 		QCoM = FRotation3::IntegrateRotationWithAngularVelocity(QCoM, Particle.W(), Dt);
 
-		FParticleUtilities::SetCoMWorldTransform(&Particle, PCoM, QCoM);
+		FParticleUtilitiesPQ::SetCoMWorldTransform(&Particle, PCoM, QCoM);
 	}
 };
 }
