Chaos : Collisions

- Persistent collisions updates

diff --git a/Engine/Source/Runtime/Experimental/Chaos/Private/Chaos/CollisionResolution.cpp b/Engine/Source/Runtime/Experimental/Chaos/Private/Chaos/CollisionResolution.cpp
--- a/Engine/Source/Runtime/Experimental/Chaos/Private/Chaos/CollisionResolution.cpp
+++ b/Engine/Source/Runtime/Experimental/Chaos/Private/Chaos/CollisionResolution.cpp
@@ -565,24 +565,23 @@ namespace Chaos
 		{
 			TContactPoint<T> ContactPoint = ConvexConvexContactPoint(Implicit0, Transform0, Implicit1, Transform1, Thickness);
 
-#ifdef USE_CONTACT_MANIFOLD
+			/*
 			if (ContactPoint.Phi < Thickness)
-			{asdf
+			{
 				ContactPoint.Location = Constraint.PlanePosition + Constraint.PlaneNormal * ContactPoint.Phi;
-			}
-#endif
+			}/**/
 			UpdateContactPoint(Constraint.Manifold, ContactPoint);
 		}
 
 		template<class T, int d>
 		void UpdateConvexConvexManifold(TRigidBodyPlaneContactConstraint<T, d>& Constraint, const TRigidTransform<T, d>& Transform0, const TRigidTransform<T, d>& Transform1, const T Thickness)
 		{
-#ifdef USE_CONTACT_MANIFOLD
 
 			const FImplicitObject& Implicit0 = *Constraint.Manifold.Implicit[0];
 			const FImplicitObject& Implicit1 = *Constraint.Manifold.Implicit[1];
 			TContactPoint<T> ContactPoint = ConvexConvexContactPoint(Implicit0, Transform0, Implicit1, Transform1, Thickness);
 
+			/*
 			if (ContactPoint.Phi < Thickness)
 			{
 				FRigidBodyPlaneContactConstraint::FManifold& Manifold = Constraint.Manifold;
@@ -602,8 +601,8 @@ namespace Chaos
 
 				ContactPoint.Location = Constraint.PlanePosition;
 				ContactPoint.Normal = Constraint.PlaneNormal;
-			}
-#endif
+			}/**/
+
 		}
 
 
diff --git a/Engine/Source/Runtime/Experimental/Chaos/Private/Chaos/PBDCollisionConstraints.cpp b/Engine/Source/Runtime/Experimental/Chaos/Private/Chaos/PBDCollisionConstraints.cpp
--- a/Engine/Source/Runtime/Experimental/Chaos/Private/Chaos/PBDCollisionConstraints.cpp
+++ b/Engine/Source/Runtime/Experimental/Chaos/Private/Chaos/PBDCollisionConstraints.cpp
@@ -107,6 +107,8 @@ namespace Chaos
 	template<typename T, int d>
 	void TPBDCollisionConstraints<T, d>::AddConstraint(FConstraintBase* ConstraintBase)
 	{
+		// WARNING : ConstraintBase is about to be deleted!
+
 		FConstraintContainerHandle* Handle = nullptr;
 
 		if (ConstraintBase->GetType() == TRigidBodyPointContactConstraint<T, 3>::StaticType())
@@ -131,6 +133,7 @@ namespace Chaos
 		check(Handle != nullptr);
 		Handles.Add(Handle);
 		Manifolds.Add(Handle->GetKey(), Handle);
+		Handle->GetContact().Timestamp = LifespanCounter;
 	}
 
 	template<typename T, int d>
@@ -147,11 +150,11 @@ namespace Chaos
 		SCOPE_CYCLE_COUNTER(STAT_Collisions_Reset);
 
 		TArray<FConstraintContainerHandle*> CopyOfHandles = Handles;
-		int LifespanWindow = LifespanCounter + 5;
 
+		int32 LifespanWindow = LifespanCounter - 1;
 		for (FConstraintContainerHandle* ContactHandle : CopyOfHandles)
 		{
-			//if (!bEnableCollisions && ContactHandle->GetContact().Timestamp<LifespanWindow)
+			if (!bEnableCollisions || ContactHandle->GetContact().Timestamp< LifespanWindow)
 			{
 				RemoveConstraint(ContactHandle);
 			}
@@ -253,6 +256,10 @@ namespace Chaos
 			check(ConstraintHandle != nullptr);
 			Collisions::Update<ECollisionUpdateType::Deepest, float, 3>(MThickness, ConstraintHandle->GetContact());
 
+			if (ConstraintHandle->GetContact().GetPhi() < MThickness) 
+			{
+				ConstraintHandle->GetContact().Timestamp = LifespanCounter;
+			}
 		}, bDisableCollisionParallelFor);
 	}
 
diff --git a/Engine/Source/Runtime/Experimental/Chaos/Private/Chaos/PBDCollisionConstraintsContact.cpp b/Engine/Source/Runtime/Experimental/Chaos/Private/Chaos/PBDCollisionConstraintsContact.cpp
--- a/Engine/Source/Runtime/Experimental/Chaos/Private/Chaos/PBDCollisionConstraintsContact.cpp
+++ b/Engine/Source/Runtime/Experimental/Chaos/Private/Chaos/PBDCollisionConstraintsContact.cpp
@@ -53,11 +53,8 @@ namespace Chaos
 				const TRigidTransform<T, d> Transform0 = GetTransform(Constraint.Particle[0]);
 				const TRigidTransform<T, d> Transform1 = GetTransform(Constraint.Particle[1]);
 
-				const FImplicitObject& Implicit0 = *Constraint.Particle[0]->Geometry();
-				const FImplicitObject& Implicit1 = *Constraint.Particle[1]->Geometry();
-
-				EImplicitObjectType Implicit0Type = GetInnerType(Implicit0.GetType());
-				EImplicitObjectType Implicit1Type = GetInnerType(Implicit1.GetType());
+				const FImplicitObject& Implicit0 = *Constraint.Manifold.Implicit[0];
+				const FImplicitObject& Implicit1 = *Constraint.Manifold.Implicit[1];
 
 				if (Implicit0.IsConvex() && Implicit1.IsConvex())
 				{
diff --git a/Engine/Source/Runtime/Experimental/Chaos/Private/Chaos/PBDRigidsEvolutionGBF.cpp b/Engine/Source/Runtime/Experimental/Chaos/Private/Chaos/PBDRigidsEvolutionGBF.cpp
--- a/Engine/Source/Runtime/Experimental/Chaos/Private/Chaos/PBDRigidsEvolutionGBF.cpp
+++ b/Engine/Source/Runtime/Experimental/Chaos/Private/Chaos/PBDRigidsEvolutionGBF.cpp
@@ -154,8 +154,6 @@ void TPBDRigidsEvolutionGBF<T, d>::AdvanceOneTimeStep(const T Dt, const T StepFr
 		PostIntegrateCallback();
 	}
 
-	// @todo(ccaulfield): When we have persistent constraints we will want to run collision detection before
-	// UpdateConstraintPositionBasedState. For now it is after because we just delete all collisions there right now.
 	{
 		SCOPE_CYCLE_COUNTER(STAT_Evolution_UpdateConstraintPositionBasedState);
 		UpdateConstraintPositionBasedState(Dt);
diff --git a/Engine/Source/Runtime/Experimental/Chaos/Public/Chaos/PBDCollisionConstraints.h b/Engine/Source/Runtime/Experimental/Chaos/Public/Chaos/PBDCollisionConstraints.h
--- a/Engine/Source/Runtime/Experimental/Chaos/Public/Chaos/PBDCollisionConstraints.h
+++ b/Engine/Source/Runtime/Experimental/Chaos/Public/Chaos/PBDCollisionConstraints.h
@@ -68,7 +68,14 @@ public:
 
 	virtual ~TPBDCollisionConstraints() {}
 
-
+	/**
+	*  Add the constraint to the container. 
+	*
+	*  @todo(chaos) : Update to use a custom allocator. 
+	*  The InConstraint should be a point to unmanaged, raw memory. 
+	*  This function will make a deep copy of the constraint and 
+	*  then delete the InConstraint. 
+	*/
 	void AddConstraint(FConstraintBase* InConstraint);
 
 	/**
