#UE4 Fix CIS error

diff --git a/Engine/Source/Runtime/Experimental/Chaos/Private/Chaos/CollisionResolution.cpp b/Engine/Source/Runtime/Experimental/Chaos/Private/Chaos/CollisionResolution.cpp
--- a/Engine/Source/Runtime/Experimental/Chaos/Private/Chaos/CollisionResolution.cpp
+++ b/Engine/Source/Runtime/Experimental/Chaos/Private/Chaos/CollisionResolution.cpp
@@ -569,7 +569,7 @@ namespace Chaos
 			if (ContactPoint.Phi < Thickness)
 			{
 				ContactPoint.Location = Constraint.PlanePosition + Constraint.PlaneNormal * ContactPoint.Phi;
-			}/**/
+			}*/
 			UpdateContactPoint(Constraint.Manifold, ContactPoint);
 		}
 
