Chaos : Collisions

- Force point constraints to be deleted every frame.

diff --git a/Engine/Source/Runtime/Experimental/Chaos/Private/Chaos/CollisionResolution.cpp b/Engine/Source/Runtime/Experimental/Chaos/Private/Chaos/CollisionResolution.cpp
--- a/Engine/Source/Runtime/Experimental/Chaos/Private/Chaos/CollisionResolution.cpp
+++ b/Engine/Source/Runtime/Experimental/Chaos/Private/Chaos/CollisionResolution.cpp
@@ -558,51 +558,35 @@ namespace Chaos
 			return ReturnValue;
 		}
 
-//#define USE_CONTACT_MANIFOLD
-
 		template<class T, int d>
 		void UpdateConvexConvexConstraint(const FImplicitObject& Implicit0, const TRigidTransform<T, d>& Transform0, const FImplicitObject& Implicit1, const TRigidTransform<T, d>& Transform1, const T Thickness, TRigidBodyPlaneContactConstraint<T,d>& Constraint)
 		{
-			TContactPoint<T> ContactPoint = ConvexConvexContactPoint(Implicit0, Transform0, Implicit1, Transform1, Thickness);
-
-			/*
-			if (ContactPoint.Phi < Thickness)
-			{
-				ContactPoint.Location = Constraint.PlanePosition + Constraint.PlaneNormal * ContactPoint.Phi;
-			}*/
+			TContactPoint<T> ContactPoint = ConvexConvexContactPoint(Implicit0, Transform0, Implicit1, Transform1, Thickness);	
+			//ContactPoint.Location = Constraint.PlanePosition + Constraint.PlaneNormal * ContactPoint.Phi;
+			//ContactPoint.Normal = Constraint.PlaneNormal;
 			UpdateContactPoint(Constraint.Manifold, ContactPoint);
 		}
 
 		template<class T, int d>
 		void UpdateConvexConvexManifold(TRigidBodyPlaneContactConstraint<T, d>& Constraint, const TRigidTransform<T, d>& Transform0, const TRigidTransform<T, d>& Transform1, const T Thickness)
 		{
+			TContactPoint<T> ContactPoint = ConvexConvexContactPoint(*Constraint.Manifold.Implicit[0], Transform0, *Constraint.Manifold.Implicit[1], Transform1, Thickness);
 
-			const FImplicitObject& Implicit0 = *Constraint.Manifold.Implicit[0];
-			const FImplicitObject& Implicit1 = *Constraint.Manifold.Implicit[1];
-			TContactPoint<T> ContactPoint = ConvexConvexContactPoint(Implicit0, Transform0, Implicit1, Transform1, Thickness);
-
+			FRigidBodyPlaneContactConstraint::FManifold& Manifold = Constraint.Manifold;
 			/*
-			if (ContactPoint.Phi < Thickness)
+			if (!ContactPoint.Normal.Equals(Constraint.PlaneNormal) || !Constraint.PlaneSamples.Num())
 			{
-				FRigidBodyPlaneContactConstraint::FManifold& Manifold = Constraint.Manifold;
-
-				if (!ContactPoint.Normal.Equals(Constraint.PlaneNormal) || !Constraint.PlaneSamples.Num())
-				{
-					Constraint.PlaneSamples.Reset();
-					Constraint.PlaneNormal = ContactPoint.Normal;
-				}
-
-				if (Constraint.PlaneSamples.Num() < 4)
-				{
-					TVector<T, d> SurfaceSample = ContactPoint.Location - ContactPoint.Normal * ContactPoint.Phi;
-					Constraint.PlaneSamples.Add(SurfaceSample);
-					Constraint.PlanePosition = SumArray(Constraint.PlaneSamples) / Constraint.PlaneSamples.Num();
-				}
-
-				ContactPoint.Location = Constraint.PlanePosition;
-				ContactPoint.Normal = Constraint.PlaneNormal;
-			}/**/
+				Constraint.PlaneSamples.Reset();
+				Constraint.PlaneNormal = ContactPoint.Normal;
+			}
 
+			if (Constraint.PlaneSamples.Num() < 4)
+			{
+				TVector<T, d> SurfaceSample = ContactPoint.Location - ContactPoint.Normal * ContactPoint.Phi;
+				Constraint.PlaneSamples.Add(SurfaceSample);
+				Constraint.PlanePosition = SumArray(Constraint.PlaneSamples) / Constraint.PlaneSamples.Num();
+			}
+			*/
 		}
 
 
diff --git a/Engine/Source/Runtime/Experimental/Chaos/Private/Chaos/PBDCollisionConstraints.cpp b/Engine/Source/Runtime/Experimental/Chaos/Private/Chaos/PBDCollisionConstraints.cpp
--- a/Engine/Source/Runtime/Experimental/Chaos/Private/Chaos/PBDCollisionConstraints.cpp
+++ b/Engine/Source/Runtime/Experimental/Chaos/Private/Chaos/PBDCollisionConstraints.cpp
@@ -117,6 +117,7 @@ namespace Chaos
 
 			int32 Idx = PointConstraints.Add(*PointConstraint);
 			Handle = HandleAllocator.template AllocHandle< TRigidBodyPointContactConstraint<T, d> >(this, Idx);
+			Handle->GetContact().Timestamp = -INT_MAX; // force point constraints to be deleted.
 
 			delete PointConstraint;
 		}
@@ -126,6 +127,7 @@ namespace Chaos
 
 			int32 Idx = PlaneConstraints.Add(*PlaneConstraint);
 			Handle = HandleAllocator.template AllocHandle< TRigidBodyPlaneContactConstraint<T, d> >(this, Idx);
+			Handle->GetContact().Timestamp = LifespanCounter;
 
 			delete PlaneConstraint;
 		}
@@ -133,7 +135,6 @@ namespace Chaos
 		check(Handle != nullptr);
 		Handles.Add(Handle);
 		Manifolds.Add(Handle->GetKey(), Handle);
-		Handle->GetContact().Timestamp = LifespanCounter;
 	}
 
 	template<typename T, int d>
