Chaos : Collisions

Manifold updates.

diff --git a/Engine/Source/Runtime/Experimental/Chaos/Private/Chaos/CollisionResolution.cpp b/Engine/Source/Runtime/Experimental/Chaos/Private/Chaos/CollisionResolution.cpp
--- a/Engine/Source/Runtime/Experimental/Chaos/Private/Chaos/CollisionResolution.cpp
+++ b/Engine/Source/Runtime/Experimental/Chaos/Private/Chaos/CollisionResolution.cpp
@@ -559,32 +559,43 @@ namespace Chaos
 		}
 
 		template<class T, int d>
-		void UpdateConvexConvexConstraint(const FImplicitObject& Implicit0, const TRigidTransform<T, d>& Transform0, const FImplicitObject& Implicit1, const TRigidTransform<T, d>& Transform1, const T Thickness, TRigidBodyPlaneContactConstraint<T,d>& Constraint)
+		void UpdateConvexConvexConstraint(const FImplicitObject& Implicit0, const TRigidTransform<T, d>& Transform0, const FImplicitObject& Implicit1, const TRigidTransform<T, d>& Transform1, const T Thickness, TRigidBodyIterativeContactConstraint<T,d>& Constraint)
 		{
-			TContactPoint<T> ContactPoint = ConvexConvexContactPoint(Implicit0, Transform0, Implicit1, Transform1, Thickness);	
-			//ContactPoint.Location = Constraint.PlanePosition + Constraint.PlaneNormal * ContactPoint.Phi;
-			//ContactPoint.Normal = Constraint.PlaneNormal;
+			TContactPoint<T> ContactPoint = ConvexConvexContactPoint(Implicit0, Transform0, Implicit1, Transform1, Thickness);
+
+			/*
+			TContactPoint<T> ContactPoint;
+
+			const TRigidTransform<T, d> AToBTM = Transform0.GetRelativeTransform(Transform1);
+			ContactPoint.Location = AToBTM.TransformPosition(Constraint.LocalPosition);
+
+			ContactPoint.Phi = Implicit1.PhiWithNormal(ContactPoint.Location, ContactPoint.Normal);
+			ContactPoint.Normal = Constraint.WorldNormal;
+			ContactPoint.Location = Transform0.TransformPosition(Constraint.LocalPosition);
+			*/
+
 			UpdateContactPoint(Constraint.Manifold, ContactPoint);
 		}
 
 		template<class T, int d>
-		void UpdateConvexConvexManifold(TRigidBodyPlaneContactConstraint<T, d>& Constraint, const TRigidTransform<T, d>& Transform0, const TRigidTransform<T, d>& Transform1, const T Thickness)
+		void UpdateConvexConvexManifold(TRigidBodyIterativeContactConstraint<T, d>& Constraint, const TRigidTransform<T, d>& Transform0, const TRigidTransform<T, d>& Transform1, const T Thickness)
 		{
+			/*
 			TContactPoint<T> ContactPoint = ConvexConvexContactPoint(*Constraint.Manifold.Implicit[0], Transform0, *Constraint.Manifold.Implicit[1], Transform1, Thickness);
 
-			FRigidBodyPlaneContactConstraint::FManifold& Manifold = Constraint.Manifold;
-			/*
-			if (!ContactPoint.Normal.Equals(Constraint.PlaneNormal) || !Constraint.PlaneSamples.Num())
+			FRigidBodyIterativeContactConstraint::FManifold& Manifold = Constraint.Manifold;
+
+			if (!ContactPoint.Normal.Equals(Constraint.WorldNormal) || !Constraint.LocalSamples.Num())
 			{
-				Constraint.PlaneSamples.Reset();
-				Constraint.PlaneNormal = ContactPoint.Normal;
+				Constraint.LocalSamples.Reset();
+				Constraint.WorldNormal = ContactPoint.Normal;
 			}
 
-			if (Constraint.PlaneSamples.Num() < 4)
+			if (Constraint.LocalSamples.Num() < 4)
 			{
-				TVector<T, d> SurfaceSample = ContactPoint.Location - ContactPoint.Normal * ContactPoint.Phi;
-				Constraint.PlaneSamples.Add(SurfaceSample);
-				Constraint.PlanePosition = SumArray(Constraint.PlaneSamples) / Constraint.PlaneSamples.Num();
+				TVector<T, d> SurfaceSample = Transform0.InverseTransformPosition(ContactPoint.Location);
+				Constraint.LocalSamples.Add(SurfaceSample);
+				Constraint.LocalPosition = SumArray(Constraint.LocalSamples) / Constraint.LocalSamples.Num();
 			}
 			*/
 		}
@@ -593,7 +604,7 @@ namespace Chaos
 		template<class T, int d>
 		void ConstructConvexConvexConstraints(TGeometryParticleHandle<T, d>* Particle0, TGeometryParticleHandle<T, d>* Particle1, const FImplicitObject* Implicit0, const FImplicitObject* Implicit1, const TRigidTransform<T, d>& Transform0, const TRigidTransform<T, d>& Transform1, const T Thickness, FCollisionConstraintsArray& NewConstraints)
 		{
-			FRigidBodyPlaneContactConstraint* Constraint = new FRigidBodyPlaneContactConstraint;
+			FRigidBodyIterativeContactConstraint* Constraint = new FRigidBodyIterativeContactConstraint;
 
 			Constraint->Particle[0] = Particle0;
 			Constraint->Particle[1] = Particle1;
@@ -1028,7 +1039,7 @@ namespace Chaos
 			}
 			else if(Implicit0.IsConvex() && Implicit1.IsConvex())
 			{
-				UpdateConvexConvexConstraint(Implicit0, Transform0, Implicit1, Transform1, Thickness, *ConstraintBase.template As<TRigidBodyPlaneContactConstraint<T, d>>());
+				UpdateConvexConvexConstraint(Implicit0, Transform0, Implicit1, Transform1, Thickness, *ConstraintBase.template As<TRigidBodyIterativeContactConstraint<T, d>>());
 			}
 			else if (Implicit1.IsUnderlyingUnion())
 			{
@@ -1067,9 +1078,9 @@ namespace Chaos
 						return;
 					}
 				}
-				else if (NewConstraints[i]->GetType() == TRigidBodyPlaneContactConstraint<T, d>::StaticType())
+				else if (NewConstraints[i]->GetType() == TRigidBodyIterativeContactConstraint<T, d>::StaticType())
 				{
-					if (NewConstraints[i]->As< TRigidBodyPlaneContactConstraint<T, d> >()->ContainsManifold(Implicit0, Implicit1))
+					if (NewConstraints[i]->As< TRigidBodyIterativeContactConstraint<T, d> >()->ContainsManifold(Implicit0, Implicit1))
 					{
 						return;
 					}
@@ -1215,8 +1226,8 @@ namespace Chaos
 		template void UpdateCapsuleBoxConstraint<float, 3>(const TCapsule<float>& A, const TRigidTransform<float, 3>& ATransform, const TBox<float, 3>& B, const TRigidTransform<float, 3>& BTransform, const float Thickness, TRigidBodyPointContactConstraint<float, 3>& Constraint);
 		template void ConstructCapsuleBoxConstraints<float, 3>(TGeometryParticleHandle<float, 3>* Particle0, TGeometryParticleHandle<float, 3>* Particle1, const FImplicitObject* Implicit0, const FImplicitObject* Implicit1, const TRigidTransform<float, 3>& Transform0, const TRigidTransform<float, 3>& Transform1, const float Thickness, FCollisionConstraintsArray& NewConstraints);
 
-		template void UpdateConvexConvexConstraint<float, 3>(const FImplicitObject& A, const TRigidTransform<float, 3>& ATM, const FImplicitObject& B, const TRigidTransform<float, 3>& BTM, const float Thickness, TRigidBodyPlaneContactConstraint<float, 3>& Constraint);
-		template void UpdateConvexConvexManifold<float, 3>(TRigidBodyPlaneContactConstraint<float, 3>& Constraint, const TRigidTransform<float, 3>& ATM, const TRigidTransform<float, 3>& BTM, const float Thickness);
+		template void UpdateConvexConvexConstraint<float, 3>(const FImplicitObject& A, const TRigidTransform<float, 3>& ATM, const FImplicitObject& B, const TRigidTransform<float, 3>& BTM, const float Thickness, TRigidBodyIterativeContactConstraint<float, 3>& Constraint);
+		template void UpdateConvexConvexManifold<float, 3>(TRigidBodyIterativeContactConstraint<float, 3>& Constraint, const TRigidTransform<float, 3>& ATM, const TRigidTransform<float, 3>& BTM, const float Thickness);
 		template void ConstructConvexConvexConstraints<float, 3>(TGeometryParticleHandle<float, 3>* Particle0, TGeometryParticleHandle<float, 3>* Particle1, const FImplicitObject* Implicit0, const FImplicitObject* Implicit1, const TRigidTransform<float, 3>& Transform0, const TRigidTransform<float, 3>& Transform1, const float Thickness, FCollisionConstraintsArray& NewConstraints);
 
 		template void UpdateLevelsetLevelsetConstraint<ECollisionUpdateType::Any, float, 3>(const float Thickness, TRigidBodyPointContactConstraint<float, 3>& Constraint);
diff --git a/Engine/Source/Runtime/Experimental/Chaos/Private/Chaos/PBDCollisionConstraints.cpp b/Engine/Source/Runtime/Experimental/Chaos/Private/Chaos/PBDCollisionConstraints.cpp
--- a/Engine/Source/Runtime/Experimental/Chaos/Private/Chaos/PBDCollisionConstraints.cpp
+++ b/Engine/Source/Runtime/Experimental/Chaos/Private/Chaos/PBDCollisionConstraints.cpp
@@ -121,12 +121,12 @@ namespace Chaos
 
 			delete PointConstraint;
 		}
-		else if (ConstraintBase->GetType() == TRigidBodyPlaneContactConstraint<T, 3>::StaticType())
+		else if (ConstraintBase->GetType() == TRigidBodyIterativeContactConstraint<T, 3>::StaticType())
 		{
-			TRigidBodyPlaneContactConstraint<T, d>* PlaneConstraint = ConstraintBase->template As< TRigidBodyPlaneContactConstraint<T, d> >();
+			TRigidBodyIterativeContactConstraint<T, d>* PlaneConstraint = ConstraintBase->template As< TRigidBodyIterativeContactConstraint<T, d> >();
 
 			int32 Idx = PlaneConstraints.Add(*PlaneConstraint);
-			Handle = HandleAllocator.template AllocHandle< TRigidBodyPlaneContactConstraint<T, d> >(this, Idx);
+			Handle = HandleAllocator.template AllocHandle< TRigidBodyIterativeContactConstraint<T, d> >(this, Idx);
 			Handle->GetContact().Timestamp = LifespanCounter;
 
 			delete PlaneConstraint;
@@ -218,7 +218,7 @@ namespace Chaos
 			PointConstraints.RemoveAtSwap(Idx);
 
 		}
-		else if (ConstraintType == FCollisionConstraintBase::FType::Plane)
+		else if (ConstraintType == FCollisionConstraintBase::FType::MultiPoint)
 		{
 			if (Idx < PlaneConstraints.Num() - 1)
 			{
@@ -270,7 +270,7 @@ namespace Chaos
 		PhysicsParallelFor(PlaneConstraints.Num(), [&](int32 ConstraintIndex)
 		{
 			FConstraintBase& ConstraintBase = PlaneConstraints[ConstraintIndex];
-			if (ConstraintBase.GetType() == FCollisionConstraintBase::FType::Plane)
+			if (ConstraintBase.GetType() == FCollisionConstraintBase::FType::MultiPoint)
 			{
 				Collisions::UpdateManifold<float, 3>(MThickness, ConstraintBase);
 			}
diff --git a/Engine/Source/Runtime/Experimental/Chaos/Private/Chaos/PBDCollisionConstraintsContact.cpp b/Engine/Source/Runtime/Experimental/Chaos/Private/Chaos/PBDCollisionConstraintsContact.cpp
--- a/Engine/Source/Runtime/Experimental/Chaos/Private/Chaos/PBDCollisionConstraintsContact.cpp
+++ b/Engine/Source/Runtime/Experimental/Chaos/Private/Chaos/PBDCollisionConstraintsContact.cpp
@@ -39,7 +39,7 @@ namespace Chaos
 					UpdateConstraintImp<UpdateType>(*Constraint.Particle[0]->Geometry(), ParticleTM, *Constraint.Particle[1]->Geometry(), LevelsetTM, Thickness, Constraint);
 				}
 			}
-			else if(Constraint.GetType() == TCollisionConstraintBase<T, d>::FType::Plane)
+			else if(Constraint.GetType() == TCollisionConstraintBase<T, d>::FType::MultiPoint)
 			{
 				UpdateConstraintImp<UpdateType>(*Constraint.Particle[0]->Geometry(), ParticleTM, *Constraint.Particle[1]->Geometry(), LevelsetTM, Thickness, Constraint);
 			}
@@ -48,7 +48,7 @@ namespace Chaos
 		template<typename T, int d>
 		void UpdateManifold(const T Thickness, TCollisionConstraintBase<T, d>& Constraint)
 		{
-			if (Constraint.GetType() == TCollisionConstraintBase<T, d>::FType::Plane)
+			if (Constraint.GetType() == TCollisionConstraintBase<T, d>::FType::MultiPoint)
 			{
 				const TRigidTransform<T, d> Transform0 = GetTransform(Constraint.Particle[0]);
 				const TRigidTransform<T, d> Transform1 = GetTransform(Constraint.Particle[1]);
@@ -58,7 +58,7 @@ namespace Chaos
 
 				if (Implicit0.IsConvex() && Implicit1.IsConvex())
 				{
-					UpdateConvexConvexManifold( *Constraint.template As<TRigidBodyPlaneContactConstraint<T, d>>(), Transform0, Transform1, Thickness);
+					UpdateConvexConvexManifold( *Constraint.template As<TRigidBodyIterativeContactConstraint<T, d>>(), Transform0, Transform1, Thickness);
 				}
 			}
 		}
diff --git a/Engine/Source/Runtime/Experimental/Chaos/Public/Chaos/CollisionResolution.h b/Engine/Source/Runtime/Experimental/Chaos/Public/Chaos/CollisionResolution.h
--- a/Engine/Source/Runtime/Experimental/Chaos/Public/Chaos/CollisionResolution.h
+++ b/Engine/Source/Runtime/Experimental/Chaos/Public/Chaos/CollisionResolution.h
@@ -115,10 +115,10 @@ namespace Chaos
 			// Convex-Convex
 			//
 			template <typename T, int d>
-			void CHAOS_API UpdateConvexConvexConstraint(const FImplicitObject& A, const TRigidTransform<T, d>& ATM, const FImplicitObject& B, const TRigidTransform<T, d>& BTM, const T Thickness, TRigidBodyPlaneContactConstraint<T, d>& Constraint);
+			void CHAOS_API UpdateConvexConvexConstraint(const FImplicitObject& A, const TRigidTransform<T, d>& ATM, const FImplicitObject& B, const TRigidTransform<T, d>& BTM, const T Thickness, TRigidBodyIterativeContactConstraint<T, d>& Constraint);
 
 			template<class T, int d>
-			void UpdateConvexConvexManifold(TRigidBodyPlaneContactConstraint<T, d>& Constraint, const TRigidTransform<T, d>& ATM, const TRigidTransform<T, d>& BTM, const T Thickness);
+			void UpdateConvexConvexManifold(TRigidBodyIterativeContactConstraint<T, d>& Constraint, const TRigidTransform<T, d>& ATM, const TRigidTransform<T, d>& BTM, const T Thickness);
 
 			template <typename T, int d>
 			void CHAOS_API ConstructConvexConvexConstraints(TGeometryParticleHandle<T, d>* Particle0, TGeometryParticleHandle<T, d>* Particle1, const FImplicitObject* Implicit0, const FImplicitObject* Implicit1, const TRigidTransform<T, d>& Transform0, const TRigidTransform<T, d>& Transform1, const T Thickness, FCollisionConstraintsArray& NewConstraints);
diff --git a/Engine/Source/Runtime/Experimental/Chaos/Public/Chaos/CollisionResolutionTypes.h b/Engine/Source/Runtime/Experimental/Chaos/Public/Chaos/CollisionResolutionTypes.h
--- a/Engine/Source/Runtime/Experimental/Chaos/Public/Chaos/CollisionResolutionTypes.h
+++ b/Engine/Source/Runtime/Experimental/Chaos/Public/Chaos/CollisionResolutionTypes.h
@@ -74,8 +74,8 @@ namespace Chaos
 		enum class FType
 		{
 			None = 0,     // default value also indicates a invalid constraint
-			SinglePoint,   // TRigidBodyPointContactConstraint
-			Plane         // TRigidBodyPlaneContactConstraint
+			SinglePoint,  //   TRigidBodyPointContactConstraint
+			MultiPoint    //   TRigidBodyIterativeContactConstraint
 		};
 
 		TCollisionConstraintBase(int32 InTimestamp = -INT_MAX) 
@@ -157,21 +157,21 @@ namespace Chaos
 	*
 	*/
 	template<class T, int d>
-	class CHAOS_API TRigidBodyPlaneContactConstraint : public TCollisionConstraintBase<T, d>
+	class CHAOS_API TRigidBodyIterativeContactConstraint : public TCollisionConstraintBase<T, d>
 	{
 	public:
 		using Base = TCollisionConstraintBase<T, d>;
 		using FGeometryParticleHandle = TGeometryParticleHandle<T, d>;
 		using Base::Particle;
 
-		TRigidBodyPlaneContactConstraint() : Base(Base::FType::Plane), PlaneNormal(0), PlanePosition(0) {}
-		static typename Base::FType StaticType() { return Base::FType::Plane; };
+		TRigidBodyIterativeContactConstraint() : Base(Base::FType::MultiPoint), WorldNormal(0), LocalPosition(0) {}
+		static typename Base::FType StaticType() { return Base::FType::MultiPoint; };
 
-		TVector<T, d> PlaneNormal;
-		TVector<T, d> PlanePosition;
-		TArray<TVector<T, d>> PlaneSamples; // iterative plane samples
+		TVector<T, d> WorldNormal; // sampled normal
+		TVector<T, d> LocalPosition; // averaged position
+		TArray<TVector<T, d>> LocalSamples; // iterative samples
 	};
-	typedef TRigidBodyPlaneContactConstraint<float, 3> FRigidBodyPlaneContactConstraint;
+	typedef TRigidBodyIterativeContactConstraint<float, 3> FRigidBodyIterativeContactConstraint;
 
 
 	//
@@ -240,7 +240,7 @@ namespace Chaos
 			{
 				return ConstraintContainer->PointConstraints[ConstraintIndex];
 			}
-			check(GetType() == FConstraintBase::FType::Plane);
+			check(GetType() == FConstraintBase::FType::MultiPoint);
 			return ConstraintContainer->PlaneConstraints[ConstraintIndex];
 		}
 
@@ -250,7 +250,7 @@ namespace Chaos
 			{
 				return ConstraintContainer->PointConstraints[ConstraintIndex];
 			}
-			check(GetType() == FConstraintBase::FType::Plane);
+			check(GetType() == FConstraintBase::FType::MultiPoint);
 			return ConstraintContainer->PlaneConstraints[ConstraintIndex];
 		}
 
@@ -259,8 +259,8 @@ namespace Chaos
 		const TRigidBodyPointContactConstraint<T, d>& GetPointContact() const { check(GetType() == FConstraintBase::FType::SinglePoint); return ConstraintContainer->PointConstraints[ConstraintIndex]; }
 		TRigidBodyPointContactConstraint<T, d>& GetPointContact() { check(GetType() == FConstraintBase::FType::SinglePoint); return ConstraintContainer->PointConstraints[ConstraintIndex]; }
 
-		const TRigidBodyPlaneContactConstraint<T, d>& GetPlaneContact() const { check(GetType() == FConstraintBase::FType::Plane); return ConstraintContainer->PlaneConstraints[ConstraintIndex]; }
-		TRigidBodyPlaneContactConstraint<T, d>& GetPlaneContact() { check(GetType() == FConstraintBase::FType::Plane); return ConstraintContainer->PlaneConstraints[ConstraintIndex]; }
+		const TRigidBodyIterativeContactConstraint<T, d>& GetPlaneContact() const { check(GetType() == FConstraintBase::FType::MultiPoint); return ConstraintContainer->PlaneConstraints[ConstraintIndex]; }
+		TRigidBodyIterativeContactConstraint<T, d>& GetPlaneContact() { check(GetType() == FConstraintBase::FType::MultiPoint); return ConstraintContainer->PlaneConstraints[ConstraintIndex]; }
 
 		typename FConstraintBase::FType GetType() const { return ConstraintType; }
 
diff --git a/Engine/Source/Runtime/Experimental/Chaos/Public/Chaos/PBDCollisionConstraints.h b/Engine/Source/Runtime/Experimental/Chaos/Public/Chaos/PBDCollisionConstraints.h
--- a/Engine/Source/Runtime/Experimental/Chaos/Public/Chaos/PBDCollisionConstraints.h
+++ b/Engine/Source/Runtime/Experimental/Chaos/Public/Chaos/PBDCollisionConstraints.h
@@ -57,7 +57,7 @@ public:
 	using FConstraintContainerHandle = TPBDCollisionConstraintHandle<T, d>;
 	using FConstraintBase = TCollisionConstraintBase<T, d>;
  	using FPointContactConstraint = TRigidBodyPointContactConstraint<T, d>;
-	using FPlaneContactConstraint = TRigidBodyPlaneContactConstraint<T, d>;
+	using FPlaneContactConstraint = TRigidBodyIterativeContactConstraint<T, d>;
 	using FConstraintHandleAllocator = TConstraintHandleAllocator<TPBDCollisionConstraints<T, d>>;
 	using FConstraintContainerHandleKey = typename TPBDCollisionConstraintHandle<T, d>::FHandleKey;
 
