Only integrate active particles

diff --git a/Engine/Source/Runtime/Experimental/Chaos/Private/Chaos/PBDRigidsEvolutionGBF.cpp b/Engine/Source/Runtime/Experimental/Chaos/Private/Chaos/PBDRigidsEvolutionGBF.cpp
--- a/Engine/Source/Runtime/Experimental/Chaos/Private/Chaos/PBDRigidsEvolutionGBF.cpp
+++ b/Engine/Source/Runtime/Experimental/Chaos/Private/Chaos/PBDRigidsEvolutionGBF.cpp
@@ -141,7 +141,7 @@ void TPBDRigidsEvolutionGBF<T, d>::AdvanceOneTimeStep(const T Dt, const T StepFr
 
 	{
 		SCOPE_CYCLE_COUNTER(STAT_Evolution_Integrate);
-		Integrate(Particles.GetNonDisabledDynamicView(), Dt);	//Question: should we use an awake view?
+		Integrate(Particles.GetActiveParticlesView(), Dt);
 	}
 
 	{
@@ -335,4 +335,5 @@ void TPBDRigidsEvolutionGBF<T,d>::Serialize(FChaosArchive& Ar)
 
 }
 
-template class Chaos::TPBDRigidsEvolutionGBF<float, 3>;
\ No newline at end of file
+template class Chaos::TPBDRigidsEvolutionGBF<float, 3>;
+
