Chaos : Collisions

- renaming and cleanup.

diff --git a/Engine/Source/Runtime/Experimental/Chaos/Private/Chaos/PBDCollisionConstraints.cpp b/Engine/Source/Runtime/Experimental/Chaos/Private/Chaos/PBDCollisionConstraints.cpp
--- a/Engine/Source/Runtime/Experimental/Chaos/Private/Chaos/PBDCollisionConstraints.cpp
+++ b/Engine/Source/Runtime/Experimental/Chaos/Private/Chaos/PBDCollisionConstraints.cpp
@@ -123,13 +123,13 @@ namespace Chaos
 		}
 		else if (ConstraintBase->GetType() == TRigidBodyIterativeContactConstraint<T, 3>::StaticType())
 		{
-			TRigidBodyIterativeContactConstraint<T, d>* PlaneConstraint = ConstraintBase->template As< TRigidBodyIterativeContactConstraint<T, d> >();
+			TRigidBodyIterativeContactConstraint<T, d>* IterativeConstraint = ConstraintBase->template As< TRigidBodyIterativeContactConstraint<T, d> >();
 
-			int32 Idx = PlaneConstraints.Add(*PlaneConstraint);
+			int32 Idx = IterativeConstraints.Add(*IterativeConstraint);
 			Handle = HandleAllocator.template AllocHandle< TRigidBodyIterativeContactConstraint<T, d> >(this, Idx);
 			Handle->GetContact().Timestamp = LifespanCounter;
 
-			delete PlaneConstraint;
+			delete IterativeConstraint;
 		}
 
 		check(Handle != nullptr);
@@ -220,13 +220,13 @@ namespace Chaos
 		}
 		else if (ConstraintType == FCollisionConstraintBase::FType::MultiPoint)
 		{
-			if (Idx < PlaneConstraints.Num() - 1)
+			if (Idx < IterativeConstraints.Num() - 1)
 			{
 				// update the handle
-				FConstraintContainerHandleKey Key = FPBDCollisionConstraintHandle::MakeKey(&PlaneConstraints.Last());
+				FConstraintContainerHandleKey Key = FPBDCollisionConstraintHandle::MakeKey(&IterativeConstraints.Last());
 				Manifolds[Key]->SetConstraintIndex(Idx, ConstraintType);
 			}
-			PlaneConstraints.RemoveAtSwap(Idx);
+			IterativeConstraints.RemoveAtSwap(Idx);
 		}
 		else 
 		{
@@ -236,7 +236,7 @@ namespace Chaos
 		Manifolds.Remove(KeyToRemove); // todo(brice): Add index to the handle to prevent the search. 
 		Handles.Remove(Handle);
 
-		ensure(Handles.Num() == PointConstraints.Num() + PlaneConstraints.Num());
+		ensure(Handles.Num() == PointConstraints.Num() + IterativeConstraints.Num());
 
 		HandleAllocator.FreeHandle(Handle);
 	}
@@ -267,9 +267,9 @@ namespace Chaos
 	template<typename T, int d>
 	void TPBDCollisionConstraints<T, d>::UpdateManifolds(T Dt)
 	{
-		PhysicsParallelFor(PlaneConstraints.Num(), [&](int32 ConstraintIndex)
+		PhysicsParallelFor(IterativeConstraints.Num(), [&](int32 ConstraintIndex)
 		{
-			FConstraintBase& ConstraintBase = PlaneConstraints[ConstraintIndex];
+			FConstraintBase& ConstraintBase = IterativeConstraints[ConstraintIndex];
 			if (ConstraintBase.GetType() == FCollisionConstraintBase::FType::MultiPoint)
 			{
 				Collisions::UpdateManifold<float, 3>(MThickness, ConstraintBase);
@@ -291,7 +291,7 @@ namespace Chaos
 				Collisions::Apply(Contact, MThickness, IterationParameters, ParticleParameters);
 			}
 
-			for (FPlaneContactConstraint& Contact : PlaneConstraints)
+			for (FIterativeContactConstraint& Contact : IterativeConstraints)
 			{
 				Collisions::TContactParticleParameters<T> ParticleParameters = { &MCollided, &MPhysicsMaterials, CollisionFrictionOverride, MAngularFriction };
 				Collisions::TContactIterationParameters<T> IterationParameters = { Dt, Iterations, NumIterations, MApplyPairIterations, nullptr };
@@ -321,7 +321,7 @@ namespace Chaos
 				Collisions::ApplyPushOut(Contact, MThickness, TempStatic, IterationParameters, ParticleParameters);
 			}
 
-			for (FPlaneContactConstraint& Contact : PlaneConstraints)
+			for (FIterativeContactConstraint& Contact : IterativeConstraints)
 			{
 				Collisions::TContactParticleParameters<T> ParticleParameters = { &MCollided, &MPhysicsMaterials, CollisionFrictionOverride, MAngularFriction };
 				Collisions::TContactIterationParameters<T> IterationParameters = { Dt, Iterations, NumIterations, MApplyPushOutPairIterations, &bNeedsAnotherIteration };
diff --git a/Engine/Source/Runtime/Experimental/Chaos/Public/Chaos/CollisionResolutionTypes.h b/Engine/Source/Runtime/Experimental/Chaos/Public/Chaos/CollisionResolutionTypes.h
--- a/Engine/Source/Runtime/Experimental/Chaos/Public/Chaos/CollisionResolutionTypes.h
+++ b/Engine/Source/Runtime/Experimental/Chaos/Public/Chaos/CollisionResolutionTypes.h
@@ -241,7 +241,7 @@ namespace Chaos
 				return ConstraintContainer->PointConstraints[ConstraintIndex];
 			}
 			check(GetType() == FConstraintBase::FType::MultiPoint);
-			return ConstraintContainer->PlaneConstraints[ConstraintIndex];
+			return ConstraintContainer->IterativeConstraints[ConstraintIndex];
 		}
 
 		TCollisionConstraintBase<T, d>& GetContact()
@@ -251,7 +251,7 @@ namespace Chaos
 				return ConstraintContainer->PointConstraints[ConstraintIndex];
 			}
 			check(GetType() == FConstraintBase::FType::MultiPoint);
-			return ConstraintContainer->PlaneConstraints[ConstraintIndex];
+			return ConstraintContainer->IterativeConstraints[ConstraintIndex];
 		}
 
 
@@ -259,8 +259,8 @@ namespace Chaos
 		const TRigidBodyPointContactConstraint<T, d>& GetPointContact() const { check(GetType() == FConstraintBase::FType::SinglePoint); return ConstraintContainer->PointConstraints[ConstraintIndex]; }
 		TRigidBodyPointContactConstraint<T, d>& GetPointContact() { check(GetType() == FConstraintBase::FType::SinglePoint); return ConstraintContainer->PointConstraints[ConstraintIndex]; }
 
-		const TRigidBodyIterativeContactConstraint<T, d>& GetPlaneContact() const { check(GetType() == FConstraintBase::FType::MultiPoint); return ConstraintContainer->PlaneConstraints[ConstraintIndex]; }
-		TRigidBodyIterativeContactConstraint<T, d>& GetPlaneContact() { check(GetType() == FConstraintBase::FType::MultiPoint); return ConstraintContainer->PlaneConstraints[ConstraintIndex]; }
+		const TRigidBodyIterativeContactConstraint<T, d>& GetIterativeContact() const { check(GetType() == FConstraintBase::FType::MultiPoint); return ConstraintContainer->IterativeConstraints[ConstraintIndex]; }
+		TRigidBodyIterativeContactConstraint<T, d>& GetIterativeContact() { check(GetType() == FConstraintBase::FType::MultiPoint); return ConstraintContainer->IterativeConstraints[ConstraintIndex]; }
 
 		typename FConstraintBase::FType GetType() const { return ConstraintType; }
 
diff --git a/Engine/Source/Runtime/Experimental/Chaos/Public/Chaos/PBDCollisionConstraints.h b/Engine/Source/Runtime/Experimental/Chaos/Public/Chaos/PBDCollisionConstraints.h
--- a/Engine/Source/Runtime/Experimental/Chaos/Public/Chaos/PBDCollisionConstraints.h
+++ b/Engine/Source/Runtime/Experimental/Chaos/Public/Chaos/PBDCollisionConstraints.h
@@ -57,7 +57,7 @@ public:
 	using FConstraintContainerHandle = TPBDCollisionConstraintHandle<T, d>;
 	using FConstraintBase = TCollisionConstraintBase<T, d>;
  	using FPointContactConstraint = TRigidBodyPointContactConstraint<T, d>;
-	using FPlaneContactConstraint = TRigidBodyIterativeContactConstraint<T, d>;
+	using FIterativeContactConstraint = TRigidBodyIterativeContactConstraint<T, d>;
 	using FConstraintHandleAllocator = TConstraintHandleAllocator<TPBDCollisionConstraints<T, d>>;
 	using FConstraintContainerHandleKey = typename TPBDCollisionConstraintHandle<T, d>::FHandleKey;
 
@@ -221,7 +221,7 @@ private:
 	const TPBDRigidsSOAs<T,d>& Particles;
 
 	TArray<FPointContactConstraint> PointConstraints;
-	TArray<FPlaneContactConstraint> PlaneConstraints;
+	TArray<FIterativeContactConstraint> IterativeConstraints;
 
 	TMap< FConstraintContainerHandleKey, FConstraintContainerHandle* > Manifolds;
 	TArray<FConstraintContainerHandle*> Handles;
