Crash fix, changes from CL10506091 didn't go far enough in fixing all locations where sleeping particles were being ignored from island assignment.

diff --git a/Engine/Source/Runtime/Experimental/Chaos/Private/Chaos/PBDConstraintGraph.cpp b/Engine/Source/Runtime/Experimental/Chaos/Private/Chaos/PBDConstraintGraph.cpp
--- a/Engine/Source/Runtime/Experimental/Chaos/Private/Chaos/PBDConstraintGraph.cpp
+++ b/Engine/Source/Runtime/Experimental/Chaos/Private/Chaos/PBDConstraintGraph.cpp
@@ -138,7 +138,7 @@ void FPBDConstraintGraph::InitializeGraph(const TParticleView<TGeometryParticles
 				    {
 						// This does NOT check if particle is actually dynamic. TODO confirm that this should happen to all nonDisabledParticles, not just dynamics.
 						TPBDRigidParticleHandle<FReal, 3>* PBDRigid = Particle->CastToRigidParticle();
-						if(PBDRigid && PBDRigid->ObjectState() == EObjectStateType::Dynamic)
+						if(PBDRigid && PBDRigid->ObjectState() != EObjectStateType::Kinematic)
 					    {
 						    PBDRigid->Island() = INDEX_NONE;
 					    }
@@ -362,7 +362,7 @@ void FPBDConstraintGraph::ComputeIslands(const TParticleView<TPBDRigidParticles<
 					}
 
 					TPBDRigidParticleHandle<FReal, 3>* PBDRigid = Particle->CastToRigidParticle();
-					if(PBDRigid && PBDRigid->ObjectState() == EObjectStateType::Dynamic)
+					if(PBDRigid && PBDRigid->ObjectState() != EObjectStateType::Kinematic)
 					{
 						if (!Particle->Sleeping() && bSleepState)
 						{
@@ -624,7 +624,7 @@ void FPBDConstraintGraph::WakeIsland(const int32 Island)
 	for (TGeometryParticleHandle<FReal, 3>* Particle : IslandToParticles[Island])
 	{
 		TPBDRigidParticleHandle<FReal, 3>* PBDRigid = Particle->CastToRigidParticle();
-		if(PBDRigid && PBDRigid->ObjectState() == EObjectStateType::Dynamic)
+		if(PBDRigid && PBDRigid->ObjectState() != EObjectStateType::Kinematic)
 		{
 			if (PBDRigid->Sleeping())
 			{
@@ -698,7 +698,7 @@ void FPBDConstraintGraph::EnableParticle(TGeometryParticleHandle<FReal, 3>* Part
 void FPBDConstraintGraph::DisableParticle(TGeometryParticleHandle<FReal, 3>* Particle)
 {
 	TPBDRigidParticleHandle<FReal, 3>* PBDRigid = Particle->CastToRigidParticle();
-	if(PBDRigid && PBDRigid->ObjectState() == EObjectStateType::Dynamic)
+	if(PBDRigid && PBDRigid->ObjectState() != EObjectStateType::Kinematic)
 	{
 		const int32 Island = PBDRigid->Island();
 		if (Island != INDEX_NONE)
