Chaos - Fix AABBTree UpdateElement losing track of payloads with a change in bHasBounds.

diff --git a/Engine/Source/Runtime/Experimental/Chaos/Public/Chaos/AABBTree.h b/Engine/Source/Runtime/Experimental/Chaos/Public/Chaos/AABBTree.h
--- a/Engine/Source/Runtime/Experimental/Chaos/Public/Chaos/AABBTree.h
+++ b/Engine/Source/Runtime/Experimental/Chaos/Public/Chaos/AABBTree.h
@@ -444,6 +444,19 @@ public:
 				DirtyElements[PayloadInfo->DirtyPayloadIdx].Bounds = NewBounds;
 			}
 			PayloadInfo->GlobalPayloadIdx = INDEX_NONE;
+
+			// Handle something that previously did not have bounds that may be in global elements.
+			if (PayloadInfo->GlobalPayloadIdx != INDEX_NONE)
+			{
+				if(PayloadInfo->GlobalPayloadIdx + 1 < GlobalPayloads.Num())
+				{
+					auto LastGlobalPayload = GlobalPayloads.Last().Payload;
+					PayloadToInfo.FindChecked(LastGlobalPayload).GlobalPayloadIdx = PayloadInfo->GlobalPayloadIdx;
+				}
+				GlobalPayloads.RemoveAtSwap(PayloadInfo->GlobalPayloadIdx);
+
+				PayloadInfo->GlobalPayloadIdx = INDEX_NONE;
+			}
 		}
 		else
 		{
@@ -456,7 +469,19 @@ public:
 			{
 				GlobalPayloads[PayloadInfo->GlobalPayloadIdx].Bounds = GlobalBounds;
 			}
-			PayloadInfo->DirtyPayloadIdx = INDEX_NONE;
+
+			// Handle something that previously had bounds that may be in dirty elements.
+			if (PayloadInfo->DirtyPayloadIdx != INDEX_NONE)
+			{
+				if(PayloadInfo->DirtyPayloadIdx + 1 < DirtyElements.Num())
+				{
+					auto LastDirtyPayload = DirtyElements.Last().Payload;
+					PayloadToInfo.FindChecked(LastDirtyPayload).DirtyPayloadIdx = PayloadInfo->DirtyPayloadIdx;
+				}
+				DirtyElements.RemoveAtSwap(PayloadInfo->DirtyPayloadIdx);
+
+				PayloadInfo->DirtyPayloadIdx = INDEX_NONE;
+			}
 		}
 	}
 
