Chaos - Fix AABBTree from clearing GlobalPayloadIdx before checking it.

diff --git a/Engine/Source/Runtime/Experimental/Chaos/Public/Chaos/AABBTree.h b/Engine/Source/Runtime/Experimental/Chaos/Public/Chaos/AABBTree.h
--- a/Engine/Source/Runtime/Experimental/Chaos/Public/Chaos/AABBTree.h
+++ b/Engine/Source/Runtime/Experimental/Chaos/Public/Chaos/AABBTree.h
@@ -443,7 +443,6 @@ public:
 			{
 				DirtyElements[PayloadInfo->DirtyPayloadIdx].Bounds = NewBounds;
 			}
-			PayloadInfo->GlobalPayloadIdx = INDEX_NONE;
 
 			// Handle something that previously did not have bounds that may be in global elements.
 			if (PayloadInfo->GlobalPayloadIdx != INDEX_NONE)
