Chaos : Collisions

- Pulled the application of Contact out of Apply and ApplyPushOut into separate functions to allow for iterating over multiple contacts on a single constraint.

diff --git a/Engine/Source/Runtime/Experimental/Chaos/Private/Chaos/CollisionResolution.cpp b/Engine/Source/Runtime/Experimental/Chaos/Private/Chaos/CollisionResolution.cpp
--- a/Engine/Source/Runtime/Experimental/Chaos/Private/Chaos/CollisionResolution.cpp
+++ b/Engine/Source/Runtime/Experimental/Chaos/Private/Chaos/CollisionResolution.cpp
@@ -547,13 +547,12 @@ namespace Chaos
 			});
 		}
 
-		template<class TYPE>
-		TYPE SumArray(TArray<TYPE>& Array)
+		TVector<float, 3> SumSampleData(TArray<FRigidBodyIterativeContactConstraint::FSampleData>& Array)
 		{
-			TYPE ReturnValue(0);
+			TVector<float,3> ReturnValue(0);
 			for (int i = 0; i < Array.Num(); i++)
 			{
-				ReturnValue += Array[i];
+				ReturnValue += Array[i].X;
 			}
 			return ReturnValue;
 		}
@@ -562,18 +561,18 @@ namespace Chaos
 		void UpdateConvexConvexConstraint(const FImplicitObject& Implicit0, const TRigidTransform<T, d>& Transform0, const FImplicitObject& Implicit1, const TRigidTransform<T, d>& Transform1, const T Thickness, TRigidBodyIterativeContactConstraint<T,d>& Constraint)
 		{
 			TContactPoint<T> ContactPoint = ConvexConvexContactPoint(Implicit0, Transform0, Implicit1, Transform1, Thickness);
-
 			/*
 			TContactPoint<T> ContactPoint;
 
 			const TRigidTransform<T, d> AToBTM = Transform0.GetRelativeTransform(Transform1);
-			ContactPoint.Location = AToBTM.TransformPosition(Constraint.LocalPosition);
 
-			ContactPoint.Phi = Implicit1.PhiWithNormal(ContactPoint.Location, ContactPoint.Normal);
-			ContactPoint.Normal = Constraint.WorldNormal;
-			ContactPoint.Location = Transform0.TransformPosition(Constraint.LocalPosition);
+			for (int32 Idx = 0; Idx < Constraint.LocalSamples.Num(); Idx++)
+			{
+				Constraint.LocalSamples[Idx].Manifold.Phi = Implicit1.PhiWithNormal(Constraint.LocalSamples[Idx].X, Constraint.LocalSamples[Idx].Manifold.Normal);
+				Constraint.LocalSamples[Idx].Manifold.Normal = Transform1.TransformVector(Constraint.PlaneNormal);
+				Constraint.LocalSamples[Idx].Manifold.Location = Transform0.TransformPosition(Constraint.LocalSamples[Idx].X);
+			}
 			*/
-
 			UpdateContactPoint(Constraint.Manifold, ContactPoint);
 		}
 
@@ -585,18 +584,41 @@ namespace Chaos
 
 			FRigidBodyIterativeContactConstraint::FManifold& Manifold = Constraint.Manifold;
 
-			if (!ContactPoint.Normal.Equals(Constraint.WorldNormal) || !Constraint.LocalSamples.Num())
+			if (!ContactPoint.Normal.Equals(Constraint.PlaneNormal) || !Constraint.LocalSamples.Num())
 			{
 				Constraint.LocalSamples.Reset();
-				Constraint.WorldNormal = ContactPoint.Normal;
+				Constraint.PlaneNormal = Transform1.InverseTransformVector(ContactPoint.Normal);
+				Constraint.PlanePosition = Transform1.InverseTransformPosition(ContactPoint.Location - ContactPoint.Phi*ContactPoint.Normal);
 			}
 
+			TVector<T, d> SurfaceSample = Transform0.InverseTransformPosition(ContactPoint.Location);
 			if (Constraint.LocalSamples.Num() < 4)
 			{
-				TVector<T, d> SurfaceSample = Transform0.InverseTransformPosition(ContactPoint.Location);
-				Constraint.LocalSamples.Add(SurfaceSample);
-				Constraint.LocalPosition = SumArray(Constraint.LocalSamples) / Constraint.LocalSamples.Num();
+				Constraint.LocalSamples.Add({ SurfaceSample,0.f });
 			}
+			else 
+			{
+				TVector<T,d> Center = SumSampleData(Constraint.LocalSamples) / 5.f;
+				T Delta = (Center - SurfaceSample).SizeSquared();
+				Constraint.LocalPosition = SumSampleData(Constraint.LocalSamples) / Constraint.LocalSamples.Num();
+
+				// todo(chaos) : maximize area instead
+				T SmallestDelta = FLT_MAX;
+				int32 SmallestIndex = 0;
+				for (int32 idx=0;idx<Constraint.LocalSamples.Num();idx++)
+					if (Constraint.LocalSamples[idx].Delta < SmallestDelta) {
+						SmallestDelta = Constraint.LocalSamples[idx].Delta;
+						SmallestIndex = idx;
+					}
+
+				if (Delta > SmallestDelta) {
+					Constraint.LocalSamples[SmallestIndex] = { SurfaceSample,Delta };
+				}
+			}
+
+			typedef FRigidBodyIterativeContactConstraint::FSampleData FSampleData;
+			Constraint.LocalPosition = SumSampleData(Constraint.LocalSamples) / Constraint.LocalSamples.Num();
+			for (FSampleData& Data : Constraint.LocalSamples) Data.Delta = (Constraint.LocalPosition - Data.X).SizeSquared();
 			*/
 		}
 
diff --git a/Engine/Source/Runtime/Experimental/Chaos/Private/Chaos/CollisionResolutionUtil.cpp b/Engine/Source/Runtime/Experimental/Chaos/Private/Chaos/CollisionResolutionUtil.cpp
--- a/Engine/Source/Runtime/Experimental/Chaos/Private/Chaos/CollisionResolutionUtil.cpp
+++ b/Engine/Source/Runtime/Experimental/Chaos/Private/Chaos/CollisionResolutionUtil.cpp
@@ -45,10 +45,8 @@ namespace Chaos
 		}
 
 		template<typename T, int d>
-		TVector<T, d> GetEnergyClampedImpulse(const TCollisionConstraintBase<T, d>& Constraint, const TVector<T, d>& Impulse, const TVector<T, d>& VectorToPoint1, const TVector<T, d>& VectorToPoint2, const TVector<T, d>& Velocity1, const TVector<T, d>& Velocity2)
+		TVector<T, d> GetEnergyClampedImpulse(const TPBDRigidParticleHandle<T, d>* PBDRigid0, const TPBDRigidParticleHandle<T, d>* PBDRigid1, const TVector<T, d>& Impulse, const TVector<T, d>& VectorToPoint1, const TVector<T, d>& VectorToPoint2, const TVector<T, d>& Velocity1, const TVector<T, d>& Velocity2)
 		{
-			TPBDRigidParticleHandle<T, d>* PBDRigid0 = Constraint.Particle[0]->CastToRigidParticle();
-			TPBDRigidParticleHandle<T, d>* PBDRigid1 = Constraint.Particle[1]->CastToRigidParticle();
 			const bool bIsRigidDynamic0 = PBDRigid0 && PBDRigid0->ObjectState() == EObjectStateType::Dynamic;
 			const bool bIsRigidDynamic1 = PBDRigid1 && PBDRigid1->ObjectState() == EObjectStateType::Dynamic;
 
@@ -598,7 +596,7 @@ namespace Chaos
 
 		template TRigidTransform<float, 3> GetTransform(const TGeometryParticleHandle<float, 3>* Particle);
 		template PMatrix<float, 3, 3> ComputeFactorMatrix3(const TVector<float, 3>& V, const PMatrix<float, 3, 3>& M, const float& Im);
-		template TVector<float, 3> GetEnergyClampedImpulse(const TCollisionConstraintBase<float, 3>& Constraint, const TVector<float, 3>& Impulse, const TVector<float, 3>& VectorToPoint1, const TVector<float, 3>& VectorToPoint2, const TVector<float, 3>& Velocity1, const TVector<float, 3>& Velocity2);
+		template TVector<float, 3> GetEnergyClampedImpulse(const TPBDRigidParticleHandle<float, 3>* PBDRigid0, const TPBDRigidParticleHandle<float, 3>* PBDRigid1, const TVector<float, 3>& Impulse, const TVector<float, 3>& VectorToPoint1, const TVector<float, 3>& VectorToPoint2, const TVector<float, 3>& Velocity1, const TVector<float, 3>& Velocity2);
 		template bool SampleObjectHelper<float, 3>(const FImplicitObject& Object, const TRigidTransform<float, 3>& ObjectTransform, const TRigidTransform<float, 3>& SampleToObjectTransform, const TVector<float, 3>& SampleParticle, float Thickness, TRigidBodyPointContactConstraint<float, 3>& Constraint);
 		template bool SampleObjectNoNormal(const FImplicitObject& Object, const TRigidTransform<float, 3>& ObjectTransform, const TRigidTransform<float, 3>& SampleToObjectTransform, const TVector<float, 3>& SampleParticle, float Thickness, TRigidBodyPointContactConstraint<float, 3>& Constraint);
 		template bool SampleObjectNormalAverageHelper(const FImplicitObject& Object, const TRigidTransform<float, 3>& ObjectTransform, const TRigidTransform<float, 3>& SampleToObjectTransform, const TVector<float, 3>& SampleParticle, float Thickness, float& TotalThickness, TRigidBodyPointContactConstraint<float, 3>& Constraint);
diff --git a/Engine/Source/Runtime/Experimental/Chaos/Private/Chaos/PBDCollisionConstraints.cpp b/Engine/Source/Runtime/Experimental/Chaos/Private/Chaos/PBDCollisionConstraints.cpp
--- a/Engine/Source/Runtime/Experimental/Chaos/Private/Chaos/PBDCollisionConstraints.cpp
+++ b/Engine/Source/Runtime/Experimental/Chaos/Private/Chaos/PBDCollisionConstraints.cpp
@@ -286,16 +286,16 @@ namespace Chaos
 		{
 			for (FPointContactConstraint& Contact : PointConstraints)
 			{
-				Collisions::TContactParticleParameters<T> ParticleParameters = { &MCollided, &MPhysicsMaterials, CollisionFrictionOverride, MAngularFriction };
+				Collisions::TContactParticleParameters<T> ParticleParameters = { MThickness, &MCollided, &MPhysicsMaterials, CollisionFrictionOverride, MAngularFriction };
 				Collisions::TContactIterationParameters<T> IterationParameters = { Dt, Iterations, NumIterations, MApplyPairIterations, nullptr };
-				Collisions::Apply(Contact, MThickness, IterationParameters, ParticleParameters);
+				Collisions::Apply(Contact, IterationParameters, ParticleParameters);
 			}
 
 			for (FIterativeContactConstraint& Contact : IterativeConstraints)
 			{
-				Collisions::TContactParticleParameters<T> ParticleParameters = { &MCollided, &MPhysicsMaterials, CollisionFrictionOverride, MAngularFriction };
+				Collisions::TContactParticleParameters<T> ParticleParameters = { MThickness, &MCollided, &MPhysicsMaterials, CollisionFrictionOverride, MAngularFriction };
 				Collisions::TContactIterationParameters<T> IterationParameters = { Dt, Iterations, NumIterations, MApplyPairIterations, nullptr };
-				Collisions::Apply(Contact, MThickness, IterationParameters, ParticleParameters);
+				Collisions::Apply(Contact, IterationParameters, ParticleParameters);
 			}
 		}
 
@@ -316,16 +316,16 @@ namespace Chaos
 		{
 			for (FPointContactConstraint& Contact : PointConstraints)
 			{
-				Collisions::TContactParticleParameters<T> ParticleParameters = { &MCollided, &MPhysicsMaterials, CollisionFrictionOverride, MAngularFriction };
+				Collisions::TContactParticleParameters<T> ParticleParameters = { MThickness, &MCollided, &MPhysicsMaterials, CollisionFrictionOverride, MAngularFriction };
 				Collisions::TContactIterationParameters<T> IterationParameters = { Dt, Iterations, NumIterations, MApplyPushOutPairIterations, &bNeedsAnotherIteration };
-				Collisions::ApplyPushOut(Contact, MThickness, TempStatic, IterationParameters, ParticleParameters);
+				Collisions::ApplyPushOut(Contact, TempStatic, IterationParameters, ParticleParameters);
 			}
 
 			for (FIterativeContactConstraint& Contact : IterativeConstraints)
 			{
-				Collisions::TContactParticleParameters<T> ParticleParameters = { &MCollided, &MPhysicsMaterials, CollisionFrictionOverride, MAngularFriction };
+				Collisions::TContactParticleParameters<T> ParticleParameters = { MThickness, &MCollided, &MPhysicsMaterials, CollisionFrictionOverride, MAngularFriction };
 				Collisions::TContactIterationParameters<T> IterationParameters = { Dt, Iterations, NumIterations, MApplyPushOutPairIterations, &bNeedsAnotherIteration };
-				Collisions::ApplyPushOut(Contact, MThickness, TempStatic, IterationParameters, ParticleParameters);
+				Collisions::ApplyPushOut(Contact, TempStatic, IterationParameters, ParticleParameters);
 			}
 		}
 
@@ -348,9 +348,9 @@ namespace Chaos
 				FConstraintContainerHandle* ConstraintHandle = InConstraintHandles[ConstraintHandleIndex];
 				check(ConstraintHandle != nullptr);
 
-				Collisions::TContactParticleParameters<T> ParticleParameters = { &MCollided, &MPhysicsMaterials, CollisionFrictionOverride, MAngularFriction };
+				Collisions::TContactParticleParameters<T> ParticleParameters = { MThickness, &MCollided, &MPhysicsMaterials, CollisionFrictionOverride, MAngularFriction };
 				Collisions::TContactIterationParameters<T> IterationParameters = { Dt, Iterations, NumIterations, MApplyPairIterations, nullptr };
-				Collisions::Apply(ConstraintHandle->GetContact(), MThickness, IterationParameters, ParticleParameters);
+				Collisions::Apply(ConstraintHandle->GetContact(), IterationParameters, ParticleParameters);
 
 			}, bDisableCollisionParallelFor);
 		}
@@ -375,9 +375,9 @@ namespace Chaos
 				FConstraintContainerHandle* ConstraintHandle = InConstraintHandles[ConstraintHandleIndex];
 				check(ConstraintHandle != nullptr);
 
-				Collisions::TContactParticleParameters<T> ParticleParameters = { &MCollided, &MPhysicsMaterials, CollisionFrictionOverride, MAngularFriction };
+				Collisions::TContactParticleParameters<T> ParticleParameters = { MThickness, &MCollided, &MPhysicsMaterials, CollisionFrictionOverride, MAngularFriction };
 				Collisions::TContactIterationParameters<T> IterationParameters = { Dt, Iteration, NumIterations, MApplyPushOutPairIterations, &bNeedsAnotherIteration };
-				Collisions::ApplyPushOut(ConstraintHandle->GetContact(), MThickness, IsTemporarilyStatic, IterationParameters, ParticleParameters);
+				Collisions::ApplyPushOut(ConstraintHandle->GetContact(), IsTemporarilyStatic, IterationParameters, ParticleParameters);
 
 			}, bDisableCollisionParallelFor);
 		}
diff --git a/Engine/Source/Runtime/Experimental/Chaos/Private/Chaos/PBDCollisionConstraintsContact.cpp b/Engine/Source/Runtime/Experimental/Chaos/Private/Chaos/PBDCollisionConstraintsContact.cpp
--- a/Engine/Source/Runtime/Experimental/Chaos/Private/Chaos/PBDCollisionConstraintsContact.cpp
+++ b/Engine/Source/Runtime/Experimental/Chaos/Private/Chaos/PBDCollisionConstraintsContact.cpp
@@ -65,230 +65,251 @@ namespace Chaos
 
 
 		template<typename T, int d>
-		void Apply(TCollisionConstraintBase<T, d>& Constraint, T Thickness, TContactIterationParameters<T> & IterationParameters, TContactParticleParameters<T> & ParticleParameters)
+		TVector<T,d> ApplyContact(TCollisionContact<T, d>& Contact,
+			TGenericParticleHandle<T, d> Particle0, 
+			TGenericParticleHandle<T, d> Particle1,
+			TContactIterationParameters<T> & IterationParameters, TContactParticleParameters<T> & ParticleParameters)
 		{
-			TGenericParticleHandle<T, d> Particle0 = TGenericParticleHandle<T, d>(Constraint.Particle[0]);
-			TGenericParticleHandle<T, d> Particle1 = TGenericParticleHandle<T, d>(Constraint.Particle[1]);
-			TPBDRigidParticleHandle<T, d>* PBDRigid0 = Particle0->CastToRigidParticle();
-			TPBDRigidParticleHandle<T, d>* PBDRigid1 = Particle1->CastToRigidParticle();
-			bool bIsRigidDynamic0 = PBDRigid0 && PBDRigid0->ObjectState() == EObjectStateType::Dynamic;
-			bool bIsRigidDynamic1 = PBDRigid1 && PBDRigid1->ObjectState() == EObjectStateType::Dynamic;
+			TVector<T, d> AccumulatedImpulse(0);
 
-			for (int32 PairIt = 0; PairIt < IterationParameters.NumPairIterations; ++PairIt)
+			TSerializablePtr<FChaosPhysicsMaterial> PhysicsMaterial0, PhysicsMaterial1;
+			if (ParticleParameters.PhysicsMaterials)
 			{
-				Collisions::Update<ECollisionUpdateType::Deepest>(Thickness, Constraint);
-				if (Constraint.GetPhi() >= Thickness)
-				{
-					return;
-				}
+				PhysicsMaterial0 = Particle0->AuxilaryValue(*ParticleParameters.PhysicsMaterials);
+				PhysicsMaterial1 = Particle1->AuxilaryValue(*ParticleParameters.PhysicsMaterials);
+			}
 
-				// @todo(ccaulfield): CHAOS_PARTICLEHANDLE_TODO what's the best way to manage external per-particle data?
-				if (ParticleParameters.Collided)
-				{
-					Particle0->AuxilaryValue(*ParticleParameters.Collided) = true;
-					Particle1->AuxilaryValue(*ParticleParameters.Collided) = true;
-				}
 
-				TSerializablePtr<FChaosPhysicsMaterial> PhysicsMaterial0, PhysicsMaterial1;
-				if (ParticleParameters.PhysicsMaterials)
-				{
-					PhysicsMaterial0 = Particle0->AuxilaryValue(*ParticleParameters.PhysicsMaterials);
-					PhysicsMaterial1 = Particle1->AuxilaryValue(*ParticleParameters.PhysicsMaterials);
-				}
+			TPBDRigidParticleHandle<T, d>* PBDRigid0 = Particle0->CastToRigidParticle();
+			TPBDRigidParticleHandle<T, d>* PBDRigid1 = Particle1->CastToRigidParticle();
 
-				const TVector<T, d> ZeroVector = TVector<T, d>(0);
-				TVector<T, d> P0 = FParticleUtilities::GetCoMWorldPosition(Particle0);
-				TVector<T, d> P1 = FParticleUtilities::GetCoMWorldPosition(Particle1);
-				TRotation<T, d> Q0 = FParticleUtilities::GetCoMWorldRotation(Particle0);
-				TRotation<T, d> Q1 = FParticleUtilities::GetCoMWorldRotation(Particle1);
+			bool bIsRigidDynamic0 = PBDRigid0 && PBDRigid0->ObjectState() == EObjectStateType::Dynamic;
+			bool bIsRigidDynamic1 = PBDRigid1 && PBDRigid1->ObjectState() == EObjectStateType::Dynamic;
 
-				typename TRigidBodyPointContactConstraint<T, d>::FManifold & Contact = Constraint.Manifold;
+			const TVector<T, d> ZeroVector = TVector<T, d>(0);
+			TVector<T, d> P0 = FParticleUtilities::GetCoMWorldPosition(Particle0);
+			TVector<T, d> P1 = FParticleUtilities::GetCoMWorldPosition(Particle1);
+			TRotation<T, d> Q0 = FParticleUtilities::GetCoMWorldRotation(Particle0);
+			TRotation<T, d> Q1 = FParticleUtilities::GetCoMWorldRotation(Particle1);
 
-				TVector<T, d> VectorToPoint1 = Contact.Location - P0;
-				TVector<T, d> VectorToPoint2 = Contact.Location - P1;
-				TVector<T, d> Body1Velocity = FParticleUtilities::GetVelocityAtCoMRelativePosition(Particle0, VectorToPoint1);
-				TVector<T, d> Body2Velocity = FParticleUtilities::GetVelocityAtCoMRelativePosition(Particle1, VectorToPoint2);
-				TVector<T, d> RelativeVelocity = Body1Velocity - Body2Velocity;
-				T RelativeNormalVelocity = TVector<T, d>::DotProduct(RelativeVelocity, Contact.Normal);
+			TVector<T, d> VectorToPoint1 = Contact.Location - P0;
+			TVector<T, d> VectorToPoint2 = Contact.Location - P1;
+			TVector<T, d> Body1Velocity = FParticleUtilities::GetVelocityAtCoMRelativePosition(Particle0, VectorToPoint1);
+			TVector<T, d> Body2Velocity = FParticleUtilities::GetVelocityAtCoMRelativePosition(Particle1, VectorToPoint2);
+			TVector<T, d> RelativeVelocity = Body1Velocity - Body2Velocity;
+			T RelativeNormalVelocity = TVector<T, d>::DotProduct(RelativeVelocity, Contact.Normal);
 
-				if (RelativeNormalVelocity < 0) // ignore separating constraints
+			if (RelativeNormalVelocity < 0) // ignore separating constraints
+			{
+				PMatrix<T, d, d> WorldSpaceInvI1 = bIsRigidDynamic0 ? Utilities::ComputeWorldSpaceInertia(Q0, PBDRigid0->InvI()) : PMatrix<T, d, d>(0);
+				PMatrix<T, d, d> WorldSpaceInvI2 = bIsRigidDynamic1 ? Utilities::ComputeWorldSpaceInertia(Q1, PBDRigid1->InvI()) : PMatrix<T, d, d>(0);
+				PMatrix<T, d, d> Factor =
+					(bIsRigidDynamic0 ? ComputeFactorMatrix3(VectorToPoint1, WorldSpaceInvI1, PBDRigid0->InvM()) : PMatrix<T, d, d>(0)) +
+					(bIsRigidDynamic1 ? ComputeFactorMatrix3(VectorToPoint2, WorldSpaceInvI2, PBDRigid1->InvM()) : PMatrix<T, d, d>(0));
+				TVector<T, d> Impulse;
+				TVector<T, d> AngularImpulse(0);
+
+				// Resting contact if very close to the surface
+				T Restitution = (T)0;
+				T Friction = (T)0;
+				T AngularFriction = (T)0;
+				bool bApplyRestitution = (RelativeVelocity.Size() > (2 * 980 * IterationParameters.Dt));
+				if (PhysicsMaterial0 && PhysicsMaterial1)
 				{
-					PMatrix<T, d, d> WorldSpaceInvI1 = bIsRigidDynamic0 ? Utilities::ComputeWorldSpaceInertia(Q0, PBDRigid0->InvI()) : PMatrix<T, d, d>(0);
-					PMatrix<T, d, d> WorldSpaceInvI2 = bIsRigidDynamic1 ? Utilities::ComputeWorldSpaceInertia(Q1, PBDRigid1->InvI()) : PMatrix<T, d, d>(0);
-					PMatrix<T, d, d> Factor =
-						(bIsRigidDynamic0 ? ComputeFactorMatrix3(VectorToPoint1, WorldSpaceInvI1, PBDRigid0->InvM()) : PMatrix<T, d, d>(0)) +
-						(bIsRigidDynamic1 ? ComputeFactorMatrix3(VectorToPoint2, WorldSpaceInvI2, PBDRigid1->InvM()) : PMatrix<T, d, d>(0));
-					TVector<T, d> Impulse;
-					TVector<T, d> AngularImpulse(0);
-
-					// Resting contact if very close to the surface
-					T Restitution = (T)0;
-					T Friction = (T)0;
-					T AngularFriction = (T)0;
-					bool bApplyRestitution = (RelativeVelocity.Size() > (2 * 980 * IterationParameters.Dt));
-					if (PhysicsMaterial0 && PhysicsMaterial1)
+					if (bApplyRestitution)
 					{
-						if (bApplyRestitution)
-						{
-							Restitution = FMath::Min(PhysicsMaterial0->Restitution, PhysicsMaterial1->Restitution);
-						}
-						Friction = FMath::Max(PhysicsMaterial0->Friction, PhysicsMaterial1->Friction);
+						Restitution = FMath::Min(PhysicsMaterial0->Restitution, PhysicsMaterial1->Restitution);
 					}
-					else if (PhysicsMaterial0)
+					Friction = FMath::Max(PhysicsMaterial0->Friction, PhysicsMaterial1->Friction);
+				}
+				else if (PhysicsMaterial0)
+				{
+					if (bApplyRestitution)
 					{
-						if (bApplyRestitution)
-						{
-							Restitution = PhysicsMaterial0->Restitution;
-						}
-						Friction = PhysicsMaterial0->Friction;
+						Restitution = PhysicsMaterial0->Restitution;
 					}
-					else if (PhysicsMaterial1)
+					Friction = PhysicsMaterial0->Friction;
+				}
+				else if (PhysicsMaterial1)
+				{
+					if (bApplyRestitution)
 					{
-						if (bApplyRestitution)
-						{
-							Restitution = PhysicsMaterial1->Restitution;
-						}
-						Friction = PhysicsMaterial1->Friction;
+						Restitution = PhysicsMaterial1->Restitution;
 					}
+					Friction = PhysicsMaterial1->Friction;
+				}
 
-					if (ParticleParameters.FrictionOverride >= 0)
-					{
-						Friction = ParticleParameters.FrictionOverride;
-					}
-					if (ParticleParameters.AngularFrictionOverride >= 0)
+				if (ParticleParameters.FrictionOverride >= 0)
+				{
+					Friction = ParticleParameters.FrictionOverride;
+				}
+				if (ParticleParameters.AngularFrictionOverride >= 0)
+				{
+					AngularFriction = ParticleParameters.AngularFrictionOverride;
+				}
+
+				if (Friction)
+				{
+					if (RelativeNormalVelocity > 0)
 					{
-						AngularFriction = ParticleParameters.AngularFrictionOverride;
+						RelativeNormalVelocity = 0;
 					}
-
-					if (Friction)
+					TVector<T, d> VelocityChange = -(Restitution * RelativeNormalVelocity * Contact.Normal + RelativeVelocity);
+					T NormalVelocityChange = TVector<T, d>::DotProduct(VelocityChange, Contact.Normal);
+					PMatrix<T, d, d> FactorInverse = Factor.Inverse();
+					TVector<T, d> MinimalImpulse = FactorInverse * VelocityChange;
+					const T MinimalImpulseDotNormal = TVector<T, d>::DotProduct(MinimalImpulse, Contact.Normal);
+					const T TangentialSize = (MinimalImpulse - MinimalImpulseDotNormal * Contact.Normal).Size();
+					if (TangentialSize <= Friction * MinimalImpulseDotNormal)
 					{
-						if (RelativeNormalVelocity > 0)
-						{
-							RelativeNormalVelocity = 0;
-						}
-						TVector<T, d> VelocityChange = -(Restitution * RelativeNormalVelocity * Contact.Normal + RelativeVelocity);
-						T NormalVelocityChange = TVector<T, d>::DotProduct(VelocityChange, Contact.Normal);
-						PMatrix<T, d, d> FactorInverse = Factor.Inverse();
-						TVector<T, d> MinimalImpulse = FactorInverse * VelocityChange;
-						const T MinimalImpulseDotNormal = TVector<T, d>::DotProduct(MinimalImpulse, Contact.Normal);
-						const T TangentialSize = (MinimalImpulse - MinimalImpulseDotNormal * Contact.Normal).Size();
-						if (TangentialSize <= Friction * MinimalImpulseDotNormal)
+						//within friction cone so just solve for static friction stopping the object
+						Impulse = MinimalImpulse;
+						if (AngularFriction)
 						{
-							//within friction cone so just solve for static friction stopping the object
-							Impulse = MinimalImpulse;
-							if (AngularFriction)
+							TVector<T, d> RelativeAngularVelocity = Particle0->W() - Particle1->W();
+							T AngularNormal = TVector<T, d>::DotProduct(RelativeAngularVelocity, Contact.Normal);
+							TVector<T, d> AngularTangent = RelativeAngularVelocity - AngularNormal * Contact.Normal;
+							TVector<T, d> FinalAngularVelocity = FMath::Sign(AngularNormal) * FMath::Max((T)0, FMath::Abs(AngularNormal) - AngularFriction * NormalVelocityChange) * Contact.Normal + FMath::Max((T)0, AngularTangent.Size() - AngularFriction * NormalVelocityChange) * AngularTangent.GetSafeNormal();
+							TVector<T, d> Delta = FinalAngularVelocity - RelativeAngularVelocity;
+							if (!bIsRigidDynamic0 && bIsRigidDynamic1)
 							{
-								TVector<T, d> RelativeAngularVelocity = Particle0->W() - Particle1->W();
-								T AngularNormal = TVector<T, d>::DotProduct(RelativeAngularVelocity, Contact.Normal);
-								TVector<T, d> AngularTangent = RelativeAngularVelocity - AngularNormal * Contact.Normal;
-								TVector<T, d> FinalAngularVelocity = FMath::Sign(AngularNormal) * FMath::Max((T)0, FMath::Abs(AngularNormal) - AngularFriction * NormalVelocityChange) * Contact.Normal + FMath::Max((T)0, AngularTangent.Size() - AngularFriction * NormalVelocityChange) * AngularTangent.GetSafeNormal();
-								TVector<T, d> Delta = FinalAngularVelocity - RelativeAngularVelocity;
-								if (!bIsRigidDynamic0 && bIsRigidDynamic1)
-								{
-									PMatrix<T, d, d> WorldSpaceI2 = (Q1 * FMatrix::Identity) * PBDRigid1->I() * (Q1 * FMatrix::Identity).GetTransposed();
-									TVector<T, d> ImpulseDelta = PBDRigid1->M() * TVector<T, d>::CrossProduct(VectorToPoint2, Delta);
-									Impulse += ImpulseDelta;
-									AngularImpulse += WorldSpaceI2 * Delta - TVector<T, d>::CrossProduct(VectorToPoint2, ImpulseDelta);
-								}
-								else if (bIsRigidDynamic0 && !bIsRigidDynamic1)
-								{
-									PMatrix<T, d, d> WorldSpaceI1 = (Q0 * FMatrix::Identity) * PBDRigid0->I() * (Q0 * FMatrix::Identity).GetTransposed();
-									TVector<T, d> ImpulseDelta = PBDRigid0->M() * TVector<T, d>::CrossProduct(VectorToPoint1, Delta);
-									Impulse += ImpulseDelta;
-									AngularImpulse += WorldSpaceI1 * Delta - TVector<T, d>::CrossProduct(VectorToPoint1, ImpulseDelta);
-								}
-								else if (bIsRigidDynamic0 && bIsRigidDynamic1)
-								{
-									PMatrix<T, d, d> Cross1(0, VectorToPoint1.Z, -VectorToPoint1.Y, -VectorToPoint1.Z, 0, VectorToPoint1.X, VectorToPoint1.Y, -VectorToPoint1.X, 0);
-									PMatrix<T, d, d> Cross2(0, VectorToPoint2.Z, -VectorToPoint2.Y, -VectorToPoint2.Z, 0, VectorToPoint2.X, VectorToPoint2.Y, -VectorToPoint2.X, 0);
-									PMatrix<T, d, d> CrossI1 = Cross1 * WorldSpaceInvI1;
-									PMatrix<T, d, d> CrossI2 = Cross2 * WorldSpaceInvI2;
-									PMatrix<T, d, d> Diag1 = CrossI1 * Cross1.GetTransposed() + CrossI2 * Cross2.GetTransposed();
-									Diag1.M[0][0] += PBDRigid0->InvM() + PBDRigid1->InvM();
-									Diag1.M[1][1] += PBDRigid0->InvM() + PBDRigid1->InvM();
-									Diag1.M[2][2] += PBDRigid0->InvM() + PBDRigid1->InvM();
-									PMatrix<T, d, d> OffDiag1 = (CrossI1 + CrossI2) * -1;
-									PMatrix<T, d, d> Diag2 = (WorldSpaceInvI1 + WorldSpaceInvI2).Inverse();
-									PMatrix<T, d, d> OffDiag1Diag2 = OffDiag1 * Diag2;
-									TVector<T, d> ImpulseDelta = PMatrix<T, d, d>((Diag1 - OffDiag1Diag2 * OffDiag1.GetTransposed()).Inverse())* ((OffDiag1Diag2 * -1) * Delta);
-									Impulse += ImpulseDelta;
-									AngularImpulse += Diag2 * (Delta - PMatrix<T, d, d>(OffDiag1.GetTransposed()) * ImpulseDelta);
-								}
+								PMatrix<T, d, d> WorldSpaceI2 = (Q1 * FMatrix::Identity) * PBDRigid1->I() * (Q1 * FMatrix::Identity).GetTransposed();
+								TVector<T, d> ImpulseDelta = PBDRigid1->M() * TVector<T, d>::CrossProduct(VectorToPoint2, Delta);
+								Impulse += ImpulseDelta;
+								AngularImpulse += WorldSpaceI2 * Delta - TVector<T, d>::CrossProduct(VectorToPoint2, ImpulseDelta);
 							}
-						}
-						else
-						{
-							//outside friction cone, solve for normal relative velocity and keep tangent at cone edge
-							TVector<T, d> Tangent = (RelativeVelocity - TVector<T, d>::DotProduct(RelativeVelocity, Contact.Normal) * Contact.Normal).GetSafeNormal();
-							TVector<T, d> DirectionalFactor = Factor * (Contact.Normal - Friction * Tangent);
-							T ImpulseDenominator = TVector<T, d>::DotProduct(Contact.Normal, DirectionalFactor);
-							if (!ensureMsgf(FMath::Abs(ImpulseDenominator) > SMALL_NUMBER, TEXT("Constraint:%s\n\nParticle:%s\n\nLevelset:%s\n\nDirectionalFactor:%s, ImpulseDenominator:%f"),
-								*Constraint.ToString(),
-								*Particle0->ToString(),
-								*Particle1->ToString(),
-								*DirectionalFactor.ToString(), ImpulseDenominator))
+							else if (bIsRigidDynamic0 && !bIsRigidDynamic1)
 							{
-								ImpulseDenominator = (T)1;
+								PMatrix<T, d, d> WorldSpaceI1 = (Q0 * FMatrix::Identity) * PBDRigid0->I() * (Q0 * FMatrix::Identity).GetTransposed();
+								TVector<T, d> ImpulseDelta = PBDRigid0->M() * TVector<T, d>::CrossProduct(VectorToPoint1, Delta);
+								Impulse += ImpulseDelta;
+								AngularImpulse += WorldSpaceI1 * Delta - TVector<T, d>::CrossProduct(VectorToPoint1, ImpulseDelta);
+							}
+							else if (bIsRigidDynamic0 && bIsRigidDynamic1)
+							{
+								PMatrix<T, d, d> Cross1(0, VectorToPoint1.Z, -VectorToPoint1.Y, -VectorToPoint1.Z, 0, VectorToPoint1.X, VectorToPoint1.Y, -VectorToPoint1.X, 0);
+								PMatrix<T, d, d> Cross2(0, VectorToPoint2.Z, -VectorToPoint2.Y, -VectorToPoint2.Z, 0, VectorToPoint2.X, VectorToPoint2.Y, -VectorToPoint2.X, 0);
+								PMatrix<T, d, d> CrossI1 = Cross1 * WorldSpaceInvI1;
+								PMatrix<T, d, d> CrossI2 = Cross2 * WorldSpaceInvI2;
+								PMatrix<T, d, d> Diag1 = CrossI1 * Cross1.GetTransposed() + CrossI2 * Cross2.GetTransposed();
+								Diag1.M[0][0] += PBDRigid0->InvM() + PBDRigid1->InvM();
+								Diag1.M[1][1] += PBDRigid0->InvM() + PBDRigid1->InvM();
+								Diag1.M[2][2] += PBDRigid0->InvM() + PBDRigid1->InvM();
+								PMatrix<T, d, d> OffDiag1 = (CrossI1 + CrossI2) * -1;
+								PMatrix<T, d, d> Diag2 = (WorldSpaceInvI1 + WorldSpaceInvI2).Inverse();
+								PMatrix<T, d, d> OffDiag1Diag2 = OffDiag1 * Diag2;
+								TVector<T, d> ImpulseDelta = PMatrix<T, d, d>((Diag1 - OffDiag1Diag2 * OffDiag1.GetTransposed()).Inverse())* ((OffDiag1Diag2 * -1) * Delta);
+								Impulse += ImpulseDelta;
+								AngularImpulse += Diag2 * (Delta - PMatrix<T, d, d>(OffDiag1.GetTransposed()) * ImpulseDelta);
 							}
-
-							const T ImpulseMag = -(1 + Restitution) * RelativeNormalVelocity / ImpulseDenominator;
-							Impulse = ImpulseMag * (Contact.Normal - Friction * Tangent);
 						}
 					}
 					else
 					{
-						T ImpulseDenominator = TVector<T, d>::DotProduct(Contact.Normal, Factor * Contact.Normal);
-						TVector<T, d> ImpulseNumerator = -(1 + Restitution) * TVector<T, d>::DotProduct(RelativeVelocity, Contact.Normal)* Contact.Normal;
-						if (!ensureMsgf(FMath::Abs(ImpulseDenominator) > SMALL_NUMBER, TEXT("Constraint:%s\n\nParticle:%s\n\nLevelset:%s\n\nFactor*Constraint.Normal:%s, ImpulseDenominator:%f"),
-							*Constraint.ToString(),
+						//outside friction cone, solve for normal relative velocity and keep tangent at cone edge
+						TVector<T, d> Tangent = (RelativeVelocity - TVector<T, d>::DotProduct(RelativeVelocity, Contact.Normal) * Contact.Normal).GetSafeNormal();
+						TVector<T, d> DirectionalFactor = Factor * (Contact.Normal - Friction * Tangent);
+						T ImpulseDenominator = TVector<T, d>::DotProduct(Contact.Normal, DirectionalFactor);
+						if (!ensureMsgf(FMath::Abs(ImpulseDenominator) > SMALL_NUMBER, TEXT("Contact:%s\n\nParticle:%s\n\nLevelset:%s\n\nDirectionalFactor:%s, ImpulseDenominator:%f"),
+							*Contact.ToString(),
 							*Particle0->ToString(),
 							*Particle1->ToString(),
-							*(Factor * Contact.Normal).ToString(), ImpulseDenominator))
+							*DirectionalFactor.ToString(), ImpulseDenominator))
 						{
 							ImpulseDenominator = (T)1;
 						}
-						Impulse = ImpulseNumerator / ImpulseDenominator;
-					}
-					Impulse = GetEnergyClampedImpulse(Constraint, Impulse, VectorToPoint1, VectorToPoint2, Body1Velocity, Body2Velocity);
-					Constraint.AccumulatedImpulse += Impulse;
-					if (bIsRigidDynamic0)
-					{
-						// Velocity update for next step
-						FVec3 NetAngularImpulse = FVec3::CrossProduct(VectorToPoint1, Impulse) + AngularImpulse;
-						FVec3 DV = PBDRigid0->InvM() * Impulse;
-						FVec3 DW = WorldSpaceInvI1 * NetAngularImpulse;
-						PBDRigid0->V() += DV;
-						PBDRigid0->W() += DW;
-						// Position update as part of pbd
-						P0 += (DV * IterationParameters.Dt);
-						Q0 += FRotation3::FromElements(DW, 0.f) * Q0 * IterationParameters.Dt * T(0.5);
-						Q0.Normalize();
-						FParticleUtilities::SetCoMWorldTransform(PBDRigid0, P0, Q0);
+
+						const T ImpulseMag = -(1 + Restitution) * RelativeNormalVelocity / ImpulseDenominator;
+						Impulse = ImpulseMag * (Contact.Normal - Friction * Tangent);
 					}
-					if (bIsRigidDynamic1)
+				}
+				else
+				{
+					T ImpulseDenominator = TVector<T, d>::DotProduct(Contact.Normal, Factor * Contact.Normal);
+					TVector<T, d> ImpulseNumerator = -(1 + Restitution) * TVector<T, d>::DotProduct(RelativeVelocity, Contact.Normal)* Contact.Normal;
+					if (!ensureMsgf(FMath::Abs(ImpulseDenominator) > SMALL_NUMBER, TEXT("Contact:%s\n\nParticle:%s\n\nLevelset:%s\n\nFactor*Constraint.Normal:%s, ImpulseDenominator:%f"),
+						*Contact.ToString(),
+						*Particle0->ToString(),
+						*Particle1->ToString(),
+						*(Factor * Contact.Normal).ToString(), ImpulseDenominator))
 					{
-						// Velocity update for next step
-						FVec3 NetAngularImpulse = FVec3::CrossProduct(VectorToPoint2, -Impulse) - AngularImpulse;
-						FVec3 DV = -PBDRigid1->InvM() * Impulse;
-						FVec3 DW = WorldSpaceInvI2 * NetAngularImpulse;
-						PBDRigid1->V() += DV;
-						PBDRigid1->W() += DW;
-						// Position update as part of pbd
-						P1 += (DV * IterationParameters.Dt);
-						Q1 += FRotation3::FromElements(DW, 0.f) * Q1 * IterationParameters.Dt * T(0.5);
-						Q1.Normalize();
-						FParticleUtilities::SetCoMWorldTransform(PBDRigid1, P1, Q1);
+						ImpulseDenominator = (T)1;
 					}
+					Impulse = ImpulseNumerator / ImpulseDenominator;
+				}
+
+
+				Impulse = GetEnergyClampedImpulse(Particle0->CastToRigidParticle(), Particle1->CastToRigidParticle(), Impulse, VectorToPoint1, VectorToPoint2, Body1Velocity, Body2Velocity);
+				AccumulatedImpulse += Impulse;
+				if (bIsRigidDynamic0)
+				{
+					// Velocity update for next step
+					FVec3 NetAngularImpulse = FVec3::CrossProduct(VectorToPoint1, Impulse) + AngularImpulse;
+					FVec3 DV = PBDRigid0->InvM() * Impulse;
+					FVec3 DW = WorldSpaceInvI1 * NetAngularImpulse;
+					PBDRigid0->V() += DV;
+					PBDRigid0->W() += DW;
+					// Position update as part of pbd
+					P0 += (DV * IterationParameters.Dt);
+					Q0 += FRotation3::FromElements(DW, 0.f) * Q0 * IterationParameters.Dt * T(0.5);
+					Q0.Normalize();
+					FParticleUtilities::SetCoMWorldTransform(PBDRigid0, P0, Q0);
+				}
+				if (bIsRigidDynamic1)
+				{
+					// Velocity update for next step
+					FVec3 NetAngularImpulse = FVec3::CrossProduct(VectorToPoint2, -Impulse) - AngularImpulse;
+					FVec3 DV = -PBDRigid1->InvM() * Impulse;
+					FVec3 DW = WorldSpaceInvI2 * NetAngularImpulse;
+					PBDRigid1->V() += DV;
+					PBDRigid1->W() += DW;
+					// Position update as part of pbd
+					P1 += (DV * IterationParameters.Dt);
+					Q1 += FRotation3::FromElements(DW, 0.f) * Q1 * IterationParameters.Dt * T(0.5);
+					Q1.Normalize();
+					FParticleUtilities::SetCoMWorldTransform(PBDRigid1, P1, Q1);
 				}
 			}
+			return AccumulatedImpulse;
 		}
 
 		template<typename T, int d>
-		void ApplyPushOut(TCollisionConstraintBase<T, d>& Constraint, T Thickness, const TSet<const TGeometryParticleHandle<T, d>*>& IsTemporarilyStatic,
-			TContactIterationParameters<T> & IterationParameters, TContactParticleParameters<T> & ParticleParameters)
+		void Apply(TCollisionConstraintBase<T, d>& Constraint, TContactIterationParameters<T> & IterationParameters, TContactParticleParameters<T> & ParticleParameters)
 		{
 			TGenericParticleHandle<T, d> Particle0 = TGenericParticleHandle<T, d>(Constraint.Particle[0]);
 			TGenericParticleHandle<T, d> Particle1 = TGenericParticleHandle<T, d>(Constraint.Particle[1]);
+
+			for (int32 PairIt = 0; PairIt < IterationParameters.NumPairIterations; ++PairIt)
+			{
+				Collisions::Update<ECollisionUpdateType::Deepest>(ParticleParameters.Thickness, Constraint);
+				if (Constraint.GetPhi() >= ParticleParameters.Thickness)
+				{
+					return;
+				}
+
+				// @todo(ccaulfield): CHAOS_PARTICLEHANDLE_TODO what's the best way to manage external per-particle data?
+				if (ParticleParameters.Collided)
+				{
+					Particle0->AuxilaryValue(*ParticleParameters.Collided) = true;
+					Particle1->AuxilaryValue(*ParticleParameters.Collided) = true;
+				}
+
+				Constraint.AccumulatedImpulse +=
+					ApplyContact(Constraint.Manifold, Particle0, Particle1, IterationParameters, ParticleParameters);
+			}
+		}
+
+		template<typename T, int d> 
+		TVector<T, d>
+		ApplyPushOutContact(
+			TCollisionContact<T,d>& Contact,
+			TGenericParticleHandle<T, d> Particle0, 
+			TGenericParticleHandle<T, d> Particle1,
+			const TSet<const TGeometryParticleHandle<T, d>*>& IsTemporarilyStatic,
+			TContactIterationParameters<T> & IterationParameters, TContactParticleParameters<T> & ParticleParameters)
+		{
+			TVector<T, d> AccumulatedImpulse(0);
+
 			TPBDRigidParticleHandle<T, d>* PBDRigid0 = Particle0->CastToRigidParticle();
 			TPBDRigidParticleHandle<T, d>* PBDRigid1 = Particle1->CastToRigidParticle();
 			const bool bIsRigidDynamic0 = PBDRigid0 && PBDRigid0->ObjectState() == EObjectStateType::Dynamic;
@@ -299,8 +320,8 @@ namespace Chaos
 			TVector<T, d> P1 = FParticleUtilities::GetCoMWorldPosition(Particle1);
 			TRotation<T, d> Q0 = FParticleUtilities::GetCoMWorldRotation(Particle0);
 			TRotation<T, d> Q1 = FParticleUtilities::GetCoMWorldRotation(Particle1);
-			const bool IsTemporarilyStatic0 = IsTemporarilyStatic.Contains(Constraint.Particle[0]);
-			const bool IsTemporarilyStatic1 = IsTemporarilyStatic.Contains(Constraint.Particle[1]);
+			const bool IsTemporarilyStatic0 = IsTemporarilyStatic.Contains(Particle0->GeometryParticleHandle());
+			const bool IsTemporarilyStatic1 = IsTemporarilyStatic.Contains(Particle1->GeometryParticleHandle());
 
 			TSerializablePtr<FChaosPhysicsMaterial> PhysicsMaterial0, PhysicsMaterial1;
 			if (ParticleParameters.PhysicsMaterials)
@@ -309,89 +330,101 @@ namespace Chaos
 				PhysicsMaterial1 = Particle1->AuxilaryValue(*ParticleParameters.PhysicsMaterials);
 			}
 
-			for (int32 PairIteration = 0; PairIteration < IterationParameters.NumPairIterations; ++PairIteration)
+			if (Contact.Phi >= ParticleParameters.Thickness)
 			{
-				Update<ECollisionUpdateType::Deepest>(Thickness, Constraint);
+				return AccumulatedImpulse;
+			}
 
-				const typename TRigidBodyPointContactConstraint<T, d>::FManifold & Contact = Constraint.Manifold;
+			if ((!bIsRigidDynamic0 || IsTemporarilyStatic0) && (!bIsRigidDynamic1 || IsTemporarilyStatic1))
+			{
+				return AccumulatedImpulse;
+			}
 
-				if (Contact.Phi >= Thickness)
+			*IterationParameters.NeedsAnotherIteration = true;
+			PMatrix<T, d, d> WorldSpaceInvI1 = bIsRigidDynamic0 ? Utilities::ComputeWorldSpaceInertia(Q0, PBDRigid0->InvI()) : PMatrix<T, d, d>(0);
+			PMatrix<T, d, d> WorldSpaceInvI2 = bIsRigidDynamic1 ? Utilities::ComputeWorldSpaceInertia(Q1, PBDRigid1->InvI()) : PMatrix<T, d, d>(0);
+			TVector<T, d> VectorToPoint1 = Contact.Location - P0;
+			TVector<T, d> VectorToPoint2 = Contact.Location - P1;
+			PMatrix<T, d, d> Factor =
+				(bIsRigidDynamic0 ? ComputeFactorMatrix3(VectorToPoint1, WorldSpaceInvI1, PBDRigid0->InvM()) : PMatrix<T, d, d>(0)) +
+				(bIsRigidDynamic1 ? ComputeFactorMatrix3(VectorToPoint2, WorldSpaceInvI2, PBDRigid1->InvM()) : PMatrix<T, d, d>(0));
+			T Numerator = FMath::Min((T)(IterationParameters.Iteration + 2), (T)IterationParameters.NumIterations);
+			T ScalingFactor = Numerator / (T)IterationParameters.NumIterations;
+
+			//if pushout is needed we better fix relative velocity along normal. Treat it as if 0 restitution
+			TVector<T, d> Body1Velocity = FParticleUtilities::GetVelocityAtCoMRelativePosition(Particle0, VectorToPoint1);
+			TVector<T, d> Body2Velocity = FParticleUtilities::GetVelocityAtCoMRelativePosition(Particle1, VectorToPoint2);
+			TVector<T, d> RelativeVelocity = Body1Velocity - Body2Velocity;
+			const T RelativeVelocityDotNormal = TVector<T, d>::DotProduct(RelativeVelocity, Contact.Normal);
+			if (RelativeVelocityDotNormal < 0)
+			{
+				T ImpulseDenominator = TVector<T, d>::DotProduct(Contact.Normal, Factor * Contact.Normal);
+				TVector<T, d> ImpulseNumerator = -TVector<T, d>::DotProduct(RelativeVelocity, Contact.Normal) * Contact.Normal * ScalingFactor;
+				if (!ensureMsgf(FMath::Abs(ImpulseDenominator) > SMALL_NUMBER, TEXT("ApplyPushout Contact:%s\n\nParticle:%s\n\nLevelset:%s\n\nFactor*Contact.Normal:%s, ImpulseDenominator:%f"),
+					*Contact.ToString(),
+					*Particle0->ToString(),
+					*Particle1->ToString(),
+					*(Factor*Contact.Normal).ToString(), ImpulseDenominator))
 				{
-					break;
+					ImpulseDenominator = (T)1;
 				}
 
-				if ((!bIsRigidDynamic0 || IsTemporarilyStatic0) && (!bIsRigidDynamic1 || IsTemporarilyStatic1))
+				TVector<T, d> VelocityFixImpulse = ImpulseNumerator / ImpulseDenominator;
+				VelocityFixImpulse = GetEnergyClampedImpulse(Particle0->CastToRigidParticle(), Particle1->CastToRigidParticle(), VelocityFixImpulse, VectorToPoint1, VectorToPoint2, Body1Velocity, Body2Velocity);
+				AccumulatedImpulse += VelocityFixImpulse;	//question: should we track this?
+				if (!IsTemporarilyStatic0 && bIsRigidDynamic0)
 				{
-					break;
+					TVector<T, d> AngularImpulse = TVector<T, d>::CrossProduct(VectorToPoint1, VelocityFixImpulse);
+					PBDRigid0->V() += PBDRigid0->InvM() * VelocityFixImpulse;
+					PBDRigid0->W() += WorldSpaceInvI1 * AngularImpulse;
+
 				}
 
-				*IterationParameters.NeedsAnotherIteration = true;
-				PMatrix<T, d, d> WorldSpaceInvI1 = bIsRigidDynamic0 ? Utilities::ComputeWorldSpaceInertia(Q0, PBDRigid0->InvI()) : PMatrix<T, d, d>(0);
-				PMatrix<T, d, d> WorldSpaceInvI2 = bIsRigidDynamic1 ? Utilities::ComputeWorldSpaceInertia(Q1, PBDRigid1->InvI()) : PMatrix<T, d, d>(0);
-				TVector<T, d> VectorToPoint1 = Contact.Location - P0;
-				TVector<T, d> VectorToPoint2 = Contact.Location - P1;
-				PMatrix<T, d, d> Factor =
-					(bIsRigidDynamic0 ? ComputeFactorMatrix3(VectorToPoint1, WorldSpaceInvI1, PBDRigid0->InvM()) : PMatrix<T, d, d>(0)) +
-					(bIsRigidDynamic1 ? ComputeFactorMatrix3(VectorToPoint2, WorldSpaceInvI2, PBDRigid1->InvM()) : PMatrix<T, d, d>(0));
-				T Numerator = FMath::Min((T)(IterationParameters.Iteration + 2), (T)IterationParameters.NumIterations);
-				T ScalingFactor = Numerator / (T)IterationParameters.NumIterations;
-
-				//if pushout is needed we better fix relative velocity along normal. Treat it as if 0 restitution
-				TVector<T, d> Body1Velocity = FParticleUtilities::GetVelocityAtCoMRelativePosition(Particle0, VectorToPoint1);
-				TVector<T, d> Body2Velocity = FParticleUtilities::GetVelocityAtCoMRelativePosition(Particle1, VectorToPoint2);
-				TVector<T, d> RelativeVelocity = Body1Velocity - Body2Velocity;
-				const T RelativeVelocityDotNormal = TVector<T, d>::DotProduct(RelativeVelocity, Contact.Normal);
-				if (RelativeVelocityDotNormal < 0)
+				if (!IsTemporarilyStatic1 && bIsRigidDynamic1)
 				{
-					T ImpulseDenominator = TVector<T, d>::DotProduct(Contact.Normal, Factor * Contact.Normal);
-					TVector<T, d> ImpulseNumerator = -TVector<T, d>::DotProduct(RelativeVelocity, Contact.Normal) * Contact.Normal * ScalingFactor;
-					if (!ensureMsgf(FMath::Abs(ImpulseDenominator) > SMALL_NUMBER, TEXT("ApplyPushout Constraint:%s\n\nParticle:%s\n\nLevelset:%s\n\nFactor*Contact.Normal:%s, ImpulseDenominator:%f"),
-						*Constraint.ToString(),
-						*Particle0->ToString(),
-						*Particle1->ToString(),
-						*(Factor*Contact.Normal).ToString(), ImpulseDenominator))
-					{
-						ImpulseDenominator = (T)1;
-					}
+					TVector<T, d> AngularImpulse = TVector<T, d>::CrossProduct(VectorToPoint2, -VelocityFixImpulse);
+					PBDRigid1->V() -= PBDRigid1->InvM() * VelocityFixImpulse;
+					PBDRigid1->W() += WorldSpaceInvI2 * AngularImpulse;
+				}
 
-					TVector<T, d> VelocityFixImpulse = ImpulseNumerator / ImpulseDenominator;
-					VelocityFixImpulse = GetEnergyClampedImpulse(Constraint, VelocityFixImpulse, VectorToPoint1, VectorToPoint2, Body1Velocity, Body2Velocity);
-					Constraint.AccumulatedImpulse += VelocityFixImpulse;	//question: should we track this?
-					if (!IsTemporarilyStatic0 && bIsRigidDynamic0)
-					{
-						TVector<T, d> AngularImpulse = TVector<T, d>::CrossProduct(VectorToPoint1, VelocityFixImpulse);
-						PBDRigid0->V() += PBDRigid0->InvM() * VelocityFixImpulse;
-						PBDRigid0->W() += WorldSpaceInvI1 * AngularImpulse;
+			}
 
-					}
 
-					if (!IsTemporarilyStatic1 && bIsRigidDynamic1)
-					{
-						TVector<T, d> AngularImpulse = TVector<T, d>::CrossProduct(VectorToPoint2, -VelocityFixImpulse);
-						PBDRigid1->V() -= PBDRigid1->InvM() * VelocityFixImpulse;
-						PBDRigid1->W() += WorldSpaceInvI2 * AngularImpulse;
-					}
+			TVector<T, d> Impulse = PMatrix<T, d, d>(Factor.Inverse()) * ((-Contact.Phi + ParticleParameters.Thickness) * ScalingFactor * Contact.Normal);
+			TVector<T, d> AngularImpulse1 = TVector<T, d>::CrossProduct(VectorToPoint1, Impulse);
+			TVector<T, d> AngularImpulse2 = TVector<T, d>::CrossProduct(VectorToPoint2, -Impulse);
+			if (!IsTemporarilyStatic0 && bIsRigidDynamic0)
+			{
+				P0 += PBDRigid0->InvM() * Impulse;
+				Q0 = TRotation<T, d>::FromVector(WorldSpaceInvI1 * AngularImpulse1) * Q0;
+				Q0.Normalize();
+				FParticleUtilities::SetCoMWorldTransform(Particle0, P0, Q0);
+			}
+			if (!IsTemporarilyStatic1 && bIsRigidDynamic1)
+			{
+				P1 -= PBDRigid1->InvM() * Impulse;
+				Q1 = TRotation<T, d>::FromVector(WorldSpaceInvI2 * AngularImpulse2) * Q1;
+				Q1.Normalize();
+				FParticleUtilities::SetCoMWorldTransform(Particle1, P1, Q1);
+			}
 
-				}
+			return AccumulatedImpulse;
+		}
 
 
-				TVector<T, d> Impulse = PMatrix<T, d, d>(Factor.Inverse()) * ((-Contact.Phi + Thickness) * ScalingFactor * Contact.Normal);
-				TVector<T, d> AngularImpulse1 = TVector<T, d>::CrossProduct(VectorToPoint1, Impulse);
-				TVector<T, d> AngularImpulse2 = TVector<T, d>::CrossProduct(VectorToPoint2, -Impulse);
-				if (!IsTemporarilyStatic0 && bIsRigidDynamic0)
-				{
-					P0 += PBDRigid0->InvM() * Impulse;
-					Q0 = TRotation<T, d>::FromVector(WorldSpaceInvI1 * AngularImpulse1) * Q0;
-					Q0.Normalize();
-					FParticleUtilities::SetCoMWorldTransform(Particle0, P0, Q0);
-				}
-				if (!IsTemporarilyStatic1 && bIsRigidDynamic1)
-				{
-					P1 -= PBDRigid1->InvM() * Impulse;
-					Q1 = TRotation<T, d>::FromVector(WorldSpaceInvI2 * AngularImpulse2) * Q1;
-					Q1.Normalize();
-					FParticleUtilities::SetCoMWorldTransform(Particle1, P1, Q1);
-				}
+		template<typename T, int d>
+		void ApplyPushOut(TCollisionConstraintBase<T, d>& Constraint, const TSet<const TGeometryParticleHandle<T, d>*>& IsTemporarilyStatic,
+			TContactIterationParameters<T> & IterationParameters, TContactParticleParameters<T> & ParticleParameters)
+		{
+			TGenericParticleHandle<T, d> Particle0 = TGenericParticleHandle<T, d>(Constraint.Particle[0]);
+			TGenericParticleHandle<T, d> Particle1 = TGenericParticleHandle<T, d>(Constraint.Particle[1]);
+
+			for (int32 PairIteration = 0; PairIteration < IterationParameters.NumPairIterations; ++PairIteration)
+			{
+				Update<ECollisionUpdateType::Deepest>(ParticleParameters.Thickness, Constraint);
+
+				Constraint.AccumulatedImpulse += 
+					ApplyPushOutContact(Constraint.Manifold, Particle0, Particle1, IsTemporarilyStatic, IterationParameters, ParticleParameters);
 			}
 		}
 
@@ -399,8 +432,9 @@ namespace Chaos
 		template void Update<ECollisionUpdateType::Any, float, 3>(const float, TCollisionConstraintBase<float, 3>&);
 		template void Update<ECollisionUpdateType::Deepest, float, 3>(const float, TCollisionConstraintBase<float, 3>&);
 		template void UpdateManifold<float, 3>(const float, TCollisionConstraintBase<float, 3>&);
-		template void Apply<float, 3>(TCollisionConstraintBase<float, 3>&, float, TContactIterationParameters<float> &, TContactParticleParameters<float> &);
-		template void ApplyPushOut<float, 3>(TCollisionConstraintBase<float,3>&, float , const TSet<const TGeometryParticleHandle<float,3>*>&,
+		//
+		template void Apply<float, 3>(TCollisionConstraintBase<float, 3>&, TContactIterationParameters<float> &, TContactParticleParameters<float> &);
+		template void ApplyPushOut<float, 3>(TCollisionConstraintBase<float,3>& , const TSet<const TGeometryParticleHandle<float,3>*>&,
 			TContactIterationParameters<float> & IterationParameters, TContactParticleParameters<float> & ParticleParameters);
 
 	} // Collisions
diff --git a/Engine/Source/Runtime/Experimental/Chaos/Public/Chaos/CollisionResolutionTypes.h b/Engine/Source/Runtime/Experimental/Chaos/Public/Chaos/CollisionResolutionTypes.h
--- a/Engine/Source/Runtime/Experimental/Chaos/Public/Chaos/CollisionResolutionTypes.h
+++ b/Engine/Source/Runtime/Experimental/Chaos/Public/Chaos/CollisionResolutionTypes.h
@@ -57,6 +57,12 @@ namespace Chaos
 		TVector<T, d> Location;
 		T Phi;
 
+
+		FString ToString() const
+		{
+			return FString::Printf(TEXT("Location:%s, Normal:%s, Phi:%s"), *Location.ToString(), *Normal.ToString(), Phi);
+		}
+
 		const FImplicitObject* Implicit[2]; // {Of Particle[0], Of Particle[1]}
 	};
 
@@ -162,14 +168,18 @@ namespace Chaos
 	public:
 		using Base = TCollisionConstraintBase<T, d>;
 		using FGeometryParticleHandle = TGeometryParticleHandle<T, d>;
+		using FManifold = TCollisionContact<T, d>;
 		using Base::Particle;
+		struct FSampleData { TVector<T,d> X; float Delta; FManifold Manifold; };
 
-		TRigidBodyIterativeContactConstraint() : Base(Base::FType::MultiPoint), WorldNormal(0), LocalPosition(0) {}
+		TRigidBodyIterativeContactConstraint() : Base(Base::FType::MultiPoint), PlaneNormal(0), PlanePosition(0), LocalPosition(0) {}
 		static typename Base::FType StaticType() { return Base::FType::MultiPoint; };
 
-		TVector<T, d> WorldNormal; // sampled normal
+		TVector<T, d> PlaneNormal; // local space contact normal on Particle1
+		TVector<T, d> PlanePosition; // local space surface position on Particle1
+
 		TVector<T, d> LocalPosition; // averaged position
-		TArray<TVector<T, d>> LocalSamples; // iterative samples
+		TArray< FSampleData > LocalSamples; // iterative samples
 	};
 	typedef TRigidBodyIterativeContactConstraint<float, 3> FRigidBodyIterativeContactConstraint;
 
diff --git a/Engine/Source/Runtime/Experimental/Chaos/Public/Chaos/CollisionResolutionUtil.h b/Engine/Source/Runtime/Experimental/Chaos/Public/Chaos/CollisionResolutionUtil.h
--- a/Engine/Source/Runtime/Experimental/Chaos/Public/Chaos/CollisionResolutionUtil.h
+++ b/Engine/Source/Runtime/Experimental/Chaos/Public/Chaos/CollisionResolutionUtil.h
@@ -30,7 +30,7 @@ namespace Chaos
 
 		template<typename T, int d>
 		TVector<T, d> 
-		GetEnergyClampedImpulse(const TCollisionConstraintBase<T, d>& Constraint, const TVector<T, d>& Impulse, const TVector<T, d>& VectorToPoint1, const TVector<T, d>& VectorToPoint2, const TVector<T, d>& Velocity1, const TVector<T, d>& Velocity2);
+		GetEnergyClampedImpulse(const TPBDRigidParticleHandle<T, d>* PBDRigid0, const TPBDRigidParticleHandle<T, d>* PBDRigid1, const TVector<T, d>& Impulse, const TVector<T, d>& VectorToPoint1, const TVector<T, d>& VectorToPoint2, const TVector<T, d>& Velocity1, const TVector<T, d>& Velocity2);
 
 		template <typename T, int d>
 		bool 
diff --git a/Engine/Source/Runtime/Experimental/Chaos/Public/Chaos/PBDCollisionConstraintsContact.h b/Engine/Source/Runtime/Experimental/Chaos/Public/Chaos/PBDCollisionConstraintsContact.h
--- a/Engine/Source/Runtime/Experimental/Chaos/Public/Chaos/PBDCollisionConstraintsContact.h
+++ b/Engine/Source/Runtime/Experimental/Chaos/Public/Chaos/PBDCollisionConstraintsContact.h
@@ -11,6 +11,7 @@ namespace Chaos
 	{
 		template<class T = float>
 		struct TContactParticleParameters {
+			T Thickness;
 			TArrayCollectionArray<bool>* Collided;
 			const TArrayCollectionArray<TSerializablePtr<FChaosPhysicsMaterial>>* PhysicsMaterials;
 			T FrictionOverride;
@@ -33,10 +34,10 @@ namespace Chaos
 		void UpdateManifold(const T Thickness, TCollisionConstraintBase<T, d>& Constraint);
 
 		template<typename T, int d>
-		void Apply(TCollisionConstraintBase<T, d>& Constraint, T Thickness, TContactIterationParameters<T> & IterationParameters, TContactParticleParameters<T> & ParticleParameters);
+		void Apply(TCollisionConstraintBase<T, d>& Constraint, TContactIterationParameters<T> & IterationParameters, TContactParticleParameters<T> & ParticleParameters);
 
 		template<typename T, int d>
-		void ApplyPushOut(TCollisionConstraintBase<T, d>& Constraint, T Thickness, const TSet<const TGeometryParticleHandle<T, d>*>& IsTemporarilyStatic,
+		void ApplyPushOut(TCollisionConstraintBase<T, d>& Constraint, const TSet<const TGeometryParticleHandle<T, d>*>& IsTemporarilyStatic,
 			TContactIterationParameters<T> & IterationParameters, TContactParticleParameters<T> & ParticleParameters);
 
 
diff --git a/Engine/Source/Runtime/Experimental/Chaos/Public/Chaos/ParticleHandle.h b/Engine/Source/Runtime/Experimental/Chaos/Public/Chaos/ParticleHandle.h
--- a/Engine/Source/Runtime/Experimental/Chaos/Public/Chaos/ParticleHandle.h
+++ b/Engine/Source/Runtime/Experimental/Chaos/Public/Chaos/ParticleHandle.h
@@ -888,6 +888,8 @@ public:
 	TKinematicGeometryParticleHandle<T, d>* CastToKinematicParticle() { return Handle->CastToKinematicParticle(); }
 	const TPBDRigidParticleHandle<T, d>* CastToRigidParticle() const { return Handle->CastToRigidParticle(); }
 	TPBDRigidParticleHandle<T, d>* CastToRigidParticle() { return Handle->CastToRigidParticle(); }
+	const TGeometryParticleHandle<T, d>* GeometryParticleHandle() const { return Handle; }
+	TGeometryParticleHandle<T, d>* GeometryParticleHandle() { return Handle; }
 
 	// Static Particles
 	TVector<T, d>& X() { return Handle->X(); }
