Chaos : Collisions

Set the deepest point on the constraint from the sample points.

diff --git a/Engine/Source/Runtime/Experimental/Chaos/Private/Chaos/CollisionResolution.cpp b/Engine/Source/Runtime/Experimental/Chaos/Private/Chaos/CollisionResolution.cpp
--- a/Engine/Source/Runtime/Experimental/Chaos/Private/Chaos/CollisionResolution.cpp
+++ b/Engine/Source/Runtime/Experimental/Chaos/Private/Chaos/CollisionResolution.cpp
@@ -563,14 +563,22 @@ namespace Chaos
 			TContactPoint<T> ContactPoint = ConvexConvexContactPoint(Implicit0, Transform0, Implicit1, Transform1, Thickness);
 			/*
 			TContactPoint<T> ContactPoint;
+			ContactPoint.Phi = FLT_MAX;
 
 			const TRigidTransform<T, d> AToBTM = Transform0.GetRelativeTransform(Transform1);
 
 			for (int32 Idx = 0; Idx < Constraint.LocalSamples.Num(); Idx++)
 			{
-				Constraint.LocalSamples[Idx].Manifold.Phi = Implicit1.PhiWithNormal(Constraint.LocalSamples[Idx].X, Constraint.LocalSamples[Idx].Manifold.Normal);
+				Constraint.LocalSamples[Idx].Manifold.Phi = Implicit1.PhiWithNormal(AToBTM.TransformPosition(Constraint.LocalSamples[Idx].X), Constraint.LocalSamples[Idx].Manifold.Normal);
 				Constraint.LocalSamples[Idx].Manifold.Normal = Transform1.TransformVector(Constraint.PlaneNormal);
 				Constraint.LocalSamples[Idx].Manifold.Location = Transform0.TransformPosition(Constraint.LocalSamples[Idx].X);
+
+				if (ContactPoint.Phi > Constraint.LocalSamples[Idx].Manifold.Phi)
+				{
+					ContactPoint.Phi = Constraint.LocalSamples[Idx].Manifold.Phi;
+					ContactPoint.Normal = Constraint.LocalSamples[Idx].Manifold.Normal;
+					ContactPoint.Location = Constraint.LocalSamples[Idx].Manifold.Location;
+				}
 			}
 			*/
 			UpdateContactPoint(Constraint.Manifold, ContactPoint);
diff --git a/Engine/Source/Runtime/Experimental/Chaos/Public/Chaos/CollisionResolutionTypes.h b/Engine/Source/Runtime/Experimental/Chaos/Public/Chaos/CollisionResolutionTypes.h
--- a/Engine/Source/Runtime/Experimental/Chaos/Public/Chaos/CollisionResolutionTypes.h
+++ b/Engine/Source/Runtime/Experimental/Chaos/Public/Chaos/CollisionResolutionTypes.h
@@ -65,7 +65,7 @@ namespace Chaos
 
 		const FImplicitObject* Implicit[2]; // {Of Particle[0], Of Particle[1]}
 	};
-
+	typedef TCollisionContact<float, 3> FCollisionContact;
 
 	/*
 	*
