Chaos : Collisions

Construct the collision surface from the sampled plane, and use PhiWithNormal off the plane.

diff --git a/Engine/Source/Runtime/Experimental/Chaos/Private/Chaos/CollisionResolution.cpp b/Engine/Source/Runtime/Experimental/Chaos/Private/Chaos/CollisionResolution.cpp
--- a/Engine/Source/Runtime/Experimental/Chaos/Private/Chaos/CollisionResolution.cpp
+++ b/Engine/Source/Runtime/Experimental/Chaos/Private/Chaos/CollisionResolution.cpp
@@ -567,9 +567,11 @@ namespace Chaos
 
 			const TRigidTransform<T, d> AToBTM = Transform0.GetRelativeTransform(Transform1);
 
+			TPlane<T, d> CollisionPlane(Constraint.PlanePosition, Constraint.PlaneNormal);
+
 			for (int32 Idx = 0; Idx < Constraint.LocalSamples.Num(); Idx++)
 			{
-				Constraint.LocalSamples[Idx].Manifold.Phi = Implicit1.PhiWithNormal(AToBTM.TransformPosition(Constraint.LocalSamples[Idx].X), Constraint.LocalSamples[Idx].Manifold.Normal);
+				Constraint.LocalSamples[Idx].Manifold.Phi = CollisionPlane.PhiWithNormal(AToBTM.TransformPosition(Constraint.LocalSamples[Idx].X), Constraint.LocalSamples[Idx].Manifold.Normal);
 				Constraint.LocalSamples[Idx].Manifold.Normal = Transform1.TransformVector(Constraint.PlaneNormal);
 				Constraint.LocalSamples[Idx].Manifold.Location = Transform0.TransformPosition(Constraint.LocalSamples[Idx].X);
 
