FIX ROBOMERGE

Add linear and angular damping parameters to rigid particles

diff --git a/Engine/Source/Runtime/Core/Public/UObject/ExternalPhysicsCustomObjectVersion.h b/Engine/Source/Runtime/Core/Public/UObject/ExternalPhysicsCustomObjectVersion.h
--- a/Engine/Source/Runtime/Core/Public/UObject/ExternalPhysicsCustomObjectVersion.h
+++ b/Engine/Source/Runtime/Core/Public/UObject/ExternalPhysicsCustomObjectVersion.h
@@ -80,6 +80,9 @@ struct CORE_API FExternalPhysicsCustomObjectVersion
 
 		//Heightfield cell bounds are implicit
 		HeightfieldImplicitBounds,
+
+		// Add damping to rigid particles
+		AddDampingToRigids,
 		
 		// -----<new versions can be added above this line>-------------------------------------------------
 		VersionPlusOne,
diff --git a/Engine/Source/Runtime/Engine/Private/PhysicsEngine/Experimental/PhysInterface_Chaos.cpp b/Engine/Source/Runtime/Engine/Private/PhysicsEngine/Experimental/PhysInterface_Chaos.cpp
--- a/Engine/Source/Runtime/Engine/Private/PhysicsEngine/Experimental/PhysInterface_Chaos.cpp
+++ b/Engine/Source/Runtime/Engine/Private/PhysicsEngine/Experimental/PhysInterface_Chaos.cpp
@@ -638,12 +638,26 @@ FBox FPhysInterface_Chaos::GetBounds_AssumesLocked(const FPhysicsActorHandle& In
 
 void FPhysInterface_Chaos::SetLinearDamping_AssumesLocked(const FPhysicsActorHandle& InActorReference, float InDamping)
 {
-
+	if (ensure(FPhysicsInterface::IsValid(InActorReference)))
+	{
+		Chaos::TPBDRigidParticle<float, 3>* Rigid = InActorReference->CastToRigidParticle();
+		if (ensure(Rigid))
+		{
+			Rigid->SetLinearDamping(InDamping);
+		}
+	}
 }
 
 void FPhysInterface_Chaos::SetAngularDamping_AssumesLocked(const FPhysicsActorHandle& InActorReference, float InDamping)
 {
-
+	if (ensure(FPhysicsInterface::IsValid(InActorReference)))
+	{
+		Chaos::TPBDRigidParticle<float, 3>* Rigid = InActorReference->CastToRigidParticle();
+		if (ensure(Rigid))
+		{
+			Rigid->SetAngularDamping(InDamping);
+		}
+	}
 }
 
 void FPhysInterface_Chaos::AddImpulse_AssumesLocked(const FPhysicsActorHandle& InActorReference, const FVector& InForce)
diff --git a/Engine/Source/Runtime/Experimental/Chaos/Public/Chaos/PBDRigidsEvolutionGBF.h b/Engine/Source/Runtime/Experimental/Chaos/Public/Chaos/PBDRigidsEvolutionGBF.h
--- a/Engine/Source/Runtime/Experimental/Chaos/Public/Chaos/PBDRigidsEvolutionGBF.h
+++ b/Engine/Source/Runtime/Experimental/Chaos/Public/Chaos/PBDRigidsEvolutionGBF.h
@@ -10,6 +10,7 @@
 #include "Chaos/Collision/SpatialAccelerationBroadPhase.h"
 #include "Chaos/PerParticleInitForce.h"
 #include "Chaos/PerParticleEulerStepVelocity.h"
+#include "Chaos/PerParticleVelocityDamping.h"
 #include "Chaos/PerParticleEtherDrag.h"
 #include "Chaos/PerParticleGravity.h"
 #include "Chaos/PerParticlePBDEulerStep.h"
@@ -127,6 +128,7 @@ public:
 		CHAOS_SCOPED_TIMER(Integrate);
 		TPerParticleInitForce<T, d> InitForceRule;
 		TPerParticleEulerStepVelocity<T, d> EulerStepVelocityRule;
+		TPerParticleVelocityDamping<T, d> VelocityDampingRule;
 		TPerParticleEtherDrag<T, d> EtherDragRule(HackLinearDrag, HackAngularDrag);
 		TPerParticlePBDEulerStep<T, d> EulerStepRule;
 
@@ -150,6 +152,7 @@ public:
 					ForceRule(Particle, Dt);
 				}
 				EulerStepVelocityRule.Apply(Particle, Dt);
+				VelocityDampingRule.Apply(Particle, Dt);
 				EtherDragRule.Apply(Particle, Dt);
 
 				if (HackMaxAngularVelocity >= 0.f)
diff --git a/Engine/Source/Runtime/Experimental/Chaos/Public/Chaos/ParticleDirtyFlags.h b/Engine/Source/Runtime/Experimental/Chaos/Public/Chaos/ParticleDirtyFlags.h
--- a/Engine/Source/Runtime/Experimental/Chaos/Public/Chaos/ParticleDirtyFlags.h
+++ b/Engine/Source/Runtime/Experimental/Chaos/Public/Chaos/ParticleDirtyFlags.h
@@ -27,16 +27,18 @@ namespace Chaos
 		InvI					= 1 << 15,
 		M						= 1 << 16,
 		InvM					= 1 << 17,
-		ObjectState				= 1 << 18,
-		Geometry				= 1 << 19,
-		ExternalForce			= 1 << 20,
-		ExternalTorque			= 1 << 20,
-		GravityEnabled			= 1 << 22,
-		SpatialIdx				= 1 << 23,
-		HashResult				= 1 << 24,
-		ShapeDisableCollision	= 1 << 25
+		LinearDamping			= 1 << 18,
+		AngularDamping			= 1 << 19,
+		ObjectState				= 1 << 20,
+		Geometry				= 1 << 21,
+		ExternalForce			= 1 << 22,
+		ExternalTorque			= 1 << 23,
+		GravityEnabled			= 1 << 24,
+		SpatialIdx				= 1 << 25,
+		HashResult				= 1 << 26,
+		ShapeDisableCollision	= 1 << 27
 #if CHAOS_CHECKED
-		, DebugName				= 1 << 26
+		, DebugName				= 1 << 28
 #endif
 	};
 
diff --git a/Engine/Source/Runtime/Experimental/Chaos/Public/Chaos/ParticleHandle.h b/Engine/Source/Runtime/Experimental/Chaos/Public/Chaos/ParticleHandle.h
--- a/Engine/Source/Runtime/Experimental/Chaos/Public/Chaos/ParticleHandle.h
+++ b/Engine/Source/Runtime/Experimental/Chaos/Public/Chaos/ParticleHandle.h
@@ -671,6 +671,14 @@ public:
 	T& InvM() { return PBDRigidParticles->InvM(ParticleIdx); }
 	void SetInvM(const T& InInvM) { PBDRigidParticles->InvM(ParticleIdx) = InInvM; }
 
+	T LinearDamping() const { return PBDRigidParticles->LinearDamping(ParticleIdx); }
+	T& LinearDamping() { return PBDRigidParticles->LinearDamping(ParticleIdx); }
+	void SetLinearDamping(const T& InLinearDamping) { PBDRigidParticles->LinearDamping(ParticleIdx) = InLinearDamping; }
+
+	T AngularDamping() const { return PBDRigidParticles->AngularDamping(ParticleIdx); }
+	T& AngularDamping() { return PBDRigidParticles->AngularDamping(ParticleIdx); }
+	void SetAngularDamping(const T& InAngularDamping) { PBDRigidParticles->AngularDamping(ParticleIdx) = InAngularDamping; }
+
 	int32 Island() const { return PBDRigidParticles->Island(ParticleIdx); }
 	int32& Island() { return PBDRigidParticles->Island(ParticleIdx); }
 	void SetIsland(const int32 InIsland) { PBDRigidParticles->Island(ParticleIdx) = InIsland; }
@@ -1736,6 +1744,14 @@ public:
 		Ar << MCollisionParticles;
 		Ar << MM;
 		Ar << MInvM;
+
+		Ar.UsingCustomVersion(FExternalPhysicsCustomObjectVersion::GUID);
+		if (Ar.CustomVer(FExternalPhysicsCustomObjectVersion::GUID) >= FExternalPhysicsCustomObjectVersion::AddDampingToRigids)
+		{
+			Ar << MLinearDamping;
+			Ar << MAngularDamping;
+		}
+
 		Ar << MIsland;
 		Ar << MCollisionGroup;
 		Ar << MObjectState;
@@ -1855,6 +1871,20 @@ public:
 		this->MInvM = InInvM;
 	}
 
+	T LinearDamping() const { return MLinearDamping; }
+	void SetLinearDamping(const T& InLinearDamping)
+	{
+		this->MarkDirty(EParticleFlags::LinearDamping);
+		this->MLinearDamping = InLinearDamping;
+	}
+
+	T AngularDamping() const { return MAngularDamping; }
+	void SetAngularDamping(const T& InAngularDamping)
+	{
+		this->MarkDirty(EParticleFlags::AngularDamping);
+		this->MAngularDamping = InAngularDamping;
+	}
+
 	int32 Island() const { return MIsland; }
 	// TODO(stett): Make the setter private. It is public right now to provide access to proxies.
 	void SetIsland(const int32 InIsland)
@@ -1897,6 +1927,8 @@ private:
 	TUniquePtr<TBVHParticles<T, d>> MCollisionParticles;
 	T MM;
 	T MInvM;
+	T MLinearDamping;
+	T MAngularDamping;
 	int32 MIsland;
 	int32 MCollisionGroup;
 	EObjectStateType MObjectState;
@@ -1947,6 +1979,8 @@ public:
 		, MCollisionParticles(nullptr)
 		, MM(T(0))
 		, MInvM(T(0))
+		, MLinearDamping(T(0))
+		, MAngularDamping(T(0))
 		, MIsland(INDEX_NONE)
 		, MCollisionGroup(0)
 		, MObjectState(EObjectStateType::Uninitialized)
@@ -1968,6 +2002,8 @@ public:
 		, MCollisionParticles(nullptr)
 		, MM(InParticle.M())
 		, MInvM(InParticle.InvM())
+		, MLinearDamping(InParticle.LinearDamping())
+		, MAngularDamping(InParticle.AngularDamping())
 		, MIsland(InParticle.Island())
 		, MCollisionGroup(InParticle.CollisionGroup())
 		, MObjectState(InParticle.ObjectState())
@@ -1997,6 +2033,8 @@ public:
 	const TBVHParticles<T, d> * MCollisionParticles;
 	T MM;
 	T MInvM;
+	T MLinearDamping;
+	T MAngularDamping;
 	int32 MIsland;
 	int32 MCollisionGroup;
 	EObjectStateType MObjectState;
@@ -2019,6 +2057,8 @@ public:
 		MCollisionParticles = nullptr;
 		MM = T(0);
 		MInvM = T(0);
+		MLinearDamping = T(0);
+		MAngularDamping = T(0);
 		MIsland = INDEX_NONE;
 		MCollisionGroup = 0;
 		MObjectState = EObjectStateType::Uninitialized;
diff --git a/Engine/Source/Runtime/Experimental/Chaos/Public/Chaos/PerParticleVelocityDamping.h b/Engine/Source/Runtime/Experimental/Chaos/Public/Chaos/PerParticleVelocityDamping.h
new file mode 100644
--- /dev/null
+++ b/Engine/Source/Runtime/Experimental/Chaos/Public/Chaos/PerParticleVelocityDamping.h
@@ -0,0 +1,40 @@
+// Copyright 1998-2019 Epic Games, Inc. All Rights Reserved.
+#pragma once
+
+#include "Chaos/PerParticleRule.h"
+#include "Chaos/Utilities.h"
+
+namespace Chaos
+{
+	template<class T, int d>
+	class TPerParticleVelocityDamping : public TPerParticleRule<T, d>
+	{
+	public:
+		TPerParticleVelocityDamping() {}
+		virtual ~TPerParticleVelocityDamping() {}
+
+		inline void ApplyHelper(FVec3& V, FVec3& W, FReal LinearDamp, FReal AngularDamp, T Dt) const
+		{
+			const FReal LinearMultiplier = FMath::Max(FReal(0), FReal(1) - (LinearDamp * Dt));
+			V *= LinearMultiplier;
+
+			const FReal AngularMultiplier = FMath::Max(FReal(0), FReal(1) - (AngularDamp * Dt));
+			W *= AngularMultiplier;
+		}
+
+		inline void Apply(TDynamicParticles<T, d>& InParticles, const T Dt, const int32 Index) const override //-V762
+		{
+			ensure(false);
+		}
+
+		inline void Apply(TRigidParticles<T, d>& InParticles, const T Dt, const int32 Index) const override //-V762
+		{
+			ApplyHelper(InParticles.V(Index), InParticles.W(Index), InParticles.LinearDamping(Index), InParticles.AngularDamping(Index), Dt);
+		}
+
+		inline void Apply(TTransientPBDRigidParticleHandle<T, d>& Particle, const T Dt) const override //-V762
+		{
+			ApplyHelper(Particle.V(), Particle.W(), Particle.LinearDamping(), Particle.AngularDamping(), Dt);
+		}
+	};
+}
diff --git a/Engine/Source/Runtime/Experimental/Chaos/Public/Chaos/RigidParticles.h b/Engine/Source/Runtime/Experimental/Chaos/Public/Chaos/RigidParticles.h
--- a/Engine/Source/Runtime/Experimental/Chaos/Public/Chaos/RigidParticles.h
+++ b/Engine/Source/Runtime/Experimental/Chaos/Public/Chaos/RigidParticles.h
@@ -58,6 +58,8 @@ class TRigidParticles : public TKinematicGeometryParticles<T, d>
 		TArrayCollection::AddArray(&MInvI);
 		TArrayCollection::AddArray(&MM);
 		TArrayCollection::AddArray(&MInvM);
+		TArrayCollection::AddArray(&MLinearDamping);
+		TArrayCollection::AddArray(&MAngularDamping);
 		TArrayCollection::AddArray(&MCollisionParticles);
 		TArrayCollection::AddArray(&MCollisionGroup);
 		TArrayCollection::AddArray(&MDisabled);
@@ -77,6 +79,8 @@ class TRigidParticles : public TKinematicGeometryParticles<T, d>
 		TArrayCollection::AddArray(&MInvI);
 		TArrayCollection::AddArray(&MM);
 		TArrayCollection::AddArray(&MInvM);
+		TArrayCollection::AddArray(&MLinearDamping);
+		TArrayCollection::AddArray(&MAngularDamping);
 		TArrayCollection::AddArray(&MCollisionParticles);
 		TArrayCollection::AddArray(&MCollisionGroup);
 		TArrayCollection::AddArray(&MDisabled);
@@ -112,7 +116,15 @@ class TRigidParticles : public TKinematicGeometryParticles<T, d>
 	FORCEINLINE const T InvM(const int32 Index) const { return MInvM[Index]; }
 	FORCEINLINE T& InvM(const int32 Index) { return MInvM[Index]; }
 
+const T& LinearDamping(const int32 index) const { return MLinearDamping[index]; }
+	T& LinearDamping(const int32 index) { return MLinearDamping[index]; }
+
+	const T& AngularDamping(const int32 index) const { return MAngularDamping[index]; }
+	T& AngularDamping(const int32 index) { return MAngularDamping[index]; }
+
+	int32 CollisionParticlesSize(int32 Index) const { return MCollisionParticles[Index] == nullptr ? 0 : MCollisionParticles[Index]->Size(); }
 	FORCEINLINE int32 CollisionParticlesSize(int32 Index) const { return MCollisionParticles[Index] == nullptr ? 0 : MCollisionParticles[Index]->Size(); }
+
 	void CHAOS_API CollisionParticlesInitIfNeeded(const int32 Index);
 	void CHAOS_API SetCollisionParticles(const int32 Index, TParticles<T, d>&& Particles);
 	
@@ -158,6 +170,13 @@ class TRigidParticles : public TKinematicGeometryParticles<T, d>
 	{
 		TKinematicGeometryParticles<T,d>::Serialize(Ar);
 		Ar << MF << MT << MExternalForce << MExternalTorque << MI << MInvI << MM << MInvM;
+
+		Ar.UsingCustomVersion(FExternalPhysicsCustomObjectVersion::GUID);
+		if (Ar.CustomVer(FExternalPhysicsCustomObjectVersion::GUID) >= FExternalPhysicsCustomObjectVersion::AddDampingToRigids)
+		{
+			Ar << MLinearDamping << MAngularDamping;
+		}
+
 		Ar << MCollisionParticles << MCollisionGroup << MIsland << MDisabled << MObjectState;
 	}
 
@@ -170,6 +189,8 @@ class TRigidParticles : public TKinematicGeometryParticles<T, d>
 	TArrayCollectionArray<PMatrix<T, d, d>> MInvI;
 	TArrayCollectionArray<T> MM;
 	TArrayCollectionArray<T> MInvM;
+	TArrayCollectionArray<T> MLinearDamping;
+	TArrayCollectionArray<T> MAngularDamping;
 	TArrayCollectionArray<TUniquePtr<TBVHParticles<T, d>>> MCollisionParticles;
 	TArrayCollectionArray<int32> MCollisionGroup;
 	TArrayCollectionArray<int32> MIsland;
diff --git a/Engine/Source/Runtime/Experimental/ChaosSolvers/Private/PhysicsProxy/SingleParticlePhysicsProxy.cpp b/Engine/Source/Runtime/Experimental/ChaosSolvers/Private/PhysicsProxy/SingleParticlePhysicsProxy.cpp
--- a/Engine/Source/Runtime/Experimental/ChaosSolvers/Private/PhysicsProxy/SingleParticlePhysicsProxy.cpp
+++ b/Engine/Source/Runtime/Experimental/ChaosSolvers/Private/PhysicsProxy/SingleParticlePhysicsProxy.cpp
@@ -227,6 +227,8 @@ void FSingleParticlePhysicsProxy<Chaos::TPBDRigidParticle<float, 3>>::PushToPhys
 		RigidHandle->SetInvM(Data->MInvM);
 		RigidHandle->SetI(Data->MI);
 		RigidHandle->SetInvI(Data->MInvI);
+		RigidHandle->SetLinearDamping(Data->MLinearDamping);
+		RigidHandle->SetAngularDamping(Data->MAngularDamping);
 
 		if (Data->DirtyFlags.IsDirty(Chaos::EParticleFlags::ExternalForce))
 		{
