Fix duplicate function introduction
#rnx

diff --git a/Engine/Source/Runtime/Experimental/Chaos/Public/Chaos/RigidParticles.h b/Engine/Source/Runtime/Experimental/Chaos/Public/Chaos/RigidParticles.h
--- a/Engine/Source/Runtime/Experimental/Chaos/Public/Chaos/RigidParticles.h
+++ b/Engine/Source/Runtime/Experimental/Chaos/Public/Chaos/RigidParticles.h
@@ -116,13 +116,12 @@ class TRigidParticles : public TKinematicGeometryParticles<T, d>
 	FORCEINLINE const T InvM(const int32 Index) const { return MInvM[Index]; }
 	FORCEINLINE T& InvM(const int32 Index) { return MInvM[Index]; }
 
-const T& LinearDamping(const int32 index) const { return MLinearDamping[index]; }
+	const T& LinearDamping(const int32 index) const { return MLinearDamping[index]; }
 	T& LinearDamping(const int32 index) { return MLinearDamping[index]; }
 
 	const T& AngularDamping(const int32 index) const { return MAngularDamping[index]; }
 	T& AngularDamping(const int32 index) { return MAngularDamping[index]; }
 
-	int32 CollisionParticlesSize(int32 Index) const { return MCollisionParticles[Index] == nullptr ? 0 : MCollisionParticles[Index]->Size(); }
 	FORCEINLINE int32 CollisionParticlesSize(int32 Index) const { return MCollisionParticles[Index] == nullptr ? 0 : MCollisionParticles[Index]->Size(); }
 
 	void CHAOS_API CollisionParticlesInitIfNeeded(const int32 Index);
