Remove second instantiation as a fix for mac visibility warning.

diff --git a/Engine/Source/Runtime/Experimental/Chaos/Private/Chaos/ParticleHandle.cpp b/Engine/Source/Runtime/Experimental/Chaos/Private/Chaos/ParticleHandle.cpp
--- a/Engine/Source/Runtime/Experimental/Chaos/Private/Chaos/ParticleHandle.cpp
+++ b/Engine/Source/Runtime/Experimental/Chaos/Private/Chaos/ParticleHandle.cpp
@@ -16,7 +16,6 @@ namespace Chaos
 	template class CHAOS_API TKinematicGeometryParticle<float, 3>;
 
 	template class CHAOS_API TPBDRigidParticleData<float, 3>;
-	template class CHAOS_API TPBDRigidParticle<float, 3>;
 
 	template <>
 	void Chaos::TGeometryParticle<float, 3>::MarkDirty(const EParticleFlags DirtyBits, bool bInvalidate )
