Chaos - Joint Cholesky solver math optimizations and unit tests

diff --git a/Engine/Source/Runtime/Engine/Private/PhysicsEngine/ImmediatePhysics/ImmediatePhysicsChaos/ImmediatePhysicsJointHandle_Chaos.cpp b/Engine/Source/Runtime/Engine/Private/PhysicsEngine/ImmediatePhysics/ImmediatePhysicsChaos/ImmediatePhysicsJointHandle_Chaos.cpp
--- a/Engine/Source/Runtime/Engine/Private/PhysicsEngine/ImmediatePhysics/ImmediatePhysicsChaos/ImmediatePhysicsJointHandle_Chaos.cpp
+++ b/Engine/Source/Runtime/Engine/Private/PhysicsEngine/ImmediatePhysics/ImmediatePhysicsChaos/ImmediatePhysicsJointHandle_Chaos.cpp
@@ -154,6 +154,8 @@ namespace ImmediatePhysics_Chaos
 
 		ConstraintSettings.Motion.AngularDriveStiffness = ConvertDriveAngularStiffness(ConstraintInstance->ProfileInstance.AngularDrive.TwistDrive.Stiffness, Actor1->GetInverseInertia(), Actor2->GetInverseInertia());
 
+		ConstraintSettings.Motion.Sanitize();
+
 		ConstraintHandle = Constraints->AddConstraint({ Actor1->ParticleHandle, Actor2->ParticleHandle }, ConstraintSettings);
 	}
 
diff --git a/Engine/Source/Runtime/Experimental/Chaos/Private/Chaos/PBDJointConstraintSolver.cpp b/Engine/Source/Runtime/Experimental/Chaos/Private/Chaos/PBDJointConstraintSolver.cpp
--- a/Engine/Source/Runtime/Experimental/Chaos/Private/Chaos/PBDJointConstraintSolver.cpp
+++ b/Engine/Source/Runtime/Experimental/Chaos/Private/Chaos/PBDJointConstraintSolver.cpp
@@ -21,6 +21,13 @@ namespace Chaos
 	#define CHAOS_JOINTSOLVER_SCOPE_CYCLE_STAT(X)
 #endif
 
+	void FJointConstraintSolver::DecomposeSwingTwistLocal(const FRotation3& R0, const FRotation3& R1, FRotation3& R01Twist, FRotation3& R01Swing)
+	{
+		const FRotation3 R01 = R0.Inverse() * R1;
+		R01.ToSwingTwistX(FJointConstants::TwistAxis(), R01Swing, R01Twist);
+	}
+
+
 	void FJointConstraintSolver::InitConstraints(
 		const FReal Dt,
 		const FPBDJointSolverSettings& SolverSettings,
@@ -90,21 +97,25 @@ namespace Chaos
 		BuildJacobianAndResidual(SolverSettings, JointSettings, J0, J1, C);
 
 		// InvM(6x6) = inverse mass matrix
-		const FMassMatrix InvM0 = FMassMatrix::Make(InvMs[0], Utilities::ComputeWorldSpaceInertia(Qs[0], InvILs[0]));
-		const FMassMatrix InvM1 = FMassMatrix::Make(InvMs[1], Utilities::ComputeWorldSpaceInertia(Qs[1], InvILs[1]));
+		const FMassMatrix InvM0 = FMassMatrix::Make(InvMs[0], Qs[0], InvILs[0]);
+		const FMassMatrix InvM1 = FMassMatrix::Make(InvMs[1], Qs[1], InvILs[1]);
 
-		// Joint-space mass: F(NxN) = J(Nx6).I(6x6).Jt(6xN)
-		const FDenseMatrix66 F0 = FDenseMatrix66::MultiplyAB(J0, FDenseMatrix66::MultiplyABt(InvM0, J0));
-		const FDenseMatrix66 F1 = FDenseMatrix66::MultiplyAB(J1, FDenseMatrix66::MultiplyABt(InvM1, J1));
-		const FDenseMatrix66 F = FDenseMatrix66::Add(F0, F1);
+		// IJt(6xN) = I(6x6).Jt(6xN)
+		const FDenseMatrix66 IJt0 = FDenseMatrix66::MultiplyABt(InvM0, J0);
+		const FDenseMatrix66 IJt1 = FDenseMatrix66::MultiplyABt(InvM1, J1);
+
+		// Joint-space mass: F(NxN) = J(Nx6).I(6x6).Jt(6xN) = J(Nx6).IJt(6xN)
+		// NOTE: Result is symmetric
+		const FDenseMatrix66 F0 = FDenseMatrix66::MultiplyAB_Symmetric(J0, IJt0);
+		const FDenseMatrix66 F = FDenseMatrix66::MultiplyBCAddA_Symmetric(F0, J1, IJt1);
 
 		// Joint-space correction: L(Nx1) = [1/F](NxN).C(Nx1)
 		FDenseMatrix61 L;
 		if (FDenseMatrixSolver::SolvePositiveDefinite(F, C, L))
 		{
-			// World-space correction: D(6x1) = I(6x6).Jt(6xN).L(Nx1) = I(6x6).JtL(6x1)
-			const FDenseMatrix61 D0 = FDenseMatrix61::MultiplyAB(InvM0, FDenseMatrix61::MultiplyAtB(J0, L));
-			const FDenseMatrix61 D1 = FDenseMatrix61::MultiplyAB(InvM1, FDenseMatrix61::MultiplyAtB(J1, L));
+			// World-space correction: D(6x1) = I(6x6).Jt(6xN).L(Nx1) = IJt(6xN).L(Nx1)
+			const FDenseMatrix61 D0 = FDenseMatrix61::MultiplyAB(IJt0, L);
+			const FDenseMatrix61 D1 = FDenseMatrix61::MultiplyAB(IJt1, L);
 
 			// Extract world-space position correction
 			Ps[0] = Ps[0] + FVec3(Stiffness * D0.At(0, 0), Stiffness * D0.At(1, 0), Stiffness * D0.At(2, 0));
@@ -167,20 +178,16 @@ namespace Chaos
 		//Result[2] = V1[0] * V2[1] - V1[1] * V2[0];
 
 		const FVec3 XP0 = Xs[0] - Ps[0];
-		J0.SetRowAt(RowIndex + 0, 0, FVec3(1, 0, 0));
-		J0.SetRowAt(RowIndex + 1, 0, FVec3(0, 1, 0));
-		J0.SetRowAt(RowIndex + 2, 0, FVec3(0, 0, 1));
-		J0.SetRowAt(RowIndex + 0, 3, FVec3((FReal)0, XP0[2], -XP0[1]));	// -(1,0,0) x XP0
-		J0.SetRowAt(RowIndex + 1, 3, FVec3(-XP0[2], (FReal)0, XP0[0]));	// -(0,1,0) x XP0
-		J0.SetRowAt(RowIndex + 2, 3, FVec3(XP0[1], -XP0[0], (FReal)0));	// -(0,0,1) x XP0
+		J0.SetBlockAtDiagonal33(RowIndex, 0, (FReal)1, (FReal)0);
+		J0.SetRowAt(RowIndex + 0, 3, (FReal)0, XP0[2], -XP0[1]);	// -(1,0,0) x XP0
+		J0.SetRowAt(RowIndex + 1, 3, -XP0[2], (FReal)0, XP0[0]);	// -(0,1,0) x XP0
+		J0.SetRowAt(RowIndex + 2, 3, XP0[1], -XP0[0], (FReal)0);	// -(0,0,1) x XP0
 
 		const FVec3 XP1 = Xs[1] - Ps[1];
-		J1.SetRowAt(RowIndex + 0, 0, -FVec3(1, 0, 0));
-		J1.SetRowAt(RowIndex + 1, 0, -FVec3(0, 1, 0));
-		J1.SetRowAt(RowIndex + 2, 0, -FVec3(0, 0, 1));
-		J1.SetRowAt(RowIndex + 0, 3, FVec3((FReal)0, -XP1[2], XP1[1]));	// (1,0,0) x XP1
-		J1.SetRowAt(RowIndex + 1, 3, FVec3(XP1[2], (FReal)0, -XP1[0]));	// (0,1,0) x XP1
-		J1.SetRowAt(RowIndex + 2, 3, FVec3(-XP1[1], XP1[0], (FReal)0));	// (0,0,1) x XP1
+		J1.SetBlockAtDiagonal33(RowIndex, 0, (FReal)-1, (FReal)0);
+		J1.SetRowAt(RowIndex + 0, 3, (FReal)0, -XP1[2], XP1[1]);	// (1,0,0) x XP1
+		J1.SetRowAt(RowIndex + 1, 3, XP1[2], (FReal)0, -XP1[0]);	// (0,1,0) x XP1
+		J1.SetRowAt(RowIndex + 2, 3, -XP1[1], XP1[0], (FReal)0);	// (0,0,1) x XP1
 
 		const FVec3 ConstraintSeparation = Xs[1] - Xs[0];
 		C.SetAt(RowIndex + 0, 0, ConstraintSeparation[0]);
@@ -319,40 +326,26 @@ namespace Chaos
 	void FJointConstraintSolver::AddAngularConstraints_Twist(
 		const FPBDJointSolverSettings& SolverSettings,
 		const FPBDJointSettings& JointSettings,
+		const FRotation3& R01Twist,
+		const FRotation3& R01Swing,
 		FDenseMatrix66& J0,
 		FDenseMatrix66& J1,
 		FDenseMatrix61& C)
 	{
-		// Calculate the Twist Axis and Angle for each body
-		const FRotation3 R01 = Rs[0].Inverse() * Rs[1];
-		FRotation3 R01Twist, R01Swing;
-		R01.ToSwingTwist(FJointConstants::TwistAxis(), R01Swing, R01Twist);
-		R01Swing = R01Swing.GetNormalized();
-		R01Twist = R01Twist.GetNormalized();
-
-		FVec3 TwistAxis01;
-		FReal TwistAngle;
-		R01Twist.ToAxisAndAngleSafe(TwistAxis01, TwistAngle, FJointConstants::TwistAxis(), SolverSettings.SwingTwistAngleTolerance);
+		FVec3 TwistAxis01 = FJointConstants::TwistAxis();
+		FReal TwistAngle = R01Twist.GetAngle();
+		//R01Twist.ToAxisAndAngleSafe(TwistAxis01, TwistAngle, FJointConstants::TwistAxis(), SolverSettings.SwingTwistAngleTolerance);
 		if (TwistAngle > PI)
 		{
 			TwistAngle = TwistAngle - (FReal)2 * PI;
 		}
-		if (FVec3::DotProduct(TwistAxis01, FJointConstants::TwistAxis()) < 0)
+		//if (FVec3::DotProduct(TwistAxis01, FJointConstants::TwistAxis()) < 0)
+		if (R01Twist.X < 0)
 		{
-			TwistAxis01 = -TwistAxis01;
 			TwistAngle = -TwistAngle;
 		}
 
-		FReal TwistAngleMax = FLT_MAX;
-		if (JointSettings.Motion.AngularMotionTypes[(int32)EJointAngularConstraintIndex::Twist] == EJointMotionType::Limited)
-		{
-			TwistAngleMax = JointSettings.Motion.AngularLimits[(int32)EJointAngularConstraintIndex::Twist];
-		}
-		else if (JointSettings.Motion.AngularMotionTypes[(int32)EJointAngularConstraintIndex::Twist] == EJointMotionType::Locked)
-		{
-			TwistAngleMax = 0;
-		}
-
+		const FReal TwistAngleMax = JointSettings.Motion.AngularLimits[(int32)EJointAngularConstraintIndex::Twist];
 		bool bConstraintActive = (TwistAngle <= -TwistAngleMax) || (TwistAngle >= TwistAngleMax);
 		if (bConstraintActive)
 		{
@@ -366,14 +359,13 @@ namespace Chaos
 			J1.AddRows(1);
 			C.AddRows(1);
 
-			J0.SetRowAt(RowIndex + 0, 0, FVec3(0, 0, 0));
+			J0.SetRowAt(RowIndex + 0, 0, 0, 0, 0);
 			J0.SetRowAt(RowIndex + 0, 3, Axis0);
 
-			J1.SetRowAt(RowIndex + 0, 0, FVec3(0, 0, 0));
+			J1.SetRowAt(RowIndex + 0, 0, 0, 0, 0);
 			J1.SetRowAt(RowIndex + 0, 3, -Axis1);
 
 			C.SetAt(RowIndex, 0, (TwistAngle >= TwistAngleMax) ? TwistAngle - TwistAngleMax : TwistAngle + TwistAngleMax);
-
 		}
 	}
 
@@ -381,17 +373,12 @@ namespace Chaos
 	void FJointConstraintSolver::AddAngularConstraints_Cone(
 		const FPBDJointSolverSettings& SolverSettings,
 		const FPBDJointSettings& JointSettings,
+		const FRotation3& R01Twist,
+		const FRotation3& R01Swing,
 		FDenseMatrix66& J0,
 		FDenseMatrix66& J1,
 		FDenseMatrix61& C)
 	{
-		// Calculate Swing axis for each body
-		const FRotation3 R01 = Rs[0].Inverse() * Rs[1];
-		FRotation3 R01Twist, R01Swing;
-		R01.ToSwingTwist(FJointConstants::TwistAxis(), R01Swing, R01Twist);
-		R01Swing = R01Swing.GetNormalized();
-		R01Twist = R01Twist.GetNormalized();
-
 		FVec3 SwingAxis01;
 		FReal SwingAngle;
 		R01Swing.ToAxisAndAngleSafe(SwingAxis01, SwingAngle, FJointConstants::Swing1Axis(), SolverSettings.SwingTwistAngleTolerance);
@@ -429,10 +416,10 @@ namespace Chaos
 			J1.AddRows(1);
 			C.AddRows(1);
 
-			J0.SetRowAt(RowIndex + 0, 0, FVec3(0, 0, 0));
+			J0.SetRowAt(RowIndex + 0, 0, 0, 0, 0);
 			J0.SetRowAt(RowIndex + 0, 3, Axis);
 
-			J1.SetRowAt(RowIndex + 0, 0, FVec3(0, 0, 0));
+			J1.SetRowAt(RowIndex + 0, 0, 0, 0, 0);
 			J1.SetRowAt(RowIndex + 0, 3, -Axis);
 
 			C.SetAt(RowIndex, 0, (SwingAngle >= SwingAngleMax) ? SwingAngle - SwingAngleMax : SwingAngle + SwingAngleMax);
@@ -445,17 +432,12 @@ namespace Chaos
 		const FPBDJointSettings& JointSettings,
 		const EJointAngularConstraintIndex SwingConstraintIndex,
 		const EJointAngularAxisIndex SwingAxisIndex,
+		const FRotation3& R01Twist,
+		const FRotation3& R01Swing,
 		FDenseMatrix66& J0,
 		FDenseMatrix66& J1,
 		FDenseMatrix61& C)
 	{
-		// Calculate the swing axis for each body
-		const FRotation3 R01 = Rs[0].Inverse() * Rs[1];
-		FRotation3 R01Twist, R01Swing;
-		R01.ToSwingTwist(FJointConstants::TwistAxis(), R01Swing, R01Twist);
-		R01Swing = R01Swing.GetNormalized();
-		R01Twist = R01Twist.GetNormalized();
-
 		FVec3 TwistAxis01;
 		FReal TwistAngle;
 		R01Twist.ToAxisAndAngleSafe(TwistAxis01, TwistAngle, FJointConstants::TwistAxis(), SolverSettings.SwingTwistAngleTolerance);
@@ -485,17 +467,7 @@ namespace Chaos
 				SwingAngle = (FReal)PI - SwingAngle;
 			}
 
-			FReal SwingAngleMax = FLT_MAX;
-			if (JointSettings.Motion.AngularMotionTypes[(int32)SwingConstraintIndex] == EJointMotionType::Limited)
-			{
-				FReal SwingLimit = JointSettings.Motion.AngularLimits[(int32)SwingConstraintIndex];
-				SwingAngleMax = SwingLimit;
-			}
-			else if (JointSettings.Motion.AngularMotionTypes[(int32)SwingConstraintIndex] == EJointMotionType::Locked)
-			{
-				SwingAngleMax = 0;
-			}
-
+			FReal SwingAngleMax = JointSettings.Motion.AngularLimits[(int32)SwingConstraintIndex];
 			bool bConstraintActive = (SwingAngle <= -SwingAngleMax) || (SwingAngle >= SwingAngleMax);
 			if (bConstraintActive)
 			{
@@ -508,10 +480,10 @@ namespace Chaos
 				J1.AddRows(1);
 				C.AddRows(1);
 
-				J0.SetRowAt(RowIndex + 0, 0, FVec3(0, 0, 0));
+				J0.SetRowAt(RowIndex + 0, 0, 0, 0, 0);
 				J0.SetRowAt(RowIndex + 0, 3, Axis);
 
-				J1.SetRowAt(RowIndex + 0, 0, FVec3(0, 0, 0));
+				J1.SetRowAt(RowIndex + 0, 0, 0, 0, 0);
 				J1.SetRowAt(RowIndex + 0, 3, -Axis);
 
 				C.SetAt(RowIndex, 0, (SwingAngle >= SwingAngleMax) ? SwingAngle - SwingAngleMax : SwingAngle + SwingAngleMax);
@@ -581,31 +553,38 @@ namespace Chaos
 		EJointMotionType Swing1Motion = JointSettings.Motion.AngularMotionTypes[(int32)EJointAngularConstraintIndex::Swing1];
 		EJointMotionType Swing2Motion = JointSettings.Motion.AngularMotionTypes[(int32)EJointAngularConstraintIndex::Swing2];
 
-		// Apply twist constraint
-		if (SolverSettings.bEnableTwistLimits)
-		{
-			if (TwistMotion != EJointMotionType::Free)
-			{
-				AddAngularConstraints_Twist(SolverSettings, JointSettings, J0, J1, C);
-			}
-		}
+		bool bAddTwist = SolverSettings.bEnableTwistLimits && (TwistMotion != EJointMotionType::Free);
+		bool bAddConeOrSwing = SolverSettings.bEnableSwingLimits && ((Swing1Motion != EJointMotionType::Free) || (Swing2Motion != EJointMotionType::Free));
 
-		// Apply swing constraints
-		if (SolverSettings.bEnableSwingLimits)
+		if (bAddTwist || bAddConeOrSwing)
 		{
-			if ((Swing1Motion == EJointMotionType::Limited) && (Swing2Motion == EJointMotionType::Limited))
+			// Calculate axes for each body
+			FRotation3 R01Twist, R01Swing;
+			DecomposeSwingTwistLocal(Rs[0], Rs[1], R01Twist, R01Swing);
+
+			// Add twist constraint
+			if (bAddTwist)
 			{
-				AddAngularConstraints_Cone(SolverSettings, JointSettings, J0, J1, C);
+				AddAngularConstraints_Twist(SolverSettings, JointSettings, R01Twist, R01Swing, J0, J1, C);
 			}
-			else
+
+			// Add swing constraints
+			if (bAddConeOrSwing)
 			{
-				if (Swing1Motion != EJointMotionType::Free)
+				if ((Swing1Motion == EJointMotionType::Limited) && (Swing2Motion == EJointMotionType::Limited))
 				{
-					AddAngularConstraints_Swing(SolverSettings, JointSettings, EJointAngularConstraintIndex::Swing1, EJointAngularAxisIndex::Swing1, J0, J1, C);
+					AddAngularConstraints_Cone(SolverSettings, JointSettings, R01Twist, R01Swing, J0, J1, C);
 				}
-				if (Swing2Motion != EJointMotionType::Free)
+				else
 				{
-					AddAngularConstraints_Swing(SolverSettings, JointSettings, EJointAngularConstraintIndex::Swing2, EJointAngularAxisIndex::Swing2, J0, J1, C);
+					if (Swing1Motion != EJointMotionType::Free)
+					{
+						AddAngularConstraints_Swing(SolverSettings, JointSettings, EJointAngularConstraintIndex::Swing1, EJointAngularAxisIndex::Swing1, R01Twist, R01Swing, J0, J1, C);
+					}
+					if (Swing2Motion != EJointMotionType::Free)
+					{
+						AddAngularConstraints_Swing(SolverSettings, JointSettings, EJointAngularConstraintIndex::Swing2, EJointAngularAxisIndex::Swing2, R01Twist, R01Swing, J0, J1, C);
+					}
 				}
 			}
 		}
diff --git a/Engine/Source/Runtime/Experimental/Chaos/Private/Chaos/PBDJointConstraints.cpp b/Engine/Source/Runtime/Experimental/Chaos/Private/Chaos/PBDJointConstraints.cpp
--- a/Engine/Source/Runtime/Experimental/Chaos/Private/Chaos/PBDJointConstraints.cpp
+++ b/Engine/Source/Runtime/Experimental/Chaos/Private/Chaos/PBDJointConstraints.cpp
@@ -20,6 +20,8 @@ namespace Chaos
 	DECLARE_CYCLE_STAT(TEXT("TPBDJointConstraints::ApplyPushOut"), STAT_Joints_ApplyPushOut, STATGROUP_Chaos);
 	DECLARE_CYCLE_STAT(TEXT("TPBDJointConstraints::Drives"), STAT_Joints_Drives, STATGROUP_Chaos);
 	DECLARE_CYCLE_STAT(TEXT("TPBDJointConstraints::Solve"), STAT_Joints_Solve, STATGROUP_Chaos);
+	DECLARE_CYCLE_STAT(TEXT("TPBDJointConstraints::SolveCholesky"), STAT_Joints_Solve_Cholesky, STATGROUP_Chaos);
+	DECLARE_CYCLE_STAT(TEXT("TPBDJointConstraints::SolveGaussSeidel"), STAT_Joints_Solve_GaussSeidel, STATGROUP_Chaos);
 	DECLARE_CYCLE_STAT(TEXT("TPBDJointConstraints::Project"), STAT_Joints_Project, STATGROUP_Chaos);
 
 	bool ChaosJoint_UseCholeskySolver = true;
@@ -156,6 +158,28 @@ namespace Chaos
 	{
 	}
 
+	void FPBDJointMotionSettings::Sanitize()
+	{
+		// Reset limits if they won;t be used (means we don't have to check if limited/locked in a few cases).
+		// A side effect: if we enable a constraint, we need to reset the value of the limit.
+		if ((LinearMotionTypes[0] != EJointMotionType::Limited) && (LinearMotionTypes[1] != EJointMotionType::Limited) && (LinearMotionTypes[2] != EJointMotionType::Limited))
+		{
+			LinearLimit = 0;
+		}
+		if (AngularMotionTypes[(int32)EJointAngularConstraintIndex::Twist] != EJointMotionType::Limited)
+		{
+			AngularLimits[(int32)EJointAngularConstraintIndex::Twist] = 0;
+		}
+		if (AngularMotionTypes[(int32)EJointAngularConstraintIndex::Swing1] != EJointMotionType::Limited)
+		{
+			AngularLimits[(int32)EJointAngularConstraintIndex::Swing1] = 0;
+		}
+		if (AngularMotionTypes[(int32)EJointAngularConstraintIndex::Swing2] != EJointMotionType::Limited)
+		{
+			AngularLimits[(int32)EJointAngularConstraintIndex::Swing2] = 0;
+		}
+	}
+
 	
 	FPBDJointSettings::FPBDJointSettings()
 		: ConstraintFrames({ FTransform::Identity, FTransform::Identity })
@@ -224,7 +248,7 @@ namespace Chaos
 		Settings = InSettings;
 	}
 
-	
+
 	int32 FPBDJointConstraints::NumConstraints() const
 	{
 		return ConstraintParticles.Num();
@@ -404,37 +428,18 @@ namespace Chaos
 		// Apply joint drives
 		if (Settings.DrivesPhase == EJointSolverPhase::Apply)
 		{
-			SCOPE_CYCLE_COUNTER(STAT_Joints_Drives);
-			if (It == 0)
-			{
-				for (int32 ConstraintIndex = 0; ConstraintIndex < NumConstraints(); ++ConstraintIndex)
-				{
-					ApplyDrives(Dt, ConstraintIndex, 1, 0, 1);
-				}
-			}
+			ApplyDrives(Dt, It, NumIts);
 		}
 
 		// Solve for joint position or velocity, depending on settings
 		if (Settings.ApplyPairIterations > 0)
 		{
-			SCOPE_CYCLE_COUNTER(STAT_Joints_Solve);
-			for (int32 ConstraintIndex = 0; ConstraintIndex < NumConstraints(); ++ConstraintIndex)
-			{
-				if (Settings.bEnableVelocitySolve)
-				{
-					SolveVelocity(Dt, ConstraintIndex, Settings.ApplyPairIterations, It, NumIts);
-				}
-				else
-				{
-					SolvePosition(Dt, ConstraintIndex, Settings.ApplyPairIterations, It, NumIts);
-				}
-			}
+			SolvePosition(Dt, Settings.ApplyPairIterations, It, NumIts);
 		}
 
 		// Correct remaining errors after last call to Solve if enabled in this phase
 		if (Settings.ProjectionPhase == EJointSolverPhase::Apply)
 		{
-			SCOPE_CYCLE_COUNTER(STAT_Joints_Project);
 			int32 ProjectionIt = NumIts - 1;
 			if (It == ProjectionIt)
 			{
@@ -458,31 +463,18 @@ namespace Chaos
 		// Apply joint drives
 		if (Settings.DrivesPhase == EJointSolverPhase::ApplyPushOut)
 		{
-			SCOPE_CYCLE_COUNTER(STAT_Joints_Drives);
-			if (It == 0)
-			{
-				for (int32 ConstraintIndex = 0; ConstraintIndex < NumConstraints(); ++ConstraintIndex)
-				{
-					ApplyDrives(Dt, ConstraintIndex, 1, 0, 1);
-				}
-			}
+			ApplyDrives(Dt, It, NumIts);
 		}
 
 		// Solve for positions
 		if (Settings.ApplyPushOutPairIterations > 0)
 		{
-			SCOPE_CYCLE_COUNTER(STAT_Joints_Solve);
-
-			for (int32 ConstraintIndex = 0; ConstraintIndex < NumConstraints(); ++ConstraintIndex)
-			{
-				SolvePosition(Dt, ConstraintIndex, Settings.ApplyPushOutPairIterations, It, NumIts);
-			}
+			SolvePosition(Dt, Settings.ApplyPushOutPairIterations, It, NumIts);
 		}
 
 		// Correct remaining errors after the last call to Solve (which depends on if PositionSolve is enabled in ApplyPushOut)
 		if (Settings.ProjectionPhase == EJointSolverPhase::ApplyPushOut)
 		{
-			SCOPE_CYCLE_COUNTER(STAT_Joints_Project);
 			int32 ProjectionIt = (Settings.ApplyPushOutPairIterations > 0) ? NumIts - 1 : 0;
 			if (It == ProjectionIt)
 			{
@@ -493,8 +485,42 @@ namespace Chaos
 		return bNeedsAnotherIteration;
 	}
 
+	void FPBDJointConstraints::ApplyDrives(const FReal Dt, const int32 It, const int32 NumIts)
+	{
+		if (It == 0)
+		{
+			SCOPE_CYCLE_COUNTER(STAT_Joints_Drives);
+			for (int32 ConstraintIndex = 0; ConstraintIndex < NumConstraints(); ++ConstraintIndex)
+			{
+				ApplyDrives(Dt, ConstraintIndex, 0, 1);
+			}
+		}
+	}
+
+	void FPBDJointConstraints::SolvePosition(const FReal Dt, const int32 NumPairIts, const int32 It, const int32 NumIts)
+	{
+		if (ChaosJoint_UseCholeskySolver)
+		{
+			SCOPE_CYCLE_COUNTER(STAT_Joints_Solve_Cholesky);
+			for (int32 ConstraintIndex = 0; ConstraintIndex < NumConstraints(); ++ConstraintIndex)
+			{
+				SolvePosition_Cholesky(Dt, ConstraintIndex, NumPairIts, It, NumIts);
+			}
+		}
+		else
+		{
+			SCOPE_CYCLE_COUNTER(STAT_Joints_Solve_GaussSeidel);
+			for (int32 ConstraintIndex = 0; ConstraintIndex < NumConstraints(); ++ConstraintIndex)
+			{
+				SolvePosition_GaussSiedel(Dt, ConstraintIndex, NumPairIts, It, NumIts);
+			}
+		}
+	}
+
 	void FPBDJointConstraints::ApplyProjection(const FReal Dt)
 	{
+		SCOPE_CYCLE_COUNTER(STAT_Joints_Project);
+
 		for (int32 ConstraintIndex = 0; ConstraintIndex < NumConstraints(); ++ConstraintIndex)
 		{
 			ProjectPosition(Dt, ConstraintIndex, 0, 1);
@@ -532,7 +558,7 @@ namespace Chaos
 			{
 				for (FConstraintContainerHandle* ConstraintHandle : SortedConstraintHandles)
 				{
-					ApplyDrives(Dt, ConstraintHandle->GetConstraintIndex(), 1, 0, 1);
+					ApplyDrives(Dt, ConstraintHandle->GetConstraintIndex(), 0, 1);
 				}
 			}
 		}
@@ -594,7 +620,7 @@ namespace Chaos
 			{
 				for (FConstraintContainerHandle* ConstraintHandle : SortedConstraintHandles)
 				{
-					ApplyDrives(Dt, ConstraintHandle->GetConstraintIndex(), 1, 0, 1);
+					ApplyDrives(Dt, ConstraintHandle->GetConstraintIndex(), 0, 1);
 				}
 			}
 		}
@@ -659,7 +685,7 @@ namespace Chaos
 	}
 
 
-	void FPBDJointConstraints::ApplyDrives(const FReal Dt, const int32 ConstraintIndex, const int32 NumPairIts, const int32 It, const int32 NumIts)
+	void FPBDJointConstraints::ApplyDrives(const FReal Dt, const int32 ConstraintIndex, const int32 It, const int32 NumIts)
 	{
 		const FPBDJointSettings& JointSettings = ConstraintSettings[ConstraintIndex];
 
diff --git a/Engine/Source/Runtime/Experimental/Chaos/Public/Chaos/DenseMatrix.h b/Engine/Source/Runtime/Experimental/Chaos/Public/Chaos/DenseMatrix.h
--- a/Engine/Source/Runtime/Experimental/Chaos/Public/Chaos/DenseMatrix.h
+++ b/Engine/Source/Runtime/Experimental/Chaos/Public/Chaos/DenseMatrix.h
@@ -5,6 +5,7 @@
 #include "Chaos/Core.h"
 #include "Chaos/Vector.h"
 #include "Chaos/Matrix.h"
+#include "Chaos/Utilities.h"
 
 namespace Chaos
 {
@@ -37,6 +38,11 @@ namespace Chaos
 			return FMassMatrix(InM, MoveTemp(InI));
 		}
 
+		static FMassMatrix Make(const FReal InM, const FRotation3& Q, const FMatrix33& InI)
+		{
+			return FMassMatrix(InM, Q, InI);
+		}
+
 	private:
 		FMassMatrix(const FReal InM, const FMatrix33& InI)
 			: Mass(InM)
@@ -50,6 +56,12 @@ namespace Chaos
 		{
 		}
 
+		FMassMatrix(const FReal InM, const FRotation3& Q, const FMatrix33& InI)
+			: Mass(InM)
+			, Inertia(Utilities::ComputeWorldSpaceInertia(Q, InI))
+		{
+		}
+
 		FReal Mass;
 		FMatrix33 Inertia;
 	};
@@ -225,19 +237,26 @@ namespace Chaos
 		{
 			check(RowIndex + NumV < NumRows());
 			check(ColumnIndex < NumColumns());
+			FReal* Row = &At(RowIndex, ColumnIndex);
 			for (int32 II = 0; II < NumV; ++II)
 			{
-				At(RowIndex, ColumnIndex + II) = V[II];
+				*Row++ = *V++;
 			}
 		}
 
 		void SetRowAt(const int32 RowIndex, const int32 ColumnIndex, const FVec3& V)
+		{
+			SetRowAt(RowIndex, ColumnIndex, V[0], V[1], V[2]);
+		}
+
+		void SetRowAt(const int32 RowIndex, const int32 ColumnIndex, const FReal V0, const FReal V1, const FReal V2)
 		{
 			check(RowIndex + 1 <= NumRows());
 			check(ColumnIndex + 3 <= NumColumns());
-			At(RowIndex, ColumnIndex + 0) = V[0];
-			At(RowIndex, ColumnIndex + 1) = V[1];
-			At(RowIndex, ColumnIndex + 2) = V[2];
+			FReal* Row = &At(RowIndex, ColumnIndex);
+			*Row++ = V0;
+			*Row++ = V1;
+			*Row++ = V2;
 		}
 
 		/**
@@ -295,6 +314,27 @@ namespace Chaos
 			}
 		}
 
+		/**
+		 * Set the specified 3x3 block to a diagonal matrix with the specified diagonal and off-diagonal values.
+		 */
+		void SetBlockAtDiagonal33(const int32 RowOffset, const int32 ColumnOffset, const FReal VDiag, const FReal VOffDiag)
+		{
+			check(RowOffset + 3 <= NumRows());
+			check(ColumnOffset + 3 <= NumColumns());
+			FReal* Row0 = &At(RowOffset + 0, ColumnOffset);
+			*Row0++ = VDiag;
+			*Row0++ = VOffDiag;
+			*Row0++ = VOffDiag;
+			FReal* Row1 = &At(RowOffset + 1, ColumnOffset);
+			*Row1++ = VOffDiag;
+			*Row1++ = VDiag;
+			*Row1++ = VOffDiag;
+			FReal* Row2 = &At(RowOffset + 2, ColumnOffset);
+			*Row2++ = VOffDiag;
+			*Row2++ = VOffDiag;
+			*Row2++ = VDiag;
+		}
+
 		//
 		// Factory methods
 		//
@@ -438,6 +478,28 @@ namespace Chaos
 			return Result;
 		}
 
+		/**
+		 * Return C = A + B, where A and B are symetric.
+		 */
+		template<int32 T_EA, int32 T_EB>
+		static TDenseMatrix<MaxElements> Add_Symmetric(const TDenseMatrix<T_EA>& A, const TDenseMatrix<T_EB>& B)
+		{
+			// @todo(ccaulfield): optimize
+			check(A.NumColumns() == B.NumColumns());
+			check(A.NumRows() == B.NumRows());
+			TDenseMatrix<T_MAXELEMENTS> Result(A.NumRows(), A.NumColumns());
+			for (int32 IRow = 0; IRow < Result.NumRows(); ++IRow)
+			{
+				for (int32 ICol = IRow; ICol < Result.NumColumns(); ++ICol)
+				{
+					FReal V = A.At(IRow, ICol) + B.At(IRow, ICol);
+					Result.At(IRow, ICol) = V;
+					Result.At(ICol, IRow) = V;
+				}
+			}
+			return Result;
+		}
+
 		/**
 		 * Return C = A - B
 		 */
@@ -558,6 +620,61 @@ namespace Chaos
 			return Result;
 		}
 
+		/**
+		 * Return C = A x B, where the results is known to be symmetric.
+		 * /see MultiplyAtB.
+		 */
+		template<int32 T_EA, int32 T_EB>
+		static TDenseMatrix<MaxElements> MultiplyAB_Symmetric(const TDenseMatrix<T_EA>& A, const TDenseMatrix<T_EB>& B)
+		{
+			// @todo(ccaulfield): optimize
+			check(A.NumColumns() == B.NumRows());
+			TDenseMatrix<T_MAXELEMENTS> Result(A.NumRows(), B.NumColumns());
+			for (int32 IRow = 0; IRow < Result.NumRows(); ++IRow)
+			{
+				for (int32 ICol = IRow; ICol < Result.NumColumns(); ++ICol)
+				{
+					FReal V = 0;
+					for (int32 II = 0; II < A.NumColumns(); ++II)
+					{
+						V += A.At(IRow, II) * B.At(II, ICol);
+					}
+					Result.At(IRow, ICol) = V;
+					Result.At(ICol, IRow) = V;
+				}
+			}
+			return Result;
+		}
+
+		/**
+		 * Return C = A + B x C, where the A and (B x C) are known to be symmetric.
+		 * /see MultiplyAtB, Add_Symmetric.
+		 */
+		template<int32 T_EA, int32 T_EB, int32 T_EC>
+		static TDenseMatrix<MaxElements> MultiplyBCAddA_Symmetric(const TDenseMatrix<T_EA>& A, const TDenseMatrix<T_EB>& B, const TDenseMatrix<T_EC>& C)
+		{
+			// @todo(ccaulfield): optimize
+			check(B.NumColumns() == C.NumRows());
+			check(A.NumRows() == B.NumRows());
+			check(A.NumRows() == C.NumColumns());
+			TDenseMatrix<T_MAXELEMENTS> Result(A.NumRows(), A.NumColumns());
+			for (int32 IRow = 0; IRow < Result.NumRows(); ++IRow)
+			{
+				for (int32 ICol = IRow; ICol < Result.NumColumns(); ++ICol)
+				{
+					FReal V = 0;
+					for (int32 II = 0; II < B.NumColumns(); ++II)
+					{
+						V += B.At(IRow, II) * C.At(II, ICol);
+					}
+					FReal VA = A.At(IRow, ICol);
+					Result.At(IRow, ICol) = VA + V;
+					Result.At(ICol, IRow) = VA + V;
+				}
+			}
+			return Result;
+		}
+
 		/**
 		 * Return C = A x B, where A is an Nx6 matrix, and B is a 6x6 mass matrix (Mass in upper left 3x3 diagonals, Inertia in lower right 3x3).
 		 * C = |A0 A1| * |M 0| = |A0.M A1.I|
diff --git a/Engine/Source/Runtime/Experimental/Chaos/Public/Chaos/PBDJointConstraintSolver.h b/Engine/Source/Runtime/Experimental/Chaos/Public/Chaos/PBDJointConstraintSolver.h
--- a/Engine/Source/Runtime/Experimental/Chaos/Public/Chaos/PBDJointConstraintSolver.h
+++ b/Engine/Source/Runtime/Experimental/Chaos/Public/Chaos/PBDJointConstraintSolver.h
@@ -89,6 +89,8 @@ namespace Chaos
 		void AddAngularConstraints_Twist(
 			const FPBDJointSolverSettings& SolverSettings,
 			const FPBDJointSettings& JointSettings,
+			const FRotation3& R01Twist,
+			const FRotation3& R01Swing,
 			FDenseMatrix66& J0,
 			FDenseMatrix66& J1,
 			FDenseMatrix61& C);
@@ -96,6 +98,8 @@ namespace Chaos
 		void AddAngularConstraints_Cone(
 			const FPBDJointSolverSettings& SolverSettings,
 			const FPBDJointSettings& JointSettings,
+			const FRotation3& R01Twist,
+			const FRotation3& R01Swing,
 			FDenseMatrix66& J0,
 			FDenseMatrix66& J1,
 			FDenseMatrix61& C);
@@ -105,6 +109,8 @@ namespace Chaos
 			const FPBDJointSettings& JointSettings,
 			const EJointAngularConstraintIndex SwingConstraintIndex,
 			const EJointAngularAxisIndex SwingAxisIndex,
+			const FRotation3& R01Twist,
+			const FRotation3& R01Swing,
 			FDenseMatrix66& J0,
 			FDenseMatrix66& J1,
 			FDenseMatrix61& C);
@@ -130,6 +136,8 @@ namespace Chaos
 			FDenseMatrix66& J1,
 			FDenseMatrix61& C);
 
+		static void DecomposeSwingTwistLocal(const FRotation3& R0, const FRotation3& R1, FRotation3& R01Twist, FRotation3& R01Swing);
+
 		static const int32 MaxConstraintedBodies = 2;
 
 		// Local-space constraint settings
diff --git a/Engine/Source/Runtime/Experimental/Chaos/Public/Chaos/PBDJointConstraintTypes.h b/Engine/Source/Runtime/Experimental/Chaos/Public/Chaos/PBDJointConstraintTypes.h
--- a/Engine/Source/Runtime/Experimental/Chaos/Public/Chaos/PBDJointConstraintTypes.h
+++ b/Engine/Source/Runtime/Experimental/Chaos/Public/Chaos/PBDJointConstraintTypes.h
@@ -69,6 +69,8 @@ namespace Chaos
 		FPBDJointMotionSettings();
 		FPBDJointMotionSettings(const TVector<EJointMotionType, 3>& InLinearMotionTypes, const TVector<EJointMotionType, 3>& InAngularMotionTypes);
 
+		void Sanitize();
+
 		FReal Stiffness;
 		FReal LinearProjection;
 		FReal AngularProjection;
diff --git a/Engine/Source/Runtime/Experimental/Chaos/Public/Chaos/PBDJointConstraints.h b/Engine/Source/Runtime/Experimental/Chaos/Public/Chaos/PBDJointConstraints.h
--- a/Engine/Source/Runtime/Experimental/Chaos/Public/Chaos/PBDJointConstraints.h
+++ b/Engine/Source/Runtime/Experimental/Chaos/Public/Chaos/PBDJointConstraints.h
@@ -156,13 +156,17 @@ namespace Chaos
 		void CalculateConstraintSpace(int32 ConstraintIndex, FVec3& OutX0, FMatrix33& OutR0, FVec3& OutX1, FMatrix33& OutR1, FVec3& OutAngles) const;
 		void UpdateParticleState(TPBDRigidParticleHandle<FReal, 3>* Rigid, const FReal Dt, const FVec3& P, const FRotation3& Q, const bool bUpdateVelocity = true);
 
-		void ApplyDrives(const FReal Dt, const int32 ConstraintIndex, const int32 NumPairIts, const int32 It, const int32 NumIts);
+		void ApplyDrives(const FReal Dt, const int32 ConstraintIndex, const int32 It, const int32 NumIts);
 		void SolveVelocity(const FReal Dt, const int32 ConstraintIndex, const int32 NumPairIts, const int32 It, const int32 NumIts);
 		void SolvePosition(const FReal Dt, const int32 ConstraintIndex, const int32 NumPairIts, const int32 It, const int32 NumIts);
 		void SolvePosition_GaussSiedel(const FReal Dt, const int32 ConstraintIndex, const int32 NumPairIts, const int32 It, const int32 NumIts);
 		void SolvePosition_Cholesky(const FReal Dt, const int32 ConstraintIndex, const int32 NumPairIts, const int32 It, const int32 NumIts);
 		void ProjectPosition(const FReal Dt, const int32 ConstraintIndex, const int32 It, const int32 NumIts);
 
+		// All-constraints API
+		void ApplyDrives(const FReal Dt, const int32 It, const int32 NumIts);
+		void SolvePosition(const FReal Dt, const int32 NumPairIts, const int32 It, const int32 NumIts);
+
 		FPBDJointSolverSettings Settings;
 
 		TArray<FPBDJointSettings> ConstraintSettings;
diff --git a/Engine/Source/Runtime/Experimental/Chaos/Public/Chaos/Rotation.h b/Engine/Source/Runtime/Experimental/Chaos/Public/Chaos/Rotation.h
--- a/Engine/Source/Runtime/Experimental/Chaos/Public/Chaos/Rotation.h
+++ b/Engine/Source/Runtime/Experimental/Chaos/Public/Chaos/Rotation.h
@@ -102,6 +102,16 @@ namespace Chaos
 			return false;
 		}
 
+		/**
+		 * Extract the Swing and Twist rotations, assuming that the Twist Axis is (1,0,0).
+		 * /see ToSwingTwist
+		 */
+		void ToSwingTwistX(const FVector& InTwistAxis, FQuat& OutSwing, FQuat& OutTwist) const
+		{
+			OutTwist = (X != 0.0f)? FQuat(X, 0, 0, W).GetNormalized() : FQuat::Identity;
+			OutSwing = *this * OutTwist.Inverse();
+		}
+
 		/**
 		 * Return the complex conjugate of the rotation
 		 */
diff --git a/Engine/Source/Runtime/Experimental/Chaos/Public/Chaos/Utilities.h b/Engine/Source/Runtime/Experimental/Chaos/Public/Chaos/Utilities.h
--- a/Engine/Source/Runtime/Experimental/Chaos/Public/Chaos/Utilities.h
+++ b/Engine/Source/Runtime/Experimental/Chaos/Public/Chaos/Utilities.h
@@ -2,10 +2,8 @@
 
 #pragma once
 
-#include "Chaos/DynamicParticles.h"
+#include "Chaos/Core.h"
 #include "Chaos/Matrix.h"
-#include "Chaos/PBDParticles.h"
-#include "Chaos/PerParticleGravity.h"
 #include "Chaos/Transform.h"
 #include "Chaos/Vector.h"
 
@@ -143,28 +141,25 @@ namespace Chaos
 		 * Multiple two matrices: C = L.R
 		 * @note This is the mathematically expected operator. FMatrix operator* calculates C = R.Transpose(L), so this is not equivalent to that.
 		 */
-		inline PMatrix<float, 3, 3> Multiply(const PMatrix<float, 3, 3>& LIn, const PMatrix<float, 3, 3>& RIn)
+		inline PMatrix<float, 3, 3> Multiply(const PMatrix<float, 3, 3>& L, const PMatrix<float, 3, 3>& R)
 		{
-			// @todo(ccaulfield): optimize: remove transposes and use simd etc
-			PMatrix<float, 3, 3> L = LIn.GetTransposed();
-			PMatrix<float, 3, 3> R = RIn.GetTransposed();
+			// @todo(ccaulfield): optimize: simd
 
 			// We want L.R (FMatrix operator* actually calculates R.(L)T; i.e., Right is on the left, and the Left is transposed on the right.)
 			// NOTE: PMatrix constructor takes values in column order
 			return PMatrix<float, 3, 3>(
-			    L.M[0][0] * R.M[0][0] + L.M[0][1] * R.M[1][0] + L.M[0][2] * R.M[2][0], // x00
-			    L.M[1][0] * R.M[0][0] + L.M[1][1] * R.M[1][0] + L.M[1][2] * R.M[2][0], // x10
-			    L.M[2][0] * R.M[0][0] + L.M[2][1] * R.M[1][0] + L.M[2][2] * R.M[2][0], // x20
-
-			    L.M[0][0] * R.M[0][1] + L.M[0][1] * R.M[1][1] + L.M[0][2] * R.M[2][1], // x01
-			    L.M[1][0] * R.M[0][1] + L.M[1][1] * R.M[1][1] + L.M[1][2] * R.M[2][1], // x11
-			    L.M[2][0] * R.M[0][1] + L.M[2][1] * R.M[1][1] + L.M[2][2] * R.M[2][1], // x21
-
-			    L.M[0][0] * R.M[0][2] + L.M[0][1] * R.M[1][2] + L.M[0][2] * R.M[2][2], // x02
-			    L.M[1][0] * R.M[0][2] + L.M[1][1] * R.M[1][2] + L.M[1][2] * R.M[2][2], // x12
-			    L.M[2][0] * R.M[0][2] + L.M[2][1] * R.M[1][2] + L.M[2][2] * R.M[2][2] // x22
-			    )
-			    .GetTransposed();
+				L.M[0][0] * R.M[0][0] + L.M[1][0] * R.M[0][1] + L.M[2][0] * R.M[0][2],	// x00
+				L.M[0][0] * R.M[1][0] + L.M[1][0] * R.M[1][1] + L.M[2][0] * R.M[1][2],	// x01
+				L.M[0][0] * R.M[2][0] + L.M[1][0] * R.M[2][1] + L.M[2][0] * R.M[2][2],	// x02
+
+				L.M[0][1] * R.M[0][0] + L.M[1][1] * R.M[0][1] + L.M[2][1] * R.M[0][2],	// x10
+				L.M[0][1] * R.M[1][0] + L.M[1][1] * R.M[1][1] + L.M[2][1] * R.M[1][2],	// x11
+				L.M[0][1] * R.M[2][0] + L.M[1][1] * R.M[2][1] + L.M[2][1] * R.M[2][2],	// x12
+
+				L.M[0][2] * R.M[0][0] + L.M[1][2] * R.M[0][1] + L.M[2][2] * R.M[0][2],	// x20
+				L.M[0][2] * R.M[1][0] + L.M[1][2] * R.M[1][1] + L.M[2][2] * R.M[1][2],	// x21
+				L.M[0][2] * R.M[2][0] + L.M[1][2] * R.M[2][1] + L.M[2][2] * R.M[2][2]	// x22
+				);
 		}
 
 		inline PMatrix<float, 3, 3> MultiplyAB(const PMatrix<float, 3, 3>& LIn, const PMatrix<float, 3, 3>& RIn)
@@ -172,14 +167,39 @@ namespace Chaos
 			return Multiply(LIn, RIn);
 		}
 
-		inline PMatrix<float, 3, 3> MultiplyABt(const PMatrix<float, 3, 3>& LIn, const PMatrix<float, 3, 3>& RIn)
+		inline PMatrix<float, 3, 3> MultiplyABt(const PMatrix<float, 3, 3>& L, const PMatrix<float, 3, 3>& R)
 		{
-			return Multiply(LIn, RIn.GetTransposed());
+			return PMatrix<float, 3, 3>(
+				L.M[0][0] * R.M[0][0] + L.M[1][0] * R.M[1][0] + L.M[2][0] * R.M[2][0],	// x00
+				L.M[0][0] * R.M[0][1] + L.M[1][0] * R.M[1][1] + L.M[2][0] * R.M[2][1],	// x01
+				L.M[0][0] * R.M[0][2] + L.M[1][0] * R.M[1][2] + L.M[2][0] * R.M[2][2],	// x02
+
+				L.M[0][1] * R.M[0][0] + L.M[1][1] * R.M[1][0] + L.M[2][1] * R.M[2][0],	// x10
+				L.M[0][1] * R.M[0][1] + L.M[1][1] * R.M[1][1] + L.M[2][1] * R.M[2][1],	// x11
+				L.M[0][1] * R.M[0][2] + L.M[1][1] * R.M[1][2] + L.M[2][1] * R.M[2][2],	// x12
+
+				L.M[0][2] * R.M[0][0] + L.M[1][2] * R.M[1][0] + L.M[2][2] * R.M[2][0],	// x20
+				L.M[0][2] * R.M[0][1] + L.M[1][2] * R.M[1][1] + L.M[2][2] * R.M[2][1],	// x21
+				L.M[0][2] * R.M[0][2] + L.M[1][2] * R.M[1][2] + L.M[2][2] * R.M[2][2]	// x22
+				);
 		}
 
-		inline PMatrix<float, 3, 3> MultiplyAtB(const PMatrix<float, 3, 3>& LIn, const PMatrix<float, 3, 3>& RIn)
+		inline PMatrix<float, 3, 3> MultiplyAtB(const PMatrix<float, 3, 3>& L, const PMatrix<float, 3, 3>& R)
 		{
-			return Multiply(LIn.GetTransposed(), RIn);
+			return PMatrix<float, 3, 3>(
+				L.M[0][0] * R.M[0][0] + L.M[0][1] * R.M[0][1] + L.M[0][2] * R.M[0][2],	// x00
+				L.M[0][0] * R.M[1][0] + L.M[0][1] * R.M[1][1] + L.M[0][2] * R.M[1][2],	// x01
+				L.M[0][0] * R.M[2][0] + L.M[0][1] * R.M[2][1] + L.M[0][2] * R.M[2][2],	// x02
+
+				L.M[1][0] * R.M[0][0] + L.M[1][1] * R.M[0][1] + L.M[1][2] * R.M[0][2],	// x10
+				L.M[1][0] * R.M[1][0] + L.M[1][1] * R.M[1][1] + L.M[1][2] * R.M[1][2],	// x11
+				L.M[1][0] * R.M[2][0] + L.M[1][1] * R.M[2][1] + L.M[1][2] * R.M[2][2],	// x12
+
+				L.M[2][0] * R.M[0][0] + L.M[2][1] * R.M[0][1] + L.M[2][2] * R.M[0][2],	// x20
+				L.M[2][0] * R.M[1][0] + L.M[2][1] * R.M[1][1] + L.M[2][2] * R.M[1][2],	// x21
+				L.M[2][0] * R.M[2][0] + L.M[2][1] * R.M[2][1] + L.M[2][2] * R.M[2][2]	// x22
+				);
+
 		}
 
 		/**
