Chaos - speculative fix StressTestUnity build

diff --git a/Engine/Source/Runtime/Experimental/Chaos/Private/Chaos/Collision/CollisionDetector.cpp b/Engine/Source/Runtime/Experimental/Chaos/Private/Chaos/Collision/CollisionDetector.cpp
--- a/Engine/Source/Runtime/Experimental/Chaos/Private/Chaos/Collision/CollisionDetector.cpp
+++ b/Engine/Source/Runtime/Experimental/Chaos/Private/Chaos/Collision/CollisionDetector.cpp
@@ -6,4 +6,7 @@
 namespace Chaos
 {
 	DEFINE_STAT(STAT_DetectCollisions);
+
+	template class TCollisionDetector<FParticlePairBroadPhase, FNarrowPhase, FSyncCollisionReceiver, TPBDCollisionConstraints<FReal, 3>>;
+	template class TCollisionDetector<FSpatialAccelerationBroadPhase, FNarrowPhase, FAsyncCollisionReceiver, TPBDCollisionConstraints<FReal, 3>>;
 }
\ No newline at end of file
