Chaos : Collisions

Enable manifold based collisions (Iterative)

diff --git a/Engine/Source/Runtime/Experimental/Chaos/Private/Chaos/CollisionResolution.cpp b/Engine/Source/Runtime/Experimental/Chaos/Private/Chaos/CollisionResolution.cpp
--- a/Engine/Source/Runtime/Experimental/Chaos/Private/Chaos/CollisionResolution.cpp
+++ b/Engine/Source/Runtime/Experimental/Chaos/Private/Chaos/CollisionResolution.cpp
@@ -560,8 +560,6 @@ namespace Chaos
 		template<class T, int d>
 		void UpdateConvexConvexConstraint(const FImplicitObject& Implicit0, const TRigidTransform<T, d>& Transform0, const FImplicitObject& Implicit1, const TRigidTransform<T, d>& Transform1, const T Thickness, TRigidBodyIterativeContactConstraint<T,d>& Constraint)
 		{
-			TContactPoint<T> ContactPoint = ConvexConvexContactPoint(Implicit0, Transform0, Implicit1, Transform1, Thickness);
-			/*
 			TContactPoint<T> ContactPoint;
 			ContactPoint.Phi = FLT_MAX;
 
@@ -582,14 +580,13 @@ namespace Chaos
 					ContactPoint.Location = Constraint.LocalSamples[Idx].Manifold.Location;
 				}
 			}
-			*/
+
 			UpdateContactPoint(Constraint.Manifold, ContactPoint);
 		}
 
 		template<class T, int d>
 		void UpdateConvexConvexManifold(TRigidBodyIterativeContactConstraint<T, d>& Constraint, const TRigidTransform<T, d>& Transform0, const TRigidTransform<T, d>& Transform1, const T Thickness)
 		{
-			/*
 			TContactPoint<T> ContactPoint = ConvexConvexContactPoint(*Constraint.Manifold.Implicit[0], Transform0, *Constraint.Manifold.Implicit[1], Transform1, Thickness);
 
 			FRigidBodyIterativeContactConstraint::FManifold& Manifold = Constraint.Manifold;
@@ -629,7 +626,6 @@ namespace Chaos
 			typedef FRigidBodyIterativeContactConstraint::FSampleData FSampleData;
 			Constraint.LocalPosition = SumSampleData(Constraint.LocalSamples) / Constraint.LocalSamples.Num();
 			for (FSampleData& Data : Constraint.LocalSamples) Data.Delta = (Constraint.LocalPosition - Data.X).SizeSquared();
-			*/
 		}
 
 
