Chaos - Joint Solvers cleanup

diff --git a/Engine/Source/Runtime/Engine/Private/PhysicsEngine/ImmediatePhysics/ImmediatePhysicsChaos/ImmediatePhysicsSimulation_Chaos.cpp b/Engine/Source/Runtime/Engine/Private/PhysicsEngine/ImmediatePhysics/ImmediatePhysicsChaos/ImmediatePhysicsSimulation_Chaos.cpp
--- a/Engine/Source/Runtime/Engine/Private/PhysicsEngine/ImmediatePhysics/ImmediatePhysicsChaos/ImmediatePhysicsSimulation_Chaos.cpp
+++ b/Engine/Source/Runtime/Engine/Private/PhysicsEngine/ImmediatePhysics/ImmediatePhysicsChaos/ImmediatePhysicsSimulation_Chaos.cpp
@@ -33,15 +33,14 @@ int32 ChaosImmediate_Collision_PushOutPairIterations = 2;
 int32 ChaosImmediate_Collision_Priority = 1;
 float ChaosImmediate_Collision_Thickness = 0;
 FAutoConsoleVariableRef CVarChaosImmPhysDisableCollisions(TEXT("p.Chaos.ImmPhys.Collision.Enabled"), ChaosImmediate_Collision_Enabled, TEXT("Enable/Disable collisions in Immediate Physics."));
-FAutoConsoleVariableRef CVarChaosImmPhysCollisionPairIterations(TEXT("p.Chaos.ImmPhys.Collision.PairIterations"), ChaosImmediate_Collision_PairIterations, TEXT("Set the Apply() (velocity correction) internal pair iteration count (0 to disable)"));
-FAutoConsoleVariableRef CVarChaosImmPhysCollisionPushOutPairIterations(TEXT("p.Chaos.ImmPhys.Collision.PushOutPairIterations"), ChaosImmediate_Collision_PushOutPairIterations, TEXT("Set the ApplyPushOut() (position correction) internal pair iteration count (0 to disable)"));
+FAutoConsoleVariableRef CVarChaosImmPhysCollisionPairIterations(TEXT("p.Chaos.ImmPhys.Collision.PairIterations"), ChaosImmediate_Collision_PairIterations, TEXT("Set the Apply() internal pair iteration count (0 to disable)"));
+FAutoConsoleVariableRef CVarChaosImmPhysCollisionPushOutPairIterations(TEXT("p.Chaos.ImmPhys.Collision.PushOutPairIterations"), ChaosImmediate_Collision_PushOutPairIterations, TEXT("Set the ApplyPushOut() internal pair iteration count (0 to disable)"));
 FAutoConsoleVariableRef CVarChaosImmPhysCollisionPriority(TEXT("p.Chaos.ImmPhys.Collision.Priority"), ChaosImmediate_Collision_Priority, TEXT("Set the Collision constraint sort order (Joints have priority 0)"));
 FAutoConsoleVariableRef CVarChaosImmPhysThickness(TEXT("p.Chaos.ImmPhys.Collision.Thickness"), ChaosImmediate_Collision_Thickness, TEXT("ChaosImmediateThickness"));
 
 int32 ChaosImmediate_Joint_PairIterations = 0;
 int32 ChaosImmediate_Joint_PushOutPairIterations = 2;
 float ChaosImmediate_Joint_SwingTwistAngleTolerance = 1.0e-6f;
-int32 ChaosImmediate_Joint_EnableVelocitySolve = 0;
 int32 ChaosImmediate_Joint_EnableTwistLimits = 1;
 int32 ChaosImmediate_Joint_EnableSwingLimits = 1;
 int32 ChaosImmediate_Joint_EnableDrives = 1;
@@ -58,7 +57,6 @@ float ChaosImmediate_Joint_MaxInertiaRatio = 5.0f;
 FAutoConsoleVariableRef CVarChaosImmPhysPairIterations(TEXT("p.Chaos.ImmPhys.Joint.PairIterations"), ChaosImmediate_Joint_PairIterations, TEXT("PairIterations."));
 FAutoConsoleVariableRef CVarChaosImmPhysPushOutPairIterations(TEXT("p.Chaos.ImmPhys.Joint.PushOutPairIterations"), ChaosImmediate_Joint_PushOutPairIterations, TEXT("PushOutPairIterations."));
 FAutoConsoleVariableRef CVarChaosImmPhysSwingTwistAngleTolerance(TEXT("p.Chaos.ImmPhys.Joint.SwingTwistAngleTolerance"), ChaosImmediate_Joint_SwingTwistAngleTolerance, TEXT("SwingTwistAngleTolerance."));
-FAutoConsoleVariableRef CVarChaosImmPhysEnableVelocitySolve(TEXT("p.Chaos.ImmPhys.Joint.EnableVelocitySolve"), ChaosImmediate_Joint_EnableVelocitySolve, TEXT("EnableVelocitySolve."));
 FAutoConsoleVariableRef CVarChaosImmPhysEnableTwistLimits(TEXT("p.Chaos.ImmPhys.Joint.EnableTwistLimits"), ChaosImmediate_Joint_EnableTwistLimits, TEXT("EnableTwistLimits."));
 FAutoConsoleVariableRef CVarChaosImmPhysEnableSwingLimits(TEXT("p.Chaos.ImmPhys.Joint.EnableSwingLimits"), ChaosImmediate_Joint_EnableSwingLimits, TEXT("EnableSwingLimits."));
 FAutoConsoleVariableRef CVarChaosImmPhysEnableDrives(TEXT("p.Chaos.ImmPhys.Joint.EnableDrives"), ChaosImmediate_Joint_EnableDrives, TEXT("EnableDrives."));
@@ -274,6 +272,7 @@ namespace ImmediatePhysics_Chaos
 		TUniquePtr<FChaosPhysicsMaterial> Material = MakeUnique<FChaosPhysicsMaterial>();
 		ParticleMaterials.Add(MakeSerializable(Material));
 		PerParticleMaterials.Add(MoveTemp(Material));
+		CollidedParticles.Add(false);
 
 		return ActorHandle;
 	}
@@ -286,6 +285,7 @@ namespace ImmediatePhysics_Chaos
 
 		ParticleMaterials.RemoveAt(Index, 1);
 		PerParticleMaterials.RemoveAt(Index, 1);
+		CollidedParticles.RemoveAt(Index, 1);
 	}
 
 	FJointHandle* FSimulation::CreateJoint(FConstraintInstance* ConstraintInstance, FActorHandle* Body1, FActorHandle* Body2)
@@ -462,7 +462,6 @@ namespace ImmediatePhysics_Chaos
 			JointsSettings.SwingTwistAngleTolerance = ChaosImmediate_Joint_SwingTwistAngleTolerance;
 			JointsSettings.MinParentMassRatio = ChaosImmediate_Joint_MinParentMassRatio;
 			JointsSettings.MaxInertiaRatio = ChaosImmediate_Joint_MaxInertiaRatio;
-			JointsSettings.bEnableVelocitySolve = ChaosImmediate_Joint_EnableVelocitySolve != 0;
 			JointsSettings.bEnableTwistLimits = ChaosImmediate_Joint_EnableTwistLimits != 0;
 			JointsSettings.bEnableSwingLimits = ChaosImmediate_Joint_EnableSwingLimits != 0;
 			JointsSettings.bEnableDrives = ChaosImmediate_Joint_EnableDrives != 0;
diff --git a/Engine/Source/Runtime/Experimental/Chaos/Private/Chaos/Joint/PBDJointSolverCholesky.cpp b/Engine/Source/Runtime/Experimental/Chaos/Private/Chaos/Joint/PBDJointSolverCholesky.cpp
new file mode 100644
--- /dev/null
+++ b/Engine/Source/Runtime/Experimental/Chaos/Private/Chaos/Joint/PBDJointSolverCholesky.cpp
@@ -0,0 +1,582 @@
+// Copyright 1998-2019 Epic Games, Inc. All Rights Reserved.
+#include "Chaos/Joint/PBDJointSolverCholesky.h"
+#include "Chaos/ParticleHandle.h"
+#include "Chaos/PBDJointConstraintUtilities.h"
+#include "Chaos/Utilities.h"
+#include "ChaosLog.h"
+#include "ChaosStats.h"
+
+//#pragma optimize("", off)
+
+// Set to '1' enable solver stats (very high frequency, so usually disabled)
+#define CHAOS_JOINTSOLVER_STATSENABLED 0
+
+namespace Chaos
+{
+#if CHAOS_JOINTSOLVER_STATSENABLED
+	DECLARE_CYCLE_STAT(TEXT("TPBDJointConstraints::Apply"), STAT_JointSolver_Apply, STATGROUP_Chaos);
+	DECLARE_CYCLE_STAT(TEXT("TPBDJointConstraints::ApplyPushOut"), STAT_JointSolver_Jacobian, STATGROUP_Chaos);
+#define CHAOS_JOINTSOLVER_SCOPE_CYCLE_STAT SCOPE_CYCLE_COUNTER(X)
+#else
+#define CHAOS_JOINTSOLVER_SCOPE_CYCLE_STAT(X)
+#endif
+
+	FJointSolverCholesky::FJointSolverCholesky()
+	{
+	}
+
+	void FJointSolverCholesky::UpdateDerivedState()
+	{
+		Xs[0] = Ps[0] + Qs[0] * XLs[0].GetTranslation();
+		Xs[1] = Ps[1] + Qs[1] * XLs[1].GetTranslation();
+		Rs[0] = Qs[0] * XLs[0].GetRotation();
+		Rs[1] = Qs[1] * XLs[1].GetRotation();
+	}
+
+
+	void FJointSolverCholesky::InitConstraints(
+		const FReal Dt,
+		const FPBDJointSolverSettings& SolverSettings,
+		const FPBDJointSettings& JointSettings,
+		const FVec3& P0,
+		const FRotation3& Q0,
+		const FVec3& P1,
+		const FRotation3& Q1,
+		const FReal InvM0,
+		const FMatrix33& InvIL0,
+		const FReal InvM1,
+		const FMatrix33& InvIL1,
+		const FRigidTransform3& XL0,
+		const FRigidTransform3& XL1)
+	{
+		XLs[0] = XL0;
+		XLs[1] = XL1;
+		InvILs[0] = InvIL0;
+		InvILs[1] = InvIL1;
+		InvMs[0] = InvM0;
+		InvMs[1] = InvM1;
+
+		Ps[0] = P0;
+		Ps[1] = P1;
+		Qs[0] = Q0;
+		Qs[1] = Q1;
+		Qs[1].EnforceShortestArcWith(Qs[0]);
+
+		Stiffness = FPBDJointUtilities::GetLinearStiffness(SolverSettings, JointSettings);
+		SwingTwistAngleTolerance = SolverSettings.SwingTwistAngleTolerance;
+		bEnableTwistLimits = SolverSettings.bEnableTwistLimits;
+		bEnableTwistLimits = SolverSettings.bEnableSwingLimits;
+
+		UpdateDerivedState();
+	}
+
+	void FJointSolverCholesky::ApplyConstraints(
+		const FReal Dt,
+		const FPBDJointSettings& JointSettings)
+	{
+		CHAOS_JOINTSOLVER_SCOPE_CYCLE_STAT(STAT_JointSolver_Apply);
+
+		// Solving for world-space position/rotation corrections D(6x1) where
+		// D = [I.Jt / [J.I.Jt]].C = I.Jt.L
+		// I is the inverse mass matrix, J the Jacobian, C the current constraint violation,
+		// and L = [1 / [J.I.Jt]].C is the Joint-space correction.
+
+		// For N constraints
+		// Constraint error: C(Nx1)
+		// Jacobian : J(Nx6)
+		// The Jacobians will be some sub-set of the following rows, depending on the "active" constraints 
+		// ("active" = enabled and either fixed or with limits exceeded).
+		//
+		// J0(Nx6) = | XAxis          -XAxis x Connector0 |
+		//           | YAxis          -YAxis x Connector0 |
+		//           | ZAxis          -ZAxis x Connector0 |
+		//           | 0              TwistAxis           |
+		//           | 0              Swing1Axis          |
+		//           | 0              Swing2Axis          |
+		//
+		// J1(Nx6) = | -XAxis         XAxis x Connector1  |
+		//           | -YAxis         YAxis x Connector1  |
+		//           | -ZAxis         ZAxis x Connector1  |
+		//           | 0              -TwistAxis          |
+		//           | 0              -Swing1Axis         |
+		//           | 0              -Swing2Axis         |
+		//
+		FDenseMatrix61 C;
+		FDenseMatrix66 J0, J1;
+		BuildJacobianAndResidual(JointSettings, J0, J1, C);
+
+		// InvM(6x6) = inverse mass matrix
+		const FMassMatrix InvM0 = FMassMatrix::Make(InvMs[0], Qs[0], InvILs[0]);
+		const FMassMatrix InvM1 = FMassMatrix::Make(InvMs[1], Qs[1], InvILs[1]);
+
+		// IJt(6xN) = I(6x6).Jt(6xN)
+		const FDenseMatrix66 IJt0 = FDenseMatrix66::MultiplyABt(InvM0, J0);
+		const FDenseMatrix66 IJt1 = FDenseMatrix66::MultiplyABt(InvM1, J1);
+
+		// Joint-space mass: F(NxN) = J(Nx6).I(6x6).Jt(6xN) = J(Nx6).IJt(6xN)
+		// NOTE: Result is symmetric
+		const FDenseMatrix66 F0 = FDenseMatrix66::MultiplyAB_Symmetric(J0, IJt0);
+		const FDenseMatrix66 F = FDenseMatrix66::MultiplyBCAddA_Symmetric(F0, J1, IJt1);
+
+		// Joint-space correction: L(Nx1) = [1/F](NxN).C(Nx1)
+		FDenseMatrix61 L;
+		if (FDenseMatrixSolver::SolvePositiveDefinite(F, C, L))
+		{
+			// World-space correction: D(6x1) = I(6x6).Jt(6xN).L(Nx1) = IJt(6xN).L(Nx1)
+			const FDenseMatrix61 D0 = FDenseMatrix61::MultiplyAB(IJt0, L);
+			const FDenseMatrix61 D1 = FDenseMatrix61::MultiplyAB(IJt1, L);
+
+			// Extract world-space position correction
+			Ps[0] = Ps[0] + FVec3(Stiffness * D0.At(0, 0), Stiffness * D0.At(1, 0), Stiffness * D0.At(2, 0));
+			Ps[1] = Ps[1] + FVec3(Stiffness * D1.At(0, 0), Stiffness * D1.At(1, 0), Stiffness * D1.At(2, 0));
+
+			// Extract world-space rotation correction
+			const FReal HalfStiffness = (FReal)0.5 * Stiffness;
+			const FRotation3 DQ0 = FRotation3::FromElements(HalfStiffness * D0.At(3, 0), HalfStiffness * D0.At(4, 0), HalfStiffness * D0.At(5, 0), 0) * Qs[0];
+			const FRotation3 DQ1 = FRotation3::FromElements(HalfStiffness * D1.At(3, 0), HalfStiffness * D1.At(4, 0), HalfStiffness * D1.At(5, 0), 0) * Qs[1];
+			Qs[0] = (Qs[0] + DQ0).GetNormalized();
+			Qs[1] = (Qs[1] + DQ1).GetNormalized();
+			Qs[1].EnforceShortestArcWith(Qs[0]);
+
+			UpdateDerivedState();
+		}
+	}
+
+	void FJointSolverCholesky::BuildJacobianAndResidual(
+		const FPBDJointSettings& JointSettings,
+		FDenseMatrix66& J0,
+		FDenseMatrix66& J1,
+		FDenseMatrix61& C)
+	{
+		CHAOS_JOINTSOLVER_SCOPE_CYCLE_STAT(STAT_JointSolver_Jacobian);
+
+		// Initialize with zero constraints
+		J0.SetDimensions(0, 6);
+		J1.SetDimensions(0, 6);
+		C.SetDimensions(0, 1);
+
+		AddLinearConstraints(JointSettings, J0, J1, C);
+		AddAngularConstraints(JointSettings, J0, J1, C);
+	}
+
+	// 3 constraints along principle axes
+	void FJointSolverCholesky::AddLinearConstraints_Point(
+		const FPBDJointSettings& JointSettings,
+		FDenseMatrix66& J0,
+		FDenseMatrix66& J1,
+		FDenseMatrix61& C)
+	{
+		const int32 RowIndex = J0.NumRows();
+		J0.AddRows(3);
+		J1.AddRows(3);
+		C.AddRows(3);
+
+		// Cross Product
+		//Result[0] = V1[1] * V2[2] - V1[2] * V2[1];
+		//Result[1] = V1[2] * V2[0] - V1[0] * V2[2];
+		//Result[2] = V1[0] * V2[1] - V1[1] * V2[0];
+
+		const FVec3 XP0 = Xs[0] - Ps[0];
+		J0.SetBlockAtDiagonal33(RowIndex, 0, (FReal)1, (FReal)0);
+		J0.SetRowAt(RowIndex + 0, 3, (FReal)0, XP0[2], -XP0[1]);	// -(1,0,0) x XP0
+		J0.SetRowAt(RowIndex + 1, 3, -XP0[2], (FReal)0, XP0[0]);	// -(0,1,0) x XP0
+		J0.SetRowAt(RowIndex + 2, 3, XP0[1], -XP0[0], (FReal)0);	// -(0,0,1) x XP0
+
+		const FVec3 XP1 = Xs[1] - Ps[1];
+		J1.SetBlockAtDiagonal33(RowIndex, 0, (FReal)-1, (FReal)0);
+		J1.SetRowAt(RowIndex + 0, 3, (FReal)0, -XP1[2], XP1[1]);	// (1,0,0) x XP1
+		J1.SetRowAt(RowIndex + 1, 3, XP1[2], (FReal)0, -XP1[0]);	// (0,1,0) x XP1
+		J1.SetRowAt(RowIndex + 2, 3, -XP1[1], XP1[0], (FReal)0);	// (0,0,1) x XP1
+
+		const FVec3 ConstraintSeparation = Xs[1] - Xs[0];
+		C.SetAt(RowIndex + 0, 0, ConstraintSeparation[0]);
+		C.SetAt(RowIndex + 1, 0, ConstraintSeparation[1]);
+		C.SetAt(RowIndex + 2, 0, ConstraintSeparation[2]);
+	}
+
+	// up to 1 constraint limiting distance
+	void FJointSolverCholesky::AddLinearConstraints_Sphere(
+		const FPBDJointSettings& JointSettings,
+		FDenseMatrix66& J0,
+		FDenseMatrix66& J1,
+		FDenseMatrix61& C)
+	{
+		const FReal Limit = JointSettings.Motion.LinearLimit;
+		const FVec3 ConstraintSeparation = Xs[1] - Xs[0];
+		const FReal ConstraintSeparationLen = ConstraintSeparation.Size();
+
+		bool bConstraintActive = (ConstraintSeparationLen >= FMath::Max(Limit, KINDA_SMALL_NUMBER));
+		if (bConstraintActive)
+		{
+			const FVec3 XP0 = Xs[0] - Ps[0];
+			const FVec3 XP1 = Xs[1] - Ps[1];
+			const FVec3 Axis = ConstraintSeparation / ConstraintSeparationLen;
+
+			const int32 RowIndex = J0.NumRows();
+			J0.AddRows(1);
+			J1.AddRows(1);
+			C.AddRows(1);
+
+			J0.SetRowAt(RowIndex, 0, Axis);
+			J0.SetRowAt(RowIndex, 3, -FVec3::CrossProduct(Axis, XP0));
+
+			J1.SetRowAt(RowIndex, 0, -Axis);
+			J1.SetRowAt(RowIndex, 3, FVec3::CrossProduct(Axis, XP1));
+
+			C.SetAt(RowIndex, 0, ConstraintSeparationLen - Limit);
+		}
+	}
+
+	// up to 2 constraint: 1 limiting distance along the axis and another limiting lateral distance from the axis
+	void FJointSolverCholesky::AddLinearConstraints_Cylinder(
+		const FPBDJointSettings& JointSettings,
+		const EJointMotionType AxisMotion,
+		const FVec3& Axis,
+		FDenseMatrix66& J0,
+		FDenseMatrix66& J1,
+		FDenseMatrix61& C)
+	{
+		const FVec3 ConstraintSeparation = Xs[1] - Xs[0];
+		const FVec3 XP0 = Xs[0] - Ps[0];
+		const FVec3 XP1 = Xs[1] - Ps[1];
+
+		// Axial Constraint
+		const FReal ConstraintDistanceAxial = FVec3::DotProduct(ConstraintSeparation, Axis);
+		bool bAxisConstraintActive = (AxisMotion != EJointMotionType::Free);
+		if (bAxisConstraintActive)
+		{
+			const int32 RowIndex = J0.NumRows();
+			J0.AddRows(1);
+			J1.AddRows(1);
+			C.AddRows(1);
+
+			J0.SetRowAt(RowIndex, 0, Axis);
+			J0.SetRowAt(RowIndex, 3, -FVec3::CrossProduct(Axis, XP0));
+
+			J1.SetRowAt(RowIndex, 0, -Axis);
+			J1.SetRowAt(RowIndex, 3, FVec3::CrossProduct(Axis, XP1));
+
+			C.SetAt(RowIndex, 0, ConstraintDistanceAxial);
+		}
+
+		// Radial Constraint
+		const FVec3 ConstraintSeparationRadial = ConstraintSeparation - ConstraintDistanceAxial * Axis;
+		const FReal ConstraintDistanceRadial = ConstraintSeparationRadial.Size();
+		const FReal RadialLimit = JointSettings.Motion.LinearLimit;
+		bool bRadialConstraintActive = (ConstraintDistanceRadial >= RadialLimit);
+		if (bRadialConstraintActive)
+		{
+			const FVec3 RadialAxis = ConstraintSeparationRadial / ConstraintDistanceRadial;
+
+			const int32 RowIndex = J0.NumRows();
+			J0.AddRows(1);
+			J1.AddRows(1);
+			C.AddRows(1);
+
+			J0.SetRowAt(RowIndex, 0, RadialAxis);
+			J0.SetRowAt(RowIndex, 3, -FVec3::CrossProduct(RadialAxis, XP0));
+
+			J1.SetRowAt(RowIndex, 0, -RadialAxis);
+			J1.SetRowAt(RowIndex, 3, FVec3::CrossProduct(RadialAxis, XP1));
+
+			C.SetAt(RowIndex, 0, ConstraintDistanceRadial - RadialLimit);
+		}
+	}
+
+	// up to 1 constraint limiting distance along the axis (lateral motion unrestricted)
+	void FJointSolverCholesky::AddLinearConstraints_Plane(
+		const FPBDJointSettings& JointSettings,
+		const EJointMotionType AxisMotion,
+		const FVec3& Axis,
+		FDenseMatrix66& J0,
+		FDenseMatrix66& J1,
+		FDenseMatrix61& C)
+	{
+		const FReal Limit = (AxisMotion == EJointMotionType::Limited) ? JointSettings.Motion.LinearLimit : (FReal)0;
+		const FVec3 ConstraintSeparation = Xs[1] - Xs[0];
+
+		// Planar Constraint
+		const FReal ConstraintDistanceAxial = FVec3::DotProduct(ConstraintSeparation, Axis);
+		bool bAxisConstraintActive = (ConstraintDistanceAxial <= -Limit) || (ConstraintDistanceAxial >= Limit);
+		if (bAxisConstraintActive)
+		{
+			const FVec3 XP0 = Xs[0] - Ps[0];
+			const FVec3 XP1 = Xs[1] - Ps[1];
+
+			const int32 RowIndex = J0.NumRows();
+			J0.AddRows(1);
+			J1.AddRows(1);
+			C.AddRows(1);
+
+			J0.SetRowAt(RowIndex, 0, Axis);
+			J0.SetRowAt(RowIndex, 3, -FVec3::CrossProduct(Axis, XP0));
+
+			J1.SetRowAt(RowIndex, 0, -Axis);
+			J1.SetRowAt(RowIndex, 3, FVec3::CrossProduct(Axis, XP1));
+
+			C.SetAt(RowIndex, 0, (ConstraintDistanceAxial >= Limit) ? ConstraintDistanceAxial - Limit : ConstraintDistanceAxial + Limit);
+		}
+	}
+
+	// up to 1 constraint limiting rotation about twist axes
+	void FJointSolverCholesky::AddAngularConstraints_Twist(
+		const FPBDJointSettings& JointSettings,
+		const FRotation3& R01Twist,
+		const FRotation3& R01Swing,
+		FDenseMatrix66& J0,
+		FDenseMatrix66& J1,
+		FDenseMatrix61& C)
+	{
+		FVec3 TwistAxis01 = FJointConstants::TwistAxis();
+		FReal TwistAngle = R01Twist.GetAngle();
+		if (TwistAngle > PI)
+		{
+			TwistAngle = TwistAngle - (FReal)2 * PI;
+		}
+		if (R01Twist.X < 0)
+		{
+			TwistAngle = -TwistAngle;
+		}
+
+		const FReal TwistAngleMax = JointSettings.Motion.AngularLimits[(int32)EJointAngularConstraintIndex::Twist];
+		bool bConstraintActive = (TwistAngle <= -TwistAngleMax) || (TwistAngle >= TwistAngleMax);
+		if (bConstraintActive)
+		{
+			const FVec3 Axis0 = Rs[0] * TwistAxis01;
+			const FVec3 Axis1 = Rs[1] * TwistAxis01;
+			const FVec3 XP0 = Xs[0] - Ps[0];
+			const FVec3 XP1 = Xs[1] - Ps[1];
+
+			const int32 RowIndex = J0.NumRows();
+			J0.AddRows(1);
+			J1.AddRows(1);
+			C.AddRows(1);
+
+			J0.SetRowAt(RowIndex + 0, 0, 0, 0, 0);
+			J0.SetRowAt(RowIndex + 0, 3, Axis0);
+
+			J1.SetRowAt(RowIndex + 0, 0, 0, 0, 0);
+			J1.SetRowAt(RowIndex + 0, 3, -Axis1);
+
+			C.SetAt(RowIndex, 0, (TwistAngle >= TwistAngleMax) ? TwistAngle - TwistAngleMax : TwistAngle + TwistAngleMax);
+		}
+	}
+
+	// up to 1 constraint limiting angle between twist axes
+	void FJointSolverCholesky::AddAngularConstraints_Cone(
+		const FPBDJointSettings& JointSettings,
+		const FRotation3& R01Twist,
+		const FRotation3& R01Swing,
+		FDenseMatrix66& J0,
+		FDenseMatrix66& J1,
+		FDenseMatrix61& C)
+	{
+		FVec3 SwingAxis01;
+		FReal SwingAngle;
+		R01Swing.ToAxisAndAngleSafe(SwingAxis01, SwingAngle, FJointConstants::Swing1Axis(), SwingTwistAngleTolerance);
+		if (SwingAngle > PI)
+		{
+			SwingAngle = SwingAngle - (FReal)2 * PI;
+		}
+
+		// Calculate swing limit for the current swing axis
+		FReal SwingAngleMax = FLT_MAX;
+		const FReal Swing1Limit = JointSettings.Motion.AngularLimits[(int32)EJointAngularConstraintIndex::Swing1];
+		const FReal Swing2Limit = JointSettings.Motion.AngularLimits[(int32)EJointAngularConstraintIndex::Swing2];
+
+		// Circular swing limit
+		SwingAngleMax = Swing1Limit;
+
+		// Elliptical swing limit
+		if (!FMath::IsNearlyEqual(Swing1Limit, Swing2Limit, KINDA_SMALL_NUMBER))
+		{
+			// Map swing axis to ellipse and calculate limit for this swing axis
+			const FReal DotSwing1 = FMath::Abs(FVec3::DotProduct(SwingAxis01, FJointConstants::Swing1Axis()));
+			const FReal DotSwing2 = FMath::Abs(FVec3::DotProduct(SwingAxis01, FJointConstants::Swing2Axis()));
+			SwingAngleMax = FMath::Sqrt(Swing1Limit * DotSwing2 * Swing1Limit * DotSwing2 + Swing2Limit * DotSwing1 * Swing2Limit * DotSwing1);
+		}
+
+		bool bConstraintActive = (SwingAngle <= -SwingAngleMax) || (SwingAngle >= SwingAngleMax);
+		if (bConstraintActive)
+		{
+			const FVec3 XP0 = Xs[0] - Ps[0];
+			const FVec3 XP1 = Xs[1] - Ps[1];
+			const FVec3 Axis = Rs[0] * SwingAxis01;
+
+			const int32 RowIndex = J0.NumRows();
+			J0.AddRows(1);
+			J1.AddRows(1);
+			C.AddRows(1);
+
+			J0.SetRowAt(RowIndex + 0, 0, 0, 0, 0);
+			J0.SetRowAt(RowIndex + 0, 3, Axis);
+
+			J1.SetRowAt(RowIndex + 0, 0, 0, 0, 0);
+			J1.SetRowAt(RowIndex + 0, 3, -Axis);
+
+			C.SetAt(RowIndex, 0, (SwingAngle >= SwingAngleMax) ? SwingAngle - SwingAngleMax : SwingAngle + SwingAngleMax);
+		}
+	}
+
+	// up to 1 constraint limiting rotation about swing axis (relative to body 0)
+	void FJointSolverCholesky::AddAngularConstraints_Swing(
+		const FPBDJointSettings& JointSettings,
+		const EJointAngularConstraintIndex SwingConstraintIndex,
+		const EJointAngularAxisIndex SwingAxisIndex,
+		const FRotation3& R01Twist,
+		const FRotation3& R01Swing,
+		FDenseMatrix66& J0,
+		FDenseMatrix66& J1,
+		FDenseMatrix61& C)
+	{
+		FVec3 TwistAxis01;
+		FReal TwistAngle;
+		R01Twist.ToAxisAndAngleSafe(TwistAxis01, TwistAngle, FJointConstants::TwistAxis(), SwingTwistAngleTolerance);
+		if (TwistAngle > PI)
+		{
+			TwistAngle = TwistAngle - (FReal)2 * PI;
+		}
+		if (FVec3::DotProduct(TwistAxis01, FJointConstants::TwistAxis()) < 0)
+		{
+			TwistAxis01 = -TwistAxis01;
+			TwistAngle = -TwistAngle;
+		}
+		const FVec3 TwistAxis = Rs[0] * TwistAxis01;
+
+		const FRotation3 R1NoTwist = Rs[1] * R01Twist.Inverse();
+		const FMatrix33 Axes0 = Rs[0].ToMatrix();
+		const FMatrix33 Axes1 = R1NoTwist.ToMatrix();
+		FVec3 SwingCross = FVec3::CrossProduct(Axes0.GetAxis((int32)SwingAxisIndex), Axes1.GetAxis((int32)SwingAxisIndex));
+		SwingCross = SwingCross - FVec3::DotProduct(TwistAxis, SwingCross) * TwistAxis;
+		const FReal SwingCrossLen = SwingCross.Size();
+		if (SwingCrossLen > KINDA_SMALL_NUMBER)
+		{
+			FReal SwingAngle = FMath::Asin(FMath::Clamp(SwingCrossLen, (FReal)0, (FReal)1));
+			const FReal SwingDot = FVec3::DotProduct(Axes0.GetAxis((int32)SwingAxisIndex), Axes1.GetAxis((int32)SwingAxisIndex));
+			if (SwingDot < (FReal)0)
+			{
+				SwingAngle = (FReal)PI - SwingAngle;
+			}
+
+			FReal SwingAngleMax = JointSettings.Motion.AngularLimits[(int32)SwingConstraintIndex];
+			bool bConstraintActive = (SwingAngle <= -SwingAngleMax) || (SwingAngle >= SwingAngleMax);
+			if (bConstraintActive)
+			{
+				const FVec3 XP0 = Xs[0] - Ps[0];
+				const FVec3 XP1 = Xs[1] - Ps[1];
+				const FVec3 Axis = SwingCross / SwingCrossLen;
+
+				const int32 RowIndex = J0.NumRows();
+				J0.AddRows(1);
+				J1.AddRows(1);
+				C.AddRows(1);
+
+				J0.SetRowAt(RowIndex + 0, 0, 0, 0, 0);
+				J0.SetRowAt(RowIndex + 0, 3, Axis);
+
+				J1.SetRowAt(RowIndex + 0, 0, 0, 0, 0);
+				J1.SetRowAt(RowIndex + 0, 3, -Axis);
+
+				C.SetAt(RowIndex, 0, (SwingAngle >= SwingAngleMax) ? SwingAngle - SwingAngleMax : SwingAngle + SwingAngleMax);
+			}
+		}
+	}
+
+	// Add linear constraints to solver
+	void FJointSolverCholesky::AddLinearConstraints(
+		const FPBDJointSettings& JointSettings,
+		FDenseMatrix66& J0,
+		FDenseMatrix66& J1,
+		FDenseMatrix61& C)
+	{
+		const TVector<EJointMotionType, 3>& Motion = JointSettings.Motion.LinearMotionTypes;
+		if ((Motion[0] == EJointMotionType::Locked) && (Motion[1] == EJointMotionType::Locked) && (Motion[2] == EJointMotionType::Locked))
+		{
+			AddLinearConstraints_Point(JointSettings, J0, J1, C);
+		}
+		else if ((Motion[0] == EJointMotionType::Limited) && (Motion[1] == EJointMotionType::Limited) && (Motion[2] == EJointMotionType::Limited))
+		{
+			AddLinearConstraints_Sphere(JointSettings, J0, J1, C);
+		}
+		else if ((Motion[1] == EJointMotionType::Limited) && (Motion[2] == EJointMotionType::Limited))
+		{
+			// Circular Limit (X Axis)
+			AddLinearConstraints_Cylinder(JointSettings, Motion[0], Rs[0] * FVec3(1, 0, 0), J0, J1, C);
+		}
+		else if ((Motion[0] == EJointMotionType::Limited) && (Motion[2] == EJointMotionType::Limited))
+		{
+			// Circular Limit (Y Axis)
+			AddLinearConstraints_Cylinder(JointSettings, Motion[1], Rs[0] * FVec3(0, 1, 0), J0, J1, C);
+		}
+		else if ((Motion[0] == EJointMotionType::Limited) && (Motion[1] == EJointMotionType::Limited))
+		{
+			// Circular Limit (Z Axis)
+			AddLinearConstraints_Cylinder(JointSettings, Motion[2], Rs[0] * FVec3(0, 0, 1), J0, J1, C);
+		}
+		else
+		{
+			// Plane/Square/Cube Limits (no way to author square or cube limits, but would work if we wanted it)
+			if (Motion[0] != EJointMotionType::Free)
+			{
+				AddLinearConstraints_Plane(JointSettings, Motion[0], Rs[0] * FVec3(1, 0, 0), J0, J1, C);
+			}
+			if (Motion[1] != EJointMotionType::Free)
+			{
+				AddLinearConstraints_Plane(JointSettings, Motion[1], Rs[0] * FVec3(0, 1, 0), J0, J1, C);
+			}
+			if (Motion[2] != EJointMotionType::Free)
+			{
+				AddLinearConstraints_Plane(JointSettings, Motion[2], Rs[0] * FVec3(0, 0, 1), J0, J1, C);
+			}
+		}
+	}
+
+	// Add angular constraints to solver
+	void FJointSolverCholesky::AddAngularConstraints(
+		const FPBDJointSettings& JointSettings,
+		FDenseMatrix66& J0,
+		FDenseMatrix66& J1,
+		FDenseMatrix61& C)
+	{
+		EJointMotionType TwistMotion = JointSettings.Motion.AngularMotionTypes[(int32)EJointAngularConstraintIndex::Twist];
+		EJointMotionType Swing1Motion = JointSettings.Motion.AngularMotionTypes[(int32)EJointAngularConstraintIndex::Swing1];
+		EJointMotionType Swing2Motion = JointSettings.Motion.AngularMotionTypes[(int32)EJointAngularConstraintIndex::Swing2];
+
+		bool bAddTwist = bEnableTwistLimits && (TwistMotion != EJointMotionType::Free);
+		bool bAddConeOrSwing = bEnableSwingLimits && ((Swing1Motion != EJointMotionType::Free) || (Swing2Motion != EJointMotionType::Free));
+
+		if (bAddTwist || bAddConeOrSwing)
+		{
+			// Decompose rotation of body 1 relative to body 0 into swing and twist rotations, assuming twist is X axis
+			FRotation3 R01Twist, R01Swing;
+			FPBDJointUtilities::DecomposeSwingTwistLocal(Rs[0], Rs[1], R01Swing, R01Twist);
+
+			// Add twist constraint
+			if (bAddTwist)
+			{
+				AddAngularConstraints_Twist(JointSettings, R01Twist, R01Swing, J0, J1, C);
+			}
+
+			// Add swing constraints
+			if (bAddConeOrSwing)
+			{
+				if ((Swing1Motion == EJointMotionType::Limited) && (Swing2Motion == EJointMotionType::Limited))
+				{
+					AddAngularConstraints_Cone(JointSettings, R01Twist, R01Swing, J0, J1, C);
+				}
+				else
+				{
+					if (Swing1Motion != EJointMotionType::Free)
+					{
+						AddAngularConstraints_Swing(JointSettings, EJointAngularConstraintIndex::Swing1, EJointAngularAxisIndex::Swing1, R01Twist, R01Swing, J0, J1, C);
+					}
+					if (Swing2Motion != EJointMotionType::Free)
+					{
+						AddAngularConstraints_Swing(JointSettings, EJointAngularConstraintIndex::Swing2, EJointAngularAxisIndex::Swing2, R01Twist, R01Swing, J0, J1, C);
+					}
+				}
+			}
+		}
+	}
+
+}
diff --git a/Engine/Source/Runtime/Experimental/Chaos/Private/Chaos/Joint/PBDJointSolverGaussSeidel.cpp b/Engine/Source/Runtime/Experimental/Chaos/Private/Chaos/Joint/PBDJointSolverGaussSeidel.cpp
new file mode 100644
--- /dev/null
+++ b/Engine/Source/Runtime/Experimental/Chaos/Private/Chaos/Joint/PBDJointSolverGaussSeidel.cpp
@@ -0,0 +1,373 @@
+// Copyright 1998-2019 Epic Games, Inc. All Rights Reserved.
+#include "Chaos/Joint/PBDJointSolverGaussSeidel.h"
+#include "Chaos/ParticleHandle.h"
+#include "Chaos/PBDJointConstraintUtilities.h"
+#include "Chaos/Utilities.h"
+#include "ChaosLog.h"
+#include "ChaosStats.h"
+
+//#pragma optimize("", off)
+
+namespace Chaos
+{
+	FJointSolverGaussSeidel::FJointSolverGaussSeidel()
+	{
+	}
+
+	void FJointSolverGaussSeidel::UpdateDerivedState()
+	{
+		Xs[0] = Ps[0] + Qs[0] * XLs[0].GetTranslation();
+		Xs[1] = Ps[1] + Qs[1] * XLs[1].GetTranslation();
+		Rs[0] = Qs[0] * XLs[0].GetRotation();
+		Rs[1] = Qs[1] * XLs[1].GetRotation();
+	}
+
+	void FJointSolverGaussSeidel::InitConstraints(
+		const FReal Dt,
+		const FPBDJointSolverSettings& SolverSettings,
+		const FPBDJointSettings& JointSettings,
+		const FVec3& P0,
+		const FRotation3& Q0,
+		const FVec3& P1,
+		const FRotation3& Q1,
+		const FReal InvM0,
+		const FMatrix33& InvIL0,
+		const FReal InvM1,
+		const FMatrix33& InvIL1,
+		const FRigidTransform3& XL0,
+		const FRigidTransform3& XL1)
+	{
+		XLs[0] = XL0;
+		XLs[1] = XL1;
+		InvILs[0] = InvIL0;
+		InvILs[1] = InvIL1;
+		InvMs[0] = InvM0;
+		InvMs[1] = InvM1;
+		// @todo(ccaulfield): mass conditioning
+
+		Ps[0] = P0;
+		Ps[1] = P1;
+		Qs[0] = Q0;
+		Qs[1] = Q1;
+		Qs[1].EnforceShortestArcWith(Qs[0]);
+
+		LinearStiffness = FPBDJointUtilities::GetLinearStiffness(SolverSettings, JointSettings);
+		TwistStiffness = FPBDJointUtilities::GetTwistStiffness(SolverSettings, JointSettings);
+		SwingStiffness = FPBDJointUtilities::GetSwingStiffness(SolverSettings, JointSettings);
+		SwingTwistAngleTolerance = SolverSettings.SwingTwistAngleTolerance;
+		bEnableTwistLimits = SolverSettings.bEnableTwistLimits;
+		bEnableTwistLimits = SolverSettings.bEnableTwistLimits;
+
+		UpdateDerivedState();
+	}
+
+	void FJointSolverGaussSeidel::ApplyConstraints(
+		const FReal Dt,
+		const FPBDJointSettings& JointSettings)
+	{
+		ApplyRotationConstraints(Dt, JointSettings);
+		ApplyPositionConstraints(Dt, JointSettings);
+	}
+
+
+	void FJointSolverGaussSeidel::ApplyRotationConstraints(
+		const FReal Dt,
+		const FPBDJointSettings& JointSettings)
+	{
+		EJointMotionType TwistMotion = JointSettings.Motion.AngularMotionTypes[(int32)EJointAngularConstraintIndex::Twist];
+		EJointMotionType Swing1Motion = JointSettings.Motion.AngularMotionTypes[(int32)EJointAngularConstraintIndex::Swing1];
+		EJointMotionType Swing2Motion = JointSettings.Motion.AngularMotionTypes[(int32)EJointAngularConstraintIndex::Swing2];
+
+		// Decompose rotation of body 1 relative to body 0 into swing and twist rotations, assuming twist is X axis
+		FRotation3 R01Twist, R01Swing;
+		FPBDJointUtilities::DecomposeSwingTwistLocal(Rs[0], Rs[1], R01Swing, R01Twist);
+
+		// Apply twist constraint
+		if (bEnableTwistLimits)
+		{
+			if (TwistMotion != EJointMotionType::Free)
+			{
+				ApplyJointTwistConstraint(Dt, JointSettings, R01Twist, R01Swing);
+			}
+		}
+
+		// Apply swing constraints
+		if (bEnableSwingLimits)
+		{
+			if ((Swing1Motion == EJointMotionType::Limited) && (Swing2Motion == EJointMotionType::Limited))
+			{
+				// Swing Cone
+				ApplyJointConeConstraint(Dt, JointSettings, R01Twist, R01Swing);
+			}
+			else
+			{
+				if (Swing1Motion != EJointMotionType::Free)
+				{
+					// Swing Arc/Lock
+					ApplyJointSwingConstraint(Dt, JointSettings, R01Twist, R01Swing, EJointAngularConstraintIndex::Swing1, EJointAngularAxisIndex::Swing1);
+				}
+				if (Swing2Motion != EJointMotionType::Free)
+				{
+					// Swing Arc/Lock
+					ApplyJointSwingConstraint(Dt, JointSettings, R01Twist, R01Swing, EJointAngularConstraintIndex::Swing2, EJointAngularAxisIndex::Swing2);
+				}
+			}
+		}
+	}
+
+
+	void FJointSolverGaussSeidel::ApplyPositionConstraints(
+		const FReal Dt,
+		const FPBDJointSettings& JointSettings)
+	{
+		const TVector<EJointMotionType, 3>& LinearMotion = JointSettings.Motion.LinearMotionTypes;
+		if ((LinearMotion[0] != EJointMotionType::Free) || (LinearMotion[1] != EJointMotionType::Free) || (LinearMotion[2] != EJointMotionType::Free))
+		{
+			ApplyJointPositionConstraint(Dt, JointSettings);
+		}
+	}
+
+
+	void FJointSolverGaussSeidel::ApplyPositionDelta(
+		const FReal Stiffness,
+		const FVec3& DP0,
+		const FVec3& DP1)
+	{
+		Ps[0] = Ps[0] + Stiffness * DP0;
+		Ps[1] = Ps[1] + Stiffness * DP1;
+	}
+
+
+	void FJointSolverGaussSeidel::ApplyRotationDelta(
+		const FReal Stiffness,
+		const FVec3& DR0,
+		const FVec3& DR1)
+	{
+		const FRotation3 DQ0 = (FRotation3::FromElements(Stiffness * DR0, 0) * Qs[0]) * (FReal)0.5;
+		const FRotation3 DQ1 = (FRotation3::FromElements(Stiffness * DR1, 0) * Qs[1]) * (FReal)0.5;
+		Qs[0] = (Qs[0] + DQ0).GetNormalized();
+		Qs[1] = (Qs[1] + DQ1).GetNormalized();
+		Qs[1].EnforceShortestArcWith(Qs[0]);
+	}
+
+	void FJointSolverGaussSeidel::ApplyRotationDelta(
+		const FReal Stiffness,
+		const FVec3& Axis0,
+		const FReal Angle0,
+		const FVec3& Axis1,
+		const FReal Angle1)
+	{
+		const FMatrix33 InvI0 = Utilities::ComputeWorldSpaceInertia(Qs[0], InvILs[0]);
+		const FMatrix33 InvI1 = Utilities::ComputeWorldSpaceInertia(Qs[1], InvILs[1]);
+
+		const FReal L = (FReal)1 / (FVec3::DotProduct(Axis0, Utilities::Multiply(InvI0, Axis0)) + FVec3::DotProduct(Axis1, Utilities::Multiply(InvI1, Axis1)));
+		const FVec3 DR0 = Utilities::Multiply(InvI0, Axis0) * L * Angle0;
+		const FVec3 DR1 = Utilities::Multiply(InvI1, Axis1) * L * Angle1;
+
+		ApplyRotationDelta(Stiffness, DR0, DR1);
+	}
+
+	void FJointSolverGaussSeidel::ApplyJointTwistConstraint(
+		const FReal Dt,
+		const FPBDJointSettings& JointSettings,
+		const FRotation3& R01Twist,
+		const FRotation3& R01Swing)
+	{
+		FVec3 TwistAxis01;
+		FReal TwistAngle;
+		R01Twist.ToAxisAndAngleSafe(TwistAxis01, TwistAngle, FJointConstants::TwistAxis(), SwingTwistAngleTolerance);
+		if (TwistAngle > PI)
+		{
+			TwistAngle = TwistAngle - (FReal)2 * PI;
+		}
+		if (FVec3::DotProduct(TwistAxis01, FJointConstants::TwistAxis()) < 0)
+		{
+			TwistAxis01 = -TwistAxis01;
+			TwistAngle = -TwistAngle;
+		}
+
+		const FVec3 TwistAxis0 = Rs[0] * TwistAxis01;
+		const FVec3 TwistAxis1 = Rs[1] * TwistAxis01;
+		FReal TwistAngleMax = FLT_MAX;
+		if (JointSettings.Motion.AngularMotionTypes[(int32)EJointAngularConstraintIndex::Twist] == EJointMotionType::Limited)
+		{
+			TwistAngleMax = JointSettings.Motion.AngularLimits[(int32)EJointAngularConstraintIndex::Twist];
+		}
+		else if (JointSettings.Motion.AngularMotionTypes[(int32)EJointAngularConstraintIndex::Twist] == EJointMotionType::Locked)
+		{
+			TwistAngleMax = 0;
+		}
+
+		// Calculate the twist correction to apply to each body
+		FReal DTwistAngle = 0;
+		if (TwistAngle > TwistAngleMax)
+		{
+			DTwistAngle = TwistAngle - TwistAngleMax;
+		}
+		else if (TwistAngle < -TwistAngleMax)
+		{
+			DTwistAngle = TwistAngle + TwistAngleMax;
+		}
+		const FReal DTwistAngle0 = DTwistAngle;
+		const FReal DTwistAngle1 = -DTwistAngle;
+
+		// Apply twist correction
+		ApplyRotationDelta(TwistStiffness, TwistAxis0, DTwistAngle0, TwistAxis1, DTwistAngle1);
+		UpdateDerivedState();
+	}
+
+
+	void FJointSolverGaussSeidel::ApplyJointConeConstraint(
+		const FReal Dt,
+		const FPBDJointSettings& JointSettings,
+		const FRotation3& R01Twist,
+		const FRotation3& R01Swing)
+	{
+		FVec3 SwingAxis01;
+		FReal SwingAngle;
+		R01Swing.ToAxisAndAngleSafe(SwingAxis01, SwingAngle, FJointConstants::Swing1Axis(), SwingTwistAngleTolerance);
+		if (SwingAngle > PI)
+		{
+			SwingAngle = SwingAngle - (FReal)2 * PI;
+		}
+
+		const FVec3 SwingAxis = Rs[0] * SwingAxis01;
+
+		// Calculate swing limit for the current swing axis
+		const FReal Swing1Limit = JointSettings.Motion.AngularLimits[(int32)EJointAngularConstraintIndex::Swing1];
+		const FReal Swing2Limit = JointSettings.Motion.AngularLimits[(int32)EJointAngularConstraintIndex::Swing2];
+
+		// Circular swing limit
+		FReal SwingAngleMax = Swing1Limit;
+
+		// Elliptical swing limit
+		// @todo(ccaulfield): do elliptical constraints properly (axis is still for circular limit)
+		if (!FMath::IsNearlyEqual(Swing1Limit, Swing2Limit, KINDA_SMALL_NUMBER))
+		{
+			// Map swing axis to ellipse and calculate limit for this swing axis
+			const FReal DotSwing1 = FMath::Abs(FVec3::DotProduct(SwingAxis01, FJointConstants::Swing1Axis()));
+			const FReal DotSwing2 = FMath::Abs(FVec3::DotProduct(SwingAxis01, FJointConstants::Swing2Axis()));
+			SwingAngleMax = FMath::Sqrt(Swing1Limit * DotSwing2 * Swing1Limit * DotSwing2 + Swing2Limit * DotSwing1 * Swing2Limit * DotSwing1);
+		}
+
+		// Calculate swing error we need to correct
+		FReal DSwingAngle = 0;
+		if (SwingAngle > SwingAngleMax)
+		{
+			DSwingAngle = SwingAngle - SwingAngleMax;
+		}
+		else if (SwingAngle < -SwingAngleMax)
+		{
+			DSwingAngle = SwingAngle + SwingAngleMax;
+		}
+
+		// Apply swing correction
+		ApplyRotationDelta(SwingStiffness, SwingAxis, DSwingAngle, SwingAxis, -DSwingAngle);
+		UpdateDerivedState();
+	}
+
+
+	void FJointSolverGaussSeidel::ApplyJointSwingConstraint(
+		const FReal Dt,
+		const FPBDJointSettings& JointSettings,
+		const FRotation3& R01Twist,
+		const FRotation3& R01Swing,
+		const EJointAngularConstraintIndex SwingConstraintIndex,
+		const EJointAngularAxisIndex SwingAxisIndex)
+	{
+		FVec3 TwistAxis01;
+		FReal TwistAngle;
+		R01Twist.ToAxisAndAngleSafe(TwistAxis01, TwistAngle, FJointConstants::TwistAxis(), SwingTwistAngleTolerance);
+		if (TwistAngle > PI)
+		{
+			TwistAngle = TwistAngle - (FReal)2 * PI;
+		}
+		if (FVec3::DotProduct(TwistAxis01, FJointConstants::TwistAxis()) < 0)
+		{
+			TwistAxis01 = -TwistAxis01;
+			TwistAngle = -TwistAngle;
+		}
+		const FVec3 TwistAxis = Rs[0] * TwistAxis01;
+
+		const FRotation3 R1NoTwist = Rs[1] * R01Twist.Inverse();
+		const FMatrix33 Axes0 = Rs[0].ToMatrix();
+		const FMatrix33 Axes1 = R1NoTwist.ToMatrix();
+		FVec3 SwingCross = FVec3::CrossProduct(Axes0.GetAxis((int32)SwingAxisIndex), Axes1.GetAxis((int32)SwingAxisIndex));
+		SwingCross = SwingCross - FVec3::DotProduct(TwistAxis, SwingCross) * TwistAxis;
+		const FReal SwingCrossLen = SwingCross.Size();
+		if (SwingCrossLen > KINDA_SMALL_NUMBER)
+		{
+			const FVec3 SwingAxis = SwingCross / SwingCrossLen;
+
+			FReal SwingAngle = FMath::Asin(FMath::Clamp(SwingCrossLen, (FReal)0, (FReal)1));
+			const FReal SwingDot = FVec3::DotProduct(Axes0.GetAxis((int32)SwingAxisIndex), Axes1.GetAxis((int32)SwingAxisIndex));
+			if (SwingDot < (FReal)0)
+			{
+				SwingAngle = (FReal)PI - SwingAngle;
+			}
+
+			FReal SwingAngleMax = FLT_MAX;
+			if (JointSettings.Motion.AngularMotionTypes[(int32)SwingConstraintIndex] == EJointMotionType::Limited)
+			{
+				FReal SwingLimit = JointSettings.Motion.AngularLimits[(int32)SwingConstraintIndex];
+				SwingAngleMax = SwingLimit;
+			}
+			else if (JointSettings.Motion.AngularMotionTypes[(int32)SwingConstraintIndex] == EJointMotionType::Locked)
+			{
+				SwingAngleMax = 0;
+			}
+
+			// Calculate swing error we need to correct
+			FReal DSwingAngle = 0;
+			if (SwingAngle > SwingAngleMax)
+			{
+				DSwingAngle = SwingAngle - SwingAngleMax;
+			}
+			else if (SwingAngle < -SwingAngleMax)
+			{
+				DSwingAngle = SwingAngle + SwingAngleMax;
+			}
+
+			// Apply swing correction
+			ApplyRotationDelta(SwingStiffness, SwingAxis, DSwingAngle, SwingAxis, -DSwingAngle);
+			UpdateDerivedState();
+		}
+	}
+
+	void FJointSolverGaussSeidel::ApplyJointPositionConstraint(
+		const FReal Dt,
+		const FPBDJointSettings& JointSettings)
+	{
+		// @todo(ccaulfield): we should really be calculating axes based on joint config, rather than fixing the error components
+		// Calculate constraint error
+		const FVec3 CX = FPBDJointUtilities::GetLimitedPositionError(JointSettings, Rs[0], Xs[1] - Xs[0]);
+
+		const FMatrix33 InvI0 = Utilities::ComputeWorldSpaceInertia(Qs[0], InvILs[0]);
+		const FMatrix33 InvI1 = Utilities::ComputeWorldSpaceInertia(Qs[1], InvILs[1]);
+
+		// Calculate constraint correction
+		FMatrix33 M0 = FMatrix33(0, 0, 0);
+		FMatrix33 M1 = FMatrix33(0, 0, 0);
+		if (InvMs[0] > 0)
+		{
+			M0 = Utilities::ComputeJointFactorMatrix(Xs[0] - Ps[0], InvI0, InvMs[0]);
+		}
+		if (InvMs[1] > 0)
+		{
+			M1 = Utilities::ComputeJointFactorMatrix(Xs[1] - Ps[1], InvI1, InvMs[1]);
+		}
+		const FMatrix33 MI = (M0 + M1).Inverse();
+		const FVec3 DX = Utilities::Multiply(MI, CX);
+
+		// Apply constraint correction
+		const FVec3 DP0 = InvMs[0] * DX;
+		const FVec3 DP1 = -InvMs[1] * DX;
+		const FVec3 DR0 = Utilities::Multiply(InvI0, FVec3::CrossProduct(Xs[0] - Ps[0], DX));
+		const FVec3 DR1 = Utilities::Multiply(InvI1, FVec3::CrossProduct(Xs[1] - Ps[1], -DX));
+
+		ApplyPositionDelta(LinearStiffness, DP0, DP1);
+		ApplyRotationDelta(LinearStiffness, DR0, DR1);
+		UpdateDerivedState();
+	}
+
+}
diff --git a/Engine/Source/Runtime/Experimental/Chaos/Private/Chaos/PBDJointConstraintSolver.cpp b/Engine/Source/Runtime/Experimental/Chaos/Private/Chaos/PBDJointConstraintSolver.cpp
deleted file mode 100644
diff --git a/Engine/Source/Runtime/Experimental/Chaos/Private/Chaos/PBDJointConstraintUtilities.cpp b/Engine/Source/Runtime/Experimental/Chaos/Private/Chaos/PBDJointConstraintUtilities.cpp
--- a/Engine/Source/Runtime/Experimental/Chaos/Private/Chaos/PBDJointConstraintUtilities.cpp
+++ b/Engine/Source/Runtime/Experimental/Chaos/Private/Chaos/PBDJointConstraintUtilities.cpp
@@ -2,7 +2,6 @@
 #include "Chaos/PBDJointConstraintUtilities.h"
 #include "Chaos/DenseMatrix.h"
 #include "Chaos/ParticleHandle.h"
-#include "Chaos/PBDJointConstraintSolver.h"
 #include "Chaos/Utilities.h"
 
 //#pragma optimize("", off)
@@ -10,6 +9,48 @@
 namespace Chaos
 {
 
+	void FPBDJointUtilities::DecomposeSwingTwistLocal(const FRotation3& R0, const FRotation3& R1, FRotation3& R01Swing, FRotation3& R01Twist)
+	{
+		const FRotation3 R01 = R0.Inverse() * R1;
+		R01.ToSwingTwistX(R01Swing, R01Twist);
+	}
+
+	FReal FPBDJointUtilities::GetLinearStiffness(
+		const FPBDJointSolverSettings& SolverSettings,
+		const FPBDJointSettings& JointSettings)
+	{
+		const FReal SolverStiffness = (SolverSettings.Stiffness > (FReal)0) ? SolverSettings.Stiffness : JointSettings.Motion.Stiffness;
+		const FReal SoftSolverStiffness = (SolverSettings.SoftLinearStiffness > (FReal)0) ? SolverSettings.SoftLinearStiffness : JointSettings.Motion.SoftLinearStiffness;
+		const bool bIsSoft = JointSettings.Motion.bSoftLinearLimitsEnabled && ((JointSettings.Motion.LinearMotionTypes[0] == EJointMotionType::Limited) || (JointSettings.Motion.LinearMotionTypes[1] == EJointMotionType::Limited) || (JointSettings.Motion.LinearMotionTypes[2] == EJointMotionType::Limited));
+		const FReal Stiffness = bIsSoft ? SolverStiffness * SoftSolverStiffness : SolverStiffness;
+		return Stiffness;
+	}
+
+
+	FReal FPBDJointUtilities::GetTwistStiffness(
+		const FPBDJointSolverSettings& SolverSettings,
+		const FPBDJointSettings& JointSettings)
+	{
+		const FReal SolverStiffness = (SolverSettings.Stiffness > (FReal)0) ? SolverSettings.Stiffness : JointSettings.Motion.Stiffness;
+		const FReal SoftSolverStiffness = (SolverSettings.SoftAngularStiffness > (FReal)0) ? SolverSettings.SoftAngularStiffness : JointSettings.Motion.SoftTwistStiffness;
+		const bool bIsSoft = JointSettings.Motion.bSoftTwistLimitsEnabled && (JointSettings.Motion.AngularMotionTypes[(int32)EJointAngularConstraintIndex::Twist] == EJointMotionType::Limited);
+		const FReal Stiffness = bIsSoft ? SolverStiffness * SoftSolverStiffness : SolverStiffness;
+		return Stiffness;
+	}
+
+
+	FReal FPBDJointUtilities::GetSwingStiffness(
+		const FPBDJointSolverSettings& SolverSettings,
+		const FPBDJointSettings& JointSettings)
+	{
+		const FReal SolverStiffness = (SolverSettings.Stiffness > (FReal)0) ? SolverSettings.Stiffness : JointSettings.Motion.Stiffness;
+		const FReal SoftSolverStiffness = (SolverSettings.SoftAngularStiffness > (FReal)0) ? SolverSettings.SoftAngularStiffness : JointSettings.Motion.SoftSwingStiffness;
+		const bool bIsSoft = JointSettings.Motion.bSoftSwingLimitsEnabled && ((JointSettings.Motion.AngularMotionTypes[(int32)EJointAngularConstraintIndex::Swing1] == EJointMotionType::Limited) || (JointSettings.Motion.AngularMotionTypes[(int32)EJointAngularConstraintIndex::Swing2] == EJointMotionType::Limited));
+		const FReal Stiffness = bIsSoft ? SolverStiffness * SoftSolverStiffness : SolverStiffness;
+		return Stiffness;
+	}
+
+
 	FVec3 FPBDJointUtilities::ConditionInertia(const FVec3& InI, const FReal MaxRatio)
 	{
 		FReal IMin = InI.Min();
@@ -122,7 +163,7 @@ namespace Chaos
 
 
 	
-	FVec3 GetSphereLimitedPositionError(const FVec3& CX, const FReal Radius)
+	FVec3 FPBDJointUtilities::GetSphereLimitedPositionError(const FVec3& CX, const FReal Radius)
 	{
 		FReal CXLen = CX.Size();
 		if (CXLen < Radius)
@@ -138,7 +179,7 @@ namespace Chaos
 	}
 
 	
-	FVec3 GetSphereLimitedVelocityError(const FVec3& CX, const FReal Radius, const FVec3& CV)
+	FVec3 FPBDJointUtilities::GetSphereLimitedVelocityError(const FVec3& CX, const FReal Radius, const FVec3& CV)
 	{
 		FReal CXLen = CX.Size();
 		if (CXLen < Radius)
@@ -155,7 +196,7 @@ namespace Chaos
 	}
 
 	
-	FVec3 GetCylinderLimitedPositionError(const FVec3& InCX, const FVec3& Axis, const FReal Limit, const EJointMotionType AxisMotion)
+	FVec3 FPBDJointUtilities::GetCylinderLimitedPositionError(const FVec3& InCX, const FVec3& Axis, const FReal Limit, const EJointMotionType AxisMotion)
 	{
 		FVec3 CXAxis = FVec3::DotProduct(InCX, Axis) * Axis;
 		FVec3 CXPlane = InCX - CXAxis;
@@ -177,7 +218,7 @@ namespace Chaos
 	}
 
 	
-	FVec3 GetCylinderLimitedVelocityError(const FVec3& InCX, const FVec3& Axis, const FReal Limit, const EJointMotionType AxisMotion, const FVec3& CV)
+	FVec3 FPBDJointUtilities::GetCylinderLimitedVelocityError(const FVec3& InCX, const FVec3& Axis, const FReal Limit, const EJointMotionType AxisMotion, const FVec3& CV)
 	{
 		FVec3 CXAxis = FVec3::DotProduct(InCX, Axis) * Axis;
 		FVec3 CXPlane = InCX - CXAxis;
@@ -204,7 +245,7 @@ namespace Chaos
 	}
 
 	
-	FVec3 GetLineLimitedPositionError(const FVec3& CX, const FVec3& Axis, const FReal Limit, const EJointMotionType AxisMotion)
+	FVec3 FPBDJointUtilities::GetLineLimitedPositionError(const FVec3& CX, const FVec3& Axis, const FReal Limit, const EJointMotionType AxisMotion)
 	{
 		FReal CXDist = FVec3::DotProduct(CX, Axis);
 		if ((AxisMotion == EJointMotionType::Free) || (FMath::Abs(CXDist) < Limit))
@@ -222,7 +263,7 @@ namespace Chaos
 	}
 
 	
-	FVec3 GetLineLimitedVelocityError(const FVec3& CX, const FVec3& Axis, const FReal Limit, const EJointMotionType AxisMotion, const FVec3& CV)
+	FVec3 FPBDJointUtilities::GetLineLimitedVelocityError(const FVec3& CX, const FVec3& Axis, const FReal Limit, const EJointMotionType AxisMotion, const FVec3& CV)
 	{
 		FReal CXDist = FVec3::DotProduct(CX, Axis);
 		FReal CVAxis = FVec3::DotProduct(CV, Axis);
@@ -242,7 +283,7 @@ namespace Chaos
 
 
 	
-	FVec3 GetLimitedPositionError(const FPBDJointSettings& JointSettings, const FRotation3& R0, const FVec3& InCX)
+	FVec3 FPBDJointUtilities::GetLimitedPositionError(const FPBDJointSettings& JointSettings, const FRotation3& R0, const FVec3& InCX)
 	{
 		const TVector<EJointMotionType, 3>& Motion = JointSettings.Motion.LinearMotionTypes;
 		if ((Motion[0] == EJointMotionType::Locked) && (Motion[1] == EJointMotionType::Locked) && (Motion[2] == EJointMotionType::Locked))
@@ -296,7 +337,7 @@ namespace Chaos
 	}
 
 	
-	FVec3 GetLimitedVelocityError(const FPBDJointSettings& JointSettings, const FRotation3& R0, const FVec3& InCX, const FVec3& InCV)
+	FVec3 FPBDJointUtilities::GetLimitedVelocityError(const FPBDJointSettings& JointSettings, const FRotation3& R0, const FVec3& InCX, const FVec3& InCV)
 	{
 		const TVector<EJointMotionType, 3>& Motion = JointSettings.Motion.LinearMotionTypes;
 		if ((Motion[0] == EJointMotionType::Locked) && (Motion[1] == EJointMotionType::Locked) && (Motion[2] == EJointMotionType::Locked))
diff --git a/Engine/Source/Runtime/Experimental/Chaos/Private/Chaos/PBDJointConstraints.cpp b/Engine/Source/Runtime/Experimental/Chaos/Private/Chaos/PBDJointConstraints.cpp
--- a/Engine/Source/Runtime/Experimental/Chaos/Private/Chaos/PBDJointConstraints.cpp
+++ b/Engine/Source/Runtime/Experimental/Chaos/Private/Chaos/PBDJointConstraints.cpp
@@ -2,9 +2,10 @@
 #include "Chaos/PBDJointConstraints.h"
 #include "Chaos/ChaosDebugDraw.h"
 #include "Chaos/DebugDrawQueue.h"
+#include "Chaos/Joint/PBDJointSolverCholesky.h"
+#include "Chaos/Joint/PBDJointSolverGaussSeidel.h"
 #include "Chaos/Particle/ParticleUtilities.h"
 #include "Chaos/ParticleHandle.h"
-#include "Chaos/PBDJointConstraintSolver.h"
 #include "Chaos/PBDJointConstraintUtilities.h"
 #include "Chaos/Utilities.h"
 #include "ChaosLog.h"
@@ -27,40 +28,6 @@ namespace Chaos
 	bool ChaosJoint_UseCholeskySolver = true;
 	FAutoConsoleVariableRef CVarChaosImmPhysDeltaTime(TEXT("p.Chaos.Joint.UseCholeskySolver"), ChaosJoint_UseCholeskySolver, TEXT("Whether to use the new solver"));
 
-	FReal GetLinearStiffness(
-		const FPBDJointSolverSettings& SolverSettings,
-		const FPBDJointSettings& JointSettings)
-	{
-		const FReal SolverStiffness = (SolverSettings.Stiffness > (FReal)0) ? SolverSettings.Stiffness : JointSettings.Motion.Stiffness;
-		const FReal SoftSolverStiffness = (SolverSettings.SoftLinearStiffness > (FReal)0) ? SolverSettings.SoftLinearStiffness : JointSettings.Motion.SoftLinearStiffness;
-		const bool bIsSoft = JointSettings.Motion.bSoftLinearLimitsEnabled && ((JointSettings.Motion.LinearMotionTypes[0] == EJointMotionType::Limited) || (JointSettings.Motion.LinearMotionTypes[1] == EJointMotionType::Limited) || (JointSettings.Motion.LinearMotionTypes[2] == EJointMotionType::Limited));
-		const FReal Stiffness = bIsSoft ? SolverStiffness * SoftSolverStiffness : SolverStiffness;
-		return Stiffness;
-	}
-
-
-	FReal GetTwistStiffness(
-		const FPBDJointSolverSettings& SolverSettings,
-		const FPBDJointSettings& JointSettings)
-	{
-		const FReal SolverStiffness = (SolverSettings.Stiffness > (FReal)0) ? SolverSettings.Stiffness : JointSettings.Motion.Stiffness;
-		const FReal SoftSolverStiffness = (SolverSettings.SoftAngularStiffness > (FReal)0) ? SolverSettings.SoftAngularStiffness : JointSettings.Motion.SoftTwistStiffness;
-		const bool bIsSoft = JointSettings.Motion.bSoftTwistLimitsEnabled && (JointSettings.Motion.AngularMotionTypes[(int32)EJointAngularConstraintIndex::Twist] == EJointMotionType::Limited);
-		const FReal Stiffness = bIsSoft ? SolverStiffness * SoftSolverStiffness : SolverStiffness;
-		return Stiffness;
-	}
-
-
-	FReal GetSwingStiffness(
-		const FPBDJointSolverSettings& SolverSettings,
-		const FPBDJointSettings& JointSettings)
-	{
-		const FReal SolverStiffness = (SolverSettings.Stiffness > (FReal)0) ? SolverSettings.Stiffness : JointSettings.Motion.Stiffness;
-		const FReal SoftSolverStiffness = (SolverSettings.SoftAngularStiffness > (FReal)0) ? SolverSettings.SoftAngularStiffness : JointSettings.Motion.SoftSwingStiffness;
-		const bool bIsSoft = JointSettings.Motion.bSoftSwingLimitsEnabled && ((JointSettings.Motion.AngularMotionTypes[(int32)EJointAngularConstraintIndex::Swing1] == EJointMotionType::Limited) || (JointSettings.Motion.AngularMotionTypes[(int32)EJointAngularConstraintIndex::Swing2] == EJointMotionType::Limited));
-		const FReal Stiffness = bIsSoft ? SolverStiffness * SoftSolverStiffness : SolverStiffness;
-		return Stiffness;
-	}
 
 	//
 	// Constraint Handle
@@ -204,7 +171,6 @@ namespace Chaos
 		, SwingTwistAngleTolerance((FReal)1.0e-6)
 		, MinParentMassRatio(0)
 		, MaxInertiaRatio(0)
-		, bEnableVelocitySolve(false)
 		, bEnableTwistLimits(true)
 		, bEnableSwingLimits(true)
 		, bEnableDrives(true)
@@ -415,6 +381,11 @@ namespace Chaos
 		}
 	}
 
+	//////////////////////////////////////////////////////////////////////////
+	//
+	// Begin Simple API Solver. Iterate over constraints in array order.
+	//
+	//////////////////////////////////////////////////////////////////////////
 
 	void FPBDJointConstraints::Apply(const FReal Dt, const int32 It, const int32 NumIts)
 	{
@@ -527,6 +498,17 @@ namespace Chaos
 		}
 	}
 
+	//////////////////////////////////////////////////////////////////////////
+	//
+	// End Simple API Solver.
+	//
+	//////////////////////////////////////////////////////////////////////////
+
+	//////////////////////////////////////////////////////////////////////////
+	//
+	// Begin Graph API Solver. Iterate over constraints in connectivity order.
+	//
+	//////////////////////////////////////////////////////////////////////////
 
 	void FPBDJointConstraints::Apply(const FReal Dt, const TArray<FConstraintContainerHandle*>& InConstraintHandles, const int32 It, const int32 NumIts)
 	{
@@ -569,14 +551,7 @@ namespace Chaos
 			SCOPE_CYCLE_COUNTER(STAT_Joints_Solve);
 			for (FConstraintContainerHandle* ConstraintHandle : SortedConstraintHandles)
 			{
-				if (Settings.bEnableVelocitySolve)
-				{
-					SolveVelocity(Dt, ConstraintHandle->GetConstraintIndex(), Settings.ApplyPairIterations, It, NumIts);
-				}
-				else
-				{
-					SolvePosition(Dt, ConstraintHandle->GetConstraintIndex(), Settings.ApplyPairIterations, It, NumIts);
-				}
+				SolvePosition(Dt, ConstraintHandle->GetConstraintIndex(), Settings.ApplyPairIterations, It, NumIts);
 			}
 		}
 
@@ -666,6 +641,18 @@ namespace Chaos
 		}
 	}
 
+	//////////////////////////////////////////////////////////////////////////
+	//
+	// End Graph API Solver.
+	//
+	//////////////////////////////////////////////////////////////////////////
+
+
+	//////////////////////////////////////////////////////////////////////////
+	//
+	// Begin single-particle solve methods used by APIs
+	//
+	//////////////////////////////////////////////////////////////////////////
 
 	void FPBDJointConstraints::UpdateParticleState(TPBDRigidParticleHandle<FReal, 3>* Rigid, const FReal Dt, const FVec3& P, const FRotation3& Q, const bool bUpdateVelocity)
 	{
@@ -779,32 +766,25 @@ namespace Chaos
 		TGenericParticleHandle<FReal, 3> Particle0 = TGenericParticleHandle<FReal, 3>(ConstraintParticles[ConstraintIndex][Index0]);
 		TGenericParticleHandle<FReal, 3> Particle1 = TGenericParticleHandle<FReal, 3>(ConstraintParticles[ConstraintIndex][Index1]);
 
-		const FVec3 P0 = FParticleUtilities::GetCoMWorldPosition(Particle0);
-		const FRotation3 Q0 = FParticleUtilities::GetCoMWorldRotation(Particle0);
-		const FVec3 P1 = FParticleUtilities::GetCoMWorldPosition(Particle1);
-		const FRotation3 Q1 = FParticleUtilities::GetCoMWorldRotation(Particle1);
-		const FRigidTransform3 XL0 = FParticleUtilities::ParticleLocalToCoMLocal(Particle0, JointSettings.ConstraintFrames[Index0]);
-		const FRigidTransform3 XL1 = FParticleUtilities::ParticleLocalToCoMLocal(Particle1, JointSettings.ConstraintFrames[Index1]);
-
-		FJointConstraintSolver Solver;
+		FJointSolverCholesky Solver;
 		Solver.InitConstraints(
 			Dt, 
 			Settings, 
 			JointSettings, 
-			P0,
-			Q0,
-			P1,
-			Q1,
+			FParticleUtilities::GetCoMWorldPosition(Particle0),
+			FParticleUtilities::GetCoMWorldRotation(Particle0),
+			FParticleUtilities::GetCoMWorldPosition(Particle1),
+			FParticleUtilities::GetCoMWorldRotation(Particle1),
 			Particle0->InvM(), 
 			Particle0->InvI(), 
 			Particle1->InvM(), 
 			Particle1->InvI(), 
-			XL0,
-			XL1);
+			FParticleUtilities::ParticleLocalToCoMLocal(Particle0, JointSettings.ConstraintFrames[Index0]),
+			FParticleUtilities::ParticleLocalToCoMLocal(Particle1, JointSettings.ConstraintFrames[Index1]));
 		
 		for (int32 PairIt = 0; PairIt < NumPairIts; ++PairIt)
 		{
-			Solver.ApplyConstraints(Dt, Settings, JointSettings, GetLinearStiffness(Settings, JointSettings));
+			Solver.ApplyConstraints(Dt, JointSettings);
 		}
 
 		UpdateParticleState(Particle0->CastToRigidParticle(), Dt, Solver.GetP(0), Solver.GetQ(0));
@@ -827,194 +807,29 @@ namespace Chaos
 		TGenericParticleHandle<FReal, 3> Particle0 = TGenericParticleHandle<FReal, 3>(ConstraintParticles[ConstraintIndex][Index0]);
 		TGenericParticleHandle<FReal, 3> Particle1 = TGenericParticleHandle<FReal, 3>(ConstraintParticles[ConstraintIndex][Index1]);
 
-		FVec3 P0 = FParticleUtilities::GetCoMWorldPosition(Particle0);
-		FRotation3 Q0 = FParticleUtilities::GetCoMWorldRotation(Particle0);
-		FVec3 P1 = FParticleUtilities::GetCoMWorldPosition(Particle1);
-		FRotation3 Q1 = FParticleUtilities::GetCoMWorldRotation(Particle1);
-		FVec3 V0 = Particle0->V();
-		FVec3 W0 = Particle0->W();
-		FVec3 V1 = Particle1->V();
-		FVec3 W1 = Particle1->W();
-		const FRigidTransform3 XL0 = FParticleUtilities::ParticleLocalToCoMLocal(Particle0, JointSettings.ConstraintFrames[Index0]);
-		const FRigidTransform3 XL1 = FParticleUtilities::ParticleLocalToCoMLocal(Particle1, JointSettings.ConstraintFrames[Index1]);
-		float InvM0 = Particle0->InvM();
-		float InvM1 = Particle1->InvM();
-		FMatrix33 InvIL0 = Particle0->InvI();
-		FMatrix33 InvIL1 = Particle1->InvI();
-
-		Q1.EnforceShortestArcWith(Q0);
-
-		FReal LinearStiffness = GetLinearStiffness(Settings, JointSettings);
-		FReal TwistStiffness = GetTwistStiffness(Settings, JointSettings);
-		FReal SwingStiffness = GetSwingStiffness(Settings, JointSettings);
-
-		// Adjust mass for stability
-		const int32 Level0 = ConstraintStates[ConstraintIndex].ParticleLevels[Index0];
-		const int32 Level1 = ConstraintStates[ConstraintIndex].ParticleLevels[Index1];
-		if (Level0 < Level1)
-		{
-			FPBDJointUtilities::GetConditionedInverseMass(Particle0->M(), Particle0->I().GetDiagonal(), Particle1->M(), Particle1->I().GetDiagonal(), InvM0, InvM1, InvIL0, InvIL1, Settings.MinParentMassRatio, Settings.MaxInertiaRatio);
-		}
-		else if (Level0 > Level1)
-		{
-			FPBDJointUtilities::GetConditionedInverseMass(Particle1->M(), Particle1->I().GetDiagonal(), Particle0->M(), Particle0->I().GetDiagonal(), InvM1, InvM0, InvIL1, InvIL0, Settings.MinParentMassRatio, Settings.MaxInertiaRatio);
-		}
-		else
-		{
-			FPBDJointUtilities::GetConditionedInverseMass(Particle0->M(), Particle0->I().GetDiagonal(), Particle1->M(), Particle1->I().GetDiagonal(), InvM0, InvM1, InvIL0, InvIL1, (FReal)0, Settings.MaxInertiaRatio);
-		}
-
-		const TVector<EJointMotionType, 3>& LinearMotion = JointSettings.Motion.LinearMotionTypes;
-		EJointMotionType TwistMotion = JointSettings.Motion.AngularMotionTypes[(int32)EJointAngularConstraintIndex::Twist];
-		EJointMotionType Swing1Motion = JointSettings.Motion.AngularMotionTypes[(int32)EJointAngularConstraintIndex::Swing1];
-		EJointMotionType Swing2Motion = JointSettings.Motion.AngularMotionTypes[(int32)EJointAngularConstraintIndex::Swing2];
-
-		for (int32 PairIt = 0; PairIt < NumPairIts; ++PairIt)
-		{
-			// Apply twist constraint
-			if (Settings.bEnableTwistLimits)
-			{
-				if (TwistMotion != EJointMotionType::Free)
-				{
-					FPBDJointUtilities::ApplyJointTwistConstraint(Dt, Settings, JointSettings, TwistStiffness, XL0, XL1, P0, Q0, V0, W0, P1, Q1, V1, W1, InvM0, InvIL0, InvM1, InvIL1);
-				}
-			}
-
-			// Apply swing constraints
-			if (Settings.bEnableSwingLimits)
-			{
-				if ((Swing1Motion == EJointMotionType::Limited) && (Swing2Motion == EJointMotionType::Limited))
-				{
-					// Swing Cone
-					FPBDJointUtilities::ApplyJointConeConstraint(Dt, Settings, JointSettings, SwingStiffness, XL0, XL1, P0, Q0, V0, W0, P1, Q1, V1, W1, InvM0, InvIL0, InvM1, InvIL1);
-				}
-				else
-				{
-					if (Swing1Motion != EJointMotionType::Free)
-					{
-						// Swing Arc/Lock
-						FPBDJointUtilities::ApplyJointSwingConstraint(Dt, Settings, JointSettings, SwingStiffness, XL0, XL1, EJointAngularConstraintIndex::Swing1, EJointAngularAxisIndex::Swing1, P0, Q0, V0, W0, P1, Q1, V1, W1, InvM0, InvIL0, InvM1, InvIL1);
-					}
-					if (Swing2Motion != EJointMotionType::Free)
-					{
-						// Swing Arc/Lock
-						FPBDJointUtilities::ApplyJointSwingConstraint(Dt, Settings, JointSettings, SwingStiffness, XL0, XL1, EJointAngularConstraintIndex::Swing2, EJointAngularAxisIndex::Swing2, P0, Q0, V0, W0, P1, Q1, V1, W1, InvM0, InvIL0, InvM1, InvIL1);
-					}
-				}
-			}
-
-			// Apply linear constraints
-			if ((LinearMotion[0] != EJointMotionType::Free) || (LinearMotion[1] != EJointMotionType::Free) || (LinearMotion[2] != EJointMotionType::Free))
-			{
-				FPBDJointUtilities::ApplyJointPositionConstraint(Dt, Settings, JointSettings, LinearStiffness, XL0, XL1, P0, Q0, V0, W0, P1, Q1, V1, W1, InvM0, InvIL0, InvM1, InvIL1);
-			}
-		}
-
-		// Update the particles
-		UpdateParticleState(Particle0->CastToRigidParticle(), Dt, P0, Q0);
-		UpdateParticleState(Particle1->CastToRigidParticle(), Dt, P1, Q1);
-	}
-
-	
-	void FPBDJointConstraints::SolveVelocity(const FReal Dt, const int32 ConstraintIndex, const int32 NumPairIts, const int32 It, const int32 NumIts)
-	{
-		const TVector<TGeometryParticleHandle<FReal, 3>*, 2>& Constraint = ConstraintParticles[ConstraintIndex];
-		UE_LOG(LogChaosJoint, VeryVerbose, TEXT("Solve Joint Constraint %d %s %s (dt = %f; it = %d / %d)"), ConstraintIndex, *Constraint[0]->ToString(), *Constraint[1]->ToString(), Dt, It, NumIts);
-
-		const FPBDJointSettings& JointSettings = ConstraintSettings[ConstraintIndex];
-
-		// Switch particles - internally we assume the first body is the parent (i.e., the space in which constraint limits are specified)
-		const int32 Index0 = 1;
-		const int32 Index1 = 0;
-		TGenericParticleHandle<FReal, 3> Particle0 = TGenericParticleHandle<FReal, 3>(ConstraintParticles[ConstraintIndex][Index0]);
-		TGenericParticleHandle<FReal, 3> Particle1 = TGenericParticleHandle<FReal, 3>(ConstraintParticles[ConstraintIndex][Index1]);
-
-		FVec3 P0 = FParticleUtilities::GetCoMWorldPosition(Particle0);
-		FRotation3 Q0 = FParticleUtilities::GetCoMWorldRotation(Particle0);
-		FVec3 P1 = FParticleUtilities::GetCoMWorldPosition(Particle1);
-		FRotation3 Q1 = FParticleUtilities::GetCoMWorldRotation(Particle1);
-		FVec3 V0 = Particle0->V();
-		FVec3 W0 = Particle0->W();
-		FVec3 V1 = Particle1->V();
-		FVec3 W1 = Particle1->W();
-		const FRigidTransform3 XL0 = FParticleUtilities::ParticleLocalToCoMLocal(Particle0, JointSettings.ConstraintFrames[Index0]);
-		const FRigidTransform3 XL1 = FParticleUtilities::ParticleLocalToCoMLocal(Particle1, JointSettings.ConstraintFrames[Index1]);
-		float InvM0 = Particle0->InvM();
-		float InvM1 = Particle1->InvM();
-		FMatrix33 InvIL0 = Particle0->InvI();
-		FMatrix33 InvIL1 = Particle1->InvI();
-
-		Q1.EnforceShortestArcWith(Q0);
-
-		FReal LinearStiffness = GetLinearStiffness(Settings, JointSettings);
-		FReal TwistStiffness = GetTwistStiffness(Settings, JointSettings);
-		FReal SwingStiffness = GetSwingStiffness(Settings, JointSettings);
-
-		// Adjust mass for stability
-		const int32 Level0 = ConstraintStates[ConstraintIndex].ParticleLevels[Index0];
-		const int32 Level1 = ConstraintStates[ConstraintIndex].ParticleLevels[Index1];
-		if (Level0 < Level1)
-		{
-			FPBDJointUtilities::GetConditionedInverseMass(Particle0->M(), Particle0->I().GetDiagonal(), Particle1->M(), Particle1->I().GetDiagonal(), InvM0, InvM1, InvIL0, InvIL1, Settings.MinParentMassRatio, Settings.MaxInertiaRatio);
-		}
-		else if (Level0 > Level1)
-		{
-			FPBDJointUtilities::GetConditionedInverseMass(Particle1->M(), Particle1->I().GetDiagonal(), Particle0->M(), Particle0->I().GetDiagonal(), InvM1, InvM0, InvIL1, InvIL0, Settings.MinParentMassRatio, Settings.MaxInertiaRatio);
-		}
-		else
-		{
-			FPBDJointUtilities::GetConditionedInverseMass(Particle0->M(), Particle0->I().GetDiagonal(), Particle1->M(), Particle1->I().GetDiagonal(), InvM0, InvM1, InvIL0, InvIL1, (FReal)0, Settings.MaxInertiaRatio);
-		}
-
-		const TVector<EJointMotionType, 3>& LinearMotion = JointSettings.Motion.LinearMotionTypes;
-		EJointMotionType TwistMotion = JointSettings.Motion.AngularMotionTypes[(int32)EJointAngularConstraintIndex::Twist];
-		EJointMotionType Swing1Motion = JointSettings.Motion.AngularMotionTypes[(int32)EJointAngularConstraintIndex::Swing1];
-		EJointMotionType Swing2Motion = JointSettings.Motion.AngularMotionTypes[(int32)EJointAngularConstraintIndex::Swing2];
+		FJointSolverGaussSeidel Solver;
+		Solver.InitConstraints(
+			Dt,
+			Settings,
+			JointSettings,
+			FParticleUtilities::GetCoMWorldPosition(Particle0),
+			FParticleUtilities::GetCoMWorldRotation(Particle0),
+			FParticleUtilities::GetCoMWorldPosition(Particle1),
+			FParticleUtilities::GetCoMWorldRotation(Particle1),
+			Particle0->InvM(),
+			Particle0->InvI(),
+			Particle1->InvM(),
+			Particle1->InvI(),
+			FParticleUtilities::ParticleLocalToCoMLocal(Particle0, JointSettings.ConstraintFrames[Index0]),
+			FParticleUtilities::ParticleLocalToCoMLocal(Particle1, JointSettings.ConstraintFrames[Index1]));
 
 		for (int32 PairIt = 0; PairIt < NumPairIts; ++PairIt)
 		{
-			// Apply twist velocity constraint
-			if (Settings.bEnableTwistLimits)
-			{
-				if (TwistMotion != EJointMotionType::Free)
-				{
-					FPBDJointUtilities::ApplyJointTwistVelocityConstraint(Dt, Settings, JointSettings, TwistStiffness, XL0, XL1, P0, Q0, V0, W0, P1, Q1, V1, W1, InvM0, InvIL0, InvM1, InvIL1);
-				}
-			}
-
-			// Apply swing velocity constraints
-			if (Settings.bEnableSwingLimits)
-			{
-				if ((Swing1Motion == EJointMotionType::Limited) && (Swing2Motion == EJointMotionType::Limited))
-				{
-					// Swing Cone
-					FPBDJointUtilities::ApplyJointConeVelocityConstraint(Dt, Settings, JointSettings, SwingStiffness, XL0, XL1, P0, Q0, V0, W0, P1, Q1, V1, W1, InvM0, InvIL0, InvM1, InvIL1);
-				}
-				else
-				{
-					if (Swing1Motion != EJointMotionType::Free)
-					{
-						// Swing Arc/Lock
-						FPBDJointUtilities::ApplyJointSwingVelocityConstraint(Dt, Settings, JointSettings, SwingStiffness, XL0, XL1, EJointAngularConstraintIndex::Swing1, EJointAngularAxisIndex::Swing1, P0, Q0, V0, W0, P1, Q1, V1, W1, InvM0, InvIL0, InvM1, InvIL1);
-					}
-					if (Swing2Motion != EJointMotionType::Free)
-					{
-						// Swing Arc/Lock
-						FPBDJointUtilities::ApplyJointSwingVelocityConstraint(Dt, Settings, JointSettings, SwingStiffness, XL0, XL1, EJointAngularConstraintIndex::Swing2, EJointAngularAxisIndex::Swing2, P0, Q0, V0, W0, P1, Q1, V1, W1, InvM0, InvIL0, InvM1, InvIL1);
-					}
-				}
-			}
-
-			// Apply linear velocity  constraints
-			if ((LinearMotion[0] != EJointMotionType::Free) || (LinearMotion[1] != EJointMotionType::Free) || (LinearMotion[2] != EJointMotionType::Free))
-			{
-				FPBDJointUtilities::ApplyJointVelocityConstraint(Dt, Settings, JointSettings, LinearStiffness, XL0, XL1, P0, Q0, V0, W0, P1, Q1, V1, W1, InvM0, InvIL0, InvM1, InvIL1);
-			}
+			Solver.ApplyConstraints(Dt, JointSettings);
 		}
 
-		// Update the particles
-		UpdateParticleState(Particle0->CastToRigidParticle(), Dt, P0, Q0);
-		UpdateParticleState(Particle1->CastToRigidParticle(), Dt, P1, Q1);
+		UpdateParticleState(Particle0->CastToRigidParticle(), Dt, Solver.GetP(0), Solver.GetQ(0));
+		UpdateParticleState(Particle1->CastToRigidParticle(), Dt, Solver.GetP(1), Solver.GetQ(1));
 	}
 
 
diff --git a/Engine/Source/Runtime/Experimental/Chaos/Public/Chaos/Joint/PBDJointSolverCholesky.h b/Engine/Source/Runtime/Experimental/Chaos/Public/Chaos/Joint/PBDJointSolverCholesky.h
new file mode 100644
--- /dev/null
+++ b/Engine/Source/Runtime/Experimental/Chaos/Public/Chaos/Joint/PBDJointSolverCholesky.h
@@ -0,0 +1,159 @@
+// Copyright 1998-2019 Epic Games, Inc. All Rights Reserved.
+#pragma once
+
+#include "CoreMinimal.h"
+
+#include "Chaos/DenseMatrix.h"
+#include "Chaos/ParticleHandleFwd.h"
+#include "Chaos/PBDJointConstraintTypes.h"
+#include "Chaos/PBDJointConstraintUtilities.h"
+#include "Chaos/Utilities.h"
+
+namespace Chaos
+{
+	/**
+	 * Calculate new positions and rotations for a pair of bodies connected by a joint.
+	 *
+	 * This solver treats of the 6 possible constraints (up to 3 linear and 3 angular)
+	 * as a single block and resolves them in simultaneously.
+	 *
+	 * \see FJointSolverGaussSeidel
+	 */
+	class FJointSolverCholesky
+	{
+	public:
+		static const int32 MaxConstraintedBodies = 2;
+
+		using FDenseMatrix66 = TDenseMatrix<6 * 6>;
+		using FDenseMatrix61 = TDenseMatrix<6 * 1>;
+
+		FORCEINLINE const FVec3& GetP(const int32 Index) const
+		{
+			checkSlow(Index < MaxConstraintedBodies);
+			return Ps[Index];
+		}
+
+		FORCEINLINE const FRotation3& GetQ(const int32 Index) const
+		{
+			checkSlow(Index < MaxConstraintedBodies);
+			return Qs[Index];
+		}
+
+		FJointSolverCholesky();
+
+		void InitConstraints(
+			const FReal Dt,
+			const FPBDJointSolverSettings& SolverSettings,
+			const FPBDJointSettings& JointSettings,
+			const FVec3& P0,
+			const FRotation3& Q0,
+			const FVec3& P1,
+			const FRotation3& Q1,
+			const FReal InvM0,
+			const FMatrix33& InvIL0,
+			const FReal InvM1,
+			const FMatrix33& InvIL1,
+			const FRigidTransform3& XL0,
+			const FRigidTransform3& XL1);
+
+		void ApplyConstraints(
+			const FReal Dt,
+			const FPBDJointSettings& JointSettings);
+
+	private:
+
+		void UpdateDerivedState();
+
+		void AddLinearConstraints_Point(
+			const FPBDJointSettings& JointSettings,
+			FDenseMatrix66& J0,
+			FDenseMatrix66& J1,
+			FDenseMatrix61& C);
+
+		void AddLinearConstraints_Sphere(
+			const FPBDJointSettings& JointSettings,
+			FDenseMatrix66& J0,
+			FDenseMatrix66& J1,
+			FDenseMatrix61& C);
+
+		void AddLinearConstraints_Cylinder(
+			const FPBDJointSettings& JointSettings,
+			const EJointMotionType AxisMotion,
+			const FVec3& Axis,
+			FDenseMatrix66& J0,
+			FDenseMatrix66& J1,
+			FDenseMatrix61& C);
+
+		void AddLinearConstraints_Plane(
+			const FPBDJointSettings& JointSettings,
+			const EJointMotionType AxisMotion,
+			const FVec3& Axis,
+			FDenseMatrix66& J0,
+			FDenseMatrix66& J1,
+			FDenseMatrix61& C);
+
+		void AddAngularConstraints_Twist(
+			const FPBDJointSettings& JointSettings,
+			const FRotation3& R01Twist,
+			const FRotation3& R01Swing,
+			FDenseMatrix66& J0,
+			FDenseMatrix66& J1,
+			FDenseMatrix61& C);
+
+		void AddAngularConstraints_Cone(
+			const FPBDJointSettings& JointSettings,
+			const FRotation3& R01Twist,
+			const FRotation3& R01Swing,
+			FDenseMatrix66& J0,
+			FDenseMatrix66& J1,
+			FDenseMatrix61& C);
+
+		void AddAngularConstraints_Swing(
+			const FPBDJointSettings& JointSettings,
+			const EJointAngularConstraintIndex SwingConstraintIndex,
+			const EJointAngularAxisIndex SwingAxisIndex,
+			const FRotation3& R01Twist,
+			const FRotation3& R01Swing,
+			FDenseMatrix66& J0,
+			FDenseMatrix66& J1,
+			FDenseMatrix61& C);
+
+		void AddLinearConstraints(
+			const FPBDJointSettings& JointSettings,
+			FDenseMatrix66& J0,
+			FDenseMatrix66& J1,
+			FDenseMatrix61& C);
+
+		void AddAngularConstraints(
+			const FPBDJointSettings& JointSettings,
+			FDenseMatrix66& J0,
+			FDenseMatrix66& J1,
+			FDenseMatrix61& C);
+
+		void BuildJacobianAndResidual(
+			const FPBDJointSettings& JointSettings,
+			FDenseMatrix66& J0,
+			FDenseMatrix66& J1,
+			FDenseMatrix61& C);
+
+		// Local-space constraint settings
+		FRigidTransform3 XLs[MaxConstraintedBodies];	// Local-space joint connector transforms
+		FMatrix33 InvILs[MaxConstraintedBodies];		// Local-space inverse inertias
+		FReal InvMs[MaxConstraintedBodies];				// Inverse masses
+
+		// World-space constraint state
+		FVec3 Xs[MaxConstraintedBodies];				// World-space joint connector positions
+		FRotation3 Rs[MaxConstraintedBodies];			// World-space joint connector rotations
+
+		// World-space body state
+		FVec3 Ps[MaxConstraintedBodies];				// World-space particle CoM positions
+		FRotation3 Qs[MaxConstraintedBodies];			// World-space particle CoM rotations
+
+		// Settings
+		FReal Stiffness;
+		FReal SwingTwistAngleTolerance;
+		bool bEnableTwistLimits;
+		bool bEnableSwingLimits;
+	};
+
+}
\ No newline at end of file
diff --git a/Engine/Source/Runtime/Experimental/Chaos/Public/Chaos/Joint/PBDJointSolverGaussSeidel.h b/Engine/Source/Runtime/Experimental/Chaos/Public/Chaos/Joint/PBDJointSolverGaussSeidel.h
new file mode 100644
--- /dev/null
+++ b/Engine/Source/Runtime/Experimental/Chaos/Public/Chaos/Joint/PBDJointSolverGaussSeidel.h
@@ -0,0 +1,135 @@
+// Copyright 1998-2019 Epic Games, Inc. All Rights Reserved.
+#pragma once
+
+#include "CoreMinimal.h"
+
+#include "Chaos/DenseMatrix.h"
+#include "Chaos/ParticleHandleFwd.h"
+#include "Chaos/PBDJointConstraintTypes.h"
+#include "Chaos/PBDJointConstraintUtilities.h"
+#include "Chaos/Utilities.h"
+
+namespace Chaos
+{
+	/**
+	 * Calculate new positions and rotations for a pair of bodies connected by a joint.
+	 *
+	 * This solver treats of the 6 possible constraints (up to 3 linear and 3 angular)
+	 * individually and resolves them in sequence.
+	 *
+	 * \see FJointSolverCholesky
+	 */
+	class FJointSolverGaussSeidel
+	{
+	public:
+		static const int32 MaxConstraintedBodies = 2;
+
+		FORCEINLINE const FVec3& GetP(const int32 Index) const
+		{
+			checkSlow(Index < MaxConstraintedBodies);
+			return Ps[Index];
+		}
+
+		FORCEINLINE const FRotation3& GetQ(const int32 Index) const
+		{
+			checkSlow(Index < MaxConstraintedBodies);
+			return Qs[Index];
+		}
+
+		FJointSolverGaussSeidel();
+
+		void InitConstraints(
+			const FReal Dt,
+			const FPBDJointSolverSettings& SolverSettings,
+			const FPBDJointSettings& JointSettings,
+			const FVec3& P0,
+			const FRotation3& Q0,
+			const FVec3& P1,
+			const FRotation3& Q1,
+			const FReal InvM0,
+			const FMatrix33& InvIL0,
+			const FReal InvM1,
+			const FMatrix33& InvIL1,
+			const FRigidTransform3& XL0,
+			const FRigidTransform3& XL1);
+
+		void ApplyConstraints(
+			const FReal Dt,
+			const FPBDJointSettings& JointSettings);
+
+	private:
+
+		void UpdateDerivedState();
+
+		void ApplyPositionConstraints(
+			const FReal Dt,
+			const FPBDJointSettings& JointSettings);
+
+		void ApplyRotationConstraints(
+			const FReal Dt,
+			const FPBDJointSettings& JointSettings);
+
+		void ApplyPositionDelta(
+			const FReal Stiffness,
+			const FVec3& DP0,
+			const FVec3& DP1);
+
+		void ApplyRotationDelta(
+			const FReal Stiffness,
+			const FVec3& DR0,
+			const FVec3& DR1);
+
+		void ApplyRotationDelta(
+			const FReal Stiffness,
+			const FVec3& Axis0,
+			const FReal Angle0,
+			const FVec3& Axis1,
+			const FReal Angle1);
+
+		void ApplyJointTwistConstraint(
+			const FReal Dt,
+			const FPBDJointSettings& JointSettings,
+			const FRotation3& R01Twist,
+			const FRotation3& R01Swing);
+
+		void ApplyJointConeConstraint(
+			const FReal Dt,
+			const FPBDJointSettings& JointSettings,
+			const FRotation3& R01Twist,
+			const FRotation3& R01Swing);
+
+		void ApplyJointSwingConstraint(
+			const FReal Dt,
+			const FPBDJointSettings& JointSettings,
+			const FRotation3& R01Twist,
+			const FRotation3& R01Swing,
+			const EJointAngularConstraintIndex SwingConstraintIndex,
+			const EJointAngularAxisIndex SwingAxisIndex);
+
+		void ApplyJointPositionConstraint(
+			const FReal Dt,
+			const FPBDJointSettings& JointSettings);
+
+		// Local-space constraint settings
+		FRigidTransform3 XLs[MaxConstraintedBodies];	// Local-space joint connector transforms
+		FMatrix33 InvILs[MaxConstraintedBodies];		// Local-space inverse inertias
+		FReal InvMs[MaxConstraintedBodies];				// Inverse masses
+
+		// World-space constraint state
+		FVec3 Xs[MaxConstraintedBodies];				// World-space joint connector positions
+		FRotation3 Rs[MaxConstraintedBodies];			// World-space joint connector rotations
+
+		// World-space body state
+		FVec3 Ps[MaxConstraintedBodies];				// World-space particle CoM positions
+		FRotation3 Qs[MaxConstraintedBodies];			// World-space particle CoM rotations
+
+		// Settings
+		FReal LinearStiffness;
+		FReal TwistStiffness;
+		FReal SwingStiffness;
+		FReal SwingTwistAngleTolerance;
+		bool bEnableTwistLimits;
+		bool bEnableSwingLimits;
+	};
+
+}
\ No newline at end of file
diff --git a/Engine/Source/Runtime/Experimental/Chaos/Public/Chaos/PBDJointConstraintSolver.h b/Engine/Source/Runtime/Experimental/Chaos/Public/Chaos/PBDJointConstraintSolver.h
deleted file mode 100644
diff --git a/Engine/Source/Runtime/Experimental/Chaos/Public/Chaos/PBDJointConstraintTypes.h b/Engine/Source/Runtime/Experimental/Chaos/Public/Chaos/PBDJointConstraintTypes.h
--- a/Engine/Source/Runtime/Experimental/Chaos/Public/Chaos/PBDJointConstraintTypes.h
+++ b/Engine/Source/Runtime/Experimental/Chaos/Public/Chaos/PBDJointConstraintTypes.h
@@ -130,7 +130,6 @@ namespace Chaos
 		FReal MaxInertiaRatio;
 
 		// @todo(ccaulfield): remove these TEMP overrides for testing
-		bool bEnableVelocitySolve;
 		bool bEnableTwistLimits;
 		bool bEnableSwingLimits;
 		bool bEnableDrives;
diff --git a/Engine/Source/Runtime/Experimental/Chaos/Public/Chaos/PBDJointConstraintUtilities.h b/Engine/Source/Runtime/Experimental/Chaos/Public/Chaos/PBDJointConstraintUtilities.h
--- a/Engine/Source/Runtime/Experimental/Chaos/Public/Chaos/PBDJointConstraintUtilities.h
+++ b/Engine/Source/Runtime/Experimental/Chaos/Public/Chaos/PBDJointConstraintUtilities.h
@@ -12,6 +12,8 @@ namespace Chaos
 	class FPBDJointUtilities
 	{
 	public:
+		static void DecomposeSwingTwistLocal(const FRotation3& R0, const FRotation3& R1, FRotation3& R01Swing, FRotation3& R01Twist);
+
 		/**
 		 * Increase the lower inertia components to ensure that the maximum ratio between any pair of elements is MaxRatio.
 		 *
@@ -61,6 +63,29 @@ namespace Chaos
 			FMatrix33& OutInvI0, 
 			const FReal MaxInertiaRatio);
 
+
+		static FReal GetLinearStiffness(
+			const FPBDJointSolverSettings& SolverSettings,
+			const FPBDJointSettings& JointSettings);
+		
+		static FReal GetTwistStiffness(
+			const FPBDJointSolverSettings& SolverSettings,
+			const FPBDJointSettings& JointSettings);
+
+		static FReal GetSwingStiffness(
+			const FPBDJointSolverSettings& SolverSettings,
+			const FPBDJointSettings& JointSettings);
+
+
+		static FVec3 GetSphereLimitedPositionError(const FVec3& CX, const FReal Radius);
+		static FVec3 GetSphereLimitedVelocityError(const FVec3& CX, const FReal Radius, const FVec3& CV);
+		static FVec3 GetCylinderLimitedPositionError(const FVec3& InCX, const FVec3& Axis, const FReal Limit, const EJointMotionType AxisMotion);
+		static FVec3 GetCylinderLimitedVelocityError(const FVec3& InCX, const FVec3& Axis, const FReal Limit, const EJointMotionType AxisMotion, const FVec3& CV);
+		static FVec3 GetLineLimitedPositionError(const FVec3& CX, const FVec3& Axis, const FReal Limit, const EJointMotionType AxisMotion);
+		static FVec3 GetLineLimitedVelocityError(const FVec3& CX, const FVec3& Axis, const FReal Limit, const EJointMotionType AxisMotion, const FVec3& CV);
+		static FVec3 GetLimitedPositionError(const FPBDJointSettings& JointSettings, const FRotation3& R0, const FVec3& InCX);
+		static FVec3 GetLimitedVelocityError(const FPBDJointSettings& JointSettings, const FRotation3& R0, const FVec3& InCX, const FVec3& InCV);
+
 		static CHAOS_API void CalculateSwingConstraintSpace(
 			const FPBDJointSolverSettings& SolverSettings,
 			const FPBDJointSettings& JointSettings,
diff --git a/Engine/Source/Runtime/Experimental/Chaos/Public/Chaos/PBDJointConstraints.h b/Engine/Source/Runtime/Experimental/Chaos/Public/Chaos/PBDJointConstraints.h
--- a/Engine/Source/Runtime/Experimental/Chaos/Public/Chaos/PBDJointConstraints.h
+++ b/Engine/Source/Runtime/Experimental/Chaos/Public/Chaos/PBDJointConstraints.h
@@ -157,7 +157,6 @@ namespace Chaos
 		void UpdateParticleState(TPBDRigidParticleHandle<FReal, 3>* Rigid, const FReal Dt, const FVec3& P, const FRotation3& Q, const bool bUpdateVelocity = true);
 
 		void ApplyDrives(const FReal Dt, const int32 ConstraintIndex, const int32 It, const int32 NumIts);
-		void SolveVelocity(const FReal Dt, const int32 ConstraintIndex, const int32 NumPairIts, const int32 It, const int32 NumIts);
 		void SolvePosition(const FReal Dt, const int32 ConstraintIndex, const int32 NumPairIts, const int32 It, const int32 NumIts);
 		void SolvePosition_GaussSiedel(const FReal Dt, const int32 ConstraintIndex, const int32 NumPairIts, const int32 It, const int32 NumIts);
 		void SolvePosition_Cholesky(const FReal Dt, const int32 ConstraintIndex, const int32 NumPairIts, const int32 It, const int32 NumIts);
diff --git a/Engine/Source/Runtime/Experimental/Chaos/Public/Chaos/Rotation.h b/Engine/Source/Runtime/Experimental/Chaos/Public/Chaos/Rotation.h
--- a/Engine/Source/Runtime/Experimental/Chaos/Public/Chaos/Rotation.h
+++ b/Engine/Source/Runtime/Experimental/Chaos/Public/Chaos/Rotation.h
@@ -106,7 +106,7 @@ namespace Chaos
 		 * Extract the Swing and Twist rotations, assuming that the Twist Axis is (1,0,0).
 		 * /see ToSwingTwist
 		 */
-		void ToSwingTwistX(const FVector& InTwistAxis, FQuat& OutSwing, FQuat& OutTwist) const
+		void ToSwingTwistX(FQuat& OutSwing, FQuat& OutTwist) const
 		{
 			OutTwist = (X != 0.0f)? FQuat(X, 0, 0, W).GetNormalized() : FQuat::Identity;
 			OutSwing = *this * OutTwist.Inverse();
