Clang CIS fix.

diff --git a/Engine/Source/Runtime/Experimental/Chaos/Private/Chaos/ParticleHandle.cpp b/Engine/Source/Runtime/Experimental/Chaos/Private/Chaos/ParticleHandle.cpp
--- a/Engine/Source/Runtime/Experimental/Chaos/Private/Chaos/ParticleHandle.cpp
+++ b/Engine/Source/Runtime/Experimental/Chaos/Private/Chaos/ParticleHandle.cpp
@@ -16,6 +16,7 @@ namespace Chaos
 	template class CHAOS_API TKinematicGeometryParticle<float, 3>;
 
 	template class CHAOS_API TPBDRigidParticleData<float, 3>;
+	template class CHAOS_API TPBDRigidParticle<float, 3>;
 
 	template <>
 	void Chaos::TGeometryParticle<float, 3>::MarkDirty(const EParticleFlags DirtyBits, bool bInvalidate )
diff --git a/Engine/Source/Runtime/Experimental/Chaos/Public/Chaos/ParticleHandle.h b/Engine/Source/Runtime/Experimental/Chaos/Public/Chaos/ParticleHandle.h
--- a/Engine/Source/Runtime/Experimental/Chaos/Public/Chaos/ParticleHandle.h
+++ b/Engine/Source/Runtime/Experimental/Chaos/Public/Chaos/ParticleHandle.h
@@ -2237,5 +2237,14 @@ FChaosArchive& operator<<(FChaosArchive& Ar, TAccelerationStructureHandle<T, d>&
 	return Ar;
 }
 
+#if PLATFORM_MAC || PLATFORM_LINUX
+extern template class CHAOS_API TGeometryParticle<float, 3>;
+extern template class CHAOS_API TKinematicGeometryParticle<float, 3>;
+extern template class CHAOS_API TPBDRigidParticle<float, 3>;
+#else
+extern template class TGeometryParticle<float, 3>;
+extern template class TKinematicGeometryParticle<float, 3>;
+extern template class TPBDRigidParticle<float, 3>;
+#endif
 } // namespace Chaos
 
diff --git a/Engine/Source/Runtime/Experimental/Chaos/Public/Chaos/RigidParticles.h b/Engine/Source/Runtime/Experimental/Chaos/Public/Chaos/RigidParticles.h
--- a/Engine/Source/Runtime/Experimental/Chaos/Public/Chaos/RigidParticles.h
+++ b/Engine/Source/Runtime/Experimental/Chaos/Public/Chaos/RigidParticles.h
@@ -116,11 +116,11 @@ class TRigidParticles : public TKinematicGeometryParticles<T, d>
 	FORCEINLINE const T InvM(const int32 Index) const { return MInvM[Index]; }
 	FORCEINLINE T& InvM(const int32 Index) { return MInvM[Index]; }
 
-	const T& LinearDamping(const int32 index) const { return MLinearDamping[index]; }
-	T& LinearDamping(const int32 index) { return MLinearDamping[index]; }
+	CHAOS_API const T& LinearDamping(const int32 index) const { return MLinearDamping[index]; }
+	CHAOS_API T& LinearDamping(const int32 index) { return MLinearDamping[index]; }
 
-	const T& AngularDamping(const int32 index) const { return MAngularDamping[index]; }
-	T& AngularDamping(const int32 index) { return MAngularDamping[index]; }
+	CHAOS_API const T& AngularDamping(const int32 index) const { return MAngularDamping[index]; }
+	CHAOS_API T& AngularDamping(const int32 index) { return MAngularDamping[index]; }
 
 	FORCEINLINE int32 CollisionParticlesSize(int32 Index) const { return MCollisionParticles[Index] == nullptr ? 0 : MCollisionParticles[Index]->Size(); }
 
