Fix clang compile
#rnx

diff --git a/Engine/Source/Runtime/Experimental/Chaos/Public/Chaos/GeometryParticles.h b/Engine/Source/Runtime/Experimental/Chaos/Public/Chaos/GeometryParticles.h
--- a/Engine/Source/Runtime/Experimental/Chaos/Public/Chaos/GeometryParticles.h
+++ b/Engine/Source/Runtime/Experimental/Chaos/Public/Chaos/GeometryParticles.h
@@ -207,7 +207,7 @@ namespace Chaos
 
 		const TSharedPtr<FImplicitObject, ESPMode::ThreadSafe>& SharedGeometry(const int32 Index) const { return MSharedGeometry[Index]; }
 
-		const TShapesArray<T, d>& ShapesArray(const int32 Index) const { return MShapesArray[Index]; }
+		CHAOS_API const TShapesArray<T, d>& ShapesArray(const int32 Index) const { return MShapesArray[Index]; }
 
 #if CHAOS_DETERMINISTIC
 		CHAOS_API FParticleID ParticleID(const int32 Idx) const { return MParticleIDs[Idx]; }
