Chaos - fix joint Gauss Seidel solver rotation after over zealous cleanup

diff --git a/Engine/Source/Runtime/Experimental/Chaos/Private/Chaos/Joint/PBDJointSolverGaussSeidel.cpp b/Engine/Source/Runtime/Experimental/Chaos/Private/Chaos/Joint/PBDJointSolverGaussSeidel.cpp
--- a/Engine/Source/Runtime/Experimental/Chaos/Private/Chaos/Joint/PBDJointSolverGaussSeidel.cpp
+++ b/Engine/Source/Runtime/Experimental/Chaos/Private/Chaos/Joint/PBDJointSolverGaussSeidel.cpp
@@ -78,16 +78,12 @@ namespace Chaos
 		EJointMotionType Swing1Motion = JointSettings.Motion.AngularMotionTypes[(int32)EJointAngularConstraintIndex::Swing1];
 		EJointMotionType Swing2Motion = JointSettings.Motion.AngularMotionTypes[(int32)EJointAngularConstraintIndex::Swing2];
 
-		// Decompose rotation of body 1 relative to body 0 into swing and twist rotations, assuming twist is X axis
-		FRotation3 R01Twist, R01Swing;
-		FPBDJointUtilities::DecomposeSwingTwistLocal(Rs[0], Rs[1], R01Swing, R01Twist);
-
 		// Apply twist constraint
 		if (bEnableTwistLimits)
 		{
 			if (TwistMotion != EJointMotionType::Free)
 			{
-				ApplyJointTwistConstraint(Dt, JointSettings, R01Twist, R01Swing);
+				ApplyTwistConstraint(Dt, JointSettings);
 			}
 		}
 
@@ -97,19 +93,19 @@ namespace Chaos
 			if ((Swing1Motion == EJointMotionType::Limited) && (Swing2Motion == EJointMotionType::Limited))
 			{
 				// Swing Cone
-				ApplyJointConeConstraint(Dt, JointSettings, R01Twist, R01Swing);
+				ApplyConeConstraint(Dt, JointSettings);
 			}
 			else
 			{
 				if (Swing1Motion != EJointMotionType::Free)
 				{
 					// Swing Arc/Lock
-					ApplyJointSwingConstraint(Dt, JointSettings, R01Twist, R01Swing, EJointAngularConstraintIndex::Swing1, EJointAngularAxisIndex::Swing1);
+					ApplySwingConstraint(Dt, JointSettings, EJointAngularConstraintIndex::Swing1, EJointAngularAxisIndex::Swing1);
 				}
 				if (Swing2Motion != EJointMotionType::Free)
 				{
 					// Swing Arc/Lock
-					ApplyJointSwingConstraint(Dt, JointSettings, R01Twist, R01Swing, EJointAngularConstraintIndex::Swing2, EJointAngularAxisIndex::Swing2);
+					ApplySwingConstraint(Dt, JointSettings, EJointAngularConstraintIndex::Swing2, EJointAngularAxisIndex::Swing2);
 				}
 			}
 		}
@@ -123,7 +119,7 @@ namespace Chaos
 		const TVector<EJointMotionType, 3>& LinearMotion = JointSettings.Motion.LinearMotionTypes;
 		if ((LinearMotion[0] != EJointMotionType::Free) || (LinearMotion[1] != EJointMotionType::Free) || (LinearMotion[2] != EJointMotionType::Free))
 		{
-			ApplyJointPositionConstraint(Dt, JointSettings);
+			ApplyPositionConstraint(Dt, JointSettings);
 		}
 	}
 
@@ -167,12 +163,14 @@ namespace Chaos
 		ApplyRotationDelta(Stiffness, DR0, DR1);
 	}
 
-	void FJointSolverGaussSeidel::ApplyJointTwistConstraint(
+	void FJointSolverGaussSeidel::ApplyTwistConstraint(
 		const FReal Dt,
-		const FPBDJointSettings& JointSettings,
-		const FRotation3& R01Twist,
-		const FRotation3& R01Swing)
+		const FPBDJointSettings& JointSettings)
 	{
+		// Decompose rotation of body 1 relative to body 0 into swing and twist rotations, assuming twist is X axis
+		FRotation3 R01Twist, R01Swing;
+		FPBDJointUtilities::DecomposeSwingTwistLocal(Rs[0], Rs[1], R01Swing, R01Twist);
+
 		FVec3 TwistAxis01;
 		FReal TwistAngle;
 		R01Twist.ToAxisAndAngleSafe(TwistAxis01, TwistAngle, FJointConstants::TwistAxis(), SwingTwistAngleTolerance);
@@ -217,12 +215,14 @@ namespace Chaos
 	}
 
 
-	void FJointSolverGaussSeidel::ApplyJointConeConstraint(
+	void FJointSolverGaussSeidel::ApplyConeConstraint(
 		const FReal Dt,
-		const FPBDJointSettings& JointSettings,
-		const FRotation3& R01Twist,
-		const FRotation3& R01Swing)
+		const FPBDJointSettings& JointSettings)
 	{
+		// Decompose rotation of body 1 relative to body 0 into swing and twist rotations, assuming twist is X axis
+		FRotation3 R01Twist, R01Swing;
+		FPBDJointUtilities::DecomposeSwingTwistLocal(Rs[0], Rs[1], R01Swing, R01Twist);
+
 		FVec3 SwingAxis01;
 		FReal SwingAngle;
 		R01Swing.ToAxisAndAngleSafe(SwingAxis01, SwingAngle, FJointConstants::Swing1Axis(), SwingTwistAngleTolerance);
@@ -267,14 +267,16 @@ namespace Chaos
 	}
 
 
-	void FJointSolverGaussSeidel::ApplyJointSwingConstraint(
+	void FJointSolverGaussSeidel::ApplySwingConstraint(
 		const FReal Dt,
 		const FPBDJointSettings& JointSettings,
-		const FRotation3& R01Twist,
-		const FRotation3& R01Swing,
 		const EJointAngularConstraintIndex SwingConstraintIndex,
 		const EJointAngularAxisIndex SwingAxisIndex)
 	{
+		// Decompose rotation of body 1 relative to body 0 into swing and twist rotations, assuming twist is X axis
+		FRotation3 R01Twist, R01Swing;
+		FPBDJointUtilities::DecomposeSwingTwistLocal(Rs[0], Rs[1], R01Swing, R01Twist);
+
 		FVec3 TwistAxis01;
 		FReal TwistAngle;
 		R01Twist.ToAxisAndAngleSafe(TwistAxis01, TwistAngle, FJointConstants::TwistAxis(), SwingTwistAngleTolerance);
@@ -334,7 +336,7 @@ namespace Chaos
 		}
 	}
 
-	void FJointSolverGaussSeidel::ApplyJointPositionConstraint(
+	void FJointSolverGaussSeidel::ApplyPositionConstraint(
 		const FReal Dt,
 		const FPBDJointSettings& JointSettings)
 	{
diff --git a/Engine/Source/Runtime/Experimental/Chaos/Public/Chaos/Joint/PBDJointSolverGaussSeidel.h b/Engine/Source/Runtime/Experimental/Chaos/Public/Chaos/Joint/PBDJointSolverGaussSeidel.h
--- a/Engine/Source/Runtime/Experimental/Chaos/Public/Chaos/Joint/PBDJointSolverGaussSeidel.h
+++ b/Engine/Source/Runtime/Experimental/Chaos/Public/Chaos/Joint/PBDJointSolverGaussSeidel.h
@@ -86,27 +86,21 @@ namespace Chaos
 			const FVec3& Axis1,
 			const FReal Angle1);
 
-		void ApplyJointTwistConstraint(
+		void ApplyTwistConstraint(
 			const FReal Dt,
-			const FPBDJointSettings& JointSettings,
-			const FRotation3& R01Twist,
-			const FRotation3& R01Swing);
+			const FPBDJointSettings& JointSettings);
 
-		void ApplyJointConeConstraint(
+		void ApplyConeConstraint(
 			const FReal Dt,
-			const FPBDJointSettings& JointSettings,
-			const FRotation3& R01Twist,
-			const FRotation3& R01Swing);
+			const FPBDJointSettings& JointSettings);
 
-		void ApplyJointSwingConstraint(
+		void ApplySwingConstraint(
 			const FReal Dt,
 			const FPBDJointSettings& JointSettings,
-			const FRotation3& R01Twist,
-			const FRotation3& R01Swing,
 			const EJointAngularConstraintIndex SwingConstraintIndex,
 			const EJointAngularAxisIndex SwingAxisIndex);
 
-		void ApplyJointPositionConstraint(
+		void ApplyPositionConstraint(
 			const FReal Dt,
 			const FPBDJointSettings& JointSettings);
 
