ChaosCloth: Robustness. Include small vector check in Axial Spring Constraints

diff --git a/Engine/Source/Runtime/Experimental/Chaos/Public/Chaos/PBDAxialSpringConstraintsBase.h b/Engine/Source/Runtime/Experimental/Chaos/Public/Chaos/PBDAxialSpringConstraintsBase.h
--- a/Engine/Source/Runtime/Experimental/Chaos/Public/Chaos/PBDAxialSpringConstraintsBase.h
+++ b/Engine/Source/Runtime/Experimental/Chaos/Public/Chaos/PBDAxialSpringConstraintsBase.h
@@ -58,7 +58,7 @@ class TPBDAxialSpringConstraintsBase
 	}
 	virtual ~TPBDAxialSpringConstraintsBase() {}
 
-	TVector<T, d> GetDelta(const TPBDParticles<T, d>& InParticles, const int32 i) const
+	inline TVector<T, d> GetDelta(const TPBDParticles<T, d>& InParticles, const int32 i) const
 	{
 		const auto& Constraint = MConstraints[i];
 		int32 i1 = Constraint[0];
@@ -73,12 +73,13 @@ class TPBDAxialSpringConstraintsBase
 		const TVector<T, d> P = (P2 - P3) * MBarys[i] + P3;
 		TVector<T, d> Difference = P1 - P;
 		float Distance = Difference.Size();
-		check(Distance > 1e-7);
+		if (UNLIKELY(Distance <= SMALL_NUMBER))
+			return TVector<T, d>((T) 0);
 		TVector<T, d> Direction = Difference / Distance;
 		TVector<T, d> Delta = (Distance - MDists[i]) * Direction;
-		T CombinedMass = PInvMass + InParticles.InvM(i1);
-		check(CombinedMass > 1e-7);
-		return MStiffness * Delta / CombinedMass;
+		T CombinedInvMass = PInvMass + InParticles.InvM(i1);
+		ensure(CombinedInvMass > 1e-7);
+		return MStiffness * Delta / CombinedInvMass;
 	}
 
   private:
