Chaos - Enable ticking of solver in editor outside of PIE w/ 0 DT.
Acceleration Structure now correct for editor queries.

diff --git a/Engine/Source/Runtime/Engine/Private/PhysicsEngine/Experimental/PhysInterface_Chaos.cpp b/Engine/Source/Runtime/Engine/Private/PhysicsEngine/Experimental/PhysInterface_Chaos.cpp
--- a/Engine/Source/Runtime/Engine/Private/PhysicsEngine/Experimental/PhysInterface_Chaos.cpp
+++ b/Engine/Source/Runtime/Engine/Private/PhysicsEngine/Experimental/PhysInterface_Chaos.cpp
@@ -359,6 +359,14 @@ FPhysScene* FPhysInterface_Chaos::GetCurrentScene(const FPhysicsActorHandle& InH
 	return nullptr;
 }
 
+void FPhysInterface_Chaos::FlushScene(FPhysScene* InScene)
+{
+	FPhysicsCommand::ExecuteWrite(InScene, [&]()
+	{
+		InScene->Flush_AssumesLocked();
+	});
+}
+
 bool FPhysInterface_Chaos::CanSimulate_AssumesLocked(const FPhysicsActorHandle& InActorReference)
 {
 	// #todo : Implement
diff --git a/Engine/Source/Runtime/Engine/Private/PhysicsEngine/Experimental/PhysScene_Chaos.cpp b/Engine/Source/Runtime/Engine/Private/PhysicsEngine/Experimental/PhysScene_Chaos.cpp
--- a/Engine/Source/Runtime/Engine/Private/PhysicsEngine/Experimental/PhysScene_Chaos.cpp
+++ b/Engine/Source/Runtime/Engine/Private/PhysicsEngine/Experimental/PhysScene_Chaos.cpp
@@ -364,6 +364,8 @@ FPhysScene_Chaos::FPhysScene_Chaos(AActor* InSolverActor
 );
 	check(SceneSolver);
 
+	SceneSolver->SetEnabled(true);
+
 	// If we're running the physics thread, hand over the solver to it - we are no longer
 	// able to access the solver on the game thread and should only use commands
 	if(ChaosModule->GetDispatcher() && ChaosModule->GetDispatcher()->GetMode() == Chaos::EThreadingMode::DedicatedThread)
@@ -1273,6 +1275,14 @@ void FPhysScene_ChaosInterface::StartFrame()
 	FChaosSolversModule* SolverModule = FChaosSolversModule::GetModule();
 	checkSlow(SolverModule);
 
+	float Dt = MDeltaTime;
+#if WITH_EDITOR
+	if (GIsPlayInEditorWorld == false)
+	{
+		Dt = 0.0f;
+	}
+#endif
+
 	if (Chaos::IDispatcher* Dispatcher = SolverModule->GetDispatcher())
 	{
 		for (auto * Solver : SolverModule->GetSolvers())
@@ -1284,27 +1294,27 @@ void FPhysScene_ChaosInterface::StartFrame()
 		{
 		case EChaosThreadingMode::SingleThread:
 		{
-			OnPhysScenePreTick.Broadcast(this, MDeltaTime);
-			OnPhysSceneStep.Broadcast(this, MDeltaTime);
+			OnPhysScenePreTick.Broadcast(this, Dt);
+			OnPhysSceneStep.Broadcast(this, Dt);
 
 			// Here we can directly tick the scene. Single threaded mode doesn't buffer any commands
 			// that would require pumping here - everything is done on demand.
-			Scene.Tick(MDeltaTime);
+			Scene.Tick(Dt);
 		}
 		break;
 		case EChaosThreadingMode::TaskGraph:
 		{
 			check(!CompletionEvent.GetReference())
 
-				OnPhysScenePreTick.Broadcast(this, MDeltaTime);
-			OnPhysSceneStep.Broadcast(this, MDeltaTime);
+			OnPhysScenePreTick.Broadcast(this, Dt);
+			OnPhysSceneStep.Broadcast(this, Dt);
 
 			FGraphEventRef SimulationCompleteEvent = FGraphEvent::CreateGraphEvent();
 
 			// Need to fire off a parallel task to handle running physics commands and
 			// ticking the scene while the engine continues on until TG_EndPhysics
 			// (this should happen in TG_StartPhysics)
-			PhysicsTickTask = TGraphTask<FPhysicsTickTask>::CreateTask(nullptr, ENamedThreads::GameThread).ConstructAndDispatchWhenReady(SimulationCompleteEvent, MDeltaTime);
+			PhysicsTickTask = TGraphTask<FPhysicsTickTask>::CreateTask(nullptr, ENamedThreads::GameThread).ConstructAndDispatchWhenReady(SimulationCompleteEvent, Dt);
 
 			// Setup post simulate tasks
 			if (PhysicsTickTask.GetReference())
diff --git a/Engine/Source/Runtime/Engine/Private/PhysicsEngine/PhysLevel.cpp b/Engine/Source/Runtime/Engine/Private/PhysicsEngine/PhysLevel.cpp
--- a/Engine/Source/Runtime/Engine/Private/PhysicsEngine/PhysLevel.cpp
+++ b/Engine/Source/Runtime/Engine/Private/PhysicsEngine/PhysLevel.cpp
@@ -104,30 +104,37 @@ void UWorld::SetupPhysicsTickFunctions(float DeltaSeconds)
 	
 	EndPhysicsTickFunction.bCanEverTick = true;
 	EndPhysicsTickFunction.Target = this;
+
+// Chaos ticks solver for trace collisions
+#if WITH_CHAOS
+	bool bEnablePhysics = (bShouldSimulatePhysics || bEnableTraceCollision);
+#else
+	bool bEnablePhysics = bShouldSimulatePhysics;
+#endif
 	
-	// see if we need to update tick registration
-	bool bNeedToUpdateTickRegistration = (bShouldSimulatePhysics != StartPhysicsTickFunction.IsTickFunctionRegistered())
-		|| (bShouldSimulatePhysics != EndPhysicsTickFunction.IsTickFunctionRegistered());
+	// see if we need to update tick registration;
+	bool bNeedToUpdateTickRegistration = (bEnablePhysics != StartPhysicsTickFunction.IsTickFunctionRegistered())
+		|| (bEnablePhysics != EndPhysicsTickFunction.IsTickFunctionRegistered());
 
 	if (bNeedToUpdateTickRegistration && PersistentLevel)
 	{
-		if (bShouldSimulatePhysics && !StartPhysicsTickFunction.IsTickFunctionRegistered())
+		if (bEnablePhysics && !StartPhysicsTickFunction.IsTickFunctionRegistered())
 		{
 			StartPhysicsTickFunction.TickGroup = TG_StartPhysics;
 			StartPhysicsTickFunction.RegisterTickFunction(PersistentLevel);
 		}
-		else if (!bShouldSimulatePhysics && StartPhysicsTickFunction.IsTickFunctionRegistered())
+		else if (!bEnablePhysics && StartPhysicsTickFunction.IsTickFunctionRegistered())
 		{
 			StartPhysicsTickFunction.UnRegisterTickFunction();
 		}
 
-		if (bShouldSimulatePhysics && !EndPhysicsTickFunction.IsTickFunctionRegistered())
+		if (bEnablePhysics && !EndPhysicsTickFunction.IsTickFunctionRegistered())
 		{
 			EndPhysicsTickFunction.TickGroup = TG_EndPhysics;
 			EndPhysicsTickFunction.RegisterTickFunction(PersistentLevel);
 			EndPhysicsTickFunction.AddPrerequisite(this, StartPhysicsTickFunction);
 		}
-		else if (!bShouldSimulatePhysics && EndPhysicsTickFunction.IsTickFunctionRegistered())
+		else if (!bEnablePhysics && EndPhysicsTickFunction.IsTickFunctionRegistered())
 		{
 			EndPhysicsTickFunction.RemovePrerequisite(this, StartPhysicsTickFunction);
 			EndPhysicsTickFunction.UnRegisterTickFunction();
diff --git a/Engine/Source/Runtime/Engine/Private/World.cpp b/Engine/Source/Runtime/Engine/Private/World.cpp
--- a/Engine/Source/Runtime/Engine/Private/World.cpp
+++ b/Engine/Source/Runtime/Engine/Private/World.cpp
@@ -4523,10 +4523,14 @@ void UWorld::SetPhysicsScene(FPhysScene* InScene)
 	// Assign pointer
 	PhysicsScene = InScene;
 
-	// Set pointer in scene to know which world its coming from
 	if(PhysicsScene != NULL)
 	{
+		// Set pointer in scene to know which world its coming from
 		PhysicsScene->SetOwningWorld(this);
+
+#if WITH_CHAOS
+		FPhysInterface_Chaos::FlushScene(PhysicsScene);
+#endif
 	}
 }
 
diff --git a/Engine/Source/Runtime/Engine/Public/Physics/Experimental/PhysInterface_Chaos.h b/Engine/Source/Runtime/Engine/Public/Physics/Experimental/PhysInterface_Chaos.h
--- a/Engine/Source/Runtime/Engine/Public/Physics/Experimental/PhysInterface_Chaos.h
+++ b/Engine/Source/Runtime/Engine/Public/Physics/Experimental/PhysInterface_Chaos.h
@@ -223,6 +223,7 @@ public:
     static bool HasAsyncSceneData(const FPhysicsActorHandle& InHandle) { return false; }
 	static bool IsInScene(const FPhysicsActorHandle& InActorReference);
 	static FPhysScene* GetCurrentScene(const FPhysicsActorHandle& InHandle);
+	static void FlushScene(FPhysScene* InScene);
 	static bool CanSimulate_AssumesLocked(const FPhysicsActorHandle& InActorReference);
 	static float GetMass_AssumesLocked(const FPhysicsActorHandle& InActorReference);
 
diff --git a/Engine/Source/Runtime/Experimental/Chaos/Private/Chaos/PBDRigidsEvolutionGBF.cpp b/Engine/Source/Runtime/Experimental/Chaos/Private/Chaos/PBDRigidsEvolutionGBF.cpp
--- a/Engine/Source/Runtime/Experimental/Chaos/Private/Chaos/PBDRigidsEvolutionGBF.cpp
+++ b/Engine/Source/Runtime/Experimental/Chaos/Private/Chaos/PBDRigidsEvolutionGBF.cpp
@@ -139,65 +139,71 @@ void TPBDRigidsEvolutionGBF<T, d>::AdvanceOneTimeStep(const T Dt, const T StepFr
 	}
 #endif
 
+	if (Dt > 0.0f)
 	{
-		SCOPE_CYCLE_COUNTER(STAT_Evolution_Integrate);
-		Integrate(Particles.GetActiveParticlesView(), Dt);
-	}
+		{
+			SCOPE_CYCLE_COUNTER(STAT_Evolution_Integrate);
+			Integrate(Particles.GetActiveParticlesView(), Dt);
+		}
 
-	{
-		SCOPE_CYCLE_COUNTER(STAT_Evolution_KinematicTargets);
-		ApplyKinematicTargets(Dt, StepFraction);
-	}
+		{
+			SCOPE_CYCLE_COUNTER(STAT_Evolution_KinematicTargets);
+			ApplyKinematicTargets(Dt, StepFraction);
+		}
 
-	if (PostIntegrateCallback != nullptr)
-	{
-		PostIntegrateCallback();
-	}
+		if (PostIntegrateCallback != nullptr)
+		{
+			PostIntegrateCallback();
+		}
 
-	{
-		SCOPE_CYCLE_COUNTER(STAT_Evolution_UpdateConstraintPositionBasedState);
-		UpdateConstraintPositionBasedState(Dt);
+		{
+			SCOPE_CYCLE_COUNTER(STAT_Evolution_UpdateConstraintPositionBasedState);
+			UpdateConstraintPositionBasedState(Dt);
+		}
 	}
+
 	{
 		Base::ComputeIntermediateSpatialAcceleration();
 	}
+
+	if (Dt > 0.0f)
 	{
-		SCOPE_CYCLE_COUNTER(STAT_Evolution_DetectCollisions);
-		CollisionDetector.GetBroadPhase().SetSpatialAcceleration(InternalAcceleration.Get());
+		{
+			SCOPE_CYCLE_COUNTER(STAT_Evolution_DetectCollisions);
+			CollisionDetector.GetBroadPhase().SetSpatialAcceleration(InternalAcceleration.Get());
 
-		CollisionStats::FStatData StatData(bPendingHierarchyDump);
+			CollisionStats::FStatData StatData(bPendingHierarchyDump);
 
-		CollisionDetector.DetectCollisions(Dt, StatData);
+			CollisionDetector.DetectCollisions(Dt, StatData);
 
-		CHAOS_COLLISION_STAT(StatData.Print());
-	}
+			CHAOS_COLLISION_STAT(StatData.Print());
+		}
 
-	if (PostDetectCollisionsCallback != nullptr)
-	{
-		PostDetectCollisionsCallback();
-	}
+		if (PostDetectCollisionsCallback != nullptr)
+		{
+			PostDetectCollisionsCallback();
+		}
 
-	{
-		SCOPE_CYCLE_COUNTER(STAT_Evolution_CreateConstraintGraph);
-		CreateConstraintGraph();
-	}
-	{
-		SCOPE_CYCLE_COUNTER(STAT_Evolution_CreateIslands);
-		CreateIslands();
-	}
+		{
+			SCOPE_CYCLE_COUNTER(STAT_Evolution_CreateConstraintGraph);
+			CreateConstraintGraph();
+		}
+		{
+			SCOPE_CYCLE_COUNTER(STAT_Evolution_CreateIslands);
+			CreateIslands();
+		}
 
-	if (PreApplyCallback != nullptr)
-	{
-		PreApplyCallback();
-	}
+		if (PreApplyCallback != nullptr)
+		{
+			PreApplyCallback();
+		}
+
+		TArray<bool> SleepedIslands;
+		SleepedIslands.SetNum(GetConstraintGraph().NumIslands());
+		TArray<TArray<TPBDRigidParticleHandle<T, d>*>> DisabledParticles;
+		DisabledParticles.SetNum(GetConstraintGraph().NumIslands());
+		SleepedIslands.SetNum(GetConstraintGraph().NumIslands());
 
-	TArray<bool> SleepedIslands;
-	SleepedIslands.SetNum(GetConstraintGraph().NumIslands());
-	TArray<TArray<TPBDRigidParticleHandle<T,d>*>> DisabledParticles;
-	DisabledParticles.SetNum(GetConstraintGraph().NumIslands());
-	SleepedIslands.SetNum(GetConstraintGraph().NumIslands());
-	if(Dt > 0)
-	{
 		SCOPE_CYCLE_COUNTER(STAT_Evolution_ParallelSolve);
 		PhysicsParallelFor(GetConstraintGraph().NumIslands(), [&](int32 Island) {
 			const TArray<TGeometryParticleHandle<T, d>*>& IslandParticles = GetConstraintGraph().GetIslandParticles(Island);
@@ -262,26 +268,27 @@ void TPBDRigidsEvolutionGBF<T, d>::AdvanceOneTimeStep(const T Dt, const T StepFr
 
 			// Turn off if not moving
 			SleepedIslands[Island] = GetConstraintGraph().SleepInactive(Island, PhysicsMaterials);
-		});
-	}
-	{
-		SCOPE_CYCLE_COUNTER(STAT_Evolution_DeactivateSleep);
-		for (int32 Island = 0; Island < GetConstraintGraph().NumIslands(); ++Island)
+			});
+
 		{
-			if (SleepedIslands[Island])
-			{
-				Particles.DeactivateParticles(GetConstraintGraph().GetIslandParticles(Island));
-			}
-			for (const auto Particle : DisabledParticles[Island])
+			SCOPE_CYCLE_COUNTER(STAT_Evolution_DeactivateSleep);
+			for (int32 Island = 0; Island < GetConstraintGraph().NumIslands(); ++Island)
 			{
-				Particles.DisableParticle(Particle);
+				if (SleepedIslands[Island])
+				{
+					Particles.DeactivateParticles(GetConstraintGraph().GetIslandParticles(Island));
+				}
+				for (const auto Particle : DisabledParticles[Island])
+				{
+					Particles.DisableParticle(Particle);
+				}
 			}
 		}
-	}
 
-	Clustering.AdvanceClustering(Dt, GetCollisionConstraints());
+		Clustering.AdvanceClustering(Dt, GetCollisionConstraints());
 
-	ParticleUpdatePosition(Particles.GetActiveParticlesView(), Dt);
+		ParticleUpdatePosition(Particles.GetActiveParticlesView(), Dt);
+	}
 }
 
 template <typename T, int d>
