Wake physics from sleeping when external forces are applied

diff --git a/Engine/Source/Runtime/Engine/Private/PhysicsEngine/Experimental/PhysScene_Chaos.cpp b/Engine/Source/Runtime/Engine/Private/PhysicsEngine/Experimental/PhysScene_Chaos.cpp
--- a/Engine/Source/Runtime/Engine/Private/PhysicsEngine/Experimental/PhysScene_Chaos.cpp
+++ b/Engine/Source/Runtime/Engine/Private/PhysicsEngine/Experimental/PhysScene_Chaos.cpp
@@ -1080,6 +1080,8 @@ void FPhysScene_ChaosInterface::AddForce_AssumesLocked(FBodyInstance* BodyInstan
 			EObjectStateType ObjectState = Rigid->ObjectState();
 			if (CHAOS_ENSURE(ObjectState == EObjectStateType::Dynamic || ObjectState == EObjectStateType::Sleeping))
 			{
+				Rigid->SetObjectState(EObjectStateType::Dynamic);
+
 				const Chaos::TVector<float, 3> CurrentForce = Rigid->ExternalForce();
 				if (bAccelChange)
 				{
@@ -1091,6 +1093,7 @@ void FPhysScene_ChaosInterface::AddForce_AssumesLocked(FBodyInstance* BodyInstan
 				{
 					Rigid->SetExternalForce(CurrentForce + Force);
 				}
+
 			}
 		}
 	}
@@ -1114,6 +1117,8 @@ void FPhysScene_ChaosInterface::AddForceAtPosition_AssumesLocked(FBodyInstance*
 				const Chaos::FVec3& CurrentTorque = Rigid->ExternalTorque();
 				const Chaos::FVec3 WorldCOM = FParticleUtilitiesGT::GetCoMWorldPosition(Rigid);
 
+				Rigid->SetObjectState(EObjectStateType::Dynamic);
+
 				if (bIsLocalForce)
 				{
 					const Chaos::FRigidTransform3 CurrentTransform = FParticleUtilitiesGT::GetActorWorldTransform(Rigid);
@@ -1129,6 +1134,7 @@ void FPhysScene_ChaosInterface::AddForceAtPosition_AssumesLocked(FBodyInstance*
 					Rigid->SetExternalForce(CurrentForce + Force);
 					Rigid->SetExternalTorque(CurrentTorque + WorldTorque);
 				}
+
 			}
 		}
 	}
diff --git a/Engine/Source/Runtime/Experimental/Chaos/Public/Chaos/PBDRigidParticles.h b/Engine/Source/Runtime/Experimental/Chaos/Public/Chaos/PBDRigidParticles.h
--- a/Engine/Source/Runtime/Experimental/Chaos/Public/Chaos/PBDRigidParticles.h
+++ b/Engine/Source/Runtime/Experimental/Chaos/Public/Chaos/PBDRigidParticles.h
@@ -79,10 +79,11 @@ class TPBDRigidParticles : public TRigidParticles<T, d>
 			PreW(Index) = this->W(Index);
 		}
 
-		if (Sleeping(Index) != bSleeping)
+		bool CurrentlySleeping = this->ObjectState(Index) == EObjectStateType::Sleeping;
+		if (CurrentlySleeping != bSleeping)
 		{
 			TGeometryParticleHandle<T, d>* Particle = reinterpret_cast<TGeometryParticleHandle<T, d>*>(this->Handle(Index));
-			this->AddSleepData(Particle);
+			this->AddSleepData(Particle, bSleeping);
 		}
 
 		// Dynamic -> Sleeping or Sleeping -> Dynamic
@@ -147,7 +148,7 @@ class TPBDRigidParticles : public TRigidParticles<T, d>
 		if(bCurrentSleeping != bNewSleeping)
 		{
 			TGeometryParticleHandle<T, d>* Particle = reinterpret_cast<TGeometryParticleHandle<T, d>*>(this->Handle(Index));
-			this->AddSleepData(Particle);
+ 			this->AddSleepData(Particle, bNewSleeping);
 		}
 
 		this->ObjectState(Index) = InObjectState;
diff --git a/Engine/Source/Runtime/Experimental/Chaos/Public/Chaos/RigidParticles.h b/Engine/Source/Runtime/Experimental/Chaos/Public/Chaos/RigidParticles.h
--- a/Engine/Source/Runtime/Experimental/Chaos/Public/Chaos/RigidParticles.h
+++ b/Engine/Source/Runtime/Experimental/Chaos/Public/Chaos/RigidParticles.h
@@ -24,6 +24,25 @@ enum class EObjectStateType : int8
 	Count
 };
 
+template<class T, int d>
+struct TSleepData
+{
+	TSleepData()
+		: Particle(nullptr)
+		, Sleeping(true)
+	{}
+
+	TSleepData(
+		TGeometryParticleHandle<T, d>* InParticle, bool InSleeping)
+		: Particle(InParticle)
+		, Sleeping(InSleeping)
+	{}
+
+	TGeometryParticleHandle<T, d>* Particle;
+	bool Sleeping; // if !Sleeping == Awake
+};
+
+
 // Counts the number of bits needed to represent an int with a max
 constexpr int8 NumBitsNeeded(const int8 MaxValue)
 {
@@ -144,8 +163,14 @@ class TRigidParticles : public TKinematicGeometryParticles<T, d>
 	FORCEINLINE const bool ToBeRemovedOnFracture(const int32 Index) const { return MToBeRemovedOnFracture[Index]; }
 	FORCEINLINE bool& ToBeRemovedOnFracture(const int32 Index) { return MToBeRemovedOnFracture[Index]; }
 
-	FORCEINLINE TQueue<TGeometryParticleHandle<T, d>*, EQueueMode::Mpsc>& GetSleepData() { return MSleepData; }
-	FORCEINLINE void AddSleepData(TGeometryParticleHandle<T, d>* Particle) { MSleepData.Enqueue(Particle); }
+	FORCEINLINE TQueue<TSleepData<T, d>, EQueueMode::Mpsc>& GetSleepData() { return MSleepData; }
+	FORCEINLINE void AddSleepData(TGeometryParticleHandle<T, d>* Particle, bool Sleeping)
+	{ 
+		TSleepData<T, d> SleepData;
+		SleepData.Particle = Particle;
+		SleepData.Sleeping = Sleeping;
+		MSleepData.Enqueue(SleepData); 
+	}
 
 	FORCEINLINE const EObjectStateType ObjectState(const int32 Index) const { return MObjectState[Index]; }
 	FORCEINLINE EObjectStateType& ObjectState(const int32 Index) { return MObjectState[Index]; }
@@ -196,7 +221,7 @@ class TRigidParticles : public TKinematicGeometryParticles<T, d>
 	TArrayCollectionArray<bool> MDisabled;
 	TArrayCollectionArray<bool> MToBeRemovedOnFracture;
 	TArrayCollectionArray<EObjectStateType> MObjectState;
-	TQueue<TGeometryParticleHandle<T, d>*, EQueueMode::Mpsc> MSleepData;
+	TQueue<TSleepData<T, d>, EQueueMode::Mpsc> MSleepData;
 };
 
 
diff --git a/Engine/Source/Runtime/Experimental/ChaosSolverEngine/Private/Chaos/ChaosGameplayEventDispatcher.cpp b/Engine/Source/Runtime/Experimental/ChaosSolverEngine/Private/Chaos/ChaosGameplayEventDispatcher.cpp
--- a/Engine/Source/Runtime/Experimental/ChaosSolverEngine/Private/Chaos/ChaosGameplayEventDispatcher.cpp
+++ b/Engine/Source/Runtime/Experimental/ChaosSolverEngine/Private/Chaos/ChaosGameplayEventDispatcher.cpp
@@ -421,10 +421,9 @@ void UChaosGameplayEventDispatcher::HandleSleepingEvents(const Chaos::FSleepingE
 	{
 		if (SleepData.Particle->Proxy != nullptr)
 		{
-			UPrimitiveComponent* Component = Cast<UPrimitiveComponent>(SleepData.Particle->Proxy->GetOwner());
-			if (Component)
+			if (FBodyInstance* BodyInstance = FPhysicsUserData::Get<FBodyInstance>(SleepData.Particle->UserData()))
 			{
-				if (FBodyInstance* BodyInstance = Component->GetBodyInstance())
+				if (BodyInstance->bGenerateWakeEvents)
 				{
 					ESleepEvent WakeSleepEvent = SleepData.Sleeping ? ESleepEvent::SET_Sleep : ESleepEvent::SET_Wakeup;
 					AddPendingSleepingNotify(BodyInstance, WakeSleepEvent);
@@ -443,3 +442,4 @@ void UChaosGameplayEventDispatcher::AddPendingSleepingNotify(FBodyInstance* Body
 }
 
 // PRAGMA_ENABLE_OPTIMIZATION
+
diff --git a/Engine/Source/Runtime/Experimental/ChaosSolvers/Private/EventDefaults.cpp b/Engine/Source/Runtime/Experimental/ChaosSolvers/Private/EventDefaults.cpp
--- a/Engine/Source/Runtime/Experimental/ChaosSolvers/Private/EventDefaults.cpp
+++ b/Engine/Source/Runtime/Experimental/ChaosSolvers/Private/EventDefaults.cpp
@@ -365,18 +365,18 @@ namespace Chaos
 
 			Chaos::FPBDRigidsSolver* NonConstSolver = (Chaos::FPBDRigidsSolver*)(Solver);
 			auto& SolverSleepingData = NonConstSolver->Particles.GetDynamicParticles().GetSleepData();
-			TGeometryParticleHandle<float, 3>* ParticleHandle = nullptr;
-			while(SolverSleepingData.Dequeue(ParticleHandle))
+			TSleepData<float, 3> SleepData;
+			while(SolverSleepingData.Dequeue(SleepData))
 			{
-				if(ParticleHandle)
+				if(SleepData.Particle)
 				{
-					TGeometryParticle<float, 3>* Particle = ParticleHandle->GTGeometryParticle();
+					TGeometryParticle<float, 3>* Particle = SleepData.Particle->GTGeometryParticle();
 					if(Particle->Proxy != nullptr)
 					{
 						int32 NewIdx = EventSleepDataArray.Add(TSleepingData<float, 3>());
 						TSleepingData<float, 3>& SleepingDataArrayItem = EventSleepDataArray[NewIdx];
 						SleepingDataArrayItem.Particle = Particle;
-						SleepingDataArrayItem.Sleeping = Particle->ObjectState() == EObjectStateType::Sleeping;
+						SleepingDataArrayItem.Sleeping = SleepData.Sleeping;
 
 						IPhysicsProxyBase* PhysicsProxy = Particle->Proxy;
 						AllSleepIndicesByPhysicsProxy.FindOrAdd(PhysicsProxy).Add(NewIdx);
