Change TBox to TAABB where we don't need the full generality

diff --git a/Engine/Plugins/Experimental/ChaosCloth/Source/Chaos/Private/ChaosCloth/ChaosClothingSimulation.cpp b/Engine/Plugins/Experimental/ChaosCloth/Source/Chaos/Private/ChaosCloth/ChaosClothingSimulation.cpp
--- a/Engine/Plugins/Experimental/ChaosCloth/Source/Chaos/Private/ChaosCloth/ChaosClothingSimulation.cpp
+++ b/Engine/Plugins/Experimental/ChaosCloth/Source/Chaos/Private/ChaosCloth/ChaosClothingSimulation.cpp
@@ -1397,7 +1397,7 @@ void ClothingSimulation::DebugDrawCollision(USkeletalMeshComponent* OwnerCompone
 		DrawWireSphere(PDI, Transform, Color, Radius, 12, SDPG_World, 0.0f, 0.001f, false);
 	};
 
-	auto DrawBox = [&PDI](const Chaos::TBox<float, 3>& Box, const TRotation<float, 3>& Rotation, const Chaos::TVector<float, 3>& Position, const FLinearColor& Color)
+	auto DrawBox = [&PDI](const Chaos::TAABB<float, 3>& Box, const TRotation<float, 3>& Rotation, const Chaos::TVector<float, 3>& Position, const FLinearColor& Color)
 	{
 		const FMatrix BoxToWorld = FTransform(Rotation, Position).ToMatrixNoScale();
 		const FVector Radii = Box.Extents() * 0.5f;
diff --git a/Engine/Plugins/Experimental/GeometryCollectionPlugin/Source/GeometryCollectionExampleCore/Resource/FracturedGeometry.cpp b/Engine/Plugins/Experimental/GeometryCollectionPlugin/Source/GeometryCollectionExampleCore/Resource/FracturedGeometry.cpp
--- a/Engine/Plugins/Experimental/GeometryCollectionPlugin/Source/GeometryCollectionExampleCore/Resource/FracturedGeometry.cpp
+++ b/Engine/Plugins/Experimental/GeometryCollectionPlugin/Source/GeometryCollectionExampleCore/Resource/FracturedGeometry.cpp
@@ -1954,7 +1954,7 @@ GlobalFracturedGeometry::GlobalFracturedGeometry()
 	}
 
 	// Measure the extent of the second point set.
-	Chaos::TBox<float, 3> BBox(
+	Chaos::TAABB<float, 3> BBox(
 		Chaos::TVector<float, 3>(
 			RawVertexArray[NumOrigFractured + 0],
 			RawVertexArray[NumOrigFractured + 1],
diff --git a/Engine/Source/Runtime/Core/Public/UObject/ExternalPhysicsCustomObjectVersion.h b/Engine/Source/Runtime/Core/Public/UObject/ExternalPhysicsCustomObjectVersion.h
--- a/Engine/Source/Runtime/Core/Public/UObject/ExternalPhysicsCustomObjectVersion.h
+++ b/Engine/Source/Runtime/Core/Public/UObject/ExternalPhysicsCustomObjectVersion.h
@@ -83,6 +83,9 @@ struct CORE_API FExternalPhysicsCustomObjectVersion
 
 		// Add damping to rigid particles
 		AddDampingToRigids,
+
+		//Replace TBox with TAABB in many places
+		TBoxReplacedWithTAABB,
 		
 		// -----<new versions can be added above this line>-------------------------------------------------
 		VersionPlusOne,
diff --git a/Engine/Source/Runtime/Engine/Private/PhysicsEngine/Experimental/ChaosDerivedData.cpp b/Engine/Source/Runtime/Engine/Private/PhysicsEngine/Experimental/ChaosDerivedData.cpp
--- a/Engine/Source/Runtime/Engine/Private/PhysicsEngine/Experimental/ChaosDerivedData.cpp
+++ b/Engine/Source/Runtime/Engine/Private/PhysicsEngine/Experimental/ChaosDerivedData.cpp
@@ -21,7 +21,7 @@ const TCHAR* FChaosDerivedDataCooker::GetPluginName() const
 
 const TCHAR* FChaosDerivedDataCooker::GetVersionString() const
 {
-	return TEXT("8AEC8C7BFB0A4906A6EB99E602E9594B");
+	return TEXT("D809C44D3C23455296C587E17D4C36C7");
 }
 
 FString FChaosDerivedDataCooker::GetPluginSpecificCacheKeySuffix() const
diff --git a/Engine/Source/Runtime/Engine/Private/PhysicsEngine/Experimental/ChaosInterfaceUtils.cpp b/Engine/Source/Runtime/Engine/Private/PhysicsEngine/Experimental/ChaosInterfaceUtils.cpp
--- a/Engine/Source/Runtime/Engine/Private/PhysicsEngine/Experimental/ChaosInterfaceUtils.cpp
+++ b/Engine/Source/Runtime/Engine/Private/PhysicsEngine/Experimental/ChaosInterfaceUtils.cpp
@@ -93,7 +93,7 @@ namespace ChaosInterface
 			const auto& Vertex = PhysXMesh->getVertices()[j];
 			CollisionMeshParticles.X(j) = Scale * Chaos::TVector<float, 3>(Vertex.x, Vertex.y, Vertex.z);
 		}
-		Chaos::TBox<float, 3> BoundingBox(CollisionMeshParticles.X(0), CollisionMeshParticles.X(0));
+		Chaos::TAABB<float, 3> BoundingBox(CollisionMeshParticles.X(0), CollisionMeshParticles.X(0));
 		for (uint32 j = 1; j < CollisionMeshParticles.Size(); ++j)
 		{
 			BoundingBox.GrowToInclude(CollisionMeshParticles.X(j));
@@ -175,7 +175,7 @@ namespace ChaosInterface
 			HalfExtents.Y = FMath::Max(HalfExtents.Y, KINDA_SMALL_NUMBER);
 			HalfExtents.Z = FMath::Max(HalfExtents.Z, KINDA_SMALL_NUMBER);
 
-			// TBox can handle translations internally but if we have a rotation we need to wrap it in a transform
+			// TAABB can handle translations internally but if we have a rotation we need to wrap it in a transform
 			TUniquePtr<Chaos::FImplicitObject> Implicit;
 			if (!BoxTransform.GetRotation().IsIdentity())
 			{
diff --git a/Engine/Source/Runtime/Engine/Private/PhysicsEngine/Experimental/PhysScene_Chaos.cpp b/Engine/Source/Runtime/Engine/Private/PhysicsEngine/Experimental/PhysScene_Chaos.cpp
--- a/Engine/Source/Runtime/Engine/Private/PhysicsEngine/Experimental/PhysScene_Chaos.cpp
+++ b/Engine/Source/Runtime/Engine/Private/PhysicsEngine/Experimental/PhysScene_Chaos.cpp
@@ -85,7 +85,7 @@ public:
 
 	}
 
-	virtual void Box(const Chaos::TBox<float, 3>& InBox, const Chaos::TVector<float, 3>& InLinearColor, float InThickness) override
+	virtual void Box(const Chaos::TAABB<float, 3>& InBox, const Chaos::TVector<float, 3>& InLinearColor, float InThickness) override
 	{
 		DrawDebugBox(World, InBox.Center(), InBox.Extents(), FQuat::Identity, FLinearColor(InLinearColor).ToFColor(true), false, -1.0f, SDPG_Foreground, InThickness);
 	}
@@ -222,8 +222,8 @@ private:
 
 					if (bDrawObjectBounds)
 					{
-						const TArray<TBox<float, 3>>& Boxes = BV->GetWorldSpaceBoxes();
-						for (const TBox<float, 3>& Box : Boxes)
+						const TArray<TAABB<float, 3>>& Boxes = BV->GetWorldSpaceBoxes();
+						for (const TAABB<float, 3>& Box : Boxes)
 						{
 							DrawDebugBox(WorldPtr, Box.Center(), Box.Extents() / 2.0f, FQuat::Identity, FColor::Cyan, false, -1.0f, SDPG_Foreground, 1.0f);
 						}
@@ -601,11 +601,11 @@ void FPhysScene_Chaos::UpdateActorInAccelerationStructure(const FPhysicsActorHan
 		if (SpatialAcceleration)
 		{
 
-			TBox<FReal, 3> WorldBounds;
+			TAABB<FReal, 3> WorldBounds;
 			const bool bHasBounds = Actor->Geometry()->HasBoundingBox();
 			if (bHasBounds)
 			{
-				WorldBounds = Actor->Geometry()->BoundingBox().GetAABB().TransformedAABB(TRigidTransform<FReal, 3>(Actor->X(), Actor->R()));
+				WorldBounds = Actor->Geometry()->BoundingBox().TransformedAABB(TRigidTransform<FReal, 3>(Actor->X(), Actor->R()));
 			}
 
 
@@ -973,11 +973,11 @@ void FPhysScene_ChaosInterface::AddActorsToScene_AssumesLocked(TArray<FPhysicsAc
 		{
 			// Get the bounding box for the particle if it has one
 			bool bHasBounds = Handle->Geometry()->HasBoundingBox();
-			Chaos::TBox<float, 3> WorldBounds;
+			Chaos::TAABB<float, 3> WorldBounds;
 			if (bHasBounds)
 			{
-				const Chaos::TBox<float, 3> LocalBounds = Handle->Geometry()->BoundingBox();
-				WorldBounds = LocalBounds.TransformedBox(Chaos::TRigidTransform<float, 3>(Handle->X(), Handle->R()));
+				const Chaos::TAABB<float, 3> LocalBounds = Handle->Geometry()->BoundingBox();
+				WorldBounds = LocalBounds.TransformedAABB(Chaos::TRigidTransform<float, 3>(Handle->X(), Handle->R()));
 			}
 
 			// Insert the particle
diff --git a/Engine/Source/Runtime/Engine/Private/PhysicsEngine/ImmediatePhysics/ImmediatePhysicsChaos/ImmediatePhysicsActorHandle_Chaos.cpp b/Engine/Source/Runtime/Engine/Private/PhysicsEngine/ImmediatePhysics/ImmediatePhysicsChaos/ImmediatePhysicsActorHandle_Chaos.cpp
--- a/Engine/Source/Runtime/Engine/Private/PhysicsEngine/ImmediatePhysics/ImmediatePhysicsChaos/ImmediatePhysicsActorHandle_Chaos.cpp
+++ b/Engine/Source/Runtime/Engine/Private/PhysicsEngine/ImmediatePhysics/ImmediatePhysicsChaos/ImmediatePhysicsActorHandle_Chaos.cpp
@@ -247,7 +247,7 @@ namespace ImmediatePhysics_Chaos
 				{
 					ParticleHandle->SetHasBounds(true);
 					ParticleHandle->SetLocalBounds(Geometry->BoundingBox());
-					ParticleHandle->SetWorldSpaceInflatedBounds(Geometry->BoundingBox().TransformedBox(TRigidTransform<float, 3>(ParticleHandle->X(), ParticleHandle->R())));
+					ParticleHandle->SetWorldSpaceInflatedBounds(Geometry->BoundingBox().TransformedAABB(TRigidTransform<float, 3>(ParticleHandle->X(), ParticleHandle->R())));
 				}
 
 				if (auto* Kinematic = ParticleHandle->CastToKinematicParticle())
diff --git a/Engine/Source/Runtime/Engine/Public/Physics/Experimental/PhysInterface_Chaos.h b/Engine/Source/Runtime/Engine/Public/Physics/Experimental/PhysInterface_Chaos.h
--- a/Engine/Source/Runtime/Engine/Public/Physics/Experimental/PhysInterface_Chaos.h
+++ b/Engine/Source/Runtime/Engine/Public/Physics/Experimental/PhysInterface_Chaos.h
@@ -54,6 +54,9 @@ namespace Chaos
 	template <typename T>
 	class TCapsule;
 
+	template <typename T, int>
+	class TAABB;
+
 	template <typename T, int>
 	class TBox;
 
diff --git a/Engine/Source/Runtime/Experimental/Chaos/Private/Chaos/BVHParticles.cpp b/Engine/Source/Runtime/Experimental/Chaos/Private/Chaos/BVHParticles.cpp
--- a/Engine/Source/Runtime/Experimental/Chaos/Private/Chaos/BVHParticles.cpp
+++ b/Engine/Source/Runtime/Experimental/Chaos/Private/Chaos/BVHParticles.cpp
@@ -66,7 +66,7 @@ void TBVHParticles<T,d>::UpdateAccelerationStructures()
 }
 
 template<class T, int d>
-const TArray<int32> TBVHParticles<T,d>::FindAllIntersections(const TBox<T, d>& Object) const
+const TArray<int32> TBVHParticles<T,d>::FindAllIntersections(const TAABB<T, d>& Object) const
 {
 	return MBVH->FindAllIntersections(Object);
 }
diff --git a/Engine/Source/Runtime/Experimental/Chaos/Private/Chaos/BoundingVolumeHierarchy.cpp b/Engine/Source/Runtime/Experimental/Chaos/Private/Chaos/BoundingVolumeHierarchy.cpp
--- a/Engine/Source/Runtime/Experimental/Chaos/Private/Chaos/BoundingVolumeHierarchy.cpp
+++ b/Engine/Source/Runtime/Experimental/Chaos/Private/Chaos/BoundingVolumeHierarchy.cpp
@@ -37,7 +37,7 @@ TBoundingVolumeHierarchy<OBJECT_ARRAY, LEAF_TYPE, T, d>::TBoundingVolumeHierarch
 }
 
 template<class OBJECT_ARRAY, class LEAF_TYPE, class T, int d>
-const TBox<T, d>& TBoundingVolumeHierarchy<OBJECT_ARRAY, LEAF_TYPE, T, d>::GetWorldSpaceBoundingBox(const TGeometryParticles<T, d>& InParticles, const int32 Index)
+const TAABB<T, d>& TBoundingVolumeHierarchy<OBJECT_ARRAY, LEAF_TYPE, T, d>::GetWorldSpaceBoundingBox(const TGeometryParticles<T, d>& InParticles, const int32 Index)
 {
 	return Chaos::GetWorldSpaceBoundingBox(InParticles, Index, MWorldSpaceBoxes);
 }
@@ -46,7 +46,7 @@ const TBox<T, d>& TBoundingVolumeHierarchy<OBJECT_ARRAY, LEAF_TYPE, T, d>::GetWo
 template<typename T, int d>
 void DrawNodeRecursive(ISpacialDebugDrawInterface<T>* InInterface, const TBVHNode<T, d>& InNode, const TArray<TBVHNode<T,d>>& InAllNodes)
 {
-	TBox<T, d> Box(InNode.MMin, InNode.MMax);
+	TAABB<T, d> Box(InNode.MMin, InNode.MMax);
 
 	InInterface->Box(Box, {1.0f, 1.0f, 1.0f}, 1.0f);
 
@@ -153,7 +153,7 @@ void TBoundingVolumeHierarchy<OBJECT_ARRAY, LEAF_TYPE, T, d>::UpdateHierarchyImp
 	ComputeAllWorldSpaceBoundingBoxes(*MObjects, AllObjects, bUseVelocity, Dt, MWorldSpaceBoxes);
 
 	int32 Axis;
-	const TBox<T, d> GlobalBox = ComputeGlobalBoxAndSplitAxis(*MObjects, AllObjects, MWorldSpaceBoxes, bAllowMultipleSplitting, Axis);
+	const TAABB<T, d> GlobalBox = ComputeGlobalBoxAndSplitAxis(*MObjects, AllObjects, MWorldSpaceBoxes, bAllowMultipleSplitting, Axis);
 
 	TBVHNode<T,d> RootNode;
 	RootNode.MMin = GlobalBox.Min();
@@ -201,7 +201,7 @@ TArray<int32> FindAllIntersectionsLeafHelper(const TBoundingVolume<TPayloadType,
 template<class OBJECT_ARRAY, typename TPayloadType, class T, int d>
 TArray<int32> TBoundingVolumeHierarchy<OBJECT_ARRAY, TPayloadType, T, d>::FindAllIntersectionsHelper(const TBVHNode<T,d>& MyNode, const TVector<T, d>& Point) const
 {
-	TBox<T, d> MBox(MyNode.MMin, MyNode.MMax);
+	TAABB<T, d> MBox(MyNode.MMin, MyNode.MMax);
 	if (MBox.SignedDistance(Point) > 0)
 	{
 		return TArray<int32>();
@@ -230,13 +230,13 @@ TArray<int32> TBoundingVolumeHierarchy<OBJECT_ARRAY, TPayloadType, T, d>::FindAl
 }
 
 template <typename T, int d>
-bool IntersectsHelper(const TBox<T, d>& WorldSpaceBox, const TBox<T, d>& LocalBox)
+bool IntersectsHelper(const TAABB<T, d>& WorldSpaceBox, const TAABB<T, d>& LocalBox)
 {
 	return WorldSpaceBox.Intersects(LocalBox);
 }
 
 template <typename T, int d>
-bool IntersectsHelper(const TBox<T, d>& WorldSpaceBox, const TSpatialRay<T, d>& Ray)
+bool IntersectsHelper(const TAABB<T, d>& WorldSpaceBox, const TSpatialRay<T, d>& Ray)
 {
 	//Very slow, should look at using slab algorithm
 	const FBox VolumeBox(WorldSpaceBox.Min(), WorldSpaceBox.Max());
@@ -252,7 +252,7 @@ template <typename OBJECT_ARRAY>
 struct TSpecializeParticlesHelper
 {
 	template <typename T, int d, typename QUERY_OBJECT>
-	static void AccumulateChildrenResults(TArray<int32>& AccumIntersectionList, const TArray<int32>& PotentialChildren, const QUERY_OBJECT& QueryObject, const TMap<int32, TBox<T,d>>& WorldSpaceBoxes)
+	static void AccumulateChildrenResults(TArray<int32>& AccumIntersectionList, const TArray<int32>& PotentialChildren, const QUERY_OBJECT& QueryObject, const TMap<int32, TAABB<T,d>>& WorldSpaceBoxes)
 	{
 		for (int32 ChildIndex : PotentialChildren)
 		{
@@ -271,7 +271,7 @@ template <typename T, int d>
 struct TSpecializeParticlesHelper<TParticles<T,d>>
 {
 	template <typename QUERY_OBJECT>
-	static void AccumulateChildrenResults(TArray<int32>& AccumIntersectionList, const TArray<int32>& PotentialChildren, const QUERY_OBJECT& QueryObject, const TMap<int32, TBox<T, d>>& WorldSpaceBoxes)
+	static void AccumulateChildrenResults(TArray<int32>& AccumIntersectionList, const TArray<int32>& PotentialChildren, const QUERY_OBJECT& QueryObject, const TMap<int32, TAABB<T, d>>& WorldSpaceBoxes)
 	{
 		if (CheckBox)
 		{
@@ -313,7 +313,7 @@ template<class OBJECT_ARRAY, class LEAF_TYPE, class T, int d>
 template <typename QUERY_OBJECT>
 void TBoundingVolumeHierarchy<OBJECT_ARRAY, LEAF_TYPE, T, d>::FindAllIntersectionsHelperRecursive(const TBVHNode<T,d>& MyNode, const QUERY_OBJECT& QueryObject, TArray<int32>& AccumulateElements) const
 {
-	TBox<T, d> MBox(MyNode.MMin, MyNode.MMax);
+	TAABB<T, d> MBox(MyNode.MMin, MyNode.MMax);
 	if (!IntersectsHelper(MBox, QueryObject))
 	{
 		return;
@@ -335,7 +335,7 @@ int32 UseAccumulationArray = 1;
 FAutoConsoleVariableRef CVarUseAccumulationArray(TEXT("p.UseAccumulationArray"), UseAccumulationArray, TEXT(""));
 
 template<class OBJECT_ARRAY, class LEAF_TYPE, class T, int d>
-TArray<int32> TBoundingVolumeHierarchy<OBJECT_ARRAY, LEAF_TYPE, T, d>::FindAllIntersectionsHelper(const TBVHNode<T,d>& MyNode, const TBox<T, d>& ObjectBox) const
+TArray<int32> TBoundingVolumeHierarchy<OBJECT_ARRAY, LEAF_TYPE, T, d>::FindAllIntersectionsHelper(const TBVHNode<T,d>& MyNode, const TAABB<T, d>& ObjectBox) const
 {
 	TArray<int32> IntersectionList;
 	FindAllIntersectionsHelperRecursive(MyNode, ObjectBox, IntersectionList);
@@ -395,7 +395,7 @@ struct FSplitCount
 };
 
 template <typename T, int d>
-void AccumulateNextLevelCount(const TBox<T, d>& Box, const TVector<T, d>& MidPoint, FSplitCount<d>& Counts)
+void AccumulateNextLevelCount(const TAABB<T, d>& Box, const TVector<T, d>& MidPoint, FSplitCount<d>& Counts)
 {
 	//todo(ocohen): particles min = max so avoid extra work
 	for (int32 i = 0; i < d; ++i)
@@ -419,15 +419,15 @@ int32 TBoundingVolumeHierarchy<OBJECT_ARRAY, LEAF_TYPE, T, d>::GenerateNextLevel
 	LocalElements.SetNum(2);
 	LocalObjects.SetNum(2);
 	LocalLeafs.SetNum(2);
-	TBox<T, d> GlobalBox(GlobalMin, GlobalMax);
+	TAABB<T, d> GlobalBox(GlobalMin, GlobalMax);
 	const TVector<T, d> WorldCenter = GlobalBox.Center();
-	const TVector<T, d> MinCenterSearch = TBox<T, d>(GlobalMin, WorldCenter).Center();
-	const TVector<T, d> MaxCenterSearch = TBox<T, d>(WorldCenter, GlobalMax).Center();
+	const TVector<T, d> MinCenterSearch = TAABB<T, d>(GlobalMin, WorldCenter).Center();
+	const TVector<T, d> MaxCenterSearch = TAABB<T, d>(WorldCenter, GlobalMax).Center();
 
 	for (int32 i = 0; i < Objects.Num(); ++i)
 	{
 		check(Objects[i] >= 0 && Objects[i] < GetObjectCount(*MObjects));
-		const TBox<T, d>& ObjectBox = Chaos::GetWorldSpaceBoundingBox(*MObjects, Objects[i], MWorldSpaceBoxes);
+		const TAABB<T, d>& ObjectBox = Chaos::GetWorldSpaceBoundingBox(*MObjects, Objects[i], MWorldSpaceBoxes);
 		const TVector<T, d> ObjectCenter = ObjectBox.Center();
 		bool MinA = false, MaxA = false;
 		if (ObjectBox.Min()[Axis] < WorldCenter[Axis])
@@ -517,12 +517,12 @@ int32 TBoundingVolumeHierarchy<OBJECT_ARRAY, LEAF_TYPE, T, d>::GenerateNextLevel
 	LocalElements.SetNum(8);
 	LocalObjects.SetNum(8);
 	LocalLeafs.SetNum(8);
-	TBox<T, d> GlobalBox(GlobalMin, GlobalMax);
+	TAABB<T, d> GlobalBox(GlobalMin, GlobalMax);
 	const TVector<T, d> WorldCenter = GlobalBox.Center();
 	for (int32 i = 0; i < Objects.Num(); ++i)
 	{
 		check(Objects[i] >= 0 && Objects[i] < GetObjectCount(*MObjects));
-		const TBox<T, d>& ObjectBox = Chaos::GetWorldSpaceBoundingBox(*MObjects, Objects[i], MWorldSpaceBoxes);
+		const TAABB<T, d>& ObjectBox = Chaos::GetWorldSpaceBoundingBox(*MObjects, Objects[i], MWorldSpaceBoxes);
 		const TVector<T, d> ObjectCenter = ObjectBox.Center();
 		bool MinX = false, MaxX = false, MinY = false, MaxY = false, MinZ = false, MaxZ = false;
 		if (ObjectBox.Min()[0] < WorldCenter[0])
@@ -608,7 +608,7 @@ int32 TBoundingVolumeHierarchy<OBJECT_ARRAY, LEAF_TYPE, T, d>::GenerateNextLevel
 #endif
 		if (LocalObjects[i].Num() > MIN_NUM_OBJECTS && Level < MMaxLevels && LocalObjects[i].Num() < Objects.Num())
 		{
-			TBox<T, d> LocalBox(Min, Max);
+			TAABB<T, d> LocalBox(Min, Max);
 			int32 NextAxis = 0;
 			TVector<T, d> LocalExtents = LocalBox.Extents();
 			if (LocalExtents[2] > LocalExtents[0] && LocalExtents[2] > LocalExtents[1])
@@ -672,7 +672,9 @@ void FixupLeafObj(const OBJECT_ARRAY& Objects, TArray<TArray<int32>>& Leafs)
 template<class OBJECT_ARRAY, class LEAF_TYPE, class T, int d>
 void TBoundingVolumeHierarchy<OBJECT_ARRAY, LEAF_TYPE, T, d>::Serialize(FArchive& Ar)
 {
-	Ar << MGlobalObjects << MWorldSpaceBoxes << MMaxLevels << Elements << Leafs;
+	Ar << MGlobalObjects;
+	TBox<T, d>::SerializeAsAABBs(Ar, MWorldSpaceBoxes);
+	Ar << MMaxLevels << Elements << Leafs;
 	if (Ar.IsLoading())
 	{
 		FixupLeafObj(*MObjects, Leafs);
diff --git a/Engine/Source/Runtime/Experimental/Chaos/Private/Chaos/ChaosDebugDraw.cpp b/Engine/Source/Runtime/Experimental/Chaos/Private/Chaos/ChaosDebugDraw.cpp
--- a/Engine/Source/Runtime/Experimental/Chaos/Private/Chaos/ChaosDebugDraw.cpp
+++ b/Engine/Source/Runtime/Experimental/Chaos/Private/Chaos/ChaosDebugDraw.cpp
@@ -125,7 +125,7 @@ namespace Chaos
 		{
 			FColor Color = (ColorScale * FColor::White).ToFColor(false);
 
-			TBox<FReal, 3> Box = InParticle->WorldSpaceInflatedBounds();
+			TAABB<FReal, 3> Box = InParticle->WorldSpaceInflatedBounds();
 			FVec3 P = SpaceTransform.TransformPosition(Box.GetCenter());
 			FRotation3 Q = SpaceTransform.GetRotation();
 			FMatrix33 Qm = Q.ToMatrix();
diff --git a/Engine/Source/Runtime/Experimental/Chaos/Private/Chaos/CollisionResolution.cpp b/Engine/Source/Runtime/Experimental/Chaos/Private/Chaos/CollisionResolution.cpp
--- a/Engine/Source/Runtime/Experimental/Chaos/Private/Chaos/CollisionResolution.cpp
+++ b/Engine/Source/Runtime/Experimental/Chaos/Private/Chaos/CollisionResolution.cpp
@@ -74,7 +74,7 @@ namespace Chaos
 
 
 		template <typename T, int d>
-		TContactPoint<T> BoxBoxContactPoint(const TBox<T, d>& Box1, const TRigidTransform<T, d>& ATM, const TBox<T, d>& Box2, const TRigidTransform<T, d>& BTM, const T Thickness)
+		TContactPoint<T> BoxBoxContactPoint(const TAABB<T, d>& Box1, const TRigidTransform<T, d>& ATM, const TAABB<T, d>& Box2, const TRigidTransform<T, d>& BTM, const T Thickness)
 		{
 			return GJKContactPoint(Box1, ATM, Box2, BTM, Thickness);
 		}
@@ -83,7 +83,7 @@ namespace Chaos
 		//
 
 		template <typename T, int d>
-		void UpdateBoxBoxConstraint(const TBox<T, d>& Box1, const TRigidTransform<T, d>& Box1Transform, const TBox<T, d>& Box2, const TRigidTransform<T, d>& Box2Transform, const T Thickness, TRigidBodyPointContactConstraint<T, d>& Constraint)
+		void UpdateBoxBoxConstraint(const TAABB<T, d>& Box1, const TRigidTransform<T, d>& Box1Transform, const TAABB<T, d>& Box2, const TRigidTransform<T, d>& Box2Transform, const T Thickness, TRigidBodyPointContactConstraint<T, d>& Constraint)
 		{
 			UpdateContactPoint(Constraint.Manifold, BoxBoxContactPoint(Box1, Box1Transform, Box2, Box2Transform, Thickness));
 		}
@@ -101,7 +101,7 @@ namespace Chaos
 				Constraint->Particle[1] = Particle1;
 				Constraint->SetManifold(Implicit0, Implicit1);
 
-				UpdateBoxBoxConstraint(*Object0, Transform0, *Object1, Transform1, Thickness, *Constraint);
+				UpdateBoxBoxConstraint(Object0->BoundingBox(), Transform0, Object1->BoundingBox(), Transform1, Thickness, *Constraint);
 
 				if (Constraint->GetPhi() < Thickness)
 				{
@@ -119,7 +119,7 @@ namespace Chaos
 		//
 
 		template <typename T, int d>
-		bool UpdateBoxPlaneConstraint(const TBox<T, d>& Box, const TRigidTransform<T, d>& BoxTransform, const TPlane<T, d>& Plane, const TRigidTransform<T, d>& PlaneTransform, const T Thickness, TRigidBodyPointContactConstraint<T, d>& Constraint)
+		bool UpdateBoxPlaneConstraint(const TAABB<T, d>& Box, const TRigidTransform<T, d>& BoxTransform, const TPlane<T, d>& Plane, const TRigidTransform<T, d>& PlaneTransform, const T Thickness, TRigidBodyPointContactConstraint<T, d>& Constraint)
 		{
 			TCollisionContact<T, d> & Contact = Constraint.Manifold;
 
@@ -194,7 +194,7 @@ namespace Chaos
 				Constraint->Particle[1] = Particle1;
 				Constraint->SetManifold(Implicit0, Implicit1);
 
-				UpdateBoxPlaneConstraint(*Object0, Transform0, *Object1, Transform1, Thickness, *Constraint);
+				UpdateBoxPlaneConstraint(Object0->BoundingBox(), Transform0, *Object1, Transform1, Thickness, *Constraint);
 
 				if (Constraint->GetPhi() < Thickness)
 				{
@@ -317,7 +317,7 @@ namespace Chaos
 		//
 
 		template <typename T, int d>
-		TContactPoint<T> SphereBoxContactPoint(const TSphere<T, d>& Sphere, const TRigidTransform<T, d>& SphereTransform, const TBox<T, d>& Box, const TRigidTransform<T, d>& BoxTransform, const T Thickness)
+		TContactPoint<T> SphereBoxContactPoint(const TSphere<T, d>& Sphere, const TRigidTransform<T, d>& SphereTransform, const TAABB<T, d>& Box, const TRigidTransform<T, d>& BoxTransform, const T Thickness)
 		{
 			TContactPoint<T> Result;
 
@@ -335,7 +335,7 @@ namespace Chaos
 		}
 
 		template <typename T, int d>
-		void UpdateSphereBoxConstraint(const TSphere<T, d>& Sphere, const TRigidTransform<T, d>& SphereTransform, const TBox<T, d>& Box, const TRigidTransform<T, d>& BoxTransform, const T Thickness, TRigidBodyPointContactConstraint<T, d>& Constraint)
+		void UpdateSphereBoxConstraint(const TSphere<T, d>& Sphere, const TRigidTransform<T, d>& SphereTransform, const TAABB<T, d>& Box, const TRigidTransform<T, d>& BoxTransform, const T Thickness, TRigidBodyPointContactConstraint<T, d>& Constraint)
 		{
 			UpdateContactPoint(Constraint.Manifold, SphereBoxContactPoint(Sphere, SphereTransform, Box, BoxTransform, Thickness));
 		}
@@ -353,7 +353,7 @@ namespace Chaos
 				Constraint->Particle[1] = Particle1;
 				Constraint->SetManifold(Implicit0, Implicit1);
 
-				UpdateSphereBoxConstraint(*Object0, Transform0, *Object1, Transform1, Thickness, *Constraint);
+				UpdateSphereBoxConstraint(*Object0, Transform0, Object1->BoundingBox(), Transform1, Thickness, *Constraint);
 
 				if (Constraint->GetPhi() < Thickness)
 				{
@@ -494,13 +494,13 @@ namespace Chaos
 		//
 
 		template <typename T, int d>
-		TContactPoint<T> CapsuleBoxContactPoint(const TCapsule<T>& A, const TRigidTransform<T, d>& ATransform, const TBox<T, d>& B, const TRigidTransform<T, d>& BTransform, const T Thickness)
+		TContactPoint<T> CapsuleBoxContactPoint(const TCapsule<T>& A, const TRigidTransform<T, d>& ATransform, const TAABB<T, d>& B, const TRigidTransform<T, d>& BTransform, const T Thickness)
 		{
 			return GJKContactPoint(A, ATransform, B, BTransform, Thickness);
 		}
 
 		template <typename T, int d>
-		void UpdateCapsuleBoxConstraint(const TCapsule<T>& A, const TRigidTransform<T, d>& ATransform, const TBox<T, d>& B, const TRigidTransform<T, d>& BTransform, const T Thickness, TRigidBodyPointContactConstraint<T, d>& Constraint)
+		void UpdateCapsuleBoxConstraint(const TCapsule<T>& A, const TRigidTransform<T, d>& ATransform, const TAABB<T, d>& B, const TRigidTransform<T, d>& BTransform, const T Thickness, TRigidBodyPointContactConstraint<T, d>& Constraint)
 		{
 			UpdateContactPoint(Constraint.Manifold, CapsuleBoxContactPoint(A, ATransform, B, BTransform, Thickness));
 		}
@@ -518,7 +518,7 @@ namespace Chaos
 				Constraint->Particle[1] = Particle1;
 				Constraint->SetManifold(Implicit0, Implicit1);
 
-				UpdateCapsuleBoxConstraint(*Object0, Transform0, *Object1, Transform1, Thickness, *Constraint);
+				UpdateCapsuleBoxConstraint(*Object0, Transform0, Object1->BoundingBox(), Transform1, Thickness, *Constraint);
 
 				if (Constraint->GetPhi() < Thickness)
 				{
@@ -968,7 +968,7 @@ namespace Chaos
 			}
 			else if (Implicit0Type == TBox<T, d>::StaticType() && Implicit1Type == TBox<T, d>::StaticType())
 			{
-				UpdateBoxBoxConstraint(*Implicit0.template GetObject<TBox<T, d>>(), Transform0, *Implicit1.template GetObject<TBox<T, d>>(), Transform1, Thickness, *ConstraintBase.template As<TRigidBodyPointContactConstraint<T,d>>());
+				UpdateBoxBoxConstraint(Implicit0.template GetObject<TBox<T, d>>()->GetAABB(), Transform0, Implicit1.template GetObject<TBox<T, d>>()->GetAABB(), Transform1, Thickness, *ConstraintBase.template As<TRigidBodyPointContactConstraint<T,d>>());
 			}
 			else if (Implicit0Type == TSphere<T, d>::StaticType() && Implicit1Type == TSphere<T, d>::StaticType())
 			{
@@ -976,7 +976,7 @@ namespace Chaos
 			}
 			else if (Implicit0Type == TBox<T, d>::StaticType() && Implicit1Type == TPlane<T, d>::StaticType())
 			{
-				UpdateBoxPlaneConstraint(*Implicit0.template GetObject<TBox<T, d>>(), Transform0, *Implicit1.template GetObject<TPlane<T, d>>(), Transform1, Thickness, *ConstraintBase.template As<TRigidBodyPointContactConstraint<T,d>>());
+				UpdateBoxPlaneConstraint(Implicit0.template GetObject<TBox<T, d>>()->GetAABB(), Transform0, *Implicit1.template GetObject<TPlane<T, d>>(), Transform1, Thickness, *ConstraintBase.template As<TRigidBodyPointContactConstraint<T,d>>());
 			}
 			else if (Implicit0Type == TSphere<T, d>::StaticType() && Implicit1Type == TPlane<T, d>::StaticType())
 			{
@@ -984,7 +984,7 @@ namespace Chaos
 			}
 			else if (Implicit0Type == TSphere<T, d>::StaticType() && Implicit1Type == TBox<T, d>::StaticType())
 			{
-				UpdateSphereBoxConstraint(*Implicit0.template GetObject<TSphere<T, d>>(), Transform0, *Implicit1.template GetObject<TBox<T, d>>(), Transform1, Thickness, *ConstraintBase.template As<TRigidBodyPointContactConstraint<T,d>>());
+				UpdateSphereBoxConstraint(*Implicit0.template GetObject<TSphere<T, d>>(), Transform0, Implicit1.template GetObject<TBox<T, d>>()->GetAABB(), Transform1, Thickness, *ConstraintBase.template As<TRigidBodyPointContactConstraint<T,d>>());
 			}
 			else if (Implicit0Type == TSphere<T, d>::StaticType() && Implicit1Type == TCapsule<T>::StaticType())
 			{
@@ -996,13 +996,13 @@ namespace Chaos
 			}
 			else if (Implicit0Type == TCapsule<T>::StaticType() && Implicit1Type == TBox<T, d>::StaticType())
 			{
-				UpdateCapsuleBoxConstraint(*Implicit0.template GetObject<TCapsule<T>>(), Transform0, *Implicit1.template GetObject<TBox<T, d>>(), Transform1, Thickness, *ConstraintBase.template As<TRigidBodyPointContactConstraint<T,d>>());
+				UpdateCapsuleBoxConstraint(*Implicit0.template GetObject<TCapsule<T>>(), Transform0, Implicit1.template GetObject<TBox<T, d>>()->GetAABB(), Transform1, Thickness, *ConstraintBase.template As<TRigidBodyPointContactConstraint<T,d>>());
 			}
 			else if (Implicit0Type == TPlane<T, d>::StaticType() && Implicit1Type == TBox<T, d>::StaticType())
 			{
 				TRigidBodyPointContactConstraint<T, d>& Constraint = *ConstraintBase.template As<TRigidBodyPointContactConstraint<T, d> >();
 				TRigidBodyPointContactConstraint<T, d> TmpConstraint = Constraint;
-				UpdateBoxPlaneConstraint(*Implicit1.template GetObject<TBox<T, d>>(), Transform1, *Implicit0.template GetObject<TPlane<T, d>>(), Transform0, Thickness, TmpConstraint);
+				UpdateBoxPlaneConstraint(Implicit1.template GetObject<TBox<T, d>>()->GetAABB(), Transform1, *Implicit0.template GetObject<TPlane<T, d>>(), Transform0, Thickness, TmpConstraint);
 				if (TmpConstraint.GetPhi() < Constraint.GetPhi())
 				{
 					Constraint = TmpConstraint;
@@ -1024,7 +1024,7 @@ namespace Chaos
 			{
 				TRigidBodyPointContactConstraint<T, d>& Constraint = *ConstraintBase.template As<TRigidBodyPointContactConstraint<T, d> >();
 				TRigidBodyPointContactConstraint<T, d> TmpConstraint = Constraint;
-				UpdateSphereBoxConstraint(*Implicit1.template GetObject<TSphere<T, d>>(), Transform1, *Implicit0.template GetObject<TBox<T, d>>(), Transform0, Thickness, TmpConstraint);
+				UpdateSphereBoxConstraint(*Implicit1.template GetObject<TSphere<T, d>>(), Transform1, Implicit0.template GetObject<TBox<T, d>>()->GetAABB(), Transform0, Thickness, TmpConstraint);
 				if (TmpConstraint.GetPhi() < Constraint.GetPhi())
 				{
 					Constraint = TmpConstraint;
@@ -1035,7 +1035,7 @@ namespace Chaos
 			{
 				TRigidBodyPointContactConstraint<T, d>& Constraint = *ConstraintBase.template As<TRigidBodyPointContactConstraint<T, d> >();
 				TRigidBodyPointContactConstraint<T, d> TmpConstraint = Constraint;
-				UpdateCapsuleBoxConstraint(*Implicit1.template GetObject<TCapsule<T>>(), Transform1, *Implicit0.template GetObject<TBox<T, d>>(), Transform0, Thickness, TmpConstraint);
+				UpdateCapsuleBoxConstraint(*Implicit1.template GetObject<TCapsule<T>>(), Transform1, Implicit0.template GetObject<TBox<T, d>>()->GetAABB(), Transform0, Thickness, TmpConstraint);
 				if (TmpConstraint.GetPhi() < Constraint.GetPhi())
 				{
 					Constraint = TmpConstraint;
@@ -1230,10 +1230,10 @@ namespace Chaos
 
 
 
-		template void UpdateBoxBoxConstraint<float, 3>(const TBox<float, 3>& Box1, const TRigidTransform<float, 3>& Box1Transform, const TBox<float, 3>& Box2, const TRigidTransform<float, 3>& Box2Transform, const float Thickness, TRigidBodyPointContactConstraint<float, 3>& Constraint);
+		template void UpdateBoxBoxConstraint<float, 3>(const TAABB<float, 3>& Box1, const TRigidTransform<float, 3>& Box1Transform, const TAABB<float, 3>& Box2, const TRigidTransform<float, 3>& Box2Transform, const float Thickness, TRigidBodyPointContactConstraint<float, 3>& Constraint);
 		template void ConstructBoxBoxConstraints<float, 3>(TGeometryParticleHandle<float, 3>* Particle0, TGeometryParticleHandle<float, 3>* Particle1, const FImplicitObject* Implicit0, const FImplicitObject* Implicit1, const TRigidTransform<float, 3>& Transform0, const TRigidTransform<float, 3>& Transform1, const float Thickness, FCollisionConstraintsArray& NewConstraints);
 
-		template bool UpdateBoxPlaneConstraint<float, 3>(const TBox<float, 3>& Box, const TRigidTransform<float, 3>& BoxTransform, const TPlane<float, 3>& Plane, const TRigidTransform<float, 3>& PlaneTransform, const float Thickness, TRigidBodyPointContactConstraint<float, 3>& Constraint);
+		template bool UpdateBoxPlaneConstraint<float, 3>(const TAABB<float, 3>& Box, const TRigidTransform<float, 3>& BoxTransform, const TPlane<float, 3>& Plane, const TRigidTransform<float, 3>& PlaneTransform, const float Thickness, TRigidBodyPointContactConstraint<float, 3>& Constraint);
 		template void ConstructBoxPlaneConstraints<float, 3>(TGeometryParticleHandle<float, 3>* Particle0, TGeometryParticleHandle<float, 3>* Particle1, const FImplicitObject* Implicit0, const FImplicitObject* Implicit1, const TRigidTransform<float, 3>& Transform0, const TRigidTransform<float, 3>& Transform1, const float Thickness, FCollisionConstraintsArray& NewConstraints);
 
 		template void UpdateSphereSphereConstraint<float, 3>(const TSphere<float, 3>& Sphere1, const TRigidTransform<float, 3>& Sphere1Transform, const TSphere<float, 3>& Sphere2, const TRigidTransform<float, 3>& Sphere2Transform, const float Thickness, TRigidBodyPointContactConstraint<float, 3>& Constraint);
@@ -1242,7 +1242,7 @@ namespace Chaos
 		template void UpdateSpherePlaneConstraint<float, 3>(const TSphere<float, 3>& Sphere, const TRigidTransform<float, 3>& SphereTransform, const TPlane<float, 3>& Plane, const TRigidTransform<float, 3>& PlaneTransform, const float Thickness, TRigidBodyPointContactConstraint<float, 3>& Constraint);
 		template void ConstructSpherePlaneConstraints<float, 3>(TGeometryParticleHandle<float, 3>* Particle0, TGeometryParticleHandle<float, 3>* Particle1, const FImplicitObject* Implicit0, const FImplicitObject* Implicit1, const TRigidTransform<float, 3>& Transform0, const TRigidTransform<float, 3>& Transform1, const float Thickness, FCollisionConstraintsArray& NewConstraints);
 
-		template void UpdateSphereBoxConstraint<float, 3>(const TSphere<float, 3>& Sphere, const TRigidTransform<float, 3>& SphereTransform, const TBox<float, 3>& Box, const TRigidTransform<float, 3>& BoxTransform, const float Thickness, TRigidBodyPointContactConstraint<float, 3>& Constraint);
+		template void UpdateSphereBoxConstraint<float, 3>(const TSphere<float, 3>& Sphere, const TRigidTransform<float, 3>& SphereTransform, const TAABB<float, 3>& Box, const TRigidTransform<float, 3>& BoxTransform, const float Thickness, TRigidBodyPointContactConstraint<float, 3>& Constraint);
 		template void ConstructSphereBoxConstraints<float, 3>(TGeometryParticleHandle<float, 3>* Particle0, TGeometryParticleHandle<float, 3>* Particle1, const FImplicitObject* Implicit0, const FImplicitObject* Implicit1, const TRigidTransform<float, 3>& Transform0, const TRigidTransform<float, 3>& Transform1, const float Thickness, FCollisionConstraintsArray& NewConstraints);
 
 		template void UpdateSphereCapsuleConstraint<float, 3>(const TSphere<float, 3>& Sphere, const TRigidTransform<float, 3>& SphereTransform, const TCapsule<float>& Box, const TRigidTransform<float, 3>& BoxTransform, const float Thickness, TRigidBodyPointContactConstraint<float, 3>& Constraint);
@@ -1251,7 +1251,7 @@ namespace Chaos
 		template void UpdateCapsuleCapsuleConstraint<float, 3>(const TCapsule<float>& A, const TRigidTransform<float, 3>& ATransform, const TCapsule<float>& B, const TRigidTransform<float, 3>& BTransform, const float Thickness, TRigidBodyPointContactConstraint<float, 3>& Constraint);
 		template void ConstructCapsuleCapsuleConstraints(TGeometryParticleHandle<float, 3>* Particle0, TGeometryParticleHandle<float, 3>* Particle1, const FImplicitObject* Implicit0, const FImplicitObject* Implicit1, const TRigidTransform<float, 3>& Transform0, const TRigidTransform<float, 3>& Transform1, const float Thickness, FCollisionConstraintsArray& NewConstraints);
 
-		template void UpdateCapsuleBoxConstraint<float, 3>(const TCapsule<float>& A, const TRigidTransform<float, 3>& ATransform, const TBox<float, 3>& B, const TRigidTransform<float, 3>& BTransform, const float Thickness, TRigidBodyPointContactConstraint<float, 3>& Constraint);
+		template void UpdateCapsuleBoxConstraint<float, 3>(const TCapsule<float>& A, const TRigidTransform<float, 3>& ATransform, const TAABB<float, 3>& B, const TRigidTransform<float, 3>& BTransform, const float Thickness, TRigidBodyPointContactConstraint<float, 3>& Constraint);
 		template void ConstructCapsuleBoxConstraints<float, 3>(TGeometryParticleHandle<float, 3>* Particle0, TGeometryParticleHandle<float, 3>* Particle1, const FImplicitObject* Implicit0, const FImplicitObject* Implicit1, const TRigidTransform<float, 3>& Transform0, const TRigidTransform<float, 3>& Transform1, const float Thickness, FCollisionConstraintsArray& NewConstraints);
 
 		template void UpdateConvexConvexConstraint<float, 3>(const FImplicitObject& A, const TRigidTransform<float, 3>& ATM, const FImplicitObject& B, const TRigidTransform<float, 3>& BTM, const float Thickness, TRigidBodyIterativeContactConstraint<float, 3>& Constraint);
diff --git a/Engine/Source/Runtime/Experimental/Chaos/Private/Chaos/CollisionResolutionUtil.cpp b/Engine/Source/Runtime/Experimental/Chaos/Private/Chaos/CollisionResolutionUtil.cpp
--- a/Engine/Source/Runtime/Experimental/Chaos/Private/Chaos/CollisionResolutionUtil.cpp
+++ b/Engine/Source/Runtime/Experimental/Chaos/Private/Chaos/CollisionResolutionUtil.cpp
@@ -169,7 +169,7 @@ namespace Chaos
 			if (NumParticles > SampleMinParticlesForAcceleration && Object.HasBoundingBox())
 			{
 				SCOPE_CYCLE_COUNTER(STAT_UpdateLevelsetPartial);
-				TBox<float, 3> ImplicitBox = Object.BoundingBox().TransformedBox(ObjectTransform.GetRelativeTransform(SampleParticlesTransform));
+				TAABB<float, 3> ImplicitBox = Object.BoundingBox().TransformedAABB(ObjectTransform.GetRelativeTransform(SampleParticlesTransform));
 				ImplicitBox.Thicken(Thickness);
 				TArray<int32> PotentialParticles;
 				{
@@ -475,7 +475,7 @@ namespace Chaos
 			if (NumParticles > SampleMinParticlesForAcceleration && Object.HasBoundingBox())
 			{
 				SCOPE_CYCLE_COUNTER(STAT_UpdateLevelsetPartial);
-				TBox<T, d> ImplicitBox = Object.BoundingBox().TransformedBox(ObjectTransform.GetRelativeTransform(SampleParticlesTransform));
+				TAABB<T, d> ImplicitBox = Object.BoundingBox().TransformedBox(ObjectTransform.GetRelativeTransform(SampleParticlesTransform));
 				ImplicitBox.Thicken(Thickness);
 				TArray<int32> PotentialParticles;
 				{
@@ -574,7 +574,7 @@ namespace Chaos
 				if (ParticleObj->HasBoundingBox())
 				{
 					const TRigidTransform<T, d> ParticlesToLevelsetTM = ParticlesTM.GetRelativeTransform(LevelsetTM);
-					TBox<T, d> ParticleBoundsInLevelset = ParticleObj->BoundingBox().TransformedBox(ParticlesToLevelsetTM);
+					TAABB<T, d> ParticleBoundsInLevelset = ParticleObj->BoundingBox().TransformedAABB(ParticlesToLevelsetTM);
 					ParticleBoundsInLevelset.Thicken(Thickness);
 					{
 						LevelsetObj.FindAllIntersectingObjects(RelevantShapes, ParticleBoundsInLevelset);
diff --git a/Engine/Source/Runtime/Experimental/Chaos/Private/Chaos/Evolution/PBDMinEvolution.cpp b/Engine/Source/Runtime/Experimental/Chaos/Private/Chaos/Evolution/PBDMinEvolution.cpp
--- a/Engine/Source/Runtime/Experimental/Chaos/Private/Chaos/Evolution/PBDMinEvolution.cpp
+++ b/Engine/Source/Runtime/Experimental/Chaos/Private/Chaos/Evolution/PBDMinEvolution.cpp
@@ -139,8 +139,8 @@ namespace Chaos
 				{
 					const FReal BoundsThickness = 1.0f;
 					const FReal BoundsThicknessVelocityInflation = 2.0f;
-					const TBox<FReal, 3>& LocalBounds = Particle.LocalBounds();
-					TBox<FReal, 3> WorldSpaceBounds = LocalBounds.TransformedBox(FRigidTransform3(Particle.P(), Particle.Q()));
+					const TAABB<FReal, 3>& LocalBounds = Particle.LocalBounds();
+					TAABB<FReal, 3> WorldSpaceBounds = LocalBounds.TransformedAABB(FRigidTransform3(Particle.P(), Particle.Q()));
 					WorldSpaceBounds.ThickenSymmetrically(FVec3(BoundsThickness) + BoundsThicknessVelocityInflation * Particle.V().GetAbs() * Dt);
 					Particle.SetWorldSpaceInflatedBounds(WorldSpaceBounds);
 				}
diff --git a/Engine/Source/Runtime/Experimental/Chaos/Private/Chaos/GeometryParticles.cpp b/Engine/Source/Runtime/Experimental/Chaos/Private/Chaos/GeometryParticles.cpp
--- a/Engine/Source/Runtime/Experimental/Chaos/Private/Chaos/GeometryParticles.cpp
+++ b/Engine/Source/Runtime/Experimental/Chaos/Private/Chaos/GeometryParticles.cpp
@@ -79,7 +79,7 @@ namespace Chaos
 	{
 		if (Geometry && Geometry->HasBoundingBox())
 		{
-			WorldSpaceInflatedShapeBounds = Geometry->BoundingBox().GetAABB().TransformedAABB(WorldTM);
+			WorldSpaceInflatedShapeBounds = Geometry->BoundingBox().TransformedAABB(WorldTM);
 		}
 	}
 
@@ -100,7 +100,7 @@ namespace Chaos
 
 		if (Ar.CustomVer(FExternalPhysicsCustomObjectVersion::GUID) >= FExternalPhysicsCustomObjectVersion::SerializeShapeWorldSpaceBounds)
 		{
-			Ar << WorldSpaceInflatedShapeBounds;
+			TBox<T, d>::SerializeAsAABB(Ar, WorldSpaceInflatedShapeBounds);
 		}
 		else
 		{
diff --git a/Engine/Source/Runtime/Experimental/Chaos/Private/Chaos/HeightField.cpp b/Engine/Source/Runtime/Experimental/Chaos/Private/Chaos/HeightField.cpp
--- a/Engine/Source/Runtime/Experimental/Chaos/Private/Chaos/HeightField.cpp
+++ b/Engine/Source/Runtime/Experimental/Chaos/Private/Chaos/HeightField.cpp
@@ -273,7 +273,7 @@ namespace Chaos
 	};
 
 	template<typename HeightfieldT, typename BufferType>
-	void BuildGeomData(TArrayView<BufferType> BufferView, TArrayView<uint8> MaterialIndexView, int32 NumRows, int32 NumCols, const TVector<HeightfieldT, 3>& InScale, TUniqueFunction<HeightfieldT(const BufferType)> ToRealFunc, typename THeightField<HeightfieldT>::FDataType& OutData, TBox<HeightfieldT, 3>& OutBounds)
+	void BuildGeomData(TArrayView<BufferType> BufferView, TArrayView<uint8> MaterialIndexView, int32 NumRows, int32 NumCols, const TVector<HeightfieldT, 3>& InScale, TUniqueFunction<HeightfieldT(const BufferType)> ToRealFunc, typename THeightField<HeightfieldT>::FDataType& OutData, TAABB<HeightfieldT, 3>& OutBounds)
 	{
 		using FDataType = typename THeightField<HeightfieldT>::FDataType;
 
@@ -321,7 +321,7 @@ namespace Chaos
 			TVector<HeightfieldT, 3> Position(RealType(X), RealType(Y), OutData.MinValue + OutData.Heights[HeightIndex] * OutData.HeightPerUnit);
 			if(HeightIndex == 0)
 			{
-				OutBounds = TBox<HeightfieldT, 3>(Position * InScale, Position * InScale);
+				OutBounds = TAABB<HeightfieldT, 3>(Position * InScale, Position * InScale);
 			}
 			else
 			{
@@ -347,7 +347,7 @@ namespace Chaos
 	}
 
 	template<typename HeightfieldT, typename BufferType>
-	void EditGeomData(TArrayView<BufferType> BufferView, int32 InBeginRow, int32 InBeginCol, int32 NumRows, int32 NumCols, TUniqueFunction<HeightfieldT(const BufferType)> ToRealFunc, typename THeightField<HeightfieldT>::FDataType& OutData, TBox<HeightfieldT, 3>& OutBounds)
+	void EditGeomData(TArrayView<BufferType> BufferView, int32 InBeginRow, int32 InBeginCol, int32 NumRows, int32 NumCols, TUniqueFunction<HeightfieldT(const BufferType)> ToRealFunc, typename THeightField<HeightfieldT>::FDataType& OutData, TAABB<HeightfieldT, 3>& OutBounds)
 	{
 		using FDataType = typename THeightField<HeightfieldT>::FDataType;
 		using RealType = typename FDataType::RealType;
@@ -399,7 +399,7 @@ namespace Chaos
 					TVector<HeightfieldT, 3> Position(RealType(X), RealType(Y), NewMin + OutData.Heights[HeightIndex] * NewHeightPerUnit);
 					if(HeightIndex == 0)
 					{
-						OutBounds = TBox<HeightfieldT, 3>(Position, Position);
+						OutBounds = TAABB<HeightfieldT, 3>(Position, Position);
 					}
 					else
 					{
@@ -616,7 +616,7 @@ namespace Chaos
 
 				if (GetCellBounds3DScaled(Cell, Min, Max))
 				{
-					if (TBox<T, 3>::RaycastFast(Min, Max, StartPoint, Dir, InvDir, bParallel, CurrentLength, InvCurrentLength, ToI, HitPoint))
+					if (TAABB<T, 3>(Min,Max).RaycastFast(StartPoint, Dir, InvDir, bParallel, CurrentLength, InvCurrentLength, ToI, HitPoint))
 					{
 						Visitor.VisitRaycast(Cell[1] * (GeomData.NumCols - 1) + Cell[0], CurrentLength);
 					}
@@ -646,7 +646,7 @@ namespace Chaos
 					if (GetCellBounds3DScaled(StartCell, Min, Max))
 					{
 						// Check cell bounds
-						if (TBox<T, 3>::RaycastFast(Min, Max, StartPoint, Dir, InvDir, bParallel, CurrentLength, InvCurrentLength, ToI, HitPoint))
+						if (TAABB<T, 3>(Min,Max).RaycastFast(StartPoint, Dir, InvDir, bParallel, CurrentLength, InvCurrentLength, ToI, HitPoint))
 						{
 							Visitor.VisitRaycast(StartCell[1] * (GeomData.NumCols - 1) + StartCell[0], CurrentLength);
 						}
@@ -658,7 +658,7 @@ namespace Chaos
 					if (GetCellBounds3DScaled(StartCell, Min, Max))
 					{
 						// Check cell bounds
-						if (TBox<T, 3>::RaycastFast(Min, Max, StartPoint, Dir, InvDir, bParallel, CurrentLength, InvCurrentLength, ToI, HitPoint))
+						if (TAABB<T, 3>(Min,Max).RaycastFast(StartPoint, Dir, InvDir, bParallel, CurrentLength, InvCurrentLength, ToI, HitPoint))
 						{
 							Visitor.VisitRaycast(StartCell[1] * (GeomData.NumCols - 1) + StartCell[0], CurrentLength);
 						}
@@ -683,7 +683,7 @@ namespace Chaos
 						if (GetCellBounds3DScaled(StartCell, Min, Max))
 						{
 							// Check cell bounds
-							if (TBox<T, 3>::RaycastFast(Min, Max, StartPoint, Dir, InvDir, bParallel, CurrentLength, InvCurrentLength, ToI, HitPoint))
+							if (TAABB<T, 3>(Min,Max).RaycastFast(StartPoint, Dir, InvDir, bParallel, CurrentLength, InvCurrentLength, ToI, HitPoint))
 							{
 								// Visit the selected cell
 								bool bContinue = Visitor.VisitRaycast(StartCell[1] * (GeomData.NumCols - 1) + StartCell[0], CurrentLength);
@@ -838,7 +838,7 @@ namespace Chaos
 						FBounds2D CellBounds;
 						if(GetCellBounds3DScaled(NeighCoord, Min, Max, HalfExtents3D) && !Seen.Contains(NeighCoord))
 						{
-							if(TBox<T, 3>::RaycastFast(Min, Max, StartPoint, Dir, InvDir, bParallel, CurrentLength, InvCurrentLength, ToI, HitPoint))
+							if(TAABB<T, 3>(Min,Max).RaycastFast(StartPoint, Dir, InvDir, bParallel, CurrentLength, InvCurrentLength, ToI, HitPoint))
 							{
 								Seen.Add(NeighCoord);
 								Queue.Add({NeighCoord, ToI});
@@ -880,7 +880,7 @@ namespace Chaos
 					// Check the current cell, if we hit its 3D bound we can move on to narrow phase
 					const TVector<int32, 2> Coord = CellCoord.Index;
 					if(GetCellBounds3DScaled(Coord, Min, Max, HalfExtents3D) &&
-						TBox<T, 3>::RaycastFast(Min, Max, StartPoint, Dir, InvDir, bParallel, CurrentLength, InvCurrentLength, ToI, HitPoint))
+						TAABB<T, 3>(Min,Max).RaycastFast(StartPoint, Dir, InvDir, bParallel, CurrentLength, InvCurrentLength, ToI, HitPoint))
 					{
 						bool bContinue = Visitor.VisitSweep(CellCoord.Index[1] * (GeomData.NumCols - 1) + CellCoord.Index[0], CurrentLength);
 
@@ -1034,7 +1034,7 @@ namespace Chaos
 			return false;
 		};
 
-		TBox<T, 3> QueryBounds(Point, Point);
+		TAABB<T, 3> QueryBounds(Point, Point);
 		QueryBounds.Thicken(Thickness);
 
 		FBounds2D FlatQueryBounds;
@@ -1084,9 +1084,9 @@ namespace Chaos
 		};
 
 		bool bResult = false;
-		TBox<T, 3> QueryBounds = QueryGeom.BoundingBox();
+		TAABB<T, 3> QueryBounds = QueryGeom.BoundingBox();
 		QueryBounds.Thicken(Thickness);
-		QueryBounds = QueryBounds.TransformedBox(QueryTM);
+		QueryBounds = QueryBounds.TransformedAABB(QueryTM);
 
 		FBounds2D FlatQueryBounds;
 		FlatQueryBounds.Min = TVector<T, 2>(QueryBounds.Min()[0], QueryBounds.Min()[1]);
@@ -1170,7 +1170,7 @@ namespace Chaos
 	{
 		bool bHit = false;
 		THeightfieldSweepVisitor<QueryGeomType, T> SQVisitor(&GeomData, QueryGeom, StartTM, Dir, Thickness, bComputeMTD);
-		const TBox<T, 3> QueryBounds = QueryGeom.BoundingBox();
+		const TAABB<T, 3> QueryBounds = QueryGeom.BoundingBox();
 		const TVector<T, 3> StartPoint = StartTM.TransformPositionNoScale(QueryBounds.Center());
 
 		const TVector<T, 3> Inflation3D = QueryBounds.Extents() * 0.5 + TVector<T, 3>(Thickness);
@@ -1241,7 +1241,7 @@ namespace Chaos
 	{
 		const T SearchDist2 = SearchDist * SearchDist;
 
-		TBox<T, 3> QueryBounds(Position - TVector<T, 3>(SearchDist), Position + TVector<T, 3>(SearchDist));
+		TAABB<T, 3> QueryBounds(Position - TVector<T, 3>(SearchDist), Position + TVector<T, 3>(SearchDist));
 		const FBounds2D FlatBounds(QueryBounds);
 		TArray<TVector<int32, 2>> PotentialIntersections;
 		GetGridIntersections(FlatBounds, PotentialIntersections);
@@ -1357,7 +1357,7 @@ namespace Chaos
 				const TVector<int32, 2> Cell(XIndex, YIndex);
 				TVector<T, 3> Min, Max;
 				CalcCellBounds3D(Cell, Min, Max);
-				//GeomData.CellBounds[YIndex * NumX + XIndex] = TBox<T, 3>(Min, Max);
+				//GeomData.CellBounds[YIndex * NumX + XIndex] = TAABB<T, 3>(Min, Max);
 				GeomData.CellHeights[YIndex * NumX + XIndex] = Max.Z;
 			}
 		}
diff --git a/Engine/Source/Runtime/Experimental/Chaos/Private/Chaos/ImplicitObject.cpp b/Engine/Source/Runtime/Experimental/Chaos/Private/Chaos/ImplicitObject.cpp
--- a/Engine/Source/Runtime/Experimental/Chaos/Private/Chaos/ImplicitObject.cpp
+++ b/Engine/Source/Runtime/Experimental/Chaos/Private/Chaos/ImplicitObject.cpp
@@ -68,10 +68,10 @@ FVec3 FImplicitObject::Normal(const FVec3& x) const
 	return Normal;
 }
 
-const TBox<FReal, 3>& FImplicitObject::BoundingBox() const
+const TAABB<FReal, 3>& FImplicitObject::BoundingBox() const
 {
 	check(false);
-	static const TBox<FReal, 3> Unbounded(FVec3(-FLT_MAX), FVec3(FLT_MAX));
+	static const TAABB<FReal, 3> Unbounded(FVec3(-FLT_MAX), FVec3(FLT_MAX));
 	return Unbounded;
 }
 
@@ -88,7 +88,7 @@ Pair<FVec3, bool> FImplicitObject::FindDeepestIntersection(const FImplicitObject
 	FReal Phi = Thickness;
 	if (HasBoundingBox())
 	{
-		TBox<FReal, 3> ImplicitBox = BoundingBox().TransformedBox(OtherToLocalTransform.Inverse());
+		TAABB<FReal, 3> ImplicitBox = BoundingBox().TransformedAABB(OtherToLocalTransform.Inverse());
 		ImplicitBox.Thicken(Thickness);
 		TArray<int32> PotentialParticles = Particles->FindAllIntersections(ImplicitBox);
 		for (int32 i : PotentialParticles)
@@ -244,7 +244,7 @@ Pair<FVec3, bool> FImplicitObject::FindClosestIntersectionImp(const FVec3& Start
 	return MakePair(ClosestPoint, true);
 }
 
-void FImplicitObject::FindAllIntersectingObjects(TArray<Pair<const FImplicitObject*, FRigidTransform3>>& Out, const TBox<FReal, 3>& LocalBounds) const
+void FImplicitObject::FindAllIntersectingObjects(TArray<Pair<const FImplicitObject*, FRigidTransform3>>& Out, const TAABB<FReal, 3>& LocalBounds) const
 {
 	if (!HasBoundingBox() || LocalBounds.Intersects(BoundingBox()))
 	{
diff --git a/Engine/Source/Runtime/Experimental/Chaos/Private/Chaos/ImplicitObjectUnion.cpp b/Engine/Source/Runtime/Experimental/Chaos/Private/Chaos/ImplicitObjectUnion.cpp
--- a/Engine/Source/Runtime/Experimental/Chaos/Private/Chaos/ImplicitObjectUnion.cpp
+++ b/Engine/Source/Runtime/Experimental/Chaos/Private/Chaos/ImplicitObjectUnion.cpp
@@ -62,7 +62,7 @@ TImplicitObjectUnion<T, d>::~TImplicitObjectUnion()
 }
 
 template<class T, int d>
-void TImplicitObjectUnion<T,d>::FindAllIntersectingObjects(TArray < Pair<const FImplicitObject*, TRigidTransform<T, d>>>& Out, const TBox<T, d>& LocalBounds) const
+void TImplicitObjectUnion<T,d>::FindAllIntersectingObjects(TArray < Pair<const FImplicitObject*, TRigidTransform<T, d>>>& Out, const TAABB<T, d>& LocalBounds) const
 {
 	if (bHierarchyBuilt)
 	{
@@ -84,7 +84,7 @@ void TImplicitObjectUnion<T,d>::FindAllIntersectingObjects(TArray < Pair<const F
 }
 
 template<class T, int d>
-TArray<int32> TImplicitObjectUnion<T,d>::FindAllIntersectingChildren(const TBox<T, d>& LocalBounds) const
+TArray<int32> TImplicitObjectUnion<T,d>::FindAllIntersectingChildren(const TAABB<T, d>& LocalBounds) const
 {
 	TArray<int32> IntersectingChildren;
 	if (bHierarchyBuilt) //todo: make this work when hierarchy is not built
@@ -179,7 +179,9 @@ void TImplicitObjectUnion<T,d>::Serialize(FChaosArchive& Ar)
 {
 	FChaosArchiveScopedMemory ScopedMemory(Ar, GetTypeName(), false);
 	FImplicitObject::SerializeImp(Ar);
-	Ar << MObjects << MLocalBoundingBox << GeomParticles << *Hierarchy << bHierarchyBuilt;
+	Ar << MObjects;
+	TBox<T, d>::SerializeAsAABB(Ar, MLocalBoundingBox);
+	Ar << GeomParticles << *Hierarchy << bHierarchyBuilt;
 }
 
 template<class T, int d>
diff --git a/Engine/Source/Runtime/Experimental/Chaos/Private/Chaos/Levelset.cpp b/Engine/Source/Runtime/Experimental/Chaos/Private/Chaos/Levelset.cpp
--- a/Engine/Source/Runtime/Experimental/Chaos/Private/Chaos/Levelset.cpp
+++ b/Engine/Source/Runtime/Experimental/Chaos/Private/Chaos/Levelset.cpp
@@ -571,7 +571,7 @@ bool TLevelSet<T, d>::ComputeDistancesNearZeroIsocontour(FErrorReporter& ErrorRe
 	const TArray<TVector<int32, 3>>& Elements = Mesh.GetSurfaceElements();
 	if (Elements.Num() > 0)
 	{
-		MOriginalLocalBoundingBox = TBox<T, d>(InParticles.X(Elements[0][0]), InParticles.X(Elements[0][0]));
+		MOriginalLocalBoundingBox = TAABB<T, d>(InParticles.X(Elements[0][0]), InParticles.X(Elements[0][0]));
 	}
 	else
 	{
@@ -583,7 +583,7 @@ bool TLevelSet<T, d>::ComputeDistancesNearZeroIsocontour(FErrorReporter& ErrorRe
 	{
 		const auto& Element = Elements[Index];
 		TPlane<T, d> TrianglePlane(InParticles.X(Element[0]), Normals[Index]);
-		TBox<T, d> TriangleBounds(InParticles.X(Element[0]), InParticles.X(Element[0]));
+		TAABB<T, d> TriangleBounds(InParticles.X(Element[0]), InParticles.X(Element[0]));
 		TriangleBounds.GrowToInclude(InParticles.X(Element[1]));
 		TriangleBounds.GrowToInclude(InParticles.X(Element[2]));
 		MOriginalLocalBoundingBox.GrowToInclude(TriangleBounds); //also save the original bounding box
@@ -1208,7 +1208,7 @@ void TLevelSet<T, d>::ComputeNormals(const TParticles<T, d>& InParticles, const
 	const TArray<TVector<int32, 3>>& Elements = Mesh.GetSurfaceElements();
 	if (Elements.Num() > 0)
 	{
-		MOriginalLocalBoundingBox = TBox<T, d>(InParticles.X(Elements[0][0]), InParticles.X(Elements[0][0]));
+		MOriginalLocalBoundingBox = TAABB<T, d>(InParticles.X(Elements[0][0]), InParticles.X(Elements[0][0]));
 	}
 	else
 	{
@@ -1218,7 +1218,7 @@ void TLevelSet<T, d>::ComputeNormals(const TParticles<T, d>& InParticles, const
 	{
 		const auto& Element = Elements[Index];
 		TPlane<T, d> TrianglePlane(InParticles.X(Element[0]), Normals[Index]);
-		TBox<T, d> TriangleBounds(InParticles.X(Element[0]), InParticles.X(Element[0]));
+		TAABB<T, d> TriangleBounds(InParticles.X(Element[0]), InParticles.X(Element[0]));
 		TriangleBounds.GrowToInclude(InParticles.X(Element[1]));
 		TriangleBounds.GrowToInclude(InParticles.X(Element[2]));
 		MOriginalLocalBoundingBox.GrowToInclude(TriangleBounds); //also save the original bounding box
@@ -1329,7 +1329,7 @@ T TLevelSet<T, d>::PhiWithNormal(const TVector<T, d>& x, TVector<T, d>& Normal)
 	T SizeSquared = (Location - x).SizeSquared();
 	if (SizeSquared)
 	{
-		Normal = MLocalBoundingBox.Normal(x);
+		MLocalBoundingBox.PhiWithNormal(x, Normal);
 	}
 	else
 	{
diff --git a/Engine/Source/Runtime/Experimental/Chaos/Private/Chaos/PBDCollisionConstraintsPGS.cpp b/Engine/Source/Runtime/Experimental/Chaos/Private/Chaos/PBDCollisionConstraintsPGS.cpp
--- a/Engine/Source/Runtime/Experimental/Chaos/Private/Chaos/PBDCollisionConstraintsPGS.cpp
+++ b/Engine/Source/Runtime/Experimental/Chaos/Private/Chaos/PBDCollisionConstraintsPGS.cpp
@@ -96,7 +96,7 @@ void TPBDCollisionConstraintPGS<T, d>::ComputeConstraints(const TPBDRigidParticl
 		const int32 ParticleIndex = InIndices[Index];
 
 		TArray<int32> PotentialIntersections;
-		TBox<T, d> Box1;
+		TAABB<T, d> Box1;
 		if (InParticles.Geometry(ParticleIndex)->HasBoundingBox())
 		{
 			Box1 = Hierarchy.GetWorldSpaceBoundingBox(InParticles, ParticleIndex);
@@ -199,7 +199,7 @@ void TPBDCollisionConstraintPGS<T, d>::UpdateConstraints(const TPBDRigidParticle
 		if (InParticles.Disabled(Body1Index))
 			return;
 		TArray<int32> PotentialIntersections;
-		TBox<T, d> Box1;
+		TAABB<T, d> Box1;
 		if (InParticles.Geometry(Body1Index)->HasBoundingBox())
 		{
 			Box1 = Hierarchy.GetWorldSpaceBoundingBox(InParticles, Body1Index);
@@ -947,7 +947,7 @@ void TPBDCollisionConstraintPGS<T, d>::UpdateLevelsetConstraint(const T_PARTICLE
 	const TRigidTransform<T, d> LocalToWorld2 = GetTransformPGS(InParticles, Constraint.LevelsetIndex);
 	if (InParticles.Geometry(Constraint.LevelsetIndex)->HasBoundingBox())
 	{
-		TBox<T, d> ImplicitBox = InParticles.Geometry(Constraint.LevelsetIndex)->BoundingBox().TransformedBox(LocalToWorld2 * LocalToWorld1.Inverse());
+		TAABB<T, d> ImplicitBox = InParticles.Geometry(Constraint.LevelsetIndex)->BoundingBox().TransformedBox(LocalToWorld2 * LocalToWorld1.Inverse());
 		if (InParticles.CollisionParticles(Constraint.ParticleIndex))
 		{
 			TArray<int32> PotentialParticles = InParticles.CollisionParticles(Constraint.ParticleIndex)->FindAllIntersections(ImplicitBox);
@@ -1234,7 +1234,7 @@ typename TPBDCollisionConstraintPGS<T, d>::FRigidBodyContactConstraint TPBDColli
 template<class T, int d>
 typename TPBDCollisionConstraintPGS<T, d>::FRigidBodyContactConstraint TPBDCollisionConstraintPGS<T, d>::ComputeConstraint(const TPBDRigidParticles<T, d>& InParticles, int32 Body1Index, int32 Body2Index, const T Thickness)
 {
-	if (InParticles.Geometry(Body1Index)->GetType() == TBox<T, d>::GetType() && InParticles.Geometry(Body2Index)->GetType() == TBox<T, d>::GetType())
+	if (InParticles.Geometry(Body1Index)->GetType() == TAABB<T, d>::GetType() && InParticles.Geometry(Body2Index)->GetType() == TAABB<T, d>::GetType())
 	{
 		return ComputeBoxConstraint(InParticles, Body1Index, Body2Index, Thickness);
 	}
@@ -1242,11 +1242,11 @@ typename TPBDCollisionConstraintPGS<T, d>::FRigidBodyContactConstraint TPBDColli
 	{
 		return ComputeSphereConstraint(InParticles, Body1Index, Body2Index, Thickness);
 	}
-	else if (InParticles.Geometry(Body1Index)->GetType() == TBox<T, d>::GetType() && InParticles.Geometry(Body2Index)->GetType() == TPlane<T, d>::GetType())
+	else if (InParticles.Geometry(Body1Index)->GetType() == TAABB<T, d>::GetType() && InParticles.Geometry(Body2Index)->GetType() == TPlane<T, d>::GetType())
 	{
 		return ComputeBoxPlaneConstraint(InParticles, Body1Index, Body2Index, Thickness);
 	}
-	else if (InParticles.Geometry(Body2Index)->GetType() == TPlane<T, d>::GetType() && InParticles.Geometry(Body1Index)->GetType() == TBox<T, d>::GetType())
+	else if (InParticles.Geometry(Body2Index)->GetType() == TPlane<T, d>::GetType() && InParticles.Geometry(Body1Index)->GetType() == TAABB<T, d>::GetType())
 	{
 		return ComputeBoxPlaneConstraint(InParticles, Body2Index, Body1Index, Thickness);
 	}
@@ -1258,11 +1258,11 @@ typename TPBDCollisionConstraintPGS<T, d>::FRigidBodyContactConstraint TPBDColli
 	{
 		return ComputeSpherePlaneConstraint(InParticles, Body2Index, Body1Index, Thickness);
 	}
-	else if (InParticles.Geometry(Body1Index)->GetType() == TSphere<T, d>::GetType() && InParticles.Geometry(Body2Index)->GetType() == TBox<T, d>::GetType())
+	else if (InParticles.Geometry(Body1Index)->GetType() == TSphere<T, d>::GetType() && InParticles.Geometry(Body2Index)->GetType() == TAABB<T, d>::GetType())
 	{
 		return ComputeSphereBoxConstraint(InParticles, Body1Index, Body2Index, Thickness);
 	}
-	else if (InParticles.Geometry(Body2Index)->GetType() == TBox<T, d>::GetType() && InParticles.Geometry(Body1Index)->GetType() == TSphere<T, d>::GetType())
+	else if (InParticles.Geometry(Body2Index)->GetType() == TAABB<T, d>::GetType() && InParticles.Geometry(Body1Index)->GetType() == TSphere<T, d>::GetType())
 	{
 		return ComputeSphereBoxConstraint(InParticles, Body2Index, Body1Index, Thickness);
 	}
@@ -1277,7 +1277,7 @@ template<class T, int d>
 template<class T_PARTICLES>
 void TPBDCollisionConstraintPGS<T, d>::UpdateConstraint(const T_PARTICLES& InParticles, const T Thickness, FRigidBodyContactConstraint& Constraint)
 {
-	if (InParticles.Geometry(Constraint.ParticleIndex)->GetType() == TBox<T, d>::GetType() && InParticles.Geometry(Constraint.LevelsetIndex)->GetType() == TBox<T, d>::GetType())
+	if (InParticles.Geometry(Constraint.ParticleIndex)->GetType() == TAABB<T, d>::GetType() && InParticles.Geometry(Constraint.LevelsetIndex)->GetType() == TAABB<T, d>::GetType())
 	{
 		UpdateBoxConstraint(InParticles, Thickness, Constraint);
 	}
@@ -1285,7 +1285,7 @@ void TPBDCollisionConstraintPGS<T, d>::UpdateConstraint(const T_PARTICLES& InPar
 	{
 		UpdateSphereConstraint(InParticles, Thickness, Constraint);
 	}
-	else if (InParticles.Geometry(Constraint.ParticleIndex)->GetType() == TBox<T, d>::GetType() && InParticles.Geometry(Constraint.LevelsetIndex)->GetType() == TPlane<T, d>::GetType())
+	else if (InParticles.Geometry(Constraint.ParticleIndex)->GetType() == TAABB<T, d>::GetType() && InParticles.Geometry(Constraint.LevelsetIndex)->GetType() == TPlane<T, d>::GetType())
 	{
 		UpdateBoxPlaneConstraint(InParticles, Thickness, Constraint);
 	}
@@ -1293,11 +1293,11 @@ void TPBDCollisionConstraintPGS<T, d>::UpdateConstraint(const T_PARTICLES& InPar
 	{
 		UpdateSpherePlaneConstraint(InParticles, Thickness, Constraint);
 	}
-	else if (InParticles.Geometry(Constraint.ParticleIndex)->GetType() == TSphere<T, d>::GetType() && InParticles.Geometry(Constraint.LevelsetIndex)->GetType() == TBox<T, d>::GetType())
+	else if (InParticles.Geometry(Constraint.ParticleIndex)->GetType() == TSphere<T, d>::GetType() && InParticles.Geometry(Constraint.LevelsetIndex)->GetType() == TAABB<T, d>::GetType())
 	{
 		UpdateSphereBoxConstraint(InParticles, Thickness, Constraint);
 	}
-	else if (InParticles.Geometry(Constraint.ParticleIndex)->GetType() == TPlane<T, d>::GetType() && InParticles.Geometry(Constraint.LevelsetIndex)->GetType() == TBox<T, d>::GetType())
+	else if (InParticles.Geometry(Constraint.ParticleIndex)->GetType() == TPlane<T, d>::GetType() && InParticles.Geometry(Constraint.LevelsetIndex)->GetType() == TAABB<T, d>::GetType())
 	{
 		int32 Tmp = Constraint.ParticleIndex;
 		Constraint.ParticleIndex = Constraint.LevelsetIndex;
@@ -1311,7 +1311,7 @@ void TPBDCollisionConstraintPGS<T, d>::UpdateConstraint(const T_PARTICLES& InPar
 		Constraint.LevelsetIndex = Tmp;
 		UpdateSpherePlaneConstraint(InParticles, Thickness, Constraint);
 	}
-	else if (InParticles.Geometry(Constraint.ParticleIndex)->GetType() == TBox<T, d>::GetType() && InParticles.Geometry(Constraint.LevelsetIndex)->GetType() == TSphere<T, d>::GetType())
+	else if (InParticles.Geometry(Constraint.ParticleIndex)->GetType() == TAABB<T, d>::GetType() && InParticles.Geometry(Constraint.LevelsetIndex)->GetType() == TSphere<T, d>::GetType())
 	{
 		int32 Tmp = Constraint.ParticleIndex;
 		Constraint.ParticleIndex = Constraint.LevelsetIndex;
diff --git a/Engine/Source/Runtime/Experimental/Chaos/Private/Chaos/PBDRigidClustering.cpp b/Engine/Source/Runtime/Experimental/Chaos/Private/Chaos/PBDRigidClustering.cpp
--- a/Engine/Source/Runtime/Experimental/Chaos/Private/Chaos/PBDRigidClustering.cpp
+++ b/Engine/Source/Runtime/Experimental/Chaos/Private/Chaos/PBDRigidClustering.cpp
@@ -641,7 +641,7 @@ namespace Chaos
 			auto ComputeStrainLambda = [&](const TPBDRigidClusteredParticleHandle<T, d>* Cluster, const TArray<uint32>& ParentToChildren) {
 				const TRigidTransform<T, d> WorldToClusterTM = TRigidTransform<T, d>(Cluster->P(), Cluster->Q());
 				const TVector<T, d> ContactLocationClusterLocal = WorldToClusterTM.InverseTransformPosition(ContactHandle->GetContactLocation());
-				TBox<T, d> ContactBox(ContactLocationClusterLocal, ContactLocationClusterLocal);
+				TAABB<T, d> ContactBox(ContactLocationClusterLocal, ContactLocationClusterLocal);
 				ContactBox.Thicken(ClusterDistanceThreshold);
 				if (Cluster->ChildrenSpatial())
 				{
@@ -654,7 +654,7 @@ namespace Chaos
 							const int32 KeyChild = ProxyData->KeyChild;
 							const TRigidTransform<T, d> ProxyToCluster = ProxyData->RelativeToKeyChild * MParticles.ChildToParent(KeyChild);
 							const TVector<T, d> ContactLocationProxyLocal = ProxyToCluster.InverseTransformPosition(ContactLocationClusterLocal);
-							TBox<T, d> ContactBoxProxy(ContactLocationProxyLocal, ContactLocationProxyLocal);
+							TAABB<T, d> ContactBoxProxy(ContactLocationProxyLocal, ContactLocationProxyLocal);
 							ContactBoxProxy.Thicken(ClusterDistanceThreshold);
 							if (MParticles.ChildrenSpatial(Child))
 							{
@@ -1395,7 +1395,7 @@ namespace Chaos
 				ensureMsgf(false, TEXT("Checking usage with no proxy and multiple ojects with levelsets"));
 
 				TImplicitObjectUnion<T, d> UnionObject(MoveTemp(Objects));
-				TBox<T, d> Bounds = UnionObject.BoundingBox();
+				TAABB<T, d> Bounds = UnionObject.BoundingBox();
 				const TVector<T, d> BoundsExtents = Bounds.Extents();
 				if (BoundsExtents.Min() >= MinLevelsetSize) //make sure the object is not too small
 				{
@@ -1711,7 +1711,7 @@ namespace Chaos
 					continue;
 				}
 				TRigidTransform<T, d> TM1 = TRigidTransform<T, d>(MParticles.X(Child1), MParticles.R(Child1));
-				TBox<T, d> Box1 = MParticles.Geometry(Child1)->BoundingBox();
+				TAABB<T, d> Box1 = MParticles.Geometry(Child1)->BoundingBox();
 
 				TArray<TArray<TPair<uint32, uint32>>> Connections;
 				Connections.Init(TArray<TPair<uint32, uint32>>(), Children.Num() - (i + 1));
diff --git a/Engine/Source/Runtime/Experimental/Chaos/Private/Chaos/PBDRigidDynamicSpringConstraints.cpp b/Engine/Source/Runtime/Experimental/Chaos/Private/Chaos/PBDRigidDynamicSpringConstraints.cpp
--- a/Engine/Source/Runtime/Experimental/Chaos/Private/Chaos/PBDRigidDynamicSpringConstraints.cpp
+++ b/Engine/Source/Runtime/Experimental/Chaos/Private/Chaos/PBDRigidDynamicSpringConstraints.cpp
@@ -64,9 +64,9 @@ void TPBDRigidDynamicSpringConstraints<T, d>::UpdatePositionBasedState(const T D
 		if (Static0->Geometry()->HasBoundingBox() && Static1->Geometry()->HasBoundingBox())
 		{
 			// Matrix multiplication is reversed intentionally to be compatible with unreal
-			TBox<T, d> Box1 = Static0->Geometry()->BoundingBox().TransformedBox(Transform1 * Transform2.Inverse());
+			TAABB<T, d> Box1 = Static0->Geometry()->BoundingBox().TransformedAABB(Transform1 * Transform2.Inverse());
 			Box1.Thicken(CreationThreshold);
-			TBox<T, d> Box2 = Static1->Geometry()->BoundingBox();
+			TAABB<T, d> Box2 = Static1->Geometry()->BoundingBox();
 			Box2.Thicken(CreationThreshold);
 			if (!Box1.Intersects(Box2))
 			{
diff --git a/Engine/Source/Runtime/Experimental/Chaos/Private/Chaos/SpatialHash.cpp b/Engine/Source/Runtime/Experimental/Chaos/Private/Chaos/SpatialHash.cpp
--- a/Engine/Source/Runtime/Experimental/Chaos/Private/Chaos/SpatialHash.cpp
+++ b/Engine/Source/Runtime/Experimental/Chaos/Private/Chaos/SpatialHash.cpp
@@ -15,7 +15,7 @@ namespace Chaos
 
 		MCellSize = 2.0 * Radius;
 
-		MBoundingBox = TBox<T, 3>(TVector<T, 3>(0.0), Chaos::TVector<T, 3>(0.0));
+		MBoundingBox = TAABB<T, 3>(TVector<T, 3>(0.0), Chaos::TVector<T, 3>(0.0));
 		for (int32 Idx = 0; Idx < MParticles.Num(); ++Idx)
 		{
 			MBoundingBox.GrowToInclude(MParticles[Idx]);
@@ -63,7 +63,7 @@ namespace Chaos
 		double Time = 0.0;
 		FDurationTimer Timer(Time);
 
-		MBoundingBox = TBox<T, 3>(TVector<T, 3>(0.0), Chaos::TVector<T, 3>(0.0));
+		MBoundingBox = TAABB<T, 3>(TVector<T, 3>(0.0), Chaos::TVector<T, 3>(0.0));
 		for (int32 Idx = 0; Idx < MParticles.Num(); ++Idx)
 		{
 			MBoundingBox.GrowToInclude(MParticles[Idx]);
diff --git a/Engine/Source/Runtime/Experimental/Chaos/Private/Chaos/TriangleMesh.cpp b/Engine/Source/Runtime/Experimental/Chaos/Private/Chaos/TriangleMesh.cpp
--- a/Engine/Source/Runtime/Experimental/Chaos/Private/Chaos/TriangleMesh.cpp
+++ b/Engine/Source/Runtime/Experimental/Chaos/Private/Chaos/TriangleMesh.cpp
@@ -882,7 +882,7 @@ TArray<int32> TTriangleMesh<T>::GetVertexImportanceOrdering(
 
 	// Move the points to the origin to avoid floating point aliasing far away
 	// from the origin.
-	TBox<T, 3> Bbox(Points[0], Points[0]);
+	TAABB<T, 3> Bbox(Points[0], Points[0]);
 	for (int i = 1; i < NumPoints; i++)
 	{
 		Bbox.GrowToInclude(Points[i]);
@@ -892,7 +892,7 @@ TArray<int32> TTriangleMesh<T>::GetVertexImportanceOrdering(
 	TArray<TVector<T, 3>> LocalPoints;
 	LocalPoints.AddUninitialized(NumPoints);
 	LocalPoints[0] = Points[Offset] - Center;
-	TBox<T, 3> LocalBBox(LocalPoints[0], LocalPoints[0]);
+	TAABB<T, 3> LocalBBox(LocalPoints[0], LocalPoints[0]);
 	for (int i = 1; i < NumPoints; i++)
 	{
 		LocalPoints[i] = Points[Offset + i] - Center;
diff --git a/Engine/Source/Runtime/Experimental/Chaos/Private/Chaos/TriangleMeshImplicitObject.cpp b/Engine/Source/Runtime/Experimental/Chaos/Private/Chaos/TriangleMeshImplicitObject.cpp
--- a/Engine/Source/Runtime/Experimental/Chaos/Private/Chaos/TriangleMeshImplicitObject.cpp
+++ b/Engine/Source/Runtime/Experimental/Chaos/Private/Chaos/TriangleMeshImplicitObject.cpp
@@ -231,7 +231,7 @@ bool TTriangleMeshImplicitObject<T>::Raycast(const TVector<T, 3>& StartPoint, co
 template <typename T>
 bool TTriangleMeshImplicitObject<T>::Overlap(const TVector<T, 3>& Point, const T Thickness) const
 {
-	TBox<T, 3> QueryBounds(Point, Point);
+	TAABB<T, 3> QueryBounds(Point, Point);
 	QueryBounds.Thicken(Thickness);
 	const TArray<int32> PotentialIntersections = BVH.FindAllIntersections(QueryBounds);
 
@@ -337,9 +337,9 @@ template <typename QueryGeomType>
 bool TTriangleMeshImplicitObject<T>::OverlapGeomImp(const QueryGeomType& QueryGeom, const TRigidTransform<T, 3>& QueryTM, const T Thickness) const
 {
 	bool bResult = false;
-	TBox<T, 3> QueryBounds = QueryGeom.BoundingBox();
+	TAABB<T, 3> QueryBounds = QueryGeom.BoundingBox();
 	QueryBounds.Thicken(Thickness);
-	QueryBounds = QueryBounds.TransformedBox(QueryTM);
+	QueryBounds = QueryBounds.TransformedAABB(QueryTM);
 	const TArray<int32> PotentialIntersections = BVH.FindAllIntersections(QueryBounds);
 
 	const auto& InnerQueryGeom = GetGeomHelper(QueryGeom);
@@ -553,7 +553,7 @@ bool TTriangleMeshImplicitObject<T>::SweepGeomImp(const QueryGeomType& QueryGeom
 	TTriangleMeshSweepVisitor<QueryGeomType, T> SQVisitor(*this, QueryGeom, StartTM, Dir, ScaledDirNormalized, LengthScale, ScaledStartTM, Thickness, bComputeMTD);
 
 
-	const TBox<T, 3> QueryBounds = QueryGeom.BoundingBox();
+	const TAABB<T, 3> QueryBounds = QueryGeom.BoundingBox();
 	const TVector<T, 3> StartPoint = StartTM.TransformPositionNoScale(QueryBounds.Center());
 	const TVector<T, 3> Inflation = QueryBounds.Extents() * 0.5 + TVector<T, 3>(Thickness);
 	BVH.template Sweep<TTriangleMeshSweepVisitor<QueryGeomType, T>>(StartPoint, Dir, Length, Inflation, SQVisitor);
@@ -623,7 +623,7 @@ int32 TTriangleMeshImplicitObject<T>::FindMostOpposingFace(const TVector<T, 3>&
 	//todo: this is horribly slow, need adjacency information
 	const T SearchDist2 = SearchDist * SearchDist;
 	
-	TBox<T, 3> QueryBounds(Position - TVector<T,3>(SearchDist), Position + TVector<T,3>(SearchDist));
+	TAABB<T, 3> QueryBounds(Position - TVector<T,3>(SearchDist), Position + TVector<T,3>(SearchDist));
 
 	const TArray<int32> PotentialIntersections = BVH.FindAllIntersections(QueryBounds);
 	const T Epsilon = 1e-4;
diff --git a/Engine/Source/Runtime/Experimental/Chaos/Public/Chaos/AABB.h b/Engine/Source/Runtime/Experimental/Chaos/Public/Chaos/AABB.h
--- a/Engine/Source/Runtime/Experimental/Chaos/Public/Chaos/AABB.h
+++ b/Engine/Source/Runtime/Experimental/Chaos/Public/Chaos/AABB.h
@@ -27,6 +27,9 @@ namespace Chaos
 	class TAABB
 	{
 	public:
+		using TType = T;
+		static constexpr int D = d;
+
 		FORCEINLINE TAABB()
 			: MMin()
 			, MMax()
@@ -140,6 +143,10 @@ namespace Chaos
 			return true;
 		}
 
+		FORCEINLINE const TAABB<T, d>& BoundingBox() const { return *this; }
+
+		FORCEINLINE uint16 GetMaterialIndex(uint32 HintIndex) const { return 0; }
+
 		FORCEINLINE T SignedDistance(const TVector<T, d>& x) const
 		{
 			TVector<T, d> Normal;
@@ -197,7 +204,7 @@ namespace Chaos
 		bool CHAOS_API Raycast(const TVector<T, d>& StartPoint, const TVector<T, d>& Dir, const T Length, const T Thickness, T& OutTime, TVector<T, d>& OutPosition, TVector<T, d>& OutNormal, int32& OutFaceIndex) const;
 
 
-		FORCEINLINE bool RaycastFast(const TVector<T, d>& StartPoint, const TVector<T, d>& Dir, const TVector<T, d>& InvDir, const bool* bParallel, const T Length, const T InvLength, T& OutTime, TVector<T, d>& OutPosition)
+		FORCEINLINE bool RaycastFast(const TVector<T, d>& StartPoint, const TVector<T, d>& Dir, const TVector<T, d>& InvDir, const bool* bParallel, const T Length, const T InvLength, T& OutTime, TVector<T, d>& OutPosition) const
 		{
 			const TVector<T, d> StartToMin = MMin - StartPoint;
 			const TVector<T, d> StartToMax = MMax - StartPoint;
@@ -467,6 +474,8 @@ namespace Chaos
 		FORCEINLINE T GetVolume() const { return GetVolume(Extents()); }
 		FORCEINLINE static T GetVolume(const TVector<T, 3>& Dim) { return Dim.Product(); }
 
+		FORCEINLINE T GetMargin() const { return 0; }
+
 		FORCEINLINE static TAABB<T, d> EmptyAABB() { return TAABB<T, d>(TVector<T, d>(TNumericLimits<T>::Max()), TVector<T, d>(-TNumericLimits<T>::Max())); }
 		FORCEINLINE static TAABB<T, d> ZeroAABB() { return TAABB<T, d>(TVector<T, d>((T)0), TVector<T, d>((T)0)); }
 
@@ -480,16 +489,41 @@ namespace Chaos
 			return HashCombine(::GetTypeHash(MMin), ::GetTypeHash(MMax));
 		}
 
+		FORCEINLINE PMatrix<T, d, d> GetInertiaTensor(const T Mass) const { return GetInertiaTensor(Mass, Extents()); }
+		FORCEINLINE static PMatrix<T, 3, 3> GetInertiaTensor(const T Mass, const TVector<T, 3>& Dim)
+		{
+			// https://www.wolframalpha.com/input/?i=cuboid
+			const T M = Mass / 12;
+			const T WW = Dim[0] * Dim[0];
+			const T HH = Dim[1] * Dim[1];
+			const T DD = Dim[2] * Dim[2];
+			return PMatrix<T, 3, 3>(M * (HH + DD), M * (WW + DD), M * (WW + HH));
+		}
+
+		FORCEINLINE static TRotation<T, d> GetRotationOfMass()
+		{
+			return TRotation<T, d>::FromElements(TVector<T, d>(0), 1);
+		}
+
+		FORCEINLINE constexpr bool IsConvex() const { return true; }
+
 	private:
 		TVector<T, d> MMin, MMax;
 	};
 
+
 	FORCEINLINE FChaosArchive& operator<<(FChaosArchive& Ar, TAABB<FReal, 3>& AABB)
 	{
 		AABB.Serialize(Ar);
 		return Ar;
 	}
 
+	FORCEINLINE FArchive& operator<<(FArchive& Ar, TAABB<FReal, 3>& AABB)
+	{
+		AABB.Serialize(Ar);
+		return Ar;
+	}
+
 	template<typename T>
 	struct TAABBSpecializeSamplingHelper<T, 2>
 	{
diff --git a/Engine/Source/Runtime/Experimental/Chaos/Public/Chaos/AABBTree.h b/Engine/Source/Runtime/Experimental/Chaos/Public/Chaos/AABBTree.h
--- a/Engine/Source/Runtime/Experimental/Chaos/Public/Chaos/AABBTree.h
+++ b/Engine/Source/Runtime/Experimental/Chaos/Public/Chaos/AABBTree.h
@@ -1,7 +1,7 @@
 // Copyright 1998-2019 Epic Games, Inc. All Rights Reserved.
 #pragma once
 
-#include "Chaos/Box.h"
+#include "Chaos/AABB.h"
 #include "Chaos/Defines.h"
 #include "Chaos/GeometryParticles.h"
 #include "Chaos/Transform.h"
@@ -30,7 +30,7 @@ template <typename T, typename TQueryFastData, EAABBQueryType Query>
 struct TAABBTreeIntersectionHelper
 {
 	static bool Intersects(const TVector<T, 3>& Start, TQueryFastData& QueryFastData, T& TOI, TVector<T, 3>& OutPosition,
-		const TBox<T, 3>& Bounds, const TBox<T, 3>& QueryBounds, const TVector<T, 3>& QueryHalfExtents)
+		const TAABB<T, 3>& Bounds, const TAABB<T, 3>& QueryBounds, const TVector<T, 3>& QueryHalfExtents)
 	{
 		check(false);
 		return true;
@@ -41,9 +41,9 @@ template<>
 struct TAABBTreeIntersectionHelper<FReal, FQueryFastData, EAABBQueryType::Raycast>
 {
 	FORCEINLINE_DEBUGGABLE static bool Intersects(const FVec3& Start, FQueryFastData& QueryFastData, FReal& TOI, FVec3& OutPosition,
-		const TBox<FReal, 3>& Bounds, const TBox<FReal, 3>& QueryBounds, const TVector<FReal, 3>& QueryHalfExtents)
+		const TAABB<FReal, 3>& Bounds, const TAABB<FReal, 3>& QueryBounds, const TVector<FReal, 3>& QueryHalfExtents)
 	{
-		return TBox<FReal, 3>::RaycastFast(Bounds.Min(), Bounds.Max(), Start, QueryFastData.Dir, QueryFastData.InvDir, QueryFastData.bParallel, QueryFastData.CurrentLength, QueryFastData.InvCurrentLength, TOI, OutPosition);
+		return Bounds.RaycastFast(Start, QueryFastData.Dir, QueryFastData.InvDir, QueryFastData.bParallel, QueryFastData.CurrentLength, QueryFastData.InvCurrentLength, TOI, OutPosition);
 	}
 };
 
@@ -51,9 +51,10 @@ template <typename T>
 struct TAABBTreeIntersectionHelper<T, FQueryFastData, EAABBQueryType::Sweep>
 {
 	FORCEINLINE_DEBUGGABLE static bool Intersects(const TVector<T, 3>& Start, FQueryFastData& QueryFastData,
-		T& TOI, TVector<T, 3>& OutPosition, const TBox<T, 3>& Bounds, const TBox<T, 3>& QueryBounds, const TVector<T, 3>& QueryHalfExtents)
+		T& TOI, TVector<T, 3>& OutPosition, const TAABB<T, 3>& Bounds, const TAABB<T, 3>& QueryBounds, const TVector<T, 3>& QueryHalfExtents)
 	{
-		return TBox<T, 3>::RaycastFast(Bounds.Min() - QueryHalfExtents, Bounds.Max() + QueryHalfExtents, Start, QueryFastData.Dir, QueryFastData.InvDir, QueryFastData.bParallel, QueryFastData.CurrentLength, QueryFastData.InvCurrentLength, TOI, OutPosition);
+		TAABB<T, 3> SweepBounds(Bounds.Min() - QueryHalfExtents, Bounds.Max() + QueryHalfExtents);
+		return SweepBounds.RaycastFast(Start, QueryFastData.Dir, QueryFastData.InvDir, QueryFastData.bParallel, QueryFastData.CurrentLength, QueryFastData.InvCurrentLength, TOI, OutPosition);
 	}
 };
 
@@ -61,7 +62,7 @@ template <typename T>
 struct TAABBTreeIntersectionHelper<T, FQueryFastDataVoid, EAABBQueryType::Overlap>
 {
 	FORCEINLINE_DEBUGGABLE static bool Intersects(const TVector<T, 3>& Start, FQueryFastDataVoid& QueryFastData, T& TOI, TVector<T, 3>& OutPosition,
-		const TBox<T, 3>& Bounds, const TBox<T, 3>& QueryBounds, const TVector<T, 3>& QueryHalfExtents)
+		const TAABB<T, 3>& Bounds, const TAABB<T, 3>& QueryBounds, const TVector<T, 3>& QueryHalfExtents)
 	{
 		return QueryBounds.Intersects(Bounds);
 	}
@@ -90,7 +91,7 @@ struct TAABBTreeLeafArray
 	}
 
 	template <typename TSQVisitor>
-	bool OverlapFast(const TBox<T, 3>& QueryBounds, TSQVisitor& Visitor) const
+	bool OverlapFast(const TAABB<T, 3>& QueryBounds, TSQVisitor& Visitor) const
 	{
 		for (const auto& Elem : Elems)
 		{
@@ -115,7 +116,7 @@ struct TAABBTreeLeafArray
 		for (const auto& Elem : Elems)
 		{
 			const auto& InstanceBounds = Elem.Bounds;
-			if (TAABBTreeIntersectionHelper<T, TQueryFastData, bSweep ? EAABBQueryType::Sweep : EAABBQueryType::Raycast>::Intersects(Start, QueryFastData, TOI, TmpPosition, InstanceBounds, TBox<T,3>(), QueryHalfExtents))
+			if (TAABBTreeIntersectionHelper<T, TQueryFastData, bSweep ? EAABBQueryType::Sweep : EAABBQueryType::Raycast>::Intersects(Start, QueryFastData, TOI, TmpPosition, InstanceBounds, TAABB<T,3>(), QueryHalfExtents))
 			{
 				TSpatialVisitorData<TPayloadType> VisitData(Elem.Payload, true, InstanceBounds);
 				const bool bContinue = (bSweep && Visitor.VisitSweep(VisitData, QueryFastData)) || (!bSweep &&  Visitor.VisitRaycast(VisitData, QueryFastData));
@@ -161,10 +162,10 @@ struct TAABBTreeNode
 {
 	TAABBTreeNode()
 	{
-		ChildrenBounds[0] = TBox<T, 3>();
-		ChildrenBounds[1] = TBox<T, 3>();
+		ChildrenBounds[0] = TAABB<T, 3>();
+		ChildrenBounds[1] = TAABB<T, 3>();
 	}
-	TBox<T, 3> ChildrenBounds[2];
+	TAABB<T, 3> ChildrenBounds[2];
 	int32 ChildrenNodes[2];
 	bool bLeaf;
 
@@ -172,7 +173,7 @@ struct TAABBTreeNode
 	{
 		for (auto& Bounds : ChildrenBounds)
 		{
-			Ar << Bounds;
+			TBox<T, 3>::SerializeAsAABB(Ar, Bounds);
 		}
 
 		for (auto& Node : ChildrenNodes)
@@ -292,9 +293,9 @@ public:
 		GenerateTree(Particles);
 	}
 
-	virtual TArray<TPayloadType> FindAllIntersections(const TBox<T, 3>& Box) const override { return FindAllIntersectionsImp(Box); }
+	virtual TArray<TPayloadType> FindAllIntersections(const TAABB<T, 3>& Box) const override { return FindAllIntersectionsImp(Box); }
 
-	bool GetAsBoundsArray(TArray<TBox<T, 3>>& AllBounds, int32 NodeIdx, int32 ParentNode, TBox<T, 3>& Bounds)
+	bool GetAsBoundsArray(TArray<TAABB<T, 3>>& AllBounds, int32 NodeIdx, int32 ParentNode, TAABB<T, 3>& Bounds)
 	{
 		if (Nodes[NodeIdx].bLeaf)
 		{
@@ -326,13 +327,13 @@ public:
 	void Raycast(const TVector<T, 3>& Start, const TVector<T, 3>& Dir, const T Length, SQVisitor& Visitor) const
 	{
 		FQueryFastData QueryFastData(Dir, Length);
-		QueryImp<EAABBQueryType::Raycast>(Start, QueryFastData, TVector<T,3>(), TBox<T,3>(), Visitor);
+		QueryImp<EAABBQueryType::Raycast>(Start, QueryFastData, TVector<T,3>(), TAABB<T,3>(), Visitor);
 	}
 
 	template <typename SQVisitor>
 	bool RaycastFast(const TVector<T, 3>& Start, FQueryFastData& CurData, SQVisitor& Visitor) const
 	{
-		return QueryImp<EAABBQueryType::Raycast>(Start, CurData, TVector<T, 3>(), TBox<T, 3>(), Visitor);
+		return QueryImp<EAABBQueryType::Raycast>(Start, CurData, TVector<T, 3>(), TAABB<T, 3>(), Visitor);
 	}
 
 	void Sweep(const TVector<T, 3>& Start, const TVector<T, 3>& Dir, const T Length, const TVector<T, 3> QueryHalfExtents, ISpatialVisitor<TPayloadType, T>& Visitor) const override
@@ -345,29 +346,29 @@ public:
 	void Sweep(const TVector<T, 3>& Start, const TVector<T, 3>& Dir, const T Length, const TVector<T, 3> QueryHalfExtents, SQVisitor& Visitor) const
 	{
 		FQueryFastData QueryFastData(Dir, Length);
-		QueryImp<EAABBQueryType::Sweep>(Start, QueryFastData, QueryHalfExtents, TBox<T, 3>(), Visitor);
+		QueryImp<EAABBQueryType::Sweep>(Start, QueryFastData, QueryHalfExtents, TAABB<T, 3>(), Visitor);
 	}
 
 	template <typename SQVisitor>
 	bool SweepFast(const TVector<T, 3>& Start, FQueryFastData& CurData, const TVector<T, 3> QueryHalfExtents, SQVisitor& Visitor) const
 	{
-		return QueryImp<EAABBQueryType::Sweep>(Start,CurData, QueryHalfExtents, TBox<T, 3>(), Visitor);
+		return QueryImp<EAABBQueryType::Sweep>(Start,CurData, QueryHalfExtents, TAABB<T, 3>(), Visitor);
 	}
 
-	void Overlap(const TBox<T, 3>& QueryBounds, ISpatialVisitor<TPayloadType, T>& Visitor) const override
+	void Overlap(const TAABB<T, 3>& QueryBounds, ISpatialVisitor<TPayloadType, T>& Visitor) const override
 	{
 		TSpatialVisitor<TPayloadType, T> ProxyVisitor(Visitor);
 		Overlap(QueryBounds, ProxyVisitor);
 	}
 
 	template <typename SQVisitor>
-	void Overlap(const TBox<T,3>& QueryBounds, SQVisitor& Visitor) const
+	void Overlap(const TAABB<T,3>& QueryBounds, SQVisitor& Visitor) const
 	{
 		OverlapFast(QueryBounds, Visitor);
 	}
 
 	template <typename SQVisitor>
-	bool OverlapFast(const TBox<T, 3>& QueryBounds, SQVisitor& Visitor) const
+	bool OverlapFast(const TAABB<T, 3>& QueryBounds, SQVisitor& Visitor) const
 	{
 		//dummy variables to reuse templated path
 		FQueryFastDataVoid VoidData;
@@ -407,7 +408,7 @@ public:
 		}
 	}
 
-	virtual void UpdateElement(const TPayloadType& Payload, const TBox<T, 3>& NewBounds, bool bHasBounds) override
+	virtual void UpdateElement(const TPayloadType& Payload, const TAABB<T, 3>& NewBounds, bool bHasBounds) override
 	{
 		FAABBTreePayloadInfo* PayloadInfo = PayloadToInfo.Find(Payload);
 		if (PayloadInfo)
@@ -459,7 +460,7 @@ public:
 		}
 		else
 		{
-			TBox<T, 3> GlobalBounds = bTooBig ? NewBounds : TBox<T, 3>(TVector<T, 3>(TNumericLimits<T>::Lowest()), TVector<T, 3>(TNumericLimits<T>::Max()));
+			TAABB<T, 3> GlobalBounds = bTooBig ? NewBounds : TAABB<T, 3>(TVector<T, 3>(TNumericLimits<T>::Lowest()), TVector<T, 3>(TNumericLimits<T>::Max()));
 			if (PayloadInfo->GlobalPayloadIdx == INDEX_NONE)
 			{
 				PayloadInfo->GlobalPayloadIdx = GlobalPayloads.Add(FElement{ Payload, GlobalBounds });
@@ -491,7 +492,7 @@ public:
 
 	virtual void Serialize(FChaosArchive& Ar) override
 	{
-		Ar << FullBounds;
+		TBox<FReal, 3>::SerializeAsAABB(Ar, FullBounds);
 		Ar << Nodes;
 		Ar << Leaves;
 		Ar << DirtyElements;
@@ -521,7 +522,7 @@ private:
 	using FNode = TAABBTreeNode<T>;
 
 	template <EAABBQueryType Query, typename TQueryFastData, typename SQVisitor>
-	bool QueryImp(const TVector<T, 3>& Start, TQueryFastData& CurData, const TVector<T, 3> QueryHalfExtents, const TBox<T,3>& QueryBounds, SQVisitor& Visitor) const
+	bool QueryImp(const TVector<T, 3>& Start, TQueryFastData& CurData, const TVector<T, 3> QueryHalfExtents, const TAABB<T,3>& QueryBounds, SQVisitor& Visitor) const
 	{
 		TVector<T, 3> TmpPosition;
 		T TOI = 0;
@@ -618,7 +619,7 @@ private:
 			else
 			{
 				int32 Idx = 0;
-				for (const TBox<T, 3>& AABB : Node.ChildrenBounds)
+				for (const TAABB<T, 3>& AABB : Node.ChildrenBounds)
 				{
 					if(TAABBTreeIntersectionHelper<T, TQueryFastData, Query>::Intersects(Start, CurData, TOI, TmpPosition, AABB, QueryBounds, QueryHalfExtents))
 					{
@@ -648,7 +649,7 @@ private:
 		PayloadToInfo.Reset();
 		NumProcessedThisSlice = 0;
 
-		FullBounds = TBox<T, 3>::EmptyBox();
+		FullBounds = TAABB<T, 3>::EmptyAABB();
 
 		{
 			SCOPE_CYCLE_COUNTER(STAT_AABBTreeTimeSliceSetup);
@@ -659,7 +660,7 @@ private:
 			{
 				bool bHasBoundingBox = HasBoundingBox(Particle);
 				auto Payload = Particle.template GetPayload<TPayloadType>(Idx);
-				TBox<T, 3> ElemBounds = ComputeWorldSpaceBoundingBox(Particle, false, (T)0);
+				TAABB<T, 3> ElemBounds = ComputeWorldSpaceBoundingBox(Particle, false, (T)0);
 
 				if (bHasBoundingBox)
 				{
@@ -675,7 +676,7 @@ private:
 				}
 				else
 				{
-					ElemBounds = TBox<T, 3>(TVector<T, 3>(TNumericLimits<T>::Lowest()), TVector<T, 3>(TNumericLimits<T>::Max()));
+					ElemBounds = TAABB<T, 3>(TVector<T, 3>(TNumericLimits<T>::Lowest()), TVector<T, 3>(TNumericLimits<T>::Max()));
 				}
 
 				if (!bHasBoundingBox)
@@ -703,7 +704,7 @@ private:
 
 	}
 
-	void SplitNode(TBox<T, 3>& SplitBounds, TBox<T, 3>& RealBounds, TArray<FElement>& Children, T& SplitBoundsSize2, int NewNodeIdx, int NodeLevel, int BoxIdx)
+	void SplitNode(TAABB<T, 3>& SplitBounds, TAABB<T, 3>& RealBounds, TArray<FElement>& Children, T& SplitBoundsSize2, int NewNodeIdx, int NodeLevel, int BoxIdx)
 	{
 		//SCOPE_CYCLE_COUNTER(STAT_AABBTreeChildrenPhase);
 
@@ -723,12 +724,12 @@ private:
 
 	}
 
-	void SplitNode(const TBox<T, 3>& Bounds, const TArray<FElement>& Elems, int32 NodeLevel, int32 NewNodeIdx, bool Initial=false)
+	void SplitNode(const TAABB<T, 3>& Bounds, const TArray<FElement>& Elems, int32 NodeLevel, int32 NewNodeIdx, bool Initial=false)
 	{
 		// create the actual node space but might no be filled in (YET) due to time slicing exit
 		if (NewNodeIdx >= Nodes.Num())
 		{
-		 	Nodes.AddDefaulted();	//todo: remove TBox
+		 	Nodes.AddDefaulted();	//todo: remove TAABB
 		}
 
 		bool WeAreTimeslicing = (MaxNumToProcess > 0);
@@ -774,13 +775,13 @@ private:
 		bool ChildrenInBothHalves = false;
 
 		FSplitInfo SplitInfos[2];
-		SplitInfos[0].SplitBounds = TBox<T, 3>(Bounds.Min(), Bounds.Min());
-		SplitInfos[1].SplitBounds = TBox<T, 3>(Bounds.Max(), Bounds.Max());
+		SplitInfos[0].SplitBounds = TAABB<T, 3>(Bounds.Min(), Bounds.Min());
+		SplitInfos[1].SplitBounds = TAABB<T, 3>(Bounds.Max(), Bounds.Max());
 
 		const TVector<T, 3> Center = Bounds.Center();
 		for (FSplitInfo& SplitInfo : SplitInfos)
 		{
-			SplitInfo.RealBounds = TBox<T, 3>::EmptyBox();
+			SplitInfo.RealBounds = TAABB<T, 3>::EmptyAABB();
 
 			for (int32 Axis = 0; Axis < 3; ++Axis)
 			{
@@ -809,7 +810,7 @@ private:
 				int32 BoxIdx = 0;
 				for (const FSplitInfo& SplitInfo : SplitInfos)
 				{
-					TBox<T, 3> NewBox = SplitInfo.SplitBounds;
+					TAABB<T, 3> NewBox = SplitInfo.SplitBounds;
 					NewBox.GrowToInclude(Elem.Bounds);
 					const T Delta2 = NewBox.Extents().SizeSquared() - SplitInfo.SplitBoundsSize2;
 					if (Delta2 < MinDelta2)
@@ -881,7 +882,7 @@ private:
 		return; // keep working
 	}
 
-	TArray<TPayloadType> FindAllIntersectionsImp(const TBox<T, 3>& Intersection) const
+	TArray<TPayloadType> FindAllIntersectionsImp(const TAABB<T, 3>& Intersection) const
 	{
 		struct FSimpleVisitor
 		{
@@ -931,13 +932,13 @@ private:
 
 	struct FSplitInfo
 	{
-		TBox<T, 3> SplitBounds;	//Even split of parent bounds
-		TBox<T, 3> RealBounds;	//Actual bounds as children are added
+		TAABB<T, 3> SplitBounds;	//Even split of parent bounds
+		TAABB<T, 3> RealBounds;	//Actual bounds as children are added
 		TArray<FElement> Children;
 		T SplitBoundsSize2;
 	};
 
-	TBox<T, 3> FullBounds;
+	TAABB<T, 3> FullBounds;
 	TArray<FNode> Nodes;
 	TArray<TLeafType> Leaves;
 	TArray<FElement> DirtyElements;
@@ -959,14 +960,14 @@ private:
 		eTimeSlicePhase TimeslicePhase;
 
 		// used for phase 1 & 2
-		TBox<T, 3> Bounds;
+		TAABB<T, 3> Bounds;
 		TArray<FElement> Elems;
 
 		int32 NodeLevel;
 		int32 NewNodeIdx;
 	
 		// phase 2 data only
-		TBox<T, 3> RealBounds;
+		TAABB<T, 3> RealBounds;
 		int8 BoxIdx;
 		T SplitBoundsSize2;
 	};
diff --git a/Engine/Source/Runtime/Experimental/Chaos/Public/Chaos/BVHParticles.h b/Engine/Source/Runtime/Experimental/Chaos/Public/Chaos/BVHParticles.h
--- a/Engine/Source/Runtime/Experimental/Chaos/Public/Chaos/BVHParticles.h
+++ b/Engine/Source/Runtime/Experimental/Chaos/Public/Chaos/BVHParticles.h
@@ -18,7 +18,7 @@ namespace Chaos
     class TBoundingVolumeHierarchy;
 
     template<class T, int d>
-    class TBox;
+    class TAABB;
 
 	template<class T, int d>
 	class TBVHParticles final /*Note: removing this final has implications for serialization. See TImplicitObject*/ : public TParticles<T, d>
@@ -42,7 +42,7 @@ namespace Chaos
 		}
 
 		CHAOS_API void UpdateAccelerationStructures();
-		const TArray<int32> FindAllIntersections(const TBox<T, d>& Object) const;
+		const TArray<int32> FindAllIntersections(const TAABB<T, d>& Object) const;
 
 		static TBVHParticles<T,d>* SerializationFactory(FChaosArchive& Ar, TBVHParticles<T,d>* BVHParticles)
 		{
diff --git a/Engine/Source/Runtime/Experimental/Chaos/Public/Chaos/BoundingVolume.h b/Engine/Source/Runtime/Experimental/Chaos/Public/Chaos/BoundingVolume.h
--- a/Engine/Source/Runtime/Experimental/Chaos/Public/Chaos/BoundingVolume.h
+++ b/Engine/Source/Runtime/Experimental/Chaos/Public/Chaos/BoundingVolume.h
@@ -65,14 +65,14 @@ namespace Chaos
 template<typename TPayloadType, typename T, int d>
 struct TBVCellElement
 {
-	TBox<T, d> Bounds;
+	TAABB<T, d> Bounds;
 	TPayloadType Payload;
 	TVector<int32, 3> StartIdx;
 	TVector<int32, 3> EndIdx;
 
 	void Serialize(FChaosArchive& Ar)
 	{
-		Ar << Bounds;
+		TBox<T, d>::SerializeAsAABB(Ar, Bounds);
 		Ar << Payload;
 		Ar << StartIdx;
 		Ar << EndIdx;
@@ -186,7 +186,7 @@ public:
 		GenerateTree(Particles, bUseVelocity, Dt, MaxCells);
 	}
 
-	TArray<TPayloadType> FindAllIntersectionsImp(const TBox<T,d>& Intersection) const
+	TArray<TPayloadType> FindAllIntersectionsImp(const TAABB<T,d>& Intersection) const
 	{
 		struct FSimpleVisitor
 		{
@@ -230,7 +230,7 @@ public:
 		}
 	}
 
-	virtual void UpdateElement(const TPayloadType& Payload, const TBox<T,d>& NewBounds, bool bHasBounds) override
+	virtual void UpdateElement(const TPayloadType& Payload, const TAABB<T,d>& NewBounds, bool bHasBounds) override
 	{
 		SCOPE_CYCLE_COUNTER(STAT_BoundingVolumeUpdateElement);
 		if (FPayloadInfo* PayloadInfo = MPayloadInfo.Find(Payload))
@@ -250,7 +250,7 @@ public:
 	}
 
 	// Begin ISpatialAcceleration interface
-	virtual TArray<TPayloadType> FindAllIntersections(const TBox<T, d>& Box) const override { return FindAllIntersectionsImp(Box); }
+	virtual TArray<TPayloadType> FindAllIntersections(const TAABB<T, d>& Box) const override { return FindAllIntersectionsImp(Box); }
 
 	const TArray<TPayloadBoundsElement<TPayloadType, T>>& GlobalObjects() const
 	{
@@ -295,20 +295,20 @@ public:
 		SweepImp<SQVisitor, bPruneDuplicates>(Start, CurData, QueryHalfExtents, Visitor);
 	}
 
-	virtual void Overlap(const TBox<T, d>& QueryBounds, ISpatialVisitor<TPayloadType, T>& Visitor) const override
+	virtual void Overlap(const TAABB<T, d>& QueryBounds, ISpatialVisitor<TPayloadType, T>& Visitor) const override
 	{
 		TSpatialVisitor<TPayloadType, T> ProxyVisitor(Visitor);
 		Overlap(QueryBounds, ProxyVisitor);
 	}
 
 	template <typename SQVisitor>
-	bool OverlapFast(const TBox<T, d>& QueryBounds, SQVisitor& Visitor) const
+	bool OverlapFast(const TAABB<T, d>& QueryBounds, SQVisitor& Visitor) const
 	{
 		return OverlapImp(QueryBounds, Visitor);
 	}
 
 	template <typename SQVisitor, bool bPruneDuplicates = true>
-	void Overlap(const TBox<T, d>& QueryBounds, SQVisitor& Visitor) const
+	void Overlap(const TAABB<T, d>& QueryBounds, SQVisitor& Visitor) const
 	{
 		OverlapImp<SQVisitor, bPruneDuplicates>(QueryBounds, Visitor);
 	}
@@ -323,7 +323,7 @@ public:
 			MGlobalPayloads.Reserve(TmpPayloads.Num());
 			for (auto& Payload : TmpPayloads)
 			{
-				MGlobalPayloads.Add({ Payload, TBox<T,d>(TVector<T,d>(TNumericLimits<T>::Lowest()), TVector<T,d>(TNumericLimits<T>::Max())) });
+				MGlobalPayloads.Add({ Payload, TAABB<T,d>(TVector<T,d>(TNumericLimits<T>::Lowest()), TVector<T,d>(TNumericLimits<T>::Max())) });
 			}
 			MaxPayloadBounds = DefaultMaxPayloadBounds;
 		}
@@ -367,7 +367,7 @@ private:
 		for (const auto& Elem : MGlobalPayloads)
 		{
 			const auto& InstanceBounds = Elem.Bounds;
-			if (TBox<T, d>::RaycastFast(InstanceBounds.Min(), InstanceBounds.Max(), Start,
+			if (InstanceBounds.RaycastFast(Start,
 				CurData.Dir, CurData.InvDir, CurData.bParallel, CurData.CurrentLength, CurData.InvCurrentLength, TOI, TmpPosition))
 			{
 				TSpatialVisitorData<TPayloadType> VisitData(Elem.Payload, true, InstanceBounds);
@@ -383,7 +383,7 @@ private:
 		for (const auto& Elem : MDirtyElements)
 		{
 			const auto& InstanceBounds = Elem.Bounds;
-			if (TBox<T, d>::RaycastFast(InstanceBounds.Min(), InstanceBounds.Max(), Start, CurData.Dir,
+			if (InstanceBounds.RaycastFast(Start, CurData.Dir,
 				CurData.InvDir, CurData.bParallel, CurData.CurrentLength, CurData.InvCurrentLength, TOI, TmpPosition))
 			{
 				TSpatialVisitorData<TPayloadType> VisitData(Elem.Payload, true, InstanceBounds);
@@ -400,10 +400,10 @@ private:
 			return true;
 		}
 
-		TBox<T, d> GlobalBounds(MGrid.MinCorner(), MGrid.MaxCorner());
+		TAABB<T, d> GlobalBounds(MGrid.MinCorner(), MGrid.MaxCorner());
 		TVector<T, d> NextStart;
 		TVector<int32, d> CellIdx;
-		bool bCellsLeft = MElements.Num() && TBox<T, d>::RaycastFast(GlobalBounds.Min(), GlobalBounds.Max(), Start, CurData.Dir,
+		bool bCellsLeft = MElements.Num() && GlobalBounds.RaycastFast(Start, CurData.Dir,
 			CurData.InvDir, CurData.bParallel, CurData.CurrentLength, CurData.InvCurrentLength, TOI, NextStart);
 		if (bCellsLeft)
 		{
@@ -446,7 +446,7 @@ private:
 						}
 					}
 					const auto& InstanceBounds = Elem.Bounds;
-					if (TBox<T, d>::RaycastFast(InstanceBounds.Min(), InstanceBounds.Max(), Start,
+					if (InstanceBounds.RaycastFast(Start,
 						CurData.Dir, CurData.InvDir, CurData.bParallel, CurData.CurrentLength, CurData.InvCurrentLength, TOI, TmpPosition))
 					{
 						TSpatialVisitorData<TPayloadType> VisitData(Elem.Payload, true, InstanceBounds);
@@ -518,11 +518,11 @@ private:
 		T TOI = 0;
 		for (const auto& Elem : MGlobalPayloads)
 		{
-			const TBox<T, d>& InstanceBounds = Elem.Bounds;
+			const TAABB<T, d>& InstanceBounds = Elem.Bounds;
 			const TVector<T, d> Min = InstanceBounds.Min() - QueryHalfExtents;
 			const TVector<T, d> Max = InstanceBounds.Max() + QueryHalfExtents;
 			TVector<T, d> TmpPosition;
-			if (TBox<T, d>::RaycastFast(Min, Max, Start, CurData.Dir, CurData.InvDir, CurData.bParallel, CurData.CurrentLength, CurData.InvCurrentLength, TOI, TmpPosition))
+			if (TAABB<T, d>(Min,Max).RaycastFast(Start, CurData.Dir, CurData.InvDir, CurData.bParallel, CurData.CurrentLength, CurData.InvCurrentLength, TOI, TmpPosition))
 			{
 				TSpatialVisitorData<TPayloadType> VisitData(Elem.Payload, true, InstanceBounds);
 				const bool bContinue = Visitor.VisitSweep(VisitData, CurData);
@@ -535,11 +535,11 @@ private:
 
 		for (const auto& Elem : MDirtyElements)
 		{
-			const TBox<T, d>& InstanceBounds = Elem.Bounds;
+			const TAABB<T, d>& InstanceBounds = Elem.Bounds;
 			const TVector<T, d> Min = InstanceBounds.Min() - QueryHalfExtents;
 			const TVector<T, d> Max = InstanceBounds.Max() + QueryHalfExtents;
 			TVector<T, d> TmpPosition;
-			if (TBox<T, d>::RaycastFast(Min, Max, Start, CurData.Dir, CurData.InvDir, CurData.bParallel, CurData.CurrentLength, CurData.InvCurrentLength, TOI, TmpPosition))
+			if (TAABB<T, d>(Min, Max).RaycastFast(Start, CurData.Dir, CurData.InvDir, CurData.bParallel, CurData.CurrentLength, CurData.InvCurrentLength, TOI, TmpPosition))
 			{
 				TSpatialVisitorData<TPayloadType> VisitData(Elem.Payload, true, InstanceBounds);
 				const bool bContinue = Visitor.VisitSweep(VisitData, CurData);
@@ -555,7 +555,7 @@ private:
 			return true;
 		}
 
-		TBox<T, d> GlobalBounds(MGrid.MinCorner() - QueryHalfExtents, MGrid.MaxCorner() + QueryHalfExtents);
+		TAABB<T, d> GlobalBounds(MGrid.MinCorner() - QueryHalfExtents, MGrid.MaxCorner() + QueryHalfExtents);
 
 
 		struct FCellIntersection
@@ -579,10 +579,10 @@ private:
 				TVector<T, d> TmpPoint;
 				for (const auto& Elem : Elems)
 				{
-					const TBox<T, d>& InstanceBounds = Elem.Bounds;
+					const TAABB<T, d>& InstanceBounds = Elem.Bounds;
 					const TVector<T, d> Min = InstanceBounds.Min() - QueryHalfExtents;
 					const TVector<T, d> Max = InstanceBounds.Max() + QueryHalfExtents;
-					if (TBox<T, d>::RaycastFast(Min, Max, Start, CurData.Dir, CurData.InvDir, CurData.bParallel, CurData.CurrentLength, CurData.InvCurrentLength, TOI, TmpPoint))
+					if (TAABB<T, d>(Min, Max).RaycastFast(Start, CurData.Dir, CurData.InvDir, CurData.bParallel, CurData.CurrentLength, CurData.InvCurrentLength, TOI, TmpPoint))
 					{
 						TSpatialVisitorData<TPayloadType> VisitData(Elem.Payload, true, InstanceBounds);
 						const bool bContinue = Visitor.VisitSweep(VisitData, CurData);
@@ -615,8 +615,7 @@ private:
 		}
 
 		TVector<T, d> HitPoint;
-		const bool bInitialHit = TBox<T, d>::RaycastFast(GlobalBounds.Min(), GlobalBounds.Max(),
-			Start, CurData.Dir, CurData.InvDir, CurData.bParallel, CurData.CurrentLength, CurData.InvCurrentLength, TOI, HitPoint);
+		const bool bInitialHit = GlobalBounds.RaycastFast(Start, CurData.Dir, CurData.InvDir, CurData.bParallel, CurData.CurrentLength, CurData.InvCurrentLength, TOI, HitPoint);
 		if (bInitialHit)	//NOTE: it's possible to have a non empty IdxsQueue and bInitialHit be false. This is because the IdxsQueue works off clamped cells which we can skip
 		{
 			//Flood fill from inflated cell so that we get all cells along the ray
@@ -679,7 +678,7 @@ private:
 						const TVector<T, d> NeighborCenter = MGrid.Location(NeighborIdx);
 						const TVector<T, d> Min = NeighborCenter - QueryHalfExtents - HalfDx;
 						const TVector<T, d> Max = NeighborCenter + QueryHalfExtents + HalfDx;
-						if (TBox<T, d>::RaycastFast(Min, Max, Start, CurData.Dir, CurData.InvDir, CurData.bParallel, CurData.CurrentLength, CurData.InvCurrentLength, TOI, HitPoint))
+						if (TAABB<T, d>(Min, Max).RaycastFast(Start, CurData.Dir, CurData.InvDir, CurData.bParallel, CurData.CurrentLength, CurData.InvCurrentLength, TOI, HitPoint))
 						{
 							IdxsQueue.Add({ NeighborIdx, TOI });	//should we sort by TOI?
 						}
@@ -717,10 +716,10 @@ private:
 						}
 					}
 
-					const TBox<T, d>& InstanceBounds = Elem.Bounds;
+					const TAABB<T, d>& InstanceBounds = Elem.Bounds;
 					const TVector<T, d> Min = InstanceBounds.Min() - QueryHalfExtents;
 					const TVector<T, d> Max = InstanceBounds.Max() + QueryHalfExtents;
-					if (TBox<T, d>::RaycastFast(Min, Max, Start,
+					if (TAABB<T, d>(Min,Max).RaycastFast(Start,
 						CurData.Dir, CurData.InvDir, CurData.bParallel, CurData.CurrentLength, CurData.InvCurrentLength, TOI, HitPoint))
 					{
 						TSpatialVisitorData<TPayloadType> VisitData(Elem.Payload, true, InstanceBounds);
@@ -739,11 +738,11 @@ private:
 	}
 
 	template <typename SQVisitor, bool bPruneDuplicates = true>
-	bool OverlapImp(const TBox<T, d>& QueryBounds, SQVisitor& Visitor) const
+	bool OverlapImp(const TAABB<T, d>& QueryBounds, SQVisitor& Visitor) const
 	{
 		for (const auto& Elem : MGlobalPayloads)
 		{
-			const TBox<T, d>& InstanceBounds = Elem.Bounds;
+			const TAABB<T, d>& InstanceBounds = Elem.Bounds;
 			if (QueryBounds.Intersects(InstanceBounds))
 			{
 				TSpatialVisitorData<TPayloadType> VisitData(Elem.Payload, true, InstanceBounds);
@@ -756,7 +755,7 @@ private:
 
 		for (const auto& Elem : MDirtyElements)
 		{
-			const TBox<T, d>& InstanceBounds = Elem.Bounds;
+			const TAABB<T, d>& InstanceBounds = Elem.Bounds;
 			if (QueryBounds.Intersects(InstanceBounds))
 			{
 				TSpatialVisitorData<TPayloadType> VisitData(Elem.Payload, true, InstanceBounds);
@@ -772,7 +771,7 @@ private:
 			return true;
 		}
 
-		TBox<T, d> GlobalBounds(MGrid.MinCorner(), MGrid.MaxCorner());
+		TAABB<T, d> GlobalBounds(MGrid.MinCorner(), MGrid.MaxCorner());
 
 		const TVector<int32, d> StartIndex = MGrid.ClampIndex(MGrid.Cell(QueryBounds.Min()));
 		const TVector<int32, d> EndIndex = MGrid.ClampIndex(MGrid.Cell(QueryBounds.Max()));
@@ -795,7 +794,7 @@ private:
 							}
 							InstancesSeen.Add(Elem.Payload);
 						}
-						const TBox<T, d>& InstanceBounds = Elem.Bounds;
+						const TAABB<T, d>& InstanceBounds = Elem.Bounds;
 						if (QueryBounds.Intersects(InstanceBounds))
 						{
 							TSpatialVisitorData<TPayloadType> VisitData(Elem.Payload, true, InstanceBounds);
@@ -826,13 +825,13 @@ private:
 		}
 
 		SCOPE_CYCLE_COUNTER(STAT_BoundingVolumeGenerateTree);
-		TArray<TBox<T, d>> AllBounds;
+		TArray<TAABB<T, d>> AllBounds;
 		TArray<bool> HasBounds;
 
 		AllBounds.SetNum(Particles.Num());
 		HasBounds.SetNum(Particles.Num());
 		int32 MaxPayloadBoundsCopy = MaxPayloadBounds;
-		auto GetValidBounds = [MaxPayloadBoundsCopy, bUseVelocity, Dt](const auto& Particle, TBox<T,d>& OutBounds) -> bool
+		auto GetValidBounds = [MaxPayloadBoundsCopy, bUseVelocity, Dt](const auto& Particle, TAABB<T,d>& OutBounds) -> bool
 		{
 			if (HasBoundingBox(Particle))
 			{
@@ -852,17 +851,17 @@ private:
 		auto& PayloadInfos = MPayloadInfo;
 		T NumObjectsWithBounds = 0;
 
-		auto ComputeBoxAndDx = [&Particles, &AllBounds, &HasBounds, &GlobalPayloads, &PayloadInfos, &GetValidBounds, &NumObjectsWithBounds](TBox<T,d>& OutGlobalBox, bool bFirstPass) -> T
+		auto ComputeBoxAndDx = [&Particles, &AllBounds, &HasBounds, &GlobalPayloads, &PayloadInfos, &GetValidBounds, &NumObjectsWithBounds](TAABB<T,d>& OutGlobalBox, bool bFirstPass) -> T
 		{
 			SCOPE_CYCLE_COUNTER(STAT_BoundingVolumeComputeGlobalBox);
-			OutGlobalBox = TBox<T, d>::EmptyBox();
+			OutGlobalBox = TAABB<T, d>::EmptyAABB();
 			constexpr T InvD = (T)1 / d;
 			int32 Idx = 0;
 			T Dx = 0;
 			NumObjectsWithBounds = 0;
 			for (auto& Particle : Particles)
 			{
-				TBox<T,d>& Bounds = AllBounds[Idx];
+				TAABB<T,d>& Bounds = AllBounds[Idx];
 				if ((bFirstPass && GetValidBounds(Particle, Bounds)) || (!bFirstPass && HasBounds[Idx]))
 				{
 					HasBounds[Idx] = true;
@@ -877,7 +876,7 @@ private:
 
 					const int32 GlobalPayloadIdx = GlobalPayloads.Num();
 					bool bTooBig = HasBoundingBox(Particle);	//todo: avoid this as it was already called in GetValidBounds
-					GlobalPayloads.Add({ Payload, bTooBig ? Bounds : TBox<T,d>(TVector<T,d>(TNumericLimits<T>::Lowest()), TVector<T,d>(TNumericLimits<T>::Max())) });
+					GlobalPayloads.Add({ Payload, bTooBig ? Bounds : TAABB<T,d>(TVector<T,d>(TNumericLimits<T>::Lowest()), TVector<T,d>(TNumericLimits<T>::Max())) });
 					PayloadInfos.Add(Payload, FPayloadInfo{ GlobalPayloadIdx, INDEX_NONE });
 				}
 				++Idx;
@@ -887,7 +886,7 @@ private:
 			return Dx;
 		};
 
-		TBox<T, d> GlobalBox;
+		TAABB<T, d> GlobalBox;
 		T Dx = ComputeBoxAndDx(GlobalBox, /*bFirstPass=*/true);
 
 		if (FBoundingVolumeCVars::FilterFarBodies)
@@ -955,7 +954,7 @@ private:
 			{
 				if (HasBounds[Idx])
 				{
-					const TBox<T, d>& ObjectBox = AllBounds[Idx];
+					const TAABB<T, d>& ObjectBox = AllBounds[Idx];
 					NumObjectsWithBounds += 1;
 					const auto StartIndex = MGrid.ClampIndex(MGrid.Cell(ObjectBox.Min()));
 					const auto EndIndex = MGrid.ClampIndex(MGrid.Cell(ObjectBox.Max()));
@@ -1021,7 +1020,7 @@ private:
 		}
 	}
 
-	void AddElementToExistingGrid(const TPayloadType& Payload, FPayloadInfo& PayloadInfo, const TBox<T, d>& NewBounds, bool bHasBounds)
+	void AddElementToExistingGrid(const TPayloadType& Payload, FPayloadInfo& PayloadInfo, const TAABB<T, d>& NewBounds, bool bHasBounds)
 	{
 		bool bTooBig = false;
 		if (bHasBounds)
@@ -1083,11 +1082,11 @@ private:
 		{
 			PayloadInfo.GlobalPayloadIdx = MGlobalPayloads.Num();
 			PayloadInfo.DirtyPayloadIdx = INDEX_NONE;
-			MGlobalPayloads.Add({ Payload, bTooBig ? NewBounds : TBox<T,d>(TVector<T,d>(TNumericLimits<T>::Lowest()), TVector<T,d>(TNumericLimits<T>::Max())) });
+			MGlobalPayloads.Add({ Payload, bTooBig ? NewBounds : TAABB<T,d>(TVector<T,d>(TNumericLimits<T>::Lowest()), TVector<T,d>(TNumericLimits<T>::Max())) });
 		}
 	}
 
-	TArray<TPayloadType> FindAllIntersectionsHelper(const TBox<T, d>& ObjectBox) const
+	TArray<TPayloadType> FindAllIntersectionsHelper(const TAABB<T, d>& ObjectBox) const
 	{
 		TArray<TPayloadType> Intersections;
 		const auto StartIndex = MGrid.ClampIndex(MGrid.Cell(ObjectBox.Min()));
diff --git a/Engine/Source/Runtime/Experimental/Chaos/Public/Chaos/BoundingVolumeHierarchy.h b/Engine/Source/Runtime/Experimental/Chaos/Public/Chaos/BoundingVolumeHierarchy.h
--- a/Engine/Source/Runtime/Experimental/Chaos/Public/Chaos/BoundingVolumeHierarchy.h
+++ b/Engine/Source/Runtime/Experimental/Chaos/Public/Chaos/BoundingVolumeHierarchy.h
@@ -2,7 +2,7 @@
 #pragma once
 
 #include "Chaos/Array.h"
-#include "Chaos/Box.h"
+#include "Chaos/AABB.h"
 #include "Chaos/Defines.h"
 #include "Chaos/GeometryParticles.h"
 #include "Chaos/ParticleHandleFwd.h"
@@ -114,7 +114,7 @@ class TBoundingVolumeHierarchy final : public ISpatialAcceleration<int32, T,d>
 	}
 
 	// Begin ISpatialAcceleration interface
-	CHAOS_API TArray<int32> FindAllIntersections(const TBox<T, d>& Box) const { return FindAllIntersectionsImp(Box); }
+	CHAOS_API TArray<int32> FindAllIntersections(const TAABB<T, d>& Box) const { return FindAllIntersectionsImp(Box); }
 	CHAOS_API TArray<int32> FindAllIntersections(const TSpatialRay<T,d>& Ray) const { return FindAllIntersectionsImp(Ray); }
 	CHAOS_API TArray<int32> FindAllIntersections(const TVector<T, d>& Point) const { return FindAllIntersectionsImp(Point); }
 	CHAOS_API TArray<int32> FindAllIntersections(const TGeometryParticles<T, d>& InParticles, const int32 i) const;
@@ -126,7 +126,7 @@ class TBoundingVolumeHierarchy final : public ISpatialAcceleration<int32, T,d>
 	}
 
 	// TODO(mlentine): Need to move this elsewhere; probably on CollisionConstraint
-	CHAOS_API const TBox<T, d>& GetWorldSpaceBoundingBox(const TGeometryParticles<T, d>& InParticles, const int32 Index);
+	CHAOS_API const TAABB<T, d>& GetWorldSpaceBoundingBox(const TGeometryParticles<T, d>& InParticles, const int32 Index);
 
 #if !UE_BUILD_SHIPPING
 	CHAOS_API virtual void DebugDraw(ISpacialDebugDrawInterface<T>* InInterface) const override;
@@ -150,7 +150,7 @@ class TBoundingVolumeHierarchy final : public ISpatialAcceleration<int32, T,d>
 	}
 
 	CHAOS_API TArray<int32> FindAllIntersectionsHelper(const TBVHNode<T,d>& MyNode, const TVector<T, d>& Point) const;
-	CHAOS_API TArray<int32> FindAllIntersectionsHelper(const TBVHNode<T,d>& MyNode, const TBox<T, d>& ObjectBox) const;
+	CHAOS_API TArray<int32> FindAllIntersectionsHelper(const TBVHNode<T,d>& MyNode, const TAABB<T, d>& ObjectBox) const;
 	CHAOS_API TArray<int32> FindAllIntersectionsHelper(const TBVHNode<T, d>& MyNode, const TSpatialRay<T, d>& Ray) const;
 
 	template <typename QUERY_OBJECT>
@@ -162,7 +162,7 @@ class TBoundingVolumeHierarchy final : public ISpatialAcceleration<int32, T,d>
 
 	OBJECT_ARRAY const* MObjects;
 	TArray<int32> MGlobalObjects;
-	TMap<int32, TBox<T, d>> MWorldSpaceBoxes;
+	TMap<int32, TAABB<T, d>> MWorldSpaceBoxes;
 	TArray<int32> MScratchAllObjects;
 	int32 MMaxLevels;
 	TArray<TBVHNode<T,d>> Elements;
diff --git a/Engine/Source/Runtime/Experimental/Chaos/Public/Chaos/BoundingVolumeUtilities.h b/Engine/Source/Runtime/Experimental/Chaos/Public/Chaos/BoundingVolumeUtilities.h
--- a/Engine/Source/Runtime/Experimental/Chaos/Public/Chaos/BoundingVolumeUtilities.h
+++ b/Engine/Source/Runtime/Experimental/Chaos/Public/Chaos/BoundingVolumeUtilities.h
@@ -89,64 +89,64 @@ bool HasBoundingBox(const TPBDRigidParticleHandleImp<T, d, bPersistent>& Handle)
 }
 
 template<class OBJECT_ARRAY, class T, int d>
-const TBox<T, d>& GetWorldSpaceBoundingBox(const OBJECT_ARRAY& Objects, const int32 i, const TMap<int32, TBox<T, d>>& WorldSpaceBoxes)
+const TAABB<T, d>& GetWorldSpaceBoundingBox(const OBJECT_ARRAY& Objects, const int32 i, const TMap<int32, TAABB<T, d>>& WorldSpaceBoxes)
 {
 	return Objects[i]->BoundingBox();
 }
 
 template<class T, int d>
-const TBox<T, d>& GetWorldSpaceBoundingBox(const TParticles<T, d>& Objects, const int32 i, const TMap<int32, TBox<T, d>>& WorldSpaceBoxes)
+const TAABB<T, d>& GetWorldSpaceBoundingBox(const TParticles<T, d>& Objects, const int32 i, const TMap<int32, TAABB<T, d>>& WorldSpaceBoxes)
 {
 	return WorldSpaceBoxes[i];
 }
 
 template<class T, int d>
-const TBox<T, d>& GetWorldSpaceBoundingBox(const TGeometryParticles<T, d>& Objects, const int32 i, const TMap<int32, TBox<T, d>>& WorldSpaceBoxes)
+const TAABB<T, d>& GetWorldSpaceBoundingBox(const TGeometryParticles<T, d>& Objects, const int32 i, const TMap<int32, TAABB<T, d>>& WorldSpaceBoxes)
 {
 	return GetWorldSpaceBoundingBox(static_cast<const TParticles<T, d>&>(Objects), i, WorldSpaceBoxes);
 }
 
 template<class T, int d>
-const TBox<T, d>& GetWorldSpaceBoundingBox(const TPBDRigidParticles<T, d>& Objects, const int32 i, const TMap<int32, TBox<T, d>>& WorldSpaceBoxes)
+const TAABB<T, d>& GetWorldSpaceBoundingBox(const TPBDRigidParticles<T, d>& Objects, const int32 i, const TMap<int32, TAABB<T, d>>& WorldSpaceBoxes)
 {
 	return GetWorldSpaceBoundingBox(static_cast<const TParticles<T, d>&>(Objects), i, WorldSpaceBoxes);
 }
 
 template<class T, int d>
-TBox<T, d> ComputeWorldSpaceBoundingBox(const TParticles<T, d>& Objects, const int32 i, bool bUseVelocity = false, T Dt = 0)
+TAABB<T, d> ComputeWorldSpaceBoundingBox(const TParticles<T, d>& Objects, const int32 i, bool bUseVelocity = false, T Dt = 0)
 {
 	ensure(!bUseVelocity);
-	return TBox<T, d>(Objects.X(i), Objects.X(i));
+	return TAABB<T, d>(Objects.X(i), Objects.X(i));
 }
 
 template<class T, int d>
-TBox<T, d> ComputeWorldSpaceBoundingBox(const TGeometryParticles<T, d>& Objects, const int32 i, bool bUseVelocity = false, T Dt = 0)
+TAABB<T, d> ComputeWorldSpaceBoundingBox(const TGeometryParticles<T, d>& Objects, const int32 i, bool bUseVelocity = false, T Dt = 0)
 {
 	ensure(!bUseVelocity);
 	TRigidTransform<T, d> LocalToWorld(Objects.X(i), Objects.R(i));
 	const auto& LocalBoundingBox = Objects.Geometry(i)->BoundingBox();
-	return LocalBoundingBox.TransformedBox(LocalToWorld);
+	return LocalBoundingBox.TransformedAABB(LocalToWorld);
 }
 
 template<class T, int d>
-TBox<T, d> ComputeWorldSpaceBoundingBox(const TPBDRigidParticles<T, d>& Objects, const int32 i, bool bUseVelocity = false, T Dt = 0)
+TAABB<T, d> ComputeWorldSpaceBoundingBox(const TPBDRigidParticles<T, d>& Objects, const int32 i, bool bUseVelocity = false, T Dt = 0)
 {
 	TRigidTransform<T, d> LocalToWorld(Objects.P(i), Objects.Q(i));
-	TBox<T, d> WorldSpaceBox;
+	TAABB<T, d> WorldSpaceBox;
 	if (Objects.Geometry(i))
 	{
 		const auto& LocalBoundingBox = Objects.Geometry(i)->BoundingBox();
-		WorldSpaceBox = LocalBoundingBox.TransformedBox(LocalToWorld);
+		WorldSpaceBox = LocalBoundingBox.TransformedAABB(LocalToWorld);
 	}
 	else
 	{
 		check(Objects.CollisionParticles(i) && Objects.CollisionParticles(i)->Size());
-		TBox<T, d> LocalBoundingBox(Objects.CollisionParticles(i)->X(0), Objects.CollisionParticles(i)->X(0));
+		TAABB<T, d> LocalBoundingBox(Objects.CollisionParticles(i)->X(0), Objects.CollisionParticles(i)->X(0));
 		for (uint32 j = 1; j < Objects.CollisionParticles(i)->Size(); ++j)
 		{
 			LocalBoundingBox.GrowToInclude(Objects.CollisionParticles(i)->X(j));
 		}
-		WorldSpaceBox = LocalBoundingBox.TransformedBox(LocalToWorld);
+		WorldSpaceBox = LocalBoundingBox.TransformedAABB(LocalToWorld);
 	}
 	
 	if (bUseVelocity)
@@ -158,7 +158,7 @@ TBox<T, d> ComputeWorldSpaceBoundingBox(const TPBDRigidParticles<T, d>& Objects,
 }
 
 template<typename THandle, typename T, int d, bool bPersistent>
-TBox<T, d> ComputeWorldSpaceBoundingBoxForHandle(const THandle& Handle)
+TAABB<T, d> ComputeWorldSpaceBoundingBoxForHandle(const THandle& Handle)
 {
 	const auto PBDRigid = Handle.CastToRigidParticle();
 	const bool bIsRigidDynamic = PBDRigid && PBDRigid->ObjectState() == EObjectStateType::Dynamic;
@@ -172,7 +172,7 @@ TBox<T, d> ComputeWorldSpaceBoundingBoxForHandle(const THandle& Handle)
 
 	check(PBDRigid);
 	check(PBDRigid->CollisionParticles() && PBDRigid->CollisionParticles()->Size());
-	TBox<T, d> LocalBoundingBox(PBDRigid->CollisionParticles()->X(0), PBDRigid->CollisionParticles()->X(0));
+	TAABB<T, d> LocalBoundingBox(PBDRigid->CollisionParticles()->X(0), PBDRigid->CollisionParticles()->X(0));
 	for(uint32 j = 1; j < PBDRigid->CollisionParticles()->Size(); ++j)
 	{
 		LocalBoundingBox.GrowToInclude(PBDRigid->CollisionParticles()->X(j));
@@ -181,28 +181,28 @@ TBox<T, d> ComputeWorldSpaceBoundingBoxForHandle(const THandle& Handle)
 }
 
 template<typename T, int d, bool bPersistent>
-TBox<T, d> ComputeWorldSpaceBoundingBox(const TGeometryParticleHandleImp<T, d, bPersistent>& Handle, bool bUseVelocity = false, T Dt = 0)
+TAABB<T, d> ComputeWorldSpaceBoundingBox(const TGeometryParticleHandleImp<T, d, bPersistent>& Handle, bool bUseVelocity = false, T Dt = 0)
 {
 	return Handle.WorldSpaceInflatedBounds();
 }
 
 template<typename T, int d, bool bPersistent>
-TBox<T, d> ComputeWorldSpaceBoundingBox(const TPBDRigidParticleHandleImp<T, d, bPersistent>& Handle, bool bUseVelocity = false, T Dt = 0)
+TAABB<T, d> ComputeWorldSpaceBoundingBox(const TPBDRigidParticleHandleImp<T, d, bPersistent>& Handle, bool bUseVelocity = false, T Dt = 0)
 {
 	return Handle.WorldSpaceInflatedBounds();
 }
 
 template<typename GenericEntry, typename T>
-TBox<T, 3> ComputeWorldSpaceBoundingBox(const GenericEntry& InEntry, bool bUseVelocity = false, T Dt = 0)
+TAABB<T, 3> ComputeWorldSpaceBoundingBox(const GenericEntry& InEntry, bool bUseVelocity = false, T Dt = 0)
 {
 	ensure(!bUseVelocity);
 	return InEntry.BoundingBox();
 }
 
 template<typename OBJECT_ARRAY, typename T, int d>
-const TBox<T, d> ComputeGlobalBoxAndSplitAxis(const OBJECT_ARRAY& Objects, const TArray<int32>& AllObjects, const TMap<int32, TBox<T, d>>& WorldSpaceBoxes, bool bAllowMultipleSplitting, int32& OutAxis)
+const TAABB<T, d> ComputeGlobalBoxAndSplitAxis(const OBJECT_ARRAY& Objects, const TArray<int32>& AllObjects, const TMap<int32, TAABB<T, d>>& WorldSpaceBoxes, bool bAllowMultipleSplitting, int32& OutAxis)
 {
-	TBox<T, d> GlobalBox = GetWorldSpaceBoundingBox(Objects, AllObjects[0], WorldSpaceBoxes);
+	TAABB<T, d> GlobalBox = GetWorldSpaceBoundingBox(Objects, AllObjects[0], WorldSpaceBoxes);
 	for (int32 i = 1; i < AllObjects.Num(); ++i)
 	{
 		GlobalBox.GrowToInclude(GetWorldSpaceBoundingBox(Objects, AllObjects[i], WorldSpaceBoxes));
@@ -227,7 +227,7 @@ const TBox<T, d> ComputeGlobalBoxAndSplitAxis(const OBJECT_ARRAY& Objects, const
 }
 
 template<typename T, int d>
-const TBox<T, d> ComputeGlobalBoxAndSplitAxis(const TParticles<T,d>& Objects, const TArray<int32>& AllObjects, const TMap<int32, TBox<T, d>>& WorldSpaceBoxes, bool bAllowMultipleSplitting, int32& OutAxis)
+const TAABB<T, d> ComputeGlobalBoxAndSplitAxis(const TParticles<T,d>& Objects, const TArray<int32>& AllObjects, const TMap<int32, TAABB<T, d>>& WorldSpaceBoxes, bool bAllowMultipleSplitting, int32& OutAxis)
 {
 	//simple particles means we can split more efficiently
 	TPair<int32, int32> Counts[d];
@@ -247,11 +247,11 @@ const TBox<T, d> ComputeGlobalBoxAndSplitAxis(const TParticles<T,d>& Objects, co
 		};
 	};
 
-	TBox<T, d> GlobalBox = GetWorldSpaceBoundingBox(Objects, AllObjects[0], WorldSpaceBoxes);
+	TAABB<T, d> GlobalBox = GetWorldSpaceBoundingBox(Objects, AllObjects[0], WorldSpaceBoxes);
 	CountLambda(GlobalBox.Center());
 	for (int32 i = 1; i < AllObjects.Num(); ++i)
 	{
-		TBox<T, d> PtBox = GetWorldSpaceBoundingBox(Objects, AllObjects[i], WorldSpaceBoxes);
+		TAABB<T, d> PtBox = GetWorldSpaceBoundingBox(Objects, AllObjects[i], WorldSpaceBoxes);
 		GlobalBox.GrowToInclude(PtBox);
 		CountLambda(PtBox.Center());
 	}
@@ -276,13 +276,13 @@ const TBox<T, d> ComputeGlobalBoxAndSplitAxis(const TParticles<T,d>& Objects, co
 }
 
 template<class OBJECT_ARRAY, class T, int d>
-void ComputeAllWorldSpaceBoundingBoxes(const OBJECT_ARRAY& Objects, const TArray<int32>& AllObjects, const bool bUseVelocity, const T Dt, TMap<int32, TBox<T, d>>& WorldSpaceBoxes)
+void ComputeAllWorldSpaceBoundingBoxes(const OBJECT_ARRAY& Objects, const TArray<int32>& AllObjects, const bool bUseVelocity, const T Dt, TMap<int32, TAABB<T, d>>& WorldSpaceBoxes)
 {
 	check(!bUseVelocity);
 }
 
 template<class T, int d>
-void ComputeAllWorldSpaceBoundingBoxes(const TParticles<T, d>& Objects, const TArray<int32>& AllObjects, const bool bUseVelocity, const T Dt, TMap<int32, TBox<T, d>>& WorldSpaceBoxes)
+void ComputeAllWorldSpaceBoundingBoxes(const TParticles<T, d>& Objects, const TArray<int32>& AllObjects, const bool bUseVelocity, const T Dt, TMap<int32, TAABB<T, d>>& WorldSpaceBoxes)
 {
 	check(!bUseVelocity);
 	WorldSpaceBoxes.Reserve(AllObjects.Num());
@@ -295,7 +295,7 @@ void ComputeAllWorldSpaceBoundingBoxes(const TParticles<T, d>& Objects, const TA
 }
 
 template<class T, int d>
-void ComputeAllWorldSpaceBoundingBoxes(const TGeometryParticles<T, d>& Objects, const TArray<int32>& AllObjects, const bool bUseVelocity, const T Dt, TMap<int32, TBox<T, d>>& WorldSpaceBoxes)
+void ComputeAllWorldSpaceBoundingBoxes(const TGeometryParticles<T, d>& Objects, const TArray<int32>& AllObjects, const bool bUseVelocity, const T Dt, TMap<int32, TAABB<T, d>>& WorldSpaceBoxes)
 {
 	check(!bUseVelocity);
 	WorldSpaceBoxes.Reserve(AllObjects.Num());
@@ -308,14 +308,14 @@ void ComputeAllWorldSpaceBoundingBoxes(const TGeometryParticles<T, d>& Objects,
 }
 
 template<class T, int d>
-void ComputeAllWorldSpaceBoundingBoxes(const TPBDRigidParticles<T, d>& Objects, const TArray<int32>& AllObjects, const bool bUseVelocity, const T Dt, TMap<int32, TBox<T, d>>& WorldSpaceBoxes)
+void ComputeAllWorldSpaceBoundingBoxes(const TPBDRigidParticles<T, d>& Objects, const TArray<int32>& AllObjects, const bool bUseVelocity, const T Dt, TMap<int32, TAABB<T, d>>& WorldSpaceBoxes)
 {
 	WorldSpaceBoxes.Reserve(AllObjects.Num());
 	//PhysicsParallelFor(AllObjects.Num(), [&](int32 i) {
 	for (int32 i = 0; i < AllObjects.Num(); ++i)
 	{
 		const int32 BodyIndex = AllObjects[i];
-		TBox<T, d>& WorldSpaceBox = WorldSpaceBoxes.FindOrAdd(BodyIndex);
+		TAABB<T, d>& WorldSpaceBox = WorldSpaceBoxes.FindOrAdd(BodyIndex);
 		WorldSpaceBox = ComputeWorldSpaceBoundingBox(Objects, BodyIndex, bUseVelocity, Dt);
 	}
 	//});
@@ -331,7 +331,7 @@ struct CParticleView
 
 //todo: how do we protect ourselves and make it const?
 template<typename ParticleView, typename T, int d>
-typename TEnableIf<TModels<CParticleView, ParticleView>::Value>::Type ComputeAllWorldSpaceBoundingBoxes(const ParticleView& Particles, const TArray<bool>& RequiresBounds, const bool bUseVelocity, const T Dt, TArray<TBox<T, d>>& WorldSpaceBoxes)
+typename TEnableIf<TModels<CParticleView, ParticleView>::Value>::Type ComputeAllWorldSpaceBoundingBoxes(const ParticleView& Particles, const TArray<bool>& RequiresBounds, const bool bUseVelocity, const T Dt, TArray<TAABB<T, d>>& WorldSpaceBoxes)
 {
 	WorldSpaceBoxes.AddUninitialized(Particles.Num());
 	ParticlesParallelFor(Particles, [&RequiresBounds, &WorldSpaceBoxes, bUseVelocity, Dt](const auto& Particle, int32 Index)
@@ -352,7 +352,7 @@ typename TEnableIf<TModels<CParticleView, ParticleView>::Value>::Type ComputeAll
 }
 
 template<typename ParticleView, typename T, int d>
-typename TEnableIf<!TModels<CParticleView, ParticleView>::Value>::Type ComputeAllWorldSpaceBoundingBoxes(const ParticleView& Particles, const TArray<bool>& RequiresBounds, const bool bUseVelocity, const T Dt, TArray<TBox<T, d>>& WorldSpaceBoxes)
+typename TEnableIf<!TModels<CParticleView, ParticleView>::Value>::Type ComputeAllWorldSpaceBoundingBoxes(const ParticleView& Particles, const TArray<bool>& RequiresBounds, const bool bUseVelocity, const T Dt, TArray<TAABB<T, d>>& WorldSpaceBoxes)
 {
 	WorldSpaceBoxes.AddUninitialized(Particles.Num());
 	ParticlesParallelFor(Particles, [&RequiresBounds, &WorldSpaceBoxes, bUseVelocity, Dt](const auto& Particle, int32 Index)
diff --git a/Engine/Source/Runtime/Experimental/Chaos/Public/Chaos/Box.h b/Engine/Source/Runtime/Experimental/Chaos/Public/Chaos/Box.h
--- a/Engine/Source/Runtime/Experimental/Chaos/Public/Chaos/Box.h
+++ b/Engine/Source/Runtime/Experimental/Chaos/Public/Chaos/Box.h
@@ -7,6 +7,7 @@
 #include "Chaos/AABB.h"
 #include "Chaos/Transform.h"
 #include "ChaosArchive.h"
+#include "UObject/ExternalPhysicsCustomObjectVersion.h"
 
 namespace Chaos
 {
@@ -95,9 +96,9 @@ namespace Chaos
 		}
 
 		template<class TTRANSFORM>
-		TBox<T, d> TransformedBox(const TTRANSFORM& SpaceTransform) const
+		TAABB<T, d> TransformedBox(const TTRANSFORM& SpaceTransform) const
 		{
-			return TBox<T, d>(AABB.TransformedAABB(SpaceTransform));
+			return TAABB<T, d>(AABB.TransformedAABB(SpaceTransform));
 		}
 
 		static FORCEINLINE bool Intersects(const TBox<T, d>& A, const TBox<T, d>& B)
@@ -105,9 +106,14 @@ namespace Chaos
 			return A.AABB.Intersects(B.AABB);
 		}
 
-		static FORCEINLINE TBox<T, d> Intersection(const TBox<T, d>& A, const TBox<T, d>& B)
+		static FORCEINLINE TAABB<T, d> Intersection(const TAABB<T, d>& A, const TAABB<T, d>& B)
 		{
-			return TBox<T, d>(A.AABB.GetIntersection(B.AABB));
+			return TAABB<T, d>(A.AABB.GetIntersection(B.AABB));
+		}
+
+		FORCEINLINE bool Intersects(const TAABB<T, d>& Other) const
+		{
+			return AABB.Intersects(Other);
 		}
 
 		FORCEINLINE bool Intersects(const TBox<T, d>& Other) const
@@ -115,9 +121,9 @@ namespace Chaos
 			return AABB.Intersects(Other.AABB);
 		}
 
-		TBox<T, d> GetIntersection(const TBox<T, d>& Other) const
+		TAABB<T, d> GetIntersection(const TAABB<T, d>& Other) const
 		{
-			return TBox<T, d>(AABB.GetIntersection(Other.AABB));
+			return TAABB<T, d>(AABB.GetIntersection(Other.AABB));
 		}
 
 		FORCEINLINE bool Contains(const TVector<T, d>& Point) const
@@ -132,7 +138,7 @@ namespace Chaos
 
 		FORCEINLINE static constexpr EImplicitObjectType StaticType() { return ImplicitObjectType::Box; }
 
-		const TBox<T, d>& BoundingBox() const { return *this; }
+		const TAABB<T, d>& BoundingBox() const { return AABB; }
 
 		virtual T PhiWithNormal(const TVector<T, d>& x, TVector<T, d>& Normal) const override
 		{
@@ -183,12 +189,12 @@ namespace Chaos
 			AABB.GrowToInclude(V);
 		}
 
-		FORCEINLINE void GrowToInclude(const TBox<T, d>& Other)
+		FORCEINLINE void GrowToInclude(const TAABB<T, d>& Other)
 		{
 			AABB.GrowToInclude(Other.AABB);
 		}
 
-		FORCEINLINE void ShrinkToInclude(const TBox<T, d>& Other)
+		FORCEINLINE void ShrinkToInclude(const TAABB<T, d>& Other)
 		{
 			AABB.ShrinkToInclude(Other.AABB);
 		}
@@ -240,10 +246,10 @@ namespace Chaos
 
 		FORCEINLINE static TRotation<T, d> GetRotationOfMass()
 		{
-			return TRotation<T, d>::FromElements(TVector<T, d>(0), 1);
+			return TAABB<T, d>::GetRotationOfMass();
 		}
 
-		virtual FString ToString() const { return FString::Printf(TEXT("TBox Min:%s, Max:%s"), *Min().ToString(), *Max().ToString()); }
+		virtual FString ToString() const { return FString::Printf(TEXT("TAABB Min:%s, Max:%s"), *Min().ToString(), *Max().ToString()); }
 
 		FORCEINLINE void SerializeImp(FArchive& Ar)
 		{
@@ -251,6 +257,60 @@ namespace Chaos
 			AABB.Serialize(Ar);
 		}
 
+		static void SerializeAsAABB(FArchive& Ar, TAABB<T,d>& AABB)
+		{
+			Ar.UsingCustomVersion(FExternalPhysicsCustomObjectVersion::GUID);
+			if (Ar.CustomVer(FExternalPhysicsCustomObjectVersion::GUID) < FExternalPhysicsCustomObjectVersion::TBoxReplacedWithTAABB)
+			{
+				TBox<T, d> Tmp;
+				Ar << Tmp;
+				AABB = Tmp.AABB;
+			}
+			else
+			{
+				AABB.Serialize(Ar);
+			}
+		}
+
+		static void SerializeAsAABBs(FArchive& Ar, TArray<TAABB<T, d>>& AABBs)
+		{
+			Ar.UsingCustomVersion(FExternalPhysicsCustomObjectVersion::GUID);
+			if (Ar.CustomVer(FExternalPhysicsCustomObjectVersion::GUID) < FExternalPhysicsCustomObjectVersion::TBoxReplacedWithTAABB)
+			{
+				TArray<TBox<T, d>> Tmp;
+				Ar << Tmp;
+				AABBs.Reserve(Tmp.Num());
+				for (const TBox<T, d>& Box : Tmp)
+				{
+					AABBs.Add(Box.AABB);
+				}
+			}
+			else
+			{
+				Ar << AABBs;
+			}
+		}
+
+		template <typename Key>
+		static void SerializeAsAABBs(FArchive& Ar, TMap<Key, TAABB<T, d>>& AABBs)
+		{
+			Ar.UsingCustomVersion(FExternalPhysicsCustomObjectVersion::GUID);
+			if (Ar.CustomVer(FExternalPhysicsCustomObjectVersion::GUID) < FExternalPhysicsCustomObjectVersion::TBoxReplacedWithTAABB)
+			{
+				TMap<Key,TBox<T, d>> Tmp;
+				Ar << Tmp;
+
+				for (const auto Itr : Tmp)
+				{
+					AABBs.Add(Itr.Key, Itr.Value.AABB);
+				}
+			}
+			else
+			{
+				Ar << AABBs;
+			}
+		}
+
 		virtual void Serialize(FChaosArchive& Ar) override
 		{
 			FChaosArchiveScopedMemory ScopedMemory(Ar, GetTypeName());
@@ -259,8 +319,8 @@ namespace Chaos
 
 		virtual void Serialize(FArchive& Ar) override { SerializeImp(Ar); }
 
-		static TBox<T, d> EmptyBox() { return TBox<T, d>(TVector<T, d>(TNumericLimits<T>::Max()), TVector<T, d>(-TNumericLimits<T>::Max())); }
-		static TBox<T, d> ZeroBox() { return TBox<T, d>(TVector<T, d>((T)0), TVector<T, d>((T)0)); }
+		static TAABB<T, d> EmptyBox() { return TAABB<T, d>::EmptyAABB(); }
+		static TAABB<T, d> ZeroBox() { return TAABB<T, d>::ZeroAABB(); }
 
 		virtual uint32 GetTypeHash() const override 
 		{
diff --git a/Engine/Source/Runtime/Experimental/Chaos/Public/Chaos/Capsule.h b/Engine/Source/Runtime/Experimental/Chaos/Public/Chaos/Capsule.h
--- a/Engine/Source/Runtime/Experimental/Chaos/Public/Chaos/Capsule.h
+++ b/Engine/Source/Runtime/Experimental/Chaos/Public/Chaos/Capsule.h
@@ -31,7 +31,7 @@ namespace Chaos
 		    , MUnionedObjects(nullptr)
 		{
 			MLocalBoundingBox.GrowToInclude(x2);
-			MLocalBoundingBox = TBox<T, 3>(MLocalBoundingBox.Min() - TVector<T, 3>(MRadius), MLocalBoundingBox.Max() + TVector<T, 3>(MRadius));
+			MLocalBoundingBox = TAABB<T, 3>(MLocalBoundingBox.Min() - TVector<T, 3>(MRadius), MLocalBoundingBox.Max() + TVector<T, 3>(MRadius));
 			InitUnionedObjects();
 		}
 
@@ -131,7 +131,7 @@ namespace Chaos
 			return Normal.SafeNormalize() - MRadius;
 		}
 
-		virtual const TBox<T, 3>& BoundingBox() const override { return MLocalBoundingBox; }
+		virtual const TAABB<T, 3>& BoundingBox() const override { return MLocalBoundingBox; }
 
 		static bool RaycastFast(T MRadius, T MHeight, const TVector<T,3>& MVector, const TVector<T,3>& X1, const TVector<T,3>& X2, const TVector<T, 3>& StartPoint, const TVector<T, 3>& Dir, const T Length, const T Thickness, T& OutTime, TVector<T, 3>& OutPosition, TVector<T, 3>& OutNormal, int32& OutFaceIndex)
 		{
@@ -303,7 +303,8 @@ namespace Chaos
 		{
 			FImplicitObject::SerializeImp(Ar);
 			MSegment.Serialize(Ar);
-			Ar << MRadius << MLocalBoundingBox;
+			Ar << MRadius;
+			TBox<FReal, 3>::SerializeAsAABB(Ar, MLocalBoundingBox);
 		}
 
 		virtual void Serialize(FChaosArchive& Ar) override
@@ -386,7 +387,7 @@ namespace Chaos
 
 		TSegment<T> MSegment;
 		T MRadius;
-		TBox<T, 3> MLocalBoundingBox;
+		TAABB<T, 3> MLocalBoundingBox;
 		TUniquePtr<TImplicitObjectUnion<T, 3>> MUnionedObjects;
 	};
 
diff --git a/Engine/Source/Runtime/Experimental/Chaos/Public/Chaos/Collision/ParticlePairBroadPhase.h b/Engine/Source/Runtime/Experimental/Chaos/Public/Chaos/Collision/ParticlePairBroadPhase.h
--- a/Engine/Source/Runtime/Experimental/Chaos/Public/Chaos/Collision/ParticlePairBroadPhase.h
+++ b/Engine/Source/Runtime/Experimental/Chaos/Public/Chaos/Collision/ParticlePairBroadPhase.h
@@ -42,9 +42,9 @@ namespace Chaos
 				TGeometryParticleHandle<FReal, 3>* Particle0 = const_cast<TGeometryParticleHandle<FReal, 3>*>(ParticlePair[0]);
 				TGeometryParticleHandle<FReal, 3>* Particle1 = const_cast<TGeometryParticleHandle<FReal, 3>*>(ParticlePair[1]);
 
-				const TBox<FReal, 3>& Box0 = Particle0->WorldSpaceInflatedBounds();
-				const TBox<FReal, 3>& Box1 = Particle1->WorldSpaceInflatedBounds();
-				if (TBox<FReal, 3>::Intersects(Box0, Box1))
+				const TAABB<FReal, 3>& Box0 = Particle0->WorldSpaceInflatedBounds();
+				const TAABB<FReal, 3>& Box1 = Particle1->WorldSpaceInflatedBounds();
+				if (Box0.Intersects(Box1))
 				{
 					NarrowPhase.GenerateCollisions(Dt, Receiver, Particle0, Particle1, BoundsThickness, StatData);
 				}
diff --git a/Engine/Source/Runtime/Experimental/Chaos/Public/Chaos/Collision/SpatialAccelerationBroadPhase.h b/Engine/Source/Runtime/Experimental/Chaos/Public/Chaos/Collision/SpatialAccelerationBroadPhase.h
--- a/Engine/Source/Runtime/Experimental/Chaos/Public/Chaos/Collision/SpatialAccelerationBroadPhase.h
+++ b/Engine/Source/Runtime/Experimental/Chaos/Public/Chaos/Collision/SpatialAccelerationBroadPhase.h
@@ -144,9 +144,9 @@ namespace Chaos
 				{
 					// @todo(ccaulfield): COLLISION - see the NOTE below - fix it
 #if CHAOS_PARTICLEHANDLE_TODO
-					const TBox<FReal, 3> Box1 = InSpatialAcceleration.GetWorldSpaceBoundingBox(Particle1);
+					const TAABB<FReal, 3> Box1 = InSpatialAcceleration.GetWorldSpaceBoundingBox(Particle1);
 #else
-					const TBox<FReal, 3> Box1 = ComputeWorldSpaceBoundingBox(Particle1); // NOTE: this ignores the velocity expansion which is wrong
+					const TAABB<FReal, 3> Box1 = ComputeWorldSpaceBoundingBox(Particle1); // NOTE: this ignores the velocity expansion which is wrong
 #endif
 
 					CHAOS_COLLISION_STAT(StatData.RecordBoundsData(Box1));
diff --git a/Engine/Source/Runtime/Experimental/Chaos/Public/Chaos/Collision/StatsData.h b/Engine/Source/Runtime/Experimental/Chaos/Public/Chaos/Collision/StatsData.h
--- a/Engine/Source/Runtime/Experimental/Chaos/Public/Chaos/Collision/StatsData.h
+++ b/Engine/Source/Runtime/Experimental/Chaos/Public/Chaos/Collision/StatsData.h
@@ -104,7 +104,7 @@ namespace Chaos
 				++SimulatedParticles;
 			}
 
-			void RecordBoundsData(const TBox <float, 3>& Box1)
+			void RecordBoundsData(const TAABB <float, 3>& Box1)
 			{
 				BoundsDistribution.Record(Box1.Extents().GetMax());
 			}
@@ -213,7 +213,7 @@ namespace Chaos
 				CHAOS_COLLISION_STAT({ if (IsEnabled()) { StatImp->IncrementSimulatedParticles(); } });
 			}
 
-			FORCEINLINE void RecordBoundsData(const TBox<FReal, 3>& Box1)
+			FORCEINLINE void RecordBoundsData(const TAABB<FReal, 3>& Box1)
 			{
 				CHAOS_COLLISION_STAT({ if (IsEnabled()) { StatImp->RecordBoundsData(Box1); } });
 			}
diff --git a/Engine/Source/Runtime/Experimental/Chaos/Public/Chaos/CollisionConvexMesh.h b/Engine/Source/Runtime/Experimental/Chaos/Public/Chaos/CollisionConvexMesh.h
--- a/Engine/Source/Runtime/Experimental/Chaos/Public/Chaos/CollisionConvexMesh.h
+++ b/Engine/Source/Runtime/Experimental/Chaos/Public/Chaos/CollisionConvexMesh.h
@@ -20,11 +20,11 @@ namespace Chaos
 			return Cross.Size() > 1e-4;
 		}
 
-		static void Build(const TParticles<FReal, 3>& InParticles, TArray <TPlane<FReal, 3>>& OutPlanes, TArray<TArray<int32>>& OutFaceIndices, TParticles<FReal, 3>& OutSurfaceParticles, TBox<FReal, 3>& OutLocalBounds)
+		static void Build(const TParticles<FReal, 3>& InParticles, TArray <TPlane<FReal, 3>>& OutPlanes, TArray<TArray<int32>>& OutFaceIndices, TParticles<FReal, 3>& OutSurfaceParticles, TAABB<FReal, 3>& OutLocalBounds)
 		{
 			OutPlanes.Reset();
 			OutSurfaceParticles.Resize(0);
-			OutLocalBounds = TBox<FReal, 3>::EmptyBox();
+			OutLocalBounds = TAABB<FReal, 3>::EmptyAABB();
 
 			const uint32 NumParticles = InParticles.Size();
 			if(NumParticles == 0)
@@ -32,7 +32,7 @@ namespace Chaos
 				return;
 			}
 
-			OutLocalBounds = TBox<FReal, 3>(InParticles.X(0), InParticles.X(0));
+			OutLocalBounds = TAABB<FReal, 3>(InParticles.X(0), InParticles.X(0));
 			for(uint32 ParticleIndex = 0; ParticleIndex < NumParticles; ++ParticleIndex)
 			{
 				OutLocalBounds.GrowToInclude(InParticles.X(ParticleIndex));
@@ -161,7 +161,7 @@ namespace Chaos
 			return FString::Printf(TEXT("Planes %d, SurfaceParticles %d"), NumPlanes, NumParticles);
 		}
 
-		static CHAOS_API void Simplify(TArray <TPlane<FReal, 3>>& InOutPlanes, TArray<TArray<int32>>& InOutFaces, TParticles<FReal, 3>& InOutParticles, TBox<FReal, 3>& InOutLocalBounds)
+		static CHAOS_API void Simplify(TArray <TPlane<FReal, 3>>& InOutPlanes, TArray<TArray<int32>>& InOutFaces, TParticles<FReal, 3>& InOutParticles, TAABB<FReal, 3>& InOutLocalBounds)
 		{
 			struct TPair
 			{
diff --git a/Engine/Source/Runtime/Experimental/Chaos/Public/Chaos/CollisionResolution.h b/Engine/Source/Runtime/Experimental/Chaos/Public/Chaos/CollisionResolution.h
--- a/Engine/Source/Runtime/Experimental/Chaos/Public/Chaos/CollisionResolution.h
+++ b/Engine/Source/Runtime/Experimental/Chaos/Public/Chaos/CollisionResolution.h
@@ -10,7 +10,7 @@ namespace Chaos
 	class TSphere;
 
 	template <typename T, int d>
-	class TBox;
+	class TAABB;
 
 	template <typename T>
 	class TCapsule;
@@ -34,7 +34,7 @@ namespace Chaos
 			//
 
 			template <typename T, int d>
-			void CHAOS_API UpdateBoxBoxConstraint(const TBox<T, d>& Box1, const TRigidTransform<T, d>& Box1Transform, const TBox<T, d>& Box2, const TRigidTransform<T, d>& Box2Transform, const T Thickness, TRigidBodyPointContactConstraint<T, d>& Constraint);
+			void CHAOS_API UpdateBoxBoxConstraint(const TAABB<T, d>& Box1, const TRigidTransform<T, d>& Box1Transform, const TAABB<T, d>& Box2, const TRigidTransform<T, d>& Box2Transform, const T Thickness, TRigidBodyPointContactConstraint<T, d>& Constraint);
 
 			template<typename T, int d>
 			void CHAOS_API ConstructBoxBoxConstraints(TGeometryParticleHandle<T, d>* Particle0, TGeometryParticleHandle<T, d>* Particle1, const FImplicitObject* Implicit0, const FImplicitObject* Implicit1, const TRigidTransform<T, d>& Transform0, const TRigidTransform<T, d>& Transform1, const T Thickness, FCollisionConstraintsArray& NewConstraints);
@@ -45,7 +45,7 @@ namespace Chaos
 			//
 
 			template <typename T, int d>
-			bool CHAOS_API UpdateBoxPlaneConstraint(const TBox<T, d>& Box, const TRigidTransform<T, d>& BoxTransform, const TPlane<T, d>& Plane, const TRigidTransform<T, d>& PlaneTransform, const T Thickness, TRigidBodyPointContactConstraint<T, d>& Constraint);
+			bool CHAOS_API UpdateBoxPlaneConstraint(const TAABB<T, d>& Box, const TRigidTransform<T, d>& BoxTransform, const TPlane<T, d>& Plane, const TRigidTransform<T, d>& PlaneTransform, const T Thickness, TRigidBodyPointContactConstraint<T, d>& Constraint);
 
 			template<typename T, int d>
 			void CHAOS_API ConstructBoxPlaneConstraints(TGeometryParticleHandle<T, d>* Particle0, TGeometryParticleHandle<T, d>* Particle1, const FImplicitObject* Implicit0, const FImplicitObject* Implicit1, const TRigidTransform<T, d>& Transform0, const TRigidTransform<T, d>& Transform1, const T Thickness, FCollisionConstraintsArray& NewConstraints);
@@ -76,7 +76,7 @@ namespace Chaos
 			//
 
 			template <typename T, int d>
-			void CHAOS_API UpdateSphereBoxConstraint(const TSphere<T, d>& Sphere, const TRigidTransform<T, d>& SphereTransform, const TBox<T, d>& Box, const TRigidTransform<T, d>& BoxTransform, const T Thickness, TRigidBodyPointContactConstraint<T, d>& Constraint);
+			void CHAOS_API UpdateSphereBoxConstraint(const TSphere<T, d>& Sphere, const TRigidTransform<T, d>& SphereTransform, const TAABB<T, d>& Box, const TRigidTransform<T, d>& BoxTransform, const T Thickness, TRigidBodyPointContactConstraint<T, d>& Constraint);
 
 			template<typename T, int d>
 			void CHAOS_API ConstructSphereBoxConstraints(TGeometryParticleHandle<T, d>* Particle0, TGeometryParticleHandle<T, d>* Particle1, const FImplicitObject* Implicit0, const FImplicitObject* Implicit1, const TRigidTransform<T, d>& Transform0, const TRigidTransform<T, d>& Transform1, const T Thickness, FCollisionConstraintsArray& NewConstraints);
@@ -106,7 +106,7 @@ namespace Chaos
 			//
 
 			template <typename T, int d>
-			void CHAOS_API UpdateCapsuleBoxConstraint(const TCapsule<T>& A, const TRigidTransform<T, d>& ATransform, const TBox<T, d>& B, const TRigidTransform<T, d>& BTransform, const T Thickness, TRigidBodyPointContactConstraint<T, d>& Constraint);
+			void CHAOS_API UpdateCapsuleBoxConstraint(const TCapsule<T>& A, const TRigidTransform<T, d>& ATransform, const TAABB<T, d>& B, const TRigidTransform<T, d>& BTransform, const T Thickness, TRigidBodyPointContactConstraint<T, d>& Constraint);
 
 			template<typename T, int d>
 			void CHAOS_API ConstructCapsuleBoxConstraints(TGeometryParticleHandle<T, d>* Particle0, TGeometryParticleHandle<T, d>* Particle1, const FImplicitObject* Implicit0, const FImplicitObject* Implicit1, const TRigidTransform<T, d>& Transform0, const TRigidTransform<T, d>& Transform1, const T Thickness, FCollisionConstraintsArray& NewConstraints);
diff --git a/Engine/Source/Runtime/Experimental/Chaos/Public/Chaos/Convex.h b/Engine/Source/Runtime/Experimental/Chaos/Public/Chaos/Convex.h
--- a/Engine/Source/Runtime/Experimental/Chaos/Public/Chaos/Convex.h
+++ b/Engine/Source/Runtime/Experimental/Chaos/Public/Chaos/Convex.h
@@ -2,7 +2,7 @@
 #pragma once
 
 #include "Chaos/ImplicitObject.h"
-#include "Chaos/Box.h"
+#include "Chaos/AABB.h"
 #include "Chaos/MassProperties.h"
 #include "CollisionConvexMesh.h"
 #include "ChaosArchive.h"
@@ -40,7 +40,7 @@ namespace Chaos
 		    : FImplicitObject(EImplicitObject::IsConvex | EImplicitObject::HasBoundingBox, ImplicitObjectType::Convex)
 			, Planes(MoveTemp(InPlanes))
 		    , SurfaceParticles(MoveTemp(InSurfaceParticles))
-		    , LocalBoundingBox(TBox<FReal, 3>::EmptyBox())
+		    , LocalBoundingBox(TAABB<FReal, 3>::EmptyAABB())
 		{
 			for (uint32 ParticleIndex = 0; ParticleIndex < SurfaceParticles.Size(); ++ParticleIndex)
 			{
@@ -72,7 +72,7 @@ namespace Chaos
 			return ImplicitObjectType::Convex;
 		}
 
-		virtual const TBox<FReal, 3>& BoundingBox() const override
+		virtual const TAABB<FReal, 3>& BoundingBox() const override
 		{
 			return LocalBoundingBox;
 		}
@@ -231,7 +231,8 @@ namespace Chaos
 		FORCEINLINE void SerializeImp(FArchive& Ar)
 		{
 			FImplicitObject::SerializeImp(Ar);
-			Ar << Planes << SurfaceParticles << LocalBoundingBox;
+			Ar << Planes << SurfaceParticles;
+			TBox<FReal,3>::SerializeAsAABB(Ar, LocalBoundingBox);
 
 			Ar.UsingCustomVersion(FExternalPhysicsCustomObjectVersion::GUID);
 			if (Ar.CustomVer(FExternalPhysicsCustomObjectVersion::GUID) >= FExternalPhysicsCustomObjectVersion::AddConvexCenterOfMassAndVolume)
@@ -298,7 +299,7 @@ namespace Chaos
 	private:
 		TArray<TPlane<FReal, 3>> Planes;
 		TParticles<FReal, 3> SurfaceParticles;	//copy of the vertices that are just on the convex hull boundary
-		TBox<FReal, 3> LocalBoundingBox;
+		TAABB<FReal, 3> LocalBoundingBox;
 		float Volume;
 		FVec3 CenterOfMass;
 	};
diff --git a/Engine/Source/Runtime/Experimental/Chaos/Public/Chaos/Cylinder.h b/Engine/Source/Runtime/Experimental/Chaos/Public/Chaos/Cylinder.h
--- a/Engine/Source/Runtime/Experimental/Chaos/Public/Chaos/Cylinder.h
+++ b/Engine/Source/Runtime/Experimental/Chaos/Public/Chaos/Cylinder.h
@@ -28,7 +28,7 @@ namespace Chaos
 		    , MLocalBoundingBox(x1, x1)
 		{
 			MLocalBoundingBox.GrowToInclude(x2);
-			MLocalBoundingBox = TBox<T, 3>(MLocalBoundingBox.Min() - TVector<T, 3>(MRadius), MLocalBoundingBox.Max() + TVector<T, 3>(MRadius));
+			MLocalBoundingBox = TAABB<T, 3>(MLocalBoundingBox.Min() - TVector<T, 3>(MRadius), MLocalBoundingBox.Max() + TVector<T, 3>(MRadius));
 		}
 		TCylinder(const TCylinder<T>& Other)
 		    : FImplicitObject(EImplicitObject::FiniteConvex, ImplicitObjectType::Cylinder)
@@ -169,7 +169,7 @@ namespace Chaos
 			return SideDistance;
 		}
 
-		virtual const TBox<T, 3>& BoundingBox() const override { return MLocalBoundingBox; }
+		virtual const TAABB<T, 3>& BoundingBox() const override { return MLocalBoundingBox; }
 
 		T GetRadius() const { return MRadius; }
 		T GetHeight() const { return MHeight; }
@@ -218,7 +218,7 @@ namespace Chaos
 			Ar << MPlane2;
 			Ar << MHeight;
 			Ar << MRadius;
-			Ar << MLocalBoundingBox;
+			TBox<T, 3>::SerializeAsAABB(Ar, MLocalBoundingBox);
 		}
 
 	private:
@@ -269,7 +269,7 @@ namespace Chaos
 	private:
 		TPlane<T, 3> MPlane1, MPlane2;
 		T MHeight, MRadius;
-		TBox<T, 3> MLocalBoundingBox;
+		TAABB<T, 3> MLocalBoundingBox;
 	};
 
 	template<typename T>
diff --git a/Engine/Source/Runtime/Experimental/Chaos/Public/Chaos/ExternalCollisionData.h b/Engine/Source/Runtime/Experimental/Chaos/Public/Chaos/ExternalCollisionData.h
--- a/Engine/Source/Runtime/Experimental/Chaos/Public/Chaos/ExternalCollisionData.h
+++ b/Engine/Source/Runtime/Experimental/Chaos/Public/Chaos/ExternalCollisionData.h
@@ -188,7 +188,7 @@ namespace Chaos
 			, Particle(nullptr)
 			, ParticleIndex(INDEX_NONE)
 			, ParticleIndexMesh(INDEX_NONE)
-			, BoundingBox(TBox<T, d>(TVector<T, d>((T)0.0), TVector<T, d>((T)0.0)))
+			, BoundingBox(TAABB<T, d>(TVector<T, d>((T)0.0), TVector<T, d>((T)0.0)))
 		{}
 
 		TVector<T, d> Location;
@@ -199,7 +199,7 @@ namespace Chaos
 		int32 ParticleIndex; //#todo: remove this in favor of TGeometryParticle?
 		int32 ParticleIndexMesh; // If ParticleIndex points to a cluster then this index will point to an actual mesh in the cluster
 								 // It is important to be able to get extra data from the component
-		Chaos::TBox<T, d> BoundingBox;
+		Chaos::TAABB<T, d> BoundingBox;
 	};
 
 
@@ -294,7 +294,7 @@ namespace Chaos
 			, Mass((T)0.0)
 			, Particle(nullptr)
 			, ParticleIndexMesh(INDEX_NONE)
-			, BoundingBox(TBox<T, d>(TVector<T, d>((T)0.0), TVector<T, d>((T)0.0)))
+			, BoundingBox(TAABB<T, d>(TVector<T, d>((T)0.0), TVector<T, d>((T)0.0)))
 		{}
 
 		TTrailingData(TVector<T, d> InLocation
@@ -307,7 +307,7 @@ namespace Chaos
 			, float InBoundingboxExtentMin
 			, float InBoundingboxExtentMax
 			, int32 InSurfaceType
-			, Chaos::TBox<T, d>& InBoundingBox)
+			, Chaos::TAABB<T, d>& InBoundingBox)
 			: Location(InLocation)
 			, Velocity(InVelocity)
 			, AngularVelocity(InAngularVelocity)
@@ -324,7 +324,7 @@ namespace Chaos
 		TGeometryParticle<T, d>* Particle;
 		int32 ParticleIndexMesh; // If ParticleIndex points to a cluster then this index will point to an actual mesh in the cluster
 								 // It is important to be able to get extra data from the component
-		Chaos::TBox<T, d> BoundingBox;
+		Chaos::TAABB<T, d> BoundingBox;
 
 		friend inline uint32 GetTypeHash(const TTrailingData& Other)
 		{
diff --git a/Engine/Source/Runtime/Experimental/Chaos/Public/Chaos/GeometryParticles.h b/Engine/Source/Runtime/Experimental/Chaos/Public/Chaos/GeometryParticles.h
--- a/Engine/Source/Runtime/Experimental/Chaos/Public/Chaos/GeometryParticles.h
+++ b/Engine/Source/Runtime/Experimental/Chaos/Public/Chaos/GeometryParticles.h
@@ -241,12 +241,12 @@ namespace Chaos
 			MapImplicitShapes(Index);
 		}
 
-		CHAOS_API const TBox<T,d>& LocalBounds(const int32 Index) const
+		CHAOS_API const TAABB<T,d>& LocalBounds(const int32 Index) const
 		{
 			return MLocalBounds[Index];
 		}
 
-		CHAOS_API TBox<T, d>& LocalBounds(const int32 Index)
+		CHAOS_API TAABB<T, d>& LocalBounds(const int32 Index)
 		{
 			return MLocalBounds[Index];
 		}
@@ -293,12 +293,12 @@ namespace Chaos
 		}
 #endif
 
-		CHAOS_API const TBox<T, d>& WorldSpaceInflatedBounds(const int32 Index) const
+		CHAOS_API const TAABB<T, d>& WorldSpaceInflatedBounds(const int32 Index) const
 		{
 			return MWorldSpaceInflatedBounds[Index];
 		}
 
-		CHAOS_API void SetWorldSpaceInflatedBounds(const int32 Index, const TBox<T, d>& Bounds)
+		CHAOS_API void SetWorldSpaceInflatedBounds(const int32 Index, const TAABB<T, d>& Bounds)
 		{
 			MWorldSpaceInflatedBounds[Index] = Bounds;
 
@@ -352,8 +352,8 @@ namespace Chaos
 			Ar.UsingCustomVersion(FExternalPhysicsCustomObjectVersion::GUID);
 			if (Ar.CustomVer(FExternalPhysicsCustomObjectVersion::GUID) >= FExternalPhysicsCustomObjectVersion::SerializeParticleBounds)
 			{
-				Ar << MLocalBounds;
-				Ar << MWorldSpaceInflatedBounds;
+				TBox<FReal, 3>::SerializeAsAABBs(Ar, MLocalBounds);
+				TBox<FReal, 3>::SerializeAsAABBs(Ar, MWorldSpaceInflatedBounds);
 				Ar << MHasBounds;
 
 				if (Ar.CustomVer(FExternalPhysicsCustomObjectVersion::GUID) < FExternalPhysicsCustomObjectVersion::SerializeShapeWorldSpaceBounds)
@@ -374,7 +374,7 @@ namespace Chaos
 					{
 						MLocalBounds[Idx] = MGeometry[Idx]->BoundingBox();
 						//ignore velocity too, really just trying to get something reasonable)
-						SetWorldSpaceInflatedBounds(Idx, MLocalBounds[Idx].TransformedBox(TRigidTransform<T,d>(X(Idx), R(Idx))));
+						SetWorldSpaceInflatedBounds(Idx, MLocalBounds[Idx].TransformedAABB(TRigidTransform<T,d>(X(Idx), R(Idx))));
 					}
 				}
 			}
@@ -436,8 +436,8 @@ namespace Chaos
 		TArrayCollectionArray<TGeometryParticle<T, d>*> MGeometryParticle;
 		TArrayCollectionArray<TShapesArray<T,d>> MShapesArray;
 		TArrayCollectionArray<TMap<const FImplicitObject*, int32>> ImplicitShapeMap;
-		TArrayCollectionArray<TBox<T,d>> MLocalBounds;
-		TArrayCollectionArray<TBox<T, d>> MWorldSpaceInflatedBounds;
+		TArrayCollectionArray<TAABB<T,d>> MLocalBounds;
+		TArrayCollectionArray<TAABB<T, d>> MWorldSpaceInflatedBounds;
 		TArrayCollectionArray<bool> MHasBounds;
 		TArrayCollectionArray<FSpatialAccelerationIdx> MSpatialIdx;
 		TArrayCollectionArray<uint32> MHashResult;
diff --git a/Engine/Source/Runtime/Experimental/Chaos/Public/Chaos/HeightField.h b/Engine/Source/Runtime/Experimental/Chaos/Public/Chaos/HeightField.h
--- a/Engine/Source/Runtime/Experimental/Chaos/Public/Chaos/HeightField.h
+++ b/Engine/Source/Runtime/Experimental/Chaos/Public/Chaos/HeightField.h
@@ -91,9 +91,9 @@ namespace Chaos
 			return 0;
 		}
 
-		virtual const TBox<T, 3>& BoundingBox() const
+		virtual const TAABB<T, 3>& BoundingBox() const
 		{
-			CachedBounds = TBox<T, 3>(LocalBounds.Min() * GeomData.Scale, LocalBounds.Max() * GeomData.Scale);
+			CachedBounds = TAABB<T, 3>(LocalBounds.Min() * GeomData.Scale, LocalBounds.Max() * GeomData.Scale);
 			return CachedBounds;
 		}
 
@@ -128,7 +128,7 @@ namespace Chaos
 				Ar << FlatGrid;
 				Ar << FlattenedBounds.Min;
 				Ar << FlattenedBounds.Max;
-				Ar << LocalBounds;
+				TBox<FReal, 3>::SerializeAsAABB(Ar, LocalBounds);
 			}
 			else
 			{
@@ -286,12 +286,12 @@ namespace Chaos
 				, Max(0)
 			{}
 
-			explicit FBounds2D(const TBox<T, 3>& In3DBounds)
+			explicit FBounds2D(const TAABB<T, 3>& In3DBounds)
 			{
 				Set(In3DBounds);
 			}
 
-			void Set(const TBox<T, 3>& In3DBounds)
+			void Set(const TAABB<T, 3>& In3DBounds)
 			{
 				Min = {In3DBounds.Min()[0], In3DBounds.Min()[1]};
 				Max = {In3DBounds.Max()[0], In3DBounds.Max()[1]};
@@ -445,9 +445,9 @@ namespace Chaos
 		// Bounds in 2D of the whole heightfield, to clip queries against
 		FBounds2D FlattenedBounds;
 		// 3D bounds for the heightfield, for insertion to the scene structure
-		TBox<T, 3> LocalBounds;
+		TAABB<T, 3> LocalBounds;
 		// Cached when bounds are requested. Mutable to allow GetBounds to be logical const
-		mutable TBox<T, 3> CachedBounds;
+		mutable TAABB<T, 3> CachedBounds;
 
 		void CalcBounds();
 		void BuildQueryData();
diff --git a/Engine/Source/Runtime/Experimental/Chaos/Public/Chaos/ISpatialAcceleration.h b/Engine/Source/Runtime/Experimental/Chaos/Public/Chaos/ISpatialAcceleration.h
--- a/Engine/Source/Runtime/Experimental/Chaos/Public/Chaos/ISpatialAcceleration.h
+++ b/Engine/Source/Runtime/Experimental/Chaos/Public/Chaos/ISpatialAcceleration.h
@@ -64,7 +64,7 @@ struct CHAOS_API FQueryFastDataVoid : public FQueryFastData
 };
 
 template <typename T, int d>
-class TBox;
+class TAABB;
 
 template <typename T, int d>
 class TGeometryParticle;
@@ -94,14 +94,14 @@ template <typename TPayloadType>
 struct CHAOS_API TSpatialVisitorData
 {
 	TPayloadType Payload;
-	TSpatialVisitorData(const TPayloadType& InPayload, const bool bInHasBounds = false, const TBox<float, 3>& InBounds = TBox<float, 3>::ZeroBox())
+	TSpatialVisitorData(const TPayloadType& InPayload, const bool bInHasBounds = false, const TAABB<float, 3>& InBounds = TAABB<float, 3>::ZeroAABB())
 		: Payload(InPayload)
 #if !(UE_BUILD_TEST || UE_BUILD_SHIPPING)
 		, bHasBounds(bInHasBounds)
 		, Bounds(InBounds)
 	{ }
 	bool bHasBounds;
-	TBox<float, 3> Bounds;
+	TAABB<float, 3> Bounds;
 #else
 	{ }
 #endif
@@ -152,7 +152,7 @@ public:
 	
 	virtual ~ISpacialDebugDrawInterface() = default;
 
-	virtual void Box(const TBox<T, 3>& InBox, const TVector<T, 3>& InLinearColor, float InThickness) = 0;
+	virtual void Box(const TAABB<T, 3>& InBox, const TVector<T, 3>& InLinearColor, float InThickness) = 0;
 	virtual void Line(const TVector<T, 3>& InBegin, const TVector<T, 3>& InEnd, const TVector<T, 3>& InLinearColor, float InThickness)  = 0;
 
 };
@@ -174,12 +174,12 @@ template <typename TPayloadType, typename T>
 struct TPayloadBoundsElement
 {
 	TPayloadType Payload;
-	TBox<T, 3> Bounds;
+	TAABB<T, 3> Bounds;
 
 	void Serialize(FChaosArchive& Ar)
 	{
 		Ar << Payload;
-		Ar << Bounds;
+		TBox<T,3>::SerializeAsAABB(Ar, Bounds);
 	}
 
 	template <typename TPayloadType2>
@@ -187,7 +187,7 @@ struct TPayloadBoundsElement
 
 	bool HasBoundingBox() const { return true; }
 
-	const TBox<T, 3>& BoundingBox() const { return Bounds; }
+	const TAABB<T, 3>& BoundingBox() const { return Bounds; }
 };
 
 template <typename TPayloadType, typename T>
@@ -210,18 +210,18 @@ public:
 
 	virtual bool IsAsyncTimeSlicingComplete() { return AsyncTimeSlicingComplete; }
 	virtual void ProgressAsyncTimeSlicing(bool ForceBuildCompletion = false) {}
-	virtual TArray<TPayloadType> FindAllIntersections(const TBox<T, d>& Box) const { check(false); return TArray<TPayloadType>(); }
+	virtual TArray<TPayloadType> FindAllIntersections(const TAABB<T, d>& Box) const { check(false); return TArray<TPayloadType>(); }
 
 	virtual void Raycast(const TVector<T, d>& Start, const TVector<T, d>& Dir, const T Length, ISpatialVisitor<TPayloadType, T>& Visitor) const { check(false); }
 	virtual void Sweep(const TVector<T, d>& Start, const TVector<T, d>& Dir, const T Length, const TVector<T, d> QueryHalfExtents, ISpatialVisitor<TPayloadType, T>& Visitor) const { check(false);}
-	virtual void Overlap(const TBox<T, d>& QueryBounds, ISpatialVisitor<TPayloadType, T>& Visitor) const { check(false); }
+	virtual void Overlap(const TAABB<T, d>& QueryBounds, ISpatialVisitor<TPayloadType, T>& Visitor) const { check(false); }
 
 	virtual void RemoveElement(const TPayloadType& Payload)
 	{
 		check(false);	//not implemented
 	}
 
-	virtual void UpdateElement(const TPayloadType& Payload, const TBox<T, d>& NewBounds, bool bHasBounds)
+	virtual void UpdateElement(const TPayloadType& Payload, const TAABB<T, d>& NewBounds, bool bHasBounds)
 	{
 		check(false);
 	}
@@ -231,7 +231,7 @@ public:
 		RemoveElement(Payload);
 	}
 
-	virtual void UpdateElementIn(const TPayloadType& Payload, const TBox<T, d>& NewBounds, bool bHasBounds, FSpatialAccelerationIdx Idx)
+	virtual void UpdateElementIn(const TPayloadType& Payload, const TAABB<T, d>& NewBounds, bool bHasBounds, FSpatialAccelerationIdx Idx)
 	{
 		UpdateElement(Payload, NewBounds, bHasBounds);
 	}
diff --git a/Engine/Source/Runtime/Experimental/Chaos/Public/Chaos/ImplicitObject.h b/Engine/Source/Runtime/Experimental/Chaos/Public/Chaos/ImplicitObject.h
--- a/Engine/Source/Runtime/Experimental/Chaos/Public/Chaos/ImplicitObject.h
+++ b/Engine/Source/Runtime/Experimental/Chaos/Public/Chaos/ImplicitObject.h
@@ -10,7 +10,7 @@
 namespace Chaos
 {
 template<class T, int d>
-class TBox;
+class TAABB;
 template<class T>
 class TCylinder;
 template<class T, int d>
@@ -170,7 +170,7 @@ public:
 	// Explicitly non-virtual.  Must cast to derived types to target their implementation.
 	FVec3 Normal(const FVec3& x) const;
 	virtual FReal PhiWithNormal(const FVec3& x, FVec3& Normal) const = 0;
-	virtual const class TBox<FReal, 3>& BoundingBox() const;
+	virtual const class TAABB<FReal, 3>& BoundingBox() const;
 	bool HasBoundingBox() const { return bHasBoundingBox; }
 	bool IsConvex() const { return bIsConvex; }
 	void IgnoreAnalyticCollisions(const bool Ignore = true) { bIgnoreAnalyticCollisions = Ignore; }
@@ -243,7 +243,7 @@ public:
 		Out.Add(MakePair(This, ParentTM));
 	}
 
-	virtual void FindAllIntersectingObjects(TArray < Pair<const FImplicitObject*, FRigidTransform3>>& Out, const TBox<FReal, 3>& LocalBounds) const;
+	virtual void FindAllIntersectingObjects(TArray < Pair<const FImplicitObject*, FRigidTransform3>>& Out, const TAABB<FReal, 3>& LocalBounds) const;
 
 	virtual FString ToString() const
 	{
diff --git a/Engine/Source/Runtime/Experimental/Chaos/Public/Chaos/ImplicitObjectIntersection.h b/Engine/Source/Runtime/Experimental/Chaos/Public/Chaos/ImplicitObjectIntersection.h
--- a/Engine/Source/Runtime/Experimental/Chaos/Public/Chaos/ImplicitObjectIntersection.h
+++ b/Engine/Source/Runtime/Experimental/Chaos/Public/Chaos/ImplicitObjectIntersection.h
@@ -55,7 +55,7 @@ class TImplicitObjectIntersection : public FImplicitObject
 		return Phi;
 	}
 
-	virtual const TBox<T,d>& BoundingBox() const { return MLocalBoundingBox; }
+	virtual const TAABB<T,d>& BoundingBox() const { return MLocalBoundingBox; }
 
 
 	virtual uint32 GetTypeHash() const override
@@ -94,6 +94,6 @@ private:
 
   private:
 	TArray<TUniquePtr<FImplicitObject>> MObjects;
-	TBox<T, d> MLocalBoundingBox;
+	TAABB<T, d> MLocalBoundingBox;
 };
 }
diff --git a/Engine/Source/Runtime/Experimental/Chaos/Public/Chaos/ImplicitObjectScaled.h b/Engine/Source/Runtime/Experimental/Chaos/Public/Chaos/ImplicitObjectScaled.h
--- a/Engine/Source/Runtime/Experimental/Chaos/Public/Chaos/ImplicitObjectScaled.h
+++ b/Engine/Source/Runtime/Experimental/Chaos/Public/Chaos/ImplicitObjectScaled.h
@@ -94,7 +94,7 @@ public:
 	FORCEINLINE TVector<T, d> Support(const TVector<T, d>& Direction, const T Thickness) const { return MObject->Support(Direction, Thickness); }
 	FORCEINLINE TVector<T, d> Support2(const TVector<T, d>& Direction, const T Thickness) const { return MObject->Support2(Direction); }
 
-	virtual const TBox<T, d>& BoundingBox() const override { return MObject->BoundingBox(); }
+	virtual const TAABB<T, d>& BoundingBox() const override { return MObject->BoundingBox(); }
 
 	const ObjectType Object() const { return MObject; }
 
@@ -446,7 +446,7 @@ public:
 		UpdateBounds();
 	}
 
-	virtual const TBox<T, d>& BoundingBox() const override { return MLocalBoundingBox; }
+	virtual const TAABB<T, d>& BoundingBox() const override { return MLocalBoundingBox; }
 
 	const FReal GetVolume() const
 	{
@@ -472,7 +472,8 @@ public:
 	{
 		FChaosArchiveScopedMemory ScopedMemory(Ar, GetTypeName(), false);
 		FImplicitObject::SerializeImp(Ar);
-		Ar << MObject << MScale << MInvScale << MLocalBoundingBox;
+		Ar << MObject << MScale << MInvScale;
+		TBox<T,d>::SerializeAsAABB(Ar, MLocalBoundingBox);
 		ensure(MInternalThickness == 0);	//not supported: do we care?
 
 		Ar.UsingCustomVersion(FExternalPhysicsCustomObjectVersion::GUID);
@@ -503,7 +504,7 @@ private:
 	TVector<T, d> MScale;
 	TVector<T, d> MInvScale;
 	T MInternalThickness;	//Allows us to inflate the instance before the scale is applied. This is useful when sweeps need to apply a non scale on a geometry with uniform thickness
-	TBox<T, d> MLocalBoundingBox;
+	TAABB<T, d> MLocalBoundingBox;
 
 	//needed for serialization
 	TImplicitObjectScaled()
@@ -526,9 +527,9 @@ private:
 
 	void UpdateBounds()
 	{
-		const TBox<T, d> UnscaledBounds = MObject->BoundingBox();
+		const TAABB<T, d> UnscaledBounds = MObject->BoundingBox();
 		const TVector<T, d> Vector1 = UnscaledBounds.Min() * MScale;
-		MLocalBoundingBox = TBox<T, d>(Vector1, Vector1);	//need to grow it out one vector at a time in case scale is negative
+		MLocalBoundingBox = TAABB<T, d>(Vector1, Vector1);	//need to grow it out one vector at a time in case scale is negative
 		const TVector<T, d> Vector2 = UnscaledBounds.Max() *MScale;
 		MLocalBoundingBox.GrowToInclude(Vector2);
 	}
diff --git a/Engine/Source/Runtime/Experimental/Chaos/Public/Chaos/ImplicitObjectTransformed.h b/Engine/Source/Runtime/Experimental/Chaos/Public/Chaos/ImplicitObjectTransformed.h
--- a/Engine/Source/Runtime/Experimental/Chaos/Public/Chaos/ImplicitObjectTransformed.h
+++ b/Engine/Source/Runtime/Experimental/Chaos/Public/Chaos/ImplicitObjectTransformed.h
@@ -51,7 +51,7 @@ public:
 	    : FImplicitObject(EImplicitObject::HasBoundingBox, ImplicitObjectType::Transformed)
 	    , MObject(Object)
 	    , MTransform(InTransform)
-	    , MLocalBoundingBox(Object->BoundingBox().TransformedBox(InTransform))
+	    , MLocalBoundingBox(Object->BoundingBox().TransformedAABB(InTransform))
 	{
 		this->bIsConvex = Object->IsConvex();
 	}
@@ -66,7 +66,7 @@ public:
 	{
 		static_assert(bSerializable, "Non-serializable TImplicitObjectTransformed created with a UniquePtr");
 		this->MObject = FStorage::Convert(MObjectOwner);
-		this->MLocalBoundingBox = MObject->BoundingBox().TransformedBox(InTransform);
+		this->MLocalBoundingBox = MObject->BoundingBox().TransformedAABB(InTransform);
 		this->bIsConvex = MObject->IsConvex();
 	}
 
@@ -173,9 +173,9 @@ public:
 		TImplicitObjectTransformAccumulateSerializableHelper(Out, MObject, NewTM);
 	}
 
-	virtual void FindAllIntersectingObjects(TArray < Pair<const FImplicitObject*, TRigidTransform<T, d>>>& Out, const TBox<T, d>& LocalBounds) const
+	virtual void FindAllIntersectingObjects(TArray < Pair<const FImplicitObject*, TRigidTransform<T, d>>>& Out, const TAABB<T, d>& LocalBounds) const
 	{
-		const TBox<T, d> SubobjectBounds = LocalBounds.TransformedBox(MTransform.Inverse());
+		const TAABB<T, d> SubobjectBounds = LocalBounds.TransformedAABB(MTransform.Inverse());
 		int32 NumOut = Out.Num();
 		MObject->FindAllIntersectingObjects(Out, SubobjectBounds);
 		if (Out.Num() > NumOut)
@@ -184,7 +184,7 @@ public:
 		}
 	}
 
-	virtual const TBox<T, d>& BoundingBox() const override { return MLocalBoundingBox; }
+	virtual const TAABB<T, d>& BoundingBox() const override { return MLocalBoundingBox; }
 
 	const FReal GetVolume() const
 	{
@@ -214,7 +214,7 @@ public:
 		FImplicitObject::SerializeImp(Ar);
 		TImplicitObjectTransformSerializeHelper(Ar, MObject);
 		Ar << MTransform;
-		Ar << MLocalBoundingBox;
+		TBox<T, d>::SerializeAsAABB(Ar, MLocalBoundingBox);
 	}
 
 	virtual uint32 GetTypeHash() const override
@@ -232,7 +232,7 @@ private:
 	ObjectType MObject;
 	TUniquePtr<Chaos::FImplicitObject> MObjectOwner;
 	TRigidTransform<T, d> MTransform;
-	TBox<T, d> MLocalBoundingBox;
+	TAABB<T, d> MLocalBoundingBox;
 
 	//needed for serialization
 	TImplicitObjectTransformed() : FImplicitObject(EImplicitObject::HasBoundingBox, ImplicitObjectType::Transformed) {}
diff --git a/Engine/Source/Runtime/Experimental/Chaos/Public/Chaos/ImplicitObjectUnion.h b/Engine/Source/Runtime/Experimental/Chaos/Public/Chaos/ImplicitObjectUnion.h
--- a/Engine/Source/Runtime/Experimental/Chaos/Public/Chaos/ImplicitObjectUnion.h
+++ b/Engine/Source/Runtime/Experimental/Chaos/Public/Chaos/ImplicitObjectUnion.h
@@ -63,7 +63,7 @@ class TImplicitObjectUnion : public FImplicitObject
 		return Phi;
 	}
 
-	virtual const TBox<T, d>& BoundingBox() const override { return MLocalBoundingBox; }
+	virtual const TAABB<T, d>& BoundingBox() const override { return MLocalBoundingBox; }
 
 	virtual void AccumulateAllImplicitObjects(TArray<Pair<const FImplicitObject*, TRigidTransform<T, d>>>& Out, const TRigidTransform<T, d>& ParentTM) const
 	{
@@ -86,8 +86,8 @@ class TImplicitObjectUnion : public FImplicitObject
 		}
 	}
 
-	virtual void FindAllIntersectingObjects(TArray < Pair<const FImplicitObject*, TRigidTransform<T, d>>>& Out, const TBox<T, d>& LocalBounds) const;
-	TArray<int32> FindAllIntersectingChildren(const TBox<T, d>& LocalBounds) const;
+	virtual void FindAllIntersectingObjects(TArray < Pair<const FImplicitObject*, TRigidTransform<T, d>>>& Out, const TAABB<T, d>& LocalBounds) const;
+	TArray<int32> FindAllIntersectingChildren(const TAABB<T, d>& LocalBounds) const;
 	TArray<int32> FindAllIntersectingChildren(const TSpatialRay<T, d>& LocalRay) const;
 	virtual void CacheAllImplicitObjects();
 
@@ -181,7 +181,7 @@ private:
 	TArray<TUniquePtr<FImplicitObject>> MObjects;
 	TGeometryParticles<T, d> GeomParticles;
 	TBoundingVolumeHierarchy<TGeometryParticles<T, d>, TArray<int32>, T, d>* Hierarchy;
-	TBox<T, d> MLocalBoundingBox;
+	TAABB<T, d> MLocalBoundingBox;
 	bool bHierarchyBuilt;
 	TArray<int32> MOriginalParticleLookupHack;	//temp hack for finding original particles
 
diff --git a/Engine/Source/Runtime/Experimental/Chaos/Public/Chaos/Levelset.h b/Engine/Source/Runtime/Experimental/Chaos/Public/Chaos/Levelset.h
--- a/Engine/Source/Runtime/Experimental/Chaos/Public/Chaos/Levelset.h
+++ b/Engine/Source/Runtime/Experimental/Chaos/Public/Chaos/Levelset.h
@@ -35,7 +35,7 @@ class CHAOS_API TLevelSet final : public FImplicitObject
 	virtual T PhiWithNormal(const TVector<T, d>& x, TVector<T, d>& Normal) const override;
 	T SignedDistance(const TVector<T, d>& x) const;
 
-	virtual const TBox<T, d>& BoundingBox() const override { return MOriginalLocalBoundingBox; }
+	virtual const TAABB<T, d>& BoundingBox() const override { return MOriginalLocalBoundingBox; }
 
 	// Returns a const ref to the underlying phi grid
 	const TArrayND<T, d>& GetPhiArray() const { return MPhi; }
@@ -65,8 +65,8 @@ class CHAOS_API TLevelSet final : public FImplicitObject
 		Ar << MGrid;
 		Ar << MPhi;
 		Ar << MNormals;
-		Ar << MLocalBoundingBox;
-		Ar << MOriginalLocalBoundingBox;
+		TBox<FReal, 3>::SerializeAsAABB(Ar, MLocalBoundingBox);
+		TBox<FReal, 3>::SerializeAsAABB(Ar, MOriginalLocalBoundingBox);
 		Ar << MBandWidth;
 	}
 
@@ -120,8 +120,8 @@ class CHAOS_API TLevelSet final : public FImplicitObject
 	TUniformGrid<T, d> MGrid;
 	TArrayND<T, d> MPhi;
 	TArrayND<TVector<T, d>, d> MNormals;
-	TBox<T, d> MLocalBoundingBox;
-	TBox<T, d> MOriginalLocalBoundingBox;
+	TAABB<T, d> MLocalBoundingBox;
+	TAABB<T, d> MOriginalLocalBoundingBox;
 	int32 MBandWidth;
 private:
 	TLevelSet() : FImplicitObject(EImplicitObject::HasBoundingBox, ImplicitObjectType::LevelSet) {}	//needed for serialization
diff --git a/Engine/Source/Runtime/Experimental/Chaos/Public/Chaos/PBDCollisionConstraints.h b/Engine/Source/Runtime/Experimental/Chaos/Public/Chaos/PBDCollisionConstraints.h
--- a/Engine/Source/Runtime/Experimental/Chaos/Public/Chaos/PBDCollisionConstraints.h
+++ b/Engine/Source/Runtime/Experimental/Chaos/Public/Chaos/PBDCollisionConstraints.h
@@ -26,7 +26,7 @@ template <typename T, int d>
 class TBVHParticles;
 
 template <typename T, int d>
-class TBox;
+class TAABB;
 
 template<class T>
 class TChaosPhysicsMaterial;
diff --git a/Engine/Source/Runtime/Experimental/Chaos/Public/Chaos/PBDRigidClustering.h b/Engine/Source/Runtime/Experimental/Chaos/Public/Chaos/PBDRigidClustering.h
--- a/Engine/Source/Runtime/Experimental/Chaos/Public/Chaos/PBDRigidClustering.h
+++ b/Engine/Source/Runtime/Experimental/Chaos/Public/Chaos/PBDRigidClustering.h
@@ -314,7 +314,7 @@ void UpdateClusterMassProperties(TPBDRigidClusteredParticles<T, d>& Particles, c
 	const TArrayCollectionArray<FMultiChildProxyId>* MMultiChildProxyId = nullptr);
 
 template<typename T, int d>
-TArray<TVector<T, d>> CleanCollisionParticles(const TArray<TVector<T, d>>& Vertices, TBox<T, d> BBox, const float SnapDistance=0.01)
+TArray<TVector<T, d>> CleanCollisionParticles(const TArray<TVector<T, d>>& Vertices, TAABB<T, d> BBox, const float SnapDistance=0.01)
 {
 	const int32 NumPoints = Vertices.Num();
 	if (NumPoints <= 1)
@@ -384,7 +384,7 @@ TArray<TVector<T, d>> CleanCollisionParticles(const TArray<TVector<T, d>>& Verti
 	{
 		return TArray<TVector<T, d>>();
 	}
-	TBox<T, d> BBox(TBox<T, d>::EmptyBox());
+	TAABB<T, d> BBox(TAABB<T, d>::EmptyAABB());
 	for (const TVector<T, d>& Pt : Vertices)
 	{
 		BBox.GrowToInclude(Pt);
diff --git a/Engine/Source/Runtime/Experimental/Chaos/Public/Chaos/PBDRigidsEvolution.h b/Engine/Source/Runtime/Experimental/Chaos/Public/Chaos/PBDRigidsEvolution.h
--- a/Engine/Source/Runtime/Experimental/Chaos/Public/Chaos/PBDRigidsEvolution.h
+++ b/Engine/Source/Runtime/Experimental/Chaos/Public/Chaos/PBDRigidsEvolution.h
@@ -132,8 +132,8 @@ public:
 	bool HasBounds(const int32 Idx) const { return MHasBoundingBoxes[Idx]; }
 	bool& HasBounds(const int32 Idx) { return MHasBoundingBoxes[Idx]; }
 
-	const TBox<T,d>& Bounds(const int32 Idx) const { return MBounds[Idx]; }
-	TBox<T, d>& Bounds(const int32 Idx) { return MBounds[Idx]; }
+	const TAABB<T,d>& Bounds(const int32 Idx) const { return MBounds[Idx]; }
+	TAABB<T, d>& Bounds(const int32 Idx) { return MBounds[Idx]; }
 
 	const TAccelerationStructureHandle<T, d>& Payload(const int32 Idx) const { return MPayloads[Idx]; }
 	TAccelerationStructureHandle<T, d>& Payload(const int32 Idx) { return MPayloads[Idx]; }
@@ -147,7 +147,7 @@ private:
 	}
 
 	TArrayCollectionArray<bool> MHasBoundingBoxes;
-	TArrayCollectionArray<TBox<T, d>> MBounds;
+	TArrayCollectionArray<TAABB<T, d>> MBounds;
 	TArrayCollectionArray<TAccelerationStructureHandle<T, d>> MPayloads;
 
 #if PARTICLE_ITERATOR_RANGED_FOR_CHECK
@@ -178,7 +178,7 @@ struct TSpatialAccelerationCacheHandle
 		return Cache->HasBounds(EntryIdx);
 	}
 
-	const TBox<T, d>& BoundingBox() const
+	const TAABB<T, d>& BoundingBox() const
 	{
 		return Cache->Bounds(EntryIdx);
 	}
diff --git a/Engine/Source/Runtime/Experimental/Chaos/Public/Chaos/PBDRigidsEvolutionGBF.h b/Engine/Source/Runtime/Experimental/Chaos/Public/Chaos/PBDRigidsEvolutionGBF.h
--- a/Engine/Source/Runtime/Experimental/Chaos/Public/Chaos/PBDRigidsEvolutionGBF.h
+++ b/Engine/Source/Runtime/Experimental/Chaos/Public/Chaos/PBDRigidsEvolutionGBF.h
@@ -177,8 +177,8 @@ public:
 
 				if (Particle.HasBounds())
 				{
-					const TBox<T, d>& LocalBounds = Particle.LocalBounds();
-					TBox<T, d> WorldSpaceBounds = LocalBounds.TransformedBox(TRigidTransform<T, d>(Particle.P(), Particle.Q()));
+					const TAABB<T, d>& LocalBounds = Particle.LocalBounds();
+					TAABB<T, d> WorldSpaceBounds = LocalBounds.TransformedAABB(TRigidTransform<T, d>(Particle.P(), Particle.Q()));
 					WorldSpaceBounds.ThickenSymmetrically(Particle.V());
 					Particle.SetWorldSpaceInflatedBounds(WorldSpaceBounds);
 				}
diff --git a/Engine/Source/Runtime/Experimental/Chaos/Public/Chaos/ParticleHandle.h b/Engine/Source/Runtime/Experimental/Chaos/Public/Chaos/ParticleHandle.h
--- a/Engine/Source/Runtime/Experimental/Chaos/Public/Chaos/ParticleHandle.h
+++ b/Engine/Source/Runtime/Experimental/Chaos/Public/Chaos/ParticleHandle.h
@@ -379,11 +379,11 @@ public:
 
 	const TShapesArray<T,d>& ShapesArray() const { return GeometryParticles->ShapesArray(ParticleIdx); }
 
-	const TBox<T, d>& LocalBounds() const { return GeometryParticles->LocalBounds(ParticleIdx); }
-	void SetLocalBounds(const TBox<T, d>& NewBounds) { GeometryParticles->LocalBounds(ParticleIdx) = NewBounds; }
+	const TAABB<T, d>& LocalBounds() const { return GeometryParticles->LocalBounds(ParticleIdx); }
+	void SetLocalBounds(const TAABB<T, d>& NewBounds) { GeometryParticles->LocalBounds(ParticleIdx) = NewBounds; }
 
-	const TBox<T, d>& WorldSpaceInflatedBounds() const { return GeometryParticles->WorldSpaceInflatedBounds(ParticleIdx); }
-	void SetWorldSpaceInflatedBounds(const TBox<T, d>& WorldSpaceInflatedBounds)
+	const TAABB<T, d>& WorldSpaceInflatedBounds() const { return GeometryParticles->WorldSpaceInflatedBounds(ParticleIdx); }
+	void SetWorldSpaceInflatedBounds(const TAABB<T, d>& WorldSpaceInflatedBounds)
 	{
 		GeometryParticles->SetWorldSpaceInflatedBounds(ParticleIdx, WorldSpaceInflatedBounds);
 	}
diff --git a/Engine/Source/Runtime/Experimental/Chaos/Public/Chaos/SpatialAccelerationCollection.h b/Engine/Source/Runtime/Experimental/Chaos/Public/Chaos/SpatialAccelerationCollection.h
--- a/Engine/Source/Runtime/Experimental/Chaos/Public/Chaos/SpatialAccelerationCollection.h
+++ b/Engine/Source/Runtime/Experimental/Chaos/Public/Chaos/SpatialAccelerationCollection.h
@@ -223,7 +223,7 @@ struct TSpatialAccelerationCollectionHelper
 	}
 
 	template <typename SQVisitor>
-	static bool OverlapFast(const Tuple& Types, const TBox<T, d> QueryBounds, SQVisitor& Visitor)
+	static bool OverlapFast(const Tuple& Types, const TAABB<T, d> QueryBounds, SQVisitor& Visitor)
 	{
 		const auto& Accelerations = GetAccelerationsPerType<TypeIdx>(Types).Objects;
 		for (const auto& Accelerator : Accelerations)
@@ -415,14 +415,14 @@ public:
 		TSpatialAccelerationCollectionHelper<0, NumTypes, decltype(Types), TPayloadType, T, d>::SweepFast(Types, Start, QueryFastData, QueryHalfExtents, Visitor);
 	}
 
-	virtual void Overlap(const TBox<T, d>& QueryBounds, ISpatialVisitor<TPayloadType, T>& Visitor) const override
+	virtual void Overlap(const TAABB<T, d>& QueryBounds, ISpatialVisitor<TPayloadType, T>& Visitor) const override
 	{
 		TSpatialVisitor<TPayloadType, T> ProxyVisitor(Visitor);
 		Overlap(QueryBounds, ProxyVisitor);
 	}
 
 	template <typename SQVisitor>
-	void Overlap(const TBox<T, 3>& QueryBounds, SQVisitor& Visitor) const
+	void Overlap(const TAABB<T, 3>& QueryBounds, SQVisitor& Visitor) const
 	{
 		TSpatialAccelerationCollectionHelper<0, NumTypes, decltype(Types), TPayloadType, T, d>::OverlapFast(Types, QueryBounds, Visitor);
 	}
@@ -460,7 +460,7 @@ public:
 		Buckets[UseBucket].Objects[SpatialIdx.InnerIdx].Acceleration->RemoveElement(Payload);
 	}
 
-	virtual void UpdateElementIn(const TPayloadType& Payload, const TBox<T, d>& NewBounds, bool bHasBounds, FSpatialAccelerationIdx SpatialIdx)
+	virtual void UpdateElementIn(const TPayloadType& Payload, const TAABB<T, d>& NewBounds, bool bHasBounds, FSpatialAccelerationIdx SpatialIdx)
 	{
 		const uint16 UseBucket = ((1 << SpatialIdx.Bucket) & this->ActiveBucketsMask) ? SpatialIdx.Bucket : 0;
 		Buckets[UseBucket].Objects[SpatialIdx.InnerIdx].Acceleration->UpdateElement(Payload, NewBounds, bHasBounds);
diff --git a/Engine/Source/Runtime/Experimental/Chaos/Public/Chaos/SpatialHash.h b/Engine/Source/Runtime/Experimental/Chaos/Public/Chaos/SpatialHash.h
--- a/Engine/Source/Runtime/Experimental/Chaos/Public/Chaos/SpatialHash.h
+++ b/Engine/Source/Runtime/Experimental/Chaos/Public/Chaos/SpatialHash.h
@@ -66,7 +66,7 @@ private:
 private:
 	TArray<TVector<T, 3>> MParticles;
 	T MCellSize;
-	TBox<T, 3> MBoundingBox;
+	TAABB<T, 3> MBoundingBox;
 	int32 MNumberOfCellsX, MNumberOfCellsY, MNumberOfCellsZ;
 	TMap<int32, TArray<int32>> MHashTable;
 };
diff --git a/Engine/Source/Runtime/Experimental/Chaos/Public/Chaos/Sphere.h b/Engine/Source/Runtime/Experimental/Chaos/Public/Chaos/Sphere.h
--- a/Engine/Source/Runtime/Experimental/Chaos/Public/Chaos/Sphere.h
+++ b/Engine/Source/Runtime/Experimental/Chaos/Public/Chaos/Sphere.h
@@ -198,7 +198,7 @@ namespace Chaos
 			return Radius;
 		}
 
-		virtual const TBox<T, d>& BoundingBox() const 
+		virtual const TAABB<T, d>& BoundingBox() const 
 		{
 			return LocalBoundingBox; 
 		}
@@ -253,7 +253,7 @@ namespace Chaos
 			Ar << Center << Radius;
 			if (Ar.IsLoading())
 			{
-				LocalBoundingBox = TBox<T, d>(Center - Radius, Center + Radius);
+				LocalBoundingBox = TAABB<T, d>(Center - Radius, Center + Radius);
 			}
 		}
 
@@ -333,7 +333,7 @@ namespace Chaos
 	private:
 		TVector<T, d> Center;
 		T Radius;
-		TBox<T, d> LocalBoundingBox;
+		TAABB<T, d> LocalBoundingBox;
 
 	private:
 
diff --git a/Engine/Source/Runtime/Experimental/Chaos/Public/Chaos/TaperedCylinder.h b/Engine/Source/Runtime/Experimental/Chaos/Public/Chaos/TaperedCylinder.h
--- a/Engine/Source/Runtime/Experimental/Chaos/Public/Chaos/TaperedCylinder.h
+++ b/Engine/Source/Runtime/Experimental/Chaos/Public/Chaos/TaperedCylinder.h
@@ -32,7 +32,7 @@ namespace Chaos
 			T MaxRadius = MRadius1;
 			if (MaxRadius < MRadius2)
 				MaxRadius = MRadius2;
-			MLocalBoundingBox = TBox<T, 3>(MLocalBoundingBox.Min() - TVector<T, 3>(MaxRadius), MLocalBoundingBox.Max() + TVector<T, 3>(MaxRadius));
+			MLocalBoundingBox = TAABB<T, 3>(MLocalBoundingBox.Min() - TVector<T, 3>(MaxRadius), MLocalBoundingBox.Max() + TVector<T, 3>(MaxRadius));
 		}
 		TTaperedCylinder(const TTaperedCylinder<T>& Other)
 		    : FImplicitObject(EImplicitObject::FiniteConvex, ImplicitObjectType::TaperedCylinder)
@@ -112,7 +112,7 @@ namespace Chaos
 		TArray<TVector<T, 3>> ComputeSamplePoints(const T PointsPerUnitArea, const bool IncludeEndCaps = true, const int32 MinPoints = 0, const int32 MaxPoints = 1000) const
 		{ return ComputeSamplePoints(FMath::Clamp(static_cast<int32>(ceil(PointsPerUnitArea * GetArea(IncludeEndCaps))), MinPoints, MaxPoints), IncludeEndCaps); }
 
-		virtual const TBox<T, 3>& BoundingBox() const override { return MLocalBoundingBox; }
+		virtual const TAABB<T, 3>& BoundingBox() const override { return MLocalBoundingBox; }
 
 		T PhiWithNormal(const TVector<T, 3>& x, TVector<T, 3>& Normal) const
 		{
@@ -348,7 +348,7 @@ namespace Chaos
 
 		TPlane<T, 3> MPlane1, MPlane2;
 		T MHeight, MRadius1, MRadius2;
-		TBox<T, 3> MLocalBoundingBox;
+		TAABB<T, 3> MLocalBoundingBox;
 	};
 
 	template<typename T>
diff --git a/Engine/Source/Runtime/Experimental/Chaos/Public/Chaos/Triangle.h b/Engine/Source/Runtime/Experimental/Chaos/Public/Chaos/Triangle.h
--- a/Engine/Source/Runtime/Experimental/Chaos/Public/Chaos/Triangle.h
+++ b/Engine/Source/Runtime/Experimental/Chaos/Public/Chaos/Triangle.h
@@ -155,7 +155,7 @@ namespace Chaos
 			: FImplicitObject(EImplicitObject::IsConvex | EImplicitObject::HasBoundingBox, ImplicitObjectType::Triangle)
 			, Tri(InA, InB, InC)
 		{
-			Bounds = TBox<T, 3>(Tri[0], Tri[0]);
+			Bounds = TAABB<T, 3>(Tri[0], Tri[0]);
 			Bounds.GrowToInclude(Tri[1]);
 			Bounds.GrowToInclude(Tri[2]);
 		}
@@ -190,7 +190,7 @@ namespace Chaos
 			return Tri.PhiWithNormal(InSamplePoint, OutNormal);
 		}
 
-		virtual const class TBox<T, 3>& BoundingBox() const override
+		virtual const class TAABB<T, 3>& BoundingBox() const override
 		{
 			return Bounds;
 		}
@@ -238,6 +238,6 @@ namespace Chaos
 	private:
 
 		TTriangle<T> Tri;
-		TBox<T, 3> Bounds;
+		TAABB<T, 3> Bounds;
 	};
 }
\ No newline at end of file
diff --git a/Engine/Source/Runtime/Experimental/Chaos/Public/Chaos/TriangleMeshImplicitObject.h b/Engine/Source/Runtime/Experimental/Chaos/Public/Chaos/TriangleMeshImplicitObject.h
--- a/Engine/Source/Runtime/Experimental/Chaos/Public/Chaos/TriangleMeshImplicitObject.h
+++ b/Engine/Source/Runtime/Experimental/Chaos/Public/Chaos/TriangleMeshImplicitObject.h
@@ -64,7 +64,7 @@ namespace Chaos
 		virtual int32 FindMostOpposingFace(const TVector<T, 3>& Position, const TVector<T, 3>& UnitDir, int32 HintFaceIndex, T SearchDistance) const override;
 		virtual TVector<T, 3> FindGeometryOpposingNormal(const TVector<T, 3>& DenormDir, int32 FaceIndex, const TVector<T, 3>& OriginalNormal) const override;
 
-		virtual const TBox<T, 3>& BoundingBox() const
+		virtual const TAABB<T, 3>& BoundingBox() const
 		{
 			return MLocalBoundingBox;
 		}
@@ -81,7 +81,7 @@ namespace Chaos
 			FImplicitObject::SerializeImp(Ar);
 			Ar << MParticles;
 			Ar << MElements;
-			Ar << MLocalBoundingBox;
+			TBox<T, 3>::SerializeAsAABB(Ar, MLocalBoundingBox);
 
 			if (Ar.CustomVer(FExternalPhysicsCustomObjectVersion::GUID) < FExternalPhysicsCustomObjectVersion::RemovedConvexHullsFromTriangleMeshImplicitObject)
 			{
@@ -151,7 +151,7 @@ namespace Chaos
 
 		TParticles<T, 3> MParticles;
 		TArray<TVector<int32, 3>> MElements;
-		TBox<T, 3> MLocalBoundingBox;
+		TAABB<T, 3> MLocalBoundingBox;
 		TArray<uint16> MaterialIndices;
 
 		//using BVHType = TBoundingVolume<int32, T, 3>;
@@ -167,9 +167,9 @@ namespace Chaos
 
 			bool HasBoundingBox() const { return true; }
 
-			TBox<T, 3> BoundingBox() const
+			TAABB<T, 3> BoundingBox() const
 			{
-				TBox<T, 3> Bounds(TmData->MParticles.X(TmData->MElements[Index][0]), TmData->MParticles.X(TmData->MElements[Index][0]));
+				TAABB<T, 3> Bounds(TmData->MParticles.X(TmData->MElements[Index][0]), TmData->MParticles.X(TmData->MElements[Index][0]));
 
 				Bounds.GrowToInclude(TmData->MParticles.X(TmData->MElements[Index][1]));
 				Bounds.GrowToInclude(TmData->MParticles.X(TmData->MElements[Index][2]));
diff --git a/Engine/Source/Runtime/Experimental/ChaosSolvers/Private/ChaosSolversModule.cpp b/Engine/Source/Runtime/Experimental/ChaosSolvers/Private/ChaosSolversModule.cpp
--- a/Engine/Source/Runtime/Experimental/ChaosSolvers/Private/ChaosSolversModule.cpp
+++ b/Engine/Source/Runtime/Experimental/ChaosSolvers/Private/ChaosSolversModule.cpp
@@ -594,7 +594,7 @@ void FChaosSolversModule::DumpHierarchyStats(int32* OutOptMaxCellElements)
 #endif
 #if 0
 
-		const TArray<Chaos::TBox<float, 3>>& Boxes = Hierarchy->GetWorldSpaceBoxes();
+		const TArray<Chaos::TAABB<float, 3>>& Boxes = Hierarchy->GetWorldSpaceBoxes();
 
 		if(Boxes.Num() > 0)
 		{
diff --git a/Engine/Source/Runtime/Experimental/ChaosSolvers/Private/PhysicsProxy/GeometryCollectionPhysicsProxy.cpp b/Engine/Source/Runtime/Experimental/ChaosSolvers/Private/PhysicsProxy/GeometryCollectionPhysicsProxy.cpp
--- a/Engine/Source/Runtime/Experimental/ChaosSolvers/Private/PhysicsProxy/GeometryCollectionPhysicsProxy.cpp
+++ b/Engine/Source/Runtime/Experimental/ChaosSolvers/Private/PhysicsProxy/GeometryCollectionPhysicsProxy.cpp
@@ -3589,7 +3589,7 @@ void FGeometryCollectionPhysicsProxy::UpdateTrailingData(const FParticlesType& P
 									!Particles.W(IdxParticle).ContainsNaN() &&
 									FMath::IsFinite(Particles.M(IdxParticle)))
 								{
-									Chaos::TBox<float, 3> BoundingBox = Particles.Geometry(IdxParticle)->BoundingBox();
+									Chaos::TAABB<float, 3> BoundingBox = Particles.Geometry(IdxParticle)->BoundingBox();
 									Chaos::TVector<float, 3> Extents = BoundingBox.Extents();
 									float Volume = Extents[0] * Extents[1] * Extents[2];
 									float SpeedSquared = Particles.V(IdxParticle).SizeSquared();
diff --git a/Engine/Source/Runtime/Experimental/ChaosSolvers/Private/PhysicsProxy/SingleParticlePhysicsProxy.cpp b/Engine/Source/Runtime/Experimental/ChaosSolvers/Private/PhysicsProxy/SingleParticlePhysicsProxy.cpp
--- a/Engine/Source/Runtime/Experimental/ChaosSolvers/Private/PhysicsProxy/SingleParticlePhysicsProxy.cpp
+++ b/Engine/Source/Runtime/Experimental/ChaosSolvers/Private/PhysicsProxy/SingleParticlePhysicsProxy.cpp
@@ -57,7 +57,7 @@ void FSingleParticlePhysicsProxy<Chaos::TGeometryParticle<float, 3>>::PushToPhys
 		{
 			RigidHandle->SetHasBounds(true);
 			RigidHandle->SetLocalBounds(Data->Geometry->BoundingBox());
-			RigidHandle->SetWorldSpaceInflatedBounds(Data->Geometry->BoundingBox().TransformedBox(Chaos::TRigidTransform<float, 3>(Data->X, Data->R)));
+			RigidHandle->SetWorldSpaceInflatedBounds(Data->Geometry->BoundingBox().TransformedAABB(Chaos::TRigidTransform<float, 3>(Data->X, Data->R)));
 		}
 		RigidHandle->SetSpatialIdx(Data->SpatialIdx);	//todo: this needs to only happen once during initialization
 		RigidHandle->SetHashResultLowLevel(Data->HashResult);
@@ -131,7 +131,7 @@ void FSingleParticlePhysicsProxy<Chaos::TKinematicGeometryParticle<float, 3>>::P
 		{
 			RigidHandle->SetHasBounds(true);
 			RigidHandle->SetLocalBounds(Data->Geometry->BoundingBox());
-			Chaos::TBox<float, 3> WorldSpaceBox = Data->Geometry->BoundingBox().TransformedBox(Chaos::TRigidTransform<float, 3>(Data->X, Data->R));
+			Chaos::TAABB<float, 3> WorldSpaceBox = Data->Geometry->BoundingBox().TransformedAABB(Chaos::TRigidTransform<float, 3>(Data->X, Data->R));
 			WorldSpaceBox.ThickenSymmetrically(Data->MV);
 			RigidHandle->SetWorldSpaceInflatedBounds(WorldSpaceBox);
 		}
@@ -253,7 +253,7 @@ void FSingleParticlePhysicsProxy<Chaos::TPBDRigidParticle<float, 3>>::PushToPhys
 			{
 				RigidHandle->SetHasBounds(true);
 				RigidHandle->SetLocalBounds(Data->Geometry->BoundingBox());
-				Chaos::TBox<float, 3> WorldSpaceBox = Data->Geometry->BoundingBox().TransformedBox(Chaos::TRigidTransform<float, 3>(Data->X, Data->R));
+				Chaos::TAABB<float, 3> WorldSpaceBox = Data->Geometry->BoundingBox().TransformedAABB(Chaos::TRigidTransform<float, 3>(Data->X, Data->R));
 				WorldSpaceBox.ThickenSymmetrically(Data->MV);
 				RigidHandle->SetWorldSpaceInflatedBounds(WorldSpaceBox);
 			}
diff --git a/Engine/Source/Runtime/Experimental/ChaosSolvers/Private/PhysicsProxy/SkeletalMeshPhysicsProxy.cpp b/Engine/Source/Runtime/Experimental/ChaosSolvers/Private/PhysicsProxy/SkeletalMeshPhysicsProxy.cpp
--- a/Engine/Source/Runtime/Experimental/ChaosSolvers/Private/PhysicsProxy/SkeletalMeshPhysicsProxy.cpp
+++ b/Engine/Source/Runtime/Experimental/ChaosSolvers/Private/PhysicsProxy/SkeletalMeshPhysicsProxy.cpp
@@ -138,7 +138,7 @@ void FSkeletalMeshPhysicsProxy::CreateRigidBodyCallback(FParticlesType& Particle
 
 			const bool DoCollisionGeom = Parameters.CollisionType == ECollisionTypeEnum::Chaos_Surface_Volumetric;
 			TArray<Chaos::TVector<float, 3>>* SamplePoints = nullptr;
-			Chaos::TBox<float, 3> SamplePointsBBox = Chaos::TBox<float, 3>::EmptyBox();
+			Chaos::TAABB<float, 3> SamplePointsBBox = Chaos::TAABB<float, 3>::EmptyAABB();
 			if (DoCollisionGeom)
 			{
 				SamplePoints =
@@ -174,7 +174,7 @@ void FSkeletalMeshPhysicsProxy::CreateRigidBodyCallback(FParticlesType& Particle
 				// Precision gets worse in opt builds, so only do the phi test in debug.
 				checkSlow(PointVolumeRegistrationCheck(*SamplePoints, *ImplicitObject, 0.01));
 				auto PointBBoxCheck =
-					[&](const auto &InSamplePoints, const Chaos::TBox<float,3> &BBox) -> bool
+					[&](const auto &InSamplePoints, const Chaos::TAABB<float,3> &BBox) -> bool
 					{
 						for (int32 i = 0; i < InSamplePoints.Num(); i++)
 						{
@@ -190,7 +190,7 @@ void FSkeletalMeshPhysicsProxy::CreateRigidBodyCallback(FParticlesType& Particle
 				check(PointBBoxCheck(*SamplePoints, ImplicitObject->BoundingBox()));
 			}
 
-			const Chaos::TBox<float, 3>& BBox = ImplicitObject->BoundingBox();
+			const Chaos::TAABB<float, 3>& BBox = ImplicitObject->BoundingBox();
 
 			const FVector Scale = WorldTransform->GetScale3D();
 			const FVector &CenterOfMass = MassProperties.CenterOfMass;
diff --git a/Engine/Source/Runtime/Experimental/ChaosSolvers/Private/PhysicsProxy/StaticMeshPhysicsObject.cpp b/Engine/Source/Runtime/Experimental/ChaosSolvers/Private/PhysicsProxy/StaticMeshPhysicsObject.cpp
--- a/Engine/Source/Runtime/Experimental/ChaosSolvers/Private/PhysicsProxy/StaticMeshPhysicsObject.cpp
+++ b/Engine/Source/Runtime/Experimental/ChaosSolvers/Private/PhysicsProxy/StaticMeshPhysicsObject.cpp
@@ -162,7 +162,7 @@ void FStaticMeshPhysicsProxy::CreateRigidBodyCallback(FParticlesType& Particles)
 		else if (Parameters.ShapeType == EImplicitTypeEnum::Chaos_Implicit_Sphere)
 		{
 			Chaos::TSphere<float,3>* Sphere = new Chaos::TSphere<float, 3>(Chaos::TVector<float, 3>(0), Parameters.ShapeParams.SphereRadius);
-			const Chaos::TBox<float, 3> BBox = Sphere->BoundingBox();
+			const Chaos::TAABB<float, 3> BBox = Sphere->BoundingBox();
 			Bounds.Min = BBox.Min();
 			Bounds.Max = BBox.Max();
 			Particles.SetDynamicGeometry(RigidBodyId, TUniquePtr<Chaos::FImplicitObject>(Sphere));
@@ -207,7 +207,7 @@ void FStaticMeshPhysicsProxy::CreateRigidBodyCallback(FParticlesType& Particles)
 			Chaos::TVector<float, 3> x1(0, -Parameters.ShapeParams.CapsuleHalfHeightAndRadius.X, 0);
 			Chaos::TVector<float, 3> x2(0, Parameters.ShapeParams.CapsuleHalfHeightAndRadius.X, 0);
 			Chaos::TCapsule<float>* Capsule = new Chaos::TCapsule<float>(x1, x2, Parameters.ShapeParams.CapsuleHalfHeightAndRadius.Y);
-			const Chaos::TBox<float, 3> BBox = Capsule->BoundingBox();
+			const Chaos::TAABB<float, 3> BBox = Capsule->BoundingBox();
 			Bounds.Min = BBox.Min();
 			Bounds.Max = BBox.Max();
 			Particles.SetDynamicGeometry(RigidBodyId, TUniquePtr<Chaos::FImplicitObject>(Capsule));
diff --git a/Engine/Source/Runtime/Experimental/ChaosSolvers/Public/PhysicsProxy/AnalyticImplicitGroup.h b/Engine/Source/Runtime/Experimental/ChaosSolvers/Public/PhysicsProxy/AnalyticImplicitGroup.h
--- a/Engine/Source/Runtime/Experimental/ChaosSolvers/Public/PhysicsProxy/AnalyticImplicitGroup.h
+++ b/Engine/Source/Runtime/Experimental/ChaosSolvers/Public/PhysicsProxy/AnalyticImplicitGroup.h
@@ -123,7 +123,7 @@ public:
 
 		const int32 Num = NumStructures();
 		TArray<Chaos::TMassProperties<float, 3>> MPArray;
-		TArray<Chaos::TBox<float, 3>> BBoxes;
+		TArray<Chaos::TAABB<float, 3>> BBoxes;
 		MPArray.SetNum(Num);
 		BBoxes.SetNum(Num);
 
@@ -131,7 +131,7 @@ public:
 		for (Chaos::TSphere<float, 3>* Sphere : Spheres)
 		{
 			const FTransform& Xf = Transforms[TransformIndex];
-			BBoxes[TransformIndex] = Sphere->BoundingBox().TransformedBox(Xf);
+			BBoxes[TransformIndex] = Sphere->BoundingBox().TransformedAABB(Xf);
 			Chaos::TMassProperties<float, 3> &MP = MPArray[TransformIndex++];
 			MP.Volume = Sphere->GetVolume();
 			MP.CenterOfMass = Xf.TransformPositionNoScale(Sphere->GetCenterOfMass());
@@ -140,7 +140,7 @@ public:
 		for (Chaos::TBox<float, 3>* Box : Boxes)
 		{
 			const FTransform& Xf = Transforms[TransformIndex];
-			BBoxes[TransformIndex] = Box->BoundingBox().TransformedBox(Xf);
+			BBoxes[TransformIndex] = Box->BoundingBox().TransformedAABB(Xf);
 			Chaos::TMassProperties<float, 3> &MP = MPArray[TransformIndex++];
 			MP.Volume = Box->GetVolume();
 			MP.CenterOfMass = Xf.TransformPositionNoScale(Box->GetCenterOfMass());
@@ -149,7 +149,7 @@ public:
 		for (Chaos::TCapsule<float>* Capsule : Capsules)
 		{
 			const FTransform& Xf = Transforms[TransformIndex];
-			BBoxes[TransformIndex] = Capsule->BoundingBox().TransformedBox(Xf);
+			BBoxes[TransformIndex] = Capsule->BoundingBox().TransformedAABB(Xf);
 			Chaos::TMassProperties<float, 3> &MP = MPArray[TransformIndex++];
 			MP.Volume = Capsule->GetVolume();
 			MP.CenterOfMass = Xf.TransformPositionNoScale(Capsule->GetCenterOfMass());
@@ -158,7 +158,7 @@ public:
 		for (Chaos::TTaperedCylinder<float>* TaperedCylinder : TaperedCylinders)
 		{
 			const FTransform& Xf = Transforms[TransformIndex];
-			BBoxes[TransformIndex] = TaperedCylinder->BoundingBox().TransformedBox(Xf);
+			BBoxes[TransformIndex] = TaperedCylinder->BoundingBox().TransformedAABB(Xf);
 			Chaos::TMassProperties<float, 3> &MP = MPArray[TransformIndex++];
 			MP.Volume = TaperedCylinder->GetVolume();
 			MP.CenterOfMass = Xf.TransformPositionNoScale(TaperedCylinder->GetCenterOfMass());
@@ -167,7 +167,7 @@ public:
 		for (Chaos::FConvex* Convex : ConvexHulls)
 		{
 			const FTransform& Xf = Transforms[TransformIndex];
-			BBoxes[TransformIndex] = Convex->BoundingBox().TransformedBox(Xf);
+			BBoxes[TransformIndex] = Convex->BoundingBox().TransformedAABB(Xf);
 			Chaos::TMassProperties<float, 3> &MP = MPArray[TransformIndex++];
 			MP.Volume = Convex->BoundingBox().GetVolume();
 			MP.CenterOfMass = Xf.TransformPositionNoScale(Convex->BoundingBox().Center());
@@ -176,7 +176,7 @@ public:
 		for (Chaos::TLevelSet<float, 3>* LevelSet : LevelSets)
 		{
 			const FTransform& Xf = Transforms[TransformIndex];
-			BBoxes[TransformIndex] = LevelSet->BoundingBox().TransformedBox(Xf);
+			BBoxes[TransformIndex] = LevelSet->BoundingBox().TransformedAABB(Xf);
 			Chaos::TMassProperties<float, 3> &MP = MPArray[TransformIndex++];
 			MP.Volume = LevelSet->BoundingBox().GetVolume();
 			MP.CenterOfMass = Xf.TransformPositionNoScale(LevelSet->BoundingBox().Center());
@@ -188,13 +188,13 @@ public:
 		// we'd do something more accurate...
 		for (int32 i=0; i < Num-1; i++)
 		{
-			const Chaos::TBox<float, 3>& BoxI = BBoxes[i];
+			const Chaos::TAABB<float, 3>& BoxI = BBoxes[i];
 			for (int32 j = i+1; j < Num; j++)
 			{
-				const Chaos::TBox<float, 3>& BoxJ = BBoxes[j];
+				const Chaos::TAABB<float, 3>& BoxJ = BBoxes[j];
 				if (BoxI.Intersects(BoxJ))
 				{
-					Chaos::TBox<float, 3> BoxIJ = BoxI.GetIntersection(BoxJ);
+					Chaos::TAABB<float, 3> BoxIJ = BoxI.GetIntersection(BoxJ);
 					const float VolIJ = BoxIJ.GetVolume();
 					if (VolIJ > KINDA_SMALL_NUMBER)
 					{
@@ -313,7 +313,7 @@ public:
 			TArray<Chaos::TVector<float, 3>>& Points = CollisionPoints[TransformIndex];
 			if (!Points.Num())
 			{
-				const Chaos::TBox<float, 3>& BBox = Convex->BoundingBox();
+				const Chaos::TAABB<float, 3>& BBox = Convex->BoundingBox();
 				Chaos::TSphere<float, 3> Sphere(BBox.Center(), BBox.Extents().Size() / 2);
 				Points = Sphere.ComputeSamplePoints(ParticlesPerUnitArea, MinParticles, MaxParticles);
 				Chaos::TVector<float, 3> Normal;
@@ -332,7 +332,7 @@ public:
 			TArray<Chaos::TVector<float, 3>>& Points = CollisionPoints[TransformIndex];
 			if (!Points.Num())
 			{
-				const Chaos::TBox<float, 3>& BBox = LevelSet->BoundingBox();
+				const Chaos::TAABB<float, 3>& BBox = LevelSet->BoundingBox();
 				Chaos::TSphere<float, 3> Sphere(BBox.Center(), BBox.Extents().Size() / 2);
 				Points = Sphere.ComputeSamplePoints(ParticlesPerUnitArea, MinParticles, MaxParticles);
 				Chaos::TVector<float, 3> Normal;
@@ -480,7 +480,7 @@ protected:
 	template<class TImplicitShape>
 	void CullDeepPoints(TArray<Chaos::TVector<float,3>>& Points, const TImplicitShape& Shape, const FTransform& Xf)
 	{
-		const Chaos::TBox<float, 3>& BBox = Shape.BoundingBox();
+		const Chaos::TAABB<float, 3>& BBox = Shape.BoundingBox();
 		const float Tolerance = -BBox.Extents().Max() / 100.f; // -1/100th the largest dimension
 		if (Xf.Equals(FTransform::Identity))
 		{
diff --git a/Engine/Source/Runtime/Experimental/GeometryCollectionEngine/Private/GeometryCollection/GeometryCollectionActor.cpp b/Engine/Source/Runtime/Experimental/GeometryCollectionEngine/Private/GeometryCollection/GeometryCollectionActor.cpp
--- a/Engine/Source/Runtime/Experimental/GeometryCollectionEngine/Private/GeometryCollection/GeometryCollectionActor.cpp
+++ b/Engine/Source/Runtime/Experimental/GeometryCollectionEngine/Private/GeometryCollection/GeometryCollectionActor.cpp
@@ -77,7 +77,7 @@ bool LowLevelRaycastImp(const Chaos::TVector<float, 3>& Start, const Chaos::TVec
 	{
 #if TODO_REIMPLEMENT_GET_RIGID_PARTICLES
 		const TPBDRigidParticles<float, 3>& Particles = Solver->GetRigidParticles();	//todo(ocohen): should these just get passed in instead of hopping through scene?
-		TBox<float, 3> RayBox(Start, Start);
+		TAABB<float, 3> RayBox(Start, Start);
 		RayBox.Thicken(Dir * DeltaMag);
 		const auto& PotentialIntersections = Solver->GetSpatialAcceleration()->FindAllIntersections(RayBox);
 		Solver->ReleaseSpatialAcceleration();
diff --git a/Engine/Source/Runtime/Experimental/GeometryCollectionEngine/Private/GeometryCollection/GeometryCollectionParticlesData.cpp b/Engine/Source/Runtime/Experimental/GeometryCollectionEngine/Private/GeometryCollection/GeometryCollectionParticlesData.cpp
--- a/Engine/Source/Runtime/Experimental/GeometryCollectionEngine/Private/GeometryCollection/GeometryCollectionParticlesData.cpp
+++ b/Engine/Source/Runtime/Experimental/GeometryCollectionEngine/Private/GeometryCollection/GeometryCollectionParticlesData.cpp
@@ -155,7 +155,7 @@ void TGeometryCollectionParticlesData<T, d>::FData::Copy(EGeometryCollectionPart
 		if (ImplicitObject)
 		{
 			if      (ImplicitObject->template GetObject<Chaos::TSphere                   <T, d>>()) { return Chaos::ImplicitObjectType::Sphere     ; }
-			else if (ImplicitObject->template GetObject<Chaos::TBox                      <T, d>>()) { return Chaos::ImplicitObjectType::Box        ; }
+			else if (ImplicitObject->template GetObject<Chaos::TAABB                      <T, d>>()) { return Chaos::ImplicitObjectType::Box        ; }
 			else if (ImplicitObject->template GetObject<Chaos::TPlane                    <T, d>>()) { return Chaos::ImplicitObjectType::Plane      ; }
 			else if (ImplicitObject->template GetObject<Chaos::TImplicitObjectTransformed<T, d>>()) { return Chaos::ImplicitObjectType::Transformed; }
 			else if (ImplicitObject->template GetObject<Chaos::TImplicitObjectUnion      <T, d>>()) { return Chaos::ImplicitObjectType::Union      ; }
@@ -165,8 +165,8 @@ void TGeometryCollectionParticlesData<T, d>::FData::Copy(EGeometryCollectionPart
 	};
 	auto IsConvex       = [](Chaos::TSerializablePtr<Chaos::FImplicitObject> ImplicitObject) { return ImplicitObject ? ImplicitObject->IsConvex      (): false; };
 	auto HasBoundingBox = [](Chaos::TSerializablePtr<Chaos::FImplicitObject> ImplicitObject) { return ImplicitObject ? ImplicitObject->HasBoundingBox(): false; };
-	auto BoxMin         = [](Chaos::TSerializablePtr<Chaos::FImplicitObject> ImplicitObject) { const Chaos::TBox     <T, d>* Box     ; return ImplicitObject && (Box      = ImplicitObject->template GetObject<Chaos::TBox     <T, d>>()) != nullptr ? Box     ->Min    (): Chaos::TVector<T, d>(T(0)); };
-	auto BoxMax         = [](Chaos::TSerializablePtr<Chaos::FImplicitObject> ImplicitObject) { const Chaos::TBox     <T, d>* Box     ; return ImplicitObject && (Box      = ImplicitObject->template GetObject<Chaos::TBox     <T, d>>()) != nullptr ? Box     ->Max    (): Chaos::TVector<T, d>(T(0)); };
+	auto BoxMin         = [](Chaos::TSerializablePtr<Chaos::FImplicitObject> ImplicitObject) { const Chaos::TAABB     <T, d>* Box     ; return ImplicitObject && (Box      = ImplicitObject->template GetObject<Chaos::TAABB     <T, d>>()) != nullptr ? Box     ->Min    (): Chaos::TVector<T, d>(T(0)); };
+	auto BoxMax         = [](Chaos::TSerializablePtr<Chaos::FImplicitObject> ImplicitObject) { const Chaos::TAABB     <T, d>* Box     ; return ImplicitObject && (Box      = ImplicitObject->template GetObject<Chaos::TAABB     <T, d>>()) != nullptr ? Box     ->Max    (): Chaos::TVector<T, d>(T(0)); };
 	auto SphereCenter   = [](Chaos::TSerializablePtr<Chaos::FImplicitObject> ImplicitObject) { const Chaos::TSphere  <T, d>* Sphere  ; return ImplicitObject && (Sphere   = ImplicitObject->template GetObject<Chaos::TSphere  <T, d>>()) != nullptr ? Sphere  ->GetCenter (): Chaos::TVector<T, d>(T(0)); };
 	auto SphereRadius   = [](Chaos::TSerializablePtr<Chaos::FImplicitObject> ImplicitObject) { const Chaos::TSphere  <T, d>* Sphere  ; return ImplicitObject && (Sphere   = ImplicitObject->template GetObject<Chaos::TSphere  <T, d>>()) != nullptr ? Sphere  ->GetRadius (): T(0); };
 	auto LevelSetGrid   = [](Chaos::TSerializablePtr<Chaos::FImplicitObject> ImplicitObject) { const Chaos::TLevelSet<T, d>* LevelSet; return ImplicitObject && (LevelSet = ImplicitObject->template GetObject<Chaos::TLevelSet<T, d>>()) != nullptr ? LevelSet->GetGrid(): Chaos::TUniformGrid<T, d>(); };
diff --git a/Engine/Source/Runtime/Experimental/GeometryCollectionEngine/Private/GeometryCollection/GeometryCollectionSQAccelerator.cpp b/Engine/Source/Runtime/Experimental/GeometryCollectionEngine/Private/GeometryCollection/GeometryCollectionSQAccelerator.cpp
--- a/Engine/Source/Runtime/Experimental/GeometryCollectionEngine/Private/GeometryCollection/GeometryCollectionSQAccelerator.cpp
+++ b/Engine/Source/Runtime/Experimental/GeometryCollectionEngine/Private/GeometryCollection/GeometryCollectionSQAccelerator.cpp
@@ -442,7 +442,7 @@ void FGeometryCollectionSQAccelerator::Sweep(const Chaos::FImplicitObject& Query
 
 		Chaos::TCapsule<float> Capsule;
 		Chaos::TSphere<float, 3> Sphere;
-		Chaos::TBox<float, 3> Box;
+		Chaos::TAABB<float, 3> Box;
 	};
 	FLocalImplicitStorage ImplicitStorage;
 	Chaos::FImplicitObject* Implicit = nullptr;
@@ -510,7 +510,7 @@ void FGeometryCollectionSQAccelerator::Sweep(const Chaos::FImplicitObject& Query
 		Chaos::TVector<float, 3> x1(-P2UVector(PxBox.halfExtents));
 		Chaos::TVector<float, 3> x2(-x1);
 
-		Chaos::TBox<float, 3> Box(x1, x2);
+		Chaos::TAABB<float, 3> Box(x1, x2);
 		ImplicitStorage.Box = MoveTemp(Box);
 		Implicit = &ImplicitStorage.Box;
 
@@ -716,7 +716,7 @@ void FGeometryCollectionSQAccelerator::Overlap(const Chaos::FImplicitObject& Que
 
 		Chaos::TCapsule<float> Capsule;
 		Chaos::TSphere<float, 3> Sphere;
-		Chaos::TBox<float, 3> Box;
+		Chaos::TAABB<float, 3> Box;
 	};
 	FLocalImplicitStorage ImplicitStorage;
 
@@ -749,7 +749,7 @@ void FGeometryCollectionSQAccelerator::Overlap(const Chaos::FImplicitObject& Que
 		physx::PxBoxGeometry& PxBox = Holder.box();
 		Chaos::TVector<float, 3> x1(-P2UVector(PxBox.halfExtents));
 		Chaos::TVector<float, 3> x2(-x1);
-		Chaos::TBox<float, 3> Box(x1, x2);
+		Chaos::TAABB<float, 3> Box(x1, x2);
 		ImplicitStorage.Box = MoveTemp(Box);
 		Implicit = &ImplicitStorage.Box;
 	}
diff --git a/Engine/Source/Runtime/Experimental/GeometryCollectionSimulationCore/Private/GeometryCollection/GeometryCollectionCollisionStructureManager.cpp b/Engine/Source/Runtime/Experimental/GeometryCollectionSimulationCore/Private/GeometryCollection/GeometryCollectionCollisionStructureManager.cpp
--- a/Engine/Source/Runtime/Experimental/GeometryCollectionSimulationCore/Private/GeometryCollection/GeometryCollectionCollisionStructureManager.cpp
+++ b/Engine/Source/Runtime/Experimental/GeometryCollectionSimulationCore/Private/GeometryCollection/GeometryCollectionCollisionStructureManager.cpp
@@ -292,7 +292,7 @@ FCollisionStructureManager::CalculateUnitMassInertiaTensor(
 	if (ImplicitType == EImplicitTypeEnum::Chaos_Implicit_Box)
 	{
 		const Chaos::TVector<float, 3> Size = Bounds.GetSize();
-		const Chaos::PMatrix<float, 3, 3> I = Chaos::TBox<float, 3>::GetInertiaTensor(1.0, Size);
+		const Chaos::PMatrix<float, 3, 3> I = Chaos::TAABB<float, 3>::GetInertiaTensor(1.0, Size);
 		Tensor = { I.M[0][0], I.M[1][1], I.M[2][2] };
 	}
 	else if (ImplicitType == EImplicitTypeEnum::Chaos_Implicit_Sphere)
diff --git a/Engine/Source/Runtime/Landscape/Private/LandscapeCollision.cpp b/Engine/Source/Runtime/Landscape/Private/LandscapeCollision.cpp
--- a/Engine/Source/Runtime/Landscape/Private/LandscapeCollision.cpp
+++ b/Engine/Source/Runtime/Landscape/Private/LandscapeCollision.cpp
@@ -129,7 +129,7 @@ ULandscapeMeshCollisionComponent::FPhysXMeshRef::~FPhysXMeshRef()
 }
 
 // Generate a new guid to force a recache of landscape collison derived data
-#define LANDSCAPE_COLLISION_DERIVEDDATA_VER	TEXT("6DEB0A7F1BC546FAADB3D367E9C38467")
+#define LANDSCAPE_COLLISION_DERIVEDDATA_VER	TEXT("2265A519785B402D90EF2E726E6D8759")
 
 static FString GetHFDDCKeyString(const FName& Format, bool bDefMaterial, const FGuid& StateId, const TArray<UPhysicalMaterial*>& PhysicalMaterials)
 {
diff --git a/Engine/Source/Runtime/PhysicsCore/Private/PhysTestSerializer.cpp b/Engine/Source/Runtime/PhysicsCore/Private/PhysTestSerializer.cpp
--- a/Engine/Source/Runtime/PhysicsCore/Private/PhysTestSerializer.cpp
+++ b/Engine/Source/Runtime/PhysicsCore/Private/PhysTestSerializer.cpp
@@ -316,7 +316,7 @@ void FPhysTestSerializer::CreateChaosData()
 					auto& ShapeArray = GTParticle->ShapesArray();
 					for (auto& Shape : ShapeArray)
 					{
-						Shape->WorldSpaceInflatedShapeBounds = Geom->BoundingBox().GetAABB().TransformedAABB(TRigidTransform<FReal, 3>(Particle->X(), Particle->R()));
+						Shape->WorldSpaceInflatedShapeBounds = Geom->BoundingBox().TransformedAABB(TRigidTransform<FReal, 3>(Particle->X(), Particle->R()));
 					}
 				}
 			}
diff --git a/Engine/Source/Runtime/PhysicsSQ/Private/SQAccelerator.cpp b/Engine/Source/Runtime/PhysicsSQ/Private/SQAccelerator.cpp
--- a/Engine/Source/Runtime/PhysicsSQ/Private/SQAccelerator.cpp
+++ b/Engine/Source/Runtime/PhysicsSQ/Private/SQAccelerator.cpp
@@ -349,7 +349,7 @@ void SweepHelper(const QueryGeomType& QueryGeom, const Chaos::ISpatialAccelerati
 	using namespace Chaos;
 	using namespace ChaosInterface;
 
-	const TBox<float, 3> Bounds = QueryGeom.BoundingBox().TransformedBox(StartTM);
+	const TAABB<float, 3> Bounds = QueryGeom.BoundingBox().TransformedAABB(StartTM);
 	const FVector HalfExtents = Bounds.Extents() * 0.5f;
 
 	TSQVisitor<QueryGeomType, TAccelerationStructureHandle<float, 3>, FSweepHit> SweepVisitor(StartTM, Dir, HitBuffer, OutputFlags, QueryFilterData, QueryCallback, QueryGeom, DebugParams);
@@ -369,7 +369,7 @@ void OverlapHelper(const QueryGeomType& QueryGeom, const Chaos::ISpatialAccelera
 	using namespace Chaos;
 	using namespace ChaosInterface;
 
-	const TBox<float, 3> Bounds = QueryGeom.BoundingBox().TransformedBox(GeomPose);
+	const TAABB<float, 3> Bounds = QueryGeom.BoundingBox().TransformedAABB(GeomPose);
 	TSQVisitor<QueryGeomType, TAccelerationStructureHandle<float, 3>, FOverlapHit> OverlapVisitor(GeomPose, HitBuffer, QueryFilterData, QueryCallback, QueryGeom, DebugParams);
 	HitBuffer.IncFlushCount();
 	SpatialAcceleration.Overlap(Bounds, OverlapVisitor);
