Fix CIS

diff --git a/Engine/Source/Runtime/Experimental/Chaos/Private/Chaos/CollisionResolutionUtil.cpp b/Engine/Source/Runtime/Experimental/Chaos/Private/Chaos/CollisionResolutionUtil.cpp
--- a/Engine/Source/Runtime/Experimental/Chaos/Private/Chaos/CollisionResolutionUtil.cpp
+++ b/Engine/Source/Runtime/Experimental/Chaos/Private/Chaos/CollisionResolutionUtil.cpp
@@ -475,7 +475,7 @@ namespace Chaos
 			if (NumParticles > SampleMinParticlesForAcceleration && Object.HasBoundingBox())
 			{
 				SCOPE_CYCLE_COUNTER(STAT_UpdateLevelsetPartial);
-				TAABB<T, d> ImplicitBox = Object.BoundingBox().TransformedBox(ObjectTransform.GetRelativeTransform(SampleParticlesTransform));
+				TAABB<T, d> ImplicitBox = Object.BoundingBox().TransformedAABB(ObjectTransform.GetRelativeTransform(SampleParticlesTransform));
 				ImplicitBox.Thicken(Thickness);
 				TArray<int32> PotentialParticles;
 				{
