Add damping initialization, fixing SingleParticleProxySingleThreadTest

diff --git a/Engine/Source/Runtime/Experimental/Chaos/Public/Chaos/ParticleHandle.h b/Engine/Source/Runtime/Experimental/Chaos/Public/Chaos/ParticleHandle.h
--- a/Engine/Source/Runtime/Experimental/Chaos/Public/Chaos/ParticleHandle.h
+++ b/Engine/Source/Runtime/Experimental/Chaos/Public/Chaos/ParticleHandle.h
@@ -80,6 +80,8 @@ void PBDRigidParticleHandleImpDefaultConstruct(TPBDRigidParticleHandleImp<T, d,
 	Concrete.SetInvM(1);
 	Concrete.SetI(PMatrix<T, d, d>(1, 1, 1));
 	Concrete.SetInvI(PMatrix<T, d, d>(1, 1, 1));
+	Concrete.SetLinearDamping(0.f);
+	Concrete.SetAngularDamping(0.f);
 	Concrete.SetObjectStateLowLevel(Params.bStartSleeping ? EObjectStateType::Sleeping : EObjectStateType::Dynamic);
 }
 
@@ -101,6 +103,8 @@ void PBDRigidParticleDefaultConstruct(TPBDRigidParticle<T,d>& Concrete, const TP
 	Concrete.SetInvM(1);
 	Concrete.SetI(PMatrix<T, d, d>(1, 1, 1));
 	Concrete.SetInvI(PMatrix<T, d, d>(1, 1, 1));
+	Concrete.SetLinearDamping(0.f);
+	Concrete.SetAngularDamping(0.f);
 	Concrete.SetObjectState(Params.bStartSleeping ? EObjectStateType::Sleeping : EObjectStateType::Dynamic);
 	Concrete.SetGravityEnabled(Params.bGravityEnabled);
 }
