Use CHAOS_ENSURE so we can cook

diff --git a/Engine/Source/Runtime/Experimental/Chaos/Public/Chaos/ImplicitObjectScaled.h b/Engine/Source/Runtime/Experimental/Chaos/Public/Chaos/ImplicitObjectScaled.h
--- a/Engine/Source/Runtime/Experimental/Chaos/Public/Chaos/ImplicitObjectScaled.h
+++ b/Engine/Source/Runtime/Experimental/Chaos/Public/Chaos/ImplicitObjectScaled.h
@@ -8,6 +8,7 @@
 #include "Templates/ChooseClass.h"
 #include "Templates/EnableIf.h"
 #include "Math/NumericLimits.h"
+#include "ChaosCheck.h"
 
 namespace Chaos
 {
@@ -281,7 +282,7 @@ public:
 					OutPosition = MScale * UnscaledPosition;
 					OutNormal = (MInvScale * UnscaledNormal).GetSafeNormal();
 				}
-				ensure(OutTime <= Length);
+				CHAOS_ENSURE(OutTime <= Length);
 				return true;
 			}
 		}
