Fix merge fallout
#rnx

diff --git a/Engine/Source/Runtime/Experimental/ChaosSolvers/Private/PhysicsProxy/GeometryCollectionPhysicsProxy.cpp b/Engine/Source/Runtime/Experimental/ChaosSolvers/Private/PhysicsProxy/GeometryCollectionPhysicsProxy.cpp
--- a/Engine/Source/Runtime/Experimental/ChaosSolvers/Private/PhysicsProxy/GeometryCollectionPhysicsProxy.cpp
+++ b/Engine/Source/Runtime/Experimental/ChaosSolvers/Private/PhysicsProxy/GeometryCollectionPhysicsProxy.cpp
@@ -3865,7 +3865,7 @@ void FGeometryCollectionPhysicsProxy::PushToPhysicsState(const Chaos::FParticleD
 			{
 				Handle->SetHasBounds(true);
 				Handle->SetLocalBounds(Handle->Geometry()->BoundingBox());
-				Handle->SetWorldSpaceInflatedBounds(Handle->Geometry()->BoundingBox().TransformedBox(Chaos::TRigidTransform<float, 3>(ParticleTransform)));
+				Handle->SetWorldSpaceInflatedBounds(Handle->Geometry()->BoundingBox().TransformedAABB(Chaos::TRigidTransform<float, 3>(ParticleTransform)));
 			}
 		}
 	}
