Chaos : Collisions

- Removed local position.
- Hid ::Samples behind API.
- Made TArray samples a fixed array size 4.

diff --git a/Engine/Source/Runtime/Experimental/Chaos/Private/Chaos/CollisionResolution.cpp b/Engine/Source/Runtime/Experimental/Chaos/Private/Chaos/CollisionResolution.cpp
--- a/Engine/Source/Runtime/Experimental/Chaos/Private/Chaos/CollisionResolution.cpp
+++ b/Engine/Source/Runtime/Experimental/Chaos/Private/Chaos/CollisionResolution.cpp
@@ -547,12 +547,12 @@ namespace Chaos
 			});
 		}
 
-		TVector<float, 3> SumSampleData(TArray<FRigidBodyIterativeContactConstraint::FSampleData>& Array)
+		TVector<float, 3> SumSampleData(FRigidBodyIterativeContactConstraint& Constraint)
 		{
 			TVector<float,3> ReturnValue(0);
-			for (int i = 0; i < Array.Num(); i++)
+			for (int i = 0; i < Constraint.NumSamples(); i++)
 			{
-				ReturnValue += Array[i].X;
+				ReturnValue += Constraint[i].X;
 			}
 			return ReturnValue;
 		}
@@ -567,17 +567,17 @@ namespace Chaos
 
 			TPlane<T, d> CollisionPlane(Constraint.PlanePosition, Constraint.PlaneNormal);
 
-			for (int32 Idx = 0; Idx < Constraint.LocalSamples.Num(); Idx++)
+			for (int32 Idx = 0; Idx < Constraint.NumSamples(); Idx++)
 			{
-				Constraint.LocalSamples[Idx].Manifold.Phi = CollisionPlane.PhiWithNormal(AToBTM.TransformPosition(Constraint.LocalSamples[Idx].X), Constraint.LocalSamples[Idx].Manifold.Normal);
-				Constraint.LocalSamples[Idx].Manifold.Normal = Transform1.TransformVector(Constraint.PlaneNormal);
-				Constraint.LocalSamples[Idx].Manifold.Location = Transform0.TransformPosition(Constraint.LocalSamples[Idx].X);
+				Constraint[Idx].Manifold.Phi = CollisionPlane.PhiWithNormal(AToBTM.TransformPosition(Constraint[Idx].X), Constraint[Idx].Manifold.Normal);
+				Constraint[Idx].Manifold.Normal = Transform1.TransformVector(Constraint.PlaneNormal);
+				Constraint[Idx].Manifold.Location = Transform0.TransformPosition(Constraint[Idx].X);
 
-				if (ContactPoint.Phi > Constraint.LocalSamples[Idx].Manifold.Phi)
+				if (ContactPoint.Phi > Constraint[Idx].Manifold.Phi)
 				{
-					ContactPoint.Phi = Constraint.LocalSamples[Idx].Manifold.Phi;
-					ContactPoint.Normal = Constraint.LocalSamples[Idx].Manifold.Normal;
-					ContactPoint.Location = Constraint.LocalSamples[Idx].Manifold.Location;
+					ContactPoint.Phi = Constraint[Idx].Manifold.Phi;
+					ContactPoint.Normal = Constraint[Idx].Manifold.Normal;
+					ContactPoint.Location = Constraint[Idx].Manifold.Location;
 				}
 			}
 
@@ -591,41 +591,45 @@ namespace Chaos
 
 			FRigidBodyIterativeContactConstraint::FManifold& Manifold = Constraint.Manifold;
 
-			if (!ContactPoint.Normal.Equals(Constraint.PlaneNormal) || !Constraint.LocalSamples.Num())
+			if (!ContactPoint.Normal.Equals(Constraint.PlaneNormal) || !Constraint.NumSamples())
 			{
-				Constraint.LocalSamples.Reset();
+				Constraint.ResetSamples();
 				Constraint.PlaneNormal = Transform1.InverseTransformVector(ContactPoint.Normal);
 				Constraint.PlanePosition = Transform1.InverseTransformPosition(ContactPoint.Location - ContactPoint.Phi*ContactPoint.Normal);
 			}
 
 			TVector<T, d> SurfaceSample = Transform0.InverseTransformPosition(ContactPoint.Location);
-			if (Constraint.LocalSamples.Num() < 4)
+			if (Constraint.NumSamples() < 4)
 			{
-				Constraint.LocalSamples.Add({ SurfaceSample,0.f });
+				Constraint.AddSample({ SurfaceSample,0.f });
 			}
-			else 
+			else if(Constraint.NumSamples() == 4)
 			{
-				TVector<T,d> Center = SumSampleData(Constraint.LocalSamples) / 5.f;
+				TVector<T,d> Center = SumSampleData(Constraint) / 4.f;
 				T Delta = (Center - SurfaceSample).SizeSquared();
-				Constraint.LocalPosition = SumSampleData(Constraint.LocalSamples) / Constraint.LocalSamples.Num();
 
 				// todo(chaos) : maximize area instead
 				T SmallestDelta = FLT_MAX;
 				int32 SmallestIndex = 0;
-				for (int32 idx=0;idx<Constraint.LocalSamples.Num();idx++)
-					if (Constraint.LocalSamples[idx].Delta < SmallestDelta) {
-						SmallestDelta = Constraint.LocalSamples[idx].Delta;
+				for (int32 idx=0;idx<Constraint.NumSamples();idx++)
+					if (Constraint[idx].Delta < SmallestDelta) {
+						SmallestDelta = Constraint[idx].Delta;
 						SmallestIndex = idx;
 					}
 
 				if (Delta > SmallestDelta) {
-					Constraint.LocalSamples[SmallestIndex] = { SurfaceSample,Delta };
+					Constraint[SmallestIndex] = { SurfaceSample,Delta };
 				}
 			}
+			else
+			{
+				ensure(false); // max of 4 points
+			}
 
 			typedef FRigidBodyIterativeContactConstraint::FSampleData FSampleData;
-			Constraint.LocalPosition = SumSampleData(Constraint.LocalSamples) / Constraint.LocalSamples.Num();
-			for (FSampleData& Data : Constraint.LocalSamples) Data.Delta = (Constraint.LocalPosition - Data.X).SizeSquared();
+			TVector<T, d> Center = SumSampleData(Constraint) / Constraint.NumSamples();
+			for (int32 Index = 0; Index < Constraint.NumSamples(); Index++) { Constraint[Index].Delta = (Center - Constraint[Index].X).SizeSquared(); }
+
 		}
 
 
diff --git a/Engine/Source/Runtime/Experimental/Chaos/Public/Chaos/CollisionResolutionTypes.h b/Engine/Source/Runtime/Experimental/Chaos/Public/Chaos/CollisionResolutionTypes.h
--- a/Engine/Source/Runtime/Experimental/Chaos/Public/Chaos/CollisionResolutionTypes.h
+++ b/Engine/Source/Runtime/Experimental/Chaos/Public/Chaos/CollisionResolutionTypes.h
@@ -172,14 +172,22 @@ namespace Chaos
 		using Base::Particle;
 		struct FSampleData { TVector<T,d> X; float Delta; FManifold Manifold; };
 
-		TRigidBodyIterativeContactConstraint() : Base(Base::FType::MultiPoint), PlaneNormal(0), PlanePosition(0), LocalPosition(0) {}
+		TRigidBodyIterativeContactConstraint() : Base(Base::FType::MultiPoint), PlaneNormal(0), PlanePosition(0), NumberSamples(0) {}
 		static typename Base::FType StaticType() { return Base::FType::MultiPoint; };
 
 		TVector<T, d> PlaneNormal; // local space contact normal on Particle1
 		TVector<T, d> PlanePosition; // local space surface position on Particle1
 
-		TVector<T, d> LocalPosition; // averaged position
-		TArray< FSampleData > LocalSamples; // iterative samples
+		// Samples
+		int32               NumSamples()                   { return NumberSamples; }
+		void                AddSample(FSampleData && Data) {Samples[NumberSamples] = Data; NumberSamples++; };
+		void                ResetSamples()                 { NumberSamples=0; }
+		FSampleData &       operator[](int32 Index)        { return Samples[Index]; }
+		const FSampleData & operator[](int32 Index) const  { return Samples[Index]; }
+
+	private:
+		int NumberSamples;
+		FSampleData Samples[4]; // iterative samples
 	};
 	typedef TRigidBodyIterativeContactConstraint<float, 3> FRigidBodyIterativeContactConstraint;
 
