More LLM tags for chaos

diff --git a/Engine/Source/Runtime/Core/Private/HAL/LowLevelMemTracker.cpp b/Engine/Source/Runtime/Core/Private/HAL/LowLevelMemTracker.cpp
--- a/Engine/Source/Runtime/Core/Private/HAL/LowLevelMemTracker.cpp
+++ b/Engine/Source/Runtime/Core/Private/HAL/LowLevelMemTracker.cpp
@@ -89,6 +89,8 @@ DECLARE_LLM_MEMORY_STAT(TEXT("Physics"), STAT_PhysicsLLM, STATGROUP_LLMFULL);
 DECLARE_LLM_MEMORY_STAT(TEXT("PhysX"), STAT_PhysXLLM, STATGROUP_LLMFULL);
 DECLARE_LLM_MEMORY_STAT(TEXT("Chaos"), STAT_ChaosLLM, STATGROUP_LLMFULL);
 DECLARE_LLM_MEMORY_STAT(TEXT("ChaosGeometry"), STAT_ChaosGeometryLLM, STATGROUP_LLMFULL);
+DECLARE_LLM_MEMORY_STAT(TEXT("ChaosAcceleration"), STAT_ChaosAccelerationLLM, STATGROUP_LLMFULL);
+DECLARE_LLM_MEMORY_STAT(TEXT("ChaosParticles"), STAT_ChaosParticlesLLM, STATGROUP_LLMFULL);
 DECLARE_LLM_MEMORY_STAT(TEXT("EnginePreInit"), STAT_EnginePreInitLLM, STATGROUP_LLMFULL);
 DECLARE_LLM_MEMORY_STAT(TEXT("EngineInit"), STAT_EngineInitLLM, STATGROUP_LLMFULL);
 DECLARE_LLM_MEMORY_STAT(TEXT("Rendering Thread"), STAT_RenderingThreadLLM, STATGROUP_LLMFULL);
diff --git a/Engine/Source/Runtime/Core/Public/HAL/LowLevelMemTracker.h b/Engine/Source/Runtime/Core/Public/HAL/LowLevelMemTracker.h
--- a/Engine/Source/Runtime/Core/Public/HAL/LowLevelMemTracker.h
+++ b/Engine/Source/Runtime/Core/Public/HAL/LowLevelMemTracker.h
@@ -188,6 +188,8 @@ enum class ELLMTagSet : uint8
 	macro(PhysX,								"PhysX",						GET_STATFNAME(STAT_PhysXLLM),								GET_STATFNAME(STAT_PhysXSummaryLLM),			ELLMTag::Physics)\
 	macro(Chaos,								"Chaos",						GET_STATFNAME(STAT_ChaosLLM),								GET_STATFNAME(STAT_ChaosSummaryLLM),			ELLMTag::Physics)\
 	macro(ChaosGeometry,						"ChaosGeometry",				GET_STATFNAME(STAT_ChaosGeometryLLM),						GET_STATFNAME(STAT_ChaosSummaryLLM),			ELLMTag::Physics)\
+	macro(ChaosAcceleration,					"ChaosAcceleration",			GET_STATFNAME(STAT_ChaosAccelerationLLM),					GET_STATFNAME(STAT_ChaosSummaryLLM),			ELLMTag::Physics)\
+	macro(ChaosParticles,						"ChaosParticles",				GET_STATFNAME(STAT_ChaosParticlesLLM),						GET_STATFNAME(STAT_ChaosSummaryLLM),			ELLMTag::Physics)\
 	macro(EnginePreInitMemory,					"EnginePreInit",				GET_STATFNAME(STAT_EnginePreInitLLM),						GET_STATFNAME(STAT_EngineSummaryLLM),			-1)\
 	macro(EngineInitMemory,						"EngineInit",					GET_STATFNAME(STAT_EngineInitLLM),							GET_STATFNAME(STAT_EngineSummaryLLM),			-1)\
 	macro(RenderingThreadMemory,				"RenderingThread",				GET_STATFNAME(STAT_RenderingThreadLLM),						GET_STATFNAME(STAT_EngineSummaryLLM),			-1)\
diff --git a/Engine/Source/Runtime/Experimental/Chaos/Public/Chaos/AABBTree.h b/Engine/Source/Runtime/Experimental/Chaos/Public/Chaos/AABBTree.h
--- a/Engine/Source/Runtime/Experimental/Chaos/Public/Chaos/AABBTree.h
+++ b/Engine/Source/Runtime/Experimental/Chaos/Public/Chaos/AABBTree.h
@@ -243,7 +243,7 @@ public:
 	virtual void ProgressAsyncTimeSlicing(bool ForceBuildCompletion) override
 	{
 		SCOPE_CYCLE_COUNTER(STAT_AABBTreeProgressTimeSlice);
-
+		LLM_SCOPE(ELLMTag::ChaosAcceleration);
 		// force is to stop time slicing and complete the rest of the build now
 		if (ForceBuildCompletion)
 		{
@@ -377,6 +377,7 @@ public:
 
 	virtual void RemoveElement(const TPayloadType& Payload)
 	{
+		LLM_SCOPE(ELLMTag::ChaosAcceleration);
 		if (FAABBTreePayloadInfo* PayloadInfo = PayloadToInfo.Find(Payload))
 		{
 			if (PayloadInfo->GlobalPayloadIdx != INDEX_NONE)
@@ -410,6 +411,7 @@ public:
 
 	virtual void UpdateElement(const TPayloadType& Payload, const TAABB<T, 3>& NewBounds, bool bHasBounds) override
 	{
+		LLM_SCOPE(ELLMTag::ChaosAcceleration);
 		FAABBTreePayloadInfo* PayloadInfo = PayloadToInfo.Find(Payload);
 		if (PayloadInfo)
 		{
@@ -492,6 +494,7 @@ public:
 
 	virtual void Serialize(FChaosArchive& Ar) override
 	{
+		LLM_SCOPE(ELLMTag::ChaosAcceleration);
 		TBox<FReal, 3>::SerializeAsAABB(Ar, FullBounds);
 		Ar << Nodes;
 		Ar << Leaves;
@@ -637,6 +640,7 @@ private:
 	void GenerateTree(const TParticles& Particles)
 	{
 		SCOPE_CYCLE_COUNTER(STAT_AABBTreeGenerateTree);
+		LLM_SCOPE(ELLMTag::ChaosAcceleration);
 		this->SetAsyncTimeSlicingComplete(false);
 
 		TArray<FElement> ElemsWithBounds;
diff --git a/Engine/Source/Runtime/Experimental/Chaos/Public/Chaos/BoundingVolume.h b/Engine/Source/Runtime/Experimental/Chaos/Public/Chaos/BoundingVolume.h
--- a/Engine/Source/Runtime/Experimental/Chaos/Public/Chaos/BoundingVolume.h
+++ b/Engine/Source/Runtime/Experimental/Chaos/Public/Chaos/BoundingVolume.h
@@ -208,6 +208,7 @@ public:
 
 	virtual void RemoveElement(const TPayloadType& Payload) override
 	{
+		LLM_SCOPE(ELLMTag::ChaosAcceleration);
 		SCOPE_CYCLE_COUNTER(STAT_BoundingVolumeRemoveElement);
 		if (const FPayloadInfo* PayloadInfo = MPayloadInfo.Find(Payload))
 		{
@@ -232,6 +233,7 @@ public:
 
 	virtual void UpdateElement(const TPayloadType& Payload, const TAABB<T,d>& NewBounds, bool bHasBounds) override
 	{
+		LLM_SCOPE(ELLMTag::ChaosAcceleration);
 		SCOPE_CYCLE_COUNTER(STAT_BoundingVolumeUpdateElement);
 		if (FPayloadInfo* PayloadInfo = MPayloadInfo.Find(Payload))
 		{
@@ -315,6 +317,7 @@ public:
 
 	virtual void Serialize(FChaosArchive& Ar) override
 	{
+		LLM_SCOPE(ELLMTag::ChaosAcceleration);
 		Ar.UsingCustomVersion(FExternalPhysicsCustomObjectVersion::GUID);
 		if (Ar.CustomVer(FExternalPhysicsCustomObjectVersion::GUID) < FExternalPhysicsCustomObjectVersion::GlobalElementsHaveBounds)
 		{
@@ -814,6 +817,7 @@ private:
 	template <typename ParticleView>
 	void GenerateTree(const ParticleView& Particles, const bool bUseVelocity, const T Dt, const int32 MaxCells)
 	{
+		LLM_SCOPE(ELLMTag::ChaosAcceleration);
 		MGlobalPayloads.Reset();
 		MPayloadInfo.Reset();
 		bIsEmpty = true;
diff --git a/Engine/Source/Runtime/Experimental/Chaos/Public/Chaos/GeometryParticles.h b/Engine/Source/Runtime/Experimental/Chaos/Public/Chaos/GeometryParticles.h
--- a/Engine/Source/Runtime/Experimental/Chaos/Public/Chaos/GeometryParticles.h
+++ b/Engine/Source/Runtime/Experimental/Chaos/Public/Chaos/GeometryParticles.h
@@ -331,6 +331,7 @@ namespace Chaos
 
 		CHAOS_API virtual void Serialize(FChaosArchive& Ar)
 		{
+			LLM_SCOPE(ELLMTag::ChaosParticles);
 			TParticles<T, d>::Serialize(Ar);
 			Ar << MGeometry << MDynamicGeometry << MR;
 			Ar.UsingCustomVersion(FPhysicsObjectVersion::GUID);
diff --git a/Engine/Source/Runtime/Experimental/Chaos/Public/Chaos/KinematicGeometryParticles.h b/Engine/Source/Runtime/Experimental/Chaos/Public/Chaos/KinematicGeometryParticles.h
--- a/Engine/Source/Runtime/Experimental/Chaos/Public/Chaos/KinematicGeometryParticles.h
+++ b/Engine/Source/Runtime/Experimental/Chaos/Public/Chaos/KinematicGeometryParticles.h
@@ -115,6 +115,7 @@ class TKinematicGeometryParticlesImp : public TGeometryParticlesImp<T, d, SimTyp
 
 	CHAOS_API virtual void Serialize(FChaosArchive& Ar) override
 	{
+		LLM_SCOPE(ELLMTag::ChaosParticles);
 		TGeometryParticlesImp<T, d, SimType>::Serialize(Ar);
 		Ar << MV << MW;
 		
diff --git a/Engine/Source/Runtime/Experimental/Chaos/Public/Chaos/PBDRigidParticles.h b/Engine/Source/Runtime/Experimental/Chaos/Public/Chaos/PBDRigidParticles.h
--- a/Engine/Source/Runtime/Experimental/Chaos/Public/Chaos/PBDRigidParticles.h
+++ b/Engine/Source/Runtime/Experimental/Chaos/Public/Chaos/PBDRigidParticles.h
@@ -162,6 +162,7 @@ class TPBDRigidParticles : public TRigidParticles<T, d>
 
 	CHAOS_API virtual void Serialize(FChaosArchive& Ar) override
 	{
+		LLM_SCOPE(ELLMTag::ChaosParticles);
 		TRigidParticles<T, d>::Serialize(Ar);
 		Ar << MP << MQ << MPreV << MPreW;
 	}
diff --git a/Engine/Source/Runtime/Experimental/Chaos/Public/Chaos/Particles.h b/Engine/Source/Runtime/Experimental/Chaos/Public/Chaos/Particles.h
--- a/Engine/Source/Runtime/Experimental/Chaos/Public/Chaos/Particles.h
+++ b/Engine/Source/Runtime/Experimental/Chaos/Public/Chaos/Particles.h
@@ -65,12 +65,14 @@ namespace Chaos
 
 		void AddParticles(const int32 Num)
 		{
+			LLM_SCOPE(ELLMTag::ChaosParticles);
 			AddElementsHelper(Num);
 			IncrementDirtyValidation();
 		}
 
 		void DestroyParticle(const int32 Idx)
 		{
+			LLM_SCOPE(ELLMTag::ChaosParticles);
 			if (MRemoveParticleBehavior == ERemoveParticleBehavior::RemoveAtSwap)
 			{
 				RemoveAtSwapHelper(Idx);
@@ -117,6 +119,7 @@ namespace Chaos
 		
 		void Serialize(FArchive& Ar)
 		{
+			LLM_SCOPE(ELLMTag::ChaosParticles);
 			bool bSerialize = true;	//leftover from when we had view support
 			Ar << bSerialize;
 			if (ensureMsgf(bSerialize, TEXT("Cannot serialize shared views. Refactor needed to reduce memory")))
diff --git a/Engine/Source/Runtime/Experimental/Chaos/Public/Chaos/RigidParticles.h b/Engine/Source/Runtime/Experimental/Chaos/Public/Chaos/RigidParticles.h
--- a/Engine/Source/Runtime/Experimental/Chaos/Public/Chaos/RigidParticles.h
+++ b/Engine/Source/Runtime/Experimental/Chaos/Public/Chaos/RigidParticles.h
@@ -192,6 +192,7 @@ class TRigidParticles : public TKinematicGeometryParticles<T, d>
 
 	CHAOS_API virtual void Serialize(FChaosArchive& Ar) override
 	{
+		LLM_SCOPE(ELLMTag::ChaosParticles);
 		TKinematicGeometryParticles<T,d>::Serialize(Ar);
 		Ar << MF << MT << MExternalForce << MExternalTorque << MI << MInvI << MM << MInvM;
 
