Fix EPA not returning the right penetration during when lower bound passes upper bound

diff --git a/Engine/Source/Runtime/Experimental/Chaos/Public/Chaos/EPA.h b/Engine/Source/Runtime/Experimental/Chaos/Public/Chaos/EPA.h
--- a/Engine/Source/Runtime/Experimental/Chaos/Public/Chaos/EPA.h
+++ b/Engine/Source/Runtime/Experimental/Chaos/Public/Chaos/EPA.h
@@ -407,7 +407,6 @@ EPAResult EPA(TArray<TVec3<T>>& VertsABuffer, TArray<TVec3<T>>& VertsBBuffer, co
 			continue;
 		}
 
-		LastEntry = Entry;
 
 
 		const TVec3<T> ASupport = SupportA(Entry.PlaneNormal);
@@ -426,13 +425,22 @@ EPAResult EPA(TArray<TVec3<T>>& VertsABuffer, TArray<TVec3<T>>& VertsBBuffer, co
 		LowerBound = Entry.Distance;
 
 		//It's possible the origin is not contained by the CSO. In this case the upper bound will be negative, at which point we should just exit. Maybe return a different enum value?
-		if ((UpperBound - LowerBound) <= FMath::Abs(Eps * LowerBound))
+		if (FMath::Abs(UpperBound - LowerBound) <= FMath::Abs(Eps * LowerBound))
 		{
 			//UE_LOG(LogChaos, Warning, TEXT("Iteration:%d"), Iteration);
 			ComputeEPAResults(VertsABuffer.GetData(), VertsBBuffer.GetData(), Entry, OutPenetration, OutDir, WitnessA, WitnessB);
 			return EPAResult::Ok;
 		}
 
+		if (UpperBound < LowerBound)
+		{
+			//we cannot get any better than what we saw, so just return previous face
+			break;
+		}
+
+		LastEntry = Entry;
+
+
 		const int32 NewVertIdx = VertsABuffer.Add(ASupport);
 		VertsBBuffer.Add(BSupport);
 
