Fix non-unity

diff --git a/Engine/Source/Runtime/Experimental/Chaos/Public/Chaos/PBDRigidParticles.h b/Engine/Source/Runtime/Experimental/Chaos/Public/Chaos/PBDRigidParticles.h
--- a/Engine/Source/Runtime/Experimental/Chaos/Public/Chaos/PBDRigidParticles.h
+++ b/Engine/Source/Runtime/Experimental/Chaos/Public/Chaos/PBDRigidParticles.h
@@ -4,6 +4,7 @@
 #include "Chaos/ArrayCollectionArray.h"
 #include "Chaos/RigidParticles.h"
 #include "Chaos/Rotation.h"
+#include "HAL/LowLevelMemTracker.h"
 
 #if !PLATFORM_PS4
 #pragma warning(push)
