Fix CIS

diff --git a/Engine/Source/Runtime/Engine/Classes/PhysicsEngine/ConvexElem.h b/Engine/Source/Runtime/Engine/Classes/PhysicsEngine/ConvexElem.h
--- a/Engine/Source/Runtime/Engine/Classes/PhysicsEngine/ConvexElem.h
+++ b/Engine/Source/Runtime/Engine/Classes/PhysicsEngine/ConvexElem.h
@@ -5,6 +5,7 @@
 #include "CoreMinimal.h"
 #include "UObject/ObjectMacros.h"
 #include "PhysicsEngine/ShapeElem.h"
+#include "Chaos/Convex.h"
 
 #if WITH_CHAOS
 #include "Chaos/Serializable.h"
diff --git a/Engine/Source/Runtime/Experimental/Chaos/Public/Chaos/Particles.h b/Engine/Source/Runtime/Experimental/Chaos/Public/Chaos/Particles.h
--- a/Engine/Source/Runtime/Experimental/Chaos/Public/Chaos/Particles.h
+++ b/Engine/Source/Runtime/Experimental/Chaos/Public/Chaos/Particles.h
@@ -5,6 +5,7 @@
 #include "Chaos/ArrayCollectionArray.h"
 #include "Chaos/Vector.h"
 #include "ChaosArchive.h"
+#include "HAL/LowLevelMemTracker.h"
 
 #if !(UE_BUILD_SHIPPING || UE_BUILD_TEST)
 #define PARTICLE_ITERATOR_RANGED_FOR_CHECK 1
diff --git a/Engine/Source/Runtime/Experimental/Chaos/Public/Chaos/RigidParticles.h b/Engine/Source/Runtime/Experimental/Chaos/Public/Chaos/RigidParticles.h
--- a/Engine/Source/Runtime/Experimental/Chaos/Public/Chaos/RigidParticles.h
+++ b/Engine/Source/Runtime/Experimental/Chaos/Public/Chaos/RigidParticles.h
@@ -10,6 +10,7 @@
 #include "Chaos/Matrix.h"
 #include "Chaos/Particles.h"
 #include "Chaos/Rotation.h"
+#include "HAL/LowLevelMemTracker.h"
 
 namespace Chaos
 {
