Chaos
- fix AABBs of kinematics
- debug draw solver indices
- partial drives fix

diff --git a/Engine/Source/Runtime/Engine/Private/PhysicsEngine/ImmediatePhysics/ImmediatePhysicsChaos/ImmediatePhysicsJointHandle_Chaos.cpp b/Engine/Source/Runtime/Engine/Private/PhysicsEngine/ImmediatePhysics/ImmediatePhysicsChaos/ImmediatePhysicsJointHandle_Chaos.cpp
--- a/Engine/Source/Runtime/Engine/Private/PhysicsEngine/ImmediatePhysics/ImmediatePhysicsChaos/ImmediatePhysicsJointHandle_Chaos.cpp
+++ b/Engine/Source/Runtime/Engine/Private/PhysicsEngine/ImmediatePhysics/ImmediatePhysicsChaos/ImmediatePhysicsJointHandle_Chaos.cpp
@@ -22,14 +22,14 @@ namespace ImmediatePhysics_Chaos
 	bool ChaosImmediate_StiffnessUseMass = true;
 	FAutoConsoleVariableRef CVarStiffnessUseMass(TEXT("p.Chaos.ImmPhys.StiffnessUseMass"), ChaosImmediate_StiffnessUseMass, TEXT("Whether to use mass to scale stiffness in the conversion."));
 
-	float ChaosImmediate_JointMinProjection = 0.1f;
-	float ChaosImmediate_JointMaxProjection = 1.0f;
+	float ChaosImmediate_JointMinProjection = 0.0f;
+	float ChaosImmediate_JointMaxProjection = 0.3f;
 	FAutoConsoleVariableRef CVarJointMinProjection(TEXT("p.Chaos.ImmPhys.JointMinProjection"), ChaosImmediate_JointMinProjection, TEXT("Joint min projection (for joints with projection disabled)."));
 	FAutoConsoleVariableRef CVarJointMaxProjection(TEXT("p.Chaos.ImmPhys.JointMaxProjection"), ChaosImmediate_JointMaxProjection, TEXT("Joint max projection (for joints with projection enabled)."));
 
-	float ChaosImmediate_DriveStiffnessSourceMin = 100.0f;
-	float ChaosImmediate_DriveStiffnessSourceMax = 2000.0f;
-	float ChaosImmediate_DriveStiffnessTargetMin = 0.1f;
+	float ChaosImmediate_DriveStiffnessSourceMin = 100.0f;	// PhysX stiffness per inertia that we translate to a Chaos stiffness of ChaosImmediate_DriveStiffnessTargetMin
+	float ChaosImmediate_DriveStiffnessSourceMax = 2000.0f;	// PhysX stiffness per inertia that we translate to a Chaos stiffness of ChaosImmediate_DriveStiffnessTargetMax
+	float ChaosImmediate_DriveStiffnessTargetMin = 0.2f;
 	float ChaosImmediate_DriveStiffnessTargetMax = 1.0f;
 	FAutoConsoleVariableRef CVarDriveStiffnessSourceMin(TEXT("p.Chaos.ImmPhys.DriveStiffnessSourceMin"), ChaosImmediate_DriveStiffnessSourceMin, TEXT("Conversion factor for drive stiffness."));
 	FAutoConsoleVariableRef CVarDriveStiffnessSourceMax(TEXT("p.Chaos.ImmPhys.DriveStiffnessSourceMax"), ChaosImmediate_DriveStiffnessSourceMax, TEXT("Conversion factor for drive stiffness."));
@@ -56,7 +56,7 @@ namespace ImmediatePhysics_Chaos
 		if (ChaosImmediate_DriveStiffnessSourceMax > ChaosImmediate_DriveStiffnessSourceMin)
 		{
 			const float InvI = ChaosImmediate_StiffnessUseMass ? (InvI0.Min() + InvI1.Min()) : 1.0f;
-			const float F = (InStiffness - ChaosImmediate_DriveStiffnessSourceMin) * InvI  / (ChaosImmediate_DriveStiffnessSourceMax - ChaosImmediate_DriveStiffnessSourceMin);
+			const float F = FMath::Max(0.0f, (InStiffness * InvI - ChaosImmediate_DriveStiffnessSourceMin)) / FMath::Max(SMALL_NUMBER, (ChaosImmediate_DriveStiffnessSourceMax - ChaosImmediate_DriveStiffnessSourceMin));
 			return FMath::Lerp(ChaosImmediate_DriveStiffnessTargetMin, ChaosImmediate_DriveStiffnessTargetMax, FMath::Clamp(F, 0.0f, 1.0f));
 		}
 		return 0;
@@ -67,7 +67,7 @@ namespace ImmediatePhysics_Chaos
 		if (ChaosImmediate_SoftLinearStiffnessSourceMax > ChaosImmediate_SoftLinearStiffnessSourceMin)
 		{
 			const float InvM = ChaosImmediate_StiffnessUseMass ? (InvM0 + InvM1) : 1.0f;
-			const float F = (InStiffness - ChaosImmediate_SoftLinearStiffnessSourceMin) * InvM / (ChaosImmediate_SoftLinearStiffnessSourceMax - ChaosImmediate_SoftLinearStiffnessSourceMin);
+			const float F = FMath::Max(0.0f, (InStiffness * InvM - ChaosImmediate_SoftLinearStiffnessSourceMin)) / FMath::Max(SMALL_NUMBER, (ChaosImmediate_SoftLinearStiffnessSourceMax - ChaosImmediate_SoftLinearStiffnessSourceMin));
 			return FMath::Lerp(ChaosImmediate_SoftLinearStiffnessTargetMin, ChaosImmediate_SoftLinearStiffnessTargetMax, FMath::Clamp(F, 0.0f, 1.0f));
 		}
 		return 0;
@@ -78,7 +78,7 @@ namespace ImmediatePhysics_Chaos
 		if (ChaosImmediate_SoftAngularStiffnessSourceMax > ChaosImmediate_SoftAngularStiffnessSourceMin)
 		{
 			const float InvI = ChaosImmediate_StiffnessUseMass ? (InvI0.Min() + InvI1.Min()) : 1.0f;
-			const float F = (InStiffness - ChaosImmediate_SoftAngularStiffnessSourceMin) * InvI / (ChaosImmediate_SoftAngularStiffnessSourceMax - ChaosImmediate_SoftAngularStiffnessSourceMin);
+			const float F = FMath::Max(0.0f, (InStiffness * InvI - ChaosImmediate_SoftAngularStiffnessSourceMin)) / FMath::Max(SMALL_NUMBER, (ChaosImmediate_SoftAngularStiffnessSourceMax - ChaosImmediate_SoftAngularStiffnessSourceMin));
 			return FMath::Lerp(ChaosImmediate_SoftAngularStiffnessTargetMin, ChaosImmediate_SoftAngularStiffnessTargetMax, FMath::Clamp(F, 0.0f, 1.0f));
 		}
 		return 0;
diff --git a/Engine/Source/Runtime/Engine/Private/PhysicsEngine/ImmediatePhysics/ImmediatePhysicsChaos/ImmediatePhysicsSimulation_Chaos.cpp b/Engine/Source/Runtime/Engine/Private/PhysicsEngine/ImmediatePhysics/ImmediatePhysicsChaos/ImmediatePhysicsSimulation_Chaos.cpp
--- a/Engine/Source/Runtime/Engine/Private/PhysicsEngine/ImmediatePhysics/ImmediatePhysicsChaos/ImmediatePhysicsSimulation_Chaos.cpp
+++ b/Engine/Source/Runtime/Engine/Private/PhysicsEngine/ImmediatePhysics/ImmediatePhysicsChaos/ImmediatePhysicsSimulation_Chaos.cpp
@@ -21,14 +21,14 @@
 // @todo(ccaulfield): remove when finished
 //
 float ChaosImmediate_Evolution_DeltaTime = 0.03f;
-int32 ChaosImmediate_Evolution_Iterations = 2;
-int32 ChaosImmediate_Evolution_PushOutIterations = 10;
+int32 ChaosImmediate_Evolution_Iterations = 4;
+int32 ChaosImmediate_Evolution_PushOutIterations = 2;
 FAutoConsoleVariableRef CVarChaosImmPhysDeltaTime(TEXT("p.Chaos.ImmPhys.DeltaTime"), ChaosImmediate_Evolution_DeltaTime, TEXT("Override chaos immediate physics delta time if non-zero"));
 FAutoConsoleVariableRef CVarChaosImmPhysIterations(TEXT("p.Chaos.ImmPhys.Iterations"), ChaosImmediate_Evolution_Iterations, TEXT("Number of constraint solver loops in immediate physics"));
 FAutoConsoleVariableRef CVarChaosImmPhysPushOutIterations(TEXT("p.Chaos.ImmPhys.PushOutIterations"), ChaosImmediate_Evolution_PushOutIterations, TEXT("Set the ApplyPushOut() (position correction) iteration count"));
 
 int32 ChaosImmediate_Collision_Enabled = 1;
-int32 ChaosImmediate_Collision_PairIterations = 0;
+int32 ChaosImmediate_Collision_PairIterations = 2;
 int32 ChaosImmediate_Collision_PushOutPairIterations = 2;
 int32 ChaosImmediate_Collision_Priority = 1;
 float ChaosImmediate_Collision_Thickness = 0;
@@ -38,14 +38,14 @@ FAutoConsoleVariableRef CVarChaosImmPhysCollisionPushOutPairIterations(TEXT("p.C
 FAutoConsoleVariableRef CVarChaosImmPhysCollisionPriority(TEXT("p.Chaos.ImmPhys.Collision.Priority"), ChaosImmediate_Collision_Priority, TEXT("Set the Collision constraint sort order (Joints have priority 0)"));
 FAutoConsoleVariableRef CVarChaosImmPhysThickness(TEXT("p.Chaos.ImmPhys.Collision.Thickness"), ChaosImmediate_Collision_Thickness, TEXT("ChaosImmediateThickness"));
 
-int32 ChaosImmediate_Joint_PairIterations = 0;
-int32 ChaosImmediate_Joint_PushOutPairIterations = 2;
+int32 ChaosImmediate_Joint_PairIterations = 2;
+int32 ChaosImmediate_Joint_PushOutPairIterations = 0;
 float ChaosImmediate_Joint_SwingTwistAngleTolerance = 1.0e-6f;
 int32 ChaosImmediate_Joint_EnableTwistLimits = 1;
 int32 ChaosImmediate_Joint_EnableSwingLimits = 1;
 int32 ChaosImmediate_Joint_EnableDrives = 1;
 int32 ChaosImmediate_Joint_DrivesPhase = (int32)Chaos::EJointSolverPhase::Apply;
-int32 ChaosImmediate_Joint_ProjectionPhase = (int32)Chaos::EJointSolverPhase::ApplyPushOut;
+int32 ChaosImmediate_Joint_ProjectionPhase = (int32)Chaos::EJointSolverPhase::Apply;
 float ChaosImmediate_Joint_LinearProjection = 0.0f;
 float ChaosImmediate_Joint_AngularProjection = 0.0f;
 float ChaosImmediate_Joint_Stiffness = 1.0f;
@@ -70,6 +70,10 @@ FAutoConsoleVariableRef CVarChaosImmPhysSoftAngularStiffness(TEXT("p.Chaos.ImmPh
 FAutoConsoleVariableRef CVarChaosImmPhysDriveStiffness(TEXT("p.Chaos.ImmPhys.Joint.DriveStiffness"), ChaosImmediate_Joint_DriveStiffness, TEXT("6Dof joint drive stiffness override (if > 0)."));
 FAutoConsoleVariableRef CVarChaosImmPhysMinParentMassRatio(TEXT("p.Chaos.ImmPhys.Joint.MinParentMassRatio"), ChaosImmediate_Joint_MinParentMassRatio, TEXT("6Dof joint MinParentMassRatio (if > 0)"));
 FAutoConsoleVariableRef CVarChaosImmPhysMaxInertiaRatio(TEXT("p.Chaos.ImmPhys.Joint.MaxInertiaRatio"), ChaosImmediate_Joint_MaxInertiaRatio, TEXT("6Dof joint MaxInertiaRatio (if > 0)"));
+
+int32 ChaosImmediate_Material_Resitution = 0.0f;
+int32 ChaosImmediate_Material_Friction = 0.5f;
+
 //
 // end remove when finished
 //
@@ -176,22 +180,19 @@ namespace ImmediatePhysics_Chaos
 		Evolution.SetPostIntegrateCallback(
 			[this]()
 			{
-				// Dynamic only - Kinematics get drawn once at the end of frame
-				DebugDrawParticles(3, 3, 0.3f, false, true);
+				DebugDrawDynamicParticles(3, 3, FColor(32, 32, 0));
 				DebugDrawConstraints(3, 3, 0.3f);
 			});
 		Evolution.SetPostApplyCallback(
 			[this]()
 			{
-				// Dynamic only - Kinematics get drawn once at the end of frame
-				DebugDrawParticles(3, 3, 0.6f, false, true);
+				DebugDrawDynamicParticles(3, 3, FColor(128, 128, 0));
 				DebugDrawConstraints(3, 3, 0.6f);
 			});
 		Evolution.SetPostApplyPushOutCallback(
 			[this]()
 			{
-				// Dynamic only - Kinematics get drawn once at the end of frame
-				DebugDrawParticles(3, 3, 1.0f, false, true);
+				DebugDrawDynamicParticles(3, 3, FColor(255, 255, 0));
 				DebugDrawConstraints(3, 3, 1.0f);
 			});
 		Collisions.SetPostApplyCallback(
@@ -201,6 +202,7 @@ namespace ImmediatePhysics_Chaos
 				{
 					DebugDraw::DrawCollisions(SimulationSpaceTransform, InConstraintHandles, 0.3f);
 				}
+				DebugDrawDynamicParticles(4, 4, FColor(128, 0, 0));
 			});
 		Collisions.SetPostApplyPushOutCallback(
 			[this](const float Dt, const TArray<TPBDCollisionConstraintHandle<float, 3>*>& InConstraintHandles, bool bRequiresAnotherIteration)
@@ -225,6 +227,16 @@ namespace ImmediatePhysics_Chaos
 				{
 					DebugDraw::DrawJointConstraints(SimulationSpaceTransform, InConstraintHandles, 0.6f);
 				}
+				DebugDrawDynamicParticles(4, 4, FColor(0, 128, 0));
+			});
+		Joints.SetPostProjectCallback(
+			[this](const float Dt, const TArray<FPBDJointConstraintHandle*>& InConstraintHandles)
+			{
+				if (ChaosImmediate_DebugDrawJoints == 4)
+				{
+					DebugDraw::DrawJointConstraints(SimulationSpaceTransform, InConstraintHandles, 0.6f);
+				}
+				DebugDrawDynamicParticles(4, 4, FColor(0, 0, 128));
 			});
 #endif
 	}
@@ -270,6 +282,16 @@ namespace ImmediatePhysics_Chaos
 
 		// @todo(ccaulfield): add materials
 		TUniquePtr<FChaosPhysicsMaterial> Material = MakeUnique<FChaosPhysicsMaterial>();
+		if (ChaosImmediate_Material_Resitution >= 0.0f)
+		{
+			Material->Restitution = ChaosImmediate_Material_Resitution;
+		}
+		if (ChaosImmediate_Material_Friction >= 0.0f)
+		{
+			Material->Friction = ChaosImmediate_Material_Friction;
+		}
+
+
 		ParticleMaterials.Add(MakeSerializable(Material));
 		PerParticleMaterials.Add(MoveTemp(Material));
 		CollidedParticles.Add(false);
@@ -348,17 +370,20 @@ namespace ImmediatePhysics_Chaos
 		{
 			FActorHandle* ActorHandle0 = ActorHandles[ActorHandleIndex0];
 			TGeometryParticleHandle<FReal, Dimensions>* Particle0 = ActorHandle0->GetParticle();
-
-			for (int ActorHandleIndex1 = ActorHandleIndex0 + 1; ActorHandleIndex1 < NumActorHandles; ++ActorHandleIndex1)
+			//TPBDRigidParticleHandle<FReal, Dimensions>* Rigid0 = Particle0->CastToRigidParticle();
+			//if (Rigid0 && (Rigid0->ObjectState() == EObjectStateType::Dynamic))
 			{
-				FActorHandle* ActorHandle1 = ActorHandles[ActorHandleIndex1];
-				TGeometryParticleHandle<FReal, Dimensions>* Particle1 = ActorHandle1->GetParticle();
-
-				const TSet<const FParticleHandle*>* Particle0IgnoreSet = IgnoreCollisionParticlePairTable.Find(Particle0);
-				bool bIgnoreActorHandlePair = (Particle0IgnoreSet != nullptr)  && Particle0IgnoreSet->Contains(Particle1);
-				if (!bIgnoreActorHandlePair)
+				for (int ActorHandleIndex1 = ActorHandleIndex0 + 1; ActorHandleIndex1 < NumActorHandles; ++ActorHandleIndex1)
 				{
-					PotentiallyCollidingPairs.Emplace(FParticlePair(Particle0, Particle1));
+					FActorHandle* ActorHandle1 = ActorHandles[ActorHandleIndex1];
+					TGeometryParticleHandle<FReal, Dimensions>* Particle1 = ActorHandle1->GetParticle();
+
+					const TSet<const FParticleHandle*>* Particle0IgnoreSet = IgnoreCollisionParticlePairTable.Find(Particle0);
+					bool bIgnoreActorHandlePair = (Particle0IgnoreSet != nullptr) && Particle0IgnoreSet->Contains(Particle1);
+					if (!bIgnoreActorHandlePair)
+					{
+						PotentiallyCollidingPairs.Emplace(FParticlePair(Particle0, Particle1));
+					}
 				}
 			}
 		}
@@ -488,7 +513,8 @@ namespace ImmediatePhysics_Chaos
 			Collisions.SetCollisionsEnabled(ChaosImmediate_Collision_Enabled != 0);
 		}
 
-		DebugDrawParticles(2, 2, 0.7f, true, true);
+		DebugDrawKinematicParticles(2, 2, FColor(128, 0, 0));
+		DebugDrawDynamicParticles(2, 2, FColor(192, 192, 0));
 		DebugDrawConstraints(2, 2, 0.7f);
 
 		ConditionConstraints();
@@ -499,12 +525,35 @@ namespace ImmediatePhysics_Chaos
 		
 		//Evolution.EndFrame(DeltaTime);
 
-		DebugDrawParticles(1, 2, 1.0f, true, true);
+		DebugDrawKinematicParticles(1, 4, FColor(128, 0, 0));
+		DebugDrawDynamicParticles(1, 4, FColor(255, 255, 0));
 		DebugDrawConstraints(1, 2, 1.0f);
-		DebugDrawParticles(3, 3, 1.0f, true, false);	// Kinematics only
 	}
 
-	void FSimulation::DebugDrawParticles(const int32 MinDebugLevel, const int32 MaxDebugLevel, const float ColorScale, bool bDrawKinematic, bool bDrawDynamic)
+	void FSimulation::DebugDrawKinematicParticles(const int32 MinDebugLevel, const int32 MaxDebugLevel, const FColor& Color)
+	{
+#if CHAOS_DEBUG_DRAW
+		using namespace Chaos;
+		if (FDebugDrawQueue::IsDebugDrawingEnabled())
+		{
+			if ((ChaosImmediate_DebugDrawParticles >= MinDebugLevel) && (ChaosImmediate_DebugDrawParticles <= MaxDebugLevel))
+			{
+				DebugDraw::DrawParticleTransforms(SimulationSpaceTransform, Particles.GetActiveKinematicParticlesView());
+			}
+			if ((ChaosImmediate_DebugDrawShapes >= MinDebugLevel) && (ChaosImmediate_DebugDrawShapes <= MaxDebugLevel))
+			{
+				DebugDraw::DrawParticleShapes(SimulationSpaceTransform, Particles.GetActiveKinematicParticlesView(), Color);
+			}
+			if ((ChaosImmediate_DebugDrawBounds >= MinDebugLevel) && (ChaosImmediate_DebugDrawBounds <= MaxDebugLevel))
+			{
+				DebugDraw::DrawParticleBounds(SimulationSpaceTransform, Particles.GetActiveKinematicParticlesView(), Color);
+			}
+		}
+#endif
+	}
+
+
+	void FSimulation::DebugDrawDynamicParticles(const int32 MinDebugLevel, const int32 MaxDebugLevel, const FColor& Color)
 	{
 #if CHAOS_DEBUG_DRAW
 		using namespace Chaos;
@@ -512,15 +561,15 @@ namespace ImmediatePhysics_Chaos
 		{
 			if ((ChaosImmediate_DebugDrawParticles >= MinDebugLevel) && (ChaosImmediate_DebugDrawParticles <= MaxDebugLevel))
 			{
-				DebugDraw::DrawParticleTransforms(SimulationSpaceTransform, Particles.GetAllParticlesView(), ColorScale, bDrawKinematic, bDrawDynamic);
+				DebugDraw::DrawParticleTransforms(SimulationSpaceTransform, Particles.GetActiveParticlesView());
 			}
 			if ((ChaosImmediate_DebugDrawShapes >= MinDebugLevel) && (ChaosImmediate_DebugDrawShapes <= MaxDebugLevel))
 			{
-				DebugDraw::DrawParticleShapes(SimulationSpaceTransform, Particles.GetAllParticlesView(), ColorScale, bDrawKinematic, bDrawDynamic);
+				DebugDraw::DrawParticleShapes(SimulationSpaceTransform, Particles.GetActiveParticlesView(), Color);
 			}
 			if ((ChaosImmediate_DebugDrawBounds >= MinDebugLevel) && (ChaosImmediate_DebugDrawBounds <= MaxDebugLevel))
 			{
-				DebugDraw::DrawParticleBounds(SimulationSpaceTransform, Particles.GetAllParticlesView(), ColorScale, bDrawKinematic, bDrawDynamic);
+				DebugDraw::DrawParticleBounds(SimulationSpaceTransform, Particles.GetActiveParticlesView(), Color);
 			}
 		}
 #endif
diff --git a/Engine/Source/Runtime/Engine/Public/Physics/ImmediatePhysics/ImmediatePhysicsChaos/ImmediatePhysicsSimulation_Chaos.h b/Engine/Source/Runtime/Engine/Public/Physics/ImmediatePhysics/ImmediatePhysicsChaos/ImmediatePhysicsSimulation_Chaos.h
--- a/Engine/Source/Runtime/Engine/Public/Physics/ImmediatePhysics/ImmediatePhysicsChaos/ImmediatePhysicsSimulation_Chaos.h
+++ b/Engine/Source/Runtime/Engine/Public/Physics/ImmediatePhysics/ImmediatePhysicsChaos/ImmediatePhysicsSimulation_Chaos.h
@@ -90,10 +90,9 @@ namespace ImmediatePhysics_Chaos
 	private:
 		void ConditionConstraints();
 
-		void DebugDrawParticles(const int32 MinDebugLevel, const int32 MaxDebugLevel, const float ColorScale, bool bDrawKinematic, bool bDrawDynamic);
+		void DebugDrawKinematicParticles(const int32 MinDebugLevel, const int32 MaxDebugLevel, const FColor& Color);
+		void DebugDrawDynamicParticles(const int32 MinDebugLevel, const int32 MaxDebugLevel, const FColor& Color);
 		void DebugDrawConstraints(const int32 MinDebugLevel, const int32 MaxDebugLevel, const float ColorScale);
-		void DebugDrawIslandParticles(const int32 Island, const int32 MinDebugLevel, const int32 MaxDebugLevel, const float ColorScale, bool bDrawKinematic, bool bDrawDynamic);
-		void DebugDrawIslandConstraints(const int32 Island, const int32 MinDebugLevel, const int32 MaxDebugLevel, const float ColorScale);
 
 		using FCollisionConstraints = Chaos::TPBDCollisionConstraints<FReal, 3>;
 		using FCollisionDetector = Chaos::TCollisionDetector<Chaos::FParticlePairBroadPhase, Chaos::FNarrowPhase, Chaos::FSyncCollisionReceiver, FCollisionConstraints>;
diff --git a/Engine/Source/Runtime/Experimental/Chaos/Private/Chaos/ChaosDebugDraw.cpp b/Engine/Source/Runtime/Experimental/Chaos/Private/Chaos/ChaosDebugDraw.cpp
--- a/Engine/Source/Runtime/Experimental/Chaos/Private/Chaos/ChaosDebugDraw.cpp
+++ b/Engine/Source/Runtime/Experimental/Chaos/Private/Chaos/ChaosDebugDraw.cpp
@@ -23,11 +23,12 @@ namespace Chaos
 		float ContactLen = 4.0f;
 		float ContactWidth = 2.0f;
 		float ContactPhiWidth = 1.5f;
+		float ContactOwnerWidth = 0.0f;
 		float ConstraintAxisLen = 5.0f;
 		float LineThickness = 0.15f;
 		float DrawScale = 1.0f;
 		float FontHeight = 10.0f;
-		float FontScale = 1.0f;
+		float FontScale = 1.5f;
 		float ShapeThicknesScale = 2.0f;
 		float PointSize = 2.0f;
 		int DrawPriority = 10.0f;
@@ -37,6 +38,7 @@ namespace Chaos
 		FAutoConsoleVariableRef CVarContactLen(TEXT("p.Chaos.DebugDrawContactLen"), ContactLen, TEXT("ContactLen."));
 		FAutoConsoleVariableRef CVarContactWidth(TEXT("p.Chaos.DebugDrawContactWidth"), ContactWidth, TEXT("ContactWidth."));
 		FAutoConsoleVariableRef CVarContactPhiWidth(TEXT("p.Chaos.DebugDrawContactPhiWidth"), ContactPhiWidth, TEXT("ContactPhiWidth."));
+		FAutoConsoleVariableRef CVarContactOwnerWidth(TEXT("p.Chaos.DebugDrawContactOwnerWidth"), ContactOwnerWidth, TEXT("ContactOwnerWidth."));
 		FAutoConsoleVariableRef CVarConstraintAxisLen(TEXT("p.Chaos.DebugDrawConstraintAxisLen"), ConstraintAxisLen, TEXT("ConstraintAxisLen."));
 		FAutoConsoleVariableRef CVarLineThickness(TEXT("p.Chaos.DebugDrawLineThickness"), LineThickness, TEXT("LineThickness."));
 		FAutoConsoleVariableRef CVarScale(TEXT("p.Chaos.DebugDrawScale"), DrawScale, TEXT("Scale applied to all Chaos Debug Draw line lengths etc."));
@@ -111,20 +113,16 @@ namespace Chaos
 			}
 		}
 
-		void DrawParticleShapesImpl(const FRigidTransform3& SpaceTransform, const TGeometryParticleHandle<FReal, 3>* Particle, FReal ColorScale)
+		void DrawParticleShapesImpl(const FRigidTransform3& SpaceTransform, const TGeometryParticleHandle<FReal, 3>* Particle, const FColor& Color)
 		{
-			FColor ShapeColor = Particle->ObjectState() == EObjectStateType::Dynamic ? FColor::Yellow : FColor::Orange;
-			FColor Color = (((FReal)0.5 * ColorScale) * ShapeColor).ToFColor(false);
 			FVec3 P = SpaceTransform.TransformPosition(Particle->ObjectState() == EObjectStateType::Dynamic ? Particle->CastToRigidParticle()->P() : Particle->X());
 			FRotation3 Q = SpaceTransform.GetRotation() * (Particle->ObjectState() == EObjectStateType::Dynamic ? Particle->CastToRigidParticle()->Q() : Particle->R());
 
 			DrawShapesImpl(FRigidTransform3(P, Q), Particle->Geometry().Get(), Color);
 		}
 
-		void DrawParticleBoundsImpl(const FRigidTransform3& SpaceTransform, const TGeometryParticleHandle<FReal, 3>* InParticle, FReal ColorScale)
+		void DrawParticleBoundsImpl(const FRigidTransform3& SpaceTransform, const TGeometryParticleHandle<FReal, 3>* InParticle, const FColor& Color)
 		{
-			FColor Color = (ColorScale * FColor::White).ToFColor(false);
-
 			TAABB<FReal, 3> Box = InParticle->WorldSpaceInflatedBounds();
 			FVec3 P = SpaceTransform.TransformPosition(Box.GetCenter());
 			FRotation3 Q = SpaceTransform.GetRotation();
@@ -132,7 +130,7 @@ namespace Chaos
 			FDebugDrawQueue::GetInstance().DrawDebugBox(P, 0.5f * Box.Extents(), Q, Color, false, KINDA_SMALL_NUMBER, DrawPriority, LineThickness);
 		}
 
-		void DrawParticleTransformImpl(const FRigidTransform3& SpaceTransform, const TGeometryParticleHandle<FReal, 3>* InParticle, FReal ColorScale)
+		void DrawParticleTransformImpl(const FRigidTransform3& SpaceTransform, const TGeometryParticleHandle<FReal, 3>* InParticle, int32 Index, FReal ColorScale)
 		{
 			FColor Red = (ColorScale * FColor::Red).ToFColor(false);
 			FColor Green = (ColorScale * FColor::Green).ToFColor(false);
@@ -151,6 +149,11 @@ namespace Chaos
 			FVec3 PActor = SpaceTransform.TransformPosition(FParticleUtilities::GetActorWorldTransform(Particle).GetTranslation());
 			FDebugDrawQueue::GetInstance().DrawDebugPoint(PActor, Black, false, KINDA_SMALL_NUMBER, DrawPriority, DrawScale * PointSize);
 			FDebugDrawQueue::GetInstance().DrawDebugLine(PCOM, PActor, Grey, false, KINDA_SMALL_NUMBER, DrawPriority, LineThickness);
+		
+			if (Index >= 0)
+			{
+				FDebugDrawQueue::GetInstance().DrawDebugString(PCOM + FontHeight * FVec3(0, 0, 1), FString::Format(TEXT("{0}{1}"), { Particle->IsKinematic()? TEXT("K"): TEXT("D"), Index }), nullptr, FColor::White, KINDA_SMALL_NUMBER, false, FontScale);
+			}
 		}
 
 		void DrawCollisionImpl(const FRigidTransform3& SpaceTransform, const TPBDCollisionConstraintHandle<float, 3>* ConstraintHandle, float ColorScale)
@@ -183,10 +186,19 @@ namespace Chaos
 					FMatrix Axes = FRotationMatrix::MakeFromX(Normal);
 					FDebugDrawQueue::GetInstance().DrawDebugCircle(Location - Contact.GetPhi() * Normal, DrawScale * ContactPhiWidth, 12, C2, false, KINDA_SMALL_NUMBER, DrawPriority, LineThickness, Axes.GetUnitAxis(EAxis::Y), Axes.GetUnitAxis(EAxis::Z), false);
 				}
+				if (ContactOwnerWidth > 0)
+				{
+					FColor C3 = (ColorScale * FColor(128, 128, 128)).ToFColor(false);
+					FMatrix Axes = FRotationMatrix::MakeFromX(Normal);
+					FVec3 P0 = SpaceTransform.TransformPosition(ConstraintHandle->GetConstrainedParticles()[0]->X());
+					FVec3 P1 = SpaceTransform.TransformPosition(ConstraintHandle->GetConstrainedParticles()[1]->X());
+					FDebugDrawQueue::GetInstance().DrawDebugLine(Location, P0, C3, false, KINDA_SMALL_NUMBER, DrawPriority, LineThickness);
+					FDebugDrawQueue::GetInstance().DrawDebugLine(Location, P1, C3, false, KINDA_SMALL_NUMBER, DrawPriority, LineThickness);
+				}
 			}
 		}
 
-		void DrawJointConstraintImpl(const FRigidTransform3& SpaceTransform, const FVec3& InPa, const FVec3& InXa, const FMatrix33& Ra, const FVec3& InPb, const FVec3& InXb, const FMatrix33& Rb, const FVec3& CR, int32 Level, FReal ColorScale, uint32 FeatureMask)
+		void DrawJointConstraintImpl(const FRigidTransform3& SpaceTransform, const FVec3& InPa, const FVec3& InXa, const FMatrix33& Ra, const FVec3& InPb, const FVec3& InXb, const FMatrix33& Rb, const FVec3& CR, int32 Level, int32 Index, FReal ColorScale, uint32 FeatureMask)
 		{
 			using namespace Chaos::DebugDraw;
 			FColor R = (ColorScale * FColor::Red).ToFColor(false);
@@ -220,6 +232,10 @@ namespace Chaos
 			{
 				FDebugDrawQueue::GetInstance().DrawDebugString(Xb + FontHeight * FVec3(0, 0, 1), FString::Format(TEXT("{0}"), { Level }), nullptr, FColor::Red, KINDA_SMALL_NUMBER, false, FontScale);
 			}
+			if ((FeatureMask & (uint32)EDebugDrawJointFeature::Index) && (Level >= 0))
+			{
+				FDebugDrawQueue::GetInstance().DrawDebugString(Xb + FontHeight * FVec3(0, 0, 1), FString::Format(TEXT("{0}"), { Index }), nullptr, FColor::Red, KINDA_SMALL_NUMBER, false, FontScale);
+			}
 			//FDebugDrawQueue::GetInstance().DrawDebugString(Xb + 3 * FontHeight * FVec3(0, 0, 1), FString::Format(TEXT("T=  {0}"), { FMath::RadiansToDegrees(CR[(int32)E6DJointAngularConstraintIndex::Twist]) }), nullptr, FColor::Red, KINDA_SMALL_NUMBER, false, FontScale);
 			//FDebugDrawQueue::GetInstance().DrawDebugString(Xb + 2 * FontHeight * FVec3(0, 0, 1), FString::Format(TEXT("S1= {0}"), { FMath::RadiansToDegrees(CR[(int32)E6DJointAngularConstraintIndex::Swing1]) }), nullptr, FColor::Red, KINDA_SMALL_NUMBER, false, FontScale);
 			//FDebugDrawQueue::GetInstance().DrawDebugString(Xb + 1 * FontHeight * FVec3(0, 0, 1), FString::Format(TEXT("S2= {0}"), { FMath::RadiansToDegrees(CR[(int32)E6DJointAngularConstraintIndex::Swing2]) }), nullptr, FColor::Red, KINDA_SMALL_NUMBER, false, FontScale);
@@ -237,103 +253,127 @@ namespace Chaos
 				FVec3 Xa, Xb, CR;
 				FMatrix33 Ra, Rb;
 				ConstraintHandle->CalculateConstraintSpace(Xa, Ra, Xb, Rb, CR);
-				DrawJointConstraintImpl(SpaceTransform, Pa, Xa, Ra, Pb, Xb, Rb, CR, ConstraintHandle->GetConstraintLevel(), ColorScale, FeatureMask);
+				DrawJointConstraintImpl(SpaceTransform, Pa, Xa, Ra, Pb, Xb, Rb, CR, ConstraintHandle->GetConstraintLevel(), ConstraintHandle->GetConstraintIndex(), ColorScale, FeatureMask);
 			}
 		}
 
 #endif
 
-		void DrawParticleShapes(const FRigidTransform3& SpaceTransform, const TParticleView<TGeometryParticles<float, 3>>& ParticlesView, float ColorScale, bool bDrawKinemtatic, bool bDrawDynamic)
+		void DrawParticleShapes(const FRigidTransform3& SpaceTransform, const TParticleView<TGeometryParticles<float, 3>>& ParticlesView, const FColor& Color)
 		{
 #if CHAOS_DEBUG_DRAW
 			if (FDebugDrawQueue::IsDebugDrawingEnabled())
 			{
 				for (auto& Particle : ParticlesView)
 				{
-					if ((bDrawKinemtatic && Particle.ObjectState() != EObjectStateType::Dynamic) || (bDrawDynamic && Particle.ObjectState() == EObjectStateType::Dynamic))
-					{
-						DrawParticleShapesImpl(SpaceTransform, GetHandleHelper(&Particle), ColorScale);
-					}
+					DrawParticleShapesImpl(SpaceTransform, GetHandleHelper(&Particle), Color);
 				}
 			}
 #endif
 		}
 
-		void DrawParticleShapes(const FRigidTransform3& SpaceTransform, const TArray<TGeometryParticleHandle<float, 3>*>& Particles, float ColorScale, bool bDrawKinemtatic, bool bDrawDynamic)
+		void DrawParticleShapes(const FRigidTransform3& SpaceTransform, const TParticleView<TKinematicGeometryParticles<float, 3>>& ParticlesView, const FColor& Color)
 		{
 #if CHAOS_DEBUG_DRAW
 			if (FDebugDrawQueue::IsDebugDrawingEnabled())
 			{
-				for (auto& Particle : Particles)
+				for (auto& Particle : ParticlesView)
 				{
-					if ((bDrawKinemtatic && Particle->ObjectState() != EObjectStateType::Dynamic) || (bDrawDynamic && Particle->ObjectState() == EObjectStateType::Dynamic))
-					{
-						DrawParticleShapesImpl(SpaceTransform, Particle, ColorScale);
-					}
+					DrawParticleShapesImpl(SpaceTransform, GetHandleHelper(&Particle), Color);
 				}
 			}
 #endif
 		}
 
-		void DrawParticleBounds(const FRigidTransform3& SpaceTransform, const TParticleView<TGeometryParticles<float, 3>>& ParticlesView, float ColorScale, bool bDrawKinemtatic, bool bDrawDynamic)
+		void DrawParticleShapes(const FRigidTransform3& SpaceTransform, const TParticleView<TPBDRigidParticles<float, 3>>& ParticlesView, const FColor& Color)
 		{
 #if CHAOS_DEBUG_DRAW
 			if (FDebugDrawQueue::IsDebugDrawingEnabled())
 			{
 				for (auto& Particle : ParticlesView)
 				{
-					if ((bDrawKinemtatic && Particle.ObjectState() != EObjectStateType::Dynamic) || (bDrawDynamic && Particle.ObjectState() == EObjectStateType::Dynamic))
-					{
-						DrawParticleBoundsImpl(SpaceTransform, GetHandleHelper(&Particle), ColorScale);
-					}
+					DrawParticleShapesImpl(SpaceTransform, GetHandleHelper(&Particle), Color);
 				}
 			}
 #endif
 		}
 
-		void DrawParticleBounds(const FRigidTransform3& SpaceTransform, const TArray<TGeometryParticleHandle<float, 3>*>& Particles, float ColorScale, bool bDrawKinemtatic, bool bDrawDynamic)
+		void DrawParticleBounds(const FRigidTransform3& SpaceTransform, const TParticleView<TGeometryParticles<float, 3>>& ParticlesView, const FColor& Color)
 		{
 #if CHAOS_DEBUG_DRAW
 			if (FDebugDrawQueue::IsDebugDrawingEnabled())
 			{
-				for (auto& Particle : Particles)
+				for (auto& Particle : ParticlesView)
 				{
-					if ((bDrawKinemtatic && Particle->ObjectState() != EObjectStateType::Dynamic) || (bDrawDynamic && Particle->ObjectState() == EObjectStateType::Dynamic))
-					{
-						DrawParticleBoundsImpl(SpaceTransform, Particle, ColorScale);
-					}
+					DrawParticleBoundsImpl(SpaceTransform, GetHandleHelper(&Particle), Color);
 				}
 			}
 #endif
 		}
 
-		void DrawParticleTransforms(const FRigidTransform3& SpaceTransform, const TParticleView<TGeometryParticles<float, 3>>& ParticlesView, float ColorScale, bool bDrawKinemtatic, bool bDrawDynamic)
+		void DrawParticleBounds(const FRigidTransform3& SpaceTransform, const TParticleView<TKinematicGeometryParticles<float, 3>>& ParticlesView, const FColor& Color)
 		{
 #if CHAOS_DEBUG_DRAW
 			if (FDebugDrawQueue::IsDebugDrawingEnabled())
 			{
 				for (auto& Particle : ParticlesView)
 				{
-					if ((bDrawKinemtatic && Particle.ObjectState() != EObjectStateType::Dynamic) || (bDrawDynamic && Particle.ObjectState() == EObjectStateType::Dynamic))
-					{
-						DrawParticleTransformImpl(SpaceTransform, GetHandleHelper(&Particle), ColorScale);
-					}
+					DrawParticleBoundsImpl(SpaceTransform, GetHandleHelper(&Particle), Color);
 				}
 			}
 #endif
 		}
 
-		void DrawParticleTransforms(const FRigidTransform3& SpaceTransform, const TArray<TGeometryParticleHandle<float, 3>*>& Particles, float ColorScale, bool bDrawKinemtatic, bool bDrawDynamic)
+		void DrawParticleBounds(const FRigidTransform3& SpaceTransform, const TParticleView<TPBDRigidParticles<float, 3>>& ParticlesView, const FColor& Color)
 		{
 #if CHAOS_DEBUG_DRAW
 			if (FDebugDrawQueue::IsDebugDrawingEnabled())
 			{
-				for (auto& Particle : Particles)
+				for (auto& Particle : ParticlesView)
 				{
-					if ((bDrawKinemtatic && Particle->ObjectState() != EObjectStateType::Dynamic) || (bDrawDynamic && Particle->ObjectState() == EObjectStateType::Dynamic))
-					{
-						DrawParticleTransformImpl(SpaceTransform, Particle, ColorScale);
-					}
+					DrawParticleBoundsImpl(SpaceTransform, GetHandleHelper(&Particle), Color);
+				}
+			}
+#endif
+		}
+
+		void DrawParticleTransforms(const FRigidTransform3& SpaceTransform, const TParticleView<TGeometryParticles<float, 3>>& ParticlesView)
+		{
+#if CHAOS_DEBUG_DRAW
+			if (FDebugDrawQueue::IsDebugDrawingEnabled())
+			{
+				int32 Index = 0;
+				for (auto& Particle : ParticlesView)
+				{
+					DrawParticleTransformImpl(SpaceTransform, GetHandleHelper(&Particle), Index++, 1.0f);
+				}
+			}
+#endif
+		}
+
+		void DrawParticleTransforms(const FRigidTransform3& SpaceTransform, const TParticleView<TKinematicGeometryParticles<float, 3>>& ParticlesView)
+		{
+#if CHAOS_DEBUG_DRAW
+			if (FDebugDrawQueue::IsDebugDrawingEnabled())
+			{
+				int32 Index = 0;
+				for (auto& Particle : ParticlesView)
+				{
+					DrawParticleTransformImpl(SpaceTransform, GetHandleHelper(&Particle), Index++, 1.0f);
+				}
+			}
+#endif
+		}
+
+		void DrawParticleTransforms(const FRigidTransform3& SpaceTransform, const TParticleView<TPBDRigidParticles<float, 3>>& ParticlesView)
+		{
+#if CHAOS_DEBUG_DRAW
+			if (FDebugDrawQueue::IsDebugDrawingEnabled())
+			{
+				int32 Index = 0;
+				for (auto& Particle : ParticlesView)
+				{
+					DrawParticleTransformImpl(SpaceTransform, GetHandleHelper(&Particle), Index++, 1.0f);
 				}
 			}
 #endif
diff --git a/Engine/Source/Runtime/Experimental/Chaos/Private/Chaos/Evolution/PBDMinEvolution.cpp b/Engine/Source/Runtime/Experimental/Chaos/Private/Chaos/Evolution/PBDMinEvolution.cpp
--- a/Engine/Source/Runtime/Experimental/Chaos/Private/Chaos/Evolution/PBDMinEvolution.cpp
+++ b/Engine/Source/Runtime/Experimental/Chaos/Private/Chaos/Evolution/PBDMinEvolution.cpp
@@ -146,6 +146,20 @@ namespace Chaos
 				}
 			}
 		}
+
+		for (TTransientGeometryParticleHandle<FReal, 3>& Particle : Particles.GetActiveKinematicParticlesView())
+		{
+			if (Particle.HasBounds())
+			{
+				const FReal BoundsThickness = 1.0f;
+				const FReal BoundsThicknessVelocityInflation = 2.0f;
+				const TAABB<FReal, 3>& LocalBounds = Particle.LocalBounds();
+				TAABB<FReal, 3> WorldSpaceBounds = LocalBounds.TransformedAABB(FRigidTransform3(Particle.X(), Particle.R()));
+				const FVec3 VAbs = (Particle.CastToKinematicParticle()) ? Particle.CastToKinematicParticle()->V().GetAbs() : FVec3(0);
+				WorldSpaceBounds.ThickenSymmetrically(FVec3(BoundsThickness) + BoundsThicknessVelocityInflation * VAbs.GetAbs() * Dt);
+				Particle.SetWorldSpaceInflatedBounds(WorldSpaceBounds);
+			}
+		}
 	}
 
 	// @todo(ccaulfield): dedupe (PBDRigidsEvolutionGBF)
diff --git a/Engine/Source/Runtime/Experimental/Chaos/Private/Chaos/Joint/PBDJointSolverCholesky.cpp b/Engine/Source/Runtime/Experimental/Chaos/Private/Chaos/Joint/PBDJointSolverCholesky.cpp
--- a/Engine/Source/Runtime/Experimental/Chaos/Private/Chaos/Joint/PBDJointSolverCholesky.cpp
+++ b/Engine/Source/Runtime/Experimental/Chaos/Private/Chaos/Joint/PBDJointSolverCholesky.cpp
@@ -342,10 +342,7 @@ namespace Chaos
 		bool bConstraintActive = (TwistAngle <= -TwistAngleMax) || (TwistAngle >= TwistAngleMax);
 		if (bConstraintActive)
 		{
-			const FVec3 Axis0 = Rs[0] * TwistAxis01;
-			const FVec3 Axis1 = Rs[1] * TwistAxis01;
-			const FVec3 XP0 = Xs[0] - Ps[0];
-			const FVec3 XP1 = Xs[1] - Ps[1];
+			const FVec3 Axis = Rs[0] * TwistAxis01;
 
 			const int32 RowIndex = J0.NumRows();
 			J0.AddRows(1);
@@ -353,10 +350,10 @@ namespace Chaos
 			C.AddRows(1);
 
 			J0.SetRowAt(RowIndex + 0, 0, 0, 0, 0);
-			J0.SetRowAt(RowIndex + 0, 3, Axis0);
+			J0.SetRowAt(RowIndex + 0, 3, Axis);
 
 			J1.SetRowAt(RowIndex + 0, 0, 0, 0, 0);
-			J1.SetRowAt(RowIndex + 0, 3, -Axis1);
+			J1.SetRowAt(RowIndex + 0, 3, -Axis);
 
 			C.SetAt(RowIndex, 0, (TwistAngle >= TwistAngleMax) ? TwistAngle - TwistAngleMax : TwistAngle + TwistAngleMax);
 		}
@@ -399,8 +396,6 @@ namespace Chaos
 		bool bConstraintActive = (SwingAngle <= -SwingAngleMax) || (SwingAngle >= SwingAngleMax);
 		if (bConstraintActive)
 		{
-			const FVec3 XP0 = Xs[0] - Ps[0];
-			const FVec3 XP1 = Xs[1] - Ps[1];
 			const FVec3 Axis = Rs[0] * SwingAxis01;
 
 			const int32 RowIndex = J0.NumRows();
@@ -462,8 +457,6 @@ namespace Chaos
 			bool bConstraintActive = (SwingAngle <= -SwingAngleMax) || (SwingAngle >= SwingAngleMax);
 			if (bConstraintActive)
 			{
-				const FVec3 XP0 = Xs[0] - Ps[0];
-				const FVec3 XP1 = Xs[1] - Ps[1];
 				const FVec3 Axis = SwingCross / SwingCrossLen;
 
 				const int32 RowIndex = J0.NumRows();
diff --git a/Engine/Source/Runtime/Experimental/Chaos/Private/Chaos/PBDJointConstraintUtilities.cpp b/Engine/Source/Runtime/Experimental/Chaos/Private/Chaos/PBDJointConstraintUtilities.cpp
--- a/Engine/Source/Runtime/Experimental/Chaos/Private/Chaos/PBDJointConstraintUtilities.cpp
+++ b/Engine/Source/Runtime/Experimental/Chaos/Private/Chaos/PBDJointConstraintUtilities.cpp
@@ -824,7 +824,6 @@ namespace Chaos
 		const FVec3 X1 = P1 + Q1 * XL1.GetTranslation();
 		const FRotation3 R0 = Q0 * XL0.GetRotation();
 		const FRotation3 R1 = Q1 * XL1.GetRotation();
-		const FVec3 CX = GetLimitedPositionError(JointSettings, R0, X1 - X0);
 
 		// Calculate the Twist Axis and Angle for each body
 		const FRotation3 R01 = R0.Inverse() * R1;
@@ -900,7 +899,6 @@ namespace Chaos
 		const FVec3 X1 = P1 + Q1 * XL1.GetTranslation();
 		const FRotation3 R0 = Q0 * XL0.GetRotation();
 		const FRotation3 R1 = Q1 * XL1.GetRotation();
-		const FVec3 CX = GetLimitedPositionError(JointSettings, R0, X1 - X0);
 
 		// Calculate the Twist Axis and Angle for each body
 		const FRotation3 R01 = R0.Inverse() * R1;
@@ -976,7 +974,6 @@ namespace Chaos
 		const FVec3 X1 = P1 + Q1 * XL1.GetTranslation();
 		const FRotation3 R0 = Q0 * XL0.GetRotation();
 		const FRotation3 R1 = Q1 * XL1.GetRotation();
-		const FVec3 CX = GetLimitedPositionError(JointSettings, R0, X1 - X0);
 
 		// Calculate Swing axis for each body
 		const FRotation3 R01 = R0.Inverse() * R1;
@@ -1135,7 +1132,6 @@ namespace Chaos
 		const FVec3 X1 = P1 + Q1 * XL1.GetTranslation();
 		const FRotation3 R0 = Q0 * XL0.GetRotation();
 		const FRotation3 R1 = Q1 * XL1.GetRotation();
-		const FVec3 CX = GetLimitedPositionError(JointSettings, R0, X1 - X0);
 
 		// Calculate the swing axis for each body
 		const FRotation3 R01 = R0.Inverse() * R1;
@@ -1326,7 +1322,6 @@ namespace Chaos
 		const FVec3 X1 = P1 + Q1 * XL1.GetTranslation();
 		const FRotation3 R0 = Q0 * XL0.GetRotation();
 		const FRotation3 R1 = Q1 * XL1.GetRotation();
-		const FVec3 CX = GetLimitedPositionError(JointSettings, R0, X1 - X0);
 
 		const FRotation3 R01 = R0.Inverse() * R1;
 		FRotation3 R01Twist, R01Swing;
@@ -1383,7 +1378,6 @@ namespace Chaos
 		const FVec3 X1 = P1 + Q1 * XL1.GetTranslation();
 		const FRotation3 R0 = Q0 * XL0.GetRotation();
 		const FRotation3 R1 = Q1 * XL1.GetRotation();
-		const FVec3 CX = GetLimitedPositionError(JointSettings, R0, X1 - X0);
 
 		// Calculate Swing axis for each body
 		const FRotation3 R01 = R0.Inverse() * R1;
@@ -1438,11 +1432,12 @@ namespace Chaos
 		FReal InvM1,
 		const FMatrix33& InvIL1)
 	{
-		const FVec3 X0 = P0 + Q0 * XL0.GetTranslation();
-		const FVec3 X1 = P1 + Q1 * XL1.GetTranslation();
+		const FVec3 DX0 = Q0 * XL0.GetTranslation();
+		const FVec3 DX1 = Q1 * XL1.GetTranslation();
+		const FVec3 X0 = P0 + DX0;
+		const FVec3 X1 = P1 + DX1;
 		const FRotation3 R0 = Q0 * XL0.GetRotation();
 		const FRotation3 R1 = Q1 * XL1.GetRotation();
-		const FVec3 CX = GetLimitedPositionError(JointSettings, R0, X1 - X0);
 
 		// Calculate the rotation we need to apply to resolve the rotation delta
 		const FRotation3 TargetR1 = R0 * JointSettings.Motion.AngularDriveTarget;
@@ -1464,10 +1459,17 @@ namespace Chaos
 			const FReal F0 = DriveStiffness * I0 / (I0 + I1);
 			const FReal F1 = DriveStiffness * I1 / (I0 + I1);
 
-			// Apply the rotation delta
+			// Apply the rotation delta about the connector
 			Q0 = FRotation3::Slerp(Q0, TargetQ0, F0);
 			Q1 = FRotation3::Slerp(Q1, TargetQ1, F1);
 			Q1.EnforceShortestArcWith(Q0);
+
+			// @todo(ccaulfield): this does not take into account the fact that some linear dofs may be inactive
+			const FVec3 X0_2 = P0 + Q0 * XL0.GetTranslation();
+			const FVec3 X1_2 = P1 + Q1 * XL1.GetTranslation();
+			const FVec3 Delta = (X1_2 - X0_2) - (X1 - X0);
+			P0 += (InvM0 / (InvM0 + InvM1)) * Delta;
+			P1 -= (InvM1 / (InvM0 + InvM1)) * Delta;
 		}
 	}
 
diff --git a/Engine/Source/Runtime/Experimental/Chaos/Private/Chaos/PBDJointConstraints.cpp b/Engine/Source/Runtime/Experimental/Chaos/Private/Chaos/PBDJointConstraints.cpp
--- a/Engine/Source/Runtime/Experimental/Chaos/Private/Chaos/PBDJointConstraints.cpp
+++ b/Engine/Source/Runtime/Experimental/Chaos/Private/Chaos/PBDJointConstraints.cpp
@@ -17,6 +17,7 @@
 
 namespace Chaos
 {
+	DECLARE_CYCLE_STAT(TEXT("TPBDJointConstraints::Sort"), STAT_Joints_Sort, STATGROUP_Chaos);
 	DECLARE_CYCLE_STAT(TEXT("TPBDJointConstraints::Apply"), STAT_Joints_Apply, STATGROUP_Chaos);
 	DECLARE_CYCLE_STAT(TEXT("TPBDJointConstraints::ApplyPushOut"), STAT_Joints_ApplyPushOut, STATGROUP_Chaos);
 	DECLARE_CYCLE_STAT(TEXT("TPBDJointConstraints::Drives"), STAT_Joints_Drives, STATGROUP_Chaos);
@@ -194,6 +195,7 @@ namespace Chaos
 		: Settings(InSettings)
 		, PreApplyCallback(nullptr)
 		, PostApplyCallback(nullptr)
+		, bRequiresSort(false)
 	{
 	}
 
@@ -300,18 +302,30 @@ namespace Chaos
 		PreApplyCallback = nullptr;
 	}
 
-	
+
 	void FPBDJointConstraints::SetPostApplyCallback(const FJointPostApplyCallback& Callback)
 	{
 		PostApplyCallback = Callback;
 	}
 
-	
+
 	void FPBDJointConstraints::ClearPostApplyCallback()
 	{
 		PostApplyCallback = nullptr;
 	}
 
+
+	void FPBDJointConstraints::SetPostProjectCallback(const FJointPostApplyCallback& Callback)
+	{
+		PostProjectCallback = Callback;
+	}
+
+
+	void FPBDJointConstraints::ClearPostProjectCallback()
+	{
+		PostProjectCallback = nullptr;
+	}
+
 	
 	const typename FPBDJointConstraints::FConstraintContainerHandle* FPBDJointConstraints::GetConstraintHandle(int32 ConstraintIndex) const
 	{
@@ -345,13 +359,54 @@ namespace Chaos
 	
 	void FPBDJointConstraints::SetParticleLevels(int32 ConstraintIndex, const TVector<int32, 2>& ParticleLevels)
 	{
-		ConstraintStates[ConstraintIndex].Level = FMath::Min(ParticleLevels[0], ParticleLevels[1]);
+		int32 NewLevel = FMath::Min(ParticleLevels[0], ParticleLevels[1]);
+		int32 PreviousLevel = ConstraintStates[ConstraintIndex].Level;
+		ConstraintStates[ConstraintIndex].Level = NewLevel;
 		ConstraintStates[ConstraintIndex].ParticleLevels = ParticleLevels;
+		bRequiresSort = bRequiresSort || (NewLevel != PreviousLevel);
+	}
+
+
+	void FPBDJointConstraints::SortConstraints()
+	{
+		// Sort constraints so that constraints with lower level (closer to a kinematic joint) are first
+		// @todo(ccaulfield): should probably also take islands/particle order into account
+		// @todo(ccaulfield): optimize (though isn't called very often)
+		SCOPE_CYCLE_COUNTER(STAT_Joints_Sort);
+
+		FHandles SortedHandles = Handles;
+		SortedHandles.StableSort([](const FConstraintContainerHandle& L, const FConstraintContainerHandle& R)
+			{
+				return L.GetConstraintLevel() < R.GetConstraintLevel();
+			});
+
+		TArray<FPBDJointSettings> SortedConstraintSettings;
+		TArray<FParticlePair> SortedConstraintParticles;
+		TArray<FPBDJointState> SortedConstraintStates;
+		int32 SortedConstraintIndex = 0;
+		for (FConstraintContainerHandle* Handle : SortedHandles)
+		{
+			int32 UnsortedIndex = Handle->GetConstraintIndex();
+			SortedConstraintSettings.Add(ConstraintSettings[UnsortedIndex]);
+			SortedConstraintParticles.Add(ConstraintParticles[UnsortedIndex]);
+			SortedConstraintStates.Add(ConstraintStates[UnsortedIndex]);
+			SetConstraintIndex(Handle, SortedConstraintIndex++);
+		}
+
+		Swap(ConstraintSettings, SortedConstraintSettings);
+		Swap(ConstraintParticles, SortedConstraintParticles);
+		Swap(ConstraintStates, SortedConstraintStates);
+		Swap(Handles, SortedHandles);
 	}
 
 	
 	void FPBDJointConstraints::UpdatePositionBasedState(const FReal Dt)
 	{
+		if (bRequiresSort)
+		{
+			SortConstraints();
+			bRequiresSort = false;
+		}
 	}
 
 	
@@ -408,6 +463,11 @@ namespace Chaos
 			SolvePosition(Dt, Settings.ApplyPairIterations, It, NumIts);
 		}
 
+		if (PostApplyCallback != nullptr)
+		{
+			PostApplyCallback(Dt, Handles);
+		}
+
 		// Correct remaining errors after last call to Solve if enabled in this phase
 		if (Settings.ProjectionPhase == EJointSolverPhase::Apply)
 		{
@@ -417,11 +477,6 @@ namespace Chaos
 				ApplyProjection(Dt);
 			}
 		}
-
-		if (PostApplyCallback != nullptr)
-		{
-			PostApplyCallback(Dt, Handles);
-		}
 	}
 
 	bool FPBDJointConstraints::ApplyPushOut(const FReal Dt, const int32 It, const int32 NumIts)
@@ -496,6 +551,11 @@ namespace Chaos
 		{
 			ProjectPosition(Dt, ConstraintIndex, 0, 1);
 		}
+
+		if (PostProjectCallback != nullptr)
+		{
+			PostProjectCallback(Dt, Handles);
+		}
 	}
 
 	//////////////////////////////////////////////////////////////////////////
diff --git a/Engine/Source/Runtime/Experimental/Chaos/Public/Chaos/ChaosDebugDraw.h b/Engine/Source/Runtime/Experimental/Chaos/Public/Chaos/ChaosDebugDraw.h
--- a/Engine/Source/Runtime/Experimental/Chaos/Public/Chaos/ChaosDebugDraw.h
+++ b/Engine/Source/Runtime/Experimental/Chaos/Public/Chaos/ChaosDebugDraw.h
@@ -3,6 +3,8 @@
 
 #include "Chaos/Declares.h"
 #include "Chaos/ChaosDebugDrawDeclares.h"
+#include "Chaos/KinematicGeometryParticles.h"
+#include "Chaos/PBDRigidParticles.h"
 #include "Chaos/Vector.h"
 
 namespace Chaos
@@ -17,17 +19,21 @@ namespace Chaos
 			Connector = 1 << 0,
 			Axes = 1 << 1,
 			Level = 1 << 2,
+			Index = 1 << 3,
 
-			Default = Axes,
+			Default = Connector | Index,
 			All = Connector | Axes | Level
 		};
 
-		CHAOS_API void DrawParticleShapes(const FRigidTransform3& SpaceTransform, const TParticleView<TGeometryParticles<float, 3>>& ParticlesView, float ColorScale, bool bDrawKinematic = true, bool bDrawDynamic = true);
-		CHAOS_API void DrawParticleShapes(const FRigidTransform3& SpaceTransform, const TArray<TGeometryParticleHandle<float, 3>*>& Particles, float ColorScale, bool bDrawKinematic = true, bool bDrawDynamic = true);
-		CHAOS_API void DrawParticleBounds(const FRigidTransform3& SpaceTransform, const TParticleView<TGeometryParticles<float, 3>>& ParticlesView, float ColorScale, bool bDrawKinematic = true, bool bDrawDynamic = true);
-		CHAOS_API void DrawParticleBounds(const FRigidTransform3& SpaceTransform, const TArray<TGeometryParticleHandle<float, 3>*>& Particles, float ColorScale, bool bDrawKinematic = true, bool bDrawDynamic = true);
-		CHAOS_API void DrawParticleTransforms(const FRigidTransform3& SpaceTransform, const TParticleView<TGeometryParticles<float, 3>>& ParticlesView, float ColorScale, bool bDrawKinematic = true, bool bDrawDynamic = true);
-		CHAOS_API void DrawParticleTransforms(const FRigidTransform3& SpaceTransform, const TArray<TGeometryParticleHandle<float, 3>*>& Particles, float ColorScale, bool bDrawKinematic = true, bool bDrawDynamic = true);
+		CHAOS_API void DrawParticleShapes(const FRigidTransform3& SpaceTransform, const TParticleView<TGeometryParticles<float, 3>>& ParticlesView, const FColor& Color);
+		CHAOS_API void DrawParticleShapes(const FRigidTransform3& SpaceTransform, const TParticleView<TKinematicGeometryParticles<float, 3>>& ParticlesView, const FColor& Color);
+		CHAOS_API void DrawParticleShapes(const FRigidTransform3& SpaceTransform, const TParticleView<TPBDRigidParticles<float, 3>>& ParticlesView, const FColor& Color);
+		CHAOS_API void DrawParticleBounds(const FRigidTransform3& SpaceTransform, const TParticleView<TGeometryParticles<float, 3>>& ParticlesView, const FColor& Color);
+		CHAOS_API void DrawParticleBounds(const FRigidTransform3& SpaceTransform, const TParticleView<TKinematicGeometryParticles<float, 3>>& ParticlesView, const FColor& Color);
+		CHAOS_API void DrawParticleBounds(const FRigidTransform3& SpaceTransform, const TParticleView<TPBDRigidParticles<float, 3>>& ParticlesView, const FColor& Color);
+		CHAOS_API void DrawParticleTransforms(const FRigidTransform3& SpaceTransform, const TParticleView<TGeometryParticles<float, 3>>& ParticlesView);
+		CHAOS_API void DrawParticleTransforms(const FRigidTransform3& SpaceTransform, const TParticleView<TKinematicGeometryParticles<float, 3>>& ParticlesView);
+		CHAOS_API void DrawParticleTransforms(const FRigidTransform3& SpaceTransform, const TParticleView<TPBDRigidParticles<float, 3>>& ParticlesView);
 		CHAOS_API void DrawParticleCollisions(const FRigidTransform3& SpaceTransform, const TGeometryParticleHandle<float, 3>* Particle, const TPBDCollisionConstraints<float, 3>& Collisions);
 		CHAOS_API void DrawCollisions(const FRigidTransform3& SpaceTransform, const TPBDCollisionConstraints<float, 3>& Collisions, float ColorScale);
 		CHAOS_API void DrawCollisions(const FRigidTransform3& SpaceTransform, const TArray<TPBDCollisionConstraintHandle<float, 3>*>& ConstraintHandles, float ColorScale);
diff --git a/Engine/Source/Runtime/Experimental/Chaos/Public/Chaos/Collision/ParticlePairBroadPhase.h b/Engine/Source/Runtime/Experimental/Chaos/Public/Chaos/Collision/ParticlePairBroadPhase.h
--- a/Engine/Source/Runtime/Experimental/Chaos/Public/Chaos/Collision/ParticlePairBroadPhase.h
+++ b/Engine/Source/Runtime/Experimental/Chaos/Public/Chaos/Collision/ParticlePairBroadPhase.h
@@ -42,11 +42,17 @@ namespace Chaos
 				TGeometryParticleHandle<FReal, 3>* Particle0 = const_cast<TGeometryParticleHandle<FReal, 3>*>(ParticlePair[0]);
 				TGeometryParticleHandle<FReal, 3>* Particle1 = const_cast<TGeometryParticleHandle<FReal, 3>*>(ParticlePair[1]);
 
-				const TAABB<FReal, 3>& Box0 = Particle0->WorldSpaceInflatedBounds();
-				const TAABB<FReal, 3>& Box1 = Particle1->WorldSpaceInflatedBounds();
-				if (Box0.Intersects(Box1))
+				// Particles may have been disabled or made kinematic
+				bool bAnyDisabled = TGenericParticleHandle<FReal, 3>(Particle0)->Disabled() || TGenericParticleHandle<FReal, 3>(Particle1)->Disabled();
+				bool bAnyDynamic = TGenericParticleHandle<FReal, 3>(Particle0)->IsDynamic() || TGenericParticleHandle<FReal, 3>(Particle1)->IsDynamic();
+				if (bAnyDynamic && !bAnyDisabled)
 				{
-					NarrowPhase.GenerateCollisions(Dt, Receiver, Particle0, Particle1, BoundsThickness, StatData);
+					const TAABB<FReal, 3>& Box0 = Particle0->WorldSpaceInflatedBounds();
+					const TAABB<FReal, 3>& Box1 = Particle1->WorldSpaceInflatedBounds();
+					if (Box0.Intersects(Box1))
+					{
+						NarrowPhase.GenerateCollisions(Dt, Receiver, Particle0, Particle1, BoundsThickness, StatData);
+					}
 				}
 			}
 		}
diff --git a/Engine/Source/Runtime/Experimental/Chaos/Public/Chaos/CollisionResolutionTypes.h b/Engine/Source/Runtime/Experimental/Chaos/Public/Chaos/CollisionResolutionTypes.h
--- a/Engine/Source/Runtime/Experimental/Chaos/Public/Chaos/CollisionResolutionTypes.h
+++ b/Engine/Source/Runtime/Experimental/Chaos/Public/Chaos/CollisionResolutionTypes.h
@@ -60,7 +60,7 @@ namespace Chaos
 
 		FString ToString() const
 		{
-			return FString::Printf(TEXT("Location:%s, Normal:%s, Phi:%s"), *Location.ToString(), *Normal.ToString(), Phi);
+			return FString::Printf(TEXT("Location:%s, Normal:%s, Phi:%f"), *Location.ToString(), *Normal.ToString(), Phi);
 		}
 
 		const FImplicitObject* Implicit[2]; // {Of Particle[0], Of Particle[1]}
diff --git a/Engine/Source/Runtime/Experimental/Chaos/Public/Chaos/PBDJointConstraints.h b/Engine/Source/Runtime/Experimental/Chaos/Public/Chaos/PBDJointConstraints.h
--- a/Engine/Source/Runtime/Experimental/Chaos/Public/Chaos/PBDJointConstraints.h
+++ b/Engine/Source/Runtime/Experimental/Chaos/Public/Chaos/PBDJointConstraints.h
@@ -98,6 +98,9 @@ namespace Chaos
 		void SetPostApplyCallback(const FJointPostApplyCallback& Callback);
 		void ClearPostApplyCallback();
 
+		void SetPostProjectCallback(const FJointPostApplyCallback& Callback);
+		void ClearPostProjectCallback();
+
 		//
 		// Constraint API
 		//
@@ -155,6 +158,7 @@ namespace Chaos
 
 		void CalculateConstraintSpace(int32 ConstraintIndex, FVec3& OutX0, FMatrix33& OutR0, FVec3& OutX1, FMatrix33& OutR1, FVec3& OutAngles) const;
 		void UpdateParticleState(TPBDRigidParticleHandle<FReal, 3>* Rigid, const FReal Dt, const FVec3& P, const FRotation3& Q, const bool bUpdateVelocity = true);
+		void SortConstraints();
 
 		void ApplyDrives(const FReal Dt, const int32 ConstraintIndex, const int32 It, const int32 NumIts);
 		void SolvePosition(const FReal Dt, const int32 ConstraintIndex, const int32 NumPairIts, const int32 It, const int32 NumIts);
@@ -177,6 +181,9 @@ namespace Chaos
 
 		FJointPreApplyCallback PreApplyCallback;
 		FJointPostApplyCallback PostApplyCallback;
+		FJointPostApplyCallback PostProjectCallback;
+
+		bool bRequiresSort;
 	};
 
 }
