Change ensure to CHAOS_ENSURE to not block build

diff --git a/Engine/Source/Runtime/Experimental/Chaos/Public/Chaos/GJK.h b/Engine/Source/Runtime/Experimental/Chaos/Public/Chaos/GJK.h
--- a/Engine/Source/Runtime/Experimental/Chaos/Public/Chaos/GJK.h
+++ b/Engine/Source/Runtime/Experimental/Chaos/Public/Chaos/GJK.h
@@ -550,7 +550,7 @@ namespace Chaos
 		{
 			if (InGJKPreDist2 > 1e-6 && InGJKPreDist2 < TNumericLimits<T>::Max())
 			{
-				ensure(Inflation > 0);	//shouldn't end up here if there is no inflation
+				CHAOS_ENSURE(Inflation > 0);	//shouldn't end up here if there is no inflation
 				OutNormal = Normal;
 				TVector<T, 3> ClosestA(0);
 				TVector<T, 3> ClosestB(0);
