Missing include

diff --git a/Engine/Source/Runtime/Experimental/Chaos/Public/Chaos/GJK.h b/Engine/Source/Runtime/Experimental/Chaos/Public/Chaos/GJK.h
--- a/Engine/Source/Runtime/Experimental/Chaos/Public/Chaos/GJK.h
+++ b/Engine/Source/Runtime/Experimental/Chaos/Public/Chaos/GJK.h
@@ -6,6 +6,7 @@
 #include "Chaos/Box.h"
 #include "Chaos/Sphere.h"
 #include "Chaos/EPA.h"
+#include "ChaosCheck.h"
 
 namespace Chaos
 {
