Change ensureMsgf to use CHAOS_ENSURE_MSG to unblock build

diff --git a/Engine/Source/Runtime/Experimental/Chaos/Public/Chaos/Framework/Handles.h b/Engine/Source/Runtime/Experimental/Chaos/Public/Chaos/Framework/Handles.h
--- a/Engine/Source/Runtime/Experimental/Chaos/Public/Chaos/Framework/Handles.h
+++ b/Engine/Source/Runtime/Experimental/Chaos/Public/Chaos/Framework/Handles.h
@@ -4,6 +4,7 @@
 
 #include "Chaos/Core.h"
 #include "../ChaosArchive.h"
+#include "ChaosCheck.h"
 
 namespace Chaos
 {
@@ -238,7 +239,7 @@ namespace Chaos
 				}
 				else
 				{
-					ensureMsgf(false, TEXT("Failed to access handle (%u, %u). NumEntries = %d, NumFlags = %d"), InHandle.Index, InHandle.Generation, HandleEntries.Num(), Validity.Num());
+					CHAOS_ENSURE_MSG(false, TEXT("Failed to access handle (%u, %u). NumEntries = %d, NumFlags = %d"), InHandle.Index, InHandle.Generation, HandleEntries.Num(), Validity.Num());
 				}
 			}
 			return nullptr;
diff --git a/Engine/Source/Runtime/Experimental/Chaos/Public/ChaosCheck.h b/Engine/Source/Runtime/Experimental/Chaos/Public/ChaosCheck.h
--- a/Engine/Source/Runtime/Experimental/Chaos/Public/ChaosCheck.h
+++ b/Engine/Source/Runtime/Experimental/Chaos/Public/ChaosCheck.h
@@ -7,7 +7,9 @@
 #if CHAOS_CHECKED
 #define CHAOS_CHECK(Condition) check(Condition)
 #define CHAOS_ENSURE(Condition) ensure(Condition)
+#define CHAOS_ENSURE_MSG(InExpression, InFormat, ... ) ensureMsgf(InExpression, InFormat, ##__VA_ARGS__)
 #else
 #define CHAOS_CHECK(Condition) (!!(Condition))
 #define CHAOS_ENSURE(Condition) (!!(Condition))
+#define CHAOS_ENSURE_MSG(InExpression, InFormat, ... ) (!!(InExpression))
 #endif
