Fix linker error

diff --git a/Engine/Source/Runtime/Experimental/Chaos/Private/Chaos/TriangleMeshImplicitObject.cpp b/Engine/Source/Runtime/Experimental/Chaos/Private/Chaos/TriangleMeshImplicitObject.cpp
--- a/Engine/Source/Runtime/Experimental/Chaos/Private/Chaos/TriangleMeshImplicitObject.cpp
+++ b/Engine/Source/Runtime/Experimental/Chaos/Private/Chaos/TriangleMeshImplicitObject.cpp
@@ -200,6 +200,13 @@ struct TTriangleMeshRaycastVisitor
 	int32 OutFaceIndex;
 };
 
+template <typename T>
+T TTriangleMeshImplicitObject<T>::PhiWithNormal(const TVector<T, 3>& x, TVector<T, 3>& Normal) const
+{
+	ensure(false);	//not supported yet - might support it in the future or we may change the interface
+	return (T)0;
+}
+
 template <typename T>
 bool TTriangleMeshImplicitObject<T>::Raycast(const TVector<T, 3>& StartPoint, const TVector<T, 3>& Dir, const T Length, const T Thickness, T& OutTime, TVector<T, 3>& OutPosition, TVector<T, 3>& OutNormal, int32& OutFaceIndex) const
 {
@@ -709,4 +716,4 @@ void Chaos::TTriangleMeshImplicitObject<T>::RebuildBV()
 }
 
 
-template class Chaos::TTriangleMeshImplicitObject<float>;
+template class Chaos::TTriangleMeshImplicitObject<float>;
\ No newline at end of file
diff --git a/Engine/Source/Runtime/Experimental/Chaos/Public/Chaos/TriangleMeshImplicitObject.h b/Engine/Source/Runtime/Experimental/Chaos/Public/Chaos/TriangleMeshImplicitObject.h
--- a/Engine/Source/Runtime/Experimental/Chaos/Public/Chaos/TriangleMeshImplicitObject.h
+++ b/Engine/Source/Runtime/Experimental/Chaos/Public/Chaos/TriangleMeshImplicitObject.h
@@ -33,12 +33,8 @@ namespace Chaos
 		TTriangleMeshImplicitObject(TTriangleMeshImplicitObject&& Other) = default;
 		virtual ~TTriangleMeshImplicitObject() {}
 
-		virtual T PhiWithNormal(const TVector<T, 3>& x, TVector<T, 3>& Normal) const
-		{
-			ensure(false);	//not supported yet - might support it in the future or we may change the interface
-			return TNumericLimits<T>::Max();
-		}
-
+		virtual T PhiWithNormal(const TVector<T, 3>& x, TVector<T, 3>& Normal) const;
+		
 		virtual bool Raycast(const TVector<T, 3>& StartPoint, const TVector<T, 3>& Dir, const T Length, const T Thickness, T& OutTime, TVector<T,3>& OutPosition, TVector<T,3>& OutNormal, int32& OutFaceIndex) const override;
 		virtual bool Overlap(const TVector<T, 3>& Point, const T Thickness) const override;
 
