Change ensures to CHAOS_ENSURES to prevent test failures.

diff --git a/Engine/Source/Runtime/Experimental/Chaos/Private/Chaos/Convex.cpp b/Engine/Source/Runtime/Experimental/Chaos/Private/Chaos/Convex.cpp
--- a/Engine/Source/Runtime/Experimental/Chaos/Private/Chaos/Convex.cpp
+++ b/Engine/Source/Runtime/Experimental/Chaos/Private/Chaos/Convex.cpp
@@ -26,7 +26,7 @@ namespace Chaos
 				}
 			}
 		}
-		ensure(MostOpposingIdx != INDEX_NONE);
+		CHAOS_ENSURE(MostOpposingIdx != INDEX_NONE);
 		return MostOpposingIdx;
 	}
 }
\ No newline at end of file
diff --git a/Engine/Source/Runtime/Experimental/Chaos/Public/Chaos/Convex.h b/Engine/Source/Runtime/Experimental/Chaos/Public/Chaos/Convex.h
--- a/Engine/Source/Runtime/Experimental/Chaos/Public/Chaos/Convex.h
+++ b/Engine/Source/Runtime/Experimental/Chaos/Public/Chaos/Convex.h
@@ -147,7 +147,7 @@ namespace Chaos
 			// For convexes, this function must be called with a face index.
 			// If this ensure is getting hit, fix the caller so that it
 			// passes in a valid face index.
-			if (ensure(FaceIndex != INDEX_NONE))
+			if (CHAOS_ENSURE(FaceIndex != INDEX_NONE))
 			{
 				const TPlane<FReal, 3>& OpposingFace = GetFaces()[FaceIndex];
 				return OpposingFace.Normal();
