Fix Linker Error in TriangleMeshImplicitObject.

diff --git a/Engine/Source/Runtime/Experimental/Chaos/Private/Chaos/TriangleMeshImplicitObject.cpp b/Engine/Source/Runtime/Experimental/Chaos/Private/Chaos/TriangleMeshImplicitObject.cpp
--- a/Engine/Source/Runtime/Experimental/Chaos/Private/Chaos/TriangleMeshImplicitObject.cpp
+++ b/Engine/Source/Runtime/Experimental/Chaos/Private/Chaos/TriangleMeshImplicitObject.cpp
@@ -677,6 +677,13 @@ TVector<T, 3> TTriangleMeshImplicitObject<T>::FindGeometryOpposingNormal(const T
 	return GetFaceNormal(FaceIndex);
 }
 
+template<typename T>
+void TTriangleMeshImplicitObject<T>::Serialize(FChaosArchive& Ar)
+{
+	FChaosArchiveScopedMemory ScopedMemory(Ar, GetTypeName());
+	SerializeImp(Ar);
+}
+
 template <typename T>
 TVector<T, 3> TTriangleMeshImplicitObject<T>::GetFaceNormal(const int32 FaceIdx) const
 {
diff --git a/Engine/Source/Runtime/Experimental/Chaos/Public/Chaos/TriangleMeshImplicitObject.h b/Engine/Source/Runtime/Experimental/Chaos/Public/Chaos/TriangleMeshImplicitObject.h
--- a/Engine/Source/Runtime/Experimental/Chaos/Public/Chaos/TriangleMeshImplicitObject.h
+++ b/Engine/Source/Runtime/Experimental/Chaos/Public/Chaos/TriangleMeshImplicitObject.h
@@ -108,11 +108,7 @@ namespace Chaos
 			}
 		}
 
-		virtual void Serialize(FChaosArchive& Ar) override
-		{
-			FChaosArchiveScopedMemory ScopedMemory(Ar, GetTypeName());
-			SerializeImp(Ar);
-		}
+		virtual void Serialize(FChaosArchive& Ar) override;
 
 		virtual uint32 GetTypeHash() const override
 		{
