Merging CL#10631815

#Chaos Fix linker error

diff --git a/Engine/Source/Runtime/Experimental/Chaos/Private/Chaos/TriangleMeshImplicitObject.cpp b/Engine/Source/Runtime/Experimental/Chaos/Private/Chaos/TriangleMeshImplicitObject.cpp
--- a/Engine/Source/Runtime/Experimental/Chaos/Private/Chaos/TriangleMeshImplicitObject.cpp
+++ b/Engine/Source/Runtime/Experimental/Chaos/Private/Chaos/TriangleMeshImplicitObject.cpp
@@ -684,6 +684,21 @@ void TTriangleMeshImplicitObject<T>::Serialize(FChaosArchive& Ar)
 	SerializeImp(Ar);
 }
 
+template<typename T>
+uint32 TTriangleMeshImplicitObject<T>::GetTypeHash() const
+{
+	uint32 Result = MParticles.GetTypeHash();
+	Result = HashCombine(Result, MLocalBoundingBox.GetTypeHash());
+
+	for(TVector<int32, 3> Tri : MElements)
+	{
+		uint32 TriHash = HashCombine(::GetTypeHash(Tri[0]), HashCombine(::GetTypeHash(Tri[1]), ::GetTypeHash(Tri[2])));
+		Result = HashCombine(Result, TriHash);
+	}
+
+	return Result;
+}
+
 template <typename T>
 TVector<T, 3> TTriangleMeshImplicitObject<T>::GetFaceNormal(const int32 FaceIdx) const
 {
diff --git a/Engine/Source/Runtime/Experimental/Chaos/Public/Chaos/TriangleMeshImplicitObject.h b/Engine/Source/Runtime/Experimental/Chaos/Public/Chaos/TriangleMeshImplicitObject.h
--- a/Engine/Source/Runtime/Experimental/Chaos/Public/Chaos/TriangleMeshImplicitObject.h
+++ b/Engine/Source/Runtime/Experimental/Chaos/Public/Chaos/TriangleMeshImplicitObject.h
@@ -110,19 +110,7 @@ namespace Chaos
 
 		virtual void Serialize(FChaosArchive& Ar) override;
 
-		virtual uint32 GetTypeHash() const override
-		{
-			uint32 Result = MParticles.GetTypeHash();
-			Result = HashCombine(Result, MLocalBoundingBox.GetTypeHash());
-
-			for(TVector<int32, 3> Tri : MElements)
-			{
-				uint32 TriHash = HashCombine(::GetTypeHash(Tri[0]), HashCombine(::GetTypeHash(Tri[1]), ::GetTypeHash(Tri[2])));
-				Result = HashCombine(Result, TriHash);
-			}
-
-			return Result;
-		}
+		virtual uint32 GetTypeHash() const override;
 
 		TVector<T, 3> GetFaceNormal(const int32 FaceIdx) const;
 
