Merging CL#10631844

#Chaos Linker fix

diff --git a/Engine/Source/Runtime/Experimental/Chaos/Private/Chaos/TriangleMeshImplicitObject.cpp b/Engine/Source/Runtime/Experimental/Chaos/Private/Chaos/TriangleMeshImplicitObject.cpp
--- a/Engine/Source/Runtime/Experimental/Chaos/Private/Chaos/TriangleMeshImplicitObject.cpp
+++ b/Engine/Source/Runtime/Experimental/Chaos/Private/Chaos/TriangleMeshImplicitObject.cpp
@@ -719,6 +719,16 @@ TVector<T, 3> TTriangleMeshImplicitObject<T>::GetFaceNormal(const int32 FaceIdx)
 	return TVector<T, 3>(0, 0, 1);
 }
 
+uint16 TTriangleMeshImplicitObject<T>::GetMaterialIndex(uint32 HintIndex) const
+{
+	if (MaterialIndices.IsValidIndex(HintIndex))
+	{
+		return MaterialIndices[HintIndex];
+	}
+
+	// 0 should always be the default material for a shape
+	return 0;
+}
 
 template<typename T>
 void Chaos::TTriangleMeshImplicitObject<T>::RebuildBV()
diff --git a/Engine/Source/Runtime/Experimental/Chaos/Public/Chaos/TriangleMeshImplicitObject.h b/Engine/Source/Runtime/Experimental/Chaos/Public/Chaos/TriangleMeshImplicitObject.h
--- a/Engine/Source/Runtime/Experimental/Chaos/Public/Chaos/TriangleMeshImplicitObject.h
+++ b/Engine/Source/Runtime/Experimental/Chaos/Public/Chaos/TriangleMeshImplicitObject.h
@@ -114,16 +114,7 @@ namespace Chaos
 
 		TVector<T, 3> GetFaceNormal(const int32 FaceIdx) const;
 
-		virtual uint16 GetMaterialIndex(uint32 HintIndex) const override
-		{
-			if (MaterialIndices.IsValidIndex(HintIndex))
-			{
-				return MaterialIndices[HintIndex];
-			}
-
-			// 0 should always be the default material for a shape
-			return 0;
-		}
+		virtual uint16 GetMaterialIndex(uint32 HintIndex) const override;
 
 	private:
 
