#Chaos CIS fix

diff --git a/Engine/Source/Runtime/Experimental/Chaos/Private/Chaos/TriangleMeshImplicitObject.cpp b/Engine/Source/Runtime/Experimental/Chaos/Private/Chaos/TriangleMeshImplicitObject.cpp
--- a/Engine/Source/Runtime/Experimental/Chaos/Private/Chaos/TriangleMeshImplicitObject.cpp
+++ b/Engine/Source/Runtime/Experimental/Chaos/Private/Chaos/TriangleMeshImplicitObject.cpp
@@ -719,6 +719,7 @@ TVector<T, 3> TTriangleMeshImplicitObject<T>::GetFaceNormal(const int32 FaceIdx)
 	return TVector<T, 3>(0, 0, 1);
 }
 
+template <typename T>
 uint16 TTriangleMeshImplicitObject<T>::GetMaterialIndex(uint32 HintIndex) const
 {
 	if (MaterialIndices.IsValidIndex(HintIndex))
