Fix EPA giving wrong face when first face passes lower bound upper bound due to precision

diff --git a/Engine/Source/Runtime/Experimental/Chaos/Public/Chaos/EPA.h b/Engine/Source/Runtime/Experimental/Chaos/Public/Chaos/EPA.h
--- a/Engine/Source/Runtime/Experimental/Chaos/Public/Chaos/EPA.h
+++ b/Engine/Source/Runtime/Experimental/Chaos/Public/Chaos/EPA.h
@@ -387,7 +387,7 @@ EPAResult EPA(TArray<TVec3<T>>& VertsABuffer, TArray<TVec3<T>>& VertsBBuffer, co
 	//TEPAEntry<T> BestEntry;
 	//BestEntry.Distance = 0;
 
-	TEPAEntry<T> LastEntry = Entries[0];
+	TEPAEntry<T> LastEntry = Queue.size() > 0 ? Entries[Queue.top().Idx] : Entries[0];
 
 	TArray<FEPAFloodEntry> VisibilityBorder;
 	int32 Iteration = 0;
