Chaos
- Joint Drives working again
- Go back to GS solver

diff --git a/Engine/Source/Runtime/Engine/Private/PhysicsEngine/ImmediatePhysics/ImmediatePhysicsChaos/ImmediatePhysicsJointHandle_Chaos.cpp b/Engine/Source/Runtime/Engine/Private/PhysicsEngine/ImmediatePhysics/ImmediatePhysicsChaos/ImmediatePhysicsJointHandle_Chaos.cpp
--- a/Engine/Source/Runtime/Engine/Private/PhysicsEngine/ImmediatePhysics/ImmediatePhysicsChaos/ImmediatePhysicsJointHandle_Chaos.cpp
+++ b/Engine/Source/Runtime/Engine/Private/PhysicsEngine/ImmediatePhysics/ImmediatePhysicsChaos/ImmediatePhysicsJointHandle_Chaos.cpp
@@ -103,6 +103,7 @@ namespace ImmediatePhysics_Chaos
 			FParticleUtilities::ActorLocalToParticleLocal(TGenericParticleHandle<FReal, 3>(Actor2->GetParticle()), ConstraintFrame2),
 		};
 		ConstraintSettings.Motion.Stiffness = ChaosImmediate_JointStiffness;
+
 		ConstraintSettings.Motion.LinearMotionTypes = 
 		{
 			static_cast<EJointMotionType>(ConstraintInstance->GetLinearXMotion()),
@@ -110,20 +111,17 @@ namespace ImmediatePhysics_Chaos
 			static_cast<EJointMotionType>(ConstraintInstance->GetLinearZMotion()),
 		};
 		ConstraintSettings.Motion.LinearLimit = ConstraintInstance->GetLinearLimit();
-		ConstraintSettings.Motion.AngularMotionTypes =
-		{
-			static_cast<EJointMotionType>(ConstraintInstance->GetAngularTwistMotion()),
-			static_cast<EJointMotionType>(ConstraintInstance->GetAngularSwing1Motion()),
-			static_cast<EJointMotionType>(ConstraintInstance->GetAngularSwing2Motion())
-		};
-		ConstraintSettings.Motion.AngularLimits =
-		{
-			FMath::DegreesToRadians(ConstraintInstance->GetAngularTwistLimit()),
-			FMath::DegreesToRadians(ConstraintInstance->GetAngularSwing1Limit()),
-			FMath::DegreesToRadians(ConstraintInstance->GetAngularSwing2Limit())
-		};
+
+		ConstraintSettings.Motion.AngularMotionTypes[(int32)EJointAngularConstraintIndex::Twist] = static_cast<EJointMotionType>(ConstraintInstance->GetAngularTwistMotion());
+		ConstraintSettings.Motion.AngularMotionTypes[(int32)EJointAngularConstraintIndex::Swing1] = static_cast<EJointMotionType>(ConstraintInstance->GetAngularSwing1Motion());
+		ConstraintSettings.Motion.AngularMotionTypes[(int32)EJointAngularConstraintIndex::Swing2] = static_cast<EJointMotionType>(ConstraintInstance->GetAngularSwing2Motion());
+		ConstraintSettings.Motion.AngularLimits[(int32)EJointAngularConstraintIndex::Twist] = FMath::DegreesToRadians(ConstraintInstance->GetAngularTwistLimit());
+		ConstraintSettings.Motion.AngularLimits[(int32)EJointAngularConstraintIndex::Swing1] = FMath::DegreesToRadians(ConstraintInstance->GetAngularSwing1Limit());
+		ConstraintSettings.Motion.AngularLimits[(int32)EJointAngularConstraintIndex::Swing2] = FMath::DegreesToRadians(ConstraintInstance->GetAngularSwing2Limit());
+
 		ConstraintSettings.Motion.LinearProjection = ConstraintInstance->IsProjectionEnabled() ? ChaosImmediate_JointMaxProjection : ChaosImmediate_JointMinProjection;
 		ConstraintSettings.Motion.AngularProjection = ConstraintInstance->IsProjectionEnabled() ? ChaosImmediate_JointMaxProjection : ChaosImmediate_JointMinProjection;
+
 		ConstraintSettings.Motion.bSoftLinearLimitsEnabled = ConstraintInstance->GetIsSoftLinearLimit();
 		ConstraintSettings.Motion.bSoftTwistLimitsEnabled = ConstraintInstance->GetIsSoftTwistLimit();
 		ConstraintSettings.Motion.bSoftSwingLimitsEnabled = ConstraintInstance->GetIsSoftSwingLimit();
@@ -131,12 +129,11 @@ namespace ImmediatePhysics_Chaos
 		ConstraintSettings.Motion.SoftTwistStiffness = ConvertSoftAngularStiffness(ConstraintInstance->GetSoftTwistLimitStiffness(), Actor1->GetInverseInertia(), Actor2->GetInverseInertia());
 		ConstraintSettings.Motion.SoftSwingStiffness = ConvertSoftAngularStiffness(ConstraintInstance->GetSoftSwingLimitStiffness(), Actor1->GetInverseInertia(), Actor2->GetInverseInertia());
 		 
-		// @todo(ccaulfield): Remove one of these
 		ConstraintSettings.Motion.AngularDriveTarget = FQuat(ConstraintInstance->ProfileInstance.AngularDrive.OrientationTarget);
-		ConstraintSettings.Motion.AngularDriveTargetAngles = TVector<float, 3>(
-			FMath::DegreesToRadians(ConstraintInstance->ProfileInstance.AngularDrive.OrientationTarget.Roll), 
-			FMath::DegreesToRadians(ConstraintInstance->ProfileInstance.AngularDrive.OrientationTarget.Pitch),
-			FMath::DegreesToRadians(ConstraintInstance->ProfileInstance.AngularDrive.OrientationTarget.Yaw));
+		// NOTE: Hard dependence on EJointAngularConstraintIndex - the following will break if we change the order
+		ConstraintSettings.Motion.AngularDriveTargetAngles[(int32)EJointAngularConstraintIndex::Twist] = FMath::DegreesToRadians(ConstraintInstance->ProfileInstance.AngularDrive.OrientationTarget.Roll);
+		ConstraintSettings.Motion.AngularDriveTargetAngles[(int32)EJointAngularConstraintIndex::Swing1] = FMath::DegreesToRadians(ConstraintInstance->ProfileInstance.AngularDrive.OrientationTarget.Yaw);
+		ConstraintSettings.Motion.AngularDriveTargetAngles[(int32)EJointAngularConstraintIndex::Swing2] = FMath::DegreesToRadians(ConstraintInstance->ProfileInstance.AngularDrive.OrientationTarget.Pitch);
 
 		if (ConstraintInstance->ProfileInstance.AngularDrive.AngularDriveMode == EAngularDriveMode::SLERP)
 		{
diff --git a/Engine/Source/Runtime/Engine/Private/PhysicsEngine/ImmediatePhysics/ImmediatePhysicsChaos/ImmediatePhysicsSimulation_Chaos.cpp b/Engine/Source/Runtime/Engine/Private/PhysicsEngine/ImmediatePhysics/ImmediatePhysicsChaos/ImmediatePhysicsSimulation_Chaos.cpp
--- a/Engine/Source/Runtime/Engine/Private/PhysicsEngine/ImmediatePhysics/ImmediatePhysicsChaos/ImmediatePhysicsSimulation_Chaos.cpp
+++ b/Engine/Source/Runtime/Engine/Private/PhysicsEngine/ImmediatePhysics/ImmediatePhysicsChaos/ImmediatePhysicsSimulation_Chaos.cpp
@@ -44,7 +44,6 @@ float ChaosImmediate_Joint_SwingTwistAngleTolerance = 1.0e-6f;
 int32 ChaosImmediate_Joint_EnableTwistLimits = 1;
 int32 ChaosImmediate_Joint_EnableSwingLimits = 1;
 int32 ChaosImmediate_Joint_EnableDrives = 1;
-int32 ChaosImmediate_Joint_DrivesPhase = (int32)Chaos::EJointSolverPhase::Apply;
 int32 ChaosImmediate_Joint_ProjectionPhase = (int32)Chaos::EJointSolverPhase::Apply;
 float ChaosImmediate_Joint_LinearProjection = 0.0f;
 float ChaosImmediate_Joint_AngularProjection = 0.0f;
@@ -60,7 +59,6 @@ FAutoConsoleVariableRef CVarChaosImmPhysSwingTwistAngleTolerance(TEXT("p.Chaos.I
 FAutoConsoleVariableRef CVarChaosImmPhysEnableTwistLimits(TEXT("p.Chaos.ImmPhys.Joint.EnableTwistLimits"), ChaosImmediate_Joint_EnableTwistLimits, TEXT("EnableTwistLimits."));
 FAutoConsoleVariableRef CVarChaosImmPhysEnableSwingLimits(TEXT("p.Chaos.ImmPhys.Joint.EnableSwingLimits"), ChaosImmediate_Joint_EnableSwingLimits, TEXT("EnableSwingLimits."));
 FAutoConsoleVariableRef CVarChaosImmPhysEnableDrives(TEXT("p.Chaos.ImmPhys.Joint.EnableDrives"), ChaosImmediate_Joint_EnableDrives, TEXT("EnableDrives."));
-FAutoConsoleVariableRef CVarChaosImmPhysDrivesPhase(TEXT("p.Chaos.ImmPhys.Joint.DrivesPhase"), ChaosImmediate_Joint_DrivesPhase, TEXT("Drives Phase (0=Never; 1=Apply; 2=ApplyPushOut)."));
 FAutoConsoleVariableRef CVarChaosImmPhysProjectionPhase(TEXT("p.Chaos.ImmPhys.Joint.ProjectionPhase"), ChaosImmediate_Joint_ProjectionPhase, TEXT("Projection Phase (0=Never; 1=Apply; 2=ApplyPushOut)."));
 FAutoConsoleVariableRef CVarChaosImmPhysLinearProjection(TEXT("p.Chaos.ImmPhys.Joint.LinearProjection"), ChaosImmediate_Joint_LinearProjection, TEXT("6Dof joint projection amount override (if > 0)."));
 FAutoConsoleVariableRef CVarChaosImmPhysAngularProjection(TEXT("p.Chaos.ImmPhys.Joint.AngularProjection"), ChaosImmediate_Joint_AngularProjection, TEXT("6Dof joint projection amount override (if > 0)."));
@@ -490,7 +488,6 @@ namespace ImmediatePhysics_Chaos
 			JointsSettings.bEnableTwistLimits = ChaosImmediate_Joint_EnableTwistLimits != 0;
 			JointsSettings.bEnableSwingLimits = ChaosImmediate_Joint_EnableSwingLimits != 0;
 			JointsSettings.bEnableDrives = ChaosImmediate_Joint_EnableDrives != 0;
-			JointsSettings.DrivesPhase = ToJointSolverPhase(ChaosImmediate_Joint_DrivesPhase);
 			JointsSettings.ProjectionPhase = ToJointSolverPhase(ChaosImmediate_Joint_ProjectionPhase);
 			JointsSettings.LinearProjection = ChaosImmediate_Joint_LinearProjection;
 			JointsSettings.AngularProjection = ChaosImmediate_Joint_AngularProjection;
diff --git a/Engine/Source/Runtime/Experimental/Chaos/Private/Chaos/Joint/PBDJointSolverCholesky.cpp b/Engine/Source/Runtime/Experimental/Chaos/Private/Chaos/Joint/PBDJointSolverCholesky.cpp
--- a/Engine/Source/Runtime/Experimental/Chaos/Private/Chaos/Joint/PBDJointSolverCholesky.cpp
+++ b/Engine/Source/Runtime/Experimental/Chaos/Private/Chaos/Joint/PBDJointSolverCholesky.cpp
@@ -14,7 +14,8 @@
 namespace Chaos
 {
 #if CHAOS_JOINTSOLVER_STATSENABLED
-	DECLARE_CYCLE_STAT(TEXT("TPBDJointConstraints::Apply"), STAT_JointSolver_Apply, STATGROUP_Chaos);
+	DECLARE_CYCLE_STAT(TEXT("TPBDJointConstraints::ApplyConstraints"), STAT_JointSolver_ApplyConstraints, STATGROUP_Chaos);
+	DECLARE_CYCLE_STAT(TEXT("TPBDJointConstraints::ApplyDrives"), STAT_JointSolver_ApplyDrives, STATGROUP_Chaos);
 	DECLARE_CYCLE_STAT(TEXT("TPBDJointConstraints::ApplyPushOut"), STAT_JointSolver_Jacobian, STATGROUP_Chaos);
 #define CHAOS_JOINTSOLVER_SCOPE_CYCLE_STAT SCOPE_CYCLE_COUNTER(X)
 #else
@@ -63,9 +64,11 @@ namespace Chaos
 		Qs[1].EnforceShortestArcWith(Qs[0]);
 
 		Stiffness = FPBDJointUtilities::GetLinearStiffness(SolverSettings, JointSettings);
+		AngularDriveStiffness = FPBDJointUtilities::GetAngularDriveStiffness(SolverSettings, JointSettings);
 		SwingTwistAngleTolerance = SolverSettings.SwingTwistAngleTolerance;
 		bEnableTwistLimits = SolverSettings.bEnableTwistLimits;
-		bEnableTwistLimits = SolverSettings.bEnableSwingLimits;
+		bEnableSwingLimits = SolverSettings.bEnableSwingLimits;
+		bEnableDrives = SolverSettings.bEnableDrives;
 
 		UpdateDerivedState();
 	}
@@ -74,8 +77,68 @@ namespace Chaos
 		const FReal Dt,
 		const FPBDJointSettings& JointSettings)
 	{
-		CHAOS_JOINTSOLVER_SCOPE_CYCLE_STAT(STAT_JointSolver_Apply);
+		CHAOS_JOINTSOLVER_SCOPE_CYCLE_STAT(STAT_JointSolver_ApplyConstraints);
 
+		FDenseMatrix61 C;
+		FDenseMatrix66 J0, J1;
+		BuildJacobianAndResidual_Constraints(JointSettings, J0, J1, C);
+
+		SolveAndApply(JointSettings, J0, J1, C);
+	}
+
+	void FJointSolverCholesky::ApplyDrives(
+		const FReal Dt,
+		const FPBDJointSettings& JointSettings)
+	{
+		CHAOS_JOINTSOLVER_SCOPE_CYCLE_STAT(STAT_JointSolver_ApplyDrives);
+
+		FDenseMatrix61 C;
+		FDenseMatrix66 J0, J1;
+		BuildJacobianAndResidual_Drives(JointSettings, J0, J1, C);
+
+		SolveAndApply(JointSettings, J0, J1, C);
+	}
+
+	void FJointSolverCholesky::BuildJacobianAndResidual_Constraints(
+		const FPBDJointSettings& JointSettings,
+		FDenseMatrix66& J0,
+		FDenseMatrix66& J1,
+		FDenseMatrix61& C)
+	{
+		CHAOS_JOINTSOLVER_SCOPE_CYCLE_STAT(STAT_JointSolver_Jacobian);
+
+		// Initialize with zero constraints
+		J0.SetDimensions(0, 6);
+		J1.SetDimensions(0, 6);
+		C.SetDimensions(0, 1);
+
+		AddLinearConstraints(JointSettings, J0, J1, C);
+		AddAngularConstraints(JointSettings, J0, J1, C);
+	}
+
+	void FJointSolverCholesky::BuildJacobianAndResidual_Drives(
+		const FPBDJointSettings& JointSettings,
+		FDenseMatrix66& J0,
+		FDenseMatrix66& J1,
+		FDenseMatrix61& C)
+	{
+		CHAOS_JOINTSOLVER_SCOPE_CYCLE_STAT(STAT_JointSolver_Jacobian);
+
+		// Initialize with zero constraints
+		J0.SetDimensions(0, 6);
+		J1.SetDimensions(0, 6);
+		C.SetDimensions(0, 1);
+
+		//AddLinearConstraints(JointSettings, J0, J1, C);
+		AddAngularDrives(JointSettings, J0, J1, C);
+	}
+
+	void FJointSolverCholesky::SolveAndApply(
+		const FPBDJointSettings& JointSettings,
+		const FDenseMatrix66& J0,
+		const FDenseMatrix66& J1,
+		const FDenseMatrix61& C)
+	{
 		// Solving for world-space position/rotation corrections D(6x1) where
 		// D = [I.Jt / [J.I.Jt]].C = I.Jt.L
 		// I is the inverse mass matrix, J the Jacobian, C the current constraint violation,
@@ -101,9 +164,6 @@ namespace Chaos
 		//           | 0              -Swing1Axis         |
 		//           | 0              -Swing2Axis         |
 		//
-		FDenseMatrix61 C;
-		FDenseMatrix66 J0, J1;
-		BuildJacobianAndResidual(JointSettings, J0, J1, C);
 
 		// InvM(6x6) = inverse mass matrix
 		const FMassMatrix InvM0 = FMassMatrix::Make(InvMs[0], Qs[0], InvILs[0]);
@@ -142,23 +202,52 @@ namespace Chaos
 		}
 	}
 
-	void FJointSolverCholesky::BuildJacobianAndResidual(
-		const FPBDJointSettings& JointSettings,
+	void FJointSolverCholesky::AddLinearRow(
+		const FVec3& Axis,
+		const FVec3& Connector0,
+		const FVec3& Connector1,
+		const FReal Error,
 		FDenseMatrix66& J0,
 		FDenseMatrix66& J1,
 		FDenseMatrix61& C)
 	{
-		CHAOS_JOINTSOLVER_SCOPE_CYCLE_STAT(STAT_JointSolver_Jacobian);
+		const int32 RowIndex = J0.NumRows();
+		J0.AddRows(1);
+		J1.AddRows(1);
+		C.AddRows(1);
 
-		// Initialize with zero constraints
-		J0.SetDimensions(0, 6);
-		J1.SetDimensions(0, 6);
-		C.SetDimensions(0, 1);
+		J0.SetRowAt(RowIndex, 0, Axis);
+		J0.SetRowAt(RowIndex, 3, -FVec3::CrossProduct(Axis, Connector0));
 
-		AddLinearConstraints(JointSettings, J0, J1, C);
-		AddAngularConstraints(JointSettings, J0, J1, C);
+		J1.SetRowAt(RowIndex, 0, -Axis);
+		J1.SetRowAt(RowIndex, 3, FVec3::CrossProduct(Axis, Connector1));
+
+		C.SetAt(RowIndex, 0, Error);
+	}
+
+	void FJointSolverCholesky::AddAngularRow(
+		const FVec3& Axis0,
+		const FVec3& Axis1,
+		const FReal Error,
+		FDenseMatrix66& J0,
+		FDenseMatrix66& J1,
+		FDenseMatrix61& C)
+	{
+		const int32 RowIndex = J0.NumRows();
+		J0.AddRows(1);
+		J1.AddRows(1);
+		C.AddRows(1);
+
+		J0.SetRowAt(RowIndex, 0, 0, 0, 0);
+		J0.SetRowAt(RowIndex, 3, Axis0);
+
+		J1.SetRowAt(RowIndex, 0, 0, 0, 0);
+		J1.SetRowAt(RowIndex, 3, -Axis1);
+
+		C.SetAt(RowIndex, 0, Error);
 	}
 
+
 	// 3 constraints along principle axes
 	void FJointSolverCholesky::AddLinearConstraints_Point(
 		const FPBDJointSettings& JointSettings,
@@ -211,19 +300,9 @@ namespace Chaos
 			const FVec3 XP0 = Xs[0] - Ps[0];
 			const FVec3 XP1 = Xs[1] - Ps[1];
 			const FVec3 Axis = ConstraintSeparation / ConstraintSeparationLen;
+			const FReal Error = ConstraintSeparationLen - Limit;
 
-			const int32 RowIndex = J0.NumRows();
-			J0.AddRows(1);
-			J1.AddRows(1);
-			C.AddRows(1);
-
-			J0.SetRowAt(RowIndex, 0, Axis);
-			J0.SetRowAt(RowIndex, 3, -FVec3::CrossProduct(Axis, XP0));
-
-			J1.SetRowAt(RowIndex, 0, -Axis);
-			J1.SetRowAt(RowIndex, 3, FVec3::CrossProduct(Axis, XP1));
-
-			C.SetAt(RowIndex, 0, ConstraintSeparationLen - Limit);
+			AddLinearRow(Axis, XP0, XP1, Error, J0, J1, C);
 		}
 	}
 
@@ -245,18 +324,9 @@ namespace Chaos
 		bool bAxisConstraintActive = (AxisMotion != EJointMotionType::Free);
 		if (bAxisConstraintActive)
 		{
-			const int32 RowIndex = J0.NumRows();
-			J0.AddRows(1);
-			J1.AddRows(1);
-			C.AddRows(1);
-
-			J0.SetRowAt(RowIndex, 0, Axis);
-			J0.SetRowAt(RowIndex, 3, -FVec3::CrossProduct(Axis, XP0));
+			const FReal Error = ConstraintDistanceAxial;
 
-			J1.SetRowAt(RowIndex, 0, -Axis);
-			J1.SetRowAt(RowIndex, 3, FVec3::CrossProduct(Axis, XP1));
-
-			C.SetAt(RowIndex, 0, ConstraintDistanceAxial);
+			AddLinearRow(Axis, XP0, XP1, Error, J0, J1, C);
 		}
 
 		// Radial Constraint
@@ -267,19 +337,9 @@ namespace Chaos
 		if (bRadialConstraintActive)
 		{
 			const FVec3 RadialAxis = ConstraintSeparationRadial / ConstraintDistanceRadial;
+			const FReal Error = ConstraintDistanceRadial - RadialLimit;
 
-			const int32 RowIndex = J0.NumRows();
-			J0.AddRows(1);
-			J1.AddRows(1);
-			C.AddRows(1);
-
-			J0.SetRowAt(RowIndex, 0, RadialAxis);
-			J0.SetRowAt(RowIndex, 3, -FVec3::CrossProduct(RadialAxis, XP0));
-
-			J1.SetRowAt(RowIndex, 0, -RadialAxis);
-			J1.SetRowAt(RowIndex, 3, FVec3::CrossProduct(RadialAxis, XP1));
-
-			C.SetAt(RowIndex, 0, ConstraintDistanceRadial - RadialLimit);
+			AddLinearRow(RadialAxis, XP0, XP1, Error, J0, J1, C);
 		}
 	}
 
@@ -302,19 +362,9 @@ namespace Chaos
 		{
 			const FVec3 XP0 = Xs[0] - Ps[0];
 			const FVec3 XP1 = Xs[1] - Ps[1];
+			const FReal Error = (ConstraintDistanceAxial >= Limit) ? ConstraintDistanceAxial - Limit : ConstraintDistanceAxial + Limit;
 
-			const int32 RowIndex = J0.NumRows();
-			J0.AddRows(1);
-			J1.AddRows(1);
-			C.AddRows(1);
-
-			J0.SetRowAt(RowIndex, 0, Axis);
-			J0.SetRowAt(RowIndex, 3, -FVec3::CrossProduct(Axis, XP0));
-
-			J1.SetRowAt(RowIndex, 0, -Axis);
-			J1.SetRowAt(RowIndex, 3, FVec3::CrossProduct(Axis, XP1));
-
-			C.SetAt(RowIndex, 0, (ConstraintDistanceAxial >= Limit) ? ConstraintDistanceAxial - Limit : ConstraintDistanceAxial + Limit);
+			AddLinearRow(Axis, XP0, XP1, Error, J0, J1, C);
 		}
 	}
 
@@ -342,20 +392,11 @@ namespace Chaos
 		bool bConstraintActive = (TwistAngle <= -TwistAngleMax) || (TwistAngle >= TwistAngleMax);
 		if (bConstraintActive)
 		{
-			const FVec3 Axis = Rs[0] * TwistAxis01;
-
-			const int32 RowIndex = J0.NumRows();
-			J0.AddRows(1);
-			J1.AddRows(1);
-			C.AddRows(1);
-
-			J0.SetRowAt(RowIndex + 0, 0, 0, 0, 0);
-			J0.SetRowAt(RowIndex + 0, 3, Axis);
-
-			J1.SetRowAt(RowIndex + 0, 0, 0, 0, 0);
-			J1.SetRowAt(RowIndex + 0, 3, -Axis);
+			const FVec3 Axis0 = Rs[0] * TwistAxis01;
+			const FVec3 Axis1 = Rs[1] * TwistAxis01;
+			const FReal Error = (TwistAngle >= TwistAngleMax) ? TwistAngle - TwistAngleMax : TwistAngle + TwistAngleMax;
 
-			C.SetAt(RowIndex, 0, (TwistAngle >= TwistAngleMax) ? TwistAngle - TwistAngleMax : TwistAngle + TwistAngleMax);
+			AddAngularRow(Axis0, Axis1, Error, J0, J1, C);
 		}
 	}
 
@@ -397,22 +438,18 @@ namespace Chaos
 		if (bConstraintActive)
 		{
 			const FVec3 Axis = Rs[0] * SwingAxis01;
+			const FReal Error = (SwingAngle >= SwingAngleMax) ? SwingAngle - SwingAngleMax : SwingAngle + SwingAngleMax;
 
-			const int32 RowIndex = J0.NumRows();
-			J0.AddRows(1);
-			J1.AddRows(1);
-			C.AddRows(1);
-
-			J0.SetRowAt(RowIndex + 0, 0, 0, 0, 0);
-			J0.SetRowAt(RowIndex + 0, 3, Axis);
-
-			J1.SetRowAt(RowIndex + 0, 0, 0, 0, 0);
-			J1.SetRowAt(RowIndex + 0, 3, -Axis);
-
-			C.SetAt(RowIndex, 0, (SwingAngle >= SwingAngleMax) ? SwingAngle - SwingAngleMax : SwingAngle + SwingAngleMax);
+			AddAngularRow(Axis, Axis, Error, J0, J1, C);
 		}
 	}
 
+	float GetSwingAngle(const FReal SwingY, const FReal SwingW)
+	{
+		return 4.0f * FMath::Atan2(SwingY, 1.0f + SwingW);
+	}
+
+
 	// up to 1 constraint limiting rotation about swing axis (relative to body 0)
 	void FJointSolverCholesky::AddAngularConstraints_Swing(
 		const FPBDJointSettings& JointSettings,
@@ -453,28 +490,85 @@ namespace Chaos
 				SwingAngle = (FReal)PI - SwingAngle;
 			}
 
+			const FMatrix33 RM0 = Rs[0].ToMatrix();
+			const FReal R01SwingYorZ = (SwingAxisIndex == EJointAngularAxisIndex::Swing1) ? R01Swing.Y : R01Swing.Z;
+			const FReal SwingAngle2 = GetSwingAngle(R01SwingYorZ, R01Swing.W);
+
+
 			FReal SwingAngleMax = JointSettings.Motion.AngularLimits[(int32)SwingConstraintIndex];
 			bool bConstraintActive = (SwingAngle <= -SwingAngleMax) || (SwingAngle >= SwingAngleMax);
 			if (bConstraintActive)
 			{
 				const FVec3 Axis = SwingCross / SwingCrossLen;
+				const FReal Error = (SwingAngle >= SwingAngleMax) ? SwingAngle - SwingAngleMax : SwingAngle + SwingAngleMax;
 
-				const int32 RowIndex = J0.NumRows();
-				J0.AddRows(1);
-				J1.AddRows(1);
-				C.AddRows(1);
+				AddAngularRow(Axis, Axis, Error, J0, J1, C);
+			}
+		}
+	}
 
-				J0.SetRowAt(RowIndex + 0, 0, 0, 0, 0);
-				J0.SetRowAt(RowIndex + 0, 3, Axis);
+	//void FJointSolverCholesky::AddAngularConstraints_Swing(
+	//	const FPBDJointSettings& JointSettings,
+	//	const EJointAngularConstraintIndex SwingConstraintIndex,
+	//	const EJointAngularAxisIndex SwingAxisIndex,
+	//	const FRotation3& R01Twist,
+	//	const FRotation3& R01Swing,
+	//	FDenseMatrix66& J0,
+	//	FDenseMatrix66& J1,
+	//	FDenseMatrix61& C)
+	//{
+	//	const FMatrix33 RM0 = Rs[0].ToMatrix();
+	//	const FReal R01SwingYorZ = ((int32)SwingAxisIndex == 2) ? R01Swing.Z : R01Swing.Y;	// Can't index a quat :(
+	//	FReal SwingAngle = GetSwingAngle(R01SwingYorZ, R01Swing.W);
+
+	//	FReal SwingAngleMax = JointSettings.Motion.AngularLimits[(int32)SwingConstraintIndex];
+	//	bool bConstraintActive = (SwingAngle <= -SwingAngleMax) || (SwingAngle >= SwingAngleMax);
+	//	if (bConstraintActive)
+	//	{
+	//		const FVec3 Axis = RM0.GetAxis((int32)SwingAxisIndex);
+	//		const FReal Error = (SwingAngle >= SwingAngleMax) ? SwingAngle - SwingAngleMax : SwingAngle + SwingAngleMax;
+
+	//		AddAngularRow(Axis, Axis, Error, J0, J1, C);
+	//	}
+	//}
+
+	void FJointSolverCholesky::AddAngularDrive_SLerp(
+		const FPBDJointSettings& JointSettings,
+		FDenseMatrix66& J0,
+		FDenseMatrix66& J1,
+		FDenseMatrix61& C)
+	{
+		// @todo(ccaulfield): target angles
+		FRotation3 R01Twist, R01Swing;
+		FPBDJointUtilities::DecomposeSwingTwistLocal(Rs[0], Rs[1], R01Swing, R01Twist);
 
-				J1.SetRowAt(RowIndex + 0, 0, 0, 0, 0);
-				J1.SetRowAt(RowIndex + 0, 3, -Axis);
+		FVec3 SwingAxis01;
+		FReal SwingAngle;
+		R01Swing.ToAxisAndAngleSafe(SwingAxis01, SwingAngle, FJointConstants::Swing1Axis(), SwingTwistAngleTolerance);
+		if (SwingAngle > PI)
+		{
+			SwingAngle = SwingAngle - (FReal)2 * PI;
+		}
 
-				C.SetAt(RowIndex, 0, (SwingAngle >= SwingAngleMax) ? SwingAngle - SwingAngleMax : SwingAngle + SwingAngleMax);
-			}
+		{
+			const FVec3 Axis = Rs[0] * SwingAxis01;
+			const FReal Error = AngularDriveStiffness * SwingAngle;
+
+			AddAngularRow(Axis, Axis, Error, J0, J1, C);
 		}
 	}
 
+
+	void FJointSolverCholesky::AddAngularDrive_Swing(
+		const FPBDJointSettings& JointSettings,
+		const EJointAngularConstraintIndex SwingConstraintIndex,
+		const EJointAngularAxisIndex SwingAxisIndex,
+		FDenseMatrix66& J0,
+		FDenseMatrix66& J1,
+		FDenseMatrix61& C)
+	{
+	}
+
 	// Add linear constraints to solver
 	void FJointSolverCholesky::AddLinearConstraints(
 		const FPBDJointSettings& JointSettings,
@@ -572,4 +666,44 @@ namespace Chaos
 		}
 	}
 
+	// Add angular constraints to solver
+	void FJointSolverCholesky::AddAngularDrives(
+		const FPBDJointSettings& JointSettings,
+		FDenseMatrix66& J0,
+		FDenseMatrix66& J1,
+		FDenseMatrix61& C)
+	{
+		EJointMotionType TwistMotion = JointSettings.Motion.AngularMotionTypes[(int32)EJointAngularConstraintIndex::Twist];
+		EJointMotionType Swing1Motion = JointSettings.Motion.AngularMotionTypes[(int32)EJointAngularConstraintIndex::Swing1];
+		EJointMotionType Swing2Motion = JointSettings.Motion.AngularMotionTypes[(int32)EJointAngularConstraintIndex::Swing2];
+
+		bool bTwistDriveEnabled = bEnableDrives && JointSettings.Motion.bAngularTwistDriveEnabled && (TwistMotion != EJointMotionType::Locked);
+		bool bSwing1DriveEnabled = bEnableDrives && JointSettings.Motion.bAngularSwingDriveEnabled && (Swing1Motion != EJointMotionType::Locked);
+		bool bSwing2DriveEnabled = bEnableDrives && JointSettings.Motion.bAngularSwingDriveEnabled && (Swing2Motion != EJointMotionType::Locked);
+		bool bSlerpDriveEnabled = bEnableDrives && JointSettings.Motion.bAngularSLerpDriveEnabled && (Swing1Motion != EJointMotionType::Locked) && (Swing2Motion != EJointMotionType::Locked);
+		bool bAnyDriveEnabled = (bTwistDriveEnabled || bSwing1DriveEnabled || bSwing2DriveEnabled || bSlerpDriveEnabled);
+
+		if (bAnyDriveEnabled)
+		{
+			if (bSlerpDriveEnabled)
+			{
+				AddAngularDrive_SLerp(JointSettings, J0, J1, C);
+			}
+			else
+			{
+				if (bTwistDriveEnabled)
+				{
+				}
+				if (bSwing1DriveEnabled)
+				{
+					AddAngularDrive_Swing(JointSettings, EJointAngularConstraintIndex::Swing1, EJointAngularAxisIndex::Swing1, J0, J1, C);
+				}
+				if (bSwing2DriveEnabled)
+				{
+					AddAngularDrive_Swing(JointSettings, EJointAngularConstraintIndex::Swing2, EJointAngularAxisIndex::Swing2, J0, J1, C);
+				}
+			}
+		}
+	}
+
 }
diff --git a/Engine/Source/Runtime/Experimental/Chaos/Private/Chaos/Joint/PBDJointSolverGaussSeidel.cpp b/Engine/Source/Runtime/Experimental/Chaos/Private/Chaos/Joint/PBDJointSolverGaussSeidel.cpp
--- a/Engine/Source/Runtime/Experimental/Chaos/Private/Chaos/Joint/PBDJointSolverGaussSeidel.cpp
+++ b/Engine/Source/Runtime/Experimental/Chaos/Private/Chaos/Joint/PBDJointSolverGaussSeidel.cpp
@@ -54,9 +54,11 @@ namespace Chaos
 		LinearStiffness = FPBDJointUtilities::GetLinearStiffness(SolverSettings, JointSettings);
 		TwistStiffness = FPBDJointUtilities::GetTwistStiffness(SolverSettings, JointSettings);
 		SwingStiffness = FPBDJointUtilities::GetSwingStiffness(SolverSettings, JointSettings);
+		AngularDriveStiffness = FPBDJointUtilities::GetAngularDriveStiffness(SolverSettings, JointSettings);
 		SwingTwistAngleTolerance = SolverSettings.SwingTwistAngleTolerance;
 		bEnableTwistLimits = SolverSettings.bEnableTwistLimits;
-		bEnableTwistLimits = SolverSettings.bEnableTwistLimits;
+		bEnableSwingLimits = SolverSettings.bEnableSwingLimits;
+		bEnableDrives = SolverSettings.bEnableDrives;
 
 		UpdateDerivedState();
 	}
@@ -70,6 +72,14 @@ namespace Chaos
 	}
 
 
+	void FJointSolverGaussSeidel::ApplyDrives(
+		const FReal Dt,
+		const FPBDJointSettings& JointSettings)
+	{
+		// @todo(ccaulfield): linear and velocity drives
+		ApplyRotationDrives(Dt, JointSettings);
+	}
+
 	void FJointSolverGaussSeidel::ApplyRotationConstraints(
 		const FReal Dt,
 		const FPBDJointSettings& JointSettings)
@@ -112,6 +122,47 @@ namespace Chaos
 	}
 
 
+	void FJointSolverGaussSeidel::ApplyRotationDrives(
+		const FReal Dt,
+		const FPBDJointSettings& JointSettings)
+	{
+		EJointMotionType TwistMotion = JointSettings.Motion.AngularMotionTypes[(int32)EJointAngularConstraintIndex::Twist];
+		EJointMotionType Swing1Motion = JointSettings.Motion.AngularMotionTypes[(int32)EJointAngularConstraintIndex::Swing1];
+		EJointMotionType Swing2Motion = JointSettings.Motion.AngularMotionTypes[(int32)EJointAngularConstraintIndex::Swing2];
+
+		if (bEnableDrives)
+		{
+			bool bTwistLocked = TwistMotion == EJointMotionType::Locked;
+			bool bSwing1Locked = Swing1Motion == EJointMotionType::Locked;
+			bool bSwing2Locked = Swing2Motion == EJointMotionType::Locked;
+
+			// No SLerp drive if we have a locked rotation (it will be grayed out in the editor in this case, but could still have been set before the rotation was locked)
+			if (JointSettings.Motion.bAngularSLerpDriveEnabled && !bTwistLocked && !bSwing1Locked && !bSwing2Locked)
+			{
+				ApplySLerpDrive(Dt, JointSettings);
+			}
+
+			if (JointSettings.Motion.bAngularTwistDriveEnabled && !bTwistLocked)
+			{
+				ApplyTwistDrive(Dt, JointSettings);
+			}
+
+			if (JointSettings.Motion.bAngularSwingDriveEnabled && !bSwing1Locked && !bSwing2Locked)
+			{
+				ApplyConeDrive(Dt, JointSettings);
+			}
+			else if (JointSettings.Motion.bAngularSwingDriveEnabled && !bSwing1Locked)
+			{
+				ApplySwingDrive(Dt, JointSettings, EJointAngularConstraintIndex::Swing1, EJointAngularAxisIndex::Swing1);
+			}
+			else if (JointSettings.Motion.bAngularSwingDriveEnabled && !bSwing2Locked)
+			{
+				ApplySwingDrive(Dt, JointSettings, EJointAngularConstraintIndex::Swing2, EJointAngularAxisIndex::Swing2);
+			}
+		}
+	}
+
+
 	void FJointSolverGaussSeidel::ApplyPositionConstraints(
 		const FReal Dt,
 		const FPBDJointSettings& JointSettings)
@@ -178,7 +229,7 @@ namespace Chaos
 		{
 			TwistAngle = TwistAngle - (FReal)2 * PI;
 		}
-		if (FVec3::DotProduct(TwistAxis01, FJointConstants::TwistAxis()) < 0)
+		if (TwistAxis01.X < 0)
 		{
 			TwistAxis01 = -TwistAxis01;
 			TwistAngle = -TwistAngle;
@@ -206,11 +257,41 @@ namespace Chaos
 		{
 			DTwistAngle = TwistAngle + TwistAngleMax;
 		}
-		const FReal DTwistAngle0 = DTwistAngle;
-		const FReal DTwistAngle1 = -DTwistAngle;
 
 		// Apply twist correction
-		ApplyRotationDelta(TwistStiffness, TwistAxis0, DTwistAngle0, TwistAxis1, DTwistAngle1);
+		ApplyRotationDelta(TwistStiffness, TwistAxis0, DTwistAngle, TwistAxis1, -DTwistAngle);
+		UpdateDerivedState();
+	}
+
+
+	void FJointSolverGaussSeidel::ApplyTwistDrive(
+		const FReal Dt,
+		const FPBDJointSettings& JointSettings)
+	{
+		// Decompose rotation of body 1 relative to body 0 into swing and twist rotations, assuming twist is X axis
+		FRotation3 R01Twist, R01Swing;
+		FPBDJointUtilities::DecomposeSwingTwistLocal(Rs[0], Rs[1], R01Swing, R01Twist);
+
+		FVec3 TwistAxis01;
+		FReal TwistAngle;
+		R01Twist.ToAxisAndAngleSafe(TwistAxis01, TwistAngle, FJointConstants::TwistAxis(), SwingTwistAngleTolerance);
+		if (TwistAngle > PI)
+		{
+			TwistAngle = TwistAngle - (FReal)2 * PI;
+		}
+		if (TwistAxis01.X < 0)
+		{
+			TwistAxis01 = -TwistAxis01;
+			TwistAngle = -TwistAngle;
+		}
+
+		const FVec3 TwistAxis0 = Rs[0] * TwistAxis01;
+		const FVec3 TwistAxis1 = Rs[1] * TwistAxis01;
+		const FReal TwistAngleTarget = JointSettings.Motion.AngularDriveTargetAngles[(int32)EJointAngularConstraintIndex::Twist];
+		const FReal DTwistAngle = TwistAngle - TwistAngleTarget;
+
+		// Apply twist correction
+		ApplyRotationDelta(AngularDriveStiffness, TwistAxis0, DTwistAngle, TwistAxis1, -DTwistAngle);
 		UpdateDerivedState();
 	}
 
@@ -267,6 +348,37 @@ namespace Chaos
 	}
 
 
+	void FJointSolverGaussSeidel::ApplyConeDrive(
+		const FReal Dt,
+		const FPBDJointSettings& JointSettings)
+	{
+		// Decompose rotation of body 1 relative to body 0 into swing and twist rotations, assuming twist is X axis
+		FRotation3 R01Twist, R01Swing;
+		FPBDJointUtilities::DecomposeSwingTwistLocal(Rs[0], Rs[1], R01Swing, R01Twist);
+
+		FVec3 SwingAxis01;
+		FReal SwingAngle;
+		R01Swing.ToAxisAndAngleSafe(SwingAxis01, SwingAngle, FJointConstants::Swing1Axis(), SwingTwistAngleTolerance);
+		if (SwingAngle > PI)
+		{
+			SwingAngle = SwingAngle - (FReal)2 * PI;
+		}
+
+		const FVec3 SwingAxis = Rs[0] * SwingAxis01;
+
+		// Circular swing target (max of Swing1, Swing2 targets)
+		// @todo(ccaulfield): what should cone target really do?
+		FReal Swing1Target = JointSettings.Motion.AngularDriveTargetAngles[(int32)EJointAngularConstraintIndex::Swing1];
+		FReal Swing2Target = JointSettings.Motion.AngularDriveTargetAngles[(int32)EJointAngularConstraintIndex::Swing2];
+		FReal SwingAngleTarget = FMath::Max(Swing1Target, Swing2Target);
+		FReal DSwingAngle = SwingAngle - SwingAngleTarget;
+
+		// Apply swing correction
+		ApplyRotationDelta(AngularDriveStiffness, SwingAxis, DSwingAngle, SwingAxis, -DSwingAngle);
+		UpdateDerivedState();
+	}
+
+
 	void FJointSolverGaussSeidel::ApplySwingConstraint(
 		const FReal Dt,
 		const FPBDJointSettings& JointSettings,
@@ -336,6 +448,48 @@ namespace Chaos
 		}
 	}
 
+
+	void FJointSolverGaussSeidel::ApplySwingDrive(
+		const FReal Dt,
+		const FPBDJointSettings& JointSettings,
+		const EJointAngularConstraintIndex SwingConstraintIndex,
+		const EJointAngularAxisIndex SwingAxisIndex)
+	{
+		// @todo(ccaulfield): implement swing drive
+	}
+
+
+	void FJointSolverGaussSeidel::ApplySLerpDrive(
+		const FReal Dt,
+		const FPBDJointSettings& JointSettings)
+	{
+		// @todo(ccaulfield): I think this can be done in a simpler way, reusing ApplyRotationDelta...
+
+		// Calculate the rotation we need to apply to resolve the rotation delta
+		const FRotation3 TargetR1 = Rs[0] * JointSettings.Motion.AngularDriveTarget;
+		const FRotation3 DR1 = TargetR1 * Rs[1].Inverse();
+		const FRotation3 TargetQ0 = DR1.Inverse() * Qs[0];
+		const FRotation3 TargetQ1 = DR1 * Qs[1];
+
+		FVec3 SLerpAxis;
+		FReal SLerpAngle;
+		if (DR1.ToAxisAndAngleSafe(SLerpAxis, SLerpAngle, FVec3(1, 0, 0)))
+		{
+			const FMatrix33 InvI0 = Utilities::ComputeWorldSpaceInertia(Qs[0], InvILs[0]);
+			const FMatrix33 InvI1 = Utilities::ComputeWorldSpaceInertia(Qs[1], InvILs[1]);
+			const FReal I0 = FVec3::DotProduct(SLerpAxis, Utilities::Multiply(InvI0, SLerpAxis));
+			const FReal I1 = FVec3::DotProduct(SLerpAxis, Utilities::Multiply(InvI1, SLerpAxis));
+			const FReal F0 = AngularDriveStiffness * I0 / (I0 + I1);
+			const FReal F1 = AngularDriveStiffness * I1 / (I0 + I1);
+
+			// Apply the rotation delta about the connector
+			Qs[0] = FRotation3::Slerp(Qs[0], TargetQ0, F0);
+			Qs[1] = FRotation3::Slerp(Qs[1], TargetQ1, F1);
+			Qs[1].EnforceShortestArcWith(Qs[0]);
+		}
+	}
+
+
 	void FJointSolverGaussSeidel::ApplyPositionConstraint(
 		const FReal Dt,
 		const FPBDJointSettings& JointSettings)
diff --git a/Engine/Source/Runtime/Experimental/Chaos/Private/Chaos/PBDJointConstraintUtilities.cpp b/Engine/Source/Runtime/Experimental/Chaos/Private/Chaos/PBDJointConstraintUtilities.cpp
--- a/Engine/Source/Runtime/Experimental/Chaos/Private/Chaos/PBDJointConstraintUtilities.cpp
+++ b/Engine/Source/Runtime/Experimental/Chaos/Private/Chaos/PBDJointConstraintUtilities.cpp
@@ -50,6 +50,14 @@ namespace Chaos
 		return Stiffness;
 	}
 
+	FReal FPBDJointUtilities::GetAngularDriveStiffness(
+		const FPBDJointSolverSettings& SolverSettings,
+		const FPBDJointSettings& JointSettings)
+	{
+		FReal DriveStiffness = (SolverSettings.DriveStiffness > 0.0f) ? SolverSettings.DriveStiffness : JointSettings.Motion.AngularDriveStiffness;
+		return DriveStiffness;
+	}
+
 
 	FVec3 FPBDJointUtilities::ConditionInertia(const FVec3& InI, const FReal MaxRatio)
 	{
diff --git a/Engine/Source/Runtime/Experimental/Chaos/Private/Chaos/PBDJointConstraints.cpp b/Engine/Source/Runtime/Experimental/Chaos/Private/Chaos/PBDJointConstraints.cpp
--- a/Engine/Source/Runtime/Experimental/Chaos/Private/Chaos/PBDJointConstraints.cpp
+++ b/Engine/Source/Runtime/Experimental/Chaos/Private/Chaos/PBDJointConstraints.cpp
@@ -26,7 +26,7 @@ namespace Chaos
 	DECLARE_CYCLE_STAT(TEXT("TPBDJointConstraints::SolveGaussSeidel"), STAT_Joints_Solve_GaussSeidel, STATGROUP_Chaos);
 	DECLARE_CYCLE_STAT(TEXT("TPBDJointConstraints::Project"), STAT_Joints_Project, STATGROUP_Chaos);
 
-	bool ChaosJoint_UseCholeskySolver = true;
+	bool ChaosJoint_UseCholeskySolver = false;
 	FAutoConsoleVariableRef CVarChaosImmPhysDeltaTime(TEXT("p.Chaos.Joint.UseCholeskySolver"), ChaosJoint_UseCholeskySolver, TEXT("Whether to use the new solver"));
 
 
@@ -175,7 +175,6 @@ namespace Chaos
 		, bEnableTwistLimits(true)
 		, bEnableSwingLimits(true)
 		, bEnableDrives(true)
-		, DrivesPhase(EJointSolverPhase::None)
 		, ProjectionPhase(EJointSolverPhase::None)
 		, LinearProjection((FReal)0)
 		, AngularProjection((FReal)0)
@@ -451,12 +450,6 @@ namespace Chaos
 			PreApplyCallback(Dt, Handles);
 		}
 
-		// Apply joint drives
-		if (Settings.DrivesPhase == EJointSolverPhase::Apply)
-		{
-			ApplyDrives(Dt, It, NumIts);
-		}
-
 		// Solve for joint position or velocity, depending on settings
 		if (Settings.ApplyPairIterations > 0)
 		{
@@ -486,12 +479,6 @@ namespace Chaos
 		// @todo(ccaulfield): track whether we are sufficiently solved
 		bool bNeedsAnotherIteration = true;
 
-		// Apply joint drives
-		if (Settings.DrivesPhase == EJointSolverPhase::ApplyPushOut)
-		{
-			ApplyDrives(Dt, It, NumIts);
-		}
-
 		// Solve for positions
 		if (Settings.ApplyPushOutPairIterations > 0)
 		{
@@ -511,18 +498,6 @@ namespace Chaos
 		return bNeedsAnotherIteration;
 	}
 
-	void FPBDJointConstraints::ApplyDrives(const FReal Dt, const int32 It, const int32 NumIts)
-	{
-		if (It == 0)
-		{
-			SCOPE_CYCLE_COUNTER(STAT_Joints_Drives);
-			for (int32 ConstraintIndex = 0; ConstraintIndex < NumConstraints(); ++ConstraintIndex)
-			{
-				ApplyDrives(Dt, ConstraintIndex, 0, 1);
-			}
-		}
-	}
-
 	void FPBDJointConstraints::SolvePosition(const FReal Dt, const int32 NumPairIts, const int32 It, const int32 NumIts)
 	{
 		if (ChaosJoint_UseCholeskySolver)
@@ -592,19 +567,6 @@ namespace Chaos
 			PreApplyCallback(Dt, SortedConstraintHandles);
 		}
 
-		// Apply joint drives
-		if (Settings.DrivesPhase == EJointSolverPhase::Apply)
-		{
-			SCOPE_CYCLE_COUNTER(STAT_Joints_Drives);
-			if (It == 0)
-			{
-				for (FConstraintContainerHandle* ConstraintHandle : SortedConstraintHandles)
-				{
-					ApplyDrives(Dt, ConstraintHandle->GetConstraintIndex(), 0, 1);
-				}
-			}
-		}
-
 		// Solve for joint position or velocity, depending on settings
 		if (Settings.ApplyPairIterations > 0)
 		{
@@ -647,19 +609,6 @@ namespace Chaos
 				return L.GetConstraintLevel() < R.GetConstraintLevel();
 			});
 
-		// Apply joint drives
-		if (Settings.DrivesPhase == EJointSolverPhase::ApplyPushOut)
-		{
-			SCOPE_CYCLE_COUNTER(STAT_Joints_Drives);
-			if (It == 0)
-			{
-				for (FConstraintContainerHandle* ConstraintHandle : SortedConstraintHandles)
-				{
-					ApplyDrives(Dt, ConstraintHandle->GetConstraintIndex(), 0, 1);
-				}
-			}
-		}
-
 		// Solve for positions
 		if (Settings.ApplyPushOutPairIterations > 0)
 		{
@@ -732,71 +681,6 @@ namespace Chaos
 	}
 
 
-	void FPBDJointConstraints::ApplyDrives(const FReal Dt, const int32 ConstraintIndex, const int32 It, const int32 NumIts)
-	{
-		const FPBDJointSettings& JointSettings = ConstraintSettings[ConstraintIndex];
-
-		const int32 Index0 = 1;
-		const int32 Index1 = 0;
-		TGenericParticleHandle<FReal, 3> Particle0 = TGenericParticleHandle<FReal, 3>(ConstraintParticles[ConstraintIndex][Index0]);
-		TGenericParticleHandle<FReal, 3> Particle1 = TGenericParticleHandle<FReal, 3>(ConstraintParticles[ConstraintIndex][Index1]);
-
-		FVec3 P0 = FParticleUtilities::GetCoMWorldPosition(Particle0);
-		FRotation3 Q0 = FParticleUtilities::GetCoMWorldRotation(Particle0);
-		FVec3 P1 = FParticleUtilities::GetCoMWorldPosition(Particle1);
-		FRotation3 Q1 = FParticleUtilities::GetCoMWorldRotation(Particle1);
-		FVec3 V0 = Particle0->V();
-		FVec3 W0 = Particle0->W();
-		FVec3 V1 = Particle1->V();
-		FVec3 W1 = Particle1->W();
-		const FRigidTransform3 XL0 = FParticleUtilities::ParticleLocalToCoMLocal(Particle0, JointSettings.ConstraintFrames[Index0]);
-		const FRigidTransform3 XL1 = FParticleUtilities::ParticleLocalToCoMLocal(Particle1, JointSettings.ConstraintFrames[Index1]);
-		float InvM0 = Particle0->InvM();
-		float InvM1 = Particle1->InvM();
-		FMatrix33 InvIL0 = Particle0->InvI();
-		FMatrix33 InvIL1 = Particle1->InvI();
-
-		EJointMotionType TwistMotion = JointSettings.Motion.AngularMotionTypes[(int32)EJointAngularConstraintIndex::Twist];
-		EJointMotionType Swing1Motion = JointSettings.Motion.AngularMotionTypes[(int32)EJointAngularConstraintIndex::Swing1];
-		EJointMotionType Swing2Motion = JointSettings.Motion.AngularMotionTypes[(int32)EJointAngularConstraintIndex::Swing2];
-
-		// Apply angular drives (NOTE: modifies position, not velocity)
-		if (Settings.bEnableDrives)
-		{
-			bool bTwistLocked = TwistMotion == EJointMotionType::Locked;
-			bool bSwing1Locked = Swing1Motion == EJointMotionType::Locked;
-			bool bSwing2Locked = Swing2Motion == EJointMotionType::Locked;
-
-			// No SLerp drive if we have a locked rotation (it will be grayed out in the editor in this case, but could still have been set before the rotation was locked)
-			if (JointSettings.Motion.bAngularSLerpDriveEnabled && !bTwistLocked && !bSwing1Locked && !bSwing2Locked)
-			{
-				FPBDJointUtilities::ApplyJointSLerpDrive(Dt, Settings, JointSettings, XL0, XL1, P0, Q0, V0, W0, P1, Q1, V1, W1, InvM0, InvIL0, InvM1, InvIL1);
-			}
-
-			if (JointSettings.Motion.bAngularTwistDriveEnabled && !bTwistLocked)
-			{
-				FPBDJointUtilities::ApplyJointTwistDrive(Dt, Settings, JointSettings, XL0, XL1, P0, Q0, V0, W0, P1, Q1, V1, W1, InvM0, InvIL0, InvM1, InvIL1);
-			}
-
-			if (JointSettings.Motion.bAngularSwingDriveEnabled && !bSwing1Locked && !bSwing2Locked)
-			{
-				FPBDJointUtilities::ApplyJointConeDrive(Dt, Settings, JointSettings, XL0, XL1, P0, Q0, V0, W0, P1, Q1, V1, W1, InvM0, InvIL0, InvM1, InvIL1);
-			}
-			else if (JointSettings.Motion.bAngularSwingDriveEnabled && !bSwing1Locked)
-			{
-				//FPBDJointUtilities::ApplyJointSwingDrive(Dt, Settings, JointSettings, XL0, XL1, EJointAngularConstraintIndex::Swing1, P0, Q0, P1, Q1, InvM0, InvIL0, InvM1, InvIL1);
-			}
-			else if (JointSettings.Motion.bAngularSwingDriveEnabled && !bSwing2Locked)
-			{
-				//FPBDJointUtilities::ApplyJointSwingDrive(Dt, Settings, JointSettings, XL0, XL1, EJointAngularConstraintIndex::Swing2, P0, Q0, P1, Q1, InvM0, InvIL0, InvM1, InvIL1);
-			}
-		}
-
-		UpdateParticleState(Particle0->CastToRigidParticle(), Dt, P0, Q0);
-		UpdateParticleState(Particle1->CastToRigidParticle(), Dt, P1, Q1);
-	}
-
-
 	void FPBDJointConstraints::SolvePosition(const FReal Dt, const int32 ConstraintIndex, const int32 NumPairIts, const int32 It, const int32 NumIts)
 	{
 		if (ChaosJoint_UseCholeskySolver)
@@ -844,6 +728,7 @@ namespace Chaos
 		
 		for (int32 PairIt = 0; PairIt < NumPairIts; ++PairIt)
 		{
+			Solver.ApplyDrives(Dt, JointSettings);
 			Solver.ApplyConstraints(Dt, JointSettings);
 		}
 
@@ -885,6 +770,7 @@ namespace Chaos
 
 		for (int32 PairIt = 0; PairIt < NumPairIts; ++PairIt)
 		{
+			Solver.ApplyDrives(Dt, JointSettings);
 			Solver.ApplyConstraints(Dt, JointSettings);
 		}
 
diff --git a/Engine/Source/Runtime/Experimental/Chaos/Public/Chaos/Joint/PBDJointSolverCholesky.h b/Engine/Source/Runtime/Experimental/Chaos/Public/Chaos/Joint/PBDJointSolverCholesky.h
--- a/Engine/Source/Runtime/Experimental/Chaos/Public/Chaos/Joint/PBDJointSolverCholesky.h
+++ b/Engine/Source/Runtime/Experimental/Chaos/Public/Chaos/Joint/PBDJointSolverCholesky.h
@@ -60,10 +60,31 @@ namespace Chaos
 			const FReal Dt,
 			const FPBDJointSettings& JointSettings);
 
+		void ApplyDrives(
+			const FReal Dt,
+			const FPBDJointSettings& JointSettings);
+
 	private:
 
 		void UpdateDerivedState();
 
+		void AddLinearRow(
+			const FVec3& Axis,
+			const FVec3& Connector0,
+			const FVec3& Connector1,
+			const FReal Error,
+			FDenseMatrix66& J0,
+			FDenseMatrix66& J1,
+			FDenseMatrix61& C);
+
+		void AddAngularRow(
+			const FVec3& Axis0,
+			const FVec3& Axis1,
+			const FReal Error,
+			FDenseMatrix66& J0,
+			FDenseMatrix66& J1,
+			FDenseMatrix61& C);
+
 		void AddLinearConstraints_Point(
 			const FPBDJointSettings& JointSettings,
 			FDenseMatrix66& J0,
@@ -118,6 +139,20 @@ namespace Chaos
 			FDenseMatrix66& J1,
 			FDenseMatrix61& C);
 
+		void AddAngularDrive_SLerp(
+			const FPBDJointSettings& JointSettings,
+			FDenseMatrix66& J0,
+			FDenseMatrix66& J1,
+			FDenseMatrix61& C);
+
+		void AddAngularDrive_Swing(
+			const FPBDJointSettings& JointSettings,
+			const EJointAngularConstraintIndex SwingConstraintIndex,
+			const EJointAngularAxisIndex SwingAxisIndex,
+			FDenseMatrix66& J0,
+			FDenseMatrix66& J1,
+			FDenseMatrix61& C);
+
 		void AddLinearConstraints(
 			const FPBDJointSettings& JointSettings,
 			FDenseMatrix66& J0,
@@ -130,12 +165,30 @@ namespace Chaos
 			FDenseMatrix66& J1,
 			FDenseMatrix61& C);
 
-		void BuildJacobianAndResidual(
+		void AddAngularDrives(
+			const FPBDJointSettings& JointSettings,
+			FDenseMatrix66& J0,
+			FDenseMatrix66& J1,
+			FDenseMatrix61& C);
+
+		void BuildJacobianAndResidual_Constraints(
 			const FPBDJointSettings& JointSettings,
 			FDenseMatrix66& J0,
 			FDenseMatrix66& J1,
 			FDenseMatrix61& C);
 
+		void BuildJacobianAndResidual_Drives(
+			const FPBDJointSettings& JointSettings,
+			FDenseMatrix66& J0,
+			FDenseMatrix66& J1,
+			FDenseMatrix61& C);
+
+		void SolveAndApply(
+			const FPBDJointSettings& JointSettings,
+			const FDenseMatrix66& J0,
+			const FDenseMatrix66& J1,
+			const FDenseMatrix61& C);
+
 		// Local-space constraint settings
 		FRigidTransform3 XLs[MaxConstraintedBodies];	// Local-space joint connector transforms
 		FMatrix33 InvILs[MaxConstraintedBodies];		// Local-space inverse inertias
@@ -151,9 +204,11 @@ namespace Chaos
 
 		// Settings
 		FReal Stiffness;
+		FReal AngularDriveStiffness;
 		FReal SwingTwistAngleTolerance;
 		bool bEnableTwistLimits;
 		bool bEnableSwingLimits;
+		bool bEnableDrives;
 	};
 
 }
\ No newline at end of file
diff --git a/Engine/Source/Runtime/Experimental/Chaos/Public/Chaos/Joint/PBDJointSolverGaussSeidel.h b/Engine/Source/Runtime/Experimental/Chaos/Public/Chaos/Joint/PBDJointSolverGaussSeidel.h
--- a/Engine/Source/Runtime/Experimental/Chaos/Public/Chaos/Joint/PBDJointSolverGaussSeidel.h
+++ b/Engine/Source/Runtime/Experimental/Chaos/Public/Chaos/Joint/PBDJointSolverGaussSeidel.h
@@ -57,6 +57,10 @@ namespace Chaos
 			const FReal Dt,
 			const FPBDJointSettings& JointSettings);
 
+		void ApplyDrives(
+			const FReal Dt,
+			const FPBDJointSettings& JointSettings);
+
 	private:
 
 		void UpdateDerivedState();
@@ -69,6 +73,11 @@ namespace Chaos
 			const FReal Dt,
 			const FPBDJointSettings& JointSettings);
 
+		void ApplyRotationDrives(
+			const FReal Dt,
+			const FPBDJointSettings& JointSettings);
+
+
 		void ApplyPositionDelta(
 			const FReal Stiffness,
 			const FVec3& DP0,
@@ -86,24 +95,44 @@ namespace Chaos
 			const FVec3& Axis1,
 			const FReal Angle1);
 
+
 		void ApplyTwistConstraint(
 			const FReal Dt,
 			const FPBDJointSettings& JointSettings);
 
+		void ApplyTwistDrive(
+			const FReal Dt,
+			const FPBDJointSettings& JointSettings);
+
 		void ApplyConeConstraint(
 			const FReal Dt,
 			const FPBDJointSettings& JointSettings);
 
+		void ApplyConeDrive(
+			const FReal Dt,
+			const FPBDJointSettings& JointSettings);
+
 		void ApplySwingConstraint(
 			const FReal Dt,
 			const FPBDJointSettings& JointSettings,
 			const EJointAngularConstraintIndex SwingConstraintIndex,
 			const EJointAngularAxisIndex SwingAxisIndex);
 
+		void ApplySwingDrive(
+			const FReal Dt,
+			const FPBDJointSettings& JointSettings,
+			const EJointAngularConstraintIndex SwingConstraintIndex,
+			const EJointAngularAxisIndex SwingAxisIndex);
+
+		void ApplySLerpDrive(
+			const FReal Dt,
+			const FPBDJointSettings& JointSettings);
+
 		void ApplyPositionConstraint(
 			const FReal Dt,
 			const FPBDJointSettings& JointSettings);
 
+
 		// Local-space constraint settings
 		FRigidTransform3 XLs[MaxConstraintedBodies];	// Local-space joint connector transforms
 		FMatrix33 InvILs[MaxConstraintedBodies];		// Local-space inverse inertias
@@ -121,9 +150,11 @@ namespace Chaos
 		FReal LinearStiffness;
 		FReal TwistStiffness;
 		FReal SwingStiffness;
+		FReal AngularDriveStiffness;
 		FReal SwingTwistAngleTolerance;
 		bool bEnableTwistLimits;
 		bool bEnableSwingLimits;
+		bool bEnableDrives;
 	};
 
 }
\ No newline at end of file
diff --git a/Engine/Source/Runtime/Experimental/Chaos/Public/Chaos/PBDJointConstraintTypes.h b/Engine/Source/Runtime/Experimental/Chaos/Public/Chaos/PBDJointConstraintTypes.h
--- a/Engine/Source/Runtime/Experimental/Chaos/Public/Chaos/PBDJointConstraintTypes.h
+++ b/Engine/Source/Runtime/Experimental/Chaos/Public/Chaos/PBDJointConstraintTypes.h
@@ -133,7 +133,6 @@ namespace Chaos
 		bool bEnableTwistLimits;
 		bool bEnableSwingLimits;
 		bool bEnableDrives;
-		EJointSolverPhase DrivesPhase;
 		EJointSolverPhase ProjectionPhase;
 		FReal LinearProjection;
 		FReal AngularProjection;
diff --git a/Engine/Source/Runtime/Experimental/Chaos/Public/Chaos/PBDJointConstraintUtilities.h b/Engine/Source/Runtime/Experimental/Chaos/Public/Chaos/PBDJointConstraintUtilities.h
--- a/Engine/Source/Runtime/Experimental/Chaos/Public/Chaos/PBDJointConstraintUtilities.h
+++ b/Engine/Source/Runtime/Experimental/Chaos/Public/Chaos/PBDJointConstraintUtilities.h
@@ -76,6 +76,10 @@ namespace Chaos
 			const FPBDJointSolverSettings& SolverSettings,
 			const FPBDJointSettings& JointSettings);
 
+		static FReal GetAngularDriveStiffness(
+			const FPBDJointSolverSettings& SolverSettings,
+			const FPBDJointSettings& JointSettings);
+
 
 		static FVec3 GetSphereLimitedPositionError(const FVec3& CX, const FReal Radius);
 		static FVec3 GetSphereLimitedVelocityError(const FVec3& CX, const FReal Radius, const FVec3& CV);
diff --git a/Engine/Source/Runtime/Experimental/Chaos/Public/Chaos/PBDJointConstraints.h b/Engine/Source/Runtime/Experimental/Chaos/Public/Chaos/PBDJointConstraints.h
--- a/Engine/Source/Runtime/Experimental/Chaos/Public/Chaos/PBDJointConstraints.h
+++ b/Engine/Source/Runtime/Experimental/Chaos/Public/Chaos/PBDJointConstraints.h
@@ -160,14 +160,12 @@ namespace Chaos
 		void UpdateParticleState(TPBDRigidParticleHandle<FReal, 3>* Rigid, const FReal Dt, const FVec3& P, const FRotation3& Q, const bool bUpdateVelocity = true);
 		void SortConstraints();
 
-		void ApplyDrives(const FReal Dt, const int32 ConstraintIndex, const int32 It, const int32 NumIts);
 		void SolvePosition(const FReal Dt, const int32 ConstraintIndex, const int32 NumPairIts, const int32 It, const int32 NumIts);
 		void SolvePosition_GaussSiedel(const FReal Dt, const int32 ConstraintIndex, const int32 NumPairIts, const int32 It, const int32 NumIts);
 		void SolvePosition_Cholesky(const FReal Dt, const int32 ConstraintIndex, const int32 NumPairIts, const int32 It, const int32 NumIts);
 		void ProjectPosition(const FReal Dt, const int32 ConstraintIndex, const int32 It, const int32 NumIts);
 
 		// All-constraints API
-		void ApplyDrives(const FReal Dt, const int32 It, const int32 NumIts);
 		void SolvePosition(const FReal Dt, const int32 NumPairIts, const int32 It, const int32 NumIts);
 
 		FPBDJointSolverSettings Settings;
