Fix CIS - UE4Editor Debug StressTestUnity

diff --git a/Engine/Source/Runtime/Experimental/ChaosSolvers/Public/PBDRigidsSolver.h b/Engine/Source/Runtime/Experimental/ChaosSolvers/Public/PBDRigidsSolver.h
--- a/Engine/Source/Runtime/Experimental/ChaosSolvers/Public/PBDRigidsSolver.h
+++ b/Engine/Source/Runtime/Experimental/ChaosSolvers/Public/PBDRigidsSolver.h
@@ -39,7 +39,6 @@ namespace Chaos
 {
 	class AdvanceOneTimeStepTask;
 	class FPersistentPhysicsTask;
-	class FPhysicsCommand;
 	class FChaosArchive;
 	class FPBDRigidsSolver;
 
diff --git a/Engine/Source/Runtime/PhysicsCore/Public/SQCapture.h b/Engine/Source/Runtime/PhysicsCore/Public/SQCapture.h
--- a/Engine/Source/Runtime/PhysicsCore/Public/SQCapture.h
+++ b/Engine/Source/Runtime/PhysicsCore/Public/SQCapture.h
@@ -14,9 +14,13 @@
 
 
 class FPhysTestSerializer;
-struct FSweepHit;
-struct FRaycastHit;
-struct FOverlapHit;
+
+namespace ChaosInterface
+{
+	struct FSweepHit;
+	struct FRaycastHit;
+	struct FOverlapHit;
+}
 
 namespace Chaos
 {
