Chaos : Collisions

-Add support for SinglePointCollisionConstraints within UpdateConvexConvexConstraint.

diff --git a/Engine/Source/Runtime/Experimental/Chaos/Private/Chaos/CollisionResolution.cpp b/Engine/Source/Runtime/Experimental/Chaos/Private/Chaos/CollisionResolution.cpp
--- a/Engine/Source/Runtime/Experimental/Chaos/Private/Chaos/CollisionResolution.cpp
+++ b/Engine/Source/Runtime/Experimental/Chaos/Private/Chaos/CollisionResolution.cpp
@@ -558,30 +558,39 @@ namespace Chaos
 		}
 
 		template<class T, int d>
-		void UpdateConvexConvexConstraint(const FImplicitObject& Implicit0, const TRigidTransform<T, d>& Transform0, const FImplicitObject& Implicit1, const TRigidTransform<T, d>& Transform1, const T Thickness, TRigidBodyIterativeContactConstraint<T,d>& Constraint)
+		void UpdateConvexConvexConstraint(const FImplicitObject& Implicit0, const TRigidTransform<T, d>& Transform0, const FImplicitObject& Implicit1, const TRigidTransform<T, d>& Transform1, const T Thickness, TCollisionConstraintBase<T,d>& ConstraintBase)
 		{
 			TContactPoint<T> ContactPoint;
-			ContactPoint.Phi = FLT_MAX;
 
-			const TRigidTransform<T, d> AToBTM = Transform0.GetRelativeTransform(Transform1);
+			if (ConstraintBase.GetType() == FRigidBodyPointContactConstraint::StaticType())
+			{
+				ContactPoint = ConvexConvexContactPoint(Implicit0, Transform0, Implicit1, Transform1, Thickness);
+			}
+			else if (ConstraintBase.GetType() == FRigidBodyIterativeContactConstraint::StaticType())
+			{
+				TRigidBodyIterativeContactConstraint<T, d>& Constraint = *ConstraintBase.template As<TRigidBodyIterativeContactConstraint<T, d>>();
+				ContactPoint.Phi = FLT_MAX;
 
-			TPlane<T, d> CollisionPlane(Constraint.PlanePosition, Constraint.PlaneNormal);
+				const TRigidTransform<T, d> AToBTM = Transform0.GetRelativeTransform(Transform1);
 
-			for (int32 Idx = 0; Idx < Constraint.NumSamples(); Idx++)
-			{
-				Constraint[Idx].Manifold.Phi = CollisionPlane.PhiWithNormal(AToBTM.TransformPosition(Constraint[Idx].X), Constraint[Idx].Manifold.Normal);
-				Constraint[Idx].Manifold.Normal = Transform1.TransformVector(Constraint.PlaneNormal);
-				Constraint[Idx].Manifold.Location = Transform0.TransformPosition(Constraint[Idx].X);
+				TPlane<T, d> CollisionPlane(Constraint.PlanePosition, Constraint.PlaneNormal);
 
-				if (ContactPoint.Phi > Constraint[Idx].Manifold.Phi)
+				for (int32 Idx = 0; Idx < Constraint.NumSamples(); Idx++)
 				{
-					ContactPoint.Phi = Constraint[Idx].Manifold.Phi;
-					ContactPoint.Normal = Constraint[Idx].Manifold.Normal;
-					ContactPoint.Location = Constraint[Idx].Manifold.Location;
+					Constraint[Idx].Manifold.Phi = CollisionPlane.PhiWithNormal(AToBTM.TransformPosition(Constraint[Idx].X), Constraint[Idx].Manifold.Normal);
+					Constraint[Idx].Manifold.Normal = Transform1.TransformVector(Constraint.PlaneNormal);
+					Constraint[Idx].Manifold.Location = Transform0.TransformPosition(Constraint[Idx].X);
+
+					if (ContactPoint.Phi > Constraint[Idx].Manifold.Phi)
+					{
+						ContactPoint.Phi = Constraint[Idx].Manifold.Phi;
+						ContactPoint.Normal = Constraint[Idx].Manifold.Normal;
+						ContactPoint.Location = Constraint[Idx].Manifold.Location;
+					}
 				}
 			}
 
-			UpdateContactPoint(Constraint.Manifold, ContactPoint);
+			UpdateContactPoint(ConstraintBase.Manifold, ContactPoint);
 		}
 
 		template<class T, int d>
@@ -1071,7 +1080,7 @@ namespace Chaos
 			}
 			else if(Implicit0.IsConvex() && Implicit1.IsConvex())
 			{
-				UpdateConvexConvexConstraint(Implicit0, Transform0, Implicit1, Transform1, Thickness, *ConstraintBase.template As<TRigidBodyIterativeContactConstraint<T, d>>());
+				UpdateConvexConvexConstraint(Implicit0, Transform0, Implicit1, Transform1, Thickness, ConstraintBase);
 			}
 			else if (Implicit1.IsUnderlyingUnion())
 			{
@@ -1258,7 +1267,7 @@ namespace Chaos
 		template void UpdateCapsuleBoxConstraint<float, 3>(const TCapsule<float>& A, const TRigidTransform<float, 3>& ATransform, const TAABB<float, 3>& B, const TRigidTransform<float, 3>& BTransform, const float Thickness, TRigidBodyPointContactConstraint<float, 3>& Constraint);
 		template void ConstructCapsuleBoxConstraints<float, 3>(TGeometryParticleHandle<float, 3>* Particle0, TGeometryParticleHandle<float, 3>* Particle1, const FImplicitObject* Implicit0, const FImplicitObject* Implicit1, const TRigidTransform<float, 3>& Transform0, const TRigidTransform<float, 3>& Transform1, const float Thickness, FCollisionConstraintsArray& NewConstraints);
 
-		template void UpdateConvexConvexConstraint<float, 3>(const FImplicitObject& A, const TRigidTransform<float, 3>& ATM, const FImplicitObject& B, const TRigidTransform<float, 3>& BTM, const float Thickness, TRigidBodyIterativeContactConstraint<float, 3>& Constraint);
+		template void UpdateConvexConvexConstraint<float, 3>(const FImplicitObject& A, const TRigidTransform<float, 3>& ATM, const FImplicitObject& B, const TRigidTransform<float, 3>& BTM, const float Thickness, TCollisionConstraintBase<float, 3>& Constraint);
 		template void UpdateConvexConvexManifold<float, 3>(TRigidBodyIterativeContactConstraint<float, 3>& Constraint, const TRigidTransform<float, 3>& ATM, const TRigidTransform<float, 3>& BTM, const float Thickness);
 		template void ConstructConvexConvexConstraints<float, 3>(TGeometryParticleHandle<float, 3>* Particle0, TGeometryParticleHandle<float, 3>* Particle1, const FImplicitObject* Implicit0, const FImplicitObject* Implicit1, const TRigidTransform<float, 3>& Transform0, const TRigidTransform<float, 3>& Transform1, const float Thickness, FCollisionConstraintsArray& NewConstraints);
 
diff --git a/Engine/Source/Runtime/Experimental/Chaos/Public/Chaos/CollisionResolution.h b/Engine/Source/Runtime/Experimental/Chaos/Public/Chaos/CollisionResolution.h
--- a/Engine/Source/Runtime/Experimental/Chaos/Public/Chaos/CollisionResolution.h
+++ b/Engine/Source/Runtime/Experimental/Chaos/Public/Chaos/CollisionResolution.h
@@ -115,7 +115,7 @@ namespace Chaos
 			// Convex-Convex
 			//
 			template <typename T, int d>
-			void CHAOS_API UpdateConvexConvexConstraint(const FImplicitObject& A, const TRigidTransform<T, d>& ATM, const FImplicitObject& B, const TRigidTransform<T, d>& BTM, const T Thickness, TRigidBodyIterativeContactConstraint<T, d>& Constraint);
+			void CHAOS_API UpdateConvexConvexConstraint(const FImplicitObject& A, const TRigidTransform<T, d>& ATM, const FImplicitObject& B, const TRigidTransform<T, d>& BTM, const T Thickness, TCollisionConstraintBase<T, d>& Constraint);
 
 			template<class T, int d>
 			void UpdateConvexConvexManifold(TRigidBodyIterativeContactConstraint<T, d>& Constraint, const TRigidTransform<T, d>& ATM, const TRigidTransform<T, d>& BTM, const T Thickness);
