Chaos
- Joint Projection
- Fix collision on last body
- Reduce iterations

diff --git a/Engine/Source/Runtime/Engine/Private/PhysicsEngine/ImmediatePhysics/ImmediatePhysicsChaos/ImmediatePhysicsJointHandle_Chaos.cpp b/Engine/Source/Runtime/Engine/Private/PhysicsEngine/ImmediatePhysics/ImmediatePhysicsChaos/ImmediatePhysicsJointHandle_Chaos.cpp
--- a/Engine/Source/Runtime/Engine/Private/PhysicsEngine/ImmediatePhysics/ImmediatePhysicsChaos/ImmediatePhysicsJointHandle_Chaos.cpp
+++ b/Engine/Source/Runtime/Engine/Private/PhysicsEngine/ImmediatePhysics/ImmediatePhysicsChaos/ImmediatePhysicsJointHandle_Chaos.cpp
@@ -23,13 +23,13 @@ namespace ImmediatePhysics_Chaos
 	FAutoConsoleVariableRef CVarStiffnessUseMass(TEXT("p.Chaos.ImmPhys.StiffnessUseMass"), ChaosImmediate_StiffnessUseMass, TEXT("Whether to use mass to scale stiffness in the conversion."));
 
 	float ChaosImmediate_JointMinProjection = 0.0f;
-	float ChaosImmediate_JointMaxProjection = 0.3f;
+	float ChaosImmediate_JointMaxProjection = 1.0f;
 	FAutoConsoleVariableRef CVarJointMinProjection(TEXT("p.Chaos.ImmPhys.JointMinProjection"), ChaosImmediate_JointMinProjection, TEXT("Joint min projection (for joints with projection disabled)."));
 	FAutoConsoleVariableRef CVarJointMaxProjection(TEXT("p.Chaos.ImmPhys.JointMaxProjection"), ChaosImmediate_JointMaxProjection, TEXT("Joint max projection (for joints with projection enabled)."));
 
 	float ChaosImmediate_DriveStiffnessSourceMin = 100.0f;	// PhysX stiffness per inertia that we translate to a Chaos stiffness of ChaosImmediate_DriveStiffnessTargetMin
 	float ChaosImmediate_DriveStiffnessSourceMax = 2000.0f;	// PhysX stiffness per inertia that we translate to a Chaos stiffness of ChaosImmediate_DriveStiffnessTargetMax
-	float ChaosImmediate_DriveStiffnessTargetMin = 0.2f;
+	float ChaosImmediate_DriveStiffnessTargetMin = 0.6f;
 	float ChaosImmediate_DriveStiffnessTargetMax = 1.0f;
 	FAutoConsoleVariableRef CVarDriveStiffnessSourceMin(TEXT("p.Chaos.ImmPhys.DriveStiffnessSourceMin"), ChaosImmediate_DriveStiffnessSourceMin, TEXT("Conversion factor for drive stiffness."));
 	FAutoConsoleVariableRef CVarDriveStiffnessSourceMax(TEXT("p.Chaos.ImmPhys.DriveStiffnessSourceMax"), ChaosImmediate_DriveStiffnessSourceMax, TEXT("Conversion factor for drive stiffness."));
diff --git a/Engine/Source/Runtime/Engine/Private/PhysicsEngine/ImmediatePhysics/ImmediatePhysicsChaos/ImmediatePhysicsSimulation_Chaos.cpp b/Engine/Source/Runtime/Engine/Private/PhysicsEngine/ImmediatePhysics/ImmediatePhysicsChaos/ImmediatePhysicsSimulation_Chaos.cpp
--- a/Engine/Source/Runtime/Engine/Private/PhysicsEngine/ImmediatePhysics/ImmediatePhysicsChaos/ImmediatePhysicsSimulation_Chaos.cpp
+++ b/Engine/Source/Runtime/Engine/Private/PhysicsEngine/ImmediatePhysics/ImmediatePhysicsChaos/ImmediatePhysicsSimulation_Chaos.cpp
@@ -22,7 +22,7 @@
 //
 float ChaosImmediate_Evolution_DeltaTime = 0.03f;
 int32 ChaosImmediate_Evolution_Iterations = 4;
-int32 ChaosImmediate_Evolution_PushOutIterations = 2;
+int32 ChaosImmediate_Evolution_PushOutIterations = 1;
 FAutoConsoleVariableRef CVarChaosImmPhysDeltaTime(TEXT("p.Chaos.ImmPhys.DeltaTime"), ChaosImmediate_Evolution_DeltaTime, TEXT("Override chaos immediate physics delta time if non-zero"));
 FAutoConsoleVariableRef CVarChaosImmPhysIterations(TEXT("p.Chaos.ImmPhys.Iterations"), ChaosImmediate_Evolution_Iterations, TEXT("Number of constraint solver loops in immediate physics"));
 FAutoConsoleVariableRef CVarChaosImmPhysPushOutIterations(TEXT("p.Chaos.ImmPhys.PushOutIterations"), ChaosImmediate_Evolution_PushOutIterations, TEXT("Set the ApplyPushOut() (position correction) iteration count"));
@@ -38,13 +38,13 @@ FAutoConsoleVariableRef CVarChaosImmPhysCollisionPushOutPairIterations(TEXT("p.C
 FAutoConsoleVariableRef CVarChaosImmPhysCollisionPriority(TEXT("p.Chaos.ImmPhys.Collision.Priority"), ChaosImmediate_Collision_Priority, TEXT("Set the Collision constraint sort order (Joints have priority 0)"));
 FAutoConsoleVariableRef CVarChaosImmPhysThickness(TEXT("p.Chaos.ImmPhys.Collision.Thickness"), ChaosImmediate_Collision_Thickness, TEXT("ChaosImmediateThickness"));
 
-int32 ChaosImmediate_Joint_PairIterations = 2;
-int32 ChaosImmediate_Joint_PushOutPairIterations = 0;
+int32 ChaosImmediate_Joint_PairIterations = 1;
+int32 ChaosImmediate_Joint_PushOutPairIterations = 2;
 float ChaosImmediate_Joint_SwingTwistAngleTolerance = 1.0e-6f;
 int32 ChaosImmediate_Joint_EnableTwistLimits = 1;
 int32 ChaosImmediate_Joint_EnableSwingLimits = 1;
 int32 ChaosImmediate_Joint_EnableDrives = 1;
-int32 ChaosImmediate_Joint_ProjectionPhase = (int32)Chaos::EJointSolverPhase::Apply;
+int32 ChaosImmediate_Joint_ProjectionPhase = (int32)Chaos::EJointSolverPhase::ApplyPushOut;
 float ChaosImmediate_Joint_LinearProjection = 0.0f;
 float ChaosImmediate_Joint_AngularProjection = 0.0f;
 float ChaosImmediate_Joint_Stiffness = 1.0f;
@@ -59,7 +59,6 @@ FAutoConsoleVariableRef CVarChaosImmPhysSwingTwistAngleTolerance(TEXT("p.Chaos.I
 FAutoConsoleVariableRef CVarChaosImmPhysEnableTwistLimits(TEXT("p.Chaos.ImmPhys.Joint.EnableTwistLimits"), ChaosImmediate_Joint_EnableTwistLimits, TEXT("EnableTwistLimits."));
 FAutoConsoleVariableRef CVarChaosImmPhysEnableSwingLimits(TEXT("p.Chaos.ImmPhys.Joint.EnableSwingLimits"), ChaosImmediate_Joint_EnableSwingLimits, TEXT("EnableSwingLimits."));
 FAutoConsoleVariableRef CVarChaosImmPhysEnableDrives(TEXT("p.Chaos.ImmPhys.Joint.EnableDrives"), ChaosImmediate_Joint_EnableDrives, TEXT("EnableDrives."));
-FAutoConsoleVariableRef CVarChaosImmPhysProjectionPhase(TEXT("p.Chaos.ImmPhys.Joint.ProjectionPhase"), ChaosImmediate_Joint_ProjectionPhase, TEXT("Projection Phase (0=Never; 1=Apply; 2=ApplyPushOut)."));
 FAutoConsoleVariableRef CVarChaosImmPhysLinearProjection(TEXT("p.Chaos.ImmPhys.Joint.LinearProjection"), ChaosImmediate_Joint_LinearProjection, TEXT("6Dof joint projection amount override (if > 0)."));
 FAutoConsoleVariableRef CVarChaosImmPhysAngularProjection(TEXT("p.Chaos.ImmPhys.Joint.AngularProjection"), ChaosImmediate_Joint_AngularProjection, TEXT("6Dof joint projection amount override (if > 0)."));
 FAutoConsoleVariableRef CVarChaosImmPhysStiffness(TEXT("p.Chaos.ImmPhys.Joint.Stiffness"), ChaosImmediate_Joint_Stiffness, TEXT("6Dof joint stiffness override (if > 0)."));
@@ -163,7 +162,7 @@ namespace ImmediatePhysics_Chaos
 		, JointsRule(0, Joints)
 		, CollisionsRule(1, Collisions)
 		, Evolution(Particles, CollisionDetector)
-		, NumActiveActorHandles(0)
+		, NumActiveDynamicActorHandles(0)
 		, SimulationSpaceTransform(FTransform::Identity)
 	{
 		using namespace Chaos;
@@ -324,23 +323,26 @@ namespace ImmediatePhysics_Chaos
 
 	void FSimulation::SetNumActiveBodies(int32 InNumActiveActorHandles)
 	{
-		if (InNumActiveActorHandles < NumActiveActorHandles)
+		if (InNumActiveActorHandles == NumActiveDynamicActorHandles)
 		{
-			for (int ActorHandleIndex = InNumActiveActorHandles; ActorHandleIndex < NumActiveActorHandles; ++ActorHandleIndex)
-			{
-				GetActorHandle(ActorHandleIndex)->SetEnabled(false);
-			}
+			return;
 		}
-		else if (InNumActiveActorHandles > NumActiveActorHandles)
+
+		// @todo(ccaulfield): can be optimized, but I think we end up with kinematic at the start and statics at the end of the 
+		// list - maybe split them up or put kinematics at the end (in a way that does not impact particle order).
+		NumActiveDynamicActorHandles = 0;
+		for (int32 ActorHandleIndex = 0; ActorHandleIndex < ActorHandles.Num(); ++ActorHandleIndex)
 		{
-			check(InNumActiveActorHandles <= ActorHandles.Num());
-			for (int ActorHandleIndex = NumActiveActorHandles; ActorHandleIndex < InNumActiveActorHandles; ++ActorHandleIndex)
+			FActorHandle* Handle = GetActorHandle(ActorHandleIndex);
+			if (Handle->GetIsKinematic())
 			{
-				GetActorHandle(ActorHandleIndex)->SetEnabled(true);
+				continue;
 			}
+
+			bool bShouldBeActive = NumActiveDynamicActorHandles < InNumActiveActorHandles;
+			Handle->SetEnabled(bShouldBeActive);
+			++NumActiveDynamicActorHandles;
 		}
-	
-		NumActiveActorHandles = InNumActiveActorHandles;
 	}
 
 	void FSimulation::SetIgnoreCollisionPairTable(const TArray<FIgnorePair>& InIgnoreCollisionPairTable)
@@ -488,7 +490,6 @@ namespace ImmediatePhysics_Chaos
 			JointsSettings.bEnableTwistLimits = ChaosImmediate_Joint_EnableTwistLimits != 0;
 			JointsSettings.bEnableSwingLimits = ChaosImmediate_Joint_EnableSwingLimits != 0;
 			JointsSettings.bEnableDrives = ChaosImmediate_Joint_EnableDrives != 0;
-			JointsSettings.ProjectionPhase = ToJointSolverPhase(ChaosImmediate_Joint_ProjectionPhase);
 			JointsSettings.LinearProjection = ChaosImmediate_Joint_LinearProjection;
 			JointsSettings.AngularProjection = ChaosImmediate_Joint_AngularProjection;
 			JointsSettings.Stiffness = ChaosImmediate_Joint_Stiffness;
diff --git a/Engine/Source/Runtime/Engine/Public/Physics/ImmediatePhysics/ImmediatePhysicsChaos/ImmediatePhysicsSimulation_Chaos.h b/Engine/Source/Runtime/Engine/Public/Physics/ImmediatePhysics/ImmediatePhysicsChaos/ImmediatePhysicsSimulation_Chaos.h
--- a/Engine/Source/Runtime/Engine/Public/Physics/ImmediatePhysics/ImmediatePhysicsChaos/ImmediatePhysicsSimulation_Chaos.h
+++ b/Engine/Source/Runtime/Engine/Public/Physics/ImmediatePhysics/ImmediatePhysicsChaos/ImmediatePhysicsSimulation_Chaos.h
@@ -119,7 +119,7 @@ namespace ImmediatePhysics_Chaos
 		/** Mapping from entity index to handle */
 		// @todo(ccaulfield): we now have handles pointing to handles which is inefficient - we can do better than this, but don't want to change API yet
 		TArray<FActorHandle*> ActorHandles;
-		int NumActiveActorHandles;
+		int32 NumActiveDynamicActorHandles;
 
 		/** Mapping from constraint index to handle */
 		TArray<FJointHandle*> JointHandles;
diff --git a/Engine/Source/Runtime/Experimental/Chaos/Private/Chaos/Joint/PBDJointSolverGaussSeidel.cpp b/Engine/Source/Runtime/Experimental/Chaos/Private/Chaos/Joint/PBDJointSolverGaussSeidel.cpp
--- a/Engine/Source/Runtime/Experimental/Chaos/Private/Chaos/Joint/PBDJointSolverGaussSeidel.cpp
+++ b/Engine/Source/Runtime/Experimental/Chaos/Private/Chaos/Joint/PBDJointSolverGaussSeidel.cpp
@@ -14,6 +14,7 @@ namespace Chaos
 	{
 	}
 
+
 	void FJointSolverGaussSeidel::UpdateDerivedState()
 	{
 		Xs[0] = Ps[0] + Qs[0] * XLs[0].GetTranslation();
@@ -22,6 +23,7 @@ namespace Chaos
 		Rs[1] = Qs[1] * XLs[1].GetRotation();
 	}
 
+
 	void FJointSolverGaussSeidel::InitConstraints(
 		const FReal Dt,
 		const FPBDJointSolverSettings& SolverSettings,
@@ -55,6 +57,8 @@ namespace Chaos
 		TwistStiffness = FPBDJointUtilities::GetTwistStiffness(SolverSettings, JointSettings);
 		SwingStiffness = FPBDJointUtilities::GetSwingStiffness(SolverSettings, JointSettings);
 		AngularDriveStiffness = FPBDJointUtilities::GetAngularDriveStiffness(SolverSettings, JointSettings);
+		LinearProjection = FPBDJointUtilities::GetLinearProjection(SolverSettings, JointSettings);
+		AngularProjection = FPBDJointUtilities::GetAngularProjection(SolverSettings, JointSettings);
 		SwingTwistAngleTolerance = SolverSettings.SwingTwistAngleTolerance;
 		bEnableTwistLimits = SolverSettings.bEnableTwistLimits;
 		bEnableSwingLimits = SolverSettings.bEnableSwingLimits;
@@ -63,6 +67,37 @@ namespace Chaos
 		UpdateDerivedState();
 	}
 
+	void FJointSolverGaussSeidel::InitProjection(
+		const FReal Dt,
+		const FPBDJointSolverSettings& SolverSettings,
+		const FPBDJointSettings& JointSettings,
+		const FVec3& P0,
+		const FRotation3& Q0,
+		const FVec3& V0,
+		const FVec3& W0,
+		const FVec3& P1,
+		const FRotation3& Q1,
+		const FVec3& V1,
+		const FVec3& W1,
+		const FReal InvM0,
+		const FMatrix33& InvIL0,
+		const FReal InvM1,
+		const FMatrix33& InvIL1,
+		const FRigidTransform3& XL0,
+		const FRigidTransform3& XL1)
+	{
+		InitConstraints(Dt, SolverSettings, JointSettings, P0, Q0, P1, Q1, InvM0, InvIL0, InvM1, InvIL1, XL0, XL1);
+
+		InvILs[0] = FMatrix33(0, 0, 0);
+		InvMs[0] = 0;
+
+		Vs[0] = V0;
+		Vs[1] = V1;
+		Ws[0] = W0;
+		Ws[1] = W1;
+	}
+
+
 	void FJointSolverGaussSeidel::ApplyConstraints(
 		const FReal Dt,
 		const FPBDJointSettings& JointSettings)
@@ -80,6 +115,14 @@ namespace Chaos
 		ApplyRotationDrives(Dt, JointSettings);
 	}
 
+	void FJointSolverGaussSeidel::ApplyProjections(
+		const FReal Dt,
+		const FPBDJointSettings& JointSettings)
+	{
+		ApplyPositionProjection(Dt, JointSettings);
+	}
+
+
 	void FJointSolverGaussSeidel::ApplyRotationConstraints(
 		const FReal Dt,
 		const FPBDJointSettings& JointSettings)
@@ -197,6 +240,7 @@ namespace Chaos
 		Qs[1].EnforceShortestArcWith(Qs[0]);
 	}
 
+
 	void FJointSolverGaussSeidel::ApplyRotationDelta(
 		const FReal Stiffness,
 		const FVec3& Axis0,
@@ -214,6 +258,21 @@ namespace Chaos
 		ApplyRotationDelta(Stiffness, DR0, DR1);
 	}
 
+
+	void FJointSolverGaussSeidel::ApplyVelocityDelta(
+		const FReal Stiffness,
+		const FVec3& DV0,
+		const FVec3& DW0,
+		const FVec3& DV1,
+		const FVec3& DW1)
+	{
+		Vs[0] = Vs[0] + Stiffness * DV0;
+		Ws[0] = Ws[0] + Stiffness * DW0;
+		Vs[1] = Vs[1] + Stiffness * DV1;
+		Ws[1] = Ws[1] + Stiffness * DW1;
+	}
+
+
 	void FJointSolverGaussSeidel::ApplyTwistConstraint(
 		const FReal Dt,
 		const FPBDJointSettings& JointSettings)
@@ -526,4 +585,48 @@ namespace Chaos
 		UpdateDerivedState();
 	}
 
+
+	void FJointSolverGaussSeidel::ApplyPositionProjection(
+		const FReal Dt,
+		const FPBDJointSettings& JointSettings)
+	{
+		// Apply a position correction with the parent body set to infinite mass, then correct the velocity
+		// in the direction of the position correction.
+		const FVec3 CX = FPBDJointUtilities::GetLimitedPositionError(JointSettings, Rs[0], Xs[1] - Xs[0]);
+		const FReal CXLen = CX.Size();
+		if (CXLen > KINDA_SMALL_NUMBER)
+		{
+			const FVec3 CXDir = CX / CXLen;
+
+			// Get required velocity correction in direction of position correction
+			const FVec3 V0 = Vs[0] + FVec3::CrossProduct(Ws[0], Xs[0] - Ps[0]);
+			const FVec3 V1 = Vs[1] + FVec3::CrossProduct(Ws[1], Xs[1] - Ps[1]);
+			FVec3 CV = FPBDJointUtilities::GetLimitedVelocityError(JointSettings, Rs[0], CX, V1 - V0);
+			CV = CV - FVec3::DotProduct(CXDir, CV) * CXDir;
+
+			FMatrix33 InvI1 = Utilities::ComputeWorldSpaceInertia(Qs[1], InvILs[1]);
+			FMatrix33 M1 = Utilities::ComputeJointFactorMatrix(Xs[1] - Ps[1], InvI1, InvMs[1]);
+			FMatrix33 MI = M1.Inverse();
+			const FVec3 DX = Utilities::Multiply(MI, CX);
+
+			// Apply position correction
+			const FVec3 DP1 = -InvMs[1] * DX;
+			const FVec3 DR1 = Utilities::Multiply(InvI1, FVec3::CrossProduct(Xs[1] - Ps[1], -DX));
+
+			ApplyPositionDelta(LinearStiffness, FVec3(0), DP1);
+			ApplyRotationDelta(LinearStiffness, FVec3(0), DR1);
+			UpdateDerivedState();
+
+			InvI1 = Utilities::ComputeWorldSpaceInertia(Qs[1], InvILs[1]);
+			M1 = Utilities::ComputeJointFactorMatrix(Xs[1] - Ps[1], InvI1, InvMs[1]);
+			MI = M1.Inverse();
+
+			// Apply velocity correction
+			const FVec3 DV = Utilities::Multiply(MI, CV);
+			const FVec3 DV1 = -InvMs[1] * DV;
+			const FVec3 DW1 = Utilities::Multiply(InvI1, FVec3::CrossProduct(Xs[1] - Ps[1], -DV));
+
+			ApplyVelocityDelta(LinearProjection, FVec3(0), FVec3(0), DV1, DW1);
+		}
+	}
 }
diff --git a/Engine/Source/Runtime/Experimental/Chaos/Private/Chaos/PBDJointConstraintUtilities.cpp b/Engine/Source/Runtime/Experimental/Chaos/Private/Chaos/PBDJointConstraintUtilities.cpp
--- a/Engine/Source/Runtime/Experimental/Chaos/Private/Chaos/PBDJointConstraintUtilities.cpp
+++ b/Engine/Source/Runtime/Experimental/Chaos/Private/Chaos/PBDJointConstraintUtilities.cpp
@@ -54,10 +54,22 @@ namespace Chaos
 		const FPBDJointSolverSettings& SolverSettings,
 		const FPBDJointSettings& JointSettings)
 	{
-		FReal DriveStiffness = (SolverSettings.DriveStiffness > 0.0f) ? SolverSettings.DriveStiffness : JointSettings.Motion.AngularDriveStiffness;
-		return DriveStiffness;
+		return (SolverSettings.DriveStiffness > 0.0f) ? SolverSettings.DriveStiffness : JointSettings.Motion.AngularDriveStiffness;
 	}
 
+	FReal FPBDJointUtilities::GetLinearProjection(
+		const FPBDJointSolverSettings& SolverSettings,
+		const FPBDJointSettings& JointSettings)
+	{
+		return (SolverSettings.LinearProjection > 0.0f) ? SolverSettings.LinearProjection : JointSettings.Motion.LinearProjection;
+	}
+
+	FReal FPBDJointUtilities::GetAngularProjection(
+		const FPBDJointSolverSettings& SolverSettings,
+		const FPBDJointSettings& JointSettings)
+	{
+		return (SolverSettings.AngularProjection > 0.0f) ? SolverSettings.AngularProjection : JointSettings.Motion.AngularProjection;
+	}
 
 	FVec3 FPBDJointUtilities::ConditionInertia(const FVec3& InI, const FReal MaxRatio)
 	{
diff --git a/Engine/Source/Runtime/Experimental/Chaos/Private/Chaos/PBDJointConstraints.cpp b/Engine/Source/Runtime/Experimental/Chaos/Private/Chaos/PBDJointConstraints.cpp
--- a/Engine/Source/Runtime/Experimental/Chaos/Private/Chaos/PBDJointConstraints.cpp
+++ b/Engine/Source/Runtime/Experimental/Chaos/Private/Chaos/PBDJointConstraints.cpp
@@ -20,11 +20,8 @@ namespace Chaos
 	DECLARE_CYCLE_STAT(TEXT("TPBDJointConstraints::Sort"), STAT_Joints_Sort, STATGROUP_Chaos);
 	DECLARE_CYCLE_STAT(TEXT("TPBDJointConstraints::Apply"), STAT_Joints_Apply, STATGROUP_Chaos);
 	DECLARE_CYCLE_STAT(TEXT("TPBDJointConstraints::ApplyPushOut"), STAT_Joints_ApplyPushOut, STATGROUP_Chaos);
-	DECLARE_CYCLE_STAT(TEXT("TPBDJointConstraints::Drives"), STAT_Joints_Drives, STATGROUP_Chaos);
-	DECLARE_CYCLE_STAT(TEXT("TPBDJointConstraints::Solve"), STAT_Joints_Solve, STATGROUP_Chaos);
 	DECLARE_CYCLE_STAT(TEXT("TPBDJointConstraints::SolveCholesky"), STAT_Joints_Solve_Cholesky, STATGROUP_Chaos);
 	DECLARE_CYCLE_STAT(TEXT("TPBDJointConstraints::SolveGaussSeidel"), STAT_Joints_Solve_GaussSeidel, STATGROUP_Chaos);
-	DECLARE_CYCLE_STAT(TEXT("TPBDJointConstraints::Project"), STAT_Joints_Project, STATGROUP_Chaos);
 
 	bool ChaosJoint_UseCholeskySolver = false;
 	FAutoConsoleVariableRef CVarChaosImmPhysDeltaTime(TEXT("p.Chaos.Joint.UseCholeskySolver"), ChaosJoint_UseCholeskySolver, TEXT("Whether to use the new solver"));
@@ -175,7 +172,6 @@ namespace Chaos
 		, bEnableTwistLimits(true)
 		, bEnableSwingLimits(true)
 		, bEnableDrives(true)
-		, ProjectionPhase(EJointSolverPhase::None)
 		, LinearProjection((FReal)0)
 		, AngularProjection((FReal)0)
 		, Stiffness((FReal)0)
@@ -221,7 +217,24 @@ namespace Chaos
 		return ConstraintParticles.Num();
 	}
 
-	
+	void FPBDJointConstraints::GetConstrainedParticleIndices(const int32 ConstraintIndex, int32& Index0, int32& Index1) const
+	{
+		// In solvers we assume Particle0 is the parent particle (which it usually is as implemented in the editor). 
+		// However, it is possible to set it up so that the kinematic particle is the child which we don't support, so...
+		// If particle 0 is kinematic we make it the parent, otherwise particle 1 is the parent.
+		// @todo(ccaulfield): look into this and confirm/fix properly
+		if (!ConstraintParticles[ConstraintIndex][0]->CastToRigidParticle())
+		{
+			Index0 = 0;
+			Index1 = 1;
+		}
+		else
+		{
+			Index0 = 1;
+			Index1 = 0;
+		}
+	}
+
 	typename FPBDJointConstraints::FConstraintContainerHandle* FPBDJointConstraints::AddConstraint(const FParticlePair& InConstrainedParticles, const FRigidTransform3& WorldConstraintFrame)
 	{
 		FTransformPair ConstraintFrames;
@@ -411,8 +424,8 @@ namespace Chaos
 	
 	void FPBDJointConstraints::CalculateConstraintSpace(int32 ConstraintIndex, FVec3& OutX0, FMatrix33& OutR0, FVec3& OutX1, FMatrix33& OutR1, FVec3& OutCR) const
 	{
-		const int32 Index0 = 1;
-		const int32 Index1 = 0;
+		int32 Index0, Index1;
+		GetConstrainedParticleIndices(ConstraintIndex, Index0, Index1);
 		TGenericParticleHandle<FReal, 3> Particle0 = TGenericParticleHandle<FReal, 3>(ConstraintParticles[ConstraintIndex][Index0]);
 		TGenericParticleHandle<FReal, 3> Particle1 = TGenericParticleHandle<FReal, 3>(ConstraintParticles[ConstraintIndex][Index1]);
 		const FVec3 P0 = FParticleUtilities::GetCoMWorldPosition(Particle0);
@@ -450,26 +463,30 @@ namespace Chaos
 			PreApplyCallback(Dt, Handles);
 		}
 
-		// Solve for joint position or velocity, depending on settings
 		if (Settings.ApplyPairIterations > 0)
 		{
-			SolvePosition(Dt, Settings.ApplyPairIterations, It, NumIts);
+			if (ChaosJoint_UseCholeskySolver)
+			{
+				SCOPE_CYCLE_COUNTER(STAT_Joints_Solve_Cholesky);
+				for (int32 ConstraintIndex = 0; ConstraintIndex < NumConstraints(); ++ConstraintIndex)
+				{
+					SolvePosition_Cholesky(Dt, ConstraintIndex, Settings.ApplyPairIterations, It, NumIts);
+				}
+			}
+			else
+			{
+				SCOPE_CYCLE_COUNTER(STAT_Joints_Solve_GaussSeidel);
+				for (int32 ConstraintIndex = 0; ConstraintIndex < NumConstraints(); ++ConstraintIndex)
+				{
+					SolvePosition_GaussSiedel(Dt, ConstraintIndex, Settings.ApplyPairIterations, It, NumIts);
+				}
+			}
 		}
 
 		if (PostApplyCallback != nullptr)
 		{
 			PostApplyCallback(Dt, Handles);
 		}
-
-		// Correct remaining errors after last call to Solve if enabled in this phase
-		if (Settings.ProjectionPhase == EJointSolverPhase::Apply)
-		{
-			int32 ProjectionIt = NumIts - 1;
-			if (It == ProjectionIt)
-			{
-				ApplyProjection(Dt);
-			}
-		}
 	}
 
 	bool FPBDJointConstraints::ApplyPushOut(const FReal Dt, const int32 It, const int32 NumIts)
@@ -479,58 +496,27 @@ namespace Chaos
 		// @todo(ccaulfield): track whether we are sufficiently solved
 		bool bNeedsAnotherIteration = true;
 
-		// Solve for positions
 		if (Settings.ApplyPushOutPairIterations > 0)
 		{
-			SolvePosition(Dt, Settings.ApplyPushOutPairIterations, It, NumIts);
-		}
-
-		// Correct remaining errors after the last call to Solve (which depends on if PositionSolve is enabled in ApplyPushOut)
-		if (Settings.ProjectionPhase == EJointSolverPhase::ApplyPushOut)
-		{
-			int32 ProjectionIt = (Settings.ApplyPushOutPairIterations > 0) ? NumIts - 1 : 0;
-			if (It == ProjectionIt)
-			{
-				ApplyProjection(Dt);
-			}
-		}
-
-		return bNeedsAnotherIteration;
-	}
-
-	void FPBDJointConstraints::SolvePosition(const FReal Dt, const int32 NumPairIts, const int32 It, const int32 NumIts)
-	{
-		if (ChaosJoint_UseCholeskySolver)
-		{
-			SCOPE_CYCLE_COUNTER(STAT_Joints_Solve_Cholesky);
-			for (int32 ConstraintIndex = 0; ConstraintIndex < NumConstraints(); ++ConstraintIndex)
+			if (ChaosJoint_UseCholeskySolver)
 			{
-				SolvePosition_Cholesky(Dt, ConstraintIndex, NumPairIts, It, NumIts);
+				// TODO
 			}
-		}
-		else
-		{
-			SCOPE_CYCLE_COUNTER(STAT_Joints_Solve_GaussSeidel);
-			for (int32 ConstraintIndex = 0; ConstraintIndex < NumConstraints(); ++ConstraintIndex)
+			else
 			{
-				SolvePosition_GaussSiedel(Dt, ConstraintIndex, NumPairIts, It, NumIts);
+				for (int32 ConstraintIndex = 0; ConstraintIndex < NumConstraints(); ++ConstraintIndex)
+				{
+					ProjectPosition_GaussSiedel(Dt, ConstraintIndex, Settings.ApplyPushOutPairIterations, It, NumIts);
+				}
 			}
 		}
-	}
-
-	void FPBDJointConstraints::ApplyProjection(const FReal Dt)
-	{
-		SCOPE_CYCLE_COUNTER(STAT_Joints_Project);
-
-		for (int32 ConstraintIndex = 0; ConstraintIndex < NumConstraints(); ++ConstraintIndex)
-		{
-			ProjectPosition(Dt, ConstraintIndex, 0, 1);
-		}
 
 		if (PostProjectCallback != nullptr)
 		{
 			PostProjectCallback(Dt, Handles);
 		}
+
+		return bNeedsAnotherIteration;
 	}
 
 	//////////////////////////////////////////////////////////////////////////
@@ -567,24 +553,24 @@ namespace Chaos
 			PreApplyCallback(Dt, SortedConstraintHandles);
 		}
 
-		// Solve for joint position or velocity, depending on settings
+
 		if (Settings.ApplyPairIterations > 0)
 		{
-			SCOPE_CYCLE_COUNTER(STAT_Joints_Solve);
-			for (FConstraintContainerHandle* ConstraintHandle : SortedConstraintHandles)
+			if (ChaosJoint_UseCholeskySolver)
 			{
-				SolvePosition(Dt, ConstraintHandle->GetConstraintIndex(), Settings.ApplyPairIterations, It, NumIts);
+				SCOPE_CYCLE_COUNTER(STAT_Joints_Solve_Cholesky);
+				for (FConstraintContainerHandle* ConstraintHandle : SortedConstraintHandles)
+				{
+					SolvePosition_Cholesky(Dt, ConstraintHandle->GetConstraintIndex(), Settings.ApplyPairIterations, It, NumIts);
+				}
 			}
-		}
-
-		// Correct remaining errors after last call to Solve if enabled in this phase
-		if (Settings.ProjectionPhase == EJointSolverPhase::Apply)
-		{
-			SCOPE_CYCLE_COUNTER(STAT_Joints_Project);
-			int32 ProjectionIt = NumIts - 1;
-			if (It == ProjectionIt)
+			else
 			{
-				ApplyProjection(Dt, InConstraintHandles);
+				SCOPE_CYCLE_COUNTER(STAT_Joints_Solve_GaussSeidel);
+				for (FConstraintContainerHandle* ConstraintHandle : SortedConstraintHandles)
+				{
+					SolvePosition_GaussSiedel(Dt, ConstraintHandle->GetConstraintIndex(), Settings.ApplyPairIterations, It, NumIts);
+				}
 			}
 		}
 
@@ -609,45 +595,27 @@ namespace Chaos
 				return L.GetConstraintLevel() < R.GetConstraintLevel();
 			});
 
-		// Solve for positions
 		if (Settings.ApplyPushOutPairIterations > 0)
 		{
-			SCOPE_CYCLE_COUNTER(STAT_Joints_Solve);
-
-			for (FConstraintContainerHandle* ConstraintHandle : SortedConstraintHandles)
+			if (ChaosJoint_UseCholeskySolver)
 			{
-				SolvePosition(Dt, ConstraintHandle->GetConstraintIndex(), Settings.ApplyPushOutPairIterations, It, NumIts);
+				// TODO
 			}
-		}
-
-		// Correct remaining errors after the last call to Solve (which depends on if PositionSolve is enabled in ApplyPushOut)
-		if (Settings.ProjectionPhase == EJointSolverPhase::ApplyPushOut)
-		{
-			SCOPE_CYCLE_COUNTER(STAT_Joints_Project);
-			int32 ProjectionIt = (Settings.ApplyPushOutPairIterations > 0) ? NumIts - 1 : 0;
-			if (It == ProjectionIt)
+			else
 			{
-				ApplyProjection(Dt, InConstraintHandles);
+				for (FConstraintContainerHandle* ConstraintHandle : SortedConstraintHandles)
+				{
+					ProjectPosition_GaussSiedel(Dt, ConstraintHandle->GetConstraintIndex(), Settings.ApplyPushOutPairIterations, It, NumIts);
+				}
 			}
 		}
 
-		return bNeedsAnotherIteration;
-	}
-
-	
-	void FPBDJointConstraints::ApplyProjection(const FReal Dt, const TArray<FConstraintContainerHandle*>& InConstraintHandles)
-	{
-		TArray<FConstraintContainerHandle*> SortedConstraintHandles = InConstraintHandles;
-		SortedConstraintHandles.Sort([](const FConstraintContainerHandle& L, const FConstraintContainerHandle& R)
-			{
-				// Sort bodies from root to leaf
-				return L.GetConstraintLevel() < R.GetConstraintLevel();
-			});
-
-		for (FConstraintContainerHandle* ConstraintHandle : SortedConstraintHandles)
+		if (PostProjectCallback != nullptr)
 		{
-			ProjectPosition(Dt, ConstraintHandle->GetConstraintIndex(), 0, 1);
+			PostProjectCallback(Dt, SortedConstraintHandles);
 		}
+
+		return bNeedsAnotherIteration;
 	}
 
 	//////////////////////////////////////////////////////////////////////////
@@ -681,15 +649,13 @@ namespace Chaos
 	}
 
 
-	void FPBDJointConstraints::SolvePosition(const FReal Dt, const int32 ConstraintIndex, const int32 NumPairIts, const int32 It, const int32 NumIts)
+	void FPBDJointConstraints::UpdateParticleState(TPBDRigidParticleHandle<FReal, 3>* Rigid, const FReal Dt, const FVec3& P, const FRotation3& Q, const FVec3& V, const FVec3& W)
 	{
-		if (ChaosJoint_UseCholeskySolver)
-		{
-			SolvePosition_Cholesky(Dt, ConstraintIndex, NumPairIts, It, NumIts);
-		}
-		else
+		if ((Rigid != nullptr) && (Rigid->ObjectState() == EObjectStateType::Dynamic))
 		{
-			SolvePosition_GaussSiedel(Dt, ConstraintIndex, NumPairIts, It, NumIts);
+			FParticleUtilities::SetCoMWorldTransform(Rigid, P, Q);
+			Rigid->SetV(V);
+			Rigid->SetW(W);
 		}
 	}
 
@@ -704,9 +670,8 @@ namespace Chaos
 
 		const FPBDJointSettings& JointSettings = ConstraintSettings[ConstraintIndex];
 
-		// Switch particles - internally we assume the first body is the parent (i.e., the space in which constraint limits are specified)
-		const int32 Index0 = 1;
-		const int32 Index1 = 0;
+		int32 Index0, Index1;
+		GetConstrainedParticleIndices(ConstraintIndex, Index0, Index1);
 		TGenericParticleHandle<FReal, 3> Particle0 = TGenericParticleHandle<FReal, 3>(ConstraintParticles[ConstraintIndex][Index0]);
 		TGenericParticleHandle<FReal, 3> Particle1 = TGenericParticleHandle<FReal, 3>(ConstraintParticles[ConstraintIndex][Index1]);
 
@@ -746,9 +711,8 @@ namespace Chaos
 
 		const FPBDJointSettings& JointSettings = ConstraintSettings[ConstraintIndex];
 
-		// Switch particles - internally we assume the first body is the parent (i.e., the space in which constraint limits are specified)
-		const int32 Index0 = 1;
-		const int32 Index1 = 0;
+		int32 Index0, Index1;
+		GetConstrainedParticleIndices(ConstraintIndex, Index0, Index1);
 		TGenericParticleHandle<FReal, 3> Particle0 = TGenericParticleHandle<FReal, 3>(ConstraintParticles[ConstraintIndex][Index0]);
 		TGenericParticleHandle<FReal, 3> Particle1 = TGenericParticleHandle<FReal, 3>(ConstraintParticles[ConstraintIndex][Index1]);
 
@@ -778,99 +742,45 @@ namespace Chaos
 		UpdateParticleState(Particle1->CastToRigidParticle(), Dt, Solver.GetP(1), Solver.GetQ(1));
 	}
 
-
-	void FPBDJointConstraints::ProjectPosition(const FReal Dt, const int32 ConstraintIndex, const int32 It, const int32 NumIts)
+	void FPBDJointConstraints::ProjectPosition_GaussSiedel(const FReal Dt, const int32 ConstraintIndex, const int32 NumPairIts, const int32 It, const int32 NumIts)
 	{
-		const FPBDJointSettings& JointSettings = ConstraintSettings[ConstraintIndex];
-
-		// Scale projection up to ProjectionSetting over NumProjectionIts
-		const FReal LinearProjectionFactor = (Settings.LinearProjection > 0) ? Settings.LinearProjection : JointSettings.Motion.LinearProjection;
-		const FReal AngularProjectionFactor = (Settings.AngularProjection > 0) ? Settings.AngularProjection : JointSettings.Motion.AngularProjection;
-		if ((LinearProjectionFactor == (FReal)0) && (AngularProjectionFactor == (FReal)0))
-		{
-			return;
-		}
-
 		const TVector<TGeometryParticleHandle<FReal, 3>*, 2>& Constraint = ConstraintParticles[ConstraintIndex];
-		UE_LOG(LogChaosJoint, VeryVerbose, TEXT("Project Joint Constraint %d %f %f (it = %d / %d)"), ConstraintIndex, LinearProjectionFactor, AngularProjectionFactor, It, NumIts);
+		UE_LOG(LogChaosJoint, VeryVerbose, TEXT("Solve Joint Constraint %d %s %s (dt = %f; it = %d / %d)"), ConstraintIndex, *Constraint[0]->ToString(), *Constraint[1]->ToString(), Dt, It, NumIts);
 
-		// Switch particles - internally we assume the first body is the parent (i.e., the space in which constraint limits are specified)
-		const int32 Index0 = 1;
-		const int32 Index1 = 0;
+		const FPBDJointSettings& JointSettings = ConstraintSettings[ConstraintIndex];
+
+		int32 Index0, Index1;
+		GetConstrainedParticleIndices(ConstraintIndex, Index0, Index1);
 		TGenericParticleHandle<FReal, 3> Particle0 = TGenericParticleHandle<FReal, 3>(ConstraintParticles[ConstraintIndex][Index0]);
 		TGenericParticleHandle<FReal, 3> Particle1 = TGenericParticleHandle<FReal, 3>(ConstraintParticles[ConstraintIndex][Index1]);
 
-		FVec3 P0 = FParticleUtilities::GetCoMWorldPosition(Particle0);
-		FRotation3 Q0 = FParticleUtilities::GetCoMWorldRotation(Particle0);
-		FVec3 P1 = FParticleUtilities::GetCoMWorldPosition(Particle1);
-		FRotation3 Q1 = FParticleUtilities::GetCoMWorldRotation(Particle1);
-		const FRigidTransform3 XL0 = FParticleUtilities::ParticleLocalToCoMLocal(Particle0, JointSettings.ConstraintFrames[Index0]);
-		const FRigidTransform3 XL1 = FParticleUtilities::ParticleLocalToCoMLocal(Particle1, JointSettings.ConstraintFrames[Index1]);
-		float InvM0 = Particle0->InvM();
-		float InvM1 = Particle1->InvM();
-		FMatrix33 InvIL0 = Particle0->InvI();
-		FMatrix33 InvIL1 = Particle1->InvI();
-
-		// Freeze the closest to kinematic connection if there is a difference
-		const int32 Level0 = ConstraintStates[ConstraintIndex].ParticleLevels[Index0];
-		const int32 Level1 = ConstraintStates[ConstraintIndex].ParticleLevels[Index1];
-		if (Level0 < Level1)
-		{
-			InvM0 = 0;
-			InvIL0 = FMatrix33(0, 0, 0);
-		}
-		else if (Level1 < Level0)
-		{
-			InvM1 = 0;
-			InvIL1 = FMatrix33(0, 0, 0);
-		}
-
-		const TVector<EJointMotionType, 3>& LinearMotion = JointSettings.Motion.LinearMotionTypes;
-		EJointMotionType TwistMotion = JointSettings.Motion.AngularMotionTypes[(int32)EJointAngularConstraintIndex::Twist];
-		EJointMotionType Swing1Motion = JointSettings.Motion.AngularMotionTypes[(int32)EJointAngularConstraintIndex::Swing1];
-		EJointMotionType Swing2Motion = JointSettings.Motion.AngularMotionTypes[(int32)EJointAngularConstraintIndex::Swing2];
-
-		if (AngularProjectionFactor > 0)
-		{
-			if (Settings.bEnableTwistLimits)
-			{
-				// Remove Twist Error
-				FPBDJointUtilities::ApplyJointTwistProjection(Dt, Settings, JointSettings, AngularProjectionFactor, XL0, XL1, P0, Q0, P1, Q1, InvM0, InvIL0, InvM1, InvIL1);
-			}
-
-			if (Settings.bEnableSwingLimits)
-			{
-				// Remove Swing Error
-				if ((Swing1Motion == EJointMotionType::Limited) && (Swing2Motion == EJointMotionType::Limited))
-				{
-					FPBDJointUtilities::ApplyJointConeProjection(Dt, Settings, JointSettings, AngularProjectionFactor, XL0, XL1, P0, Q0, P1, Q1, InvM0, InvIL0, InvM1, InvIL1);
-				}
-				else
-				{
-					if (Swing1Motion != EJointMotionType::Free)
-					{
-						FPBDJointUtilities::ApplyJointSwingProjection(Dt, Settings, JointSettings, AngularProjectionFactor, XL0, XL1, EJointAngularConstraintIndex::Swing1, EJointAngularAxisIndex::Swing1, P0, Q0, P1, Q1, InvM0, InvIL0, InvM1, InvIL1);
-					}
-					if (Swing2Motion != EJointMotionType::Free)
-					{
-						FPBDJointUtilities::ApplyJointSwingProjection(Dt, Settings, JointSettings, AngularProjectionFactor, XL0, XL1, EJointAngularConstraintIndex::Swing2, EJointAngularAxisIndex::Swing2, P0, Q0, P1, Q1, InvM0, InvIL0, InvM1, InvIL1);
-					}
-				}
-			}
-		}
+		FJointSolverGaussSeidel Solver;
+		Solver.InitProjection(
+			Dt,
+			Settings,
+			JointSettings,
+			FParticleUtilities::GetCoMWorldPosition(Particle0),
+			FParticleUtilities::GetCoMWorldRotation(Particle0),
+			Particle0->V(),
+			Particle0->W(),
+			FParticleUtilities::GetCoMWorldPosition(Particle1),
+			FParticleUtilities::GetCoMWorldRotation(Particle1),
+			Particle1->V(),
+			Particle1->W(),
+			Particle0->InvM(),
+			Particle0->InvI(),
+			Particle1->InvM(),
+			Particle1->InvI(),
+			FParticleUtilities::ParticleLocalToCoMLocal(Particle0, JointSettings.ConstraintFrames[Index0]),
+			FParticleUtilities::ParticleLocalToCoMLocal(Particle1, JointSettings.ConstraintFrames[Index1]));
 
-		// Remove Position Error
-		if (LinearProjectionFactor > 0)
+		for (int32 PairIt = 0; PairIt < NumPairIts; ++PairIt)
 		{
-			if ((LinearMotion[0] != EJointMotionType::Free) || (LinearMotion[1] != EJointMotionType::Free) || (LinearMotion[2] != EJointMotionType::Free))
-			{
-				FPBDJointUtilities::ApplyJointPositionProjection(Dt, Settings, JointSettings, LinearProjectionFactor, XL0, XL1, P0, Q0, P1, Q1, InvM0, InvIL0, InvM1, InvIL1);
-			}
+			Solver.ApplyProjections(Dt, JointSettings);
 		}
 
-		// Update the particles
-		UpdateParticleState(Particle0->CastToRigidParticle(), Dt, P0, Q0, false);
-		UpdateParticleState(Particle1->CastToRigidParticle(), Dt, P1, Q1, false);
+		UpdateParticleState(Particle0->CastToRigidParticle(), Dt, Solver.GetP(0), Solver.GetQ(0), Solver.GetV(0), Solver.GetW(0));
+		UpdateParticleState(Particle1->CastToRigidParticle(), Dt, Solver.GetP(1), Solver.GetQ(1), Solver.GetV(1), Solver.GetW(1));
 	}
 }
 
diff --git a/Engine/Source/Runtime/Experimental/Chaos/Public/Chaos/Joint/PBDJointSolverCholesky.h b/Engine/Source/Runtime/Experimental/Chaos/Public/Chaos/Joint/PBDJointSolverCholesky.h
--- a/Engine/Source/Runtime/Experimental/Chaos/Public/Chaos/Joint/PBDJointSolverCholesky.h
+++ b/Engine/Source/Runtime/Experimental/Chaos/Public/Chaos/Joint/PBDJointSolverCholesky.h
@@ -22,20 +22,20 @@ namespace Chaos
 	class FJointSolverCholesky
 	{
 	public:
-		static const int32 MaxConstraintedBodies = 2;
+		static const int32 MaxConstrainedBodies = 2;
 
 		using FDenseMatrix66 = TDenseMatrix<6 * 6>;
 		using FDenseMatrix61 = TDenseMatrix<6 * 1>;
 
 		FORCEINLINE const FVec3& GetP(const int32 Index) const
 		{
-			checkSlow(Index < MaxConstraintedBodies);
+			checkSlow(Index < MaxConstrainedBodies);
 			return Ps[Index];
 		}
 
 		FORCEINLINE const FRotation3& GetQ(const int32 Index) const
 		{
-			checkSlow(Index < MaxConstraintedBodies);
+			checkSlow(Index < MaxConstrainedBodies);
 			return Qs[Index];
 		}
 
@@ -190,17 +190,17 @@ namespace Chaos
 			const FDenseMatrix61& C);
 
 		// Local-space constraint settings
-		FRigidTransform3 XLs[MaxConstraintedBodies];	// Local-space joint connector transforms
-		FMatrix33 InvILs[MaxConstraintedBodies];		// Local-space inverse inertias
-		FReal InvMs[MaxConstraintedBodies];				// Inverse masses
+		FRigidTransform3 XLs[MaxConstrainedBodies];	// Local-space joint connector transforms
+		FMatrix33 InvILs[MaxConstrainedBodies];		// Local-space inverse inertias
+		FReal InvMs[MaxConstrainedBodies];				// Inverse masses
 
 		// World-space constraint state
-		FVec3 Xs[MaxConstraintedBodies];				// World-space joint connector positions
-		FRotation3 Rs[MaxConstraintedBodies];			// World-space joint connector rotations
+		FVec3 Xs[MaxConstrainedBodies];				// World-space joint connector positions
+		FRotation3 Rs[MaxConstrainedBodies];			// World-space joint connector rotations
 
 		// World-space body state
-		FVec3 Ps[MaxConstraintedBodies];				// World-space particle CoM positions
-		FRotation3 Qs[MaxConstraintedBodies];			// World-space particle CoM rotations
+		FVec3 Ps[MaxConstrainedBodies];				// World-space particle CoM positions
+		FRotation3 Qs[MaxConstrainedBodies];			// World-space particle CoM rotations
 
 		// Settings
 		FReal Stiffness;
diff --git a/Engine/Source/Runtime/Experimental/Chaos/Public/Chaos/Joint/PBDJointSolverGaussSeidel.h b/Engine/Source/Runtime/Experimental/Chaos/Public/Chaos/Joint/PBDJointSolverGaussSeidel.h
--- a/Engine/Source/Runtime/Experimental/Chaos/Public/Chaos/Joint/PBDJointSolverGaussSeidel.h
+++ b/Engine/Source/Runtime/Experimental/Chaos/Public/Chaos/Joint/PBDJointSolverGaussSeidel.h
@@ -22,20 +22,32 @@ namespace Chaos
 	class FJointSolverGaussSeidel
 	{
 	public:
-		static const int32 MaxConstraintedBodies = 2;
+		static const int32 MaxConstrainedBodies = 2;
 
 		FORCEINLINE const FVec3& GetP(const int32 Index) const
 		{
-			checkSlow(Index < MaxConstraintedBodies);
+			checkSlow(Index < MaxConstrainedBodies);
 			return Ps[Index];
 		}
 
 		FORCEINLINE const FRotation3& GetQ(const int32 Index) const
 		{
-			checkSlow(Index < MaxConstraintedBodies);
+			checkSlow(Index < MaxConstrainedBodies);
 			return Qs[Index];
 		}
 
+		FORCEINLINE const FVec3& GetV(const int32 Index) const
+		{
+			checkSlow(Index < MaxConstrainedBodies);
+			return Vs[Index];
+		}
+
+		FORCEINLINE const FVec3& GetW(const int32 Index) const
+		{
+			checkSlow(Index < MaxConstrainedBodies);
+			return Ws[Index];
+		}
+
 		FJointSolverGaussSeidel();
 
 		void InitConstraints(
@@ -53,6 +65,25 @@ namespace Chaos
 			const FRigidTransform3& XL0,
 			const FRigidTransform3& XL1);
 
+		void InitProjection(
+			const FReal Dt,
+			const FPBDJointSolverSettings& SolverSettings,
+			const FPBDJointSettings& JointSettings,
+			const FVec3& P0,
+			const FRotation3& Q0,
+			const FVec3& V0,
+			const FVec3& W0,
+			const FVec3& P1,
+			const FRotation3& Q1,
+			const FVec3& V1,
+			const FVec3& W1,
+			const FReal InvM0,
+			const FMatrix33& InvIL0,
+			const FReal InvM1,
+			const FMatrix33& InvIL1,
+			const FRigidTransform3& XL0,
+			const FRigidTransform3& XL1);
+
 		void ApplyConstraints(
 			const FReal Dt,
 			const FPBDJointSettings& JointSettings);
@@ -61,6 +92,10 @@ namespace Chaos
 			const FReal Dt,
 			const FPBDJointSettings& JointSettings);
 
+		void ApplyProjections(
+			const FReal Dt,
+			const FPBDJointSettings& JointSettings);
+
 	private:
 
 		void UpdateDerivedState();
@@ -95,6 +130,13 @@ namespace Chaos
 			const FVec3& Axis1,
 			const FReal Angle1);
 
+		void ApplyVelocityDelta(
+			const FReal Stiffness,
+			const FVec3& DV0,
+			const FVec3& DW0,
+			const FVec3& DV1,
+			const FVec3& DW1);
+
 
 		void ApplyTwistConstraint(
 			const FReal Dt,
@@ -132,25 +174,33 @@ namespace Chaos
 			const FReal Dt,
 			const FPBDJointSettings& JointSettings);
 
+		void ApplyPositionProjection(
+			const FReal Dt,
+			const FPBDJointSettings& JointSettings);
+
 
 		// Local-space constraint settings
-		FRigidTransform3 XLs[MaxConstraintedBodies];	// Local-space joint connector transforms
-		FMatrix33 InvILs[MaxConstraintedBodies];		// Local-space inverse inertias
-		FReal InvMs[MaxConstraintedBodies];				// Inverse masses
+		FRigidTransform3 XLs[MaxConstrainedBodies];	// Local-space joint connector transforms
+		FMatrix33 InvILs[MaxConstrainedBodies];		// Local-space inverse inertias
+		FReal InvMs[MaxConstrainedBodies];			// Inverse masses
 
 		// World-space constraint state
-		FVec3 Xs[MaxConstraintedBodies];				// World-space joint connector positions
-		FRotation3 Rs[MaxConstraintedBodies];			// World-space joint connector rotations
+		FVec3 Xs[MaxConstrainedBodies];				// World-space joint connector positions
+		FRotation3 Rs[MaxConstrainedBodies];		// World-space joint connector rotations
 
 		// World-space body state
-		FVec3 Ps[MaxConstraintedBodies];				// World-space particle CoM positions
-		FRotation3 Qs[MaxConstraintedBodies];			// World-space particle CoM rotations
+		FVec3 Ps[MaxConstrainedBodies];				// World-space particle CoM positions
+		FRotation3 Qs[MaxConstrainedBodies];		// World-space particle CoM rotations
+		FVec3 Vs[MaxConstrainedBodies];				// World-space particle CoM velocities
+		FVec3 Ws[MaxConstrainedBodies];				// World-space particle CoM angular velocities
 
 		// Settings
 		FReal LinearStiffness;
 		FReal TwistStiffness;
 		FReal SwingStiffness;
 		FReal AngularDriveStiffness;
+		FReal LinearProjection;
+		FReal AngularProjection;
 		FReal SwingTwistAngleTolerance;
 		bool bEnableTwistLimits;
 		bool bEnableSwingLimits;
diff --git a/Engine/Source/Runtime/Experimental/Chaos/Public/Chaos/PBDJointConstraintTypes.h b/Engine/Source/Runtime/Experimental/Chaos/Public/Chaos/PBDJointConstraintTypes.h
--- a/Engine/Source/Runtime/Experimental/Chaos/Public/Chaos/PBDJointConstraintTypes.h
+++ b/Engine/Source/Runtime/Experimental/Chaos/Public/Chaos/PBDJointConstraintTypes.h
@@ -133,7 +133,6 @@ namespace Chaos
 		bool bEnableTwistLimits;
 		bool bEnableSwingLimits;
 		bool bEnableDrives;
-		EJointSolverPhase ProjectionPhase;
 		FReal LinearProjection;
 		FReal AngularProjection;
 		FReal Stiffness;
diff --git a/Engine/Source/Runtime/Experimental/Chaos/Public/Chaos/PBDJointConstraintUtilities.h b/Engine/Source/Runtime/Experimental/Chaos/Public/Chaos/PBDJointConstraintUtilities.h
--- a/Engine/Source/Runtime/Experimental/Chaos/Public/Chaos/PBDJointConstraintUtilities.h
+++ b/Engine/Source/Runtime/Experimental/Chaos/Public/Chaos/PBDJointConstraintUtilities.h
@@ -80,6 +80,13 @@ namespace Chaos
 			const FPBDJointSolverSettings& SolverSettings,
 			const FPBDJointSettings& JointSettings);
 
+		static FReal GetLinearProjection(
+			const FPBDJointSolverSettings& SolverSettings,
+			const FPBDJointSettings& JointSettings);
+
+		static FReal GetAngularProjection(
+			const FPBDJointSolverSettings& SolverSettings,
+			const FPBDJointSettings& JointSettings);
 
 		static FVec3 GetSphereLimitedPositionError(const FVec3& CX, const FReal Radius);
 		static FVec3 GetSphereLimitedVelocityError(const FVec3& CX, const FReal Radius, const FVec3& CV);
diff --git a/Engine/Source/Runtime/Experimental/Chaos/Public/Chaos/PBDJointConstraints.h b/Engine/Source/Runtime/Experimental/Chaos/Public/Chaos/PBDJointConstraints.h
--- a/Engine/Source/Runtime/Experimental/Chaos/Public/Chaos/PBDJointConstraints.h
+++ b/Engine/Source/Runtime/Experimental/Chaos/Public/Chaos/PBDJointConstraints.h
@@ -138,7 +138,6 @@ namespace Chaos
 
 		void Apply(const FReal Dt, const int32 It, const int32 NumIts);
 		bool ApplyPushOut(const FReal Dt, const int32 It, const int32 NumIts);
-		void ApplyProjection(const FReal Dt);
 
 		//
 		// Island Rule API
@@ -146,7 +145,6 @@ namespace Chaos
 
 		void Apply(const FReal Dt, const TArray<FConstraintContainerHandle*>& InConstraintHandles, const int32 It, const int32 NumIts);
 		bool ApplyPushOut(const FReal Dt, const TArray<FConstraintContainerHandle*>& InConstraintHandles, const int32 It, const int32 NumIts);
-		void ApplyProjection(const FReal Dt, const TArray<FConstraintContainerHandle*>& InConstraintHandles);
 
 
 	protected:
@@ -156,17 +154,15 @@ namespace Chaos
 	private:
 		friend class FPBDJointConstraintHandle;
 
+		void GetConstrainedParticleIndices(const int32 ConstraintIndex, int32& Index0, int32& Index1) const;
 		void CalculateConstraintSpace(int32 ConstraintIndex, FVec3& OutX0, FMatrix33& OutR0, FVec3& OutX1, FMatrix33& OutR1, FVec3& OutAngles) const;
 		void UpdateParticleState(TPBDRigidParticleHandle<FReal, 3>* Rigid, const FReal Dt, const FVec3& P, const FRotation3& Q, const bool bUpdateVelocity = true);
+		void UpdateParticleState(TPBDRigidParticleHandle<FReal, 3>* Rigid, const FReal Dt, const FVec3& P, const FRotation3& Q, const FVec3& V, const FVec3& W);
 		void SortConstraints();
 
-		void SolvePosition(const FReal Dt, const int32 ConstraintIndex, const int32 NumPairIts, const int32 It, const int32 NumIts);
-		void SolvePosition_GaussSiedel(const FReal Dt, const int32 ConstraintIndex, const int32 NumPairIts, const int32 It, const int32 NumIts);
 		void SolvePosition_Cholesky(const FReal Dt, const int32 ConstraintIndex, const int32 NumPairIts, const int32 It, const int32 NumIts);
-		void ProjectPosition(const FReal Dt, const int32 ConstraintIndex, const int32 It, const int32 NumIts);
-
-		// All-constraints API
-		void SolvePosition(const FReal Dt, const int32 NumPairIts, const int32 It, const int32 NumIts);
+		void SolvePosition_GaussSiedel(const FReal Dt, const int32 ConstraintIndex, const int32 NumPairIts, const int32 It, const int32 NumIts);
+		void ProjectPosition_GaussSiedel(const FReal Dt, const int32 ConstraintIndex, const int32 NumPairIts, const int32 It, const int32 NumIts);
 
 		FPBDJointSolverSettings Settings;
 
