Chaos: FixCrash on Editor Save - Make Trimeshes on BodySetup SharedPtr, and store SharedPtr on ImplicitObjectScaled to prevent GC'd BodySetup from invalidating Trimeshes.

diff --git a/Engine/Source/Runtime/Engine/Classes/PhysicsEngine/BodySetup.h b/Engine/Source/Runtime/Engine/Classes/PhysicsEngine/BodySetup.h
--- a/Engine/Source/Runtime/Engine/Classes/PhysicsEngine/BodySetup.h
+++ b/Engine/Source/Runtime/Engine/Classes/PhysicsEngine/BodySetup.h
@@ -266,7 +266,7 @@ public:
 
 #if WITH_CHAOS
 	//FBodySetupTriMeshes* TriMeshWrapper;
-	TArray<TUniquePtr<Chaos::TTriangleMeshImplicitObject<float>>> ChaosTriMeshes;
+	TArray<TSharedPtr<Chaos::TTriangleMeshImplicitObject<float>, ESPMode::ThreadSafe>> ChaosTriMeshes;
 #endif
 
 	/** Additional UV info, if available. Used for determining UV for a line trace impact. */
diff --git a/Engine/Source/Runtime/Engine/Private/PhysicsEngine/Experimental/ChaosInterfaceUtils.cpp b/Engine/Source/Runtime/Engine/Private/PhysicsEngine/Experimental/ChaosInterfaceUtils.cpp
--- a/Engine/Source/Runtime/Engine/Private/PhysicsEngine/Experimental/ChaosInterfaceUtils.cpp
+++ b/Engine/Source/Runtime/Engine/Private/PhysicsEngine/Experimental/ChaosInterfaceUtils.cpp
@@ -272,7 +272,7 @@ namespace ChaosInterface
 
 		for (const auto& ChaosTriMesh : InParams.ChaosTriMeshes)
 		{
-			auto Implicit = MakeUnique<Chaos::TImplicitObjectScaled<Chaos::TTriangleMeshImplicitObject<float>>>(MakeSerializable(ChaosTriMesh), Scale);
+			auto Implicit = MakeUnique<Chaos::TImplicitObjectScaled<Chaos::TTriangleMeshImplicitObject<float>>>(ChaosTriMesh, Scale);
 			auto NewShape = NewShapeHelper(MakeSerializable(Implicit), nullptr, true);
 			Shapes.Emplace(MoveTemp(NewShape));
 			Geoms.Add(MoveTemp(Implicit));
diff --git a/Engine/Source/Runtime/Engine/Public/Physics/Experimental/ChaosDerivedDataReader.h b/Engine/Source/Runtime/Engine/Public/Physics/Experimental/ChaosDerivedDataReader.h
--- a/Engine/Source/Runtime/Engine/Public/Physics/Experimental/ChaosDerivedDataReader.h
+++ b/Engine/Source/Runtime/Engine/Public/Physics/Experimental/ChaosDerivedDataReader.h
@@ -25,7 +25,7 @@ public:
 	explicit FChaosDerivedDataReader(FUntypedBulkData* InBulkData);
 
 	TArray<TUniquePtr<Chaos::FConvex>> ConvexImplicitObjects;
-	TArray<TUniquePtr<Chaos::TTriangleMeshImplicitObject<T>>> TrimeshImplicitObjects;
+	TArray<TSharedPtr<Chaos::TTriangleMeshImplicitObject<T>, ESPMode::ThreadSafe>> TrimeshImplicitObjects;
 	FBodySetupUVInfo UVInfo;
 	TArray<int32> FaceRemap;
 
diff --git a/Engine/Source/Runtime/Engine/Public/Physics/PhysicsInterfaceTypes.h b/Engine/Source/Runtime/Engine/Public/Physics/PhysicsInterfaceTypes.h
--- a/Engine/Source/Runtime/Engine/Public/Physics/PhysicsInterfaceTypes.h
+++ b/Engine/Source/Runtime/Engine/Public/Physics/PhysicsInterfaceTypes.h
@@ -56,7 +56,7 @@ struct FGeometryAddParams
 	TArrayView<physx::PxTriangleMesh*> TriMeshes;
 #endif
 #if WITH_CHAOS
-	TArrayView<TUniquePtr<Chaos::TTriangleMeshImplicitObject<float>>> ChaosTriMeshes;
+	TArrayView<TSharedPtr<Chaos::TTriangleMeshImplicitObject<float>, ESPMode::ThreadSafe>> ChaosTriMeshes;
 #endif
 };
 
diff --git a/Engine/Source/Runtime/Experimental/Chaos/Public/Chaos/ImplicitObjectScaled.h b/Engine/Source/Runtime/Experimental/Chaos/Public/Chaos/ImplicitObjectScaled.h
--- a/Engine/Source/Runtime/Experimental/Chaos/Public/Chaos/ImplicitObjectScaled.h
+++ b/Engine/Source/Runtime/Experimental/Chaos/Public/Chaos/ImplicitObjectScaled.h
@@ -139,6 +139,7 @@ public:
 	TImplicitObjectScaled(ObjectType Object, const TVector<T, d>& Scale, T Thickness = 0)
 	    : FImplicitObject(EImplicitObject::HasBoundingBox, Object->GetType() | ImplicitObjectType::IsScaled)
 	    , MObject(MoveTemp(Object))
+		, MSharedPtrForRefCount(nullptr)
 		, MInternalThickness(Thickness)
 	{
 		ensureMsgf((IsScaled(MObject->GetType(true)) == false), TEXT("Scaled objects should not contain each other."));
@@ -154,9 +155,31 @@ public:
 		this->bIsConvex = MObject->IsConvex();
 		SetScale(Scale);
 	}
+
+	TImplicitObjectScaled(TSharedPtr<TConcrete, ESPMode::ThreadSafe> Object, const TVector<T, d>& Scale, T Thickness = 0)
+	    : FImplicitObject(EImplicitObject::HasBoundingBox, Object->GetType() | ImplicitObjectType::IsScaled)
+	    , MObject(MakeSerializable<TConcrete, ESPMode::ThreadSafe>(Object))
+		, MSharedPtrForRefCount(Object)
+		, MInternalThickness(Thickness)
+	{
+		ensureMsgf((IsScaled(MObject->GetType(true)) == false), TEXT("Scaled objects should not contain each other."));
+		ensureMsgf((IsInstanced(MObject->GetType(true)) == false), TEXT("Scaled objects should not contain instances."));
+		switch (MObject->GetType(true))
+		{
+		case ImplicitObjectType::Transformed:
+		case ImplicitObjectType::Union:
+			check(false);	//scale is only supported for concrete types like sphere, capsule, convex, levelset, etc... Nothing that contains other objects
+		default:
+			break;
+		}
+		this->bIsConvex = MObject->IsConvex();
+		SetScale(Scale);
+	}
+
 	TImplicitObjectScaled(ObjectType Object, TUniquePtr<Chaos::FImplicitObject> &&ObjectOwner, const TVector<T, d>& Scale, T Thickness = 0)
 	    : FImplicitObject(EImplicitObject::HasBoundingBox, Object->GetType() | ImplicitObjectType::IsScaled)
 	    , MObject(Object)
+		, MSharedPtrForRefCount(nullptr)
 		, MInternalThickness(Thickness)
 	{
 		ensureMsgf((IsScaled(MObject->GetType(true)) == false), TEXT("Scaled objects should not contain each other."));
@@ -169,6 +192,7 @@ public:
 	TImplicitObjectScaled(TImplicitObjectScaled<TConcrete, bInstanced>&& Other)
 	    : FImplicitObject(EImplicitObject::HasBoundingBox, Other.MObject->GetType() | ImplicitObjectType::IsScaled)
 	    , MObject(MoveTemp(Other.MObject))
+		, MSharedPtrForRefCount(MoveTemp(Other.MSharedPtrForRefCount))
 	    , MScale(Other.MScale)
 		, MInvScale(Other.MInvScale)
 		, MInternalThickness(Other.MInternalThickness)
@@ -502,6 +526,7 @@ public:
 #endif
 private:
 	ObjectType MObject;
+	TSharedPtr<TConcrete, ESPMode::ThreadSafe> MSharedPtrForRefCount; // Temporary solution to force ref counting on trianglemesh from body setup.
 	TVector<T, d> MScale;
 	TVector<T, d> MInvScale;
 	T MInternalThickness;	//Allows us to inflate the instance before the scale is applied. This is useful when sweeps need to apply a non scale on a geometry with uniform thickness
