Chaos : Collisions

- Removed "less than" Union type constraint.

diff --git a/Engine/Source/Runtime/Experimental/Chaos/Private/Chaos/CollisionResolution.cpp b/Engine/Source/Runtime/Experimental/Chaos/Private/Chaos/CollisionResolution.cpp
--- a/Engine/Source/Runtime/Experimental/Chaos/Private/Chaos/CollisionResolution.cpp
+++ b/Engine/Source/Runtime/Experimental/Chaos/Private/Chaos/CollisionResolution.cpp
@@ -1066,13 +1066,13 @@ namespace Chaos
 					Constraint.SetNormal(-Constraint.GetNormal());
 				}
 			}
-			else if (Implicit0Type < TImplicitObjectUnion<T, d>::StaticType() && Implicit1Type == TImplicitObjectUnion<T, d>::StaticType())
+			else if (Implicit0Type != TImplicitObjectUnion<T, d>::StaticType() && Implicit1Type == TImplicitObjectUnion<T, d>::StaticType())
 			{
 				return UpdateSingleUnionConstraint<UpdateType>(Implicit0, Transform0, Implicit1, Transform1, Thickness, ConstraintBase);
 			}
-			else if (Implicit0Type == TImplicitObjectUnion<T, d>::StaticType() && Implicit1Type < TImplicitObjectUnion<T, d>::StaticType())
+			else if (Implicit0Type == TImplicitObjectUnion<T, d>::StaticType() && Implicit1Type != TImplicitObjectUnion<T, d>::StaticType())
 			{
-				check(false);	//should not be possible to get this ordering (see ComputeConstraint)
+				check(false);	//should not be possible to get this ordering (see [Note] in ConstructPairConstraintImpl)
 			}
 			else if (Implicit0Type == TImplicitObjectUnion<T, d>::StaticType() && Implicit1Type == TImplicitObjectUnion<T, d>::StaticType())
 			{
@@ -1200,12 +1200,13 @@ namespace Chaos
 			{
 				ConstructCapsuleBoxConstraints(Particle1, Particle0, Implicit1, Implicit0, Transform1, Transform0, Thickness, NewConstraints);
 			}
-			else if (Implicit0Type < TImplicitObjectUnion<T, d>::StaticType() && Implicit1Type == TImplicitObjectUnion<T, d>::StaticType())
+			else if (Implicit0Type != TImplicitObjectUnion<T, d>::StaticType() && Implicit1Type == TImplicitObjectUnion<T, d>::StaticType())
 			{
 				ConstructSingleUnionConstraints(Particle0, Particle1, Implicit0, Implicit1, Transform0, Transform1, Thickness, NewConstraints);
 			}
-			else if (Implicit0Type == TImplicitObjectUnion<T, d>::StaticType() && Implicit1Type < TImplicitObjectUnion<T, d>::StaticType())
+			else if (Implicit0Type == TImplicitObjectUnion<T, d>::StaticType() && Implicit1Type != TImplicitObjectUnion<T, d>::StaticType())
 			{
+				// [Note] forces non-unions into particle[0] position
 				ConstructSingleUnionConstraints(Particle1, Particle0, Implicit1, Implicit0, Transform1, Transform0, Thickness, NewConstraints);
 			}
 			else if (Implicit0Type == TImplicitObjectUnion<T, d>::StaticType() && Implicit1Type == TImplicitObjectUnion<T, d>::StaticType())
