Chaos : Collisions

- Removed unneeded check for pre-existing constraints.

diff --git a/Engine/Source/Runtime/Experimental/Chaos/Private/Chaos/CollisionResolution.cpp b/Engine/Source/Runtime/Experimental/Chaos/Private/Chaos/CollisionResolution.cpp
--- a/Engine/Source/Runtime/Experimental/Chaos/Private/Chaos/CollisionResolution.cpp
+++ b/Engine/Source/Runtime/Experimental/Chaos/Private/Chaos/CollisionResolution.cpp
@@ -1108,26 +1108,6 @@ namespace Chaos
 				return;
 			}
 
-			// See if we already have a constraint for this shape pair
-			// todo(brice) : Not sure this actually prevents duplicate constraints.
-			for (int32 i = 0; i < NewConstraints.Num(); i++)
-			{
-				if (NewConstraints[i]->GetType() == TRigidBodyPointContactConstraint<T, d>::StaticType())
-				{
-					if (NewConstraints[i]->As< TRigidBodyPointContactConstraint<T, d> >()->ContainsManifold(Implicit0, Implicit1))
-					{
-						return;
-					}
-				}
-				else if (NewConstraints[i]->GetType() == TRigidBodyIterativeContactConstraint<T, d>::StaticType())
-				{
-					if (NewConstraints[i]->As< TRigidBodyIterativeContactConstraint<T, d> >()->ContainsManifold(Implicit0, Implicit1))
-					{
-						return;
-					}
-				}
-			}
-
 			if (!Implicit0 || !Implicit1)
 			{
 				ConstructLevelsetLevelsetConstraints(Particle0, Particle1, Implicit0, Implicit1, Transform0, Transform1, Thickness, NewConstraints);
