Change "Damping" to "EtherDrag"

diff --git a/Engine/Source/Runtime/Engine/Private/PhysicsEngine/Experimental/PhysInterface_Chaos.cpp b/Engine/Source/Runtime/Engine/Private/PhysicsEngine/Experimental/PhysInterface_Chaos.cpp
--- a/Engine/Source/Runtime/Engine/Private/PhysicsEngine/Experimental/PhysInterface_Chaos.cpp
+++ b/Engine/Source/Runtime/Engine/Private/PhysicsEngine/Experimental/PhysInterface_Chaos.cpp
@@ -647,14 +647,14 @@ FBox FPhysInterface_Chaos::GetBounds_AssumesLocked(const FPhysicsActorHandle& In
     return FBox(FVector(-0.5), FVector(0.5));
 }
 
-void FPhysInterface_Chaos::SetLinearDamping_AssumesLocked(const FPhysicsActorHandle& InActorReference, float InDamping)
+void FPhysInterface_Chaos::SetLinearDamping_AssumesLocked(const FPhysicsActorHandle& InActorReference, float InDrag)
 {
 	if (ensure(FPhysicsInterface::IsValid(InActorReference)))
 	{
 		Chaos::TPBDRigidParticle<float, 3>* Rigid = InActorReference->CastToRigidParticle();
 		if (ensure(Rigid))
 		{
-			Rigid->SetLinearDamping(InDamping);
+			Rigid->SetLinearEtherDrag(InDrag);
 		}
 	}
 }
@@ -666,7 +666,7 @@ void FPhysInterface_Chaos::SetAngularDamping_AssumesLocked(const FPhysicsActorHa
 		Chaos::TPBDRigidParticle<float, 3>* Rigid = InActorReference->CastToRigidParticle();
 		if (ensure(Rigid))
 		{
-			Rigid->SetAngularDamping(InDamping);
+			Rigid->SetAngularEtherDrag(InDamping);
 		}
 	}
 }
diff --git a/Engine/Source/Runtime/Experimental/Chaos/Public/Chaos/PBDRigidsEvolutionGBF.h b/Engine/Source/Runtime/Experimental/Chaos/Public/Chaos/PBDRigidsEvolutionGBF.h
--- a/Engine/Source/Runtime/Experimental/Chaos/Public/Chaos/PBDRigidsEvolutionGBF.h
+++ b/Engine/Source/Runtime/Experimental/Chaos/Public/Chaos/PBDRigidsEvolutionGBF.h
@@ -10,7 +10,6 @@
 #include "Chaos/Collision/SpatialAccelerationBroadPhase.h"
 #include "Chaos/PerParticleInitForce.h"
 #include "Chaos/PerParticleEulerStepVelocity.h"
-#include "Chaos/PerParticleVelocityDamping.h"
 #include "Chaos/PerParticleEtherDrag.h"
 #include "Chaos/PerParticleGravity.h"
 #include "Chaos/PerParticlePBDEulerStep.h"
@@ -128,8 +127,7 @@ public:
 		CHAOS_SCOPED_TIMER(Integrate);
 		TPerParticleInitForce<T, d> InitForceRule;
 		TPerParticleEulerStepVelocity<T, d> EulerStepVelocityRule;
-		TPerParticleVelocityDamping<T, d> VelocityDampingRule;
-		TPerParticleEtherDrag<T, d> EtherDragRule(HackLinearDrag, HackAngularDrag);
+		TPerParticleEtherDrag<T, d> EtherDragRule;
 		TPerParticlePBDEulerStep<T, d> EulerStepRule;
 
 		const T MaxAngularSpeedSq = HackMaxAngularVelocity * HackMaxAngularVelocity;
@@ -152,7 +150,6 @@ public:
 					ForceRule(Particle, Dt);
 				}
 				EulerStepVelocityRule.Apply(Particle, Dt);
-				VelocityDampingRule.Apply(Particle, Dt);
 				EtherDragRule.Apply(Particle, Dt);
 
 				if (HackMaxAngularVelocity >= 0.f)
diff --git a/Engine/Source/Runtime/Experimental/Chaos/Public/Chaos/ParticleDirtyFlags.h b/Engine/Source/Runtime/Experimental/Chaos/Public/Chaos/ParticleDirtyFlags.h
--- a/Engine/Source/Runtime/Experimental/Chaos/Public/Chaos/ParticleDirtyFlags.h
+++ b/Engine/Source/Runtime/Experimental/Chaos/Public/Chaos/ParticleDirtyFlags.h
@@ -27,8 +27,8 @@ namespace Chaos
 		InvI					= 1 << 15,
 		M						= 1 << 16,
 		InvM					= 1 << 17,
-		LinearDamping			= 1 << 18,
-		AngularDamping			= 1 << 19,
+		LinearEtherDrag			= 1 << 18,
+		AngularEtherDrag			= 1 << 19,
 		ObjectState				= 1 << 20,
 		Geometry				= 1 << 21,
 		ExternalForce			= 1 << 22,
diff --git a/Engine/Source/Runtime/Experimental/Chaos/Public/Chaos/ParticleHandle.h b/Engine/Source/Runtime/Experimental/Chaos/Public/Chaos/ParticleHandle.h
--- a/Engine/Source/Runtime/Experimental/Chaos/Public/Chaos/ParticleHandle.h
+++ b/Engine/Source/Runtime/Experimental/Chaos/Public/Chaos/ParticleHandle.h
@@ -80,8 +80,8 @@ void PBDRigidParticleHandleImpDefaultConstruct(TPBDRigidParticleHandleImp<T, d,
 	Concrete.SetInvM(1);
 	Concrete.SetI(PMatrix<T, d, d>(1, 1, 1));
 	Concrete.SetInvI(PMatrix<T, d, d>(1, 1, 1));
-	Concrete.SetLinearDamping(0.f);
-	Concrete.SetAngularDamping(0.f);
+	Concrete.SetLinearEtherDrag(0.f);
+	Concrete.SetAngularEtherDrag(0.f);
 	Concrete.SetObjectStateLowLevel(Params.bStartSleeping ? EObjectStateType::Sleeping : EObjectStateType::Dynamic);
 }
 
@@ -103,8 +103,8 @@ void PBDRigidParticleDefaultConstruct(TPBDRigidParticle<T,d>& Concrete, const TP
 	Concrete.SetInvM(1);
 	Concrete.SetI(PMatrix<T, d, d>(1, 1, 1));
 	Concrete.SetInvI(PMatrix<T, d, d>(1, 1, 1));
-	Concrete.SetLinearDamping(0.f);
-	Concrete.SetAngularDamping(0.f);
+	Concrete.SetLinearEtherDrag(0.f);
+	Concrete.SetAngularEtherDrag(0.f);
 	Concrete.SetObjectState(Params.bStartSleeping ? EObjectStateType::Sleeping : EObjectStateType::Dynamic);
 	Concrete.SetGravityEnabled(Params.bGravityEnabled);
 }
@@ -675,13 +675,13 @@ public:
 	T& InvM() { return PBDRigidParticles->InvM(ParticleIdx); }
 	void SetInvM(const T& InInvM) { PBDRigidParticles->InvM(ParticleIdx) = InInvM; }
 
-	T LinearDamping() const { return PBDRigidParticles->LinearDamping(ParticleIdx); }
-	T& LinearDamping() { return PBDRigidParticles->LinearDamping(ParticleIdx); }
-	void SetLinearDamping(const T& InLinearDamping) { PBDRigidParticles->LinearDamping(ParticleIdx) = InLinearDamping; }
+	T LinearEtherDrag() const { return PBDRigidParticles->LinearEtherDrag(ParticleIdx); }
+	T& LinearEtherDrag() { return PBDRigidParticles->LinearEtherDrag(ParticleIdx); }
+	void SetLinearEtherDrag(const T& InLinearEtherDrag) { PBDRigidParticles->LinearEtherDrag(ParticleIdx) = InLinearEtherDrag; }
 
-	T AngularDamping() const { return PBDRigidParticles->AngularDamping(ParticleIdx); }
-	T& AngularDamping() { return PBDRigidParticles->AngularDamping(ParticleIdx); }
-	void SetAngularDamping(const T& InAngularDamping) { PBDRigidParticles->AngularDamping(ParticleIdx) = InAngularDamping; }
+	T AngularEtherDrag() const { return PBDRigidParticles->AngularEtherDrag(ParticleIdx); }
+	T& AngularEtherDrag() { return PBDRigidParticles->AngularEtherDrag(ParticleIdx); }
+	void SetAngularEtherDrag(const T& InAngularEtherDrag) { PBDRigidParticles->AngularEtherDrag(ParticleIdx) = InAngularEtherDrag; }
 
 	int32 Island() const { return PBDRigidParticles->Island(ParticleIdx); }
 	int32& Island() { return PBDRigidParticles->Island(ParticleIdx); }
@@ -1752,8 +1752,8 @@ public:
 		Ar.UsingCustomVersion(FExternalPhysicsCustomObjectVersion::GUID);
 		if (Ar.CustomVer(FExternalPhysicsCustomObjectVersion::GUID) >= FExternalPhysicsCustomObjectVersion::AddDampingToRigids)
 		{
-			Ar << MLinearDamping;
-			Ar << MAngularDamping;
+			Ar << MLinearEtherDrag;
+			Ar << MAngularEtherDrag;
 		}
 
 		Ar << MIsland;
@@ -1875,18 +1875,18 @@ public:
 		this->MInvM = InInvM;
 	}
 
-	T LinearDamping() const { return MLinearDamping; }
-	void SetLinearDamping(const T& InLinearDamping)
+	T LinearEtherDrag() const { return MLinearEtherDrag; }
+	void SetLinearEtherDrag(const T& InLinearEtherDrag)
 	{
-		this->MarkDirty(EParticleFlags::LinearDamping);
-		this->MLinearDamping = InLinearDamping;
+		this->MarkDirty(EParticleFlags::LinearEtherDrag);
+		this->MLinearEtherDrag = InLinearEtherDrag;
 	}
 
-	T AngularDamping() const { return MAngularDamping; }
-	void SetAngularDamping(const T& InAngularDamping)
+	T AngularEtherDrag() const { return MAngularEtherDrag; }
+	void SetAngularEtherDrag(const T& InAngularEtherDrag)
 	{
-		this->MarkDirty(EParticleFlags::AngularDamping);
-		this->MAngularDamping = InAngularDamping;
+		this->MarkDirty(EParticleFlags::AngularEtherDrag);
+		this->MAngularEtherDrag = InAngularEtherDrag;
 	}
 
 	int32 Island() const { return MIsland; }
@@ -1931,8 +1931,8 @@ private:
 	TUniquePtr<TBVHParticles<T, d>> MCollisionParticles;
 	T MM;
 	T MInvM;
-	T MLinearDamping;
-	T MAngularDamping;
+	T MLinearEtherDrag;
+	T MAngularEtherDrag;
 	int32 MIsland;
 	int32 MCollisionGroup;
 	EObjectStateType MObjectState;
@@ -1983,8 +1983,8 @@ public:
 		, MCollisionParticles(nullptr)
 		, MM(T(0))
 		, MInvM(T(0))
-		, MLinearDamping(T(0))
-		, MAngularDamping(T(0))
+		, MLinearEtherDrag(T(0))
+		, MAngularEtherDrag(T(0))
 		, MIsland(INDEX_NONE)
 		, MCollisionGroup(0)
 		, MObjectState(EObjectStateType::Uninitialized)
@@ -2006,8 +2006,8 @@ public:
 		, MCollisionParticles(nullptr)
 		, MM(InParticle.M())
 		, MInvM(InParticle.InvM())
-		, MLinearDamping(InParticle.LinearDamping())
-		, MAngularDamping(InParticle.AngularDamping())
+		, MLinearEtherDrag(InParticle.LinearEtherDrag())
+		, MAngularEtherDrag(InParticle.AngularEtherDrag())
 		, MIsland(InParticle.Island())
 		, MCollisionGroup(InParticle.CollisionGroup())
 		, MObjectState(InParticle.ObjectState())
@@ -2037,8 +2037,8 @@ public:
 	const TBVHParticles<T, d> * MCollisionParticles;
 	T MM;
 	T MInvM;
-	T MLinearDamping;
-	T MAngularDamping;
+	T MLinearEtherDrag;
+	T MAngularEtherDrag;
 	int32 MIsland;
 	int32 MCollisionGroup;
 	EObjectStateType MObjectState;
@@ -2061,8 +2061,8 @@ public:
 		MCollisionParticles = nullptr;
 		MM = T(0);
 		MInvM = T(0);
-		MLinearDamping = T(0);
-		MAngularDamping = T(0);
+		MLinearEtherDrag = T(0);
+		MAngularEtherDrag = T(0);
 		MIsland = INDEX_NONE;
 		MCollisionGroup = 0;
 		MObjectState = EObjectStateType::Uninitialized;
diff --git a/Engine/Source/Runtime/Experimental/Chaos/Public/Chaos/PerParticleEtherDrag.h b/Engine/Source/Runtime/Experimental/Chaos/Public/Chaos/PerParticleEtherDrag.h
--- a/Engine/Source/Runtime/Experimental/Chaos/Public/Chaos/PerParticleEtherDrag.h
+++ b/Engine/Source/Runtime/Experimental/Chaos/Public/Chaos/PerParticleEtherDrag.h
@@ -1,53 +1,40 @@
 // Copyright 1998-2019 Epic Games, Inc. All Rights Reserved.
 #pragma once
 
-#include "Chaos/Matrix.h"
 #include "Chaos/PerParticleRule.h"
+#include "Chaos/Utilities.h"
 
 namespace Chaos
 {
-template<class T, int d>
-class TPerParticleEtherDrag : public TPerParticleRule<T, d>
-{
-  public:
-	TPerParticleEtherDrag(const T Coefficient = 0.01, const T AngularCoefficient = 0.01)
-	    : MCoefficient(Coefficient), MAngularCoefficient(AngularCoefficient)
-	{
-	}
-	virtual ~TPerParticleEtherDrag() {}
-
-	template<class T_PARTICLES>
-	inline void ApplyHelper(T_PARTICLES& InParticles, const T Dt, const int32 Index) const
+	template<class T, int d>
+	class TPerParticleEtherDrag : public TPerParticleRule<T, d>
 	{
-		InParticles.V(Index) -= MCoefficient * InParticles.V(Index);
-	}
+	public:
+		TPerParticleEtherDrag() {}
+		virtual ~TPerParticleEtherDrag() {}
 
-	inline void Apply(TDynamicParticles<T, d>& InParticles, const T Dt, const int32 Index) const override //-V762
-	{
-		if (InParticles.InvM(Index) == 0)
+		inline void ApplyHelper(FVec3& V, FVec3& W, FReal LinearDamp, FReal AngularDamp, T Dt) const
 		{
-			return; // Do not damp kinematic particles
+			const FReal LinearMultiplier = FMath::Max(FReal(0), FReal(1) - (LinearDamp * Dt));
+			V *= LinearMultiplier;
+
+			const FReal AngularMultiplier = FMath::Max(FReal(0), FReal(1) - (AngularDamp * Dt));
+			W *= AngularMultiplier;
 		}
-		ApplyHelper(InParticles, Dt, Index);
-	}
 
-	inline void Apply(TRigidParticles<T, d>& InParticles, const T Dt, const int32 Index) const override //-V762
-	{
-		if (InParticles.InvM(Index) == 0)
+		inline void Apply(TDynamicParticles<T, d>& InParticles, const T Dt, const int32 Index) const override //-V762
 		{
-			return; // Do not damp kinematic rigid bodies
+			ensure(false);
 		}
-		ApplyHelper(InParticles, Dt, Index);
-		InParticles.W(Index) -= MAngularCoefficient * InParticles.W(Index);
-	}
 
-	inline void Apply(TTransientPBDRigidParticleHandle<T, d>& Particle, const T Dt) const override
-	{
-		Particle.V() -= MCoefficient * Particle.V();
-	}
+		inline void Apply(TRigidParticles<T, d>& InParticles, const T Dt, const int32 Index) const override //-V762
+		{
+			ApplyHelper(InParticles.V(Index), InParticles.W(Index), InParticles.LinearEtherDrag(Index), InParticles.AngularEtherDrag(Index), Dt);
+		}
 
-  private:
-	T MCoefficient;
-	T MAngularCoefficient;
-};
+		inline void Apply(TTransientPBDRigidParticleHandle<T, d>& Particle, const T Dt) const override //-V762
+		{
+			ApplyHelper(Particle.V(), Particle.W(), Particle.LinearEtherDrag(), Particle.AngularEtherDrag(), Dt);
+		}
+	};
 }
diff --git a/Engine/Source/Runtime/Experimental/Chaos/Public/Chaos/PerParticleVelocityDamping.h b/Engine/Source/Runtime/Experimental/Chaos/Public/Chaos/PerParticleVelocityDamping.h
deleted file mode 100644
diff --git a/Engine/Source/Runtime/Experimental/Chaos/Public/Chaos/RigidParticles.h b/Engine/Source/Runtime/Experimental/Chaos/Public/Chaos/RigidParticles.h
--- a/Engine/Source/Runtime/Experimental/Chaos/Public/Chaos/RigidParticles.h
+++ b/Engine/Source/Runtime/Experimental/Chaos/Public/Chaos/RigidParticles.h
@@ -78,8 +78,8 @@ class TRigidParticles : public TKinematicGeometryParticles<T, d>
 		TArrayCollection::AddArray(&MInvI);
 		TArrayCollection::AddArray(&MM);
 		TArrayCollection::AddArray(&MInvM);
-		TArrayCollection::AddArray(&MLinearDamping);
-		TArrayCollection::AddArray(&MAngularDamping);
+		TArrayCollection::AddArray(&MLinearEtherDrag);
+		TArrayCollection::AddArray(&MAngularEtherDrag);
 		TArrayCollection::AddArray(&MCollisionParticles);
 		TArrayCollection::AddArray(&MCollisionGroup);
 		TArrayCollection::AddArray(&MDisabled);
@@ -99,8 +99,8 @@ class TRigidParticles : public TKinematicGeometryParticles<T, d>
 		TArrayCollection::AddArray(&MInvI);
 		TArrayCollection::AddArray(&MM);
 		TArrayCollection::AddArray(&MInvM);
-		TArrayCollection::AddArray(&MLinearDamping);
-		TArrayCollection::AddArray(&MAngularDamping);
+		TArrayCollection::AddArray(&MLinearEtherDrag);
+		TArrayCollection::AddArray(&MAngularEtherDrag);
 		TArrayCollection::AddArray(&MCollisionParticles);
 		TArrayCollection::AddArray(&MCollisionGroup);
 		TArrayCollection::AddArray(&MDisabled);
@@ -136,11 +136,11 @@ class TRigidParticles : public TKinematicGeometryParticles<T, d>
 	FORCEINLINE const T InvM(const int32 Index) const { return MInvM[Index]; }
 	FORCEINLINE T& InvM(const int32 Index) { return MInvM[Index]; }
 
-	CHAOS_API const T& LinearDamping(const int32 index) const { return MLinearDamping[index]; }
-	CHAOS_API T& LinearDamping(const int32 index) { return MLinearDamping[index]; }
+	FORCEINLINE const T& LinearEtherDrag(const int32 index) const { return MLinearEtherDrag[index]; }
+	FORCEINLINE T& LinearEtherDrag(const int32 index) { return MLinearEtherDrag[index]; }
 
-	CHAOS_API const T& AngularDamping(const int32 index) const { return MAngularDamping[index]; }
-	CHAOS_API T& AngularDamping(const int32 index) { return MAngularDamping[index]; }
+	FORCEINLINE const T& AngularEtherDrag(const int32 index) const { return MAngularEtherDrag[index]; }
+	FORCEINLINE T& AngularEtherDrag(const int32 index) { return MAngularEtherDrag[index]; }
 
 	FORCEINLINE int32 CollisionParticlesSize(int32 Index) const { return MCollisionParticles[Index] == nullptr ? 0 : MCollisionParticles[Index]->Size(); }
 
@@ -200,7 +200,7 @@ class TRigidParticles : public TKinematicGeometryParticles<T, d>
 		Ar.UsingCustomVersion(FExternalPhysicsCustomObjectVersion::GUID);
 		if (Ar.CustomVer(FExternalPhysicsCustomObjectVersion::GUID) >= FExternalPhysicsCustomObjectVersion::AddDampingToRigids)
 		{
-			Ar << MLinearDamping << MAngularDamping;
+			Ar << MLinearEtherDrag << MAngularEtherDrag;
 		}
 
 		Ar << MCollisionParticles << MCollisionGroup << MIsland << MDisabled << MObjectState;
@@ -215,8 +215,8 @@ class TRigidParticles : public TKinematicGeometryParticles<T, d>
 	TArrayCollectionArray<PMatrix<T, d, d>> MInvI;
 	TArrayCollectionArray<T> MM;
 	TArrayCollectionArray<T> MInvM;
-	TArrayCollectionArray<T> MLinearDamping;
-	TArrayCollectionArray<T> MAngularDamping;
+	TArrayCollectionArray<T> MLinearEtherDrag;
+	TArrayCollectionArray<T> MAngularEtherDrag;
 	TArrayCollectionArray<TUniquePtr<TBVHParticles<T, d>>> MCollisionParticles;
 	TArrayCollectionArray<int32> MCollisionGroup;
 	TArrayCollectionArray<int32> MIsland;
diff --git a/Engine/Source/Runtime/Experimental/ChaosSolvers/Private/PhysicsProxy/SingleParticlePhysicsProxy.cpp b/Engine/Source/Runtime/Experimental/ChaosSolvers/Private/PhysicsProxy/SingleParticlePhysicsProxy.cpp
--- a/Engine/Source/Runtime/Experimental/ChaosSolvers/Private/PhysicsProxy/SingleParticlePhysicsProxy.cpp
+++ b/Engine/Source/Runtime/Experimental/ChaosSolvers/Private/PhysicsProxy/SingleParticlePhysicsProxy.cpp
@@ -227,8 +227,8 @@ void FSingleParticlePhysicsProxy<Chaos::TPBDRigidParticle<float, 3>>::PushToPhys
 		RigidHandle->SetInvM(Data->MInvM);
 		RigidHandle->SetI(Data->MI);
 		RigidHandle->SetInvI(Data->MInvI);
-		RigidHandle->SetLinearDamping(Data->MLinearDamping);
-		RigidHandle->SetAngularDamping(Data->MAngularDamping);
+		RigidHandle->SetLinearEtherDrag(Data->MLinearEtherDrag);
+		RigidHandle->SetAngularEtherDrag(Data->MAngularEtherDrag);
 
 		if (Data->DirtyFlags.IsDirty(Chaos::EParticleFlags::ExternalForce))
 		{
