Fixed Clang compile errors in Chaos ParticleDirtyFlags.h

diff --git a/Engine/Source/Runtime/Experimental/Chaos/Public/Chaos/ParticleDirtyFlags.h b/Engine/Source/Runtime/Experimental/Chaos/Public/Chaos/ParticleDirtyFlags.h
--- a/Engine/Source/Runtime/Experimental/Chaos/Public/Chaos/ParticleDirtyFlags.h
+++ b/Engine/Source/Runtime/Experimental/Chaos/Public/Chaos/ParticleDirtyFlags.h
@@ -7,7 +7,7 @@ namespace Chaos
 	// There is a dirty flag for every user-settable particle property.
 	// Dirty property values will get copied from game to physics thread buffers,
 	// but clean property values will get overridden with physics thread results.
-	enum EParticleFlags : int32
+	enum class EParticleFlags : int32
 	{
 		X						= 1 << 0,
 		R						= 1 << 1,
diff --git a/Engine/Source/Runtime/Experimental/ChaosSolvers/Private/PhysicsProxy/SingleParticlePhysicsProxy.cpp b/Engine/Source/Runtime/Experimental/ChaosSolvers/Private/PhysicsProxy/SingleParticlePhysicsProxy.cpp
--- a/Engine/Source/Runtime/Experimental/ChaosSolvers/Private/PhysicsProxy/SingleParticlePhysicsProxy.cpp
+++ b/Engine/Source/Runtime/Experimental/ChaosSolvers/Private/PhysicsProxy/SingleParticlePhysicsProxy.cpp
@@ -247,7 +247,7 @@ void FSingleParticlePhysicsProxy<Chaos::TPBDRigidParticle<float, 3>>::PushToPhys
 			GetSolver()->GetEvolution()->GetGravityForces().SetEnabled(*RigidHandle, Data->MGravityEnabled);
 		}
 
-		if (Data->DirtyFlags.IsDirty(Chaos::EParticleFlags::X | Chaos::EParticleFlags::R | Chaos::EParticleFlags::V | Chaos::EParticleFlags::Geometry))
+		if (Data->DirtyFlags.IsDirty((int32)Chaos::EParticleFlags::X | (int32)Chaos::EParticleFlags::R | (int32)Chaos::EParticleFlags::V | (int32)Chaos::EParticleFlags::Geometry))
 		{
 			if (Data->Geometry && Data->Geometry->HasBoundingBox())
 			{
