Fix nan in gjk when close enough is actually the case regardless of inflation

diff --git a/Engine/Source/Runtime/Experimental/Chaos/Public/Chaos/GJK.h b/Engine/Source/Runtime/Experimental/Chaos/Public/Chaos/GJK.h
--- a/Engine/Source/Runtime/Experimental/Chaos/Public/Chaos/GJK.h
+++ b/Engine/Source/Runtime/Experimental/Chaos/Public/Chaos/GJK.h
@@ -517,7 +517,7 @@ namespace Chaos
 				InGJKPreDist2 = NewDist2;
 
 
-				if (bComputeMTD && bCloseEnough && Lambda == 0 && Inflation2 > 1e-6 && SimplexIDs.NumVerts < 4)
+				if (bComputeMTD && bCloseEnough && Lambda == 0 && InGJKPreDist2 > 1e-6 && Inflation2 > 1e-6 && SimplexIDs.NumVerts < 4)
 				{
 					//For mtd of inflated shapes we have to find the closest point, so we have to keep going
 					bCloseEnough = false;
