Chaos : Collisions

-Single shot manifold generation

diff --git a/Engine/Source/Runtime/Experimental/Chaos/Private/Chaos/CollisionResolution.cpp b/Engine/Source/Runtime/Experimental/Chaos/Private/Chaos/CollisionResolution.cpp
--- a/Engine/Source/Runtime/Experimental/Chaos/Private/Chaos/CollisionResolution.cpp
+++ b/Engine/Source/Runtime/Experimental/Chaos/Private/Chaos/CollisionResolution.cpp
@@ -575,12 +575,15 @@ namespace Chaos
 
 				TPlane<T, d> CollisionPlane(Constraint.PlanePosition, Constraint.PlaneNormal);
 
+				// re-sample the constraint based on the distance from the collision plane.
 				for (int32 Idx = 0; Idx < Constraint.NumSamples(); Idx++)
 				{
 					Constraint[Idx].Manifold.Phi = CollisionPlane.PhiWithNormal(AToBTM.TransformPosition(Constraint[Idx].X), Constraint[Idx].Manifold.Normal);
 					Constraint[Idx].Manifold.Normal = Transform1.TransformVector(Constraint.PlaneNormal);
 					Constraint[Idx].Manifold.Location = Transform0.TransformPosition(Constraint[Idx].X);
 
+					// save the best point for collision processing
+					// @todo(chaos) : Consider appling all constraint in ::Apply at each iteration, right now it just takes the deepest.
 					if (ContactPoint.Phi > Constraint[Idx].Manifold.Phi)
 					{
 						ContactPoint.Phi = Constraint[Idx].Manifold.Phi;
@@ -596,49 +599,81 @@ namespace Chaos
 		template<class T, int d>
 		void UpdateConvexConvexManifold(TRigidBodyIterativeContactConstraint<T, d>& Constraint, const TRigidTransform<T, d>& Transform0, const TRigidTransform<T, d>& Transform1, const T Thickness)
 		{
-			TContactPoint<T> ContactPoint = ConvexConvexContactPoint(*Constraint.Manifold.Implicit[0], Transform0, *Constraint.Manifold.Implicit[1], Transform1, Thickness);
-
-			FRigidBodyIterativeContactConstraint::FManifold& Manifold = Constraint.Manifold;
-
-			if (!ContactPoint.Normal.Equals(Constraint.PlaneNormal) || !Constraint.NumSamples())
+			if (ensure(Constraint.GetType() == FRigidBodyIterativeContactConstraint::StaticType()))
 			{
-				Constraint.ResetSamples();
-				Constraint.PlaneNormal = Transform1.InverseTransformVector(ContactPoint.Normal);
-				Constraint.PlanePosition = Transform1.InverseTransformPosition(ContactPoint.Location - ContactPoint.Phi*ContactPoint.Normal);
-			}
+				TContactPoint<T> ContactPoint = ConvexConvexContactPoint(*Constraint.Manifold.Implicit[0], Transform0, *Constraint.Manifold.Implicit[1], Transform1, Thickness);
 
-			TVector<T, d> SurfaceSample = Transform0.InverseTransformPosition(ContactPoint.Location);
-			if (Constraint.NumSamples() < 4)
-			{
-				Constraint.AddSample({ SurfaceSample,0.f });
-			}
-			else if(Constraint.NumSamples() == 4)
-			{
-				TVector<T,d> Center = SumSampleData(Constraint) / 4.f;
-				T Delta = (Center - SurfaceSample).SizeSquared();
+				if (GetInnerType(Constraint.Manifold.Implicit[0]->GetType()) == ImplicitObjectType::Convex)
+				{
+					// single shot manifolds for TConvex implicit object in the constraints implicit[0] position. 
 
-				// todo(chaos) : maximize area instead
-				T SmallestDelta = FLT_MAX;
-				int32 SmallestIndex = 0;
-				for (int32 idx=0;idx<Constraint.NumSamples();idx++)
-					if (Constraint[idx].Delta < SmallestDelta) {
-						SmallestDelta = Constraint[idx].Delta;
-						SmallestIndex = idx;
+					TArray<FVec3> CollisionSamples;
+					int32 FaceIndex = Constraint.Manifold.Implicit[0]->FindClosestFaceAndVertices(Transform0.InverseTransformPosition(ContactPoint.Location), CollisionSamples, 1.f);
+
+					if (!ContactPoint.Normal.Equals(Constraint.PlaneNormal)
+						|| FaceIndex != Constraint.SourceNormalIndex
+						|| !Constraint.NumSamples())
+					{
+						Constraint.ResetSamples();
+						Constraint.PlaneNormal = Transform1.InverseTransformVector(ContactPoint.Normal);
+						Constraint.PlanePosition = Transform1.InverseTransformPosition(ContactPoint.Location - ContactPoint.Phi*ContactPoint.Normal);
 					}
 
-				if (Delta > SmallestDelta) {
-					Constraint[SmallestIndex] = { SurfaceSample,Delta };
+
+					if (Constraint.NumSamples() == 0)
+					{
+						// @todo(chaos) : Only save the four best samples and hard-code the size of Constraint.Samples to [len:4]
+						Constraint.SourceNormalIndex = FaceIndex;
+						for (FVec3 Sample : CollisionSamples)
+						{
+							Constraint.AddSample({ Sample,0.f });
+						}
+					}
 				}
-			}
-			else
-			{
-				ensure(false); // max of 4 points
-			}
+				else
+				{
+					// iterative manifolds for non TConvex implicit objects that require sampling 
 
-			typedef FRigidBodyIterativeContactConstraint::FSampleData FSampleData;
-			TVector<T, d> Center = SumSampleData(Constraint) / Constraint.NumSamples();
-			for (int32 Index = 0; Index < Constraint.NumSamples(); Index++) { Constraint[Index].Delta = (Center - Constraint[Index].X).SizeSquared(); }
+					if (!ContactPoint.Normal.Equals(Constraint.PlaneNormal) || !Constraint.NumSamples())
+					{
+						Constraint.ResetSamples();
+						Constraint.PlaneNormal = Transform1.InverseTransformVector(ContactPoint.Normal);
+						Constraint.PlanePosition = Transform1.InverseTransformPosition(ContactPoint.Location - ContactPoint.Phi*ContactPoint.Normal);
+					}
 
+					TVector<T, d> SurfaceSample = Transform0.InverseTransformPosition(ContactPoint.Location);
+					if (Constraint.NumSamples() < 4)
+					{
+						Constraint.AddSample({ SurfaceSample,0.f });
+					}
+					else if (Constraint.NumSamples() == 4)
+					{
+						TVector<T, d> Center = SumSampleData(Constraint) / Constraint.NumSamples();
+						T Delta = (Center - SurfaceSample).SizeSquared();
+
+						// todo(chaos) : maximize for area instead
+						T SmallestDelta = FLT_MAX;
+						int32 SmallestIndex = 0;
+						for (int32 idx = 0; idx < Constraint.NumSamples(); idx++)
+							if (Constraint[idx].Delta < SmallestDelta) {
+								SmallestDelta = Constraint[idx].Delta;
+								SmallestIndex = idx;
+							}
+
+						if (Delta > SmallestDelta) {
+							Constraint[SmallestIndex] = { SurfaceSample,Delta };
+						}
+					}
+					else
+					{
+						ensure(false); // max of 4 points
+					}
+
+					typedef FRigidBodyIterativeContactConstraint::FSampleData FSampleData;
+					TVector<T, d> Center = SumSampleData(Constraint) / Constraint.NumSamples();
+					for (int32 Index = 0; Index < Constraint.NumSamples(); Index++) { Constraint[Index].Delta = (Center - Constraint[Index].X).SizeSquared(); }
+				}
+			}
 		}
 
 
diff --git a/Engine/Source/Runtime/Experimental/Chaos/Private/Chaos/Convex.cpp b/Engine/Source/Runtime/Experimental/Chaos/Private/Chaos/Convex.cpp
--- a/Engine/Source/Runtime/Experimental/Chaos/Private/Chaos/Convex.cpp
+++ b/Engine/Source/Runtime/Experimental/Chaos/Private/Chaos/Convex.cpp
@@ -29,4 +29,35 @@ namespace Chaos
 		CHAOS_ENSURE(MostOpposingIdx != INDEX_NONE);
 		return MostOpposingIdx;
 	}
+
+	int32 FConvex::FindClosestFaceAndVertices(const FVec3& Position, TArray<FVec3>& FaceVertices, FReal SearchDist) const
+	{
+		//
+		//  todo(chaos) : Create a correspondence between the faces and surface particles on construction.
+		//
+
+		int32 ReturnIndex = INDEX_NONE;
+		TSet<int32> IncludedParticles;
+		for (int32 Idx = 0; Idx < Planes.Num(); ++Idx)
+		{
+			const TPlane<FReal, 3>& Plane = Planes[Idx];
+			FReal AbsOfSignedDistance = FMath::Abs(Plane.SignedDistance(Position));
+			if (AbsOfSignedDistance < SearchDist)
+			{
+				for (int32 Fdx = 0; Fdx < (int32)SurfaceParticles.Size(); Fdx++)
+				{
+					if (!IncludedParticles.Contains(Fdx))
+					{
+						if (FMath::Abs(Plane.SignedDistance(SurfaceParticles.X(Fdx))) < SearchDist)
+						{
+							FaceVertices.Add(SurfaceParticles.X(Fdx));
+							IncludedParticles.Add(Fdx);
+						}
+					}
+				}
+				ReturnIndex = Idx;
+			}
+		}
+		return ReturnIndex;
+	}
 }
\ No newline at end of file
diff --git a/Engine/Source/Runtime/Experimental/Chaos/Public/Chaos/CollisionResolutionTypes.h b/Engine/Source/Runtime/Experimental/Chaos/Public/Chaos/CollisionResolutionTypes.h
--- a/Engine/Source/Runtime/Experimental/Chaos/Public/Chaos/CollisionResolutionTypes.h
+++ b/Engine/Source/Runtime/Experimental/Chaos/Public/Chaos/CollisionResolutionTypes.h
@@ -172,22 +172,22 @@ namespace Chaos
 		using Base::Particle;
 		struct FSampleData { TVector<T,d> X; float Delta; FManifold Manifold; };
 
-		TRigidBodyIterativeContactConstraint() : Base(Base::FType::MultiPoint), PlaneNormal(0), PlanePosition(0), NumberSamples(0) {}
+		TRigidBodyIterativeContactConstraint() : Base(Base::FType::MultiPoint), SourceNormalIndex(INDEX_NONE), PlaneNormal(0), PlanePosition(0) {}
 		static typename Base::FType StaticType() { return Base::FType::MultiPoint; };
 
+		int SourceNormalIndex; // index of normal on particle[0] body;
 		TVector<T, d> PlaneNormal; // local space contact normal on Particle1
 		TVector<T, d> PlanePosition; // local space surface position on Particle1
 
 		// Samples
-		int32               NumSamples()                   { return NumberSamples; }
-		void                AddSample(FSampleData && Data) {Samples[NumberSamples] = Data; NumberSamples++; };
-		void                ResetSamples()                 { NumberSamples=0; }
-		FSampleData &       operator[](int32 Index)        { return Samples[Index]; }
-		const FSampleData & operator[](int32 Index) const  { return Samples[Index]; }
+		int32               NumSamples()                   { return Samples.Num(); }
+		void                AddSample(FSampleData && Data) { Samples.Add(Data); };
+		void                ResetSamples()                 { Samples.Reset(); }
+		FSampleData &       operator[](int32 Index) { ensure(0 <= Index && Index < NumSamples()); return Samples[Index]; }
+		const FSampleData & operator[](int32 Index) const  { ensure(0 <= Index && Index < NumSamples()); return Samples[Index]; }
 
 	private:
-		int NumberSamples;
-		FSampleData Samples[4]; // iterative samples
+		TArray<FSampleData> Samples; // iterative samples
 	};
 	typedef TRigidBodyIterativeContactConstraint<float, 3> FRigidBodyIterativeContactConstraint;
 
diff --git a/Engine/Source/Runtime/Experimental/Chaos/Public/Chaos/Convex.h b/Engine/Source/Runtime/Experimental/Chaos/Public/Chaos/Convex.h
--- a/Engine/Source/Runtime/Experimental/Chaos/Public/Chaos/Convex.h
+++ b/Engine/Source/Runtime/Experimental/Chaos/Public/Chaos/Convex.h
@@ -140,7 +140,6 @@ namespace Chaos
 		}
 
 		virtual int32 FindMostOpposingFace(const FVec3& Position, const FVec3& UnitDir, int32 HintFaceIndex, FReal SearchDist) const override;
-		
 
 		FVec3 FindGeometryOpposingNormal(const FVec3& DenormDir, int32 FaceIndex, const FVec3& OriginalNormal) const
 		{
@@ -155,6 +154,9 @@ namespace Chaos
 			return FVec3(0.f, 0.f, 1.f);
 		}
 
+	
+		virtual int32 FindClosestFaceAndVertices(const FVec3& Position, TArray<FVec3>& FaceVertices, FReal SearchDist = 0.01) const override;
+
 		FORCEINLINE FReal GetMargin() const { return 0; }
 
 		FORCEINLINE FVec3 Support2(const FVec3& Direction) const { return Support(Direction, 0); }
diff --git a/Engine/Source/Runtime/Experimental/Chaos/Public/Chaos/ImplicitObject.h b/Engine/Source/Runtime/Experimental/Chaos/Public/Chaos/ImplicitObject.h
--- a/Engine/Source/Runtime/Experimental/Chaos/Public/Chaos/ImplicitObject.h
+++ b/Engine/Source/Runtime/Experimental/Chaos/Public/Chaos/ImplicitObject.h
@@ -214,6 +214,19 @@ public:
 		return INDEX_NONE;
 	}
 
+
+	/** Finds the first intersecting face at given position
+	@param Position - local position to search around (for example a point on the surface of a convex hull)
+	@param FaceIndices - Vertices that lie on the face plane.
+	@param SearchDistance - distance to surface [def:0.01]
+	*/
+	virtual int32 FindClosestFaceAndVertices(const FVec3& Position, TArray<FVec3>& FaceVertices, FReal SearchDist = 0.01) const
+	{
+		//Many objects have no concept of a face
+		return INDEX_NONE;
+	}
+
+
 	/** Given a normal and a face index, compute the most opposing normal associated with the underlying geometry features.
 		For example a sphere swept against a box may not give a normal associated with one of the box faces. This function will return a normal associated with one of the faces.
 		@param DenormDir - the direction we want to oppose
diff --git a/Engine/Source/Runtime/Experimental/Chaos/Public/Chaos/ImplicitObjectScaled.h b/Engine/Source/Runtime/Experimental/Chaos/Public/Chaos/ImplicitObjectScaled.h
--- a/Engine/Source/Runtime/Experimental/Chaos/Public/Chaos/ImplicitObjectScaled.h
+++ b/Engine/Source/Runtime/Experimental/Chaos/Public/Chaos/ImplicitObjectScaled.h
@@ -422,6 +422,22 @@ public:
 		return ClosestIntersection;
 	}
 
+	virtual int32 FindClosestFaceAndVertices(const FVec3& Position, TArray<FVec3>& FaceVertices, FReal SearchDist = 0.01) const override
+	{
+		const FVec3 UnscaledPoint = MInvScale * Position;
+		const FReal UnscaledSearchDist = SearchDist * MInvScale.Max();	//this is not quite right since it's no longer a sphere, but the whole thing is fuzzy anyway
+		int32 FaceIndex =  MObject->FindClosestFaceAndVertices(UnscaledPoint, FaceVertices, UnscaledSearchDist);
+		if (FaceIndex != INDEX_NONE)
+		{
+			for (FVec3& Vec : FaceVertices)
+			{
+				Vec = Vec * MScale;
+			}
+		}
+		return FaceIndex;
+	}
+
+
 	FORCEINLINE_DEBUGGABLE TVector<T, d> Support(const TVector<T, d>& Direction, const T Thickness) const
 	{
 		// Support_obj(dir) = pt => for all x in obj, pt \dot dir >= x \dot dir
diff --git a/Engine/Source/Runtime/Experimental/Chaos/Public/Chaos/ImplicitObjectTransformed.h b/Engine/Source/Runtime/Experimental/Chaos/Public/Chaos/ImplicitObjectTransformed.h
--- a/Engine/Source/Runtime/Experimental/Chaos/Public/Chaos/ImplicitObjectTransformed.h
+++ b/Engine/Source/Runtime/Experimental/Chaos/Public/Chaos/ImplicitObjectTransformed.h
@@ -153,6 +153,20 @@ public:
 		return ClosestIntersection;
 	}
 
+	virtual int32 FindClosestFaceAndVertices(const FVec3& Position, TArray<FVec3>& FaceVertices, FReal SearchDist = 0.01) const override
+	{
+		const FVec3 LocalPoint = MTransform.InverseTransformPosition(Position);
+		int32 FaceIndex = MObject->FindClosestFaceAndVertices(LocalPoint, FaceVertices, SearchDist);
+		if (FaceIndex != INDEX_NONE)
+		{
+			for (FVec3& Vec : FaceVertices)
+			{
+				Vec = MTransform.TransformPosition(Vec);
+			}
+		}
+		return FaceIndex;
+	}
+
 	const TRigidTransform<T, d>& GetTransform() const { return MTransform; }
 	void SetTransform(const TRigidTransform<T, d>& InTransform)
 	{
diff --git a/Engine/Source/Runtime/Experimental/Chaos/Public/Chaos/Plane.h b/Engine/Source/Runtime/Experimental/Chaos/Public/Chaos/Plane.h
--- a/Engine/Source/Runtime/Experimental/Chaos/Public/Chaos/Plane.h
+++ b/Engine/Source/Runtime/Experimental/Chaos/Public/Chaos/Plane.h
@@ -254,4 +254,16 @@ TVector<T, d> FindClosestPointOnTriangle(const TPlane<T, d>& TrianglePlane, cons
 	const TVector<T, d> PointOnPlane = TrianglePlane.FindClosestPoint(P);
 	return FindClosestPointOnTriangle(PointOnPlane, P0, P1, P2, P);
 }
+
+
+template<typename T, int d>
+bool IntersectPlanes2(TVector<T,d>& I, TVector<T,d>& D, const TPlane<T,d>& P1, const TPlane<T,d>& P2)
+{
+	FVector LI = I, LD = D;
+	FPlane LP1(P1.X(), P1.Normal()), LP2(P2.X(), P2.Normal());
+	bool RetVal = FMath::IntersectPlanes2(LI,LD,LP1,LP2);
+	I = LI; D = LD;
+	return RetVal;
+}
+
 }
