Chaos : Collisions

- Updated comments.

diff --git a/Engine/Source/Runtime/Experimental/Chaos/Private/Chaos/CollisionResolution.cpp b/Engine/Source/Runtime/Experimental/Chaos/Private/Chaos/CollisionResolution.cpp
--- a/Engine/Source/Runtime/Experimental/Chaos/Private/Chaos/CollisionResolution.cpp
+++ b/Engine/Source/Runtime/Experimental/Chaos/Private/Chaos/CollisionResolution.cpp
@@ -582,8 +582,7 @@ namespace Chaos
 					Constraint[Idx].Manifold.Normal = Transform1.TransformVector(Constraint.PlaneNormal);
 					Constraint[Idx].Manifold.Location = Transform0.TransformPosition(Constraint[Idx].X);
 
-					// save the best point for collision processing
-					// @todo(chaos) : Consider appling all constraint in ::Apply at each iteration, right now it just takes the deepest.
+					// save the best point for collision processing	
 					if (ContactPoint.Phi > Constraint[Idx].Manifold.Phi)
 					{
 						ContactPoint.Phi = Constraint[Idx].Manifold.Phi;
@@ -608,6 +607,13 @@ namespace Chaos
 					// single shot manifolds for TConvex implicit object in the constraints implicit[0] position. 
 
 					TArray<FVec3> CollisionSamples;
+					//
+					//  @todo(chaos) : Collision Manifold
+					//   Remove the dependency on the virtual calls on the Implicit. Don't use FindClosestFaceAndVertices
+					//   this relies on virtual calls on the ImplicitObject. Instead pass a parameters structures into 
+					//   ConvexConvexContactPoint that can collect the face indices during evaluation of the support functions. 
+					//   This can be implemented without virtual calls. 
+					//
 					int32 FaceIndex = Constraint.Manifold.Implicit[0]->FindClosestFaceAndVertices(Transform0.InverseTransformPosition(ContactPoint.Location), CollisionSamples, 1.f);
 
 					if (!ContactPoint.Normal.Equals(Constraint.PlaneNormal)
@@ -622,7 +628,11 @@ namespace Chaos
 
 					if (Constraint.NumSamples() == 0)
 					{
-						// @todo(chaos) : Only save the four best samples and hard-code the size of Constraint.Samples to [len:4]
+						//
+						// @todo(chaos) : Collision Manifold
+						//   Only save the four best samples and hard-code the size of Constraint.Samples to [len:4].
+						//   Currently this just grabs all points and uses the deepest point for resolution. 
+						//
 						Constraint.SourceNormalIndex = FaceIndex;
 						for (FVec3 Sample : CollisionSamples)
 						{
@@ -651,7 +661,11 @@ namespace Chaos
 						TVector<T, d> Center = SumSampleData(Constraint) / Constraint.NumSamples();
 						T Delta = (Center - SurfaceSample).SizeSquared();
 
-						// todo(chaos) : maximize for area instead
+						//
+						// @todo(chaos) : Collision Manifold
+						//    The iterative manifold need to be maximized for area instead of largest 
+						//    distance from center.
+						//
 						T SmallestDelta = FLT_MAX;
 						int32 SmallestIndex = 0;
 						for (int32 idx = 0; idx < Constraint.NumSamples(); idx++)
diff --git a/Engine/Source/Runtime/Experimental/Chaos/Private/Chaos/Convex.cpp b/Engine/Source/Runtime/Experimental/Chaos/Private/Chaos/Convex.cpp
--- a/Engine/Source/Runtime/Experimental/Chaos/Private/Chaos/Convex.cpp
+++ b/Engine/Source/Runtime/Experimental/Chaos/Private/Chaos/Convex.cpp
@@ -33,7 +33,10 @@ namespace Chaos
 	int32 FConvex::FindClosestFaceAndVertices(const FVec3& Position, TArray<FVec3>& FaceVertices, FReal SearchDist) const
 	{
 		//
-		//  todo(chaos) : Create a correspondence between the faces and surface particles on construction.
+		//  @todo(chaos) : Collision Manifold 
+		//     Create a correspondence between the faces and surface particles on construction.
+		//     The correspondence will provide an index between the Planes and the SurfaceParticles,
+		//     removing the need for the exhaustive search here. 
 		//
 
 		int32 ReturnIndex = INDEX_NONE;
diff --git a/Engine/Source/Runtime/Experimental/Chaos/Private/Chaos/PBDCollisionConstraintsContact.cpp b/Engine/Source/Runtime/Experimental/Chaos/Private/Chaos/PBDCollisionConstraintsContact.cpp
--- a/Engine/Source/Runtime/Experimental/Chaos/Private/Chaos/PBDCollisionConstraintsContact.cpp
+++ b/Engine/Source/Runtime/Experimental/Chaos/Private/Chaos/PBDCollisionConstraintsContact.cpp
@@ -294,6 +294,13 @@ namespace Chaos
 					Particle1->AuxilaryValue(*ParticleParameters.Collided) = true;
 				}
 
+				//
+				// @todo(chaos) : Collision Constraints
+				//   Consider applying all constraints in ::Apply at each iteration, right now it just takes the deepest.
+				//   For example, and iterative constraint might have 4 penetrating points that need to be resolved. 
+				//
+
+
 				Constraint.AccumulatedImpulse +=
 					ApplyContact(Constraint.Manifold, Particle0, Particle1, IterationParameters, ParticleParameters);
 			}
diff --git a/Engine/Source/Runtime/Experimental/Chaos/Public/Chaos/PBDCollisionConstraints.h b/Engine/Source/Runtime/Experimental/Chaos/Public/Chaos/PBDCollisionConstraints.h
--- a/Engine/Source/Runtime/Experimental/Chaos/Public/Chaos/PBDCollisionConstraints.h
+++ b/Engine/Source/Runtime/Experimental/Chaos/Public/Chaos/PBDCollisionConstraints.h
@@ -71,7 +71,8 @@ public:
 	/**
 	*  Add the constraint to the container. 
 	*
-	*  @todo(chaos) : Update to use a custom allocator. 
+	*  @todo(chaos) : Collision Constraints 
+	*  Update to use a custom allocator. 
 	*  The InConstraint should be a point to unmanaged, raw memory. 
 	*  This function will make a deep copy of the constraint and 
 	*  then delete the InConstraint. 
