Remove template from TriangleMeshImplicitObject

diff --git a/Engine/Source/Runtime/Engine/Classes/PhysicsEngine/BodySetup.h b/Engine/Source/Runtime/Engine/Classes/PhysicsEngine/BodySetup.h
--- a/Engine/Source/Runtime/Engine/Classes/PhysicsEngine/BodySetup.h
+++ b/Engine/Source/Runtime/Engine/Classes/PhysicsEngine/BodySetup.h
@@ -44,8 +44,7 @@ namespace Chaos
 {
 	class FImplicitObject;
 
-	template <typename T>
-	class TTriangleMeshImplicitObject;
+	class FTriangleMeshImplicitObject;
 }
 
 template<typename T, int d>
@@ -266,7 +265,7 @@ public:
 
 #if WITH_CHAOS
 	//FBodySetupTriMeshes* TriMeshWrapper;
-	TArray<TSharedPtr<Chaos::TTriangleMeshImplicitObject<float>, ESPMode::ThreadSafe>> ChaosTriMeshes;
+	TArray<TSharedPtr<Chaos::FTriangleMeshImplicitObject, ESPMode::ThreadSafe>> ChaosTriMeshes;
 #endif
 
 	/** Additional UV info, if available. Used for determining UV for a line trace impact. */
diff --git a/Engine/Source/Runtime/Engine/Private/PhysicsEngine/BodyInstance.cpp b/Engine/Source/Runtime/Engine/Private/PhysicsEngine/BodyInstance.cpp
--- a/Engine/Source/Runtime/Engine/Private/PhysicsEngine/BodyInstance.cpp
+++ b/Engine/Source/Runtime/Engine/Private/PhysicsEngine/BodyInstance.cpp
@@ -1841,10 +1841,10 @@ bool FBodyInstance::UpdateBodyScale(const FVector& InScale3D, bool bForceUpdate)
 					// PhysX supports translation, we currently do not.
 					CHAOS_ENSURE(RelativeTM.GetTranslation() == FVector(0, 0, 0));
 
-					const TImplicitObjectScaled<TTriangleMeshImplicitObject<FReal>>* ScaledTriangleMesh = (static_cast<const TImplicitObjectScaled<TTriangleMeshImplicitObject<FReal>>*>(&ImplicitObject));
+					const TImplicitObjectScaled<FTriangleMeshImplicitObject>* ScaledTriangleMesh = (static_cast<const TImplicitObjectScaled<FTriangleMeshImplicitObject>*>(&ImplicitObject));
 					auto UnscaledTriangleMesh = ScaledTriangleMesh->Object();
 
-					TUniquePtr<TImplicitObjectScaled<TTriangleMeshImplicitObject<FReal>>> NewTriangleMesh = MakeUnique<TImplicitObjectScaled<TTriangleMeshImplicitObject<FReal>>>(UnscaledTriangleMesh, AdjustedScale3D);
+					TUniquePtr<TImplicitObjectScaled<FTriangleMeshImplicitObject>> NewTriangleMesh = MakeUnique<TImplicitObjectScaled<FTriangleMeshImplicitObject>>(UnscaledTriangleMesh, AdjustedScale3D);
 					NewGeometry.Emplace(MoveTemp(NewTriangleMesh));
 
 					bSuccess = true;
diff --git a/Engine/Source/Runtime/Engine/Private/PhysicsEngine/Experimental/ChaosDerivedData.cpp b/Engine/Source/Runtime/Engine/Private/PhysicsEngine/Experimental/ChaosDerivedData.cpp
--- a/Engine/Source/Runtime/Engine/Private/PhysicsEngine/Experimental/ChaosDerivedData.cpp
+++ b/Engine/Source/Runtime/Engine/Private/PhysicsEngine/Experimental/ChaosDerivedData.cpp
@@ -61,7 +61,7 @@ bool FChaosDerivedDataCooker::Build(TArray<uint8>& OutData)
 		int32 PrecisionSize = (int32)sizeof(BuildPrecision);
 
 		Ar << PrecisionSize;
-		BuildInternal<BuildPrecision>(Ar, CookInfo);
+		BuildInternal(Ar, CookInfo);
 
 		bSucceeded = true;
 	}
@@ -84,8 +84,7 @@ FChaosDerivedDataCooker::FChaosDerivedDataCooker(UBodySetup* InSetup, FName InFo
 
 }
 
-template<typename Precision>
-void FChaosDerivedDataCooker::BuildTriangleMeshes(TArray<TUniquePtr<Chaos::TTriangleMeshImplicitObject<Precision>>>& OutTriangleMeshes, TArray<int32>& OutFaceRemap, const FCookBodySetupInfo& InParams)
+void FChaosDerivedDataCooker::BuildTriangleMeshes(TArray<TUniquePtr<Chaos::FTriangleMeshImplicitObject>>& OutTriangleMeshes, TArray<int32>& OutFaceRemap, const FCookBodySetupInfo& InParams)
 {
 	if(!InParams.bCookTriMesh)
 	{
@@ -108,7 +107,7 @@ void FChaosDerivedDataCooker::BuildTriangleMeshes(TArray<TUniquePtr<Chaos::TTria
 	Chaos::CleanTrimesh(FinalVerts, FinalIndices, &OutFaceRemap);
 
 	// Build particle list #BG Maybe allow TParticles to copy vectors?
-	Chaos::TParticles<Precision, 3> TriMeshParticles;
+	Chaos::TParticles<Chaos::FReal, 3> TriMeshParticles;
 	TriMeshParticles.AddParticles(FinalVerts.Num());
 
 	const int32 NumVerts = FinalVerts.Num();
@@ -169,13 +168,9 @@ void FChaosDerivedDataCooker::BuildTriangleMeshes(TArray<TUniquePtr<Chaos::TTria
 		}
 	}
 
-	OutTriangleMeshes.Emplace(new Chaos::TTriangleMeshImplicitObject<Precision>(MoveTemp(TriMeshParticles), MoveTemp(Triangles), MoveTemp(MaterialIndices)));
+	OutTriangleMeshes.Emplace(new Chaos::FTriangleMeshImplicitObject(MoveTemp(TriMeshParticles), MoveTemp(Triangles), MoveTemp(MaterialIndices)));
 }
 
-template void FChaosDerivedDataCooker::BuildTriangleMeshes(TArray<TUniquePtr<Chaos::TTriangleMeshImplicitObject<float>>>& OutTriangleMeshes, TArray<int32>& OutFaceRemap, const FCookBodySetupInfo& InParams);
-//#BGTODO When it's possible to build with doubles, re-enable this. (Currently at least TRigidTransform cannot build with double precision because we don't have a base transform implementation using them)
-//template void FChaosDerivedDataCooker::BuildTriangleMeshes(TArray<TUniquePtr<Chaos::FImplicitObject>>& OutTriangleMeshes, const FCookBodySetupInfo& InParams);
-
 void FChaosDerivedDataCooker::BuildConvexMeshes(TArray<TUniquePtr<Chaos::FImplicitObject>>& OutConvexMeshes, const FCookBodySetupInfo& InParams)
 {
 	auto BuildConvexFromVerts = [](TArray<TUniquePtr<Chaos::FImplicitObject>>& OutConvexes, const TArray<TArray<FVector>>& InMeshVerts, const bool bMirrored)
@@ -244,11 +239,10 @@ void BuildSimpleShapes(TArray<TUniquePtr<Chaos::FImplicitObject>>& OutImplicits,
 	UE_CLOG(NumTaperedCapsules > 0, LogChaos, Warning, TEXT("Ignoring %d tapered spheres when building collision data for body setup %s"), NumTaperedCapsules, *InSetup->GetName());
 }
 
-template<typename Precision>
 void FChaosDerivedDataCooker::BuildInternal(Chaos::FChaosArchive& Ar, FCookBodySetupInfo& InInfo)
 {
 	TArray<TUniquePtr<Chaos::FImplicitObject>> SimpleImplicits;
-	TArray<TUniquePtr<Chaos::TTriangleMeshImplicitObject<Precision>>> ComplexImplicits;
+	TArray<TUniquePtr<Chaos::FTriangleMeshImplicitObject>> ComplexImplicits;
 
 	TArray<int32> FaceRemap;
 	//BuildSimpleShapes(SimpleImplicits, Setup);
@@ -271,9 +265,5 @@ void FChaosDerivedDataCooker::BuildInternal(Chaos::FChaosArchive& Ar, FCookBodyS
 	Ar << SimpleImplicits << ComplexImplicits << UVInfo << FaceRemap;
 }
 
-template void FChaosDerivedDataCooker::BuildInternal<float>(Chaos::FChaosArchive& Ar, FCookBodySetupInfo& InInfo);
-//#BGTODO When it's possible to build with doubles, re-enable this. (Currently at least TRigidTransform cannot build with double precision because we don't have a base transform implementation using them)
-//template void FChaosDerivedDataCooker::BuildInternal<float>(FArchive& Ar, FCookBodySetupInfo& InInfo);
-
 #endif
  
diff --git a/Engine/Source/Runtime/Engine/Private/PhysicsEngine/Experimental/ChaosDerivedData.h b/Engine/Source/Runtime/Engine/Private/PhysicsEngine/Experimental/ChaosDerivedData.h
--- a/Engine/Source/Runtime/Engine/Private/PhysicsEngine/Experimental/ChaosDerivedData.h
+++ b/Engine/Source/Runtime/Engine/Private/PhysicsEngine/Experimental/ChaosDerivedData.h
@@ -20,8 +20,7 @@ namespace Chaos
 
 	class FImplicitObject;
 
-	template<typename T>
-	class TTriangleMeshImplicitObject;
+	class FTriangleMeshImplicitObject;
 }
 
 struct FChaosTriMeshCollisionBuildParameters
@@ -65,11 +64,9 @@ public:
 
 private:
 
-	template<typename Precision>
 	void BuildInternal(Chaos::FChaosArchive& Ar, FCookBodySetupInfo& InInfo);
 
-	template<typename Precision>
-	void BuildTriangleMeshes(TArray<TUniquePtr<Chaos::TTriangleMeshImplicitObject<Precision>>>& OutTriangleMeshes, TArray<int32>& OutFaceRemap, const FCookBodySetupInfo& InParams);
+	void BuildTriangleMeshes(TArray<TUniquePtr<Chaos::FTriangleMeshImplicitObject>>& OutTriangleMeshes, TArray<int32>& OutFaceRemap, const FCookBodySetupInfo& InParams);
 
 	void BuildConvexMeshes(TArray<TUniquePtr<Chaos::FImplicitObject>>& OutTriangleMeshes, const FCookBodySetupInfo& InParams);
 
@@ -77,12 +74,5 @@ private:
 	FName RequestedFormat;
 };
 
-extern template void FChaosDerivedDataCooker::BuildTriangleMeshes(TArray<TUniquePtr<Chaos::TTriangleMeshImplicitObject<float>>>& OutTriangleMeshes, TArray<int32>& OutFaceRemap, const FCookBodySetupInfo& InParams);
-// #BGTODO When it's possible to build with doubles, re-enable this. (Currently at least TRigidTransform cannot build with double precision because we don't have a base transform implementation using them)
-//extern template void FChaosDerivedDataCooker::BuildTriangleMeshes(TArray<Chaos::TCollisionTriangleMesh<double>>& OutTriangleMeshes, const FCookBodySetupInfo& InParams);
-
-extern template void FChaosDerivedDataCooker::BuildInternal<float>(Chaos::FChaosArchive& Ar, FCookBodySetupInfo& InInfo);
-// #BGTODO When it's possible to build with doubles, re-enable this. (Currently at least TRigidTransform cannot build with double precision because we don't have a base transform implementation using them)
-//extern template void FChaosDerivedDataCooker::BuildInternal<float>(FArchive& Ar, FCookBodySetupInfo& InInfo);
 
 #endif
\ No newline at end of file
diff --git a/Engine/Source/Runtime/Engine/Private/PhysicsEngine/Experimental/ChaosInterfaceUtils.cpp b/Engine/Source/Runtime/Engine/Private/PhysicsEngine/Experimental/ChaosInterfaceUtils.cpp
--- a/Engine/Source/Runtime/Engine/Private/PhysicsEngine/Experimental/ChaosInterfaceUtils.cpp
+++ b/Engine/Source/Runtime/Engine/Private/PhysicsEngine/Experimental/ChaosInterfaceUtils.cpp
@@ -272,7 +272,7 @@ namespace ChaosInterface
 
 		for (const auto& ChaosTriMesh : InParams.ChaosTriMeshes)
 		{
-			auto Implicit = MakeUnique<Chaos::TImplicitObjectScaled<Chaos::TTriangleMeshImplicitObject<float>>>(ChaosTriMesh, Scale);
+			auto Implicit = MakeUnique<Chaos::TImplicitObjectScaled<Chaos::FTriangleMeshImplicitObject>>(ChaosTriMesh, Scale);
 			auto NewShape = NewShapeHelper(MakeSerializable(Implicit), nullptr, true);
 			Shapes.Emplace(MoveTemp(NewShape));
 			Geoms.Add(MoveTemp(Implicit));
diff --git a/Engine/Source/Runtime/Engine/Private/PhysicsEngine/Experimental/PhysInterface_Chaos.cpp b/Engine/Source/Runtime/Engine/Private/PhysicsEngine/Experimental/PhysInterface_Chaos.cpp
--- a/Engine/Source/Runtime/Engine/Private/PhysicsEngine/Experimental/PhysInterface_Chaos.cpp
+++ b/Engine/Source/Runtime/Engine/Private/PhysicsEngine/Experimental/PhysInterface_Chaos.cpp
@@ -1401,9 +1401,9 @@ const Chaos::FConvex& FPhysicsGeometryCollection_Chaos::GetConvexGeometry() cons
 	return Geom.GetObjectChecked<Chaos::FConvex>();
 }
 
-const Chaos::TTriangleMeshImplicitObject<float>& FPhysicsGeometryCollection_Chaos::GetTriMeshGeometry() const
+const Chaos::FTriangleMeshImplicitObject& FPhysicsGeometryCollection_Chaos::GetTriMeshGeometry() const
 {
-	return Geom.GetObjectChecked<Chaos::TTriangleMeshImplicitObject<float>>();
+	return Geom.GetObjectChecked<Chaos::FTriangleMeshImplicitObject>();
 }
 
 FPhysicsGeometryCollection_Chaos::FPhysicsGeometryCollection_Chaos(const FPhysicsShapeReference_Chaos& InShape)
diff --git a/Engine/Source/Runtime/Engine/Public/Physics/Experimental/ChaosDerivedDataReader.h b/Engine/Source/Runtime/Engine/Public/Physics/Experimental/ChaosDerivedDataReader.h
--- a/Engine/Source/Runtime/Engine/Public/Physics/Experimental/ChaosDerivedDataReader.h
+++ b/Engine/Source/Runtime/Engine/Public/Physics/Experimental/ChaosDerivedDataReader.h
@@ -8,8 +8,7 @@ namespace Chaos
 {
 	class FImplicitObject;
 
-	template <typename T>
-	class TTriangleMeshImplicitObject;
+	class FTriangleMeshImplicitObject;
 
 	class FConvex;
 }
@@ -25,7 +24,7 @@ public:
 	explicit FChaosDerivedDataReader(FUntypedBulkData* InBulkData);
 
 	TArray<TUniquePtr<Chaos::FConvex>> ConvexImplicitObjects;
-	TArray<TSharedPtr<Chaos::TTriangleMeshImplicitObject<T>, ESPMode::ThreadSafe>> TrimeshImplicitObjects;
+	TArray<TSharedPtr<Chaos::FTriangleMeshImplicitObject, ESPMode::ThreadSafe>> TrimeshImplicitObjects;
 	FBodySetupUVInfo UVInfo;
 	TArray<int32> FaceRemap;
 
diff --git a/Engine/Source/Runtime/Engine/Public/Physics/Experimental/PhysInterface_Chaos.h b/Engine/Source/Runtime/Engine/Public/Physics/Experimental/PhysInterface_Chaos.h
--- a/Engine/Source/Runtime/Engine/Public/Physics/Experimental/PhysInterface_Chaos.h
+++ b/Engine/Source/Runtime/Engine/Public/Physics/Experimental/PhysInterface_Chaos.h
@@ -150,7 +150,7 @@ struct ENGINE_API FPhysicsGeometryCollection_Chaos
 	const Chaos::TSphere<float, 3>&  GetSphereGeometry() const;
 	const Chaos::TCapsule<float>&  GetCapsuleGeometry() const;
 	const Chaos::FConvex& GetConvexGeometry() const;
-	const Chaos::TTriangleMeshImplicitObject<float>& GetTriMeshGeometry() const;
+	const Chaos::FTriangleMeshImplicitObject& GetTriMeshGeometry() const;
 
 private:
 	friend class FPhysInterface_Chaos;
diff --git a/Engine/Source/Runtime/Engine/Public/Physics/PhysicsInterfaceTypes.h b/Engine/Source/Runtime/Engine/Public/Physics/PhysicsInterfaceTypes.h
--- a/Engine/Source/Runtime/Engine/Public/Physics/PhysicsInterfaceTypes.h
+++ b/Engine/Source/Runtime/Engine/Public/Physics/PhysicsInterfaceTypes.h
@@ -18,8 +18,7 @@ namespace Chaos
 
 	class FImplicitObject;
 
-	template <typename T>
-	class TTriangleMeshImplicitObject;
+	class FTriangleMeshImplicitObject;
 }
 #endif
 
@@ -56,7 +55,7 @@ struct FGeometryAddParams
 	TArrayView<physx::PxTriangleMesh*> TriMeshes;
 #endif
 #if WITH_CHAOS
-	TArrayView<TSharedPtr<Chaos::TTriangleMeshImplicitObject<float>, ESPMode::ThreadSafe>> ChaosTriMeshes;
+	TArrayView<TSharedPtr<Chaos::FTriangleMeshImplicitObject, ESPMode::ThreadSafe>> ChaosTriMeshes;
 #endif
 };
 
diff --git a/Engine/Source/Runtime/Experimental/Chaos/Private/Chaos/ImplicitObject.cpp b/Engine/Source/Runtime/Experimental/Chaos/Private/Chaos/ImplicitObject.cpp
--- a/Engine/Source/Runtime/Experimental/Chaos/Private/Chaos/ImplicitObject.cpp
+++ b/Engine/Source/Runtime/Experimental/Chaos/Private/Chaos/ImplicitObject.cpp
@@ -350,7 +350,7 @@ FImplicitObject* FImplicitObject::SerializationFactory(FChaosArchive& Ar, FImpli
 			switch (InnerType)
 			{
 			case ImplicitObjectType::Convex: if (Ar.IsLoading()) { return new TImplicitObjectScaled<FConvex>(); } break;
-			case ImplicitObjectType::TriangleMesh: if (Ar.IsLoading()) { return new TImplicitObjectScaled<TTriangleMeshImplicitObject<FReal>>(); } break;
+			case ImplicitObjectType::TriangleMesh: if (Ar.IsLoading()) { return new TImplicitObjectScaled<FTriangleMeshImplicitObject>(); } break;
 			default: check(false);
 			}
 
@@ -369,7 +369,7 @@ FImplicitObject* FImplicitObject::SerializationFactory(FChaosArchive& Ar, FImpli
 	case ImplicitObjectType::LevelSet: if (Ar.IsLoading()) { return new TLevelSet<FReal, 3>(); } break;
 	case ImplicitObjectType::Convex: if (Ar.IsLoading()) { return new FConvex(); } break;
 	case ImplicitObjectType::TaperedCylinder: if (Ar.IsLoading()) { return new TTaperedCylinder<FReal>(); } break;
-	case ImplicitObjectType::TriangleMesh: if (Ar.IsLoading()) { return new TTriangleMeshImplicitObject<FReal>(); } break;
+	case ImplicitObjectType::TriangleMesh: if (Ar.IsLoading()) { return new FTriangleMeshImplicitObject(); } break;
 	case ImplicitObjectType::DEPRECATED_Scaled:
 	{
 		ensure(Ar.IsLoading() && (Ar.CustomVer(FExternalPhysicsCustomObjectVersion::GUID) < FExternalPhysicsCustomObjectVersion::ScaledGeometryIsConcrete));
diff --git a/Engine/Source/Runtime/Experimental/Chaos/Private/Chaos/TriangleMeshImplicitObject.cpp b/Engine/Source/Runtime/Experimental/Chaos/Private/Chaos/TriangleMeshImplicitObject.cpp
--- a/Engine/Source/Runtime/Experimental/Chaos/Private/Chaos/TriangleMeshImplicitObject.cpp
+++ b/Engine/Source/Runtime/Experimental/Chaos/Private/Chaos/TriangleMeshImplicitObject.cpp
@@ -8,8 +8,7 @@
 
 namespace Chaos
 {
-template <typename T>
-TTriangleMeshImplicitObject<T>::TTriangleMeshImplicitObject(TParticles<T, 3>&& Particles, TArray<TVector<int32, 3>>&& Elements, TArray<uint16>&& InMaterialIndices)
+FTriangleMeshImplicitObject::FTriangleMeshImplicitObject(TParticles<FReal, 3>&& Particles, TArray<TVector<int32, 3>>&& Elements, TArray<uint16>&& InMaterialIndices)
 	: FImplicitObject(EImplicitObject::HasBoundingBox, ImplicitObjectType::TriangleMesh)
 	, MParticles(MoveTemp(Particles))
 	, MElements(MoveTemp(Elements))
@@ -23,16 +22,15 @@ TTriangleMeshImplicitObject<T>::TTriangleMeshImplicitObject(TParticles<T, 3>&& P
 	RebuildBV();
 }
 
-template <typename T>
-struct TTriangleMeshRaycastVisitor
+struct FTriangleMeshRaycastVisitor
 {
-	TTriangleMeshRaycastVisitor(const TVector<T,3>& InStart, const TVector<T,3>& InDir, const T InThickness, const TParticles<T,3>& InParticles, const TArray<TVector<int32, 3>>& InElements)
+	FTriangleMeshRaycastVisitor(const FVec3& InStart, const FVec3& InDir, const FReal InThickness, const TParticles<FReal,3>& InParticles, const TArray<TVector<int32, 3>>& InElements)
 	: Particles(InParticles)
 	, Elements(InElements)
 	, StartPoint(InStart)
 	, Dir(InDir)
 	, Thickness(InThickness)
-	, OutTime(TNumericLimits<T>::Max())
+	, OutTime(TNumericLimits<FReal>::Max())
 	{
 	}
 
@@ -45,45 +43,45 @@ struct TTriangleMeshRaycastVisitor
 	template <ERaycastType SQType>
 	bool Visit(int32 TriIdx, FQueryFastData& CurData)
 	{
-		constexpr T Epsilon = 1e-4;
-		constexpr T Epsilon2 = Epsilon * Epsilon;
-		const T Thickness2 = SQType == ERaycastType::Sweep ? Thickness * Thickness : 0;
-		T MinTime = 0;	//no need to initialize, but fixes warning
-
-		const T R = Thickness + Epsilon;
-		const T R2 = R * R;
-
-		const TVector<T, 3>& A = Particles.X(Elements[TriIdx][0]);
-		const TVector<T, 3>& B = Particles.X(Elements[TriIdx][1]);
-		const TVector<T, 3>& C = Particles.X(Elements[TriIdx][2]);
-
-		const TVector<T, 3> AB = B - A;
-		const TVector<T, 3> AC = C - A;
-		TVector<T, 3> TriNormal = TVector<T, 3>::CrossProduct(AB, AC);
-		const T NormalLength = TriNormal.SafeNormalize();
+		constexpr FReal Epsilon = 1e-4;
+		constexpr FReal Epsilon2 = Epsilon * Epsilon;
+		const FReal Thickness2 = SQType == ERaycastType::Sweep ? Thickness * Thickness : 0;
+		FReal MinTime = 0;	//no need to initialize, but fixes warning
+
+		const FReal R = Thickness + Epsilon;
+		const FReal R2 = R * R;
+
+		const FVec3& A = Particles.X(Elements[TriIdx][0]);
+		const FVec3& B = Particles.X(Elements[TriIdx][1]);
+		const FVec3& C = Particles.X(Elements[TriIdx][2]);
+
+		const FVec3 AB = B - A;
+		const FVec3 AC = C - A;
+		FVec3 TriNormal = FVec3::CrossProduct(AB, AC);
+		const FReal NormalLength = TriNormal.SafeNormalize();
 		if (!CHAOS_ENSURE(NormalLength > Epsilon))
 		{
 			//hitting degenerate triangle so keep searching - should be fixed before we get to this stage
 			return true;
 		}
 
-		const TPlane<T, 3> TriPlane{ A, TriNormal };
-		TVector<T, 3> RaycastPosition;
-		TVector<T, 3> RaycastNormal;
-		T Time;
+		const TPlane<FReal, 3> TriPlane{ A, TriNormal };
+		FVec3 RaycastPosition;
+		FVec3 RaycastNormal;
+		FReal Time;
 
 		//Check if we even intersect with triangle plane
 		int32 DummyFaceIndex;
 		if (TriPlane.Raycast(StartPoint, Dir, CurData.CurrentLength, Thickness, Time, RaycastPosition, RaycastNormal, DummyFaceIndex))
 		{
-			TVector<T, 3> IntersectionPosition = RaycastPosition;
-			TVector<T, 3> IntersectionNormal = RaycastNormal;
+			FVec3 IntersectionPosition = RaycastPosition;
+			FVec3 IntersectionNormal = RaycastNormal;
 			bool bTriangleIntersects = false;
 			if (Time == 0)
 			{
 				//Initial overlap so no point of intersection, do an explicit sphere triangle test.
-				const TVector<T, 3> ClosestPtOnTri = FindClosestPointOnTriangle(TriPlane, A, B, C, StartPoint);
-				const T DistToTriangle2 = (StartPoint - ClosestPtOnTri).SizeSquared();
+				const FVec3 ClosestPtOnTri = FindClosestPointOnTriangle(TriPlane, A, B, C, StartPoint);
+				const FReal DistToTriangle2 = (StartPoint - ClosestPtOnTri).SizeSquared();
 				if (DistToTriangle2 <= R2)
 				{
 					OutTime = 0;
@@ -93,39 +91,39 @@ struct TTriangleMeshRaycastVisitor
 			}
 			else
 			{
-				const TVector<T, 3> ClosestPtOnTri = FindClosestPointOnTriangle(RaycastPosition, A, B, C, RaycastPosition);	//We know Position is on the triangle plane
-				const T DistToTriangle2 = (RaycastPosition - ClosestPtOnTri).SizeSquared();
+				const FVec3 ClosestPtOnTri = FindClosestPointOnTriangle(RaycastPosition, A, B, C, RaycastPosition);	//We know Position is on the triangle plane
+				const FReal DistToTriangle2 = (RaycastPosition - ClosestPtOnTri).SizeSquared();
 				bTriangleIntersects = DistToTriangle2 <= Epsilon2;	//raycast gave us the intersection point so sphere radius is already accounted for
 			}
 
 			if (SQType == ERaycastType::Sweep && !bTriangleIntersects)
 			{
 				//sphere is not immediately touching the triangle, but it could start intersecting the perimeter as it sweeps by
-				TVector<T, 3> BorderPositions[3];
-				TVector<T, 3> BorderNormals[3];
-				T BorderTimes[3];
+				FVec3 BorderPositions[3];
+				FVec3 BorderNormals[3];
+				FReal BorderTimes[3];
 				bool bBorderIntersections[3];
 
 				{
-					TVector<T, 3> ABCapsuleAxis = B - A;
-					T ABHeight = ABCapsuleAxis.SafeNormalize();
-					bBorderIntersections[0] = TCapsule<T>::RaycastFast(Thickness, ABHeight, ABCapsuleAxis, A, B, StartPoint, Dir, CurData.CurrentLength, 0, BorderTimes[0], BorderPositions[0], BorderNormals[0], DummyFaceIndex);
+					FVec3 ABCapsuleAxis = B - A;
+					FReal ABHeight = ABCapsuleAxis.SafeNormalize();
+					bBorderIntersections[0] = TCapsule<FReal>::RaycastFast(Thickness, ABHeight, ABCapsuleAxis, A, B, StartPoint, Dir, CurData.CurrentLength, 0, BorderTimes[0], BorderPositions[0], BorderNormals[0], DummyFaceIndex);
 				}
 				
 				{
-					TVector<T, 3> BCCapsuleAxis = C - B;
-					T BCHeight = BCCapsuleAxis.SafeNormalize();
-					bBorderIntersections[1] = TCapsule<T>::RaycastFast(Thickness, BCHeight, BCCapsuleAxis, B, C, StartPoint, Dir, CurData.CurrentLength, 0, BorderTimes[1], BorderPositions[1], BorderNormals[1], DummyFaceIndex);
+					FVec3 BCCapsuleAxis = C - B;
+					FReal BCHeight = BCCapsuleAxis.SafeNormalize();
+					bBorderIntersections[1] = TCapsule<FReal>::RaycastFast(Thickness, BCHeight, BCCapsuleAxis, B, C, StartPoint, Dir, CurData.CurrentLength, 0, BorderTimes[1], BorderPositions[1], BorderNormals[1], DummyFaceIndex);
 				}
 				
 				{
-					TVector<T, 3> ACCapsuleAxis = C - A;
-					T ACHeight = ACCapsuleAxis.SafeNormalize();
-					bBorderIntersections[2] = TCapsule<T>::RaycastFast(Thickness, ACHeight, ACCapsuleAxis, A, C, StartPoint, Dir, CurData.CurrentLength, 0, BorderTimes[2], BorderPositions[2], BorderNormals[2], DummyFaceIndex);
+					FVec3 ACCapsuleAxis = C - A;
+					FReal ACHeight = ACCapsuleAxis.SafeNormalize();
+					bBorderIntersections[2] = TCapsule<FReal>::RaycastFast(Thickness, ACHeight, ACCapsuleAxis, A, C, StartPoint, Dir, CurData.CurrentLength, 0, BorderTimes[2], BorderPositions[2], BorderNormals[2], DummyFaceIndex);
 				}
 
 				int32 MinBorderIdx = INDEX_NONE;
-				T MinBorderTime = 0;	//initialization not needed, but fixes warning
+				FReal MinBorderTime = 0;	//initialization not needed, but fixes warning
 
 				for (int32 BorderIdx = 0; BorderIdx < 3; ++BorderIdx)
 				{
@@ -148,8 +146,8 @@ struct TTriangleMeshRaycastVisitor
 					if (Time == 0)
 					{
 						//we were initially overlapping with triangle plane so no normal was given. Compute it now
-						TVector<T, 3> TmpNormal;
-						const T SignedDistance = TriPlane.PhiWithNormal(StartPoint, TmpNormal);
+						FVec3 TmpNormal;
+						const FReal SignedDistance = TriPlane.PhiWithNormal(StartPoint, TmpNormal);
 						RaycastNormal = SignedDistance >= 0 ? TmpNormal : -TmpNormal;
 					}
 
@@ -189,32 +187,30 @@ struct TTriangleMeshRaycastVisitor
 		return true;
 	}
 
-	const TParticles<T, 3>& Particles;
+	const TParticles<FReal, 3>& Particles;
 	const TArray<TVector<int32, 3>>& Elements;
-	const TVector<T, 3>& StartPoint;
-	const TVector<T, 3>& Dir;
-	const T Thickness;
-	T OutTime;
-	TVector<T, 3> OutPosition;
-	TVector<T, 3> OutNormal;
+	const FVec3& StartPoint;
+	const FVec3& Dir;
+	const FReal Thickness;
+	FReal OutTime;
+	FVec3 OutPosition;
+	FVec3 OutNormal;
 	int32 OutFaceIndex;
 };
 
-template <typename T>
-T TTriangleMeshImplicitObject<T>::PhiWithNormal(const TVector<T, 3>& x, TVector<T, 3>& Normal) const
+FReal FTriangleMeshImplicitObject::PhiWithNormal(const FVec3& x, FVec3& Normal) const
 {
 	ensure(false);	//not supported yet - might support it in the future or we may change the interface
-	return (T)0;
+	return 0;
 }
 
-template <typename T>
-bool TTriangleMeshImplicitObject<T>::Raycast(const TVector<T, 3>& StartPoint, const TVector<T, 3>& Dir, const T Length, const T Thickness, T& OutTime, TVector<T, 3>& OutPosition, TVector<T, 3>& OutNormal, int32& OutFaceIndex) const
+bool FTriangleMeshImplicitObject::Raycast(const FVec3& StartPoint, const FVec3& Dir, const FReal Length, const FReal Thickness, FReal& OutTime, FVec3& OutPosition, FVec3& OutNormal, int32& OutFaceIndex) const
 {
-	TTriangleMeshRaycastVisitor<T> SQVisitor(StartPoint, Dir, Thickness, MParticles, MElements);
+	FTriangleMeshRaycastVisitor SQVisitor(StartPoint, Dir, Thickness, MParticles, MElements);
 
 	if (Thickness > 0)
 	{
-		BVH.Sweep(StartPoint, Dir, Length, TVector<T, 3>(Thickness), SQVisitor);
+		BVH.Sweep(StartPoint, Dir, Length, FVec3(Thickness), SQVisitor);
 	}
 	else
 	{
@@ -235,36 +231,35 @@ bool TTriangleMeshImplicitObject<T>::Raycast(const TVector<T, 3>& StartPoint, co
 	}
 }
 
-template <typename T>
-bool TTriangleMeshImplicitObject<T>::Overlap(const TVector<T, 3>& Point, const T Thickness) const
+bool FTriangleMeshImplicitObject::Overlap(const FVec3& Point, const FReal Thickness) const
 {
-	TAABB<T, 3> QueryBounds(Point, Point);
+	TAABB<FReal, 3> QueryBounds(Point, Point);
 	QueryBounds.Thicken(Thickness);
 	const TArray<int32> PotentialIntersections = BVH.FindAllIntersections(QueryBounds);
 
-	const T Epsilon = 1e-4;
+	const FReal Epsilon = 1e-4;
 	//ensure(Thickness > Epsilon);	//There's no hope for this to work unless thickness is large (really a sphere overlap test)
 	//todo: turn ensure back on, off until some other bug is fixed
 
 	for (int32 TriIdx : PotentialIntersections)
 	{
-		const TVector<T, 3>& A = MParticles.X(MElements[TriIdx][0]);
-		const TVector<T, 3>& B = MParticles.X(MElements[TriIdx][1]);
-		const TVector<T, 3>& C = MParticles.X(MElements[TriIdx][2]);
-
-		const TVector<T, 3> AB = B - A;
-		const TVector<T, 3> AC = C - A;
-		TVector<T, 3> Normal = TVector<T, 3>::CrossProduct(AB, AC);
-		const T NormalLength = Normal.SafeNormalize();
+		const FVec3& A = MParticles.X(MElements[TriIdx][0]);
+		const FVec3& B = MParticles.X(MElements[TriIdx][1]);
+		const FVec3& C = MParticles.X(MElements[TriIdx][2]);
+
+		const FVec3 AB = B - A;
+		const FVec3 AC = C - A;
+		FVec3 Normal = FVec3::CrossProduct(AB, AC);
+		const FReal NormalLength = Normal.SafeNormalize();
 		if (!ensure(NormalLength > Epsilon))
 		{
 			//hitting degenerate triangle - should be fixed before we get to this stage
 			continue;
 		}
 
-		const TPlane<T, 3> TriPlane{ A, Normal };
-		const TVector<T,3> ClosestPointOnTri = FindClosestPointOnTriangle(TriPlane, A, B, C, Point);
-		const T Distance2 = (ClosestPointOnTri - Point).SizeSquared();
+		const TPlane<FReal, 3> TriPlane{ A, Normal };
+		const FVec3 ClosestPointOnTri = FindClosestPointOnTriangle(TriPlane, A, B, C, Point);
+		const FReal Distance2 = (ClosestPointOnTri - Point).SizeSquared();
 		if (Distance2 <= Thickness * Thickness)	//This really only has a hope in working if thickness is > 0
 		{
 			return true;
@@ -273,20 +268,20 @@ bool TTriangleMeshImplicitObject<T>::Overlap(const TVector<T, 3>& Point, const T
 	return false;
 }
 
-template <typename QueryGeomType, typename T>
-void TransformVertsHelper(const QueryGeomType& QueryGeom, int32 TriIdx, const TParticles<T, 3>& Particles,
-	const TArray<TVector<int32, 3>>& Elements, TVec3<T>& OutA, TVec3<T>& OutB, TVec3<T>& OutC)
+template <typename QueryGeomType>
+void TransformVertsHelper(const QueryGeomType& QueryGeom, int32 TriIdx, const TParticles<FReal, 3>& Particles,
+	const TArray<TVector<int32, 3>>& Elements, TVec3<FReal>& OutA, TVec3<FReal>& OutB, TVec3<FReal>& OutC)
 {
 	OutA = Particles.X(Elements[TriIdx][0]);
 	OutB = Particles.X(Elements[TriIdx][1]);
 	OutC = Particles.X(Elements[TriIdx][2]);
 }
 
-template <typename QueryGeomType, typename T>
-void TransformVertsHelper(const TImplicitObjectScaled<QueryGeomType>& QueryGeom, int32 TriIdx, const TParticles<T, 3>& Particles,
-	const TArray<TVector<int32, 3>>& Elements, TVec3<T>& OutA, TVec3<T>& OutB, TVec3<T>& OutC)
+template <typename QueryGeomType>
+void TransformVertsHelper(const TImplicitObjectScaled<QueryGeomType>& QueryGeom, int32 TriIdx, const TParticles<FReal, 3>& Particles,
+	const TArray<TVector<int32, 3>>& Elements, TVec3<FReal>& OutA, TVec3<FReal>& OutB, TVec3<FReal>& OutC)
 {
-	const TVec3<T> InvScale = QueryGeom.GetInvScale();
+	const TVec3<FReal> InvScale = QueryGeom.GetInvScale();
 	OutA = Particles.X(Elements[TriIdx][0]) * InvScale;
 	OutB = Particles.X(Elements[TriIdx][1]) * InvScale;
 	OutC = Particles.X(Elements[TriIdx][2]) * InvScale;
@@ -304,47 +299,46 @@ const QueryGeomType& GetGeomHelper(const TImplicitObjectScaled<QueryGeomType>& Q
 	return *QueryGeom.GetUnscaledObject();
 }
 
-template <typename QueryGeomType, typename T>
-void TransformSweepOutputsHelper(const QueryGeomType& QueryGeom, const TVec3<T>& HitNormal, const TVec3<T>& HitPosition, const T LengthScale,
-	const T Time,  TVec3<T>& OutNormal, TVec3<T>& OutPosition, T& OutTime)
+template <typename QueryGeomType>
+void TransformSweepOutputsHelper(const QueryGeomType& QueryGeom, const TVec3<FReal>& HitNormal, const TVec3<FReal>& HitPosition, const FReal LengthScale,
+	const FReal Time,  TVec3<FReal>& OutNormal, TVec3<FReal>& OutPosition, FReal& OutTime)
 {
 	OutNormal = HitNormal;
 	OutPosition = HitPosition;
 	OutTime = Time;
 }
 
-template <typename QueryGeomType, typename T>
-void TransformSweepOutputsHelper(const TImplicitObjectScaled<QueryGeomType>& QueryGeom, const TVec3<T>& HitNormal, const TVec3<T>& HitPosition,  const T LengthScale,
-	const T Time,  TVec3<T>& OutNormal, TVec3<T>& OutPosition, T& OutTime)
+template <typename QueryGeomType>
+void TransformSweepOutputsHelper(const TImplicitObjectScaled<QueryGeomType>& QueryGeom, const TVec3<FReal>& HitNormal, const TVec3<FReal>& HitPosition,  const FReal LengthScale,
+	const FReal Time,  TVec3<FReal>& OutNormal, TVec3<FReal>& OutPosition, FReal& OutTime)
 {
-	const TVec3<T> InvScale = QueryGeom.GetInvScale();
-	const TVec3<T> Scale = QueryGeom.GetScale();
+	const TVec3<FReal> InvScale = QueryGeom.GetInvScale();
+	const TVec3<FReal> Scale = QueryGeom.GetScale();
 
 	OutTime = Time / LengthScale;
 	OutNormal = (InvScale * HitNormal).GetSafeNormal();
 	OutPosition = Scale * HitPosition;
 }
 
-template <typename QueryGeomType, typename T>
-void TransformOverlapInputsHelper(const QueryGeomType& QueryGeom, const TRigidTransform<T, 3>& QueryTM, TRigidTransform<T, 3>& OutScaledQueryTM)
+template <typename QueryGeomType>
+void TransformOverlapInputsHelper(const QueryGeomType& QueryGeom, const FRigidTransform3& QueryTM, FRigidTransform3& OutScaledQueryTM)
 {
 	OutScaledQueryTM = QueryTM;
 }
 
-template <typename QueryGeomType, typename T>
-void TransformOverlapInputsHelper(const TImplicitObjectScaled<QueryGeomType>& QueryGeom, const TRigidTransform<T, 3>& QueryTM, TRigidTransform<T, 3>& OutScaledQueryTM)
+template <typename QueryGeomType>
+void TransformOverlapInputsHelper(const TImplicitObjectScaled<QueryGeomType>& QueryGeom, const FRigidTransform3& QueryTM, FRigidTransform3& OutScaledQueryTM)
 {
-	const TVec3<T> InvScale = QueryGeom.GetInvScale();
+	const TVec3<FReal> InvScale = QueryGeom.GetInvScale();
 	OutScaledQueryTM = TRigidTransform<FReal, 3>(QueryTM.GetLocation() * InvScale, QueryTM.GetRotation());
 }
 
 
-template <typename T>
 template <typename QueryGeomType>
-bool TTriangleMeshImplicitObject<T>::OverlapGeomImp(const QueryGeomType& QueryGeom, const TRigidTransform<T, 3>& QueryTM, const T Thickness) const
+bool FTriangleMeshImplicitObject::OverlapGeomImp(const QueryGeomType& QueryGeom, const FRigidTransform3& QueryTM, const FReal Thickness) const
 {
 	bool bResult = false;
-	TAABB<T, 3> QueryBounds = QueryGeom.BoundingBox();
+	TAABB<FReal, 3> QueryBounds = QueryGeom.BoundingBox();
 	QueryBounds.Thicken(Thickness);
 	QueryBounds = QueryBounds.TransformedAABB(QueryTM);
 	const TArray<int32> PotentialIntersections = BVH.FindAllIntersections(QueryBounds);
@@ -356,17 +350,17 @@ bool TTriangleMeshImplicitObject<T>::OverlapGeomImp(const QueryGeomType& QueryGe
 	
 	for (int32 TriIdx : PotentialIntersections)
 	{
-		TVec3<T> A, B, C;
+		TVec3<FReal> A, B, C;
 		TransformVertsHelper(QueryGeom, TriIdx, MParticles, MElements, A, B, C);
 
-		const TVector<T, 3> AB = B - A;
-		const TVector<T, 3> AC = C - A;
+		const FVec3 AB = B - A;
+		const FVec3 AC = C - A;
 
 		//It's most likely that the query object is in front of the triangle since queries tend to be on the outside.
 		//However, maybe we should check if it's behind the triangle plane. Also, we should enforce this winding in some way
-		const TVector<T, 3> Offset = TVector<T, 3>::CrossProduct(AB, AC);
+		const FVec3 Offset = FVec3::CrossProduct(AB, AC);
 
-		if (GJKIntersection(TTriangle<T>(A, B, C), InnerQueryGeom, TransformedQueryTM, Thickness, Offset))
+		if (GJKIntersection(TTriangle<FReal>(A, B, C), InnerQueryGeom, TransformedQueryTM, Thickness, Offset))
 		{
 			return true;
 		}
@@ -375,67 +369,57 @@ bool TTriangleMeshImplicitObject<T>::OverlapGeomImp(const QueryGeomType& QueryGe
 	return false;
 }
 
-template <typename T>
-bool TTriangleMeshImplicitObject<T>::OverlapGeom(const TSphere<T, 3>& QueryGeom, const TRigidTransform<T, 3>& QueryTM, const T Thickness) const
+bool FTriangleMeshImplicitObject::OverlapGeom(const TSphere<FReal, 3>& QueryGeom, const FRigidTransform3& QueryTM, const FReal Thickness) const
 {
 	return OverlapGeomImp(QueryGeom, QueryTM, Thickness);
 }
 
-template <typename T>
-bool TTriangleMeshImplicitObject<T>::OverlapGeom(const TBox<T, 3>& QueryGeom, const TRigidTransform<T, 3>& QueryTM, const T Thickness) const
+bool FTriangleMeshImplicitObject::OverlapGeom(const TBox<FReal, 3>& QueryGeom, const FRigidTransform3& QueryTM, const FReal Thickness) const
 {
 	return OverlapGeomImp(QueryGeom, QueryTM, Thickness);
 }
 
-template <typename T>
-bool TTriangleMeshImplicitObject<T>::OverlapGeom(const TCapsule<T>& QueryGeom, const TRigidTransform<T, 3>& QueryTM, const T Thickness) const
+bool FTriangleMeshImplicitObject::OverlapGeom(const TCapsule<FReal>& QueryGeom, const FRigidTransform3& QueryTM, const FReal Thickness) const
 {
 	return OverlapGeomImp(QueryGeom, QueryTM, Thickness);
 }
 
-template <typename T>
-bool TTriangleMeshImplicitObject<T>::OverlapGeom(const FConvex& QueryGeom, const TRigidTransform<T, 3>& QueryTM, const T Thickness) const
+bool FTriangleMeshImplicitObject::OverlapGeom(const FConvex& QueryGeom, const FRigidTransform3& QueryTM, const FReal Thickness) const
 {
 	return OverlapGeomImp(QueryGeom, QueryTM, Thickness);
 }
 
-template <typename T>
-bool TTriangleMeshImplicitObject<T>::OverlapGeom(const TImplicitObjectScaled<TSphere<T, 3>>& QueryGeom, const TRigidTransform<T, 3>& QueryTM, const T Thickness) const
+bool FTriangleMeshImplicitObject::OverlapGeom(const TImplicitObjectScaled<TSphere<FReal, 3>>& QueryGeom, const FRigidTransform3& QueryTM, const FReal Thickness) const
 {
 	return OverlapGeomImp(QueryGeom, QueryTM, Thickness);
 }
 
-template <typename T>
-bool TTriangleMeshImplicitObject<T>::OverlapGeom(const TImplicitObjectScaled<TBox<T, 3>>& QueryGeom, const TRigidTransform<T, 3>& QueryTM, const T Thickness) const
+bool FTriangleMeshImplicitObject::OverlapGeom(const TImplicitObjectScaled<TBox<FReal, 3>>& QueryGeom, const FRigidTransform3& QueryTM, const FReal Thickness) const
 {
 	return OverlapGeomImp(QueryGeom, QueryTM, Thickness);
 }
 
-template <typename T>
-bool TTriangleMeshImplicitObject<T>::OverlapGeom(const TImplicitObjectScaled<TCapsule<T>>& QueryGeom, const TRigidTransform<T, 3>& QueryTM, const T Thickness) const
+bool FTriangleMeshImplicitObject::OverlapGeom(const TImplicitObjectScaled<TCapsule<FReal>>& QueryGeom, const FRigidTransform3& QueryTM, const FReal Thickness) const
 {
 	return OverlapGeomImp(QueryGeom, QueryTM, Thickness);
 }
 
-template <typename T>
-bool TTriangleMeshImplicitObject<T>::OverlapGeom(const TImplicitObjectScaled<FConvex>& QueryGeom, const TRigidTransform<T, 3>& QueryTM, const T Thickness) const
+bool FTriangleMeshImplicitObject::OverlapGeom(const TImplicitObjectScaled<FConvex>& QueryGeom, const FRigidTransform3& QueryTM, const FReal Thickness) const
 {
 	return OverlapGeomImp(QueryGeom, QueryTM, Thickness);
 }
 
-template <typename T>
-bool TTriangleMeshImplicitObject<T>::OverlapGeom(const TImplicitObjectScaled<TImplicitObjectScaled<FConvex>>& QueryGeom, const TRigidTransform<T, 3>& QueryTM, const T Thickness) const
+bool FTriangleMeshImplicitObject::OverlapGeom(const TImplicitObjectScaled<TImplicitObjectScaled<FConvex>>& QueryGeom, const FRigidTransform3& QueryTM, const FReal Thickness) const
 {
 	return OverlapGeomImp(QueryGeom, QueryTM, Thickness);
 }
 
 
-
-template <typename QueryGeomType, typename T>
-struct TTriangleMeshSweepVisitor
+template <typename QueryGeomType>
+struct FTriangleMeshSweepVisitor
 {
-	TTriangleMeshSweepVisitor(const TTriangleMeshImplicitObject<T>& InTriMesh, const QueryGeomType& InQueryGeom, const TRigidTransform<T,3>& InStartTM, const TVector<T,3>& InDir,
-		const TVector<T, 3>& InScaledDirNormalized, const T InLengthScale, const TRigidTransform<T, 3>& InScaledStartTM, const T InThickness, const bool InComputeMTD)
+	FTriangleMeshSweepVisitor(const FTriangleMeshImplicitObject& InTriMesh, const QueryGeomType& InQueryGeom, const TRigidTransform<FReal,3>& InStartTM, const FVec3& InDir,
+		const FVec3& InScaledDirNormalized, const FReal InLengthScale, const FRigidTransform3& InScaledStartTM, const FReal InThickness, const bool InComputeMTD)
 	: TriMesh(InTriMesh)
 	, StartTM(InStartTM)
 	, QueryGeom(InQueryGeom)
@@ -445,7 +429,7 @@ struct TTriangleMeshSweepVisitor
 	, ScaledDirNormalized(InScaledDirNormalized)
 	, LengthScale(InLengthScale)
 	, ScaledStartTM(InScaledStartTM)
-	, OutTime(TNumericLimits<T>::Max())
+	, OutTime(TNumericLimits<FReal>::Max())
 	{
 	}
 
@@ -465,17 +449,17 @@ struct TTriangleMeshSweepVisitor
 	{
 		const int32 TriIdx = VisitData.Payload;
 
-		T Time;
-		TVector<T, 3> HitPosition;
-		TVector<T, 3> HitNormal;
+		FReal Time;
+		FVec3 HitPosition;
+		FVec3 HitNormal;
 
-		TVec3<T> A, B, C;
+		FVec3 A, B, C;
 		TransformVertsHelper(QueryGeom, TriIdx, TriMesh.MParticles, TriMesh.MElements, A, B, C);
-		TTriangle<T> Tri(A, B, C);
+		TTriangle<FReal> Tri(A, B, C);
 
 		const auto& InnerQueryGeom = GetGeomHelper(QueryGeom);
 
-		if(GJKRaycast2<T>(Tri, InnerQueryGeom, ScaledStartTM, ScaledDirNormalized, LengthScale * CurData.CurrentLength, Time, HitPosition, HitNormal, Thickness, bComputeMTD))
+		if(GJKRaycast2<FReal>(Tri, InnerQueryGeom, ScaledStartTM, ScaledDirNormalized, LengthScale * CurData.CurrentLength, Time, HitPosition, HitNormal, Thickness, bComputeMTD))
 		{
 			if(Time < OutTime)
 			{
@@ -498,72 +482,71 @@ struct TTriangleMeshSweepVisitor
 		return true;
 	}
 
-	const TTriangleMeshImplicitObject<T>& TriMesh;
-	const TRigidTransform<T, 3> StartTM;
+	const FTriangleMeshImplicitObject& TriMesh;
+	const FRigidTransform3 StartTM;
 	const QueryGeomType& QueryGeom;
-	const TVector<T, 3>& Dir;
-	const T Thickness;
+	const FVec3& Dir;
+	const FReal Thickness;
 	const bool bComputeMTD;
 
 	// Cache these values for Scaled Triangle Mesh, as they are needed for transformation when sweeping against triangles.
-	TVector<T, 3> ScaledDirNormalized;
-	T LengthScale;
-	TRigidTransform<T, 3> ScaledStartTM;
+	FVec3 ScaledDirNormalized;
+	FReal LengthScale;
+	FRigidTransform3 ScaledStartTM;
 
-	T OutTime;
-	TVector<T, 3> OutPosition;
-	TVector<T, 3> OutNormal;
+	FReal OutTime;
+	FVec3 OutPosition;
+	FVec3 OutNormal;
 	int32 OutFaceIndex;
 };
 
-template <typename QueryGeomType, typename T>
-void ComputeScaledSweepInputs(const QueryGeomType& QueryGeom, const TRigidTransform<T, 3>& StartTM, const TVector<T, 3>& Dir, const T Length,
-	TVector<T, 3>& OutScaledDirNormalized, T& OutLengthScale, TRigidTransform<T, 3>& OutScaledStartTM)
+template <typename QueryGeomType>
+void ComputeScaledSweepInputs(const QueryGeomType& QueryGeom, const FRigidTransform3& StartTM, const FVec3& Dir, const FReal Length,
+	FVec3& OutScaledDirNormalized, FReal& OutLengthScale, FRigidTransform3& OutScaledStartTM)
 {
 	OutScaledDirNormalized = Dir;
 	OutLengthScale = 1.0f;
 	OutScaledStartTM = StartTM;
 }
 
-template<typename QueryGeomType, typename T>
-void ComputeScaledSweepInputs(const TImplicitObjectScaled<QueryGeomType>& QueryGeom, const TRigidTransform<T, 3>& StartTM, const TVector<T, 3>& Dir, const T Length,
-	TVector<T, 3>& OutScaledDirNormalized, T& OutLengthScale, TRigidTransform<T, 3>& OutScaledStartTM)
+template<typename QueryGeomType>
+void ComputeScaledSweepInputs(const TImplicitObjectScaled<QueryGeomType>& QueryGeom, const FRigidTransform3& StartTM, const FVec3& Dir, const FReal Length,
+	FVec3& OutScaledDirNormalized, FReal& OutLengthScale, FRigidTransform3& OutScaledStartTM)
 {
-	const TVector<T, 3>& InvScale = QueryGeom.GetInvScale();
+	const FVec3& InvScale = QueryGeom.GetInvScale();
 
-	const TVector<T, 3> UnscaledDirDenorm = InvScale * Dir;
-	const T LengthScale = UnscaledDirDenorm.Size();
-	if (CHAOS_ENSURE(LengthScale > TNumericLimits<T>::Min()))
+	const FVec3 UnscaledDirDenorm = InvScale * Dir;
+	const FReal LengthScale = UnscaledDirDenorm.Size();
+	if (CHAOS_ENSURE(LengthScale > TNumericLimits<FReal>::Min()))
 	{
-		const T LengthScaleInv = 1.f / LengthScale;
+		const FReal LengthScaleInv = 1.f / LengthScale;
 		OutScaledDirNormalized = UnscaledDirDenorm * LengthScaleInv;
 	}
 
 
 	OutLengthScale = LengthScale;
-	OutScaledStartTM = TRigidTransform<T, 3>(StartTM.GetLocation() * InvScale, StartTM.GetRotation());
+	OutScaledStartTM = FRigidTransform3(StartTM.GetLocation() * InvScale, StartTM.GetRotation());
 }
 
-template <typename T>
 template <typename QueryGeomType>
-bool TTriangleMeshImplicitObject<T>::SweepGeomImp(const QueryGeomType& QueryGeom, const TRigidTransform<T, 3>& StartTM, const TVector<T, 3>& Dir, const T Length,
-	T& OutTime, TVector<T, 3>& OutPosition, TVector<T, 3>& OutNormal, int32& OutFaceIndex, const T Thickness, const bool bComputeMTD) const
+bool FTriangleMeshImplicitObject::SweepGeomImp(const QueryGeomType& QueryGeom, const FRigidTransform3& StartTM, const FVec3& Dir, const FReal Length,
+	FReal& OutTime, FVec3& OutPosition, FVec3& OutNormal, int32& OutFaceIndex, const FReal Thickness, const bool bComputeMTD) const
 {
 
 	// Compute scaled sweep inputs to cache in visitor.
-	TVector<T, 3> ScaledDirNormalized;
-	T LengthScale;
-	TRigidTransform<T, 3> ScaledStartTM;
+	FVec3 ScaledDirNormalized;
+	FReal LengthScale;
+	FRigidTransform3 ScaledStartTM;
 	ComputeScaledSweepInputs(QueryGeom, StartTM, Dir, Length, ScaledDirNormalized, LengthScale, ScaledStartTM);
 
 	bool bHit = false;
-	TTriangleMeshSweepVisitor<QueryGeomType, T> SQVisitor(*this, QueryGeom, StartTM, Dir, ScaledDirNormalized, LengthScale, ScaledStartTM, Thickness, bComputeMTD);
+	FTriangleMeshSweepVisitor<QueryGeomType> SQVisitor(*this, QueryGeom, StartTM, Dir, ScaledDirNormalized, LengthScale, ScaledStartTM, Thickness, bComputeMTD);
 
 
-	const TAABB<T, 3> QueryBounds = QueryGeom.BoundingBox();
-	const TVector<T, 3> StartPoint = StartTM.TransformPositionNoScale(QueryBounds.Center());
-	const TVector<T, 3> Inflation = QueryBounds.Extents() * 0.5 + TVector<T, 3>(Thickness);
-	BVH.template Sweep<TTriangleMeshSweepVisitor<QueryGeomType, T>>(StartPoint, Dir, Length, Inflation, SQVisitor);
+	const TAABB<FReal, 3> QueryBounds = QueryGeom.BoundingBox();
+	const FVec3 StartPoint = StartTM.TransformPositionNoScale(QueryBounds.Center());
+	const FVec3 Inflation = QueryBounds.Extents() * 0.5 + FVec3(Thickness);
+	BVH.template Sweep<FTriangleMeshSweepVisitor<QueryGeomType>>(StartPoint, Dir, Length, Inflation, SQVisitor);
 
 	if (SQVisitor.OutTime <= Length)
 	{
@@ -576,90 +559,81 @@ bool TTriangleMeshImplicitObject<T>::SweepGeomImp(const QueryGeomType& QueryGeom
 	return bHit;
 }
 
-template <typename T>
-bool TTriangleMeshImplicitObject<T>::SweepGeom(const TSphere<T,3>& QueryGeom, const TRigidTransform<T, 3>& StartTM, const TVector<T, 3>& Dir, const T Length, T& OutTime, TVector<T, 3>& OutPosition, TVector<T, 3>& OutNormal, int32& OutFaceIndex, const T Thickness, const bool bComputeMTD) const
+bool FTriangleMeshImplicitObject::SweepGeom(const TSphere<FReal,3>& QueryGeom, const FRigidTransform3& StartTM, const FVec3& Dir, const FReal Length, FReal& OutTime, FVec3& OutPosition, FVec3& OutNormal, int32& OutFaceIndex, const FReal Thickness, const bool bComputeMTD) const
 {
 	return SweepGeomImp(QueryGeom, StartTM, Dir, Length, OutTime, OutPosition, OutNormal, OutFaceIndex, Thickness, bComputeMTD);
 }
 
-template <typename T>
-bool TTriangleMeshImplicitObject<T>::SweepGeom(const TBox<T, 3>& QueryGeom, const TRigidTransform<T, 3>& StartTM, const TVector<T, 3>& Dir, const T Length, T& OutTime, TVector<T, 3>& OutPosition, TVector<T, 3>& OutNormal, int32& OutFaceIndex, const T Thickness, const bool bComputeMTD) const
+bool FTriangleMeshImplicitObject::SweepGeom(const TBox<FReal, 3>& QueryGeom, const FRigidTransform3& StartTM, const FVec3& Dir, const FReal Length, FReal& OutTime, FVec3& OutPosition, FVec3& OutNormal, int32& OutFaceIndex, const FReal Thickness, const bool bComputeMTD) const
 {
 	return SweepGeomImp(QueryGeom, StartTM, Dir, Length, OutTime, OutPosition, OutNormal, OutFaceIndex, Thickness, bComputeMTD);
 }
 
-template <typename T>
-bool TTriangleMeshImplicitObject<T>::SweepGeom(const TCapsule<T>& QueryGeom, const TRigidTransform<T, 3>& StartTM, const TVector<T, 3>& Dir, const T Length, T& OutTime, TVector<T, 3>& OutPosition, TVector<T, 3>& OutNormal, int32& OutFaceIndex, const T Thickness, const bool bComputeMTD) const
+bool FTriangleMeshImplicitObject::SweepGeom(const TCapsule<FReal>& QueryGeom, const FRigidTransform3& StartTM, const FVec3& Dir, const FReal Length, FReal& OutTime, FVec3& OutPosition, FVec3& OutNormal, int32& OutFaceIndex, const FReal Thickness, const bool bComputeMTD) const
 {
 	return SweepGeomImp(QueryGeom, StartTM, Dir, Length, OutTime, OutPosition, OutNormal, OutFaceIndex, Thickness, bComputeMTD);
 }
 
-template <typename T>
-bool TTriangleMeshImplicitObject<T>::SweepGeom(const FConvex& QueryGeom, const TRigidTransform<T, 3>& StartTM, const TVector<T, 3>& Dir, const T Length, T& OutTime, TVector<T, 3>& OutPosition, TVector<T, 3>& OutNormal, int32& OutFaceIndex, const T Thickness, const bool bComputeMTD) const
+bool FTriangleMeshImplicitObject::SweepGeom(const FConvex& QueryGeom, const FRigidTransform3& StartTM, const FVec3& Dir, const FReal Length, FReal& OutTime, FVec3& OutPosition, FVec3& OutNormal, int32& OutFaceIndex, const FReal Thickness, const bool bComputeMTD) const
 {
 	return SweepGeomImp(QueryGeom, StartTM, Dir, Length, OutTime, OutPosition, OutNormal, OutFaceIndex, Thickness, bComputeMTD);
 }
 
-template <typename T>
-bool TTriangleMeshImplicitObject<T>::SweepGeom(const TImplicitObjectScaled<TSphere<T, 3>>& QueryGeom, const TRigidTransform<T, 3>& StartTM, const TVector<T, 3>& Dir, const T Length, T& OutTime, TVector<T, 3>& OutPosition, TVector<T, 3>& OutNormal, int32& OutFaceIndex, const T Thickness, const bool bComputeMTD) const
+bool FTriangleMeshImplicitObject::SweepGeom(const TImplicitObjectScaled<TSphere<FReal, 3>>& QueryGeom, const FRigidTransform3& StartTM, const FVec3& Dir, const FReal Length, FReal& OutTime, FVec3& OutPosition, FVec3& OutNormal, int32& OutFaceIndex, const FReal Thickness, const bool bComputeMTD) const
 {
 	return SweepGeomImp(QueryGeom, StartTM, Dir, Length, OutTime, OutPosition, OutNormal, OutFaceIndex, Thickness, bComputeMTD);
 }
 
-template <typename T>
-bool TTriangleMeshImplicitObject<T>::SweepGeom(const TImplicitObjectScaled<TBox<T, 3>>& QueryGeom, const TRigidTransform<T, 3>& StartTM, const TVector<T, 3>& Dir, const T Length, T& OutTime, TVector<T, 3>& OutPosition, TVector<T, 3>& OutNormal, int32& OutFaceIndex, const T Thickness, const bool bComputeMTD) const
+bool FTriangleMeshImplicitObject::SweepGeom(const TImplicitObjectScaled<TBox<FReal, 3>>& QueryGeom, const FRigidTransform3& StartTM, const FVec3& Dir, const FReal Length, FReal& OutTime, FVec3& OutPosition, FVec3& OutNormal, int32& OutFaceIndex, const FReal Thickness, const bool bComputeMTD) const
 {
 	return SweepGeomImp(QueryGeom, StartTM, Dir, Length, OutTime, OutPosition, OutNormal, OutFaceIndex, Thickness, bComputeMTD);
 }
 
-template <typename T>
-bool TTriangleMeshImplicitObject<T>::SweepGeom(const TImplicitObjectScaled<TCapsule<T>>& QueryGeom, const TRigidTransform<T, 3>& StartTM, const TVector<T, 3>& Dir, const T Length, T& OutTime, TVector<T, 3>& OutPosition, TVector<T, 3>& OutNormal, int32& OutFaceIndex, const T Thickness, const bool bComputeMTD) const
+bool FTriangleMeshImplicitObject::SweepGeom(const TImplicitObjectScaled<TCapsule<FReal>>& QueryGeom, const FRigidTransform3& StartTM, const FVec3& Dir, const FReal Length, FReal& OutTime, FVec3& OutPosition, FVec3& OutNormal, int32& OutFaceIndex, const FReal Thickness, const bool bComputeMTD) const
 {
 	return SweepGeomImp(QueryGeom, StartTM, Dir, Length, OutTime, OutPosition, OutNormal, OutFaceIndex, Thickness, bComputeMTD);
 }
 
-template <typename T>
-bool TTriangleMeshImplicitObject<T>::SweepGeom(const TImplicitObjectScaled<FConvex>& QueryGeom, const TRigidTransform<T, 3>& StartTM, const TVector<T, 3>& Dir, const T Length, T& OutTime, TVector<T, 3>& OutPosition, TVector<T, 3>& OutNormal, int32& OutFaceIndex, const T Thickness, const bool bComputeMTD) const
+bool FTriangleMeshImplicitObject::SweepGeom(const TImplicitObjectScaled<FConvex>& QueryGeom, const FRigidTransform3& StartTM, const FVec3& Dir, const FReal Length, FReal& OutTime, FVec3& OutPosition, FVec3& OutNormal, int32& OutFaceIndex, const FReal Thickness, const bool bComputeMTD) const
 {
 	return SweepGeomImp(QueryGeom, StartTM, Dir, Length, OutTime, OutPosition, OutNormal, OutFaceIndex, Thickness, bComputeMTD);
 }
 
-template <typename T>
-int32 TTriangleMeshImplicitObject<T>::FindMostOpposingFace(const TVector<T, 3>& Position, const TVector<T, 3>& UnitDir, int32 HintFaceIndex, T SearchDist) const
+int32 FTriangleMeshImplicitObject::FindMostOpposingFace(const FVec3& Position, const FVec3& UnitDir, int32 HintFaceIndex, FReal SearchDist) const
 {
 	//todo: this is horribly slow, need adjacency information
-	const T SearchDist2 = SearchDist * SearchDist;
+	const FReal SearchDist2 = SearchDist * SearchDist;
 	
-	TAABB<T, 3> QueryBounds(Position - TVector<T,3>(SearchDist), Position + TVector<T,3>(SearchDist));
+	TAABB<FReal, 3> QueryBounds(Position - FVec3(SearchDist), Position + FVec3(SearchDist));
 
 	const TArray<int32> PotentialIntersections = BVH.FindAllIntersections(QueryBounds);
-	const T Epsilon = 1e-4;
+	const FReal Epsilon = 1e-4;
 
-	T MostOpposingDot = TNumericLimits<T>::Max();
+	FReal MostOpposingDot = TNumericLimits<FReal>::Max();
 	int32 MostOpposingFace = HintFaceIndex;
 	
 	for (int32 TriIdx : PotentialIntersections)
 	{
-		const TVector<T, 3>& A = MParticles.X(MElements[TriIdx][0]);
-		const TVector<T, 3>& B = MParticles.X(MElements[TriIdx][1]);
-		const TVector<T, 3>& C = MParticles.X(MElements[TriIdx][2]);
-
-		const TVector<T, 3> AB = B - A;
-		const TVector<T, 3> AC = C - A;
-		TVector<T, 3> Normal = TVector<T, 3>::CrossProduct(AB, AC);
-		const T NormalLength = Normal.SafeNormalize();
+		const FVec3& A = MParticles.X(MElements[TriIdx][0]);
+		const FVec3& B = MParticles.X(MElements[TriIdx][1]);
+		const FVec3& C = MParticles.X(MElements[TriIdx][2]);
+
+		const FVec3 AB = B - A;
+		const FVec3 AC = C - A;
+		FVec3 Normal = FVec3::CrossProduct(AB, AC);
+		const FReal NormalLength = Normal.SafeNormalize();
 		if (!ensure(NormalLength > Epsilon))
 		{
 			//hitting degenerate triangle - should be fixed before we get to this stage
 			continue;
 		}
 
-		const TPlane<T, 3> TriPlane{ A, Normal };
-		const TVector<T, 3> ClosestPointOnTri = FindClosestPointOnTriangle(TriPlane, A, B, C, Position);
-		const T Distance2 = (ClosestPointOnTri - Position).SizeSquared();
+		const TPlane<FReal, 3> TriPlane{ A, Normal };
+		const FVec3 ClosestPointOnTri = FindClosestPointOnTriangle(TriPlane, A, B, C, Position);
+		const FReal Distance2 = (ClosestPointOnTri - Position).SizeSquared();
 		if (Distance2 < SearchDist2)
 		{
-			const T Dot = TVector<T, 3>::DotProduct(Normal, UnitDir);
+			const FReal Dot = FVec3::DotProduct(Normal, UnitDir);
 			if (Dot < MostOpposingDot)
 			{
 				MostOpposingDot = Dot;
@@ -671,21 +645,18 @@ int32 TTriangleMeshImplicitObject<T>::FindMostOpposingFace(const TVector<T, 3>&
 	return MostOpposingFace;
 }
 
-template <typename T>
-TVector<T, 3> TTriangleMeshImplicitObject<T>::FindGeometryOpposingNormal(const TVector<T, 3>& DenormDir, int32 FaceIndex, const TVector<T, 3>& OriginalNormal) const
+FVec3 FTriangleMeshImplicitObject::FindGeometryOpposingNormal(const FVec3& DenormDir, int32 FaceIndex, const FVec3& OriginalNormal) const
 {
 	return GetFaceNormal(FaceIndex);
 }
 
-template<typename T>
-void TTriangleMeshImplicitObject<T>::Serialize(FChaosArchive& Ar)
+void FTriangleMeshImplicitObject::Serialize(FChaosArchive& Ar)
 {
 	FChaosArchiveScopedMemory ScopedMemory(Ar, GetTypeName());
 	SerializeImp(Ar);
 }
 
-template<typename T>
-uint32 TTriangleMeshImplicitObject<T>::GetTypeHash() const
+uint32 FTriangleMeshImplicitObject::GetTypeHash() const
 {
 	uint32 Result = MParticles.GetTypeHash();
 	Result = HashCombine(Result, MLocalBoundingBox.GetTypeHash());
@@ -699,28 +670,26 @@ uint32 TTriangleMeshImplicitObject<T>::GetTypeHash() const
 	return Result;
 }
 
-template <typename T>
-TVector<T, 3> TTriangleMeshImplicitObject<T>::GetFaceNormal(const int32 FaceIdx) const
+FVec3 FTriangleMeshImplicitObject::GetFaceNormal(const int32 FaceIdx) const
 {
 	if (ensure(FaceIdx != INDEX_NONE))
 	{
-		const TVector<T, 3>& A = MParticles.X(MElements[FaceIdx][0]);
-		const TVector<T, 3>& B = MParticles.X(MElements[FaceIdx][1]);
-		const TVector<T, 3>& C = MParticles.X(MElements[FaceIdx][2]);
-
-		const TVector<T, 3> AB = B - A;
-		const TVector<T, 3> AC = C - A;
-		TVector<T, 3> Normal = TVector<T, 3>::CrossProduct(AB, AC);
-		const T Length = Normal.SafeNormalize();
+		const FVec3& A = MParticles.X(MElements[FaceIdx][0]);
+		const FVec3& B = MParticles.X(MElements[FaceIdx][1]);
+		const FVec3& C = MParticles.X(MElements[FaceIdx][2]);
+
+		const FVec3 AB = B - A;
+		const FVec3 AC = C - A;
+		FVec3 Normal = FVec3::CrossProduct(AB, AC);
+		const FReal Length = Normal.SafeNormalize();
 		ensure(Length);
 		return Normal;
 	}
 
-	return TVector<T, 3>(0, 0, 1);
+	return FVec3(0, 0, 1);
 }
 
-template <typename T>
-uint16 TTriangleMeshImplicitObject<T>::GetMaterialIndex(uint32 HintIndex) const
+uint16 FTriangleMeshImplicitObject::GetMaterialIndex(uint32 HintIndex) const
 {
 	if (MaterialIndices.IsValidIndex(HintIndex))
 	{
@@ -731,8 +700,7 @@ uint16 TTriangleMeshImplicitObject<T>::GetMaterialIndex(uint32 HintIndex) const
 	return 0;
 }
 
-template<typename T>
-void Chaos::TTriangleMeshImplicitObject<T>::RebuildBV()
+void Chaos::FTriangleMeshImplicitObject::RebuildBV()
 {
 	const int32 NumTris = MElements.Num();
 	BVEntries.Reset(NumTris);
@@ -747,6 +715,3 @@ void Chaos::TTriangleMeshImplicitObject<T>::RebuildBV()
 
 
 }
-
-
-template class Chaos::TTriangleMeshImplicitObject<float>;
\ No newline at end of file
diff --git a/Engine/Source/Runtime/Experimental/Chaos/Public/Chaos/TriangleMeshImplicitObject.h b/Engine/Source/Runtime/Experimental/Chaos/Public/Chaos/TriangleMeshImplicitObject.h
--- a/Engine/Source/Runtime/Experimental/Chaos/Public/Chaos/TriangleMeshImplicitObject.h
+++ b/Engine/Source/Runtime/Experimental/Chaos/Public/Chaos/TriangleMeshImplicitObject.h
@@ -22,45 +22,44 @@ namespace Chaos
 
 	class FConvex;
 
-	template<typename T>
-	class CHAOS_API TTriangleMeshImplicitObject final : public FImplicitObject
+	class CHAOS_API FTriangleMeshImplicitObject final : public FImplicitObject
 	{
 	public:
 		using FImplicitObject::GetTypeName;
 
-		TTriangleMeshImplicitObject(TParticles<T,3>&& Particles, TArray<TVector<int32, 3>>&& Elements, TArray<uint16>&& InMaterialIndices);
-		TTriangleMeshImplicitObject(const TTriangleMeshImplicitObject& Other) = delete;
-		TTriangleMeshImplicitObject(TTriangleMeshImplicitObject&& Other) = default;
-		virtual ~TTriangleMeshImplicitObject() {}
+		FTriangleMeshImplicitObject(TParticles<FReal,3>&& Particles, TArray<TVector<int32, 3>>&& Elements, TArray<uint16>&& InMaterialIndices);
+		FTriangleMeshImplicitObject(const FTriangleMeshImplicitObject& Other) = delete;
+		FTriangleMeshImplicitObject(FTriangleMeshImplicitObject&& Other) = default;
+		virtual ~FTriangleMeshImplicitObject() {}
 
-		virtual T PhiWithNormal(const TVector<T, 3>& x, TVector<T, 3>& Normal) const;
+		virtual FReal PhiWithNormal(const FVec3& x, FVec3& Normal) const;
 		
-		virtual bool Raycast(const TVector<T, 3>& StartPoint, const TVector<T, 3>& Dir, const T Length, const T Thickness, T& OutTime, TVector<T,3>& OutPosition, TVector<T,3>& OutNormal, int32& OutFaceIndex) const override;
-		virtual bool Overlap(const TVector<T, 3>& Point, const T Thickness) const override;
-
-		bool OverlapGeom(const TSphere<T,3>& QueryGeom, const TRigidTransform<T, 3>& QueryTM, const T Thickness) const;
-		bool OverlapGeom(const TBox<T, 3>& QueryGeom, const TRigidTransform<T, 3>& QueryTM, const T Thickness) const;
-		bool OverlapGeom(const TCapsule<T>& QueryGeom, const TRigidTransform<T, 3>& QueryTM, const T Thickness) const;
-		bool OverlapGeom(const FConvex& QueryGeom, const TRigidTransform<T, 3>& QueryTM, const T Thickness) const;
-		bool OverlapGeom(const TImplicitObjectScaled<TSphere<T, 3>>& QueryGeom, const TRigidTransform<T, 3>& QueryTM, const T Thickness) const;
-		bool OverlapGeom(const TImplicitObjectScaled<TBox<T, 3>>& QueryGeom, const TRigidTransform<T, 3>& QueryTM, const T Thickness) const;
-		bool OverlapGeom(const TImplicitObjectScaled<TCapsule<T>>& QueryGeom, const TRigidTransform<T, 3>& QueryTM, const T Thickness) const;
-		bool OverlapGeom(const TImplicitObjectScaled<FConvex>& QueryGeom, const TRigidTransform<T, 3>& QueryTM, const T Thickness) const;
-		bool OverlapGeom(const TImplicitObjectScaled<TImplicitObjectScaled<FConvex>>& QueryGeom, const TRigidTransform<T, 3>& QueryTM, const T Thickness) const;
-
-		bool SweepGeom(const TSphere<T,3>& QueryGeom, const TRigidTransform<T, 3>& StartTM, const TVector<T,3>& Dir, const T Length, T& OutTime, TVector<T,3>& OutPosition, TVector<T,3>& OutNormal, int32& OutFaceIndex, const T Thickness = 0, const bool bComputeMTD = false) const;
-		bool SweepGeom(const TBox<T, 3>& QueryGeom, const TRigidTransform<T, 3>& StartTM, const TVector<T, 3>& Dir, const T Length, T& OutTime, TVector<T, 3>& OutPosition, TVector<T, 3>& OutNormal, int32& OutFaceIndex, const T Thickness = 0, const bool bComputeMTD = false) const;
-		bool SweepGeom(const TCapsule<T>& QueryGeom, const TRigidTransform<T, 3>& StartTM, const TVector<T, 3>& Dir, const T Length, T& OutTime, TVector<T, 3>& OutPosition, TVector<T, 3>& OutNormal, int32& OutFaceIndex, const T Thickness = 0, const bool bComputeMTD = false) const;
-		bool SweepGeom(const FConvex& QueryGeom, const TRigidTransform<T, 3>& StartTM, const TVector<T, 3>& Dir, const T Length, T& OutTime, TVector<T, 3>& OutPosition, TVector<T, 3>& OutNormal, int32& OutFaceIndex, const T Thickness = 0, const bool bComputeMTD = false) const;
-		bool SweepGeom(const TImplicitObjectScaled<TSphere<T, 3>>& QueryGeom, const TRigidTransform<T, 3>& StartTM, const TVector<T, 3>& Dir, const T Length, T& OutTime, TVector<T, 3>& OutPosition, TVector<T, 3>& OutNormal, int32& OutFaceIndex, const T Thickness = 0, const bool bComputeMTD = false) const;
-		bool SweepGeom(const TImplicitObjectScaled<TBox<T, 3>>& QueryGeom, const TRigidTransform<T, 3>& StartTM, const TVector<T, 3>& Dir, const T Length, T& OutTime, TVector<T, 3>& OutPosition, TVector<T, 3>& OutNormal, int32& OutFaceIndex, const T Thickness = 0, const bool bComputeMTD = false) const;
-		bool SweepGeom(const TImplicitObjectScaled<TCapsule<T>>& QueryGeom, const TRigidTransform<T, 3>& StartTM, const TVector<T, 3>& Dir, const T Length, T& OutTime, TVector<T, 3>& OutPosition, TVector<T, 3>& OutNormal, int32& OutFaceIndex, const T Thickness = 0, const bool bComputeMTD = false) const;
-		bool SweepGeom(const TImplicitObjectScaled<FConvex>& QueryGeom, const TRigidTransform<T, 3>& StartTM, const TVector<T, 3>& Dir, const T Length, T& OutTime, TVector<T, 3>& OutPosition, TVector<T, 3>& OutNormal, int32& OutFaceIndex, const T Thickness = 0, const bool bComputeMTD = false) const;
-
-		virtual int32 FindMostOpposingFace(const TVector<T, 3>& Position, const TVector<T, 3>& UnitDir, int32 HintFaceIndex, T SearchDistance) const override;
-		virtual TVector<T, 3> FindGeometryOpposingNormal(const TVector<T, 3>& DenormDir, int32 FaceIndex, const TVector<T, 3>& OriginalNormal) const override;
-
-		virtual const TAABB<T, 3>& BoundingBox() const
+		virtual bool Raycast(const FVec3& StartPoint, const FVec3& Dir, const FReal Length, const FReal Thickness, FReal& OutTime, FVec3& OutPosition, FVec3& OutNormal, int32& OutFaceIndex) const override;
+		virtual bool Overlap(const FVec3& Point, const FReal Thickness) const override;
+
+		bool OverlapGeom(const TSphere<FReal,3>& QueryGeom, const FRigidTransform3& QueryTM, const FReal Thickness) const;
+		bool OverlapGeom(const TBox<FReal, 3>& QueryGeom, const FRigidTransform3& QueryTM, const FReal Thickness) const;
+		bool OverlapGeom(const TCapsule<FReal>& QueryGeom, const FRigidTransform3& QueryTM, const FReal Thickness) const;
+		bool OverlapGeom(const FConvex& QueryGeom, const FRigidTransform3& QueryTM, const FReal Thickness) const;
+		bool OverlapGeom(const TImplicitObjectScaled<TSphere<FReal, 3>>& QueryGeom, const FRigidTransform3& QueryTM, const FReal Thickness) const;
+		bool OverlapGeom(const TImplicitObjectScaled<TBox<FReal, 3>>& QueryGeom, const FRigidTransform3& QueryTM, const FReal Thickness) const;
+		bool OverlapGeom(const TImplicitObjectScaled<TCapsule<FReal>>& QueryGeom, const FRigidTransform3& QueryTM, const FReal Thickness) const;
+		bool OverlapGeom(const TImplicitObjectScaled<FConvex>& QueryGeom, const FRigidTransform3& QueryTM, const FReal Thickness) const;
+		bool OverlapGeom(const TImplicitObjectScaled<TImplicitObjectScaled<FConvex>>& QueryGeom, const FRigidTransform3& QueryTM, const FReal Thickness) const;
+
+		bool SweepGeom(const TSphere<FReal,3>& QueryGeom, const FRigidTransform3& StartTM, const FVec3& Dir, const FReal Length, FReal& OutTime, FVec3& OutPosition, FVec3& OutNormal, int32& OutFaceIndex, const FReal Thickness = 0, const bool bComputeMTD = false) const;
+		bool SweepGeom(const TBox<FReal, 3>& QueryGeom, const FRigidTransform3& StartTM, const FVec3& Dir, const FReal Length, FReal& OutTime, FVec3& OutPosition, FVec3& OutNormal, int32& OutFaceIndex, const FReal Thickness = 0, const bool bComputeMTD = false) const;
+		bool SweepGeom(const TCapsule<FReal>& QueryGeom, const FRigidTransform3& StartTM, const FVec3& Dir, const FReal Length, FReal& OutTime, FVec3& OutPosition, FVec3& OutNormal, int32& OutFaceIndex, const FReal Thickness = 0, const bool bComputeMTD = false) const;
+		bool SweepGeom(const FConvex& QueryGeom, const FRigidTransform3& StartTM, const FVec3& Dir, const FReal Length, FReal& OutTime, FVec3& OutPosition, FVec3& OutNormal, int32& OutFaceIndex, const FReal Thickness = 0, const bool bComputeMTD = false) const;
+		bool SweepGeom(const TImplicitObjectScaled<TSphere<FReal, 3>>& QueryGeom, const FRigidTransform3& StartTM, const FVec3& Dir, const FReal Length, FReal& OutTime, FVec3& OutPosition, FVec3& OutNormal, int32& OutFaceIndex, const FReal Thickness = 0, const bool bComputeMTD = false) const;
+		bool SweepGeom(const TImplicitObjectScaled<TBox<FReal, 3>>& QueryGeom, const FRigidTransform3& StartTM, const FVec3& Dir, const FReal Length, FReal& OutTime, FVec3& OutPosition, FVec3& OutNormal, int32& OutFaceIndex, const FReal Thickness = 0, const bool bComputeMTD = false) const;
+		bool SweepGeom(const TImplicitObjectScaled<TCapsule<FReal>>& QueryGeom, const FRigidTransform3& StartTM, const FVec3& Dir, const FReal Length, FReal& OutTime, FVec3& OutPosition, FVec3& OutNormal, int32& OutFaceIndex, const FReal Thickness = 0, const bool bComputeMTD = false) const;
+		bool SweepGeom(const TImplicitObjectScaled<FConvex>& QueryGeom, const FRigidTransform3& StartTM, const FVec3& Dir, const FReal Length, FReal& OutTime, FVec3& OutPosition, FVec3& OutNormal, int32& OutFaceIndex, const FReal Thickness = 0, const bool bComputeMTD = false) const;
+
+		virtual int32 FindMostOpposingFace(const FVec3& Position, const FVec3& UnitDir, int32 HintFaceIndex, FReal SearchDistance) const override;
+		virtual FVec3 FindGeometryOpposingNormal(const FVec3& DenormDir, int32 FaceIndex, const FVec3& OriginalNormal) const override;
+
+		virtual const TAABB<FReal, 3>& BoundingBox() const
 		{
 			return MLocalBoundingBox;
 		}
@@ -77,11 +76,11 @@ namespace Chaos
 			FImplicitObject::SerializeImp(Ar);
 			Ar << MParticles;
 			Ar << MElements;
-			TBox<T, 3>::SerializeAsAABB(Ar, MLocalBoundingBox);
+			TBox<FReal, 3>::SerializeAsAABB(Ar, MLocalBoundingBox);
 
 			if (Ar.CustomVer(FExternalPhysicsCustomObjectVersion::GUID) < FExternalPhysicsCustomObjectVersion::RemovedConvexHullsFromTriangleMeshImplicitObject)
 			{
-				TUniquePtr<TGeometryParticles<T, 3>> ConvexHulls;
+				TUniquePtr<TGeometryParticles<FReal, 3>> ConvexHulls;
 				Ar << ConvexHulls;
 			}
 
@@ -92,7 +91,7 @@ namespace Chaos
 			}
 			else if(Ar.CustomVer(FExternalPhysicsCustomObjectVersion::GUID) < FExternalPhysicsCustomObjectVersion::TrimeshSerializesAABBTree)
 			{
-				TBoundingVolume<int32, T, 3> Dummy;
+				TBoundingVolume<int32, FReal, 3> Dummy;
 				Ar << Dummy;
 				RebuildBV();
 			}
@@ -112,7 +111,7 @@ namespace Chaos
 
 		virtual uint32 GetTypeHash() const override;
 
-		TVector<T, 3> GetFaceNormal(const int32 FaceIdx) const;
+		FVec3 GetFaceNormal(const int32 FaceIdx) const;
 
 		virtual uint16 GetMaterialIndex(uint32 HintIndex) const override;
 
@@ -120,27 +119,27 @@ namespace Chaos
 
 		void RebuildBV();
 
-		TParticles<T, 3> MParticles;
+		TParticles<FReal, 3> MParticles;
 		TArray<TVector<int32, 3>> MElements;
-		TAABB<T, 3> MLocalBoundingBox;
+		TAABB<FReal, 3> MLocalBoundingBox;
 		TArray<uint16> MaterialIndices;
 
 		//using BVHType = TBoundingVolume<int32, T, 3>;
-		using BVHType = TAABBTree<int32, TAABBTreeLeafArray<int32, T>, T>;
+		using BVHType = TAABBTree<int32, TAABBTreeLeafArray<int32, FReal>, FReal>;
 
 		template<typename InStorageType, typename InRealType>
 		friend struct FBvEntry;
 
 		struct FBvEntry
 		{
-			TTriangleMeshImplicitObject* TmData;
+			FTriangleMeshImplicitObject* TmData;
 			int32 Index;
 
 			bool HasBoundingBox() const { return true; }
 
-			TAABB<T, 3> BoundingBox() const
+			TAABB<FReal, 3> BoundingBox() const
 			{
-				TAABB<T, 3> Bounds(TmData->MParticles.X(TmData->MElements[Index][0]), TmData->MParticles.X(TmData->MElements[Index][0]));
+				TAABB<FReal, 3> Bounds(TmData->MParticles.X(TmData->MElements[Index][0]), TmData->MParticles.X(TmData->MElements[Index][0]));
 
 				Bounds.GrowToInclude(TmData->MParticles.X(TmData->MElements[Index][1]));
 				Bounds.GrowToInclude(TmData->MParticles.X(TmData->MElements[Index][2]));
@@ -149,7 +148,7 @@ namespace Chaos
 			}
 
 			template <typename TPayloadType>
-			int32 GetPayload(int32 Idx) const
+			uint16 GetPayload(uint16 Idx) const
 			{
 				return Idx;
 			}
@@ -158,21 +157,21 @@ namespace Chaos
 
 		BVHType BVH;
 
-		template <typename Geom, typename R>
-		friend struct TTriangleMeshSweepVisitor;
+		template <typename Geom>
+		friend struct FTriangleMeshSweepVisitor;
 
 		// Required by implicit object serialization, disabled for general use.
 		friend class FImplicitObject;
 
-		TTriangleMeshImplicitObject() 
+		FTriangleMeshImplicitObject() 
 			: FImplicitObject(EImplicitObject::HasBoundingBox, ImplicitObjectType::TriangleMesh)
 		{};
 
 		template <typename QueryGeomType>
-		bool OverlapGeomImp(const QueryGeomType& QueryGeom, const TRigidTransform<T, 3>& QueryTM, const T Thickness) const;
+		bool OverlapGeomImp(const QueryGeomType& QueryGeom, const FRigidTransform3& QueryTM, const FReal Thickness) const;
 
 		template <typename QueryGeomType>
-		bool SweepGeomImp(const QueryGeomType& QueryGeom, const TRigidTransform<T, 3>& StartTM, const TVector<T, 3>& Dir, const T Length, T& OutTime, TVector<T, 3>& OutPosition, TVector<T, 3>& OutNormal, int32& OutFaceIndex, const T Thickness, const bool bComputeMTD) const;
+		bool SweepGeomImp(const QueryGeomType& QueryGeom, const FRigidTransform3& StartTM, const FVec3& Dir, const FReal Length, FReal& OutTime, FVec3& OutPosition, FVec3& OutNormal, int32& OutFaceIndex, const FReal Thickness, const bool bComputeMTD) const;
 	};
 }
 
diff --git a/Engine/Source/Runtime/PhysicsCore/Public/PhysXToChaosUtil.h b/Engine/Source/Runtime/PhysicsCore/Public/PhysXToChaosUtil.h
--- a/Engine/Source/Runtime/PhysicsCore/Public/PhysXToChaosUtil.h
+++ b/Engine/Source/Runtime/PhysicsCore/Public/PhysXToChaosUtil.h
@@ -178,7 +178,7 @@ inline TUniquePtr<Chaos::TImplicitObjectTransformed<float, 3>> PxShapeToChaosGeo
 			}
 		}
 
-		TUniquePtr<TTriangleMeshImplicitObject<float>> TriMeshObj = MakeUnique<TTriangleMeshImplicitObject<float>>(MoveTemp(Particles), MoveTemp(Triangles), MoveTemp(MaterialIndices));
+		TUniquePtr<FTriangleMeshImplicitObject> TriMeshObj = MakeUnique<FTriangleMeshImplicitObject>(MoveTemp(Particles), MoveTemp(Triangles), MoveTemp(MaterialIndices));
 		if (TriMeshGeom.scale.isIdentity())
 		{
 			InnerObj = MoveTemp(TriMeshObj);
@@ -186,7 +186,7 @@ inline TUniquePtr<Chaos::TImplicitObjectTransformed<float, 3>> PxShapeToChaosGeo
 		else
 		{
 			ensure(TriMeshGeom.scale.rotation == PxQuat(PxIdentity));
-			InnerObj = MakeUnique<TImplicitObjectScaled<TTriangleMeshImplicitObject<float>, /*bInstanced=*/false>>(MoveTemp(TriMeshObj), P2UVector(TriMeshGeom.scale.scale));	//todo: make instanced
+			InnerObj = MakeUnique<TImplicitObjectScaled<FTriangleMeshImplicitObject, /*bInstanced=*/false>>(MoveTemp(TriMeshObj), P2UVector(TriMeshGeom.scale.scale));	//todo: make instanced
 		}
 		break;
 	}
