Merging CL#10682670

Missing file

diff --git a/Engine/Source/Runtime/Experimental/Chaos/Public/Chaos/GeometryQueries.h b/Engine/Source/Runtime/Experimental/Chaos/Public/Chaos/GeometryQueries.h
--- a/Engine/Source/Runtime/Experimental/Chaos/Public/Chaos/GeometryQueries.h
+++ b/Engine/Source/Runtime/Experimental/Chaos/Public/Chaos/GeometryQueries.h
@@ -97,13 +97,13 @@ namespace Chaos
 			}
 			case ImplicitObjectType::TriangleMesh:
 			{
-				const TTriangleMeshImplicitObject<FReal>& ATriangleMesh = static_cast<const TTriangleMeshImplicitObject<FReal>&>(A);
+				const FTriangleMeshImplicitObject& ATriangleMesh = static_cast<const FTriangleMeshImplicitObject&>(A);
 				return ATriangleMesh.OverlapGeom(B, BToATM, Thickness);
 			}
 			default:
 				if (IsScaled(AType))
 				{
-					const auto& AScaled = TImplicitObjectScaled<TTriangleMeshImplicitObject<FReal>>::AsScaledChecked(A);
+					const auto& AScaled = TImplicitObjectScaled<FTriangleMeshImplicitObject>::AsScaledChecked(A);
 					return AScaled.LowLevelOverlapGeom(B, BToATM, Thickness);
 				}
 				else
@@ -200,14 +200,14 @@ namespace Chaos
 			}
 			case ImplicitObjectType::TriangleMesh:
 			{
-				const auto& ATriangleMesh = static_cast<const TTriangleMeshImplicitObject<FReal>&>(A);
+				const auto& ATriangleMesh = static_cast<const FTriangleMeshImplicitObject&>(A);
 				bResult = ATriangleMesh.SweepGeom(B, BToATM, LocalDir, Length, OutTime, LocalPosition, LocalNormal, OutFaceIndex, Thickness, bComputeMTD);
 				break;
 			}
 			default:
 				if (IsScaled(AType))
 				{
-					const auto& AScaled = TImplicitObjectScaled<TTriangleMeshImplicitObject<FReal>>::AsScaledChecked(A);
+					const auto& AScaled = TImplicitObjectScaled<FTriangleMeshImplicitObject>::AsScaledChecked(A);
 					bResult = AScaled.LowLevelSweepGeom(B, BToATM, LocalDir, Length, OutTime, LocalPosition, LocalNormal, OutFaceIndex, Thickness, bComputeMTD);
 					break;
 				}
