add missing header
#rnx

diff --git a/Engine/Source/Runtime/Experimental/InteractiveToolsFramework/Public/InteractiveToolObjects.h b/Engine/Source/Runtime/Experimental/InteractiveToolsFramework/Public/InteractiveToolObjects.h
--- a/Engine/Source/Runtime/Experimental/InteractiveToolsFramework/Public/InteractiveToolObjects.h
+++ b/Engine/Source/Runtime/Experimental/InteractiveToolsFramework/Public/InteractiveToolObjects.h
@@ -3,6 +3,7 @@
 #pragma once
 
 #include "CoreMinimal.h"
+#include "GameFramework/Actor.h"
 #include "InteractiveToolObjects.generated.h"
 
 
