Simplify mechanism to provide and obtain MeshDescriptionSources.  Also expose an API to provide and control access to Sources for custom Components types

#rnx

diff --git a/Engine/Source/Editor/Experimental/EditorInteractiveToolsFramework/Private/EdModeInteractiveToolsContext.cpp b/Engine/Source/Editor/Experimental/EditorInteractiveToolsFramework/Private/EdModeInteractiveToolsContext.cpp
--- a/Engine/Source/Editor/Experimental/EditorInteractiveToolsFramework/Private/EdModeInteractiveToolsContext.cpp
+++ b/Engine/Source/Editor/Experimental/EditorInteractiveToolsFramework/Private/EdModeInteractiveToolsContext.cpp
@@ -55,7 +55,6 @@ public:
 		StateOut.World = EditorMode->GetWorld();
 		StateOut.SelectedActors = EditorMode->GetModeManager()->GetSelectedActors();
 		StateOut.SelectedComponents = EditorMode->GetModeManager()->GetSelectedComponents();
-		StateOut.SourceBuilder = ToolsContext->GetComponentSourceFactory();
 	}
 
 
@@ -309,7 +308,6 @@ UEdModeInteractiveToolsContext::UEdModeInteractiveToolsContext()
 	QueriesAPI = nullptr;
 	TransactionAPI = nullptr;
 	AssetAPI = nullptr;
-	SourceFactory = nullptr;
 }
 
 
@@ -351,7 +349,6 @@ void UEdModeInteractiveToolsContext::InitializeContextFromEdMode(FEdMode* Editor
 	this->TransactionAPI = new FEdModeToolsContextTransactionImpl(this, EditorModeIn);
 	this->QueriesAPI = new FEdModeToolsContextQueriesImpl(this, EditorModeIn);
 	this->AssetAPI = new FEditorToolAssetAPI();
-	this->SourceFactory = new FEditorComponentSourceFactory();
 
 	Initialize(QueriesAPI, TransactionAPI);
 
@@ -387,12 +384,6 @@ void UEdModeInteractiveToolsContext::ShutdownContext()
 		AssetAPI = nullptr;
 	}
 
-	if (SourceFactory != nullptr)
-	{
-		delete SourceFactory;
-		SourceFactory = nullptr;
-	}
-
 	this->EditorMode = nullptr;
 }
 
@@ -847,4 +838,4 @@ void UEdModeInteractiveToolsContext::RestoreEditorState()
 		GCurrentLevelEditingViewportClient->EngineShowFlags.SetAntiAliasing(bSavedAntiAliasingState);
 		bHaveSavedEditorState = false;
 	}
-}
\ No newline at end of file
+}
diff --git a/Engine/Source/Editor/Experimental/EditorInteractiveToolsFramework/Private/EditorComponentSourceFactory.cpp b/Engine/Source/Editor/Experimental/EditorInteractiveToolsFramework/Private/EditorComponentSourceFactory.cpp
--- a/Engine/Source/Editor/Experimental/EditorInteractiveToolsFramework/Private/EditorComponentSourceFactory.cpp
+++ b/Engine/Source/Editor/Experimental/EditorInteractiveToolsFramework/Private/EditorComponentSourceFactory.cpp
@@ -8,29 +8,24 @@
 #include "PhysicsEngine/BodySetup.h"
 
 
-TUniquePtr<IMeshDescriptionSource> FEditorComponentSourceFactory::MakeMeshDescriptionSource(UActorComponent* Component)
+TUniquePtr<IMeshDescriptionSource> MakeMeshDescriptionSourceForStaticMesh(UActorComponent* Component)
 {
 	UStaticMeshComponent* StaticMeshComp = Cast<UStaticMeshComponent>(Component);
 	if (StaticMeshComp != nullptr)
 	{
-		auto NewSource = new FStaticMeshComponentMeshDescriptionSource(StaticMeshComp, 0);
+		auto NewSource =
+			new FStaticMeshComponentMeshDescriptionSource(StaticMeshComp, 0);
 		return TUniquePtr<IMeshDescriptionSource>(NewSource);
 	}
 	return nullptr;
 }
 
 
-
-
-
 FStaticMeshComponentMeshDescriptionSource::FStaticMeshComponentMeshDescriptionSource(
-	UStaticMeshComponent* ComponentIn, int LODIndex)
+	UStaticMeshComponent* ComponentIn, int LODIndex) : Component{ComponentIn}, LODIndex(LODIndex)
 {
-	this->Component = ComponentIn;
-	this->LODIndex = LODIndex;
 }
 
-
 AActor* FStaticMeshComponentMeshDescriptionSource::GetOwnerActor() const
 {
 	return Component->GetOwner();
@@ -105,4 +100,4 @@ bool FStaticMeshComponentMeshDescriptionSource::HitTest(const FRay& WorldRay, FH
 	}
 
 	return false;
-}
\ No newline at end of file
+}
diff --git a/Engine/Source/Editor/Experimental/EditorInteractiveToolsFramework/Private/EditorInteractiveToolsFrameworkModule.cpp b/Engine/Source/Editor/Experimental/EditorInteractiveToolsFramework/Private/EditorInteractiveToolsFrameworkModule.cpp
--- a/Engine/Source/Editor/Experimental/EditorInteractiveToolsFramework/Private/EditorInteractiveToolsFrameworkModule.cpp
+++ b/Engine/Source/Editor/Experimental/EditorInteractiveToolsFramework/Private/EditorInteractiveToolsFrameworkModule.cpp
@@ -1,12 +1,16 @@
 // Copyright 1998-2019 Epic Games, Inc. All Rights Reserved.
 
 #include "EditorInteractiveToolsFrameworkModule.h"
+#include "ComponentSourceInterfaces.h"
+#include "EditorComponentSourceFactory.h"
 
 #define LOCTEXT_NAMESPACE "FEditorInteractiveToolsFrameworkModule"
 
 void FEditorInteractiveToolsFrameworkModule::StartupModule()
 {
 	// This code will execute after your module is loaded into memory; the exact timing is specified in the .uplugin file per-module
+	AddMeshDescriptionSourceBuilder(MakeMeshDescriptionSourceForStaticMesh);
+	UE_LOG(LogTemp, Warning, TEXT("FEditorInteractiveToolsFrameworkModule StartupModule"));
 }
 
 void FEditorInteractiveToolsFrameworkModule::ShutdownModule()
@@ -17,4 +21,4 @@ void FEditorInteractiveToolsFrameworkModule::ShutdownModule()
 
 #undef LOCTEXT_NAMESPACE
 	
-IMPLEMENT_MODULE(FEditorInteractiveToolsFrameworkModule, EditorInteractiveToolsFramework)
\ No newline at end of file
+IMPLEMENT_MODULE(FEditorInteractiveToolsFrameworkModule, EditorInteractiveToolsFramework)
diff --git a/Engine/Source/Editor/Experimental/EditorInteractiveToolsFramework/Public/EdModeInteractiveToolsContext.h b/Engine/Source/Editor/Experimental/EditorInteractiveToolsFramework/Public/EdModeInteractiveToolsContext.h
--- a/Engine/Source/Editor/Experimental/EditorInteractiveToolsFramework/Public/EdModeInteractiveToolsContext.h
+++ b/Engine/Source/Editor/Experimental/EditorInteractiveToolsFramework/Public/EdModeInteractiveToolsContext.h
@@ -34,7 +34,6 @@ public:
 	IToolsContextQueriesAPI* GetQueriesAPI() const { return QueriesAPI; }
 	IToolsContextTransactionsAPI* GetTransactionAPI() const { return TransactionAPI; }
 	IToolsContextAssetAPI* GetAssetAPI() const { return AssetAPI; }
-	IComponentSourceFactory* GetComponentSourceFactory() const { return SourceFactory; }
 
 	virtual void PostInvalidation();
 
@@ -94,7 +93,6 @@ protected:
 	IToolsContextQueriesAPI* QueriesAPI;
 	IToolsContextTransactionsAPI* TransactionAPI;
 	IToolsContextAssetAPI* AssetAPI;
-	IComponentSourceFactory* SourceFactory;
 
 	// if true, we invalidate the ViewportClient on next tick
 	bool bInvalidationPending;
diff --git a/Engine/Source/Editor/Experimental/EditorInteractiveToolsFramework/Public/EditorComponentSourceFactory.h b/Engine/Source/Editor/Experimental/EditorInteractiveToolsFramework/Public/EditorComponentSourceFactory.h
--- a/Engine/Source/Editor/Experimental/EditorInteractiveToolsFramework/Public/EditorComponentSourceFactory.h
+++ b/Engine/Source/Editor/Experimental/EditorInteractiveToolsFramework/Public/EditorComponentSourceFactory.h
@@ -8,21 +8,6 @@
 
 
 
-/**
- * Component Source Factory for use in the Editor (ie inside FEdMode, etc)
- */
-class EDITORINTERACTIVETOOLSFRAMEWORK_API FEditorComponentSourceFactory : public IComponentSourceFactory
-{
-public:
-	virtual ~FEditorComponentSourceFactory() {}
-
-	virtual TUniquePtr<IMeshDescriptionSource> MakeMeshDescriptionSource(UActorComponent* Component);
-};
-
-
-
-
-
 class UStaticMeshComponent;
 
 /**
@@ -31,8 +16,8 @@ class UStaticMeshComponent;
 class EDITORINTERACTIVETOOLSFRAMEWORK_API FStaticMeshComponentMeshDescriptionSource : public IMeshDescriptionSource
 {
 public:
-	UStaticMeshComponent* Component;
-	int LODIndex;
+	UStaticMeshComponent* Component{nullptr};
+	int LODIndex{0};
 
 	FStaticMeshComponentMeshDescriptionSource(
 		UStaticMeshComponent* ComponentIn,
@@ -52,3 +37,6 @@ public:
 	virtual void CommitInPlaceModification(const TFunction<void(FMeshDescription*)>& ModifyFunction) override;
 };
 
+EDITORINTERACTIVETOOLSFRAMEWORK_API TUniquePtr<IMeshDescriptionSource> MakeMeshDescriptionSourceForStaticMesh(UActorComponent* Component);
+
+
diff --git a/Engine/Source/Editor/UnrealEd/Private/EditorEngine.cpp b/Engine/Source/Editor/UnrealEd/Private/EditorEngine.cpp
--- a/Engine/Source/Editor/UnrealEd/Private/EditorEngine.cpp
+++ b/Engine/Source/Editor/UnrealEd/Private/EditorEngine.cpp
@@ -897,7 +897,8 @@ void UEditorEngine::Init(IEngineLoop* InEngineLoop)
 			TEXT("MergeActors"),
 			TEXT("InputBindingEditor"),
 			TEXT("AudioEditor"),
-			TEXT("TimeManagementEditor")
+			TEXT("TimeManagementEditor"),
+			TEXT("EditorInteractiveToolsFramework")
 		};
 
 		FScopedSlowTask ModuleSlowTask(ARRAY_COUNT(ModuleNames));
diff --git a/Engine/Source/Runtime/Experimental/InteractiveToolsFramework/Private/BaseTools/MeshSurfacePointTool.cpp b/Engine/Source/Runtime/Experimental/InteractiveToolsFramework/Private/BaseTools/MeshSurfacePointTool.cpp
--- a/Engine/Source/Runtime/Experimental/InteractiveToolsFramework/Private/BaseTools/MeshSurfacePointTool.cpp
+++ b/Engine/Source/Runtime/Experimental/InteractiveToolsFramework/Private/BaseTools/MeshSurfacePointTool.cpp
@@ -31,8 +31,7 @@ void UMeshSurfacePointToolBuilder::InitializeNewTool(UMeshSurfacePointTool* NewT
 {
 	UActorComponent* MeshComponent = ToolBuilderUtil::FindFirstComponent(SceneState, ToolBuilderUtil::IsMeshDescriptionSourceComponent);
 	check(MeshComponent != nullptr);
-	NewTool->SetMeshSource(
-		SceneState.SourceBuilder->MakeMeshDescriptionSource(MeshComponent) );
+	NewTool->SetMeshSource( MakeMeshDescriptionSource(MeshComponent) );
 }
 
 
diff --git a/Engine/Source/Runtime/Experimental/InteractiveToolsFramework/Public/ComponentSourceInterfaces.h b/Engine/Source/Runtime/Experimental/InteractiveToolsFramework/Public/ComponentSourceInterfaces.h
--- a/Engine/Source/Runtime/Experimental/InteractiveToolsFramework/Public/ComponentSourceInterfaces.h
+++ b/Engine/Source/Runtime/Experimental/InteractiveToolsFramework/Public/ComponentSourceInterfaces.h
@@ -4,6 +4,7 @@
 
 #include "CoreMinimal.h"
 #include "Math/UnrealMath.h"
+#include "Templates/Function.h"
 #include "Engine/EngineTypes.h"    // FHitResult
 
 // predeclarations
@@ -82,19 +83,18 @@ public:
 
 
 
+using MeshDescriptionSourceBuilder = TFunction<TUniquePtr<IMeshDescriptionSource>(UActorComponent*)>;
+
 /**
- * Interface to a Factory that knows how to build Sources for UObjects. 
+ * Add a factory method to make MeshDescriptionSources from UActorComponent*
+ * @param Builder The MeshDescriptionSourceBuilder
+ * @return void
  */
-class IComponentSourceFactory
-{
-public:
-	virtual ~IComponentSourceFactory() {}
-
-	/**
-	 * Create a MeshDescription source for the given Component
-	 * @param Component A UObject that can provide a MeshDescription. Assumption is this is a Component of an Actor.
-	 * @return A MeshDescriptionSource instance. Must not return null.
-	 */
-	virtual TUniquePtr<IMeshDescriptionSource> MakeMeshDescriptionSource(UActorComponent* Component) = 0;
-};
+INTERACTIVETOOLSFRAMEWORK_API void AddMeshDescriptionSourceBuilder( MeshDescriptionSourceBuilder Builder );
 
+/**
+ * Create a MeshDescription source for the given Component
+ * @param Component A UObject that can provide a MeshDescription. Assumption is this is a Component of an Actor.
+ * @return A MeshDescriptionSource instance. Must not return null.
+ */
+INTERACTIVETOOLSFRAMEWORK_API TUniquePtr<IMeshDescriptionSource> MakeMeshDescriptionSource(UActorComponent* Component);
diff --git a/Engine/Source/Runtime/Experimental/InteractiveToolsFramework/Public/ToolContextInterfaces.h b/Engine/Source/Runtime/Experimental/InteractiveToolsFramework/Public/ToolContextInterfaces.h
--- a/Engine/Source/Runtime/Experimental/InteractiveToolsFramework/Public/ToolContextInterfaces.h
+++ b/Engine/Source/Runtime/Experimental/InteractiveToolsFramework/Public/ToolContextInterfaces.h
@@ -32,13 +32,6 @@ struct INTERACTIVETOOLSFRAMEWORK_API FToolBuilderState
 	USelection* SelectedActors = nullptr;
 	/** Current selected Components. May be empty or nullptr. */
 	USelection* SelectedComponents = nullptr;
-
-	/** Implementation that can build Sources (like MeshDescriptionSource) for Components */
-	IComponentSourceFactory* SourceBuilder = nullptr;
-
-	FToolBuilderState()
-	{
-	}
 };
 
 
