Add straggler ComponentSourceInterfaces.cpp to fix the build

#rnx

diff --git a/Engine/Source/Runtime/Experimental/InteractiveToolsFramework/Private/ComponentSourceInterfaces.cpp b/Engine/Source/Runtime/Experimental/InteractiveToolsFramework/Private/ComponentSourceInterfaces.cpp
new file mode 100644
--- /dev/null
+++ b/Engine/Source/Runtime/Experimental/InteractiveToolsFramework/Private/ComponentSourceInterfaces.cpp
@@ -0,0 +1,30 @@
+// Copyright 1998-2019 Epic Games, Inc. All Rights Reserved.
+
+#include "ComponentSourceInterfaces.h"
+
+#include "Containers/Array.h"
+
+namespace
+{
+	TArray<MeshDescriptionSourceBuilder> Builders;
+}
+
+void
+AddMeshDescriptionSourceBuilder( MeshDescriptionSourceBuilder Builder )
+{
+	Builders.Push(MoveTemp(Builder));
+}
+
+TUniquePtr<IMeshDescriptionSource>
+MakeMeshDescriptionSource(UActorComponent* Component)
+{
+	for ( const auto& Builder : Builders )
+	{
+		auto Source = Builder( Component );
+		if ( Source )
+		{
+			return Source;
+		}
+	}
+	return {};
+}
