Split MeshSource into TargetComponent, MeshSource, and MeshSink

#rnx

diff --git a/Engine/Source/Editor/Experimental/EditorInteractiveToolsFramework/Private/EditorComponentSourceFactory.cpp b/Engine/Source/Editor/Experimental/EditorInteractiveToolsFramework/Private/EditorComponentSourceFactory.cpp
--- a/Engine/Source/Editor/Experimental/EditorInteractiveToolsFramework/Private/EditorComponentSourceFactory.cpp
+++ b/Engine/Source/Editor/Experimental/EditorInteractiveToolsFramework/Private/EditorComponentSourceFactory.cpp
@@ -3,70 +3,20 @@
 #include "EditorComponentSourceFactory.h"
 
 #include "Engine/StaticMesh.h"
-#include "Components/ActorComponent.h"
+#include "Components/PrimitiveComponent.h"
 #include "Components/StaticMeshComponent.h"
 #include "PhysicsEngine/BodySetup.h"
 
 
-TUniquePtr<IMeshDescriptionSource> MakeMeshDescriptionSourceForStaticMesh(UActorComponent* Component)
-{
-	UStaticMeshComponent* StaticMeshComp = Cast<UStaticMeshComponent>(Component);
-	if (StaticMeshComp != nullptr)
-	{
-		auto NewSource =
-			new FStaticMeshComponentMeshDescriptionSource(StaticMeshComp, 0);
-		return TUniquePtr<IMeshDescriptionSource>(NewSource);
-	}
-	return nullptr;
-}
-
-
-FStaticMeshComponentMeshDescriptionSource::FStaticMeshComponentMeshDescriptionSource(
-	UStaticMeshComponent* ComponentIn, int LODIndex) : Component{ComponentIn}, LODIndex(LODIndex)
-{
-}
-
-AActor* FStaticMeshComponentMeshDescriptionSource::GetOwnerActor() const
-{
-	return Component->GetOwner();
-}
-
-UActorComponent* FStaticMeshComponentMeshDescriptionSource::GetOwnerComponent() const
-{
-	return Component;
-}
-
-
-void FStaticMeshComponentMeshDescriptionSource::SetOwnerVisibility(bool bVisible) const
-{
-	Component->SetVisibility(bVisible);
+namespace {
+FMeshDescription *
+GetMeshDescription(UStaticMeshComponent* Component, int LODIndex) {
+	return Component->GetStaticMesh()->GetMeshDescription(LODIndex);
 }
 
-FMeshDescription* FStaticMeshComponentMeshDescriptionSource::GetMeshDescription() const
-{
-	return Component->GetStaticMesh()->GetMeshDescription(0);
-}
-
-UMaterialInterface* FStaticMeshComponentMeshDescriptionSource::GetMaterial(int32 MaterialIndex) const
-{
-	return Component->GetMaterial(MaterialIndex);
-}
-
-FTransform FStaticMeshComponentMeshDescriptionSource::GetWorldTransform() const
-{
-	//return Component->GetOwner()->GetActorTransform();
-	return Component->GetComponentTransform();
-}
-
-
-
-bool FStaticMeshComponentMeshDescriptionSource::IsReadOnly() const
-{
-	return false;
-}
-
-
-void FStaticMeshComponentMeshDescriptionSource::CommitInPlaceModification(const TFunction<void(FMeshDescription*)>& ModifyFunction) 
+void
+CommitInPlaceModification(UStaticMeshComponent* Component, int LODIndex,
+						  const TFunction<void(FMeshDescription*)>& ModifyFunction)
 {
 	//bool bSaved = Component->Modify();
 	//check(bSaved);
@@ -89,15 +39,25 @@ void FStaticMeshComponentMeshDescriptionSource::CommitInPlaceModification(const
 
 	//Component->PostEditChange();
 }
+} // namespace
 
-
-bool FStaticMeshComponentMeshDescriptionSource::HitTest(const FRay& WorldRay, FHitResult& OutHit) const
+FMeshDescriptionBridge
+MakeStaticMeshDescriptionBridge(UPrimitiveComponent* Component)
 {
-	FVector End = WorldRay.PointAt(HALF_WORLD_MAX);
-	if (Component->LineTraceComponent(OutHit, WorldRay.Origin, End, FCollisionQueryParams(SCENE_QUERY_STAT(HitTest), true)))
+	const int LODIndex{0};
+	UStaticMeshComponent* StaticMeshComponent = Cast<UStaticMeshComponent>(Component);
+	if (StaticMeshComponent != nullptr)
 	{
-		return true;
+		return {
+			[StaticMeshComponent, LODIndex]()
+			{
+				return GetMeshDescription(StaticMeshComponent, LODIndex);
+			},
+			[StaticMeshComponent, LODIndex](const FMeshDescriptionBridge::FCommitter& Committer)
+			{
+				CommitInPlaceModification(StaticMeshComponent, LODIndex, Committer);
+			}
+		};
 	}
-
-	return false;
+	return FMeshDescriptionBridge{};
 }
diff --git a/Engine/Source/Editor/Experimental/EditorInteractiveToolsFramework/Private/EditorInteractiveToolsFrameworkModule.cpp b/Engine/Source/Editor/Experimental/EditorInteractiveToolsFramework/Private/EditorInteractiveToolsFrameworkModule.cpp
--- a/Engine/Source/Editor/Experimental/EditorInteractiveToolsFramework/Private/EditorInteractiveToolsFrameworkModule.cpp
+++ b/Engine/Source/Editor/Experimental/EditorInteractiveToolsFramework/Private/EditorInteractiveToolsFrameworkModule.cpp
@@ -9,7 +9,7 @@
 void FEditorInteractiveToolsFrameworkModule::StartupModule()
 {
 	// This code will execute after your module is loaded into memory; the exact timing is specified in the .uplugin file per-module
-	AddMeshDescriptionSourceBuilder(MakeMeshDescriptionSourceForStaticMesh);
+	AddMeshDescriptionBridgeBuilder(MakeStaticMeshDescriptionBridge);
 	UE_LOG(LogTemp, Warning, TEXT("FEditorInteractiveToolsFrameworkModule StartupModule"));
 }
 
diff --git a/Engine/Source/Editor/Experimental/EditorInteractiveToolsFramework/Public/EditorComponentSourceFactory.h b/Engine/Source/Editor/Experimental/EditorInteractiveToolsFramework/Public/EditorComponentSourceFactory.h
--- a/Engine/Source/Editor/Experimental/EditorInteractiveToolsFramework/Public/EditorComponentSourceFactory.h
+++ b/Engine/Source/Editor/Experimental/EditorInteractiveToolsFramework/Public/EditorComponentSourceFactory.h
@@ -5,38 +5,4 @@
 #include "CoreMinimal.h"
 #include "ComponentSourceInterfaces.h"
 
-
-
-
-class UStaticMeshComponent;
-
-/**
- * This MeshDescriptionSource provides a specific LOD from a StaticMeshComponent
- */
-class EDITORINTERACTIVETOOLSFRAMEWORK_API FStaticMeshComponentMeshDescriptionSource : public IMeshDescriptionSource
-{
-public:
-	UStaticMeshComponent* Component{nullptr};
-	int LODIndex{0};
-
-	FStaticMeshComponentMeshDescriptionSource(
-		UStaticMeshComponent* ComponentIn,
-		int LODIndex = 0);
-	virtual ~FStaticMeshComponentMeshDescriptionSource() {}
-
-	virtual AActor* GetOwnerActor() const override;
-	virtual UActorComponent* GetOwnerComponent() const override;
-	virtual FMeshDescription* GetMeshDescription() const override;
-	virtual UMaterialInterface* GetMaterial(int32 MaterialIndex) const override;
-	virtual FTransform GetWorldTransform() const override;
-	virtual bool HitTest(const FRay& WorldRay, FHitResult& OutHit) const override;
-
-	virtual bool IsReadOnly() const override;
-
-	virtual void SetOwnerVisibility(bool bVisible) const override;
-	virtual void CommitInPlaceModification(const TFunction<void(FMeshDescription*)>& ModifyFunction) override;
-};
-
-EDITORINTERACTIVETOOLSFRAMEWORK_API TUniquePtr<IMeshDescriptionSource> MakeMeshDescriptionSourceForStaticMesh(UActorComponent* Component);
-
-
+EDITORINTERACTIVETOOLSFRAMEWORK_API FMeshDescriptionBridge MakeStaticMeshDescriptionBridge(UPrimitiveComponent* Component);
diff --git a/Engine/Source/Runtime/Experimental/InteractiveToolsFramework/Private/BaseTools/MeshSurfacePointTool.cpp b/Engine/Source/Runtime/Experimental/InteractiveToolsFramework/Private/BaseTools/MeshSurfacePointTool.cpp
--- a/Engine/Source/Runtime/Experimental/InteractiveToolsFramework/Private/BaseTools/MeshSurfacePointTool.cpp
+++ b/Engine/Source/Runtime/Experimental/InteractiveToolsFramework/Private/BaseTools/MeshSurfacePointTool.cpp
@@ -29,22 +29,16 @@ UMeshSurfacePointTool* UMeshSurfacePointToolBuilder::CreateNewTool(const FToolBu
 
 void UMeshSurfacePointToolBuilder::InitializeNewTool(UMeshSurfacePointTool* NewTool, const FToolBuilderState& SceneState) const
 {
-	UActorComponent* MeshComponent = ToolBuilderUtil::FindFirstComponent(SceneState, ToolBuilderUtil::IsMeshDescriptionSourceComponent);
+	UActorComponent* ActorComponent = ToolBuilderUtil::FindFirstComponent(SceneState, ToolBuilderUtil::IsMeshDescriptionSourceComponent);
+	UPrimitiveComponent* MeshComponent = Cast<UPrimitiveComponent>(ActorComponent);
 	check(MeshComponent != nullptr);
-	NewTool->SetMeshSource( MakeMeshDescriptionSource(MeshComponent) );
+	NewTool->SetSelection( MakeComponentTarget(MeshComponent) );
 }
 
 
 /*
  * Tool
  */
-
-void UMeshSurfacePointTool::SetMeshSource(TUniquePtr<IMeshDescriptionSource> MeshSourceIn)
-{
-	this->MeshSource = TUniquePtr<IMeshDescriptionSource>(MoveTemp(MeshSourceIn));
-}
-
-
 void UMeshSurfacePointTool::Setup()
 {
 	UInteractiveTool::Setup();
@@ -64,7 +58,7 @@ void UMeshSurfacePointTool::Setup()
 
 bool UMeshSurfacePointTool::HitTest(const FRay& Ray, FHitResult& OutHit)
 {
-	return MeshSource->HitTest(Ray, OutHit);
+	return ComponentTarget.HitTest(Ray, OutHit);
 }
 
 
diff --git a/Engine/Source/Runtime/Experimental/InteractiveToolsFramework/Private/ComponentSourceInterfaces.cpp b/Engine/Source/Runtime/Experimental/InteractiveToolsFramework/Private/ComponentSourceInterfaces.cpp
--- a/Engine/Source/Runtime/Experimental/InteractiveToolsFramework/Private/ComponentSourceInterfaces.cpp
+++ b/Engine/Source/Runtime/Experimental/InteractiveToolsFramework/Private/ComponentSourceInterfaces.cpp
@@ -2,29 +2,73 @@
 
 #include "ComponentSourceInterfaces.h"
 
+#include "Components/PrimitiveComponent.h"
 #include "Containers/Array.h"
 
 namespace
 {
-	TArray<MeshDescriptionSourceBuilder> Builders;
+	TArray<FMeshDescriptionBridge::FBuilder> Builders;
 }
 
 void
-AddMeshDescriptionSourceBuilder( MeshDescriptionSourceBuilder Builder )
+AddMeshDescriptionBridgeBuilder( FMeshDescriptionBridge::FBuilder Builder )
 {
 	Builders.Push(MoveTemp(Builder));
 }
 
-TUniquePtr<IMeshDescriptionSource>
-MakeMeshDescriptionSource(UActorComponent* Component)
+FComponentTarget
+MakeComponentTarget(UPrimitiveComponent* Component)
 {
 	for ( const auto& Builder : Builders )
 	{
-		auto Source = Builder( Component );
-		if ( Source )
+		auto Bridge = Builder( Component );
+		if ( Bridge.HasSource() )
 		{
-			return Source;
+			return {Bridge, Component};
 		}
 	}
 	return {};
 }
+
+AActor*
+FComponentTarget::GetOwnerActor() const
+{
+	return Component->GetOwner();
+}
+
+UPrimitiveComponent*
+FComponentTarget::GetOwnerComponent() const
+{
+	return Component;
+}
+
+
+void
+FComponentTarget::SetOwnerVisibility(bool bVisible) const
+{
+	Component->SetVisibility(bVisible);
+}
+
+UMaterialInterface*
+FComponentTarget::GetMaterial(int32 MaterialIndex) const
+{
+	return Component->GetMaterial(MaterialIndex);
+}
+
+FTransform
+FComponentTarget::GetWorldTransform() const
+{
+	//return Component->GetOwner()->GetActorTransform();
+	return Component->GetComponentTransform();
+}
+
+bool
+FComponentTarget::HitTest(const FRay& WorldRay, FHitResult& OutHit) const
+{
+	FVector End = WorldRay.PointAt(HALF_WORLD_MAX);
+	if (Component->LineTraceComponent(OutHit, WorldRay.Origin, End, FCollisionQueryParams(SCENE_QUERY_STAT(HitTest), true)))
+	{
+		return true;
+	}
+	return false;
+}
diff --git a/Engine/Source/Runtime/Experimental/InteractiveToolsFramework/Public/BaseTools/MeshSurfacePointTool.h b/Engine/Source/Runtime/Experimental/InteractiveToolsFramework/Public/BaseTools/MeshSurfacePointTool.h
--- a/Engine/Source/Runtime/Experimental/InteractiveToolsFramework/Public/BaseTools/MeshSurfacePointTool.h
+++ b/Engine/Source/Runtime/Experimental/InteractiveToolsFramework/Public/BaseTools/MeshSurfacePointTool.h
@@ -3,7 +3,7 @@
 #pragma once
 
 #include "CoreMinimal.h"
-#include "InteractiveTool.h"
+#include "SingleSelectionTool.h"
 #include "InteractiveToolBuilder.h"
 #include "InputBehaviorSet.h"
 #include "BaseBehaviors/BehaviorTargetInterfaces.h"
@@ -47,7 +47,7 @@ public:
  * to implement custom behavior.
  */
 UCLASS()
-class INTERACTIVETOOLSFRAMEWORK_API UMeshSurfacePointTool : public UInteractiveTool, public IHoverBehaviorTarget
+class INTERACTIVETOOLSFRAMEWORK_API UMeshSurfacePointTool : public USingleSelectionTool, public IHoverBehaviorTarget
 {
 	GENERATED_BODY()
 
@@ -60,14 +60,6 @@ public:
 
 	// UMeshSurfacePointTool API
 
-	/**
-	 * Set the Target MeshDescription-providing Object for this tool
-	 * @param MeshSource object that can provide a MeshDescriptionSource. Note UMeshSurfacePointTool
-	 *        only calls MeshSource->HitTest, so this source technically does not have to actually provide
-	 *        a valid MeshDescription
-	 */
-	virtual void SetMeshSource(TUniquePtr<IMeshDescriptionSource> MeshSource);
-
 	/**
 	 * @return true if the target MeshSource is hit by the Ray
 	 */
@@ -104,9 +96,6 @@ public:
 
 
 protected:
-	/** Target source object that provides a MeshDescription and various other query API calls */
-	TUniquePtr<IMeshDescriptionSource> MeshSource;
-
 	/** Current state of the shift modifier toggle */
 	bool bShiftToggle;
 };
diff --git a/Engine/Source/Runtime/Experimental/InteractiveToolsFramework/Public/ComponentSourceInterfaces.h b/Engine/Source/Runtime/Experimental/InteractiveToolsFramework/Public/ComponentSourceInterfaces.h
--- a/Engine/Source/Runtime/Experimental/InteractiveToolsFramework/Public/ComponentSourceInterfaces.h
+++ b/Engine/Source/Runtime/Experimental/InteractiveToolsFramework/Public/ComponentSourceInterfaces.h
@@ -4,15 +4,38 @@
 
 #include "CoreMinimal.h"
 #include "Math/UnrealMath.h"
+#include "Misc/Optional.h"
 #include "Templates/Function.h"
 #include "Engine/EngineTypes.h"    // FHitResult
 
 // predeclarations
 class AActor;
-class UActorComponent;
+class UPrimitiveComponent;
 struct FMeshDescription;
 class UMaterialInterface;
 
+template <class MeshT>
+struct TMeshBridge
+{
+	using  FSource    = TFunction< MeshT*() >;
+	using  FCommitter = TFunction< void( MeshT* ) >;
+	using  FSink      = TFunction< void( const FCommitter& ) >;
+	using  FBuilder   = TFunction< TMeshBridge ( UPrimitiveComponent* ) >;
+
+	bool HasSource(){ return !!GetMesh; }
+	bool HasSink(){ return !!CommitMesh; }
+
+	FSource GetMesh{nullptr};
+	FSink   CommitMesh{nullptr};
+};
+using FMeshDescriptionBridge = TMeshBridge<FMeshDescription>;
+
+/**
+ * Add a factory method to make MeshDescriptionSources from UActorComponent*
+ * @param Builder The MeshDescriptionSourceBuilder
+ * @return void
+ */
+INTERACTIVETOOLSFRAMEWORK_API void AddMeshDescriptionBridgeBuilder( FMeshDescriptionBridge::FBuilder Builder );
 
 
 /**
@@ -20,36 +43,30 @@ class UMaterialInterface;
  * (optionally) bake a modified MeshDescription back to this Component.
  * An example of a Source might be a StaticMeshComponent. How a modified
  * MeshDescription is committed back is context-dependent (in Editor vs PIE vs Runtime, etc).
- * 
+ *
  * (Conceivably this doesn't have to be backed by a Component, but most usage will assume there is an Actor)
  */
-class IMeshDescriptionSource
+class INTERACTIVETOOLSFRAMEWORK_API FComponentTarget
 {
 public:
-	virtual ~IMeshDescriptionSource() {}
-
-	
 	/** @return the Actor that owns this Component */
-	virtual AActor* GetOwnerActor() const = 0;
+	AActor* GetOwnerActor() const;
 
 	/** @return the Component this is a Source for */
-	virtual UActorComponent* GetOwnerComponent() const = 0;
-
-	/** @return Pointer to the MeshDescription this Source is providing */
-	virtual FMeshDescription* GetMeshDescription() const = 0;
+	UPrimitiveComponent* GetOwnerComponent() const;
 
 	/**
 	 * Get pointer to a Material provided by this Source
 	 * @param MaterialIndex index of the material
 	 * @return MaterialInterface pointer, or null if MaterialIndex is invalid
 	 */
-	virtual UMaterialInterface* GetMaterial(int32 MaterialIndex) const = 0;
+	UMaterialInterface* GetMaterial(int32 MaterialIndex) const;
 
-	/** 
+	/**
 	 * @return the transform on this component
 	 * @todo Do we need to return a list of transforms here?
 	 */
-	virtual FTransform GetWorldTransform() const = 0;
+	FTransform GetWorldTransform() const;
 
 	/**
 	 * Compute ray intersection with the MeshDescription this Source is providing
@@ -57,44 +74,23 @@ public:
 	 * @param OutHit hit test data
 	 * @return true if ray intersected Component
 	 */
-	virtual bool HitTest(const FRay& WorldRay, FHitResult& OutHit) const = 0;
+	bool HitTest(const FRay& WorldRay, FHitResult& OutHit) const;
 
 	/**
 	 * Set the visibility of the Component associated with this Source (ie to hide during Tool usage)
 	 * @param bVisible desired visibility
 	 */
-	virtual void SetOwnerVisibility(bool bVisible) const = 0;
-
+	void SetOwnerVisibility(bool bVisible) const;
 
-	/** @return true if this Source is read-only, ie Commit functions cannot be called */
-	virtual bool IsReadOnly() const { return true; }
-
-	/**
-	 * Call this to modify the MeshDescription provided by the Source. You provide a callback
-	 * function that the Source will immediately call with a suitable MeshDescription instance.
-	 * The Source will then update the Component as necessary. The Source is responsible for
-	 * making sure that it is safe to modify this MeshDescription instance.
-	 * @param ModifyFunction callback function that updates/modifies a MeshDescription
-	 */
-	virtual void CommitInPlaceModification(const TFunction<void(FMeshDescription*)>& ModifyFunction) { check(false); }
+	FMeshDescriptionBridge MeshDescriptionBridge{};
+	UPrimitiveComponent* Component{};
 };
 
-
-
-
-
-using MeshDescriptionSourceBuilder = TFunction<TUniquePtr<IMeshDescriptionSource>(UActorComponent*)>;
-
-/**
- * Add a factory method to make MeshDescriptionSources from UActorComponent*
- * @param Builder The MeshDescriptionSourceBuilder
- * @return void
- */
-INTERACTIVETOOLSFRAMEWORK_API void AddMeshDescriptionSourceBuilder( MeshDescriptionSourceBuilder Builder );
-
 /**
- * Create a MeshDescription source for the given Component
- * @param Component A UObject that can provide a MeshDescription. Assumption is this is a Component of an Actor.
- * @return A MeshDescriptionSource instance. Must not return null.
+ * Create a TargetComponent for the given Component
+ * @param Component A UObject that we would like to use as tool target. This must presently descend from
+ * UPrimitiveComponent
+ * @return An FComponentTarget instance. Must not return null, though the MeshSource and MeshSink in it's MeshBridge may
+ * be
  */
-INTERACTIVETOOLSFRAMEWORK_API TUniquePtr<IMeshDescriptionSource> MakeMeshDescriptionSource(UActorComponent* Component);
+INTERACTIVETOOLSFRAMEWORK_API FComponentTarget MakeComponentTarget(UPrimitiveComponent* Component);
diff --git a/Engine/Source/Runtime/Experimental/InteractiveToolsFramework/Public/MultiSelectionTool.h b/Engine/Source/Runtime/Experimental/InteractiveToolsFramework/Public/MultiSelectionTool.h
new file mode 100644
--- /dev/null
+++ b/Engine/Source/Runtime/Experimental/InteractiveToolsFramework/Public/MultiSelectionTool.h
@@ -0,0 +1,20 @@
+// Copyright 1998-2019 Epic Games, Inc. All Rights Reserved.
+
+#pragma once
+
+#include "InteractiveTool.h"
+#include "ComponentSourceInterfaces.h"
+#include "MultiSelectionTool.generated.h"
+
+UCLASS(Transient)
+class INTERACTIVETOOLSFRAMEWORK_API UMultiSelectionTool : public UInteractiveTool
+{
+GENERATED_BODY()
+public:
+	void SetSelection(TArray<FComponentTarget> ComponentTargetsIn)
+    {
+		ComponentTargets = MoveTemp(ComponentTargetsIn);
+	}
+protected:
+	TArray<FComponentTarget> ComponentTargets{};
+};
diff --git a/Engine/Source/Runtime/Experimental/InteractiveToolsFramework/Public/SingleSelectionTool.h b/Engine/Source/Runtime/Experimental/InteractiveToolsFramework/Public/SingleSelectionTool.h
new file mode 100644
--- /dev/null
+++ b/Engine/Source/Runtime/Experimental/InteractiveToolsFramework/Public/SingleSelectionTool.h
@@ -0,0 +1,20 @@
+// Copyright 1998-2019 Epic Games, Inc. All Rights Reserved.
+
+#pragma once
+
+#include "InteractiveTool.h"
+#include "ComponentSourceInterfaces.h"
+#include "SingleSelectionTool.generated.h"
+
+UCLASS(Transient)
+class INTERACTIVETOOLSFRAMEWORK_API USingleSelectionTool : public UInteractiveTool
+{
+GENERATED_BODY()
+public:
+	void SetSelection(FComponentTarget ComponentTargetIn)
+    {
+		ComponentTarget = MoveTemp(ComponentTargetIn);
+	}
+protected:
+	FComponentTarget ComponentTarget{};
+};
