Change ToolContextInterfaces to use TArray instead of USelection

#rnx

diff --git a/Engine/Source/Editor/Experimental/EditorInteractiveToolsFramework/Private/EdModeInteractiveToolsContext.cpp b/Engine/Source/Editor/Experimental/EditorInteractiveToolsFramework/Private/EdModeInteractiveToolsContext.cpp
--- a/Engine/Source/Editor/Experimental/EditorInteractiveToolsFramework/Private/EdModeInteractiveToolsContext.cpp
+++ b/Engine/Source/Editor/Experimental/EditorInteractiveToolsFramework/Private/EdModeInteractiveToolsContext.cpp
@@ -53,8 +53,8 @@ public:
 		StateOut.ToolManager = ToolsContext->ToolManager;
 		StateOut.GizmoManager = ToolsContext->GizmoManager;
 		StateOut.World = EditorMode->GetWorld();
-		StateOut.SelectedActors = EditorMode->GetModeManager()->GetSelectedActors();
-		StateOut.SelectedComponents = EditorMode->GetModeManager()->GetSelectedComponents();
+		EditorMode->GetModeManager()->GetSelectedActors()->GetSelectedObjects(StateOut.SelectedActors);
+		EditorMode->GetModeManager()->GetSelectedComponents()->GetSelectedObjects(StateOut.SelectedComponents);
 	}
 
 
diff --git a/Engine/Source/Runtime/Experimental/InteractiveToolsFramework/Private/ToolBuilderUtil.cpp b/Engine/Source/Runtime/Experimental/InteractiveToolsFramework/Private/ToolBuilderUtil.cpp
--- a/Engine/Source/Runtime/Experimental/InteractiveToolsFramework/Private/ToolBuilderUtil.cpp
+++ b/Engine/Source/Runtime/Experimental/InteractiveToolsFramework/Private/ToolBuilderUtil.cpp
@@ -2,46 +2,31 @@
 
 #include "ToolBuilderUtil.h"
 #include "CoreMinimal.h"
+#include "Algo/Accumulate.h"
+#include "Algo/Count.h"
+#include "Algo/Find.h"
 #include "GameFramework/Actor.h"
 #include "Engine/Selection.h"
 #include "Components/StaticMeshComponent.h"
 
-
-
 int ToolBuilderUtil::CountComponents(const FToolBuilderState& InputState, const TFunction<bool(UActorComponent*)>& Predicate)
 {
-	int nTypedComponents = 0;
+	int nTypedComponents{};
 
-	if (InputState.SelectedComponents != nullptr && InputState.SelectedComponents->Num() > 0)
+	if (InputState.SelectedComponents.Num() > 0)
 	{
-		for (FSelectionIterator Iter(*InputState.SelectedComponents); Iter; ++Iter)
-		{
-			if (UActorComponent* Component = Cast<UActorComponent>(*Iter))
-			{
-				if (Predicate(Component))
-				{
-					nTypedComponents++;
-				}
-			}
-		}
+		nTypedComponents = Algo::CountIf(InputState.SelectedComponents, Predicate);
 	}
 	else
 	{
-		for (FSelectionIterator Iter(*InputState.SelectedActors); Iter; ++Iter)
-		{
-			if (AActor* Actor = Cast<AActor>(*Iter))
-			{
-				for (UActorComponent* Component : Actor->GetComponents())
-				{
-					if (Predicate(Component))
-					{
-						nTypedComponents++;
-					}
-				}
-			}
-		}
+		nTypedComponents =
+			Algo::TransformAccumulate(InputState.SelectedActors,
+									  [&Predicate](AActor* Actor)
+									  {
+										  return Algo::CountIf(Actor->GetComponents(), Predicate);
+									  },
+									  0);
 	}
-
 	return nTypedComponents;
 }
 
@@ -50,36 +35,21 @@ int ToolBuilderUtil::CountComponents(const FToolBuilderState& InputState, const
 
 UActorComponent* ToolBuilderUtil::FindFirstComponent(const FToolBuilderState& InputState, const TFunction<bool(UActorComponent*)>& Predicate)
 {
-	if (InputState.SelectedComponents != nullptr && InputState.SelectedComponents->Num() > 0)
+	if (InputState.SelectedComponents.Num() > 0)
 	{
-		for (FSelectionIterator Iter(*InputState.SelectedComponents); Iter; ++Iter)
-		{
-			if (UActorComponent* Component = Cast<UActorComponent>(*Iter))
-			{
-				if (Predicate(Component))
-				{
-					return Component;
-				}
-			}
-		}
+		return *InputState.SelectedComponents.FindByPredicate(Predicate);
 	}
 	else
 	{
-		for (FSelectionIterator Iter(*InputState.SelectedActors); Iter; ++Iter)
+		for ( AActor* Actor : InputState.SelectedActors )
 		{
-			if (AActor* Actor = Cast<AActor>(*Iter))
+			UActorComponent* Component = *Algo::FindByPredicate(Actor->GetComponents(), Predicate);
+			if ( Component )
 			{
-				for (UActorComponent* Component : Actor->GetComponents())
-				{
-					if (Predicate(Component))
-					{
-						return Component;
-					}
-				}
+				return Component;
 			}
 		}
 	}
-
 	return nullptr;
 }
 
@@ -88,38 +58,24 @@ UActorComponent* ToolBuilderUtil::FindFirstComponent(const FToolBuilderState& In
 
 TArray<UActorComponent*> ToolBuilderUtil::FindAllComponents(const FToolBuilderState& InputState, const TFunction<bool(UActorComponent*)>& Predicate)
 {
-	TArray<UActorComponent*> Components;
-
-	if (InputState.SelectedComponents != nullptr && InputState.SelectedComponents->Num() > 0)
+	if (InputState.SelectedComponents.Num() > 0)
 	{
-		for (FSelectionIterator Iter(*InputState.SelectedComponents); Iter; ++Iter)
-		{
-			if (UActorComponent* Component = Cast<UActorComponent>(*Iter))
-			{
-				if (Predicate(Component))
-				{
-					Components.AddUnique(Component);
-				}
-			}
-		}
+		return InputState.SelectedComponents.FilterByPredicate(Predicate);
 	}
 	else
 	{
-		for (FSelectionIterator Iter(*InputState.SelectedActors); Iter; ++Iter)
-		{
-			if (AActor* Actor = Cast<AActor>(*Iter))
-			{
-				for (UActorComponent* Component : Actor->GetComponents())
-				{
-					if (Predicate(Component))
-					{
-						Components.AddUnique(Component);
-					}
-				}
-			}
-		}
+		return Algo::TransformAccumulate(InputState.SelectedActors,
+										 [&Predicate](AActor* Actor)
+										 {
+											 TInlineComponentArray<UActorComponent*> ActorComponents;
+											 Actor->GetComponents(ActorComponents);
+											 return ActorComponents.FilterByPredicate(Predicate);
+										 },
+										 TArray<UActorComponent*>{},
+										 [](TArray<UActorComponent*> FoundComponents, TArray<UActorComponent*> ActorComponents)
+										 {
+											 FoundComponents.Insert(MoveTemp(ActorComponents), FoundComponents.Num());
+											 return FoundComponents;
+										 });
 	}
-
-	return Components;
 }
-
diff --git a/Engine/Source/Runtime/Experimental/InteractiveToolsFramework/Public/ToolContextInterfaces.h b/Engine/Source/Runtime/Experimental/InteractiveToolsFramework/Public/ToolContextInterfaces.h
--- a/Engine/Source/Runtime/Experimental/InteractiveToolsFramework/Public/ToolContextInterfaces.h
+++ b/Engine/Source/Runtime/Experimental/InteractiveToolsFramework/Public/ToolContextInterfaces.h
@@ -7,10 +7,11 @@
 
 
 // predeclarations so we don't have to include these in all tools
+class AActor;
+class UActorComponent;
 class FChange;
 class UPackage;
 class FPrimitiveDrawInterface;
-class USelection;
 class UInteractiveToolManager;
 class UInteractiveGizmoManager;
 
@@ -29,9 +30,9 @@ struct INTERACTIVETOOLSFRAMEWORK_API FToolBuilderState
 	UInteractiveGizmoManager* GizmoManager = nullptr;
 
 	/** Current selected Actors. May be empty or nullptr. */
-	USelection* SelectedActors = nullptr;
+	TArray<AActor*> SelectedActors;
 	/** Current selected Components. May be empty or nullptr. */
-	USelection* SelectedComponents = nullptr;
+	TArray<UActorComponent*> SelectedComponents;
 };
 
 
@@ -281,7 +282,6 @@ public:
 
 };
 
-
 /**
  * Users of the Tools Framework need to implement IToolsContextRenderAPI to allow
  * Tools, Indicators, and Gizmos to make low-level rendering calls for things like line drawing.
@@ -297,9 +297,6 @@ public:
 	virtual FPrimitiveDrawInterface* GetPrimitiveDrawInterface() = 0;
 };
 
-
-
-
 /**
  * Users of the Tools Framework need to provide an IToolsContextAssetAPI implementation
  * that allows Packages and Assets to be created/saved. Note that this is not strictly
