Convert remaining ToolBuilderUtil funcs to use TArray instead of USelection

#rnx

diff --git a/Engine/Source/Runtime/Experimental/InteractiveToolsFramework/Private/ToolBuilderUtil.cpp b/Engine/Source/Runtime/Experimental/InteractiveToolsFramework/Private/ToolBuilderUtil.cpp
--- a/Engine/Source/Runtime/Experimental/InteractiveToolsFramework/Private/ToolBuilderUtil.cpp
+++ b/Engine/Source/Runtime/Experimental/InteractiveToolsFramework/Private/ToolBuilderUtil.cpp
@@ -3,8 +3,6 @@
 #include "ToolBuilderUtil.h"
 #include "CoreMinimal.h"
 #include "Algo/Accumulate.h"
-#include "Algo/Count.h"
-#include "Algo/Find.h"
 #include "GameFramework/Actor.h"
 #include "Engine/Selection.h"
 #include "Components/StaticMeshComponent.h"
diff --git a/Engine/Source/Runtime/Experimental/InteractiveToolsFramework/Public/ToolBuilderUtil.h b/Engine/Source/Runtime/Experimental/InteractiveToolsFramework/Public/ToolBuilderUtil.h
--- a/Engine/Source/Runtime/Experimental/InteractiveToolsFramework/Public/ToolBuilderUtil.h
+++ b/Engine/Source/Runtime/Experimental/InteractiveToolsFramework/Public/ToolBuilderUtil.h
@@ -3,6 +3,8 @@
 #pragma once
 
 #include "CoreMinimal.h"
+#include "Algo/Count.h"
+#include "Algo/Find.h"
 #include "GameFramework/Actor.h"
 #include "Engine/Selection.h"
 #include "InteractiveToolBuilder.h"
@@ -35,79 +37,18 @@ namespace ToolBuilderUtil
 
 }
 
-
-
-
 /*
  * Template Implementations
  */
-
-
 template<typename ComponentType>
 int ToolBuilderUtil::CountSelectedComponentsOfType(const FToolBuilderState& InputState)
 {
-	int nTypedComponents = 0;
-
-	if (InputState.SelectedComponents != nullptr && InputState.SelectedComponents->Num() > 0)
-	{
-		for (FSelectionIterator Iter(*InputState.SelectedComponents); Iter; ++Iter)
-		{
-			if (ComponentType* TypedComponent = Cast<ComponentType>(*Iter))
-			{
-				nTypedComponents++;
-			}
-		}
-	}
-	else
-	{
-		for (FSelectionIterator Iter(*InputState.SelectedActors); Iter; ++Iter)
-		{
-			if (AActor* Actor = Cast<AActor>(*Iter))
-			{
-				// [RMS] is there a more efficient way to do this?
-				TArray<UActorComponent*> TypedComponents = Actor->GetComponentsByClass(ComponentType::StaticClass());
-				nTypedComponents += TypedComponents.Num();
-			}
-		}
-	}
-
-	return nTypedComponents;
+	return CountComponents(InputState, [](UActorComponent* Actor) { return Cast<ComponentType>(Actor) != nullptr; });
 }
 
-
-
-
-
-
 template<typename ComponentType>
 ComponentType* ToolBuilderUtil::FindFirstComponentOfType(const FToolBuilderState& InputState)
 {
-	if (InputState.SelectedComponents != nullptr && InputState.SelectedComponents->Num() > 0)
-	{
-		for (FSelectionIterator Iter(*InputState.SelectedComponents); Iter; ++Iter)
-		{
-			if (ComponentType* TypedComponent = Cast<ComponentType>(*Iter))
-			{
-				return TypedComponent;
-			}
-		}
-	}
-	else
-	{
-		for (FSelectionIterator Iter(*InputState.SelectedActors); Iter; ++Iter)
-		{
-			if (AActor* Actor = Cast<AActor>(*Iter))
-			{
-				// [RMS] is there a more efficient way to do this?
-				TArray<UActorComponent*> TypedComponents = Actor->GetComponentsByClass(ComponentType::StaticClass());
-				if (TypedComponents.Num() > 0)
-				{
-					return (ComponentType*)TypedComponents[0];
-				}
-			}
-		}
-	}
-
-	return nullptr;
+	return FindFirstComponent(InputState, [](UActorComponent* Actor) { return Cast<ComponentType>(Actor) != nullptr; });
 }
 
