add missing UPrimitiveComponent header
#rnx

diff --git a/Engine/Source/Runtime/Experimental/InteractiveToolsFramework/Private/BaseTools/MeshSurfacePointTool.cpp b/Engine/Source/Runtime/Experimental/InteractiveToolsFramework/Private/BaseTools/MeshSurfacePointTool.cpp
--- a/Engine/Source/Runtime/Experimental/InteractiveToolsFramework/Private/BaseTools/MeshSurfacePointTool.cpp
+++ b/Engine/Source/Runtime/Experimental/InteractiveToolsFramework/Private/BaseTools/MeshSurfacePointTool.cpp
@@ -4,6 +4,7 @@
 #include "BaseBehaviors/MouseHoverBehavior.h"
 #include "InteractiveToolManager.h"
 #include "ToolBuilderUtil.h"
+#include "Components/PrimitiveComponent.h"
 
 
 /*
