ToolsFramework: Standard implementations of Gizmos for axis position/angle and plane position. Used to implement a standard 3-axis translate/rotate transformer.

InteractiveGizmoManager now tracks Owner of Gizmos, for delete-by-owner support. Added Create3AxisTransformGizmo() to create standard Gizmo.

GizmoInterfaces.h: UInterfaces for gizmo components
AxisPositionGizmo, AxisAngleGizmo, PlanePositionGizmo: standard implementations of UInteractiveGizmo
GizmoBaseComponent, GizmoArrowComponent, GizmoRectangleComponent, GizmoCircleComponent: custom UPrimitiveComponent implementations for standard Gizmos
TransformGizmo: 3-axis transform gizmo built out of above elements
AxisSources.h, ParameterSourcesFloat/Vec2.h, StateTargets.h, HitTargets.h, TransformSources.h: UObjects used by Gizmos to expose implementation for customization.
ParameterToTransformAdapters.h: Gizmo ParameterSource implementations that convert axis/plane position and axis-angle to 3D translations/rotations
TransformProxy: wrapper around set of objects that provides a single FTransform, used by TransformGizmo to apply to multiple objects/etc
GizmoMath.h, GizmoRenderingUtil.h: utility functions for Gizmo implementation

#rnx

diff --git a/Engine/Source/Runtime/Experimental/InteractiveToolsFramework/Private/BaseGizmos/AxisAngleGizmo.cpp b/Engine/Source/Runtime/Experimental/InteractiveToolsFramework/Private/BaseGizmos/AxisAngleGizmo.cpp
new file mode 100644
--- /dev/null
+++ b/Engine/Source/Runtime/Experimental/InteractiveToolsFramework/Private/BaseGizmos/AxisAngleGizmo.cpp
@@ -0,0 +1,167 @@
+// Copyright 1998-2019 Epic Games, Inc. All Rights Reserved.
+
+#include "BaseGizmos/AxisAngleGizmo.h"
+#include "InteractiveGizmoManager.h"
+#include "BaseBehaviors/ClickDragBehavior.h"
+#include "BaseBehaviors/MouseHoverBehavior.h"
+#include "BaseGizmos/GizmoMath.h"
+
+
+
+UInteractiveGizmo* UAxisAngleGizmoBuilder::BuildGizmo(const FToolBuilderState& SceneState) const
+{
+	UAxisAngleGizmo* NewGizmo = NewObject<UAxisAngleGizmo>(SceneState.GizmoManager);
+	return NewGizmo;
+}
+
+
+
+
+void UAxisAngleGizmo::Setup()
+{
+	UInteractiveGizmo::Setup();
+
+	// Add default mouse input behavior
+	UClickDragInputBehavior* MouseBehavior = NewObject<UClickDragInputBehavior>();
+	MouseBehavior->Initialize(this);
+	MouseBehavior->SetDefaultPriority(FInputCapturePriority(FInputCapturePriority::DEFAULT_GIZMO_PRIORITY));
+	AddInputBehavior(MouseBehavior);
+	UMouseHoverBehavior* HoverBehavior = NewObject<UMouseHoverBehavior>();
+	HoverBehavior->Initialize(this);
+	HoverBehavior->SetDefaultPriority(FInputCapturePriority(FInputCapturePriority::DEFAULT_GIZMO_PRIORITY));
+	AddInputBehavior(HoverBehavior);
+
+	AxisSource = NewObject<UGizmoConstantAxisSource>(this);
+	AngleSource = NewObject<UGizmoLocalFloatParameterSource>(this);
+	HitTarget = NewObject<UGizmoComponentHitTarget>(this);
+	StateTarget = NewObject<UGizmoNilStateTarget>(this);
+
+	bInInteraction = false;
+}
+
+
+FInputRayHit UAxisAngleGizmo::CanBeginClickDragSequence(const FInputDeviceRay& PressPos)
+{
+	FInputRayHit GizmoHit;
+	if (HitTarget && AxisSource && AngleSource)
+	{
+		GizmoHit = HitTarget->IsHit(PressPos);
+		if (GizmoHit.bHit)
+		{
+			LastHitPosition = PressPos.WorldRay.PointAt(GizmoHit.HitDepth);
+		}
+	}
+	return GizmoHit;
+}
+
+void UAxisAngleGizmo::OnClickPress(const FInputDeviceRay& PressPos)
+{
+	check(bInInteraction == false);
+
+	RotationOrigin = GizmoMath::ProjectPointOntoLine(LastHitPosition, AxisSource->GetOrigin(), RotationAxis);
+	RotationAxis = AxisSource->GetDirection();
+
+	if (AxisSource->HasTangentVectors())
+	{
+		AxisSource->GetTangentVectors(RotationPlaneX, RotationPlaneY);
+	}
+	else
+	{
+		GizmoMath::MakeNormalPlaneBasis(RotationAxis, RotationPlaneX, RotationPlaneY);
+	}
+
+	bool bIntersects; FVector IntersectionPoint;
+	GizmoMath::RayPlaneIntersectionPoint(
+		RotationOrigin, RotationAxis,
+		PressPos.WorldRay.Origin, PressPos.WorldRay.Direction,
+		bIntersects, IntersectionPoint);
+	check(bIntersects);  // need to handle this case...
+
+	InteractionStartPoint = IntersectionPoint;
+	InteractionCurPoint = InteractionStartPoint;
+
+	InteractionStartAngle = GizmoMath::ComputeAngleInPlane(InteractionCurPoint,
+		RotationOrigin, RotationAxis, RotationPlaneX, RotationPlaneY);
+	InteractionCurAngle = InteractionStartAngle;
+
+	InitialTargetAngle = AngleSource->GetParameter();
+	AngleSource->BeginModify();
+
+	bInInteraction = true;
+
+	if (StateTarget)
+	{
+		StateTarget->BeginUpdate();
+	}
+}
+
+
+void UAxisAngleGizmo::OnClickDrag(const FInputDeviceRay& DragPos)
+{
+	check(bInInteraction);
+
+	bool bIntersects; FVector IntersectionPoint;
+	GizmoMath::RayPlaneIntersectionPoint(
+		RotationOrigin, RotationAxis,
+		DragPos.WorldRay.Origin, DragPos.WorldRay.Direction,
+		bIntersects, IntersectionPoint);
+
+	if (bIntersects == false)
+	{
+		return;
+	}
+
+	InteractionCurPoint = IntersectionPoint;
+
+	InteractionCurAngle = GizmoMath::ComputeAngleInPlane(InteractionCurPoint,
+		RotationOrigin, RotationAxis, RotationPlaneX, RotationPlaneY);
+
+	float DeltaAngle = InteractionCurAngle - InteractionStartAngle;
+	float NewAngle = InitialTargetAngle + DeltaAngle;
+
+	AngleSource->SetParameter(NewAngle);
+}
+
+
+void UAxisAngleGizmo::OnClickRelease(const FInputDeviceRay& ReleasePos)
+{
+	check(bInInteraction);
+
+	AngleSource->EndModify();
+	if (StateTarget)
+	{
+		StateTarget->EndUpdate();
+	}
+	bInInteraction = false;
+}
+
+
+void UAxisAngleGizmo::OnTerminateDragSequence()
+{
+	check(bInInteraction);
+
+	AngleSource->EndModify();
+	if (StateTarget)
+	{
+		StateTarget->EndUpdate();
+	}
+	bInInteraction = false;
+}
+
+
+
+void UAxisAngleGizmo::OnUpdateHover(const FInputDeviceRay& DevicePos)
+{
+	FInputRayHit GizmoHit;
+	if (HitTarget)
+	{
+		GizmoHit = HitTarget->IsHit(DevicePos);
+		if (GizmoHit.bHit)
+		{
+			HitTarget->UpdateHoverState(true);
+			return;
+		}
+	}
+	HitTarget->UpdateHoverState(false);
+}
+
diff --git a/Engine/Source/Runtime/Experimental/InteractiveToolsFramework/Private/BaseGizmos/AxisPositionGizmo.cpp b/Engine/Source/Runtime/Experimental/InteractiveToolsFramework/Private/BaseGizmos/AxisPositionGizmo.cpp
new file mode 100644
--- /dev/null
+++ b/Engine/Source/Runtime/Experimental/InteractiveToolsFramework/Private/BaseGizmos/AxisPositionGizmo.cpp
@@ -0,0 +1,141 @@
+// Copyright 1998-2019 Epic Games, Inc. All Rights Reserved.
+
+#include "BaseGizmos/AxisPositionGizmo.h"
+#include "InteractiveGizmoManager.h"
+#include "BaseBehaviors/ClickDragBehavior.h"
+#include "BaseBehaviors/MouseHoverBehavior.h"
+#include "BaseGizmos/GizmoMath.h"
+
+
+
+UInteractiveGizmo* UAxisPositionGizmoBuilder::BuildGizmo(const FToolBuilderState& SceneState) const
+{
+	UAxisPositionGizmo* NewGizmo = NewObject<UAxisPositionGizmo>(SceneState.GizmoManager);
+	return NewGizmo;
+}
+
+
+
+
+void UAxisPositionGizmo::Setup()
+{
+	UInteractiveGizmo::Setup();
+
+	// Add default mouse input behavior
+	UClickDragInputBehavior* MouseBehavior = NewObject<UClickDragInputBehavior>();
+	MouseBehavior->Initialize(this);
+	MouseBehavior->SetDefaultPriority(FInputCapturePriority(FInputCapturePriority::DEFAULT_GIZMO_PRIORITY));
+	AddInputBehavior(MouseBehavior);
+	UMouseHoverBehavior* HoverBehavior = NewObject<UMouseHoverBehavior>();
+	HoverBehavior->Initialize(this);
+	HoverBehavior->SetDefaultPriority(FInputCapturePriority(FInputCapturePriority::DEFAULT_GIZMO_PRIORITY));
+	AddInputBehavior(HoverBehavior);
+
+	AxisSource = NewObject<UGizmoConstantAxisSource>(this);
+	ParameterSource = NewObject<UGizmoLocalFloatParameterSource>(this);
+	HitTarget = NewObject<UGizmoComponentHitTarget>(this);
+	StateTarget = NewObject<UGizmoNilStateTarget>(this);
+
+	bInInteraction = false;
+}
+
+
+FInputRayHit UAxisPositionGizmo::CanBeginClickDragSequence(const FInputDeviceRay& PressPos)
+{
+	FInputRayHit GizmoHit;
+	if (HitTarget && AxisSource && ParameterSource)
+	{
+		GizmoHit = HitTarget->IsHit(PressPos);
+		if (GizmoHit.bHit)
+		{
+			LastHitPosition = PressPos.WorldRay.PointAt(GizmoHit.HitDepth);
+		}
+	}
+	return GizmoHit;
+}
+
+void UAxisPositionGizmo::OnClickPress(const FInputDeviceRay& PressPos)
+{
+	InteractionOrigin = LastHitPosition;
+	InteractionAxis = AxisSource->GetDirection();
+
+	FVector RayNearestPt; float RayNearestParam;
+	GizmoMath::NearestPointOnLineToRay(InteractionOrigin, InteractionAxis,
+		PressPos.WorldRay.Origin, PressPos.WorldRay.Direction,
+		InteractionStartPoint, InteractionStartParameter,
+		RayNearestPt, RayNearestParam);
+
+	InteractionCurPoint = InteractionStartPoint;
+	InteractionCurParameter = InteractionStartParameter;
+
+	InitialTargetParameter = ParameterSource->GetParameter();
+	ParameterSource->BeginModify();
+
+	bInInteraction = true;
+
+	if (StateTarget)
+	{
+		StateTarget->BeginUpdate();
+	}
+}
+
+void UAxisPositionGizmo::OnClickDrag(const FInputDeviceRay& DragPos)
+{
+	FVector AxisNearestPt; float AxisNearestParam;
+	FVector RayNearestPt; float RayNearestParam;
+	GizmoMath::NearestPointOnLineToRay(InteractionOrigin, InteractionAxis,
+		DragPos.WorldRay.Origin, DragPos.WorldRay.Direction,
+		AxisNearestPt, AxisNearestParam,
+		RayNearestPt, RayNearestParam);
+
+	InteractionCurPoint = AxisNearestPt;
+	InteractionCurParameter = AxisNearestParam;
+
+	float DeltaParam = InteractionCurParameter - InteractionStartParameter;
+	float NewValue = InitialTargetParameter + DeltaParam;
+
+	ParameterSource->SetParameter(NewValue);
+}
+
+void UAxisPositionGizmo::OnClickRelease(const FInputDeviceRay& ReleasePos)
+{
+	check(bInInteraction);
+
+	ParameterSource->EndModify();
+	if (StateTarget)
+	{
+		StateTarget->EndUpdate();
+	}
+	bInInteraction = false;
+}
+
+
+void UAxisPositionGizmo::OnTerminateDragSequence()
+{
+	check(bInInteraction);
+
+	ParameterSource->EndModify();
+	if (StateTarget)
+	{
+		StateTarget->EndUpdate();
+	}
+	bInInteraction = false;
+}
+
+
+
+void UAxisPositionGizmo::OnUpdateHover(const FInputDeviceRay& DevicePos)
+{
+	FInputRayHit GizmoHit;
+	if (HitTarget)
+	{
+		GizmoHit = HitTarget->IsHit(DevicePos);
+		if (GizmoHit.bHit)
+		{
+			HitTarget->UpdateHoverState(true);
+			return;
+		}
+	}
+	HitTarget->UpdateHoverState(false);
+}
+
diff --git a/Engine/Source/Runtime/Experimental/InteractiveToolsFramework/Private/BaseGizmos/AxisSources.cpp b/Engine/Source/Runtime/Experimental/InteractiveToolsFramework/Private/BaseGizmos/AxisSources.cpp
new file mode 100644
--- /dev/null
+++ b/Engine/Source/Runtime/Experimental/InteractiveToolsFramework/Private/BaseGizmos/AxisSources.cpp
@@ -0,0 +1,48 @@
+// Copyright 1998-2019 Epic Games, Inc. All Rights Reserved.
+
+#include "BaseGizmos/AxisSources.h"
+#include "Components/SceneComponent.h"
+
+FVector UGizmoComponentAxisSource::GetOrigin() const
+{
+	const FTransform& WorldTransform = Component->GetComponentToWorld();
+	return WorldTransform.GetLocation();
+}
+
+FVector UGizmoComponentAxisSource::GetDirection() const
+{
+	const FTransform& WorldTransform = Component->GetComponentToWorld();
+	FVector Axis(0, 0, 0);
+	Axis[FMath::Clamp(AxisIndex, 0, 2)] = 1.0;
+	if (bLocalAxes)
+	{
+		return WorldTransform.GetRotation().RotateVector(Axis);
+	} 
+	else
+	{
+		return Axis;
+	}
+}
+
+void UGizmoComponentAxisSource::GetTangentVectors(FVector& TangentXOut, FVector& TangentYOut) const
+{
+	const FTransform& WorldTransform = Component->GetComponentToWorld();
+	TangentXOut = FVector(0, 1, 0);
+	TangentYOut = FVector(0, 0, 1);
+	int Index = FMath::Clamp(AxisIndex, 0, 2);
+	if (Index == 1)
+	{
+		TangentXOut = FVector(-1, 0, 0);
+		TangentYOut = FVector(0, 0, 1);
+	}
+	else if (Index == 2)
+	{
+		TangentXOut = FVector(1, 0, 0);
+		TangentYOut = FVector(0, 1, 0);
+	}
+	if (bLocalAxes)
+	{
+		TangentXOut = WorldTransform.GetRotation().RotateVector(TangentXOut);
+		TangentYOut = WorldTransform.GetRotation().RotateVector(TangentYOut);
+	}
+}
\ No newline at end of file
diff --git a/Engine/Source/Runtime/Experimental/InteractiveToolsFramework/Private/BaseGizmos/GizmoActor.cpp b/Engine/Source/Runtime/Experimental/InteractiveToolsFramework/Private/BaseGizmos/GizmoActor.cpp
new file mode 100644
--- /dev/null
+++ b/Engine/Source/Runtime/Experimental/InteractiveToolsFramework/Private/BaseGizmos/GizmoActor.cpp
@@ -0,0 +1,61 @@
+// Copyright 1998-2019 Epic Games, Inc. All Rights Reserved.
+
+#include "BaseGizmos/GizmoActor.h"
+
+#include "BaseGizmos/GizmoArrowComponent.h"
+#include "BaseGizmos/GizmoRectangleComponent.h"
+#include "BaseGizmos/GizmoCircleComponent.h"
+
+#include "Components/PrimitiveComponent.h"
+#include "Engine/World.h"
+
+
+#define LOCTEXT_NAMESPACE "UTransformGizmo"
+
+
+UGizmoArrowComponent* AGizmoActor::AddDefaultArrowComponent(
+	UWorld* World, AActor* Actor,
+	const FLinearColor& Color, const FVector& LocalDirection)
+{
+	UGizmoArrowComponent* NewArrow = NewObject<UGizmoArrowComponent>(Actor);
+	Actor->AddInstanceComponent(NewArrow);
+	NewArrow->AttachToComponent(Actor->GetRootComponent(), FAttachmentTransformRules::KeepRelativeTransform);
+	NewArrow->Direction = LocalDirection;
+	NewArrow->Color = Color;
+	NewArrow->RegisterComponent();
+	return NewArrow;
+}
+
+
+
+UGizmoRectangleComponent* AGizmoActor::AddDefaultRectangleComponent(
+	UWorld* World, AActor* Actor,
+	const FLinearColor& Color, const FVector& PlaneAxis1, const FVector& PlaneAxisx2)
+{
+	UGizmoRectangleComponent* NewRectangle = NewObject<UGizmoRectangleComponent>(Actor);
+	Actor->AddInstanceComponent(NewRectangle);
+	NewRectangle->AttachToComponent(Actor->GetRootComponent(), FAttachmentTransformRules::KeepRelativeTransform);
+	NewRectangle->DirectionX = PlaneAxis1;
+	NewRectangle->DirectionY = PlaneAxisx2;
+	NewRectangle->Color = Color;
+	NewRectangle->RegisterComponent();
+	return NewRectangle;
+}
+
+
+UGizmoCircleComponent* AGizmoActor::AddDefaultCircleComponent(
+	UWorld* World, AActor* Actor,
+	const FLinearColor& Color, const FVector& PlaneNormal)
+{
+	UGizmoCircleComponent* NewCircle = NewObject<UGizmoCircleComponent>(Actor);
+	Actor->AddInstanceComponent(NewCircle);
+	NewCircle->AttachToComponent(Actor->GetRootComponent(), FAttachmentTransformRules::KeepRelativeTransform);
+	NewCircle->Normal = PlaneNormal;
+	NewCircle->Color = Color;
+	NewCircle->RegisterComponent();
+	return NewCircle;
+}
+
+
+
+#undef LOCTEXT_NAMESPACE
\ No newline at end of file
diff --git a/Engine/Source/Runtime/Experimental/InteractiveToolsFramework/Private/BaseGizmos/GizmoArrowComponent.cpp b/Engine/Source/Runtime/Experimental/InteractiveToolsFramework/Private/BaseGizmos/GizmoArrowComponent.cpp
new file mode 100644
--- /dev/null
+++ b/Engine/Source/Runtime/Experimental/InteractiveToolsFramework/Private/BaseGizmos/GizmoArrowComponent.cpp
@@ -0,0 +1,194 @@
+// Copyright 1998-2019 Epic Games, Inc. All Rights Reserved.
+ 
+#include "BaseGizmos/GizmoArrowComponent.h"
+#include "BaseGizmos/GizmoRenderingUtil.h"
+#include "PrimitiveSceneProxy.h"
+#include "Materials/Material.h"
+#include "MaterialShared.h"
+#include "SceneManagement.h"
+
+
+class FGizmoArrowComponentSceneProxy final : public FPrimitiveSceneProxy
+{
+public:
+	SIZE_T GetTypeHash() const override
+	{
+		static size_t UniquePointer;
+		return reinterpret_cast<size_t>(&UniquePointer);
+	}
+
+	FGizmoArrowComponentSceneProxy(const UGizmoArrowComponent* InComponent)
+		: FPrimitiveSceneProxy(InComponent),
+		Color(InComponent->Color),
+		Direction(InComponent->Direction),
+		Gap(InComponent->Gap),
+		Length(InComponent->Length),
+		Thickness(InComponent->Thickness),
+		HoverThicknessMultiplier(InComponent->HoverSizeMultiplier)
+	{
+	}
+
+	virtual void GetDynamicMeshElements(const TArray<const FSceneView*>& Views, const FSceneViewFamily& ViewFamily, uint32 VisibilityMap, FMeshElementCollector& Collector) const override
+	{
+		// find view to use for gizmo sizing/etc
+		const FSceneView* GizmoControlView = GizmoRenderingUtil::FindActiveSceneView(Views, ViewFamily, VisibilityMap);
+		if (GizmoControlView == nullptr)
+		{
+			return;
+		}
+
+		const FMatrix& LocalToWorldMatrix = GetLocalToWorld();
+		FVector Origin = LocalToWorldMatrix.TransformPosition(FVector::ZeroVector);
+
+		// direction to origin of gizmo
+		FVector ViewDirection = Origin - GizmoControlView->ViewLocation;
+		ViewDirection.Normalize();
+
+		FVector ArrowDirection = LocalToWorldMatrix.TransformVector(Direction);
+		bool bFlipped = (FVector::DotProduct(ViewDirection, ArrowDirection) > 0);
+		ArrowDirection = (bFlipped) ? -ArrowDirection : ArrowDirection;
+		if (bFlippedExternal != nullptr)
+		{
+			*bFlippedExternal = bFlipped;
+		}
+
+		if (bExternalRenderVisibility != nullptr)
+		{
+			//*bExternalRenderVisibility = FMath::Abs(FVector::DotProduct(ArrowDirection, ViewDirection)) < 0.985f;   // ~10 degrees
+			*bExternalRenderVisibility = FMath::Abs(FVector::DotProduct(ArrowDirection, ViewDirection)) < 0.965f;   // ~15 degrees
+			if (*bExternalRenderVisibility == false)
+			{
+				return;
+			}
+		}
+
+		float LengthScale = 1.0f;
+		if (ExternalDynamicPixelToWorldScale != nullptr)
+		{
+			float PixelToWorldScale = GizmoRenderingUtil::CalculateLocalPixelToWorldScale(GizmoControlView, Origin);
+			*ExternalDynamicPixelToWorldScale = PixelToWorldScale;
+			LengthScale = PixelToWorldScale;
+		}
+
+		float UseThickness = (bExternalHoverState != nullptr && *bExternalHoverState == true) ?
+			(HoverThicknessMultiplier*Thickness) : (Thickness);
+
+		double StartDist = LengthScale * Gap;
+		double EndDist = LengthScale * (Gap+Length);
+
+		for (int32 ViewIndex = 0; ViewIndex < Views.Num(); ViewIndex++)
+		{
+			if (VisibilityMap & (1 << ViewIndex))
+			{
+				const FSceneView* View = Views[ViewIndex];
+				FPrimitiveDrawInterface* PDI = Collector.GetPDI(ViewIndex);
+
+				FVector StartPoint = Origin + StartDist*ArrowDirection;
+				FVector EndPoint = Origin + EndDist*ArrowDirection;
+
+				PDI->DrawLine(StartPoint, EndPoint, Color, SDPG_Foreground, UseThickness, 0.0f, true);
+			}
+		}
+	}
+
+	virtual FPrimitiveViewRelevance GetViewRelevance(const FSceneView* View) const override
+	{
+		FPrimitiveViewRelevance Result;
+		Result.bDrawRelevance = IsShown(View);
+		Result.bDynamicRelevance = true;
+		Result.bShadowRelevance = false;
+		Result.bEditorPrimitiveRelevance = UseEditorCompositing(View);
+		Result.bRenderCustomDepth = ShouldRenderCustomDepth();
+		return Result;
+	}
+
+	virtual bool CanBeOccluded() const override
+	{
+		return false;
+	}
+
+	virtual uint32 GetMemoryFootprint(void) const override { return sizeof *this + GetAllocatedSize(); }
+	uint32 GetAllocatedSize(void) const { return FPrimitiveSceneProxy::GetAllocatedSize(); }
+
+
+	void SetExternalFlip(bool* bFlipped)
+	{
+		bFlippedExternal = bFlipped;
+	}
+
+	void SetExternalDynamicPixelToWorldScale(float* DynamicPixelToWorldScale)
+	{
+		ExternalDynamicPixelToWorldScale = DynamicPixelToWorldScale;
+	}
+
+	void SetExternalRenderVisibility(bool* bRenderVisibility)
+	{
+		bExternalRenderVisibility = bRenderVisibility;
+	}
+
+	void SetExternalHoverState(bool* HoverState)
+	{
+		bExternalHoverState = HoverState;
+	}
+
+
+private:
+	FLinearColor Color;
+	FVector Direction;
+	float Gap;
+	float Length;
+	float Thickness;
+	float HoverThicknessMultiplier;
+
+	bool* bFlippedExternal = nullptr;
+	float* ExternalDynamicPixelToWorldScale = nullptr;
+	bool* bExternalRenderVisibility = nullptr;
+	bool* bExternalHoverState = nullptr;
+};
+
+
+FPrimitiveSceneProxy* UGizmoArrowComponent::CreateSceneProxy()
+{
+	FGizmoArrowComponentSceneProxy* NewProxy = new FGizmoArrowComponentSceneProxy(this);
+	NewProxy->SetExternalFlip(&bFlipped);
+	NewProxy->SetExternalDynamicPixelToWorldScale(&DynamicPixelToWorldScale);
+	NewProxy->SetExternalRenderVisibility(&bRenderVisibility);
+	NewProxy->SetExternalHoverState(&bHovering);
+	return NewProxy;
+}
+
+FBoxSphereBounds UGizmoArrowComponent::CalcBounds(const FTransform& LocalToWorld) const
+{
+	return FBoxSphereBounds(FSphere(FVector::ZeroVector, Gap+Length).TransformBy(LocalToWorld));
+}
+
+bool UGizmoArrowComponent::LineTraceComponent(FHitResult& OutHit, const FVector Start, const FVector End, const FCollisionQueryParams& Params)
+{
+	if (bRenderVisibility == false)
+	{
+		return false;
+	}
+
+	const FTransform& Transform = this->GetComponentToWorld();
+
+	FVector UseDirection = (bFlipped) ? -Direction : Direction;
+	float LengthScale = DynamicPixelToWorldScale;
+	double StartDist = LengthScale * Gap;
+	double EndDist = LengthScale * (Gap + Length);
+
+	FVector Point0 = Transform.TransformPosition( StartDist * UseDirection);
+	FVector Point1 = Transform.TransformPosition( EndDist * UseDirection);
+
+	FVector NearestArrow, NearestLine;
+	FMath::SegmentDistToSegmentSafe(Point0, Point1, Start, End, NearestArrow, NearestLine);
+	double Distance = FVector::Distance(NearestArrow, NearestLine);
+	if (Distance > PixelHitDistanceThreshold*DynamicPixelToWorldScale)
+	{
+		return false;
+	}
+
+	OutHit.Component = this;
+	OutHit.Distance = FVector::Distance(Start, NearestLine);
+	OutHit.ImpactPoint = NearestLine;
+	return true;
+}
\ No newline at end of file
diff --git a/Engine/Source/Runtime/Experimental/InteractiveToolsFramework/Private/BaseGizmos/GizmoCircleComponent.cpp b/Engine/Source/Runtime/Experimental/InteractiveToolsFramework/Private/BaseGizmos/GizmoCircleComponent.cpp
new file mode 100644
--- /dev/null
+++ b/Engine/Source/Runtime/Experimental/InteractiveToolsFramework/Private/BaseGizmos/GizmoCircleComponent.cpp
@@ -0,0 +1,236 @@
+// Copyright 1998-2019 Epic Games, Inc. All Rights Reserved.
+ 
+#include "BaseGizmos/GizmoCircleComponent.h"
+#include "BaseGizmos/GizmoRenderingUtil.h"
+#include "BaseGizmos/GizmoMath.h"
+#include "PrimitiveSceneProxy.h"
+
+
+class FGizmoCircleComponentSceneProxy final : public FPrimitiveSceneProxy
+{
+public:
+	SIZE_T GetTypeHash() const override
+	{
+		static size_t UniquePointer;
+		return reinterpret_cast<size_t>(&UniquePointer);
+	}
+
+	FGizmoCircleComponentSceneProxy(const UGizmoCircleComponent* InComponent)
+		: FPrimitiveSceneProxy(InComponent),
+		Color(InComponent->Color),
+		Normal(InComponent->Normal),
+		Radius(InComponent->Radius),
+		Thickness(InComponent->Thickness),
+		NumSides(InComponent->NumSides),
+		bViewAligned(InComponent->bViewAligned),
+		HoverThicknessMultiplier(InComponent->HoverSizeMultiplier)
+	{
+	}
+
+	virtual void GetDynamicMeshElements(const TArray<const FSceneView*>& Views, const FSceneViewFamily& ViewFamily, uint32 VisibilityMap, FMeshElementCollector& Collector) const override
+	{
+		// find view to use for gizmo sizing/etc
+		const FSceneView* GizmoControlView = GizmoRenderingUtil::FindActiveSceneView(Views, ViewFamily, VisibilityMap);
+		if (GizmoControlView == nullptr)
+		{
+			return;
+		}
+
+		const FMatrix& LocalToWorldMatrix = GetLocalToWorld();
+		FVector Origin = LocalToWorldMatrix.TransformPosition(FVector::ZeroVector);
+
+		FVector PlaneX, PlaneY;
+		GizmoMath::MakeNormalPlaneBasis(Normal, PlaneX, PlaneY);
+
+		// direction to origin of gizmo
+		FVector GizmoViewDirection = Origin - GizmoControlView->ViewLocation;
+		GizmoViewDirection.Normalize();
+
+		float LengthScale = 1.0f;
+		if (ExternalDynamicPixelToWorldScale != nullptr)
+		{
+			float PixelToWorldScale = GizmoRenderingUtil::CalculateLocalPixelToWorldScale(GizmoControlView, Origin);
+			*ExternalDynamicPixelToWorldScale = PixelToWorldScale;
+			LengthScale = PixelToWorldScale;
+		}
+
+		double UseRadius = LengthScale * Radius;
+
+		FLinearColor BackColor = FLinearColor(0.5f, 0.5f, 0.5f);
+		float BackThickness = 0.5f;
+
+		float UseThickness = (bExternalHoverState != nullptr && *bExternalHoverState == true) ? 
+			(HoverThicknessMultiplier*Thickness) : (Thickness);
+
+		for (int32 ViewIndex = 0; ViewIndex < Views.Num(); ViewIndex++)
+		{
+			if (VisibilityMap & (1 << ViewIndex))
+			{
+				const FSceneView* View = Views[ViewIndex];
+				FPrimitiveDrawInterface* PDI = Collector.GetPDI(ViewIndex);
+
+				FVector UpVector = View->GetViewUp();
+
+				const float	AngleDelta = 2.0f * PI / NumSides;
+
+				if (bViewAligned)
+				{
+					FVector ViewVector = View->GetViewDirection();
+					FVector WorldOrigin = LocalToWorldMatrix.TransformPosition(FVector::ZeroVector);
+					FVector WorldPlaneX, WorldPlaneY;
+					GizmoMath::MakeNormalPlaneBasis(ViewVector, WorldPlaneX, WorldPlaneY);
+
+					FVector	LastVertex = WorldOrigin + WorldPlaneX * UseRadius;
+					for (int32 SideIndex = 0; SideIndex < NumSides; SideIndex++)
+					{
+						float DeltaX = FMath::Cos(AngleDelta * (SideIndex + 1));
+						float DeltaY = FMath::Sin(AngleDelta * (SideIndex + 1));
+						const FVector DeltaVector = WorldPlaneX * DeltaX + WorldPlaneY * DeltaY;
+						const FVector Vertex = WorldOrigin + UseRadius * DeltaVector;
+						PDI->DrawLine(LastVertex, Vertex, Color, SDPG_Foreground, UseThickness, 0.0f, true);
+						LastVertex = Vertex;
+					}
+				}
+				else
+				{
+					FVector WorldOrigin = LocalToWorldMatrix.TransformPosition(FVector::ZeroVector);
+					FVector WorldPlaneX = LocalToWorldMatrix.TransformVector(PlaneX);
+					FVector WorldPlaneY = LocalToWorldMatrix.TransformVector(PlaneY);
+
+					FVector	LastVertex = WorldOrigin + WorldPlaneX * UseRadius;
+					bool bLastVisible = FVector::DotProduct(WorldPlaneX, GizmoViewDirection) < 0;
+					for (int32 SideIndex = 0; SideIndex < NumSides; SideIndex++)
+					{
+						float DeltaX = FMath::Cos(AngleDelta * (SideIndex + 1));
+						float DeltaY = FMath::Sin(AngleDelta * (SideIndex + 1));
+						const FVector DeltaVector = WorldPlaneX * DeltaX + WorldPlaneY * DeltaY;
+						const FVector Vertex = WorldOrigin + UseRadius * DeltaVector;
+						bool bVertexVisible = FVector::DotProduct(DeltaVector, GizmoViewDirection) < 0;
+						if (bLastVisible && bVertexVisible)
+						{
+							PDI->DrawLine(LastVertex, Vertex, Color, SDPG_Foreground, UseThickness, 0.0f, true);
+						}
+						else
+						{
+							PDI->DrawLine(LastVertex, Vertex, BackColor, SDPG_Foreground, BackThickness, 0.0f, true);
+						}
+						bLastVisible = bVertexVisible;
+						LastVertex = Vertex;
+					}
+				}
+			}
+		}
+	}
+
+	virtual FPrimitiveViewRelevance GetViewRelevance(const FSceneView* View) const override
+	{
+		FPrimitiveViewRelevance Result;
+		Result.bDrawRelevance = IsShown(View);
+		Result.bDynamicRelevance = true;
+		Result.bShadowRelevance = false;
+		Result.bEditorPrimitiveRelevance = UseEditorCompositing(View);
+		Result.bRenderCustomDepth = ShouldRenderCustomDepth();
+		return Result;
+	}
+
+	virtual bool CanBeOccluded() const override
+	{
+		return false;
+	}
+
+	virtual uint32 GetMemoryFootprint(void) const override { return sizeof *this + GetAllocatedSize(); }
+	uint32 GetAllocatedSize(void) const { return FPrimitiveSceneProxy::GetAllocatedSize(); }
+
+
+	void SetExternalDynamicPixelToWorldScale(float* DynamicPixelToWorldScale)
+	{
+		ExternalDynamicPixelToWorldScale = DynamicPixelToWorldScale;
+	}
+
+	void SetExternalHoverState(bool* HoverState)
+	{
+		bExternalHoverState = HoverState;
+	}
+
+private:
+	FLinearColor Color;
+	FVector Normal;
+	float Radius;
+	float Thickness;
+	int NumSides;
+	bool bViewAligned;
+	float HoverThicknessMultiplier;
+
+	float* ExternalDynamicPixelToWorldScale = nullptr;
+	bool* bExternalHoverState = nullptr;
+};
+
+
+
+
+FPrimitiveSceneProxy* UGizmoCircleComponent::CreateSceneProxy()
+{
+	FGizmoCircleComponentSceneProxy* NewProxy = new FGizmoCircleComponentSceneProxy(this);
+	NewProxy->SetExternalDynamicPixelToWorldScale(&DynamicPixelToWorldScale);
+	NewProxy->SetExternalHoverState(&bHovering);
+	return NewProxy;
+}
+
+FBoxSphereBounds UGizmoCircleComponent::CalcBounds(const FTransform& LocalToWorld) const
+{
+	return FBoxSphereBounds(FSphere(FVector::ZeroVector, Radius).TransformBy(LocalToWorld));
+}
+
+bool UGizmoCircleComponent::LineTraceComponent(FHitResult& OutHit, const FVector Start, const FVector End, const FCollisionQueryParams& Params)
+{
+	float LengthScale = DynamicPixelToWorldScale;
+	double UseRadius = LengthScale * Radius;
+
+	const FTransform& Transform = this->GetComponentToWorld();
+	FVector WorldNormal = Transform.TransformVector(Normal);
+	FVector WorldOrigin = Transform.TransformPosition(FVector::ZeroVector);
+	FPlane CirclePlane(WorldOrigin, WorldNormal);
+
+	FRay Ray(Start, End - Start, false);
+	FVector HitPos = FMath::RayPlaneIntersection(Ray.Origin, Ray.Direction, CirclePlane);
+	float RayParameter = Ray.GetParameter(HitPos);
+	if (RayParameter < 0 || RayParameter > Ray.GetParameter(End))
+	{
+		return false;
+	}
+
+	FVector NearestCircle;
+	GizmoMath::ClosetPointOnCircle(HitPos, WorldOrigin, WorldNormal, UseRadius, NearestCircle);
+
+	FVector NearestRay = Ray.ClosestPoint(NearestCircle);
+
+	double Distance = FVector::Distance(NearestCircle, NearestRay);
+	if (Distance > PixelHitDistanceThreshold*DynamicPixelToWorldScale)
+	{
+		return false;
+	}
+
+
+	// filter out hits on "back" of sphere that circle lies on
+	if (bOnlyAllowFrontFacingHits)
+	{
+		bool bSphereIntersects = false;
+		FVector SphereHitPoint;
+		FVector RayToCirclePointDirection = (NearestCircle - Ray.Origin);
+		RayToCirclePointDirection.Normalize();
+		GizmoMath::RaySphereIntersection(
+			WorldOrigin, UseRadius, Ray.Origin, RayToCirclePointDirection, bSphereIntersects, SphereHitPoint);
+		if (bSphereIntersects)
+		{
+			if (FVector::DistSquared(SphereHitPoint, NearestCircle) > UseRadius*0.1f)
+			{
+				return false;
+			}
+		}
+	}
+
+	OutHit.Component = this;
+	OutHit.Distance = FVector::Distance(Start, NearestRay);
+	OutHit.ImpactPoint = NearestRay;
+	return true;
+}
\ No newline at end of file
diff --git a/Engine/Source/Runtime/Experimental/InteractiveToolsFramework/Private/BaseGizmos/GizmoInterfaces.cpp b/Engine/Source/Runtime/Experimental/InteractiveToolsFramework/Private/BaseGizmos/GizmoInterfaces.cpp
new file mode 100644
--- /dev/null
+++ b/Engine/Source/Runtime/Experimental/InteractiveToolsFramework/Private/BaseGizmos/GizmoInterfaces.cpp
@@ -0,0 +1,20 @@
+// Copyright 1998-2019 Epic Games, Inc. All Rights Reserved.
+
+#include "BaseGizmos/GizmoInterfaces.h"
+#include "BaseGizmos/GizmoMath.h"
+
+
+
+void IGizmoAxisSource::GetAxisFrame(
+	FVector& PlaneNormalOut, FVector& PlaneAxis1Out, FVector& PlaneAxis2Out) const
+{
+	PlaneNormalOut = GetDirection();
+	if (HasTangentVectors())
+	{
+		GetTangentVectors(PlaneAxis1Out, PlaneAxis2Out);
+	}
+	else
+	{
+		GizmoMath::MakeNormalPlaneBasis(PlaneNormalOut, PlaneAxis1Out, PlaneAxis2Out);
+	}
+}
diff --git a/Engine/Source/Runtime/Experimental/InteractiveToolsFramework/Private/BaseGizmos/GizmoMath.cpp b/Engine/Source/Runtime/Experimental/InteractiveToolsFramework/Private/BaseGizmos/GizmoMath.cpp
new file mode 100644
--- /dev/null
+++ b/Engine/Source/Runtime/Experimental/InteractiveToolsFramework/Private/BaseGizmos/GizmoMath.cpp
@@ -0,0 +1,208 @@
+// Copyright 1998-2019 Epic Games, Inc. All Rights Reserved.
+
+#include "BaseGizmos/GizmoMath.h"
+
+
+
+
+FVector GizmoMath::ProjectPointOntoLine(
+	const FVector& Point,
+	const FVector& LineOrigin, const FVector& LineDirection)
+{
+	float ProjectionParam = FVector::DotProduct((Point - LineOrigin), LineDirection);
+	return LineOrigin + ProjectionParam * LineDirection;
+}
+
+
+void GizmoMath::NearestPointOnLine(
+	const FVector& LineOrigin, const FVector& LineDirection,
+	const FVector& QueryPoint,
+	FVector& NearestPointOut, float& LineParameterOut)
+{
+	check(LineDirection.IsNormalized());
+	LineParameterOut = FVector::DotProduct( (QueryPoint - LineOrigin), LineDirection);
+	NearestPointOut = LineOrigin + LineParameterOut * LineDirection;
+}
+
+
+void GizmoMath::NearestPointOnLineToRay(
+	const FVector& LineOrigin, const FVector& LineDirection,
+	const FVector& RayOrigin, const FVector& RayDirection,
+	FVector& NearestLinePointOut, float& LineParameterOut,
+	FVector& NearestRayPointOut, float& RayParameterOut)
+{
+	FVector kDiff = LineOrigin - RayOrigin;
+	float a01 = -FVector::DotProduct(LineDirection, RayDirection);
+	float b0 = FVector::DotProduct(kDiff, LineDirection);
+	float c = kDiff.SizeSquared();
+	float det = FMath::Abs((float)1 - a01 * a01);
+	float b1, s0, s1;
+
+	if (det >= SMALL_NUMBER)
+	{
+		b1 = -FVector::DotProduct(kDiff, RayDirection);
+		s1 = a01 * b0 - b1;
+
+		if (s1 >= (float)0)
+		{
+			// Two interior points are closest, one on Line and one on Ray
+			float invDet = ((float)1) / det;
+			s0 = (a01 * b1 - b0) * invDet;
+			s1 *= invDet;
+		}
+		else
+		{
+			// Origin of Ray and interior point of Line are closest.
+			s0 = -b0;
+			s1 = (float)0;
+		}
+	}
+	else
+	{
+		// Lines are parallel, closest pair with one point at Ray origin.
+		s0 = -b0;
+		s1 = (float)0;
+	}
+
+	NearestLinePointOut = LineOrigin + s0 * LineDirection;
+	NearestRayPointOut = RayOrigin + s1 * RayDirection;
+	LineParameterOut = s0;
+	RayParameterOut = s1;
+}
+
+
+
+
+void GizmoMath::RayPlaneIntersectionPoint(
+	const FVector& PlaneOrigin, const FVector& PlaneNormal,
+	const FVector& RayOrigin, const FVector& RayDirection,
+	bool& bIntersectsOut, FVector& PlaneIntersectionPointOut)
+{
+	bIntersectsOut = false;
+	PlaneIntersectionPointOut = PlaneOrigin;
+
+	float PlaneEquationD = -FVector::DotProduct(PlaneOrigin, PlaneNormal);
+	float NormalDot = FVector::DotProduct(RayDirection, PlaneNormal);
+
+	if (FMath::Abs(NormalDot) < SMALL_NUMBER)
+	{
+		return;
+	}
+
+	float RayParam = -( FVector::DotProduct(RayOrigin, PlaneNormal) + PlaneEquationD) / NormalDot;
+	if (RayParam < 0)
+	{
+		return;
+	}
+
+	PlaneIntersectionPointOut = RayOrigin + RayParam * RayDirection;
+	bIntersectsOut = true;
+}
+
+
+void GizmoMath::RaySphereIntersection(
+	const FVector& SphereOrigin, const float SphereRadius,
+	const FVector& RayOrigin, const FVector& RayDirection,
+	bool& bIntersectsOut, FVector& SphereIntersectionPointOut)
+{
+	bIntersectsOut = false;
+	SphereIntersectionPointOut = RayOrigin;
+	 
+	FVector DeltaPos = RayOrigin - SphereOrigin;
+	float a0 = DeltaPos.SizeSquared() - SphereRadius*SphereRadius;
+	float a1 = FVector::DotProduct(RayDirection, DeltaPos);
+	float discr = a1 * a1 - a0;
+	if (discr > 0)   // intersection only when roots are real
+	{
+		bIntersectsOut = true;
+		float root = FMath::Sqrt(discr);
+		float NearRayParam = -a1 + root;		// isn't it always this one?
+		float NearRayParam2 = -a1 - root;
+		float UseRayParam = FMath::Min(NearRayParam, NearRayParam2);
+		SphereIntersectionPointOut = RayOrigin + UseRayParam * RayDirection;
+	}
+}
+
+
+
+void GizmoMath::ClosetPointOnCircle(
+	const FVector& QueryPoint,
+	const FVector& CircleOrigin, const FVector& CircleNormal, float CircleRadius,
+	FVector& ClosestPointOut)
+{
+	FVector PointDelta = QueryPoint - CircleOrigin;
+	FVector DeltaInPlane = PointDelta - FVector::DotProduct(CircleNormal,PointDelta)*CircleNormal;
+	float OriginDist = DeltaInPlane.Size();
+	if (OriginDist > 0.0f)
+	{
+		ClosestPointOut =  CircleOrigin + (CircleRadius / OriginDist) * DeltaInPlane;
+	}
+	else    // all points equidistant, use any one
+	{
+		FVector PlaneX, PlaneY;
+		MakeNormalPlaneBasis(CircleNormal, PlaneX, PlaneY);
+		ClosestPointOut = CircleOrigin + CircleRadius * PlaneX;
+	}
+}
+
+
+
+void GizmoMath::MakeNormalPlaneBasis(
+	const FVector& PlaneNormal,
+	FVector& BasisAxis1Out, FVector& BasisAxis2Out)
+{
+	// Duff et al method, from https://graphics.pixar.com/library/OrthonormalB/paper.pdf
+	if (PlaneNormal.Z < 0)
+	{
+		float A = 1.0f / (1.0f - PlaneNormal.Z);
+		float B = PlaneNormal.X * PlaneNormal.Y * A;
+		BasisAxis1Out.X = 1.0f - PlaneNormal.X * PlaneNormal.X * A;
+		BasisAxis1Out.Y = -B;
+		BasisAxis1Out.Z = PlaneNormal.X;
+		BasisAxis2Out.X = B;
+		BasisAxis2Out.Y = PlaneNormal.Y * PlaneNormal.Y * A - 1.0f;
+		BasisAxis2Out.Z = -PlaneNormal.Y;
+	}
+	else
+	{
+		float A = 1.0f / (1.0f + PlaneNormal.Z);
+		float B = -PlaneNormal.X * PlaneNormal.Y * A;
+		BasisAxis1Out.X = 1.0f - PlaneNormal.X * PlaneNormal.X * A;
+		BasisAxis1Out.Y = B;
+		BasisAxis1Out.Z = -PlaneNormal.X;
+		BasisAxis2Out.X = B;
+		BasisAxis2Out.Y = 1.0f - PlaneNormal.Y * PlaneNormal.Y * A;
+		BasisAxis2Out.Z = -PlaneNormal.Y;
+	}
+}
+
+
+
+float GizmoMath::ComputeAngleInPlane(
+	const FVector& Point,
+	const FVector& PlaneOrigin, const FVector& PlaneNormal,
+	const FVector& PlaneAxis1, const FVector& PlaneAxis2)
+{
+	// project point into plane
+	FVector LocalPoint = Point - PlaneOrigin;
+
+	float X = FVector::DotProduct(LocalPoint, PlaneAxis1);
+	float Y = FVector::DotProduct(LocalPoint, PlaneAxis2);
+
+	float SignedAngle = (float)atan2(Y, X);
+	return SignedAngle;
+}
+
+
+
+
+FVector2D GizmoMath::ComputeCoordinatesInPlane(
+	const FVector& Point,
+	const FVector& PlaneOrigin, const FVector& PlaneNormal,
+	const FVector& PlaneAxis1, const FVector& PlaneAxis2)
+{
+	FVector LocalPoint = Point - PlaneOrigin;
+	float X = FVector::DotProduct(LocalPoint, PlaneAxis1);
+	float Y = FVector::DotProduct(LocalPoint, PlaneAxis2);
+	return FVector2D(X, Y);
+}
diff --git a/Engine/Source/Runtime/Experimental/InteractiveToolsFramework/Private/BaseGizmos/GizmoRectangleComponent.cpp b/Engine/Source/Runtime/Experimental/InteractiveToolsFramework/Private/BaseGizmos/GizmoRectangleComponent.cpp
new file mode 100644
--- /dev/null
+++ b/Engine/Source/Runtime/Experimental/InteractiveToolsFramework/Private/BaseGizmos/GizmoRectangleComponent.cpp
@@ -0,0 +1,250 @@
+// Copyright 1998-2019 Epic Games, Inc. All Rights Reserved.
+
+#include "BaseGizmos/GizmoRectangleComponent.h"
+#include "BaseGizmos/GizmoRenderingUtil.h"
+#include "PrimitiveSceneProxy.h"
+
+
+class FGizmoRectangleComponentSceneProxy final : public FPrimitiveSceneProxy
+{
+public:
+	SIZE_T GetTypeHash() const override
+	{
+		static size_t UniquePointer;
+		return reinterpret_cast<size_t>(&UniquePointer);
+	}
+
+	FGizmoRectangleComponentSceneProxy(const UGizmoRectangleComponent* InComponent)
+		: FPrimitiveSceneProxy(InComponent),
+		Color(InComponent->Color),
+		DirectionX(InComponent->DirectionX),
+		DirectionY(InComponent->DirectionY),
+		OffsetX(InComponent->OffsetX),
+		OffsetY(InComponent->OffsetY),
+		LengthX(InComponent->LengthX),
+		LengthY(InComponent->LengthY),
+		Thickness(InComponent->Thickness),
+		HoverThicknessMultiplier(InComponent->HoverSizeMultiplier)
+	{
+	}
+
+	virtual void GetDynamicMeshElements(const TArray<const FSceneView*>& Views, const FSceneViewFamily& ViewFamily, uint32 VisibilityMap, FMeshElementCollector& Collector) const override
+	{
+		// find view to use for gizmo sizing/etc
+		const FSceneView* GizmoControlView = GizmoRenderingUtil::FindActiveSceneView(Views, ViewFamily, VisibilityMap);
+		if (GizmoControlView == nullptr)
+		{
+			return;
+		}
+
+		const FMatrix& LocalToWorldMatrix = GetLocalToWorld();
+		FVector Origin = LocalToWorldMatrix.TransformPosition(FVector::ZeroVector);
+
+		// direction to origin of gizmo
+		FVector ViewDirection = Origin - GizmoControlView->ViewLocation;
+		ViewDirection.Normalize();
+
+		FVector UseDirectionX = LocalToWorldMatrix.TransformVector(DirectionX);
+		bool bFlippedX = (FVector::DotProduct(ViewDirection, UseDirectionX) > 0);
+		UseDirectionX = (bFlippedX) ? -UseDirectionX : UseDirectionX;
+		if (bFlippedXExternal)
+		{
+			*bFlippedXExternal = bFlippedX;
+		}
+
+		FVector UseDirectionY = LocalToWorldMatrix.TransformVector(DirectionY);
+		bool bFlippedY = (FVector::DotProduct(ViewDirection, UseDirectionY) > 0);
+		UseDirectionY = (bFlippedY) ? -UseDirectionY : UseDirectionY;
+		if (bFlippedYExternal)
+		{
+			*bFlippedYExternal = bFlippedY;
+		}
+
+		if (bExternalRenderVisibility != nullptr)
+		{
+			FVector PlaneNormal = FVector::CrossProduct(UseDirectionX, UseDirectionY);
+			*bExternalRenderVisibility = FMath::Abs(FVector::DotProduct(PlaneNormal, ViewDirection)) > 0.25f;
+			if (*bExternalRenderVisibility == false)
+			{
+				return;
+			}
+		}
+
+		float LengthScale = 1.0f;
+		if (ExternalDynamicPixelToWorldScale != nullptr)
+		{
+			float PixelToWorldScale = GizmoRenderingUtil::CalculateLocalPixelToWorldScale(GizmoControlView, Origin);
+			*ExternalDynamicPixelToWorldScale = PixelToWorldScale;
+			LengthScale = PixelToWorldScale;
+		}
+
+
+		float UseThickness = (bExternalHoverState != nullptr && *bExternalHoverState == true) ?
+			(HoverThicknessMultiplier*Thickness) : (Thickness);
+
+		double UseOffsetX = LengthScale * OffsetX;
+		double UseOffsetLengthX = LengthScale * (OffsetX + LengthX);
+		double UseOffsetY = LengthScale * OffsetY;
+		double UseOffsetLengthY = LengthScale * (OffsetY + LengthY);
+
+		for (int32 ViewIndex = 0; ViewIndex < Views.Num(); ViewIndex++)
+		{
+			if (VisibilityMap & (1 << ViewIndex))
+			{
+				const FSceneView* View = Views[ViewIndex];
+				FPrimitiveDrawInterface* PDI = Collector.GetPDI(ViewIndex);
+
+				FVector Point00 = Origin + UseOffsetX*UseDirectionX + UseOffsetY*UseDirectionY;
+				FVector Point10 = Origin + UseOffsetLengthX*UseDirectionX + UseOffsetY*UseDirectionY;
+				FVector Point11 = Origin + UseOffsetLengthX*UseDirectionX + UseOffsetLengthY*UseDirectionY;
+				FVector Point01 = Origin + UseOffsetX*UseDirectionX + UseOffsetLengthY*UseDirectionY;
+
+				PDI->DrawLine(Point10, Point11, Color, SDPG_Foreground, UseThickness, 0.0f, true);
+				PDI->DrawLine(Point11, Point01, Color, SDPG_Foreground, UseThickness, 0.0f, true);
+
+
+				/* 
+				 
+				// Draw solid square. Will be occluded but can be used w/ custom depth to show hidden
+				 
+				FDynamicMeshBuilder MeshBuilder(View->GetFeatureLevel());
+				MeshBuilder.AddVertex(Point00, FVector2D::ZeroVector, FVector::ZeroVector, FVector::ZeroVector, FVector::ZeroVector, UseColor);
+				MeshBuilder.AddVertex(Point10, FVector2D::ZeroVector, FVector::ZeroVector, FVector::ZeroVector, FVector::ZeroVector, UseColor);
+				MeshBuilder.AddVertex(Point11, FVector2D::ZeroVector, FVector::ZeroVector, FVector::ZeroVector, FVector::ZeroVector, UseColor);
+				MeshBuilder.AddVertex(Point01, FVector2D::ZeroVector, FVector::ZeroVector, FVector::ZeroVector, FVector::ZeroVector, UseColor);
+				MeshBuilder.AddTriangle(0, 1, 2);
+				MeshBuilder.AddTriangle(0, 2, 3);
+
+				//FMaterialRenderProxy* MaterialRenderProxy = Material->GetRenderProxy();
+				FMaterialRenderProxy* MaterialRenderProxy = GEngine->VertexColorMaterial->GetRenderProxy();
+				MeshBuilder.GetMesh(FMatrix::Identity, MaterialRenderProxy, SDPG_Foreground, true, false, ViewIndex, Collector);
+				//MeshBuilder.GetMesh(FMatrix::Identity, MaterialRenderProxy, SDPG_World, true, false, ViewIndex, Collector);
+				*/
+			}
+		}
+	}
+
+	virtual FPrimitiveViewRelevance GetViewRelevance(const FSceneView* View) const override
+	{
+		FPrimitiveViewRelevance Result;
+		Result.bDrawRelevance = IsShown(View);
+		Result.bDynamicRelevance = true;
+		Result.bShadowRelevance = false;
+		Result.bEditorPrimitiveRelevance = UseEditorCompositing(View);
+		Result.bRenderCustomDepth = ShouldRenderCustomDepth();
+
+		return Result;
+	}
+
+	virtual bool CanBeOccluded() const override
+	{
+		return false;
+	}
+
+	virtual uint32 GetMemoryFootprint(void) const override { return sizeof *this + GetAllocatedSize(); }
+	uint32 GetAllocatedSize(void) const { return FPrimitiveSceneProxy::GetAllocatedSize(); }
+
+	void SetExternalFlip(bool* bFlippedX, bool* bFlippedY)
+	{
+		bFlippedXExternal = bFlippedX;
+		bFlippedYExternal = bFlippedY;
+	}
+
+	void SetExternalDynamicPixelToWorldScale(float* DynamicPixelToWorldScale)
+	{
+		ExternalDynamicPixelToWorldScale = DynamicPixelToWorldScale;
+	}
+
+	void SetExternalRenderVisibility(bool* bRenderVisibility)
+	{
+		bExternalRenderVisibility = bRenderVisibility;
+	}
+
+	void SetExternalHoverState(bool* HoverState)
+	{
+		bExternalHoverState = HoverState;
+	}
+
+private:
+	FLinearColor Color;
+	FVector DirectionX, DirectionY;
+	float OffsetX, OffsetY;
+	float LengthX, LengthY;
+	float Thickness;
+	float HoverThicknessMultiplier;
+
+	bool* bFlippedXExternal = nullptr;
+	bool* bFlippedYExternal = nullptr;
+	float* ExternalDynamicPixelToWorldScale = nullptr;
+	bool* bExternalRenderVisibility = nullptr;
+	bool* bExternalHoverState = nullptr;
+};
+
+
+
+FPrimitiveSceneProxy* UGizmoRectangleComponent::CreateSceneProxy()
+{
+	FGizmoRectangleComponentSceneProxy* NewProxy = new FGizmoRectangleComponentSceneProxy(this);
+	NewProxy->SetExternalFlip(&bFlippedX, &bFlippedY);
+	NewProxy->SetExternalDynamicPixelToWorldScale(&DynamicPixelToWorldScale);
+	NewProxy->SetExternalRenderVisibility(&bRenderVisibility);
+	NewProxy->SetExternalHoverState(&bHovering);
+	return NewProxy;
+}
+
+FBoxSphereBounds UGizmoRectangleComponent::CalcBounds(const FTransform& LocalToWorld) const
+{
+	float MaxOffset = FMath::Max(OffsetX, OffsetY);
+	float MaxLength = FMath::Max(LengthX, LengthY);
+	return FBoxSphereBounds(FSphere(FVector::ZeroVector, 100*MaxOffset+MaxLength).TransformBy(LocalToWorld));
+}
+
+bool UGizmoRectangleComponent::LineTraceComponent(FHitResult& OutHit, const FVector Start, const FVector End, const FCollisionQueryParams& Params)
+{
+	if (bRenderVisibility == false)
+	{
+		return false;
+	}
+
+	const FTransform& Transform = this->GetComponentToWorld();
+
+	FVector UseDirectionX = (bFlippedX) ? -DirectionX : DirectionX;
+	FVector UseDirectionY = (bFlippedY) ? -DirectionY : DirectionY;
+
+	float LengthScale = DynamicPixelToWorldScale;
+	double UseOffsetX = LengthScale * OffsetX;
+	double UseOffsetLengthX = LengthScale * (OffsetX + LengthX);
+	double UseOffsetY = LengthScale * OffsetY;
+	double UseOffsetLengthY = LengthScale * (OffsetY + LengthY);
+
+	FVector Points[4] = {
+		Transform.TransformPosition(UseOffsetX*UseDirectionX + UseOffsetY *UseDirectionY),
+		Transform.TransformPosition(UseOffsetLengthX*UseDirectionX + UseOffsetY *UseDirectionY),
+		Transform.TransformPosition(UseOffsetLengthX*UseDirectionX + UseOffsetLengthY *UseDirectionY),
+		Transform.TransformPosition(UseOffsetX*UseDirectionX + UseOffsetLengthY *UseDirectionY)
+	};
+	static const int Triangles[2][3] = { {0,1,2}, {0,2,3} };
+
+	for (int j = 0; j < 2; ++j)
+	{
+		const int* Triangle = Triangles[j];
+		FVector HitPoint, HitNormal;
+		if (FMath::SegmentTriangleIntersection(Start, End,
+			Points[Triangle[0]], Points[Triangle[1]], Points[Triangle[2]], 
+			HitPoint, HitNormal) )
+		{
+			OutHit.Component = this;
+			OutHit.Distance = FVector::Distance(Start, HitPoint);
+			OutHit.ImpactPoint = HitPoint;
+			OutHit.ImpactNormal = HitNormal;
+			return true;
+		}
+	}
+
+	return false;
+}
+
+void UGizmoRectangleComponent::GetUsedMaterials(TArray<UMaterialInterface*>& OutMaterials, bool bGetDebugMaterials) const
+{
+	//OutMaterials.Add(GEngine->VertexColorMaterial);
+}
\ No newline at end of file
diff --git a/Engine/Source/Runtime/Experimental/InteractiveToolsFramework/Private/BaseGizmos/GizmoRenderingUtil.cpp b/Engine/Source/Runtime/Experimental/InteractiveToolsFramework/Private/BaseGizmos/GizmoRenderingUtil.cpp
new file mode 100644
--- /dev/null
+++ b/Engine/Source/Runtime/Experimental/InteractiveToolsFramework/Private/BaseGizmos/GizmoRenderingUtil.cpp
@@ -0,0 +1,60 @@
+// Copyright 1998-2019 Epic Games, Inc. All Rights Reserved.
+
+#include "BaseGizmos/GizmoRenderingUtil.h"
+
+
+const FSceneView* GizmoRenderingUtil::FindActiveSceneView(
+	const TArray<const FSceneView*>& Views,
+	const FSceneViewFamily& ViewFamily,
+	uint32 VisibilityMap)
+{
+	// can we tell focus here?
+
+	const FSceneView* FirstValidView = nullptr;
+	const FSceneView* GizmoControlView = nullptr;
+	for (int32 ViewIndex = 0; ViewIndex < Views.Num(); ViewIndex++)
+	{
+		if (VisibilityMap & (1 << ViewIndex))
+		{
+			const FSceneView* View = Views[ViewIndex];
+			if (FirstValidView == nullptr)
+			{
+				FirstValidView = View;
+			}
+			if (View->IsPerspectiveProjection())
+			{
+				GizmoControlView = View;
+			}
+		}
+	}
+
+	if (FirstValidView == nullptr)
+	{
+		return nullptr;
+	}
+	if (GizmoControlView == nullptr)
+	{
+		GizmoControlView = FirstValidView;
+	}
+
+	return GizmoControlView;
+}
+
+
+
+float GizmoRenderingUtil::CalculateLocalPixelToWorldScale(
+	const FSceneView* View,
+	const FVector& Location)
+{
+	FVector2D PixelA, PixelB;
+	View->WorldToPixel(Location, PixelA);
+	FVector OffsetPointWorld = Location + 0.01f*View->GetViewRight() + 0.01f*View->GetViewUp();
+	View->WorldToPixel(OffsetPointWorld, PixelB);
+
+	float PixelDelta = (PixelA - PixelB).Size();
+	float WorldDelta = (Location - OffsetPointWorld).Size();
+
+	float PixelToWorldScale = WorldDelta / PixelDelta;
+
+	return PixelToWorldScale;
+}
\ No newline at end of file
diff --git a/Engine/Source/Runtime/Experimental/InteractiveToolsFramework/Private/BaseGizmos/HitTargets.cpp b/Engine/Source/Runtime/Experimental/InteractiveToolsFramework/Private/BaseGizmos/HitTargets.cpp
new file mode 100644
--- /dev/null
+++ b/Engine/Source/Runtime/Experimental/InteractiveToolsFramework/Private/BaseGizmos/HitTargets.cpp
@@ -0,0 +1,23 @@
+// Copyright 1998-2019 Epic Games, Inc. All Rights Reserved.
+
+#include "BaseGizmos/HitTargets.h"
+#include "Components/PrimitiveComponent.h"
+#include "Engine/EngineTypes.h"    // FHitResult
+
+
+
+FInputRayHit UGizmoComponentHitTarget::IsHit(const FInputDeviceRay& ClickPos) const
+{
+	FInputRayHit Hit;
+	if (Component)
+	{
+		FVector End = ClickPos.WorldRay.PointAt(HALF_WORLD_MAX);
+		FHitResult OutHit;
+		if (Component->LineTraceComponent(OutHit, ClickPos.WorldRay.Origin, End, FCollisionQueryParams(SCENE_QUERY_STAT(HitTest), true)))
+		{
+			return FInputRayHit(OutHit.Distance);
+		}
+	}
+	return Hit;
+}
+
diff --git a/Engine/Source/Runtime/Experimental/InteractiveToolsFramework/Private/BaseGizmos/ParameterSourcesFloat.cpp b/Engine/Source/Runtime/Experimental/InteractiveToolsFramework/Private/BaseGizmos/ParameterSourcesFloat.cpp
new file mode 100644
--- /dev/null
+++ b/Engine/Source/Runtime/Experimental/InteractiveToolsFramework/Private/BaseGizmos/ParameterSourcesFloat.cpp
@@ -0,0 +1,3 @@
+// Copyright 1998-2019 Epic Games, Inc. All Rights Reserved.
+
+#include "BaseGizmos/ParameterSourcesFloat.h"
\ No newline at end of file
diff --git a/Engine/Source/Runtime/Experimental/InteractiveToolsFramework/Private/BaseGizmos/PlanePositionGizmo.cpp b/Engine/Source/Runtime/Experimental/InteractiveToolsFramework/Private/BaseGizmos/PlanePositionGizmo.cpp
new file mode 100644
--- /dev/null
+++ b/Engine/Source/Runtime/Experimental/InteractiveToolsFramework/Private/BaseGizmos/PlanePositionGizmo.cpp
@@ -0,0 +1,158 @@
+// Copyright 1998-2019 Epic Games, Inc. All Rights Reserved.
+
+#include "BaseGizmos/PlanePositionGizmo.h"
+#include "InteractiveGizmoManager.h"
+#include "BaseBehaviors/ClickDragBehavior.h"
+#include "BaseBehaviors/MouseHoverBehavior.h"
+#include "BaseGizmos/GizmoMath.h"
+
+
+
+UInteractiveGizmo* UPlanePositionGizmoBuilder::BuildGizmo(const FToolBuilderState& SceneState) const
+{
+	UPlanePositionGizmo* NewGizmo = NewObject<UPlanePositionGizmo>(SceneState.GizmoManager);
+	return NewGizmo;
+}
+
+
+
+
+void UPlanePositionGizmo::Setup()
+{
+	UInteractiveGizmo::Setup();
+
+	// Add default mouse input behavior
+	UClickDragInputBehavior* MouseBehavior = NewObject<UClickDragInputBehavior>();
+	MouseBehavior->Initialize(this);
+	MouseBehavior->SetDefaultPriority(FInputCapturePriority(FInputCapturePriority::DEFAULT_GIZMO_PRIORITY));
+	AddInputBehavior(MouseBehavior);
+	UMouseHoverBehavior* HoverBehavior = NewObject<UMouseHoverBehavior>();
+	HoverBehavior->Initialize(this);
+	HoverBehavior->SetDefaultPriority(FInputCapturePriority(FInputCapturePriority::DEFAULT_GIZMO_PRIORITY));
+	AddInputBehavior(HoverBehavior);
+
+	AxisSource = NewObject<UGizmoConstantAxisSource>(this);
+	ParameterSource = NewObject<UGizmoLocalVec2ParameterSource>(this);
+	HitTarget = NewObject<UGizmoComponentHitTarget>(this);
+	StateTarget = NewObject<UGizmoNilStateTarget>(this);
+
+	bInInteraction = false;
+}
+
+
+FInputRayHit UPlanePositionGizmo::CanBeginClickDragSequence(const FInputDeviceRay& PressPos)
+{
+	FInputRayHit GizmoHit;
+	if (HitTarget && AxisSource && ParameterSource)
+	{
+		GizmoHit = HitTarget->IsHit(PressPos);
+		if (GizmoHit.bHit)
+		{
+			LastHitPosition = PressPos.WorldRay.PointAt(GizmoHit.HitDepth);
+		}
+	}
+	return GizmoHit;
+}
+
+void UPlanePositionGizmo::OnClickPress(const FInputDeviceRay& PressPos)
+{
+	InteractionOrigin = LastHitPosition;
+	InteractionNormal = AxisSource->GetDirection();
+	if (AxisSource->HasTangentVectors())
+	{
+		AxisSource->GetTangentVectors(InteractionAxisX, InteractionAxisY);
+	}
+	else
+	{
+		GizmoMath::MakeNormalPlaneBasis(InteractionNormal, InteractionAxisX, InteractionAxisY);
+	}
+
+	bool bIntersects; FVector IntersectionPoint;
+	GizmoMath::RayPlaneIntersectionPoint(
+		InteractionOrigin, InteractionNormal,
+		PressPos.WorldRay.Origin, PressPos.WorldRay.Direction,
+		bIntersects, IntersectionPoint);
+	check(bIntersects);  // need to handle this case...
+
+	InteractionStartPoint = InteractionCurPoint = IntersectionPoint;
+
+	InteractionStartParameter = GizmoMath::ComputeCoordinatesInPlane(IntersectionPoint,
+			InteractionOrigin, InteractionNormal, InteractionAxisX, InteractionAxisY);
+	InteractionCurParameter = InteractionStartParameter;
+
+	InitialTargetParameter = ParameterSource->GetParameter();
+	ParameterSource->BeginModify();
+
+	bInInteraction = true;
+
+	if (StateTarget)
+	{
+		StateTarget->BeginUpdate();
+	}
+}
+
+void UPlanePositionGizmo::OnClickDrag(const FInputDeviceRay& DragPos)
+{
+	bool bIntersects; FVector IntersectionPoint;
+	GizmoMath::RayPlaneIntersectionPoint(
+		InteractionOrigin, InteractionNormal,
+		DragPos.WorldRay.Origin, DragPos.WorldRay.Direction,
+		bIntersects, IntersectionPoint);
+
+	if (bIntersects == false)
+	{
+		return;
+	}
+
+	InteractionCurPoint = IntersectionPoint;
+	InteractionCurParameter = GizmoMath::ComputeCoordinatesInPlane(IntersectionPoint,
+		InteractionOrigin, InteractionNormal, InteractionAxisX, InteractionAxisY);
+
+	FVector2D DeltaParam = InteractionCurParameter - InteractionStartParameter;
+	FVector2D NewValue = InitialTargetParameter + DeltaParam;
+
+	ParameterSource->SetParameter(NewValue);
+}
+
+void UPlanePositionGizmo::OnClickRelease(const FInputDeviceRay& ReleasePos)
+{
+	check(bInInteraction);
+
+	ParameterSource->EndModify();
+	if (StateTarget)
+	{
+		StateTarget->EndUpdate();
+	}
+	bInInteraction = false;
+}
+
+
+void UPlanePositionGizmo::OnTerminateDragSequence()
+{
+	check(bInInteraction);
+
+	ParameterSource->EndModify();
+	if (StateTarget)
+	{
+		StateTarget->EndUpdate();
+	}
+	bInInteraction = false;
+}
+
+
+
+void UPlanePositionGizmo::OnUpdateHover(const FInputDeviceRay& DevicePos)
+{
+	FInputRayHit GizmoHit;
+	if (HitTarget)
+	{
+		GizmoHit = HitTarget->IsHit(DevicePos);
+		if (GizmoHit.bHit)
+		{
+			HitTarget->UpdateHoverState(true);
+			return;
+		}
+	}
+	HitTarget->UpdateHoverState(false);
+}
+
diff --git a/Engine/Source/Runtime/Experimental/InteractiveToolsFramework/Private/BaseGizmos/StateTargets.cpp b/Engine/Source/Runtime/Experimental/InteractiveToolsFramework/Private/BaseGizmos/StateTargets.cpp
new file mode 100644
--- /dev/null
+++ b/Engine/Source/Runtime/Experimental/InteractiveToolsFramework/Private/BaseGizmos/StateTargets.cpp
@@ -0,0 +1,3 @@
+// Copyright 1998-2019 Epic Games, Inc. All Rights Reserved.
+
+#include "BaseGizmos/StateTargets.h"
\ No newline at end of file
diff --git a/Engine/Source/Runtime/Experimental/InteractiveToolsFramework/Private/BaseGizmos/TransformGizmo.cpp b/Engine/Source/Runtime/Experimental/InteractiveToolsFramework/Private/BaseGizmos/TransformGizmo.cpp
new file mode 100644
--- /dev/null
+++ b/Engine/Source/Runtime/Experimental/InteractiveToolsFramework/Private/BaseGizmos/TransformGizmo.cpp
@@ -0,0 +1,347 @@
+// Copyright 1998-2019 Epic Games, Inc. All Rights Reserved.
+
+#include "BaseGizmos/TransformGizmo.h"
+#include "InteractiveGizmoManager.h"
+#include "BaseGizmos/AxisPositionGizmo.h"
+#include "BaseGizmos/PlanePositionGizmo.h"
+#include "BaseGizmos/AxisAngleGizmo.h"
+#include "BaseGizmos/GizmoComponents.h"
+
+#include "BaseGizmos/GizmoArrowComponent.h"
+#include "BaseGizmos/GizmoRectangleComponent.h"
+#include "BaseGizmos/GizmoCircleComponent.h"
+
+// need this to implement hover
+#include "BaseGizmos/GizmoBaseComponent.h"
+
+#include "Components/SphereComponent.h"
+#include "Components/PrimitiveComponent.h"
+#include "Engine/World.h"
+#include "Engine/CollisionProfile.h"
+
+
+#define LOCTEXT_NAMESPACE "UTransformGizmo"
+
+
+ATransformGizmoActor::ATransformGizmoActor()
+{
+	// root component is a hidden sphere
+	USphereComponent* SphereComponent = CreateDefaultSubobject<USphereComponent>(TEXT("GizmoCenter"));
+	RootComponent = SphereComponent;
+	SphereComponent->InitSphereRadius(1.0f);
+	SphereComponent->SetVisibility(false);
+	SphereComponent->SetCollisionProfileName(UCollisionProfile::NoCollision_ProfileName);
+}
+
+
+
+ATransformGizmoActor* ATransformGizmoActor::ConstructDefault3AxisGizmo(UWorld* World)
+{
+	FActorSpawnParameters SpawnInfo;
+	ATransformGizmoActor* NewActor = World->SpawnActor<ATransformGizmoActor>(FVector::ZeroVector, FRotator::ZeroRotator, SpawnInfo);
+
+	NewActor->TranslateX = AddDefaultArrowComponent(World, NewActor, FLinearColor::Red, FVector(1, 0, 0));
+	NewActor->TranslateY = AddDefaultArrowComponent(World, NewActor, FLinearColor::Green, FVector(0, 1, 0));
+	NewActor->TranslateZ = AddDefaultArrowComponent(World, NewActor, FLinearColor::Blue, FVector(0, 0, 1));
+
+	NewActor->TranslateYZ = AddDefaultRectangleComponent(World, NewActor, FLinearColor::Red, FVector(0,1,0), FVector(0,0,1));
+	NewActor->TranslateXZ = AddDefaultRectangleComponent(World, NewActor, FLinearColor::Green, FVector(1, 0, 0), FVector(0, 0, 1));
+	NewActor->TranslateXY = AddDefaultRectangleComponent(World, NewActor, FLinearColor::Blue, FVector(1, 0, 0), FVector(0, 1, 0));
+
+	NewActor->RotateX = AddDefaultCircleComponent(World, NewActor, FLinearColor::Red, FVector(1, 0, 0));
+	NewActor->RotateY = AddDefaultCircleComponent(World, NewActor, FLinearColor::Green, FVector(0, 1, 0));
+	NewActor->RotateZ = AddDefaultCircleComponent(World, NewActor, FLinearColor::Blue, FVector(0, 0, 1));
+
+
+	// add a non-interactive view-aligned circle element, so the axes look like a sphere.
+	UGizmoCircleComponent* SphereEdge = NewObject<UGizmoCircleComponent>(NewActor);
+	NewActor->AddInstanceComponent(SphereEdge);
+	SphereEdge->AttachToComponent(NewActor->GetRootComponent(), FAttachmentTransformRules::KeepRelativeTransform);
+	SphereEdge->Color = FLinearColor::Gray;
+	SphereEdge->Thickness = 1.0f;
+	SphereEdge->bViewAligned = true;
+	SphereEdge->RegisterComponent();
+
+
+	return NewActor;
+}
+
+
+
+
+
+ATransformGizmoActor* FTransformGizmoActorFactory::CreateNewGizmoActor(UWorld* World) const
+{
+	return ATransformGizmoActor::ConstructDefault3AxisGizmo(World);
+}
+
+
+
+UInteractiveGizmo* UTransformGizmoBuilder::BuildGizmo(const FToolBuilderState& SceneState) const
+{
+	UTransformGizmo* NewGizmo = NewObject<UTransformGizmo>(SceneState.GizmoManager);
+	NewGizmo->SetWorld(SceneState.World);
+
+	// use default gizmo actor if client has not given us a new builder
+	NewGizmo->SetGizmoActorBuilder( (GizmoActorBuilder) ? GizmoActorBuilder : MakeShared<FTransformGizmoActorFactory>() );
+
+	// override default hover function if proposed
+	if (UpdateHoverFunction)
+	{
+		NewGizmo->SetUpdateHoverFunction(UpdateHoverFunction);
+	}
+
+	return NewGizmo;
+}
+
+
+
+void UTransformGizmo::SetWorld(UWorld* WorldIn)
+{
+	this->World = WorldIn;
+}
+
+
+void UTransformGizmo::SetGizmoActorBuilder(TSharedPtr<FTransformGizmoActorFactory> Builder)
+{
+	GizmoActorBuilder = Builder;
+}
+
+void UTransformGizmo::SetUpdateHoverFunction(TFunction<void(UPrimitiveComponent*, bool)> HoverFunction)
+{
+	UpdateHoverFunction = HoverFunction;
+}
+
+
+void UTransformGizmo::Setup()
+{
+	UInteractiveGizmo::Setup();
+
+	UpdateHoverFunction = [](UPrimitiveComponent* Component, bool bHovering)
+	{
+		if (Cast<UGizmoBaseComponent>(Component) != nullptr)
+		{
+			Cast<UGizmoBaseComponent>(Component)->UpdateHoverState(bHovering);
+		}
+	};
+
+	GizmoActor = GizmoActorBuilder->CreateNewGizmoActor(World);
+}
+
+
+
+void UTransformGizmo::Shutdown()
+{
+	ClearActiveTarget();
+
+	if (GizmoActor)
+	{
+		GizmoActor->Destroy();
+		GizmoActor = nullptr;
+	}
+}
+
+
+void UTransformGizmo::SetActiveTarget(UTransformProxy* Target)
+{
+	if (ActiveTarget != nullptr)
+	{
+		ClearActiveTarget();
+	}
+
+	ActiveTarget = Target;
+
+	// move gizmo to target location
+	USceneComponent* GizmoComponent = GizmoActor->GetRootComponent();
+
+	FTransform TargetTransform = Target->GetTransform();
+	FVector SaveScale = TargetTransform.GetScale3D();
+	TargetTransform.SetScale3D(FVector(1, 1, 1));
+	GizmoComponent->SetWorldTransform(TargetTransform);
+
+	// target tracks location of GizmoComponent
+	GizmoActor->GetRootComponent()->TransformUpdated.AddLambda(
+		[this, SaveScale](USceneComponent* Component, EUpdateTransformFlags /*UpdateTransformFlags*/, ETeleportType /*Teleport*/) {
+		//this->GetGizmoManager()->PostMessage(TEXT("TRANSFORM UPDATED"), EToolMessageLevel::Internal);
+		FTransform NewXForm = Component->GetComponentToWorld();
+		NewXForm.SetScale3D(SaveScale);
+		this->ActiveTarget->SetTransform(NewXForm);
+	});
+
+
+	UGizmoObjectModifyStateTarget* StateTarget = UGizmoObjectModifyStateTarget::Construct(GizmoComponent,
+		LOCTEXT("UTransformGizmoTransaction", "Transform"), GetGizmoManager(), this);
+
+	UGizmoComponentWorldTransformSource* TransformSource = 
+		UGizmoComponentWorldTransformSource::Construct(GizmoComponent, this);
+
+	// root component provides local X/Y/Z axis, identified by AxisIndex
+	AxisXSource = UGizmoComponentAxisSource::Construct(GizmoComponent, 0, true, this);
+	AxisYSource = UGizmoComponentAxisSource::Construct(GizmoComponent, 1, true, this);
+	AxisZSource = UGizmoComponentAxisSource::Construct(GizmoComponent, 2, true, this);
+
+	// todo should we hold onto these?
+	if (GizmoActor->TranslateX != nullptr)
+	{
+		AddAxisTranslationGizmo(GizmoActor->TranslateX, GizmoComponent, TEXT("TranslateX"), AxisXSource, TransformSource, StateTarget);
+	}
+	if (GizmoActor->TranslateY != nullptr)
+	{
+		AddAxisTranslationGizmo(GizmoActor->TranslateY, GizmoComponent, TEXT("TranslateY"), AxisYSource, TransformSource, StateTarget);
+	}
+	if (GizmoActor->TranslateZ != nullptr)
+	{
+		AddAxisTranslationGizmo(GizmoActor->TranslateZ, GizmoComponent, TEXT("TranslateZ"), AxisZSource, TransformSource, StateTarget);
+	}
+
+
+	if (GizmoActor->TranslateYZ != nullptr)
+	{
+		AddPlaneTranslationGizmo(GizmoActor->TranslateYZ, GizmoComponent, TEXT("TranslateYZ"), AxisXSource, TransformSource, StateTarget);
+	}
+	if (GizmoActor->TranslateXZ != nullptr)
+	{
+		AddPlaneTranslationGizmo(GizmoActor->TranslateXZ, GizmoComponent, TEXT("TranslateXZ"), AxisYSource, TransformSource, StateTarget);
+	}
+	if (GizmoActor->TranslateXY != nullptr)
+	{
+		AddPlaneTranslationGizmo(GizmoActor->TranslateXY, GizmoComponent, TEXT("TranslateXY"), AxisZSource, TransformSource, StateTarget);
+	}
+
+	if (GizmoActor->RotateX != nullptr)
+	{
+		AddAxisRotationGizmo(GizmoActor->RotateX, GizmoComponent, TEXT("RotateX"), AxisXSource, TransformSource, StateTarget);
+	}
+	if (GizmoActor->RotateY != nullptr)
+	{
+		AddAxisRotationGizmo(GizmoActor->RotateY, GizmoComponent, TEXT("RotateY"), AxisYSource, TransformSource, StateTarget);
+	}
+	if (GizmoActor->RotateZ != nullptr)
+	{
+		AddAxisRotationGizmo(GizmoActor->RotateZ, GizmoComponent, TEXT("RotateZ"), AxisZSource, TransformSource, StateTarget);
+	}
+}
+
+
+
+UInteractiveGizmo* UTransformGizmo::AddAxisTranslationGizmo(
+	UPrimitiveComponent* AxisComponent, USceneComponent* RootComponent,
+	const FString& Identifier,
+	IGizmoAxisSource* AxisSource,
+	IGizmoTransformSource* TransformSource,
+	IGizmoStateTarget* StateTarget)
+{
+	// create axis-position gizmo, axis-position parameter will drive translation
+	UAxisPositionGizmo* TranslateGizmo = Cast<UAxisPositionGizmo>(GetGizmoManager()->CreateGizmo(
+		UInteractiveGizmoManager::DefaultAxisPositionBuilderIdentifier, Identifier));
+	check(TranslateGizmo);
+
+	// axis source provides the translation axis
+	TranslateGizmo->AxisSource = Cast<UObject>(AxisSource);
+
+	// parameter source maps axis-parameter-change to translation of TransformSource's transform
+	TranslateGizmo->ParameterSource = UGizmoAxisTranslationParameterSource::Construct(AxisSource, TransformSource, this);
+
+	// sub-component provides hit target
+	UGizmoComponentHitTarget* HitTarget = UGizmoComponentHitTarget::Construct(AxisComponent, this);
+	if (this->UpdateHoverFunction)
+	{
+		HitTarget->UpdateHoverFunction = [AxisComponent, this](bool bHovering) { this->UpdateHoverFunction(AxisComponent, bHovering); };
+	}
+	TranslateGizmo->HitTarget = HitTarget;
+
+	TranslateGizmo->StateTarget = Cast<UObject>(StateTarget);
+
+	ActiveGizmos.Add(TranslateGizmo);
+	return TranslateGizmo;
+}
+
+
+
+UInteractiveGizmo* UTransformGizmo::AddPlaneTranslationGizmo(
+	UPrimitiveComponent* AxisComponent, USceneComponent* RootComponent,
+	const FString& Identifier,
+	IGizmoAxisSource* AxisSource,
+	IGizmoTransformSource* TransformSource,
+	IGizmoStateTarget* StateTarget)
+{
+	// create axis-position gizmo, axis-position parameter will drive translation
+	UPlanePositionGizmo* TranslateGizmo = Cast<UPlanePositionGizmo>(GetGizmoManager()->CreateGizmo(
+		UInteractiveGizmoManager::DefaultPlanePositionBuilderIdentifier, Identifier));
+	check(TranslateGizmo);
+
+	// axis source provides the translation axis
+	TranslateGizmo->AxisSource = Cast<UObject>(AxisSource);
+
+	// parameter source maps axis-parameter-change to translation of TransformSource's transform
+	TranslateGizmo->ParameterSource = UGizmoPlaneTranslationParameterSource::Construct(AxisSource, TransformSource, this);
+
+	// sub-component provides hit target
+	UGizmoComponentHitTarget* HitTarget = UGizmoComponentHitTarget::Construct(AxisComponent, this);
+	if (this->UpdateHoverFunction)
+	{
+		HitTarget->UpdateHoverFunction = [AxisComponent, this](bool bHovering) { this->UpdateHoverFunction(AxisComponent, bHovering); };
+	}
+	TranslateGizmo->HitTarget = HitTarget;
+
+	TranslateGizmo->StateTarget = Cast<UObject>(StateTarget);
+
+	ActiveGizmos.Add(TranslateGizmo);
+	return TranslateGizmo;
+}
+
+
+
+
+
+UInteractiveGizmo* UTransformGizmo::AddAxisRotationGizmo(
+	UPrimitiveComponent* AxisComponent, USceneComponent* RootComponent,
+	const FString& Identifier,
+	IGizmoAxisSource* AxisSource,
+	IGizmoTransformSource* TransformSource,
+	IGizmoStateTarget* StateTarget)
+{
+	// create axis-angle gizmo, angle will drive axis-rotation
+	UAxisAngleGizmo* RotateGizmo = Cast<UAxisAngleGizmo>(GetGizmoManager()->CreateGizmo(
+		UInteractiveGizmoManager::DefaultAxisAngleBuilderIdentifier, Identifier));
+	check(RotateGizmo);
+
+	// axis source provides the rotation axis
+	RotateGizmo->AxisSource = Cast<UObject>(AxisSource);
+
+	// parameter source maps angle-parameter-change to rotation of TransformSource's transform
+	RotateGizmo->AngleSource = UGizmoAxisRotationParameterSource::Construct(AxisSource, TransformSource, this);
+
+	// sub-component provides hit target
+	UGizmoComponentHitTarget* HitTarget = UGizmoComponentHitTarget::Construct(AxisComponent, this);
+	if (this->UpdateHoverFunction)
+	{
+		HitTarget->UpdateHoverFunction = [AxisComponent, this](bool bHovering) { this->UpdateHoverFunction(AxisComponent, bHovering); };
+	}
+	RotateGizmo->HitTarget = HitTarget;
+
+	RotateGizmo->StateTarget = Cast<UObject>(StateTarget);
+
+	ActiveGizmos.Add(RotateGizmo);
+
+	return RotateGizmo;
+}
+
+
+
+
+
+void UTransformGizmo::ClearActiveTarget()
+{
+	for (UInteractiveGizmo* Gizmo : ActiveGizmos)
+	{
+		GetGizmoManager()->DestroyGizmo(Gizmo);
+	}
+	ActiveGizmos.SetNum(0);
+
+	ActiveTarget = nullptr;
+}
+
+
+
+#undef LOCTEXT_NAMESPACE
\ No newline at end of file
diff --git a/Engine/Source/Runtime/Experimental/InteractiveToolsFramework/Private/BaseGizmos/TransformProxy.cpp b/Engine/Source/Runtime/Experimental/InteractiveToolsFramework/Private/BaseGizmos/TransformProxy.cpp
new file mode 100644
--- /dev/null
+++ b/Engine/Source/Runtime/Experimental/InteractiveToolsFramework/Private/BaseGizmos/TransformProxy.cpp
@@ -0,0 +1,92 @@
+// Copyright 1998-2019 Epic Games, Inc. All Rights Reserved.
+
+#include "BaseGizmos/TransformProxy.h"
+#include "Components/SceneComponent.h"
+
+
+void UTransformProxy::AddComponent(USceneComponent* Component, bool bModifyComponentOnTransform)
+{
+	check(Component);
+
+	FRelativeObject NewObj;
+	NewObj.Component = Component;
+	NewObj.bModifyComponentOnTransform = bModifyComponentOnTransform;
+	NewObj.StartTransform = Component->GetComponentToWorld();
+	NewObj.RelativeTransform = FTransform::Identity;
+	Objects.Add(NewObj);
+
+	UpdateSharedTransform();
+	OnTransformChanged.Broadcast(this, SharedTransform);
+}
+
+
+FTransform UTransformProxy::GetTransform() const
+{
+	return SharedTransform;
+}
+
+void UTransformProxy::SetTransform(const FTransform& TransformIn)
+{
+	SharedTransform = TransformIn;
+
+	UpdateObjects();
+
+	OnTransformChanged.Broadcast(this, SharedTransform);
+}
+
+
+
+
+void UTransformProxy::UpdateObjects()
+{
+	for (FRelativeObject& Obj : Objects)
+	{
+		FTransform CombinedTransform;
+		FTransform::Multiply(&CombinedTransform, &SharedTransform, &Obj.RelativeTransform);
+		
+		if (Obj.Component.IsValid())
+		{
+			if (Obj.bModifyComponentOnTransform)
+			{
+				Obj.Component->Modify();
+			}
+
+			Obj.Component->SetWorldTransform(CombinedTransform);
+		}
+	}
+}
+
+
+
+
+void UTransformProxy::UpdateSharedTransform()
+{
+	if (Objects.Num() == 0)
+	{
+		SharedTransform = FTransform::Identity;
+	}
+	else if (Objects.Num() == 1)
+	{
+		SharedTransform = Objects[0].StartTransform;
+
+		Objects[0].RelativeTransform = FTransform::Identity;
+	}
+	else
+	{
+		FVector Origin = FVector::ZeroVector;
+		for (const FRelativeObject& Obj : Objects)
+		{
+			Origin += Obj.StartTransform.GetLocation();
+		}
+		Origin /= (float)Objects.Num();
+
+		SharedTransform = FTransform(Origin);
+
+		for (FRelativeObject& Obj : Objects)
+		{
+			Obj.RelativeTransform = Obj.StartTransform;
+			Obj.RelativeTransform.SetToRelativeTransform(SharedTransform);
+		}
+	}
+
+}
\ No newline at end of file
diff --git a/Engine/Source/Runtime/Experimental/InteractiveToolsFramework/Private/BaseGizmos/TransformSources.cpp b/Engine/Source/Runtime/Experimental/InteractiveToolsFramework/Private/BaseGizmos/TransformSources.cpp
new file mode 100644
--- /dev/null
+++ b/Engine/Source/Runtime/Experimental/InteractiveToolsFramework/Private/BaseGizmos/TransformSources.cpp
@@ -0,0 +1,35 @@
+// Copyright 1998-2019 Epic Games, Inc. All Rights Reserved.
+
+#include "BaseGizmos/TransformSources.h"
+#include "Components/SceneComponent.h" 
+
+
+
+FTransform UGizmoComponentWorldTransformSource::GetTransform() const
+{
+	return Component->GetComponentToWorld();
+}
+
+void UGizmoComponentWorldTransformSource::SetTransform(const FTransform& NewTransform)
+{
+	if (bModifyComponentOnTransform)
+	{
+		Component->Modify();
+	}
+	Component->SetWorldTransform(NewTransform);
+	OnTransformChanged.Broadcast(this);
+}
+
+
+
+
+FTransform UGizmoTransformProxyTransformSource::GetTransform() const
+{
+	return Proxy->GetTransform();
+}
+
+void UGizmoTransformProxyTransformSource::SetTransform(const FTransform& NewTransform)
+{
+	Proxy->SetTransform(NewTransform);
+	OnTransformChanged.Broadcast(this);
+}
\ No newline at end of file
diff --git a/Engine/Source/Runtime/Experimental/InteractiveToolsFramework/Private/InteractiveGizmoManager.cpp b/Engine/Source/Runtime/Experimental/InteractiveToolsFramework/Private/InteractiveGizmoManager.cpp
--- a/Engine/Source/Runtime/Experimental/InteractiveToolsFramework/Private/InteractiveGizmoManager.cpp
+++ b/Engine/Source/Runtime/Experimental/InteractiveToolsFramework/Private/InteractiveGizmoManager.cpp
@@ -4,6 +4,11 @@
 #include "InteractiveGizmoManager.h"
 #include "InteractiveToolsContext.h"
 
+#include "BaseGizmos/AxisPositionGizmo.h"
+#include "BaseGizmos/PlanePositionGizmo.h"
+#include "BaseGizmos/AxisAngleGizmo.h"
+#include "BaseGizmos/TransformGizmo.h"
+
 
 UInteractiveGizmoManager::UInteractiveGizmoManager()
 {
@@ -33,6 +38,14 @@ void UInteractiveGizmoManager::Shutdown()
 	ActiveGizmos.Reset();
 
 	this->TransactionsAPI = nullptr;
+
+	if (bDefaultGizmosRegistered)
+	{
+		DeregisterGizmoType(DefaultAxisPositionBuilderIdentifier);
+		DeregisterGizmoType(DefaultPlanePositionBuilderIdentifier);
+		DeregisterGizmoType(DefaultAxisAngleBuilderIdentifier);
+		DeregisterGizmoType(DefaultThreeAxisTransformBuilderIdentifier);
+	}
 }
 
 
@@ -58,7 +71,7 @@ bool UInteractiveGizmoManager::DeregisterGizmoType(const FString& BuilderIdentif
 
 
 
-UInteractiveGizmo* UInteractiveGizmoManager::CreateGizmo(const FString& BuilderIdentifier, const FString& InstanceIdentifier)
+UInteractiveGizmo* UInteractiveGizmoManager::CreateGizmo(const FString& BuilderIdentifier, const FString& InstanceIdentifier, void* Owner)
 {
 	if ( GizmoBuilders.Contains(BuilderIdentifier) == false )
 	{
@@ -94,7 +107,7 @@ UInteractiveGizmo* UInteractiveGizmoManager::CreateGizmo(const FString& BuilderI
 
 	PostInvalidation();
 
-	FActiveGizmo ActiveGizmo = { NewGizmo, BuilderIdentifier, InstanceIdentifier };
+	FActiveGizmo ActiveGizmo = { NewGizmo, BuilderIdentifier, InstanceIdentifier, Owner };
 	ActiveGizmos.Add(ActiveGizmo);
 
 	return NewGizmo;
@@ -159,6 +172,23 @@ void UInteractiveGizmoManager::DestroyAllGizmosOfType(const FString& BuilderIden
 }
 
 
+void UInteractiveGizmoManager::DestroyAllGizmosByOwner(void* Owner)
+{
+	TArray<UInteractiveGizmo*> Found;
+	for ( const FActiveGizmo& ActiveGizmo : ActiveGizmos )
+	{
+		if (ActiveGizmo.Owner == Owner)
+		{
+			Found.Add(ActiveGizmo.Gizmo);
+		}
+	}
+	for (UInteractiveGizmo* Gizmo : Found)
+	{
+		DestroyGizmo(Gizmo);
+	}
+}
+
+
 
 UInteractiveGizmo* UInteractiveGizmoManager::FindGizmoByInstanceIdentifier(const FString& Identifier)
 {
@@ -226,3 +256,37 @@ void UInteractiveGizmoManager::EmitObjectChange(UObject* TargetObject, TUniquePt
 }
 
 
+
+
+FString UInteractiveGizmoManager::DefaultAxisPositionBuilderIdentifier = TEXT("StandardXFormAxisTranslationGizmo");
+FString UInteractiveGizmoManager::DefaultPlanePositionBuilderIdentifier = TEXT("StandardXFormPlaneTranslationGizmo");
+FString UInteractiveGizmoManager::DefaultAxisAngleBuilderIdentifier = TEXT("StandardXFormAxisRotationGizmo");
+FString UInteractiveGizmoManager::DefaultThreeAxisTransformBuilderIdentifier = TEXT("DefaultThreeAxisTransformBuilderIdentifier");
+
+
+void UInteractiveGizmoManager::RegisterDefaultGizmos()
+{
+	check(bDefaultGizmosRegistered == false);
+
+	UAxisPositionGizmoBuilder* AxisTranslationBuilder = NewObject<UAxisPositionGizmoBuilder>();
+	RegisterGizmoType(DefaultAxisPositionBuilderIdentifier, AxisTranslationBuilder);
+
+	UPlanePositionGizmoBuilder* PlaneTranslationBuilder = NewObject<UPlanePositionGizmoBuilder>();
+	RegisterGizmoType(DefaultPlanePositionBuilderIdentifier, PlaneTranslationBuilder);
+
+	UAxisAngleGizmoBuilder* AxisRotationBuilder = NewObject<UAxisAngleGizmoBuilder>();
+	RegisterGizmoType(DefaultAxisAngleBuilderIdentifier, AxisRotationBuilder);
+
+	UTransformGizmoBuilder* TransformBuilder = NewObject<UTransformGizmoBuilder>();
+	RegisterGizmoType(DefaultThreeAxisTransformBuilderIdentifier, TransformBuilder);
+
+	bDefaultGizmosRegistered = true;
+}
+
+UTransformGizmo* UInteractiveGizmoManager::Create3AxisTransformGizmo(const FString& InstanceIdentifier, void* Owner)
+{
+	check(bDefaultGizmosRegistered);
+	UInteractiveGizmo* NewGizmo = CreateGizmo(DefaultThreeAxisTransformBuilderIdentifier, InstanceIdentifier, Owner);
+	check(NewGizmo);
+	return Cast<UTransformGizmo>(NewGizmo);
+}
\ No newline at end of file
diff --git a/Engine/Source/Runtime/Experimental/InteractiveToolsFramework/Private/InteractiveToolsContext.cpp b/Engine/Source/Runtime/Experimental/InteractiveToolsFramework/Private/InteractiveToolsContext.cpp
--- a/Engine/Source/Runtime/Experimental/InteractiveToolsFramework/Private/InteractiveToolsContext.cpp
+++ b/Engine/Source/Runtime/Experimental/InteractiveToolsFramework/Private/InteractiveToolsContext.cpp
@@ -20,6 +20,8 @@ void UInteractiveToolsContext::Initialize(IToolsContextQueriesAPI* QueriesAPI, I
 
 	GizmoManager = NewObject<UInteractiveGizmoManager>(this);
 	GizmoManager->Initialize(QueriesAPI, TransactionsAPI, InputRouter);
+
+	GizmoManager->RegisterDefaultGizmos();
 }
 
 
diff --git a/Engine/Source/Runtime/Experimental/InteractiveToolsFramework/Public/BaseGizmos/AxisAngleGizmo.h b/Engine/Source/Runtime/Experimental/InteractiveToolsFramework/Public/BaseGizmos/AxisAngleGizmo.h
new file mode 100644
--- /dev/null
+++ b/Engine/Source/Runtime/Experimental/InteractiveToolsFramework/Public/BaseGizmos/AxisAngleGizmo.h
@@ -0,0 +1,95 @@
+// Copyright 1998-2019 Epic Games, Inc. All Rights Reserved.
+
+#pragma once
+
+#include "CoreMinimal.h"
+#include "InteractiveGizmo.h"
+#include "InteractiveGizmoBuilder.h"
+#include "BaseBehaviors/BehaviorTargetInterfaces.h"
+#include "BaseGizmos/GizmoComponents.h"
+#include "AxisAngleGizmo.generated.h"
+
+
+
+
+UCLASS()
+class INTERACTIVETOOLSFRAMEWORK_API UAxisAngleGizmoBuilder : public UInteractiveGizmoBuilder
+{
+	GENERATED_BODY()
+
+public:
+	virtual UInteractiveGizmo* BuildGizmo(const FToolBuilderState& SceneState) const override;
+};
+
+
+/**
+ *
+ */
+UCLASS()
+class INTERACTIVETOOLSFRAMEWORK_API UAxisAngleGizmo : public UInteractiveGizmo, public IClickDragBehaviorTarget, public IHoverBehaviorTarget
+{
+	GENERATED_BODY()
+
+public:
+	// UInteractiveGizmo overrides
+
+	virtual void Setup() override;
+
+	// IClickDragBehaviorTarget implementation
+
+	virtual FInputRayHit CanBeginClickDragSequence(const FInputDeviceRay& PressPos) override;
+	virtual void OnClickPress(const FInputDeviceRay& PressPos) override;
+	virtual void OnClickDrag(const FInputDeviceRay& DragPos) override;
+	virtual void OnClickRelease(const FInputDeviceRay& ReleasePos) override;
+	virtual void OnTerminateDragSequence() override;
+
+	// IHoverBehaviorTarget implementation
+	virtual void OnUpdateHover(const FInputDeviceRay& DevicePos) override;
+
+public:
+	UPROPERTY()
+	TScriptInterface<IGizmoAxisSource> AxisSource;
+
+	UPROPERTY()
+	TScriptInterface<IGizmoFloatParameterSource> AngleSource;
+
+	UPROPERTY()
+	TScriptInterface<IGizmoClickTarget> HitTarget;
+
+	UPROPERTY()
+	TScriptInterface<IGizmoStateTarget> StateTarget;
+
+public:
+	UPROPERTY()
+	bool bInInteraction = false;
+
+	UPROPERTY()
+	FVector RotationOrigin;
+	
+	UPROPERTY()
+	FVector RotationAxis;
+
+	UPROPERTY()
+	FVector RotationPlaneX;
+
+	UPROPERTY()
+	FVector RotationPlaneY;
+
+
+	UPROPERTY()
+	FVector InteractionStartPoint;
+
+	UPROPERTY()
+	FVector InteractionCurPoint;
+
+	UPROPERTY()
+	float InteractionStartAngle;
+
+	UPROPERTY()
+	float InteractionCurAngle;
+
+protected:
+	FVector LastHitPosition;
+	float InitialTargetAngle;
+};
+
diff --git a/Engine/Source/Runtime/Experimental/InteractiveToolsFramework/Public/BaseGizmos/AxisPositionGizmo.h b/Engine/Source/Runtime/Experimental/InteractiveToolsFramework/Public/BaseGizmos/AxisPositionGizmo.h
new file mode 100644
--- /dev/null
+++ b/Engine/Source/Runtime/Experimental/InteractiveToolsFramework/Public/BaseGizmos/AxisPositionGizmo.h
@@ -0,0 +1,89 @@
+// Copyright 1998-2019 Epic Games, Inc. All Rights Reserved.
+
+#pragma once
+
+#include "CoreMinimal.h"
+#include "InteractiveGizmo.h"
+#include "InteractiveGizmoBuilder.h"
+#include "BaseBehaviors/BehaviorTargetInterfaces.h"
+#include "BaseGizmos/GizmoComponents.h"
+#include "AxisPositionGizmo.generated.h"
+
+
+
+
+UCLASS()
+class INTERACTIVETOOLSFRAMEWORK_API UAxisPositionGizmoBuilder : public UInteractiveGizmoBuilder
+{
+	GENERATED_BODY()
+
+public:
+	virtual UInteractiveGizmo* BuildGizmo(const FToolBuilderState& SceneState) const override;
+};
+
+
+/**
+ * TODO rename to UAxisPositionGizmo
+ */
+UCLASS()
+class INTERACTIVETOOLSFRAMEWORK_API UAxisPositionGizmo : public UInteractiveGizmo, public IClickDragBehaviorTarget, public IHoverBehaviorTarget
+{
+	GENERATED_BODY()
+
+public:
+	// UInteractiveGizmo overrides
+
+	virtual void Setup() override;
+
+	// IClickDragBehaviorTarget implementation
+
+	virtual FInputRayHit CanBeginClickDragSequence(const FInputDeviceRay& PressPos) override;
+	virtual void OnClickPress(const FInputDeviceRay& PressPos) override;
+	virtual void OnClickDrag(const FInputDeviceRay& DragPos) override;
+	virtual void OnClickRelease(const FInputDeviceRay& ReleasePos) override;
+	virtual void OnTerminateDragSequence() override;
+
+	// IHoverBehaviorTarget implementation
+	virtual void OnUpdateHover(const FInputDeviceRay& DevicePos) override;
+
+public:
+	UPROPERTY()
+	TScriptInterface<IGizmoAxisSource> AxisSource;
+
+	UPROPERTY()
+	TScriptInterface<IGizmoFloatParameterSource> ParameterSource;
+
+	UPROPERTY()
+	TScriptInterface<IGizmoClickTarget> HitTarget;
+
+	UPROPERTY()
+	TScriptInterface<IGizmoStateTarget> StateTarget;
+
+public:
+	UPROPERTY()
+	bool bInInteraction = false;
+
+	UPROPERTY()
+	FVector InteractionOrigin;
+
+	UPROPERTY()
+	FVector InteractionAxis;
+
+	UPROPERTY()
+	FVector InteractionStartPoint;
+
+	UPROPERTY()
+	FVector InteractionCurPoint;
+
+	UPROPERTY()
+	float InteractionStartParameter;
+
+	UPROPERTY()
+	float InteractionCurParameter;
+
+
+protected:
+	FVector LastHitPosition;
+	float InitialTargetParameter;
+};
+
diff --git a/Engine/Source/Runtime/Experimental/InteractiveToolsFramework/Public/BaseGizmos/AxisSources.h b/Engine/Source/Runtime/Experimental/InteractiveToolsFramework/Public/BaseGizmos/AxisSources.h
new file mode 100644
--- /dev/null
+++ b/Engine/Source/Runtime/Experimental/InteractiveToolsFramework/Public/BaseGizmos/AxisSources.h
@@ -0,0 +1,127 @@
+// Copyright 1998-2019 Epic Games, Inc. All Rights Reserved.
+
+#pragma once
+
+#include "CoreMinimal.h"
+#include "BaseGizmos/GizmoInterfaces.h"
+#include "AxisSources.generated.h"
+
+class USceneComponent;
+
+
+/**
+ * UGizmoConstantAxisSource is an IGizmoAxisSource implementation that
+ * internally stores the Origin and Direction constants
+ */
+UCLASS()
+class INTERACTIVETOOLSFRAMEWORK_API UGizmoConstantAxisSource : public UObject, public IGizmoAxisSource
+{
+	GENERATED_BODY()
+public:
+	virtual FVector GetOrigin() const
+	{
+		return Origin;
+	}
+
+	virtual FVector GetDirection() const
+	{
+		return Direction;
+	}
+
+	UPROPERTY()
+	FVector Origin = FVector::ZeroVector;
+
+	UPROPERTY()
+	FVector Direction = FVector(0, 0, 1);
+};
+
+
+
+/**
+ * UGizmoWorldAxisSource is an IGizmoAxisSource that provides one of the World axes
+ * (ie X/Y/Z axes) based on an integer AxisIndex in range [0,2]. The Orgin is
+ * internally stored as a UProperty.
+ */
+UCLASS()
+class INTERACTIVETOOLSFRAMEWORK_API UGizmoWorldAxisSource : public UObject, public IGizmoAxisSource
+{
+	GENERATED_BODY()
+public:
+	virtual FVector GetOrigin() const
+	{
+		return Origin;
+	}
+
+	virtual FVector GetDirection() const
+	{
+		FVector Result(0, 0, 0);
+		Result[FMath::Clamp(AxisIndex, 0, 2)] = 1.0;
+		return Result;
+	}
+
+	UPROPERTY()
+	FVector Origin = FVector::ZeroVector;
+
+	/** Clamped internally to 0,1,2 */
+	UPROPERTY()
+	int AxisIndex = 2;
+
+public:
+	static UGizmoWorldAxisSource* Construct(
+		const FVector& WorldOrigin,
+		int WorldAxisIndex,
+		UObject* Outer = (UObject*)GetTransientPackage())
+	{
+		UGizmoWorldAxisSource* NewSource = NewObject<UGizmoWorldAxisSource>(Outer);
+		NewSource->Origin = WorldOrigin;
+		NewSource->AxisIndex = FMath::Clamp(WorldAxisIndex, 0, 2);
+		return NewSource;
+	}
+};
+
+
+
+
+/**
+ * UGizmoComponentAxisSource is an IGizmoAxisSource implementation that provides one of the
+ * X/Y/Z axes of a Component's local coordinate system, mapped to World, based on an integer AxisIndex in range [0,2].
+ * The Axis Origin is the Component's transform origin. Tangent vectors are provided.
+ */
+UCLASS()
+class INTERACTIVETOOLSFRAMEWORK_API UGizmoComponentAxisSource : public UObject, public IGizmoAxisSource
+{
+	GENERATED_BODY()
+public:
+	virtual FVector GetOrigin() const;
+
+	virtual FVector GetDirection() const;
+
+	virtual bool HasTangentVectors() const { return true; }
+
+	virtual void GetTangentVectors(FVector& TangentXOut, FVector& TangentYOut) const;
+
+	UPROPERTY()
+	USceneComponent* Component;
+
+	/** Clamped internally to 0,1,2 */
+	UPROPERTY()
+	int AxisIndex = 2;
+
+	/** If false, returns World axes */
+	UPROPERTY()
+	bool bLocalAxes = true;
+
+public:
+	static UGizmoComponentAxisSource* Construct(
+		USceneComponent* ComponentIn,
+		int LocalAxisIndex,
+		bool bUseLocalAxes,
+		UObject* Outer = (UObject*)GetTransientPackage())
+	{
+		UGizmoComponentAxisSource* NewSource = NewObject<UGizmoComponentAxisSource>(Outer);
+		NewSource->Component = ComponentIn;
+		NewSource->AxisIndex = LocalAxisIndex;
+		NewSource->bLocalAxes = bUseLocalAxes;
+		return NewSource;
+	}
+};
\ No newline at end of file
diff --git a/Engine/Source/Runtime/Experimental/InteractiveToolsFramework/Public/BaseGizmos/GizmoActor.h b/Engine/Source/Runtime/Experimental/InteractiveToolsFramework/Public/BaseGizmos/GizmoActor.h
new file mode 100644
--- /dev/null
+++ b/Engine/Source/Runtime/Experimental/InteractiveToolsFramework/Public/BaseGizmos/GizmoActor.h
@@ -0,0 +1,40 @@
+// Copyright 1998-2019 Epic Games, Inc. All Rights Reserved.
+
+#pragma once
+
+#include "CoreMinimal.h"
+#include "InteractiveToolObjects.h"
+#include "GizmoActor.generated.h"
+
+class UGizmoArrowComponent;
+class UGizmoRectangleComponent;
+class UGizmoCircleComponent;
+
+
+/**
+ * AGizmoActor is a base class for Actors that would be created to represent Gizmos in the scene.
+ * Currently this does not involve any special functionality, but a set of static functions
+ * are provided to create default Components commonly used in Gizmos.
+ */
+UCLASS(Transient)
+class INTERACTIVETOOLSFRAMEWORK_API AGizmoActor : public AInternalToolFrameworkActor
+{
+	GENERATED_BODY()
+public:
+
+	/** Add standard arrow component to Actor, generally used for axis-translation */
+	static UGizmoArrowComponent* AddDefaultArrowComponent(
+		UWorld* World, AActor* Actor,
+		const FLinearColor& Color, const FVector& LocalDirection
+	);
+	/** Add standard rectangle component to Actor, generally used for plane-translation */
+	static UGizmoRectangleComponent* AddDefaultRectangleComponent(
+		UWorld* World, AActor* Actor,
+		const FLinearColor& Color, const FVector& PlaneAxis1, const FVector& PlaneAxisx2
+	);
+	/** Add standard circle component to Actor, generally used for axis-rotation */
+	static UGizmoCircleComponent* AddDefaultCircleComponent(
+		UWorld* World, AActor* Actor,
+		const FLinearColor& Color, const FVector& PlaneNormal
+	);
+};
diff --git a/Engine/Source/Runtime/Experimental/InteractiveToolsFramework/Public/BaseGizmos/GizmoArrowComponent.h b/Engine/Source/Runtime/Experimental/InteractiveToolsFramework/Public/BaseGizmos/GizmoArrowComponent.h
new file mode 100644
--- /dev/null
+++ b/Engine/Source/Runtime/Experimental/InteractiveToolsFramework/Public/BaseGizmos/GizmoArrowComponent.h
@@ -0,0 +1,47 @@
+// Copyright 1998-2019 Epic Games, Inc. All Rights Reserved.
+
+#pragma once
+
+#include "CoreMinimal.h"
+#include "UObject/ObjectMacros.h"
+#include "BaseGizmos/GizmoBaseComponent.h"
+#include "GizmoArrowComponent.generated.h"
+
+/**
+ * Simple Component intended to be used as part of 3D Gizmos.
+ * Currently draws the "arrow" as a single 3D line.
+ */
+UCLASS(ClassGroup = Utility, HideCategories = (Physics, Collision, Lighting, Rendering, Mobile))
+class INTERACTIVETOOLSFRAMEWORK_API UGizmoArrowComponent : public UGizmoBaseComponent
+{
+	GENERATED_BODY()
+
+public:
+	UPROPERTY(EditAnywhere, Category = Options)
+	FVector Direction = FVector(0,0,1);
+
+	UPROPERTY(EditAnywhere, Category = Options)
+	float Gap = 5.0f;
+
+	UPROPERTY(EditAnywhere, Category = Options)
+	float Length = 60.0f;
+
+	UPROPERTY(EditAnywhere, Category = Options)
+	float Thickness = 2.0f;
+
+private:
+	//~ Begin UPrimitiveComponent Interface.
+	virtual FPrimitiveSceneProxy* CreateSceneProxy() override;
+	virtual bool LineTraceComponent(FHitResult& OutHit, const FVector Start, const FVector End, const FCollisionQueryParams& Params) override;
+	//~ End UPrimitiveComponent Interface.
+
+	//~ Begin USceneComponent Interface.
+	virtual FBoxSphereBounds CalcBounds(const FTransform& LocalToWorld) const override;
+	//~ Begin USceneComponent Interface.
+
+	// if true, we drew along -Direction instead of Direction, and so should hit-test accordingly
+	bool bFlipped;
+
+	// gizmo visibility
+	bool bRenderVisibility = true;
+};
\ No newline at end of file
diff --git a/Engine/Source/Runtime/Experimental/InteractiveToolsFramework/Public/BaseGizmos/GizmoBaseComponent.h b/Engine/Source/Runtime/Experimental/InteractiveToolsFramework/Public/BaseGizmos/GizmoBaseComponent.h
new file mode 100644
--- /dev/null
+++ b/Engine/Source/Runtime/Experimental/InteractiveToolsFramework/Public/BaseGizmos/GizmoBaseComponent.h
@@ -0,0 +1,57 @@
+// Copyright 1998-2019 Epic Games, Inc. All Rights Reserved.
+
+#pragma once
+
+#include "CoreMinimal.h"
+#include "UObject/ObjectMacros.h"
+#include "Components/PrimitiveComponent.h"
+#include "GizmoBaseComponent.generated.h"
+
+/**
+ * Base class for simple Components intended to be used as part of 3D Gizmos.
+ * Contains common properties and utility functions.
+ * This class does nothing by itself, use subclasses like UGizmoCircleComponent
+ */
+UCLASS(ClassGroup = Utility, HideCategories = (Physics, Collision, Mobile))
+class INTERACTIVETOOLSFRAMEWORK_API UGizmoBaseComponent : public UPrimitiveComponent
+{
+	GENERATED_BODY()
+
+public:
+	UGizmoBaseComponent()
+	{
+		bUseEditorCompositing = true;
+	}
+
+
+public:
+	UPROPERTY(EditAnywhere, Category = Options)
+	FLinearColor Color = FLinearColor::Red;
+
+
+	UPROPERTY(EditAnywhere, Category = Options)
+	float HoverSizeMultiplier = 2.0f;
+
+
+	UPROPERTY(EditAnywhere, Category = Options)
+	float PixelHitDistanceThreshold = 7.0f;
+
+
+public:
+	UFUNCTION()
+	void UpdateHoverState(bool bHoveringIn)
+	{
+		if (bHoveringIn != bHovering)
+		{
+			bHovering = bHoveringIn;
+		}
+	}
+
+
+protected:
+	// scale factor between pixel distances and world distances at Gizmo origin
+	float DynamicPixelToWorldScale = 1.0f;
+
+	// hover state
+	bool bHovering = false;
+};
\ No newline at end of file
diff --git a/Engine/Source/Runtime/Experimental/InteractiveToolsFramework/Public/BaseGizmos/GizmoCircleComponent.h b/Engine/Source/Runtime/Experimental/InteractiveToolsFramework/Public/BaseGizmos/GizmoCircleComponent.h
new file mode 100644
--- /dev/null
+++ b/Engine/Source/Runtime/Experimental/InteractiveToolsFramework/Public/BaseGizmos/GizmoCircleComponent.h
@@ -0,0 +1,47 @@
+// Copyright 1998-2019 Epic Games, Inc. All Rights Reserved.
+
+#pragma once
+
+#include "CoreMinimal.h"
+#include "UObject/ObjectMacros.h"
+#include "BaseGizmos/GizmoBaseComponent.h"
+#include "GizmoCircleComponent.generated.h"
+
+/**
+ * Simple Component intended to be used as part of 3D Gizmos.
+ * Draws a 3D circle based on parameters.
+ */
+UCLASS(ClassGroup = Utility, HideCategories = (Physics, Collision, Mobile))
+class INTERACTIVETOOLSFRAMEWORK_API UGizmoCircleComponent : public UGizmoBaseComponent
+{
+	GENERATED_BODY()
+
+public:
+	UPROPERTY(EditAnywhere, Category = Options)
+	FVector Normal = FVector(0,0,1);
+
+	UPROPERTY(EditAnywhere, Category = Options)
+	float Radius = 100.0f;
+
+	UPROPERTY(EditAnywhere, Category = Options)
+	float Thickness = 2.0f;
+
+	UPROPERTY(EditAnywhere, Category = Options)
+	int NumSides = 64;
+
+	UPROPERTY(EditAnywhere, Category = Options)
+	bool bViewAligned = false;
+
+	UPROPERTY(EditAnywhere, Category = Options)
+	bool bOnlyAllowFrontFacingHits = true;
+
+private:
+	//~ Begin UPrimitiveComponent Interface.
+	virtual FPrimitiveSceneProxy* CreateSceneProxy() override;
+	virtual bool LineTraceComponent(FHitResult& OutHit, const FVector Start, const FVector End, const FCollisionQueryParams& Params) override;
+	//~ End UPrimitiveComponent Interface.
+
+	//~ Begin USceneComponent Interface.
+	virtual FBoxSphereBounds CalcBounds(const FTransform& LocalToWorld) const override;
+	//~ Begin USceneComponent Interface.
+};
\ No newline at end of file
diff --git a/Engine/Source/Runtime/Experimental/InteractiveToolsFramework/Public/BaseGizmos/GizmoComponents.h b/Engine/Source/Runtime/Experimental/InteractiveToolsFramework/Public/BaseGizmos/GizmoComponents.h
new file mode 100644
--- /dev/null
+++ b/Engine/Source/Runtime/Experimental/InteractiveToolsFramework/Public/BaseGizmos/GizmoComponents.h
@@ -0,0 +1,13 @@
+// Copyright 1998-2019 Epic Games, Inc. All Rights Reserved.
+
+#pragma once
+
+#include "BaseGizmos/GizmoInterfaces.h"
+#include "BaseGizmos/AxisSources.h"
+#include "BaseGizmos/ParameterSourcesFloat.h"
+#include "BaseGizmos/ParameterSourcesVec2.h"
+#include "BaseGizmos/HitTargets.h"
+#include "BaseGizmos/StateTargets.h"
+#include "BaseGizmos/TransformProxy.h"
+#include "BaseGizmos/TransformSources.h"
+#include "BaseGizmos/ParameterToTransformAdapters.h"
diff --git a/Engine/Source/Runtime/Experimental/InteractiveToolsFramework/Public/BaseGizmos/GizmoInterfaces.h b/Engine/Source/Runtime/Experimental/InteractiveToolsFramework/Public/BaseGizmos/GizmoInterfaces.h
new file mode 100644
--- /dev/null
+++ b/Engine/Source/Runtime/Experimental/InteractiveToolsFramework/Public/BaseGizmos/GizmoInterfaces.h
@@ -0,0 +1,203 @@
+// Copyright 1998-2019 Epic Games, Inc. All Rights Reserved.
+
+#pragma once
+
+#include "CoreMinimal.h"
+#include "UObject/Interface.h"
+#include "InputState.h"
+#include "GizmoInterfaces.generated.h"
+
+
+//
+// UInterfaces for the various UObjects used in the Standard Gizmo Library.
+// 
+
+
+
+UINTERFACE()
+class INTERACTIVETOOLSFRAMEWORK_API UGizmoTransformSource : public UInterface
+{
+	GENERATED_BODY()
+};
+/**
+ * IGizmoTransformSource is an interface which is used to Get/Set an FTransform.
+ */
+class INTERACTIVETOOLSFRAMEWORK_API IGizmoTransformSource
+{
+	GENERATED_BODY()
+public:
+	UFUNCTION()
+	virtual FTransform GetTransform() const = 0;
+
+	UFUNCTION()
+	virtual void SetTransform(const FTransform& NewTransform) = 0;
+};
+
+
+
+
+
+UINTERFACE()
+class INTERACTIVETOOLSFRAMEWORK_API UGizmoAxisSource : public UInterface
+{
+	GENERATED_BODY()
+};
+/**
+ * IGizmoAxisSource is an interface which is used to get information about a 3D Axis.
+ * At minimum this includes a 3D Direction Vector and Origin Point.
+ * Optionally the implementation may provide two Tangent Vectors which are
+ * assumed to be mutually-orthogonal and perpendicular to the Axis Direction
+ * (ie that's the normal and the 3 vectors form a coordinate frame).
+ */
+class INTERACTIVETOOLSFRAMEWORK_API IGizmoAxisSource
+{
+	GENERATED_BODY()
+public:
+	/** @return Origin Point of axis */
+	UFUNCTION()
+	virtual FVector GetOrigin() const = 0;
+
+	/** @return Direction Vector of axis */
+	UFUNCTION()
+	virtual FVector GetDirection() const = 0;
+
+	/** @return true if this AxisSource has tangent vectors orthogonal to the Direction vector */
+	UFUNCTION()
+	virtual bool HasTangentVectors() const { return false; }
+
+	/** Get the two tangent vectors that are orthogonal to the Direction vector. 
+	 * @warning Only valid if HasTangentVectors() returns true
+	 */
+	UFUNCTION()
+	virtual void GetTangentVectors(FVector& TangentXOut, FVector& TangentYOut) const { }
+
+
+	/**
+	 * Utility function that always returns a 3D coordinate system (ie plane normal and perpendicular axes).
+	 * Internally calls GetTangentVectors() if available, otherwise constructs arbitrary mutually perpendicular vectors. 
+	 */
+	void GetAxisFrame(
+		FVector& PlaneNormalOut, FVector& PlaneAxis1Out, FVector& PlaneAxis2Out) const;
+};
+
+
+
+
+
+
+UINTERFACE()
+class INTERACTIVETOOLSFRAMEWORK_API UGizmoClickTarget : public UInterface
+{
+	GENERATED_BODY()
+};
+/**
+ * IGizmoClickTarget is an interface used to provide a ray-object hit test.
+ */
+class INTERACTIVETOOLSFRAMEWORK_API IGizmoClickTarget
+{
+	GENERATED_BODY()
+public:
+	/**
+	 * @return FInputRayHit indicating whether or not the target object was hit by the device-ray at ClickPos
+	 */
+	//UFUNCTION()    // FInputDeviceRay is not USTRUCT because FRay is not USTRUCT
+	virtual FInputRayHit IsHit(const FInputDeviceRay& ClickPos) const = 0;
+
+	UFUNCTION()
+	virtual void UpdateHoverState(bool bHovering) const = 0;
+};
+
+
+
+
+UINTERFACE()
+class INTERACTIVETOOLSFRAMEWORK_API UGizmoStateTarget : public UInterface
+{
+	GENERATED_BODY()
+};
+/**
+ * IGizmoStateTarget is an interface that is used to pass notifications about significant gizmo state updates
+ */
+class INTERACTIVETOOLSFRAMEWORK_API IGizmoStateTarget
+{
+	GENERATED_BODY()
+public:
+	/**
+	 * BeginUpdate is called before a standard Gizmo begins changing a parameter (via a ParameterSource)
+	 */
+	UFUNCTION()
+	virtual void BeginUpdate() = 0;
+
+	/**
+	 * EndUpdate is called when a standard Gizmo is finished changing a parameter (via a ParameterSource)
+	 */
+	UFUNCTION()
+	virtual void EndUpdate() = 0;
+};
+
+
+
+
+
+UINTERFACE()
+class INTERACTIVETOOLSFRAMEWORK_API UGizmoFloatParameterSource : public UInterface
+{
+	GENERATED_BODY()
+};
+/**
+ * IGizmoFloatParameterSource provides Get and Set for an arbitrary float-valued parameter.
+ */
+class INTERACTIVETOOLSFRAMEWORK_API IGizmoFloatParameterSource
+{
+	GENERATED_BODY()
+public:
+
+	/** @return value of parameter */
+	UFUNCTION()
+	virtual float GetParameter() const = 0;
+
+	/** notify ParameterSource that a parameter modification is about to begin */
+	UFUNCTION()
+	virtual void BeginModify() = 0;
+
+	/** set value of parameter */
+	UFUNCTION()
+	virtual void SetParameter(float NewValue) = 0;
+
+	/** notify ParameterSource that a parameter modification is complete */
+	UFUNCTION()
+	virtual void EndModify() = 0;
+};
+
+
+
+
+UINTERFACE()
+class INTERACTIVETOOLSFRAMEWORK_API UGizmoVec2ParameterSource : public UInterface
+{
+	GENERATED_BODY()
+};
+/**
+ * IGizmoVec2ParameterSource provides Get and Set for an arbitrary 2D-vector-valued parameter.
+ */
+class INTERACTIVETOOLSFRAMEWORK_API IGizmoVec2ParameterSource
+{
+	GENERATED_BODY()
+public:
+
+	/** @return value of parameter */
+	UFUNCTION()
+	virtual FVector2D GetParameter() const = 0;
+
+	/** notify ParameterSource that a parameter modification is about to begin */
+	UFUNCTION()
+	virtual void BeginModify() = 0;
+
+	/** set value of parameter */
+	UFUNCTION()
+	virtual void SetParameter(const FVector2D& NewValue) = 0;
+
+	/** notify ParameterSource that a parameter modification is complete */
+	UFUNCTION()
+	virtual void EndModify() = 0;
+};
diff --git a/Engine/Source/Runtime/Experimental/InteractiveToolsFramework/Public/BaseGizmos/GizmoMath.h b/Engine/Source/Runtime/Experimental/InteractiveToolsFramework/Public/BaseGizmos/GizmoMath.h
new file mode 100644
--- /dev/null
+++ b/Engine/Source/Runtime/Experimental/InteractiveToolsFramework/Public/BaseGizmos/GizmoMath.h
@@ -0,0 +1,103 @@
+// Copyright 1998-2019 Epic Games, Inc. All Rights Reserved.
+
+#pragma once
+
+#include "CoreMinimal.h"
+#include "BaseGizmos/GizmoInterfaces.h"
+
+/**
+ * GizmoMath functions implement various vector math/geometry operations
+ * required to implement various standard Gizmos. Although some of these
+ * are standard queries, in the context of user interface Gizmos they may
+ * have behaviors (like assumptions about valid ranges or error-handling)
+ * that would not be suitable for general-purpose math code.
+ */
+namespace GizmoMath
+{
+	/**
+	 * Project a point onto the line defined by LineOrigin and LineDirection
+	 * @return projected point
+	 */
+	INTERACTIVETOOLSFRAMEWORK_API FVector ProjectPointOntoLine(
+		const FVector& Point,
+		const FVector& LineOrigin, const FVector& LineDirection);
+	
+	/**
+	 * Find the nearest point to QueryPoint on the line defined by LineOrigin and LineDirection.
+	 * Returns this point and the associated line-equation parameter
+	 */
+	INTERACTIVETOOLSFRAMEWORK_API void NearestPointOnLine(
+		const FVector& LineOrigin, const FVector& LineDirection,
+		const FVector& QueryPoint,
+		FVector& NearestPointOut, float& LineParameterOut);
+
+	/**
+	 * Find the closest pair of points on a 3D line and 3D ray.
+	 * Returns the points and line-equation parameters on both the line and ray.
+	 */
+	INTERACTIVETOOLSFRAMEWORK_API void NearestPointOnLineToRay(
+		const FVector& LineOrigin, const FVector& LineDirection,
+		const FVector& RayOrigin, const FVector& RayDirection,
+		FVector& NearestLinePointOut, float& LineParameterOut,
+		FVector& NearestRayPointOut, float& RayParameterOut);
+
+	/**
+	 * Find the intersection of the ray defined by RayOrigin and RayDirection
+	 * with the plane defined by PlaneOrigin and PlaneNormal.
+	 * Returns intersection success/failure in bIntersectsOut and the intersection point in PlaneIntersectionPointOut
+	 */
+	INTERACTIVETOOLSFRAMEWORK_API void RayPlaneIntersectionPoint(
+		const FVector& PlaneOrigin, const FVector& PlaneNormal,
+		const FVector& RayOrigin, const FVector& RayDirection,
+		bool& bIntersectsOut, FVector& PlaneIntersectionPointOut);
+
+	/**
+	 * Find the intersection of the ray defined by RayOrigin and RayDirection
+	 * with the sphere defined by SphereOrigin and SphereRadius.
+	 * Returns intersection success/failure in bIntersectsOut and the intersection point in SphereIntersectionPointOut
+	 */
+	INTERACTIVETOOLSFRAMEWORK_API void RaySphereIntersection(
+		const FVector& SphereOrigin, const float SphereRadius,
+		const FVector& RayOrigin, const FVector& RayDirection,
+		bool& bIntersectsOut, FVector& SphereIntersectionPointOut);
+
+	/**
+	 * Find the closest point to QueryPoint that is on the circle
+	 * defined by CircleOrigin, CircleNormal, and CircleRadius.
+	 */
+	INTERACTIVETOOLSFRAMEWORK_API void ClosetPointOnCircle(
+		const FVector& QueryPoint, 
+		const FVector& CircleOrigin, const FVector& CircleNormal, float CircleRadius,
+		FVector& ClosestPointOut);
+
+	/**
+	 * Construct two mutually-perpendicular unit vectors BasisAxis1Out and BasisAxis2Out
+	 * which are orthogonal to PlaneNormal
+	 */
+	INTERACTIVETOOLSFRAMEWORK_API void MakeNormalPlaneBasis(
+		const FVector& PlaneNormal,
+		FVector& BasisAxis1Out, FVector& BasisAxis2Out);
+
+	/**
+	 * Project the input Point into the plane defined by PlaneOrigin and PlaneNormal,
+	 * and then measure it's signed angle relative to the in-plane X and Y axes
+	 * defined by PlaneAxis1 and PlaneAxis2
+	 * @return the signed angle
+	 */
+	INTERACTIVETOOLSFRAMEWORK_API float ComputeAngleInPlane(
+		const FVector& Point,
+		const FVector& PlaneOrigin, const FVector& PlaneNormal,
+		const FVector& PlaneAxis1, const FVector& PlaneAxis2);
+
+
+	/**
+	 * Project the input Point into the plane defined by PlaneOrigin and PlaneNormal,
+	 * and then calculate it's UV coordinates in the space defined by the plane 
+	 * axes PlaneAxis1 and PlaneAxis2
+	 * @return the UV coordinates
+	 */
+	INTERACTIVETOOLSFRAMEWORK_API FVector2D ComputeCoordinatesInPlane(
+		const FVector& Point,
+		const FVector& PlaneOrigin, const FVector& PlaneNormal,
+		const FVector& PlaneAxis1, const FVector& PlaneAxis2);
+}
\ No newline at end of file
diff --git a/Engine/Source/Runtime/Experimental/InteractiveToolsFramework/Public/BaseGizmos/GizmoRectangleComponent.h b/Engine/Source/Runtime/Experimental/InteractiveToolsFramework/Public/BaseGizmos/GizmoRectangleComponent.h
new file mode 100644
--- /dev/null
+++ b/Engine/Source/Runtime/Experimental/InteractiveToolsFramework/Public/BaseGizmos/GizmoRectangleComponent.h
@@ -0,0 +1,60 @@
+// Copyright 1998-2019 Epic Games, Inc. All Rights Reserved.
+
+#pragma once
+
+#include "CoreMinimal.h"
+#include "UObject/ObjectMacros.h"
+#include "BaseGizmos/GizmoBaseComponent.h"
+#include "GizmoRectangleComponent.generated.h"
+
+/**
+ * Simple Component intended to be used as part of 3D Gizmos. 
+ * Draws outline of 3D rectangle based on parameters.
+ */
+UCLASS(ClassGroup = Utility, HideCategories = (Physics, Collision, Lighting, Rendering, Mobile))
+class INTERACTIVETOOLSFRAMEWORK_API UGizmoRectangleComponent : public UGizmoBaseComponent
+{
+	GENERATED_BODY()
+
+public:
+	UPROPERTY(EditAnywhere, Category = Options)
+	FVector DirectionX = FVector(0, 0, 1);
+
+	UPROPERTY(EditAnywhere, Category = Options)
+	FVector DirectionY = FVector(0, 1, 0);
+
+	UPROPERTY(EditAnywhere, Category = Options)
+	float OffsetX = 0.0f;
+
+	UPROPERTY(EditAnywhere, Category = Options)
+	float OffsetY = 0.0f;
+
+	UPROPERTY(EditAnywhere, Category = Options)
+	float LengthX = 20.0f;
+
+	UPROPERTY(EditAnywhere, Category = Options)
+	float LengthY = 20.0f;
+
+	UPROPERTY(EditAnywhere, Category = Options)
+	float Thickness = 2.0f;
+
+
+private:
+	//~ Begin UPrimitiveComponent Interface.
+	virtual FPrimitiveSceneProxy* CreateSceneProxy() override;
+	virtual bool LineTraceComponent(FHitResult& OutHit, const FVector Start, const FVector End, const FCollisionQueryParams& Params) override;
+	virtual void GetUsedMaterials(TArray<UMaterialInterface*>& OutMaterials, bool bGetDebugMaterials = false) const override;
+	//~ End UPrimitiveComponent Interface.
+
+	//~ Begin USceneComponent Interface.
+	virtual FBoxSphereBounds CalcBounds(const FTransform& LocalToWorld) const override;
+	//~ Begin USceneComponent Interface.
+
+	// if true, we drew along -Direction instead of Direction, and so should hit-test accordingly
+	bool bFlippedX;
+	bool bFlippedY;
+
+	// gizmo visibility
+	bool bRenderVisibility = true;
+
+};
\ No newline at end of file
diff --git a/Engine/Source/Runtime/Experimental/InteractiveToolsFramework/Public/BaseGizmos/GizmoRenderingUtil.h b/Engine/Source/Runtime/Experimental/InteractiveToolsFramework/Public/BaseGizmos/GizmoRenderingUtil.h
new file mode 100644
--- /dev/null
+++ b/Engine/Source/Runtime/Experimental/InteractiveToolsFramework/Public/BaseGizmos/GizmoRenderingUtil.h
@@ -0,0 +1,31 @@
+// Copyright 1998-2019 Epic Games, Inc. All Rights Reserved.
+
+#pragma once
+
+#include "CoreMinimal.h"
+#include "Components/PrimitiveComponent.h"
+
+
+/**
+ * Utility functions for standard GizmoComponent rendering
+ */
+namespace GizmoRenderingUtil
+{
+	/**
+	 * @return best guess at "active" scene view. Prefers Perspective view if available,
+	 * otherwise first valid view.
+	 */
+	INTERACTIVETOOLSFRAMEWORK_API const FSceneView* FindActiveSceneView(
+		const TArray<const FSceneView*>& Views,
+		const FSceneViewFamily& ViewFamily, 
+		uint32 VisibilityMap);
+
+	/**
+	 * @return Conversion factor between pixel and world-space coordinates at 3D point Location in View.
+	 * @warning This is a local estimate and is increasingly incorrect as the 3D point gets further from Location
+	 */
+	INTERACTIVETOOLSFRAMEWORK_API float CalculateLocalPixelToWorldScale(
+		const FSceneView* View,
+		const FVector& Location);
+
+}
diff --git a/Engine/Source/Runtime/Experimental/InteractiveToolsFramework/Public/BaseGizmos/HitTargets.h b/Engine/Source/Runtime/Experimental/InteractiveToolsFramework/Public/BaseGizmos/HitTargets.h
new file mode 100644
--- /dev/null
+++ b/Engine/Source/Runtime/Experimental/InteractiveToolsFramework/Public/BaseGizmos/HitTargets.h
@@ -0,0 +1,90 @@
+// Copyright 1998-2019 Epic Games, Inc. All Rights Reserved.
+
+#pragma once
+
+#include "CoreMinimal.h"
+#include "BaseGizmos/GizmoInterfaces.h"
+#include "HitTargets.generated.h"
+
+class UPrimitiveComponent;
+
+
+/**
+ * UGizmoLambdaHitTarget is an IGizmoClickTarget implementation that
+ * forwards the hit-test function to a TFunction
+ */
+UCLASS()
+class INTERACTIVETOOLSFRAMEWORK_API UGizmoLambdaHitTarget : public UObject, public IGizmoClickTarget
+{
+	GENERATED_BODY()
+public:
+	/** This function is called to determine if target is hit */
+	TUniqueFunction<FInputRayHit(const FInputDeviceRay&)> IsHitFunction;
+
+	/** This function is called to update hover state of the target */
+	TFunction<void(bool)> UpdateHoverFunction;
+
+
+public:
+	virtual FInputRayHit IsHit(const FInputDeviceRay& ClickPos) const
+	{
+		if (IsHitFunction)
+		{
+			return IsHit(ClickPos);
+		}
+		return FInputRayHit();
+	}
+
+	virtual void UpdateHoverState(bool bHovering) const
+	{
+		if (UpdateHoverFunction)
+		{
+			UpdateHoverFunction(bHovering);
+		}
+	}
+};
+
+
+
+/**
+ * UGizmoComponentHitTarget is an IGizmoClickTarget implementation that
+ * hit-tests a UPrimitiveComponent
+ */
+UCLASS()
+class INTERACTIVETOOLSFRAMEWORK_API UGizmoComponentHitTarget : public UObject, public IGizmoClickTarget
+{
+	GENERATED_BODY()
+public:
+
+	/**
+	 * Component->LineTraceComponent() is called to determine if the target is hit
+	 */
+	UPROPERTY()
+	UPrimitiveComponent* Component;
+
+	/** This function is called to update hover state of the target */
+	TFunction<void(bool)> UpdateHoverFunction;
+
+
+public:
+	virtual FInputRayHit IsHit(const FInputDeviceRay& ClickPos) const;
+
+	virtual void UpdateHoverState(bool bHovering) const
+	{
+		if (UpdateHoverFunction)
+		{
+			UpdateHoverFunction(bHovering);
+		}
+	}
+
+public:
+	static UGizmoComponentHitTarget* Construct(
+		UPrimitiveComponent* Component,
+		UObject* Outer = (UObject*)GetTransientPackage())
+	{
+		UGizmoComponentHitTarget* NewHitTarget = NewObject<UGizmoComponentHitTarget>(Outer);
+		NewHitTarget->Component = Component;
+		return NewHitTarget;
+	}
+};
+
diff --git a/Engine/Source/Runtime/Experimental/InteractiveToolsFramework/Public/BaseGizmos/ParameterSourcesFloat.h b/Engine/Source/Runtime/Experimental/InteractiveToolsFramework/Public/BaseGizmos/ParameterSourcesFloat.h
new file mode 100644
--- /dev/null
+++ b/Engine/Source/Runtime/Experimental/InteractiveToolsFramework/Public/BaseGizmos/ParameterSourcesFloat.h
@@ -0,0 +1,116 @@
+// Copyright 1998-2019 Epic Games, Inc. All Rights Reserved.
+
+#pragma once
+
+#include "CoreMinimal.h"
+#include "BaseGizmos/GizmoInterfaces.h"
+#include "ParameterSourcesFloat.generated.h"
+
+/**
+ * FGizmoVec2ParameterChange represents a change in the value of an IGizmoFloatParameterSource.
+ * IGizmoFloatParameterSource implementations use this to track changes and emit delta information.
+ */
+USTRUCT()
+struct INTERACTIVETOOLSFRAMEWORK_API FGizmoFloatParameterChange
+{
+	GENERATED_BODY()
+
+	UPROPERTY()
+	float InitialValue = 0;
+
+	UPROPERTY()
+	float CurrentValue = 0;
+
+
+	float GetChangeDelta() const
+	{
+		return CurrentValue - InitialValue;
+	}
+
+	explicit FGizmoFloatParameterChange(float StartValue = 0)
+	{
+		InitialValue = CurrentValue = StartValue;
+	}
+};
+
+
+/**
+ * UGizmoBaseFloatParameterSource is a base implementation of IGizmoFloatParameterSource,
+ * which is not functional but adds an OnParameterChanged delegate for further subclasses.
+ */
+UCLASS()
+class INTERACTIVETOOLSFRAMEWORK_API UGizmoBaseFloatParameterSource : public UObject, public IGizmoFloatParameterSource
+{
+	GENERATED_BODY()
+public:
+	virtual float GetParameter() const
+	{
+		check(false);    // this is an abstract base class
+		return 0.0f;
+	}
+
+	virtual void SetParameter(float NewValue)
+	{
+		check(false);    // this is an abstract base class
+	}
+
+
+	virtual void BeginModify()
+	{
+	}
+
+	virtual void EndModify()
+	{
+	}
+
+
+public:
+	DECLARE_MULTICAST_DELEGATE_TwoParams(FOnGizmoFloatParameterSourceChanged, IGizmoFloatParameterSource*, FGizmoFloatParameterChange);
+	FOnGizmoFloatParameterSourceChanged OnParameterChanged;
+};
+
+
+
+/**
+ * UGizmoLocalFloatParameterSource is an implementation of IGizmoFloatParameterSource
+ * (by way of UGizmoBaseFloatParameterSource) which locally stores the relevant Parameter
+ * and emits update events via the OnParameterChanged delegate.
+ */
+UCLASS()
+class INTERACTIVETOOLSFRAMEWORK_API UGizmoLocalFloatParameterSource : public UGizmoBaseFloatParameterSource
+{
+	GENERATED_BODY()
+public:
+	virtual float GetParameter() const override
+	{
+		return Value;
+	}
+
+	virtual void SetParameter(float NewValue) override
+	{
+		Value = NewValue;
+		LastChange.CurrentValue = NewValue;
+		OnParameterChanged.Broadcast(this, LastChange);
+	}
+
+	virtual void BeginModify()
+	{
+		LastChange = FGizmoFloatParameterChange(Value);
+	}
+
+	virtual void EndModify()
+	{
+	}
+
+
+public:
+
+	UPROPERTY()
+	float Value = 0.0f;
+
+	UPROPERTY()
+	FGizmoFloatParameterChange LastChange;
+};
+
+
+
diff --git a/Engine/Source/Runtime/Experimental/InteractiveToolsFramework/Public/BaseGizmos/ParameterSourcesVec2.h b/Engine/Source/Runtime/Experimental/InteractiveToolsFramework/Public/BaseGizmos/ParameterSourcesVec2.h
new file mode 100644
--- /dev/null
+++ b/Engine/Source/Runtime/Experimental/InteractiveToolsFramework/Public/BaseGizmos/ParameterSourcesVec2.h
@@ -0,0 +1,122 @@
+// Copyright 1998-2019 Epic Games, Inc. All Rights Reserved.
+
+#pragma once
+
+#include "CoreMinimal.h"
+#include "BaseGizmos/GizmoInterfaces.h"
+#include "ParameterSourcesVec2.generated.h"
+
+//
+// various implementations of IGizmoVec2ParameterSource
+// 
+
+
+/**
+ * FGizmoVec2ParameterChange represents a change in the value of an IGizmoVec2ParameterSource.
+ * IGizmoVec2ParameterSource implementations use this to track changes and emit delta information.
+ */
+USTRUCT()
+struct INTERACTIVETOOLSFRAMEWORK_API FGizmoVec2ParameterChange
+{
+	GENERATED_BODY()
+
+	UPROPERTY()
+	FVector2D InitialValue = FVector2D::ZeroVector;
+
+	UPROPERTY()
+	FVector2D CurrentValue = FVector2D::ZeroVector;
+
+
+	FVector2D GetChangeDelta() const
+	{
+		return CurrentValue - InitialValue;
+	}
+
+	explicit FGizmoVec2ParameterChange(const FVector2D& StartValue = FVector2D::ZeroVector)
+	{
+		InitialValue = CurrentValue = StartValue;
+	}
+};
+
+
+
+/**
+ * UGizmoBaseVec2ParameterSource is a base implementation of IGizmoVec2ParameterSource,
+ * which is not functional but adds an OnParameterChanged delegate for further subclasses.
+ */
+UCLASS()
+class INTERACTIVETOOLSFRAMEWORK_API UGizmoBaseVec2ParameterSource : public UObject, public IGizmoVec2ParameterSource
+{
+	GENERATED_BODY()
+public:
+	virtual FVector2D GetParameter() const
+	{
+		check(false);    // this is an abstract base class
+		return FVector2D::ZeroVector;
+	}
+
+	virtual void SetParameter(const FVector2D& NewValue)
+	{
+		check(false);    // this is an abstract base class
+	}
+
+
+	virtual void BeginModify()
+	{
+	}
+
+	virtual void EndModify()
+	{
+	}
+
+
+public:
+	DECLARE_MULTICAST_DELEGATE_TwoParams(FOnGizmoVec2ParameterSourceChanged, IGizmoVec2ParameterSource*, FGizmoVec2ParameterChange);
+	FOnGizmoVec2ParameterSourceChanged OnParameterChanged;
+};
+
+
+
+/**
+ * UGizmoBaseVec2ParameterSource is an implementation of IGizmoVec2ParameterSource 
+ * (by way of UGizmoBaseVec2ParameterSource) which locally stores the relevant Parameter
+ * and emits update events via the OnParameterChanged delegate.
+ */
+UCLASS()
+class INTERACTIVETOOLSFRAMEWORK_API UGizmoLocalVec2ParameterSource : public UGizmoBaseVec2ParameterSource
+{
+	GENERATED_BODY()
+public:
+	virtual FVector2D GetParameter() const override
+	{
+		return Value;
+	}
+
+	virtual void SetParameter(const FVector2D& NewValue) override
+	{
+		Value = NewValue;
+		LastChange.CurrentValue = NewValue;
+		OnParameterChanged.Broadcast(this, LastChange);
+	}
+
+	virtual void BeginModify()
+	{
+		LastChange = FGizmoVec2ParameterChange(Value);
+	}
+
+	virtual void EndModify()
+	{
+	}
+
+
+public:
+
+	UPROPERTY()
+	FVector2D Value = FVector2D::ZeroVector;
+
+	UPROPERTY()
+	FGizmoVec2ParameterChange LastChange;
+};
+
+
+
diff --git a/Engine/Source/Runtime/Experimental/InteractiveToolsFramework/Public/BaseGizmos/ParameterToTransformAdapters.h b/Engine/Source/Runtime/Experimental/InteractiveToolsFramework/Public/BaseGizmos/ParameterToTransformAdapters.h
new file mode 100644
--- /dev/null
+++ b/Engine/Source/Runtime/Experimental/InteractiveToolsFramework/Public/BaseGizmos/ParameterToTransformAdapters.h
@@ -0,0 +1,338 @@
+// Copyright 1998-2019 Epic Games, Inc. All Rights Reserved.
+
+#pragma once
+
+#include "CoreMinimal.h"
+#include "BaseGizmos/GizmoInterfaces.h"
+#include "ParameterToTransformAdapters.generated.h"
+
+
+//
+// Various 1D and 2D ParameterSource converters intended to be used to create 3D transformation gizmos.
+// Based on base classes in ParameterSourcesFloat.h and ParameterSourcesVec2.h
+// 
+
+
+
+/**
+ * UGizmoAxisTranslationParameterSource is an IGizmoFloatParameterSource implementation that
+ * interprets the float value as the parameter of a line equation, and maps this parameter to a 3D translation 
+ * along a line with origin/direction given by an IGizmoAxisSource. This translation is applied to an IGizmoTransformSource.
+ *
+ * This ParameterSource is intended to be used to create 3D Axis Translation Gizmos.
+ */
+UCLASS()
+class INTERACTIVETOOLSFRAMEWORK_API UGizmoAxisTranslationParameterSource : public UGizmoBaseFloatParameterSource
+{
+	GENERATED_BODY()
+public:
+	virtual float GetParameter() const override
+	{
+		return Parameter;
+	}
+
+	virtual void SetParameter(float NewValue) override
+	{
+		Parameter = NewValue;
+		LastChange.CurrentValue = NewValue;
+
+		// construct translation as delta from initial position
+		FVector Translation = LastChange.GetChangeDelta() * CurTranslationAxis;
+
+		// translate the initial transform
+		FTransform NewTransform = InitialTransform;
+		NewTransform.AddToTranslation(Translation);
+		TransformSource->SetTransform(NewTransform);
+
+		OnParameterChanged.Broadcast(this, LastChange);
+	}
+
+	virtual void BeginModify()
+	{
+		check(AxisSource);
+
+		LastChange = FGizmoFloatParameterChange(Parameter);
+
+		InitialTransform = TransformSource->GetTransform();
+		CurTranslationAxis = AxisSource->GetDirection();
+		CurTranslationOrigin = AxisSource->GetOrigin();
+	}
+
+	virtual void EndModify()
+	{
+	}
+
+
+public:
+	/** The Parameter line-equation value is converted to a 3D Translation along this Axis */
+	UPROPERTY()
+	TScriptInterface<IGizmoAxisSource> AxisSource;
+
+	/** This TransformSource is updated by applying the constructed 3D translation */
+	UPROPERTY()
+	TScriptInterface<IGizmoTransformSource> TransformSource;
+
+
+public:
+	/** Parameter is the line-equation parameter that this FloatParameterSource provides */
+	UPROPERTY()
+	float Parameter = 0.0f;
+
+	/** Active parameter change (only valid between BeginModify/EndModify) */
+	UPROPERTY()
+	FGizmoFloatParameterChange LastChange;
+
+	/** tranlsation axis for current parameter edit (only valid between BeginModify/EndModify) */
+	UPROPERTY()
+	FVector CurTranslationAxis;
+
+	/** tranlsation origin for current parameter edit (only valid between BeginModify/EndModify) */
+	UPROPERTY()
+	FVector CurTranslationOrigin;
+
+	/** Saved copy of Initial Transform for current parameter edit (only valid between BeginModify/EndModify) */
+	UPROPERTY()
+	FTransform InitialTransform;
+
+public:
+
+	/**
+	 * Create a standard instance of this ParameterSource, with the given AxisSource and TransformSource
+	 */
+	static UGizmoAxisTranslationParameterSource* Construct(
+		IGizmoAxisSource* AxisSourceIn,
+		IGizmoTransformSource* TransformSourceIn,
+		UObject* Outer = (UObject*)GetTransientPackage())
+	{
+		UGizmoAxisTranslationParameterSource* NewSource = NewObject<UGizmoAxisTranslationParameterSource>(Outer);
+		NewSource->AxisSource = Cast<UObject>(AxisSourceIn);
+		NewSource->TransformSource = Cast<UObject>(TransformSourceIn);
+		return NewSource;
+	}
+};
+
+
+
+
+
+/**
+ * UGizmoAxisRotationParameterSource is an IGizmoVec2ParameterSource implementation that
+ * interprets the FVector2D parameter as a position in a 2D plane, and maps this position to a 3D translation
+ * a plane with origin/normal given by an IGizmoAxisSource. This translation is applied to an IGizmoTransformSource.
+ * 
+ * This ParameterSource is intended to be used to create 3D Plane Translation Gizmos.
+ */
+UCLASS()
+class INTERACTIVETOOLSFRAMEWORK_API UGizmoPlaneTranslationParameterSource : public UGizmoBaseVec2ParameterSource
+{
+	GENERATED_BODY()
+public:
+	virtual FVector2D GetParameter() const override
+	{
+		return Parameter;
+	}
+
+	virtual void SetParameter(const FVector2D& NewValue) override
+	{
+		Parameter = NewValue;
+		LastChange.CurrentValue = NewValue;
+
+		// construct translation as delta from initial position
+		FVector2D Delta = LastChange.GetChangeDelta();
+		FVector Translation = Delta.X*CurTranslationAxisX + Delta.Y*CurTranslationAxisY;
+
+		// apply translation to initial transform
+		FTransform NewTransform = InitialTransform;
+		NewTransform.AddToTranslation(Translation);
+		TransformSource->SetTransform(NewTransform);
+
+		OnParameterChanged.Broadcast(this, LastChange);
+	}
+
+	virtual void BeginModify()
+	{
+		check(AxisSource);
+
+		LastChange = FGizmoVec2ParameterChange(Parameter);
+
+		// save initial transformation and axis information
+		InitialTransform = TransformSource->GetTransform();
+		CurTranslationOrigin = AxisSource->GetOrigin();
+		AxisSource->GetAxisFrame(CurTranslationNormal, CurTranslationAxisX, CurTranslationAxisY);
+	}
+
+	virtual void EndModify()
+	{
+	}
+
+
+public:
+	/** AxisSource provides the 3D plane (origin/normal/u/v) that is used to interpret the 2D parameters */
+	UPROPERTY()
+	TScriptInterface<IGizmoAxisSource> AxisSource;
+
+	/** This TransformSource is updated by applying the constructed 3D translation */
+	UPROPERTY()
+	TScriptInterface<IGizmoTransformSource> TransformSource;
+
+
+public:
+	/** Parameter is the two line-equation parameters that this Vec2ParameterSource provides */
+	UPROPERTY()
+	FVector2D Parameter = FVector2D::ZeroVector;
+
+	/** Active parameter change (only valid between BeginModify/EndModify) */
+	UPROPERTY()
+	FGizmoVec2ParameterChange LastChange;
+
+	/** plane origin for current parameter edit (only valid between BeginModify/EndModify) */
+	UPROPERTY()
+	FVector CurTranslationOrigin;
+
+	/** plane normal for current parameter edit (only valid between BeginModify/EndModify) */
+	UPROPERTY()
+	FVector CurTranslationNormal;
+
+	/** in-plane axis X for current parameter edit (only valid between BeginModify/EndModify) */
+	UPROPERTY()
+	FVector CurTranslationAxisX;
+
+	/** in-plane axis Y for current parameter edit (only valid between BeginModify/EndModify) */
+	UPROPERTY()
+	FVector CurTranslationAxisY;
+
+	/** Saved copy of Initial Transform for current parameter edit (only valid between BeginModify/EndModify) */
+	UPROPERTY()
+	FTransform InitialTransform;
+
+public:
+
+	/**
+	 * Create a standard instance of this ParameterSource, with the given AxisSource and TransformSource
+	 */
+	static UGizmoPlaneTranslationParameterSource* Construct(
+		IGizmoAxisSource* AxisSourceIn,
+		IGizmoTransformSource* TransformSourceIn,
+		UObject* Outer = (UObject*)GetTransientPackage())
+	{
+		UGizmoPlaneTranslationParameterSource* NewSource = NewObject<UGizmoPlaneTranslationParameterSource>(Outer);
+		NewSource->AxisSource = Cast<UObject>(AxisSourceIn);
+		NewSource->TransformSource = Cast<UObject>(TransformSourceIn);
+		return NewSource;
+	}
+};
+
+
+
+
+
+
+/**
+ * UGizmoAxisRotationParameterSource is an IGizmoFloatParameterSource implementation that
+ * interprets the float parameter as an angle, and maps this angle to a 3D rotation
+ * around an IGizmoAxisSource (ie 3D axis). This rotation is applied to an IGizmoTransformSource.
+ * This ParameterSource is intended to be used to create 3D Rotation Gizmos.
+ */
+UCLASS()
+class INTERACTIVETOOLSFRAMEWORK_API UGizmoAxisRotationParameterSource : public UGizmoBaseFloatParameterSource
+{
+	GENERATED_BODY()
+public:
+	virtual float GetParameter() const override
+	{
+		return Angle;
+	}
+
+	virtual void SetParameter(float NewValue) override
+	{
+		Angle = NewValue;
+		LastChange.CurrentValue = NewValue;
+
+		// construct rotation as delta from initial position
+		FQuat DeltaRotation(CurRotationAxis, LastChange.GetChangeDelta());
+
+		// rotate the vector from the rotation origin to the transform origin, 
+		// to get the translation of the origin produced by the rotation
+		FVector DeltaPosition = InitialTransform.GetLocation() - CurRotationOrigin;
+		DeltaPosition = DeltaRotation * DeltaPosition;
+		FVector NewLocation = CurRotationOrigin + DeltaPosition;
+
+		// rotate the initial transform by the rotation
+		FQuat NewRotation = DeltaRotation * InitialTransform.GetRotation();
+
+		// construct new transform
+		FTransform NewTransform = InitialTransform;
+		NewTransform.SetLocation(NewLocation);
+		NewTransform.SetRotation(NewRotation);
+		TransformSource->SetTransform(NewTransform);
+
+		OnParameterChanged.Broadcast(this, LastChange);
+	}
+
+	virtual void BeginModify()
+	{
+		check(AxisSource != nullptr);
+
+		LastChange = FGizmoFloatParameterChange(Angle);
+
+		// save initial transformation and axis information
+		InitialTransform = TransformSource->GetTransform();
+		CurRotationAxis = AxisSource->GetDirection();
+		CurRotationOrigin = AxisSource->GetOrigin();
+	}
+
+	virtual void EndModify()
+	{
+	}
+
+
+public:
+	/** float-parameter Angle is mapped to a 3D Rotation around this Axis */
+	UPROPERTY()
+	TScriptInterface<IGizmoAxisSource> AxisSource;
+
+	/** This TransformSource is updated by applying the constructed 3D rotation */
+	UPROPERTY()
+	TScriptInterface<IGizmoTransformSource> TransformSource;
+
+
+public:
+	/** Angle is the parameter that this FloatParameterSource provides */
+	UPROPERTY()
+	float Angle = 0.0f;
+
+	/** Active parameter change (only valid between BeginModify/EndModify) */
+	UPROPERTY()
+	FGizmoFloatParameterChange LastChange;
+
+	/** Rotation axis for current parameter edit (only valid between BeginModify/EndModify) */
+	UPROPERTY()
+	FVector CurRotationAxis;
+
+	/** Rotation origin for current parameter edit (only valid between BeginModify/EndModify) */
+	UPROPERTY()
+	FVector CurRotationOrigin;
+
+	/** Saved copy of Initial Transform for current parameter edit (only valid between BeginModify/EndModify) */
+	UPROPERTY()
+	FTransform InitialTransform;
+
+
+public:
+
+	/**
+	 * Create a standard instance of this ParameterSource, with the given AxisSource and TransformSource
+	 */
+	static UGizmoAxisRotationParameterSource* Construct(
+		IGizmoAxisSource* AxisSourceIn,
+		IGizmoTransformSource* TransformSourceIn,
+		UObject* Outer = (UObject*)GetTransientPackage())
+	{
+		UGizmoAxisRotationParameterSource* NewSource = NewObject<UGizmoAxisRotationParameterSource>(Outer);
+		NewSource->AxisSource = Cast<UObject>(AxisSourceIn);
+		NewSource->TransformSource = Cast<UObject>(TransformSourceIn);
+		return NewSource;
+	}
+};
+
+
diff --git a/Engine/Source/Runtime/Experimental/InteractiveToolsFramework/Public/BaseGizmos/PlanePositionGizmo.h b/Engine/Source/Runtime/Experimental/InteractiveToolsFramework/Public/BaseGizmos/PlanePositionGizmo.h
new file mode 100644
--- /dev/null
+++ b/Engine/Source/Runtime/Experimental/InteractiveToolsFramework/Public/BaseGizmos/PlanePositionGizmo.h
@@ -0,0 +1,96 @@
+// Copyright 1998-2019 Epic Games, Inc. All Rights Reserved.
+
+#pragma once
+
+#include "CoreMinimal.h"
+#include "InteractiveGizmo.h"
+#include "InteractiveGizmoBuilder.h"
+#include "BaseBehaviors/BehaviorTargetInterfaces.h"
+#include "BaseGizmos/GizmoComponents.h"
+#include "PlanePositionGizmo.generated.h"
+
+
+
+
+UCLASS()
+class INTERACTIVETOOLSFRAMEWORK_API UPlanePositionGizmoBuilder : public UInteractiveGizmoBuilder
+{
+	GENERATED_BODY()
+
+public:
+	virtual UInteractiveGizmo* BuildGizmo(const FToolBuilderState& SceneState) const override;
+};
+
+
+/**
+ *
+ */
+UCLASS()
+class INTERACTIVETOOLSFRAMEWORK_API UPlanePositionGizmo : public UInteractiveGizmo, public IClickDragBehaviorTarget, public IHoverBehaviorTarget
+{
+	GENERATED_BODY()
+
+public:
+	// UInteractiveGizmo overrides
+
+	virtual void Setup() override;
+
+	// IClickDragBehaviorTarget implementation
+
+	virtual FInputRayHit CanBeginClickDragSequence(const FInputDeviceRay& PressPos) override;
+	virtual void OnClickPress(const FInputDeviceRay& PressPos) override;
+	virtual void OnClickDrag(const FInputDeviceRay& DragPos) override;
+	virtual void OnClickRelease(const FInputDeviceRay& ReleasePos) override;
+	virtual void OnTerminateDragSequence() override;
+
+	// IHoverBehaviorTarget implementation
+	virtual void OnUpdateHover(const FInputDeviceRay& DevicePos) override;
+
+public:
+	UPROPERTY()
+	TScriptInterface<IGizmoAxisSource> AxisSource;
+
+	UPROPERTY()
+	TScriptInterface<IGizmoVec2ParameterSource> ParameterSource;
+
+	UPROPERTY()
+	TScriptInterface<IGizmoClickTarget> HitTarget;
+
+	UPROPERTY()
+	TScriptInterface<IGizmoStateTarget> StateTarget;
+
+public:
+	UPROPERTY()
+	bool bInInteraction = false;
+
+	UPROPERTY()
+	FVector InteractionOrigin;
+
+	UPROPERTY()
+	FVector InteractionNormal;
+
+	UPROPERTY()
+	FVector InteractionAxisX;
+
+	UPROPERTY()
+	FVector InteractionAxisY;
+
+
+	UPROPERTY()
+	FVector InteractionStartPoint;
+
+	UPROPERTY()
+	FVector InteractionCurPoint;
+
+	UPROPERTY()
+	FVector2D InteractionStartParameter;
+
+	UPROPERTY()
+	FVector2D InteractionCurParameter;
+
+
+protected:
+	FVector LastHitPosition;
+	FVector2D InitialTargetParameter;
+};
+
diff --git a/Engine/Source/Runtime/Experimental/InteractiveToolsFramework/Public/BaseGizmos/StateTargets.h b/Engine/Source/Runtime/Experimental/InteractiveToolsFramework/Public/BaseGizmos/StateTargets.h
new file mode 100644
--- /dev/null
+++ b/Engine/Source/Runtime/Experimental/InteractiveToolsFramework/Public/BaseGizmos/StateTargets.h
@@ -0,0 +1,127 @@
+// Copyright 1998-2019 Epic Games, Inc. All Rights Reserved.
+
+#pragma once
+
+#include "CoreMinimal.h"
+#include "InteractiveGizmoManager.h"
+#include "BaseGizmos/GizmoInterfaces.h"
+#include "StateTargets.generated.h"
+
+
+/**
+ * UGizmoNilStateTarget is an implementation of IGizmoStateTarget that does nothing
+ */
+UCLASS()
+class INTERACTIVETOOLSFRAMEWORK_API UGizmoNilStateTarget : public UObject, public IGizmoStateTarget
+{
+	GENERATED_BODY()
+public:
+	virtual void BeginUpdate() final
+	{
+	}
+
+	virtual void EndUpdate() final
+	{
+	}
+
+};
+
+
+
+/**
+ * UGizmoLambdaStateTarget is an implementation of IGizmoStateTarget that forwards
+ * calls to its interface functions to external TFunctions
+ */
+UCLASS()
+class INTERACTIVETOOLSFRAMEWORK_API UGizmoLambdaStateTarget : public UObject, public IGizmoStateTarget
+{
+	GENERATED_BODY()
+public:
+	virtual void BeginUpdate()
+	{
+		if (BeginUpdateFunction)
+		{
+			BeginUpdateFunction();
+		}
+	}
+
+	virtual void EndUpdate()
+	{
+		if (EndUpdateFunction)
+		{
+			EndUpdateFunction();
+		}
+	}
+
+	TUniqueFunction<void(void)> BeginUpdateFunction = TUniqueFunction<void(void)>();
+	TUniqueFunction<void(void)> EndUpdateFunction = TUniqueFunction<void(void)>();
+};
+
+
+
+/**
+ * UGizmoObjectModifyStateTarget is an implementation of IGizmoStateTarget that 
+ * opens and closes change transactions on a target UObject via a GizmoManager.
+ */
+UCLASS()
+class INTERACTIVETOOLSFRAMEWORK_API UGizmoObjectModifyStateTarget : public UObject, public IGizmoStateTarget
+{
+	GENERATED_BODY()
+public:
+	virtual void BeginUpdate()
+	{
+		if (GizmoManager.IsValid())
+		{
+			GizmoManager->BeginUndoTransaction(TransactionDescription);
+		}
+		if (ModifyObject.IsValid())
+		{
+			ModifyObject->Modify();
+		}
+	}
+
+	virtual void EndUpdate()
+	{
+		if (GizmoManager.IsValid())
+		{
+			GizmoManager->EndUndoTransaction();
+		}
+	}
+
+
+	/**
+	 * The object that will be changed, ie have Modify() called on it on BeginUpdate()
+	 */
+	TWeakObjectPtr<UObject> ModifyObject;
+
+	/**
+	 * Localized text description of the transaction (will be visible in Editor on undo/redo)
+	 */
+	FText TransactionDescription;
+
+	/**
+	 * Pointer to the GizmoManager that is used to open/close the transaction
+	 */
+	TWeakObjectPtr<UInteractiveGizmoManager> GizmoManager;
+
+
+public:
+	/**
+	 * Create and initialize an standard instance of UGizmoObjectModifyStateTarget
+	 * @param ModifyObjectIn the object this StateTarget will call Modify() on
+	 * @param DescriptionIn Localized text description of the transaction
+	 * @param GizmoManagerIn pointer to the GizmoManager used to manage transactions
+	 */
+	static UGizmoObjectModifyStateTarget* Construct(
+		UObject* ModifyObjectIn,
+		const FText& DescriptionIn,
+		UInteractiveGizmoManager* GizmoManagerIn,
+		UObject* Outer = (UObject*)GetTransientPackage())
+	{
+		UGizmoObjectModifyStateTarget* NewTarget = NewObject<UGizmoObjectModifyStateTarget>(Outer);
+		NewTarget->ModifyObject = ModifyObjectIn;
+		NewTarget->TransactionDescription = DescriptionIn;
+		NewTarget->GizmoManager = GizmoManagerIn;
+		return NewTarget;
+	}
+};
diff --git a/Engine/Source/Runtime/Experimental/InteractiveToolsFramework/Public/BaseGizmos/TransformGizmo.h b/Engine/Source/Runtime/Experimental/InteractiveToolsFramework/Public/BaseGizmos/TransformGizmo.h
new file mode 100644
--- /dev/null
+++ b/Engine/Source/Runtime/Experimental/InteractiveToolsFramework/Public/BaseGizmos/TransformGizmo.h
@@ -0,0 +1,244 @@
+// Copyright 1998-2019 Epic Games, Inc. All Rights Reserved.
+
+#pragma once
+
+#include "CoreMinimal.h"
+#include "InteractiveGizmo.h"
+#include "InteractiveGizmoBuilder.h"
+#include "InteractiveToolObjects.h"
+#include "BaseGizmos/GizmoActor.h"
+#include "BaseGizmos/TransformProxy.h"
+#include "TransformGizmo.generated.h"
+
+class UInteractiveGizmoManager;
+class IGizmoAxisSource;
+class IGizmoTransformSource;
+class IGizmoStateTarget;
+class UGizmoComponentAxisSource;
+
+
+/**
+ * ATransformGizmoActor is an Actor type intended to be used with UTransformGizmo,
+ * as the in-scene visual representation of the Gizmo.
+ * 
+ * FTransformGizmoActorFactory returns an instance of this Actor type (or a subclass), and based on
+ * which Translate and Rotate UProperties are initialized, will associate those Components
+ * with UInteractiveGizmo's that implement Axis Translation, Plane Translation, and Axis Rotation.
+ * 
+ * If a particular sub-Gizmo is not required, simply set that UProperty to null.
+ * 
+ * The static factory method ::ConstructDefault3AxisGizmo() creates and initializes an 
+ * Actor suitable for use in a standard 3-axis Transformation Gizmo.
+ */
+UCLASS(Transient)
+class INTERACTIVETOOLSFRAMEWORK_API ATransformGizmoActor : public AGizmoActor
+{
+	GENERATED_BODY()
+public:
+
+	ATransformGizmoActor();
+
+public:
+	/** X Axis Translation Component */
+	UPROPERTY()
+	UPrimitiveComponent* TranslateX;
+
+	/** Y Axis Translation Component */
+	UPROPERTY()
+	UPrimitiveComponent* TranslateY;
+
+	/** Z Axis Translation Component */
+	UPROPERTY()
+	UPrimitiveComponent* TranslateZ;
+
+
+	/** YZ Plane Translation Component */
+	UPROPERTY()
+	UPrimitiveComponent* TranslateYZ;
+
+	/** XZ Plane Translation Component */
+	UPROPERTY()
+	UPrimitiveComponent* TranslateXZ;
+
+	/** XY Plane Translation Component */
+	UPROPERTY()
+	UPrimitiveComponent* TranslateXY;
+
+
+	/** X Axis Rotation Component */
+	UPROPERTY()
+	UPrimitiveComponent* RotateX;
+
+	/** Y Axis Rotation Component */
+	UPROPERTY()
+	UPrimitiveComponent* RotateY;
+
+	/** Z Axis Rotation Component */
+	UPROPERTY()
+	UPrimitiveComponent* RotateZ;
+
+
+public:
+	/**
+	 * Create a new instance of ATransformGizmoActor and populate the various
+	 * sub-components with standard GizmoXComponent instances suitable for a 3-axis transformer Gizmo
+	 */
+	static ATransformGizmoActor* ConstructDefault3AxisGizmo(
+		UWorld* World
+	);
+};
+
+
+
+
+
+/**
+ * FTransformGizmoActorFactory creates new instances of ATransformGizmoActor which
+ * are used by UTransformGizmo to implement 3D transformation Gizmos. 
+ * An instance of FTransformGizmoActorFactory is passed to UTransformGizmo
+ * (by way of UTransformGizmoBuilder), which then calls CreateNewGizmoActor()
+ * to spawn new Gizmo Actors.
+ * 
+ * By default CreateNewGizmoActor() returns a default Gizmo Actor suitable for
+ * a three-axis transformation Gizmo, override this function to customize
+ * the Actor sub-elements.
+ */
+class INTERACTIVETOOLSFRAMEWORK_API FTransformGizmoActorFactory
+{
+public:
+	/**
+	 * @param World the UWorld to create the new Actor in
+	 * @return new ATransformGizmoActor instance with members initialized with Components suitable for a transformation Gizmo
+	 */
+	virtual ATransformGizmoActor* CreateNewGizmoActor(UWorld* World) const;
+};
+
+
+
+
+
+
+UCLASS()
+class INTERACTIVETOOLSFRAMEWORK_API UTransformGizmoBuilder : public UInteractiveGizmoBuilder
+{
+	GENERATED_BODY()
+
+public:
+	/**
+	 * If set, this Actor Builder will be passed to UTransformGizmo instances.
+	 * Otherwise new instances of the base FTransformGizmoActorFactory are created internally.
+	 */
+	TSharedPtr<FTransformGizmoActorFactory> GizmoActorBuilder;
+
+	/**
+	 * If set, this hover function will be passed to UTransformGizmo instances to use instead of the default.
+	 * Hover is complicated for UTransformGizmo because all it knows about the different gizmo scene elements
+	 * is that they are UPrimitiveComponent (coming from the ATransformGizmoActor). The default hover
+	 * function implementation is to try casting to UGizmoBaseComponent and calling ::UpdateHoverState().
+	 * If you are using different Components that do not subclass UGizmoBaseComponent, and you want hover to 
+	 * work, you will need to provide a different hover update function.
+	 */
+	TFunction<void(UPrimitiveComponent*, bool)> UpdateHoverFunction;
+
+	virtual UInteractiveGizmo* BuildGizmo(const FToolBuilderState& SceneState) const override;
+};
+
+
+/**
+ * UTransformGizmo provides standard Transformation Gizmo interactions,
+ * applied to a UTransformProxy target object. By default the Gizmo will be
+ * a standard XYZ translate/rotate Gizmo (axis and plane translation).
+ * 
+ * The in-scene representation of the Gizmo is a ATransformGizmoActor (or subclass).
+ * This Actor has UProperty members for the various sub-widgets, each as a separate Component.
+ * Any particular sub-widget of the Gizmo can be disabled by setting the respective
+ * Actor Component to null. 
+ * 
+ * So, to create non-standard variants of the Transform Gizmo, set a new GizmoActorBuilder 
+ * in the UTransformGizmoBuilder registered with the GizmoManager. Return
+ * a suitably-configured GizmoActor and everything else will be handled automatically.
+ * 
+ */
+UCLASS()
+class INTERACTIVETOOLSFRAMEWORK_API UTransformGizmo : public UInteractiveGizmo
+{
+	GENERATED_BODY()
+
+public:
+
+	virtual void SetWorld(UWorld* World);
+	virtual void SetGizmoActorBuilder(TSharedPtr<FTransformGizmoActorFactory> Builder);
+	virtual void SetUpdateHoverFunction(TFunction<void(UPrimitiveComponent*, bool)> HoverFunction);
+
+	// UInteractiveGizmo overrides
+	virtual void Setup() override;
+	virtual void Shutdown() override;
+
+
+	/**
+	 * Set the active target object for the Gizmo
+	 * @param Target active target
+	 */
+	virtual void SetActiveTarget(UTransformProxy* Target);
+
+	/**
+	 * Clear the active target object for the Gizmo
+	 */
+	virtual void ClearActiveTarget();
+
+	/** The active target object for the Gizmo */
+	UPROPERTY()
+	UTransformProxy* ActiveTarget;
+
+
+	ATransformGizmoActor* GetGizmoActor() const { return GizmoActor; }
+
+
+protected:
+	TSharedPtr<FTransformGizmoActorFactory> GizmoActorBuilder;
+
+	TFunction<void(UPrimitiveComponent*, bool)> UpdateHoverFunction;
+
+	UPROPERTY()
+	TArray<UInteractiveGizmo*> ActiveGizmos;
+
+	UWorld* World;
+	ATransformGizmoActor* GizmoActor;
+
+	/** X-axis source is shared across Gizmos, and created internally during SetActiveTarget() */
+	UPROPERTY()
+	UGizmoComponentAxisSource* AxisXSource;
+
+	/** Y-axis source is shared across Gizmos, and created internally during SetActiveTarget() */
+	UPROPERTY()
+	UGizmoComponentAxisSource* AxisYSource;
+
+	/** Z-axis source is shared across Gizmos, and created internally during SetActiveTarget() */
+	UPROPERTY()
+	UGizmoComponentAxisSource* AxisZSource;
+
+	/** @return a new instance of the standard axis-translation Gizmo */
+	virtual UInteractiveGizmo* AddAxisTranslationGizmo(
+		UPrimitiveComponent* AxisComponent, USceneComponent* RootComponent,
+		const FString& Identifier,
+		IGizmoAxisSource* AxisSource,
+		IGizmoTransformSource* TransformSource, 
+		IGizmoStateTarget* StateTarget);
+
+	/** @return a new instance of the standard plane-translation Gizmo */
+	virtual UInteractiveGizmo* AddPlaneTranslationGizmo(
+		UPrimitiveComponent* AxisComponent, USceneComponent* RootComponent,
+		const FString& Identifier,
+		IGizmoAxisSource* AxisSource,
+		IGizmoTransformSource* TransformSource,
+		IGizmoStateTarget* StateTarget);
+
+	/** @return a new instance of the standard axis-rotation Gizmo */
+	virtual UInteractiveGizmo* AddAxisRotationGizmo(
+		UPrimitiveComponent* AxisComponent, USceneComponent* RootComponent,
+		const FString& Identifier,
+		IGizmoAxisSource* AxisSource,
+		IGizmoTransformSource* TransformSource,
+		IGizmoStateTarget* StateTarget);
+};
+
diff --git a/Engine/Source/Runtime/Experimental/InteractiveToolsFramework/Public/BaseGizmos/TransformProxy.h b/Engine/Source/Runtime/Experimental/InteractiveToolsFramework/Public/BaseGizmos/TransformProxy.h
new file mode 100644
--- /dev/null
+++ b/Engine/Source/Runtime/Experimental/InteractiveToolsFramework/Public/BaseGizmos/TransformProxy.h
@@ -0,0 +1,81 @@
+// Copyright 1998-2019 Epic Games, Inc. All Rights Reserved.
+
+#pragma once
+
+#include "CoreMinimal.h"
+#include "UObject/ObjectMacros.h"
+#include "TransformProxy.generated.h"
+
+class USceneComponent;
+
+/**
+ * UTransformProxy is used to transform a set of sub-ojects. An internal
+ * FTransform is generated based on the sub-object set, and the relative
+ * FTransform of each sub-object is stored. Then as this main transform
+ * is updated, the sub-objects are also updated.
+ * 
+ * Currently only USceneComponent sub-objects are supported.
+ * 
+ * If only one sub-object is set, the main transform is the sub-object transform.
+ * Otherwise the main transform is centered at the average origin and
+ * has no rotation.
+ */
+UCLASS()
+class INTERACTIVETOOLSFRAMEWORK_API UTransformProxy : public UObject
+{
+	GENERATED_BODY()
+public:
+
+	/**
+	 * Add a component sub-object to the proxy set. 
+	 * @param bModifyComponentOnTransform if true, Component->Modify() is called before the Component transform is updated
+	 * @warning The internal shared transform is regenerated each time a component is added
+	 */
+	virtual void AddComponent(USceneComponent* Component, bool bModifyComponentOnTransform = true);
+
+	/**
+	 * @return the shared transform for all the sub-objects
+	 */
+	virtual FTransform GetTransform() const;
+
+	/**
+	 * Update the main transform and then update the sub-objects based on their relative transformations
+	 */
+	virtual void SetTransform(const FTransform& Transform);
+
+
+public:
+	/**
+	 * This delegate is fired whenever the internal transform changes, ie
+	 * on AddComponent and SetTransform
+	 */
+	DECLARE_MULTICAST_DELEGATE_TwoParams(FOnTransformChanged, UTransformProxy*, FTransform);
+	FOnTransformChanged OnTransformChanged;
+
+	
+protected:
+
+	struct FRelativeObject
+	{
+		TWeakObjectPtr<USceneComponent> Component;
+		bool bModifyComponentOnTransform;
+
+		/** The initial transform of the object, set during UpdateSharedTransform() */
+		FTransform StartTransform;
+		/** The transform of the object relative to */
+		FTransform RelativeTransform;
+	};
+
+	/** List of sub-objects */
+	TArray<FRelativeObject> Objects;
+
+	/** The main transform */
+	UPROPERTY()
+	FTransform SharedTransform;
+
+	/** Recalculate main SharedTransform when object set changes*/
+	virtual void UpdateSharedTransform();
+
+	/** Propagate a transform update to the sub-objects */
+	virtual void UpdateObjects();
+};
\ No newline at end of file
diff --git a/Engine/Source/Runtime/Experimental/InteractiveToolsFramework/Public/BaseGizmos/TransformSources.h b/Engine/Source/Runtime/Experimental/InteractiveToolsFramework/Public/BaseGizmos/TransformSources.h
new file mode 100644
--- /dev/null
+++ b/Engine/Source/Runtime/Experimental/InteractiveToolsFramework/Public/BaseGizmos/TransformSources.h
@@ -0,0 +1,109 @@
+// Copyright 1998-2019 Epic Games, Inc. All Rights Reserved.
+
+#pragma once
+
+#include "CoreMinimal.h"
+#include "BaseGizmos/GizmoInterfaces.h"
+#include "BaseGizmos/TransformProxy.h"
+#include "TransformSources.generated.h"
+
+
+
+/**
+ * UGizmoBaseTransformSource is a base implementation of IGizmoTransformSource that 
+ * adds an OnTransformChanged delegate. This class cannot be used directly and must be subclassed.
+ */
+UCLASS()
+class INTERACTIVETOOLSFRAMEWORK_API UGizmoBaseTransformSource : public UObject, public IGizmoTransformSource
+{
+	GENERATED_BODY()
+public:
+	virtual FTransform GetTransform() const
+	{
+		return FTransform::Identity;
+	}
+
+	virtual void SetTransform(const FTransform& NewTransform)
+	{
+		check(false);   // not implemented
+	}
+
+public:
+	DECLARE_MULTICAST_DELEGATE_OneParam(FOnGizmoTransformSourceChanged, IGizmoTransformSource*);
+	FOnGizmoTransformSourceChanged OnTransformChanged;
+};
+
+
+
+
+/**
+ * UGizmoComponentWorldTransformSource implements IGizmoTransformSource (via UGizmoBaseTransformSource)
+ * based on the internal transform of a USceneComponent.
+ */
+UCLASS()
+class INTERACTIVETOOLSFRAMEWORK_API UGizmoComponentWorldTransformSource : public UGizmoBaseTransformSource
+{
+	GENERATED_BODY()
+public:
+
+	virtual FTransform GetTransform() const override;
+
+	virtual void SetTransform(const FTransform& NewTransform) override;
+
+	UPROPERTY()
+	USceneComponent* Component;
+
+	/** If true, Component->Modify() is called on SetTransform */
+	UPROPERTY()
+	bool bModifyComponentOnTransform = true;
+
+
+public:
+	/**
+	 * Construct a default instance of UGizmoComponentWorldTransformSource with the given Component
+	 */
+	static UGizmoComponentWorldTransformSource* Construct(
+		USceneComponent* Component,
+		UObject* Outer = (UObject*)GetTransientPackage())
+	{
+		UGizmoComponentWorldTransformSource* NewSource = NewObject<UGizmoComponentWorldTransformSource>(Outer);
+		NewSource->Component = Component;
+		return NewSource;
+	}
+};
+
+
+
+
+
+/**
+ * UGizmoTransformProxyTransformSource implements IGizmoTransformSource (via UGizmoBaseTransformSource)
+ * based on the internal transform of a UTransformProxy.
+ */
+UCLASS()
+class INTERACTIVETOOLSFRAMEWORK_API UGizmoTransformProxyTransformSource : public UGizmoBaseTransformSource
+{
+	GENERATED_BODY()
+public:
+
+	virtual FTransform GetTransform() const override;
+
+	virtual void SetTransform(const FTransform& NewTransform) override;
+
+	UPROPERTY()
+	UTransformProxy* Proxy;
+
+public:
+	/**
+	 * Construct a default instance of UGizmoComponentWorldTransformSource with the given Proxy
+	 */
+	static UGizmoTransformProxyTransformSource* Construct(
+		UTransformProxy* Proxy,
+		UObject* Outer = (UObject*)GetTransientPackage())
+	{
+		UGizmoTransformProxyTransformSource* NewSource = NewObject<UGizmoTransformProxyTransformSource>(Outer);
+		NewSource->Proxy = Proxy;
+		return NewSource;
+	}
+};
+
diff --git a/Engine/Source/Runtime/Experimental/InteractiveToolsFramework/Public/InteractiveGizmoManager.h b/Engine/Source/Runtime/Experimental/InteractiveToolsFramework/Public/InteractiveGizmoManager.h
--- a/Engine/Source/Runtime/Experimental/InteractiveToolsFramework/Public/InteractiveGizmoManager.h
+++ b/Engine/Source/Runtime/Experimental/InteractiveToolsFramework/Public/InteractiveGizmoManager.h
@@ -11,7 +11,7 @@
 #include "ToolContextInterfaces.h"
 #include "InteractiveGizmoManager.generated.h"
 
-
+class UTransformGizmo;
 
 USTRUCT()
 struct FActiveGizmo
@@ -21,6 +21,7 @@ struct FActiveGizmo
 	UInteractiveGizmo* Gizmo;
 	FString BuilderIdentifier;
 	FString InstanceIdentifier;
+	void* Owner = nullptr;
 };
 
 
@@ -72,9 +73,10 @@ public:
 	 * Try to activate a new Gizmo instance on the given Side
 	 * @param BuilderIdentifier string used to identify Builder that should be called
 	 * @param InstanceIdentifier client-defined string that can be used to locate this instance
+	 * @param Owner void pointer to whatever "owns" this Gizmo. Allows Gizmo to later be deleted using DestroyAllGizmosByOwner()
 	 * @return new Gizmo instance that has been created and initialized
 	 */	
-	virtual UInteractiveGizmo* CreateGizmo(const FString& BuilderIdentifier, const FString& InstanceIdentifier );
+	virtual UInteractiveGizmo* CreateGizmo(const FString& BuilderIdentifier, const FString& InstanceIdentifier, void* Owner = nullptr);
 
 
 	/**
@@ -90,6 +92,13 @@ public:
 	 */
 	virtual void DestroyAllGizmosOfType(const FString& BuilderIdentifier);
 
+	/**
+	 * Destroy all Gizmos that are owned by the given pointer
+	 * @param Owner pointer that was passed to CreateGizmo
+	 */
+	virtual void DestroyAllGizmosByOwner(void* Owner);
+
+
 	/**
 	 * Find all the existing Gizmo instances that were created by the identified GizmoBuilder
 	 * @param BuilderIdentifier the Builder string registered with RegisterGizmoType
@@ -158,7 +167,36 @@ public:
 	virtual IToolsContextQueriesAPI* GetContextQueriesAPI() { return QueriesAPI; }
 
 
+
+
+public:
+	//
+	// Standard Gizmos
+	// 
+
+	/**
+	 * Register default gizmo types
+	 */
+	virtual void RegisterDefaultGizmos();
+
+	/**
+	 * Activate a new instance of the default 3-axis transformation Gizmo. RegisterDefaultGizmos() must have been called first.
+	 * @param InstanceIdentifier client-defined *unique* string that can be used to locate this instance
+	 * @param Owner void pointer to whatever "owns" this Gizmo. Allows Gizmo to later be deleted using DestroyAllGizmosByOwner()
+	 * @return new Gizmo instance that has been created and initialized
+	 */
+	virtual UTransformGizmo* Create3AxisTransformGizmo(const FString& InstanceIdentifier, void* Owner = nullptr);
+
+
 public:
+	// builder identifiers for default gizmo types. Perhaps should have an API for this...
+	static FString DefaultAxisPositionBuilderIdentifier;
+	static FString DefaultPlanePositionBuilderIdentifier;
+	static FString DefaultAxisAngleBuilderIdentifier;
+	static FString DefaultThreeAxisTransformBuilderIdentifier;
+
+
+protected:
 	/** set of Currently-active Gizmos */
 	UPROPERTY()
 	TArray<FActiveGizmo> ActiveGizmos;
@@ -176,4 +214,6 @@ protected:
 	/** Current set of named GizmoBuilders */
 	UPROPERTY()
 	TMap<FString, UInteractiveGizmoBuilder*> GizmoBuilders;
+
+	bool bDefaultGizmosRegistered = false;
 };
\ No newline at end of file
