add missing header
#rnx

diff --git a/Engine/Source/Runtime/Experimental/InteractiveToolsFramework/Public/BaseGizmos/ParameterToTransformAdapters.h b/Engine/Source/Runtime/Experimental/InteractiveToolsFramework/Public/BaseGizmos/ParameterToTransformAdapters.h
--- a/Engine/Source/Runtime/Experimental/InteractiveToolsFramework/Public/BaseGizmos/ParameterToTransformAdapters.h
+++ b/Engine/Source/Runtime/Experimental/InteractiveToolsFramework/Public/BaseGizmos/ParameterToTransformAdapters.h
@@ -4,6 +4,7 @@
 
 #include "CoreMinimal.h"
 #include "BaseGizmos/GizmoInterfaces.h"
+#include "BaseGizmos/GizmoComponents.h"
 #include "ParameterToTransformAdapters.generated.h"
 
 
