ToolsFramework: Improve Undo/Redo handling with new "Expired" functionality.

InteractiveToolManager now wraps incoming FChanges from the Tools in a FToolChangeWrapperChange that Expires itself once the Tool is no longer active.

InteractiveToolManager now pushes new FBeginToolChange at the start of each Tool. Undoing past start of Tool cancels the Tool, and then this change is Expired.

#rnx

diff --git a/Engine/Source/Runtime/Experimental/InteractiveToolsFramework/Private/InteractiveToolManager.cpp b/Engine/Source/Runtime/Experimental/InteractiveToolsFramework/Private/InteractiveToolManager.cpp
--- a/Engine/Source/Runtime/Experimental/InteractiveToolsFramework/Private/InteractiveToolManager.cpp
+++ b/Engine/Source/Runtime/Experimental/InteractiveToolsFramework/Private/InteractiveToolManager.cpp
@@ -5,6 +5,10 @@
 #include "InteractiveToolsContext.h"
 
 
+
+#define LOCTEXT_NAMESPACE "UInteractiveToolManager"
+
+
 UInteractiveToolManager::UInteractiveToolManager()
 {
 	QueriesAPI = nullptr;
@@ -133,6 +137,9 @@ bool UInteractiveToolManager::ActivateTool(EToolSide Side)
 
 	OnToolStarted.Broadcast(this, ActiveLeftTool);
 
+	// emit a change so that we can undo to cancel the tool
+	EmitObjectChange(this, MakeUnique<FBeginToolChange>(), LOCTEXT("ActivateToolChange", "Activate Tool"));
+
 	return true;
 }
 
@@ -263,10 +270,97 @@ void UInteractiveToolManager::EndUndoTransaction()
 
 void UInteractiveToolManager::EmitObjectChange(UObject* TargetObject, TUniquePtr<FChange> Change, const FText& Description)
 {
-	TransactionsAPI->AppendChange(TargetObject, MoveTemp(Change), Description );
+	// wrap change 
+	check(HasActiveTool(EToolSide::Left));
+	TUniquePtr<FToolChangeWrapperChange> Wrapper = MakeUnique<FToolChangeWrapperChange>();
+	Wrapper->ToolManager = this;
+	Wrapper->ActiveTool = GetActiveTool(EToolSide::Left);
+	Wrapper->ToolChange = MoveTemp(Change);
+
+	TransactionsAPI->AppendChange(TargetObject, MoveTemp(Wrapper), Description );
 }
 
 bool UInteractiveToolManager::RequestSelectionChange(const FSelectedOjectsChangeList& SelectionChange)
 {
 	return TransactionsAPI->RequestSelectionChange(SelectionChange);
 }
+
+
+
+
+
+
+
+void FBeginToolChange::Apply(UObject* Object)
+{
+	// do nothing on apply, we do not want to re-enter the tool
+}
+
+void FBeginToolChange::Revert(UObject* Object)
+{
+	// On revert, if a tool is active, we cancel it.
+	// Note that this should only happen once, because any further tool activations
+	// would be pushing their own FBeginToolChange
+	UInteractiveToolManager* ToolManager = CastChecked<UInteractiveToolManager>(Object);
+	if (ToolManager->HasAnyActiveTool())
+	{
+		ToolManager->DeactivateTool(EToolSide::Left, EToolShutdownType::Cancel);
+	}
+}
+
+bool FBeginToolChange::HasExpired( UObject* Object ) const
+{
+	UInteractiveToolManager* ToolManager = CastChecked<UInteractiveToolManager>(Object);
+	return (ToolManager == nullptr) || (ToolManager->HasAnyActiveTool() == false);
+}
+
+FString FBeginToolChange::ToString() const
+{
+	return FString(TEXT("Begin Tool"));
+}
+
+
+
+
+
+void FToolChangeWrapperChange::Apply(UObject* Object)
+{
+	if (ToolChange.IsValid())
+	{
+		ToolChange->Apply(Object);
+	}
+}
+
+void FToolChangeWrapperChange::Revert(UObject* Object)
+{
+	if (ToolChange.IsValid())
+	{
+		ToolChange->Revert(Object);
+	}
+}
+
+bool FToolChangeWrapperChange::HasExpired(UObject* Object) const
+{
+	if (ToolChange.IsValid() && ToolManager.IsValid() && ActiveTool.IsValid())
+	{
+		if (ToolManager->GetActiveTool(EToolSide::Left) == ActiveTool.Get())
+		{
+			return false;
+		}
+	}
+	return true;
+}
+
+FString FToolChangeWrapperChange::ToString() const
+{
+	if (ToolChange.IsValid())
+	{
+		return ToolChange->ToString();
+	}
+	return FString();
+}
+
+
+
+
+#undef LOCTEXT_NAMESPACE
\ No newline at end of file
diff --git a/Engine/Source/Runtime/Experimental/InteractiveToolsFramework/Public/InteractiveToolManager.h b/Engine/Source/Runtime/Experimental/InteractiveToolsFramework/Public/InteractiveToolManager.h
--- a/Engine/Source/Runtime/Experimental/InteractiveToolsFramework/Public/InteractiveToolManager.h
+++ b/Engine/Source/Runtime/Experimental/InteractiveToolsFramework/Public/InteractiveToolManager.h
@@ -233,4 +233,45 @@ protected:
 	/** Currently-active Right ToolBuilder */
 	UInteractiveToolBuilder* ActiveRightBuilder;
 
+};
+
+
+
+
+/**
+ * FBeginToolChange is used by UInteractiveToolManager to back out of a Tool on Undo.
+ * No action is taken on Redo, ie we do not re-start the Tool on Redo.
+ */
+class INTERACTIVETOOLSFRAMEWORK_API FBeginToolChange : public FCommandChange
+{
+public:
+	virtual void Apply(UObject* Object) override;
+
+	virtual void Revert(UObject* Object) override;
+
+	virtual bool HasExpired(UObject* Object) const override;
+
+	virtual FString ToString() const override;
+};
+
+
+
+/**
+ * FToolChangeWrapperChange wraps an FChange emitted by an InteractiveTool, allowing
+ * us to Expire the change without each FChange implementation needing to handle this explicitly.
+ */
+class INTERACTIVETOOLSFRAMEWORK_API FToolChangeWrapperChange : public FCommandChange
+{
+public:
+	TWeakObjectPtr<UInteractiveToolManager> ToolManager;
+	TWeakObjectPtr<UInteractiveTool> ActiveTool;
+	TUniquePtr<FChange> ToolChange;
+
+	virtual void Apply(UObject* Object) override;
+
+	virtual void Revert(UObject* Object) override;
+
+	virtual bool HasExpired(UObject* Object) const override;
+
+	virtual FString ToString() const override;
 };
\ No newline at end of file
