ToolsFramework: add new USelectionSet and UMeshSelectionSet to represent sub-selections of objects/meshes
#rnx

diff --git a/Engine/Source/Runtime/Experimental/InteractiveToolsFramework/Private/SelectionSet.cpp b/Engine/Source/Runtime/Experimental/InteractiveToolsFramework/Private/SelectionSet.cpp
new file mode 100644
--- /dev/null
+++ b/Engine/Source/Runtime/Experimental/InteractiveToolsFramework/Private/SelectionSet.cpp
@@ -0,0 +1,72 @@
+// Copyright 1998-2019 Epic Games, Inc. All Rights Reserved.
+
+#include "SelectionSet.h"
+
+
+UMeshSelectionSet::UMeshSelectionSet()
+{
+	//SetFlags(RF_Transactional);
+}
+
+
+TArray<int>& UMeshSelectionSet::GetElements(EMeshSelectionElementType ElementType)
+{
+	switch (ElementType)
+	{
+	default:
+	case EMeshSelectionElementType::Vertex:
+		return Vertices;
+	case EMeshSelectionElementType::Edge:
+		return Edges;
+	case EMeshSelectionElementType::Face:
+		return Faces;
+	case EMeshSelectionElementType::Group:
+		return Groups;
+	}
+	check(false);
+}
+
+const TArray<int>& UMeshSelectionSet::GetElements(EMeshSelectionElementType ElementType) const
+{
+	switch (ElementType)
+	{
+	default:
+	case EMeshSelectionElementType::Vertex:
+		return Vertices;
+	case EMeshSelectionElementType::Edge:
+		return Edges;
+	case EMeshSelectionElementType::Face:
+		return Faces;
+	case EMeshSelectionElementType::Group:
+		return Groups;
+	}
+	check(false);
+}
+
+
+
+void UMeshSelectionSet::AddIndices(EMeshSelectionElementType ElementType, const TArray<int>& Indices)
+{
+	TArray<int>& CurElements = GetElements(ElementType);
+
+	int N = Indices.Num();
+	for (int k = 0; k < N; ++k)
+	{
+		CurElements.Add(Indices[k]);
+	}
+	NotifySelectionSetModified();
+}
+
+
+void UMeshSelectionSet::RemoveIndices(EMeshSelectionElementType ElementType, const TArray<int>& Indices)
+{
+	TArray<int>& CurElements = GetElements(ElementType);
+
+	// @todo if we are removing many elements it is maybe cheaper to make a new array...
+	int N = Indices.Num();
+	for (int k = 0; k < N; ++k)
+	{
+		CurElements.RemoveSwap(Indices[k]);
+	}
+	NotifySelectionSetModified();
+}
\ No newline at end of file
diff --git a/Engine/Source/Runtime/Experimental/InteractiveToolsFramework/Public/SelectionSet.h b/Engine/Source/Runtime/Experimental/InteractiveToolsFramework/Public/SelectionSet.h
new file mode 100644
--- /dev/null
+++ b/Engine/Source/Runtime/Experimental/InteractiveToolsFramework/Public/SelectionSet.h
@@ -0,0 +1,101 @@
+// Copyright 1998-2019 Epic Games, Inc. All Rights Reserved.
+
+#pragma once
+
+#include "CoreMinimal.h"
+#include "SelectionSet.generated.h"
+
+
+class USelectionSet;
+
+/** This delegate is used by UInteractiveToolPropertySet */
+DECLARE_MULTICAST_DELEGATE_OneParam(FSelectionSetModifiedSignature, USelectionSet*);
+
+
+/**
+ * USelectionSet is a base class for Selection objects. 
+ * For example the UMeshSelectionSet implementation stores information on selected
+ * triangles, vertices, etc. 
+ */
+UCLASS(Transient)
+class INTERACTIVETOOLSFRAMEWORK_API USelectionSet : public UObject
+{
+	GENERATED_BODY()
+
+protected:
+	FSelectionSetModifiedSignature OnModified;
+
+
+public:
+
+	/** @return the multicast delegate that is called when properties are modified */
+	FSelectionSetModifiedSignature& GetOnModified()
+	{
+		return OnModified;
+	}
+
+	/** posts a message to the OnModified delegate with the modified UProperty */
+	void NotifySelectionSetModified()
+	{
+		OnModified.Broadcast(this);
+	}
+
+};
+
+
+
+
+enum class EMeshSelectionElementType
+{
+	Vertex = 0,
+	Face = 1,
+	Edge = 2,
+	Group = 3
+};
+
+
+
+/**
+ * UMeshSelectionSet is an implementation of USelectionSet that represents selections on indexed meshes.
+ * Vertices, Edges, Faces, and Groups can be selected.
+ */
+UCLASS(Transient)
+class INTERACTIVETOOLSFRAMEWORK_API UMeshSelectionSet : public USelectionSet
+{
+	GENERATED_BODY()
+
+public:
+	UMeshSelectionSet();
+
+	UPROPERTY()
+	TArray<int32> Vertices;
+
+	UPROPERTY()
+	TArray<int32> Edges;
+
+	UPROPERTY()
+	TArray<int32> Faces;
+
+	UPROPERTY()
+	TArray<int32> Groups;
+
+
+public:
+	/** @return current selected elements of the given ElementType */
+	TArray<int>& GetElements(EMeshSelectionElementType ElementType);
+
+	/** @return current selected elements of the given ElementType */
+	const TArray<int>& GetElements(EMeshSelectionElementType ElementType) const;
+
+	/**
+	 * Add list of Indices to the current selection for the given ElementType
+	 * @warning no duplicate detection is currently done
+	 */
+	void AddIndices(EMeshSelectionElementType ElementType, const TArray<int>& Indices);
+	
+	/**
+	 * Remove list of Indices from the current selection for the given ElementType
+	 */
+	void RemoveIndices(EMeshSelectionElementType ElementType, const TArray<int>& Indices);
+
+};
\ No newline at end of file
