ToolsFramework: minor Gizmo improvements. Gizmo unique identifier is now optional, allowing multiple TransformGizmo to be created. Added several options to TransformProxy to support moving pivot and separated rotations.
#rnx

diff --git a/Engine/Source/Runtime/Experimental/InteractiveToolsFramework/Private/BaseGizmos/TransformGizmo.cpp b/Engine/Source/Runtime/Experimental/InteractiveToolsFramework/Private/BaseGizmos/TransformGizmo.cpp
--- a/Engine/Source/Runtime/Experimental/InteractiveToolsFramework/Private/BaseGizmos/TransformGizmo.cpp
+++ b/Engine/Source/Runtime/Experimental/InteractiveToolsFramework/Private/BaseGizmos/TransformGizmo.cpp
@@ -183,42 +183,42 @@ void UTransformGizmo::SetActiveTarget(UTransformProxy* Target)
 	// todo should we hold onto these?
 	if (GizmoActor->TranslateX != nullptr)
 	{
-		AddAxisTranslationGizmo(GizmoActor->TranslateX, GizmoComponent, TEXT("TranslateX"), AxisXSource, TransformSource, StateTarget);
+		AddAxisTranslationGizmo(GizmoActor->TranslateX, GizmoComponent, AxisXSource, TransformSource, StateTarget);
 	}
 	if (GizmoActor->TranslateY != nullptr)
 	{
-		AddAxisTranslationGizmo(GizmoActor->TranslateY, GizmoComponent, TEXT("TranslateY"), AxisYSource, TransformSource, StateTarget);
+		AddAxisTranslationGizmo(GizmoActor->TranslateY, GizmoComponent, AxisYSource, TransformSource, StateTarget);
 	}
 	if (GizmoActor->TranslateZ != nullptr)
 	{
-		AddAxisTranslationGizmo(GizmoActor->TranslateZ, GizmoComponent, TEXT("TranslateZ"), AxisZSource, TransformSource, StateTarget);
+		AddAxisTranslationGizmo(GizmoActor->TranslateZ, GizmoComponent, AxisZSource, TransformSource, StateTarget);
 	}
 
 
 	if (GizmoActor->TranslateYZ != nullptr)
 	{
-		AddPlaneTranslationGizmo(GizmoActor->TranslateYZ, GizmoComponent, TEXT("TranslateYZ"), AxisXSource, TransformSource, StateTarget);
+		AddPlaneTranslationGizmo(GizmoActor->TranslateYZ, GizmoComponent, AxisXSource, TransformSource, StateTarget);
 	}
 	if (GizmoActor->TranslateXZ != nullptr)
 	{
-		AddPlaneTranslationGizmo(GizmoActor->TranslateXZ, GizmoComponent, TEXT("TranslateXZ"), AxisYSource, TransformSource, StateTarget);
+		AddPlaneTranslationGizmo(GizmoActor->TranslateXZ, GizmoComponent, AxisYSource, TransformSource, StateTarget);
 	}
 	if (GizmoActor->TranslateXY != nullptr)
 	{
-		AddPlaneTranslationGizmo(GizmoActor->TranslateXY, GizmoComponent, TEXT("TranslateXY"), AxisZSource, TransformSource, StateTarget);
+		AddPlaneTranslationGizmo(GizmoActor->TranslateXY, GizmoComponent, AxisZSource, TransformSource, StateTarget);
 	}
 
 	if (GizmoActor->RotateX != nullptr)
 	{
-		AddAxisRotationGizmo(GizmoActor->RotateX, GizmoComponent, TEXT("RotateX"), AxisXSource, TransformSource, StateTarget);
+		AddAxisRotationGizmo(GizmoActor->RotateX, GizmoComponent, AxisXSource, TransformSource, StateTarget);
 	}
 	if (GizmoActor->RotateY != nullptr)
 	{
-		AddAxisRotationGizmo(GizmoActor->RotateY, GizmoComponent, TEXT("RotateY"), AxisYSource, TransformSource, StateTarget);
+		AddAxisRotationGizmo(GizmoActor->RotateY, GizmoComponent, AxisYSource, TransformSource, StateTarget);
 	}
 	if (GizmoActor->RotateZ != nullptr)
 	{
-		AddAxisRotationGizmo(GizmoActor->RotateZ, GizmoComponent, TEXT("RotateZ"), AxisZSource, TransformSource, StateTarget);
+		AddAxisRotationGizmo(GizmoActor->RotateZ, GizmoComponent, AxisZSource, TransformSource, StateTarget);
 	}
 }
 
@@ -226,14 +226,13 @@ void UTransformGizmo::SetActiveTarget(UTransformProxy* Target)
 
 UInteractiveGizmo* UTransformGizmo::AddAxisTranslationGizmo(
 	UPrimitiveComponent* AxisComponent, USceneComponent* RootComponent,
-	const FString& Identifier,
 	IGizmoAxisSource* AxisSource,
 	IGizmoTransformSource* TransformSource,
 	IGizmoStateTarget* StateTarget)
 {
 	// create axis-position gizmo, axis-position parameter will drive translation
 	UAxisPositionGizmo* TranslateGizmo = Cast<UAxisPositionGizmo>(GetGizmoManager()->CreateGizmo(
-		UInteractiveGizmoManager::DefaultAxisPositionBuilderIdentifier, Identifier));
+		UInteractiveGizmoManager::DefaultAxisPositionBuilderIdentifier));
 	check(TranslateGizmo);
 
 	// axis source provides the translation axis
@@ -260,14 +259,13 @@ UInteractiveGizmo* UTransformGizmo::AddAxisTranslationGizmo(
 
 UInteractiveGizmo* UTransformGizmo::AddPlaneTranslationGizmo(
 	UPrimitiveComponent* AxisComponent, USceneComponent* RootComponent,
-	const FString& Identifier,
 	IGizmoAxisSource* AxisSource,
 	IGizmoTransformSource* TransformSource,
 	IGizmoStateTarget* StateTarget)
 {
 	// create axis-position gizmo, axis-position parameter will drive translation
 	UPlanePositionGizmo* TranslateGizmo = Cast<UPlanePositionGizmo>(GetGizmoManager()->CreateGizmo(
-		UInteractiveGizmoManager::DefaultPlanePositionBuilderIdentifier, Identifier));
+		UInteractiveGizmoManager::DefaultPlanePositionBuilderIdentifier));
 	check(TranslateGizmo);
 
 	// axis source provides the translation axis
@@ -296,14 +294,13 @@ UInteractiveGizmo* UTransformGizmo::AddPlaneTranslationGizmo(
 
 UInteractiveGizmo* UTransformGizmo::AddAxisRotationGizmo(
 	UPrimitiveComponent* AxisComponent, USceneComponent* RootComponent,
-	const FString& Identifier,
 	IGizmoAxisSource* AxisSource,
 	IGizmoTransformSource* TransformSource,
 	IGizmoStateTarget* StateTarget)
 {
 	// create axis-angle gizmo, angle will drive axis-rotation
 	UAxisAngleGizmo* RotateGizmo = Cast<UAxisAngleGizmo>(GetGizmoManager()->CreateGizmo(
-		UInteractiveGizmoManager::DefaultAxisAngleBuilderIdentifier, Identifier));
+		UInteractiveGizmoManager::DefaultAxisAngleBuilderIdentifier));
 	check(RotateGizmo);
 
 	// axis source provides the rotation axis
diff --git a/Engine/Source/Runtime/Experimental/InteractiveToolsFramework/Private/BaseGizmos/TransformProxy.cpp b/Engine/Source/Runtime/Experimental/InteractiveToolsFramework/Private/BaseGizmos/TransformProxy.cpp
--- a/Engine/Source/Runtime/Experimental/InteractiveToolsFramework/Private/BaseGizmos/TransformProxy.cpp
+++ b/Engine/Source/Runtime/Experimental/InteractiveToolsFramework/Private/BaseGizmos/TransformProxy.cpp
@@ -29,9 +29,15 @@ void UTransformProxy::SetTransform(const FTransform& TransformIn)
 {
 	SharedTransform = TransformIn;
 
-	UpdateObjects();
-
-	OnTransformChanged.Broadcast(this, SharedTransform);
+	if (bSetPivotMode)
+	{
+		UpdateObjectTransforms();
+	}
+	else
+	{
+		UpdateObjects();
+		OnTransformChanged.Broadcast(this, SharedTransform);
+	}
 }
 
 
@@ -42,7 +48,18 @@ void UTransformProxy::UpdateObjects()
 	for (FRelativeObject& Obj : Objects)
 	{
 		FTransform CombinedTransform;
-		FTransform::Multiply(&CombinedTransform, &SharedTransform, &Obj.RelativeTransform);
+		if (bRotatePerObject)
+		{
+			FTransform Temp = SharedTransform.GetRelativeTransform(InitialSharedTransform);
+
+			CombinedTransform = Obj.StartTransform;
+			CombinedTransform.AddToTranslation(Temp.GetTranslation());
+			CombinedTransform.ConcatenateRotation(Temp.GetRotation());
+		}
+		else
+		{
+			FTransform::Multiply(&CombinedTransform, &Obj.RelativeTransform, &SharedTransform);
+		}
 		
 		if (Obj.Component.IsValid())
 		{
@@ -89,4 +106,20 @@ void UTransformProxy::UpdateSharedTransform()
 		}
 	}
 
+	InitialSharedTransform = SharedTransform;
+}
+
+
+
+void UTransformProxy::UpdateObjectTransforms()
+{
+	for (FRelativeObject& Obj : Objects)
+	{
+		if (Obj.Component != nullptr)
+		{
+			Obj.StartTransform = Obj.Component->GetComponentToWorld();
+		}
+		Obj.RelativeTransform = Obj.StartTransform;
+		Obj.RelativeTransform.SetToRelativeTransform(SharedTransform);
+	}
 }
\ No newline at end of file
diff --git a/Engine/Source/Runtime/Experimental/InteractiveToolsFramework/Private/InteractiveGizmoManager.cpp b/Engine/Source/Runtime/Experimental/InteractiveToolsFramework/Private/InteractiveGizmoManager.cpp
--- a/Engine/Source/Runtime/Experimental/InteractiveToolsFramework/Private/InteractiveGizmoManager.cpp
+++ b/Engine/Source/Runtime/Experimental/InteractiveToolsFramework/Private/InteractiveGizmoManager.cpp
@@ -81,12 +81,15 @@ UInteractiveGizmo* UInteractiveGizmoManager::CreateGizmo(const FString& BuilderI
 	UInteractiveGizmoBuilder* FoundBuilder = GizmoBuilders[BuilderIdentifier];
 
 	// check if we have used this instance identifier
-	for (FActiveGizmo& ActiveGizmo : ActiveGizmos)
+	if (InstanceIdentifier.IsEmpty() == false)
 	{
-		if (ActiveGizmo.InstanceIdentifier == InstanceIdentifier)
+		for (FActiveGizmo& ActiveGizmo : ActiveGizmos)
 		{
-			PostMessage(FString::Printf(TEXT("UInteractiveGizmoManager::CreateGizmo: instance identifier %s already in use!"), *InstanceIdentifier), EToolMessageLevel::Internal);
-			return nullptr;
+			if (ActiveGizmo.InstanceIdentifier == InstanceIdentifier)
+			{
+				PostMessage(FString::Printf(TEXT("UInteractiveGizmoManager::CreateGizmo: instance identifier %s already in use!"), *InstanceIdentifier), EToolMessageLevel::Internal);
+				return nullptr;
+			}
 		}
 	}
 
diff --git a/Engine/Source/Runtime/Experimental/InteractiveToolsFramework/Public/BaseGizmos/TransformGizmo.h b/Engine/Source/Runtime/Experimental/InteractiveToolsFramework/Public/BaseGizmos/TransformGizmo.h
--- a/Engine/Source/Runtime/Experimental/InteractiveToolsFramework/Public/BaseGizmos/TransformGizmo.h
+++ b/Engine/Source/Runtime/Experimental/InteractiveToolsFramework/Public/BaseGizmos/TransformGizmo.h
@@ -220,7 +220,6 @@ protected:
 	/** @return a new instance of the standard axis-translation Gizmo */
 	virtual UInteractiveGizmo* AddAxisTranslationGizmo(
 		UPrimitiveComponent* AxisComponent, USceneComponent* RootComponent,
-		const FString& Identifier,
 		IGizmoAxisSource* AxisSource,
 		IGizmoTransformSource* TransformSource, 
 		IGizmoStateTarget* StateTarget);
@@ -228,7 +227,6 @@ protected:
 	/** @return a new instance of the standard plane-translation Gizmo */
 	virtual UInteractiveGizmo* AddPlaneTranslationGizmo(
 		UPrimitiveComponent* AxisComponent, USceneComponent* RootComponent,
-		const FString& Identifier,
 		IGizmoAxisSource* AxisSource,
 		IGizmoTransformSource* TransformSource,
 		IGizmoStateTarget* StateTarget);
@@ -236,7 +234,6 @@ protected:
 	/** @return a new instance of the standard axis-rotation Gizmo */
 	virtual UInteractiveGizmo* AddAxisRotationGizmo(
 		UPrimitiveComponent* AxisComponent, USceneComponent* RootComponent,
-		const FString& Identifier,
 		IGizmoAxisSource* AxisSource,
 		IGizmoTransformSource* TransformSource,
 		IGizmoStateTarget* StateTarget);
diff --git a/Engine/Source/Runtime/Experimental/InteractiveToolsFramework/Public/BaseGizmos/TransformProxy.h b/Engine/Source/Runtime/Experimental/InteractiveToolsFramework/Public/BaseGizmos/TransformProxy.h
--- a/Engine/Source/Runtime/Experimental/InteractiveToolsFramework/Public/BaseGizmos/TransformProxy.h
+++ b/Engine/Source/Runtime/Experimental/InteractiveToolsFramework/Public/BaseGizmos/TransformProxy.h
@@ -52,6 +52,18 @@ public:
 	DECLARE_MULTICAST_DELEGATE_TwoParams(FOnTransformChanged, UTransformProxy*, FTransform);
 	FOnTransformChanged OnTransformChanged;
 
+	/**
+	 * If true, relative rotation of shared transform is applied to objects before relative translation (ie they rotate in place)
+	 */
+	UPROPERTY()
+	bool bRotatePerObject = false;
+
+	/**
+	 * If true, then on SetTransform() the components are not moved, and their local transforms are recalculated
+	 */
+	UPROPERTY()
+	bool bSetPivotMode = false;
+
 	
 protected:
 
@@ -73,9 +85,16 @@ protected:
 	UPROPERTY()
 	FTransform SharedTransform;
 
+	/** The main transform */
+	UPROPERTY()
+	FTransform InitialSharedTransform;
+
 	/** Recalculate main SharedTransform when object set changes*/
 	virtual void UpdateSharedTransform();
 
+	/** Recalculate per-object relative transforms */
+	virtual void UpdateObjectTransforms();
+
 	/** Propagate a transform update to the sub-objects */
 	virtual void UpdateObjects();
 };
\ No newline at end of file
diff --git a/Engine/Source/Runtime/Experimental/InteractiveToolsFramework/Public/InteractiveGizmoManager.h b/Engine/Source/Runtime/Experimental/InteractiveToolsFramework/Public/InteractiveGizmoManager.h
--- a/Engine/Source/Runtime/Experimental/InteractiveToolsFramework/Public/InteractiveGizmoManager.h
+++ b/Engine/Source/Runtime/Experimental/InteractiveToolsFramework/Public/InteractiveGizmoManager.h
@@ -72,11 +72,11 @@ public:
 	/**
 	 * Try to activate a new Gizmo instance on the given Side
 	 * @param BuilderIdentifier string used to identify Builder that should be called
-	 * @param InstanceIdentifier client-defined string that can be used to locate this instance
+	 * @param InstanceIdentifier optional client-defined string that can be used to locate this instance (must be unique across all Gizmos)
 	 * @param Owner void pointer to whatever "owns" this Gizmo. Allows Gizmo to later be deleted using DestroyAllGizmosByOwner()
 	 * @return new Gizmo instance that has been created and initialized
 	 */	
-	virtual UInteractiveGizmo* CreateGizmo(const FString& BuilderIdentifier, const FString& InstanceIdentifier, void* Owner = nullptr);
+	virtual UInteractiveGizmo* CreateGizmo(const FString& BuilderIdentifier, const FString& InstanceIdentifier = FString(), void* Owner = nullptr);
 
 
 	/**
@@ -185,7 +185,7 @@ public:
 	 * @param Owner void pointer to whatever "owns" this Gizmo. Allows Gizmo to later be deleted using DestroyAllGizmosByOwner()
 	 * @return new Gizmo instance that has been created and initialized
 	 */
-	virtual UTransformGizmo* Create3AxisTransformGizmo(const FString& InstanceIdentifier, void* Owner = nullptr);
+	virtual UTransformGizmo* Create3AxisTransformGizmo(const FString& InstanceIdentifier = FString(), void* Owner = nullptr);
 
 
 public:
