resolve minor CIS errors/warnings and remove initials from comment
#rnx

diff --git a/Engine/Source/Editor/UnrealEd/Private/Tools/EditorToolAssetAPI.cpp b/Engine/Source/Editor/UnrealEd/Private/Tools/EditorToolAssetAPI.cpp
--- a/Engine/Source/Editor/UnrealEd/Private/Tools/EditorToolAssetAPI.cpp
+++ b/Engine/Source/Editor/UnrealEd/Private/Tools/EditorToolAssetAPI.cpp
@@ -17,7 +17,7 @@
 
 FString FEditorToolAssetAPI::GetActiveAssetFolderPath()
 {
-	check(false);  	// [RMS] this only works if we have an asset selected!
+	check(false);  	// this only works if we have an asset selected!
 
 	IContentBrowserSingleton& ContentBrowser = FModuleManager::LoadModuleChecked<FContentBrowserModule>("ContentBrowser").Get();
 
diff --git a/Engine/Source/Runtime/Experimental/InteractiveToolsFramework/Private/BaseGizmos/GizmoRenderingUtil.cpp b/Engine/Source/Runtime/Experimental/InteractiveToolsFramework/Private/BaseGizmos/GizmoRenderingUtil.cpp
--- a/Engine/Source/Runtime/Experimental/InteractiveToolsFramework/Private/BaseGizmos/GizmoRenderingUtil.cpp
+++ b/Engine/Source/Runtime/Experimental/InteractiveToolsFramework/Private/BaseGizmos/GizmoRenderingUtil.cpp
@@ -17,13 +17,16 @@ const FSceneView* GizmoRenderingUtil::FindActiveSceneView(
 		if (VisibilityMap & (1 << ViewIndex))
 		{
 			const FSceneView* View = Views[ViewIndex];
-			if (FirstValidView == nullptr)
+			if (View != nullptr)
 			{
-				FirstValidView = View;
-			}
-			if (View->IsPerspectiveProjection())
-			{
-				GizmoControlView = View;
+				if (FirstValidView == nullptr)
+				{
+					FirstValidView = View;
+				}
+				if (View->IsPerspectiveProjection())
+				{
+					GizmoControlView = View;
+				}
 			}
 		}
 	}
diff --git a/Engine/Source/Runtime/Experimental/InteractiveToolsFramework/Public/BaseGizmos/ParameterToTransformAdapters.h b/Engine/Source/Runtime/Experimental/InteractiveToolsFramework/Public/BaseGizmos/ParameterToTransformAdapters.h
--- a/Engine/Source/Runtime/Experimental/InteractiveToolsFramework/Public/BaseGizmos/ParameterToTransformAdapters.h
+++ b/Engine/Source/Runtime/Experimental/InteractiveToolsFramework/Public/BaseGizmos/ParameterToTransformAdapters.h
@@ -4,7 +4,8 @@
 
 #include "CoreMinimal.h"
 #include "BaseGizmos/GizmoInterfaces.h"
-#include "BaseGizmos/GizmoComponents.h"
+#include "BaseGizmos/ParameterSourcesFloat.h"
+#include "BaseGizmos/ParameterSourcesVec2.h"
 #include "ParameterToTransformAdapters.generated.h"
 
 
