ToolsFramework: gizmo actors should not be transactional
#rnx

diff --git a/Engine/Source/Runtime/Experimental/InteractiveToolsFramework/Private/BaseGizmos/GizmoActor.cpp b/Engine/Source/Runtime/Experimental/InteractiveToolsFramework/Private/BaseGizmos/GizmoActor.cpp
--- a/Engine/Source/Runtime/Experimental/InteractiveToolsFramework/Private/BaseGizmos/GizmoActor.cpp
+++ b/Engine/Source/Runtime/Experimental/InteractiveToolsFramework/Private/BaseGizmos/GizmoActor.cpp
@@ -13,6 +13,13 @@
 #define LOCTEXT_NAMESPACE "UTransformGizmo"
 
 
+AGizmoActor::AGizmoActor()
+{
+	// generally gizmo actor creation/destruction should not be transacted
+	ClearFlags(RF_Transactional);
+}
+
+
 UGizmoArrowComponent* AGizmoActor::AddDefaultArrowComponent(
 	UWorld* World, AActor* Actor,
 	const FLinearColor& Color, const FVector& LocalDirection)
diff --git a/Engine/Source/Runtime/Experimental/InteractiveToolsFramework/Public/BaseGizmos/GizmoActor.h b/Engine/Source/Runtime/Experimental/InteractiveToolsFramework/Public/BaseGizmos/GizmoActor.h
--- a/Engine/Source/Runtime/Experimental/InteractiveToolsFramework/Public/BaseGizmos/GizmoActor.h
+++ b/Engine/Source/Runtime/Experimental/InteractiveToolsFramework/Public/BaseGizmos/GizmoActor.h
@@ -21,6 +21,7 @@ class INTERACTIVETOOLSFRAMEWORK_API AGizmoActor : public AInternalToolFrameworkA
 {
 	GENERATED_BODY()
 public:
+	AGizmoActor();
 
 	/** Add standard arrow component to Actor, generally used for axis-translation */
 	static UGizmoArrowComponent* AddDefaultArrowComponent(
