ToolsFramework: replace FString/TCHAR in various PostMessage functions with FText. Fix up call sites. Add area to modeling panel to display notification messages.
#rnx

diff --git a/Engine/Plugins/Experimental/MeshModelingToolset/Source/MeshModelingTools/Private/AddPrimitiveTool.cpp b/Engine/Plugins/Experimental/MeshModelingToolset/Source/MeshModelingTools/Private/AddPrimitiveTool.cpp
--- a/Engine/Plugins/Experimental/MeshModelingToolset/Source/MeshModelingTools/Private/AddPrimitiveTool.cpp
+++ b/Engine/Plugins/Experimental/MeshModelingToolset/Source/MeshModelingTools/Private/AddPrimitiveTool.cpp
@@ -112,6 +112,10 @@ void UAddPrimitiveTool::Setup()
 	PreviewMesh->SetMaterial(ShapeSettings->Material);
 
 	UpdatePreviewMesh();
+
+	GetToolManager()->PostMessage(
+		LOCTEXT("OnStartAddPrimitiveTool", "Position the Primitive by moving the mouse over the scene. Drop a new instance by Left-clicking."),
+		EToolMessageLevel::UserNotification);
 }
 
 
diff --git a/Engine/Plugins/Experimental/MeshModelingToolset/Source/MeshModelingTools/Private/ConvertToPolygonsTool.cpp b/Engine/Plugins/Experimental/MeshModelingToolset/Source/MeshModelingTools/Private/ConvertToPolygonsTool.cpp
--- a/Engine/Plugins/Experimental/MeshModelingToolset/Source/MeshModelingTools/Private/ConvertToPolygonsTool.cpp
+++ b/Engine/Plugins/Experimental/MeshModelingToolset/Source/MeshModelingTools/Private/ConvertToPolygonsTool.cpp
@@ -150,7 +150,9 @@ void UConvertToPolygonsTool::UpdatePolygons()
 	Polygons.FindPolygons(DotTolerance);
 	Polygons.FindPolygonEdges();
 
-	GetToolManager()->PostMessage(FString::Printf(TEXT("ConvertToPolygons - found %d polys in %d triangles"), Polygons.FoundPolygons.Num(), SearchMesh.TriangleCount()), EToolMessageLevel::Internal);
+	GetToolManager()->PostMessage(
+		FText::Format(LOCTEXT("UpdatePolygonsMessage", "ConvertToPolygons - found {0} polys in {1} triangles"), 
+			FText::AsNumber(Polygons.FoundPolygons.Num()), FText::AsNumber(SearchMesh.TriangleCount()) ), EToolMessageLevel::Internal);
 
 	if (Settings->bCalculateNormals)
 	{
diff --git a/Engine/Plugins/Experimental/MeshModelingToolset/Source/MeshModelingTools/Private/DrawPolygonTool.cpp b/Engine/Plugins/Experimental/MeshModelingToolset/Source/MeshModelingTools/Private/DrawPolygonTool.cpp
--- a/Engine/Plugins/Experimental/MeshModelingToolset/Source/MeshModelingTools/Private/DrawPolygonTool.cpp
+++ b/Engine/Plugins/Experimental/MeshModelingToolset/Source/MeshModelingTools/Private/DrawPolygonTool.cpp
@@ -146,6 +146,10 @@ void UDrawPolygonTool::Setup()
 	// add self as tool props
 	AddToolPropertySource(PolygonProperties);
 	AddToolPropertySource(SnapProperties);
+
+	GetToolManager()->PostMessage(
+		LOCTEXT("OnStartDrawPolygonTool", "Left-click to place points on the Drawing Plane. Ctrl-click on the scene to reposition the Plane, and click-drag on the gizmo to move it."),
+		EToolMessageLevel::UserNotification);
 }
 
 
diff --git a/Engine/Plugins/Experimental/ModelingToolsEditorMode/Source/ModelingToolsEditorMode/Private/ModelingToolsEditorMode.cpp b/Engine/Plugins/Experimental/ModelingToolsEditorMode/Source/ModelingToolsEditorMode/Private/ModelingToolsEditorMode.cpp
--- a/Engine/Plugins/Experimental/ModelingToolsEditorMode/Source/ModelingToolsEditorMode/Private/ModelingToolsEditorMode.cpp
+++ b/Engine/Plugins/Experimental/ModelingToolsEditorMode/Source/ModelingToolsEditorMode/Private/ModelingToolsEditorMode.cpp
@@ -312,19 +312,22 @@ bool FModelingToolsEditorMode::MouseLeave(FEditorViewportClient* ViewportClient,
 
 
 
-
-static FString MeshInspectorToolName("MeshInspectorTool");
-
-
 void FModelingToolsEditorMode::Enter()
 {
 	FEdMode::Enter();
 
 	// initialize FEdMode ToolsContext adapter
 
-	ToolsContext = NewObject<UEdModeInteractiveToolsContext>(GetTransientPackage(), TEXT("ToolsContext"), RF_Transient);
+	//ToolsContext = NewObject<UEdModeInteractiveToolsContext>(GetTransientPackage(), TEXT("ToolsContext"), RF_Transient);
+	ToolsContext = NewObject<UEdModeInteractiveToolsContext>();
 	ToolsContext->InitializeContextFromEdMode(this);
 
+	ToolsContext->OnToolNotificationMessage.AddLambda([this](const FText& Message)
+	{
+		this->OnToolNotificationMessage.Broadcast(Message);
+	});
+
+
 	if (!Toolkit.IsValid() && UsesToolkits())
 	{
 		Toolkit = MakeShareable(new FModelingToolsEditorModeToolkit);
@@ -509,9 +512,10 @@ void FModelingToolsEditorMode::Enter()
 
 
 
-
 void FModelingToolsEditorMode::Exit()
 {
+	OnToolNotificationMessage.Clear();
+
 	ToolsContext->ShutdownContext();
 	ToolsContext = nullptr;
 
diff --git a/Engine/Plugins/Experimental/ModelingToolsEditorMode/Source/ModelingToolsEditorMode/Private/ModelingToolsEditorModeToolkit.cpp b/Engine/Plugins/Experimental/ModelingToolsEditorMode/Source/ModelingToolsEditorMode/Private/ModelingToolsEditorModeToolkit.cpp
--- a/Engine/Plugins/Experimental/ModelingToolsEditorMode/Source/ModelingToolsEditorMode/Private/ModelingToolsEditorModeToolkit.cpp
+++ b/Engine/Plugins/Experimental/ModelingToolsEditorMode/Source/ModelingToolsEditorMode/Private/ModelingToolsEditorModeToolkit.cpp
@@ -20,44 +20,6 @@ FModelingToolsEditorModeToolkit::FModelingToolsEditorModeToolkit()
 }
 
 
-TSharedRef<SButton> FModelingToolsEditorModeToolkit::MakeToolButton(const FText& ButtonLabel, const FString& ToolIdentifier)
-{
-	return SNew(SButton).Text(ButtonLabel)
-		.OnClicked_Lambda([this, ToolIdentifier]() { GetToolsContext()->StartTool(ToolIdentifier); return FReply::Handled(); })
-		.IsEnabled_Lambda([this, ToolIdentifier]() { return GetToolsContext()->CanStartTool(ToolIdentifier); });
-}
-
-SVerticalBox::FSlot& FModelingToolsEditorModeToolkit::MakeToolButtonSlotV(const FText& ButtonLabel, const FString& ToolIdentifier)
-{
-	return SVerticalBox::Slot()
-		.HAlign(HAlign_Center)
-		.AutoHeight()
-		[MakeToolButton(ButtonLabel, ToolIdentifier)];
-}
-
-SHorizontalBox::FSlot& FModelingToolsEditorModeToolkit::MakeToolButtonSlotH(const FText& ButtonLabel, const FString& ToolIdentifier)
-{
-	return SHorizontalBox::Slot()
-		.HAlign(HAlign_Center)
-		.AutoWidth()
-		[MakeToolButton(ButtonLabel, ToolIdentifier)];
-}
-
-
-SVerticalBox::FSlot& FModelingToolsEditorModeToolkit::MakeSetToolLabelV(const FText& LabelText)
-{
-	return SVerticalBox::Slot()
-		.AutoHeight()
-		.HAlign(HAlign_Center)
-		.Padding(5)
-		[
-			SNew(STextBlock)
-			.AutoWrapText(true)
-			.Text(LabelText)
-		];
-}
-
-
 
 void FModelingToolsEditorModeToolkit::Init(const TSharedPtr<IToolkitHost>& InitToolkitHost)
 {
@@ -80,10 +42,21 @@ void FModelingToolsEditorModeToolkit::Init(const TSharedPtr<IToolkitHost>& InitT
 	DetailsView = PropertyEditorModule.CreateDetailView(DetailsViewArgs);
 
 
-	ToolHeaderLabel = SNew(STextBlock).AutoWrapText(true);
+	ToolHeaderLabel = SNew(STextBlock)
+		.AutoWrapText(true)
+		.Font(FCoreStyle::GetDefaultFontStyle("Bold", 12));
 	ToolHeaderLabel->SetText(LOCTEXT("SelectToolLabel", "Select a Tool from the Toolbar"));
 	ToolHeaderLabel->SetJustification(ETextJustify::Center);
 
+	//const FTextBlockStyle DefaultText = FTextBlockStyle()
+	//	.SetFont(DEFAULT_FONT("Bold", 10));
+
+	ToolMessageArea = SNew(STextBlock)
+		.AutoWrapText(true)
+		.Font(FCoreStyle::GetDefaultFontStyle("Italic", 9))
+		.ColorAndOpacity(FSlateColor(FLinearColor::White * 0.7f));
+	ToolMessageArea->SetText(LOCTEXT("ToolMessageLabel", ""));
+
 	SAssignNew(ToolkitWidget, SBorder)
 		.HAlign(HAlign_Fill)
 		.Padding(4)
@@ -91,7 +64,14 @@ void FModelingToolsEditorModeToolkit::Init(const TSharedPtr<IToolkitHost>& InitT
 			SNew(SVerticalBox)
 
 			+ SVerticalBox::Slot().AutoHeight().HAlign(HAlign_Fill).Padding(5)
-				[ToolHeaderLabel->AsShared()]
+				[
+					ToolHeaderLabel->AsShared()
+				]
+
+			+ SVerticalBox::Slot().AutoHeight().HAlign(HAlign_Fill).AutoHeight().Padding(10, 10, 10, 20)
+				[
+					ToolMessageArea->AsShared()
+				]
 
 			+ SVerticalBox::Slot().HAlign(HAlign_Fill).AutoHeight().MaxHeight(500.0f)
 				[
@@ -117,10 +97,29 @@ void FModelingToolsEditorModeToolkit::Init(const TSharedPtr<IToolkitHost>& InitT
 	{
 		DetailsView->SetObject(nullptr);
 		ToolHeaderLabel->SetText(LOCTEXT("SelectToolLabel", "Select a Tool from the Toolbar"));
+		ClearNotification();
+	});
+
+
+	GetToolsEditorMode()->OnToolNotificationMessage.AddLambda([this](const FText& Message)
+	{
+		PostNotification(Message);
 	});
 
+
+}
+
+
+
+void FModelingToolsEditorModeToolkit::PostNotification(const FText& Message)
+{
+	ToolMessageArea->SetText(Message);
 }
 
+void FModelingToolsEditorModeToolkit::ClearNotification()
+{
+	ToolMessageArea->SetText(FText());
+}
 
 
 
diff --git a/Engine/Plugins/Experimental/ModelingToolsEditorMode/Source/ModelingToolsEditorMode/Public/ModelingToolsEditorMode.h b/Engine/Plugins/Experimental/ModelingToolsEditorMode/Source/ModelingToolsEditorMode/Public/ModelingToolsEditorMode.h
--- a/Engine/Plugins/Experimental/ModelingToolsEditorMode/Source/ModelingToolsEditorMode/Public/ModelingToolsEditorMode.h
+++ b/Engine/Plugins/Experimental/ModelingToolsEditorMode/Source/ModelingToolsEditorMode/Public/ModelingToolsEditorMode.h
@@ -133,6 +133,12 @@ public:
 	//////////////////
 
 
+
+public:
+	DECLARE_MULTICAST_DELEGATE_OneParam(FOnModelingModeToolNotification, const FText&);
+	FOnModelingModeToolNotification OnToolNotificationMessage;
+
+
 public:
 	virtual UEdModeInteractiveToolsContext* GetToolsContext() const
 	{
diff --git a/Engine/Plugins/Experimental/ModelingToolsEditorMode/Source/ModelingToolsEditorMode/Public/ModelingToolsEditorModeToolkit.h b/Engine/Plugins/Experimental/ModelingToolsEditorMode/Source/ModelingToolsEditorMode/Public/ModelingToolsEditorModeToolkit.h
--- a/Engine/Plugins/Experimental/ModelingToolsEditorMode/Source/ModelingToolsEditorMode/Public/ModelingToolsEditorModeToolkit.h
+++ b/Engine/Plugins/Experimental/ModelingToolsEditorMode/Source/ModelingToolsEditorMode/Public/ModelingToolsEditorModeToolkit.h
@@ -31,18 +31,17 @@ public:
 	virtual FModelingToolsEditorMode* GetToolsEditorMode() const;
 	virtual UEdModeInteractiveToolsContext* GetToolsContext() const;
 
+	// set/clear notification message area
+	virtual void PostNotification(const FText& Message);
+	virtual void ClearNotification();
+
 private:
 
 	TSharedPtr<SWidget> ToolkitWidget;
 	TSharedPtr<IDetailsView> DetailsView;
 
-	TSharedRef<SButton> MakeToolButton(const FText& ButtonLabel, const FString& ToolIdentifier);
-	SVerticalBox::FSlot& MakeToolButtonSlotV(const FText& ButtonLabel, const FString& ToolIdentifier);
-	SHorizontalBox::FSlot& MakeToolButtonSlotH(const FText& ButtonLabel, const FString& ToolIdentifier);
-
-	SVerticalBox::FSlot& MakeSetToolLabelV(const FText& LabelText);
-
 	TSharedPtr<STextBlock> ToolHeaderLabel;
+	TSharedPtr<STextBlock> ToolMessageArea;
 	TSharedPtr<SButton> AcceptButton;
 	TSharedPtr<SButton> CancelButton;
 	TSharedPtr<SButton> CompletedButton;
diff --git a/Engine/Source/Editor/Experimental/EditorInteractiveToolsFramework/Private/EdModeInteractiveToolsContext.cpp b/Engine/Source/Editor/Experimental/EditorInteractiveToolsFramework/Private/EdModeInteractiveToolsContext.cpp
--- a/Engine/Source/Editor/Experimental/EditorInteractiveToolsFramework/Private/EdModeInteractiveToolsContext.cpp
+++ b/Engine/Source/Editor/Experimental/EditorInteractiveToolsFramework/Private/EdModeInteractiveToolsContext.cpp
@@ -240,9 +240,16 @@ public:
 	}
 
 
-	virtual void PostMessage(const TCHAR* Message, EToolMessageLevel Level) override
+	virtual void PostMessage(const FText& Message, EToolMessageLevel Level) override
 	{
-		UE_LOG(LogTemp, Warning, TEXT("[ToolsContext] %s"), Message);
+		if (Level == EToolMessageLevel::UserNotification)
+		{
+			ToolsContext->PostToolNotificationMessage(Message);
+		}
+		else
+		{
+			UE_LOG(LogTemp, Warning, TEXT("%s"), *Message.ToString());
+		}
 	}
 
 
@@ -372,6 +379,8 @@ void UEdModeInteractiveToolsContext::ShutdownContext()
 {
 	Shutdown();
 
+	OnToolNotificationMessage.Clear();
+
 	if (QueriesAPI != nullptr)
 	{
 		delete QueriesAPI;
@@ -394,6 +403,11 @@ void UEdModeInteractiveToolsContext::ShutdownContext()
 }
 
 
+void UEdModeInteractiveToolsContext::PostToolNotificationMessage(const FText& Message)
+{
+	OnToolNotificationMessage.Broadcast(Message);
+}
+
 
 
 void UEdModeInteractiveToolsContext::TerminateActiveToolsOnPIEStart()
diff --git a/Engine/Source/Editor/Experimental/EditorInteractiveToolsFramework/Public/EdModeInteractiveToolsContext.h b/Engine/Source/Editor/Experimental/EditorInteractiveToolsFramework/Public/EdModeInteractiveToolsContext.h
--- a/Engine/Source/Editor/Experimental/EditorInteractiveToolsFramework/Public/EdModeInteractiveToolsContext.h
+++ b/Engine/Source/Editor/Experimental/EditorInteractiveToolsFramework/Public/EdModeInteractiveToolsContext.h
@@ -5,6 +5,7 @@
 #include "CoreMinimal.h"
 #include "EdMode.h"
 #include "InteractiveToolsContext.h"
+#include "Delegates/Delegate.h"
 #include "EdModeInteractiveToolsContext.generated.h"
 
 /**
@@ -69,6 +70,14 @@ public:
 	virtual void EndTool(EToolShutdownType ShutdownType);
 
 
+public:
+	// forwards message to OnToolNotificationMessage delegate
+	virtual void PostToolNotificationMessage(const FText& Message);
+
+	DECLARE_MULTICAST_DELEGATE_OneParam(FEdModeToolsContextToolNotification, const FText&);
+	FEdModeToolsContextToolNotification OnToolNotificationMessage;
+
+
 protected:
 	// we hide these 
 	virtual void Initialize(IToolsContextQueriesAPI* QueriesAPI, IToolsContextTransactionsAPI* TransactionsAPI) override;
diff --git a/Engine/Source/Editor/UnrealEd/Private/Tools/UEdMode.cpp b/Engine/Source/Editor/UnrealEd/Private/Tools/UEdMode.cpp
--- a/Engine/Source/Editor/UnrealEd/Private/Tools/UEdMode.cpp
+++ b/Engine/Source/Editor/UnrealEd/Private/Tools/UEdMode.cpp
@@ -237,13 +237,11 @@ public:
 		EdMode = InEdMode;
 	}
 
-
-	virtual void PostMessage(const TCHAR* Message, EToolMessageLevel Level)
+	virtual void PostMessage(const FText& Message, EToolMessageLevel Level) override
 	{
-		UE_LOG(LogTemp, Warning, TEXT("[ToolsContext] %s"), Message);
+		UE_LOG(LogTemp, Warning, TEXT("%s"), *Message.ToString());
 	}
 
-
 	virtual void PostInvalidation() override
 	{
 		EdMode->PostInvalidation();
diff --git a/Engine/Source/Runtime/Experimental/InteractiveToolsFramework/Private/BaseTools/ClickDragTool.cpp b/Engine/Source/Runtime/Experimental/InteractiveToolsFramework/Private/BaseTools/ClickDragTool.cpp
--- a/Engine/Source/Runtime/Experimental/InteractiveToolsFramework/Private/BaseTools/ClickDragTool.cpp
+++ b/Engine/Source/Runtime/Experimental/InteractiveToolsFramework/Private/BaseTools/ClickDragTool.cpp
@@ -5,6 +5,9 @@
 #include "BaseBehaviors/ClickDragBehavior.h"
 
 
+#define LOCTEXT_NAMESPACE "UClickDragTool"
+
+
 /*
  * ToolBuilder
  */
@@ -48,8 +51,9 @@ void UClickDragTool::OnClickPress(const FInputDeviceRay& ClickPos)
 {
 	// print debug message
 	GetToolManager()->PostMessage( 
-		FString::Printf( TEXT("UClickDragTool::OnClickPress at (%f,%f)"), ClickPos.ScreenPosition.X, ClickPos.ScreenPosition.Y), 
-		EToolMessageLevel::Internal );
+		FText::Format(LOCTEXT("OnClickPressMessage", "UClickDragTool::OnClickPress: clicked at ({0},{1})"),
+			FText::AsNumber(ClickPos.ScreenPosition.X), FText::AsNumber(ClickPos.ScreenPosition.Y)),
+		EToolMessageLevel::Internal);
 }
 
 void UClickDragTool::OnClickDrag(const FInputDeviceRay& ClickPos)
@@ -61,7 +65,8 @@ void UClickDragTool::OnClickRelease(const FInputDeviceRay& ClickPos)
 {
 	// print debug message
 	GetToolManager()->PostMessage(
-		FString::Printf(TEXT("UClickDragTool::OnClickRelease at (%f,%f)"), ClickPos.ScreenPosition.X, ClickPos.ScreenPosition.Y),
+		FText::Format(LOCTEXT("OnClickReleaseMessage", "UClickDragTool::OnClickRelease: released at ({0},{1})"),
+			FText::AsNumber(ClickPos.ScreenPosition.X), FText::AsNumber(ClickPos.ScreenPosition.Y)),
 		EToolMessageLevel::Internal);
 }
 
@@ -69,4 +74,7 @@ void UClickDragTool::OnClickRelease(const FInputDeviceRay& ClickPos)
 void UClickDragTool::OnTerminateDragSequence()
 {
 
-}
\ No newline at end of file
+}
+
+
+#undef LOCTEXT_NAMESPACE
\ No newline at end of file
diff --git a/Engine/Source/Runtime/Experimental/InteractiveToolsFramework/Private/BaseTools/MeshSurfacePointTool.cpp b/Engine/Source/Runtime/Experimental/InteractiveToolsFramework/Private/BaseTools/MeshSurfacePointTool.cpp
--- a/Engine/Source/Runtime/Experimental/InteractiveToolsFramework/Private/BaseTools/MeshSurfacePointTool.cpp
+++ b/Engine/Source/Runtime/Experimental/InteractiveToolsFramework/Private/BaseTools/MeshSurfacePointTool.cpp
@@ -6,6 +6,8 @@
 #include "ToolBuilderUtil.h"
 #include "Components/PrimitiveComponent.h"
 
+#define LOCTEXT_NAMESPACE "UMeshSurfacePointTool"
+
 
 /*
  * ToolBuilder
@@ -75,7 +77,9 @@ void UMeshSurfacePointTool::OnUpdateDrag(const FRay& Ray)
 	if ( HitTest(Ray, OutHit) ) 
 	{
 		GetToolManager()->PostMessage( 
-			FString::Printf(TEXT("[UMeshSurfacePointTool::OnUpdateDrag] Hit triangle index %d at ray distance %f"), OutHit.FaceIndex, OutHit.Distance), EToolMessageLevel::Internal);
+			FText::Format(LOCTEXT("OnUpdateDragMessage", "UMeshSurfacePointTool::OnUpdateDrag: Hit triangle index {0} at ray distance {1}"),
+				FText::AsNumber(OutHit.FaceIndex), FText::AsNumber(OutHit.Distance)),
+			EToolMessageLevel::Internal);
 	}
 }
 
@@ -156,5 +160,4 @@ void UMeshSurfacePointToolMouseBehavior::ForceEndCapture(const FInputCaptureData
 
 
 
-
-
+#undef LOCTEXT_NAMESPACE
\ No newline at end of file
diff --git a/Engine/Source/Runtime/Experimental/InteractiveToolsFramework/Private/BaseTools/SingleClickTool.cpp b/Engine/Source/Runtime/Experimental/InteractiveToolsFramework/Private/BaseTools/SingleClickTool.cpp
--- a/Engine/Source/Runtime/Experimental/InteractiveToolsFramework/Private/BaseTools/SingleClickTool.cpp
+++ b/Engine/Source/Runtime/Experimental/InteractiveToolsFramework/Private/BaseTools/SingleClickTool.cpp
@@ -3,6 +3,9 @@
 #include "BaseTools/SingleClickTool.h"
 #include "InteractiveToolManager.h"
 
+#define LOCTEXT_NAMESPACE "USingleClickTool"
+
+
 
 /*
  * ToolBuilder
@@ -46,11 +49,13 @@ bool USingleClickTool::IsHitByClick(const FInputDeviceRay& ClickPos)
 void USingleClickTool::OnClicked(const FInputDeviceRay& ClickPos)
 {
 	// print debug message
-	GetToolManager()->PostMessage( 
-		FString::Printf( TEXT("USingleClickTool::OnClicked at (%f,%f)"), ClickPos.ScreenPosition.X, ClickPos.ScreenPosition.Y), 
+	GetToolManager()->PostMessage(
+		FText::Format(LOCTEXT("OnClickedMessage", "USingleClickTool::OnClicked at ({0},{1})"),
+			FText::AsNumber(ClickPos.ScreenPosition.X), FText::AsNumber(ClickPos.ScreenPosition.Y)),
 		EToolMessageLevel::Internal );
 }
 
 
 
 
+#undef LOCTEXT_NAMESPACE
\ No newline at end of file
diff --git a/Engine/Source/Runtime/Experimental/InteractiveToolsFramework/Private/InputRouter.cpp b/Engine/Source/Runtime/Experimental/InteractiveToolsFramework/Private/InputRouter.cpp
--- a/Engine/Source/Runtime/Experimental/InteractiveToolsFramework/Private/InputRouter.cpp
+++ b/Engine/Source/Runtime/Experimental/InteractiveToolsFramework/Private/InputRouter.cpp
@@ -4,6 +4,8 @@
 #include "InputRouter.h"
 
 
+#define LOCTEXT_NAMESPACE "UInputRouter"
+
 UInputRouter::UInputRouter()
 {
 	ActiveLeftCapture = nullptr;
@@ -65,7 +67,7 @@ void UInputRouter::PostInputEvent(const FInputDeviceState& Input)
 	else
 	{
 		unimplemented();
-		TransactionsAPI->PostMessage(TEXT("UInteractiveToolManager::PostInputEvent - input device is not currently supported."), EToolMessageLevel::Internal);
+		TransactionsAPI->PostMessage(LOCTEXT("PostInputEventMessage", "UInteractiveToolManager::PostInputEvent - input device is not currently supported."), EToolMessageLevel::Internal);
 		return;
 	}
 }
@@ -136,7 +138,7 @@ void UInputRouter::HandleCapturedKeyboardInput(const FInputDeviceState& Input)
 	}
 	else if (Result.State != EInputCaptureState::Continue)
 	{
-		TransactionsAPI->PostMessage(TEXT("UInteractiveToolManager::HandleCapturedKeyboardInput - unexpected capture state!"), EToolMessageLevel::Internal);
+		TransactionsAPI->PostMessage(LOCTEXT("HandleCapturedKeyboardInputMessage", "UInteractiveToolManager::HandleCapturedKeyboardInput - unexpected capture state!"), EToolMessageLevel::Internal);
 	}
 
 	if (bAutoInvalidateOnCapture)
@@ -248,7 +250,7 @@ void UInputRouter::HandleCapturedMouseInput(const FInputDeviceState& Input)
 	}
 	else if (Result.State != EInputCaptureState::Continue)
 	{
-		TransactionsAPI->PostMessage(TEXT("UInteractiveToolManager::HandleCapturedMouseInput - unexpected capture state!"), EToolMessageLevel::Internal);
+		TransactionsAPI->PostMessage(LOCTEXT("HandleCapturedMouseInputMessage", "UInteractiveToolManager::HandleCapturedMouseInput - unexpected capture state!"), EToolMessageLevel::Internal);
 	}
 
 	if (bAutoInvalidateOnCapture)
@@ -316,3 +318,6 @@ void UInputRouter::ForceTerminateSource(IInputBehaviorSource* Source)
 	}
 }
 
+
+
+#undef LOCTEXT_NAMESPACE
\ No newline at end of file
diff --git a/Engine/Source/Runtime/Experimental/InteractiveToolsFramework/Private/InteractiveGizmoManager.cpp b/Engine/Source/Runtime/Experimental/InteractiveToolsFramework/Private/InteractiveGizmoManager.cpp
--- a/Engine/Source/Runtime/Experimental/InteractiveToolsFramework/Private/InteractiveGizmoManager.cpp
+++ b/Engine/Source/Runtime/Experimental/InteractiveToolsFramework/Private/InteractiveGizmoManager.cpp
@@ -9,6 +9,8 @@
 #include "BaseGizmos/AxisAngleGizmo.h"
 #include "BaseGizmos/TransformGizmo.h"
 
+#define LOCTEXT_NAMESPACE "UInteractiveGizmoManager"
+
 
 UInteractiveGizmoManager::UInteractiveGizmoManager()
 {
@@ -61,7 +63,9 @@ bool UInteractiveGizmoManager::DeregisterGizmoType(const FString& BuilderIdentif
 {
 	if (GizmoBuilders.Contains(BuilderIdentifier) == false)
 	{
-		PostMessage(FString::Printf(TEXT("UInteractiveGizmoManager::DeregisterGizmoType: could not find requested type %s"), *BuilderIdentifier), EToolMessageLevel::Internal);
+		PostMessage(
+			FText::Format(LOCTEXT("DeregisterFailedMessage", "UInteractiveGizmoManager::DeregisterGizmoType: could not find requested type {0}"), FText::FromString(BuilderIdentifier) ),
+			EToolMessageLevel::Internal);
 		return false;
 	}
 	GizmoBuilders.Remove(BuilderIdentifier);
@@ -75,7 +79,9 @@ UInteractiveGizmo* UInteractiveGizmoManager::CreateGizmo(const FString& BuilderI
 {
 	if ( GizmoBuilders.Contains(BuilderIdentifier) == false )
 	{
-		PostMessage(FString::Printf(TEXT("UInteractiveGizmoManager::CreateGizmo: could not find requested type %s"), *BuilderIdentifier), EToolMessageLevel::Internal);
+		PostMessage(
+			FText::Format(LOCTEXT("CreateGizmoCannotFindFailedMessage", "UInteractiveGizmoManager::CreateGizmo: could not find requested type {0}"), FText::FromString(BuilderIdentifier) ),
+			EToolMessageLevel::Internal);
 		return nullptr;
 	}
 	UInteractiveGizmoBuilder* FoundBuilder = GizmoBuilders[BuilderIdentifier];
@@ -87,7 +93,9 @@ UInteractiveGizmo* UInteractiveGizmoManager::CreateGizmo(const FString& BuilderI
 		{
 			if (ActiveGizmo.InstanceIdentifier == InstanceIdentifier)
 			{
-				PostMessage(FString::Printf(TEXT("UInteractiveGizmoManager::CreateGizmo: instance identifier %s already in use!"), *InstanceIdentifier), EToolMessageLevel::Internal);
+				PostMessage(
+					FText::Format(LOCTEXT("CreateGizmoExistsMessage", "UInteractiveGizmoManager::CreateGizmo: instance identifier {0} already in use!"), FText::FromString(InstanceIdentifier) ),
+					EToolMessageLevel::Internal);
 				return nullptr;
 			}
 		}
@@ -99,7 +107,7 @@ UInteractiveGizmo* UInteractiveGizmoManager::CreateGizmo(const FString& BuilderI
 	UInteractiveGizmo* NewGizmo = FoundBuilder->BuildGizmo(CurrentSceneState);
 	if (NewGizmo == nullptr)
 	{
-		PostMessage(FString::Printf(TEXT("UInteractiveGizmoManager::CreateGizmo: BuildGizmo() returned null")), EToolMessageLevel::Internal);
+		PostMessage(LOCTEXT("CreateGizmoReturnNullMessage", "UInteractiveGizmoManager::CreateGizmo: BuildGizmo() returned null"), EToolMessageLevel::Internal);
 		return nullptr;
 	}
 
@@ -225,16 +233,11 @@ void UInteractiveGizmoManager::Render(IToolsContextRenderAPI* RenderAPI)
 
 }
 
-void UInteractiveGizmoManager::PostMessage(const TCHAR* Message, EToolMessageLevel Level)
+void UInteractiveGizmoManager::PostMessage(const FText& Message, EToolMessageLevel Level)
 {
 	TransactionsAPI->PostMessage(Message, Level);
 }
 
-void UInteractiveGizmoManager::PostMessage(const FString& Message, EToolMessageLevel Level)
-{
-	TransactionsAPI->PostMessage(*Message, Level);
-}
-
 void UInteractiveGizmoManager::PostInvalidation()
 {
 	TransactionsAPI->PostInvalidation();
@@ -292,4 +295,7 @@ UTransformGizmo* UInteractiveGizmoManager::Create3AxisTransformGizmo(const FStri
 	UInteractiveGizmo* NewGizmo = CreateGizmo(DefaultThreeAxisTransformBuilderIdentifier, InstanceIdentifier, Owner);
 	check(NewGizmo);
 	return Cast<UTransformGizmo>(NewGizmo);
-}
\ No newline at end of file
+}
+
+
+#undef LOCTEXT_NAMESPACE
\ No newline at end of file
diff --git a/Engine/Source/Runtime/Experimental/InteractiveToolsFramework/Private/InteractiveToolManager.cpp b/Engine/Source/Runtime/Experimental/InteractiveToolsFramework/Private/InteractiveToolManager.cpp
--- a/Engine/Source/Runtime/Experimental/InteractiveToolsFramework/Private/InteractiveToolManager.cpp
+++ b/Engine/Source/Runtime/Experimental/InteractiveToolsFramework/Private/InteractiveToolManager.cpp
@@ -118,7 +118,7 @@ bool UInteractiveToolManager::ActivateTool(EToolSide Side)
 
 	if (ActiveLeftBuilder->CanBuildTool(InputState) == false)
 	{
-		TransactionsAPI->PostMessage( TEXT("UInteractiveToolManager::ActivateTool: CanBuildTool returned false."), EToolMessageLevel::Internal);
+		TransactionsAPI->PostMessage( LOCTEXT("ActivateToolCanBuildFailMessage", "UInteractiveToolManager::ActivateTool: CanBuildTool returned false."), EToolMessageLevel::Internal);
 		return false;
 	}
 
@@ -240,16 +240,11 @@ UInteractiveGizmoManager* UInteractiveToolManager::GetPairedGizmoManager()
 	return Cast<UInteractiveToolsContext>(GetOuter())->GizmoManager;
 }
 
-void UInteractiveToolManager::PostMessage(const TCHAR* Message, EToolMessageLevel Level)
+void UInteractiveToolManager::PostMessage(const FText& Message, EToolMessageLevel Level)
 {
 	TransactionsAPI->PostMessage(Message, Level);
 }
 
-void UInteractiveToolManager::PostMessage(const FString& Message, EToolMessageLevel Level)
-{
-	TransactionsAPI->PostMessage(*Message, Level);
-}
-
 void UInteractiveToolManager::PostInvalidation()
 {
 	TransactionsAPI->PostInvalidation();
diff --git a/Engine/Source/Runtime/Experimental/InteractiveToolsFramework/Public/InteractiveGizmoManager.h b/Engine/Source/Runtime/Experimental/InteractiveToolsFramework/Public/InteractiveGizmoManager.h
--- a/Engine/Source/Runtime/Experimental/InteractiveToolsFramework/Public/InteractiveGizmoManager.h
+++ b/Engine/Source/Runtime/Experimental/InteractiveToolsFramework/Public/InteractiveGizmoManager.h
@@ -121,10 +121,7 @@ public:
 	//
 	
 	/** Post a message via the Transactions API */
-	virtual void PostMessage(const TCHAR* Message, EToolMessageLevel Level);
-
-	/** Post a message via the Transactions API */
-	virtual void PostMessage(const FString& Message, EToolMessageLevel Level);
+	virtual void PostMessage(const FText& Message, EToolMessageLevel Level);
 
 	/** Request an Invalidation via the Transactions API (ie to cause a repaint, etc) */
 	virtual void PostInvalidation();
diff --git a/Engine/Source/Runtime/Experimental/InteractiveToolsFramework/Public/InteractiveToolManager.h b/Engine/Source/Runtime/Experimental/InteractiveToolsFramework/Public/InteractiveToolManager.h
--- a/Engine/Source/Runtime/Experimental/InteractiveToolsFramework/Public/InteractiveToolManager.h
+++ b/Engine/Source/Runtime/Experimental/InteractiveToolsFramework/Public/InteractiveToolManager.h
@@ -141,10 +141,7 @@ public:
 	//
 	
 	/** Post a message via the Transactions API */
-	virtual void PostMessage(const TCHAR* Message, EToolMessageLevel Level);
-
-	/** Post a message via the Transactions API */
-	virtual void PostMessage(const FString& Message, EToolMessageLevel Level);
+	virtual void PostMessage(const FText& Message, EToolMessageLevel Level);
 
 	/** Request an Invalidation via the Transactions API (ie to cause a repaint, etc) */
 	virtual void PostInvalidation();
diff --git a/Engine/Source/Runtime/Experimental/InteractiveToolsFramework/Public/ToolContextInterfaces.h b/Engine/Source/Runtime/Experimental/InteractiveToolsFramework/Public/ToolContextInterfaces.h
--- a/Engine/Source/Runtime/Experimental/InteractiveToolsFramework/Public/ToolContextInterfaces.h
+++ b/Engine/Source/Runtime/Experimental/InteractiveToolsFramework/Public/ToolContextInterfaces.h
@@ -201,6 +201,7 @@ enum class EToolMessageLevel
 	UserError = 3
 };
 
+
 /** Type of change we want to apply to a selection */
 UENUM()
 enum class ESelectedObjectsModificationType
@@ -239,7 +240,7 @@ public:
 	 * @param Message text of message
 	 * @param Level severity level of message
 	 */
-	virtual void PostMessage(const TCHAR* Message, EToolMessageLevel Level) = 0;
+	virtual void PostMessage(const FText& Message, EToolMessageLevel Level) = 0;
 
 	/** 
 	 * Forward an invalidation request from Tools framework, to cause repaint/etc. 
