ToolsFramework: GizmoManager can now spawn customized variants of the standard transform gizmo. Minor API changes.
#rnx

diff --git a/Engine/Plugins/Experimental/MeshModelingToolset/Source/MeshModelingTools/Private/PlaneCutTool.cpp b/Engine/Plugins/Experimental/MeshModelingToolset/Source/MeshModelingTools/Private/PlaneCutTool.cpp
--- a/Engine/Plugins/Experimental/MeshModelingToolset/Source/MeshModelingTools/Private/PlaneCutTool.cpp
+++ b/Engine/Plugins/Experimental/MeshModelingToolset/Source/MeshModelingTools/Private/PlaneCutTool.cpp
@@ -109,7 +109,7 @@ void UPlaneCutTool::Setup()
 	// create proxy and gizmo (but don't attach yet)
 	UInteractiveGizmoManager* GizmoManager = GetToolManager()->GetPairedGizmoManager();
 	PlaneTransformProxy = NewObject<UTransformProxy>(this);
-	PlaneTransformGizmo = GizmoManager->Create3AxisTransformGizmo(TEXT("PlaneCutToolTransformGizmo"), this);
+	PlaneTransformGizmo = GizmoManager->Create3AxisTransformGizmo(this);
 
 	BasicProperties = NewObject<UPlaneCutToolProperties>(this, TEXT("Plane Cut Settings"));
 	AdvancedProperties = NewObject<UPlaneCutAdvancedProperties>(this, TEXT("Advanced Settings"));
diff --git a/Engine/Plugins/Experimental/MeshModelingToolset/Source/MeshModelingTools/Private/TransformMeshesTool.cpp b/Engine/Plugins/Experimental/MeshModelingToolset/Source/MeshModelingTools/Private/TransformMeshesTool.cpp
--- a/Engine/Plugins/Experimental/MeshModelingToolset/Source/MeshModelingTools/Private/TransformMeshesTool.cpp
+++ b/Engine/Plugins/Experimental/MeshModelingToolset/Source/MeshModelingTools/Private/TransformMeshesTool.cpp
@@ -169,7 +169,7 @@ void UTransformMeshesTool::SetActiveGizmos_Single(bool bLocalRotations)
 	{
 		Transformable.TransformProxy->AddComponent(Target->GetOwnerComponent());
 	}
-	Transformable.TransformGizmo = GizmoManager->Create3AxisTransformGizmo(FString(), this);
+	Transformable.TransformGizmo = GizmoManager->Create3AxisTransformGizmo(this);
 	Transformable.TransformGizmo->SetActiveTarget(Transformable.TransformProxy);
 
 	ActiveGizmos.Add(Transformable);
@@ -184,7 +184,7 @@ void UTransformMeshesTool::SetActiveGizmos_PerObject()
 		FTransformMeshesTarget Transformable;
 		Transformable.TransformProxy = NewObject<UTransformProxy>(this);
 		Transformable.TransformProxy->AddComponent(Target->GetOwnerComponent());
-		Transformable.TransformGizmo = GizmoManager->Create3AxisTransformGizmo(FString(), this);
+		Transformable.TransformGizmo = GizmoManager->Create3AxisTransformGizmo(this);
 		Transformable.TransformGizmo->SetActiveTarget(Transformable.TransformProxy);
 		ActiveGizmos.Add(Transformable);
 	}
diff --git a/Engine/Plugins/Experimental/MeshModelingToolset/Source/MeshModelingTools/Private/UVProjectionTool.cpp b/Engine/Plugins/Experimental/MeshModelingToolset/Source/MeshModelingTools/Private/UVProjectionTool.cpp
--- a/Engine/Plugins/Experimental/MeshModelingToolset/Source/MeshModelingTools/Private/UVProjectionTool.cpp
+++ b/Engine/Plugins/Experimental/MeshModelingToolset/Source/MeshModelingTools/Private/UVProjectionTool.cpp
@@ -175,7 +175,7 @@ void UUVProjectionTool::UpdateNumPreviews()
 			TransformProxy->SetTransform(LocalXF * ComponentTargets[PreviewIdx]->GetWorldTransform());
 			TransformProxy->OnTransformChanged.AddUObject(this, &UUVProjectionTool::TransformChanged);
 
-			UTransformGizmo* TransformGizmo = TransformGizmos.Add_GetRef(GizmoManager->Create3AxisTransformGizmo(TEXT("UVProjectionToolTransformGizmo"), this));
+			UTransformGizmo* TransformGizmo = TransformGizmos.Add_GetRef(GizmoManager->Create3AxisTransformGizmo(this));
 			TransformGizmo->SetActiveTarget(TransformProxy);
 		}
 		check(TransformProxies.Num() == TargetNumPreview);
diff --git a/Engine/Source/Runtime/Experimental/InteractiveToolsFramework/Private/BaseGizmos/TransformGizmo.cpp b/Engine/Source/Runtime/Experimental/InteractiveToolsFramework/Private/BaseGizmos/TransformGizmo.cpp
--- a/Engine/Source/Runtime/Experimental/InteractiveToolsFramework/Private/BaseGizmos/TransformGizmo.cpp
+++ b/Engine/Source/Runtime/Experimental/InteractiveToolsFramework/Private/BaseGizmos/TransformGizmo.cpp
@@ -36,32 +36,76 @@ ATransformGizmoActor::ATransformGizmoActor()
 
 
 ATransformGizmoActor* ATransformGizmoActor::ConstructDefault3AxisGizmo(UWorld* World)
+{
+	return ConstructCustom3AxisGizmo(World, 
+		ETransformGizmoSubElements::TranslateAllAxes |
+		ETransformGizmoSubElements::TranslateAllPlanes |
+		ETransformGizmoSubElements::RotateAllAxes);
+}
+
+
+ATransformGizmoActor* ATransformGizmoActor::ConstructCustom3AxisGizmo(
+	UWorld* World,
+	ETransformGizmoSubElements Elements)
 {
 	FActorSpawnParameters SpawnInfo;
 	ATransformGizmoActor* NewActor = World->SpawnActor<ATransformGizmoActor>(FVector::ZeroVector, FRotator::ZeroRotator, SpawnInfo);
 
-	NewActor->TranslateX = AddDefaultArrowComponent(World, NewActor, FLinearColor::Red, FVector(1, 0, 0));
-	NewActor->TranslateY = AddDefaultArrowComponent(World, NewActor, FLinearColor::Green, FVector(0, 1, 0));
-	NewActor->TranslateZ = AddDefaultArrowComponent(World, NewActor, FLinearColor::Blue, FVector(0, 0, 1));
+	if ((Elements & ETransformGizmoSubElements::TranslateAxisX) != ETransformGizmoSubElements::None)
+	{
+		NewActor->TranslateX = AddDefaultArrowComponent(World, NewActor, FLinearColor::Red, FVector(1, 0, 0));
+	}
+	if ((Elements & ETransformGizmoSubElements::TranslateAxisY) != ETransformGizmoSubElements::None)
+	{
+		NewActor->TranslateY = AddDefaultArrowComponent(World, NewActor, FLinearColor::Green, FVector(0, 1, 0));
+	}
+	if ((Elements & ETransformGizmoSubElements::TranslateAxisZ) != ETransformGizmoSubElements::None)
+	{
+		NewActor->TranslateZ = AddDefaultArrowComponent(World, NewActor, FLinearColor::Blue, FVector(0, 0, 1));
+	}
 
-	NewActor->TranslateYZ = AddDefaultRectangleComponent(World, NewActor, FLinearColor::Red, FVector(0,1,0), FVector(0,0,1));
-	NewActor->TranslateXZ = AddDefaultRectangleComponent(World, NewActor, FLinearColor::Green, FVector(1, 0, 0), FVector(0, 0, 1));
-	NewActor->TranslateXY = AddDefaultRectangleComponent(World, NewActor, FLinearColor::Blue, FVector(1, 0, 0), FVector(0, 1, 0));
+	if ((Elements & ETransformGizmoSubElements::TranslatePlaneYZ) != ETransformGizmoSubElements::None)
+	{
+		NewActor->TranslateYZ = AddDefaultRectangleComponent(World, NewActor, FLinearColor::Red, FVector(0, 1, 0), FVector(0, 0, 1));
+	}
+	if ((Elements & ETransformGizmoSubElements::TranslatePlaneXZ) != ETransformGizmoSubElements::None)
+	{
+		NewActor->TranslateXZ = AddDefaultRectangleComponent(World, NewActor, FLinearColor::Green, FVector(1, 0, 0), FVector(0, 0, 1));
+	}
+	if ((Elements & ETransformGizmoSubElements::TranslatePlaneXY) != ETransformGizmoSubElements::None)
+	{
+		NewActor->TranslateXY = AddDefaultRectangleComponent(World, NewActor, FLinearColor::Blue, FVector(1, 0, 0), FVector(0, 1, 0));
+	}
 
-	NewActor->RotateX = AddDefaultCircleComponent(World, NewActor, FLinearColor::Red, FVector(1, 0, 0));
-	NewActor->RotateY = AddDefaultCircleComponent(World, NewActor, FLinearColor::Green, FVector(0, 1, 0));
-	NewActor->RotateZ = AddDefaultCircleComponent(World, NewActor, FLinearColor::Blue, FVector(0, 0, 1));
+	bool bAnyRotate = false;
+	if ((Elements & ETransformGizmoSubElements::RotateAxisX) != ETransformGizmoSubElements::None)
+	{
+		NewActor->RotateX = AddDefaultCircleComponent(World, NewActor, FLinearColor::Red, FVector(1, 0, 0));
+		bAnyRotate = true;
+	}
+	if ((Elements & ETransformGizmoSubElements::RotateAxisY) != ETransformGizmoSubElements::None)
+	{
+		NewActor->RotateY = AddDefaultCircleComponent(World, NewActor, FLinearColor::Green, FVector(0, 1, 0));
+		bAnyRotate = true;
+	}
+	if ((Elements & ETransformGizmoSubElements::RotateAxisZ) != ETransformGizmoSubElements::None)
+	{
+		NewActor->RotateZ = AddDefaultCircleComponent(World, NewActor, FLinearColor::Blue, FVector(0, 0, 1));
+		bAnyRotate = true;
+	}
 
 
 	// add a non-interactive view-aligned circle element, so the axes look like a sphere.
-	UGizmoCircleComponent* SphereEdge = NewObject<UGizmoCircleComponent>(NewActor);
-	NewActor->AddInstanceComponent(SphereEdge);
-	SphereEdge->AttachToComponent(NewActor->GetRootComponent(), FAttachmentTransformRules::KeepRelativeTransform);
-	SphereEdge->Color = FLinearColor::Gray;
-	SphereEdge->Thickness = 1.0f;
-	SphereEdge->bViewAligned = true;
-	SphereEdge->RegisterComponent();
-
+	if (bAnyRotate)
+	{
+		UGizmoCircleComponent* SphereEdge = NewObject<UGizmoCircleComponent>(NewActor);
+		NewActor->AddInstanceComponent(SphereEdge);
+		SphereEdge->AttachToComponent(NewActor->GetRootComponent(), FAttachmentTransformRules::KeepRelativeTransform);
+		SphereEdge->Color = FLinearColor::Gray;
+		SphereEdge->Thickness = 1.0f;
+		SphereEdge->bViewAligned = true;
+		SphereEdge->RegisterComponent();
+	}
 
 	return NewActor;
 }
@@ -69,10 +113,9 @@ ATransformGizmoActor* ATransformGizmoActor::ConstructDefault3AxisGizmo(UWorld* W
 
 
 
-
 ATransformGizmoActor* FTransformGizmoActorFactory::CreateNewGizmoActor(UWorld* World) const
 {
-	return ATransformGizmoActor::ConstructDefault3AxisGizmo(World);
+	return ATransformGizmoActor::ConstructCustom3AxisGizmo(World, EnableElements);
 }
 
 
diff --git a/Engine/Source/Runtime/Experimental/InteractiveToolsFramework/Private/InteractiveGizmoManager.cpp b/Engine/Source/Runtime/Experimental/InteractiveToolsFramework/Private/InteractiveGizmoManager.cpp
--- a/Engine/Source/Runtime/Experimental/InteractiveToolsFramework/Private/InteractiveGizmoManager.cpp
+++ b/Engine/Source/Runtime/Experimental/InteractiveToolsFramework/Private/InteractiveGizmoManager.cpp
@@ -268,7 +268,7 @@ FString UInteractiveGizmoManager::DefaultAxisPositionBuilderIdentifier = TEXT("S
 FString UInteractiveGizmoManager::DefaultPlanePositionBuilderIdentifier = TEXT("StandardXFormPlaneTranslationGizmo");
 FString UInteractiveGizmoManager::DefaultAxisAngleBuilderIdentifier = TEXT("StandardXFormAxisRotationGizmo");
 FString UInteractiveGizmoManager::DefaultThreeAxisTransformBuilderIdentifier = TEXT("DefaultThreeAxisTransformBuilderIdentifier");
-
+const FString UInteractiveGizmoManager::CustomThreeAxisTransformBuilderIdentifier = TEXT("CustomThreeAxisTransformBuilderIdentifier");
 
 void UInteractiveGizmoManager::RegisterDefaultGizmos()
 {
@@ -286,10 +286,14 @@ void UInteractiveGizmoManager::RegisterDefaultGizmos()
 	UTransformGizmoBuilder* TransformBuilder = NewObject<UTransformGizmoBuilder>();
 	RegisterGizmoType(DefaultThreeAxisTransformBuilderIdentifier, TransformBuilder);
 
+	CustomThreeAxisBuilder = NewObject<UTransformGizmoBuilder>();
+	CustomThreeAxisBuilder->GizmoActorBuilder = MakeShared<FTransformGizmoActorFactory>();
+	RegisterGizmoType(CustomThreeAxisTransformBuilderIdentifier, CustomThreeAxisBuilder);
+
 	bDefaultGizmosRegistered = true;
 }
 
-UTransformGizmo* UInteractiveGizmoManager::Create3AxisTransformGizmo(const FString& InstanceIdentifier, void* Owner)
+UTransformGizmo* UInteractiveGizmoManager::Create3AxisTransformGizmo(void* Owner, const FString& InstanceIdentifier)
 {
 	check(bDefaultGizmosRegistered);
 	UInteractiveGizmo* NewGizmo = CreateGizmo(DefaultThreeAxisTransformBuilderIdentifier, InstanceIdentifier, Owner);
@@ -297,5 +301,14 @@ UTransformGizmo* UInteractiveGizmoManager::Create3AxisTransformGizmo(const FStri
 	return Cast<UTransformGizmo>(NewGizmo);
 }
 
+UTransformGizmo* UInteractiveGizmoManager::CreateCustomTransformGizmo(ETransformGizmoSubElements Elements, void* Owner, const FString& InstanceIdentifier)
+{
+	check(bDefaultGizmosRegistered);
+	CustomThreeAxisBuilder->GizmoActorBuilder->EnableElements = Elements;
+	UInteractiveGizmo* NewGizmo = CreateGizmo(CustomThreeAxisTransformBuilderIdentifier, InstanceIdentifier, Owner);
+	check(NewGizmo);
+	return Cast<UTransformGizmo>(NewGizmo);
+}
+
 
 #undef LOCTEXT_NAMESPACE
\ No newline at end of file
diff --git a/Engine/Source/Runtime/Experimental/InteractiveToolsFramework/Public/BaseGizmos/TransformGizmo.h b/Engine/Source/Runtime/Experimental/InteractiveToolsFramework/Public/BaseGizmos/TransformGizmo.h
--- a/Engine/Source/Runtime/Experimental/InteractiveToolsFramework/Public/BaseGizmos/TransformGizmo.h
+++ b/Engine/Source/Runtime/Experimental/InteractiveToolsFramework/Public/BaseGizmos/TransformGizmo.h
@@ -86,6 +86,15 @@ public:
 	static ATransformGizmoActor* ConstructDefault3AxisGizmo(
 		UWorld* World
 	);
+
+	/**
+	 * Create a new instance of ATransformGizmoActor. Populate the sub-components 
+	 * specified by Elements with standard GizmoXComponent instances suitable for a 3-axis transformer Gizmo
+	 */
+	static ATransformGizmoActor* ConstructCustom3AxisGizmo(
+		UWorld* World,
+		ETransformGizmoSubElements Elements
+	);
 };
 
 
@@ -106,6 +115,10 @@ public:
 class INTERACTIVETOOLSFRAMEWORK_API FTransformGizmoActorFactory
 {
 public:
+	/** Only these members of the ATransformGizmoActor gizmo will be initialized */
+	ETransformGizmoSubElements EnableElements =
+		ETransformGizmoSubElements::TranslateAllAxes | ETransformGizmoSubElements::TranslateAllPlanes | ETransformGizmoSubElements::RotateAllAxes;
+
 	/**
 	 * @param World the UWorld to create the new Actor in
 	 * @return new ATransformGizmoActor instance with members initialized with Components suitable for a transformation Gizmo
diff --git a/Engine/Source/Runtime/Experimental/InteractiveToolsFramework/Public/InteractiveGizmo.h b/Engine/Source/Runtime/Experimental/InteractiveToolsFramework/Public/InteractiveGizmo.h
--- a/Engine/Source/Runtime/Experimental/InteractiveToolsFramework/Public/InteractiveGizmo.h
+++ b/Engine/Source/Runtime/Experimental/InteractiveToolsFramework/Public/InteractiveGizmo.h
@@ -81,3 +81,31 @@ protected:
 };
 
 
+
+/**
+ * ETransformGizmoSubElements identifies the sub-elements of a standard 3-axis transformation Gizmo.
+ * Used by GizmoManager to customize UTransformGizmo instances.
+ */
+UENUM()
+enum class ETransformGizmoSubElements
+{
+	None = 0,
+
+	TranslateAxisX = 1,
+	TranslateAxisY = 2,
+	TranslateAxisZ = 4,
+	TranslateAllAxes = TranslateAxisX | TranslateAxisY | TranslateAxisZ,
+
+	TranslatePlaneXY = 8,
+	TranslatePlaneXZ = 16,
+	TranslatePlaneYZ = 32,
+	TranslateAllPlanes = TranslatePlaneXY | TranslatePlaneXZ | TranslatePlaneYZ,
+
+	RotateAxisX = 64,
+	RotateAxisY = 128,
+	RotateAxisZ = 256,
+	RotateAllAxes = RotateAxisX | RotateAxisY | RotateAxisZ,
+
+	StanedardTranslateRotate = TranslateAllAxes | TranslateAllPlanes | RotateAllAxes
+};
+ENUM_CLASS_FLAGS(ETransformGizmoSubElements);
\ No newline at end of file
diff --git a/Engine/Source/Runtime/Experimental/InteractiveToolsFramework/Public/InteractiveGizmoManager.h b/Engine/Source/Runtime/Experimental/InteractiveToolsFramework/Public/InteractiveGizmoManager.h
--- a/Engine/Source/Runtime/Experimental/InteractiveToolsFramework/Public/InteractiveGizmoManager.h
+++ b/Engine/Source/Runtime/Experimental/InteractiveToolsFramework/Public/InteractiveGizmoManager.h
@@ -12,11 +12,12 @@
 #include "InteractiveGizmoManager.generated.h"
 
 class UTransformGizmo;
+class UTransformGizmoBuilder;
 
 USTRUCT()
 struct FActiveGizmo
 {
-	GENERATED_BODY();
+	GENERATED_BODY()
 
 	UInteractiveGizmo* Gizmo;
 	FString BuilderIdentifier;
@@ -178,11 +179,20 @@ public:
 
 	/**
 	 * Activate a new instance of the default 3-axis transformation Gizmo. RegisterDefaultGizmos() must have been called first.
-	 * @param InstanceIdentifier client-defined *unique* string that can be used to locate this instance
-	 * @param Owner void pointer to whatever "owns" this Gizmo. Allows Gizmo to later be deleted using DestroyAllGizmosByOwner()
+	 * @param Owner optional void pointer to whatever "owns" this Gizmo. Allows Gizmo to later be deleted using DestroyAllGizmosByOwner()
+	 * @param InstanceIdentifier optional client-defined *unique* string that can be used to locate this instance
+	 * @return new Gizmo instance that has been created and initialized
+	 */
+	virtual UTransformGizmo* Create3AxisTransformGizmo(void* Owner = nullptr, const FString& InstanceIdentifier = FString());
+
+	/**
+	 * Activate a new customized instance of the default 3-axis transformation Gizmo, with only certain elements included. RegisterDefaultGizmos() must have been called first.
+	 * @param Elements flags that indicate which standard gizmo sub-elements should be included
+	 * @param Owner optional void pointer to whatever "owns" this Gizmo. Allows Gizmo to later be deleted using DestroyAllGizmosByOwner()
+	 * @param InstanceIdentifier optional client-defined *unique* string that can be used to locate this instance
 	 * @return new Gizmo instance that has been created and initialized
 	 */
-	virtual UTransformGizmo* Create3AxisTransformGizmo(const FString& InstanceIdentifier = FString(), void* Owner = nullptr);
+	virtual UTransformGizmo* CreateCustomTransformGizmo(ETransformGizmoSubElements Elements, void* Owner = nullptr, const FString& InstanceIdentifier = FString());
 
 
 public:
@@ -191,6 +201,7 @@ public:
 	static FString DefaultPlanePositionBuilderIdentifier;
 	static FString DefaultAxisAngleBuilderIdentifier;
 	static FString DefaultThreeAxisTransformBuilderIdentifier;
+	static const FString CustomThreeAxisTransformBuilderIdentifier;
 
 
 protected:
@@ -213,4 +224,5 @@ protected:
 	TMap<FString, UInteractiveGizmoBuilder*> GizmoBuilders;
 
 	bool bDefaultGizmosRegistered = false;
+	UTransformGizmoBuilder* CustomThreeAxisBuilder = nullptr;
 };
\ No newline at end of file
