ToolsFramework: add trivial function to return material index count from Component
#rnx

diff --git a/Engine/Source/Runtime/Experimental/InteractiveToolsFramework/Private/ComponentSourceInterfaces.cpp b/Engine/Source/Runtime/Experimental/InteractiveToolsFramework/Private/ComponentSourceInterfaces.cpp
--- a/Engine/Source/Runtime/Experimental/InteractiveToolsFramework/Private/ComponentSourceInterfaces.cpp
+++ b/Engine/Source/Runtime/Experimental/InteractiveToolsFramework/Private/ComponentSourceInterfaces.cpp
@@ -61,6 +61,12 @@ FPrimitiveComponentTarget::SetOwnerVisibility(bool bVisible) const
 	Component->SetVisibility(bVisible);
 }
 
+
+int32 FPrimitiveComponentTarget::GetNumMaterials() const
+{
+	return Component->GetNumMaterials();
+}
+
 UMaterialInterface*
 FPrimitiveComponentTarget::GetMaterial(int32 MaterialIndex) const
 {
diff --git a/Engine/Source/Runtime/Experimental/InteractiveToolsFramework/Public/ComponentSourceInterfaces.h b/Engine/Source/Runtime/Experimental/InteractiveToolsFramework/Public/ComponentSourceInterfaces.h
--- a/Engine/Source/Runtime/Experimental/InteractiveToolsFramework/Public/ComponentSourceInterfaces.h
+++ b/Engine/Source/Runtime/Experimental/InteractiveToolsFramework/Public/ComponentSourceInterfaces.h
@@ -39,6 +39,9 @@ public:
 	/** @return the Component this is a Source for */
 	UPrimitiveComponent* GetOwnerComponent() const;
 
+	/** @return number of material indices in use by this Component */
+	int32 GetNumMaterials() const;
+
 	/**
 	 * Get pointer to a Material provided by this Source
 	 * @param MaterialIndex index of the material
