Fix compile error with UnEdMode PostMessage on windows due to a window headers defining PostMessage
#rnx

diff --git a/Engine/Plugins/Experimental/MeshModelingToolset/Source/MeshModelingTools/Private/AddPrimitiveTool.cpp b/Engine/Plugins/Experimental/MeshModelingToolset/Source/MeshModelingTools/Private/AddPrimitiveTool.cpp
--- a/Engine/Plugins/Experimental/MeshModelingToolset/Source/MeshModelingTools/Private/AddPrimitiveTool.cpp
+++ b/Engine/Plugins/Experimental/MeshModelingToolset/Source/MeshModelingTools/Private/AddPrimitiveTool.cpp
@@ -113,7 +113,7 @@ void UAddPrimitiveTool::Setup()
 
 	UpdatePreviewMesh();
 
-	GetToolManager()->PostMessage(
+	GetToolManager()->DisplayMessage(
 		LOCTEXT("OnStartAddPrimitiveTool", "Position the Primitive by moving the mouse over the scene. Drop a new instance by Left-clicking."),
 		EToolMessageLevel::UserNotification);
 }
diff --git a/Engine/Plugins/Experimental/MeshModelingToolset/Source/MeshModelingTools/Private/ConvertToPolygonsTool.cpp b/Engine/Plugins/Experimental/MeshModelingToolset/Source/MeshModelingTools/Private/ConvertToPolygonsTool.cpp
--- a/Engine/Plugins/Experimental/MeshModelingToolset/Source/MeshModelingTools/Private/ConvertToPolygonsTool.cpp
+++ b/Engine/Plugins/Experimental/MeshModelingToolset/Source/MeshModelingTools/Private/ConvertToPolygonsTool.cpp
@@ -150,7 +150,7 @@ void UConvertToPolygonsTool::UpdatePolygons()
 	Polygons.FindPolygons(DotTolerance);
 	Polygons.FindPolygonEdges();
 
-	GetToolManager()->PostMessage(
+	GetToolManager()->DisplayMessage(
 		FText::Format(LOCTEXT("UpdatePolygonsMessage", "ConvertToPolygons - found {0} polys in {1} triangles"), 
 			FText::AsNumber(Polygons.FoundPolygons.Num()), FText::AsNumber(SearchMesh.TriangleCount()) ), EToolMessageLevel::Internal);
 
diff --git a/Engine/Plugins/Experimental/MeshModelingToolset/Source/MeshModelingTools/Private/DrawPolygonTool.cpp b/Engine/Plugins/Experimental/MeshModelingToolset/Source/MeshModelingTools/Private/DrawPolygonTool.cpp
--- a/Engine/Plugins/Experimental/MeshModelingToolset/Source/MeshModelingTools/Private/DrawPolygonTool.cpp
+++ b/Engine/Plugins/Experimental/MeshModelingToolset/Source/MeshModelingTools/Private/DrawPolygonTool.cpp
@@ -147,7 +147,7 @@ void UDrawPolygonTool::Setup()
 	AddToolPropertySource(PolygonProperties);
 	AddToolPropertySource(SnapProperties);
 
-	GetToolManager()->PostMessage(
+	GetToolManager()->DisplayMessage(
 		LOCTEXT("OnStartDrawPolygonTool", "Left-click to place points on the Drawing Plane. Ctrl-click on the scene to reposition the Plane, and click-drag on the gizmo to move it."),
 		EToolMessageLevel::UserNotification);
 }
diff --git a/Engine/Source/Editor/Experimental/EditorInteractiveToolsFramework/Private/EdModeInteractiveToolsContext.cpp b/Engine/Source/Editor/Experimental/EditorInteractiveToolsFramework/Private/EdModeInteractiveToolsContext.cpp
--- a/Engine/Source/Editor/Experimental/EditorInteractiveToolsFramework/Private/EdModeInteractiveToolsContext.cpp
+++ b/Engine/Source/Editor/Experimental/EditorInteractiveToolsFramework/Private/EdModeInteractiveToolsContext.cpp
@@ -240,7 +240,7 @@ public:
 	}
 
 
-	virtual void PostMessage(const FText& Message, EToolMessageLevel Level) override
+	virtual void DisplayMessage(const FText& Message, EToolMessageLevel Level) override
 	{
 		if (Level == EToolMessageLevel::UserNotification)
 		{
diff --git a/Engine/Source/Editor/UnrealEd/Private/Tools/UEdMode.cpp b/Engine/Source/Editor/UnrealEd/Private/Tools/UEdMode.cpp
--- a/Engine/Source/Editor/UnrealEd/Private/Tools/UEdMode.cpp
+++ b/Engine/Source/Editor/UnrealEd/Private/Tools/UEdMode.cpp
@@ -237,7 +237,7 @@ public:
 		EdMode = InEdMode;
 	}
 
-	virtual void PostMessage(const FText& Message, EToolMessageLevel Level) override
+	virtual void DisplayMessage(const FText& Message, EToolMessageLevel Level) override
 	{
 		UE_LOG(LogTemp, Warning, TEXT("%s"), *Message.ToString());
 	}
diff --git a/Engine/Source/Runtime/ApplicationCore/Private/Windows/WindowsPlatformSplash.cpp b/Engine/Source/Runtime/ApplicationCore/Private/Windows/WindowsPlatformSplash.cpp
--- a/Engine/Source/Runtime/ApplicationCore/Private/Windows/WindowsPlatformSplash.cpp
+++ b/Engine/Source/Runtime/ApplicationCore/Private/Windows/WindowsPlatformSplash.cpp
@@ -700,7 +700,7 @@ void FWindowsPlatformSplash::Hide()
 		if(GSplashScreenWnd)
 		{
 			// Send message to splash screen window to destroy itself
-			PostMessage(GSplashScreenWnd, WM_DESTROY, 0, 0);
+			PostMessageW(GSplashScreenWnd, WM_DESTROY, 0, 0);
 		}
 
 		// Wait for splash screen thread to finish
@@ -714,7 +714,7 @@ void FWindowsPlatformSplash::Hide()
 		// Close the Z-Order guard window
 		if ( GSplashScreenGuard )
 		{
-			PostMessage(GSplashScreenGuard, WM_DESTROY, 0, 0);
+			PostMessageW(GSplashScreenGuard, WM_DESTROY, 0, 0);
 			GSplashScreenGuard = NULL;
 		}
 	}
diff --git a/Engine/Source/Runtime/Core/Public/Windows/PostWindowsApi.h b/Engine/Source/Runtime/Core/Public/Windows/PostWindowsApi.h
--- a/Engine/Source/Runtime/Core/Public/Windows/PostWindowsApi.h
+++ b/Engine/Source/Runtime/Core/Public/Windows/PostWindowsApi.h
@@ -35,6 +35,7 @@ THIRD_PARTY_INCLUDES_END
 #undef CreateFont
 #undef CreateDesktop
 #undef GetMessage
+#undef PostMessage
 #undef GetCommandLine
 #undef GetProp
 #undef SetProp
diff --git a/Engine/Source/Runtime/D3D12RHI/Private/Windows/WindowsD3D12Viewport.cpp b/Engine/Source/Runtime/D3D12RHI/Private/Windows/WindowsD3D12Viewport.cpp
--- a/Engine/Source/Runtime/D3D12RHI/Private/Windows/WindowsD3D12Viewport.cpp
+++ b/Engine/Source/Runtime/D3D12RHI/Private/Windows/WindowsD3D12Viewport.cpp
@@ -214,7 +214,7 @@ void FD3D12Viewport::Init()
 #if !PLATFORM_HOLOLENS
 	// Tell the window to redraw when they can.
 	// @todo: For Slate viewports, it doesn't make sense to post WM_PAINT messages (we swallow those.)
-	::PostMessage(WindowHandle, WM_PAINT, 0, 0);
+	::PostMessageW(WindowHandle, WM_PAINT, 0, 0);
 #endif
 	// : END HoloLens support
 }
diff --git a/Engine/Source/Runtime/Experimental/InteractiveToolsFramework/Private/BaseGizmos/TransformGizmo.cpp b/Engine/Source/Runtime/Experimental/InteractiveToolsFramework/Private/BaseGizmos/TransformGizmo.cpp
--- a/Engine/Source/Runtime/Experimental/InteractiveToolsFramework/Private/BaseGizmos/TransformGizmo.cpp
+++ b/Engine/Source/Runtime/Experimental/InteractiveToolsFramework/Private/BaseGizmos/TransformGizmo.cpp
@@ -205,7 +205,7 @@ void UTransformGizmo::SetActiveTarget(UTransformProxy* Target)
 	// target tracks location of GizmoComponent
 	GizmoActor->GetRootComponent()->TransformUpdated.AddLambda(
 		[this, SaveScale](USceneComponent* Component, EUpdateTransformFlags /*UpdateTransformFlags*/, ETeleportType /*Teleport*/) {
-		//this->GetGizmoManager()->PostMessage(TEXT("TRANSFORM UPDATED"), EToolMessageLevel::Internal);
+		//this->GetGizmoManager()->DisplayMessage(TEXT("TRANSFORM UPDATED"), EToolMessageLevel::Internal);
 		FTransform NewXForm = Component->GetComponentToWorld();
 		NewXForm.SetScale3D(SaveScale);
 		this->ActiveTarget->SetTransform(NewXForm);
diff --git a/Engine/Source/Runtime/Experimental/InteractiveToolsFramework/Private/BaseTools/ClickDragTool.cpp b/Engine/Source/Runtime/Experimental/InteractiveToolsFramework/Private/BaseTools/ClickDragTool.cpp
--- a/Engine/Source/Runtime/Experimental/InteractiveToolsFramework/Private/BaseTools/ClickDragTool.cpp
+++ b/Engine/Source/Runtime/Experimental/InteractiveToolsFramework/Private/BaseTools/ClickDragTool.cpp
@@ -50,7 +50,7 @@ FInputRayHit UClickDragTool::CanBeginClickDragSequence(const FInputDeviceRay& Cl
 void UClickDragTool::OnClickPress(const FInputDeviceRay& ClickPos)
 {
 	// print debug message
-	GetToolManager()->PostMessage( 
+	GetToolManager()->DisplayMessage( 
 		FText::Format(LOCTEXT("OnClickPressMessage", "UClickDragTool::OnClickPress: clicked at ({0},{1})"),
 			FText::AsNumber(ClickPos.ScreenPosition.X), FText::AsNumber(ClickPos.ScreenPosition.Y)),
 		EToolMessageLevel::Internal);
@@ -64,7 +64,7 @@ void UClickDragTool::OnClickDrag(const FInputDeviceRay& ClickPos)
 void UClickDragTool::OnClickRelease(const FInputDeviceRay& ClickPos)
 {
 	// print debug message
-	GetToolManager()->PostMessage(
+	GetToolManager()->DisplayMessage(
 		FText::Format(LOCTEXT("OnClickReleaseMessage", "UClickDragTool::OnClickRelease: released at ({0},{1})"),
 			FText::AsNumber(ClickPos.ScreenPosition.X), FText::AsNumber(ClickPos.ScreenPosition.Y)),
 		EToolMessageLevel::Internal);
diff --git a/Engine/Source/Runtime/Experimental/InteractiveToolsFramework/Private/BaseTools/MeshSurfacePointTool.cpp b/Engine/Source/Runtime/Experimental/InteractiveToolsFramework/Private/BaseTools/MeshSurfacePointTool.cpp
--- a/Engine/Source/Runtime/Experimental/InteractiveToolsFramework/Private/BaseTools/MeshSurfacePointTool.cpp
+++ b/Engine/Source/Runtime/Experimental/InteractiveToolsFramework/Private/BaseTools/MeshSurfacePointTool.cpp
@@ -76,7 +76,7 @@ void UMeshSurfacePointTool::OnUpdateDrag(const FRay& Ray)
 	FHitResult OutHit;
 	if ( HitTest(Ray, OutHit) ) 
 	{
-		GetToolManager()->PostMessage( 
+		GetToolManager()->DisplayMessage( 
 			FText::Format(LOCTEXT("OnUpdateDragMessage", "UMeshSurfacePointTool::OnUpdateDrag: Hit triangle index {0} at ray distance {1}"),
 				FText::AsNumber(OutHit.FaceIndex), FText::AsNumber(OutHit.Distance)),
 			EToolMessageLevel::Internal);
@@ -85,7 +85,7 @@ void UMeshSurfacePointTool::OnUpdateDrag(const FRay& Ray)
 
 void UMeshSurfacePointTool::OnEndDrag(const FRay& Ray)
 {
-	//GetToolManager()->PostMessage(TEXT("UMeshSurfacePointTool::OnEndDrag!"), EToolMessageLevel::Internal);
+	//GetToolManager()->DisplayMessage(TEXT("UMeshSurfacePointTool::OnEndDrag!"), EToolMessageLevel::Internal);
 }
 
 
diff --git a/Engine/Source/Runtime/Experimental/InteractiveToolsFramework/Private/BaseTools/SingleClickTool.cpp b/Engine/Source/Runtime/Experimental/InteractiveToolsFramework/Private/BaseTools/SingleClickTool.cpp
--- a/Engine/Source/Runtime/Experimental/InteractiveToolsFramework/Private/BaseTools/SingleClickTool.cpp
+++ b/Engine/Source/Runtime/Experimental/InteractiveToolsFramework/Private/BaseTools/SingleClickTool.cpp
@@ -49,7 +49,7 @@ bool USingleClickTool::IsHitByClick(const FInputDeviceRay& ClickPos)
 void USingleClickTool::OnClicked(const FInputDeviceRay& ClickPos)
 {
 	// print debug message
-	GetToolManager()->PostMessage(
+	GetToolManager()->DisplayMessage(
 		FText::Format(LOCTEXT("OnClickedMessage", "USingleClickTool::OnClicked at ({0},{1})"),
 			FText::AsNumber(ClickPos.ScreenPosition.X), FText::AsNumber(ClickPos.ScreenPosition.Y)),
 		EToolMessageLevel::Internal );
diff --git a/Engine/Source/Runtime/Experimental/InteractiveToolsFramework/Private/InputRouter.cpp b/Engine/Source/Runtime/Experimental/InteractiveToolsFramework/Private/InputRouter.cpp
--- a/Engine/Source/Runtime/Experimental/InteractiveToolsFramework/Private/InputRouter.cpp
+++ b/Engine/Source/Runtime/Experimental/InteractiveToolsFramework/Private/InputRouter.cpp
@@ -67,7 +67,7 @@ void UInputRouter::PostInputEvent(const FInputDeviceState& Input)
 	else
 	{
 		unimplemented();
-		TransactionsAPI->PostMessage(LOCTEXT("PostInputEventMessage", "UInteractiveToolManager::PostInputEvent - input device is not currently supported."), EToolMessageLevel::Internal);
+		TransactionsAPI->DisplayMessage(LOCTEXT("PostInputEventMessage", "UInteractiveToolManager::PostInputEvent - input device is not currently supported."), EToolMessageLevel::Internal);
 		return;
 	}
 }
@@ -138,7 +138,7 @@ void UInputRouter::HandleCapturedKeyboardInput(const FInputDeviceState& Input)
 	}
 	else if (Result.State != EInputCaptureState::Continue)
 	{
-		TransactionsAPI->PostMessage(LOCTEXT("HandleCapturedKeyboardInputMessage", "UInteractiveToolManager::HandleCapturedKeyboardInput - unexpected capture state!"), EToolMessageLevel::Internal);
+		TransactionsAPI->DisplayMessage(LOCTEXT("HandleCapturedKeyboardInputMessage", "UInteractiveToolManager::HandleCapturedKeyboardInput - unexpected capture state!"), EToolMessageLevel::Internal);
 	}
 
 	if (bAutoInvalidateOnCapture)
@@ -250,7 +250,7 @@ void UInputRouter::HandleCapturedMouseInput(const FInputDeviceState& Input)
 	}
 	else if (Result.State != EInputCaptureState::Continue)
 	{
-		TransactionsAPI->PostMessage(LOCTEXT("HandleCapturedMouseInputMessage", "UInteractiveToolManager::HandleCapturedMouseInput - unexpected capture state!"), EToolMessageLevel::Internal);
+		TransactionsAPI->DisplayMessage(LOCTEXT("HandleCapturedMouseInputMessage", "UInteractiveToolManager::HandleCapturedMouseInput - unexpected capture state!"), EToolMessageLevel::Internal);
 	}
 
 	if (bAutoInvalidateOnCapture)
diff --git a/Engine/Source/Runtime/Experimental/InteractiveToolsFramework/Private/InteractiveGizmoManager.cpp b/Engine/Source/Runtime/Experimental/InteractiveToolsFramework/Private/InteractiveGizmoManager.cpp
--- a/Engine/Source/Runtime/Experimental/InteractiveToolsFramework/Private/InteractiveGizmoManager.cpp
+++ b/Engine/Source/Runtime/Experimental/InteractiveToolsFramework/Private/InteractiveGizmoManager.cpp
@@ -63,7 +63,7 @@ bool UInteractiveGizmoManager::DeregisterGizmoType(const FString& BuilderIdentif
 {
 	if (GizmoBuilders.Contains(BuilderIdentifier) == false)
 	{
-		PostMessage(
+		DisplayMessage(
 			FText::Format(LOCTEXT("DeregisterFailedMessage", "UInteractiveGizmoManager::DeregisterGizmoType: could not find requested type {0}"), FText::FromString(BuilderIdentifier) ),
 			EToolMessageLevel::Internal);
 		return false;
@@ -79,7 +79,7 @@ UInteractiveGizmo* UInteractiveGizmoManager::CreateGizmo(const FString& BuilderI
 {
 	if ( GizmoBuilders.Contains(BuilderIdentifier) == false )
 	{
-		PostMessage(
+		DisplayMessage(
 			FText::Format(LOCTEXT("CreateGizmoCannotFindFailedMessage", "UInteractiveGizmoManager::CreateGizmo: could not find requested type {0}"), FText::FromString(BuilderIdentifier) ),
 			EToolMessageLevel::Internal);
 		return nullptr;
@@ -93,7 +93,7 @@ UInteractiveGizmo* UInteractiveGizmoManager::CreateGizmo(const FString& BuilderI
 		{
 			if (ActiveGizmo.InstanceIdentifier == InstanceIdentifier)
 			{
-				PostMessage(
+				DisplayMessage(
 					FText::Format(LOCTEXT("CreateGizmoExistsMessage", "UInteractiveGizmoManager::CreateGizmo: instance identifier {0} already in use!"), FText::FromString(InstanceIdentifier) ),
 					EToolMessageLevel::Internal);
 				return nullptr;
@@ -107,7 +107,7 @@ UInteractiveGizmo* UInteractiveGizmoManager::CreateGizmo(const FString& BuilderI
 	UInteractiveGizmo* NewGizmo = FoundBuilder->BuildGizmo(CurrentSceneState);
 	if (NewGizmo == nullptr)
 	{
-		PostMessage(LOCTEXT("CreateGizmoReturnNullMessage", "UInteractiveGizmoManager::CreateGizmo: BuildGizmo() returned null"), EToolMessageLevel::Internal);
+		DisplayMessage(LOCTEXT("CreateGizmoReturnNullMessage", "UInteractiveGizmoManager::CreateGizmo: BuildGizmo() returned null"), EToolMessageLevel::Internal);
 		return nullptr;
 	}
 
@@ -233,9 +233,9 @@ void UInteractiveGizmoManager::Render(IToolsContextRenderAPI* RenderAPI)
 
 }
 
-void UInteractiveGizmoManager::PostMessage(const FText& Message, EToolMessageLevel Level)
+void UInteractiveGizmoManager::DisplayMessage(const FText& Message, EToolMessageLevel Level)
 {
-	TransactionsAPI->PostMessage(Message, Level);
+	TransactionsAPI->DisplayMessage(Message, Level);
 }
 
 void UInteractiveGizmoManager::PostInvalidation()
diff --git a/Engine/Source/Runtime/Experimental/InteractiveToolsFramework/Private/InteractiveToolManager.cpp b/Engine/Source/Runtime/Experimental/InteractiveToolsFramework/Private/InteractiveToolManager.cpp
--- a/Engine/Source/Runtime/Experimental/InteractiveToolsFramework/Private/InteractiveToolManager.cpp
+++ b/Engine/Source/Runtime/Experimental/InteractiveToolsFramework/Private/InteractiveToolManager.cpp
@@ -118,7 +118,7 @@ bool UInteractiveToolManager::ActivateTool(EToolSide Side)
 
 	if (ActiveLeftBuilder->CanBuildTool(InputState) == false)
 	{
-		TransactionsAPI->PostMessage( LOCTEXT("ActivateToolCanBuildFailMessage", "UInteractiveToolManager::ActivateTool: CanBuildTool returned false."), EToolMessageLevel::Internal);
+		TransactionsAPI->DisplayMessage( LOCTEXT("ActivateToolCanBuildFailMessage", "UInteractiveToolManager::ActivateTool: CanBuildTool returned false."), EToolMessageLevel::Internal);
 		return false;
 	}
 
@@ -240,9 +240,9 @@ UInteractiveGizmoManager* UInteractiveToolManager::GetPairedGizmoManager()
 	return Cast<UInteractiveToolsContext>(GetOuter())->GizmoManager;
 }
 
-void UInteractiveToolManager::PostMessage(const FText& Message, EToolMessageLevel Level)
+void UInteractiveToolManager::DisplayMessage(const FText& Message, EToolMessageLevel Level)
 {
-	TransactionsAPI->PostMessage(Message, Level);
+	TransactionsAPI->DisplayMessage(Message, Level);
 }
 
 void UInteractiveToolManager::PostInvalidation()
diff --git a/Engine/Source/Runtime/Experimental/InteractiveToolsFramework/Public/InteractiveGizmoManager.h b/Engine/Source/Runtime/Experimental/InteractiveToolsFramework/Public/InteractiveGizmoManager.h
--- a/Engine/Source/Runtime/Experimental/InteractiveToolsFramework/Public/InteractiveGizmoManager.h
+++ b/Engine/Source/Runtime/Experimental/InteractiveToolsFramework/Public/InteractiveGizmoManager.h
@@ -122,7 +122,7 @@ public:
 	//
 	
 	/** Post a message via the Transactions API */
-	virtual void PostMessage(const FText& Message, EToolMessageLevel Level);
+	virtual void DisplayMessage(const FText& Message, EToolMessageLevel Level);
 
 	/** Request an Invalidation via the Transactions API (ie to cause a repaint, etc) */
 	virtual void PostInvalidation();
diff --git a/Engine/Source/Runtime/Experimental/InteractiveToolsFramework/Public/InteractiveToolManager.h b/Engine/Source/Runtime/Experimental/InteractiveToolsFramework/Public/InteractiveToolManager.h
--- a/Engine/Source/Runtime/Experimental/InteractiveToolsFramework/Public/InteractiveToolManager.h
+++ b/Engine/Source/Runtime/Experimental/InteractiveToolsFramework/Public/InteractiveToolManager.h
@@ -141,7 +141,7 @@ public:
 	//
 	
 	/** Post a message via the Transactions API */
-	virtual void PostMessage(const FText& Message, EToolMessageLevel Level);
+	virtual void DisplayMessage(const FText& Message, EToolMessageLevel Level);
 
 	/** Request an Invalidation via the Transactions API (ie to cause a repaint, etc) */
 	virtual void PostInvalidation();
diff --git a/Engine/Source/Runtime/Experimental/InteractiveToolsFramework/Public/ToolContextInterfaces.h b/Engine/Source/Runtime/Experimental/InteractiveToolsFramework/Public/ToolContextInterfaces.h
--- a/Engine/Source/Runtime/Experimental/InteractiveToolsFramework/Public/ToolContextInterfaces.h
+++ b/Engine/Source/Runtime/Experimental/InteractiveToolsFramework/Public/ToolContextInterfaces.h
@@ -240,7 +240,7 @@ public:
 	 * @param Message text of message
 	 * @param Level severity level of message
 	 */
-	virtual void PostMessage(const FText& Message, EToolMessageLevel Level) = 0;
+	virtual void DisplayMessage(const FText& Message, EToolMessageLevel Level) = 0;
 
 	/** 
 	 * Forward an invalidation request from Tools framework, to cause repaint/etc. 
diff --git a/Engine/Source/Runtime/Windows/D3D11RHI/Private/Windows/WindowsD3D11Viewport.cpp b/Engine/Source/Runtime/Windows/D3D11RHI/Private/Windows/WindowsD3D11Viewport.cpp
--- a/Engine/Source/Runtime/Windows/D3D11RHI/Private/Windows/WindowsD3D11Viewport.cpp
+++ b/Engine/Source/Runtime/Windows/D3D11RHI/Private/Windows/WindowsD3D11Viewport.cpp
@@ -173,7 +173,7 @@ FD3D11Viewport::FD3D11Viewport(FD3D11DynamicRHI* InD3DRHI,HWND InWindowHandle,ui
 
 	// Tell the window to redraw when they can.
 	// @todo: For Slate viewports, it doesn't make sense to post WM_PAINT messages (we swallow those.)
-	::PostMessage( WindowHandle, WM_PAINT, 0, 0 );
+	::PostMessageW( WindowHandle, WM_PAINT, 0, 0 );
 
 	BeginInitResource(&FrameSyncEvent);
 }
