ToolsFramework: expand InputBehavior/InputRouter hover implementation so it has hit-test/capture/release like mouse capture. This is required to resolve multiple hoverable objects under. Update all behaviors/tools with new API.
#rnx

diff --git a/Engine/Plugins/Experimental/MeshModelingToolset/Source/MeshModelingTools/Private/AddPrimitiveTool.cpp b/Engine/Plugins/Experimental/MeshModelingToolset/Source/MeshModelingTools/Private/AddPrimitiveTool.cpp
--- a/Engine/Plugins/Experimental/MeshModelingToolset/Source/MeshModelingTools/Private/AddPrimitiveTool.cpp
+++ b/Engine/Plugins/Experimental/MeshModelingToolset/Source/MeshModelingTools/Private/AddPrimitiveTool.cpp
@@ -155,11 +155,28 @@ void UAddPrimitiveTool::OnPropertyModified(UObject* PropertySet, UProperty* Prop
 
 
 
-void UAddPrimitiveTool::OnUpdateHover(const FInputDeviceRay& DevicePos)
+FInputRayHit UAddPrimitiveTool::BeginHoverSequenceHitTest(const FInputDeviceRay& PressPos)
+{
+	return FInputRayHit(0.0f);		// always hit in hover 
+}
+
+void UAddPrimitiveTool::OnBeginHover(const FInputDeviceRay& DevicePos)
 {
 	UpdatePreviewPosition(DevicePos);
 }
 
+bool UAddPrimitiveTool::OnUpdateHover(const FInputDeviceRay& DevicePos)
+{
+	UpdatePreviewPosition(DevicePos);
+	return true;
+}
+
+void UAddPrimitiveTool::OnEndHover()
+{
+	// do nothing
+}
+
+
 
 void UAddPrimitiveTool::UpdatePreviewPosition(const FInputDeviceRay& DeviceClickPos)
 {
diff --git a/Engine/Plugins/Experimental/MeshModelingToolset/Source/MeshModelingTools/Private/BaseBrushTool.cpp b/Engine/Plugins/Experimental/MeshModelingToolset/Source/MeshModelingTools/Private/BaseBrushTool.cpp
--- a/Engine/Plugins/Experimental/MeshModelingToolset/Source/MeshModelingTools/Private/BaseBrushTool.cpp
+++ b/Engine/Plugins/Experimental/MeshModelingToolset/Source/MeshModelingTools/Private/BaseBrushTool.cpp
@@ -152,7 +152,7 @@ void UBaseBrushTool::OnEndDrag(const FRay& Ray)
 }
 
 
-void UBaseBrushTool::OnUpdateHover(const FInputDeviceRay& DevicePos)
+bool UBaseBrushTool::OnUpdateHover(const FInputDeviceRay& DevicePos)
 {
 	FHitResult OutHit;
 	if (HitTest(DevicePos.WorldRay, OutHit))
@@ -162,6 +162,7 @@ void UBaseBrushTool::OnUpdateHover(const FInputDeviceRay& DevicePos)
 		LastBrushStamp.WorldNormal = OutHit.Normal;
 		LastBrushStamp.HitResult = OutHit;
 	}
+	return true;
 }
 
 
diff --git a/Engine/Plugins/Experimental/MeshModelingToolset/Source/MeshModelingTools/Private/DynamicMeshSculptTool.cpp b/Engine/Plugins/Experimental/MeshModelingToolset/Source/MeshModelingTools/Private/DynamicMeshSculptTool.cpp
--- a/Engine/Plugins/Experimental/MeshModelingToolset/Source/MeshModelingTools/Private/DynamicMeshSculptTool.cpp
+++ b/Engine/Plugins/Experimental/MeshModelingToolset/Source/MeshModelingTools/Private/DynamicMeshSculptTool.cpp
@@ -659,7 +659,7 @@ void UDynamicMeshSculptTool::OnEndDrag(const FRay& Ray)
 }
 
 
-void UDynamicMeshSculptTool::OnUpdateHover(const FInputDeviceRay& DevicePos)
+bool UDynamicMeshSculptTool::OnUpdateHover(const FInputDeviceRay& DevicePos)
 {
 	if (bInDrag)
 	{
@@ -676,6 +676,7 @@ void UDynamicMeshSculptTool::OnUpdateHover(const FInputDeviceRay& DevicePos)
 			LastBrushPosNormalWorld = OutHit.Normal;
 		}
 	}
+	return true;
 }
 
 
diff --git a/Engine/Plugins/Experimental/MeshModelingToolset/Source/MeshModelingTools/Private/EditMeshPolygonsTool.cpp b/Engine/Plugins/Experimental/MeshModelingToolset/Source/MeshModelingTools/Private/EditMeshPolygonsTool.cpp
--- a/Engine/Plugins/Experimental/MeshModelingToolset/Source/MeshModelingTools/Private/EditMeshPolygonsTool.cpp
+++ b/Engine/Plugins/Experimental/MeshModelingToolset/Source/MeshModelingTools/Private/EditMeshPolygonsTool.cpp
@@ -1043,7 +1043,7 @@ void UEditMeshPolygonsTool::OnEndDrag(const FRay& Ray)
 
 
 
-void UEditMeshPolygonsTool::OnUpdateHover(const FInputDeviceRay& DevicePos)
+bool UEditMeshPolygonsTool::OnUpdateHover(const FInputDeviceRay& DevicePos)
 {
 	//if (!bNeedEmitEndChange)
 	if (ActiveVertexChange == nullptr)
@@ -1067,6 +1067,7 @@ void UEditMeshPolygonsTool::OnUpdateHover(const FInputDeviceRay& DevicePos)
 			UpdateQuickTransformer();
 		}
 	}
+	return true;
 }
 
 
diff --git a/Engine/Plugins/Experimental/MeshModelingToolset/Source/MeshModelingTools/Private/MeshSelectionTool.cpp b/Engine/Plugins/Experimental/MeshModelingToolset/Source/MeshModelingTools/Private/MeshSelectionTool.cpp
--- a/Engine/Plugins/Experimental/MeshModelingToolset/Source/MeshModelingTools/Private/MeshSelectionTool.cpp
+++ b/Engine/Plugins/Experimental/MeshModelingToolset/Source/MeshModelingTools/Private/MeshSelectionTool.cpp
@@ -284,7 +284,7 @@ void UMeshSelectionTool::OnEndDrag(const FRay& Ray)
 }
 
 
-void UMeshSelectionTool::OnUpdateHover(const FInputDeviceRay& DevicePos)
+bool UMeshSelectionTool::OnUpdateHover(const FInputDeviceRay& DevicePos)
 {
 	UDynamicMeshBrushTool::OnUpdateHover(DevicePos);
 
@@ -302,6 +302,8 @@ void UMeshSelectionTool::OnUpdateHover(const FInputDeviceRay& DevicePos)
 			CalculateVertexROI(LastBrushStamp, PreviewBrushROI);
 		}
 	}
+
+	return true;
 }
 
 
diff --git a/Engine/Plugins/Experimental/MeshModelingToolset/Source/MeshModelingTools/Private/MeshSpaceDeformerTool.cpp b/Engine/Plugins/Experimental/MeshModelingToolset/Source/MeshModelingTools/Private/MeshSpaceDeformerTool.cpp
--- a/Engine/Plugins/Experimental/MeshModelingToolset/Source/MeshModelingTools/Private/MeshSpaceDeformerTool.cpp
+++ b/Engine/Plugins/Experimental/MeshModelingToolset/Source/MeshModelingTools/Private/MeshSpaceDeformerTool.cpp
@@ -399,7 +399,7 @@ void UMeshSpaceDeformerTool::UpdateIntervalsFromDrag()
 	}
 }
 
-void UMeshSpaceDeformerTool::OnUpdateHover(const FInputDeviceRay& DevicePos)
+bool UMeshSpaceDeformerTool::OnUpdateHover(const FInputDeviceRay& DevicePos)
 {
 	if (!bInDrag)
 	{
@@ -408,6 +408,7 @@ void UMeshSpaceDeformerTool::OnUpdateHover(const FInputDeviceRay& DevicePos)
 		};
 		Handle.UpdateHover(DevicePos, ComponentTarget->GetWorldTransform(), TFunction<bool(const FVector3d&, const FVector3d&)>{ToleranceTest});
 	}
+	return true;
 }
 
 #if WITH_EDITOR
diff --git a/Engine/Plugins/Experimental/MeshModelingToolset/Source/MeshModelingTools/Public/AddPrimitiveTool.h b/Engine/Plugins/Experimental/MeshModelingToolset/Source/MeshModelingTools/Public/AddPrimitiveTool.h
--- a/Engine/Plugins/Experimental/MeshModelingToolset/Source/MeshModelingTools/Public/AddPrimitiveTool.h
+++ b/Engine/Plugins/Experimental/MeshModelingToolset/Source/MeshModelingTools/Public/AddPrimitiveTool.h
@@ -161,10 +161,15 @@ public:
 
 	virtual void OnPropertyModified(UObject* PropertySet, UProperty* Property) override;
 
-	virtual void OnUpdateHover(const FInputDeviceRay& DevicePos) override;
 	virtual void OnClicked(const FInputDeviceRay& ClickPos) override;
 
 
+	// IHoverBehaviorTarget interface
+	virtual FInputRayHit BeginHoverSequenceHitTest(const FInputDeviceRay& PressPos) override;
+	virtual void OnBeginHover(const FInputDeviceRay& DevicePos) override;
+	virtual bool OnUpdateHover(const FInputDeviceRay& DevicePos) override;
+	virtual void OnEndHover() override;
+
 
 protected:
 	UPROPERTY()
diff --git a/Engine/Plugins/Experimental/MeshModelingToolset/Source/MeshModelingTools/Public/BaseBrushTool.h b/Engine/Plugins/Experimental/MeshModelingToolset/Source/MeshModelingTools/Public/BaseBrushTool.h
--- a/Engine/Plugins/Experimental/MeshModelingToolset/Source/MeshModelingTools/Public/BaseBrushTool.h
+++ b/Engine/Plugins/Experimental/MeshModelingToolset/Source/MeshModelingTools/Public/BaseBrushTool.h
@@ -88,7 +88,7 @@ public:
 	virtual void OnBeginDrag(const FRay& Ray) override;
 	virtual void OnUpdateDrag(const FRay& Ray) override;
 	virtual void OnEndDrag(const FRay& Ray) override;
-	virtual void OnUpdateHover(const FInputDeviceRay& DevicePos) override;
+	virtual bool OnUpdateHover(const FInputDeviceRay& DevicePos) override;
 
 public:
 
diff --git a/Engine/Plugins/Experimental/MeshModelingToolset/Source/MeshModelingTools/Public/DynamicMeshSculptTool.h b/Engine/Plugins/Experimental/MeshModelingToolset/Source/MeshModelingTools/Public/DynamicMeshSculptTool.h
--- a/Engine/Plugins/Experimental/MeshModelingToolset/Source/MeshModelingTools/Public/DynamicMeshSculptTool.h
+++ b/Engine/Plugins/Experimental/MeshModelingToolset/Source/MeshModelingTools/Public/DynamicMeshSculptTool.h
@@ -198,7 +198,7 @@ public:
 	virtual void OnBeginDrag(const FRay& Ray) override;
 	virtual void OnUpdateDrag(const FRay& Ray) override;
 	virtual void OnEndDrag(const FRay& Ray) override;
-	virtual void OnUpdateHover(const FInputDeviceRay& DevicePos) override;
+	virtual bool OnUpdateHover(const FInputDeviceRay& DevicePos) override;
 
 	virtual void SetEnableRemeshing(bool bEnable) { bEnableRemeshing = bEnable; }
 	virtual bool GetEnableRemeshing() const { return bEnableRemeshing; }
diff --git a/Engine/Plugins/Experimental/MeshModelingToolset/Source/MeshModelingTools/Public/EditMeshPolygonsTool.h b/Engine/Plugins/Experimental/MeshModelingToolset/Source/MeshModelingTools/Public/EditMeshPolygonsTool.h
--- a/Engine/Plugins/Experimental/MeshModelingToolset/Source/MeshModelingTools/Public/EditMeshPolygonsTool.h
+++ b/Engine/Plugins/Experimental/MeshModelingToolset/Source/MeshModelingTools/Public/EditMeshPolygonsTool.h
@@ -423,7 +423,7 @@ public:
 	virtual void OnBeginDrag(const FRay& Ray) override;
 	virtual void OnUpdateDrag(const FRay& Ray) override;
 	virtual void OnEndDrag(const FRay& Ray) override;
-	virtual void OnUpdateHover(const FInputDeviceRay& DevicePos) override;
+	virtual bool OnUpdateHover(const FInputDeviceRay& DevicePos) override;
 
 public:
 	virtual void ToggleLocalTransformFrameAction();
diff --git a/Engine/Plugins/Experimental/MeshModelingToolset/Source/MeshModelingTools/Public/MeshSelectionTool.h b/Engine/Plugins/Experimental/MeshModelingToolset/Source/MeshModelingTools/Public/MeshSelectionTool.h
--- a/Engine/Plugins/Experimental/MeshModelingToolset/Source/MeshModelingTools/Public/MeshSelectionTool.h
+++ b/Engine/Plugins/Experimental/MeshModelingToolset/Source/MeshModelingTools/Public/MeshSelectionTool.h
@@ -55,7 +55,7 @@ public:
 	virtual void OnBeginDrag(const FRay& Ray) override;
 	virtual void OnUpdateDrag(const FRay& Ray) override;
 	virtual void OnEndDrag(const FRay& Ray) override;
-	virtual void OnUpdateHover(const FInputDeviceRay& DevicePos) override;
+	virtual bool OnUpdateHover(const FInputDeviceRay& DevicePos) override;
 
 public:
 
diff --git a/Engine/Plugins/Experimental/MeshModelingToolset/Source/MeshModelingTools/Public/MeshSpaceDeformerTool.h b/Engine/Plugins/Experimental/MeshModelingToolset/Source/MeshModelingTools/Public/MeshSpaceDeformerTool.h
--- a/Engine/Plugins/Experimental/MeshModelingToolset/Source/MeshModelingTools/Public/MeshSpaceDeformerTool.h
+++ b/Engine/Plugins/Experimental/MeshModelingToolset/Source/MeshModelingTools/Public/MeshSpaceDeformerTool.h
@@ -359,7 +359,7 @@ public:
 	virtual void OnUpdateDrag(const FRay& Ray) override;
 	virtual void OnEndDrag(const FRay& Ray) override;
 
-	virtual void OnUpdateHover(const FInputDeviceRay& DevicePos) override;
+	virtual bool OnUpdateHover(const FInputDeviceRay& DevicePos) override;
 
 #if WITH_EDITOR
 	virtual void PostEditChangeProperty(FPropertyChangedEvent& PropertyChangedEvent) override;
diff --git a/Engine/Source/Runtime/Experimental/InteractiveToolsFramework/Private/BaseBehaviors/MouseHoverBehavior.cpp b/Engine/Source/Runtime/Experimental/InteractiveToolsFramework/Private/BaseBehaviors/MouseHoverBehavior.cpp
--- a/Engine/Source/Runtime/Experimental/InteractiveToolsFramework/Private/BaseBehaviors/MouseHoverBehavior.cpp
+++ b/Engine/Source/Runtime/Experimental/InteractiveToolsFramework/Private/BaseBehaviors/MouseHoverBehavior.cpp
@@ -23,17 +23,46 @@ bool UMouseHoverBehavior::WantsHoverEvents()
 	return true;
 }
 
-void UMouseHoverBehavior::UpdateHover(const FInputDeviceState& input)
+
+FInputCaptureRequest UMouseHoverBehavior::WantsHoverCapture(const FInputDeviceState& InputState)
 {
 	if (Target != nullptr)
 	{
-		Modifiers.UpdateModifiers(input, Target);
+		Modifiers.UpdateModifiers(InputState, Target);
 
-		Target->OnUpdateHover( FInputDeviceRay(input.Mouse.WorldRay, input.Mouse.Position2D) );
+		FInputRayHit Hit = Target->BeginHoverSequenceHitTest(InputState.Mouse.WorldRay);
+		if (Hit.bHit)
+		{
+			return FInputCaptureRequest::Begin(this, EInputCaptureSide::Any, Hit.HitDepth);
+		}
 	}
+	return FInputCaptureRequest::Ignore();
+}
+
+FInputCaptureUpdate UMouseHoverBehavior::BeginHoverCapture(const FInputDeviceState& InputState, EInputCaptureSide eSide)
+{
+	check(Target != nullptr);
+	Modifiers.UpdateModifiers(InputState, Target);
+	Target->OnBeginHover(InputState.Mouse.WorldRay);
+	return FInputCaptureUpdate::Begin(this, eSide);
 }
 
-void UMouseHoverBehavior::EndHover(const FInputDeviceState& input)
+FInputCaptureUpdate UMouseHoverBehavior::UpdateHoverCapture(const FInputDeviceState& InputState)
 {
+	check(Target != nullptr);
+	Modifiers.UpdateModifiers(InputState, Target);
+	if (Target->OnUpdateHover(InputState.Mouse.WorldRay))
+	{
+		return FInputCaptureUpdate::Continue();
+	}
+	return FInputCaptureUpdate::End();
 }
 
+void UMouseHoverBehavior::EndHoverCapture()
+{
+	check(Target != nullptr);
+	Target->OnEndHover();
+}
+
+
+
diff --git a/Engine/Source/Runtime/Experimental/InteractiveToolsFramework/Private/BaseBehaviors/MultiClickSequenceInputBehavior.cpp b/Engine/Source/Runtime/Experimental/InteractiveToolsFramework/Private/BaseBehaviors/MultiClickSequenceInputBehavior.cpp
--- a/Engine/Source/Runtime/Experimental/InteractiveToolsFramework/Private/BaseBehaviors/MultiClickSequenceInputBehavior.cpp
+++ b/Engine/Source/Runtime/Experimental/InteractiveToolsFramework/Private/BaseBehaviors/MultiClickSequenceInputBehavior.cpp
@@ -88,15 +88,36 @@ bool UMultiClickSequenceInputBehavior::WantsHoverEvents()
 	return true;
 }
 
-void UMultiClickSequenceInputBehavior::UpdateHover(const FInputDeviceState& Input)
+
+FInputCaptureRequest UMultiClickSequenceInputBehavior::WantsHoverCapture(const FInputDeviceState& InputState)
+{
+	return FInputCaptureRequest::Begin(this, EInputCaptureSide::Any);
+}
+
+FInputCaptureUpdate UMultiClickSequenceInputBehavior::BeginHoverCapture(const FInputDeviceState& InputState, EInputCaptureSide eSide) 
+{
+	if (Target != nullptr)
+	{
+		Modifiers.UpdateModifiers(InputState, Target);
+		Target->OnBeginSequencePreview(FInputDeviceRay(InputState.Mouse.WorldRay, InputState.Mouse.Position2D));
+		return FInputCaptureUpdate::Begin(this, EInputCaptureSide::Any);
+	}
+	return FInputCaptureUpdate::Ignore();
+}
+
+FInputCaptureUpdate UMultiClickSequenceInputBehavior::UpdateHoverCapture(const FInputDeviceState& InputState)
 {
 	if (Target != nullptr)
 	{
-		Modifiers.UpdateModifiers(Input, Target);
-		Target->OnBeginSequencePreview(FInputDeviceRay(Input.Mouse.WorldRay, Input.Mouse.Position2D));
+		Modifiers.UpdateModifiers(InputState, Target);
+		Target->OnBeginSequencePreview(FInputDeviceRay(InputState.Mouse.WorldRay, InputState.Mouse.Position2D));
+		return FInputCaptureUpdate::Continue();
 	}
+	return FInputCaptureUpdate::End();
 }
 
-void UMultiClickSequenceInputBehavior::EndHover(const FInputDeviceState& input)
+
+void UMultiClickSequenceInputBehavior::EndHoverCapture()
 {
-}
\ No newline at end of file
+	// end
+}
diff --git a/Engine/Source/Runtime/Experimental/InteractiveToolsFramework/Private/BaseGizmos/AxisAngleGizmo.cpp b/Engine/Source/Runtime/Experimental/InteractiveToolsFramework/Private/BaseGizmos/AxisAngleGizmo.cpp
--- a/Engine/Source/Runtime/Experimental/InteractiveToolsFramework/Private/BaseGizmos/AxisAngleGizmo.cpp
+++ b/Engine/Source/Runtime/Experimental/InteractiveToolsFramework/Private/BaseGizmos/AxisAngleGizmo.cpp
@@ -150,18 +150,32 @@ void UAxisAngleGizmo::OnTerminateDragSequence()
 
 
 
-void UAxisAngleGizmo::OnUpdateHover(const FInputDeviceRay& DevicePos)
+
+FInputRayHit UAxisAngleGizmo::BeginHoverSequenceHitTest(const FInputDeviceRay& PressPos)
 {
 	FInputRayHit GizmoHit;
 	if (HitTarget)
 	{
-		GizmoHit = HitTarget->IsHit(DevicePos);
-		if (GizmoHit.bHit)
-		{
-			HitTarget->UpdateHoverState(true);
-			return;
-		}
+		GizmoHit = HitTarget->IsHit(PressPos);
 	}
+	return GizmoHit;
+}
+
+void UAxisAngleGizmo::OnBeginHover(const FInputDeviceRay& DevicePos)
+{
+	HitTarget->UpdateHoverState(true);
+}
+
+bool UAxisAngleGizmo::OnUpdateHover(const FInputDeviceRay& DevicePos)
+{
+	// not necessary...
+	HitTarget->UpdateHoverState(true);
+	return true;
+}
+
+void UAxisAngleGizmo::OnEndHover()
+{
 	HitTarget->UpdateHoverState(false);
 }
 
+
diff --git a/Engine/Source/Runtime/Experimental/InteractiveToolsFramework/Private/BaseGizmos/AxisPositionGizmo.cpp b/Engine/Source/Runtime/Experimental/InteractiveToolsFramework/Private/BaseGizmos/AxisPositionGizmo.cpp
--- a/Engine/Source/Runtime/Experimental/InteractiveToolsFramework/Private/BaseGizmos/AxisPositionGizmo.cpp
+++ b/Engine/Source/Runtime/Experimental/InteractiveToolsFramework/Private/BaseGizmos/AxisPositionGizmo.cpp
@@ -124,18 +124,29 @@ void UAxisPositionGizmo::OnTerminateDragSequence()
 
 
 
-void UAxisPositionGizmo::OnUpdateHover(const FInputDeviceRay& DevicePos)
+FInputRayHit UAxisPositionGizmo::BeginHoverSequenceHitTest(const FInputDeviceRay& PressPos)
 {
 	FInputRayHit GizmoHit;
 	if (HitTarget)
 	{
-		GizmoHit = HitTarget->IsHit(DevicePos);
-		if (GizmoHit.bHit)
-		{
-			HitTarget->UpdateHoverState(true);
-			return;
-		}
+		GizmoHit = HitTarget->IsHit(PressPos);
 	}
-	HitTarget->UpdateHoverState(false);
+	return GizmoHit;
+}
+
+void UAxisPositionGizmo::OnBeginHover(const FInputDeviceRay& DevicePos)
+{
+	HitTarget->UpdateHoverState(true);
 }
 
+bool UAxisPositionGizmo::OnUpdateHover(const FInputDeviceRay& DevicePos)
+{
+	// not necessary...
+	HitTarget->UpdateHoverState(true);
+	return true;
+}
+
+void UAxisPositionGizmo::OnEndHover()
+{
+	HitTarget->UpdateHoverState(false);
+}
\ No newline at end of file
diff --git a/Engine/Source/Runtime/Experimental/InteractiveToolsFramework/Private/BaseGizmos/PlanePositionGizmo.cpp b/Engine/Source/Runtime/Experimental/InteractiveToolsFramework/Private/BaseGizmos/PlanePositionGizmo.cpp
--- a/Engine/Source/Runtime/Experimental/InteractiveToolsFramework/Private/BaseGizmos/PlanePositionGizmo.cpp
+++ b/Engine/Source/Runtime/Experimental/InteractiveToolsFramework/Private/BaseGizmos/PlanePositionGizmo.cpp
@@ -141,18 +141,31 @@ void UPlanePositionGizmo::OnTerminateDragSequence()
 
 
 
-void UPlanePositionGizmo::OnUpdateHover(const FInputDeviceRay& DevicePos)
+
+
+FInputRayHit UPlanePositionGizmo::BeginHoverSequenceHitTest(const FInputDeviceRay& PressPos)
 {
 	FInputRayHit GizmoHit;
 	if (HitTarget)
 	{
-		GizmoHit = HitTarget->IsHit(DevicePos);
-		if (GizmoHit.bHit)
-		{
-			HitTarget->UpdateHoverState(true);
-			return;
-		}
+		GizmoHit = HitTarget->IsHit(PressPos);
 	}
-	HitTarget->UpdateHoverState(false);
+	return GizmoHit;
 }
 
+void UPlanePositionGizmo::OnBeginHover(const FInputDeviceRay& DevicePos)
+{
+	HitTarget->UpdateHoverState(true);
+}
+
+bool UPlanePositionGizmo::OnUpdateHover(const FInputDeviceRay& DevicePos)
+{
+	// not necessary...
+	HitTarget->UpdateHoverState(true);
+	return true;
+}
+
+void UPlanePositionGizmo::OnEndHover()
+{
+	HitTarget->UpdateHoverState(false);
+}
diff --git a/Engine/Source/Runtime/Experimental/InteractiveToolsFramework/Private/BaseTools/MeshSurfacePointTool.cpp b/Engine/Source/Runtime/Experimental/InteractiveToolsFramework/Private/BaseTools/MeshSurfacePointTool.cpp
--- a/Engine/Source/Runtime/Experimental/InteractiveToolsFramework/Private/BaseTools/MeshSurfacePointTool.cpp
+++ b/Engine/Source/Runtime/Experimental/InteractiveToolsFramework/Private/BaseTools/MeshSurfacePointTool.cpp
@@ -97,6 +97,21 @@ void UMeshSurfacePointTool::SetShiftToggle(bool bShiftDown)
 
 
 
+FInputRayHit UMeshSurfacePointTool::BeginHoverSequenceHitTest(const FInputDeviceRay& PressPos)
+{
+	FHitResult OutHit;
+	if (HitTest(PressPos.WorldRay, OutHit))
+	{
+		return FInputRayHit(OutHit.Distance);
+	}
+	return FInputRayHit();
+}
+
+
+
+
+
+
 /*
  * Mouse Input Behavior
  */
diff --git a/Engine/Source/Runtime/Experimental/InteractiveToolsFramework/Private/InputBehavior.cpp b/Engine/Source/Runtime/Experimental/InteractiveToolsFramework/Private/InputBehavior.cpp
--- a/Engine/Source/Runtime/Experimental/InteractiveToolsFramework/Private/InputBehavior.cpp
+++ b/Engine/Source/Runtime/Experimental/InteractiveToolsFramework/Private/InputBehavior.cpp
@@ -65,12 +65,24 @@ bool UInputBehavior::WantsHoverEvents()
 	return false;
 }
 
-void UInputBehavior::UpdateHover(const FInputDeviceState& input)
+FInputCaptureRequest UInputBehavior::WantsHoverCapture(const FInputDeviceState& input)
 {
-	return;
+	return FInputCaptureRequest::Ignore();
+}
+
+
+FInputCaptureUpdate UInputBehavior::BeginHoverCapture(const FInputDeviceState& input, EInputCaptureSide eSide)
+{
+	return FInputCaptureUpdate::Ignore();
+}
+
+
+FInputCaptureUpdate UInputBehavior::UpdateHoverCapture(const FInputDeviceState& input)
+{
+	return FInputCaptureUpdate::End();
 }
 
-void UInputBehavior::EndHover(const FInputDeviceState& input)
+void UInputBehavior::EndHoverCapture()
 {
 	return;
 }
diff --git a/Engine/Source/Runtime/Experimental/InteractiveToolsFramework/Private/InputBehaviorSet.cpp b/Engine/Source/Runtime/Experimental/InteractiveToolsFramework/Private/InputBehaviorSet.cpp
--- a/Engine/Source/Runtime/Experimental/InteractiveToolsFramework/Private/InputBehaviorSet.cpp
+++ b/Engine/Source/Runtime/Experimental/InteractiveToolsFramework/Private/InputBehaviorSet.cpp
@@ -116,38 +116,21 @@ void UInputBehaviorSet::CollectWantsCapture(const FInputDeviceState& input, TArr
 
 
 
-bool UInputBehaviorSet::UpdateHover(const FInputDeviceState& input)
+void UInputBehaviorSet::CollectWantsHoverCapture(const FInputDeviceState& input, TArray<FInputCaptureRequest>& result)
 {
-	bool bAnyWantedHover = false;
-	for (auto b : Behaviors) 
-	{
-		if (SupportsInputType(b.Behavior, input)) 
-		{
-			if (b.Behavior->WantsHoverEvents()) 
-			{
-				b.Behavior->UpdateHover(input);
-				bAnyWantedHover = true;
-			}
-		}
-	}
-	return bAnyWantedHover;
-}
-
-bool UInputBehaviorSet::EndHover(const FInputDeviceState& input)
-{
-	bool bAnyWantedHover = false;
 	for (auto b : Behaviors)
 	{
-		if (SupportsInputType(b.Behavior, input)) 
+		// only call WantsCapture if the Behavior supports the current input device
+		if (b.Behavior->WantsHoverEvents() && SupportsInputType(b.Behavior, input))
 		{
-			if (b.Behavior->WantsHoverEvents()) 
+			FInputCaptureRequest request = b.Behavior->WantsHoverCapture(input);
+			if (request.Type != EInputCaptureRequestType::Ignore)
 			{
-				b.Behavior->EndHover(input);
-				bAnyWantedHover = true;
+				request.Owner = b.Source;
+				result.Add(request);
 			}
 		}
 	}
-	return bAnyWantedHover;
 }
 
 
diff --git a/Engine/Source/Runtime/Experimental/InteractiveToolsFramework/Private/InputRouter.cpp b/Engine/Source/Runtime/Experimental/InteractiveToolsFramework/Private/InputRouter.cpp
--- a/Engine/Source/Runtime/Experimental/InteractiveToolsFramework/Private/InputRouter.cpp
+++ b/Engine/Source/Runtime/Experimental/InteractiveToolsFramework/Private/InputRouter.cpp
@@ -168,10 +168,10 @@ void UInputRouter::PostInputEvent_Mouse(const FInputDeviceState& Input)
 	}
 
 	// update hover if nobody is capturing
-	if (ActiveLeftCapture == nullptr && ActiveRightCapture == nullptr)
+	if (ActiveLeftCapture == nullptr)
 	{
-		bool bProcessed = ActiveInputBehaviors->UpdateHover(Input);
-		if (bProcessed && bAutoInvalidateOnHover)
+		bool bHoverStateUpdated = ProcessMouseHover(Input);
+		if (bHoverStateUpdated && bAutoInvalidateOnHover)
 		{
 			TransactionsAPI->PostInvalidation();
 		}
@@ -182,9 +182,8 @@ void UInputRouter::PostInputEvent_Mouse(const FInputDeviceState& Input)
 
 void UInputRouter::PostHoverInputEvent(const FInputDeviceState& Input)
 {
-	LastHoverInput = Input;
-	bool bProcessed = ActiveInputBehaviors->UpdateHover(Input);
-	if (bProcessed && bAutoInvalidateOnHover)
+	bool bHoverStateUpdated = ProcessMouseHover(Input);
+	if (bHoverStateUpdated && bAutoInvalidateOnHover)
 	{
 		TransactionsAPI->PostInvalidation();
 	}
@@ -217,8 +216,8 @@ void UInputRouter::CheckForMouseCaptures(const FInputDeviceState& Input)
 			CaptureRequests[i].Source->BeginCapture(Input, EInputCaptureSide::Left);
 		if (Result.State == EInputCaptureState::Begin)
 		{
-			// end outstanding hovers
-			ActiveInputBehaviors->EndHover(Input);
+			// end outstanding hover
+			TerminateHover(EInputCaptureSide::Left);
 
 			ActiveLeftCapture = Result.Source;
 			ActiveLeftCaptureOwner = CaptureRequests[i].Owner;
@@ -261,6 +260,85 @@ void UInputRouter::HandleCapturedMouseInput(const FInputDeviceState& Input)
 
 
 
+
+void UInputRouter::TerminateHover(EInputCaptureSide Side)
+{
+	if (Side == EInputCaptureSide::Left && ActiveLeftHoverCapture != nullptr)
+	{
+		ActiveLeftHoverCapture->EndHoverCapture();
+		ActiveLeftHoverCapture = nullptr;
+		ActiveLeftCaptureOwner = nullptr;
+	}
+}
+
+
+bool UInputRouter::ProcessMouseHover(const FInputDeviceState& Input)
+{
+	TArray<FInputCaptureRequest> CaptureRequests;
+	ActiveInputBehaviors->CollectWantsHoverCapture(Input, CaptureRequests);
+
+	if (CaptureRequests.Num() == 0 )
+	{
+		if (ActiveLeftHoverCapture != nullptr)
+		{
+			TerminateHover(EInputCaptureSide::Left);
+			return true;
+		}
+		return false;
+	}
+
+	CaptureRequests.StableSort();
+
+	// if we have an active hover, either update it, or terminate if we got a new best hit
+	bool bHoverStateModified = false;
+	if (ActiveLeftHoverCapture != nullptr)
+	{
+		bool bTerminateActiveHover = false;
+		if (CaptureRequests[0].Source == ActiveLeftHoverCapture)
+		{
+			FInputCaptureUpdate Result =
+				ActiveLeftHoverCapture->UpdateHoverCapture(Input);
+			bTerminateActiveHover = (Result.State == EInputCaptureState::End);
+		}
+		else
+		{
+			bTerminateActiveHover = true;
+		}
+
+		if (bTerminateActiveHover)
+		{
+			TerminateHover(EInputCaptureSide::Left);
+			bHoverStateModified = true;
+		}
+		else
+		{
+			return true;		// hover has been consumed
+		}
+	}
+
+	// if we get here, we have a new hover
+	bool bAccepted = false;
+	for (int i = 0; i < CaptureRequests.Num() && bAccepted == false; ++i)
+	{
+		FInputCaptureUpdate Result =
+			CaptureRequests[i].Source->BeginHoverCapture(Input, EInputCaptureSide::Left);
+		if (Result.State == EInputCaptureState::Begin)
+		{
+			ActiveLeftHoverCapture = Result.Source;
+			ActiveLeftCaptureOwner = CaptureRequests[i].Owner;
+			bAccepted = true;
+			return true;
+		}
+	}
+	
+	// no hover! but we might have terminated an active hover
+	return bHoverStateModified;
+}
+
+
+
+
+
 void UInputRouter::ForceTerminateAll()
 {
 	if (ActiveKeyboardCapture != nullptr)
@@ -287,7 +365,10 @@ void UInputRouter::ForceTerminateAll()
 		ActiveRightCaptureData = FInputCaptureData();
 	}
 
-	ActiveInputBehaviors->EndHover(LastHoverInput);
+	if (ActiveLeftHoverCapture != nullptr)
+	{
+		TerminateHover(EInputCaptureSide::Left);
+	}
 }
 
 
@@ -316,6 +397,11 @@ void UInputRouter::ForceTerminateSource(IInputBehaviorSource* Source)
 		ActiveRightCaptureOwner = nullptr;
 		ActiveRightCaptureData = FInputCaptureData();
 	}
+
+	if (ActiveLeftHoverCapture != nullptr && ActiveLeftHoverCaptureOwner == Source)
+	{
+		TerminateHover(EInputCaptureSide::Left);
+	}
 }
 
 
diff --git a/Engine/Source/Runtime/Experimental/InteractiveToolsFramework/Public/BaseBehaviors/BehaviorTargetInterfaces.h b/Engine/Source/Runtime/Experimental/InteractiveToolsFramework/Public/BaseBehaviors/BehaviorTargetInterfaces.h
--- a/Engine/Source/Runtime/Experimental/InteractiveToolsFramework/Public/BaseBehaviors/BehaviorTargetInterfaces.h
+++ b/Engine/Source/Runtime/Experimental/InteractiveToolsFramework/Public/BaseBehaviors/BehaviorTargetInterfaces.h
@@ -167,7 +167,22 @@ public:
 	virtual ~IHoverBehaviorTarget() {}
 
 	/**
-	 * Notify Target about a hover event
+	 * Do hover hit-test
 	 */
-	virtual void OnUpdateHover(const FInputDeviceRay& DevicePos) = 0;
+	virtual FInputRayHit BeginHoverSequenceHitTest(const FInputDeviceRay& PressPos) = 0;
+
+	/**
+	 * Initialize hover sequence at given position
+	 */
+	virtual void OnBeginHover(const FInputDeviceRay& DevicePos) = 0;
+
+	/**
+	 * Update active hover sequence with new input position
+	 */
+	virtual bool OnUpdateHover(const FInputDeviceRay& DevicePos) = 0;
+
+	/**
+	 * Terminate active hover sequence
+	 */
+	virtual void OnEndHover() = 0;
 };
diff --git a/Engine/Source/Runtime/Experimental/InteractiveToolsFramework/Public/BaseBehaviors/MouseHoverBehavior.h b/Engine/Source/Runtime/Experimental/InteractiveToolsFramework/Public/BaseBehaviors/MouseHoverBehavior.h
--- a/Engine/Source/Runtime/Experimental/InteractiveToolsFramework/Public/BaseBehaviors/MouseHoverBehavior.h
+++ b/Engine/Source/Runtime/Experimental/InteractiveToolsFramework/Public/BaseBehaviors/MouseHoverBehavior.h
@@ -35,8 +35,10 @@ public:
 	virtual EInputDevices GetSupportedDevices() override;
 
 	virtual bool WantsHoverEvents() override;
-	virtual void UpdateHover(const FInputDeviceState& input) override;
-	virtual void EndHover(const FInputDeviceState& input) override;
+	virtual FInputCaptureRequest WantsHoverCapture(const FInputDeviceState& InputState) override;
+	virtual FInputCaptureUpdate BeginHoverCapture(const FInputDeviceState& InputState, EInputCaptureSide eSide) override;
+	virtual FInputCaptureUpdate UpdateHoverCapture(const FInputDeviceState& InputState) override;
+	virtual void EndHoverCapture() override;
 
 protected:
 	IHoverBehaviorTarget* Target;
diff --git a/Engine/Source/Runtime/Experimental/InteractiveToolsFramework/Public/BaseBehaviors/MultiClickSequenceInputBehavior.h b/Engine/Source/Runtime/Experimental/InteractiveToolsFramework/Public/BaseBehaviors/MultiClickSequenceInputBehavior.h
--- a/Engine/Source/Runtime/Experimental/InteractiveToolsFramework/Public/BaseBehaviors/MultiClickSequenceInputBehavior.h
+++ b/Engine/Source/Runtime/Experimental/InteractiveToolsFramework/Public/BaseBehaviors/MultiClickSequenceInputBehavior.h
@@ -59,8 +59,10 @@ public:
 	virtual void ForceEndCapture(const FInputCaptureData& Data) override;
 
 	virtual bool WantsHoverEvents() override;
-	virtual void UpdateHover(const FInputDeviceState& input) override;
-	virtual void EndHover(const FInputDeviceState& input) override;
+	virtual FInputCaptureRequest WantsHoverCapture(const FInputDeviceState& InputState) override;
+	virtual FInputCaptureUpdate BeginHoverCapture(const FInputDeviceState& InputState, EInputCaptureSide eSide) override;
+	virtual FInputCaptureUpdate UpdateHoverCapture(const FInputDeviceState& InputState) override;
+	virtual void EndHoverCapture() override;
 
 
 public:
diff --git a/Engine/Source/Runtime/Experimental/InteractiveToolsFramework/Public/BaseGizmos/AxisAngleGizmo.h b/Engine/Source/Runtime/Experimental/InteractiveToolsFramework/Public/BaseGizmos/AxisAngleGizmo.h
--- a/Engine/Source/Runtime/Experimental/InteractiveToolsFramework/Public/BaseGizmos/AxisAngleGizmo.h
+++ b/Engine/Source/Runtime/Experimental/InteractiveToolsFramework/Public/BaseGizmos/AxisAngleGizmo.h
@@ -44,7 +44,10 @@ public:
 	virtual void OnTerminateDragSequence() override;
 
 	// IHoverBehaviorTarget implementation
-	virtual void OnUpdateHover(const FInputDeviceRay& DevicePos) override;
+	virtual FInputRayHit BeginHoverSequenceHitTest(const FInputDeviceRay& PressPos) override;
+	virtual void OnBeginHover(const FInputDeviceRay& DevicePos) override;
+	virtual bool OnUpdateHover(const FInputDeviceRay& DevicePos) override;
+	virtual void OnEndHover() override;
 
 public:
 	UPROPERTY()
diff --git a/Engine/Source/Runtime/Experimental/InteractiveToolsFramework/Public/BaseGizmos/AxisPositionGizmo.h b/Engine/Source/Runtime/Experimental/InteractiveToolsFramework/Public/BaseGizmos/AxisPositionGizmo.h
--- a/Engine/Source/Runtime/Experimental/InteractiveToolsFramework/Public/BaseGizmos/AxisPositionGizmo.h
+++ b/Engine/Source/Runtime/Experimental/InteractiveToolsFramework/Public/BaseGizmos/AxisPositionGizmo.h
@@ -44,7 +44,11 @@ public:
 	virtual void OnTerminateDragSequence() override;
 
 	// IHoverBehaviorTarget implementation
-	virtual void OnUpdateHover(const FInputDeviceRay& DevicePos) override;
+	virtual FInputRayHit BeginHoverSequenceHitTest(const FInputDeviceRay& PressPos) override;
+	virtual void OnBeginHover(const FInputDeviceRay& DevicePos) override;
+	virtual bool OnUpdateHover(const FInputDeviceRay& DevicePos) override;
+	virtual void OnEndHover() override;
+
 
 public:
 	UPROPERTY()
diff --git a/Engine/Source/Runtime/Experimental/InteractiveToolsFramework/Public/BaseGizmos/PlanePositionGizmo.h b/Engine/Source/Runtime/Experimental/InteractiveToolsFramework/Public/BaseGizmos/PlanePositionGizmo.h
--- a/Engine/Source/Runtime/Experimental/InteractiveToolsFramework/Public/BaseGizmos/PlanePositionGizmo.h
+++ b/Engine/Source/Runtime/Experimental/InteractiveToolsFramework/Public/BaseGizmos/PlanePositionGizmo.h
@@ -44,7 +44,10 @@ public:
 	virtual void OnTerminateDragSequence() override;
 
 	// IHoverBehaviorTarget implementation
-	virtual void OnUpdateHover(const FInputDeviceRay& DevicePos) override;
+	virtual FInputRayHit BeginHoverSequenceHitTest(const FInputDeviceRay& PressPos) override;
+	virtual void OnBeginHover(const FInputDeviceRay& DevicePos) override;
+	virtual bool OnUpdateHover(const FInputDeviceRay& DevicePos) override;
+	virtual void OnEndHover() override;
 
 public:
 	UPROPERTY()
diff --git a/Engine/Source/Runtime/Experimental/InteractiveToolsFramework/Public/BaseTools/MeshSurfacePointTool.h b/Engine/Source/Runtime/Experimental/InteractiveToolsFramework/Public/BaseTools/MeshSurfacePointTool.h
--- a/Engine/Source/Runtime/Experimental/InteractiveToolsFramework/Public/BaseTools/MeshSurfacePointTool.h
+++ b/Engine/Source/Runtime/Experimental/InteractiveToolsFramework/Public/BaseTools/MeshSurfacePointTool.h
@@ -81,10 +81,15 @@ public:
 	 */
 	virtual void OnEndDrag(const FRay& Ray);
 
-	/**
-	 * This function is called each frame while the tool is active, to support hover updates
-	 */
-	virtual void OnUpdateHover(const FInputDeviceRay& DevicePos) {}
+
+
+	// IHoverBehaviorTarget implementation
+	
+	virtual FInputRayHit BeginHoverSequenceHitTest(const FInputDeviceRay& PressPos) override;
+	virtual void OnBeginHover(const FInputDeviceRay& DevicePos) override {}
+	virtual bool OnUpdateHover(const FInputDeviceRay& DevicePos) override { return true; }
+	virtual void OnEndHover() override {}
+
 
 
 
diff --git a/Engine/Source/Runtime/Experimental/InteractiveToolsFramework/Public/InputBehavior.h b/Engine/Source/Runtime/Experimental/InteractiveToolsFramework/Public/InputBehavior.h
--- a/Engine/Source/Runtime/Experimental/InteractiveToolsFramework/Public/InputBehavior.h
+++ b/Engine/Source/Runtime/Experimental/InteractiveToolsFramework/Public/InputBehavior.h
@@ -288,11 +288,17 @@ public:
 	/** return true if this Behavior supports hover (ie passive input events) */
 	virtual bool WantsHoverEvents();
 
-	/** called on each new hover input event, ie if no other behavior is actively capturing input */
-	virtual void UpdateHover(const FInputDeviceState& InputState);
+	/** Given the input state, does this Behavior want to begin capturing some input devices for hover */
+	virtual FInputCaptureRequest WantsHoverCapture(const FInputDeviceState& InputState);
 
-	/** if a capture begins or focus is lost, any active hover visualization needs to terminate */
-	virtual void EndHover(const FInputDeviceState& InputState);
+	/** Called after WantsHoverCapture() returns a capture request that was accepted */
+	virtual FInputCaptureUpdate BeginHoverCapture(const FInputDeviceState& InputState, EInputCaptureSide eSide);
+
+	/** Called on each new hover input event, ie if no other behavior is actively capturing input */
+	virtual FInputCaptureUpdate UpdateHoverCapture(const FInputDeviceState& InputState);
+
+	/** If a different hover capture begins, focus is lost, a tool starts, etc, any active hover visualization needs to terminate */
+	virtual void EndHoverCapture();
 
 
 protected:
diff --git a/Engine/Source/Runtime/Experimental/InteractiveToolsFramework/Public/InputBehaviorSet.h b/Engine/Source/Runtime/Experimental/InteractiveToolsFramework/Public/InputBehaviorSet.h
--- a/Engine/Source/Runtime/Experimental/InteractiveToolsFramework/Public/InputBehaviorSet.h
+++ b/Engine/Source/Runtime/Experimental/InteractiveToolsFramework/Public/InputBehaviorSet.h
@@ -115,18 +115,12 @@ public:
 	virtual void CollectWantsCapture(const FInputDeviceState& InputState, TArray<FInputCaptureRequest>& ResultOut);
 
 	/**
-	 * Call UInputBehavior::UpdateHover() on each valid Behavior that wants hover events
+	 * Call UInputBehavior::WantsCapture() on each valid Behavior and collect up the requests that indicated a Capture was desired.
 	 * @param InputState current input device state information
-	 * @return true if any InputBehavior indicated that it wants hover events
+	 * @param ResultOut returned set of non-Ignoring FInputCaptureRequests returned by Behaviors
 	 */
-	virtual bool UpdateHover(const FInputDeviceState&  InputState);
+	virtual void CollectWantsHoverCapture(const FInputDeviceState& InputState, TArray<FInputCaptureRequest>& ResultOut);
 
-	/**
-	 * Call UInputBehavior::EndHover() on each valid Behavior that wants hover events
-	 * @param InputState current input device state information
-	 * @return true if any InputBehavior indicated that it wants hover events
-	 */
-	virtual bool EndHover(const FInputDeviceState&  InputState);
 
 
 protected:
diff --git a/Engine/Source/Runtime/Experimental/InteractiveToolsFramework/Public/InputRouter.h b/Engine/Source/Runtime/Experimental/InteractiveToolsFramework/Public/InputRouter.h
--- a/Engine/Source/Runtime/Experimental/InteractiveToolsFramework/Public/InputRouter.h
+++ b/Engine/Source/Runtime/Experimental/InteractiveToolsFramework/Public/InputRouter.h
@@ -81,20 +81,18 @@ protected:
 	UPROPERTY()
 	UInputBehaviorSet* ActiveInputBehaviors;
 
-	UInputBehavior* ActiveKeyboardCapture;
-	void* ActiveKeyboardCaptureOwner;
+	UInputBehavior* ActiveKeyboardCapture = nullptr;
+	void* ActiveKeyboardCaptureOwner = nullptr;
 	FInputCaptureData ActiveKeyboardCaptureData;
 
-	UInputBehavior* ActiveLeftCapture;
-	void* ActiveLeftCaptureOwner;
+	UInputBehavior* ActiveLeftCapture = nullptr;
+	void* ActiveLeftCaptureOwner = nullptr;
 	FInputCaptureData ActiveLeftCaptureData;
 
-	UInputBehavior* ActiveRightCapture;
-	void* ActiveRightCaptureOwner;
+	UInputBehavior* ActiveRightCapture = nullptr;
+	void* ActiveRightCaptureOwner = nullptr;
 	FInputCaptureData ActiveRightCaptureData;
 
-	FInputDeviceState LastHoverInput;
-
 	virtual void PostInputEvent_Keyboard(const FInputDeviceState& Input);
 	void CheckForKeyboardCaptures(const FInputDeviceState& Input);
 	void HandleCapturedKeyboardInput(const FInputDeviceState& Input);
@@ -102,4 +100,15 @@ protected:
 	virtual void PostInputEvent_Mouse(const FInputDeviceState& Input);
 	void CheckForMouseCaptures(const FInputDeviceState& Input);
 	void HandleCapturedMouseInput(const FInputDeviceState& Input);
+
+
+	//
+	// Hover support
+	//
+
+	UInputBehavior* ActiveLeftHoverCapture = nullptr;
+	void* ActiveLeftHoverCaptureOwner = nullptr;
+
+	void TerminateHover(EInputCaptureSide Side);
+	bool ProcessMouseHover(const FInputDeviceState& Input);
 };
\ No newline at end of file
