ToolsFramework: add FComponentWorldTransformChange, tracks changes world transform of a component. Add UGizmoTransformChangeStateTarget, a gizmo state target that builds and emits FComponentWorldTransformChange. Add hit normal to FInputRayHit.
#rnx

diff --git a/Engine/Source/Runtime/Experimental/InteractiveToolsFramework/Private/Changes/TransformChange.cpp b/Engine/Source/Runtime/Experimental/InteractiveToolsFramework/Private/Changes/TransformChange.cpp
new file mode 100644
--- /dev/null
+++ b/Engine/Source/Runtime/Experimental/InteractiveToolsFramework/Private/Changes/TransformChange.cpp
@@ -0,0 +1,44 @@
+// Copyright 1998-2019 Epic Games, Inc. All Rights Reserved.
+
+#include "Changes/TransformChange.h"
+#include "Components/SceneComponent.h"
+
+FComponentWorldTransformChange::FComponentWorldTransformChange()
+{
+	FromWorldTransform = ToWorldTransform = FTransform::Identity;
+}
+
+FComponentWorldTransformChange::FComponentWorldTransformChange(const FTransform& From, const FTransform& To)
+{
+	FromWorldTransform = From;
+	ToWorldTransform = To;
+}
+
+void FComponentWorldTransformChange::Apply(UObject* Object)
+{
+	USceneComponent* SceneComponent = CastChecked<USceneComponent>(Object);
+	SceneComponent->SetWorldTransform(ToWorldTransform);
+
+	if (OnChangeAppliedFunc)
+	{
+		OnChangeAppliedFunc(this, Object, true);
+	}
+}
+
+void FComponentWorldTransformChange::Revert(UObject* Object)
+{
+	USceneComponent* SceneComponent = CastChecked<USceneComponent>(Object);
+	SceneComponent->SetWorldTransform(FromWorldTransform);
+
+	if (OnChangeAppliedFunc)
+	{
+		OnChangeAppliedFunc(this, Object, false);
+	}
+}
+
+
+FString FComponentWorldTransformChange::ToString() const
+{
+	return FString(TEXT("Transform Change"));
+}
+
diff --git a/Engine/Source/Runtime/Experimental/InteractiveToolsFramework/Public/BaseGizmos/StateTargets.h b/Engine/Source/Runtime/Experimental/InteractiveToolsFramework/Public/BaseGizmos/StateTargets.h
--- a/Engine/Source/Runtime/Experimental/InteractiveToolsFramework/Public/BaseGizmos/StateTargets.h
+++ b/Engine/Source/Runtime/Experimental/InteractiveToolsFramework/Public/BaseGizmos/StateTargets.h
@@ -5,6 +5,7 @@
 #include "CoreMinimal.h"
 #include "InteractiveGizmoManager.h"
 #include "BaseGizmos/GizmoInterfaces.h"
+#include "Changes/TransformChange.h"
 #include "StateTargets.generated.h"
 
 
@@ -125,3 +126,95 @@ public:
 		return NewTarget;
 	}
 };
+
+
+
+
+
+
+
+/**
+ * UGizmoTransformChangeStateTarget is an implementation of IGizmoStateTarget that
+ * emits an FComponentWorldTransformChange on a target USceneComponent. This StateTarget
+ * also opens/closes an undo transaction via GizmoManager.
+ */
+UCLASS()
+class INTERACTIVETOOLSFRAMEWORK_API UGizmoTransformChangeStateTarget : public UObject, public IGizmoStateTarget
+{
+	GENERATED_BODY()
+public:
+	virtual void BeginUpdate()
+	{
+		if (TargetComponent.IsValid())
+		{
+			if (GizmoManager.IsValid())
+			{
+				GizmoManager->BeginUndoTransaction(ChangeDescription);
+			}
+
+			InitialTransform = TargetComponent->GetComponentTransform();
+		}
+	}
+
+	virtual void EndUpdate()
+	{
+		if (TargetComponent.IsValid())
+		{
+			FinalTransform = TargetComponent->GetComponentTransform();
+
+			if (GizmoManager.IsValid())
+			{
+				TUniquePtr<FComponentWorldTransformChange> TransformChange 
+					= MakeUnique<FComponentWorldTransformChange>(InitialTransform, FinalTransform);
+				GizmoManager->EmitObjectChange(TargetComponent.Get(), MoveTemp(TransformChange), ChangeDescription);
+
+				if (GizmoManager.IsValid())
+				{
+					GizmoManager->EndUndoTransaction();
+				}
+			}
+		}
+	}
+
+
+	/**
+	 * The object that will be changed, ie have Modify() called on it on BeginUpdate()
+	 */
+	TWeakObjectPtr<USceneComponent> TargetComponent;
+
+	/**
+	 * Localized text description of the transaction (will be visible in Editor on undo/redo)
+	 */
+	FText ChangeDescription;
+
+	/**
+	 * Pointer to the GizmoManager that is used to open/close the transaction
+	 */
+	TWeakObjectPtr<UInteractiveGizmoManager> GizmoManager;
+
+	/** Start Transform, saved on BeginUpdate() */
+	FTransform InitialTransform;
+	/** End Transform, saved on EndUpdate() */
+	FTransform FinalTransform;
+
+
+public:
+	/**
+	 * Create and initialize an standard instance of UGizmoTransformChangeStateTarget
+	 * @param TargetComponentIn the USceneComponent this StateTarget will track transform changes on
+	 * @param DescriptionIn Localized text description of the transaction
+	 * @param GizmoManagerIn pointer to the GizmoManager used to manage transactions
+	 */
+	static UGizmoTransformChangeStateTarget* Construct(
+		USceneComponent* TargetComponentIn,
+		const FText& DescriptionIn,
+		UInteractiveGizmoManager* GizmoManagerIn,
+		UObject* Outer = (UObject*)GetTransientPackage())
+	{
+		UGizmoTransformChangeStateTarget* NewTarget = NewObject<UGizmoTransformChangeStateTarget>(Outer);
+		NewTarget->TargetComponent = TargetComponentIn;
+		NewTarget->ChangeDescription = DescriptionIn;
+		NewTarget->GizmoManager = GizmoManagerIn;
+		return NewTarget;
+	}
+};
diff --git a/Engine/Source/Runtime/Experimental/InteractiveToolsFramework/Public/Changes/TransformChange.h b/Engine/Source/Runtime/Experimental/InteractiveToolsFramework/Public/Changes/TransformChange.h
new file mode 100644
--- /dev/null
+++ b/Engine/Source/Runtime/Experimental/InteractiveToolsFramework/Public/Changes/TransformChange.h
@@ -0,0 +1,35 @@
+// Copyright 1998-2019 Epic Games, Inc. All Rights Reserved.
+
+#pragma once
+
+#include "CoreMinimal.h"
+#include "Misc/Change.h"
+
+
+/**
+ * FComponentWorldTransformChange represents an undoable change to the world transform of a USceneComponent.
+ */
+class INTERACTIVETOOLSFRAMEWORK_API FComponentWorldTransformChange : public FCommandChange
+{
+public:
+	FComponentWorldTransformChange();
+	FComponentWorldTransformChange(const FTransform& From, const FTransform& To);
+
+	FTransform FromWorldTransform;
+	FTransform ToWorldTransform;
+
+	/** This function is called on Apply and Revert (last argument is true on Apply)*/
+	TFunction<void(FComponentWorldTransformChange*, UObject*, bool)> OnChangeAppliedFunc;
+
+	/** Makes the change to the object */
+	virtual void Apply(UObject* Object) override;
+
+	/** Reverts change to the object */
+	virtual void Revert(UObject* Object) override;
+
+	/** Describes this change (for debugging) */
+	virtual FString ToString() const override;
+};
+
+
+
diff --git a/Engine/Source/Runtime/Experimental/InteractiveToolsFramework/Public/InputState.h b/Engine/Source/Runtime/Experimental/InteractiveToolsFramework/Public/InputState.h
--- a/Engine/Source/Runtime/Experimental/InteractiveToolsFramework/Public/InputState.h
+++ b/Engine/Source/Runtime/Experimental/InteractiveToolsFramework/Public/InputState.h
@@ -80,6 +80,12 @@ struct INTERACTIVETOOLSFRAMEWORK_API FInputRayHit
 	/** distance along ray at which intersection occurred */
 	float HitDepth;
 
+	/** Normal at hit point, if available */
+	FVector HitNormal;
+
+	/** True if HitNormal was set */
+	bool bHasHitNormal;
+
 	/** client-defined integer identifier for hit object/element/target/etc */
 	int32 HitIdentifier;
 
@@ -90,6 +96,8 @@ struct INTERACTIVETOOLSFRAMEWORK_API FInputRayHit
 	{
 		bHit = false;
 		HitDepth = TNumericLimits<float>::Max();
+		HitNormal = FVector(0, 0, 1);
+		bHasHitNormal = false;
 		HitIdentifier = 0;
 		HitOwner = nullptr;
 	}
@@ -98,6 +106,18 @@ struct INTERACTIVETOOLSFRAMEWORK_API FInputRayHit
 	{
 		bHit = true;
 		HitDepth = HitDepthIn;
+		HitNormal = FVector(0, 0, 1);
+		bHasHitNormal = false;
+		HitIdentifier = 0;
+		HitOwner = nullptr;
+	}
+
+	explicit FInputRayHit(float HitDepthIn, const FVector& HitNormalIn)
+	{
+		bHit = true;
+		HitDepth = HitDepthIn;
+		HitNormal = HitNormalIn;
+		bHasHitNormal = true;
 		HitIdentifier = 0;
 		HitOwner = nullptr;
 	}
