Add missing header to fix CIS errors
#rnx

diff --git a/Engine/Source/Runtime/Experimental/InteractiveToolsFramework/Public/BaseGizmos/StateTargets.h b/Engine/Source/Runtime/Experimental/InteractiveToolsFramework/Public/BaseGizmos/StateTargets.h
--- a/Engine/Source/Runtime/Experimental/InteractiveToolsFramework/Public/BaseGizmos/StateTargets.h
+++ b/Engine/Source/Runtime/Experimental/InteractiveToolsFramework/Public/BaseGizmos/StateTargets.h
@@ -3,6 +3,7 @@
 #pragma once
 
 #include "CoreMinimal.h"
+#include "Components/SceneComponent.h"
 #include "InteractiveGizmoManager.h"
 #include "BaseGizmos/GizmoInterfaces.h"
 #include "Changes/TransformChange.h"
