ToolsFramework: use new UGizmoTransformChangeStateTarget in TransformGizmo, so that gizmo position is properly updated on undo/redo
#rnx

diff --git a/Engine/Source/Runtime/Experimental/InteractiveToolsFramework/Private/BaseGizmos/TransformGizmo.cpp b/Engine/Source/Runtime/Experimental/InteractiveToolsFramework/Private/BaseGizmos/TransformGizmo.cpp
--- a/Engine/Source/Runtime/Experimental/InteractiveToolsFramework/Private/BaseGizmos/TransformGizmo.cpp
+++ b/Engine/Source/Runtime/Experimental/InteractiveToolsFramework/Private/BaseGizmos/TransformGizmo.cpp
@@ -202,8 +202,9 @@ void UTransformGizmo::SetActiveTarget(UTransformProxy* Target)
 	TargetTransform.SetScale3D(FVector(1, 1, 1));
 	GizmoComponent->SetWorldTransform(TargetTransform);
 
-	// target tracks location of GizmoComponent
-	GizmoActor->GetRootComponent()->TransformUpdated.AddLambda(
+	// Target tracks location of GizmoComponent. Note that TransformUpdated is not called during undo/redo transactions!
+	// We currently rely on the transaction system to undo/redo target object locations. This will not work during runtime...
+	GizmoComponent->TransformUpdated.AddLambda(
 		[this, SaveScale](USceneComponent* Component, EUpdateTransformFlags /*UpdateTransformFlags*/, ETeleportType /*Teleport*/) {
 		//this->GetGizmoManager()->DisplayMessage(TEXT("TRANSFORM UPDATED"), EToolMessageLevel::Internal);
 		FTransform NewXForm = Component->GetComponentToWorld();
@@ -212,7 +213,9 @@ void UTransformGizmo::SetActiveTarget(UTransformProxy* Target)
 	});
 
 
-	UGizmoObjectModifyStateTarget* StateTarget = UGizmoObjectModifyStateTarget::Construct(GizmoComponent,
+	// This state target emits an explicit FChange that moves the GizmoActor root component during undo/redo.
+	// It also opens/closes the Transaction that saves/restores the target object locations.
+	UGizmoTransformChangeStateTarget* StateTarget = UGizmoTransformChangeStateTarget::Construct(GizmoComponent,
 		LOCTEXT("UTransformGizmoTransaction", "Transform"), GetGizmoManager(), this);
 
 	UGizmoComponentWorldTransformSource* TransformSource = 
