ToolsFramework: add simple method to save/restore Tool UInteractiveToolPropertySet properties, via CDO.
#rnx

diff --git a/Engine/Source/Runtime/Experimental/InteractiveToolsFramework/Public/InteractiveTool.h b/Engine/Source/Runtime/Experimental/InteractiveToolsFramework/Public/InteractiveTool.h
--- a/Engine/Source/Runtime/Experimental/InteractiveToolsFramework/Public/InteractiveTool.h
+++ b/Engine/Source/Runtime/Experimental/InteractiveToolsFramework/Public/InteractiveTool.h
@@ -28,13 +28,13 @@ DECLARE_MULTICAST_DELEGATE_TwoParams(FInteractiveToolPropertySetModifiedSignatur
 
 
 /**
- * a UInteractiveTool contains a set of UObjects that contain "properties" of the Tool, ie
+ * A UInteractiveTool contains a set of UObjects that contain "properties" of the Tool, ie
  * the configuration flags, parameters, etc that control the Tool. Currently any UObject
  * can be added as a property set, however there is no automatic mechanism for those child 
  * UObjects to notify the Tool when a property changes.
  * 
  * If you make your property set UObjects subclasses of UInteractiveToolPropertySet, then
- * when the properties are changed in the Editor, the parent Tool will be automatically notified.
+ * when the Tool Properties are changed *in the Editor*, the parent Tool will be automatically notified.
  * You can override UInteractiveTool::OnPropertyModified() to act on these notifications
  */
 UCLASS(Transient)
@@ -53,11 +53,57 @@ public:
 		return OnModified;
 	}
 
-	/** posts a message to the OnModified delegate with the modified UProperty */
+	/** 
+	  * Posts a message to the OnModified delegate with the modified UProperty 
+	  * @warning this function is currently only called in Editor (not at runtime)
+	  */
 	void PostEditChangeProperty(FPropertyChangedEvent& PropertyChangedEvent)
 	{
 		OnModified.Broadcast(this, PropertyChangedEvent.Property);
 	}
+
+
+
+	//
+	// Setting saving/serialization
+	//
+
+public:
+	/**
+	 * Save values of current Tool Properties. Implementing these functions is *optional*
+	 * and how it is implemented is up to the PropertySet implementation.
+	 * It is not necessary to save/restore all possible Properties (in many cases this would not make sense).
+	 * GetPropertyCache() can be used to return an instance of subclasses that is an easy
+	 * place to save/restore these properties
+	 */
+	virtual void SaveProperties(UInteractiveTool* SaveFromTool) {}
+
+	/**
+	 * Restore saved property values
+	 */
+	virtual void RestoreProperties(UInteractiveTool* RestoreToTool) {}
+
+protected:
+	/**
+	 * GetPropertyCache returns a class-internal object that subclasses can use to save/restore properties.
+	 * If the subclass is UMyPropertySet, this function should only ever be called as GetPropertyCache<UMyPropertySet>().
+	 */
+	template<typename ObjType>
+	ObjType* GetPropertyCache()
+	{
+		ObjType* CDO = GetMutableDefault<ObjType>();
+		if (CDO->CachedProperties == nullptr)
+		{
+			CDO->CachedProperties = NewObject<ObjType>();
+		}
+		return CastChecked<ObjType>(CDO->CachedProperties);
+	}
+
+private:
+	// CachedProperties should only ever be set to an instance of the subclass, ideally via GetPropertyCache().
+	UPROPERTY()
+	UObject* CachedProperties = nullptr;
+
 };
 
 
