ToolsFramework: Add support for querying external world/local coordinate system setting via IToolsContextQueriesAPI (eg like the current mode in the Editor UI).
#rnx

diff --git a/Engine/Source/Editor/Experimental/EditorInteractiveToolsFramework/Private/EdModeInteractiveToolsContext.cpp b/Engine/Source/Editor/Experimental/EditorInteractiveToolsFramework/Private/EdModeInteractiveToolsContext.cpp
--- a/Engine/Source/Editor/Experimental/EditorInteractiveToolsFramework/Private/EdModeInteractiveToolsContext.cpp
+++ b/Engine/Source/Editor/Experimental/EditorInteractiveToolsFramework/Private/EdModeInteractiveToolsContext.cpp
@@ -64,6 +64,13 @@ public:
 	}
 
 
+	virtual EToolContextCoordinateSystem GetCurrentCoordinateSystem() const override
+	{
+		ECoordSystem CoordSys = EditorMode->GetModeManager()->GetCoordSystem();
+		return (CoordSys == COORD_World) ? EToolContextCoordinateSystem::World : EToolContextCoordinateSystem::Local;
+	}
+
+
 	virtual bool ExecuteSceneSnapQuery(const FSceneSnapQueryRequest& Request, TArray<FSceneSnapQueryResult>& Results) const override
 	{
 		if (Request.RequestType != ESceneSnapQueryType::Position)
diff --git a/Engine/Source/Editor/UnrealEd/Private/Tools/UEdMode.cpp b/Engine/Source/Editor/UnrealEd/Private/Tools/UEdMode.cpp
--- a/Engine/Source/Editor/UnrealEd/Private/Tools/UEdMode.cpp
+++ b/Engine/Source/Editor/UnrealEd/Private/Tools/UEdMode.cpp
@@ -63,6 +63,13 @@ public:
 	}
 
 
+	virtual EToolContextCoordinateSystem GetCurrentCoordinateSystem() const override
+	{
+		ECoordSystem CoordSys = EdMode->GetModeManager()->GetCoordSystem();
+		return (CoordSys == COORD_World) ? EToolContextCoordinateSystem::World : EToolContextCoordinateSystem::Local;
+	}
+
+
 	virtual bool ExecuteSceneSnapQuery(const FSceneSnapQueryRequest& Request, TArray<FSceneSnapQueryResult>& Results) const override
 	{
 		if (Request.RequestType != ESceneSnapQueryType::Position)
diff --git a/Engine/Source/Runtime/Experimental/InteractiveToolsFramework/Public/ToolContextInterfaces.h b/Engine/Source/Runtime/Experimental/InteractiveToolsFramework/Public/ToolContextInterfaces.h
--- a/Engine/Source/Runtime/Experimental/InteractiveToolsFramework/Public/ToolContextInterfaces.h
+++ b/Engine/Source/Runtime/Experimental/InteractiveToolsFramework/Public/ToolContextInterfaces.h
@@ -142,6 +142,14 @@ enum class EStandardToolContextMaterials
 };
 
 
+/** Types of coordinate systems that a Tool/Gizmo might use */
+UENUM()
+enum class EToolContextCoordinateSystem
+{
+	World = 0,
+	Local = 1
+};
+
 
 /**
  * Users of the Tools Framework need to implement IToolsContextQueriesAPI to provide
@@ -166,6 +174,12 @@ public:
 	virtual void GetCurrentViewState(FViewCameraState& StateOut) const = 0;
 
 
+	/**
+	 * Request current external coordinate-system setting
+	 */
+	virtual EToolContextCoordinateSystem GetCurrentCoordinateSystem() const = 0;	
+
+
 	/**
 	 * Try to find Snap Targets in the scene that satisfy the Snap Query.
 	 * @param Request snap query configuration
