ToolsFramework: Support for world/local frame configuration in TransformGizmo and EditMeshPolygonsTool
#rnx

diff --git a/Engine/Plugins/Experimental/MeshModelingToolset/Source/MeshModelingTools/Private/EditMeshPolygonsTool.cpp b/Engine/Plugins/Experimental/MeshModelingToolset/Source/MeshModelingTools/Private/EditMeshPolygonsTool.cpp
--- a/Engine/Plugins/Experimental/MeshModelingToolset/Source/MeshModelingTools/Private/EditMeshPolygonsTool.cpp
+++ b/Engine/Plugins/Experimental/MeshModelingToolset/Source/MeshModelingTools/Private/EditMeshPolygonsTool.cpp
@@ -47,7 +47,6 @@ UMeshSurfacePointTool* UEditMeshPolygonsToolBuilder::CreateNewTool(const FToolBu
  */
 UPolyEditTransformProperties::UPolyEditTransformProperties()
 {
-	bUseLocalAxis = true;
 	TransformMode = EQuickTransformerMode::AxisTranslation;
 	bSelectVertices = true;
 	bSelectFaces = true;
@@ -732,14 +731,6 @@ void UEditMeshPolygonsTool::Shutdown(EToolShutdownType ShutdownType)
 
 
 
-void UEditMeshPolygonsTool::ToggleLocalTransformFrameAction()
-{
-	TransformProps->bUseLocalAxis = !TransformProps->bUseLocalAxis;
-	if (bInDrag)
-	{
-		UpdateQuickTransformer();
-	}
-}
 
 void UEditMeshPolygonsTool::NextTransformTypeAction()
 {
@@ -761,14 +752,6 @@ void UEditMeshPolygonsTool::NextTransformTypeAction()
 
 void UEditMeshPolygonsTool::RegisterActions(FInteractiveToolActionSet& ActionSet)
 {
-	ActionSet.RegisterAction(this, (int32)EStandardToolActions::BaseClientDefinedActionID + 1,
-		TEXT("ToggleFrame"),
-		LOCTEXT("ToggleFrame", "Toggle World/Local Frame"),
-		LOCTEXT("ToggleFrameTooltip", "Toggle between World and Local transformation frames"),
-		EModifierKey::None, EKeys::A,
-		[this]() { ToggleLocalTransformFrameAction(); });
-
-
 	ActionSet.RegisterAction(this, (int32)EStandardToolActions::BaseClientDefinedActionID + 2,
 		TEXT("NextTransformType"),
 		LOCTEXT("NextTransformType", "Next Transform Type"),
@@ -1015,7 +998,9 @@ FQuickTransformer* UEditMeshPolygonsTool::GetActiveQuickTransformer()
 
 void UEditMeshPolygonsTool::UpdateQuickTransformer()
 {
-	if (TransformProps->bUseLocalAxis)
+	bool bUseLocalAxes =
+		(GetToolManager()->GetContextQueriesAPI()->GetCurrentCoordinateSystem() == EToolContextCoordinateSystem::Local);
+	if (bUseLocalAxes)
 	{
 		GetActiveQuickTransformer()->SetActiveWorldFrame(ActiveSurfaceFrame);
 	}
diff --git a/Engine/Plugins/Experimental/MeshModelingToolset/Source/MeshModelingTools/Public/EditMeshPolygonsTool.h b/Engine/Plugins/Experimental/MeshModelingToolset/Source/MeshModelingTools/Public/EditMeshPolygonsTool.h
--- a/Engine/Plugins/Experimental/MeshModelingToolset/Source/MeshModelingTools/Public/EditMeshPolygonsTool.h
+++ b/Engine/Plugins/Experimental/MeshModelingToolset/Source/MeshModelingTools/Public/EditMeshPolygonsTool.h
@@ -111,9 +111,6 @@ public:
 	UPROPERTY(EditAnywhere, Category = Options, meta = (DisplayName = "Deformation Type", ToolTip = "Select the type of deformation you wish to employ on a polygroup."))
 	EGroupTopologyDeformationStrategy DeformationStrategy;
 
-	UPROPERTY(EditAnywhere, Category = Options)
-	bool bUseLocalAxis;
-
 	UPROPERTY(EditAnywhere, Category = Options)
 	EQuickTransformerMode TransformMode;
 
@@ -428,7 +425,6 @@ public:
 	virtual bool OnUpdateHover(const FInputDeviceRay& DevicePos) override;
 
 public:
-	virtual void ToggleLocalTransformFrameAction();
 	virtual void NextTransformTypeAction();
 
 	//
diff --git a/Engine/Source/Runtime/Experimental/InteractiveToolsFramework/Private/BaseGizmos/GizmoArrowComponent.cpp b/Engine/Source/Runtime/Experimental/InteractiveToolsFramework/Private/BaseGizmos/GizmoArrowComponent.cpp
--- a/Engine/Source/Runtime/Experimental/InteractiveToolsFramework/Private/BaseGizmos/GizmoArrowComponent.cpp
+++ b/Engine/Source/Runtime/Experimental/InteractiveToolsFramework/Private/BaseGizmos/GizmoArrowComponent.cpp
@@ -44,7 +44,8 @@ public:
 		FVector ViewDirection = Origin - GizmoControlView->ViewLocation;
 		ViewDirection.Normalize();
 
-		FVector ArrowDirection = LocalToWorldMatrix.TransformVector(Direction);
+		bool bWorldAxis = (bExternalWorldLocalState) ? (*bExternalWorldLocalState) : false;
+		FVector ArrowDirection = (bWorldAxis) ? Direction : LocalToWorldMatrix.TransformVector(Direction);
 		bool bFlipped = (FVector::DotProduct(ViewDirection, ArrowDirection) > 0);
 		ArrowDirection = (bFlipped) ? -ArrowDirection : ArrowDirection;
 		if (bFlippedExternal != nullptr)
@@ -131,6 +132,10 @@ public:
 		bExternalHoverState = HoverState;
 	}
 
+	void SetExternalWorldLocalState(bool* bWorldLocalState)
+	{
+		bExternalWorldLocalState = bWorldLocalState;
+	}
 
 private:
 	FLinearColor Color;
@@ -144,6 +149,7 @@ private:
 	float* ExternalDynamicPixelToWorldScale = nullptr;
 	bool* bExternalRenderVisibility = nullptr;
 	bool* bExternalHoverState = nullptr;
+	bool* bExternalWorldLocalState = nullptr;
 };
 
 
@@ -154,6 +160,7 @@ FPrimitiveSceneProxy* UGizmoArrowComponent::CreateSceneProxy()
 	NewProxy->SetExternalDynamicPixelToWorldScale(&DynamicPixelToWorldScale);
 	NewProxy->SetExternalRenderVisibility(&bRenderVisibility);
 	NewProxy->SetExternalHoverState(&bHovering);
+	NewProxy->SetExternalWorldLocalState(&bWorld);
 	return NewProxy;
 }
 
@@ -176,8 +183,10 @@ bool UGizmoArrowComponent::LineTraceComponent(FHitResult& OutHit, const FVector
 	double StartDist = LengthScale * Gap;
 	double EndDist = LengthScale * (Gap + Length);
 
-	FVector Point0 = Transform.TransformPosition( StartDist * UseDirection);
-	FVector Point1 = Transform.TransformPosition( EndDist * UseDirection);
+	UseDirection = (bWorld) ? UseDirection : Transform.TransformVector(UseDirection);
+	FVector UseOrigin = Transform.TransformPosition(FVector::ZeroVector);
+	FVector Point0 = UseOrigin + StartDist * UseDirection;
+	FVector Point1 = UseOrigin + EndDist * UseDirection;
 
 	FVector NearestArrow, NearestLine;
 	FMath::SegmentDistToSegmentSafe(Point0, Point1, Start, End, NearestArrow, NearestLine);
diff --git a/Engine/Source/Runtime/Experimental/InteractiveToolsFramework/Private/BaseGizmos/GizmoCircleComponent.cpp b/Engine/Source/Runtime/Experimental/InteractiveToolsFramework/Private/BaseGizmos/GizmoCircleComponent.cpp
--- a/Engine/Source/Runtime/Experimental/InteractiveToolsFramework/Private/BaseGizmos/GizmoCircleComponent.cpp
+++ b/Engine/Source/Runtime/Experimental/InteractiveToolsFramework/Private/BaseGizmos/GizmoCircleComponent.cpp
@@ -94,8 +94,9 @@ public:
 				else
 				{
 					FVector WorldOrigin = LocalToWorldMatrix.TransformPosition(FVector::ZeroVector);
-					FVector WorldPlaneX = LocalToWorldMatrix.TransformVector(PlaneX);
-					FVector WorldPlaneY = LocalToWorldMatrix.TransformVector(PlaneY);
+					bool bWorldAxis = (bExternalWorldLocalState) ? (*bExternalWorldLocalState) : false;
+					FVector WorldPlaneX = (bWorldAxis) ? PlaneX : LocalToWorldMatrix.TransformVector(PlaneX);
+					FVector WorldPlaneY = (bWorldAxis) ? PlaneY : LocalToWorldMatrix.TransformVector(PlaneY);
 
 					FVector	LastVertex = WorldOrigin + WorldPlaneX * UseRadius;
 					bool bLastVisible = FVector::DotProduct(WorldPlaneX, GizmoViewDirection) < 0;
@@ -152,6 +153,12 @@ public:
 		bExternalHoverState = HoverState;
 	}
 
+	void SetExternalWorldLocalState(bool* bWorldLocalState)
+	{
+		bExternalWorldLocalState = bWorldLocalState;
+	}
+
+
 private:
 	FLinearColor Color;
 	FVector Normal;
@@ -163,6 +170,7 @@ private:
 
 	float* ExternalDynamicPixelToWorldScale = nullptr;
 	bool* bExternalHoverState = nullptr;
+	bool* bExternalWorldLocalState = nullptr;
 };
 
 
@@ -173,6 +181,7 @@ FPrimitiveSceneProxy* UGizmoCircleComponent::CreateSceneProxy()
 	FGizmoCircleComponentSceneProxy* NewProxy = new FGizmoCircleComponentSceneProxy(this);
 	NewProxy->SetExternalDynamicPixelToWorldScale(&DynamicPixelToWorldScale);
 	NewProxy->SetExternalHoverState(&bHovering);
+	NewProxy->SetExternalWorldLocalState(&bWorld);
 	return NewProxy;
 }
 
@@ -187,7 +196,7 @@ bool UGizmoCircleComponent::LineTraceComponent(FHitResult& OutHit, const FVector
 	double UseRadius = LengthScale * Radius;
 
 	const FTransform& Transform = this->GetComponentToWorld();
-	FVector WorldNormal = Transform.TransformVector(Normal);
+	FVector WorldNormal = (bWorld) ? Normal : Transform.TransformVector(Normal);
 	FVector WorldOrigin = Transform.TransformPosition(FVector::ZeroVector);
 	FPlane CirclePlane(WorldOrigin, WorldNormal);
 
diff --git a/Engine/Source/Runtime/Experimental/InteractiveToolsFramework/Private/BaseGizmos/GizmoRectangleComponent.cpp b/Engine/Source/Runtime/Experimental/InteractiveToolsFramework/Private/BaseGizmos/GizmoRectangleComponent.cpp
--- a/Engine/Source/Runtime/Experimental/InteractiveToolsFramework/Private/BaseGizmos/GizmoRectangleComponent.cpp
+++ b/Engine/Source/Runtime/Experimental/InteractiveToolsFramework/Private/BaseGizmos/GizmoRectangleComponent.cpp
@@ -44,7 +44,8 @@ public:
 		FVector ViewDirection = Origin - GizmoControlView->ViewLocation;
 		ViewDirection.Normalize();
 
-		FVector UseDirectionX = LocalToWorldMatrix.TransformVector(DirectionX);
+		bool bWorldAxis = (bExternalWorldLocalState) ? (*bExternalWorldLocalState) : false;
+		FVector UseDirectionX = (bWorldAxis) ? DirectionX : LocalToWorldMatrix.TransformVector(DirectionX);
 		bool bFlippedX = (FVector::DotProduct(ViewDirection, UseDirectionX) > 0);
 		UseDirectionX = (bFlippedX) ? -UseDirectionX : UseDirectionX;
 		if (bFlippedXExternal)
@@ -52,7 +53,7 @@ public:
 			*bFlippedXExternal = bFlippedX;
 		}
 
-		FVector UseDirectionY = LocalToWorldMatrix.TransformVector(DirectionY);
+		FVector UseDirectionY = (bWorldAxis) ? DirectionY : LocalToWorldMatrix.TransformVector(DirectionY);
 		bool bFlippedY = (FVector::DotProduct(ViewDirection, UseDirectionY) > 0);
 		UseDirectionY = (bFlippedY) ? -UseDirectionY : UseDirectionY;
 		if (bFlippedYExternal)
@@ -165,6 +166,12 @@ public:
 		bExternalHoverState = HoverState;
 	}
 
+	void SetExternalWorldLocalState(bool* bWorldLocalState)
+	{
+		bExternalWorldLocalState = bWorldLocalState;
+	}
+
+
 private:
 	FLinearColor Color;
 	FVector DirectionX, DirectionY;
@@ -178,6 +185,7 @@ private:
 	float* ExternalDynamicPixelToWorldScale = nullptr;
 	bool* bExternalRenderVisibility = nullptr;
 	bool* bExternalHoverState = nullptr;
+	bool* bExternalWorldLocalState = nullptr;
 };
 
 
@@ -189,6 +197,7 @@ FPrimitiveSceneProxy* UGizmoRectangleComponent::CreateSceneProxy()
 	NewProxy->SetExternalDynamicPixelToWorldScale(&DynamicPixelToWorldScale);
 	NewProxy->SetExternalRenderVisibility(&bRenderVisibility);
 	NewProxy->SetExternalHoverState(&bHovering);
+	NewProxy->SetExternalWorldLocalState(&bWorld);
 	return NewProxy;
 }
 
@@ -209,7 +218,10 @@ bool UGizmoRectangleComponent::LineTraceComponent(FHitResult& OutHit, const FVec
 	const FTransform& Transform = this->GetComponentToWorld();
 
 	FVector UseDirectionX = (bFlippedX) ? -DirectionX : DirectionX;
+	UseDirectionX = (bWorld) ? UseDirectionX : Transform.TransformVector(UseDirectionX);
 	FVector UseDirectionY = (bFlippedY) ? -DirectionY : DirectionY;
+	UseDirectionY = (bWorld) ? DirectionY : Transform.TransformVector(UseDirectionY);
+	FVector UseOrigin = Transform.TransformPosition(FVector::ZeroVector);
 
 	float LengthScale = DynamicPixelToWorldScale;
 	double UseOffsetX = LengthScale * OffsetX;
@@ -218,10 +230,10 @@ bool UGizmoRectangleComponent::LineTraceComponent(FHitResult& OutHit, const FVec
 	double UseOffsetLengthY = LengthScale * (OffsetY + LengthY);
 
 	FVector Points[4] = {
-		Transform.TransformPosition(UseOffsetX*UseDirectionX + UseOffsetY *UseDirectionY),
-		Transform.TransformPosition(UseOffsetLengthX*UseDirectionX + UseOffsetY *UseDirectionY),
-		Transform.TransformPosition(UseOffsetLengthX*UseDirectionX + UseOffsetLengthY *UseDirectionY),
-		Transform.TransformPosition(UseOffsetX*UseDirectionX + UseOffsetLengthY *UseDirectionY)
+		UseOrigin + UseOffsetX*UseDirectionX + UseOffsetY*UseDirectionY,
+		UseOrigin + UseOffsetLengthX*UseDirectionX + UseOffsetY*UseDirectionY,
+		UseOrigin + UseOffsetLengthX*UseDirectionX + UseOffsetLengthY*UseDirectionY,
+		UseOrigin + UseOffsetX*UseDirectionX + UseOffsetLengthY*UseDirectionY
 	};
 	static const int Triangles[2][3] = { {0,1,2}, {0,2,3} };
 
diff --git a/Engine/Source/Runtime/Experimental/InteractiveToolsFramework/Private/BaseGizmos/TransformGizmo.cpp b/Engine/Source/Runtime/Experimental/InteractiveToolsFramework/Private/BaseGizmos/TransformGizmo.cpp
--- a/Engine/Source/Runtime/Experimental/InteractiveToolsFramework/Private/BaseGizmos/TransformGizmo.cpp
+++ b/Engine/Source/Runtime/Experimental/InteractiveToolsFramework/Private/BaseGizmos/TransformGizmo.cpp
@@ -134,6 +134,11 @@ UInteractiveGizmo* UTransformGizmoBuilder::BuildGizmo(const FToolBuilderState& S
 		NewGizmo->SetUpdateHoverFunction(UpdateHoverFunction);
 	}
 
+	if (UpdateCoordSystemFunction)
+	{
+		NewGizmo->SetUpdateCoordSystemFunction(UpdateCoordSystemFunction);
+	}
+
 	return NewGizmo;
 }
 
@@ -155,6 +160,11 @@ void UTransformGizmo::SetUpdateHoverFunction(TFunction<void(UPrimitiveComponent*
 	UpdateHoverFunction = HoverFunction;
 }
 
+void UTransformGizmo::SetUpdateCoordSystemFunction(TFunction<void(UPrimitiveComponent*, EToolContextCoordinateSystem)> CoordSysFunction)
+{
+	UpdateCoordSystemFunction = CoordSysFunction;
+}
+
 
 void UTransformGizmo::Setup()
 {
@@ -168,6 +178,14 @@ void UTransformGizmo::Setup()
 		}
 	};
 
+	UpdateCoordSystemFunction = [](UPrimitiveComponent* Component, EToolContextCoordinateSystem CoordSystem)
+	{
+		if (Cast<UGizmoBaseComponent>(Component) != nullptr)
+		{
+			Cast<UGizmoBaseComponent>(Component)->UpdateWorldLocalState(CoordSystem == EToolContextCoordinateSystem::World);
+		}
+	};
+
 	GizmoActor = GizmoActorBuilder->CreateNewGizmoActor(World);
 }
 
@@ -185,6 +203,30 @@ void UTransformGizmo::Shutdown()
 }
 
 
+
+void UTransformGizmo::Tick(float DeltaTime)
+{	
+	EToolContextCoordinateSystem CoordSystem = GetGizmoManager()->GetContextQueriesAPI()->GetCurrentCoordinateSystem();
+	check(CoordSystem == EToolContextCoordinateSystem::World || CoordSystem == EToolContextCoordinateSystem::Local)
+	bool bUseLocalAxes =
+		(GetGizmoManager()->GetContextQueriesAPI()->GetCurrentCoordinateSystem() == EToolContextCoordinateSystem::Local);
+	if (AxisXSource != nullptr && AxisYSource != nullptr && AxisZSource != nullptr)
+	{
+		AxisXSource->bLocalAxes = bUseLocalAxes;
+		AxisYSource->bLocalAxes = bUseLocalAxes;
+		AxisZSource->bLocalAxes = bUseLocalAxes;
+	}
+	if (UpdateCoordSystemFunction)
+	{
+		for (UPrimitiveComponent* Component : ActiveComponents)
+		{
+			UpdateCoordSystemFunction(Component, CoordSystem);
+		}
+	}
+}
+
+
+
 void UTransformGizmo::SetActiveTarget(UTransformProxy* Target)
 {
 	if (ActiveTarget != nullptr)
@@ -230,41 +272,50 @@ void UTransformGizmo::SetActiveTarget(UTransformProxy* Target)
 	if (GizmoActor->TranslateX != nullptr)
 	{
 		AddAxisTranslationGizmo(GizmoActor->TranslateX, GizmoComponent, AxisXSource, TransformSource, StateTarget);
+		ActiveComponents.Add(GizmoActor->TranslateX);
 	}
 	if (GizmoActor->TranslateY != nullptr)
 	{
 		AddAxisTranslationGizmo(GizmoActor->TranslateY, GizmoComponent, AxisYSource, TransformSource, StateTarget);
+		ActiveComponents.Add(GizmoActor->TranslateY);
 	}
 	if (GizmoActor->TranslateZ != nullptr)
 	{
 		AddAxisTranslationGizmo(GizmoActor->TranslateZ, GizmoComponent, AxisZSource, TransformSource, StateTarget);
+		ActiveComponents.Add(GizmoActor->TranslateZ);
 	}
 
 
 	if (GizmoActor->TranslateYZ != nullptr)
 	{
 		AddPlaneTranslationGizmo(GizmoActor->TranslateYZ, GizmoComponent, AxisXSource, TransformSource, StateTarget);
+		ActiveComponents.Add(GizmoActor->TranslateYZ);
 	}
 	if (GizmoActor->TranslateXZ != nullptr)
 	{
 		AddPlaneTranslationGizmo(GizmoActor->TranslateXZ, GizmoComponent, AxisYSource, TransformSource, StateTarget);
+		ActiveComponents.Add(GizmoActor->TranslateXZ);
 	}
 	if (GizmoActor->TranslateXY != nullptr)
 	{
 		AddPlaneTranslationGizmo(GizmoActor->TranslateXY, GizmoComponent, AxisZSource, TransformSource, StateTarget);
+		ActiveComponents.Add(GizmoActor->TranslateXY);
 	}
 
 	if (GizmoActor->RotateX != nullptr)
 	{
 		AddAxisRotationGizmo(GizmoActor->RotateX, GizmoComponent, AxisXSource, TransformSource, StateTarget);
+		ActiveComponents.Add(GizmoActor->RotateX);
 	}
 	if (GizmoActor->RotateY != nullptr)
 	{
 		AddAxisRotationGizmo(GizmoActor->RotateY, GizmoComponent, AxisYSource, TransformSource, StateTarget);
+		ActiveComponents.Add(GizmoActor->RotateY);
 	}
 	if (GizmoActor->RotateZ != nullptr)
 	{
 		AddAxisRotationGizmo(GizmoActor->RotateZ, GizmoComponent, AxisZSource, TransformSource, StateTarget);
+		ActiveComponents.Add(GizmoActor->RotateZ);
 	}
 }
 
@@ -381,6 +432,7 @@ void UTransformGizmo::ClearActiveTarget()
 		GetGizmoManager()->DestroyGizmo(Gizmo);
 	}
 	ActiveGizmos.SetNum(0);
+	ActiveComponents.SetNum(0);
 
 	ActiveTarget = nullptr;
 }
diff --git a/Engine/Source/Runtime/Experimental/InteractiveToolsFramework/Public/BaseGizmos/GizmoBaseComponent.h b/Engine/Source/Runtime/Experimental/InteractiveToolsFramework/Public/BaseGizmos/GizmoBaseComponent.h
--- a/Engine/Source/Runtime/Experimental/InteractiveToolsFramework/Public/BaseGizmos/GizmoBaseComponent.h
+++ b/Engine/Source/Runtime/Experimental/InteractiveToolsFramework/Public/BaseGizmos/GizmoBaseComponent.h
@@ -47,6 +47,15 @@ public:
 		}
 	}
 
+	UFUNCTION()
+	void UpdateWorldLocalState(bool bWorldIn)
+	{
+		if (bWorldIn != bWorld)
+		{
+			bWorld = bWorldIn;
+		}
+	}
+
 
 protected:
 	// scale factor between pixel distances and world distances at Gizmo origin
@@ -54,4 +63,7 @@ protected:
 
 	// hover state
 	bool bHovering = false;
+
+	// world/local coordinates state
+	bool bWorld = false;
 };
\ No newline at end of file
diff --git a/Engine/Source/Runtime/Experimental/InteractiveToolsFramework/Public/BaseGizmos/TransformGizmo.h b/Engine/Source/Runtime/Experimental/InteractiveToolsFramework/Public/BaseGizmos/TransformGizmo.h
--- a/Engine/Source/Runtime/Experimental/InteractiveToolsFramework/Public/BaseGizmos/TransformGizmo.h
+++ b/Engine/Source/Runtime/Experimental/InteractiveToolsFramework/Public/BaseGizmos/TransformGizmo.h
@@ -153,6 +153,17 @@ public:
 	 */
 	TFunction<void(UPrimitiveComponent*, bool)> UpdateHoverFunction;
 
+	/**
+	 * If set, this coord-system function will be passed to UTransformGizmo instances to use instead
+	 * of the default UpdateCoordSystemFunction. By default the UTransformGizmo will query the external Context
+	 * to ask whether it should be using world or local coordinate system. Then the default UpdateCoordSystemFunction
+	 * will try casting to UGizmoBaseCmponent and passing that info on via UpdateWorldLocalState();
+	 * If you are using different Components that do not subclass UGizmoBaseComponent, and you want the coord system
+	 * to be configurable, you will need to provide a different update function.
+	 */
+	TFunction<void(UPrimitiveComponent*, EToolContextCoordinateSystem)> UpdateCoordSystemFunction;
+
+
 	virtual UInteractiveGizmo* BuildGizmo(const FToolBuilderState& SceneState) const override;
 };
 
@@ -182,10 +193,12 @@ public:
 	virtual void SetWorld(UWorld* World);
 	virtual void SetGizmoActorBuilder(TSharedPtr<FTransformGizmoActorFactory> Builder);
 	virtual void SetUpdateHoverFunction(TFunction<void(UPrimitiveComponent*, bool)> HoverFunction);
+	virtual void SetUpdateCoordSystemFunction(TFunction<void(UPrimitiveComponent*, EToolContextCoordinateSystem)> CoordSysFunction);
 
 	// UInteractiveGizmo overrides
 	virtual void Setup() override;
 	virtual void Shutdown() override;
+	virtual void Tick(float DeltaTime) override;
 
 
 	/**
@@ -210,12 +223,28 @@ public:
 protected:
 	TSharedPtr<FTransformGizmoActorFactory> GizmoActorBuilder;
 
+	// This function is called on each active GizmoActor Component to update it's hover state.
+	// If the Component is not a UGizmoBaseCmponent, the client needs to provide a different implementation
+	// of this function via the ToolBuilder
 	TFunction<void(UPrimitiveComponent*, bool)> UpdateHoverFunction;
 
+	// This function is called on each active GizmoActor Component to update it's coordinate system (eg world/local).
+	// If the Component is not a UGizmoBaseCmponent, the client needs to provide a different implementation
+	// of this function via the ToolBuilder
+	TFunction<void(UPrimitiveComponent*, EToolContextCoordinateSystem)> UpdateCoordSystemFunction;
+
+	/** list of current-active child components */
+	UPROPERTY()
+	TArray<UPrimitiveComponent*> ActiveComponents;
+
+	/** list of currently-active child gizmos */
 	UPROPERTY()
 	TArray<UInteractiveGizmo*> ActiveGizmos;
 
+	/** GizmoActors will be spawned in this World */
 	UWorld* World;
+
+	/** Current active GizmoActor that was spawned by this Gizmo. Will be destroyed when Gizmo is. */
 	ATransformGizmoActor* GizmoActor;
 
 	/** X-axis source is shared across Gizmos, and created internally during SetActiveTarget() */
