ToolsFramework: add IToolContextTransactionProvider UInterface which provides the transaction/change functions that ToolManager and GizmoManager both implement.

UGizmoTransformChangeStateTarget can now take either. TransformGizmo::SetActiveTarget() now takes optional override (uses GizmoManager by default).

This allows us to pass in ToolManager instead inside Tools where the Gizmo is fully local to Tool, in which case Gizmo undo/redo events should be auto-expired after the Tool exits (which ToolManager will do, but GizmoManager does not)

#rnx

diff --git a/Engine/Source/Runtime/Experimental/InteractiveToolsFramework/Private/BaseGizmos/TransformGizmo.cpp b/Engine/Source/Runtime/Experimental/InteractiveToolsFramework/Private/BaseGizmos/TransformGizmo.cpp
--- a/Engine/Source/Runtime/Experimental/InteractiveToolsFramework/Private/BaseGizmos/TransformGizmo.cpp
+++ b/Engine/Source/Runtime/Experimental/InteractiveToolsFramework/Private/BaseGizmos/TransformGizmo.cpp
@@ -227,7 +227,7 @@ void UTransformGizmo::Tick(float DeltaTime)
 
 
 
-void UTransformGizmo::SetActiveTarget(UTransformProxy* Target)
+void UTransformGizmo::SetActiveTarget(UTransformProxy* Target, IToolContextTransactionProvider* TransactionProvider)
 {
 	if (ActiveTarget != nullptr)
 	{
@@ -257,8 +257,12 @@ void UTransformGizmo::SetActiveTarget(UTransformProxy* Target)
 
 	// This state target emits an explicit FChange that moves the GizmoActor root component during undo/redo.
 	// It also opens/closes the Transaction that saves/restores the target object locations.
+	if (TransactionProvider == nullptr)
+	{
+		TransactionProvider = GetGizmoManager();
+	}
 	UGizmoTransformChangeStateTarget* StateTarget = UGizmoTransformChangeStateTarget::Construct(GizmoComponent,
-		LOCTEXT("UTransformGizmoTransaction", "Transform"), GetGizmoManager(), this);
+		LOCTEXT("UTransformGizmoTransaction", "Transform"), TransactionProvider, this);
 
 	UGizmoComponentWorldTransformSource* TransformSource = 
 		UGizmoComponentWorldTransformSource::Construct(GizmoComponent, this);
diff --git a/Engine/Source/Runtime/Experimental/InteractiveToolsFramework/Public/BaseGizmos/StateTargets.h b/Engine/Source/Runtime/Experimental/InteractiveToolsFramework/Public/BaseGizmos/StateTargets.h
--- a/Engine/Source/Runtime/Experimental/InteractiveToolsFramework/Public/BaseGizmos/StateTargets.h
+++ b/Engine/Source/Runtime/Experimental/InteractiveToolsFramework/Public/BaseGizmos/StateTargets.h
@@ -5,6 +5,7 @@
 #include "CoreMinimal.h"
 #include "Components/SceneComponent.h"
 #include "InteractiveGizmoManager.h"
+#include "InteractiveToolManager.h"
 #include "BaseGizmos/GizmoInterfaces.h"
 #include "Changes/TransformChange.h"
 #include "StateTargets.generated.h"
@@ -148,9 +149,9 @@ public:
 	{
 		if (TargetComponent.IsValid())
 		{
-			if (GizmoManager.IsValid())
+			if (TransactionManager)
 			{
-				GizmoManager->BeginUndoTransaction(ChangeDescription);
+				TransactionManager->BeginUndoTransaction(ChangeDescription);
 			}
 
 			InitialTransform = TargetComponent->GetComponentTransform();
@@ -163,16 +164,12 @@ public:
 		{
 			FinalTransform = TargetComponent->GetComponentTransform();
 
-			if (GizmoManager.IsValid())
+			if (TransactionManager)
 			{
 				TUniquePtr<FComponentWorldTransformChange> TransformChange 
 					= MakeUnique<FComponentWorldTransformChange>(InitialTransform, FinalTransform);
-				GizmoManager->EmitObjectChange(TargetComponent.Get(), MoveTemp(TransformChange), ChangeDescription);
-
-				if (GizmoManager.IsValid())
-				{
-					GizmoManager->EndUndoTransaction();
-				}
+				TransactionManager->EmitObjectChange(TargetComponent.Get(), MoveTemp(TransformChange), ChangeDescription);
+				TransactionManager->EndUndoTransaction();
 			}
 		}
 	}
@@ -189,9 +186,10 @@ public:
 	FText ChangeDescription;
 
 	/**
-	 * Pointer to the GizmoManager that is used to open/close the transaction
+	 * Pointer to the GizmoManager or ToolManager that is used to open/close the transaction
 	 */
-	TWeakObjectPtr<UInteractiveGizmoManager> GizmoManager;
+	UPROPERTY()
+	TScriptInterface<IToolContextTransactionProvider> TransactionManager;
 
 	/** Start Transform, saved on BeginUpdate() */
 	FTransform InitialTransform;
@@ -200,22 +198,29 @@ public:
 
 
 public:
+
 	/**
 	 * Create and initialize an standard instance of UGizmoTransformChangeStateTarget
 	 * @param TargetComponentIn the USceneComponent this StateTarget will track transform changes on
 	 * @param DescriptionIn Localized text description of the transaction
-	 * @param GizmoManagerIn pointer to the GizmoManager used to manage transactions
+	 * @param TransactionManagerIn pointer to the GizmoManager or ToolManager used to manage transactions
 	 */
 	static UGizmoTransformChangeStateTarget* Construct(
 		USceneComponent* TargetComponentIn,
 		const FText& DescriptionIn,
-		UInteractiveGizmoManager* GizmoManagerIn,
+		IToolContextTransactionProvider* TransactionManagerIn,
 		UObject* Outer = (UObject*)GetTransientPackage())
 	{
 		UGizmoTransformChangeStateTarget* NewTarget = NewObject<UGizmoTransformChangeStateTarget>(Outer);
 		NewTarget->TargetComponent = TargetComponentIn;
 		NewTarget->ChangeDescription = DescriptionIn;
-		NewTarget->GizmoManager = GizmoManagerIn;
+
+		// have to explicitly configure this because we only have IToolContextTransactionProvider pointer
+		NewTarget->TransactionManager.SetInterface(TransactionManagerIn);
+		NewTarget->TransactionManager.SetObject(CastChecked<UObject>(TransactionManagerIn));
+		
 		return NewTarget;
 	}
+
+
 };
diff --git a/Engine/Source/Runtime/Experimental/InteractiveToolsFramework/Public/BaseGizmos/TransformGizmo.h b/Engine/Source/Runtime/Experimental/InteractiveToolsFramework/Public/BaseGizmos/TransformGizmo.h
--- a/Engine/Source/Runtime/Experimental/InteractiveToolsFramework/Public/BaseGizmos/TransformGizmo.h
+++ b/Engine/Source/Runtime/Experimental/InteractiveToolsFramework/Public/BaseGizmos/TransformGizmo.h
@@ -204,8 +204,9 @@ public:
 	/**
 	 * Set the active target object for the Gizmo
 	 * @param Target active target
+	 * @param TransactionProvider optional IToolContextTransactionProvider implementation to use - by default uses GizmoManager
 	 */
-	virtual void SetActiveTarget(UTransformProxy* Target);
+	virtual void SetActiveTarget(UTransformProxy* Target, IToolContextTransactionProvider* TransactionProvider = nullptr);
 
 	/**
 	 * Clear the active target object for the Gizmo
diff --git a/Engine/Source/Runtime/Experimental/InteractiveToolsFramework/Public/InteractiveGizmoManager.h b/Engine/Source/Runtime/Experimental/InteractiveToolsFramework/Public/InteractiveGizmoManager.h
--- a/Engine/Source/Runtime/Experimental/InteractiveToolsFramework/Public/InteractiveGizmoManager.h
+++ b/Engine/Source/Runtime/Experimental/InteractiveToolsFramework/Public/InteractiveGizmoManager.h
@@ -33,7 +33,7 @@ struct FActiveGizmo
  * 
  */
 UCLASS(Transient)
-class INTERACTIVETOOLSFRAMEWORK_API UInteractiveGizmoManager : public UObject
+class INTERACTIVETOOLSFRAMEWORK_API UInteractiveGizmoManager : public UObject, public IToolContextTransactionProvider
 {
 	GENERATED_BODY()
 
diff --git a/Engine/Source/Runtime/Experimental/InteractiveToolsFramework/Public/InteractiveToolManager.h b/Engine/Source/Runtime/Experimental/InteractiveToolsFramework/Public/InteractiveToolManager.h
--- a/Engine/Source/Runtime/Experimental/InteractiveToolsFramework/Public/InteractiveToolManager.h
+++ b/Engine/Source/Runtime/Experimental/InteractiveToolsFramework/Public/InteractiveToolManager.h
@@ -41,7 +41,7 @@ enum class EToolSide
  *
  */
 UCLASS(Transient)
-class INTERACTIVETOOLSFRAMEWORK_API UInteractiveToolManager : public UObject
+class INTERACTIVETOOLSFRAMEWORK_API UInteractiveToolManager : public UObject, public IToolContextTransactionProvider
 {
 	GENERATED_BODY()
 
diff --git a/Engine/Source/Runtime/Experimental/InteractiveToolsFramework/Public/ToolContextInterfaces.h b/Engine/Source/Runtime/Experimental/InteractiveToolsFramework/Public/ToolContextInterfaces.h
--- a/Engine/Source/Runtime/Experimental/InteractiveToolsFramework/Public/ToolContextInterfaces.h
+++ b/Engine/Source/Runtime/Experimental/InteractiveToolsFramework/Public/ToolContextInterfaces.h
@@ -3,7 +3,9 @@
 #pragma once
 
 #include "CoreMinimal.h"
+#include "UObject/Interface.h"
 #include "ComponentSourceInterfaces.h"
+#include "ToolContextInterfaces.generated.h"
 
 
 // predeclarations so we don't have to include these in all tools
@@ -15,6 +17,44 @@ class FPrimitiveDrawInterface;
 class UInteractiveToolManager;
 class UInteractiveGizmoManager;
 
+
+
+UINTERFACE()
+class INTERACTIVETOOLSFRAMEWORK_API UToolContextTransactionProvider : public UInterface
+{
+	GENERATED_BODY()
+};
+/**
+ * IToolContextTransactionProvider is a UInterface that defines several functions that InteractiveTool code
+ * uses to interface with the higher-level transaction system. UInteractiveToolManager and UInteractiveGizmoManager
+ * both implement this interface.
+ */
+class INTERACTIVETOOLSFRAMEWORK_API IToolContextTransactionProvider
+{
+	GENERATED_BODY()
+public:
+	/**
+	 * Request that the Context open a Transaction, whatever that means to the current Context
+	 * @param Description text description of this transaction (this is the string that appears on undo/redo in the UE Editor)
+	 */
+	virtual void BeginUndoTransaction(const FText& Description) = 0;
+
+	/** Request that the Context close and commit the open Transaction */
+	virtual void EndUndoTransaction() = 0;
+
+	/**
+	 * Forward an FChange object to the Context
+	 * @param TargetObject the object that the FChange applies to
+	 * @param Change the change object that the Context should insert into the transaction history
+	 * @param Description text description of this change (this is the string that appears on undo/redo in the UE Editor)
+	 */
+	virtual void EmitObjectChange(UObject* TargetObject, TUniquePtr<FChange> Change, const FText& Description) = 0;
+};
+
+
+
+
+
 /**
  * FToolBuilderState is a bucket of state information that a ToolBuilder might need
  * to construct a Tool. This information comes from a level above the Tools framework,
