Fix strict build by disambiguating some types

diff --git a/Engine/Source/Runtime/Experimental/InteractiveToolsFramework/Private/BaseGizmos/GizmoArrowComponent.cpp b/Engine/Source/Runtime/Experimental/InteractiveToolsFramework/Private/BaseGizmos/GizmoArrowComponent.cpp
--- a/Engine/Source/Runtime/Experimental/InteractiveToolsFramework/Private/BaseGizmos/GizmoArrowComponent.cpp
+++ b/Engine/Source/Runtime/Experimental/InteractiveToolsFramework/Private/BaseGizmos/GizmoArrowComponent.cpp
@@ -45,7 +45,7 @@ public:
 		ViewDirection.Normalize();
 
 		bool bWorldAxis = (bExternalWorldLocalState) ? (*bExternalWorldLocalState) : false;
-		FVector ArrowDirection = (bWorldAxis) ? Direction : LocalToWorldMatrix.TransformVector(Direction);
+		FVector ArrowDirection = (bWorldAxis) ? Direction : FVector{ LocalToWorldMatrix.TransformVector(Direction) };
 		bool bFlipped = (FVector::DotProduct(ViewDirection, ArrowDirection) > 0);
 		ArrowDirection = (bFlipped) ? -ArrowDirection : ArrowDirection;
 		if (bFlippedExternal != nullptr)
diff --git a/Engine/Source/Runtime/Experimental/InteractiveToolsFramework/Private/BaseGizmos/GizmoCircleComponent.cpp b/Engine/Source/Runtime/Experimental/InteractiveToolsFramework/Private/BaseGizmos/GizmoCircleComponent.cpp
--- a/Engine/Source/Runtime/Experimental/InteractiveToolsFramework/Private/BaseGizmos/GizmoCircleComponent.cpp
+++ b/Engine/Source/Runtime/Experimental/InteractiveToolsFramework/Private/BaseGizmos/GizmoCircleComponent.cpp
@@ -95,8 +95,8 @@ public:
 				{
 					FVector WorldOrigin = LocalToWorldMatrix.TransformPosition(FVector::ZeroVector);
 					bool bWorldAxis = (bExternalWorldLocalState) ? (*bExternalWorldLocalState) : false;
-					FVector WorldPlaneX = (bWorldAxis) ? PlaneX : LocalToWorldMatrix.TransformVector(PlaneX);
-					FVector WorldPlaneY = (bWorldAxis) ? PlaneY : LocalToWorldMatrix.TransformVector(PlaneY);
+					FVector WorldPlaneX = (bWorldAxis) ? PlaneX : FVector{ LocalToWorldMatrix.TransformVector(PlaneX) };
+					FVector WorldPlaneY = (bWorldAxis) ? PlaneY : FVector{ LocalToWorldMatrix.TransformVector(PlaneY) };
 
 					FVector	LastVertex = WorldOrigin + WorldPlaneX * UseRadius;
 					bool bLastVisible = FVector::DotProduct(WorldPlaneX, GizmoViewDirection) < 0;
diff --git a/Engine/Source/Runtime/Experimental/InteractiveToolsFramework/Private/BaseGizmos/GizmoRectangleComponent.cpp b/Engine/Source/Runtime/Experimental/InteractiveToolsFramework/Private/BaseGizmos/GizmoRectangleComponent.cpp
--- a/Engine/Source/Runtime/Experimental/InteractiveToolsFramework/Private/BaseGizmos/GizmoRectangleComponent.cpp
+++ b/Engine/Source/Runtime/Experimental/InteractiveToolsFramework/Private/BaseGizmos/GizmoRectangleComponent.cpp
@@ -45,7 +45,7 @@ public:
 		ViewDirection.Normalize();
 
 		bool bWorldAxis = (bExternalWorldLocalState) ? (*bExternalWorldLocalState) : false;
-		FVector UseDirectionX = (bWorldAxis) ? DirectionX : LocalToWorldMatrix.TransformVector(DirectionX);
+		FVector UseDirectionX = (bWorldAxis) ? DirectionX : FVector{ LocalToWorldMatrix.TransformVector(DirectionX) };
 		bool bFlippedX = (FVector::DotProduct(ViewDirection, UseDirectionX) > 0);
 		UseDirectionX = (bFlippedX) ? -UseDirectionX : UseDirectionX;
 		if (bFlippedXExternal)
@@ -53,7 +53,7 @@ public:
 			*bFlippedXExternal = bFlippedX;
 		}
 
-		FVector UseDirectionY = (bWorldAxis) ? DirectionY : LocalToWorldMatrix.TransformVector(DirectionY);
+		FVector UseDirectionY = (bWorldAxis) ? DirectionY : FVector{ LocalToWorldMatrix.TransformVector(DirectionY) };
 		bool bFlippedY = (FVector::DotProduct(ViewDirection, UseDirectionY) > 0);
 		UseDirectionY = (bFlippedY) ? -UseDirectionY : UseDirectionY;
 		if (bFlippedYExternal)
