ToolsFramework: add change event support for UTransformProxy. UGizmoTransformChangeStateTarget can now have attached change-generators. Use this in TransformGizmo to generate changes for TransformProxy when gizmo state change happens. Add TransformGizmo::SetNewGizmoTransform() which allows for repositioning gizmo via transaction/changes that work with undo/redo.
#rnx

diff --git a/Engine/Source/Runtime/Experimental/InteractiveToolsFramework/Private/BaseGizmos/TransformGizmo.cpp b/Engine/Source/Runtime/Experimental/InteractiveToolsFramework/Private/BaseGizmos/TransformGizmo.cpp
--- a/Engine/Source/Runtime/Experimental/InteractiveToolsFramework/Private/BaseGizmos/TransformGizmo.cpp
+++ b/Engine/Source/Runtime/Experimental/InteractiveToolsFramework/Private/BaseGizmos/TransformGizmo.cpp
@@ -261,8 +261,9 @@ void UTransformGizmo::SetActiveTarget(UTransformProxy* Target, IToolContextTrans
 	{
 		TransactionProvider = GetGizmoManager();
 	}
-	UGizmoTransformChangeStateTarget* StateTarget = UGizmoTransformChangeStateTarget::Construct(GizmoComponent,
+	StateTarget = UGizmoTransformChangeStateTarget::Construct(GizmoComponent,
 		LOCTEXT("UTransformGizmoTransaction", "Transform"), TransactionProvider, this);
+	StateTarget->DependentChangeSources.Add(MakeUnique<FTransformProxyChangeSource>(Target));
 
 	UGizmoComponentWorldTransformSource* TransformSource = 
 		UGizmoComponentWorldTransformSource::Construct(GizmoComponent, this);
@@ -325,11 +326,25 @@ void UTransformGizmo::SetActiveTarget(UTransformProxy* Target, IToolContextTrans
 
 
 
+void UTransformGizmo::SetNewGizmoTransform(const FTransform& NewTransform)
+{
+	check(ActiveTarget != nullptr);
+
+	StateTarget->BeginUpdate();
+
+	USceneComponent* GizmoComponent = GizmoActor->GetRootComponent();
+	GizmoComponent->SetWorldTransform(NewTransform);
+	//ActiveTarget->SetTransform(NewTransform);		// this will happen in the GizmoComponent.TransformUpdated delegate handler above
+
+	StateTarget->EndUpdate();
+}
+
+
 UInteractiveGizmo* UTransformGizmo::AddAxisTranslationGizmo(
 	UPrimitiveComponent* AxisComponent, USceneComponent* RootComponent,
 	IGizmoAxisSource* AxisSource,
 	IGizmoTransformSource* TransformSource,
-	IGizmoStateTarget* StateTarget)
+	IGizmoStateTarget* StateTargetIn)
 {
 	// create axis-position gizmo, axis-position parameter will drive translation
 	UAxisPositionGizmo* TranslateGizmo = Cast<UAxisPositionGizmo>(GetGizmoManager()->CreateGizmo(
@@ -350,7 +365,7 @@ UInteractiveGizmo* UTransformGizmo::AddAxisTranslationGizmo(
 	}
 	TranslateGizmo->HitTarget = HitTarget;
 
-	TranslateGizmo->StateTarget = Cast<UObject>(StateTarget);
+	TranslateGizmo->StateTarget = Cast<UObject>(StateTargetIn);
 
 	ActiveGizmos.Add(TranslateGizmo);
 	return TranslateGizmo;
@@ -362,7 +377,7 @@ UInteractiveGizmo* UTransformGizmo::AddPlaneTranslationGizmo(
 	UPrimitiveComponent* AxisComponent, USceneComponent* RootComponent,
 	IGizmoAxisSource* AxisSource,
 	IGizmoTransformSource* TransformSource,
-	IGizmoStateTarget* StateTarget)
+	IGizmoStateTarget* StateTargetIn)
 {
 	// create axis-position gizmo, axis-position parameter will drive translation
 	UPlanePositionGizmo* TranslateGizmo = Cast<UPlanePositionGizmo>(GetGizmoManager()->CreateGizmo(
@@ -383,7 +398,7 @@ UInteractiveGizmo* UTransformGizmo::AddPlaneTranslationGizmo(
 	}
 	TranslateGizmo->HitTarget = HitTarget;
 
-	TranslateGizmo->StateTarget = Cast<UObject>(StateTarget);
+	TranslateGizmo->StateTarget = Cast<UObject>(StateTargetIn);
 
 	ActiveGizmos.Add(TranslateGizmo);
 	return TranslateGizmo;
@@ -397,7 +412,7 @@ UInteractiveGizmo* UTransformGizmo::AddAxisRotationGizmo(
 	UPrimitiveComponent* AxisComponent, USceneComponent* RootComponent,
 	IGizmoAxisSource* AxisSource,
 	IGizmoTransformSource* TransformSource,
-	IGizmoStateTarget* StateTarget)
+	IGizmoStateTarget* StateTargetIn)
 {
 	// create axis-angle gizmo, angle will drive axis-rotation
 	UAxisAngleGizmo* RotateGizmo = Cast<UAxisAngleGizmo>(GetGizmoManager()->CreateGizmo(
@@ -418,7 +433,7 @@ UInteractiveGizmo* UTransformGizmo::AddAxisRotationGizmo(
 	}
 	RotateGizmo->HitTarget = HitTarget;
 
-	RotateGizmo->StateTarget = Cast<UObject>(StateTarget);
+	RotateGizmo->StateTarget = Cast<UObject>(StateTargetIn);
 
 	ActiveGizmos.Add(RotateGizmo);
 
@@ -438,6 +453,11 @@ void UTransformGizmo::ClearActiveTarget()
 	ActiveGizmos.SetNum(0);
 	ActiveComponents.SetNum(0);
 
+	AxisXSource = nullptr;
+	AxisYSource = nullptr;
+	AxisZSource = nullptr;
+	StateTarget = nullptr;
+
 	ActiveTarget = nullptr;
 }
 
diff --git a/Engine/Source/Runtime/Experimental/InteractiveToolsFramework/Private/BaseGizmos/TransformProxy.cpp b/Engine/Source/Runtime/Experimental/InteractiveToolsFramework/Private/BaseGizmos/TransformProxy.cpp
--- a/Engine/Source/Runtime/Experimental/InteractiveToolsFramework/Private/BaseGizmos/TransformProxy.cpp
+++ b/Engine/Source/Runtime/Experimental/InteractiveToolsFramework/Private/BaseGizmos/TransformProxy.cpp
@@ -4,6 +4,9 @@
 #include "Components/SceneComponent.h"
 
 
+#define LOCTEXT_NAMESPACE "UTransformProxy"
+
+
 void UTransformProxy::AddComponent(USceneComponent* Component, bool bModifyComponentOnTransform)
 {
 	check(Component);
@@ -20,6 +23,7 @@ void UTransformProxy::AddComponent(USceneComponent* Component, bool bModifyCompo
 }
 
 
+
 FTransform UTransformProxy::GetTransform() const
 {
 	return SharedTransform;
@@ -122,4 +126,53 @@ void UTransformProxy::UpdateObjectTransforms()
 		Obj.RelativeTransform = Obj.StartTransform;
 		Obj.RelativeTransform.SetToRelativeTransform(SharedTransform);
 	}
-}
\ No newline at end of file
+}
+
+
+
+
+
+void FTransformProxyChange::Apply(UObject* Object)
+{
+	UTransformProxy* Proxy = CastChecked<UTransformProxy>(Object);
+	Proxy->SetTransform(To);
+}
+
+void FTransformProxyChange::Revert(UObject* Object)
+{
+	UTransformProxy* Proxy = CastChecked<UTransformProxy>(Object);
+	Proxy->SetTransform(From);
+}
+
+
+void FTransformProxyChangeSource::BeginChange()
+{
+	if (Proxy.IsValid())
+	{
+		ActiveChange = MakeUnique<FTransformProxyChange>();
+		ActiveChange->From = Proxy->GetTransform();
+	}
+}
+
+TUniquePtr<FToolCommandChange> FTransformProxyChangeSource::EndChange()
+{
+	if (Proxy.IsValid())
+	{
+		ActiveChange->To = Proxy->GetTransform();
+		return MoveTemp(ActiveChange);
+	}
+	return TUniquePtr<FToolCommandChange>();
+}
+
+UObject* FTransformProxyChangeSource::GetChangeTarget()
+{
+	return Proxy.Get();
+}
+
+FText FTransformProxyChangeSource::GetChangeDescription()
+{
+	return LOCTEXT("FTransformProxyChangeDescription", "TransformProxyChange");
+}
+
+
+#undef LOCTEXT_NAMESPACE
\ No newline at end of file
diff --git a/Engine/Source/Runtime/Experimental/InteractiveToolsFramework/Public/BaseGizmos/StateTargets.h b/Engine/Source/Runtime/Experimental/InteractiveToolsFramework/Public/BaseGizmos/StateTargets.h
--- a/Engine/Source/Runtime/Experimental/InteractiveToolsFramework/Public/BaseGizmos/StateTargets.h
+++ b/Engine/Source/Runtime/Experimental/InteractiveToolsFramework/Public/BaseGizmos/StateTargets.h
@@ -155,6 +155,11 @@ public:
 			}
 
 			InitialTransform = TargetComponent->GetComponentTransform();
+
+			for (TUniquePtr<IToolCommandChangeSource>& Source : DependentChangeSources)
+			{
+				Source->BeginChange();
+			}
 		}
 	}
 
@@ -169,6 +174,17 @@ public:
 				TUniquePtr<FComponentWorldTransformChange> TransformChange 
 					= MakeUnique<FComponentWorldTransformChange>(InitialTransform, FinalTransform);
 				TransactionManager->EmitObjectChange(TargetComponent.Get(), MoveTemp(TransformChange), ChangeDescription);
+
+				for (TUniquePtr<IToolCommandChangeSource>& Source : DependentChangeSources)
+				{
+					TUniquePtr<FToolCommandChange> Change = Source->EndChange();
+					if (Change)
+					{
+						UObject* Target = Source->GetChangeTarget();
+						TransactionManager->EmitObjectChange(Target, MoveTemp(Change), Source->GetChangeDescription());
+					}
+				}
+
 				TransactionManager->EndUndoTransaction();
 			}
 		}
@@ -197,6 +213,12 @@ public:
 	FTransform FinalTransform;
 
 
+	/** 
+	 * Dependent-change generators. This will be told about update start/end, and any generated changes will also be emitted. 
+	 * This allows (eg) TransformProxy change events to be collected at the same time as changes to a gizmo target component.
+	 */
+	TArray<TUniquePtr<IToolCommandChangeSource>> DependentChangeSources;
+
 public:
 
 	/**
diff --git a/Engine/Source/Runtime/Experimental/InteractiveToolsFramework/Public/BaseGizmos/TransformGizmo.h b/Engine/Source/Runtime/Experimental/InteractiveToolsFramework/Public/BaseGizmos/TransformGizmo.h
--- a/Engine/Source/Runtime/Experimental/InteractiveToolsFramework/Public/BaseGizmos/TransformGizmo.h
+++ b/Engine/Source/Runtime/Experimental/InteractiveToolsFramework/Public/BaseGizmos/TransformGizmo.h
@@ -6,6 +6,7 @@
 #include "InteractiveGizmo.h"
 #include "InteractiveGizmoBuilder.h"
 #include "InteractiveToolObjects.h"
+#include "InteractiveToolChange.h"
 #include "BaseGizmos/GizmoActor.h"
 #include "BaseGizmos/TransformProxy.h"
 #include "TransformGizmo.generated.h"
@@ -15,6 +16,7 @@ class IGizmoAxisSource;
 class IGizmoTransformSource;
 class IGizmoStateTarget;
 class UGizmoComponentAxisSource;
+class UGizmoTransformChangeStateTarget;
 
 
 /**
@@ -217,9 +219,17 @@ public:
 	UPROPERTY()
 	UTransformProxy* ActiveTarget;
 
-
+	/**
+	 * @return the internal GizmoActor used by the Gizmo
+	 */
 	ATransformGizmoActor* GetGizmoActor() const { return GizmoActor; }
 
+	/**
+	 * Set a new position for the Gizmo. This is done via the same mechanisms as the sub-gizmos,
+	 * so it generates the same Change/Modify() events, and hence works with Undo/Redo
+	 */
+	virtual void SetNewGizmoTransform(const FTransform& NewTransform);
+
 
 protected:
 	TSharedPtr<FTransformGizmoActorFactory> GizmoActorBuilder;
@@ -260,6 +270,10 @@ protected:
 	UPROPERTY()
 	UGizmoComponentAxisSource* AxisZSource;
 
+	/** State target is shared across gizmos, and created internally during SetActiveTarget() */
+	UPROPERTY()
+	UGizmoTransformChangeStateTarget* StateTarget;
+
 	/** @return a new instance of the standard axis-translation Gizmo */
 	virtual UInteractiveGizmo* AddAxisTranslationGizmo(
 		UPrimitiveComponent* AxisComponent, USceneComponent* RootComponent,
diff --git a/Engine/Source/Runtime/Experimental/InteractiveToolsFramework/Public/BaseGizmos/TransformProxy.h b/Engine/Source/Runtime/Experimental/InteractiveToolsFramework/Public/BaseGizmos/TransformProxy.h
--- a/Engine/Source/Runtime/Experimental/InteractiveToolsFramework/Public/BaseGizmos/TransformProxy.h
+++ b/Engine/Source/Runtime/Experimental/InteractiveToolsFramework/Public/BaseGizmos/TransformProxy.h
@@ -4,6 +4,7 @@
 
 #include "CoreMinimal.h"
 #include "UObject/ObjectMacros.h"
+#include "BaseGizmos/StateTargets.h"
 #include "TransformProxy.generated.h"
 
 class USceneComponent;
@@ -97,4 +98,44 @@ protected:
 
 	/** Propagate a transform update to the sub-objects */
 	virtual void UpdateObjects();
+};
+
+
+
+/**
+ * FTransformProxyChange tracks a change to the base transform for a TransformProxy
+ */
+class INTERACTIVETOOLSFRAMEWORK_API FTransformProxyChange : public FToolCommandChange
+{
+public:
+	FTransform From;
+	FTransform To;
+
+	virtual void Apply(UObject* Object) override;
+	virtual void Revert(UObject* Object) override;
+
+	virtual FString ToString() const override { return TEXT("FTransformProxyChange"); }
+};
+
+/**
+ * FTransformProxyChangeSource generates FTransformProxyChange instances on Begin/End.
+ * Instances of this class can (for example) be attached to a UGizmoTransformChangeStateTarget for use TransformGizmo change tracking.
+ */
+class INTERACTIVETOOLSFRAMEWORK_API FTransformProxyChangeSource : public IToolCommandChangeSource
+{
+public:
+	FTransformProxyChangeSource(UTransformProxy* ProxyIn)
+	{
+		Proxy = ProxyIn;
+	}
+
+	virtual ~FTransformProxyChangeSource() {}
+
+	TWeakObjectPtr<UTransformProxy> Proxy;
+	TUniquePtr<FTransformProxyChange> ActiveChange;
+
+	virtual void BeginChange() override;
+	virtual TUniquePtr<FToolCommandChange> EndChange() override;
+	virtual UObject* GetChangeTarget() override;
+	virtual FText GetChangeDescription() override;
 };
\ No newline at end of file
