fixes to keyasmodifierinputbehavior, to make it correctly respond to key releases

#rnx

diff --git a/Engine/Source/Runtime/Experimental/InteractiveToolsFramework/Private/BaseBehaviors/KeyAsModifierInputBehavior.cpp b/Engine/Source/Runtime/Experimental/InteractiveToolsFramework/Private/BaseBehaviors/KeyAsModifierInputBehavior.cpp
--- a/Engine/Source/Runtime/Experimental/InteractiveToolsFramework/Private/BaseBehaviors/KeyAsModifierInputBehavior.cpp
+++ b/Engine/Source/Runtime/Experimental/InteractiveToolsFramework/Private/BaseBehaviors/KeyAsModifierInputBehavior.cpp
@@ -13,7 +13,7 @@ void UKeyAsModifierInputBehavior::Initialize(IModifierToggleBehaviorTarget* Targ
 	FKey TempKey = ModifierKeyIn;
 	Modifiers.RegisterModifier(ModifierID, [TempKey](const FInputDeviceState& Input)
 	{
-		return Input.Keyboard.ActiveKey.Button == TempKey;
+		return (TempKey == EKeys::AnyKey || Input.Keyboard.ActiveKey.Button == TempKey) && Input.Keyboard.ActiveKey.bPressed;
 	});
 }
 
@@ -21,7 +21,7 @@ FInputCaptureRequest UKeyAsModifierInputBehavior::WantsCapture(const FInputDevic
 {
 	if ((ModifierCheckFunc == nullptr || ModifierCheckFunc(Input)))
 	{
-		if (Input.Keyboard.ActiveKey.Button == ModifierKey && Input.Keyboard.ActiveKey.bPressed)
+		if ((ModifierKey == EKeys::AnyKey || Input.Keyboard.ActiveKey.Button == ModifierKey) && Input.Keyboard.ActiveKey.bPressed)
 		{
 			return FInputCaptureRequest::Begin(this, EInputCaptureSide::Any);
 		}
@@ -32,6 +32,7 @@ FInputCaptureRequest UKeyAsModifierInputBehavior::WantsCapture(const FInputDevic
 
 FInputCaptureUpdate UKeyAsModifierInputBehavior::BeginCapture(const FInputDeviceState& Input, EInputCaptureSide Side)
 {
+	PressedButton = Input.Keyboard.ActiveKey.Button;
 	Modifiers.UpdateModifiers(Input, Target);
 	return FInputCaptureUpdate::Begin(this, EInputCaptureSide::Any);
 }
@@ -39,6 +40,11 @@ FInputCaptureUpdate UKeyAsModifierInputBehavior::BeginCapture(const FInputDevice
 
 FInputCaptureUpdate UKeyAsModifierInputBehavior::UpdateCapture(const FInputDeviceState& Input, const FInputCaptureData& Data)
 {
+	if (Input.Keyboard.ActiveKey.Button != PressedButton)
+	{
+		return FInputCaptureUpdate::Continue();
+	}
+
 	Modifiers.UpdateModifiers(Input, Target);
 
 	if (Input.Keyboard.ActiveKey.bReleased)
diff --git a/Engine/Source/Runtime/Experimental/InteractiveToolsFramework/Public/BaseBehaviors/KeyAsModifierInputBehavior.h b/Engine/Source/Runtime/Experimental/InteractiveToolsFramework/Public/BaseBehaviors/KeyAsModifierInputBehavior.h
--- a/Engine/Source/Runtime/Experimental/InteractiveToolsFramework/Public/BaseBehaviors/KeyAsModifierInputBehavior.h
+++ b/Engine/Source/Runtime/Experimental/InteractiveToolsFramework/Public/BaseBehaviors/KeyAsModifierInputBehavior.h
@@ -58,4 +58,7 @@ protected:
 
 	/** Modifier set for this behavior, internally initialized with check on ModifierKey */
 	FInputBehaviorModifierStates Modifiers;
+
+	/** The key that was pressed to activate capture */
+	FKey PressedButton;
 };
