ModelingTools: add some utility functions for selection list construction
#rnx

diff --git a/Engine/Plugins/Experimental/MeshModelingToolset/Source/ModelingComponents/Private/Changes/MeshSelectionChange.cpp b/Engine/Plugins/Experimental/MeshModelingToolset/Source/ModelingComponents/Private/Changes/MeshSelectionChange.cpp
--- a/Engine/Plugins/Experimental/MeshModelingToolset/Source/ModelingComponents/Private/Changes/MeshSelectionChange.cpp
+++ b/Engine/Plugins/Experimental/MeshModelingToolset/Source/ModelingComponents/Private/Changes/MeshSelectionChange.cpp
@@ -40,7 +40,6 @@ FString FMeshSelectionChange::ToString() const
 
 
 
-
 FMeshSelectionChangeBuilder::FMeshSelectionChangeBuilder(EMeshSelectionElementType ElementType, bool bAdding)
 {
 	Change = MakeUnique<FMeshSelectionChange>();
@@ -48,7 +47,26 @@ FMeshSelectionChangeBuilder::FMeshSelectionChangeBuilder(EMeshSelectionElementTy
 	Change->bAdded = bAdding;
 }
 
-void FMeshSelectionChangeBuilder::Add(int ElementID)
+void FMeshSelectionChangeBuilder::Add(int32 ElementID)
 {
 	Change->Indices.Add(ElementID);
 }
+
+
+
+void FMeshSelectionChangeBuilder::Add(const TArray<int32>& Elements)
+{
+	for (int32 Index : Elements)
+	{
+		Change->Indices.Add(Index);
+	}
+}
+
+void FMeshSelectionChangeBuilder::Add(const TSet<int32>& Elements)
+{
+	for (int32 Index : Elements)
+	{
+		Change->Indices.Add(Index);
+	}
+}
+
diff --git a/Engine/Plugins/Experimental/MeshModelingToolset/Source/ModelingComponents/Public/Changes/MeshSelectionChange.h b/Engine/Plugins/Experimental/MeshModelingToolset/Source/ModelingComponents/Public/Changes/MeshSelectionChange.h
--- a/Engine/Plugins/Experimental/MeshModelingToolset/Source/ModelingComponents/Public/Changes/MeshSelectionChange.h
+++ b/Engine/Plugins/Experimental/MeshModelingToolset/Source/ModelingComponents/Public/Changes/MeshSelectionChange.h
@@ -16,7 +16,7 @@ class MODELINGCOMPONENTS_API FMeshSelectionChange : public FToolCommandChange
 {
 public:
 	EMeshSelectionElementType ElementType = EMeshSelectionElementType::Vertex;
-	TArray<int> Indices;
+	TArray<int32> Indices;
 	bool bAdded;
 
 	/** Makes the change to the object */
@@ -50,5 +50,15 @@ public:
 	/**
 	 * Add ElementID to list of changed elements  (which may then be added or removed, depending on bAdding parameter to constructor)
 	 */
-	void Add(int ElementID);
+	void Add(int32 ElementID);
+
+	/**
+	 * Add Array of elements
+	 */
+	void Add(const TArray<int32>& Elements);
+
+	/**
+	 * Add Set of elements
+	 */
+	void Add(const TSet<int32>& Elements);
 };
diff --git a/Engine/Source/Runtime/Experimental/InteractiveToolsFramework/Private/SelectionSet.cpp b/Engine/Source/Runtime/Experimental/InteractiveToolsFramework/Private/SelectionSet.cpp
--- a/Engine/Source/Runtime/Experimental/InteractiveToolsFramework/Private/SelectionSet.cpp
+++ b/Engine/Source/Runtime/Experimental/InteractiveToolsFramework/Private/SelectionSet.cpp
@@ -45,9 +45,9 @@ const TArray<int>& UMeshSelectionSet::GetElements(EMeshSelectionElementType Elem
 
 
 
-void UMeshSelectionSet::AddIndices(EMeshSelectionElementType ElementType, const TArray<int>& Indices)
+void UMeshSelectionSet::AddIndices(EMeshSelectionElementType ElementType, const TArray<int32>& Indices)
 {
-	TArray<int>& CurElements = GetElements(ElementType);
+	TArray<int32>& CurElements = GetElements(ElementType);
 
 	int N = Indices.Num();
 	for (int k = 0; k < N; ++k)
@@ -57,16 +57,41 @@ void UMeshSelectionSet::AddIndices(EMeshSelectionElementType ElementType, const
 	NotifySelectionSetModified();
 }
 
-
-void UMeshSelectionSet::RemoveIndices(EMeshSelectionElementType ElementType, const TArray<int>& Indices)
+void UMeshSelectionSet::AddIndices(EMeshSelectionElementType ElementType, const TSet<int32>& Indices)
 {
 	TArray<int>& CurElements = GetElements(ElementType);
 
+	for ( int32 Index : Indices )
+	{
+		CurElements.Add(Index);
+	}
+	NotifySelectionSetModified();
+}
+
+
+
+void UMeshSelectionSet::RemoveIndices(EMeshSelectionElementType ElementType, const TArray<int32>& Indices)
+{
+	TArray<int32>& CurElements = GetElements(ElementType);
+
 	// @todo if we are removing many elements it is maybe cheaper to make a new array...
 	int N = Indices.Num();
-	for (int k = 0; k < N; ++k)
+	for (int32 k = 0; k < N; ++k)
 	{
 		CurElements.RemoveSwap(Indices[k]);
 	}
 	NotifySelectionSetModified();
+}
+
+
+void UMeshSelectionSet::RemoveIndices(EMeshSelectionElementType ElementType, const TSet<int32>& Indices)
+{
+	TArray<int32>& CurElements = GetElements(ElementType);
+
+	// @todo if we are removing many elements it is maybe cheaper to make a new array...
+	for (int32 Index : Indices)
+	{
+		CurElements.RemoveSwap(Index);
+	}
+	NotifySelectionSetModified();
 }
\ No newline at end of file
diff --git a/Engine/Source/Runtime/Experimental/InteractiveToolsFramework/Public/SelectionSet.h b/Engine/Source/Runtime/Experimental/InteractiveToolsFramework/Public/SelectionSet.h
--- a/Engine/Source/Runtime/Experimental/InteractiveToolsFramework/Public/SelectionSet.h
+++ b/Engine/Source/Runtime/Experimental/InteractiveToolsFramework/Public/SelectionSet.h
@@ -91,11 +91,22 @@ public:
 	 * Add list of Indices to the current selection for the given ElementType
 	 * @warning no duplicate detection is currently done
 	 */
-	void AddIndices(EMeshSelectionElementType ElementType, const TArray<int>& Indices);
-	
+	void AddIndices(EMeshSelectionElementType ElementType, const TArray<int32>& Indices);
+
+	/**
+	 * Add list of Indices to the current selection for the given ElementType
+	 * @warning no duplicate detection is currently done
+	 */
+	void AddIndices(EMeshSelectionElementType ElementType, const TSet<int32>& Indices);
+
+
 	/**
 	 * Remove list of Indices from the current selection for the given ElementType
 	 */
-	void RemoveIndices(EMeshSelectionElementType ElementType, const TArray<int>& Indices);
+	void RemoveIndices(EMeshSelectionElementType ElementType, const TArray<int32>& Indices);
 
+	/**
+	 * Remove list of Indices from the current selection for the given ElementType
+	 */
+	void RemoveIndices(EMeshSelectionElementType ElementType, const TSet<int32>& Indices);
 };
\ No newline at end of file
