Moving generic notification functions from EdModeInteractiveToolsContext to InteractiveToolsContext

diff --git a/Engine/Source/Editor/Experimental/EditorInteractiveToolsFramework/Private/EdModeInteractiveToolsContext.cpp b/Engine/Source/Editor/Experimental/EditorInteractiveToolsFramework/Private/EdModeInteractiveToolsContext.cpp
--- a/Engine/Source/Editor/Experimental/EditorInteractiveToolsFramework/Private/EdModeInteractiveToolsContext.cpp
+++ b/Engine/Source/Editor/Experimental/EditorInteractiveToolsFramework/Private/EdModeInteractiveToolsContext.cpp
@@ -414,18 +414,6 @@ void UEdModeInteractiveToolsContext::ShutdownContext()
 }
 
 
-void UEdModeInteractiveToolsContext::PostToolNotificationMessage(const FText& Message)
-{
-	OnToolNotificationMessage.Broadcast(Message);
-}
-
-void UEdModeInteractiveToolsContext::PostToolWarningMessage(const FText& Message)
-{
-	OnToolWarningMessage.Broadcast(Message);
-}
-
-
-
 void UEdModeInteractiveToolsContext::TerminateActiveToolsOnPIEStart()
 {
 	DeactivateAllActiveTools();
@@ -435,10 +423,6 @@ void UEdModeInteractiveToolsContext::TerminateActiveToolsOnSaveWorld()
 	DeactivateAllActiveTools();
 }
 
-
-
-
-
 void UEdModeInteractiveToolsContext::PostInvalidation()
 {
 	bInvalidationPending = true;
@@ -534,7 +518,7 @@ bool UEdModeInteractiveToolsContext::InputKey(FEditorViewportClient* ViewportCli
 	{
 		if (ToolManager->HasAnyActiveTool())
 		{
-			if (ToolManager->HasActiveTool(EToolSide::Mouse))
+			if (ToolManager->HasActiveTool(EToolSide::Mouse) && ToolManager->CanCancelActiveTool(EToolSide::Mouse))
 			{
 				DeactivateActiveTool(EToolSide::Mouse, EToolShutdownType::Cancel);
 			}
diff --git a/Engine/Source/Editor/Experimental/EditorInteractiveToolsFramework/Public/EdModeInteractiveToolsContext.h b/Engine/Source/Editor/Experimental/EditorInteractiveToolsFramework/Public/EdModeInteractiveToolsContext.h
--- a/Engine/Source/Editor/Experimental/EditorInteractiveToolsFramework/Public/EdModeInteractiveToolsContext.h
+++ b/Engine/Source/Editor/Experimental/EditorInteractiveToolsFramework/Public/EdModeInteractiveToolsContext.h
@@ -69,17 +69,6 @@ public:
 	virtual void StartTool(const FString& ToolTypeIdentifier);
 	virtual void EndTool(EToolShutdownType ShutdownType);
 
-
-public:
-	// forwards message to OnToolNotificationMessage delegate
-	virtual void PostToolNotificationMessage(const FText& Message);
-	virtual void PostToolWarningMessage(const FText& Message);
-
-	DECLARE_MULTICAST_DELEGATE_OneParam(FEdModeToolsContextToolNotification, const FText&);
-	FEdModeToolsContextToolNotification OnToolNotificationMessage;
-	FEdModeToolsContextToolNotification OnToolWarningMessage;
-
-
 protected:
 	// we hide these 
 	virtual void Initialize(IToolsContextQueriesAPI* QueriesAPI, IToolsContextTransactionsAPI* TransactionsAPI) override;
diff --git a/Engine/Source/Runtime/Experimental/InteractiveToolsFramework/Private/InteractiveToolsContext.cpp b/Engine/Source/Runtime/Experimental/InteractiveToolsFramework/Private/InteractiveToolsContext.cpp
--- a/Engine/Source/Runtime/Experimental/InteractiveToolsFramework/Private/InteractiveToolsContext.cpp
+++ b/Engine/Source/Runtime/Experimental/InteractiveToolsFramework/Private/InteractiveToolsContext.cpp
@@ -8,6 +8,7 @@ UInteractiveToolsContext::UInteractiveToolsContext()
 {
 	InputRouter = nullptr;
 	ToolManager = nullptr;
+	ToolManagerClass = UInteractiveToolManager::StaticClass();
 }
 
 void UInteractiveToolsContext::Initialize(IToolsContextQueriesAPI* QueriesAPI, IToolsContextTransactionsAPI* TransactionsAPI)
@@ -15,7 +16,8 @@ void UInteractiveToolsContext::Initialize(IToolsContextQueriesAPI* QueriesAPI, I
 	InputRouter = NewObject<UInputRouter>(this);
 	InputRouter->Initialize(TransactionsAPI);
 
-	ToolManager = NewObject<UInteractiveToolManager>(this);
+	ToolManager = NewObject<UInteractiveToolManager>(this, ToolManagerClass.Get());
+
 	ToolManager->Initialize(QueriesAPI, TransactionsAPI, InputRouter);
 
 	GizmoManager = NewObject<UInteractiveGizmoManager>(this);
@@ -106,3 +108,18 @@ void UInteractiveToolsContext::EndTool(EToolSide WhichSide, EToolShutdownType Sh
 {
 	DeactivateActiveTool(WhichSide, ShutdownType);
 }
+
+bool UInteractiveToolsContext::IsToolActive(EToolSide WhichSide, UInteractiveToolBuilder* Builder)
+{
+	return ToolManager->GetActiveToolBuilder(WhichSide) == Builder;
+}
+
+void UInteractiveToolsContext::PostToolNotificationMessage(const FText& Message)
+{
+	OnToolNotificationMessage.Broadcast(Message);
+}
+
+void UInteractiveToolsContext::PostToolWarningMessage(const FText& Message)
+{
+	OnToolWarningMessage.Broadcast(Message);
+}
diff --git a/Engine/Source/Runtime/Experimental/InteractiveToolsFramework/Public/InteractiveToolsContext.h b/Engine/Source/Runtime/Experimental/InteractiveToolsFramework/Public/InteractiveToolsContext.h
--- a/Engine/Source/Runtime/Experimental/InteractiveToolsFramework/Public/InteractiveToolsContext.h
+++ b/Engine/Source/Runtime/Experimental/InteractiveToolsFramework/Public/InteractiveToolsContext.h
@@ -40,6 +40,17 @@ public:
 	bool CanCompleteActiveTool(EToolSide WhichSide) const;
 	bool StartTool(EToolSide WhichSide, const FString& ToolTypeIdentifier);
 	void EndTool(EToolSide WhichSide, EToolShutdownType ShutdownType);
+	bool IsToolActive(EToolSide WhichSide, UInteractiveToolBuilder* Builder);
+
+public:
+	// forwards message to OnToolNotificationMessage delegate
+	virtual void PostToolNotificationMessage(const FText& Message);
+	virtual void PostToolWarningMessage(const FText& Message);
+
+	DECLARE_MULTICAST_DELEGATE_OneParam(FToolsContextToolNotification, const FText&);
+	FToolsContextToolNotification OnToolNotificationMessage;
+	FToolsContextToolNotification OnToolWarningMessage;
+
 public:
 	/** current UInputRouter for this Context */
 	UPROPERTY()
@@ -52,4 +63,8 @@ public:
 	/** current UInteractiveGizmoManager for this Context */
 	UPROPERTY()
 	UInteractiveGizmoManager* GizmoManager;
+
+protected:
+	UPROPERTY()
+	TSoftClassPtr<UInteractiveToolManager> ToolManagerClass;
 };
\ No newline at end of file
