Moving some of the more generic brush and tool systems from mesh modeling plugin to InteractiveToolsFramework

diff --git a/Engine/Config/BaseEngine.ini b/Engine/Config/BaseEngine.ini
--- a/Engine/Config/BaseEngine.ini
+++ b/Engine/Config/BaseEngine.ini
@@ -732,6 +732,9 @@ bUseStreamingPause=false
 +EnumRedirects=(OldName="/Script/Engine.EPreviewAnimationBlueprintApplicationMethod",NewName="/Script/Engine.EPreviewAnimationBlueprintApplicationMethod",ValueChanges=(("OverlayLayer", "LinkedLayers"), ("SubInstance", "LinkedAnimGraph"))
 +PropertyRedirects=(OldName="AnimClassData.SubInstanceNodeProperties",NewName="LinkedAnimGraphNodeProperties")
 +PropertyRedirects=(OldName="AnimClassData.LayerNodeProperties",NewName="LinkedAnimLayerNodeProperties")
++ClassRedirects=(OldName="/Script/MeshEditingToolset.BaseBrushTool",NewName="/Script/InteractiveToolsFramework.BaseBrushTool")
++ClassRedirects=(OldName="/Script/MeshEditingToolset.BrushBaseProperties",NewName="/Script/InteractiveToolsFramework.BrushBaseProperties")
+
 
 [CoreUObject.Metadata]
 ; Note: UnrealHeaderTool should be rerun after making changes to MetadataRedirects to catch any keys specified in class headers
diff --git a/Engine/Plugins/Experimental/MeshModelingToolset/Source/MeshModelingTools/Private/BaseBrushTool.cpp b/Engine/Plugins/Experimental/MeshModelingToolset/Source/MeshModelingTools/Private/BaseBrushTool.cpp
deleted file mode 100644
diff --git a/Engine/Plugins/Experimental/MeshModelingToolset/Source/MeshModelingTools/Private/DynamicMeshSculptTool.cpp b/Engine/Plugins/Experimental/MeshModelingToolset/Source/MeshModelingTools/Private/DynamicMeshSculptTool.cpp
--- a/Engine/Plugins/Experimental/MeshModelingToolset/Source/MeshModelingTools/Private/DynamicMeshSculptTool.cpp
+++ b/Engine/Plugins/Experimental/MeshModelingToolset/Source/MeshModelingTools/Private/DynamicMeshSculptTool.cpp
@@ -14,13 +14,16 @@
 #include "MeshIndexUtil.h"
 #include "Drawing/MeshDebugDrawing.h"
 #include "PreviewMesh.h"
-#include "ToolSetupUtil.h"
+#include "BaseTools/ToolSetupUtil.h"
 
 #include "Changes/MeshVertexChange.h"
 #include "Changes/MeshChange.h"
 #include "DynamicMeshChangeTracker.h"
 
 #include "Async/ParallelFor.h"
+#include "ToolDataVisualizer.h"
+#include "Components/PrimitiveComponent.h"
+#include "Generators/SphereGenerator.h"
 
 #define LOCTEXT_NAMESPACE "UDynamicMeshSculptTool"
 
@@ -167,7 +170,7 @@ void UDynamicMeshSculptTool::Setup()
 	// register and spawn brush indicator gizmo
 	GetToolManager()->GetPairedGizmoManager()->RegisterGizmoType(BrushIndicatorGizmoType, NewObject<UBrushStampIndicatorBuilder>());
 	BrushIndicator = GetToolManager()->GetPairedGizmoManager()->CreateGizmo<UBrushStampIndicator>(BrushIndicatorGizmoType, FString(), this);
-	BrushIndicatorMesh = UBrushStampIndicator::MakeDefaultSphereMesh(this, TargetWorld);
+	BrushIndicatorMesh = MakeDefaultSphereMesh(this, TargetWorld);
 	BrushIndicator->AttachedComponent = BrushIndicatorMesh->GetRootComponent();
 	BrushIndicator->bDrawIndicatorLines = false;
 
@@ -1534,6 +1537,19 @@ double UDynamicMeshSculptTool::EstimateIntialSafeTargetLength(const FDynamicMesh
 
 
 
+UPreviewMesh* UDynamicMeshSculptTool::MakeDefaultSphereMesh(UObject* Parent, UWorld* World, int Resolution /*= 32*/)
+{
+	UPreviewMesh* SphereMesh = NewObject<UPreviewMesh>(Parent);
+	SphereMesh->CreateInWorld(World, FTransform::Identity);
+	FSphereGenerator SphereGen;
+	SphereGen.NumPhi = SphereGen.NumTheta = Resolution;
+	SphereGen.Generate();
+	FDynamicMesh3 Mesh(&SphereGen);
+	SphereMesh->UpdatePreview(&Mesh);
+	SphereMesh->SetMaterial(ToolSetupUtil::GetDefaultBrushVolumeMaterial(nullptr));
+	return SphereMesh;
+}
+
 void UDynamicMeshSculptTool::IncreaseBrushRadiusAction()
 {
 	BrushProperties->BrushSize = FMath::Clamp(BrushProperties->BrushSize + 0.025f, 0.0f, 1.0f);
diff --git a/Engine/Plugins/Experimental/MeshModelingToolset/Source/MeshModelingTools/Private/EditNormalsTool.cpp b/Engine/Plugins/Experimental/MeshModelingToolset/Source/MeshModelingTools/Private/EditNormalsTool.cpp
--- a/Engine/Plugins/Experimental/MeshModelingToolset/Source/MeshModelingTools/Private/EditNormalsTool.cpp
+++ b/Engine/Plugins/Experimental/MeshModelingToolset/Source/MeshModelingTools/Private/EditNormalsTool.cpp
@@ -4,7 +4,7 @@
 #include "InteractiveToolManager.h"
 #include "ToolBuilderUtil.h"
 
-#include "ToolSetupUtil.h"
+#include "BaseTools/ToolSetupUtil.h"
 
 #include "DynamicMesh3.h"
 #include "BaseBehaviors/MultiClickSequenceInputBehavior.h"
diff --git a/Engine/Plugins/Experimental/MeshModelingToolset/Source/MeshModelingTools/Private/MeshSelectionTool.cpp b/Engine/Plugins/Experimental/MeshModelingToolset/Source/MeshModelingTools/Private/MeshSelectionTool.cpp
--- a/Engine/Plugins/Experimental/MeshModelingToolset/Source/MeshModelingTools/Private/MeshSelectionTool.cpp
+++ b/Engine/Plugins/Experimental/MeshModelingToolset/Source/MeshModelingTools/Private/MeshSelectionTool.cpp
@@ -10,7 +10,7 @@
 #include "Changes/MeshChange.h"
 #include "MeshIndexUtil.h"
 #include "AssetGenerationUtil.h"
-#include "ToolSetupUtil.h"
+#include "BaseTools/ToolSetupUtil.h"
 
 #define LOCTEXT_NAMESPACE "UMeshSelectionTool"
 
diff --git a/Engine/Plugins/Experimental/MeshModelingToolset/Source/MeshModelingTools/Private/PlaneCutTool.cpp b/Engine/Plugins/Experimental/MeshModelingToolset/Source/MeshModelingTools/Private/PlaneCutTool.cpp
--- a/Engine/Plugins/Experimental/MeshModelingToolset/Source/MeshModelingTools/Private/PlaneCutTool.cpp
+++ b/Engine/Plugins/Experimental/MeshModelingToolset/Source/MeshModelingTools/Private/PlaneCutTool.cpp
@@ -4,7 +4,7 @@
 #include "InteractiveToolManager.h"
 #include "ToolBuilderUtil.h"
 
-#include "ToolSetupUtil.h"
+#include "BaseTools/ToolSetupUtil.h"
 
 #include "DynamicMesh3.h"
 #include "BaseBehaviors/MultiClickSequenceInputBehavior.h"
diff --git a/Engine/Plugins/Experimental/MeshModelingToolset/Source/MeshModelingTools/Private/RemoveOccludedTrianglesTool.cpp b/Engine/Plugins/Experimental/MeshModelingToolset/Source/MeshModelingTools/Private/RemoveOccludedTrianglesTool.cpp
--- a/Engine/Plugins/Experimental/MeshModelingToolset/Source/MeshModelingTools/Private/RemoveOccludedTrianglesTool.cpp
+++ b/Engine/Plugins/Experimental/MeshModelingToolset/Source/MeshModelingTools/Private/RemoveOccludedTrianglesTool.cpp
@@ -4,7 +4,7 @@
 #include "InteractiveToolManager.h"
 #include "ToolBuilderUtil.h"
 
-#include "ToolSetupUtil.h"
+#include "BaseTools/ToolSetupUtil.h"
 
 #include "DynamicMesh3.h"
 #include "BaseBehaviors/MultiClickSequenceInputBehavior.h"
diff --git a/Engine/Plugins/Experimental/MeshModelingToolset/Source/MeshModelingTools/Private/SmoothMeshTool.cpp b/Engine/Plugins/Experimental/MeshModelingToolset/Source/MeshModelingTools/Private/SmoothMeshTool.cpp
--- a/Engine/Plugins/Experimental/MeshModelingToolset/Source/MeshModelingTools/Private/SmoothMeshTool.cpp
+++ b/Engine/Plugins/Experimental/MeshModelingToolset/Source/MeshModelingTools/Private/SmoothMeshTool.cpp
@@ -11,7 +11,7 @@
 #include "SimpleDynamicMeshComponent.h"
 
 #include "SceneManagement.h" // for FPrimitiveDrawInterface
-#include "ToolSetupUtil.h"
+#include "BaseTools/ToolSetupUtil.h"
 
 // Smoothing operators
 #include "SmoothingOps/IterativeSmoothingOp.h"
diff --git a/Engine/Plugins/Experimental/MeshModelingToolset/Source/MeshModelingTools/Private/TransformMeshesTool.cpp b/Engine/Plugins/Experimental/MeshModelingToolset/Source/MeshModelingTools/Private/TransformMeshesTool.cpp
--- a/Engine/Plugins/Experimental/MeshModelingToolset/Source/MeshModelingTools/Private/TransformMeshesTool.cpp
+++ b/Engine/Plugins/Experimental/MeshModelingToolset/Source/MeshModelingTools/Private/TransformMeshesTool.cpp
@@ -4,7 +4,7 @@
 #include "InteractiveToolManager.h"
 #include "InteractiveGizmoManager.h"
 #include "ToolBuilderUtil.h"
-#include "ToolSetupUtil.h"
+#include "BaseTools/ToolSetupUtil.h"
 #include "DynamicMesh3.h"
 #include "BaseBehaviors/ClickDragBehavior.h"
 
diff --git a/Engine/Plugins/Experimental/MeshModelingToolset/Source/MeshModelingTools/Private/UVProjectionTool.cpp b/Engine/Plugins/Experimental/MeshModelingToolset/Source/MeshModelingTools/Private/UVProjectionTool.cpp
--- a/Engine/Plugins/Experimental/MeshModelingToolset/Source/MeshModelingTools/Private/UVProjectionTool.cpp
+++ b/Engine/Plugins/Experimental/MeshModelingToolset/Source/MeshModelingTools/Private/UVProjectionTool.cpp
@@ -4,7 +4,7 @@
 #include "InteractiveToolManager.h"
 #include "ToolBuilderUtil.h"
 
-#include "ToolSetupUtil.h"
+#include "BaseTools/ToolSetupUtil.h"
 
 #include "DynamicMesh3.h"
 #include "BaseBehaviors/MultiClickSequenceInputBehavior.h"
diff --git a/Engine/Plugins/Experimental/MeshModelingToolset/Source/MeshModelingTools/Public/BaseBrushTool.h b/Engine/Plugins/Experimental/MeshModelingToolset/Source/MeshModelingTools/Public/BaseBrushTool.h
deleted file mode 100644
diff --git a/Engine/Plugins/Experimental/MeshModelingToolset/Source/MeshModelingTools/Public/DynamicMeshBrushTool.h b/Engine/Plugins/Experimental/MeshModelingToolset/Source/MeshModelingTools/Public/DynamicMeshBrushTool.h
--- a/Engine/Plugins/Experimental/MeshModelingToolset/Source/MeshModelingTools/Public/DynamicMeshBrushTool.h
+++ b/Engine/Plugins/Experimental/MeshModelingToolset/Source/MeshModelingTools/Public/DynamicMeshBrushTool.h
@@ -5,7 +5,7 @@
 #include "CoreMinimal.h"
 #include "UObject/NoExportTypes.h"
 #include "InteractiveToolBuilder.h"
-#include "BaseBrushTool.h"
+#include "BaseTools/BaseBrushTool.h"
 #include "DynamicMeshAABBTree3.h"
 #include "SimpleDynamicMeshComponent.h"
 #include "PreviewMesh.h"
diff --git a/Engine/Plugins/Experimental/MeshModelingToolset/Source/MeshModelingTools/Public/DynamicMeshSculptTool.h b/Engine/Plugins/Experimental/MeshModelingToolset/Source/MeshModelingTools/Public/DynamicMeshSculptTool.h
--- a/Engine/Plugins/Experimental/MeshModelingToolset/Source/MeshModelingTools/Public/DynamicMeshSculptTool.h
+++ b/Engine/Plugins/Experimental/MeshModelingToolset/Source/MeshModelingTools/Public/DynamicMeshSculptTool.h
@@ -9,10 +9,10 @@
 #include "DynamicMeshAABBTree3.h"
 #include "DynamicMeshOctree3.h"
 #include "MeshNormals.h"
-#include "BaseBrushTool.h"
-#include "Drawing/ToolDataVisualizer.h"
+#include "BaseTools/BaseBrushTool.h"
+#include "ToolDataVisualizer.h"
 #include "Changes/ValueWatcher.h"
-#include "Gizmos/BrushStampIndicator.h"
+#include "BaseGizmos/BrushStampIndicator.h"
 #include "Properties/MeshMaterialProperties.h"
 #include "TransformTypes.h"
 #include "DynamicMeshSculptTool.generated.h"
@@ -377,6 +377,7 @@ protected:
 	TArray<EDynamicMeshSculptBrushType> BrushTypeHistory;
 	int BrushTypeHistoryIndex = 0;
 
+	UPreviewMesh* MakeDefaultSphereMesh(UObject* Parent, UWorld* World, int Resolution = 32);
 };
 
 
diff --git a/Engine/Plugins/Experimental/MeshModelingToolset/Source/MeshModelingTools/Public/EditMeshPolygonsTool.h b/Engine/Plugins/Experimental/MeshModelingToolset/Source/MeshModelingTools/Public/EditMeshPolygonsTool.h
--- a/Engine/Plugins/Experimental/MeshModelingToolset/Source/MeshModelingTools/Public/EditMeshPolygonsTool.h
+++ b/Engine/Plugins/Experimental/MeshModelingToolset/Source/MeshModelingTools/Public/EditMeshPolygonsTool.h
@@ -7,7 +7,7 @@
 #include "BaseTools/MeshSurfacePointTool.h"
 #include "SimpleDynamicMeshComponent.h"
 #include "DynamicMeshAABBTree3.h"
-#include "Drawing/ToolDataVisualizer.h"
+#include "ToolDataVisualizer.h"
 #include "Transforms/QuickAxisTranslater.h"
 #include "Transforms/QuickAxisRotator.h"
 #include "Changes/MeshVertexChange.h"
diff --git a/Engine/Plugins/Experimental/MeshModelingToolset/Source/MeshModelingTools/Public/MeshSpaceDeformerTool.h b/Engine/Plugins/Experimental/MeshModelingToolset/Source/MeshModelingTools/Public/MeshSpaceDeformerTool.h
--- a/Engine/Plugins/Experimental/MeshModelingToolset/Source/MeshModelingTools/Public/MeshSpaceDeformerTool.h
+++ b/Engine/Plugins/Experimental/MeshModelingToolset/Source/MeshModelingTools/Public/MeshSpaceDeformerTool.h
@@ -6,7 +6,7 @@
 #include "UObject/NoExportTypes.h"
 #include "BaseTools/MeshSurfacePointTool.h"
 #include "DynamicMeshAABBTree3.h"
-#include "Drawing/ToolDataVisualizer.h"
+#include "ToolDataVisualizer.h"
 #include "Transforms/QuickAxisTranslater.h"
 #include "Transforms/QuickAxisRotator.h"
 #include "Changes/MeshVertexChange.h"
diff --git a/Engine/Plugins/Experimental/MeshModelingToolset/Source/MeshModelingTools/Public/UVProjectionTool.h b/Engine/Plugins/Experimental/MeshModelingToolset/Source/MeshModelingTools/Public/UVProjectionTool.h
--- a/Engine/Plugins/Experimental/MeshModelingToolset/Source/MeshModelingTools/Public/UVProjectionTool.h
+++ b/Engine/Plugins/Experimental/MeshModelingToolset/Source/MeshModelingTools/Public/UVProjectionTool.h
@@ -7,7 +7,7 @@
 #include "MultiSelectionTool.h"
 #include "InteractiveToolBuilder.h"
 #include "MeshOpPreviewHelpers.h"
-#include "Drawing/ToolDataVisualizer.h"
+#include "ToolDataVisualizer.h"
 #include "ParameterizationOps/UVProjectionOp.h"
 #include "DynamicMesh3.h"
 #include "BaseTools/SingleClickTool.h"
diff --git a/Engine/Plugins/Experimental/MeshModelingToolset/Source/MeshModelingToolsEditorOnly/Private/AttributeEditorTool.cpp b/Engine/Plugins/Experimental/MeshModelingToolset/Source/MeshModelingToolsEditorOnly/Private/AttributeEditorTool.cpp
--- a/Engine/Plugins/Experimental/MeshModelingToolset/Source/MeshModelingToolsEditorOnly/Private/AttributeEditorTool.cpp
+++ b/Engine/Plugins/Experimental/MeshModelingToolset/Source/MeshModelingToolsEditorOnly/Private/AttributeEditorTool.cpp
@@ -4,7 +4,7 @@
 #include "InteractiveToolManager.h"
 #include "ToolBuilderUtil.h"
 
-#include "ToolSetupUtil.h"
+#include "BaseTools/ToolSetupUtil.h"
 #include "MathUtil.h"
 
 #include "MeshDescription.h"
diff --git a/Engine/Plugins/Experimental/MeshModelingToolset/Source/MeshModelingToolsEditorOnly/Private/MergeMeshesTool.cpp b/Engine/Plugins/Experimental/MeshModelingToolset/Source/MeshModelingToolsEditorOnly/Private/MergeMeshesTool.cpp
--- a/Engine/Plugins/Experimental/MeshModelingToolset/Source/MeshModelingToolsEditorOnly/Private/MergeMeshesTool.cpp
+++ b/Engine/Plugins/Experimental/MeshModelingToolset/Source/MeshModelingToolsEditorOnly/Private/MergeMeshesTool.cpp
@@ -3,7 +3,7 @@
 #include "MergeMeshesTool.h"
 #include "InteractiveToolManager.h"
 #include "ToolBuilderUtil.h"
-#include "ToolSetupUtil.h"
+#include "BaseTools/ToolSetupUtil.h"
 #include "AssetGenerationUtil.h"
 #include "Selection/ToolSelectionUtil.h"
 
diff --git a/Engine/Plugins/Experimental/MeshModelingToolset/Source/MeshModelingToolsEditorOnly/Private/ParameterizeMeshTool.cpp b/Engine/Plugins/Experimental/MeshModelingToolset/Source/MeshModelingToolsEditorOnly/Private/ParameterizeMeshTool.cpp
--- a/Engine/Plugins/Experimental/MeshModelingToolset/Source/MeshModelingToolsEditorOnly/Private/ParameterizeMeshTool.cpp
+++ b/Engine/Plugins/Experimental/MeshModelingToolset/Source/MeshModelingToolsEditorOnly/Private/ParameterizeMeshTool.cpp
@@ -15,7 +15,7 @@
 #include "MeshDescriptionToDynamicMesh.h"
 
 #include "ParameterizationOps/ParameterizeMeshOp.h"
-#include "ToolSetupUtil.h"
+#include "BaseTools/ToolSetupUtil.h"
 
 #define LOCTEXT_NAMESPACE "UParameterizeMeshTool"
 
diff --git a/Engine/Plugins/Experimental/MeshModelingToolset/Source/MeshModelingToolsEditorOnly/Private/PolygonOnMeshTool.cpp b/Engine/Plugins/Experimental/MeshModelingToolset/Source/MeshModelingToolsEditorOnly/Private/PolygonOnMeshTool.cpp
--- a/Engine/Plugins/Experimental/MeshModelingToolset/Source/MeshModelingToolsEditorOnly/Private/PolygonOnMeshTool.cpp
+++ b/Engine/Plugins/Experimental/MeshModelingToolset/Source/MeshModelingToolsEditorOnly/Private/PolygonOnMeshTool.cpp
@@ -4,7 +4,7 @@
 #include "InteractiveToolManager.h"
 #include "ToolBuilderUtil.h"
 
-#include "ToolSetupUtil.h"
+#include "BaseTools/ToolSetupUtil.h"
 
 #include "DynamicMesh3.h"
 #include "BaseBehaviors/MultiClickSequenceInputBehavior.h"
diff --git a/Engine/Plugins/Experimental/MeshModelingToolset/Source/MeshModelingToolsEditorOnly/Private/VoxelCSGMeshesTool.cpp b/Engine/Plugins/Experimental/MeshModelingToolset/Source/MeshModelingToolsEditorOnly/Private/VoxelCSGMeshesTool.cpp
--- a/Engine/Plugins/Experimental/MeshModelingToolset/Source/MeshModelingToolsEditorOnly/Private/VoxelCSGMeshesTool.cpp
+++ b/Engine/Plugins/Experimental/MeshModelingToolset/Source/MeshModelingToolsEditorOnly/Private/VoxelCSGMeshesTool.cpp
@@ -3,7 +3,7 @@
 #include "VoxelCSGMeshesTool.h"
 #include "InteractiveToolManager.h"
 #include "ToolBuilderUtil.h"
-#include "ToolSetupUtil.h"
+#include "BaseTools/ToolSetupUtil.h"
 #include "AssetGenerationUtil.h"
 #include "Selection/ToolSelectionUtil.h"
 
diff --git a/Engine/Plugins/Experimental/MeshModelingToolset/Source/ModelingComponents/Private/Drawing/ToolDataVisualizer.cpp b/Engine/Plugins/Experimental/MeshModelingToolset/Source/ModelingComponents/Private/Drawing/ToolDataVisualizer.cpp
deleted file mode 100644
diff --git a/Engine/Plugins/Experimental/MeshModelingToolset/Source/ModelingComponents/Private/Gizmos/BrushStampIndicator.cpp b/Engine/Plugins/Experimental/MeshModelingToolset/Source/ModelingComponents/Private/Gizmos/BrushStampIndicator.cpp
deleted file mode 100644
diff --git a/Engine/Plugins/Experimental/MeshModelingToolset/Source/ModelingComponents/Private/Selection/GroupTopologySelector.cpp b/Engine/Plugins/Experimental/MeshModelingToolset/Source/ModelingComponents/Private/Selection/GroupTopologySelector.cpp
--- a/Engine/Plugins/Experimental/MeshModelingToolset/Source/ModelingComponents/Private/Selection/GroupTopologySelector.cpp
+++ b/Engine/Plugins/Experimental/MeshModelingToolset/Source/ModelingComponents/Private/Selection/GroupTopologySelector.cpp
@@ -3,7 +3,7 @@
 
 #include "Selection/GroupTopologySelector.h"
 #include "MeshQueries.h"
-#include "Drawing/ToolDataVisualizer.h"
+#include "ToolDataVisualizer.h"
 #include "ToolSceneQueriesUtil.h"
 
 
diff --git a/Engine/Plugins/Experimental/MeshModelingToolset/Source/ModelingComponents/Private/Snapping/RaySpatialSnapSolver.cpp b/Engine/Plugins/Experimental/MeshModelingToolset/Source/ModelingComponents/Private/Snapping/RaySpatialSnapSolver.cpp
--- a/Engine/Plugins/Experimental/MeshModelingToolset/Source/ModelingComponents/Private/Snapping/RaySpatialSnapSolver.cpp
+++ b/Engine/Plugins/Experimental/MeshModelingToolset/Source/ModelingComponents/Private/Snapping/RaySpatialSnapSolver.cpp
@@ -6,7 +6,7 @@
 #include "LineTypes.h"
 #include "Distance/DistLine3Ray3.h"
 #include "Distance/DistLine3Circle3.h"
-#include "Drawing/ToolDataVisualizer.h"
+#include "ToolDataVisualizer.h"
 
 
 
diff --git a/Engine/Plugins/Experimental/MeshModelingToolset/Source/ModelingComponents/Public/Drawing/ToolDataVisualizer.h b/Engine/Plugins/Experimental/MeshModelingToolset/Source/ModelingComponents/Public/Drawing/ToolDataVisualizer.h
deleted file mode 100644
diff --git a/Engine/Plugins/Experimental/MeshModelingToolset/Source/ModelingComponents/Public/Gizmos/BrushStampIndicator.h b/Engine/Plugins/Experimental/MeshModelingToolset/Source/ModelingComponents/Public/Gizmos/BrushStampIndicator.h
deleted file mode 100644
diff --git a/Engine/Plugins/Experimental/MeshModelingToolset/Source/ModelingComponents/Public/Transforms/QuickTransformer.h b/Engine/Plugins/Experimental/MeshModelingToolset/Source/ModelingComponents/Public/Transforms/QuickTransformer.h
--- a/Engine/Plugins/Experimental/MeshModelingToolset/Source/ModelingComponents/Public/Transforms/QuickTransformer.h
+++ b/Engine/Plugins/Experimental/MeshModelingToolset/Source/ModelingComponents/Public/Transforms/QuickTransformer.h
@@ -4,7 +4,7 @@
 
 #include "CoreMinimal.h"
 #include "ToolContextInterfaces.h"
-#include "Drawing/ToolDataVisualizer.h"
+#include "ToolDataVisualizer.h"
 #include "FrameTypes.h"
 
 
diff --git a/Engine/Source/Runtime/Experimental/InteractiveToolsFramework/Private/BaseGizmos/BrushStampIndicator.cpp b/Engine/Source/Runtime/Experimental/InteractiveToolsFramework/Private/BaseGizmos/BrushStampIndicator.cpp
new file mode 100644
--- /dev/null
+++ b/Engine/Source/Runtime/Experimental/InteractiveToolsFramework/Private/BaseGizmos/BrushStampIndicator.cpp
@@ -0,0 +1,76 @@
+// Copyright 1998-2019 Epic Games, Inc. All Rights Reserved.
+
+#include "BaseGizmos/BrushStampIndicator.h"
+#include "InteractiveGizmoManager.h"
+#include "Components/PrimitiveComponent.h"
+#include "ToolDataVisualizer.h"
+
+
+UInteractiveGizmo* UBrushStampIndicatorBuilder::BuildGizmo(const FToolBuilderState& SceneState) const
+{
+	UBrushStampIndicator* NewGizmo = NewObject<UBrushStampIndicator>(SceneState.GizmoManager);
+	return NewGizmo;
+}
+
+
+
+void UBrushStampIndicator::Setup()
+{
+}
+
+void UBrushStampIndicator::Shutdown()
+{
+}
+
+void UBrushStampIndicator::Render(IToolsContextRenderAPI* RenderAPI)
+{
+	if (bDrawIndicatorLines)
+	{
+		FToolDataVisualizer Draw;
+		Draw.BeginFrame(RenderAPI);
+
+		Draw.DrawCircle(BrushPosition, BrushNormal, BrushRadius, SampleStepCount, LineColor, LineThickness, bDepthTested);
+
+		if (bDrawSecondaryLines)
+		{
+			Draw.DrawCircle(BrushPosition, BrushNormal, BrushRadius*0.5f, SampleStepCount, SecondaryLineColor, SecondaryLineThickness, bDepthTested);
+			Draw.DrawLine(BrushPosition, BrushPosition + BrushRadius * BrushNormal, SecondaryLineColor, SecondaryLineThickness, bDepthTested);
+		}
+
+		Draw.EndFrame();
+	}
+}
+
+void UBrushStampIndicator::Tick(float DeltaTime)
+{
+}
+
+
+void UBrushStampIndicator::Update(float Radius, const FVector& Position, const FVector& Normal)
+{
+	BrushRadius = Radius;
+	BrushPosition = Position;
+	BrushNormal = Normal;
+
+	if (AttachedComponent != nullptr)
+	{
+		FTransform Transform = AttachedComponent->GetComponentTransform();
+
+		if (ScaleInitializedComponent != AttachedComponent)
+		{
+			InitialComponentScale = Transform.GetScale3D();
+			InitialComponentScale *= 1.0f / InitialComponentScale.Z;
+			ScaleInitializedComponent = AttachedComponent;
+		}
+
+		Transform.SetTranslation(BrushPosition);
+
+		FQuat CurRotation = Transform.GetRotation();
+		FQuat ApplyRotation = FQuat::FindBetween(CurRotation.GetAxisZ(), BrushNormal);
+		Transform.SetRotation(ApplyRotation * CurRotation);
+
+		Transform.SetScale3D(Radius * InitialComponentScale);
+
+		AttachedComponent->SetWorldTransform(Transform);
+	}
+}
diff --git a/Engine/Source/Runtime/Experimental/InteractiveToolsFramework/Private/BaseTools/BaseBrushTool.cpp b/Engine/Source/Runtime/Experimental/InteractiveToolsFramework/Private/BaseTools/BaseBrushTool.cpp
new file mode 100644
--- /dev/null
+++ b/Engine/Source/Runtime/Experimental/InteractiveToolsFramework/Private/BaseTools/BaseBrushTool.cpp
@@ -0,0 +1,209 @@
+// Copyright 1998-2019 Epic Games, Inc. All Rights Reserved.
+
+#include "BaseTools/BaseBrushTool.h"
+#include "InteractiveToolManager.h"
+#include "InteractiveGizmoManager.h"
+#include "ToolBuilderUtil.h"
+#include "InteractiveTool.h"
+#include "BaseGizmos/BrushStampIndicator.h"
+
+
+#define LOCTEXT_NAMESPACE "UBaseBrushTool"
+
+
+
+UBrushBaseProperties::UBrushBaseProperties()
+{
+	BrushSize = 0.25f;
+	bSpecifyRadius = false;
+	BrushRadius = 10.0f;
+	BrushStrength = 0.5f;
+}
+
+void UBrushBaseProperties::SaveProperties(UInteractiveTool* SaveFromTool)
+{
+	UBrushBaseProperties* PropertyCache = GetPropertyCache<UBrushBaseProperties>();
+	PropertyCache->BrushSize = this->BrushSize;
+	PropertyCache->bSpecifyRadius = this->bSpecifyRadius;
+	PropertyCache->BrushRadius = this->BrushRadius;
+}
+
+void UBrushBaseProperties::RestoreProperties(UInteractiveTool* RestoreToTool)
+{
+	UBrushBaseProperties* PropertyCache = GetPropertyCache<UBrushBaseProperties>();
+	this->BrushSize = PropertyCache->BrushSize;
+	this->bSpecifyRadius = PropertyCache->bSpecifyRadius;
+	this->BrushRadius = PropertyCache->BrushRadius;
+}
+
+
+
+
+UBaseBrushTool::UBaseBrushTool()
+{
+}
+
+
+void UBaseBrushTool::Setup()
+{
+	UMeshSurfacePointTool::Setup();
+
+	double MaxDimension = EstimateMaximumTargetDimension();
+	BrushRelativeSizeRange = TInterval<float>(MaxDimension*0.01, MaxDimension);
+	BrushProperties = NewObject<UBrushBaseProperties>(this, TEXT("Brush"));
+	RecalculateBrushRadius();
+
+	// initialize our properties
+	AddToolPropertySource(BrushProperties);
+
+	SetupBrushStampIndicator();
+}
+
+
+void UBaseBrushTool::Shutdown(EToolShutdownType ShutdownType)
+{
+	ShutdownBrushStampIndicator();
+}
+
+
+void UBaseBrushTool::OnPropertyModified(UObject* PropertySet, UProperty* Property)
+{
+	if (PropertySet == BrushProperties)
+	{
+		RecalculateBrushRadius();
+	}
+}
+
+
+
+void UBaseBrushTool::IncreaseBrushSizeAction()
+{
+	BrushProperties->BrushSize = FMath::Clamp(BrushProperties->BrushSize + 0.025f, 0.0f, 1.0f);
+	RecalculateBrushRadius();
+}
+
+void UBaseBrushTool::DecreaseBrushSizeAction()
+{
+	BrushProperties->BrushSize = FMath::Clamp(BrushProperties->BrushSize - 0.025f, 0.0f, 1.0f);
+	RecalculateBrushRadius();
+}
+
+
+void UBaseBrushTool::RegisterActions(FInteractiveToolActionSet& ActionSet)
+{
+	ActionSet.RegisterAction(this, (int32)EStandardToolActions::BaseClientDefinedActionID+12,
+		TEXT("BrushIncreaseSize"), 
+		LOCTEXT("BrushIncreaseSize", "Increase Brush Size"),
+		LOCTEXT("BrushIncreaseSizeTooltip", "Increase size of brush"),
+		EModifierKey::None, EKeys::RightBracket,
+		[this]() { IncreaseBrushSizeAction(); });
+
+	ActionSet.RegisterAction(this, (int32)EStandardToolActions::BaseClientDefinedActionID + 37,
+		TEXT("BrushDecreaseSize"), 
+		LOCTEXT("BrushDecreaseSize", "Decrease Brush Size"),
+		LOCTEXT("BrushDecreaseSizeTooltip", "Decrease size of brush"),
+		EModifierKey::None, EKeys::LeftBracket,
+		[this]() { DecreaseBrushSizeAction(); });
+}
+
+
+
+
+void UBaseBrushTool::RecalculateBrushRadius()
+{
+	CurrentBrushRadius = 0.5 * BrushRelativeSizeRange.Interpolate(BrushProperties->BrushSize);
+	if (BrushProperties->bSpecifyRadius)
+	{
+		CurrentBrushRadius = BrushProperties->BrushRadius;
+	}
+	else
+	{
+		BrushProperties->BrushRadius = CurrentBrushRadius;
+	}
+}
+
+
+void UBaseBrushTool::OnBeginDrag(const FRay& Ray)
+{
+	FHitResult OutHit;
+	if (HitTest(Ray, OutHit))
+	{
+		LastBrushStamp.Radius = BrushProperties->BrushRadius;
+		LastBrushStamp.WorldPosition = OutHit.ImpactPoint;
+		LastBrushStamp.WorldNormal = OutHit.Normal;
+		LastBrushStamp.HitResult = OutHit;
+	}
+	bInBrushStroke = true;
+}
+
+void UBaseBrushTool::OnUpdateDrag(const FRay& Ray)
+{
+	FHitResult OutHit;
+	if (HitTest(Ray, OutHit))
+	{
+		LastBrushStamp.Radius = BrushProperties->BrushRadius;
+		LastBrushStamp.WorldPosition = OutHit.ImpactPoint;
+		LastBrushStamp.WorldNormal = OutHit.Normal;
+		LastBrushStamp.HitResult = OutHit;
+	}
+}
+
+void UBaseBrushTool::OnEndDrag(const FRay& Ray)
+{
+	bInBrushStroke = false;
+}
+
+
+bool UBaseBrushTool::OnUpdateHover(const FInputDeviceRay& DevicePos)
+{
+	FHitResult OutHit;
+	if (HitTest(DevicePos.WorldRay, OutHit))
+	{
+		LastBrushStamp.Radius = BrushProperties->BrushRadius;
+		LastBrushStamp.WorldPosition = OutHit.ImpactPoint;
+		LastBrushStamp.WorldNormal = OutHit.Normal;
+		LastBrushStamp.HitResult = OutHit;
+	}
+	return true;
+}
+
+
+
+void UBaseBrushTool::Render(IToolsContextRenderAPI* RenderAPI)
+{
+	UMeshSurfacePointTool::Render(RenderAPI);
+
+	UpdateBrushStampIndicator();
+}
+
+void UBaseBrushTool::Tick(float DeltaTime)
+{
+	UMeshSurfacePointTool::Tick(DeltaTime);
+}
+
+
+
+
+const FString BaseBrushIndicatorGizmoType = TEXT("BrushIndicatorGizmoType");
+
+void UBaseBrushTool::SetupBrushStampIndicator()
+{
+	// register and spawn brush indicator gizmo
+	GetToolManager()->GetPairedGizmoManager()->RegisterGizmoType(BaseBrushIndicatorGizmoType, NewObject<UBrushStampIndicatorBuilder>());
+	BrushStampIndicator = GetToolManager()->GetPairedGizmoManager()->CreateGizmo<UBrushStampIndicator>(BaseBrushIndicatorGizmoType, FString(), this);
+}
+
+void UBaseBrushTool::UpdateBrushStampIndicator()
+{
+	BrushStampIndicator->Update(LastBrushStamp.Radius, LastBrushStamp.WorldPosition, LastBrushStamp.WorldNormal);
+}
+
+void UBaseBrushTool::ShutdownBrushStampIndicator()
+{
+	GetToolManager()->GetPairedGizmoManager()->DestroyAllGizmosByOwner(this);
+	BrushStampIndicator = nullptr;
+	GetToolManager()->GetPairedGizmoManager()->DeregisterGizmoType(BaseBrushIndicatorGizmoType);
+}
+
+
+#undef LOCTEXT_NAMESPACE
diff --git a/Engine/Source/Runtime/Experimental/InteractiveToolsFramework/Private/ToolDataVisualizer.cpp b/Engine/Source/Runtime/Experimental/InteractiveToolsFramework/Private/ToolDataVisualizer.cpp
new file mode 100644
--- /dev/null
+++ b/Engine/Source/Runtime/Experimental/InteractiveToolsFramework/Private/ToolDataVisualizer.cpp
@@ -0,0 +1,193 @@
+// Copyright 1998-2019 Epic Games, Inc. All Rights Reserved.
+
+#include "ToolDataVisualizer.h"
+#include "SceneManagement.h" 
+#include "ToolContextInterfaces.h"
+#include "SceneManagement.h"   // DrawCircle
+#include "BaseGizmos/GizmoMath.h"
+
+
+const int BoxFaces[6][4] =
+{
+	{ 0, 1, 2, 3 },     // back, -z
+	{ 5, 4, 7, 6 },     // front, +z
+	{ 4, 0, 3, 7 },     // left, -x
+	{ 1, 5, 6, 2 },     // right, +x,
+	{ 4, 5, 1, 0 },     // bottom, -y
+	{ 3, 2, 6, 7 }      // top, +y
+};
+
+FToolDataVisualizer::FToolDataVisualizer()
+{
+	LineColor = FLinearColor(0.95f, 0.05f, 0.05f);
+	PointColor = FLinearColor(0.95f, 0.05f, 0.05f);
+	PopAllTransforms();
+}
+
+
+void FToolDataVisualizer::BeginFrame(IToolsContextRenderAPI* RenderAPI, const FViewCameraState& CameraStateIn)
+{
+	checkf(CurrentPDI == nullptr, TEXT("FToolDataVisualizer::BeginFrame: matching EndFrame was not called last frame!"));
+	CurrentPDI = RenderAPI->GetPrimitiveDrawInterface();
+	CameraState = CameraStateIn;
+	bHaveCameraState = true;
+}
+
+void FToolDataVisualizer::BeginFrame(IToolsContextRenderAPI* RenderAPI)
+{
+	checkf(CurrentPDI == nullptr, TEXT("FToolDataVisualizer::BeginFrame: matching EndFrame was not called last frame!"));
+	CurrentPDI = RenderAPI->GetPrimitiveDrawInterface();
+	bHaveCameraState = false;
+}
+
+void FToolDataVisualizer::EndFrame()
+{
+	// not safe to hold PDI
+	CurrentPDI = nullptr;
+}
+
+
+void FToolDataVisualizer::SetTransform(const FTransform& Transform)
+{
+	TransformStack.Reset();
+	TransformStack.Add(Transform);
+	TotalTransform = Transform;
+}
+
+
+void FToolDataVisualizer::PushTransform(const FTransform& Transform)
+{
+	TransformStack.Add(Transform);
+	TotalTransform *= Transform;
+}
+
+void FToolDataVisualizer::PopTransform()
+{
+	TransformStack.Pop(false);
+	TotalTransform = FTransform::Identity;
+	for (const FTransform& Transform : TransformStack)
+	{
+		TotalTransform *= Transform;
+	}
+}
+
+void FToolDataVisualizer::PopAllTransforms()
+{
+	TransformStack.Reset();
+	TotalTransform = FTransform::Identity;
+}
+
+
+
+
+void FToolDataVisualizer::InternalDrawTransformedLine(const FVector& A, const FVector& B, const FLinearColor& ColorIn, float LineThicknessIn, bool bDepthTestedIn)
+{
+	CurrentPDI->DrawLine(A, B, ColorIn,
+		(bDepthTestedIn) ? SDPG_World : SDPG_Foreground,
+		LineThicknessIn, DepthBias, true);
+}
+
+
+void FToolDataVisualizer::InternalDrawTransformedPoint(const FVector& Position, const FLinearColor& ColorIn, float PointSizeIn, bool bDepthTestedIn)
+{
+	CurrentPDI->DrawPoint(Position, ColorIn, PointSizeIn,
+		(bDepthTestedIn) ? SDPG_World : SDPG_Foreground);
+}
+
+
+void FToolDataVisualizer::InternalDrawCircle(const FVector& Position, const FVector& Normal, float Radius, int Steps, const FLinearColor& Color, float LineThicknessIn, bool bDepthTestedIn)
+{
+	FVector Tan1, Tan2;
+	GizmoMath::MakeNormalPlaneBasis((FVector)TransformN(Normal), Tan1, Tan2);
+	Tan1.Normalize(); Tan2.Normalize();
+
+	// this function is from SceneManagement.h
+	::DrawCircle(CurrentPDI, TransformP(Position), (FVector)Tan1, (FVector)Tan2, 
+		Color, Radius, Steps,
+		(bDepthTestedIn) ? SDPG_World : SDPG_Foreground,
+		LineThicknessIn, DepthBias, true);
+}
+
+void FToolDataVisualizer::InternalDrawWireBox(const FBox& Box, const FLinearColor& ColorIn, float LineThicknessIn, bool bDepthTestedIn)
+{
+	// corners [ (-x,-y), (x,-y), (x,y), (-x,y) ], -z, then +z
+	FVector Corners[8] =
+	{
+		TransformP(Box.Min),
+		TransformP(FVector(Box.Max.X, Box.Min.Y, Box.Min.Z)),
+		TransformP(FVector(Box.Max.X, Box.Max.Y, Box.Min.Z)),
+		TransformP(FVector(Box.Min.X, Box.Max.Y, Box.Min.Z)),
+		TransformP(FVector(Box.Min.X, Box.Min.Y, Box.Max.Z)),
+		TransformP(FVector(Box.Max.X, Box.Min.Y, Box.Max.Z)),
+		TransformP(Box.Max),
+		TransformP(FVector(Box.Min.X, Box.Max.Y, Box.Max.Z))
+	};
+	for (int FaceIdx = 0; FaceIdx < 6; FaceIdx++)
+	{
+		for (int Last = 3, Cur = 0; Cur < 4; Last = Cur++)
+		{
+			InternalDrawTransformedLine(Corners[BoxFaces[FaceIdx][Last]], Corners[BoxFaces[FaceIdx][Cur]], 
+										ColorIn, LineThicknessIn, bDepthTestedIn);
+		}
+	}
+}
+
+void FToolDataVisualizer::InternalDrawSquare(const FVector& Center, const FVector& SideA, const FVector& SideB, const FLinearColor& Color, float LineThicknessIn, bool bDepthTestedIn)
+{
+	FVector CC = TransformP(Center);
+	FVector SA = TransformV(SideA);
+	FVector SB = TransformV(SideB);
+	FVector HalfDiag = (SA + SB) * .5f;
+	FVector C00 = CC - HalfDiag;
+	FVector C11 = CC + HalfDiag;
+	FVector C01 = C00 + SB;
+	FVector C10 = C00 + SA;
+	InternalDrawTransformedLine(C00, C01, Color, LineThicknessIn, bDepthTestedIn);
+	InternalDrawTransformedLine(C01, C11, Color, LineThicknessIn, bDepthTestedIn);
+	InternalDrawTransformedLine(C10, C11, Color, LineThicknessIn, bDepthTestedIn);
+	InternalDrawTransformedLine(C00, C10, Color, LineThicknessIn, bDepthTestedIn);
+}
+
+void FToolDataVisualizer::InternalDrawWireCylinder(const FVector& Position, const FVector& Normal, float Radius, float Height, int Steps, const FLinearColor& Color, float LineThicknessIn, bool bDepthTestedIn)
+{
+	FVector Tan1, Tan2;
+	GizmoMath::MakeNormalPlaneBasis(Normal, Tan1, Tan2);
+	
+	const float	AngleDelta = 2.0f * PI / Steps;
+	FVector X(Tan1), Y(Tan2);
+	FVector	LastVertex = TransformP(Position + X * Radius);
+	FVector LastVertexB = TransformP(Position + X * Radius + Normal * Height);
+
+	for (int32 Step = 0; Step < Steps; Step++)
+	{
+		float Angle = (Step + 1) * AngleDelta;
+		FVector A = Position + (X * FMath::Cos(Angle) + Y * FMath::Sin(Angle)) * Radius;
+		FVector B = A + Normal * Height;
+		FVector Vertex = TransformP(A);
+		FVector VertexB = TransformP(B);
+		InternalDrawTransformedLine(LastVertex, Vertex, Color, LineThicknessIn, bDepthTestedIn);
+		InternalDrawTransformedLine(Vertex, VertexB, Color, LineThicknessIn, bDepthTestedIn);
+		InternalDrawTransformedLine(LastVertexB, VertexB, Color, LineThicknessIn, bDepthTestedIn);
+		LastVertex = Vertex;
+		LastVertexB = VertexB;
+	}
+
+}
+
+
+void FToolDataVisualizer::InternalDrawViewFacingCircle(const FVector& Position, float Radius, int Steps, const FLinearColor& Color, float LineThicknessIn, bool bDepthTestedIn)
+{
+	checkf(bHaveCameraState, TEXT("To call this function, you must first call the version of BeginFrame that takes the CameraState"));
+
+	FVector WorldPosition = TransformP(Position);
+	FVector WorldNormal = (CameraState.Position - WorldPosition);
+	WorldNormal.Normalize();
+	FVector Tan1, Tan2;
+	GizmoMath::MakeNormalPlaneBasis(WorldNormal, Tan1, Tan2);
+
+	// this function is from SceneManagement.h
+	::DrawCircle(CurrentPDI, WorldPosition, (FVector)Tan1, (FVector)Tan2,
+		Color, Radius, Steps,
+		(bDepthTestedIn) ? SDPG_World : SDPG_Foreground,
+		LineThicknessIn, DepthBias, true);
+}
diff --git a/Engine/Source/Runtime/Experimental/InteractiveToolsFramework/Public/BaseGizmos/BrushStampIndicator.h b/Engine/Source/Runtime/Experimental/InteractiveToolsFramework/Public/BaseGizmos/BrushStampIndicator.h
new file mode 100644
--- /dev/null
+++ b/Engine/Source/Runtime/Experimental/InteractiveToolsFramework/Public/BaseGizmos/BrushStampIndicator.h
@@ -0,0 +1,95 @@
+// Copyright 1998-2019 Epic Games, Inc. All Rights Reserved.
+
+#pragma once
+
+#include "InteractiveGizmo.h"
+#include "InteractiveGizmoBuilder.h"
+#include "BrushStampIndicator.generated.h"
+
+class UPrimitiveComponent;
+class UPreviewMesh;
+
+UCLASS()
+class INTERACTIVETOOLSFRAMEWORK_API UBrushStampIndicatorBuilder : public UInteractiveGizmoBuilder
+{
+	GENERATED_BODY()
+public:
+	virtual UInteractiveGizmo* BuildGizmo(const FToolBuilderState& SceneState) const override;
+};
+
+
+
+/*
+ * UBrushStampIndicator is a simple 3D brush indicator. 
+ */
+UCLASS(Transient)
+class INTERACTIVETOOLSFRAMEWORK_API UBrushStampIndicator : public UInteractiveGizmo
+{
+	GENERATED_BODY()
+
+public:
+
+	// UInteractiveGizmo interface/implementation
+
+	virtual void Setup() override;
+	virtual void Shutdown() override;
+	virtual void Render(IToolsContextRenderAPI* RenderAPI) override;
+	virtual void Tick(float DeltaTime) override;
+
+
+	/**
+	 * Update the Radius, Position, and Normal of the stamp indicator
+	 */
+	virtual void Update(float Radius, const FVector& Position, const FVector& Normal);
+
+public:
+
+	UPROPERTY()
+	float BrushRadius = 1.0f;
+
+	UPROPERTY()
+	FVector BrushPosition = FVector::ZeroVector;
+
+	UPROPERTY()
+	FVector BrushNormal = FVector(0, 0, 1);;
+
+
+
+	UPROPERTY()
+	bool bDrawIndicatorLines = true;
+
+
+	UPROPERTY()
+	int SampleStepCount = 32;
+
+	UPROPERTY()
+	FLinearColor LineColor = FLinearColor(0.96f, 0.06f, 0.06f);
+
+	UPROPERTY()
+	float LineThickness = 2.0f;
+
+	UPROPERTY()
+	bool bDepthTested = false;
+
+
+
+	UPROPERTY()
+	bool bDrawSecondaryLines = true;
+
+	UPROPERTY()
+	float SecondaryLineThickness = 0.5f;
+
+	UPROPERTY()
+	FLinearColor SecondaryLineColor = FLinearColor(0.5f, 0.5f, 0.5f, 0.5f);
+
+
+	/**
+	 * Optional Component that will be transformed such that it tracks the Radius/Position/Normal
+	 */
+	UPROPERTY()
+	UPrimitiveComponent* AttachedComponent;
+
+protected:
+	UPrimitiveComponent* ScaleInitializedComponent = nullptr;		// we are just using this as a key, never calling any functions on it
+	FVector InitialComponentScale;
+};
\ No newline at end of file
diff --git a/Engine/Source/Runtime/Experimental/InteractiveToolsFramework/Public/BaseTools/BaseBrushTool.h b/Engine/Source/Runtime/Experimental/InteractiveToolsFramework/Public/BaseTools/BaseBrushTool.h
new file mode 100644
--- /dev/null
+++ b/Engine/Source/Runtime/Experimental/InteractiveToolsFramework/Public/BaseTools/BaseBrushTool.h
@@ -0,0 +1,150 @@
+// Copyright 1998-2019 Epic Games, Inc. All Rights Reserved.
+
+#pragma once
+
+#include "CoreMinimal.h"
+#include "BaseTools/MeshSurfacePointTool.h"
+#include "InteractiveTool.h"
+#include "BaseBrushTool.generated.h"
+
+class UBrushStampIndicator;
+/**
+ * Standard properties for a Brush-type Tool
+ */
+UCLASS()
+class INTERACTIVETOOLSFRAMEWORK_API UBrushBaseProperties : public UInteractiveToolPropertySet
+{
+	GENERATED_BODY()
+
+public:
+	UBrushBaseProperties();
+
+	/** Relative size of brush */
+	UPROPERTY(EditAnywhere, Category = Brush, meta = (DisplayName = "Size", UIMin = "0.0", UIMax = "1.0", ClampMin = "0.0", ClampMax = "10.0"))
+	float BrushSize;
+
+	/** If true, ignore relative Brush Size and use explicit World Radius */
+	UPROPERTY(EditAnywhere, Category = Brush, AdvancedDisplay)
+	bool bSpecifyRadius;
+
+	/** Radius of brush */
+	UPROPERTY(EditAnywhere, Category = Brush, AdvancedDisplay, meta = (DisplayName = "Radius", UIMin = "1.0", UIMax = "1000.0", ClampMin = "0.1", ClampMax = "50000.0"))
+	float BrushRadius;
+
+	/** Strength of the brush (0.0 - 1.0) */
+	UPROPERTY(EditAnywhere, Category = Brush, meta = (DisplayName = "Strength", UIMin = "0.0", UIMax = "1.0", ClampMin = "0.0", ClampMax = "1.0"))
+	float BrushStrength;
+
+	/** Amount of falloff to apply (0.0 - 1.0) */
+	UPROPERTY(EditAnywhere, Category = Brush, meta = (DisplayName = "Falloff", UIMin = "0.0", UIMax = "1.0", ClampMin = "0.0", ClampMax = "1.0"))
+	float BrushFalloffAmount;
+	//
+	// save/restore support
+	//
+	virtual void SaveProperties(UInteractiveTool* SaveFromTool) override;
+	virtual void RestoreProperties(UInteractiveTool* RestoreToTool) override;
+};
+
+
+/**
+ * Generic Brush Stamp data
+ */
+USTRUCT()
+struct INTERACTIVETOOLSFRAMEWORK_API FBrushStampData
+{
+	GENERATED_BODY();
+	/** Radius of brush stamp */
+	float Radius;
+	/** World Position of brush stamp */
+	FVector WorldPosition;
+	/** World Normal of brush stamp */
+	FVector WorldNormal;
+
+	/** Hit Result provided by implementations - may not be fully populated */
+	FHitResult HitResult;
+};
+
+
+
+/**
+ * UBaseBrushTool implements standard brush-style functionality for an InteractiveTool.
+ * This includes:
+ *   1) brush radius property set with dimension-relative brush sizing and default brush radius hotkeys
+ *   2) brush indicator visualization
+ *   3) tracking of last brush stamp location via .LastBrushStamp UProperty
+ *   4) status of brush stroke via .bInBrushStroke UProperty
+ */
+UCLASS()
+class INTERACTIVETOOLSFRAMEWORK_API UBaseBrushTool : public UMeshSurfacePointTool
+{
+	GENERATED_BODY()
+
+public:
+	UBaseBrushTool();
+
+	virtual void RegisterActions(FInteractiveToolActionSet& ActionSet) override;
+
+	virtual void Setup() override;
+	virtual void Shutdown(EToolShutdownType ShutdownType) override;
+
+	virtual void Tick(float DeltaTime) override;
+	virtual void Render(IToolsContextRenderAPI* RenderAPI) override;
+
+	virtual void OnPropertyModified(UObject* PropertySet, UProperty* Property) override;
+
+
+	//
+	// UMeshSurfacePointTool implementation
+	//
+	virtual void OnBeginDrag(const FRay& Ray) override;
+	virtual void OnUpdateDrag(const FRay& Ray) override;
+	virtual void OnEndDrag(const FRay& Ray) override;
+	virtual bool OnUpdateHover(const FInputDeviceRay& DevicePos) override;
+
+public:
+
+	/** Properties that control brush size/etc*/
+	UPROPERTY()
+	UBrushBaseProperties* BrushProperties;
+
+	/** Set to true by Tool if user is currently in an active brush stroke*/
+	UPROPERTY()
+	bool bInBrushStroke = false;
+
+	/** Position of brush at last update (both during stroke and during Hover) */
+	UPROPERTY()
+	FBrushStampData LastBrushStamp;
+
+public:
+
+	virtual void IncreaseBrushSizeAction();
+	virtual void DecreaseBrushSizeAction();
+
+	virtual bool IsInBrushStroke() const { return bInBrushStroke; }
+
+protected:
+
+	/**
+	 * Subclasses should implement this to give an estimate of target dimension for brush size scaling
+	 */
+	virtual double EstimateMaximumTargetDimension() { return 100.0; }
+
+protected:
+	TInterval<float> BrushRelativeSizeRange;
+	double CurrentBrushRadius;
+	void RecalculateBrushRadius();
+
+
+
+	//
+	// Brush Indicator support
+	//
+protected:
+
+	UPROPERTY()
+	UBrushStampIndicator* BrushStampIndicator;
+
+	virtual void SetupBrushStampIndicator();
+	virtual void UpdateBrushStampIndicator();
+	virtual void ShutdownBrushStampIndicator();
+};
diff --git a/Engine/Source/Runtime/Experimental/InteractiveToolsFramework/Public/ToolDataVisualizer.h b/Engine/Source/Runtime/Experimental/InteractiveToolsFramework/Public/ToolDataVisualizer.h
new file mode 100644
--- /dev/null
+++ b/Engine/Source/Runtime/Experimental/InteractiveToolsFramework/Public/ToolDataVisualizer.h
@@ -0,0 +1,248 @@
+// Copyright 1998-2019 Epic Games, Inc. All Rights Reserved.
+
+#pragma once
+
+#include "CoreMinimal.h"
+#include "ToolContextInterfaces.h"  // for FViewCameraState
+
+class FPrimitiveDrawInterface;
+class IToolsContextRenderAPI;
+
+/**
+ * FToolDataVisualizer is a utility class for Tool and Gizmo implementations
+ * to use to draw 3D lines, points, etc.
+ * 
+ * Currently uses PDI drawing but may use different back-ends in the future
+ */
+class INTERACTIVETOOLSFRAMEWORK_API FToolDataVisualizer
+{
+public:
+	/** Default color used for drawing lines */
+	FLinearColor LineColor;   // = Red
+	/** Default thickness used for drawing lines */
+	float LineThickness = 1.0f;
+
+	/** Default color used for drawing points */
+	FLinearColor PointColor;   // = Red
+	/** Default size used for drawing points */
+	float PointSize = 1.0f;
+
+	/** Depth bias applied to lines */
+	float DepthBias = 0.0f;
+	/** Should lines be clipped by 3D geometry, or should they be drawn with quasi-transparency */
+	bool bDepthTested = true;
+
+public:
+	FToolDataVisualizer();
+	virtual ~FToolDataVisualizer() {}
+
+
+	//
+	// Frame initialization/cleanup
+	// 
+
+
+	/** This must be called every frame to allow Visualizer to extract necessary rendering data/objects */
+	void BeginFrame(IToolsContextRenderAPI* RenderAPI, const FViewCameraState& CameraState);
+
+	/** This must be called every frame to allow Visualizer to extract necessary rendering data/objects */
+	void BeginFrame(IToolsContextRenderAPI* RenderAPI);
+
+	/** this should be called at the end of every frame to allow for necessary cleanup */
+	void EndFrame();
+
+
+	//
+	// Transform support
+	// 
+
+	/** Clear transform stack and push the given Transform */
+	void SetTransform(const FTransform& Transform);
+	/** Push a Transform onto the transform stack */
+	void PushTransform(const FTransform& Transform);
+	/** Pop a transform from the transform stack */
+	void PopTransform();
+	/** Clear the transform stack to identity */
+	void PopAllTransforms();
+
+	/** @return input Point transformed by transform stack */
+	FVector TransformP(const FVector& Point)
+	{
+		return TotalTransform.TransformPosition(Point);
+	}
+
+	/** @return input Vector transformed by transform stack */
+	FVector TransformV(const FVector& Vector)
+	{
+		return TotalTransform.TransformVector(Vector);
+	}
+
+	/** @return input Normal transformed by transform stack */
+	FVector TransformN(const FVector& Normal)
+	{
+		// transform normal by a safe inverse scale + normalize, and a standard rotation (TODO: move implementation into FTransform)
+		const FVector& S = TotalTransform.GetScale3D();
+		float DetSign = FMath::Sign(S.X*S.Y*S.Z); // we only need to multiply by the sign of the determinant, rather than divide by it, since we normalize later anyway
+		FVector SafeInvS(S.Y*S.Z*DetSign, S.X*S.Z*DetSign, S.X*S.Y*DetSign);
+		return TotalTransform.TransformVectorNoScale((SafeInvS*Normal).GetSafeNormal());
+	}
+
+
+	//
+	// Parameters
+	//
+
+
+	/** Update the default line color and thickness */
+	void SetLineParameters(const FLinearColor& Color, float Thickness)
+	{
+		LineColor = Color;
+		LineThickness = Thickness;
+	}
+
+	/** Update the default point color and size */
+	void SetPointParameters(const FLinearColor& Color, float Size)
+	{
+		PointColor = Color;
+		PointSize = Size;
+	}
+
+
+	//
+	// Drawing functions
+	// 
+
+
+	/** Draw a line with default parameters */
+	template<typename PointType>
+	void DrawLine(const PointType& A, const PointType& B)
+	{
+		InternalDrawTransformedLine(TransformP((FVector)A), TransformP((FVector)B), LineColor, LineThickness, bDepthTested);
+	}
+
+	/** Draw a line with the given Color, otherwise use default parameters */
+	template<typename PointType>
+	void DrawLine(const PointType& A, const PointType& B, const FLinearColor& Color)
+	{
+		InternalDrawTransformedLine(TransformP((FVector)A), TransformP((FVector)B), Color, LineThickness, bDepthTested);
+	}
+
+	/** Draw a line with the given Color and Thickness, otherwise use default parameters */
+	template<typename PointType>
+	void DrawLine(const PointType& A, const PointType& B, const FLinearColor& Color, float LineThicknessIn)
+	{
+		InternalDrawTransformedLine(TransformP((FVector)A), TransformP((FVector)B), Color, LineThicknessIn, bDepthTested);
+	}
+
+	/** Draw a line with the given parameters */
+	template<typename PointType>
+	void DrawLine(const PointType& A, const PointType& B, const FLinearColor& Color, float LineThicknessIn, bool bDepthTestedIn)
+	{
+		InternalDrawTransformedLine(TransformP((FVector)A), TransformP((FVector)B), Color, LineThicknessIn, bDepthTestedIn);
+	}
+
+
+	/** Draw a point with default parameters */
+	template<typename PointType>
+	void DrawPoint(const PointType& Position)
+	{
+		InternalDrawTransformedPoint(TransformP((FVector)Position), PointColor, PointSize, bDepthTested);
+	}
+
+	/** Draw a point with the given Color, otherwise use default parameters */
+	template<typename PointType>
+	void DrawPoint(const PointType& Position, const FLinearColor& Color)
+	{
+		InternalDrawTransformedPoint(TransformP((FVector)Position), Color, PointSize, bDepthTested);
+	}
+
+	/** Draw a point with the given parameters */
+	template<typename PointType>
+	void DrawPoint(const PointType& Position, const FLinearColor& Color, float PointSizeIn, bool bDepthTestedIn)
+	{
+		InternalDrawTransformedPoint(TransformP((FVector)Position), Color, PointSizeIn, bDepthTestedIn);
+	}
+
+
+	/** Draw a 3D circle at given position/normal with the given parameters. Tangent axes are defined internally */
+	template<typename PointType>
+	void DrawCircle(const PointType& Position, const PointType& Normal, float Radius, int Steps, const FLinearColor& Color, float LineThicknessIn, bool bDepthTestedIn)
+	{
+		InternalDrawCircle((FVector)Position, (FVector)Normal, Radius, Steps, Color, LineThicknessIn, bDepthTestedIn);
+	}
+
+	/** Draw a 3D circle at given position/normal with the given parameters. Tangent axes are defined internally */
+	template<typename PointType>
+	void DrawViewFacingCircle(const PointType& Position, float Radius, int Steps, const FLinearColor& Color, float LineThicknessIn, bool bDepthTestedIn)
+	{
+		InternalDrawViewFacingCircle((FVector)Position, Radius, Steps, Color, LineThicknessIn, bDepthTestedIn);
+	}
+
+	/** Draw a 3D cylinder, parameterized the same as the 3D circle but extruded by Height */
+	template<typename PointType>
+	void DrawWireCylinder(const PointType& Position, const PointType& Normal, float Radius, float Height, int Steps, const FLinearColor& Color, float LineThicknessIn, bool bDepthTestedIn)
+	{
+		InternalDrawWireCylinder((FVector)Position, (FVector)Normal, Radius, Height, Steps, Color, LineThicknessIn, bDepthTestedIn);
+	}
+
+	/** Draw a 3D cylinder, parameterized the same as the 3D circle but extruded by Height */
+	template<typename PointType>
+	void DrawWireCylinder(const PointType& Position, const PointType& Normal, float Radius, float Height, int Steps)
+	{
+		InternalDrawWireCylinder((FVector)Position, (FVector)Normal, Radius, Height, Steps, LineColor, LineThickness, bDepthTested);
+	}
+
+	/** Draw a 3D box, parameterized the same as the 3D circle but extruded by Height */
+	void DrawWireBox(const FBox& Box)
+	{
+		InternalDrawWireBox(Box, LineColor, LineThickness, bDepthTested);
+	}
+
+	void DrawWireBox(const FBox& Box, const FLinearColor& Color, float LineThicknessIn, bool bDepthTestedIn)
+	{
+		InternalDrawWireBox(Box, Color, LineThicknessIn, bDepthTestedIn);
+	}
+
+	template<typename PointType>
+	void DrawSquare(const PointType& Center, const PointType& SideA, const PointType& SideB, const FLinearColor& Color, float LineThicknessIn, bool bDepthTestedIn)
+	{
+		InternalDrawSquare(Center, SideA, SideB, Color, LineThicknessIn, bDepthTestedIn);
+	}
+
+	template<typename PointType>
+	void DrawSquare(const PointType& Center, const PointType& SideA, const PointType& SideB)
+	{
+		InternalDrawSquare(Center, SideA, SideB, LineColor, LineThickness, bDepthTested);
+	}
+
+
+protected:
+	/** We use this for drawing, extracted in InitializeFrame */
+	FPrimitiveDrawInterface* CurrentPDI = nullptr;
+	
+	FViewCameraState CameraState;
+	bool bHaveCameraState;
+
+
+	TArray<FTransform> TransformStack;
+	FTransform TotalTransform;
+
+	// actually does the line drawing; assumes A and B are already transformed
+	virtual void InternalDrawTransformedLine(const FVector& A, const FVector& B, const FLinearColor& Color, float LineThickness, bool bDepthTested);
+	// actually does the point drawing; assumes Position is already transformed
+	virtual void InternalDrawTransformedPoint(const FVector& Position, const FLinearColor& Color, float PointSize, bool bDepthTested);
+
+
+	// actually does the circle drawing
+	virtual void InternalDrawCircle(const FVector& Position, const FVector& Normal, float Radius, int Steps, const FLinearColor& Color, float LineThickness, bool bDepthTested);
+	// actually does the circle drawing
+	virtual void InternalDrawViewFacingCircle(const FVector& Position, float Radius, int Steps, const FLinearColor& Color, float LineThickness, bool bDepthTested);
+	// actually does the cylinder drawing
+	virtual void InternalDrawWireCylinder(const FVector& Position, const FVector& Normal, float Radius, float Height, int Steps, const FLinearColor& Color, float LineThickness, bool bDepthTested);
+	// actually does the box drawing
+	virtual void InternalDrawWireBox(const FBox& Box, const FLinearColor& Color, float LineThickness, bool bDepthTested);
+	// actually does the square drawing
+	virtual void InternalDrawSquare(const FVector& Center, const FVector& SideA, const FVector& SideB, const FLinearColor& Color, float LineThickness, bool bDepthTested);
+
+
+};
\ No newline at end of file
