GetActiveToolBuilder system fully hooked up

diff --git a/Engine/Source/Runtime/Experimental/InteractiveToolsFramework/Private/InteractiveToolManager.cpp b/Engine/Source/Runtime/Experimental/InteractiveToolsFramework/Private/InteractiveToolManager.cpp
--- a/Engine/Source/Runtime/Experimental/InteractiveToolsFramework/Private/InteractiveToolManager.cpp
+++ b/Engine/Source/Runtime/Experimental/InteractiveToolsFramework/Private/InteractiveToolManager.cpp
@@ -182,6 +182,11 @@ UInteractiveTool* UInteractiveToolManager::GetActiveTool(EToolSide Side)
 	return (Side == EToolSide::Left) ? ActiveLeftTool : ActiveRightTool;
 }
 
+UInteractiveToolBuilder* UInteractiveToolManager::GetActiveToolBuilder(EToolSide Side)
+{
+	return (Side == EToolSide::Left) ? ActiveLeftBuilder : ActiveRightBuilder;
+}
+
 
 bool UInteractiveToolManager::CanAcceptActiveTool(EToolSide Side)
 {
diff --git a/Engine/Source/Runtime/Experimental/InteractiveToolsFramework/Private/InteractiveToolsContext.cpp b/Engine/Source/Runtime/Experimental/InteractiveToolsFramework/Private/InteractiveToolsContext.cpp
--- a/Engine/Source/Runtime/Experimental/InteractiveToolsFramework/Private/InteractiveToolsContext.cpp
+++ b/Engine/Source/Runtime/Experimental/InteractiveToolsFramework/Private/InteractiveToolsContext.cpp
@@ -109,7 +109,7 @@ void UInteractiveToolsContext::EndTool(EToolSide WhichSide, EToolShutdownType Sh
 	DeactivateActiveTool(WhichSide, ShutdownType);
 }
 
-bool UInteractiveToolsContext::IsToolActive(EToolSide WhichSide, UInteractiveToolBuilder* Builder)
+bool UInteractiveToolsContext::IsToolBuilderActive(EToolSide WhichSide, UInteractiveToolBuilder* Builder)
 {
 	return ToolManager->GetActiveToolBuilder(WhichSide) == Builder;
 }
diff --git a/Engine/Source/Runtime/Experimental/InteractiveToolsFramework/Public/InteractiveToolManager.h b/Engine/Source/Runtime/Experimental/InteractiveToolsFramework/Public/InteractiveToolManager.h
--- a/Engine/Source/Runtime/Experimental/InteractiveToolsFramework/Public/InteractiveToolManager.h
+++ b/Engine/Source/Runtime/Experimental/InteractiveToolsFramework/Public/InteractiveToolManager.h
@@ -112,6 +112,13 @@ public:
 	 */
 	virtual UInteractiveTool* GetActiveTool(EToolSide Side);
 
+	/**
+	 * Get pointer to active Tool Builder on a given side
+	 * @param Side which Side is being requested
+	 * @return pointer to Tool Builder instance active on that Side, or nullptr if no such ToolBuilder exists
+	 */
+	virtual UInteractiveToolBuilder* GetActiveToolBuilder(EToolSide Side);
+
 
 	/**
 	 * Check if an active Tool on the given Side can be Accepted in its current state
diff --git a/Engine/Source/Runtime/Experimental/InteractiveToolsFramework/Public/InteractiveToolsContext.h b/Engine/Source/Runtime/Experimental/InteractiveToolsFramework/Public/InteractiveToolsContext.h
--- a/Engine/Source/Runtime/Experimental/InteractiveToolsFramework/Public/InteractiveToolsContext.h
+++ b/Engine/Source/Runtime/Experimental/InteractiveToolsFramework/Public/InteractiveToolsContext.h
@@ -40,7 +40,7 @@ public:
 	bool CanCompleteActiveTool(EToolSide WhichSide) const;
 	bool StartTool(EToolSide WhichSide, const FString& ToolTypeIdentifier);
 	void EndTool(EToolSide WhichSide, EToolShutdownType ShutdownType);
-	bool IsToolActive(EToolSide WhichSide, UInteractiveToolBuilder* Builder);
+	bool IsToolBuilderActive(EToolSide WhichSide, UInteractiveToolBuilder* Builder);
 
 public:
 	// forwards message to OnToolNotificationMessage delegate
