Changing default behavior when swapping away from tools to auto-accept instead of auto-cancel

diff --git a/Engine/Source/Runtime/Experimental/InteractiveToolsFramework/Private/InteractiveToolManager.cpp b/Engine/Source/Runtime/Experimental/InteractiveToolsFramework/Private/InteractiveToolManager.cpp
--- a/Engine/Source/Runtime/Experimental/InteractiveToolsFramework/Private/InteractiveToolManager.cpp
+++ b/Engine/Source/Runtime/Experimental/InteractiveToolsFramework/Private/InteractiveToolManager.cpp
@@ -80,11 +80,6 @@ bool UInteractiveToolManager::CanActivateTool(EToolSide Side, const FString& Ide
 {
 	check(Side == EToolSide::Left);   // TODO: support right-side tool
 
-	if (ActiveLeftTool != nullptr)
-	{
-		return false;
-	}
-
 	if (ToolBuilders.Contains(Identifier))
 	{
 		FToolBuilderState InputState;
@@ -104,7 +99,7 @@ bool UInteractiveToolManager::ActivateTool(EToolSide Side)
 
 	if (ActiveLeftTool != nullptr) 
 	{
-		DeactivateTool(EToolSide::Left, EToolShutdownType::Cancel);
+		DeactivateTool(EToolSide::Left, EToolShutdownType::Accept);
 	}
 
 	if (ActiveLeftBuilder == nullptr) 
