Changing FindFirstComponent to verify that the FindByPredicate found a result before dereferencing

diff --git a/Engine/Source/Runtime/Experimental/InteractiveToolsFramework/Private/ToolBuilderUtil.cpp b/Engine/Source/Runtime/Experimental/InteractiveToolsFramework/Private/ToolBuilderUtil.cpp
--- a/Engine/Source/Runtime/Experimental/InteractiveToolsFramework/Private/ToolBuilderUtil.cpp
+++ b/Engine/Source/Runtime/Experimental/InteractiveToolsFramework/Private/ToolBuilderUtil.cpp
@@ -35,16 +35,20 @@ UActorComponent* ToolBuilderUtil::FindFirstComponent(const FToolBuilderState& In
 {
 	if (InputState.SelectedComponents.Num() > 0)
 	{
-		return *InputState.SelectedComponents.FindByPredicate(Predicate);
+		UActorComponent* const* ComponentPtr = InputState.SelectedComponents.FindByPredicate(Predicate);
+		if (ComponentPtr)
+		{
+			return *ComponentPtr;
+		}
 	}
 	else
 	{
 		for ( AActor* Actor : InputState.SelectedActors )
 		{
-			UActorComponent* Component = *Algo::FindByPredicate(Actor->GetComponents(), Predicate);
-			if ( Component )
+			UActorComponent* const* ComponentPtr = Algo::FindByPredicate(Actor->GetComponents(), Predicate);
+			if (ComponentPtr)
 			{
-				return Component;
+				return *ComponentPtr;
 			}
 		}
 	}
