Geometry Processing - Interactive Tools.  Adding Interval Gizmo.

diff --git a/Engine/Source/Runtime/Experimental/InteractiveToolsFramework/Private/BaseGizmos/GizmoActor.cpp b/Engine/Source/Runtime/Experimental/InteractiveToolsFramework/Private/BaseGizmos/GizmoActor.cpp
--- a/Engine/Source/Runtime/Experimental/InteractiveToolsFramework/Private/BaseGizmos/GizmoActor.cpp
+++ b/Engine/Source/Runtime/Experimental/InteractiveToolsFramework/Private/BaseGizmos/GizmoActor.cpp
@@ -5,6 +5,7 @@
 #include "BaseGizmos/GizmoArrowComponent.h"
 #include "BaseGizmos/GizmoRectangleComponent.h"
 #include "BaseGizmos/GizmoCircleComponent.h"
+#include "BaseGizmos/GizmoLineHandleComponent.h"
 
 #include "Components/PrimitiveComponent.h"
 #include "Engine/World.h"
@@ -27,13 +28,14 @@ AGizmoActor::AGizmoActor()
 
 UGizmoArrowComponent* AGizmoActor::AddDefaultArrowComponent(
 	UWorld* World, AActor* Actor,
-	const FLinearColor& Color, const FVector& LocalDirection)
+	const FLinearColor& Color, const FVector& LocalDirection, const float Length)
 {
 	UGizmoArrowComponent* NewArrow = NewObject<UGizmoArrowComponent>(Actor);
 	Actor->AddInstanceComponent(NewArrow);
 	NewArrow->AttachToComponent(Actor->GetRootComponent(), FAttachmentTransformRules::KeepRelativeTransform);
 	NewArrow->Direction = LocalDirection;
 	NewArrow->Color = Color;
+	NewArrow->Length = Length;
 	NewArrow->RegisterComponent();
 	return NewArrow;
 }
@@ -68,6 +70,21 @@ UGizmoCircleComponent* AGizmoActor::AddDefaultCircleComponent(
 	return NewCircle;
 }
 
-
+UGizmoLineHandleComponent* AGizmoActor::AddDefaultLineHandleComponent(
+	UWorld* World, AActor* Actor,
+	const FLinearColor& Color, const FVector& HandleNormal, const FVector& LocalDirection,
+    const float Length, const bool bImageScale)
+{
+	UGizmoLineHandleComponent* LineHandle = NewObject<UGizmoLineHandleComponent>(Actor);
+	Actor->AddInstanceComponent(LineHandle);
+	LineHandle->AttachToComponent(Actor->GetRootComponent(), FAttachmentTransformRules::KeepRelativeTransform);
+	LineHandle->Normal = HandleNormal;
+	LineHandle->Direction = LocalDirection;
+	LineHandle->Length = Length;
+	LineHandle->bImageScale = bImageScale;
+	LineHandle->Color = Color;
+	LineHandle->RegisterComponent();
+	return LineHandle;
+}
 
 #undef LOCTEXT_NAMESPACE
\ No newline at end of file
diff --git a/Engine/Source/Runtime/Experimental/InteractiveToolsFramework/Private/BaseGizmos/GizmoLineHandleComponent.cpp b/Engine/Source/Runtime/Experimental/InteractiveToolsFramework/Private/BaseGizmos/GizmoLineHandleComponent.cpp
new file mode 100644
--- /dev/null
+++ b/Engine/Source/Runtime/Experimental/InteractiveToolsFramework/Private/BaseGizmos/GizmoLineHandleComponent.cpp
@@ -0,0 +1,211 @@
+// Copyright 1998-2019 Epic Games, Inc. All Rights Reserved.
+
+#include "BaseGizmos/GizmoLineHandleComponent.h"
+#include "BaseGizmos/GizmoRenderingUtil.h"
+#include "BaseGizmos/GizmoMath.h"
+#include "PrimitiveSceneProxy.h"
+
+
+
+
+class FGizmoLineHandleComponentSceneProxy final : public FPrimitiveSceneProxy
+{
+public:
+	SIZE_T GetTypeHash() const override
+	{
+		static size_t UniquePointer;
+		return reinterpret_cast<size_t>(&UniquePointer);
+	}
+
+	FGizmoLineHandleComponentSceneProxy(const UGizmoLineHandleComponent* InComponent)
+		: FPrimitiveSceneProxy(InComponent),
+		Color(InComponent->Color),
+		Normal(InComponent->Normal),
+		Direction(InComponent->Direction),
+		HandleSize(InComponent->HandleSize),
+		Thickness(InComponent->Thickness),
+		bBoundaryOnly(false),
+		bImageScale(InComponent->bImageScale),
+		HoverThicknessMultiplier(InComponent->HoverSizeMultiplier)
+	{
+	}
+
+	virtual void GetDynamicMeshElements(const TArray<const FSceneView*>& Views, const FSceneViewFamily& ViewFamily, uint32 VisibilityMap, FMeshElementCollector& Collector) const override
+	{
+		// find view to use for gizmo sizing/etc
+		const FSceneView* GizmoControlView = GizmoRenderingUtil::FindActiveSceneView(Views, ViewFamily, VisibilityMap);
+		if (GizmoControlView == nullptr)
+		{
+			return;
+		}
+
+		FVector LocalOffset = Direction;
+		if (ExternalDistance != nullptr)
+		{
+			LocalOffset *= (*ExternalDistance);
+		}
+
+		const FMatrix& LocalToWorldMatrix = GetLocalToWorld();
+
+		float IntervalMarkerSize = HandleSize;
+		FVector WorldIntervalEnd = LocalToWorldMatrix.TransformPosition(LocalOffset + IntervalMarkerSize * Normal);
+
+		FVector WorldDiskOrigin = LocalToWorldMatrix.TransformPosition(LocalOffset);
+		FVector WorldBaseOrigin = LocalToWorldMatrix.TransformPosition(FVector::ZeroVector);
+
+
+		float LengthScale = 1.0f;
+		if (ExternalDynamicPixelToWorldScale != nullptr && bImageScale)
+		{
+			float PixelToWorldScale = GizmoRenderingUtil::CalculateLocalPixelToWorldScale(GizmoControlView, WorldDiskOrigin);
+			*ExternalDynamicPixelToWorldScale = PixelToWorldScale;
+			LengthScale = PixelToWorldScale;
+		}
+
+
+		FVector ScaledDiskOrigin = LengthScale * (WorldDiskOrigin - WorldBaseOrigin) + WorldBaseOrigin;
+		FVector ScaledIntevalStart = -LengthScale * (WorldIntervalEnd - WorldDiskOrigin) + WorldDiskOrigin;
+		FVector ScaledIntevalEnd = LengthScale * (WorldIntervalEnd - WorldDiskOrigin) + WorldDiskOrigin;
+
+
+		//double UseRadius = LengthScale * Radius;
+
+		FLinearColor BackColor = FLinearColor(0.5f, 0.5f, 0.5f);
+		float BackThickness = 0.5f;
+
+		float UseThickness = (bExternalHoverState != nullptr && *bExternalHoverState == true) ?
+			(HoverThicknessMultiplier*Thickness) : (Thickness);
+
+		for (int32 ViewIndex = 0; ViewIndex < Views.Num(); ViewIndex++)
+		{
+			if (VisibilityMap & (1 << ViewIndex))
+			{
+				const FSceneView* View = Views[ViewIndex];
+				FPrimitiveDrawInterface* PDI = Collector.GetPDI(ViewIndex);
+
+				FVector UpVector = View->GetViewUp();
+
+
+				FVector ViewVector = View->GetViewDirection();
+
+				// From base origin to disk origin
+				PDI->DrawLine(WorldBaseOrigin, ScaledDiskOrigin, Color, SDPG_Foreground, UseThickness, 0.0f, true);
+
+				// Draw the interval marker
+				PDI->DrawLine(ScaledIntevalStart, ScaledIntevalEnd, Color, SDPG_Foreground, 2. * UseThickness, 0.0f, true);
+
+			}
+		}
+	}
+
+	virtual FPrimitiveViewRelevance GetViewRelevance(const FSceneView* View) const override
+	{
+		FPrimitiveViewRelevance Result;
+		Result.bDrawRelevance = IsShown(View);
+		Result.bDynamicRelevance = true;
+		Result.bShadowRelevance = false;
+		Result.bEditorPrimitiveRelevance = UseEditorCompositing(View);
+		Result.bRenderCustomDepth = ShouldRenderCustomDepth();
+		return Result;
+	}
+
+	virtual bool CanBeOccluded() const override
+	{
+		return false;
+	}
+
+	virtual uint32 GetMemoryFootprint(void) const override { return sizeof *this + GetAllocatedSize(); }
+	uint32 GetAllocatedSize(void) const { return FPrimitiveSceneProxy::GetAllocatedSize(); }
+
+
+	void SetExternalDynamicPixelToWorldScale(float* DynamicPixelToWorldScale)
+	{
+		ExternalDynamicPixelToWorldScale = DynamicPixelToWorldScale;
+	}
+
+	void SetExternalHoverState(bool* HoverState)
+	{
+		bExternalHoverState = HoverState;
+	}
+
+	void SetExternalWorldLocalState(bool* bWorldLocalState)
+	{
+		bExternalWorldLocalState = bWorldLocalState;
+	}
+
+	void SetLengthScale(float* Distance)
+	{
+		ExternalDistance = Distance;
+	}
+
+private:
+	FLinearColor Color;
+	FVector Normal;
+	FVector Direction;
+	float HandleSize;
+	float Thickness;
+	bool bBoundaryOnly;
+	bool bImageScale;
+	float HoverThicknessMultiplier;
+
+	float* ExternalDistance = nullptr;
+	float* ExternalDynamicPixelToWorldScale = nullptr;
+	bool* bExternalHoverState = nullptr;
+	bool* bExternalWorldLocalState = nullptr;
+};
+
+
+FPrimitiveSceneProxy* UGizmoLineHandleComponent::CreateSceneProxy()
+{
+	FGizmoLineHandleComponentSceneProxy* NewProxy = new FGizmoLineHandleComponentSceneProxy(this);
+	NewProxy->SetExternalDynamicPixelToWorldScale(&DynamicPixelToWorldScale);
+	NewProxy->SetExternalHoverState(&bHovering);
+	NewProxy->SetExternalWorldLocalState(&bWorld);
+	NewProxy->SetLengthScale(&Length);
+	return NewProxy;
+}
+
+FBoxSphereBounds UGizmoLineHandleComponent::CalcBounds(const FTransform& LocalToWorld) const
+{
+	// the handle looks like
+	//                         ------|
+	// where '------' has length "Length" and '|' is of length 2*Handlesize
+	//  
+	float Radius = FMath::Sqrt(Length * Length + HandleSize * HandleSize);
+	return FBoxSphereBounds(FSphere(FVector::ZeroVector, Radius).TransformBy(LocalToWorld));
+}
+
+bool UGizmoLineHandleComponent::LineTraceComponent(FHitResult& OutHit, const FVector Start, const FVector End, const FCollisionQueryParams& Params)
+{
+	float LengthScale = (bImageScale) ? DynamicPixelToWorldScale : 1.f;
+	double UseHandleSize = LengthScale * HandleSize;
+	FVector LocalOffset = Length * Direction;
+
+	const FTransform& Transform = this->GetComponentToWorld();
+	FVector HandleDir = (bWorld) ? Normal : Transform.TransformVector(Normal);
+	FVector WorldBaseOrigin = Transform.TransformPosition(FVector::ZeroVector);
+	FVector WorldHandleOrigin = Transform.TransformPosition(LocalOffset);
+
+	FVector BaseToHandle = WorldHandleOrigin - WorldBaseOrigin;
+
+	// where the handle crosses the connecting line
+	FVector ScaledHandleOrigin = LengthScale * BaseToHandle + WorldBaseOrigin;
+
+	// start and end point of the handle.
+	FVector HandleStart = ScaledHandleOrigin + LengthScale * HandleDir;
+	FVector HandleEnd = ScaledHandleOrigin - LengthScale * HandleDir;
+
+	FVector NearestOnHandle, NearestOnLine;
+	FMath::SegmentDistToSegmentSafe(HandleStart, HandleEnd, Start, End, NearestOnHandle, NearestOnLine);
+	double Distance = FVector::Distance(NearestOnHandle, NearestOnLine);
+	if (Distance > PixelHitDistanceThreshold*DynamicPixelToWorldScale)
+	{
+		return false;
+	}
+
+	OutHit.Component = this;
+	OutHit.Distance = FVector::Distance(Start, NearestOnLine);
+	OutHit.ImpactPoint = NearestOnLine;
+	return true;
+
+}
\ No newline at end of file
diff --git a/Engine/Source/Runtime/Experimental/InteractiveToolsFramework/Private/BaseGizmos/IntervalGizmo.cpp b/Engine/Source/Runtime/Experimental/InteractiveToolsFramework/Private/BaseGizmos/IntervalGizmo.cpp
new file mode 100644
--- /dev/null
+++ b/Engine/Source/Runtime/Experimental/InteractiveToolsFramework/Private/BaseGizmos/IntervalGizmo.cpp
@@ -0,0 +1,432 @@
+// Copyright 1998-2019 Epic Games, Inc. All Rights Reserved.
+
+#include "BaseGizmos/IntervalGizmo.h"
+#include "InteractiveGizmoManager.h"
+#include "BaseGizmos/AxisPositionGizmo.h"
+
+#include "BaseGizmos/GizmoComponents.h"
+#include "BaseGizmos/GizmoLineHandleComponent.h"
+
+// need this to implement hover
+#include "BaseGizmos/GizmoBaseComponent.h"
+
+#include "Components/SphereComponent.h"
+#include "Components/PrimitiveComponent.h"
+#include "Engine/World.h"
+#include "Engine/CollisionProfile.h"
+
+
+#define LOCTEXT_NAMESPACE "UIntervalGizmo"
+
+/**
+* FFloatParameterProxyChange tracks a change to the base transform for a FloatParameter
+*/
+class FFloatParameterProxyChange : public FToolCommandChange
+{
+public:
+	FGizmoFloatParameterChange To;
+	FGizmoFloatParameterChange From;
+
+
+	virtual void Apply(UObject* Object) override 
+	{
+		UGizmoLocalFloatParameterSource* ParameterSource = CastChecked<UGizmoLocalFloatParameterSource>(Object);
+		ParameterSource->SetParameter(To.CurrentValue);
+	}
+	virtual void Revert(UObject* Object) override
+	{
+		UGizmoLocalFloatParameterSource* ParameterSource = CastChecked<UGizmoLocalFloatParameterSource>(Object);
+		ParameterSource->SetParameter(From.CurrentValue);
+	}
+
+	virtual FString ToString() const override { return TEXT("FFloatParameterProxyChange"); }
+};
+
+/**
+ * FGizmoFloatParameterChangeSource generates FFloatParameterProxyChange instances on Begin/End.
+ * Instances of this class can (for example) be attached to a UGizmoTransformChangeStateTarget for use TransformGizmo change tracking.
+ */
+class FGizmoFloatParameterChangeSource : public IToolCommandChangeSource
+{
+public:
+	FGizmoFloatParameterChangeSource(UGizmoLocalFloatParameterSource* ProxyIn)
+	{
+		Proxy = ProxyIn;
+	}
+
+	virtual ~FGizmoFloatParameterChangeSource() {}
+
+	TWeakObjectPtr<UGizmoLocalFloatParameterSource> Proxy;
+	TUniquePtr<FFloatParameterProxyChange> ActiveChange;
+
+	virtual void BeginChange() override
+	{
+		if (Proxy.IsValid())
+		{
+			ActiveChange = MakeUnique<FFloatParameterProxyChange>();
+			ActiveChange->From = Proxy->LastChange;
+		}
+	}
+	virtual TUniquePtr<FToolCommandChange> EndChange() override
+	{
+		if (Proxy.IsValid())
+		{
+			ActiveChange->To = Proxy->LastChange;
+			return MoveTemp(ActiveChange);
+		}
+		return TUniquePtr<FToolCommandChange>();
+	}
+	virtual UObject* GetChangeTarget() override
+	{
+		return Proxy.Get();
+	}
+	virtual FText GetChangeDescription() override
+	{
+		return LOCTEXT("FFGizmoFloatParameterChangeDescription", "GizmoFloatParameterChange");
+	}
+};
+
+
+AIntervalGizmoActor::AIntervalGizmoActor()
+{
+	// root component is a hidden sphere
+	USphereComponent* SphereComponent = CreateDefaultSubobject<USphereComponent>(TEXT("GizmoCenter"));
+	RootComponent = SphereComponent;
+	SphereComponent->InitSphereRadius(1.0f);
+	SphereComponent->SetVisibility(false);
+	SphereComponent->SetCollisionProfileName(UCollisionProfile::NoCollision_ProfileName);
+}
+
+AIntervalGizmoActor* AIntervalGizmoActor::ConstructDefaultIntervalGizmo(UWorld* World)
+{
+	FActorSpawnParameters SpawnInfo;
+	AIntervalGizmoActor* NewActor = World->SpawnActor<AIntervalGizmoActor>(FVector::ZeroVector, FRotator::ZeroRotator, SpawnInfo);
+
+	const FLinearColor MintGreen(152 / 255.f, 1.f, 152 / 255.f);
+	
+	// add components 
+	NewActor->UpIntervalComponent      = AddDefaultLineHandleComponent(World, NewActor, MintGreen, FVector(0, 1, 0), FVector(0, 0, 1));
+	NewActor->DownIntervalComponent    = AddDefaultLineHandleComponent(World, NewActor, MintGreen, FVector(0, 1, 0), FVector(0, 0, 1));
+	NewActor->ForwardIntervalComponent = AddDefaultLineHandleComponent(World, NewActor, MintGreen, FVector(1, 0, 0), FVector(0, 1, 0));
+
+
+	return NewActor;
+}
+
+
+
+UInteractiveGizmo* UIntervalGizmoBuilder::BuildGizmo(const FToolBuilderState& SceneState) const
+{
+	UIntervalGizmo* NewGizmo = NewObject<UIntervalGizmo>(SceneState.GizmoManager);
+	NewGizmo->SetWorld(SceneState.World);
+
+	// use default gizmo actor if client has not given us a new builder
+	NewGizmo->SetGizmoActorBuilder((GizmoActorBuilder) ? GizmoActorBuilder : MakeShared<FIntervalGizmoActorFactory>());
+
+	// override default hover function if proposed
+	if (UpdateHoverFunction)
+	{
+		NewGizmo->SetUpdateHoverFunction(UpdateHoverFunction);
+	}
+
+	if (UpdateCoordSystemFunction)
+	{
+		NewGizmo->SetUpdateCoordSystemFunction(UpdateCoordSystemFunction);
+	}
+
+	return NewGizmo;
+}
+
+// Init static FName
+FString UIntervalGizmo::GizmoName = TEXT("IntervalGizmo");
+
+void UIntervalGizmo::SetWorld(UWorld* WorldIn)
+{
+	this->World = WorldIn;
+}
+
+void UIntervalGizmo::SetGizmoActorBuilder(TSharedPtr<FIntervalGizmoActorFactory> Builder)
+{
+	GizmoActorBuilder = Builder;
+}
+
+void UIntervalGizmo::SetUpdateHoverFunction(TFunction<void(UPrimitiveComponent*, bool)> HoverFunction)
+{
+	UpdateHoverFunction = HoverFunction;
+}
+
+void UIntervalGizmo::SetUpdateCoordSystemFunction(TFunction<void(UPrimitiveComponent*, EToolContextCoordinateSystem)> CoordSysFunction)
+{
+	UpdateCoordSystemFunction = CoordSysFunction;
+}
+
+void UIntervalGizmo::Setup()
+{
+	UInteractiveGizmo::Setup();
+
+	UpdateHoverFunction = [](UPrimitiveComponent* Component, bool bHovering)
+	{
+		if (Cast<UGizmoBaseComponent>(Component) != nullptr)
+		{
+			Cast<UGizmoBaseComponent>(Component)->UpdateHoverState(bHovering);
+		}
+	};
+
+	UpdateCoordSystemFunction = [](UPrimitiveComponent* Component, EToolContextCoordinateSystem CoordSystem)
+	{
+		if (Cast<UGizmoBaseComponent>(Component) != nullptr)
+		{
+			Cast<UGizmoBaseComponent>(Component)->UpdateWorldLocalState(CoordSystem == EToolContextCoordinateSystem::World);
+		}
+	};
+
+	GizmoActor = GizmoActorBuilder->CreateNewGizmoActor(World);
+}
+
+void UIntervalGizmo::Shutdown()
+{
+	ClearActiveTarget();
+
+	if (GizmoActor)
+	{
+		GizmoActor->Destroy();
+		GizmoActor = nullptr;
+	}
+
+	ClearSources();
+
+}
+
+void UIntervalGizmo::Tick(float DeltaTime)
+{
+	EToolContextCoordinateSystem CoordSystem = GetGizmoManager()->GetContextQueriesAPI()->GetCurrentCoordinateSystem();
+	check(CoordSystem == EToolContextCoordinateSystem::World || CoordSystem == EToolContextCoordinateSystem::Local)
+		bool bUseLocalAxes =
+		(GetGizmoManager()->GetContextQueriesAPI()->GetCurrentCoordinateSystem() == EToolContextCoordinateSystem::Local);
+
+	// Update gizmo location.
+	{
+		USceneComponent* GizmoComponent = GizmoActor->GetRootComponent();
+		// move gizmo to target location
+		FTransform TargetTransform = TransformProxy->GetTransform();
+		FVector SaveScale = TargetTransform.GetScale3D();
+		TargetTransform.SetScale3D(FVector(1, 1, 1));
+		GizmoComponent->SetWorldTransform(TargetTransform);
+	}
+	// Update the lengths
+	{
+		if (GizmoActor->UpIntervalComponent != nullptr && UpIntervalSource != nullptr)
+		{
+			GizmoActor->UpIntervalComponent->Length = UpIntervalSource->GetParameter();
+		}
+		if (GizmoActor->DownIntervalComponent != nullptr && DownIntervalSource != nullptr)
+		{
+			GizmoActor->DownIntervalComponent->Length = DownIntervalSource->GetParameter();
+		}
+		if (GizmoActor->ForwardIntervalComponent != nullptr && ForwardIntervalSource != nullptr)
+		{
+			GizmoActor->ForwardIntervalComponent->Length = ForwardIntervalSource->GetParameter();
+		}
+	}
+
+	if (UpdateCoordSystemFunction)
+	{
+		for (UPrimitiveComponent* Component : ActiveComponents)
+		{
+			UpdateCoordSystemFunction(Component, CoordSystem);
+		}
+	}
+}
+
+void UIntervalGizmo::SetActiveTarget(UTransformProxy* TransformTargetIn, UGizmoLocalFloatParameterSource* UpInterval , UGizmoLocalFloatParameterSource* DownInterval, UGizmoLocalFloatParameterSource* ForwardInterval, IToolContextTransactionProvider* TransactionProvider)
+{
+	if (TransformProxy != nullptr)
+	{
+		ClearActiveTarget();
+		ClearSources();
+
+	}
+
+	// This state target emits an explicit FChange that moves the GizmoActor root component during undo/redo.
+	// It also opens/closes the Transaction that saves/restores the target object locations.
+	if (TransactionProvider == nullptr)
+	{
+		TransactionProvider = GetGizmoManager();
+	}
+
+	TransformProxy = TransformTargetIn;
+
+	// parameters and init lengths for each interval
+	UpIntervalSource      = UpInterval;
+	DownIntervalSource    = DownInterval;
+	ForwardIntervalSource = ForwardInterval;
+
+
+	USceneComponent* GizmoComponent = GizmoActor->GetRootComponent();
+
+	// move gizmo to target location
+	FTransform TargetTransform = TransformTargetIn->GetTransform();
+	FVector SaveScale = TargetTransform.GetScale3D();
+	TargetTransform.SetScale3D(FVector(1, 1, 1));
+	GizmoComponent->SetWorldTransform(TargetTransform);
+
+	
+	// TargetTransform tracks location of GizmoComponent. Note that TransformUpdated is not called during undo/redo transactions!
+	// We currently rely on the transaction system to undo/redo target object locations. This will not work during runtime...
+	GizmoComponent->TransformUpdated.AddLambda(
+		[this, SaveScale](USceneComponent* Component, EUpdateTransformFlags /*UpdateTransformFlags*/, ETeleportType /*Teleport*/) {
+		//this->GetGizmoManager()->DisplayMessage(TEXT("TRANSFORM UPDATED"), EToolMessageLevel::Internal);
+		FTransform NewXForm = Component->GetComponentToWorld();
+		NewXForm.SetScale3D(SaveScale);
+		this->TransformProxy->SetTransform(NewXForm);
+	});
+
+	
+
+
+	StateTarget = UGizmoTransformChangeStateTarget::Construct(GizmoComponent,
+		LOCTEXT("UIntervalGizmoTransaction", "Interval"), TransactionProvider, this);
+	StateTarget->DependentChangeSources.Add(MakeUnique<FTransformProxyChangeSource>(TransformProxy));
+	StateTarget->DependentChangeSources.Add(MakeUnique<FGizmoFloatParameterChangeSource>(UpIntervalSource));
+	StateTarget->DependentChangeSources.Add(MakeUnique<FGizmoFloatParameterChangeSource>(DownIntervalSource));
+	StateTarget->DependentChangeSources.Add(MakeUnique<FGizmoFloatParameterChangeSource>(ForwardIntervalSource));
+
+
+	UGizmoTransformProxyTransformSource* TransformSource =
+		UGizmoTransformProxyTransformSource::Construct(TransformProxy, this);
+
+	// root component provides local Y/Z axis, identified by AxisIndex
+	AxisYSource = UGizmoComponentAxisSource::Construct(GizmoComponent, 1, true, this);
+	AxisZSource = UGizmoComponentAxisSource::Construct(GizmoComponent, 2, true, this);
+
+	
+	
+	if (GizmoActor->UpIntervalComponent != nullptr)
+	{
+		float IntervalMin = 0.f;
+		float IntervalMax = FLT_MAX;
+		AddIntervalHandleGizmo(GizmoComponent, GizmoActor->UpIntervalComponent, AxisZSource, UpIntervalSource, IntervalMin, IntervalMax, StateTarget);
+		ActiveComponents.Add(GizmoActor->UpIntervalComponent);
+	}
+
+	if (GizmoActor->DownIntervalComponent != nullptr)
+	{
+		float IntervalMin = -FLT_MAX;
+		float IntervalMax = 0.f;
+		AddIntervalHandleGizmo(GizmoComponent, GizmoActor->DownIntervalComponent, AxisZSource, DownIntervalSource, IntervalMin, IntervalMax, StateTarget);
+		ActiveComponents.Add(GizmoActor->DownIntervalComponent);
+	}
+
+	if (GizmoActor->ForwardIntervalComponent != nullptr)
+	{
+		float IntervalMin = -FLT_MAX;
+		float IntervalMax = FLT_MAX;
+		AddIntervalHandleGizmo(GizmoComponent, GizmoActor->ForwardIntervalComponent, AxisYSource, ForwardIntervalSource, IntervalMin, IntervalMax, StateTarget);
+		ActiveComponents.Add(GizmoActor->ForwardIntervalComponent);
+	}
+
+}
+
+void UIntervalGizmo::ClearSources()
+{
+	UpIntervalSource = nullptr;
+	DownIntervalSource = nullptr;
+	ForwardIntervalSource = nullptr;
+}
+
+void UIntervalGizmo::ClearActiveTarget()
+{
+	for (UInteractiveGizmo* Gizmo : ActiveGizmos)
+	{
+		GetGizmoManager()->DestroyGizmo(Gizmo);
+	}
+	ActiveGizmos.SetNum(0);
+	ActiveComponents.SetNum(0);
+	
+	ClearSources();
+
+	TransformProxy = nullptr;
+}
+
+
+UInteractiveGizmo* UIntervalGizmo::AddIntervalHandleGizmo(
+	USceneComponent* RootComponent,
+	UPrimitiveComponent* HandleComponent,
+	IGizmoAxisSource* AxisSource,
+	IGizmoFloatParameterSource* FloatParameterSource,
+	float MinParameter,
+	float MaxParameter,
+	IGizmoStateTarget* StateTargetIn)
+{
+	// create axis-position gizmo, axis-position parameter will drive translation
+	UAxisPositionGizmo* IntervalGizmo = Cast<UAxisPositionGizmo>(GetGizmoManager()->CreateGizmo(
+		UInteractiveGizmoManager::DefaultAxisPositionBuilderIdentifier));
+	check(IntervalGizmo);
+
+	// axis source provides the scale axis
+	IntervalGizmo->AxisSource = Cast<UObject>(AxisSource);
+
+
+	// parameter source maps axis-parameter-change to change in interval length
+	IntervalGizmo->ParameterSource = UGizmoAxisIntervalParameterSource::Construct(FloatParameterSource, MinParameter, MaxParameter, this);
+
+	// sub-component provides hit target
+	UGizmoComponentHitTarget* HitTarget = UGizmoComponentHitTarget::Construct(HandleComponent, this);
+	if (this->UpdateHoverFunction)
+	{
+		HitTarget->UpdateHoverFunction = [HandleComponent, this](bool bHovering) { this->UpdateHoverFunction(HandleComponent, bHovering); };
+	}
+	IntervalGizmo->HitTarget = HitTarget;
+
+	IntervalGizmo->StateTarget = Cast<UObject>(StateTargetIn);
+
+	ActiveGizmos.Add(IntervalGizmo);
+
+	return IntervalGizmo;
+}
+
+float UGizmoAxisIntervalParameterSource::GetParameter() const
+{
+	return FloatParameterSource->GetParameter();
+}
+
+void UGizmoAxisIntervalParameterSource::SetParameter(float NewValue) 
+{
+
+	NewValue = FMath::Clamp(NewValue, MinParameter, MaxParameter);
+
+	FloatParameterSource->SetParameter(NewValue);
+
+}
+
+void UGizmoAxisIntervalParameterSource::BeginModify()
+{
+	FloatParameterSource->BeginModify();
+}
+
+void UGizmoAxisIntervalParameterSource::EndModify()
+{
+	FloatParameterSource->EndModify();
+}
+
+UGizmoAxisIntervalParameterSource* UGizmoAxisIntervalParameterSource::Construct(
+	IGizmoFloatParameterSource* FloatSourceIn,
+	float ParameterMin,
+	float ParameterMax,
+	UObject* Outer)
+{
+	UGizmoAxisIntervalParameterSource* NewSource = NewObject<UGizmoAxisIntervalParameterSource>(Outer);
+
+	NewSource->FloatParameterSource = Cast<UObject>(FloatSourceIn);
+
+	// Clamp the initial value
+	float DefaultValue = NewSource->FloatParameterSource->GetParameter();
+	DefaultValue = FMath::Clamp(DefaultValue, ParameterMin, ParameterMax);
+	NewSource->FloatParameterSource->SetParameter(DefaultValue);
+
+	// record the min / max allowed
+	NewSource->MinParameter = ParameterMin;
+	NewSource->MaxParameter = ParameterMax;
+
+	return NewSource;
+}
\ No newline at end of file
diff --git a/Engine/Source/Runtime/Experimental/InteractiveToolsFramework/Private/InteractiveGizmoManager.cpp b/Engine/Source/Runtime/Experimental/InteractiveToolsFramework/Private/InteractiveGizmoManager.cpp
--- a/Engine/Source/Runtime/Experimental/InteractiveToolsFramework/Private/InteractiveGizmoManager.cpp
+++ b/Engine/Source/Runtime/Experimental/InteractiveToolsFramework/Private/InteractiveGizmoManager.cpp
@@ -8,6 +8,7 @@
 #include "BaseGizmos/PlanePositionGizmo.h"
 #include "BaseGizmos/AxisAngleGizmo.h"
 #include "BaseGizmos/TransformGizmo.h"
+#include "BaseGizmos/IntervalGizmo.h"
 
 #define LOCTEXT_NAMESPACE "UInteractiveGizmoManager"
 
@@ -290,6 +291,9 @@ void UInteractiveGizmoManager::RegisterDefaultGizmos()
 	CustomThreeAxisBuilder->GizmoActorBuilder = MakeShared<FTransformGizmoActorFactory>();
 	RegisterGizmoType(CustomThreeAxisTransformBuilderIdentifier, CustomThreeAxisBuilder);
 
+	UIntervalGizmoBuilder* IntervalGizmoBuilder = NewObject<UIntervalGizmoBuilder>();
+	RegisterGizmoType(UIntervalGizmo::GizmoName, IntervalGizmoBuilder);
+
 	bDefaultGizmosRegistered = true;
 }
 
diff --git a/Engine/Source/Runtime/Experimental/InteractiveToolsFramework/Public/BaseGizmos/GizmoActor.h b/Engine/Source/Runtime/Experimental/InteractiveToolsFramework/Public/BaseGizmos/GizmoActor.h
--- a/Engine/Source/Runtime/Experimental/InteractiveToolsFramework/Public/BaseGizmos/GizmoActor.h
+++ b/Engine/Source/Runtime/Experimental/InteractiveToolsFramework/Public/BaseGizmos/GizmoActor.h
@@ -9,6 +9,7 @@
 class UGizmoArrowComponent;
 class UGizmoRectangleComponent;
 class UGizmoCircleComponent;
+class UGizmoLineHandleComponent;
 
 
 /**
@@ -26,7 +27,7 @@ public:
 	/** Add standard arrow component to Actor, generally used for axis-translation */
 	static UGizmoArrowComponent* AddDefaultArrowComponent(
 		UWorld* World, AActor* Actor,
-		const FLinearColor& Color, const FVector& LocalDirection
+		const FLinearColor& Color, const FVector& LocalDirection, const float Length = 60.f
 	);
 	/** Add standard rectangle component to Actor, generally used for plane-translation */
 	static UGizmoRectangleComponent* AddDefaultRectangleComponent(
@@ -38,4 +39,11 @@ public:
 		UWorld* World, AActor* Actor,
 		const FLinearColor& Color, const FVector& PlaneNormal
 	);
+
+	/** Add standard disk component to Actor, generally used for handles */
+	static UGizmoLineHandleComponent* AddDefaultLineHandleComponent(
+		UWorld* World, AActor* Actor,
+		const FLinearColor& Color, const FVector& PlaneNormal, const FVector& LocalDirection,
+		const float Length = 60.f, const bool bImageScale = false
+	);
 };
diff --git a/Engine/Source/Runtime/Experimental/InteractiveToolsFramework/Public/BaseGizmos/GizmoLineHandleComponent.h b/Engine/Source/Runtime/Experimental/InteractiveToolsFramework/Public/BaseGizmos/GizmoLineHandleComponent.h
new file mode 100644
--- /dev/null
+++ b/Engine/Source/Runtime/Experimental/InteractiveToolsFramework/Public/BaseGizmos/GizmoLineHandleComponent.h
@@ -0,0 +1,47 @@
+// Copyright 1998-2019 Epic Games, Inc. All Rights Reserved.
+
+#pragma once
+
+#include "CoreMinimal.h"
+#include "UObject/ObjectMacros.h"
+#include "BaseGizmos/GizmoBaseComponent.h"
+#include "GizmoLineHandleComponent.generated.h"
+
+
+/**
+ * Simple Component intended to be used as part of 3D Gizmos.
+ * Draws line terminated by a short perpendicular handle line based on parameters.
+ */
+UCLASS(ClassGroup = Utility, HideCategories = (Physics, Collision, Lighting, Rendering, Mobile))
+class INTERACTIVETOOLSFRAMEWORK_API UGizmoLineHandleComponent : public UGizmoBaseComponent
+{
+	GENERATED_BODY()
+
+public:
+	UPROPERTY(EditAnywhere, Category = Options)
+	FVector Normal = FVector(0, 0, 1);
+
+	UPROPERTY(EditAnywhere, Category = Options)
+	float HandleSize = 10.0f;
+
+	UPROPERTY(EditAnywhere, Category = Options)
+	float Thickness = 2.0f;
+
+	UPROPERTY(EditAnywhere, Category = Options)
+	FVector Direction = FVector(0, 0, 1.);
+
+	UPROPERTY(EditAnywhere, Category = Options)
+	float Length;
+
+	UPROPERTY(EditAnywhere, Category = Options)
+	bool bImageScale = true;
+private:
+	//~ Begin UPrimitiveComponent Interface.
+	virtual FPrimitiveSceneProxy* CreateSceneProxy() override;
+	virtual bool LineTraceComponent(FHitResult& OutHit, const FVector Start, const FVector End, const FCollisionQueryParams& Params) override;
+	//~ End UPrimitiveComponent Interface.
+
+	//~ Begin USceneComponent Interface.
+	virtual FBoxSphereBounds CalcBounds(const FTransform& LocalToWorld) const override;
+	//~ Begin USceneComponent Interface.
+};
\ No newline at end of file
diff --git a/Engine/Source/Runtime/Experimental/InteractiveToolsFramework/Public/BaseGizmos/IntervalGizmo.h b/Engine/Source/Runtime/Experimental/InteractiveToolsFramework/Public/BaseGizmos/IntervalGizmo.h
new file mode 100644
--- /dev/null
+++ b/Engine/Source/Runtime/Experimental/InteractiveToolsFramework/Public/BaseGizmos/IntervalGizmo.h
@@ -0,0 +1,270 @@
+// Copyright 1998-2019 Epic Games, Inc. All Rights Reserved.
+
+#pragma once
+
+#include "CoreMinimal.h"
+#include "InteractiveGizmo.h"
+#include "InteractiveGizmoBuilder.h"
+#include "InteractiveToolObjects.h"
+#include "InteractiveToolChange.h"
+#include "BaseGizmos/GizmoActor.h"
+#include "BaseGizmos/TransformProxy.h"
+#include "ParameterSourcesFloat.h"
+#include "IntervalGizmo.generated.h"
+
+class UInteractiveGizmoManager;
+class IGizmoAxisSource;
+class IGizmoTransformSource;
+class IGizmoStateTarget;
+class UGizmoComponentAxisSource;
+class IGizmoFloatParameterSource;
+class UGizmoLocalFloatParameterSource;
+
+
+/**
+ *  AIntervalGizmoActor is an Actor type intended to be used with UIntervalGizmo,
+ * as the in-scene visual representation of the Gizmo.
+ *
+ * FIntervalGizmoActorFactory returns an instance of this Actor type (or a subclass).
+ *
+ * If a particular sub-Gizmo is not required, simply set that UProperty to null.
+
+ */
+UCLASS(Transient)
+class INTERACTIVETOOLSFRAMEWORK_API AIntervalGizmoActor : public AGizmoActor
+{
+	GENERATED_BODY()
+public:
+
+	AIntervalGizmoActor();
+	
+public:
+
+	UPROPERTY()
+	UGizmoLineHandleComponent* UpIntervalComponent = nullptr;
+
+	UPROPERTY()
+	UGizmoLineHandleComponent* DownIntervalComponent = nullptr;
+
+	UPROPERTY()
+	UGizmoLineHandleComponent* ForwardIntervalComponent = nullptr;
+
+
+public:
+	/**
+	 * Create a new instance of AIntervalGizmoActor and populate the various
+	 * sub-components with standard GizmoXComponent instances suitable for a 3-interval Gizmo
+	 */
+	static AIntervalGizmoActor* ConstructDefaultIntervalGizmo(UWorld* World);
+
+};
+
+
+/**
+ * FIntervalGizmoActorFactory creates new instances of AIntervalGizmoActor which
+ * are used by UIntervalGizmo to implement 3D interval Gizmos.
+ * An instance of FIntervalGizmoActorFactory is passed to UIntervalGizmo
+ * (by way of UIntervalGizmoBuilder), which then calls CreateNewGizmoActor()
+ * to spawn new Gizmo Actors.
+ *
+ * By default CreateNewGizmoActor() returns a default Gizmo Actor suitable for
+ * a three-axis Interval Gizmo, override this function to customize
+ * the Actor sub-elements.
+ */
+class INTERACTIVETOOLSFRAMEWORK_API FIntervalGizmoActorFactory
+{
+public:
+
+	/**
+	 * @param World the UWorld to create the new Actor in
+	 * @return new ATransformGizmoActor instance with members initialized with Components suitable for a transformation Gizmo
+	 */
+	virtual AIntervalGizmoActor* CreateNewGizmoActor(UWorld* World) const
+	{
+		return AIntervalGizmoActor::ConstructDefaultIntervalGizmo(World);
+	}
+};
+
+
+
+UCLASS()
+class INTERACTIVETOOLSFRAMEWORK_API UIntervalGizmoBuilder : public UInteractiveGizmoBuilder
+{
+	GENERATED_BODY()
+
+public:
+	/**
+	 * If set, this Actor Builder will be passed to UIntervalGizmo instances.
+	 * Otherwise new instances of the base FIntervalGizmoActorFactory are created internally.
+	 */
+	TSharedPtr<FIntervalGizmoActorFactory> GizmoActorBuilder;
+
+	/**
+	 * If set, this hover function will be passed to UIntervalGizmo instances to use instead of the default.
+	 * Hover is complicated for UIntervalGizmo because all it knows about the different gizmo scene elements
+	 * is that they are UPrimitiveComponent (coming from the AIntervalGizmoActor). The default hover
+	 * function implementation is to try casting to UGizmoBaseComponent and calling ::UpdateHoverState().
+	 * If you are using different Components that do not subclass UGizmoBaseComponent, and you want hover to
+	 * work, you will need to provide a different hover update function.
+	 */
+	TFunction<void(UPrimitiveComponent*, bool)> UpdateHoverFunction;
+
+	/**
+	 * If set, this coord-system function will be passed to UIntervalGizmo instances to use instead
+	 * of the default UpdateCoordSystemFunction. By default theUIntervalGizmo will query the external Context
+	 * to ask whether it should be using world or local coordinate system. Then the default UpdateCoordSystemFunction
+	 * will try casting to UGizmoBaseCmponent and passing that info on via UpdateWorldLocalState();
+	 * If you are using different Components that do not subclass UGizmoBaseComponent, and you want the coord system
+	 * to be configurable, you will need to provide a different update function.
+	 */
+	TFunction<void(UPrimitiveComponent*, EToolContextCoordinateSystem)> UpdateCoordSystemFunction;
+
+
+	virtual UInteractiveGizmo* BuildGizmo(const FToolBuilderState& SceneState) const override;
+};
+
+
+UCLASS()
+class INTERACTIVETOOLSFRAMEWORK_API UIntervalGizmo : public UInteractiveGizmo
+{
+	GENERATED_BODY()
+
+public:
+
+	static FString GizmoName;
+
+	virtual void SetWorld(UWorld* WorldIn);
+	virtual void SetGizmoActorBuilder(TSharedPtr<FIntervalGizmoActorFactory> Builder);
+	virtual void SetUpdateHoverFunction(TFunction<void(UPrimitiveComponent*, bool)> HoverFunction);
+	virtual void SetUpdateCoordSystemFunction(TFunction<void(UPrimitiveComponent*, EToolContextCoordinateSystem)> CoordSysFunction);
+
+	// UInteractiveGizmo overrides
+	virtual void Setup() override;
+	virtual void Shutdown() override;
+	virtual void Tick(float DeltaTime) override;
+
+
+	virtual void SetActiveTarget(UTransformProxy* TransformTargetIn, UGizmoLocalFloatParameterSource* UpInterval, UGizmoLocalFloatParameterSource* DownInterval, UGizmoLocalFloatParameterSource* ForwardInterval,
+		                         IToolContextTransactionProvider* TransactionProvider = nullptr);
+
+	/**
+	* Clear the parameter sources for this gizmo
+	*/
+	virtual void ClearSources();
+
+	/**
+	 * Clear the active target object for the Gizmo
+	 */
+	virtual void ClearActiveTarget();
+
+	/** State target is shared across gizmos, and created internally during SetActiveTarget() */
+	UPROPERTY()
+	UGizmoTransformChangeStateTarget* StateTarget = nullptr;
+
+
+protected:
+
+	/** GizmoActors will be spawned in this World */
+	UWorld* World;
+
+	AIntervalGizmoActor* GizmoActor = nullptr;
+
+	/** The active target object for the Gizmo */
+	UPROPERTY()
+	UTransformProxy* TransformProxy = nullptr;
+
+	/** list of current-active child components */
+	UPROPERTY()
+	TArray<UPrimitiveComponent*> ActiveComponents;
+
+	/** list of currently-active child gizmos */
+	UPROPERTY()
+	TArray<UInteractiveGizmo*> ActiveGizmos;
+
+	UGizmoLocalFloatParameterSource* UpIntervalSource = nullptr;
+	UGizmoLocalFloatParameterSource* DownIntervalSource = nullptr;
+	UGizmoLocalFloatParameterSource* ForwardIntervalSource = nullptr;
+
+	
+	/** Y-axis source is shared across Gizmos, and created internally during SetActiveTarget() */
+	UPROPERTY()
+	UGizmoComponentAxisSource* AxisYSource = nullptr;
+
+	/** Z-axis source is shared across Gizmos, and created internally during SetActiveTarget() */
+	UPROPERTY()
+	UGizmoComponentAxisSource* AxisZSource = nullptr;
+
+protected:
+	TSharedPtr<FIntervalGizmoActorFactory> GizmoActorBuilder;
+
+	// This function is called on each active GizmoActor Component to update it's hover state.
+	// If the Component is not a UGizmoBaseCmponent, the client needs to provide a different implementation
+	// of this function via the ToolBuilder
+	TFunction<void(UPrimitiveComponent*, bool)> UpdateHoverFunction;
+
+	// This function is called on each active GizmoActor Component to update it's coordinate system (eg world/local).
+	// If the Component is not a UGizmoBaseCmponent, the client needs to provide a different implementation
+	// of this function via the ToolBuilder
+	TFunction<void(UPrimitiveComponent*, EToolContextCoordinateSystem)> UpdateCoordSystemFunction;
+
+protected:
+
+	/** @return a new instance of the standard axis-handle Gizmo */
+	virtual UInteractiveGizmo* AddIntervalHandleGizmo(
+		USceneComponent* RootComponent,
+		UPrimitiveComponent* HandleComponent, 
+		IGizmoAxisSource* AxisSource,
+		IGizmoFloatParameterSource* FloatParameterSource,
+		float MinParameter,
+		float MaxParameter,
+		IGizmoStateTarget* StateTargetIn);
+
+
+};
+
+
+/**
+ * UGizmoAxisIntervalParameterSource is an IGizmoFloatParameterSource implementation that
+ * interprets the float value as the parameter of a line equation, and maps this parameter to a 3D translation
+ * along a line with origin/direction given by an IGizmoAxisSource. This translation is applied to an IGizmoTransformSource.
+ *
+ * This ParameterSource is intended to be used to create 3D Axis Interval Gizmos.
+ */
+UCLASS()
+class INTERACTIVETOOLSFRAMEWORK_API UGizmoAxisIntervalParameterSource : public UGizmoBaseFloatParameterSource
+{
+	GENERATED_BODY()
+public:
+	virtual float GetParameter() const override;
+
+	virtual void SetParameter(float NewValue) override;
+
+	virtual void BeginModify();
+	
+	virtual void EndModify();
+
+
+public:
+	
+	UPROPERTY()
+	TScriptInterface<IGizmoFloatParameterSource> FloatParameterSource;
+
+public:
+
+	UPROPERTY()
+	float MinParameter;
+
+	UPROPERTY()
+	float MaxParameter;
+
+public:
+
+	/**
+	 * Create a standard instance of this ParameterSource, with the given AxisSource and TransformSource
+	 */
+	static UGizmoAxisIntervalParameterSource* Construct(
+		IGizmoFloatParameterSource* FloatSourceIn,
+		float ParameterMin = -FLT_MAX,
+		float ParameterMax = FLT_MAX,
+		UObject* Outer = (UObject*)GetTransientPackage());
+};
\ No newline at end of file
