Geometry Processing - Interactive Tools.  Interval Gizmo.  adding missing undef LOCTEXT_NAMESPACE

diff --git a/Engine/Source/Runtime/Experimental/InteractiveToolsFramework/Private/BaseGizmos/IntervalGizmo.cpp b/Engine/Source/Runtime/Experimental/InteractiveToolsFramework/Private/BaseGizmos/IntervalGizmo.cpp
--- a/Engine/Source/Runtime/Experimental/InteractiveToolsFramework/Private/BaseGizmos/IntervalGizmo.cpp
+++ b/Engine/Source/Runtime/Experimental/InteractiveToolsFramework/Private/BaseGizmos/IntervalGizmo.cpp
@@ -429,4 +429,6 @@ UGizmoAxisIntervalParameterSource* UGizmoAxisIntervalParameterSource::Construct(
 	NewSource->MaxParameter = ParameterMax;
 
 	return NewSource;
-}
\ No newline at end of file
+}
+
+#undef LOCTEXT_NAMESPACE
\ No newline at end of file
