ToolsFramework: fix typo resulting in wrong direction being used for plane-translate gizmo
#rnx

diff --git a/Engine/Source/Runtime/Experimental/InteractiveToolsFramework/Private/BaseGizmos/GizmoRectangleComponent.cpp b/Engine/Source/Runtime/Experimental/InteractiveToolsFramework/Private/BaseGizmos/GizmoRectangleComponent.cpp
--- a/Engine/Source/Runtime/Experimental/InteractiveToolsFramework/Private/BaseGizmos/GizmoRectangleComponent.cpp
+++ b/Engine/Source/Runtime/Experimental/InteractiveToolsFramework/Private/BaseGizmos/GizmoRectangleComponent.cpp
@@ -220,7 +220,7 @@ bool UGizmoRectangleComponent::LineTraceComponent(FHitResult& OutHit, const FVec
 	FVector UseDirectionX = (bFlippedX) ? -DirectionX : DirectionX;
 	UseDirectionX = (bWorld) ? UseDirectionX : Transform.TransformVector(UseDirectionX);
 	FVector UseDirectionY = (bFlippedY) ? -DirectionY : DirectionY;
-	UseDirectionY = (bWorld) ? DirectionY : Transform.TransformVector(UseDirectionY);
+	UseDirectionY = (bWorld) ? UseDirectionY : Transform.TransformVector(UseDirectionY);
 	FVector UseOrigin = Transform.TransformPosition(FVector::ZeroVector);
 
 	float LengthScale = DynamicPixelToWorldScale;
