ToolsFramework: TransformGizmo can now snap to world grid during axis/plane translation. Moved snapping util funcs to GizmoMath. Axis and Plane TransformSources have optional Constraint functions which are used by TransformGizmo to apply calculated snaps. Snapping is controlled by boolean in TransformGizmo, currently Tools must explicitly enable/disable (ie no hotkey/etc at Gizmo level)
#rnx

diff --git a/Engine/Source/Runtime/Experimental/InteractiveToolsFramework/Private/BaseGizmos/AxisAngleGizmo.cpp b/Engine/Source/Runtime/Experimental/InteractiveToolsFramework/Private/BaseGizmos/AxisAngleGizmo.cpp
--- a/Engine/Source/Runtime/Experimental/InteractiveToolsFramework/Private/BaseGizmos/AxisAngleGizmo.cpp
+++ b/Engine/Source/Runtime/Experimental/InteractiveToolsFramework/Private/BaseGizmos/AxisAngleGizmo.cpp
@@ -106,38 +106,6 @@ void UAxisAngleGizmo::OnClickPress(const FInputDeviceRay& PressPos)
 }
 
 
-// these functions should move to a utility namespace
-
-static float SnapToIncrement(float fValue, float fIncrement, float offset = 0)
-{
-	if (!FMath::IsFinite(fValue))
-	{
-		return 0;
-	}
-	fValue -= offset;
-	float sign = FMath::Sign(fValue);
-	fValue = FMath::Abs(fValue);
-	int nInc = (int)(fValue / fIncrement);
-	float fRem = (float)fmod(fValue, fIncrement);
-	if (fRem > fIncrement / 2)
-	{
-		++nInc;
-	}
-	return sign * (float)nInc * fIncrement + offset;
-}
-
-
-static float SnapToNearbyIncrement(float fValue, float fIncrement, float fTolerance)
-{
-	float snapped = SnapToIncrement(fValue, fIncrement);
-	if (FMath::Abs(snapped - fValue) < fTolerance)
-	{
-		return snapped;
-	}
-	return fValue;
-}
-
-
 
 void UAxisAngleGizmo::OnClickDrag(const FInputDeviceRay& DragPos)
 {
@@ -162,7 +130,7 @@ void UAxisAngleGizmo::OnClickDrag(const FInputDeviceRay& DragPos)
 	float DeltaAngle = InteractionCurAngle - InteractionStartAngle;
 	if (bEnableSnapAngleModifier)
 	{
-		DeltaAngle = SnapToIncrement(DeltaAngle, 3.14159f / 4.0f);				// hardcoded 45-degree snap
+		DeltaAngle = GizmoMath::SnapToIncrement(DeltaAngle, 3.14159f / 4.0f);				// hardcoded 45-degree snap
 	}
 	float NewAngle = InitialTargetAngle + DeltaAngle;
 
diff --git a/Engine/Source/Runtime/Experimental/InteractiveToolsFramework/Private/BaseGizmos/GizmoMath.cpp b/Engine/Source/Runtime/Experimental/InteractiveToolsFramework/Private/BaseGizmos/GizmoMath.cpp
--- a/Engine/Source/Runtime/Experimental/InteractiveToolsFramework/Private/BaseGizmos/GizmoMath.cpp
+++ b/Engine/Source/Runtime/Experimental/InteractiveToolsFramework/Private/BaseGizmos/GizmoMath.cpp
@@ -206,3 +206,32 @@ FVector2D GizmoMath::ComputeCoordinatesInPlane(
 	float Y = FVector::DotProduct(LocalPoint, PlaneAxis2);
 	return FVector2D(X, Y);
 }
+
+
+FVector GizmoMath::ProjectPointOntoPlane(
+	const FVector& Point,
+	const FVector& PlaneOrigin, const FVector& PlaneNormal)
+{
+	FVector LocalPoint = Point - PlaneOrigin;
+	float NormalDot = FVector::DotProduct(LocalPoint, PlaneNormal);
+	return Point - NormalDot * PlaneNormal;
+}
+
+
+
+float GizmoMath::SnapToIncrement(float Value, float Increment)
+{
+	if (!FMath::IsFinite(Value))
+	{
+		return 0;
+	}
+	float Sign = FMath::Sign(Value);
+	Value = FMath::Abs(Value);
+	int IntIncrement = (int)(Value / Increment);
+	float Remainder = (float)fmod(Value, Increment);
+	if (Remainder > IntIncrement / 2)
+	{
+		++IntIncrement;
+	}
+	return Sign * (float)IntIncrement * Increment;
+}
\ No newline at end of file
diff --git a/Engine/Source/Runtime/Experimental/InteractiveToolsFramework/Private/BaseGizmos/TransformGizmo.cpp b/Engine/Source/Runtime/Experimental/InteractiveToolsFramework/Private/BaseGizmos/TransformGizmo.cpp
--- a/Engine/Source/Runtime/Experimental/InteractiveToolsFramework/Private/BaseGizmos/TransformGizmo.cpp
+++ b/Engine/Source/Runtime/Experimental/InteractiveToolsFramework/Private/BaseGizmos/TransformGizmo.cpp
@@ -355,7 +355,9 @@ UInteractiveGizmo* UTransformGizmo::AddAxisTranslationGizmo(
 	TranslateGizmo->AxisSource = Cast<UObject>(AxisSource);
 
 	// parameter source maps axis-parameter-change to translation of TransformSource's transform
-	TranslateGizmo->ParameterSource = UGizmoAxisTranslationParameterSource::Construct(AxisSource, TransformSource, this);
+	UGizmoAxisTranslationParameterSource* ParamSource = UGizmoAxisTranslationParameterSource::Construct(AxisSource, TransformSource, this);
+	ParamSource->PositionConstraintFunction = [this](const FVector& Pos, FVector& Snapped) { return PositionSnapFunction(Pos, Snapped); };
+	TranslateGizmo->ParameterSource = ParamSource;
 
 	// sub-component provides hit target
 	UGizmoComponentHitTarget* HitTarget = UGizmoComponentHitTarget::Construct(AxisComponent, this);
@@ -388,7 +390,9 @@ UInteractiveGizmo* UTransformGizmo::AddPlaneTranslationGizmo(
 	TranslateGizmo->AxisSource = Cast<UObject>(AxisSource);
 
 	// parameter source maps axis-parameter-change to translation of TransformSource's transform
-	TranslateGizmo->ParameterSource = UGizmoPlaneTranslationParameterSource::Construct(AxisSource, TransformSource, this);
+	UGizmoPlaneTranslationParameterSource* ParamSource = UGizmoPlaneTranslationParameterSource::Construct(AxisSource, TransformSource, this);
+	ParamSource->PositionConstraintFunction = [this](const FVector& Pos, FVector& Snapped) { return PositionSnapFunction(Pos, Snapped); };
+	TranslateGizmo->ParameterSource = ParamSource;
 
 	// sub-component provides hit target
 	UGizmoComponentHitTarget* HitTarget = UGizmoComponentHitTarget::Construct(AxisComponent, this);
@@ -463,4 +467,36 @@ void UTransformGizmo::ClearActiveTarget()
 
 
 
+
+bool UTransformGizmo::PositionSnapFunction(const FVector& WorldPosition, FVector& SnappedPositionOut)
+{
+	SnappedPositionOut = WorldPosition;
+
+	// only snap if we want snapping obvs
+	if (bSnapToWorldGrid == false)
+	{
+		return false;
+	}
+
+	// only snap to world grid when using world axes
+	if (GetGizmoManager()->GetContextQueriesAPI()->GetCurrentCoordinateSystem() != EToolContextCoordinateSystem::World)
+	{
+		return false;
+	}
+
+	FSceneSnapQueryRequest Request;
+	Request.RequestType = ESceneSnapQueryType::Position;
+	Request.TargetTypes = ESceneSnapQueryTargetType::Grid;
+	Request.Position = WorldPosition;
+	TArray<FSceneSnapQueryResult> Results;
+	if (GetGizmoManager()->GetContextQueriesAPI()->ExecuteSceneSnapQuery(Request, Results))
+	{
+		SnappedPositionOut = Results[0].Position;
+		return true;
+	};
+	
+	return false;
+}
+
+
 #undef LOCTEXT_NAMESPACE
\ No newline at end of file
diff --git a/Engine/Source/Runtime/Experimental/InteractiveToolsFramework/Public/BaseGizmos/GizmoMath.h b/Engine/Source/Runtime/Experimental/InteractiveToolsFramework/Public/BaseGizmos/GizmoMath.h
--- a/Engine/Source/Runtime/Experimental/InteractiveToolsFramework/Public/BaseGizmos/GizmoMath.h
+++ b/Engine/Source/Runtime/Experimental/InteractiveToolsFramework/Public/BaseGizmos/GizmoMath.h
@@ -21,7 +21,7 @@ namespace GizmoMath
 	INTERACTIVETOOLSFRAMEWORK_API FVector ProjectPointOntoLine(
 		const FVector& Point,
 		const FVector& LineOrigin, const FVector& LineDirection);
-	
+
 	/**
 	 * Find the nearest point to QueryPoint on the line defined by LineOrigin and LineDirection.
 	 * Returns this point and the associated line-equation parameter
@@ -100,4 +100,19 @@ namespace GizmoMath
 		const FVector& Point,
 		const FVector& PlaneOrigin, const FVector& PlaneNormal,
 		const FVector& PlaneAxis1, const FVector& PlaneAxis2);
+
+	/**
+	 * Project the input Point onto the plane defined by PlaneOrigin and PlaneNormal,
+	 * @return Nearest position to Point lying in the plane
+	 */
+	INTERACTIVETOOLSFRAMEWORK_API FVector ProjectPointOntoPlane(
+		const FVector& Point,
+		const FVector& PlaneOrigin, const FVector& PlaneNormal);
+
+
+	/**
+	 * Round Value to nearest step of Increment
+	 * @return snapped/rounded value
+	 */
+	INTERACTIVETOOLSFRAMEWORK_API float SnapToIncrement(float Value, float Increment);
 }
\ No newline at end of file
diff --git a/Engine/Source/Runtime/Experimental/InteractiveToolsFramework/Public/BaseGizmos/ParameterToTransformAdapters.h b/Engine/Source/Runtime/Experimental/InteractiveToolsFramework/Public/BaseGizmos/ParameterToTransformAdapters.h
--- a/Engine/Source/Runtime/Experimental/InteractiveToolsFramework/Public/BaseGizmos/ParameterToTransformAdapters.h
+++ b/Engine/Source/Runtime/Experimental/InteractiveToolsFramework/Public/BaseGizmos/ParameterToTransformAdapters.h
@@ -6,6 +6,7 @@
 #include "BaseGizmos/GizmoInterfaces.h"
 #include "BaseGizmos/ParameterSourcesFloat.h"
 #include "BaseGizmos/ParameterSourcesVec2.h"
+#include "BaseGizmos/GizmoMath.h"
 #include "ParameterToTransformAdapters.generated.h"
 
 
@@ -28,6 +29,15 @@ class INTERACTIVETOOLSFRAMEWORK_API UGizmoAxisTranslationParameterSource : publi
 {
 	GENERATED_BODY()
 public:
+
+	/**
+	 * Optional position constraint function. Called during interaction with the new transform origin.
+	 * To snap the transform to a new position, return as second value, and return true from your lambda.
+	 * Note that returned snap point will be projected onto the current translation origin/axis.
+	 * @return true if constraint point was found and should be used, false to ignore
+	 */
+	TUniqueFunction<bool(const FVector&, FVector&)> PositionConstraintFunction = [](const FVector&, FVector&) { return false; };
+
 	virtual float GetParameter() const override
 	{
 		return Parameter;
@@ -44,6 +54,15 @@ public:
 		// translate the initial transform
 		FTransform NewTransform = InitialTransform;
 		NewTransform.AddToTranslation(Translation);
+
+		// apply position constraint
+		FVector SnappedPos;
+		if (PositionConstraintFunction(NewTransform.GetTranslation(), SnappedPos))
+		{
+			FVector SnappedLinePos = GizmoMath::ProjectPointOntoLine(SnappedPos, CurTranslationOrigin, CurTranslationAxis);
+			NewTransform.SetTranslation(SnappedLinePos);
+		}
+
 		TransformSource->SetTransform(NewTransform);
 
 		OnParameterChanged.Broadcast(this, LastChange);
@@ -129,6 +148,16 @@ class INTERACTIVETOOLSFRAMEWORK_API UGizmoPlaneTranslationParameterSource : publ
 {
 	GENERATED_BODY()
 public:
+
+	/**
+	 * Optional position constraint function. Called during interaction with the new transform origin.
+	 * To snap the transform to a new position, return as second value, and return true from your lambda.
+	 * Note that returned snap point will be projected onto the current translation origin/axis.
+	 * @return true if constraint point was found and should be used, false to ignore
+	 */
+	TUniqueFunction<bool(const FVector&, FVector&)> PositionConstraintFunction = [](const FVector&, FVector&) { return false; };
+
+
 	virtual FVector2D GetParameter() const override
 	{
 		return Parameter;
@@ -146,6 +175,15 @@ public:
 		// apply translation to initial transform
 		FTransform NewTransform = InitialTransform;
 		NewTransform.AddToTranslation(Translation);
+
+		// apply position constraint
+		FVector SnappedPos;
+		if (PositionConstraintFunction(NewTransform.GetTranslation(), SnappedPos))
+		{
+			FVector PlanePos = GizmoMath::ProjectPointOntoPlane(SnappedPos, CurTranslationOrigin, CurTranslationNormal);
+			NewTransform.SetTranslation(PlanePos);
+		}
+
 		TransformSource->SetTransform(NewTransform);
 
 		OnParameterChanged.Broadcast(this, LastChange);
diff --git a/Engine/Source/Runtime/Experimental/InteractiveToolsFramework/Public/BaseGizmos/TransformGizmo.h b/Engine/Source/Runtime/Experimental/InteractiveToolsFramework/Public/BaseGizmos/TransformGizmo.h
--- a/Engine/Source/Runtime/Experimental/InteractiveToolsFramework/Public/BaseGizmos/TransformGizmo.h
+++ b/Engine/Source/Runtime/Experimental/InteractiveToolsFramework/Public/BaseGizmos/TransformGizmo.h
@@ -230,6 +230,12 @@ public:
 	 */
 	virtual void SetNewGizmoTransform(const FTransform& NewTransform);
 
+	/**
+	 * If true, then when using world frame, Axis and Plane translation snap to the world grid via the ContextQueriesAPI (in PositionSnapFunction)
+	 */
+	UPROPERTY()
+	bool bSnapToWorldGrid = false;
+
 
 protected:
 	TSharedPtr<FTransformGizmoActorFactory> GizmoActorBuilder;
@@ -294,5 +300,9 @@ protected:
 		IGizmoAxisSource* AxisSource,
 		IGizmoTransformSource* TransformSource,
 		IGizmoStateTarget* StateTarget);
+
+
+	// Axis and Plane TransformSources use this function to execute worldgrid snap queries
+	bool PositionSnapFunction(const FVector& WorldPosition, FVector& SnappedPositionOut);
 };
 
