Fix Landscape edit not working (Partial fix, not perfect)

diff --git a/Engine/Source/Runtime/Experimental/Chaos/Private/Chaos/HeightField.cpp b/Engine/Source/Runtime/Experimental/Chaos/Private/Chaos/HeightField.cpp
--- a/Engine/Source/Runtime/Experimental/Chaos/Private/Chaos/HeightField.cpp
+++ b/Engine/Source/Runtime/Experimental/Chaos/Private/Chaos/HeightField.cpp
@@ -468,6 +468,9 @@ namespace Chaos
 			};
 
 			EditGeomData<T, const uint16>(InHeights, InBeginRow, InBeginCol, InNumRows, InNumCols, MoveTemp(ConversionFunc), GeomData, LocalBounds);
+
+			// Slow and dumb. TODO: Actually fix CellHeights inside EditGeomData.
+			CalcBounds();
 		}
 	}
 
@@ -486,6 +489,9 @@ namespace Chaos
 			};
 
 			EditGeomData<T, T>(InHeights, InBeginRow, InBeginCol, InNumRows, InNumCols, MoveTemp(ConversionFunc), GeomData, LocalBounds);
+
+			// Slow and dumb. TODO: Actually fix CellHeights inside EditGeomData.
+			CalcBounds();
 		}
 	}
 
