Fix trimesh returning uint16 payload by accident

diff --git a/Engine/Source/Runtime/Experimental/Chaos/Public/Chaos/TriangleMeshImplicitObject.h b/Engine/Source/Runtime/Experimental/Chaos/Public/Chaos/TriangleMeshImplicitObject.h
--- a/Engine/Source/Runtime/Experimental/Chaos/Public/Chaos/TriangleMeshImplicitObject.h
+++ b/Engine/Source/Runtime/Experimental/Chaos/Public/Chaos/TriangleMeshImplicitObject.h
@@ -148,7 +148,7 @@ namespace Chaos
 			}
 
 			template <typename TPayloadType>
-			uint16 GetPayload(uint16 Idx) const
+			int32 GetPayload(int32 Idx) const
 			{
 				return Idx;
 			}
