Chaos : Collisions

- Updated collision constraint construction to buffer the transform of the transformed/scaled  implicit object, allowing for the update pipeline to be simplified.

diff --git a/Engine/Source/Runtime/Experimental/Chaos/Private/Chaos/CollisionResolution.cpp b/Engine/Source/Runtime/Experimental/Chaos/Private/Chaos/CollisionResolution.cpp
--- a/Engine/Source/Runtime/Experimental/Chaos/Private/Chaos/CollisionResolution.cpp
+++ b/Engine/Source/Runtime/Experimental/Chaos/Private/Chaos/CollisionResolution.cpp
@@ -96,10 +96,9 @@ namespace Chaos
 			const TBox<T, d> * Object1 = Implicit1->template GetObject<const TBox<T, d> >();
 			if (ensure(Object0 && Object1))
 			{
-				FRigidBodyPointContactConstraint* Constraint = new FRigidBodyPointContactConstraint;
-				Constraint->Particle[0] = Particle0;
-				Constraint->Particle[1] = Particle1;
-				Constraint->SetManifold(Implicit0, Implicit1);
+				TRigidTransform<T, d> ParticleImplicit0TM = Transform0.GetRelativeTransform(Collisions::GetTransform(Particle0));
+				TRigidTransform<T, d> ParticleImplicit1TM = Transform1.GetRelativeTransform(Collisions::GetTransform(Particle1));
+				FRigidBodyPointContactConstraint* Constraint = new FRigidBodyPointContactConstraint(Particle0, Implicit0, ParticleImplicit0TM, Particle1, Implicit1, ParticleImplicit1TM);
 
 				UpdateBoxBoxConstraint(Object0->BoundingBox(), Transform0, Object1->BoundingBox(), Transform1, Thickness, *Constraint);
 
@@ -189,10 +188,9 @@ namespace Chaos
 			const TPlane<T, d> * Object1 = Implicit1->template GetObject<const TPlane<T, d> >();
 			if (ensure(Object0 && Object1))
 			{
-				FRigidBodyPointContactConstraint* Constraint = new FRigidBodyPointContactConstraint;
-				Constraint->Particle[0] = Particle0;
-				Constraint->Particle[1] = Particle1;
-				Constraint->SetManifold(Implicit0, Implicit1);
+				TRigidTransform<T, d> ParticleImplicit0TM = Transform0.GetRelativeTransform(Collisions::GetTransform(Particle0));
+				TRigidTransform<T, d> ParticleImplicit1TM = Transform1.GetRelativeTransform(Collisions::GetTransform(Particle1));
+				FRigidBodyPointContactConstraint* Constraint = new FRigidBodyPointContactConstraint(Particle0, Implicit0, ParticleImplicit0TM, Particle1, Implicit1, ParticleImplicit1TM);
 
 				UpdateBoxPlaneConstraint(Object0->BoundingBox(), Transform0, *Object1, Transform1, Thickness, *Constraint);
 
@@ -243,10 +241,9 @@ namespace Chaos
 			const TSphere<T, d> * Object1 = Implicit1->template GetObject<const TSphere<T, d> >();
 			if (ensure(Object0 && Object1))
 			{
-				FRigidBodyPointContactConstraint* Constraint = new FRigidBodyPointContactConstraint;
-				Constraint->Particle[0] = Particle0;
-				Constraint->Particle[1] = Particle1;
-				Constraint->SetManifold(Implicit0, Implicit1);
+				TRigidTransform<T, d> ParticleImplicit0TM = Transform0.GetRelativeTransform(Collisions::GetTransform(Particle0));
+				TRigidTransform<T, d> ParticleImplicit1TM = Transform1.GetRelativeTransform(Collisions::GetTransform(Particle1));
+				FRigidBodyPointContactConstraint* Constraint = new FRigidBodyPointContactConstraint(Particle0, Implicit0, ParticleImplicit0TM, Particle1, Implicit1, ParticleImplicit1TM);
 
 				UpdateSphereSphereConstraint(*Object0, Transform0, *Object1, Transform1, Thickness, *Constraint);
 
@@ -293,11 +290,9 @@ namespace Chaos
 			const TPlane<T, d> * Object1 = Implicit1->template GetObject<const TPlane<T, d> >();
 			if (ensure(Object0 && Object1))
 			{
-
-				FRigidBodyPointContactConstraint* Constraint = new FRigidBodyPointContactConstraint;
-				Constraint->Particle[0] = Particle0;
-				Constraint->Particle[1] = Particle1;
-				Constraint->SetManifold(Implicit0, Implicit1);
+				TRigidTransform<T, d> ParticleImplicit0TM = Transform0.GetRelativeTransform(Collisions::GetTransform(Particle0));
+				TRigidTransform<T, d> ParticleImplicit1TM = Transform1.GetRelativeTransform(Collisions::GetTransform(Particle1));
+				FRigidBodyPointContactConstraint* Constraint = new FRigidBodyPointContactConstraint(Particle0, Implicit0, ParticleImplicit0TM, Particle1, Implicit1, ParticleImplicit1TM);
 
 				UpdateSpherePlaneConstraint(*Object0, Transform0, *Object1, Transform1, Thickness, *Constraint);
 
@@ -348,10 +343,9 @@ namespace Chaos
 			const TBox<T, d> * Object1 = Implicit1->template GetObject<const TBox<T, d> >();
 			if (ensure(Object0 && Object1))
 			{
-				FRigidBodyPointContactConstraint* Constraint = new FRigidBodyPointContactConstraint;
-				Constraint->Particle[0] = Particle0;
-				Constraint->Particle[1] = Particle1;
-				Constraint->SetManifold(Implicit0, Implicit1);
+				TRigidTransform<T, d> ParticleImplicit0TM = Transform0.GetRelativeTransform(Collisions::GetTransform(Particle0));
+				TRigidTransform<T, d> ParticleImplicit1TM = Transform1.GetRelativeTransform(Collisions::GetTransform(Particle1));
+				FRigidBodyPointContactConstraint* Constraint = new FRigidBodyPointContactConstraint(Particle0, Implicit0, ParticleImplicit0TM, Particle1, Implicit1, ParticleImplicit1TM);
 
 				UpdateSphereBoxConstraint(*Object0, Transform0, Object1->BoundingBox(), Transform1, Thickness, *Constraint);
 
@@ -408,10 +402,9 @@ namespace Chaos
 			const TCapsule<T>* Object1 = Implicit1->template GetObject<const TCapsule<T> >();
 			if (ensure(Object0 && Object1))
 			{
-				FRigidBodyPointContactConstraint* Constraint = new FRigidBodyPointContactConstraint;
-				Constraint->Particle[0] = Particle0;
-				Constraint->Particle[1] = Particle1;
-				Constraint->SetManifold(Implicit0, Implicit1);
+				TRigidTransform<T, d> ParticleImplicit0TM = Transform0.GetRelativeTransform(Collisions::GetTransform(Particle0));
+				TRigidTransform<T, d> ParticleImplicit1TM = Transform1.GetRelativeTransform(Collisions::GetTransform(Particle1));
+				FRigidBodyPointContactConstraint* Constraint = new FRigidBodyPointContactConstraint(Particle0, Implicit0, ParticleImplicit0TM, Particle1, Implicit1, ParticleImplicit1TM);
 
 				UpdateSphereCapsuleConstraint(*Object0, Transform0, *Object1, Transform1, Thickness, *Constraint);
 
@@ -471,10 +464,9 @@ namespace Chaos
 			const TCapsule<T> * Object1 = Implicit1->template GetObject<const TCapsule<T> >();
 			if (ensure(Object0 && Object1))
 			{
-				FRigidBodyPointContactConstraint* Constraint = new FRigidBodyPointContactConstraint;
-				Constraint->Particle[0] = Particle0;
-				Constraint->Particle[1] = Particle1;
-				Constraint->SetManifold(Implicit0, Implicit1);
+				TRigidTransform<T, d> ParticleImplicit0TM = Transform0.GetRelativeTransform(Collisions::GetTransform(Particle0));
+				TRigidTransform<T, d> ParticleImplicit1TM = Transform1.GetRelativeTransform(Collisions::GetTransform(Particle1));
+				FRigidBodyPointContactConstraint* Constraint = new FRigidBodyPointContactConstraint(Particle0, Implicit0, ParticleImplicit0TM, Particle1, Implicit1, ParticleImplicit1TM);
 
 				UpdateCapsuleCapsuleConstraint(*Object0, Transform0, *Object1, Transform1, Thickness, *Constraint);
 
@@ -513,10 +505,9 @@ namespace Chaos
 			const TBox<T, d> * Object1 = Implicit1->template GetObject<const TBox<T, d> >();
 			if (ensure(Object0 && Object1))
 			{
-				FRigidBodyPointContactConstraint* Constraint = new FRigidBodyPointContactConstraint;
-				Constraint->Particle[0] = Particle0;
-				Constraint->Particle[1] = Particle1;
-				Constraint->SetManifold(Implicit0, Implicit1);
+				TRigidTransform<T, d> ParticleImplicit0TM = Transform0.GetRelativeTransform(Collisions::GetTransform(Particle0));
+				TRigidTransform<T, d> ParticleImplicit1TM = Transform1.GetRelativeTransform(Collisions::GetTransform(Particle1));
+				FRigidBodyPointContactConstraint* Constraint = new FRigidBodyPointContactConstraint(Particle0, Implicit0, ParticleImplicit0TM, Particle1, Implicit1, ParticleImplicit1TM);
 
 				UpdateCapsuleBoxConstraint(*Object0, Transform0, Object1->BoundingBox(), Transform1, Thickness, *Constraint);
 
@@ -608,11 +599,11 @@ namespace Chaos
 
 					TArray<FVec3> CollisionSamples;
 					//
-					//  @todo(chaos) : Collision Manifold
+					//Â  @todo(chaos) : Collision Manifold
 					//   Remove the dependency on the virtual calls on the Implicit. Don't use FindClosestFaceAndVertices
-					//   this relies on virtual calls on the ImplicitObject. Instead pass a parameters structures into 
-					//   ConvexConvexContactPoint that can collect the face indices during evaluation of the support functions. 
-					//   This can be implemented without virtual calls. 
+					//   this relies on virtual calls on the ImplicitObject. Instead pass a parametersÂ structures into 
+					//   ConvexConvexContactPoint that can collectÂ the face indices during evaluation of the support functions. 
+					//   This can be implemented without virtual calls.Â 
 					//
 					int32 FaceIndex = Constraint.Manifold.Implicit[0]->FindClosestFaceAndVertices(Transform0.InverseTransformPosition(ContactPoint.Location), CollisionSamples, 1.f);
 
@@ -694,14 +685,11 @@ namespace Chaos
 		template<class T, int d>
 		void ConstructConvexConvexConstraints(TGeometryParticleHandle<T, d>* Particle0, TGeometryParticleHandle<T, d>* Particle1, const FImplicitObject* Implicit0, const FImplicitObject* Implicit1, const TRigidTransform<T, d>& Transform0, const TRigidTransform<T, d>& Transform1, const T Thickness, FCollisionConstraintsArray& NewConstraints)
 		{
-			FRigidBodyIterativeContactConstraint* Constraint = new FRigidBodyIterativeContactConstraint;
-
-			Constraint->Particle[0] = Particle0;
-			Constraint->Particle[1] = Particle1;
-			Constraint->SetManifold(Implicit0, Implicit1);
+			TRigidTransform<T, d> ParticleImplicit0TM  = Transform0.GetRelativeTransform(Collisions::GetTransform(Particle0));
+			TRigidTransform<T, d> ParticleImplicit1TM = Transform1.GetRelativeTransform(Collisions::GetTransform(Particle1));
+			FRigidBodyIterativeContactConstraint* Constraint = new FRigidBodyIterativeContactConstraint(Particle0, Implicit0, ParticleImplicit0TM, Particle1, Implicit1, ParticleImplicit1TM);
 
 			UpdateConvexConvexManifold(*Constraint, Transform0, Transform1, Thickness);
-
 			UpdateConvexConvexConstraint(*Implicit0, Transform0, *Implicit1, Transform1, Thickness, *Constraint);
 
 			if (Constraint->GetPhi() < Thickness)
@@ -753,7 +741,9 @@ namespace Chaos
 		template<typename T, int d>
 		void ConstructLevelsetLevelsetConstraints(TGeometryParticleHandle<T, d>* Particle0, TGeometryParticleHandle<T, d>* Particle1, const FImplicitObject* Implicit0, const FImplicitObject* Implicit1, const TRigidTransform<T, d>& Transform0, const TRigidTransform<T, d>& Transform1, const T Thickness, FCollisionConstraintsArray& NewConstraints)
 		{
-			FRigidBodyPointContactConstraint* Constraint = new FRigidBodyPointContactConstraint;
+			TRigidTransform<T, d> ParticleImplicit0TM = Transform0.GetRelativeTransform(Collisions::GetTransform(Particle0));
+			TRigidTransform<T, d> ParticleImplicit1TM = Transform1.GetRelativeTransform(Collisions::GetTransform(Particle1));
+			FRigidBodyPointContactConstraint* Constraint = new FRigidBodyPointContactConstraint(nullptr, nullptr, ParticleImplicit0TM, nullptr, nullptr, ParticleImplicit1TM);
 
 			bool bIsParticleDynamic0 = Particle0->CastToRigidParticle() && Particle0->ObjectState() == EObjectStateType::Dynamic;
 			if (!Particle1->Geometry() || (bIsParticleDynamic0 && !Particle0->CastToRigidParticle()->CollisionParticlesSize() && Particle0->Geometry() && !Particle0->Geometry()->IsUnderlyingUnion()))
@@ -781,165 +771,11 @@ namespace Chaos
 			}
 		}
 
-		//
-		// Levelset-Union
-		//
-
-		DECLARE_CYCLE_STAT(TEXT("TPBDCollisionConstraints::UpdateLevelsetUnionConstraint"), STAT_UpdateLevelsetUnionConstraint, STATGROUP_ChaosWide);
-		template<ECollisionUpdateType UpdateType, typename T, int d>
-		void UpdateLevelsetUnionConstraint(const T Thickness, TRigidBodyPointContactConstraint<T, d>& Constraint)
-		{
-			SCOPE_CYCLE_COUNTER(STAT_UpdateLevelsetUnionConstraint);
-
-			TGenericParticleHandle<T, d> Particle0 = Constraint.Particle[0];
-			TGenericParticleHandle<T, d> Particle1 = Constraint.Particle[1];
-
-			TRigidTransform<T, d> ParticlesTM = TRigidTransform<T, d>(Particle0->P(), Particle0->Q());
-			TRigidTransform<T, d> LevelsetTM = TRigidTransform<T, d>(Particle1->P(), Particle1->Q());
-
-			const FImplicitObject* ParticleObj = Particle0->Geometry().Get();
-			const FImplicitObject* LevelsetObj = Particle1->Geometry().Get();
-
-			if (!(ensure(!FMath::IsNaN(ParticlesTM.GetTranslation().X)) && ensure(!FMath::IsNaN(ParticlesTM.GetTranslation().Y)) && ensure(!FMath::IsNaN(ParticlesTM.GetTranslation().Z))))
-			{
-				return;
-			}
-
-			if (!(ensure(!FMath::IsNaN(LevelsetTM.GetTranslation().X)) && ensure(!FMath::IsNaN(LevelsetTM.GetTranslation().Y)) && ensure(!FMath::IsNaN(LevelsetTM.GetTranslation().Z))))
-			{
-				return;
-			}
-
-#if CHAOS_PARTICLEHANDLE_TODO
-			TArray<Pair<const FImplicitObject*, TRigidTransform<T, d>>> ParticleShapes = FindRelevantShapes(LevelsetObj, LevelsetTM, *ParticleObj, ParticlesTM, Thickness);
-			check(ParticleObj->IsUnderlyingUnion());
-			const TImplicitObjectUnion<T, d>* UnionObj = static_cast<const TImplicitObjectUnion<T, d>*>(ParticleObj);
-			for (const Pair<const FImplicitObject*, TRigidTransform<T, d>>& ParticleObjPair : ParticleShapes)
-			{
-				const FImplicitObject* Object = ParticleObjPair.First;
-
-				if (const int32* OriginalIdx = UnionObj->MCollisionParticleLookupHack.Find(Object))
-				{
-					const TBVHParticles<T, d>& SampleParticles = *InParticles.CollisionParticles(*OriginalIdx).Get();
-					const TRigidTransform<T, d> ObjectTM = ParticleObjPair.Second * ParticlesTM;
-
-					SampleObject2<UpdateType>(*LevelsetObj, LevelsetTM, SampleParticles, ObjectTM, Thickness, Constraint);
-					if (UpdateType == ECollisionUpdateType::Any && Constraint.Phi < Thickness)
-					{
-						return;
-					}
-				}
-			}
-#endif
-		}
-
-
-		DECLARE_CYCLE_STAT(TEXT("TPBDCollisionConstraints::UpdateUnionLevelsetConstraint"), STAT_UpdateUnionLevelsetConstraint, STATGROUP_ChaosWide);
-		template<ECollisionUpdateType UpdateType, typename T, int d>
-		void UpdateUnionLevelsetConstraint(const T Thickness, TRigidBodyPointContactConstraint<T, d>& Constraint)
-		{
-			SCOPE_CYCLE_COUNTER(STAT_UpdateUnionLevelsetConstraint);
-
-			TGenericParticleHandle<T, d> Particle0 = Constraint.Particle[0];
-			TGenericParticleHandle<T, d> Particle1 = Constraint.Particle[1];
-
-			TRigidTransform<T, d> ParticlesTM = TRigidTransform<T, d>(Particle0->P(), Particle0->Q());
-			TRigidTransform<T, d> LevelsetTM = TRigidTransform<T, d>(Particle1->P(), Particle1->Q());
-
-			if (!(ensure(!FMath::IsNaN(ParticlesTM.GetTranslation().X)) && ensure(!FMath::IsNaN(ParticlesTM.GetTranslation().Y)) && ensure(!FMath::IsNaN(ParticlesTM.GetTranslation().Z))))
-			{
-				return;
-			}
-
-			if (!(ensure(!FMath::IsNaN(LevelsetTM.GetTranslation().X)) && ensure(!FMath::IsNaN(LevelsetTM.GetTranslation().Y)) && ensure(!FMath::IsNaN(LevelsetTM.GetTranslation().Z))))
-			{
-				return;
-			}
-
-			const FImplicitObject* ParticleObj = Particle0->Geometry().Get();
-			const FImplicitObject* LevelsetObj = Particle1->Geometry().Get();
-			TArray<Pair<const FImplicitObject*, TRigidTransform<T, d>>> LevelsetShapes = FindRelevantShapes(ParticleObj, ParticlesTM, *LevelsetObj, LevelsetTM, Thickness);
-
-			if (LevelsetShapes.Num() && Particle0->CollisionParticles().Get())
-			{
-				const TBVHParticles<T, d>& SampleParticles = *Particle0->CollisionParticles().Get();
-				if (SampleParticles.Size())
-				{
-					for (const Pair<const FImplicitObject*, TRigidTransform<T, d>>& LevelsetObjPair : LevelsetShapes)
-					{
-						const FImplicitObject* Object = LevelsetObjPair.First;
-						const TRigidTransform<T, d> ObjectTM = LevelsetObjPair.Second * LevelsetTM;
-						// todo(brice) : SampleObject<UpdateType>(*Object, ObjectTM, SampleParticles, ParticlesTM, Thickness, Constraint);
-						if (UpdateType == ECollisionUpdateType::Any && Constraint.GetPhi() < Thickness)
-						{
-							return;
-						}
-					}
-				}
-#if CHAOS_PARTICLEHANDLE_TODO
-				else if (ParticleObj && ParticleObj->IsUnderlyingUnion())
-				{
-					const TImplicitObjectUnion<T, d>* UnionObj = static_cast<const TImplicitObjectUnion<T, d>*>(ParticleObj);
-					//need to traverse shapes to get their collision particles
-					for (const Pair<const FImplicitObject*, TRigidTransform<T, d>>& LevelsetObjPair : LevelsetShapes)
-					{
-						const FImplicitObject* LevelsetInnerObject = LevelsetObjPair.First;
-						const TRigidTransform<T, d> LevelsetInnerObjectTM = LevelsetObjPair.Second * LevelsetTM;
-
-						TArray<Pair<const FImplicitObject*, TRigidTransform<T, d>>> ParticleShapes = FindRelevantShapes(LevelsetInnerObject, LevelsetInnerObjectTM, *ParticleObj, ParticlesTM, Thickness);
-						for (const Pair<const FImplicitObject*, TRigidTransform<T, d>>& ParticleObjPair : ParticleShapes)
-						{
-							const FImplicitObject* ParticleInnerObject = ParticleObjPair.First;
-							const TRigidTransform<T, d> ParticleInnerObjectTM = ParticleObjPair.Second * ParticlesTM;
-
-							if (const int32* OriginalIdx = UnionObj->MCollisionParticleLookupHack.Find(ParticleInnerObject))
-							{
-								const TBVHParticles<T, d>& InnerSampleParticles = *InParticles.CollisionParticles(*OriginalIdx).Get();
-								SampleObject2<UpdateType>(*LevelsetInnerObject, LevelsetInnerObjectTM, InnerSampleParticles, ParticleInnerObjectTM, Thickness, Constraint);
-								if (UpdateType == ECollisionUpdateType::Any && Constraint.Phi < Thickness)
-								{
-									return;
-								}
-							}
-						}
-
-					}
-				}
-#endif
-			}
-		}
-
 		//
 		//  Union-Union
 		//
-		DECLARE_CYCLE_STAT(TEXT("TPBDCollisionConstraints::UpdateUnionUnionConstraint"), STAT_UpdateUnionUnionConstraint, STATGROUP_ChaosWide);
-
-		template<ECollisionUpdateType UpdateType, typename T, int d>
-		void UpdateUnionUnionConstraint(const FImplicitObject& Implicit0, const TRigidTransform<T, d>& Transform0, const FImplicitObject& Implicit1, const TRigidTransform<T, d>& Transform1, const T Thickness, TCollisionConstraintBase<T, d>& Constraint)
-		{
-			SCOPE_CYCLE_COUNTER(STAT_UpdateUnionUnionConstraint);
-
-			const TArray<Pair<const FImplicitObject*, TRigidTransform<T, d>>> LevelsetShapes = FindRelevantShapes(&Implicit0, Transform0, Implicit1, Transform1, Thickness);
-			for (const Pair<const FImplicitObject*, TRigidTransform<T, d>>& LevelsetObjPair : LevelsetShapes)
-			{
-				const FImplicitObject& LevelsetInnerObj = *LevelsetObjPair.First;
-				const TRigidTransform<T, d>& LevelsetInnerObjTM = LevelsetObjPair.Second * Transform1;
-
-				//now find all particle inner objects
-				const TArray<Pair<const FImplicitObject*, TRigidTransform<T, d>>> ParticleShapes = FindRelevantShapes(&LevelsetInnerObj, LevelsetInnerObjTM, Implicit0, Transform0, Thickness);
-
-				//for each inner obj pair, update constraint
-				for (const Pair<const FImplicitObject*, TRigidTransform<T, d>>& ParticlePair : ParticleShapes)
-				{
-					const FImplicitObject& ParticleInnerObj = *ParticlePair.First;
-					const TRigidTransform<T, d> ParticleInnerObjTM = ParticlePair.Second * Transform0;
-					UpdateConstraintImp<UpdateType>(ParticleInnerObj, ParticleInnerObjTM, LevelsetInnerObj, LevelsetInnerObjTM, Thickness, Constraint);
-				}
-			}
-		}
-
-
 
+		DECLARE_CYCLE_STAT(TEXT("TPBDCollisionConstraints::UpdateUnionUnionConstraint"), STAT_UpdateUnionUnionConstraint, STATGROUP_ChaosWide);
 		template<typename T, int d>
 		void ConstructUnionUnionConstraints(TGeometryParticleHandle<T, d>* Particle0, TGeometryParticleHandle<T, d>* Particle1, const FImplicitObject* Implicit0, const FImplicitObject* Implicit1, const TRigidTransform<T, d>& Transform0, const TRigidTransform<T, d>& Transform1, const T Thickness, FCollisionConstraintsArray& NewConstraints)
 		{
@@ -959,74 +795,50 @@ namespace Chaos
 				for (const Pair<const FImplicitObject*, TRigidTransform<T, d>>& ParticlePair : ParticleShapes)
 				{
 					const FImplicitObject* ParticleInnerObj = ParticlePair.First;
-					ConstructConstraintsImpl<T, d>(Particle0, Particle1, ParticleInnerObj, LevelsetInnerObj, Transform0, Transform1, Thickness, NewConstraints);
+					ConstructConstraints<T, d>(Particle0, Particle1, ParticleInnerObj, LevelsetInnerObj, Transform0, Transform1, Thickness, NewConstraints);
 				}
 			}
 		}
 
 		//
-		// Single-Union
+		// Constraint API
 		//
 
-		DECLARE_CYCLE_STAT(TEXT("TPBDCollisionConstraints::UpdateSingleUnionConstraint"), STAT_UpdateSingleUnionConstraint, STATGROUP_ChaosWide);
 		template<ECollisionUpdateType UpdateType, typename T, int d>
-		void UpdateSingleUnionConstraint(const FImplicitObject& Implicit0, const TRigidTransform<T, d>& Transform0, const FImplicitObject& Implicit1, const TRigidTransform<T, d>& Transform1, const T Thickness, TCollisionConstraintBase<T, d>& Constraint)
+		void UpdateConstraint(TCollisionConstraintBase<T, d>& ConstraintBase, const TRigidTransform<T, d>& ParticleTransform0, const TRigidTransform<T, d>& ParticleTransform1, const T Thickness)
 		{
-			SCOPE_CYCLE_COUNTER(STAT_UpdateSingleUnionConstraint);
+			const FImplicitObject& Implicit0 = *ConstraintBase.Manifold.Implicit[0];
+			const FImplicitObject& Implicit1 = *ConstraintBase.Manifold.Implicit[1];
 
-			const TArray<Pair<const FImplicitObject*, TRigidTransform<T, d>>> LevelsetShapes = FindRelevantShapes(&Implicit0, Transform0, Implicit1, Transform1, Thickness);
+			const TRigidTransform<T, d>& Transform0 = ConstraintBase.ImplicitTransform[0] * ParticleTransform0;
+			const TRigidTransform<T, d>& Transform1 = ConstraintBase.ImplicitTransform[1] * ParticleTransform1;
 
-			for (const Pair<const FImplicitObject*, TRigidTransform<T, d>>& LevelsetObjPair : LevelsetShapes)
+			//
+			// @todo(chaos): Collision Constraints (CollisionMap)
+			//    Modify Construct() and Update() use a CollisionMap indexed on EImplicitObjectType, instead of the if/else chain
+			//
+			EImplicitObjectType Implicit0Type = Implicit0.GetType();
+			EImplicitObjectType Implicit1Type = Implicit1.GetType();
+
+			if (Implicit0Type == TImplicitObjectTransformed<T, d>::StaticType())
 			{
-				const FImplicitObject& Implicit1InnerObj = *LevelsetObjPair.First;
-				const TRigidTransform<T, d> Implicit1InnerObjTM = LevelsetObjPair.Second * Transform1;
-				UpdateConstraintImp<UpdateType>(Implicit0, Transform0, Implicit1InnerObj, Implicit1InnerObjTM, Thickness, Constraint);
+				ensure(false);//should not be possible to get this type, it should already be resolved by the constraint. (see ConstructConstraints)
 			}
-		}
-
-
-
-		template<typename T, int d>
-		void ConstructSingleUnionConstraints(TGeometryParticleHandle<T, d>* Particle0, TGeometryParticleHandle<T, d>* Particle1, const FImplicitObject* Implicit0, const FImplicitObject* Implicit1, const TRigidTransform<T, d>& Transform0, const TRigidTransform<T, d>& Transform1, const T Thickness, FCollisionConstraintsArray& NewConstraints)
-		{
-			FRigidBodyPointContactConstraint* Constraint = new FRigidBodyPointContactConstraint;
-			Constraint->Particle[0] = Particle0;
-			Constraint->Particle[1] = Particle1;
-			Constraint->SetManifold(Implicit0, Implicit1);
-
-			UpdateSingleUnionConstraint<ECollisionUpdateType::Any>(*Implicit0, Transform0, *Implicit1, Transform1, Thickness, *Constraint);
-
-			if (Constraint->GetPhi() < Thickness)
+			else if (Implicit1Type == TImplicitObjectTransformed<T, d>::StaticType())
 			{
-				NewConstraints.Add(Constraint);
+				ensure(false);//should not be possible to get this type, it should already be resolved by the constraint. (see ConstructConstraints)
 			}
-			else
+			else if (Implicit0Type != TImplicitObjectUnion<T, d>::StaticType() && Implicit1Type == TImplicitObjectUnion<T, d>::StaticType())
 			{
-				delete Constraint;
+				ensure(false);//should not be possible to get this type, it should already be resolved by the constraint. (see ConstructConstraints)
 			}
-		}
-
-		//
-		// Constraint API
-		//
-
-		template<ECollisionUpdateType UpdateType, typename T, int d>
-		void UpdateConstraintImp(const FImplicitObject& Implicit0, const TRigidTransform<T, d>& Transform0, const FImplicitObject& Implicit1, const TRigidTransform<T, d>& Transform1, const T Thickness, TCollisionConstraintBase<T, d>& ConstraintBase)
-		{
-			EImplicitObjectType Implicit0Type = GetInnerType(Implicit0.GetType());
-			EImplicitObjectType Implicit1Type = GetInnerType(Implicit1.GetType());
-
-			if (Implicit0Type == TImplicitObjectTransformed<T, d>::StaticType())
+			else if (Implicit0Type == TImplicitObjectUnion<T, d>::StaticType() && Implicit1Type != TImplicitObjectUnion<T, d>::StaticType())
 			{
-				const TImplicitObjectTransformed<FReal, 3>* TransformedImplicit0 = Implicit0.template GetObject<const TImplicitObjectTransformed<FReal, 3>>();
-				TRigidTransform<T, d> TransformedTransform0 = TransformedImplicit0->GetTransform() * Transform0;
-				UpdateConstraintImp<UpdateType>(*TransformedImplicit0->GetTransformedObject(), TransformedTransform0, Implicit1, Transform1, Thickness, ConstraintBase);
+				ensure(false);//should not be possible to get this type, it should already be resolved by the constraint. (see ConstructConstraints)
 			}
-			else if (Implicit1Type == TImplicitObjectTransformed<T, d>::StaticType())
+			else if (Implicit0Type == TImplicitObjectUnion<T, d>::StaticType() && Implicit1Type == TImplicitObjectUnion<T, d>::StaticType())
 			{
-				const TImplicitObjectTransformed<FReal, 3>* TransformedImplicit1 = Implicit1.template GetObject<const TImplicitObjectTransformed<FReal, 3>>();
-				TRigidTransform<T, d> TransformedTransform1 = TransformedImplicit1->GetTransform() * Transform1;
-				UpdateConstraintImp<UpdateType>(Implicit0, Transform0, *TransformedImplicit1->GetTransformedObject(), TransformedTransform1, Thickness, ConstraintBase);
+				ensure(false);//should not be possible to get this type, it should already be resolved by the constraint. (see ConstructConstraints)
 			}
 			else if (Implicit0Type == TBox<T, d>::StaticType() && Implicit1Type == TBox<T, d>::StaticType())
 			{
@@ -1115,30 +927,10 @@ namespace Chaos
 					Constraint.SetNormal(-Constraint.GetNormal());
 				}
 			}
-			else if (Implicit0Type != TImplicitObjectUnion<T, d>::StaticType() && Implicit1Type == TImplicitObjectUnion<T, d>::StaticType())
-			{
-				return UpdateSingleUnionConstraint<UpdateType>(Implicit0, Transform0, Implicit1, Transform1, Thickness, ConstraintBase);
-			}
-			else if (Implicit0Type == TImplicitObjectUnion<T, d>::StaticType() && Implicit1Type != TImplicitObjectUnion<T, d>::StaticType())
-			{
-				check(false);	//should not be possible to get this ordering (see [Note] in ConstructPairConstraintImpl)
-			}
-			else if (Implicit0Type == TImplicitObjectUnion<T, d>::StaticType() && Implicit1Type == TImplicitObjectUnion<T, d>::StaticType())
-			{
-				return UpdateUnionUnionConstraint<UpdateType>(Implicit0, Transform0, Implicit1, Transform1, Thickness, ConstraintBase);
-			}
 			else if(Implicit0.IsConvex() && Implicit1.IsConvex())
 			{
 				UpdateConvexConvexConstraint(Implicit0, Transform0, Implicit1, Transform1, Thickness, ConstraintBase);
 			}
-			else if (Implicit1.IsUnderlyingUnion())
-			{
-				UpdateUnionLevelsetConstraint<UpdateType>(Thickness, *ConstraintBase.template As<TRigidBodyPointContactConstraint<T,d> >());
-			}
-			else if (Implicit0.IsUnderlyingUnion())
-			{
-				UpdateLevelsetUnionConstraint<UpdateType>(Thickness, *ConstraintBase.template As<TRigidBodyPointContactConstraint<T,d> >());
-			}
 			else
 			{
 				UpdateLevelsetLevelsetConstraint<UpdateType>(Thickness, *ConstraintBase.template As<TRigidBodyPointContactConstraint<T,d> >());
@@ -1146,12 +938,15 @@ namespace Chaos
 		}
 
 		template<typename T, int d>
-		void ConstructPairConstraintImpl(TGeometryParticleHandle<T, d>* Particle0, TGeometryParticleHandle<T, d>* Particle1, const FImplicitObject* Implicit0, const FImplicitObject* Implicit1, const TRigidTransform<T, d>& Transform0, const TRigidTransform<T, d>& Transform1, const T Thickness, FCollisionConstraintsArray& NewConstraints)
+		void ConstructConstraints(TGeometryParticleHandle<T, d>* Particle0, TGeometryParticleHandle<T, d>* Particle1, const FImplicitObject* Implicit0, const FImplicitObject* Implicit1, const TRigidTransform<T, d>& Transform0, const TRigidTransform<T, d>& Transform1, const T Thickness, FCollisionConstraintsArray& NewConstraints)
 		{
-			const TPerShapeData<T, d>* Shape0 = Particle0->GetImplicitShape(Implicit0);
-			const TPerShapeData<T, d>* Shape1 = Particle1->GetImplicitShape(Implicit1);
+			//  TriMesh Culling
+			if (Implicit0 && Implicit0->GetType() == ImplicitObjectType::TriangleMesh) return;
+			if (Implicit1 && Implicit1->GetType() == ImplicitObjectType::TriangleMesh) return;
 
 			// If either shape is disabled for collision bail without constructing a constraint
+			const TPerShapeData<T, d>* Shape0 = Particle0->GetImplicitShape(Implicit0);
+			const TPerShapeData<T, d>* Shape1 = Particle1->GetImplicitShape(Implicit1);
 			if((Shape0 && Shape0->bDisable) || (Shape1 && Shape1->bDisable))
 			{
 				return;
@@ -1163,19 +958,49 @@ namespace Chaos
 				return;
 			}
 			
-			EImplicitObjectType Implicit0Type = GetInnerType(Implicit0->GetType());
-			EImplicitObjectType Implicit1Type = GetInnerType(Implicit1->GetType());
+			//
+			// @todo(chaos): Collision Constraints (CollisionMap)
+			//    Modify Construct() and Update() use a CollisionMap indexed on EImplicitObjectType, instead of the if/else chain
+			//
+			EImplicitObjectType Implicit0Type = Implicit0->GetType();
+			EImplicitObjectType Implicit1Type = Implicit1->GetType();
+
 			if (Implicit0Type == TImplicitObjectTransformed<T, d>::StaticType())
 			{
 				const TImplicitObjectTransformed<FReal, 3>* TransformedImplicit0 = Implicit0->template GetObject<const TImplicitObjectTransformed<FReal, 3>>();
 				TRigidTransform<T, d> TransformedTransform0 = TransformedImplicit0->GetTransform() * Transform0;
-				ConstructPairConstraintImpl(Particle0, Particle1, TransformedImplicit0->GetTransformedObject(), Implicit1, TransformedTransform0, Transform1, Thickness, NewConstraints);
+				ConstructConstraints(Particle0, Particle1, TransformedImplicit0->GetTransformedObject(), Implicit1, TransformedTransform0, Transform1, Thickness, NewConstraints);
 			}
 			else if (Implicit1Type == TImplicitObjectTransformed<T, d>::StaticType())
 			{
 				const TImplicitObjectTransformed<FReal, 3>* TransformedImplicit1 = Implicit1->template GetObject<const TImplicitObjectTransformed<FReal, 3>>();
 				TRigidTransform<T, d> TransformedTransform1 = TransformedImplicit1->GetTransform() * Transform1;
-				ConstructPairConstraintImpl(Particle0, Particle1, Implicit0, TransformedImplicit1->GetTransformedObject(), Transform0, TransformedTransform1, Thickness, NewConstraints);
+				ConstructConstraints(Particle0, Particle1, Implicit0, TransformedImplicit1->GetTransformedObject(), Transform0, TransformedTransform1, Thickness, NewConstraints);
+			}
+			else if (Implicit0Type != TImplicitObjectUnion<T, d>::StaticType() && Implicit1Type == TImplicitObjectUnion<T, d>::StaticType())
+			{
+				const TArray<Pair<const FImplicitObject*, TRigidTransform<T, d>>> LevelsetShapes = FindRelevantShapes(Implicit0, Transform0, *Implicit1, Transform1, Thickness);
+				for (const Pair<const FImplicitObject*, TRigidTransform<T, d>>& LevelsetObjPair : LevelsetShapes)
+				{
+					const FImplicitObject* Implicit1InnerObj = LevelsetObjPair.First;
+					const TRigidTransform<T, d> Implicit1InnerObjTM = LevelsetObjPair.Second * Transform1;
+					ConstructConstraints(Particle0, Particle1, Implicit0, Implicit1InnerObj, Transform0, Implicit1InnerObjTM, Thickness, NewConstraints);
+				}
+			}
+			else if (Implicit0Type == TImplicitObjectUnion<T, d>::StaticType() && Implicit1Type != TImplicitObjectUnion<T, d>::StaticType())
+			{
+				// [Note] forces non-unions into particle[0] position
+				const TArray<Pair<const FImplicitObject*, TRigidTransform<T, d>>> LevelsetShapes = FindRelevantShapes(Implicit1, Transform1, *Implicit0, Transform0, Thickness);
+				for (const Pair<const FImplicitObject*, TRigidTransform<T, d>>& LevelsetObjPair : LevelsetShapes)
+				{
+					const FImplicitObject* Implicit0InnerObj = LevelsetObjPair.First;
+					const TRigidTransform<T, d> Implicit0InnerObjTM = LevelsetObjPair.Second * Transform0;
+					ConstructConstraints(Particle1, Particle0, Implicit1, Implicit0InnerObj, Transform1, Implicit0InnerObjTM, Thickness, NewConstraints);
+				}
+			}
+			else if (Implicit0Type == TImplicitObjectUnion<T, d>::StaticType() && Implicit1Type == TImplicitObjectUnion<T, d>::StaticType())
+			{
+				ConstructUnionUnionConstraints(Particle0, Particle1, Implicit0, Implicit1, Transform0, Transform1, Thickness, NewConstraints);
 			}
 			else if (Implicit0Type == TBox<T, d>::StaticType() && Implicit1Type == TBox<T, d>::StaticType())
 			{
@@ -1229,20 +1054,6 @@ namespace Chaos
 			{
 				ConstructCapsuleBoxConstraints(Particle1, Particle0, Implicit1, Implicit0, Transform1, Transform0, Thickness, NewConstraints);
 			}
-			else if (Implicit0Type != TImplicitObjectUnion<T, d>::StaticType() && Implicit1Type == TImplicitObjectUnion<T, d>::StaticType())
-			{
-				ConstructSingleUnionConstraints(Particle0, Particle1, Implicit0, Implicit1, Transform0, Transform1, Thickness, NewConstraints);
-			}
-			else if (Implicit0Type == TImplicitObjectUnion<T, d>::StaticType() && Implicit1Type != TImplicitObjectUnion<T, d>::StaticType())
-			{
-				// [Note] forces non-unions into particle[0] position
-				ConstructSingleUnionConstraints(Particle1, Particle0, Implicit1, Implicit0, Transform1, Transform0, Thickness, NewConstraints);
-			}
-			else if (Implicit0Type == TImplicitObjectUnion<T, d>::StaticType() && Implicit1Type == TImplicitObjectUnion<T, d>::StaticType())
-			{
-				// Union-union creates multiple manifolds - we should never get here for this pair type. See ConstructConstraintsImpl and ConstructUnionUnionConstraints
-				ensure(false);
-			}
 			else if (Implicit0->IsConvex() && Implicit1->IsConvex())
 			{
 				ConstructConvexConvexConstraints(Particle0, Particle1, Implicit0, Implicit1, Transform0, Transform1, Thickness, NewConstraints);
@@ -1253,25 +1064,6 @@ namespace Chaos
 			}
 		}
 
-		template<typename T, int d>
-		void ConstructConstraintsImpl(TGeometryParticleHandle<T, d>* Particle0, TGeometryParticleHandle<T, d>* Particle1, const FImplicitObject* Implicit0, const FImplicitObject* Implicit1, const TRigidTransform<T, d>& Transform0, const TRigidTransform<T, d>& Transform1, const T Thickness, FCollisionConstraintsArray& NewConstraints)
-		{
-			// TriangleMesh implicits are for scene query only.
-			if (Implicit0 && GetInnerType(Implicit0->GetType()) == ImplicitObjectType::TriangleMesh) return;
-			if (Implicit1 && GetInnerType(Implicit1->GetType()) == ImplicitObjectType::TriangleMesh) return;
-
-			if (Implicit0 && Implicit0->GetType() == TImplicitObjectUnion<T, d>::StaticType() &&
-				Implicit1 && Implicit1->GetType() == TImplicitObjectUnion<T, d>::StaticType())
-			{
-				return ConstructUnionUnionConstraints(Particle0, Particle1, Implicit0, Implicit1, Transform0, Transform1, Thickness, NewConstraints);
-			}
-			else
-			{
-				return ConstructPairConstraintImpl(Particle0, Particle1, Implicit0, Implicit1, Transform0, Transform1, Thickness, NewConstraints);
-			}
-		}
-
-
 
 		template void UpdateBoxBoxConstraint<float, 3>(const TAABB<float, 3>& Box1, const TRigidTransform<float, 3>& Box1Transform, const TAABB<float, 3>& Box2, const TRigidTransform<float, 3>& Box2Transform, const float Thickness, TRigidBodyPointContactConstraint<float, 3>& Constraint);
 		template void ConstructBoxBoxConstraints<float, 3>(TGeometryParticleHandle<float, 3>* Particle0, TGeometryParticleHandle<float, 3>* Particle1, const FImplicitObject* Implicit0, const FImplicitObject* Implicit1, const TRigidTransform<float, 3>& Transform0, const TRigidTransform<float, 3>& Transform1, const float Thickness, FCollisionConstraintsArray& NewConstraints);
@@ -1305,25 +1097,13 @@ namespace Chaos
 		template void UpdateLevelsetLevelsetConstraint<ECollisionUpdateType::Deepest, float, 3>(const float Thickness, TRigidBodyPointContactConstraint<float, 3>& Constraint);
 		template void ConstructLevelsetLevelsetConstraints<float, 3>(TGeometryParticleHandle<float, 3>* Particle0, TGeometryParticleHandle<float, 3>* Particle1, const FImplicitObject* Implicit0, const FImplicitObject* Implicit1, const TRigidTransform<float, 3>& Transform0, const TRigidTransform<float, 3>& Transform1, const float Thickness, FCollisionConstraintsArray& NewConstraints);
 
-		template void UpdateLevelsetUnionConstraint<ECollisionUpdateType::Any, float, 3>(const float Thickness, TRigidBodyPointContactConstraint<float, 3>& Constraint);
-		template void UpdateLevelsetUnionConstraint<ECollisionUpdateType::Deepest, float, 3>(const float Thickness, TRigidBodyPointContactConstraint<float, 3>& Constraint);
-
-		template void UpdateUnionLevelsetConstraint<ECollisionUpdateType::Any, float, 3>(const float Thickness, TRigidBodyPointContactConstraint<float, 3>& Constraint);
-		template void UpdateUnionLevelsetConstraint<ECollisionUpdateType::Deepest, float, 3>(const float Thickness, TRigidBodyPointContactConstraint<float, 3>& Constraint);
-
 		template void UpdateUnionUnionConstraint< ECollisionUpdateType::Any, float, 3>(const FImplicitObject& Implicit0, const TRigidTransform<float, 3>& Transform0, const FImplicitObject& Implicit1, const TRigidTransform<float, 3>& Transform1, const float Thickness, TCollisionConstraintBase<float, 3>& Constraint);
 		template void UpdateUnionUnionConstraint< ECollisionUpdateType::Deepest, float, 3>(const FImplicitObject& Implicit0, const TRigidTransform<float, 3>& Transform0, const FImplicitObject& Implicit1, const TRigidTransform<float, 3>& Transform1, const float Thickness, TCollisionConstraintBase<float, 3>& Constraint);
 		template void ConstructUnionUnionConstraints<float, 3>(TGeometryParticleHandle<float, 3>* Particle0, TGeometryParticleHandle<float, 3>* Particle1, const FImplicitObject* Implicit0, const FImplicitObject* Implicit1, const TRigidTransform<float, 3>& Transform0, const TRigidTransform<float, 3>& Transform1, const float Thickness, FCollisionConstraintsArray& NewConstraints);
 
-		template void UpdateSingleUnionConstraint< ECollisionUpdateType::Any, float, 3>(const FImplicitObject& Implicit0, const TRigidTransform<float, 3>& Transform0, const FImplicitObject& Implicit1, const TRigidTransform<float, 3>& Transform1, const float Thickness, TCollisionConstraintBase<float, 3>& Constraint);
-		template void UpdateSingleUnionConstraint< ECollisionUpdateType::Deepest, float, 3>(const FImplicitObject& Implicit0, const TRigidTransform<float, 3>& Transform0, const FImplicitObject& Implicit1, const TRigidTransform<float, 3>& Transform1, const float Thickness, TCollisionConstraintBase<float, 3>& Constraint);
-		template void ConstructSingleUnionConstraints<float, 3>(TGeometryParticleHandle<float, 3>* Particle0, TGeometryParticleHandle<float, 3>* Particle1, const FImplicitObject* Implicit0, const FImplicitObject* Implicit1, const TRigidTransform<float, 3>& Transform0, const TRigidTransform<float, 3>& Transform1, const float Thickness, FCollisionConstraintsArray& NewConstraints);
-
-		template void UpdateConstraintImp<ECollisionUpdateType::Any, float, 3>(const FImplicitObject& Implicit0, const TRigidTransform<float, 3>& Transform0, const FImplicitObject& Implicit1, const TRigidTransform<float, 3>& Transform1, const float Thickness, TCollisionConstraintBase<float, 3>& Constraints);
-		template void UpdateConstraintImp<ECollisionUpdateType::Deepest, float, 3>(const FImplicitObject& Implicit0, const TRigidTransform<float, 3>& Transform0, const FImplicitObject& Implicit1, const TRigidTransform<float, 3>& Transform1, const float Thickness, TCollisionConstraintBase<float, 3>& Constraint);
-
-		template void ConstructPairConstraintImpl<float, 3>(TGeometryParticleHandle<float, 3>* Particle0, TGeometryParticleHandle<float, 3>* Particle1, const FImplicitObject* Implicit0, const FImplicitObject* Implicit1, const TRigidTransform<float, 3>& Transform0, const TRigidTransform<float, 3>& Transform1, const float Thickness, FCollisionConstraintsArray& NewConstraints);
-		template void ConstructConstraintsImpl<float, 3>(TGeometryParticleHandle<float, 3>* Particle0, TGeometryParticleHandle<float, 3>* Particle1, const FImplicitObject* Implicit0, const FImplicitObject* Implicit1, const TRigidTransform<float, 3>& Transform0, const TRigidTransform<float, 3>& Transform1, const float Thickness, FCollisionConstraintsArray& NewConstraints);
+		template void UpdateConstraint<ECollisionUpdateType::Any, float, 3>(TCollisionConstraintBase<float,3>& ConstraintBase, const TRigidTransform<float,3>& Transform0, const TRigidTransform<float,3>& Transform1, const float Thickness);
+		template void UpdateConstraint<ECollisionUpdateType::Deepest, float, 3>(TCollisionConstraintBase<float, 3>& ConstraintBase, const TRigidTransform<float, 3>& Transform0, const TRigidTransform<float, 3>& Transform1, const float Thickness);
+		template void ConstructConstraints<float, 3>(TGeometryParticleHandle<float, 3>* Particle0, TGeometryParticleHandle<float, 3>* Particle1, const FImplicitObject* Implicit0, const FImplicitObject* Implicit1, const TRigidTransform<float, 3>& Transform0, const TRigidTransform<float, 3>& Transform1, const float Thickness, FCollisionConstraintsArray& NewConstraints);
 
 	} // Collisions
 
diff --git a/Engine/Source/Runtime/Experimental/Chaos/Private/Chaos/PBDCollisionConstraintsContact.cpp b/Engine/Source/Runtime/Experimental/Chaos/Private/Chaos/PBDCollisionConstraintsContact.cpp
--- a/Engine/Source/Runtime/Experimental/Chaos/Private/Chaos/PBDCollisionConstraintsContact.cpp
+++ b/Engine/Source/Runtime/Experimental/Chaos/Private/Chaos/PBDCollisionConstraintsContact.cpp
@@ -18,31 +18,7 @@ namespace Chaos
 			const TRigidTransform<T, d> ParticleTM = GetTransform(Constraint.Particle[0]);
 			const TRigidTransform<T, d> LevelsetTM = GetTransform(Constraint.Particle[1]);
 
-			if (Constraint.GetType() == TCollisionConstraintBase<T, d>::FType::SinglePoint)
-			{
-				if (!Constraint.Particle[0]->Geometry())
-				{
-					if (Constraint.Particle[1]->Geometry())
-					{
-						if (!Constraint.Particle[1]->Geometry()->IsUnderlyingUnion())
-						{
-							UpdateLevelsetLevelsetConstraint<UpdateType>(Thickness, *Constraint.template As<TRigidBodyPointContactConstraint<T, d>>());
-						}
-						else
-						{
-							UpdateUnionLevelsetConstraint<UpdateType>(Thickness, *Constraint.template As<TRigidBodyPointContactConstraint<T, d>>());
-						}
-					}
-				}
-				else
-				{
-					UpdateConstraintImp<UpdateType>(*Constraint.Particle[0]->Geometry(), ParticleTM, *Constraint.Particle[1]->Geometry(), LevelsetTM, Thickness, Constraint);
-				}
-			}
-			else if(Constraint.GetType() == TCollisionConstraintBase<T, d>::FType::MultiPoint)
-			{
-				UpdateConstraintImp<UpdateType>(*Constraint.Particle[0]->Geometry(), ParticleTM, *Constraint.Particle[1]->Geometry(), LevelsetTM, Thickness, Constraint);
-			}
+			UpdateConstraint<UpdateType>(Constraint, ParticleTM, LevelsetTM, Thickness);
 		}
 
 		template<typename T, int d>
diff --git a/Engine/Source/Runtime/Experimental/Chaos/Public/Chaos/Collision/NarrowPhase.h b/Engine/Source/Runtime/Experimental/Chaos/Public/Chaos/Collision/NarrowPhase.h
--- a/Engine/Source/Runtime/Experimental/Chaos/Public/Chaos/Collision/NarrowPhase.h
+++ b/Engine/Source/Runtime/Experimental/Chaos/Public/Chaos/Collision/NarrowPhase.h
@@ -39,7 +39,14 @@ namespace Chaos
 		{
 			if (ensure(Particle0 && Particle1))
 			{
-				Collisions::ConstructConstraintsImpl<FReal, 3>(Particle0, Particle1, Particle0->Geometry().Get(), Particle1->Geometry().Get(), Collisions::GetTransform(Particle0), Collisions::GetTransform(Particle1), Thickness, NewConstraints);
+				//
+				// @todo(chaos) : Collision Constraints
+				//   This is not efficient. The constraint has to go through a construction 
+				//   process, only to be deleted later if it already existed. This should 
+				//   determine if the constraint is already defined, and then opt out of 
+				//   the creation process. 
+				//
+				Collisions::ConstructConstraints<FReal, 3>(Particle0, Particle1, Particle0->Geometry().Get(), Particle1->Geometry().Get(), Collisions::GetTransform(Particle0), Collisions::GetTransform(Particle1), Thickness, NewConstraints);
 
 				CHAOS_COLLISION_STAT(if (NewConstraints.Num()) { StatData.IncrementCountNP(NewConstraints.Num()); });
 				CHAOS_COLLISION_STAT(if (!NewConstraints.Num()) { StatData.IncrementRejectedNP(); });
diff --git a/Engine/Source/Runtime/Experimental/Chaos/Public/Chaos/CollisionResolution.h b/Engine/Source/Runtime/Experimental/Chaos/Public/Chaos/CollisionResolution.h
--- a/Engine/Source/Runtime/Experimental/Chaos/Public/Chaos/CollisionResolution.h
+++ b/Engine/Source/Runtime/Experimental/Chaos/Public/Chaos/CollisionResolution.h
@@ -133,17 +133,6 @@ namespace Chaos
 			void CHAOS_API ConstructLevelsetLevelsetConstraints(TGeometryParticleHandle<T, d>* Particle0, TGeometryParticleHandle<T, d>* Particle1, const FImplicitObject* Implicit0, const FImplicitObject* Implicit1, const TRigidTransform<T, d>& Transform0, const TRigidTransform<T, d>& Transform1, const T Thickness, FCollisionConstraintsArray& NewConstraints);
 
 
-			//
-			//  Levelset-Union and Union-Levelset
-			//
-
-			template<ECollisionUpdateType UpdateType, typename T, int d>
-			void CHAOS_API UpdateLevelsetUnionConstraint(const T Thickness, TRigidBodyPointContactConstraint<T, d>& Constraint);
-
-			template<ECollisionUpdateType UpdateType, typename T, int d>
-			void CHAOS_API UpdateUnionLevelsetConstraint(const T Thickness, TRigidBodyPointContactConstraint<T, d>& Constraint);
-
-
 			//
 			// Union-Union
 			//
@@ -154,28 +143,16 @@ namespace Chaos
 			template<typename T, int d>
 			void CHAOS_API ConstructUnionUnionConstraints(TGeometryParticleHandle<T, d>* Particle0, TGeometryParticleHandle<T, d>* Particle1, const FImplicitObject* Implicit0, const FImplicitObject* Implicit1, const TRigidTransform<T, d>& Transform0, const TRigidTransform<T, d>& Transform1, const T Thickness, FCollisionConstraintsArray& NewConstraints);
 
-			//
-			// Single-Union
-			//
-
-			template<ECollisionUpdateType UpdateType, typename T, int d>
-			void CHAOS_API UpdateSingleUnionConstraint(const FImplicitObject& Implicit0, const TRigidTransform<T, d>& Transform0, const FImplicitObject& Implicit1, const TRigidTransform<T, d>& Transform1, const T Thickness, TCollisionConstraintBase<T, d>& Constraint);
-
-			template<typename T, int d>
-			void CHAOS_API ConstructSingleUnionConstraints(TGeometryParticleHandle<T, d>* Particle0, TGeometryParticleHandle<T, d>* Particle1, const FImplicitObject* Implicit0, const FImplicitObject* Implicit1, const TRigidTransform<T, d>& Transform0, const TRigidTransform<T, d>& Transform1, const  T Thickness, FCollisionConstraintsArray& NewConstraints);
-
 
 			//
 			// Constraint API
 			//
 
 			template<ECollisionUpdateType UpdateType, typename T, int d>
-			void CHAOS_API UpdateConstraintImp(const FImplicitObject& ParticleObject, const TRigidTransform<T, d>& ParticleTM, const FImplicitObject& LevelsetObject, const TRigidTransform<T, d>& LevelsetTM, const T Thickness, TCollisionConstraintBase<T, d>& Constraint);
+			void CHAOS_API UpdateConstraint(TCollisionConstraintBase<T, d>& ConstraintBase, const TRigidTransform<T, d>& Transform0, const TRigidTransform<T, d>& Transform1, const T Thickness);
 
 			template<typename T, int d>
-			void CHAOS_API ConstructPairConstraintImpl(TGeometryParticleHandle<T, d>* Particle0, TGeometryParticleHandle<T, d>* Particle1, const FImplicitObject* Implicit0, const FImplicitObject* Implicit1, const TRigidTransform<T, d>& Transform0, const TRigidTransform<T, d>& Transform1, const T Thickness, FCollisionConstraintsArray& NewConstraints);
+			void CHAOS_API ConstructConstraints(TGeometryParticleHandle<T, d>* Particle0, TGeometryParticleHandle<T, d>* Particle1, const FImplicitObject* Implicit0, const FImplicitObject* Implicit1, const TRigidTransform<T, d>& Transform0, const TRigidTransform<T, d>& Transform1, const T Thickness, FCollisionConstraintsArray& NewConstraints);
 
-			template<typename T, int d>
-			void CHAOS_API ConstructConstraintsImpl(TGeometryParticleHandle<T, d>* Particle0, TGeometryParticleHandle<T, d>* Particle1, const FImplicitObject* Implicit0, const FImplicitObject* Implicit1, const TRigidTransform<T, d>& Transform0, const TRigidTransform<T, d>& Transform1, const T Thickness, FCollisionConstraintsArray& NewConstraints);
 	}
 }
diff --git a/Engine/Source/Runtime/Experimental/Chaos/Public/Chaos/CollisionResolutionTypes.h b/Engine/Source/Runtime/Experimental/Chaos/Public/Chaos/CollisionResolutionTypes.h
--- a/Engine/Source/Runtime/Experimental/Chaos/Public/Chaos/CollisionResolutionTypes.h
+++ b/Engine/Source/Runtime/Experimental/Chaos/Public/Chaos/CollisionResolutionTypes.h
@@ -84,17 +84,28 @@ namespace Chaos
 			MultiPoint    //   TRigidBodyIterativeContactConstraint
 		};
 
-		TCollisionConstraintBase(int32 InTimestamp = -INT_MAX) 
+		TCollisionConstraintBase(FType InType = FType::None)
 			: AccumulatedImpulse(0)
-			, Timestamp(InTimestamp)
-			, Type(FType::None)
-		{ Particle[0] = nullptr; Particle[1] = nullptr; }
+			, Timestamp(-INT_MAX)
+			, Type(InType)
+		{ 
+			ImplicitTransform[0] = TRigidTransform<T, d>::Identity; ImplicitTransform[1] = TRigidTransform<T,d>::Identity;
+			Manifold.Implicit[0] = nullptr; Manifold.Implicit[1] = nullptr;
+			Particle[0] = nullptr; Particle[1] = nullptr; 
+		}
 		
-		TCollisionConstraintBase(FType InType, int32 InTimestamp = -INT_MAX) 
+		TCollisionConstraintBase(
+			FGeometryParticleHandle* Particle0, const FImplicitObject* Implicit0, const TRigidTransform<T, d>& Transform0,
+			FGeometryParticleHandle* Particle1, const FImplicitObject* Implicit1, const TRigidTransform<T, d>& Transform1,
+			FType InType, int32 InTimestamp = -INT_MAX)
 			: AccumulatedImpulse(0)
 			, Timestamp(InTimestamp)
 			, Type(InType)
-		{ Particle[0] = nullptr; Particle[1] = nullptr; }
+		{
+			ImplicitTransform[0] = Transform0; ImplicitTransform[1] = Transform1;
+			Manifold.Implicit[0] = Implicit0; Manifold.Implicit[1] = Implicit1;
+			Particle[0] = Particle0; Particle[1] = Particle1; 
+		}
 
 		FType GetType() const { return Type; }
 
@@ -129,6 +140,7 @@ namespace Chaos
 		}
 
 
+		TRigidTransform<T, d> ImplicitTransform[2]; // { Point, Volume }
 		FGeometryParticleHandle* Particle[2]; // { Point, Volume }
 		TVector<T, d> AccumulatedImpulse;
 		FManifold Manifold;
@@ -154,6 +166,11 @@ namespace Chaos
 		using Base::Particle;
 
 		TRigidBodyPointContactConstraint() : Base(Base::FType::SinglePoint) {}
+		TRigidBodyPointContactConstraint(
+			FGeometryParticleHandle* Particle0, const FImplicitObject* Implicit0, const TRigidTransform<T, d>& Transform0,
+			FGeometryParticleHandle* Particle1, const FImplicitObject* Implicit1, const TRigidTransform<T, d>& Transform1)
+			: Base(Particle0, Implicit0, Transform0, Particle1, Implicit1, Transform1, Base::FType::SinglePoint) {}
+
 		static typename Base::FType StaticType() { return Base::FType::SinglePoint; };
 	};
 	typedef TRigidBodyPointContactConstraint<float, 3> FRigidBodyPointContactConstraint;
@@ -173,6 +190,13 @@ namespace Chaos
 		struct FSampleData { TVector<T,d> X; float Delta; FManifold Manifold; };
 
 		TRigidBodyIterativeContactConstraint() : Base(Base::FType::MultiPoint), SourceNormalIndex(INDEX_NONE), PlaneNormal(0), PlanePosition(0) {}
+		TRigidBodyIterativeContactConstraint(
+			FGeometryParticleHandle* Particle0, const FImplicitObject* Implicit0, const TRigidTransform<T, d>& Transform0,
+			FGeometryParticleHandle* Particle1, const FImplicitObject* Implicit1, const TRigidTransform<T, d>& Transform1)
+			: Base(Particle0, Implicit0, Transform0, Particle1, Implicit1, Transform1, Base::FType::MultiPoint)
+			, SourceNormalIndex(INDEX_NONE), PlaneNormal(0), PlanePosition(0) 
+		{}
+
 		static typename Base::FType StaticType() { return Base::FType::MultiPoint; };
 
 		int SourceNormalIndex; // index of normal on particle[0] body;
