Make sure to rotate bounds for internal bvh of trimesh

diff --git a/Engine/Source/Runtime/Experimental/Chaos/Private/Chaos/TriangleMeshImplicitObject.cpp b/Engine/Source/Runtime/Experimental/Chaos/Private/Chaos/TriangleMeshImplicitObject.cpp
--- a/Engine/Source/Runtime/Experimental/Chaos/Private/Chaos/TriangleMeshImplicitObject.cpp
+++ b/Engine/Source/Runtime/Experimental/Chaos/Private/Chaos/TriangleMeshImplicitObject.cpp
@@ -543,7 +543,7 @@ bool FTriangleMeshImplicitObject::SweepGeomImp(const QueryGeomType& QueryGeom, c
 	FTriangleMeshSweepVisitor<QueryGeomType> SQVisitor(*this, QueryGeom, StartTM, Dir, ScaledDirNormalized, LengthScale, ScaledStartTM, Thickness, bComputeMTD);
 
 
-	const TAABB<FReal, 3> QueryBounds = QueryGeom.BoundingBox();
+	const TAABB<FReal, 3> QueryBounds = QueryGeom.BoundingBox().TransformedAABB(FRigidTransform3(FVec3::ZeroVector, StartTM.GetRotation()));
 	const FVec3 StartPoint = StartTM.TransformPositionNoScale(QueryBounds.Center());
 	const FVec3 Inflation = QueryBounds.Extents() * 0.5 + FVec3(Thickness);
 	BVH.template Sweep<FTriangleMeshSweepVisitor<QueryGeomType>>(StartPoint, Dir, Length, Inflation, SQVisitor);
