- Fix CIS

diff --git a/Engine/Source/Runtime/Experimental/Chaos/Private/Chaos/CollisionResolution.cpp b/Engine/Source/Runtime/Experimental/Chaos/Private/Chaos/CollisionResolution.cpp
--- a/Engine/Source/Runtime/Experimental/Chaos/Private/Chaos/CollisionResolution.cpp
+++ b/Engine/Source/Runtime/Experimental/Chaos/Private/Chaos/CollisionResolution.cpp
@@ -775,11 +775,11 @@ namespace Chaos
 		//  Union-Union
 		//
 
-		DECLARE_CYCLE_STAT(TEXT("TPBDCollisionConstraints::UpdateUnionUnionConstraint"), STAT_UpdateUnionUnionConstraint, STATGROUP_ChaosWide);
+		DECLARE_CYCLE_STAT(TEXT("TPBDCollisionConstraints::ConstructUnionUnionConstraints"), STAT_ConstructUnionUnionConstraints, STATGROUP_ChaosWide);
 		template<typename T, int d>
 		void ConstructUnionUnionConstraints(TGeometryParticleHandle<T, d>* Particle0, TGeometryParticleHandle<T, d>* Particle1, const FImplicitObject* Implicit0, const FImplicitObject* Implicit1, const TRigidTransform<T, d>& Transform0, const TRigidTransform<T, d>& Transform1, const T Thickness, FCollisionConstraintsArray& NewConstraints)
 		{
-			SCOPE_CYCLE_COUNTER(STAT_UpdateUnionUnionConstraint);
+			SCOPE_CYCLE_COUNTER(STAT_ConstructUnionUnionConstraints);
 
 			const TArray<Pair<const FImplicitObject*, TRigidTransform<T, d>>> LevelsetShapes = FindRelevantShapes(Implicit0, Transform0, *Implicit1, Transform1, Thickness);
 
@@ -1097,8 +1097,6 @@ namespace Chaos
 		template void UpdateLevelsetLevelsetConstraint<ECollisionUpdateType::Deepest, float, 3>(const float Thickness, TRigidBodyPointContactConstraint<float, 3>& Constraint);
 		template void ConstructLevelsetLevelsetConstraints<float, 3>(TGeometryParticleHandle<float, 3>* Particle0, TGeometryParticleHandle<float, 3>* Particle1, const FImplicitObject* Implicit0, const FImplicitObject* Implicit1, const TRigidTransform<float, 3>& Transform0, const TRigidTransform<float, 3>& Transform1, const float Thickness, FCollisionConstraintsArray& NewConstraints);
 
-		template void UpdateUnionUnionConstraint< ECollisionUpdateType::Any, float, 3>(const FImplicitObject& Implicit0, const TRigidTransform<float, 3>& Transform0, const FImplicitObject& Implicit1, const TRigidTransform<float, 3>& Transform1, const float Thickness, TCollisionConstraintBase<float, 3>& Constraint);
-		template void UpdateUnionUnionConstraint< ECollisionUpdateType::Deepest, float, 3>(const FImplicitObject& Implicit0, const TRigidTransform<float, 3>& Transform0, const FImplicitObject& Implicit1, const TRigidTransform<float, 3>& Transform1, const float Thickness, TCollisionConstraintBase<float, 3>& Constraint);
 		template void ConstructUnionUnionConstraints<float, 3>(TGeometryParticleHandle<float, 3>* Particle0, TGeometryParticleHandle<float, 3>* Particle1, const FImplicitObject* Implicit0, const FImplicitObject* Implicit1, const TRigidTransform<float, 3>& Transform0, const TRigidTransform<float, 3>& Transform1, const float Thickness, FCollisionConstraintsArray& NewConstraints);
 
 		template void UpdateConstraint<ECollisionUpdateType::Any, float, 3>(TCollisionConstraintBase<float,3>& ConstraintBase, const TRigidTransform<float,3>& Transform0, const TRigidTransform<float,3>& Transform1, const float Thickness);
diff --git a/Engine/Source/Runtime/Experimental/Chaos/Public/Chaos/CollisionResolution.h b/Engine/Source/Runtime/Experimental/Chaos/Public/Chaos/CollisionResolution.h
--- a/Engine/Source/Runtime/Experimental/Chaos/Public/Chaos/CollisionResolution.h
+++ b/Engine/Source/Runtime/Experimental/Chaos/Public/Chaos/CollisionResolution.h
@@ -136,14 +136,9 @@ namespace Chaos
 			//
 			// Union-Union
 			//
-
-			template<ECollisionUpdateType UpdateType, typename T, int d>
-			void CHAOS_API UpdateUnionUnionConstraint(const FImplicitObject& Implicit0, const TRigidTransform<T, d>& Transform0, const FImplicitObject& Implicit1, const TRigidTransform<T, d>& Transform1, const T Thickness, TCollisionConstraintBase<T, d>& Constraint);
-
 			template<typename T, int d>
 			void CHAOS_API ConstructUnionUnionConstraints(TGeometryParticleHandle<T, d>* Particle0, TGeometryParticleHandle<T, d>* Particle1, const FImplicitObject* Implicit0, const FImplicitObject* Implicit1, const TRigidTransform<T, d>& Transform0, const TRigidTransform<T, d>& Transform1, const T Thickness, FCollisionConstraintsArray& NewConstraints);
 
-
 			//
 			// Constraint API
 			//
