- update to comments

diff --git a/Engine/Source/Runtime/Engine/Private/PhysicsEngine/Experimental/PhysInterface_Chaos.cpp b/Engine/Source/Runtime/Engine/Private/PhysicsEngine/Experimental/PhysInterface_Chaos.cpp
--- a/Engine/Source/Runtime/Engine/Private/PhysicsEngine/Experimental/PhysInterface_Chaos.cpp
+++ b/Engine/Source/Runtime/Engine/Private/PhysicsEngine/Experimental/PhysInterface_Chaos.cpp
@@ -1354,9 +1354,7 @@ void FPhysInterface_Chaos::AddGeometry(FPhysicsActorHandle& InActor, const FGeom
 }
 
 
-// todo(brice): Implicit Initialization Pipeline(WIP)
-// ... add virtual TImplicitObject::NewCopy()
-// @todo(mlentine,brice): We probably need to actually duplicate the data here, add virtual TImplicitObject::NewCopy()
+// @todo(chaos): We probably need to actually duplicate the data here, add virtual TImplicitObject::NewCopy()
 FPhysicsShapeHandle FPhysInterface_Chaos::CloneShape(const FPhysicsShapeHandle& InShape)
 {
 	FPhysicsActorHandle NewActor = nullptr; // why zero and not the default INDEX_NONE?
diff --git a/Engine/Source/Runtime/Experimental/Chaos/Private/Chaos/PBDCollisionConstraints.cpp b/Engine/Source/Runtime/Experimental/Chaos/Private/Chaos/PBDCollisionConstraints.cpp
--- a/Engine/Source/Runtime/Experimental/Chaos/Private/Chaos/PBDCollisionConstraints.cpp
+++ b/Engine/Source/Runtime/Experimental/Chaos/Private/Chaos/PBDCollisionConstraints.cpp
@@ -233,7 +233,10 @@ namespace Chaos
 			check(false);
 		}
 
-		Manifolds.Remove(KeyToRemove); // todo(brice): Add index to the handle to prevent the search. 
+		// @todo(chaos): Collision Manifold
+		//   Add an index to the handle in the Manifold.Value 
+		//   to prevent the search in Handles when removed.
+		Manifolds.Remove(KeyToRemove);  
 		Handles.Remove(Handle);
 
 		ensure(Handles.Num() == PointConstraints.Num() + IterativeConstraints.Num());
diff --git a/Engine/Source/Runtime/Experimental/Chaos/Public/Chaos/Collision/CollisionReceiver.h b/Engine/Source/Runtime/Experimental/Chaos/Public/Chaos/Collision/CollisionReceiver.h
--- a/Engine/Source/Runtime/Experimental/Chaos/Public/Chaos/Collision/CollisionReceiver.h
+++ b/Engine/Source/Runtime/Experimental/Chaos/Public/Chaos/Collision/CollisionReceiver.h
@@ -40,10 +40,11 @@ namespace Chaos
 			FCollisionConstraintBase* ConstraintBase = nullptr;
 			while (Queue.Dequeue(ConstraintBase))
 			{
-				// todo(brice) : This needs to be moved within the MidPhase (which does not exist yet). The 
-				//               NarrowPhase is creating temporary constraints for all interacting bodies 
-				//               even though we might already have those bodies within the constraint.
-				//               So there is a bunch of wasted computation when the CreateConstraints is called. 
+				// @todo(chaos) : Collision Constraints
+				//    This needs to be moved within the MidPhase (which does not exist yet). The 
+				//    NarrowPhase is creating temporary constraints for all interacting bodies 
+				//    even though we might already have those bodies within the constraint.
+				//    So there is a bunch of wasted computation when the CreateConstraints is called. 
 				if (!CollisionConstraints.Contains(ConstraintBase))
 				{
 					CollisionConstraints.AddConstraint(ConstraintBase);
diff --git a/Engine/Source/Runtime/Experimental/ChaosSolvers/Private/PhysicsProxy/FieldSystemPhysicsProxy.cpp b/Engine/Source/Runtime/Experimental/ChaosSolvers/Private/PhysicsProxy/FieldSystemPhysicsProxy.cpp
--- a/Engine/Source/Runtime/Experimental/ChaosSolvers/Private/PhysicsProxy/FieldSystemPhysicsProxy.cpp
+++ b/Engine/Source/Runtime/Experimental/ChaosSolvers/Private/PhysicsProxy/FieldSystemPhysicsProxy.cpp
@@ -117,7 +117,7 @@ void FFieldSystemPhysicsProxy::FieldParameterUpdateCallback(
 					TArrayView<ContextIndex> SampleIndicesView(&(SampleIndices[0]), SampleIndices.Num());
 
 					FFieldContext Context(
-						SampleIndicesView, // @todo(brice) important: an empty index array should evaluate everything
+						SampleIndicesView, // @todo(chaos) important: an empty index array should evaluate everything
 						SamplePointsView,
 						Command.MetaData);
 
@@ -251,7 +251,7 @@ void FFieldSystemPhysicsProxy::FieldParameterUpdateCallback(
 					TArrayView<FVector> SamplesView(tptr, int32(Particles.Size()));
 
 					FFieldContext Context{
-						IndexView, // @todo(brice) important: an empty index array should evaluate everything
+						IndexView, // @todo(chaos) important: an empty index array should evaluate everything
 						SamplesView,
 						Command.MetaData
 					};
