Change check() in TaperedCylinder.h to ensure to avoid crashing game until bad geometry problem is fixed.

diff --git a/Engine/Source/Runtime/Experimental/Chaos/Public/Chaos/TaperedCylinder.h b/Engine/Source/Runtime/Experimental/Chaos/Public/Chaos/TaperedCylinder.h
--- a/Engine/Source/Runtime/Experimental/Chaos/Public/Chaos/TaperedCylinder.h
+++ b/Engine/Source/Runtime/Experimental/Chaos/Public/Chaos/TaperedCylinder.h
@@ -120,7 +120,7 @@ namespace Chaos
 			const T Distance1 = MPlane1.PhiWithNormal(x, Normal1);
 			if (Distance1 < SMALL_NUMBER)
 			{
-				check(MPlane2.PhiWithNormal(x, Normal2) > 0.);
+				ensure(MPlane2.PhiWithNormal(x, Normal2) > 0.);
 				const TVector<T, 3> v = x - TVector<T, 3>(Normal1 * Distance1 + MPlane1.X());
 				if (v.Size() > MRadius1)
 				{
@@ -138,7 +138,7 @@ namespace Chaos
 			const T Distance2 = MPlane2.PhiWithNormal(x, Normal2);
 			if (Distance2 < SMALL_NUMBER)
 			{
-				check(MPlane1.PhiWithNormal(x, Normal1) > 0.);
+				ensure(MPlane1.PhiWithNormal(x, Normal1) > 0.);
 				const TVector<T, 3> v = x - TVector<T, 3>(Normal2 * Distance2 + MPlane2.X());
 				if (v.Size() > MRadius2)
 				{
@@ -153,7 +153,7 @@ namespace Chaos
 					return -Distance2;
 				}
 			}
-			check(Distance1 <= MHeight && Distance2 <= MHeight);
+			ensure(Distance1 <= MHeight && Distance2 <= MHeight);
 			const TVector<T, 3> SideVector = (x - TVector<T, 3>(Normal1 * Distance1 + MPlane1.X()));
 			const T SideDistance = SideVector.Size() - GetRadius(Distance1);
 			if (SideDistance < 0.)
