Fix CIS - TriangleMeshImplicitObject Clang error

diff --git a/Engine/Source/Runtime/Experimental/Chaos/Public/Chaos/TriangleMeshImplicitObject.h b/Engine/Source/Runtime/Experimental/Chaos/Public/Chaos/TriangleMeshImplicitObject.h
--- a/Engine/Source/Runtime/Experimental/Chaos/Public/Chaos/TriangleMeshImplicitObject.h
+++ b/Engine/Source/Runtime/Experimental/Chaos/Public/Chaos/TriangleMeshImplicitObject.h
@@ -29,7 +29,7 @@ namespace Chaos
 
 		FTriangleMeshImplicitObject(TParticles<FReal,3>&& Particles, TArray<TVector<int32, 3>>&& Elements, TArray<uint16>&& InMaterialIndices);
 		FTriangleMeshImplicitObject(const FTriangleMeshImplicitObject& Other) = delete;
-		FTriangleMeshImplicitObject(FTriangleMeshImplicitObject&& Other) = default;
+		FTriangleMeshImplicitObject(FTriangleMeshImplicitObject&& Other) = delete;
 		virtual ~FTriangleMeshImplicitObject() {}
 
 		virtual FReal PhiWithNormal(const FVec3& x, FVec3& Normal) const;
