Integrate 10712088  from 11.40: Chaos - Change HeightField GridCast/GridSweep so it will not move diagonally on grid and cut corners.

diff --git a/Engine/Source/Runtime/Experimental/Chaos/Private/Chaos/HeightField.cpp b/Engine/Source/Runtime/Experimental/Chaos/Private/Chaos/HeightField.cpp
--- a/Engine/Source/Runtime/Experimental/Chaos/Private/Chaos/HeightField.cpp
+++ b/Engine/Source/Runtime/Experimental/Chaos/Private/Chaos/HeightField.cpp
@@ -719,10 +719,10 @@ namespace Chaos
 				EndCell = FlatGrid.ClampIndex(EndCell);
 
 				int32 DeltaX = FMath::Abs(EndCell[0] - StartCell[0]);
-				int32 DeltaY = -FMath::Abs(EndCell[1] - StartCell[1]);
+				int32 DeltaY = FMath::Abs(EndCell[1] - StartCell[1]);
 				int32 DirX = StartCell[0] < EndCell[0] ? 1 : -1;
 				int32 DirY = StartCell[1] < EndCell[1] ? 1 : -1;
-				int32 Error = DeltaX + DeltaY;
+				int32 Error = 0;
 
 				if(StartCell == EndCell)
 				{
@@ -751,17 +751,18 @@ namespace Chaos
 
 					do 
 					{
-						const int32 DoubleError = Error * 2;
+						const int32 ErrorX = Error + DeltaY;
+						const int32 ErrorY = Error - DeltaX;
 
-						if(DoubleError >= DeltaY)
+
+						if(FMath::Abs(ErrorX) < FMath::Abs(ErrorY))
 						{
-							Error += DeltaY;
+							Error = ErrorX;
 							StartCell[0] += DirX;
 						}
-
-						if(DoubleError <= DeltaX)
+						else
 						{
-							Error += DeltaX;
+							Error = ErrorY;
 							StartCell[1] += DirY;
 						}
 
@@ -853,12 +854,12 @@ namespace Chaos
 			EndCell = FlatGrid.ClampIndex(EndCell);
 
 			const int32 DeltaX = FMath::Abs(EndCell[0] - StartCell[0]);
-			const int32 DeltaY = -FMath::Abs(EndCell[1] - StartCell[1]);
+			const int32 DeltaY = FMath::Abs(EndCell[1] - StartCell[1]);
 			const bool bSameCell = DeltaX == 0 && DeltaY == 0;
 
 			const int32 DirX = StartCell[0] < EndCell[0] ? 1 : -1;
 			const int32 DirY = StartCell[1] < EndCell[1] ? 1 : -1;
-			int32 Error = DeltaX + DeltaY;
+			int32 Error = 0;
 			const TVector<int32, 2> ThickenDir = FMath::Abs(DeltaX) > FMath::Abs(DeltaY) ? TVector<int32, 2>(0, 1) : TVector<int32, 2>(1, 0);
 
 			struct FQueueEntry
@@ -997,17 +998,17 @@ namespace Chaos
 						// Walk the line and add to the queue
 						if(StartCell != EndCell)
 						{
-							const int32 DoubleError = Error * 2;
+							const int32 ErrorX = Error + DeltaY;
+							const int32 ErrorY = Error - DeltaX;
 
-							if(DoubleError >= DeltaY)
+							if(FMath::Abs(ErrorX) < FMath::Abs(ErrorY))
 							{
-								Error += DeltaY;
+								Error = ErrorX;
 								StartCell[0] += DirX;
 							}
-
-							if(DoubleError <= DeltaX)
+							else
 							{
-								Error += DeltaX;
+								Error = ErrorY;
 								StartCell[1] += DirY;
 							}
 
diff --git a/Engine/Source/Runtime/Landscape/Private/LandscapeCollision.cpp b/Engine/Source/Runtime/Landscape/Private/LandscapeCollision.cpp
--- a/Engine/Source/Runtime/Landscape/Private/LandscapeCollision.cpp
+++ b/Engine/Source/Runtime/Landscape/Private/LandscapeCollision.cpp
@@ -509,16 +509,10 @@ void ULandscapeHeightfieldCollisionComponent::OnCreatePhysicsState()
 
 				TArray<FPhysicsActorHandle> Actors;
 				Actors.Add(PhysHandle);
-				PhysScene->AddActorsToScene_AssumesLocked(Actors);
 
-#if WITH_EDITOR
-				// If we're in an editor world we will never simulate but we require up to date SQ for the tools
-				// #BGTODO consider a more automatic way to have scenes update when they aren't simulating to avoid this explicit flush
-				if(!GetWorld()->IsGameWorld())
-				{
-					PhysScene->Flush_AssumesLocked();
-				}
-#endif
+				bool bImmediateAccelStructureInsertion = true;
+				PhysScene->AddActorsToScene_AssumesLocked(Actors, bImmediateAccelStructureInsertion);
+
 			}
 #endif // WITH_CHAOS
 		}
