Added support for a non-mutable AABBTree

diff --git a/Engine/Source/Runtime/Experimental/Chaos/Public/Chaos/AABBTree.h b/Engine/Source/Runtime/Experimental/Chaos/Public/Chaos/AABBTree.h
--- a/Engine/Source/Runtime/Experimental/Chaos/Public/Chaos/AABBTree.h
+++ b/Engine/Source/Runtime/Experimental/Chaos/Public/Chaos/AABBTree.h
@@ -8,6 +8,7 @@
 #include "ChaosLog.h"
 #include "Chaos/ISpatialAcceleration.h"
 #include "Templates/Models.h"
+#include "Chaos/BoundingVolume.h"
 
 namespace Chaos
 {
@@ -218,7 +219,7 @@ inline FArchive& operator<<(FArchive& Ar, FAABBTreePayloadInfo& PayloadInfo)
 	return Ar;
 }
 
-template <typename TPayloadType, typename TLeafType, typename T>
+template <typename TPayloadType, typename TLeafType, typename T, bool bMutable = true>
 class TAABBTree final : public ISpatialAcceleration<TPayloadType, T, 3>
 {
 public:
@@ -243,7 +244,6 @@ public:
 	virtual void ProgressAsyncTimeSlicing(bool ForceBuildCompletion) override
 	{
 		SCOPE_CYCLE_COUNTER(STAT_AABBTreeProgressTimeSlice);
-		LLM_SCOPE(ELLMTag::ChaosAcceleration);
 		// force is to stop time slicing and complete the rest of the build now
 		if (ForceBuildCompletion)
 		{
@@ -314,7 +314,7 @@ public:
 
 	virtual TUniquePtr<ISpatialAcceleration<TPayloadType, T, 3>> Copy() const override
 	{
-		return TUniquePtr<ISpatialAcceleration<TPayloadType, T, 3>>(new TAABBTree<TPayloadType, TLeafType, T>(*this));
+		return TUniquePtr<ISpatialAcceleration<TPayloadType, T, 3>>(new TAABBTree<TPayloadType, TLeafType, T, bMutable>(*this));
 	}
 
 	virtual void Raycast(const TVector<T, 3>& Start, const TVector<T, 3>& Dir, const T Length, ISpatialVisitor<TPayloadType, T>& Visitor) const override
@@ -377,112 +377,116 @@ public:
 
 	virtual void RemoveElement(const TPayloadType& Payload)
 	{
-		LLM_SCOPE(ELLMTag::ChaosAcceleration);
-		if (FAABBTreePayloadInfo* PayloadInfo = PayloadToInfo.Find(Payload))
+		if (ensure(bMutable))
 		{
-			if (PayloadInfo->GlobalPayloadIdx != INDEX_NONE)
+			if (FAABBTreePayloadInfo* PayloadInfo = PayloadToInfo.Find(Payload))
 			{
-				ensure(PayloadInfo->DirtyPayloadIdx == INDEX_NONE);
-				ensure(PayloadInfo->LeafIdx == INDEX_NONE);
-				if(PayloadInfo->GlobalPayloadIdx + 1 < GlobalPayloads.Num())
+				if (PayloadInfo->GlobalPayloadIdx != INDEX_NONE)
 				{
-					auto LastGlobalPayload = GlobalPayloads.Last().Payload;
-					PayloadToInfo.FindChecked(LastGlobalPayload).GlobalPayloadIdx = PayloadInfo->GlobalPayloadIdx;
+					ensure(PayloadInfo->DirtyPayloadIdx == INDEX_NONE);
+					ensure(PayloadInfo->LeafIdx == INDEX_NONE);
+					if (PayloadInfo->GlobalPayloadIdx + 1 < GlobalPayloads.Num())
+					{
+						auto LastGlobalPayload = GlobalPayloads.Last().Payload;
+						PayloadToInfo.FindChecked(LastGlobalPayload).GlobalPayloadIdx = PayloadInfo->GlobalPayloadIdx;
+					}
+					GlobalPayloads.RemoveAtSwap(PayloadInfo->GlobalPayloadIdx);
 				}
-				GlobalPayloads.RemoveAtSwap(PayloadInfo->GlobalPayloadIdx);
-			}
-			else if(PayloadInfo->DirtyPayloadIdx != INDEX_NONE)
-			{
-				if(PayloadInfo->DirtyPayloadIdx + 1 < DirtyElements.Num())
+				else if (PayloadInfo->DirtyPayloadIdx != INDEX_NONE)
 				{
-					auto LastDirtyPayload = DirtyElements.Last().Payload;
-					PayloadToInfo.FindChecked(LastDirtyPayload).DirtyPayloadIdx = PayloadInfo->DirtyPayloadIdx;
+					if (PayloadInfo->DirtyPayloadIdx + 1 < DirtyElements.Num())
+					{
+						auto LastDirtyPayload = DirtyElements.Last().Payload;
+						PayloadToInfo.FindChecked(LastDirtyPayload).DirtyPayloadIdx = PayloadInfo->DirtyPayloadIdx;
+					}
+					DirtyElements.RemoveAtSwap(PayloadInfo->DirtyPayloadIdx);
+				}
+				else if (ensure(PayloadInfo->LeafIdx != INDEX_NONE))
+				{
+					Leaves[PayloadInfo->LeafIdx].RemoveElement(Payload);
 				}
-				DirtyElements.RemoveAtSwap(PayloadInfo->DirtyPayloadIdx);
-			}
-			else if(ensure(PayloadInfo->LeafIdx != INDEX_NONE))
-			{
-				Leaves[PayloadInfo->LeafIdx].RemoveElement(Payload);
-			}
 
-			PayloadToInfo.Remove(Payload);
+				PayloadToInfo.Remove(Payload);
+			}
 		}
 	}
 
 	virtual void UpdateElement(const TPayloadType& Payload, const TAABB<T, 3>& NewBounds, bool bHasBounds) override
 	{
-		LLM_SCOPE(ELLMTag::ChaosAcceleration);
-		FAABBTreePayloadInfo* PayloadInfo = PayloadToInfo.Find(Payload);
-		if (PayloadInfo)
-		{
-			if (PayloadInfo->LeafIdx != INDEX_NONE)
-			{
-				Leaves[PayloadInfo->LeafIdx].RemoveElement(Payload);
-				PayloadInfo->LeafIdx = INDEX_NONE;
-			}
-		}
-		else
-		{
-			PayloadInfo = &PayloadToInfo.Add(Payload);
-		}
-
-		bool bTooBig = false;
-		if (bHasBounds)
-		{
-			if (NewBounds.Extents().Max() > MaxPayloadBounds)
-			{
-				bTooBig = true;
-				bHasBounds = false;
-			}
-		}
-		
-		if (bHasBounds)
+		if (ensure(bMutable))
 		{
-			if (PayloadInfo->DirtyPayloadIdx == INDEX_NONE)
+			FAABBTreePayloadInfo* PayloadInfo = PayloadToInfo.Find(Payload);
+			if (PayloadInfo)
 			{
-				PayloadInfo->DirtyPayloadIdx = DirtyElements.Add(FElement{ Payload, NewBounds });
+				if (PayloadInfo->LeafIdx != INDEX_NONE)
+				{
+					Leaves[PayloadInfo->LeafIdx].RemoveElement(Payload);
+					PayloadInfo->LeafIdx = INDEX_NONE;
+				}
 			}
 			else
 			{
-				DirtyElements[PayloadInfo->DirtyPayloadIdx].Bounds = NewBounds;
+				PayloadInfo = &PayloadToInfo.Add(Payload);
 			}
 
-			// Handle something that previously did not have bounds that may be in global elements.
-			if (PayloadInfo->GlobalPayloadIdx != INDEX_NONE)
+			bool bTooBig = false;
+			if (bHasBounds)
 			{
-				if(PayloadInfo->GlobalPayloadIdx + 1 < GlobalPayloads.Num())
+				if (NewBounds.Extents().Max() > MaxPayloadBounds)
 				{
-					auto LastGlobalPayload = GlobalPayloads.Last().Payload;
-					PayloadToInfo.FindChecked(LastGlobalPayload).GlobalPayloadIdx = PayloadInfo->GlobalPayloadIdx;
+					bTooBig = true;
+					bHasBounds = false;
 				}
-				GlobalPayloads.RemoveAtSwap(PayloadInfo->GlobalPayloadIdx);
-
-				PayloadInfo->GlobalPayloadIdx = INDEX_NONE;
 			}
-		}
-		else
-		{
-			TAABB<T, 3> GlobalBounds = bTooBig ? NewBounds : TAABB<T, 3>(TVector<T, 3>(TNumericLimits<T>::Lowest()), TVector<T, 3>(TNumericLimits<T>::Max()));
-			if (PayloadInfo->GlobalPayloadIdx == INDEX_NONE)
+
+			if (bHasBounds)
 			{
-				PayloadInfo->GlobalPayloadIdx = GlobalPayloads.Add(FElement{ Payload, GlobalBounds });
+				if (PayloadInfo->DirtyPayloadIdx == INDEX_NONE)
+				{
+					PayloadInfo->DirtyPayloadIdx = DirtyElements.Add(FElement{ Payload, NewBounds });
+				}
+				else
+				{
+					DirtyElements[PayloadInfo->DirtyPayloadIdx].Bounds = NewBounds;
+				}
+
+				// Handle something that previously did not have bounds that may be in global elements.
+				if (PayloadInfo->GlobalPayloadIdx != INDEX_NONE)
+				{
+					if (PayloadInfo->GlobalPayloadIdx + 1 < GlobalPayloads.Num())
+					{
+						auto LastGlobalPayload = GlobalPayloads.Last().Payload;
+						PayloadToInfo.FindChecked(LastGlobalPayload).GlobalPayloadIdx = PayloadInfo->GlobalPayloadIdx;
+					}
+					GlobalPayloads.RemoveAtSwap(PayloadInfo->GlobalPayloadIdx);
+
+					PayloadInfo->GlobalPayloadIdx = INDEX_NONE;
+				}
 			}
 			else
 			{
-				GlobalPayloads[PayloadInfo->GlobalPayloadIdx].Bounds = GlobalBounds;
-			}
-
-			// Handle something that previously had bounds that may be in dirty elements.
-			if (PayloadInfo->DirtyPayloadIdx != INDEX_NONE)
-			{
-				if(PayloadInfo->DirtyPayloadIdx + 1 < DirtyElements.Num())
+				TAABB<T, 3> GlobalBounds = bTooBig ? NewBounds : TAABB<T, 3>(TVector<T, 3>(TNumericLimits<T>::Lowest()), TVector<T, 3>(TNumericLimits<T>::Max()));
+				if (PayloadInfo->GlobalPayloadIdx == INDEX_NONE)
 				{
-					auto LastDirtyPayload = DirtyElements.Last().Payload;
-					PayloadToInfo.FindChecked(LastDirtyPayload).DirtyPayloadIdx = PayloadInfo->DirtyPayloadIdx;
+					PayloadInfo->GlobalPayloadIdx = GlobalPayloads.Add(FElement{ Payload, GlobalBounds });
 				}
-				DirtyElements.RemoveAtSwap(PayloadInfo->DirtyPayloadIdx);
+				else
+				{
+					GlobalPayloads[PayloadInfo->GlobalPayloadIdx].Bounds = GlobalBounds;
+				}
+
+				// Handle something that previously had bounds that may be in dirty elements.
+				if (PayloadInfo->DirtyPayloadIdx != INDEX_NONE)
+				{
+					if (PayloadInfo->DirtyPayloadIdx + 1 < DirtyElements.Num())
+					{
+						auto LastDirtyPayload = DirtyElements.Last().Payload;
+						PayloadToInfo.FindChecked(LastDirtyPayload).DirtyPayloadIdx = PayloadInfo->DirtyPayloadIdx;
+					}
+					DirtyElements.RemoveAtSwap(PayloadInfo->DirtyPayloadIdx);
 
-				PayloadInfo->DirtyPayloadIdx = INDEX_NONE;
+					PayloadInfo->DirtyPayloadIdx = INDEX_NONE;
+				}
 			}
 		}
 	}
@@ -494,7 +498,6 @@ public:
 
 	virtual void Serialize(FChaosArchive& Ar) override
 	{
-		LLM_SCOPE(ELLMTag::ChaosAcceleration);
 		TBox<FReal, 3>::SerializeAsAABB(Ar, FullBounds);
 		Ar << Nodes;
 		Ar << Leaves;
@@ -513,6 +516,10 @@ public:
 			auto& Info = PayloadToInfo.FindOrAdd(Payload);
 			Ar << Info;
 		}
+		if (!bMutable)
+		{
+			PayloadToInfo.Empty();	//todo: avoid this entirely
+		}
 
 		Ar << MaxChildrenInLeaf;
 		Ar << MaxTreeDepth;
@@ -553,26 +560,29 @@ private:
 			}
 		}
 
-		for (const auto& Elem : DirtyElements)
+		if (bMutable)
 		{
-			const auto& InstanceBounds = Elem.Bounds;
-			if(TAABBTreeIntersectionHelper<T, TQueryFastData, Query>::Intersects(Start, CurData, TOI, TmpPosition, InstanceBounds, QueryBounds, QueryHalfExtents))
+			for (const auto& Elem : DirtyElements)
 			{
-				TSpatialVisitorData<TPayloadType> VisitData(Elem.Payload, true, InstanceBounds);
-				
-				bool bContinue;
-				if (Query == EAABBQueryType::Overlap)
-				{
-					bContinue = Visitor.VisitOverlap(VisitData);
-				}
-				else
-				{
-					bContinue = Query == EAABBQueryType::Sweep ? Visitor.VisitSweep(VisitData, CurData) : Visitor.VisitRaycast(VisitData, CurData);
-				}
-				
-				if (!bContinue)
+				const auto& InstanceBounds = Elem.Bounds;
+				if (TAABBTreeIntersectionHelper<T, TQueryFastData, Query>::Intersects(Start, CurData, TOI, TmpPosition, InstanceBounds, QueryBounds, QueryHalfExtents))
 				{
-					return false;
+					TSpatialVisitorData<TPayloadType> VisitData(Elem.Payload, true, InstanceBounds);
+
+					bool bContinue;
+					if (Query == EAABBQueryType::Overlap)
+					{
+						bContinue = Visitor.VisitOverlap(VisitData);
+					}
+					else
+					{
+						bContinue = Query == EAABBQueryType::Sweep ? Visitor.VisitSweep(VisitData, CurData) : Visitor.VisitRaycast(VisitData, CurData);
+					}
+
+					if (!bContinue)
+					{
+						return false;
+					}
 				}
 			}
 		}
@@ -640,7 +650,6 @@ private:
 	void GenerateTree(const TParticles& Particles)
 	{
 		SCOPE_CYCLE_COUNTER(STAT_AABBTreeGenerateTree);
-		LLM_SCOPE(ELLMTag::ChaosAcceleration);
 		this->SetAsyncTimeSlicingComplete(false);
 
 		TArray<FElement> ElemsWithBounds;
@@ -685,7 +694,10 @@ private:
 
 				if (!bHasBoundingBox)
 				{
-					PayloadToInfo.Add(Payload, FAABBTreePayloadInfo{ GlobalPayloads.Num(), INDEX_NONE, INDEX_NONE });
+					if (bMutable)
+					{
+						PayloadToInfo.Add(Payload, FAABBTreePayloadInfo{ GlobalPayloads.Num(), INDEX_NONE, INDEX_NONE });
+					}
 					GlobalPayloads.Add(FElement{ Payload, ElemBounds });
 				}
 
@@ -758,9 +770,12 @@ private:
 		auto& NodesRef = Nodes;
 		auto MakeLeaf = [NewNodeIdx, &PayloadToInfoRef, &Elems, &LeavesRef, &NodesRef]()
 		{
-			for (const FElement& Elem : Elems)
+			if (bMutable)
 			{
-				PayloadToInfoRef.Add(Elem.Payload, FAABBTreePayloadInfo{ INDEX_NONE, INDEX_NONE, LeavesRef.Num() });
+				for (const FElement& Elem : Elems)
+				{
+					PayloadToInfoRef.Add(Elem.Payload, FAABBTreePayloadInfo{ INDEX_NONE, INDEX_NONE, LeavesRef.Num() });
+				}
 			}
 
 			NodesRef[NewNodeIdx].bLeaf = true;
@@ -917,7 +932,7 @@ private:
 	}
 
 
-	TAABBTree(const TAABBTree<TPayloadType, TLeafType, T>& Other)
+	TAABBTree(const TAABBTree<TPayloadType, TLeafType, T, bMutable>& Other)
 		: ISpatialAcceleration<TPayloadType, T, 3>(StaticType)
 		, FullBounds(Other.FullBounds)
 		, Nodes(Other.Nodes)
@@ -980,8 +995,8 @@ private:
 	TQueue<FWorkSnapshot> TimeSliceWorkToComplete;
 };
 
-template<typename TPayloadType, typename TLeafType, class T>
-FArchive& operator<<(FChaosArchive& Ar, TAABBTree<TPayloadType, TLeafType, T>& AABBTree)
+template<typename TPayloadType, typename TLeafType, class T, bool bMutable>
+FArchive& operator<<(FChaosArchive& Ar, TAABBTree<TPayloadType, TLeafType, T, bMutable>& AABBTree)
 {
 	AABBTree.Serialize(Ar);
 	return Ar;
