Added geometry memory tracking in Chaos

diff --git a/Engine/Source/Runtime/Engine/Private/PhysicsEngine/BodySetup.cpp b/Engine/Source/Runtime/Engine/Private/PhysicsEngine/BodySetup.cpp
--- a/Engine/Source/Runtime/Engine/Private/PhysicsEngine/BodySetup.cpp
+++ b/Engine/Source/Runtime/Engine/Private/PhysicsEngine/BodySetup.cpp
@@ -735,6 +735,12 @@ void UBodySetup::FinishCreatingPhysicsMeshes_Chaos(FChaosDerivedDataReader<float
 	ChaosTriMeshes = MoveTemp(InReader.TrimeshImplicitObjects);
 	UVInfo = MoveTemp(InReader.UVInfo);
 	FaceRemap = MoveTemp(InReader.FaceRemap);
+#if TRACK_CHAOS_GEOMETRY
+	for (auto& TriMesh : ChaosTriMeshes)
+	{
+		TriMesh->Track(Chaos::MakeSerializable(TriMesh), FullName);
+	}
+#endif
 
 	// Clear the cooked data
 	if (!GIsEditor && !bSharedCookedData)
diff --git a/Engine/Source/Runtime/Engine/Private/PhysicsEngine/PhysUtils.cpp b/Engine/Source/Runtime/Engine/Private/PhysicsEngine/PhysUtils.cpp
--- a/Engine/Source/Runtime/Engine/Private/PhysicsEngine/PhysUtils.cpp
+++ b/Engine/Source/Runtime/Engine/Private/PhysicsEngine/PhysUtils.cpp
@@ -19,6 +19,7 @@
 #include "PhysicsSolver.h"
 #include "Chaos/PBDRigidsEvolutionGBF.h"
 #include "Chaos/ChaosArchive.h"
+#include "Chaos/TrackedGeometryManager.h"
 
 /** Returns false if ModelToHulls operation should halt because of vertex count overflow. */
 static bool AddConvexPrim(FKAggregateGeom* OutGeom, TArray<FPlane> &Planes, UModel* InModel)
@@ -473,7 +474,11 @@ bool FPhysScene::ExecApexVis(const TCHAR* Cmd, FOutputDevice* Ar)
 #if WITH_CHAOS
 bool FPhysicsInterface::ExecPhysCommands(const TCHAR* Cmd, FOutputDevice* OutputDevice, UWorld* InWorld)
 {
-
+	if (FParse::Command(&Cmd, TEXT("ChaosGeometryMemory")))
+	{
+		Chaos::FTrackedGeometryManager::Get().DumpMemoryUsage(OutputDevice);
+		return true;
+	}
 #if CHAOS_MEMORY_TRACKING
 	if (FParse::Command(&Cmd, TEXT("ChaosMemoryDistribution")))
 	{
diff --git a/Engine/Source/Runtime/Experimental/Chaos/Private/Chaos/ImplicitObject.cpp b/Engine/Source/Runtime/Experimental/Chaos/Private/Chaos/ImplicitObject.cpp
--- a/Engine/Source/Runtime/Experimental/Chaos/Private/Chaos/ImplicitObject.cpp
+++ b/Engine/Source/Runtime/Experimental/Chaos/Private/Chaos/ImplicitObject.cpp
@@ -12,6 +12,7 @@
 #include "Chaos/Plane.h"
 #include "Chaos/Sphere.h"
 #include "Chaos/TaperedCylinder.h"
+#include "Chaos/TrackedGeometryManager.h"
 #include "Chaos/TriangleMeshImplicitObject.h"
 #include "HAL/IConsoleManager.h"
 #include "UObject/DestructionObjectVersion.h"
@@ -23,13 +24,38 @@ FImplicitObject::FImplicitObject(int32 Flags, EImplicitObjectType InType)
     , bIsConvex(!!(Flags & EImplicitObject::IsConvex))
     , bIgnoreAnalyticCollisions(!!(Flags & EImplicitObject::IgnoreAnalyticCollisions))
     , bHasBoundingBox(!!(Flags & EImplicitObject::HasBoundingBox))
+#if TRACK_CHAOS_GEOMETRY
+    , bIsTracked(false)
+#endif
 {
 }
 
 FImplicitObject::~FImplicitObject()
 {
+#if TRACK_CHAOS_GEOMETRY
+	if (bIsTracked)
+	{
+		FTrackedGeometryManager::Get().RemoveGeometry(this);
+	}
+#endif
 }
 
+#if TRACK_CHAOS_GEOMETRY
+
+void FImplicitObject::Track(TSerializablePtr<FImplicitObject> This, const FString& DebugInfo)
+{
+	if (ensure(This.Get() == this))
+	{
+		if (!bIsTracked)
+		{
+			bIsTracked = true;
+			FTrackedGeometryManager::Get().AddGeometry(This, DebugInfo);
+		}
+	}
+}
+
+#endif
+
 EImplicitObjectType FImplicitObject::GetType(bool bGetTrueType) const
 {
 	if (bIgnoreAnalyticCollisions && !bGetTrueType)
diff --git a/Engine/Source/Runtime/Experimental/Chaos/Public/Chaos/ImplicitObject.h b/Engine/Source/Runtime/Experimental/Chaos/Public/Chaos/ImplicitObject.h
--- a/Engine/Source/Runtime/Experimental/Chaos/Public/Chaos/ImplicitObject.h
+++ b/Engine/Source/Runtime/Experimental/Chaos/Public/Chaos/ImplicitObject.h
@@ -7,6 +7,10 @@
 
 #include <functional>
 
+#ifndef TRACK_CHAOS_GEOMETRY
+#define TRACK_CHAOS_GEOMETRY !(UE_BUILD_SHIPPING || UE_BUILD_TEST)
+#endif
+
 namespace Chaos
 {
 template<class T, int d>
@@ -176,6 +180,12 @@ public:
 	void IgnoreAnalyticCollisions(const bool Ignore = true) { bIgnoreAnalyticCollisions = Ignore; }
 	bool GetIgnoreAnalyticCollisions() const { return bIgnoreAnalyticCollisions; }
 	void SetConvex(const bool Convex = true) { bIsConvex = Convex; }
+	
+#if TRACK_CHAOS_GEOMETRY
+	//Turn on memory tracking. Must pass object itself as a serializable ptr so we can save it out
+	void Track(TSerializablePtr<FImplicitObject> This, const FString& DebugInfo);
+#endif
+
 	virtual bool IsPerformanceWarning() const { return false; }
 	virtual FString PerformanceWarningAndSimplifaction() 
 	{
@@ -293,6 +303,10 @@ protected:
 	bool bIgnoreAnalyticCollisions;
 	bool bHasBoundingBox;
 
+#if TRACK_CHAOS_GEOMETRY
+	bool bIsTracked;
+#endif
+
 private:
 	virtual Pair<FVec3, bool> FindClosestIntersectionImp(const FVec3& StartPoint, const FVec3& EndPoint, const FReal Thickness) const;
 };
diff --git a/Engine/Source/Runtime/Experimental/Chaos/Public/Chaos/TrackedGeometryManager.h b/Engine/Source/Runtime/Experimental/Chaos/Public/Chaos/TrackedGeometryManager.h
new file mode 100644
--- /dev/null
+++ b/Engine/Source/Runtime/Experimental/Chaos/Public/Chaos/TrackedGeometryManager.h
@@ -0,0 +1,78 @@
+// Copyright 1998-2019 Epic Games, Inc. All Rights Reserved.
+#pragma once
+
+namespace Chaos
+{
+	class FImplicitObject;
+
+	class CHAOS_API FTrackedGeometryManager
+	{
+	public:
+		static FTrackedGeometryManager& Get()
+		{
+			static FTrackedGeometryManager Singleton;
+			return Singleton;
+		}
+
+		void DumpMemoryUsage(FOutputDevice* Ar) const
+		{
+			struct FMemInfo
+			{
+				uint32 NumBytes;
+				FString DebugInfo;
+
+				bool operator<(const FMemInfo& Other) const { return NumBytes < Other.NumBytes; }
+			};
+
+			TArray<FMemInfo> MemEntries;
+			uint32 TotalBytes = 0;
+			for (const auto& Itr : SharedGeometry)
+			{
+				FMemInfo Info;
+				Info.DebugInfo = Itr.Value;
+				TArray<uint8> Data;
+				FMemoryWriter MemAr(Data);
+				FChaosArchive ChaosAr(MemAr);
+				FImplicitObject* NonConst = const_cast<FImplicitObject*>(Itr.Key.Get());	//only doing this to write out, serialize is non const for read in
+				NonConst->Serialize(ChaosAr);
+				Info.NumBytes = Data.Num();
+				MemEntries.Add(Info);
+				TotalBytes += Info.NumBytes;
+			}
+
+			MemEntries.Sort();
+
+			Ar->Logf(TEXT(""));
+			Ar->Logf(TEXT("Chaos Tracked Geometry:"));
+			Ar->Logf(TEXT(""));
+
+			for (const FMemInfo& Info : MemEntries)
+			{
+				Ar->Logf(TEXT("%-10d %s"), Info.NumBytes, *Info.DebugInfo);
+			}
+
+			Ar->Logf(TEXT("%-10d Total"), TotalBytes);
+		}
+
+	private:
+		TMap<TSerializablePtr<FImplicitObject>, FString> SharedGeometry;
+		FCriticalSection CriticalSection;
+
+		friend FImplicitObject;
+
+		//These are private because of various threading considerations. ImplicitObject does the cleanup because it needs extra information
+		void AddGeometry(TSerializablePtr<FImplicitObject> Geometry, const FString& DebugInfo)
+		{
+			FScopeLock Lock(&CriticalSection);
+			SharedGeometry.Add(Geometry, DebugInfo);
+		}
+
+		void RemoveGeometry(const FImplicitObject* Geometry)
+		{
+			FScopeLock Lock(&CriticalSection);
+			TSerializablePtr<FImplicitObject> Dummy;
+			Dummy.SetFromRawLowLevel(Geometry);
+			SharedGeometry.Remove(Dummy);
+		}
+	};
+}
