Make trimesh's AABBTree immutable, saves 20% memory per trimesh

diff --git a/Engine/Source/Runtime/Core/Public/UObject/ExternalPhysicsCustomObjectVersion.h b/Engine/Source/Runtime/Core/Public/UObject/ExternalPhysicsCustomObjectVersion.h
--- a/Engine/Source/Runtime/Core/Public/UObject/ExternalPhysicsCustomObjectVersion.h
+++ b/Engine/Source/Runtime/Core/Public/UObject/ExternalPhysicsCustomObjectVersion.h
@@ -89,6 +89,9 @@ struct CORE_API FExternalPhysicsCustomObjectVersion
 
 		// Serialize bSimulate on PerShapeData
 		SerializePerShapeDataSimulateFlag,
+
+		// Serialize whether or not an AABBTree is immutable
+		ImmutableAABBTree,
 		
 		// -----<new versions can be added above this line>-------------------------------------------------
 		VersionPlusOne,
diff --git a/Engine/Source/Runtime/Engine/Private/PhysicsEngine/Experimental/ChaosDerivedData.cpp b/Engine/Source/Runtime/Engine/Private/PhysicsEngine/Experimental/ChaosDerivedData.cpp
--- a/Engine/Source/Runtime/Engine/Private/PhysicsEngine/Experimental/ChaosDerivedData.cpp
+++ b/Engine/Source/Runtime/Engine/Private/PhysicsEngine/Experimental/ChaosDerivedData.cpp
@@ -21,7 +21,7 @@ const TCHAR* FChaosDerivedDataCooker::GetPluginName() const
 
 const TCHAR* FChaosDerivedDataCooker::GetVersionString() const
 {
-	return TEXT("6871810764D045EDA573A3173329490D");
+	return TEXT("E138DCCB71AD4DBCA7AA0E469571F332");
 }
 
 FString FChaosDerivedDataCooker::GetPluginSpecificCacheKeySuffix() const
diff --git a/Engine/Source/Runtime/Experimental/Chaos/Public/Chaos/AABBTree.h b/Engine/Source/Runtime/Experimental/Chaos/Public/Chaos/AABBTree.h
--- a/Engine/Source/Runtime/Experimental/Chaos/Public/Chaos/AABBTree.h
+++ b/Engine/Source/Runtime/Experimental/Chaos/Public/Chaos/AABBTree.h
@@ -498,27 +498,42 @@ public:
 
 	virtual void Serialize(FChaosArchive& Ar) override
 	{
+		Ar.UsingCustomVersion(FExternalPhysicsCustomObjectVersion::GUID);
+
 		TBox<FReal, 3>::SerializeAsAABB(Ar, FullBounds);
 		Ar << Nodes;
 		Ar << Leaves;
 		Ar << DirtyElements;
 		Ar << GlobalPayloads;
 
-		TArray<TPayloadType> Payloads;
-		if (!Ar.IsLoading())
+		bool bSerializePayloadToInfo = !bMutable;
+		if (Ar.CustomVer(FExternalPhysicsCustomObjectVersion::GUID) >= FExternalPhysicsCustomObjectVersion::ImmutableAABBTree)
 		{
-			PayloadToInfo.GenerateKeyArray(Payloads);
+			Ar << bSerializePayloadToInfo;
 		}
-		Ar << Payloads;
-
-		for (auto Payload : Payloads)
+		else
 		{
-			auto& Info = PayloadToInfo.FindOrAdd(Payload);
-			Ar << Info;
+			bSerializePayloadToInfo = true;
 		}
-		if (!bMutable)
+
+		if (bSerializePayloadToInfo)
 		{
-			PayloadToInfo.Empty();	//todo: avoid this entirely
+			TArray<TPayloadType> Payloads;
+			if (!Ar.IsLoading())
+			{
+				PayloadToInfo.GenerateKeyArray(Payloads);
+			}
+			Ar << Payloads;
+
+			for (auto Payload : Payloads)
+			{
+				auto& Info = PayloadToInfo.FindOrAdd(Payload);
+				Ar << Info;
+			}
+			if (!bMutable)	//if immutable empty this even if we had to serialize it in for backwards compat
+			{
+				PayloadToInfo.Empty();
+			}
 		}
 
 		Ar << MaxChildrenInLeaf;
diff --git a/Engine/Source/Runtime/Experimental/Chaos/Public/Chaos/TriangleMeshImplicitObject.h b/Engine/Source/Runtime/Experimental/Chaos/Public/Chaos/TriangleMeshImplicitObject.h
--- a/Engine/Source/Runtime/Experimental/Chaos/Public/Chaos/TriangleMeshImplicitObject.h
+++ b/Engine/Source/Runtime/Experimental/Chaos/Public/Chaos/TriangleMeshImplicitObject.h
@@ -128,7 +128,7 @@ namespace Chaos
 		TArray<uint16> MaterialIndices;
 
 		//using BVHType = TBoundingVolume<int32, T, 3>;
-		using BVHType = TAABBTree<int32, TAABBTreeLeafArray<int32, FReal>, FReal>;
+		using BVHType = TAABBTree<int32, TAABBTreeLeafArray<int32, FReal>, FReal, /*bMutable=*/false>;
 
 		template<typename InStorageType, typename InRealType>
 		friend struct FBvEntry;
