Chaos : Collisions

- Added THeightField to TBox collisions.

diff --git a/Engine/Source/Runtime/Experimental/Chaos/Private/Chaos/CollisionResolution.cpp b/Engine/Source/Runtime/Experimental/Chaos/Private/Chaos/CollisionResolution.cpp
--- a/Engine/Source/Runtime/Experimental/Chaos/Private/Chaos/CollisionResolution.cpp
+++ b/Engine/Source/Runtime/Experimental/Chaos/Private/Chaos/CollisionResolution.cpp
@@ -8,6 +8,7 @@
 #include "Chaos/CollisionResolutionUtil.h"
 #include "Chaos/Convex.h"
 #include "Chaos/Defines.h"
+#include "Chaos/HeightField.h"
 #include "Chaos/ImplicitObjectScaled.h"
 #include "Chaos/ImplicitObjectTransformed.h"
 #include "Chaos/Levelset.h"
@@ -113,6 +114,53 @@ namespace Chaos
 			}
 		}
 
+		//
+		// Box - HeightField
+		//
+
+		template <typename T, int d>
+		TContactPoint<T> BoxHeightFieldContactPoint(const TAABB<T, d>& A, const TRigidTransform<T, d>& ATransform, const THeightField<T>& B, const TRigidTransform<T, d>& BTransform, const T Thickness)
+		{
+			TContactPoint<T> Contact;
+			B.GJKContactPoint(TBox<float, 3>(A), ATransform.GetRelativeTransform(BTransform), Thickness, Contact.Location, Contact.Normal, Contact.Phi);
+			Contact.Location = BTransform.TransformPosition(Contact.Location);
+			Contact.Normal = BTransform.TransformVector(Contact.Normal);
+			return Contact;
+		}
+
+		template <typename T, int d>
+		void UpdateBoxHeightFieldConstraint(const TAABB<T, d>& A, const TRigidTransform<T, d>& ATransform, const THeightField<T>& B, const TRigidTransform<T, d>& BTransform, const T Thickness, TRigidBodyPointContactConstraint<T, d>& Constraint)
+		{
+			UpdateContactPoint(Constraint.Manifold, BoxHeightFieldContactPoint(A, ATransform, B, BTransform, Thickness));
+		}
+
+		template<typename T, int d>
+		void ConstructBoxHeightFieldConstraints(TGeometryParticleHandle<T, d>* Particle0, TGeometryParticleHandle<T, d>* Particle1, const FImplicitObject* Implicit0, const FImplicitObject* Implicit1, const TRigidTransform<T, d>& Transform0, const TRigidTransform<T, d>& Transform1, const T Thickness, FCollisionConstraintsArray& NewConstraints)
+		{
+
+			const TBox<T, d> * Object0 = Implicit0->template GetObject<const TBox<T, d> >();
+			const THeightField<T> * Object1 = Implicit1->template GetObject<const THeightField<T> >();
+			if (ensure(Object0 && Object1))
+			{
+				TRigidTransform<T, d> ParticleImplicit0TM = Transform0.GetRelativeTransform(Collisions::GetTransform(Particle0));
+				TRigidTransform<T, d> ParticleImplicit1TM = Transform1.GetRelativeTransform(Collisions::GetTransform(Particle1));
+				FRigidBodyPointContactConstraint* Constraint = new FRigidBodyPointContactConstraint(Particle0, Implicit0, ParticleImplicit0TM, Particle1, Implicit1, ParticleImplicit1TM);
+
+				UpdateBoxHeightFieldConstraint(Object0->BoundingBox(), Transform0, *Object1, Transform1, Thickness, *Constraint);
+
+				if (Constraint->GetPhi() < Thickness)
+				{
+					NewConstraints.Add(Constraint);
+				}
+				else
+				{
+					delete Constraint;
+				}
+			}
+		}
+
+
+
 		//
 		// Box-Plane
 		//
@@ -815,7 +863,10 @@ namespace Chaos
 
 			//
 			// @todo(chaos): Collision Constraints (CollisionMap)
-			//    Modify Construct() and Update() use a CollisionMap indexed on EImplicitObjectType, instead of the if/else chain
+			//    Modify Construct() and Update() use a CollisionMap indexed on 
+			//    EImplicitObjectType, instead of the if/else chain. Also, remove 
+			//    the blocks with the ensure(false), they are just for validation 
+			//    after the recent change. 
 			//
 			EImplicitObjectType Implicit0Type = Implicit0.GetType();
 			EImplicitObjectType Implicit1Type = Implicit1.GetType();
@@ -844,6 +895,10 @@ namespace Chaos
 			{
 				UpdateBoxBoxConstraint(Implicit0.template GetObject<TBox<T, d>>()->GetAABB(), Transform0, Implicit1.template GetObject<TBox<T, d>>()->GetAABB(), Transform1, Thickness, *ConstraintBase.template As<TRigidBodyPointContactConstraint<T,d>>());
 			}
+			else if (Implicit0Type == TBox<T, d>::StaticType() && Implicit1Type == THeightField<T>::StaticType())
+			{
+				UpdateBoxHeightFieldConstraint(Implicit0.template GetObject<TBox<T, d>>()->GetAABB(), Transform0, *Implicit1.template GetObject< THeightField<T> >(), Transform1, Thickness, *ConstraintBase.template As<TRigidBodyPointContactConstraint<T, d>>());
+			}
 			else if (Implicit0Type == TSphere<T, d>::StaticType() && Implicit1Type == TSphere<T, d>::StaticType())
 			{
 				UpdateSphereSphereConstraint(*Implicit0.template GetObject<TSphere<T, d>>(), Transform0, *Implicit1.template GetObject<TSphere<T, d>>(), Transform1, Thickness, *ConstraintBase.template As<TRigidBodyPointContactConstraint<T,d>>());
@@ -872,6 +927,17 @@ namespace Chaos
 			{
 				UpdateCapsuleBoxConstraint(*Implicit0.template GetObject<TCapsule<T>>(), Transform0, Implicit1.template GetObject<TBox<T, d>>()->GetAABB(), Transform1, Thickness, *ConstraintBase.template As<TRigidBodyPointContactConstraint<T,d>>());
 			}
+			else if (Implicit0Type == THeightField<T>::StaticType() && Implicit1Type == TBox<T, d>::StaticType())
+			{
+				TRigidBodyPointContactConstraint<T, d>& Constraint = *ConstraintBase.template As<TRigidBodyPointContactConstraint<T, d> >();
+				TRigidBodyPointContactConstraint<T, d> TmpConstraint = Constraint;
+				UpdateBoxHeightFieldConstraint(Implicit1.template GetObject<TBox<T, d>>()->GetAABB(), Transform1, *Implicit0.template GetObject<THeightField<T>>(), Transform0, Thickness, TmpConstraint);
+				if (TmpConstraint.GetPhi() < Constraint.GetPhi())
+				{
+					Constraint = TmpConstraint;
+					Constraint.SetNormal(-Constraint.GetNormal());
+				}
+			}
 			else if (Implicit0Type == TPlane<T, d>::StaticType() && Implicit1Type == TBox<T, d>::StaticType())
 			{
 				TRigidBodyPointContactConstraint<T, d>& Constraint = *ConstraintBase.template As<TRigidBodyPointContactConstraint<T, d> >();
@@ -1006,9 +1072,13 @@ namespace Chaos
 			{
 				ConstructBoxBoxConstraints(Particle0, Particle1, Implicit0, Implicit1, Transform0, Transform1, Thickness, NewConstraints);
 			}
-			else if (Implicit0Type == TSphere<T, d>::StaticType() && Implicit1Type == TSphere<T, d>::StaticType())
+			else if (Implicit0Type == TBox<T, d>::StaticType() && Implicit1Type == THeightField<T>::StaticType())
 			{
-				ConstructSphereSphereConstraints(Particle0, Particle1, Implicit0, Implicit1, Transform0, Transform1, Thickness, NewConstraints);
+				ConstructBoxHeightFieldConstraints(Particle0, Particle1, Implicit0, Implicit1, Transform0, Transform1, Thickness, NewConstraints);
+			}
+			else if (Implicit0Type == THeightField<T>::StaticType() && Implicit1Type == TBox<T,d>::StaticType())
+			{
+				ConstructBoxHeightFieldConstraints(Particle1, Particle0, Implicit1, Implicit0, Transform1, Transform0, Thickness, NewConstraints);
 			}
 			else if (Implicit0Type == TBox<T, d>::StaticType() && Implicit1Type == TPlane<T, d>::StaticType())
 			{
@@ -1018,6 +1088,10 @@ namespace Chaos
 			{
 				ConstructBoxPlaneConstraints(Particle1, Particle0, Implicit1, Implicit0, Transform1, Transform0, Thickness, NewConstraints);
 			}
+			else if (Implicit0Type == TSphere<T, d>::StaticType() && Implicit1Type == TSphere<T, d>::StaticType())
+			{
+				ConstructSphereSphereConstraints(Particle0, Particle1, Implicit0, Implicit1, Transform0, Transform1, Thickness, NewConstraints);
+			}
 			else if (Implicit0Type == TSphere<T, d>::StaticType() && Implicit1Type == TPlane<T, d>::StaticType())
 			{
 				ConstructSpherePlaneConstraints(Particle0, Particle1, Implicit0, Implicit1, Transform0, Transform1, Thickness, NewConstraints);
@@ -1071,6 +1145,9 @@ namespace Chaos
 		template bool UpdateBoxPlaneConstraint<float, 3>(const TAABB<float, 3>& Box, const TRigidTransform<float, 3>& BoxTransform, const TPlane<float, 3>& Plane, const TRigidTransform<float, 3>& PlaneTransform, const float Thickness, TRigidBodyPointContactConstraint<float, 3>& Constraint);
 		template void ConstructBoxPlaneConstraints<float, 3>(TGeometryParticleHandle<float, 3>* Particle0, TGeometryParticleHandle<float, 3>* Particle1, const FImplicitObject* Implicit0, const FImplicitObject* Implicit1, const TRigidTransform<float, 3>& Transform0, const TRigidTransform<float, 3>& Transform1, const float Thickness, FCollisionConstraintsArray& NewConstraints);
 
+		template void UpdateBoxHeightFieldConstraint(const TAABB<float, 3>& A, const TRigidTransform<float, 3>& ATransform, const THeightField<float>& B, const TRigidTransform<float, 3>& BTransform, const float Thickness, TRigidBodyPointContactConstraint<float, 3>& Constraint);
+		template void ConstructBoxHeightFieldConstraints(TGeometryParticleHandle<float, 3>* Particle0, TGeometryParticleHandle<float, 3>* Particle1, const FImplicitObject* Implicit0, const FImplicitObject* Implicit1, const TRigidTransform<float, 3>& Transform0, const TRigidTransform<float, 3>& Transform1, const float Thickness, FCollisionConstraintsArray& NewConstraints);
+
 		template void UpdateSphereSphereConstraint<float, 3>(const TSphere<float, 3>& Sphere1, const TRigidTransform<float, 3>& Sphere1Transform, const TSphere<float, 3>& Sphere2, const TRigidTransform<float, 3>& Sphere2Transform, const float Thickness, TRigidBodyPointContactConstraint<float, 3>& Constraint);
 		template void ConstructSphereSphereConstraints<float, 3>(TGeometryParticleHandle<float, 3>* Particle0, TGeometryParticleHandle<float, 3>* Particle1, const FImplicitObject* Implicit0, const FImplicitObject* Implicit1, const TRigidTransform<float, 3>& Transform0, const TRigidTransform<float, 3>& Transform1, const float Thickness, FCollisionConstraintsArray& NewConstraints);
 
diff --git a/Engine/Source/Runtime/Experimental/Chaos/Private/Chaos/HeightField.cpp b/Engine/Source/Runtime/Experimental/Chaos/Private/Chaos/HeightField.cpp
--- a/Engine/Source/Runtime/Experimental/Chaos/Private/Chaos/HeightField.cpp
+++ b/Engine/Source/Runtime/Experimental/Chaos/Private/Chaos/HeightField.cpp
@@ -1152,6 +1152,90 @@ namespace Chaos
 	}
 
 
+	template <typename T>
+	template <typename GeomType>
+	bool THeightField<T>::GJKContactPointImp(const GeomType& QueryGeom, const TRigidTransform<T, 3>& QueryTM, const T Thickness,
+		TVector<T, 3>& ContactLocation, TVector<T, 3>& ContactNormal, T& ContactPhi) const
+	{
+		auto OverlapTriangle = [&](const TVector<T, 3>& A, const TVector<T, 3>& B, const TVector<T, 3>& C,
+			TVector<T, 3>& LocalContactLocation, TVector<T, 3>& LocalContactNormal, T& LocalContactPhi) -> bool
+		{
+			const TVector<T, 3> AB = B - A;
+			const TVector<T, 3> AC = C - A;
+
+			const TVector<T, 3> Offset = TVector<T, 3>::CrossProduct(AB, AC);
+
+			TTriangle<T> TriangleConvex(A, B, C);
+
+			T Penetration;
+			TVec3<T> ClosestA, ClosestB, Normal;
+			if (GJKPenetration(TriangleConvex, QueryGeom, QueryTM, Penetration, ClosestA, ClosestB, Normal, (T)0))
+			{
+				TVec3<T> TestVector = QueryTM.InverseTransformVector(Normal);
+
+				LocalContactLocation = ClosestB;
+				LocalContactNormal = Normal; 
+				LocalContactPhi = -Penetration;
+				return true;
+			}
+
+			return LocalContactPhi<0.f;
+		};
+
+		bool bResult = false;
+		TAABB<T, 3> QueryBounds = QueryGeom.BoundingBox();
+		QueryBounds.Thicken(Thickness);
+		QueryBounds = QueryBounds.TransformedAABB(QueryTM);
+
+		FBounds2D FlatQueryBounds;
+		FlatQueryBounds.Min = TVector<T, 2>(QueryBounds.Min()[0], QueryBounds.Min()[1]);
+		FlatQueryBounds.Max = TVector<T, 2>(QueryBounds.Max()[0], QueryBounds.Max()[1]);
+
+		TArray<TVector<int32, 2>> Intersections;
+		TVector<T, 3> Points[4];
+
+		GetGridIntersections(FlatQueryBounds, Intersections);
+
+		T LocalContactPhi = FLT_MAX;
+		TVector<T, 3> LocalContactLocation, LocalContactNormal;
+		for (const TVector<int32, 2>& Cell : Intersections)
+		{
+			const int32 SingleIndex = Cell[1] * GeomData.NumCols + Cell[0];
+			GeomData.GetPointsScaled(SingleIndex, Points);
+
+			if (OverlapTriangle(Points[0], Points[1], Points[3], LocalContactLocation, LocalContactNormal, LocalContactPhi))
+			{
+				if (LocalContactPhi < ContactPhi)
+				{
+					ContactPhi = LocalContactPhi;
+					ContactLocation = LocalContactLocation;
+					ContactNormal = LocalContactNormal;
+				}
+			}
+
+			if (OverlapTriangle(Points[0], Points[3], Points[2], LocalContactLocation, LocalContactNormal, LocalContactPhi))
+			{
+				if (LocalContactPhi < ContactPhi)
+				{
+					ContactPhi = LocalContactPhi;
+					ContactLocation = LocalContactLocation;
+					ContactNormal = LocalContactNormal;
+				}
+			}
+		}
+
+		if (ContactPhi < 0)
+			return true;
+		return false;
+	}
+
+	template <typename T>
+	bool THeightField<T>::GJKContactPoint(const TBox<T, 3>& QueryGeom, const TRigidTransform<T, 3>& QueryTM, const T Thickness, TVector<T, 3>& ContactLocation, TVector<T, 3>& ContactNormal, T& ContactPhi) const
+	{
+		return GJKContactPointImp(QueryGeom, QueryTM, Thickness, ContactLocation, ContactNormal, ContactPhi);
+	}
+
+
 	template <typename T>
 	template <typename QueryGeomType>
 	bool THeightField<T>::OverlapGeomImp(const QueryGeomType& QueryGeom, const TRigidTransform<T, 3>& QueryTM, const T Thickness) const
diff --git a/Engine/Source/Runtime/Experimental/Chaos/Public/Chaos/CollisionResolution.h b/Engine/Source/Runtime/Experimental/Chaos/Public/Chaos/CollisionResolution.h
--- a/Engine/Source/Runtime/Experimental/Chaos/Public/Chaos/CollisionResolution.h
+++ b/Engine/Source/Runtime/Experimental/Chaos/Public/Chaos/CollisionResolution.h
@@ -18,6 +18,9 @@ namespace Chaos
 	template <typename T, int d>
 	class TPlane;
 
+	template <typename T>
+	class THeightField;
+
 	class FImplicitObject;
 
 	template<typename T, int d>
@@ -40,6 +43,16 @@ namespace Chaos
 			void CHAOS_API ConstructBoxBoxConstraints(TGeometryParticleHandle<T, d>* Particle0, TGeometryParticleHandle<T, d>* Particle1, const FImplicitObject* Implicit0, const FImplicitObject* Implicit1, const TRigidTransform<T, d>& Transform0, const TRigidTransform<T, d>& Transform1, const T Thickness, FCollisionConstraintsArray& NewConstraints);
 
 
+			//
+			// Box-HeightField
+			//
+
+			template <typename T, int d>
+			void UpdateBoxHeightFieldConstraint(const TAABB<T, d>& A, const TRigidTransform<T, d>& ATransform, const THeightField<T>& B, const TRigidTransform<T, d>& BTransform, const T Thickness, TRigidBodyPointContactConstraint<T, d>& Constraint);
+
+			template<typename T, int d>
+			void ConstructBoxHeightFieldConstraints(TGeometryParticleHandle<T, d>* Particle0, TGeometryParticleHandle<T, d>* Particle1, const FImplicitObject* Implicit0, const FImplicitObject* Implicit1, const TRigidTransform<T, d>& Transform0, const TRigidTransform<T, d>& Transform1, const T Thickness, FCollisionConstraintsArray& NewConstraints);
+
 			//
 			// Box-Plane
 			//
@@ -51,6 +64,7 @@ namespace Chaos
 			void CHAOS_API ConstructBoxPlaneConstraints(TGeometryParticleHandle<T, d>* Particle0, TGeometryParticleHandle<T, d>* Particle1, const FImplicitObject* Implicit0, const FImplicitObject* Implicit1, const TRigidTransform<T, d>& Transform0, const TRigidTransform<T, d>& Transform1, const T Thickness, FCollisionConstraintsArray& NewConstraints);
 
 
+
 			//
 			// Sphere-Sphere
 			//
diff --git a/Engine/Source/Runtime/Experimental/Chaos/Public/Chaos/HeightField.h b/Engine/Source/Runtime/Experimental/Chaos/Public/Chaos/HeightField.h
--- a/Engine/Source/Runtime/Experimental/Chaos/Public/Chaos/HeightField.h
+++ b/Engine/Source/Runtime/Experimental/Chaos/Public/Chaos/HeightField.h
@@ -71,9 +71,15 @@ namespace Chaos
 		bool SweepGeom(const TImplicitObjectScaled<TCapsule<T>>& QueryGeom, const TRigidTransform<T, 3>& StartTM, const TVector<T, 3>& Dir, const T Length, T& OutTime, TVector<T, 3>& OutPosition, TVector<T, 3>& OutNormal, int32& OutFaceIndex, const T Thickness = 0, bool bComputeMTD = false) const;
 		bool SweepGeom(const TImplicitObjectScaled<FConvex>& QueryGeom, const TRigidTransform<T, 3>& StartTM, const TVector<T, 3>& Dir, const T Length, T& OutTime, TVector<T, 3>& OutPosition, TVector<T, 3>& OutNormal, int32& OutFaceIndex, const T Thickness = 0, bool bComputeMTD = false) const;
 
+		bool GJKContactPoint(const TBox<T, 3>& QueryGeom, const TRigidTransform<T, 3>& QueryTM, const T Thickness, TVector<T, 3>& ContactLocation, TVector<T, 3>& ContactNormal, T& ContactPhi) const;
+
+
 		virtual int32 FindMostOpposingFace(const TVector<T, 3>& Position, const TVector<T, 3>& UnitDir, int32 HintFaceIndex, T SearchDist) const override;
 		virtual TVector<T, 3> FindGeometryOpposingNormal(const TVector<T, 3>& DenormDir, int32 FaceIndex, const TVector<T, 3>& OriginalNormal) const override;
 
+
+
+
 		virtual uint16 GetMaterialIndex(uint32 HintIndex) const override
 		{
 			ensure(GeomData.MaterialIndices.Num() > 0);
@@ -482,5 +488,8 @@ namespace Chaos
 		template <typename QueryGeomType>
 		bool SweepGeomImp(const QueryGeomType& QueryGeom, const TRigidTransform<T, 3>& StartTM, const TVector<T, 3>& Dir, const T Length, T& OutTime, TVector<T, 3>& OutPosition, TVector<T, 3>& OutNormal, int32& OutFaceIndex, const T Thickness, bool bComputeMTD) const;
 
+		template <typename GeomType>
+		bool GJKContactPointImp(const GeomType& QueryGeom, const TRigidTransform<T, 3>& QueryTM, const T Thickness, TVector<T, 3>& ContactLocation, TVector<T, 3>& ContactNormal, T& ContactPhi) const;
+
 	};
 }
