Chaos : Collisions

- THightField to TSphere collisions

diff --git a/Engine/Source/Runtime/Experimental/Chaos/Private/Chaos/CollisionResolution.cpp b/Engine/Source/Runtime/Experimental/Chaos/Private/Chaos/CollisionResolution.cpp
--- a/Engine/Source/Runtime/Experimental/Chaos/Private/Chaos/CollisionResolution.cpp
+++ b/Engine/Source/Runtime/Experimental/Chaos/Private/Chaos/CollisionResolution.cpp
@@ -306,6 +306,51 @@ namespace Chaos
 			}
 		}
 
+		//
+		// Sphere - HeightField
+		//
+
+		template <typename T, int d>
+		TContactPoint<T> SphereHeightFieldContactPoint(const TSphere<T, d>& A, const TRigidTransform<T, d>& ATransform, const THeightField<T>& B, const TRigidTransform<T, d>& BTransform, const T Thickness)
+		{
+			TContactPoint<T> Contact;
+			B.GJKContactPoint(TSphere<float, 3>(A), ATransform.GetRelativeTransform(BTransform), Thickness, Contact.Location, Contact.Normal, Contact.Phi);
+			Contact.Location = BTransform.TransformPosition(Contact.Location);
+			Contact.Normal = BTransform.TransformVector(Contact.Normal);
+			return Contact;
+		}
+
+		template <typename T, int d>
+		void UpdateSphereHeightFieldConstraint(const TSphere<T, d>& A, const TRigidTransform<T, d>& ATransform, const THeightField<T>& B, const TRigidTransform<T, d>& BTransform, const T Thickness, TRigidBodyPointContactConstraint<T, d>& Constraint)
+		{
+			UpdateContactPoint(Constraint.Manifold, SphereHeightFieldContactPoint(A, ATransform, B, BTransform, Thickness));
+		}
+
+		template<typename T, int d>
+		void ConstructSphereHeightFieldConstraints(TGeometryParticleHandle<T, d>* Particle0, TGeometryParticleHandle<T, d>* Particle1, const FImplicitObject* Implicit0, const FImplicitObject* Implicit1, const TRigidTransform<T, d>& Transform0, const TRigidTransform<T, d>& Transform1, const T Thickness, FCollisionConstraintsArray& NewConstraints)
+		{
+
+			const TSphere<T, d> * Object0 = Implicit0->template GetObject<const TSphere<T, d> >();
+			const THeightField<T> * Object1 = Implicit1->template GetObject<const THeightField<T> >();
+			if (ensure(Object0 && Object1))
+			{
+				TRigidTransform<T, d> ParticleImplicit0TM = Transform0.GetRelativeTransform(Collisions::GetTransform(Particle0));
+				TRigidTransform<T, d> ParticleImplicit1TM = Transform1.GetRelativeTransform(Collisions::GetTransform(Particle1));
+				FRigidBodyPointContactConstraint* Constraint = new FRigidBodyPointContactConstraint(Particle0, Implicit0, ParticleImplicit0TM, Particle1, Implicit1, ParticleImplicit1TM);
+
+				UpdateSphereHeightFieldConstraint(*Object0, Transform0, *Object1, Transform1, Thickness, *Constraint);
+
+				if (Constraint->GetPhi() < Thickness)
+				{
+					NewConstraints.Add(Constraint);
+				}
+				else
+				{
+					delete Constraint;
+				}
+			}
+		}
+
 		//
 		//  Sphere-Plane
 		//
@@ -903,6 +948,10 @@ namespace Chaos
 			{
 				UpdateSphereSphereConstraint(*Implicit0.template GetObject<TSphere<T, d>>(), Transform0, *Implicit1.template GetObject<TSphere<T, d>>(), Transform1, Thickness, *ConstraintBase.template As<TRigidBodyPointContactConstraint<T,d>>());
 			}
+			else if (Implicit0Type == TSphere<T, d>::StaticType() && Implicit1Type == THeightField<T>::StaticType())
+			{
+				UpdateSphereHeightFieldConstraint(*Implicit0.template GetObject<TSphere<T, d>>(), Transform0, *Implicit1.template GetObject< THeightField<T> >(), Transform1, Thickness, *ConstraintBase.template As<TRigidBodyPointContactConstraint<T, d>>());
+			}
 			else if (Implicit0Type == TBox<T, d>::StaticType() && Implicit1Type == TPlane<T, d>::StaticType())
 			{
 				UpdateBoxPlaneConstraint(Implicit0.template GetObject<TBox<T, d>>()->GetAABB(), Transform0, *Implicit1.template GetObject<TPlane<T, d>>(), Transform1, Thickness, *ConstraintBase.template As<TRigidBodyPointContactConstraint<T,d>>());
@@ -929,6 +978,8 @@ namespace Chaos
 			}
 			else if (Implicit0Type == THeightField<T>::StaticType() && Implicit1Type == TBox<T, d>::StaticType())
 			{
+				// @todo(chaos) : Collision Constraints
+				//                Does this really work? 
 				TRigidBodyPointContactConstraint<T, d>& Constraint = *ConstraintBase.template As<TRigidBodyPointContactConstraint<T, d> >();
 				TRigidBodyPointContactConstraint<T, d> TmpConstraint = Constraint;
 				UpdateBoxHeightFieldConstraint(Implicit1.template GetObject<TBox<T, d>>()->GetAABB(), Transform1, *Implicit0.template GetObject<THeightField<T>>(), Transform0, Thickness, TmpConstraint);
@@ -949,6 +1000,19 @@ namespace Chaos
 					Constraint.SetNormal(-Constraint.GetNormal());
 				}
 			}
+			else if (Implicit0Type == THeightField<T>::StaticType() && Implicit1Type == TSphere<T, d>::StaticType())
+			{
+				// @todo(chaos) : Collision Constraints
+				//                Does this really work? 
+				TRigidBodyPointContactConstraint<T, d>& Constraint = *ConstraintBase.template As<TRigidBodyPointContactConstraint<T, d> >();
+				TRigidBodyPointContactConstraint<T, d> TmpConstraint = Constraint;
+				UpdateSphereHeightFieldConstraint(*Implicit1.template GetObject<TSphere<T, d>>(), Transform1, *Implicit0.template GetObject<THeightField<T>>(), Transform0, Thickness, TmpConstraint);
+				if (TmpConstraint.GetPhi() < Constraint.GetPhi())
+				{
+					Constraint = TmpConstraint;
+					Constraint.SetNormal(-Constraint.GetNormal());
+				}
+			}
 			else if (Implicit0Type == TPlane<T, d>::StaticType() && Implicit1Type == TSphere<T, d>::StaticType())
 			{
 				TRigidBodyPointContactConstraint<T, d>& Constraint = *ConstraintBase.template As<TRigidBodyPointContactConstraint<T, d> >();
@@ -1092,6 +1156,14 @@ namespace Chaos
 			{
 				ConstructSphereSphereConstraints(Particle0, Particle1, Implicit0, Implicit1, Transform0, Transform1, Thickness, NewConstraints);
 			}
+			else if (Implicit0Type == TSphere<T, d>::StaticType() && Implicit1Type == THeightField<T>::StaticType())
+			{
+				ConstructSphereHeightFieldConstraints(Particle0, Particle1, Implicit0, Implicit1, Transform0, Transform1, Thickness, NewConstraints);
+			}
+			else if (Implicit0Type == THeightField<T>::StaticType() && Implicit1Type == TSphere<T, d>::StaticType())
+			{
+				ConstructSphereHeightFieldConstraints(Particle1, Particle0, Implicit1, Implicit0, Transform1, Transform0, Thickness, NewConstraints);
+			}
 			else if (Implicit0Type == TSphere<T, d>::StaticType() && Implicit1Type == TPlane<T, d>::StaticType())
 			{
 				ConstructSpherePlaneConstraints(Particle0, Particle1, Implicit0, Implicit1, Transform0, Transform1, Thickness, NewConstraints);
@@ -1151,6 +1223,9 @@ namespace Chaos
 		template void UpdateSphereSphereConstraint<float, 3>(const TSphere<float, 3>& Sphere1, const TRigidTransform<float, 3>& Sphere1Transform, const TSphere<float, 3>& Sphere2, const TRigidTransform<float, 3>& Sphere2Transform, const float Thickness, TRigidBodyPointContactConstraint<float, 3>& Constraint);
 		template void ConstructSphereSphereConstraints<float, 3>(TGeometryParticleHandle<float, 3>* Particle0, TGeometryParticleHandle<float, 3>* Particle1, const FImplicitObject* Implicit0, const FImplicitObject* Implicit1, const TRigidTransform<float, 3>& Transform0, const TRigidTransform<float, 3>& Transform1, const float Thickness, FCollisionConstraintsArray& NewConstraints);
 
+		template void UpdateSphereHeightFieldConstraint(const TSphere<float, 3>& A, const TRigidTransform<float, 3>& ATransform, const THeightField<float>& B, const TRigidTransform<float, 3>& BTransform, const float Thickness, TRigidBodyPointContactConstraint<float, 3>& Constraint);
+		template void ConstructSphereHeightFieldConstraints(TGeometryParticleHandle<float, 3>* Particle0, TGeometryParticleHandle<float, 3>* Particle1, const FImplicitObject* Implicit0, const FImplicitObject* Implicit1, const TRigidTransform<float, 3>& Transform0, const TRigidTransform<float, 3>& Transform1, const float Thickness, FCollisionConstraintsArray& NewConstraints);
+
 		template void UpdateSpherePlaneConstraint<float, 3>(const TSphere<float, 3>& Sphere, const TRigidTransform<float, 3>& SphereTransform, const TPlane<float, 3>& Plane, const TRigidTransform<float, 3>& PlaneTransform, const float Thickness, TRigidBodyPointContactConstraint<float, 3>& Constraint);
 		template void ConstructSpherePlaneConstraints<float, 3>(TGeometryParticleHandle<float, 3>* Particle0, TGeometryParticleHandle<float, 3>* Particle1, const FImplicitObject* Implicit0, const FImplicitObject* Implicit1, const TRigidTransform<float, 3>& Transform0, const TRigidTransform<float, 3>& Transform1, const float Thickness, FCollisionConstraintsArray& NewConstraints);
 
diff --git a/Engine/Source/Runtime/Experimental/Chaos/Private/Chaos/HeightField.cpp b/Engine/Source/Runtime/Experimental/Chaos/Private/Chaos/HeightField.cpp
--- a/Engine/Source/Runtime/Experimental/Chaos/Private/Chaos/HeightField.cpp
+++ b/Engine/Source/Runtime/Experimental/Chaos/Private/Chaos/HeightField.cpp
@@ -1235,6 +1235,12 @@ namespace Chaos
 		return GJKContactPointImp(QueryGeom, QueryTM, Thickness, ContactLocation, ContactNormal, ContactPhi);
 	}
 
+	template <typename T>
+	bool THeightField<T>::GJKContactPoint(const TSphere<T, 3>& QueryGeom, const TRigidTransform<T, 3>& QueryTM, const T Thickness, TVector<T, 3>& ContactLocation, TVector<T, 3>& ContactNormal, T& ContactPhi) const
+	{
+		return GJKContactPointImp(QueryGeom, QueryTM, Thickness, ContactLocation, ContactNormal, ContactPhi);
+	}
+
 
 	template <typename T>
 	template <typename QueryGeomType>
diff --git a/Engine/Source/Runtime/Experimental/Chaos/Public/Chaos/CollisionResolution.h b/Engine/Source/Runtime/Experimental/Chaos/Public/Chaos/CollisionResolution.h
--- a/Engine/Source/Runtime/Experimental/Chaos/Public/Chaos/CollisionResolution.h
+++ b/Engine/Source/Runtime/Experimental/Chaos/Public/Chaos/CollisionResolution.h
@@ -75,6 +75,16 @@ namespace Chaos
 			template<typename T, int d>
 			void CHAOS_API ConstructSphereSphereConstraints(TGeometryParticleHandle<T, d>* Particle0, TGeometryParticleHandle<T, d>* Particle1, const FImplicitObject* Implicit0, const FImplicitObject* Implicit1, const TRigidTransform<T, d>& Transform0, const TRigidTransform<T, d>& Transform1, const T Thickness, FCollisionConstraintsArray& NewConstraints);
 
+			//
+			// Sphere-HeightField
+			//
+
+			template <typename T, int d>
+			void UpdateSphereHeightFieldConstraint(const TSphere<T, d>& A, const TRigidTransform<T, d>& ATransform, const THeightField<T>& B, const TRigidTransform<T, d>& BTransform, const T Thickness, TRigidBodyPointContactConstraint<T, d>& Constraint);
+
+			template<typename T, int d>
+			void ConstructSphereHeightFieldConstraints(TGeometryParticleHandle<T, d>* Particle0, TGeometryParticleHandle<T, d>* Particle1, const FImplicitObject* Implicit0, const FImplicitObject* Implicit1, const TRigidTransform<T, d>& Transform0, const TRigidTransform<T, d>& Transform1, const T Thickness, FCollisionConstraintsArray& NewConstraints);
+				
 			//
 			// Sphere-Plane
 			//
diff --git a/Engine/Source/Runtime/Experimental/Chaos/Public/Chaos/HeightField.h b/Engine/Source/Runtime/Experimental/Chaos/Public/Chaos/HeightField.h
--- a/Engine/Source/Runtime/Experimental/Chaos/Public/Chaos/HeightField.h
+++ b/Engine/Source/Runtime/Experimental/Chaos/Public/Chaos/HeightField.h
@@ -72,6 +72,7 @@ namespace Chaos
 		bool SweepGeom(const TImplicitObjectScaled<FConvex>& QueryGeom, const TRigidTransform<T, 3>& StartTM, const TVector<T, 3>& Dir, const T Length, T& OutTime, TVector<T, 3>& OutPosition, TVector<T, 3>& OutNormal, int32& OutFaceIndex, const T Thickness = 0, bool bComputeMTD = false) const;
 
 		bool GJKContactPoint(const TBox<T, 3>& QueryGeom, const TRigidTransform<T, 3>& QueryTM, const T Thickness, TVector<T, 3>& ContactLocation, TVector<T, 3>& ContactNormal, T& ContactPhi) const;
+		bool GJKContactPoint(const TSphere<T, 3>& QueryGeom, const TRigidTransform<T, 3>& QueryTM, const T Thickness, TVector<T, 3>& ContactLocation, TVector<T, 3>& ContactNormal, T& ContactPhi) const;
 
 
 		virtual int32 FindMostOpposingFace(const TVector<T, 3>& Position, const TVector<T, 3>& UnitDir, int32 HintFaceIndex, T SearchDist) const override;
