Chaos : Collisions

- THeightField to TCapsule

diff --git a/Engine/Source/Runtime/Experimental/Chaos/Private/Chaos/CollisionResolution.cpp b/Engine/Source/Runtime/Experimental/Chaos/Private/Chaos/CollisionResolution.cpp
--- a/Engine/Source/Runtime/Experimental/Chaos/Private/Chaos/CollisionResolution.cpp
+++ b/Engine/Source/Runtime/Experimental/Chaos/Private/Chaos/CollisionResolution.cpp
@@ -615,6 +615,52 @@ namespace Chaos
 			}
 		}
 
+		//
+		// Capsule-HeightField
+		//
+
+		template <typename T, int d>
+		TContactPoint<T> CapsuleHeightFieldContactPoint(const TCapsule<T>& A, const TRigidTransform<T, d>& ATransform, const THeightField<T>& B, const TRigidTransform<T, d>& BTransform, const T Thickness)
+		{
+			TContactPoint<T> Contact;
+			B.GJKContactPoint(TCapsule<float>(A), ATransform.GetRelativeTransform(BTransform), Thickness, Contact.Location, Contact.Normal, Contact.Phi);
+			Contact.Location = BTransform.TransformPosition(Contact.Location);
+			Contact.Normal = BTransform.TransformVector(Contact.Normal);
+			return Contact;
+		}
+
+		template <typename T, int d>
+		void UpdateCapsuleHeightFieldConstraint(const TCapsule<T>& A, const TRigidTransform<T, d>& ATransform, const THeightField<T>& B, const TRigidTransform<T, d>& BTransform, const T Thickness, TRigidBodyPointContactConstraint<T, d>& Constraint)
+		{
+			UpdateContactPoint(Constraint.Manifold, CapsuleHeightFieldContactPoint(A, ATransform, B, BTransform, Thickness));
+		}
+
+		template<typename T, int d>
+		void ConstructCapsuleHeightFieldConstraints(TGeometryParticleHandle<T, d>* Particle0, TGeometryParticleHandle<T, d>* Particle1, const FImplicitObject* Implicit0, const FImplicitObject* Implicit1, const TRigidTransform<T, d>& Transform0, const TRigidTransform<T, d>& Transform1, const T Thickness, FCollisionConstraintsArray& NewConstraints)
+		{
+
+			const TCapsule<T> * Object0 = Implicit0->template GetObject<const TCapsule<T> >();
+			const THeightField<T> * Object1 = Implicit1->template GetObject<const THeightField<T> >();
+			if (ensure(Object0 && Object1))
+			{
+				TRigidTransform<T, d> ParticleImplicit0TM = Transform0.GetRelativeTransform(Collisions::GetTransform(Particle0));
+				TRigidTransform<T, d> ParticleImplicit1TM = Transform1.GetRelativeTransform(Collisions::GetTransform(Particle1));
+				FRigidBodyPointContactConstraint* Constraint = new FRigidBodyPointContactConstraint(Particle0, Implicit0, ParticleImplicit0TM, Particle1, Implicit1, ParticleImplicit1TM);
+
+				UpdateCapsuleHeightFieldConstraint(*Object0, Transform0, *Object1, Transform1, Thickness, *Constraint);
+
+				if (Constraint->GetPhi() < Thickness)
+				{
+					NewConstraints.Add(Constraint);
+				}
+				else
+				{
+					delete Constraint;
+				}
+			}
+		}
+
+
 		//
 		// Convex - Convex
 		//
@@ -976,18 +1022,18 @@ namespace Chaos
 			{
 				UpdateCapsuleBoxConstraint(*Implicit0.template GetObject<TCapsule<T>>(), Transform0, Implicit1.template GetObject<TBox<T, d>>()->GetAABB(), Transform1, Thickness, *ConstraintBase.template As<TRigidBodyPointContactConstraint<T,d>>());
 			}
+			else if (Implicit0Type == TCapsule<T>::StaticType() && Implicit1Type == THeightField<T>::StaticType())
+			{
+				UpdateCapsuleHeightFieldConstraint(*Implicit0.template GetObject<TCapsule<T>>(), Transform0, *Implicit1.template GetObject< THeightField<T> >(), Transform1, Thickness, *ConstraintBase.template As<TRigidBodyPointContactConstraint<T, d>>());
+			}
 			else if (Implicit0Type == THeightField<T>::StaticType() && Implicit1Type == TBox<T, d>::StaticType())
 			{
-				// @todo(chaos) : Collision Constraints
-				//                Does this really work? 
-				TRigidBodyPointContactConstraint<T, d>& Constraint = *ConstraintBase.template As<TRigidBodyPointContactConstraint<T, d> >();
-				TRigidBodyPointContactConstraint<T, d> TmpConstraint = Constraint;
-				UpdateBoxHeightFieldConstraint(Implicit1.template GetObject<TBox<T, d>>()->GetAABB(), Transform1, *Implicit0.template GetObject<THeightField<T>>(), Transform0, Thickness, TmpConstraint);
-				if (TmpConstraint.GetPhi() < Constraint.GetPhi())
-				{
-					Constraint = TmpConstraint;
-					Constraint.SetNormal(-Constraint.GetNormal());
-				}
+				//     This case should not be necessary. The height fields 
+				//     will only ever be collided against, so ideally will never 
+				//     be in index[0] position of the constraint, also the construction
+				//     of the constraint will just switch the index position so its always 
+				//     second.
+				ensure(false);
 			}
 			else if (Implicit0Type == TPlane<T, d>::StaticType() && Implicit1Type == TBox<T, d>::StaticType())
 			{
@@ -1002,16 +1048,12 @@ namespace Chaos
 			}
 			else if (Implicit0Type == THeightField<T>::StaticType() && Implicit1Type == TSphere<T, d>::StaticType())
 			{
-				// @todo(chaos) : Collision Constraints
-				//                Does this really work? 
-				TRigidBodyPointContactConstraint<T, d>& Constraint = *ConstraintBase.template As<TRigidBodyPointContactConstraint<T, d> >();
-				TRigidBodyPointContactConstraint<T, d> TmpConstraint = Constraint;
-				UpdateSphereHeightFieldConstraint(*Implicit1.template GetObject<TSphere<T, d>>(), Transform1, *Implicit0.template GetObject<THeightField<T>>(), Transform0, Thickness, TmpConstraint);
-				if (TmpConstraint.GetPhi() < Constraint.GetPhi())
-				{
-					Constraint = TmpConstraint;
-					Constraint.SetNormal(-Constraint.GetNormal());
-				}
+				//     This case should not be necessary. The height fields 
+				//     will only ever be collided against, so ideally will never 
+				//     be in index[0] position of the constraint, also the construction
+				//     of the constraint will just switch the index position so its always 
+				//     second.
+				ensure(false);
 			}
 			else if (Implicit0Type == TPlane<T, d>::StaticType() && Implicit1Type == TSphere<T, d>::StaticType())
 			{
@@ -1057,6 +1099,15 @@ namespace Chaos
 					Constraint.SetNormal(-Constraint.GetNormal());
 				}
 			}
+			else if (Implicit0Type == THeightField<T>::StaticType() && Implicit1Type == TSphere<T, d>::StaticType())
+			{
+				//     This case should not be necessary. The height fields 
+				//     will only ever be collided against, so ideally will never 
+				//     be in index[0] position of the constraint, also the construction
+				//     of the constraint will just switch the index position so its always 
+				//     second.
+				ensure(false);
+			}
 			else if(Implicit0.IsConvex() && Implicit1.IsConvex())
 			{
 				UpdateConvexConvexConstraint(Implicit0, Transform0, Implicit1, Transform1, Thickness, ConstraintBase);
@@ -1125,7 +1176,7 @@ namespace Chaos
 				{
 					const FImplicitObject* Implicit0InnerObj = LevelsetObjPair.First;
 					const TRigidTransform<T, d> Implicit0InnerObjTM = LevelsetObjPair.Second * Transform0;
-					ConstructConstraints(Particle1, Particle0, Implicit1, Implicit0InnerObj, Transform1, Implicit0InnerObjTM, Thickness, NewConstraints);
+					ConstructConstraints(Particle0, Particle1, Implicit0InnerObj, Implicit1, Implicit0InnerObjTM, Transform1, Thickness, NewConstraints);
 				}
 			}
 			else if (Implicit0Type == TImplicitObjectUnion<T, d>::StaticType() && Implicit1Type == TImplicitObjectUnion<T, d>::StaticType())
@@ -1200,6 +1251,14 @@ namespace Chaos
 			{
 				ConstructCapsuleBoxConstraints(Particle1, Particle0, Implicit1, Implicit0, Transform1, Transform0, Thickness, NewConstraints);
 			}
+			else if (Implicit0Type == TCapsule<T>::StaticType() && Implicit1Type == THeightField<T>::StaticType())
+			{
+				ConstructCapsuleHeightFieldConstraints(Particle0, Particle1, Implicit0, Implicit1, Transform0, Transform1, Thickness, NewConstraints);
+			}
+			else if (Implicit0Type == THeightField<T>::StaticType() && Implicit1Type == TCapsule<T>::StaticType())
+			{
+				ConstructSphereHeightFieldConstraints(Particle1, Particle0, Implicit1, Implicit0, Transform1, Transform0, Thickness, NewConstraints);
+			}
 			else if (Implicit0->IsConvex() && Implicit1->IsConvex())
 			{
 				ConstructConvexConvexConstraints(Particle0, Particle1, Implicit0, Implicit1, Transform0, Transform1, Thickness, NewConstraints);
@@ -1241,6 +1300,9 @@ namespace Chaos
 		template void UpdateCapsuleBoxConstraint<float, 3>(const TCapsule<float>& A, const TRigidTransform<float, 3>& ATransform, const TAABB<float, 3>& B, const TRigidTransform<float, 3>& BTransform, const float Thickness, TRigidBodyPointContactConstraint<float, 3>& Constraint);
 		template void ConstructCapsuleBoxConstraints<float, 3>(TGeometryParticleHandle<float, 3>* Particle0, TGeometryParticleHandle<float, 3>* Particle1, const FImplicitObject* Implicit0, const FImplicitObject* Implicit1, const TRigidTransform<float, 3>& Transform0, const TRigidTransform<float, 3>& Transform1, const float Thickness, FCollisionConstraintsArray& NewConstraints);
 
+		template void UpdateCapsuleHeightFieldConstraint(const TCapsule<float>& A, const TRigidTransform<float, 3>& ATransform, const THeightField<float>& B, const TRigidTransform<float, 3>& BTransform, const float Thickness, TRigidBodyPointContactConstraint<float, 3>& Constraint);
+		template void ConstructCapsuleHeightFieldConstraints(TGeometryParticleHandle<float, 3>* Particle0, TGeometryParticleHandle<float, 3>* Particle1, const FImplicitObject* Implicit0, const FImplicitObject* Implicit1, const TRigidTransform<float, 3>& Transform0, const TRigidTransform<float, 3>& Transform1, const float Thickness, FCollisionConstraintsArray& NewConstraints);
+
 		template void UpdateConvexConvexConstraint<float, 3>(const FImplicitObject& A, const TRigidTransform<float, 3>& ATM, const FImplicitObject& B, const TRigidTransform<float, 3>& BTM, const float Thickness, TCollisionConstraintBase<float, 3>& Constraint);
 		template void UpdateConvexConvexManifold<float, 3>(TRigidBodyIterativeContactConstraint<float, 3>& Constraint, const TRigidTransform<float, 3>& ATM, const TRigidTransform<float, 3>& BTM, const float Thickness);
 		template void ConstructConvexConvexConstraints<float, 3>(TGeometryParticleHandle<float, 3>* Particle0, TGeometryParticleHandle<float, 3>* Particle1, const FImplicitObject* Implicit0, const FImplicitObject* Implicit1, const TRigidTransform<float, 3>& Transform0, const TRigidTransform<float, 3>& Transform1, const float Thickness, FCollisionConstraintsArray& NewConstraints);
diff --git a/Engine/Source/Runtime/Experimental/Chaos/Private/Chaos/HeightField.cpp b/Engine/Source/Runtime/Experimental/Chaos/Private/Chaos/HeightField.cpp
--- a/Engine/Source/Runtime/Experimental/Chaos/Private/Chaos/HeightField.cpp
+++ b/Engine/Source/Runtime/Experimental/Chaos/Private/Chaos/HeightField.cpp
@@ -1241,6 +1241,12 @@ namespace Chaos
 		return GJKContactPointImp(QueryGeom, QueryTM, Thickness, ContactLocation, ContactNormal, ContactPhi);
 	}
 
+	template <typename T>
+	bool THeightField<T>::GJKContactPoint(const TCapsule<T>& QueryGeom, const TRigidTransform<T, 3>& QueryTM, const T Thickness, TVector<T, 3>& ContactLocation, TVector<T, 3>& ContactNormal, T& ContactPhi) const
+	{
+		return GJKContactPointImp(QueryGeom, QueryTM, Thickness, ContactLocation, ContactNormal, ContactPhi);
+	}
+
 
 	template <typename T>
 	template <typename QueryGeomType>
diff --git a/Engine/Source/Runtime/Experimental/Chaos/Public/Chaos/CollisionResolution.h b/Engine/Source/Runtime/Experimental/Chaos/Public/Chaos/CollisionResolution.h
--- a/Engine/Source/Runtime/Experimental/Chaos/Public/Chaos/CollisionResolution.h
+++ b/Engine/Source/Runtime/Experimental/Chaos/Public/Chaos/CollisionResolution.h
@@ -135,6 +135,16 @@ namespace Chaos
 			template<typename T, int d>
 			void CHAOS_API ConstructCapsuleBoxConstraints(TGeometryParticleHandle<T, d>* Particle0, TGeometryParticleHandle<T, d>* Particle1, const FImplicitObject* Implicit0, const FImplicitObject* Implicit1, const TRigidTransform<T, d>& Transform0, const TRigidTransform<T, d>& Transform1, const T Thickness, FCollisionConstraintsArray& NewConstraints);
 
+			//
+			// Capsule-HeightField
+			//
+
+			template <typename T, int d>
+			void UpdateCapsuleHeightFieldConstraint(const TCapsule<T>& A, const TRigidTransform<T, d>& ATransform, const THeightField<T>& B, const TRigidTransform<T, d>& BTransform, const T Thickness, TRigidBodyPointContactConstraint<T, d>& Constraint);
+
+			template<typename T, int d>
+			void ConstructCapsuleHeightFieldConstraints(TGeometryParticleHandle<T, d>* Particle0, TGeometryParticleHandle<T, d>* Particle1, const FImplicitObject* Implicit0, const FImplicitObject* Implicit1, const TRigidTransform<T, d>& Transform0, const TRigidTransform<T, d>& Transform1, const T Thickness, FCollisionConstraintsArray& NewConstraints);
+
 			//
 			// Convex-Convex
 			//
diff --git a/Engine/Source/Runtime/Experimental/Chaos/Public/Chaos/HeightField.h b/Engine/Source/Runtime/Experimental/Chaos/Public/Chaos/HeightField.h
--- a/Engine/Source/Runtime/Experimental/Chaos/Public/Chaos/HeightField.h
+++ b/Engine/Source/Runtime/Experimental/Chaos/Public/Chaos/HeightField.h
@@ -73,6 +73,7 @@ namespace Chaos
 
 		bool GJKContactPoint(const TBox<T, 3>& QueryGeom, const TRigidTransform<T, 3>& QueryTM, const T Thickness, TVector<T, 3>& ContactLocation, TVector<T, 3>& ContactNormal, T& ContactPhi) const;
 		bool GJKContactPoint(const TSphere<T, 3>& QueryGeom, const TRigidTransform<T, 3>& QueryTM, const T Thickness, TVector<T, 3>& ContactLocation, TVector<T, 3>& ContactNormal, T& ContactPhi) const;
+		bool GJKContactPoint(const TCapsule<T>& QueryGeom, const TRigidTransform<T, 3>& QueryTM, const T Thickness, TVector<T, 3>& ContactLocation, TVector<T, 3>& ContactNormal, T& ContactPhi) const;
 
 
 		virtual int32 FindMostOpposingFace(const TVector<T, 3>& Position, const TVector<T, 3>& UnitDir, int32 HintFaceIndex, T SearchDist) const override;
