Chaos RBAN
- XPBD angular drives
- Mass independent angular drives
- Support custom masses and mass/inertia scaling
- Support COM Nudge
- Fix linear projection
- Fix elliptical cone limits

diff --git a/Engine/Source/Runtime/Engine/Private/PhysicsEngine/ImmediatePhysics/ImmediatePhysicsChaos/ImmediatePhysicsActorHandle_Chaos.cpp b/Engine/Source/Runtime/Engine/Private/PhysicsEngine/ImmediatePhysics/ImmediatePhysicsChaos/ImmediatePhysicsActorHandle_Chaos.cpp
--- a/Engine/Source/Runtime/Engine/Private/PhysicsEngine/ImmediatePhysics/ImmediatePhysicsChaos/ImmediatePhysicsActorHandle_Chaos.cpp
+++ b/Engine/Source/Runtime/Engine/Private/PhysicsEngine/ImmediatePhysics/ImmediatePhysicsChaos/ImmediatePhysicsActorHandle_Chaos.cpp
@@ -157,18 +157,23 @@ namespace ImmediatePhysics_Chaos
 		UBodySetup* BodySetup = BodyInstance->BodySetup.Get();
 
 #if WITH_CHAOS && !PHYSICS_INTERFACE_PHYSX
-		float Density = 1.e-3f;	// 1g/cm3
 		Chaos::TMassProperties<float, 3> MassProperties;
 		CalculateMassProperties<float, 3>(Scale, FTransform::Identity, BodySetup->AggGeom, MassProperties);
-		OutMass = Density * MassProperties.Volume;
-		OutInertia = Density * Chaos::TVector<float, 3>(MassProperties.InertiaTensor.M[0][0], MassProperties.InertiaTensor.M[1][1], MassProperties.InertiaTensor.M[2][2]);
-		OutCoMTransform = FTransform(MassProperties.RotationOfMass, MassProperties.CenterOfMass);
+		float Density = 1.e-3f;	// 1g/cm3	@todo(ccaulfield): should come from material
+		if (BodyInstance->bOverrideMass)
+		{
+			Density = BodyInstance->GetMassOverride() / MassProperties.Volume;
+		}
+		OutMass = Density * (BodyInstance->MassScale * MassProperties.Volume);
+		OutInertia = Density * (BodyInstance->InertiaTensorScale * Chaos::TVector<float, 3>(MassProperties.InertiaTensor.M[0][0], MassProperties.InertiaTensor.M[1][1], MassProperties.InertiaTensor.M[2][2]));
+		OutCoMTransform = FTransform(MassProperties.RotationOfMass, MassProperties.CenterOfMass + BodyInstance->COMNudge);
 #else
 		OutMass = BodyInstance->GetBodyMass();
 		OutInertia = BodyInstance->GetBodyInertiaTensor();
 		OutCoMTransform = BodyInstance->GetMassSpaceLocal();
 #endif
 
+
 		FBodyCollisionData BodyCollisionData;
 		BodyInstance->BuildBodyFilterData(BodyCollisionData.CollisionFilterData);
 		BodyInstance->BuildBodyCollisionFlags(BodyCollisionData.CollisionFlags, BodyInstance->GetCollisionEnabled(), BodyInstance->BodySetup->GetCollisionTraceFlag() == CTF_UseComplexAsSimple);
@@ -502,6 +507,11 @@ namespace ImmediatePhysics_Chaos
 		return Handle()->InvM();
 	}
 
+	float FActorHandle::GetMass() const
+	{
+		return Handle()->M();
+	}
+
 	void FActorHandle::SetInverseInertia(const FVector& NewInverseInertia)
 	{
 		using namespace Chaos;
@@ -523,10 +533,18 @@ namespace ImmediatePhysics_Chaos
 	{
 		using namespace Chaos;
 
-		PMatrix<float, 3, 3> InvI = Handle()->InvI();
+		const PMatrix<float, 3, 3>& InvI = Handle()->InvI();
 		return { InvI.M[0][0], InvI.M[1][1], InvI.M[2][2] };
 	}
 
+	FVector FActorHandle::GetInertia() const
+	{
+		using namespace Chaos;
+
+		const PMatrix<float, 3, 3>& I = Handle()->I();
+		return { I.M[0][0], I.M[1][1], I.M[2][2] };
+	}
+
 	void FActorHandle::SetMaxDepenetrationVelocity(float NewMaxDepenetrationVelocity)
 	{
 #if IMMEDIATEPHYSICS_CHAOS_TODO
diff --git a/Engine/Source/Runtime/Engine/Private/PhysicsEngine/ImmediatePhysics/ImmediatePhysicsChaos/ImmediatePhysicsJointHandle_Chaos.cpp b/Engine/Source/Runtime/Engine/Private/PhysicsEngine/ImmediatePhysics/ImmediatePhysicsChaos/ImmediatePhysicsJointHandle_Chaos.cpp
--- a/Engine/Source/Runtime/Engine/Private/PhysicsEngine/ImmediatePhysics/ImmediatePhysicsChaos/ImmediatePhysicsJointHandle_Chaos.cpp
+++ b/Engine/Source/Runtime/Engine/Private/PhysicsEngine/ImmediatePhysics/ImmediatePhysicsChaos/ImmediatePhysicsJointHandle_Chaos.cpp
@@ -19,11 +19,14 @@ namespace ImmediatePhysics_Chaos
 	float ChaosImmediate_JointStiffness = 1.0f;
 	FAutoConsoleVariableRef CVarJointStiffness(TEXT("p.Chaos.ImmPhys.JointStiffness"), ChaosImmediate_JointStiffness, TEXT("Joint solver stiffness."));
 
-	bool ChaosImmediate_StiffnessUseMass = true;
+	bool ChaosImmediate_StiffnessUseMass = false;
 	FAutoConsoleVariableRef CVarStiffnessUseMass(TEXT("p.Chaos.ImmPhys.StiffnessUseMass"), ChaosImmediate_StiffnessUseMass, TEXT("Whether to use mass to scale stiffness in the conversion."));
 
+	float ChaosImmediate_DriveStiffnessScale = 8.0f;
+	FAutoConsoleVariableRef CVarDriveStiffnessScale(TEXT("p.Chaos.ImmPhys.DriveStiffnessScale"), ChaosImmediate_DriveStiffnessScale, TEXT("Conversion factor for drive stiffness."));
+
 	float ChaosImmediate_JointMinProjection = 0.0f;
-	float ChaosImmediate_JointMaxProjection = 1.0f;
+	float ChaosImmediate_JointMaxProjection = 0.5f;
 	FAutoConsoleVariableRef CVarJointMinProjection(TEXT("p.Chaos.ImmPhys.JointMinProjection"), ChaosImmediate_JointMinProjection, TEXT("Joint min projection (for joints with projection disabled)."));
 	FAutoConsoleVariableRef CVarJointMaxProjection(TEXT("p.Chaos.ImmPhys.JointMaxProjection"), ChaosImmediate_JointMaxProjection, TEXT("Joint max projection (for joints with projection enabled)."));
 
@@ -50,16 +53,20 @@ namespace ImmediatePhysics_Chaos
 
 	// Convert the UE (PhysX) drive spring stiffness to a joint stiffness [0,1] value for use in the solver.
 	// We linearly map a range of UE stiffness values to a range of Chaos values, optionally scaling by the effective mass.
-	// Ranges and effective mass use are vcar controlled.
-	float ConvertDriveAngularStiffness(const float InStiffness, const Chaos::TVector<float, 3>& InvI0, const Chaos::TVector<float, 3>& InvI1)
+	float ConvertDriveAngularStiffness(const float InStiffness, const Chaos::FVec3& X0, const FReal M0, const Chaos::FVec3& I0, const Chaos::FVec3& X1, const FReal M1, const Chaos::FVec3& I1)
 	{
-		if (ChaosImmediate_DriveStiffnessSourceMax > ChaosImmediate_DriveStiffnessSourceMin)
+		using namespace Chaos;
+
+		// Calculate inertia of the system about the joint connector using parallel axis theorem
+		float XI = (FReal)1.0;
+		if (ChaosImmediate_StiffnessUseMass)
 		{
-			const float InvI = ChaosImmediate_StiffnessUseMass ? (InvI0.Min() + InvI1.Min()) : 1.0f;
-			const float F = FMath::Max(0.0f, (InStiffness * InvI - ChaosImmediate_DriveStiffnessSourceMin)) / FMath::Max(SMALL_NUMBER, (ChaosImmediate_DriveStiffnessSourceMax - ChaosImmediate_DriveStiffnessSourceMin));
-			return FMath::Lerp(ChaosImmediate_DriveStiffnessTargetMin, ChaosImmediate_DriveStiffnessTargetMax, FMath::Clamp(F, 0.0f, 1.0f));
+			const FVec3 XI0 = I0 + M0 * (X0 * X0);
+			const FVec3 XI1 = I1 + M1 * (X1 * X1);
+			XI = FMath::Max(KINDA_SMALL_NUMBER, XI0.Max() + XI1.Max());
 		}
-		return 0;
+
+		return ChaosImmediate_DriveStiffnessScale * InStiffness / XI;
 	}
 
 	float ConvertSoftLinearStiffness(const float InStiffness, const float InvM0, const float InvM1)
@@ -149,7 +156,7 @@ namespace ImmediatePhysics_Chaos
 			ConstraintSettings.Motion.bAngularSwingDriveEnabled = ConstraintInstance->ProfileInstance.AngularDrive.SwingDrive.bEnablePositionDrive;
 		}
 
-		ConstraintSettings.Motion.AngularDriveStiffness = ConvertDriveAngularStiffness(ConstraintInstance->ProfileInstance.AngularDrive.TwistDrive.Stiffness, Actor1->GetInverseInertia(), Actor2->GetInverseInertia());
+		ConstraintSettings.Motion.AngularDriveStiffness = ConvertDriveAngularStiffness(ConstraintInstance->ProfileInstance.AngularDrive.TwistDrive.Stiffness, ConstraintSettings.ConstraintFrames[0].GetTranslation(), Actor1->GetMass(), Actor1->GetInertia(), ConstraintSettings.ConstraintFrames[1].GetTranslation(), Actor2->GetMass(), Actor2->GetInertia());
 
 		ConstraintSettings.Motion.Sanitize();
 
diff --git a/Engine/Source/Runtime/Engine/Private/PhysicsEngine/ImmediatePhysics/ImmediatePhysicsChaos/ImmediatePhysicsSimulation_Chaos.cpp b/Engine/Source/Runtime/Engine/Private/PhysicsEngine/ImmediatePhysics/ImmediatePhysicsChaos/ImmediatePhysicsSimulation_Chaos.cpp
--- a/Engine/Source/Runtime/Engine/Private/PhysicsEngine/ImmediatePhysics/ImmediatePhysicsChaos/ImmediatePhysicsSimulation_Chaos.cpp
+++ b/Engine/Source/Runtime/Engine/Private/PhysicsEngine/ImmediatePhysics/ImmediatePhysicsChaos/ImmediatePhysicsSimulation_Chaos.cpp
@@ -51,6 +51,7 @@ float ChaosImmediate_Joint_Stiffness = 1.0f;
 float ChaosImmediate_Joint_SoftLinearStiffness = 0.0f;
 float ChaosImmediate_Joint_SoftAngularStiffness = 0.0f;
 float ChaosImmediate_Joint_DriveStiffness = 0.0f;
+float ChaosImmediate_Joint_DriveDamping = 0.0f;
 float ChaosImmediate_Joint_MinParentMassRatio = 0.5f;
 float ChaosImmediate_Joint_MaxInertiaRatio = 5.0f;
 FAutoConsoleVariableRef CVarChaosImmPhysPairIterations(TEXT("p.Chaos.ImmPhys.Joint.PairIterations"), ChaosImmediate_Joint_PairIterations, TEXT("PairIterations."));
@@ -65,6 +66,7 @@ FAutoConsoleVariableRef CVarChaosImmPhysStiffness(TEXT("p.Chaos.ImmPhys.Joint.St
 FAutoConsoleVariableRef CVarChaosImmPhysSoftLinearStiffness(TEXT("p.Chaos.ImmPhys.Joint.SoftLinearStiffness"), ChaosImmediate_Joint_SoftLinearStiffness, TEXT("6Dof joint soft linear stiffness override (if > 0)."));
 FAutoConsoleVariableRef CVarChaosImmPhysSoftAngularStiffness(TEXT("p.Chaos.ImmPhys.Joint.SoftAngularStiffness"), ChaosImmediate_Joint_SoftAngularStiffness, TEXT("6Dof joint soft angular stiffness override (if > 0)."));
 FAutoConsoleVariableRef CVarChaosImmPhysDriveStiffness(TEXT("p.Chaos.ImmPhys.Joint.DriveStiffness"), ChaosImmediate_Joint_DriveStiffness, TEXT("6Dof joint drive stiffness override (if > 0)."));
+FAutoConsoleVariableRef CVarChaosImmPhysDriveDamping(TEXT("p.Chaos.ImmPhys.Joint.DriveDamping"), ChaosImmediate_Joint_DriveDamping, TEXT("6Dof joint drive damping override (if > 0)."));
 FAutoConsoleVariableRef CVarChaosImmPhysMinParentMassRatio(TEXT("p.Chaos.ImmPhys.Joint.MinParentMassRatio"), ChaosImmediate_Joint_MinParentMassRatio, TEXT("6Dof joint MinParentMassRatio (if > 0)"));
 FAutoConsoleVariableRef CVarChaosImmPhysMaxInertiaRatio(TEXT("p.Chaos.ImmPhys.Joint.MaxInertiaRatio"), ChaosImmediate_Joint_MaxInertiaRatio, TEXT("6Dof joint MaxInertiaRatio (if > 0)"));
 
@@ -496,6 +498,7 @@ namespace ImmediatePhysics_Chaos
 			JointsSettings.SoftLinearStiffness = ChaosImmediate_Joint_SoftLinearStiffness;
 			JointsSettings.SoftAngularStiffness = ChaosImmediate_Joint_SoftAngularStiffness;
 			JointsSettings.DriveStiffness = ChaosImmediate_Joint_DriveStiffness;
+			JointsSettings.DriveDamping = ChaosImmediate_Joint_DriveDamping;
 
 			Joints.SetSettings(JointsSettings);
 
diff --git a/Engine/Source/Runtime/Engine/Public/Physics/ImmediatePhysics/ImmediatePhysicsChaos/ImmediatePhysicsActorHandle_Chaos.h b/Engine/Source/Runtime/Engine/Public/Physics/ImmediatePhysics/ImmediatePhysicsChaos/ImmediatePhysicsActorHandle_Chaos.h
--- a/Engine/Source/Runtime/Engine/Public/Physics/ImmediatePhysics/ImmediatePhysicsChaos/ImmediatePhysicsActorHandle_Chaos.h
+++ b/Engine/Source/Runtime/Engine/Public/Physics/ImmediatePhysics/ImmediatePhysicsChaos/ImmediatePhysicsActorHandle_Chaos.h
@@ -85,12 +85,14 @@ namespace ImmediatePhysics_Chaos
 
 		/** Get the inverse mass. */
 		float GetInverseMass() const;
+		float GetMass() const;
 
 		/** Set the inverse inertia. Mass-space inverse inertia diagonal vector */
 		void SetInverseInertia(const FVector& NewInverseInertia);
 
 		/** Get the inverse inertia. Mass-space inverse inertia diagonal vector */
 		FVector GetInverseInertia() const;
+		FVector GetInertia() const;
 
 		/** Set the max depenetration velocity*/
 		void SetMaxDepenetrationVelocity(float NewMaxDepenetrationVelocity);
diff --git a/Engine/Source/Runtime/Experimental/Chaos/Private/Chaos/ChaosDebugDraw.cpp b/Engine/Source/Runtime/Experimental/Chaos/Private/Chaos/ChaosDebugDraw.cpp
--- a/Engine/Source/Runtime/Experimental/Chaos/Private/Chaos/ChaosDebugDraw.cpp
+++ b/Engine/Source/Runtime/Experimental/Chaos/Private/Chaos/ChaosDebugDraw.cpp
@@ -198,7 +198,7 @@ namespace Chaos
 			}
 		}
 
-		void DrawJointConstraintImpl(const FRigidTransform3& SpaceTransform, const FVec3& InPa, const FVec3& InXa, const FMatrix33& Ra, const FVec3& InPb, const FVec3& InXb, const FMatrix33& Rb, const FVec3& CR, int32 Level, int32 Index, FReal ColorScale, uint32 FeatureMask)
+		void DrawJointConstraintImpl(const FRigidTransform3& SpaceTransform, const FVec3& InPa, const FVec3& InXa, const FMatrix33& Ra, const FVec3& InPb, const FVec3& InXb, const FMatrix33& Rb, int32 Level, int32 Index, FReal ColorScale, uint32 FeatureMask)
 		{
 			using namespace Chaos::DebugDraw;
 			FColor R = (ColorScale * FColor::Red).ToFColor(false);
@@ -236,9 +236,6 @@ namespace Chaos
 			{
 				FDebugDrawQueue::GetInstance().DrawDebugString(Xb + FontHeight * FVec3(0, 0, 1), FString::Format(TEXT("{0}"), { Index }), nullptr, FColor::Red, KINDA_SMALL_NUMBER, false, FontScale);
 			}
-			//FDebugDrawQueue::GetInstance().DrawDebugString(Xb + 3 * FontHeight * FVec3(0, 0, 1), FString::Format(TEXT("T=  {0}"), { FMath::RadiansToDegrees(CR[(int32)E6DJointAngularConstraintIndex::Twist]) }), nullptr, FColor::Red, KINDA_SMALL_NUMBER, false, FontScale);
-			//FDebugDrawQueue::GetInstance().DrawDebugString(Xb + 2 * FontHeight * FVec3(0, 0, 1), FString::Format(TEXT("S1= {0}"), { FMath::RadiansToDegrees(CR[(int32)E6DJointAngularConstraintIndex::Swing1]) }), nullptr, FColor::Red, KINDA_SMALL_NUMBER, false, FontScale);
-			//FDebugDrawQueue::GetInstance().DrawDebugString(Xb + 1 * FontHeight * FVec3(0, 0, 1), FString::Format(TEXT("S2= {0}"), { FMath::RadiansToDegrees(CR[(int32)E6DJointAngularConstraintIndex::Swing2]) }), nullptr, FColor::Red, KINDA_SMALL_NUMBER, false, FontScale);
 		}
 
 		void DrawJointConstraintImpl(const FRigidTransform3& SpaceTransform, const FPBDJointConstraintHandle* ConstraintHandle, FReal ColorScale, uint32 FeatureMask)
@@ -250,10 +247,10 @@ namespace Chaos
 			{
 				FVec3 Pa = FParticleUtilities::GetCoMWorldPosition(TConstGenericParticleHandle<FReal, 3>(ConstraintHandle->GetConstrainedParticles()[1]));
 				FVec3 Pb = FParticleUtilities::GetCoMWorldPosition(TConstGenericParticleHandle<FReal, 3>(ConstraintHandle->GetConstrainedParticles()[0]));
-				FVec3 Xa, Xb, CR;
+				FVec3 Xa, Xb;
 				FMatrix33 Ra, Rb;
-				ConstraintHandle->CalculateConstraintSpace(Xa, Ra, Xb, Rb, CR);
-				DrawJointConstraintImpl(SpaceTransform, Pa, Xa, Ra, Pb, Xb, Rb, CR, ConstraintHandle->GetConstraintLevel(), ConstraintHandle->GetConstraintIndex(), ColorScale, FeatureMask);
+				ConstraintHandle->CalculateConstraintSpace(Xa, Ra, Xb, Rb);
+				DrawJointConstraintImpl(SpaceTransform, Pa, Xa, Ra, Pb, Xb, Rb, ConstraintHandle->GetConstraintLevel(), ConstraintHandle->GetConstraintIndex(), ColorScale, FeatureMask);
 			}
 		}
 
diff --git a/Engine/Source/Runtime/Experimental/Chaos/Private/Chaos/Joint/PBDJointSolverGaussSeidel.cpp b/Engine/Source/Runtime/Experimental/Chaos/Private/Chaos/Joint/PBDJointSolverGaussSeidel.cpp
--- a/Engine/Source/Runtime/Experimental/Chaos/Private/Chaos/Joint/PBDJointSolverGaussSeidel.cpp
+++ b/Engine/Source/Runtime/Experimental/Chaos/Private/Chaos/Joint/PBDJointSolverGaussSeidel.cpp
@@ -24,14 +24,12 @@ namespace Chaos
 	}
 
 
-	void FJointSolverGaussSeidel::InitConstraints(
+	void FJointSolverGaussSeidel::Init(
 		const FReal Dt,
 		const FPBDJointSolverSettings& SolverSettings,
 		const FPBDJointSettings& JointSettings,
-		const FVec3& P0,
-		const FRotation3& Q0,
-		const FVec3& P1,
-		const FRotation3& Q1,
+		const FRotation3& PrevQ0,
+		const FRotation3& PrevQ1,
 		const FReal InvM0,
 		const FMatrix33& InvIL0,
 		const FReal InvM1,
@@ -47,30 +45,28 @@ namespace Chaos
 		InvMs[1] = InvM1;
 		// @todo(ccaulfield): mass conditioning
 
-		Ps[0] = P0;
-		Ps[1] = P1;
-		Qs[0] = Q0;
-		Qs[1] = Q1;
-		Qs[1].EnforceShortestArcWith(Qs[0]);
+		PrevQs[0] = PrevQ0;
+		PrevQs[1] = PrevQ1;
+
+		TwistLambda = (FReal)0;
+		SwingLambda = (FReal)0;
 
 		LinearStiffness = FPBDJointUtilities::GetLinearStiffness(SolverSettings, JointSettings);
 		TwistStiffness = FPBDJointUtilities::GetTwistStiffness(SolverSettings, JointSettings);
 		SwingStiffness = FPBDJointUtilities::GetSwingStiffness(SolverSettings, JointSettings);
 		AngularDriveStiffness = FPBDJointUtilities::GetAngularDriveStiffness(SolverSettings, JointSettings);
+		AngularDriveDamping = FPBDJointUtilities::GetAngularDriveDamping(SolverSettings, JointSettings);
 		LinearProjection = FPBDJointUtilities::GetLinearProjection(SolverSettings, JointSettings);
 		AngularProjection = FPBDJointUtilities::GetAngularProjection(SolverSettings, JointSettings);
 		SwingTwistAngleTolerance = SolverSettings.SwingTwistAngleTolerance;
 		bEnableTwistLimits = SolverSettings.bEnableTwistLimits;
 		bEnableSwingLimits = SolverSettings.bEnableSwingLimits;
 		bEnableDrives = SolverSettings.bEnableDrives;
-
-		UpdateDerivedState();
 	}
 
-	void FJointSolverGaussSeidel::InitProjection(
+
+	void FJointSolverGaussSeidel::Update(
 		const FReal Dt,
-		const FPBDJointSolverSettings& SolverSettings,
-		const FPBDJointSettings& JointSettings,
 		const FVec3& P0,
 		const FRotation3& Q0,
 		const FVec3& V0,
@@ -78,23 +74,20 @@ namespace Chaos
 		const FVec3& P1,
 		const FRotation3& Q1,
 		const FVec3& V1,
-		const FVec3& W1,
-		const FReal InvM0,
-		const FMatrix33& InvIL0,
-		const FReal InvM1,
-		const FMatrix33& InvIL1,
-		const FRigidTransform3& XL0,
-		const FRigidTransform3& XL1)
+		const FVec3& W1)
 	{
-		InitConstraints(Dt, SolverSettings, JointSettings, P0, Q0, P1, Q1, InvM0, InvIL0, InvM1, InvIL1, XL0, XL1);
-
-		InvILs[0] = FMatrix33(0, 0, 0);
-		InvMs[0] = 0;
+		Ps[0] = P0;
+		Ps[1] = P1;
+		Qs[0] = Q0;
+		Qs[1] = Q1;
+		Qs[1].EnforceShortestArcWith(Qs[0]);
 
 		Vs[0] = V0;
 		Vs[1] = V1;
 		Ws[0] = W0;
 		Ws[1] = W1;
+
+		UpdateDerivedState();
 	}
 
 
@@ -244,20 +237,79 @@ namespace Chaos
 	void FJointSolverGaussSeidel::ApplyRotationDelta(
 		const FReal Stiffness,
 		const FVec3& Axis0,
-		const FReal Angle0,
 		const FVec3& Axis1,
-		const FReal Angle1)
+		const FReal Angle)
 	{
 		const FMatrix33 InvI0 = Utilities::ComputeWorldSpaceInertia(Qs[0], InvILs[0]);
 		const FMatrix33 InvI1 = Utilities::ComputeWorldSpaceInertia(Qs[1], InvILs[1]);
 
-		const FReal L = (FReal)1 / (FVec3::DotProduct(Axis0, Utilities::Multiply(InvI0, Axis0)) + FVec3::DotProduct(Axis1, Utilities::Multiply(InvI1, Axis1)));
-		const FVec3 DR0 = Utilities::Multiply(InvI0, Axis0) * L * Angle0;
-		const FVec3 DR1 = Utilities::Multiply(InvI1, Axis1) * L * Angle1;
+		const FReal M0 = FVec3::DotProduct(Axis0, Utilities::Multiply(InvI0, Axis0));
+		const FReal M1 = FVec3::DotProduct(Axis1, Utilities::Multiply(InvI1, Axis1));
+
+		//const FVec3 DR0 = Utilities::Multiply(InvI0, Axis0) * (Angle / (M0 + M1));
+		//const FVec3 DR1 = Utilities::Multiply(InvI1, Axis1) * -(Angle / (M0 + M1));
+		const FVec3 DR0 = Axis0 * (Angle * M0 / (M0 + M1));
+		const FVec3 DR1 = Axis1 * -(Angle * M1 / (M0 + M1));
 
 		ApplyRotationDelta(Stiffness, DR0, DR1);
 	}
 
+	// See "XPBD: Position-Based Simulation of Compliant Constrained Dynamics"
+	void FJointSolverGaussSeidel::ApplyDriveRotationDelta(
+		const FReal Dt,
+		const FReal Stiffness,
+		const FReal Damping,
+		const FVec3& Axis0,
+		const FVec3& Axis1,
+		const FReal Angle,
+		FReal& Lambda)
+	{
+		const FMatrix33 InvI0 = Utilities::ComputeWorldSpaceInertia(Qs[0], InvILs[0]);
+		const FMatrix33 InvI1 = Utilities::ComputeWorldSpaceInertia(Qs[1], InvILs[1]);
+
+		const FReal IM0 = FVec3::DotProduct(Axis0, Utilities::Multiply(InvI0, Axis0));
+		const FReal IM1 = FVec3::DotProduct(Axis1, Utilities::Multiply(InvI1, Axis1));
+		const FReal IM = (IM0 + IM1);
+
+		FReal JV = 0;
+		if (Damping > KINDA_SMALL_NUMBER)
+		{
+			const FVec3 W0 = FRotation3::CalculateAngularVelocity(PrevQs[0], Qs[0], Dt);
+			const FVec3 W1 = FRotation3::CalculateAngularVelocity(PrevQs[1], Qs[1], Dt);
+			JV = FVec3::DotProduct(Axis0, W0) - FVec3::DotProduct(Axis1, W1);
+		}
+
+		//const FReal MassScale = (FReal)1;	// For force springs. Make this an option...
+		const FReal MassScale = IM;			// For acceleration springs. Make this an option...
+		FReal DLambda = 0;
+		if (Stiffness > KINDA_SMALL_NUMBER)
+		{
+			const FReal Alpha = MassScale / (Stiffness * Dt * Dt);
+			const FReal AlphaBeta = Damping / Stiffness;
+			const FReal DLambdaNumerator = (Angle - Alpha * Lambda - AlphaBeta * JV);
+			const FReal DLambdaDenominator = (((FReal)1 + AlphaBeta / Dt) * IM + Alpha);
+			DLambda = DLambdaNumerator / DLambdaDenominator;
+		}
+		else
+		{
+			const FReal Beta = Damping / (MassScale * Dt * Dt);
+			const FReal DLambdaNumerator = -Beta * JV;
+			const FReal DLambdaDenominator = IM;
+			DLambda = DLambdaNumerator / DLambdaDenominator;
+		}
+
+		//const FVec3 DR0 = Utilities::Multiply(InvI0, Axis0) * DLambda;
+		//const FVec3 DR1 = Utilities::Multiply(InvI1, Axis1) * -DLambda;
+		const FVec3 DR0 = Axis0 * (DLambda * IM0);
+		const FVec3 DR1 = Axis1 * -(DLambda * IM1);
+
+		const FRotation3 DQ0 = (FRotation3::FromElements(DR0, 0) * Qs[0]) * (FReal)0.5;
+		const FRotation3 DQ1 = (FRotation3::FromElements(DR1, 0) * Qs[1]) * (FReal)0.5;
+		Qs[0] = (Qs[0] + DQ0).GetNormalized();
+		Qs[1] = (Qs[1] + DQ1).GetNormalized();
+		Qs[1].EnforceShortestArcWith(Qs[0]);
+		Lambda += DLambda;
+	}
 
 	void FJointSolverGaussSeidel::ApplyVelocityDelta(
 		const FReal Stiffness,
@@ -281,16 +333,14 @@ namespace Chaos
 		FRotation3 R01Twist, R01Swing;
 		FPBDJointUtilities::DecomposeSwingTwistLocal(Rs[0], Rs[1], R01Swing, R01Twist);
 
-		FVec3 TwistAxis01;
-		FReal TwistAngle;
-		R01Twist.ToAxisAndAngleSafe(TwistAxis01, TwistAngle, FJointConstants::TwistAxis(), SwingTwistAngleTolerance);
+		FVec3 TwistAxis01 = FJointConstants::TwistAxis();
+		FReal TwistAngle = R01Twist.GetAngle();
 		if (TwistAngle > PI)
 		{
 			TwistAngle = TwistAngle - (FReal)2 * PI;
 		}
-		if (TwistAxis01.X < 0)
+		if (R01Twist.X < 0)
 		{
-			TwistAxis01 = -TwistAxis01;
 			TwistAngle = -TwistAngle;
 		}
 
@@ -318,7 +368,7 @@ namespace Chaos
 		}
 
 		// Apply twist correction
-		ApplyRotationDelta(TwistStiffness, TwistAxis0, DTwistAngle, TwistAxis1, -DTwistAngle);
+		ApplyRotationDelta(TwistStiffness, TwistAxis0, TwistAxis1, DTwistAngle);
 		UpdateDerivedState();
 	}
 
@@ -331,16 +381,14 @@ namespace Chaos
 		FRotation3 R01Twist, R01Swing;
 		FPBDJointUtilities::DecomposeSwingTwistLocal(Rs[0], Rs[1], R01Swing, R01Twist);
 
-		FVec3 TwistAxis01;
-		FReal TwistAngle;
-		R01Twist.ToAxisAndAngleSafe(TwistAxis01, TwistAngle, FJointConstants::TwistAxis(), SwingTwistAngleTolerance);
+		FVec3 TwistAxis01 = FJointConstants::TwistAxis();
+		FReal TwistAngle = R01Twist.GetAngle();
 		if (TwistAngle > PI)
 		{
 			TwistAngle = TwistAngle - (FReal)2 * PI;
 		}
-		if (TwistAxis01.X < 0)
+		if (R01Twist.X < 0)
 		{
-			TwistAxis01 = -TwistAxis01;
 			TwistAngle = -TwistAngle;
 		}
 
@@ -350,7 +398,7 @@ namespace Chaos
 		const FReal DTwistAngle = TwistAngle - TwistAngleTarget;
 
 		// Apply twist correction
-		ApplyRotationDelta(AngularDriveStiffness, TwistAxis0, DTwistAngle, TwistAxis1, -DTwistAngle);
+		ApplyDriveRotationDelta(Dt, AngularDriveStiffness, AngularDriveDamping, TwistAxis0, TwistAxis1, DTwistAngle, TwistLambda);
 		UpdateDerivedState();
 	}
 
@@ -387,7 +435,7 @@ namespace Chaos
 			// Map swing axis to ellipse and calculate limit for this swing axis
 			const FReal DotSwing1 = FMath::Abs(FVec3::DotProduct(SwingAxis01, FJointConstants::Swing1Axis()));
 			const FReal DotSwing2 = FMath::Abs(FVec3::DotProduct(SwingAxis01, FJointConstants::Swing2Axis()));
-			SwingAngleMax = FMath::Sqrt(Swing1Limit * DotSwing2 * Swing1Limit * DotSwing2 + Swing2Limit * DotSwing1 * Swing2Limit * DotSwing1);
+			SwingAngleMax = FMath::Sqrt(Swing1Limit * DotSwing1 * Swing1Limit * DotSwing1 + Swing2Limit * DotSwing2 * Swing2Limit * DotSwing2);
 		}
 
 		// Calculate swing error we need to correct
@@ -402,7 +450,7 @@ namespace Chaos
 		}
 
 		// Apply swing correction
-		ApplyRotationDelta(SwingStiffness, SwingAxis, DSwingAngle, SwingAxis, -DSwingAngle);
+		ApplyRotationDelta(SwingStiffness, SwingAxis, SwingAxis, DSwingAngle);
 		UpdateDerivedState();
 	}
 
@@ -432,8 +480,7 @@ namespace Chaos
 		FReal SwingAngleTarget = FMath::Max(Swing1Target, Swing2Target);
 		FReal DSwingAngle = SwingAngle - SwingAngleTarget;
 
-		// Apply swing correction
-		ApplyRotationDelta(AngularDriveStiffness, SwingAxis, DSwingAngle, SwingAxis, -DSwingAngle);
+		ApplyDriveRotationDelta(Dt, AngularDriveStiffness, AngularDriveDamping, SwingAxis, SwingAxis, DSwingAngle, SwingLambda);
 		UpdateDerivedState();
 	}
 
@@ -444,20 +491,20 @@ namespace Chaos
 		const EJointAngularConstraintIndex SwingConstraintIndex,
 		const EJointAngularAxisIndex SwingAxisIndex)
 	{
+		// @todo(ccaulfield): this can be made much simpler I think...swing axis should just be Body0's swing axis
+
 		// Decompose rotation of body 1 relative to body 0 into swing and twist rotations, assuming twist is X axis
 		FRotation3 R01Twist, R01Swing;
 		FPBDJointUtilities::DecomposeSwingTwistLocal(Rs[0], Rs[1], R01Swing, R01Twist);
 
-		FVec3 TwistAxis01;
-		FReal TwistAngle;
-		R01Twist.ToAxisAndAngleSafe(TwistAxis01, TwistAngle, FJointConstants::TwistAxis(), SwingTwistAngleTolerance);
+		FVec3 TwistAxis01 = FJointConstants::TwistAxis();
+		FReal TwistAngle = R01Twist.GetAngle();
 		if (TwistAngle > PI)
 		{
 			TwistAngle = TwistAngle - (FReal)2 * PI;
 		}
-		if (FVec3::DotProduct(TwistAxis01, FJointConstants::TwistAxis()) < 0)
+		if (R01Twist.X < 0)
 		{
-			TwistAxis01 = -TwistAxis01;
 			TwistAngle = -TwistAngle;
 		}
 		const FVec3 TwistAxis = Rs[0] * TwistAxis01;
@@ -502,7 +549,7 @@ namespace Chaos
 			}
 
 			// Apply swing correction
-			ApplyRotationDelta(SwingStiffness, SwingAxis, DSwingAngle, SwingAxis, -DSwingAngle);
+			ApplyRotationDelta(SwingStiffness, SwingAxis, SwingAxis, DSwingAngle);
 			UpdateDerivedState();
 		}
 	}
@@ -522,29 +569,19 @@ namespace Chaos
 		const FReal Dt,
 		const FPBDJointSettings& JointSettings)
 	{
-		// @todo(ccaulfield): I think this can be done in a simpler way, reusing ApplyRotationDelta...
-
 		// Calculate the rotation we need to apply to resolve the rotation delta
 		const FRotation3 TargetR1 = Rs[0] * JointSettings.Motion.AngularDriveTarget;
 		const FRotation3 DR1 = TargetR1 * Rs[1].Inverse();
-		const FRotation3 TargetQ0 = DR1.Inverse() * Qs[0];
-		const FRotation3 TargetQ1 = DR1 * Qs[1];
 
 		FVec3 SLerpAxis;
 		FReal SLerpAngle;
 		if (DR1.ToAxisAndAngleSafe(SLerpAxis, SLerpAngle, FVec3(1, 0, 0)))
 		{
-			const FMatrix33 InvI0 = Utilities::ComputeWorldSpaceInertia(Qs[0], InvILs[0]);
-			const FMatrix33 InvI1 = Utilities::ComputeWorldSpaceInertia(Qs[1], InvILs[1]);
-			const FReal I0 = FVec3::DotProduct(SLerpAxis, Utilities::Multiply(InvI0, SLerpAxis));
-			const FReal I1 = FVec3::DotProduct(SLerpAxis, Utilities::Multiply(InvI1, SLerpAxis));
-			const FReal F0 = AngularDriveStiffness * I0 / (I0 + I1);
-			const FReal F1 = AngularDriveStiffness * I1 / (I0 + I1);
-
-			// Apply the rotation delta about the connector
-			Qs[0] = FRotation3::Slerp(Qs[0], TargetQ0, F0);
-			Qs[1] = FRotation3::Slerp(Qs[1], TargetQ1, F1);
-			Qs[1].EnforceShortestArcWith(Qs[0]);
+			if (SLerpAngle > (FReal)PI)
+			{
+				SLerpAngle = SLerpAngle - (FReal)2 * PI;
+			}
+			ApplyDriveRotationDelta(Dt, AngularDriveStiffness, AngularDriveDamping, SLerpAxis, SLerpAxis, -SLerpAngle, SwingLambda);
 		}
 	}
 
@@ -590,43 +627,32 @@ namespace Chaos
 		const FReal Dt,
 		const FPBDJointSettings& JointSettings)
 	{
-		// Apply a position correction with the parent body set to infinite mass, then correct the velocity
-		// in the direction of the position correction.
-		const FVec3 CX = FPBDJointUtilities::GetLimitedPositionError(JointSettings, Rs[0], Xs[1] - Xs[0]);
+		// Apply a position correction with the parent body set to infinite mass, then correct the velocity.
+		FVec3 CX = FPBDJointUtilities::GetLimitedPositionError(JointSettings, Rs[0], Xs[1] - Xs[0]);
 		const FReal CXLen = CX.Size();
 		if (CXLen > KINDA_SMALL_NUMBER)
 		{
 			const FVec3 CXDir = CX / CXLen;
-
-			// Get required velocity correction in direction of position correction
 			const FVec3 V0 = Vs[0] + FVec3::CrossProduct(Ws[0], Xs[0] - Ps[0]);
 			const FVec3 V1 = Vs[1] + FVec3::CrossProduct(Ws[1], Xs[1] - Ps[1]);
-			FVec3 CV = FPBDJointUtilities::GetLimitedVelocityError(JointSettings, Rs[0], CX, V1 - V0);
-			CV = CV - FVec3::DotProduct(CXDir, CV) * CXDir;
+			FVec3 CV = FVec3::DotProduct(V1 - V0, CXDir) * CXDir;
 
 			FMatrix33 InvI1 = Utilities::ComputeWorldSpaceInertia(Qs[1], InvILs[1]);
 			FMatrix33 M1 = Utilities::ComputeJointFactorMatrix(Xs[1] - Ps[1], InvI1, InvMs[1]);
 			FMatrix33 MI = M1.Inverse();
+
 			const FVec3 DX = Utilities::Multiply(MI, CX);
+			const FVec3 DV = Utilities::Multiply(MI, CV);
 
-			// Apply position correction
 			const FVec3 DP1 = -InvMs[1] * DX;
 			const FVec3 DR1 = Utilities::Multiply(InvI1, FVec3::CrossProduct(Xs[1] - Ps[1], -DX));
-
-			ApplyPositionDelta(LinearStiffness, FVec3(0), DP1);
-			ApplyRotationDelta(LinearStiffness, FVec3(0), DR1);
-			UpdateDerivedState();
-
-			InvI1 = Utilities::ComputeWorldSpaceInertia(Qs[1], InvILs[1]);
-			M1 = Utilities::ComputeJointFactorMatrix(Xs[1] - Ps[1], InvI1, InvMs[1]);
-			MI = M1.Inverse();
-
-			// Apply velocity correction
-			const FVec3 DV = Utilities::Multiply(MI, CV);
 			const FVec3 DV1 = -InvMs[1] * DV;
 			const FVec3 DW1 = Utilities::Multiply(InvI1, FVec3::CrossProduct(Xs[1] - Ps[1], -DV));
 
+			ApplyPositionDelta(LinearProjection, FVec3(0), DP1);
+			ApplyRotationDelta(LinearProjection, FVec3(0), DR1);
 			ApplyVelocityDelta(LinearProjection, FVec3(0), FVec3(0), DV1, DW1);
+			UpdateDerivedState();
 		}
 	}
 }
diff --git a/Engine/Source/Runtime/Experimental/Chaos/Private/Chaos/PBDJointConstraintUtilities.cpp b/Engine/Source/Runtime/Experimental/Chaos/Private/Chaos/PBDJointConstraintUtilities.cpp
--- a/Engine/Source/Runtime/Experimental/Chaos/Private/Chaos/PBDJointConstraintUtilities.cpp
+++ b/Engine/Source/Runtime/Experimental/Chaos/Private/Chaos/PBDJointConstraintUtilities.cpp
@@ -57,6 +57,14 @@ namespace Chaos
 		return (SolverSettings.DriveStiffness > 0.0f) ? SolverSettings.DriveStiffness : JointSettings.Motion.AngularDriveStiffness;
 	}
 
+
+	FReal FPBDJointUtilities::GetAngularDriveDamping(
+		const FPBDJointSolverSettings& SolverSettings,
+		const FPBDJointSettings& JointSettings)
+	{
+		return (SolverSettings.DriveDamping > 0.0f) ? SolverSettings.DriveDamping: JointSettings.Motion.AngularDriveDamping;
+	}
+
 	FReal FPBDJointUtilities::GetLinearProjection(
 		const FPBDJointSolverSettings& SolverSettings,
 		const FPBDJointSettings& JointSettings)
diff --git a/Engine/Source/Runtime/Experimental/Chaos/Private/Chaos/PBDJointConstraints.cpp b/Engine/Source/Runtime/Experimental/Chaos/Private/Chaos/PBDJointConstraints.cpp
--- a/Engine/Source/Runtime/Experimental/Chaos/Private/Chaos/PBDJointConstraints.cpp
+++ b/Engine/Source/Runtime/Experimental/Chaos/Private/Chaos/PBDJointConstraints.cpp
@@ -43,9 +43,9 @@ namespace Chaos
 	}
 
 	
-	void FPBDJointConstraintHandle::CalculateConstraintSpace(FVec3& OutXa, FMatrix33& OutRa, FVec3& OutXb, FMatrix33& OutRb, FVec3& OutCR) const
+	void FPBDJointConstraintHandle::CalculateConstraintSpace(FVec3& OutXa, FMatrix33& OutRa, FVec3& OutXb, FMatrix33& OutRb) const
 	{
-		ConstraintContainer->CalculateConstraintSpace(ConstraintIndex, OutXa, OutRa, OutXb, OutRb, OutCR);
+		ConstraintContainer->CalculateConstraintSpace(ConstraintIndex, OutXa, OutRa, OutXb, OutRb);
 	}
 
 	
@@ -96,6 +96,7 @@ namespace Chaos
 		, bAngularTwistDriveEnabled(false)
 		, bAngularSwingDriveEnabled(false)
 		, AngularDriveStiffness(0)
+		, AngularDriveDamping(0)
 	{
 	}
 
@@ -120,6 +121,7 @@ namespace Chaos
 		, bAngularTwistDriveEnabled(false)
 		, bAngularSwingDriveEnabled(false)
 		, AngularDriveStiffness(0)
+		, AngularDriveDamping(0)
 	{
 	}
 
@@ -176,6 +178,7 @@ namespace Chaos
 		, AngularProjection((FReal)0)
 		, Stiffness((FReal)0)
 		, DriveStiffness((FReal)0)
+		, DriveDamping((FReal)0)
 		, SoftLinearStiffness((FReal)0)
 		, SoftAngularStiffness((FReal)0)
 	{
@@ -419,10 +422,35 @@ namespace Chaos
 			SortConstraints();
 			bRequiresSort = false;
 		}
+
+		ConstraintSolvers.SetNum(NumConstraints());
+		for (int32 ConstraintIndex = 0; ConstraintIndex < NumConstraints(); ++ConstraintIndex)
+		{
+			const FPBDJointSettings& JointSettings = ConstraintSettings[ConstraintIndex];
+			FJointSolverGaussSeidel& Solver = ConstraintSolvers[ConstraintIndex];
+
+			int32 Index0, Index1;
+			GetConstrainedParticleIndices(ConstraintIndex, Index0, Index1);
+			TGenericParticleHandle<FReal, 3> Particle0 = TGenericParticleHandle<FReal, 3>(ConstraintParticles[ConstraintIndex][Index0]);
+			TGenericParticleHandle<FReal, 3> Particle1 = TGenericParticleHandle<FReal, 3>(ConstraintParticles[ConstraintIndex][Index1]);
+
+			Solver.Init(
+				Dt,
+				Settings,
+				JointSettings,
+				FParticleUtilitiesXR::GetCoMWorldRotation(Particle0),	// Prev rotation
+				FParticleUtilitiesXR::GetCoMWorldRotation(Particle1),	// Prev rotation
+				Particle0->InvM(),
+				Particle0->InvI(),
+				Particle1->InvM(),
+				Particle1->InvI(),
+				FParticleUtilities::ParticleLocalToCoMLocal(Particle0, JointSettings.ConstraintFrames[Index0]),
+				FParticleUtilities::ParticleLocalToCoMLocal(Particle1, JointSettings.ConstraintFrames[Index1]));
+		}
 	}
 
 	
-	void FPBDJointConstraints::CalculateConstraintSpace(int32 ConstraintIndex, FVec3& OutX0, FMatrix33& OutR0, FVec3& OutX1, FMatrix33& OutR1, FVec3& OutCR) const
+	void FPBDJointConstraints::CalculateConstraintSpace(int32 ConstraintIndex, FVec3& OutX0, FMatrix33& OutR0, FVec3& OutX1, FMatrix33& OutR1) const
 	{
 		int32 Index0, Index1;
 		GetConstrainedParticleIndices(ConstraintIndex, Index0, Index1);
@@ -436,16 +464,11 @@ namespace Chaos
 		const FPBDJointSettings& JointSettings = ConstraintSettings[ConstraintIndex];
 		const FRigidTransform3 XL0 = FParticleUtilities::ParticleLocalToCoMLocal(Particle0, JointSettings.ConstraintFrames[Index0]);
 		const FRigidTransform3 XL1 = FParticleUtilities::ParticleLocalToCoMLocal(Particle1, JointSettings.ConstraintFrames[Index1]);
-		EJointMotionType Swing1Motion = JointSettings.Motion.AngularMotionTypes[(int32)EJointAngularConstraintIndex::Swing1];
-		EJointMotionType Swing2Motion = JointSettings.Motion.AngularMotionTypes[(int32)EJointAngularConstraintIndex::Swing2];
-		if ((Swing1Motion == EJointMotionType::Limited) && (Swing2Motion == EJointMotionType::Limited))
-		{
-			FPBDJointUtilities::CalculateConeConstraintSpace(Settings, ConstraintSettings[ConstraintIndex], XL0, XL1, P0, Q0, P1, Q1, OutX0, OutR0, OutX1, OutR1, OutCR);
-		}
-		else
-		{
-			FPBDJointUtilities::CalculateSwingConstraintSpace(Settings, ConstraintSettings[ConstraintIndex], XL0, XL1, P0, Q0, P1, Q1, OutX0, OutR0, OutX1, OutR1, OutCR);
-		}
+
+		OutX0 = P0 + Q0 * XL0.GetTranslation();
+		OutX1 = P1 + Q1 * XL1.GetTranslation();
+		OutR0 = FRotation3(Q0 * XL0.GetRotation()).ToMatrix();
+		OutR1 = FRotation3(Q1 * XL1.GetRotation()).ToMatrix();
 	}
 
 	//////////////////////////////////////////////////////////////////////////
@@ -710,27 +733,23 @@ namespace Chaos
 		UE_LOG(LogChaosJoint, VeryVerbose, TEXT("Solve Joint Constraint %d %s %s (dt = %f; it = %d / %d)"), ConstraintIndex, *Constraint[0]->ToString(), *Constraint[1]->ToString(), Dt, It, NumIts);
 
 		const FPBDJointSettings& JointSettings = ConstraintSettings[ConstraintIndex];
+		FJointSolverGaussSeidel& Solver = ConstraintSolvers[ConstraintIndex];
 
 		int32 Index0, Index1;
 		GetConstrainedParticleIndices(ConstraintIndex, Index0, Index1);
 		TGenericParticleHandle<FReal, 3> Particle0 = TGenericParticleHandle<FReal, 3>(ConstraintParticles[ConstraintIndex][Index0]);
 		TGenericParticleHandle<FReal, 3> Particle1 = TGenericParticleHandle<FReal, 3>(ConstraintParticles[ConstraintIndex][Index1]);
 
-		FJointSolverGaussSeidel Solver;
-		Solver.InitConstraints(
+		Solver.Update(
 			Dt,
-			Settings,
-			JointSettings,
 			FParticleUtilities::GetCoMWorldPosition(Particle0),
 			FParticleUtilities::GetCoMWorldRotation(Particle0),
+			Particle0->V(),
+			Particle0->W(),
 			FParticleUtilities::GetCoMWorldPosition(Particle1),
 			FParticleUtilities::GetCoMWorldRotation(Particle1),
-			Particle0->InvM(),
-			Particle0->InvI(),
-			Particle1->InvM(),
-			Particle1->InvI(),
-			FParticleUtilities::ParticleLocalToCoMLocal(Particle0, JointSettings.ConstraintFrames[Index0]),
-			FParticleUtilities::ParticleLocalToCoMLocal(Particle1, JointSettings.ConstraintFrames[Index1]));
+			Particle1->V(),
+			Particle1->W());
 
 		for (int32 PairIt = 0; PairIt < NumPairIts; ++PairIt)
 		{
@@ -748,17 +767,15 @@ namespace Chaos
 		UE_LOG(LogChaosJoint, VeryVerbose, TEXT("Solve Joint Constraint %d %s %s (dt = %f; it = %d / %d)"), ConstraintIndex, *Constraint[0]->ToString(), *Constraint[1]->ToString(), Dt, It, NumIts);
 
 		const FPBDJointSettings& JointSettings = ConstraintSettings[ConstraintIndex];
+		FJointSolverGaussSeidel& Solver = ConstraintSolvers[ConstraintIndex];
 
 		int32 Index0, Index1;
 		GetConstrainedParticleIndices(ConstraintIndex, Index0, Index1);
 		TGenericParticleHandle<FReal, 3> Particle0 = TGenericParticleHandle<FReal, 3>(ConstraintParticles[ConstraintIndex][Index0]);
 		TGenericParticleHandle<FReal, 3> Particle1 = TGenericParticleHandle<FReal, 3>(ConstraintParticles[ConstraintIndex][Index1]);
 
-		FJointSolverGaussSeidel Solver;
-		Solver.InitProjection(
+		Solver.Update(
 			Dt,
-			Settings,
-			JointSettings,
 			FParticleUtilities::GetCoMWorldPosition(Particle0),
 			FParticleUtilities::GetCoMWorldRotation(Particle0),
 			Particle0->V(),
@@ -766,13 +783,7 @@ namespace Chaos
 			FParticleUtilities::GetCoMWorldPosition(Particle1),
 			FParticleUtilities::GetCoMWorldRotation(Particle1),
 			Particle1->V(),
-			Particle1->W(),
-			Particle0->InvM(),
-			Particle0->InvI(),
-			Particle1->InvM(),
-			Particle1->InvI(),
-			FParticleUtilities::ParticleLocalToCoMLocal(Particle0, JointSettings.ConstraintFrames[Index0]),
-			FParticleUtilities::ParticleLocalToCoMLocal(Particle1, JointSettings.ConstraintFrames[Index1]));
+			Particle1->W());
 
 		for (int32 PairIt = 0; PairIt < NumPairIts; ++PairIt)
 		{
diff --git a/Engine/Source/Runtime/Experimental/Chaos/Public/Chaos/Joint/PBDJointSolverGaussSeidel.h b/Engine/Source/Runtime/Experimental/Chaos/Public/Chaos/Joint/PBDJointSolverGaussSeidel.h
--- a/Engine/Source/Runtime/Experimental/Chaos/Public/Chaos/Joint/PBDJointSolverGaussSeidel.h
+++ b/Engine/Source/Runtime/Experimental/Chaos/Public/Chaos/Joint/PBDJointSolverGaussSeidel.h
@@ -50,14 +50,12 @@ namespace Chaos
 
 		FJointSolverGaussSeidel();
 
-		void InitConstraints(
+		void Init(
 			const FReal Dt,
 			const FPBDJointSolverSettings& SolverSettings,
 			const FPBDJointSettings& JointSettings,
-			const FVec3& P0,
-			const FRotation3& Q0,
-			const FVec3& P1,
-			const FRotation3& Q1,
+			const FRotation3& PrevQ0,
+			const FRotation3& PrevQ1,
 			const FReal InvM0,
 			const FMatrix33& InvIL0,
 			const FReal InvM1,
@@ -65,10 +63,8 @@ namespace Chaos
 			const FRigidTransform3& XL0,
 			const FRigidTransform3& XL1);
 
-		void InitProjection(
+		void Update(
 			const FReal Dt,
-			const FPBDJointSolverSettings& SolverSettings,
-			const FPBDJointSettings& JointSettings,
 			const FVec3& P0,
 			const FRotation3& Q0,
 			const FVec3& V0,
@@ -76,13 +72,7 @@ namespace Chaos
 			const FVec3& P1,
 			const FRotation3& Q1,
 			const FVec3& V1,
-			const FVec3& W1,
-			const FReal InvM0,
-			const FMatrix33& InvIL0,
-			const FReal InvM1,
-			const FMatrix33& InvIL1,
-			const FRigidTransform3& XL0,
-			const FRigidTransform3& XL1);
+			const FVec3& W1);
 
 		void ApplyConstraints(
 			const FReal Dt,
@@ -126,9 +116,17 @@ namespace Chaos
 		void ApplyRotationDelta(
 			const FReal Stiffness,
 			const FVec3& Axis0,
-			const FReal Angle0,
 			const FVec3& Axis1,
-			const FReal Angle1);
+			const FReal Angle);
+
+		void ApplyDriveRotationDelta(
+			const FReal Dt,
+			const FReal Stiffness,
+			const FReal Damping,
+			const FVec3& Axis0,
+			const FVec3& Axis1,
+			const FReal Angle,
+			FReal& Lambda);
 
 		void ApplyVelocityDelta(
 			const FReal Stiffness,
@@ -194,11 +192,19 @@ namespace Chaos
 		FVec3 Vs[MaxConstrainedBodies];				// World-space particle CoM velocities
 		FVec3 Ws[MaxConstrainedBodies];				// World-space particle CoM angular velocities
 
+		// XPBD Previous iteration world-space body state
+		FRotation3 PrevQs[MaxConstrainedBodies];	// World-space particle CoM rotations
+
+		// XPBD joint drive constraint multipliers
+		FReal TwistLambda;
+		FReal SwingLambda;
+
 		// Settings
 		FReal LinearStiffness;
 		FReal TwistStiffness;
 		FReal SwingStiffness;
 		FReal AngularDriveStiffness;
+		FReal AngularDriveDamping;
 		FReal LinearProjection;
 		FReal AngularProjection;
 		FReal SwingTwistAngleTolerance;
diff --git a/Engine/Source/Runtime/Experimental/Chaos/Public/Chaos/PBDJointConstraintTypes.h b/Engine/Source/Runtime/Experimental/Chaos/Public/Chaos/PBDJointConstraintTypes.h
--- a/Engine/Source/Runtime/Experimental/Chaos/Public/Chaos/PBDJointConstraintTypes.h
+++ b/Engine/Source/Runtime/Experimental/Chaos/Public/Chaos/PBDJointConstraintTypes.h
@@ -96,6 +96,7 @@ namespace Chaos
 		bool bAngularTwistDriveEnabled;
 		bool bAngularSwingDriveEnabled;
 		FReal AngularDriveStiffness;
+		FReal AngularDriveDamping;
 	};
 
 
@@ -137,6 +138,7 @@ namespace Chaos
 		FReal AngularProjection;
 		FReal Stiffness;
 		FReal DriveStiffness;
+		FReal DriveDamping;
 		FReal SoftLinearStiffness;
 		FReal SoftAngularStiffness;
 	};
diff --git a/Engine/Source/Runtime/Experimental/Chaos/Public/Chaos/PBDJointConstraintUtilities.h b/Engine/Source/Runtime/Experimental/Chaos/Public/Chaos/PBDJointConstraintUtilities.h
--- a/Engine/Source/Runtime/Experimental/Chaos/Public/Chaos/PBDJointConstraintUtilities.h
+++ b/Engine/Source/Runtime/Experimental/Chaos/Public/Chaos/PBDJointConstraintUtilities.h
@@ -80,6 +80,10 @@ namespace Chaos
 			const FPBDJointSolverSettings& SolverSettings,
 			const FPBDJointSettings& JointSettings);
 
+		static FReal GetAngularDriveDamping(
+			const FPBDJointSolverSettings& SolverSettings,
+			const FPBDJointSettings& JointSettings);
+
 		static FReal GetLinearProjection(
 			const FPBDJointSolverSettings& SolverSettings,
 			const FPBDJointSettings& JointSettings);
diff --git a/Engine/Source/Runtime/Experimental/Chaos/Public/Chaos/PBDJointConstraints.h b/Engine/Source/Runtime/Experimental/Chaos/Public/Chaos/PBDJointConstraints.h
--- a/Engine/Source/Runtime/Experimental/Chaos/Public/Chaos/PBDJointConstraints.h
+++ b/Engine/Source/Runtime/Experimental/Chaos/Public/Chaos/PBDJointConstraints.h
@@ -9,6 +9,7 @@
 #include "Chaos/Vector.h"
 
 #include "Chaos/ConstraintHandle.h"
+#include "Chaos/Joint/PBDJointSolverGaussSeidel.h"
 #include "Chaos/ParticleHandleFwd.h"
 #include "Chaos/PBDConstraintContainer.h"
 #include "Chaos/PBDJointConstraintTypes.h"
@@ -25,7 +26,7 @@ namespace Chaos
 		CHAOS_API FPBDJointConstraintHandle();
 		CHAOS_API FPBDJointConstraintHandle(FConstraintContainer* InConstraintContainer, int32 InConstraintIndex);
 
-		CHAOS_API void CalculateConstraintSpace(FVec3& OutXa, FMatrix33& OutRa, FVec3& OutXb, FMatrix33& OutRb, FVec3& OutCR) const;
+		CHAOS_API void CalculateConstraintSpace(FVec3& OutXa, FMatrix33& OutRa, FVec3& OutXb, FMatrix33& OutRb) const;
 		CHAOS_API void SetParticleLevels(const TVector<int32, 2>& ParticleLevels);
 		CHAOS_API int32 GetConstraintLevel() const;
 		CHAOS_API const FPBDJointSettings& GetSettings() const;
@@ -154,8 +155,10 @@ namespace Chaos
 	private:
 		friend class FPBDJointConstraintHandle;
 
+		void PrepareConstraints(FReal Dt);
+
 		void GetConstrainedParticleIndices(const int32 ConstraintIndex, int32& Index0, int32& Index1) const;
-		void CalculateConstraintSpace(int32 ConstraintIndex, FVec3& OutX0, FMatrix33& OutR0, FVec3& OutX1, FMatrix33& OutR1, FVec3& OutAngles) const;
+		void CalculateConstraintSpace(int32 ConstraintIndex, FVec3& OutX0, FMatrix33& OutR0, FVec3& OutX1, FMatrix33& OutR1) const;
 		void UpdateParticleState(TPBDRigidParticleHandle<FReal, 3>* Rigid, const FReal Dt, const FVec3& P, const FRotation3& Q, const bool bUpdateVelocity = true);
 		void UpdateParticleState(TPBDRigidParticleHandle<FReal, 3>* Rigid, const FReal Dt, const FVec3& P, const FRotation3& Q, const FVec3& V, const FVec3& W);
 		void SortConstraints();
@@ -177,6 +180,9 @@ namespace Chaos
 		FJointPostApplyCallback PostApplyCallback;
 		FJointPostApplyCallback PostProjectCallback;
 
+		// @todo(ccaulfield): optimize storage for joint solver
+		TArray<FJointSolverGaussSeidel> ConstraintSolvers;
+
 		bool bRequiresSort;
 	};
 
