Chaos - Prevent NaN in PerParticleDampVelocity is called with no active particle indices.

diff --git a/Engine/Source/Runtime/Experimental/Chaos/Public/Chaos/PerParticleDampVelocity.h b/Engine/Source/Runtime/Experimental/Chaos/Public/Chaos/PerParticleDampVelocity.h
--- a/Engine/Source/Runtime/Experimental/Chaos/Public/Chaos/PerParticleDampVelocity.h
+++ b/Engine/Source/Runtime/Experimental/Chaos/Public/Chaos/PerParticleDampVelocity.h
@@ -34,8 +34,13 @@ namespace Chaos
 				MVcm += InParticles.V(Index) * InParticles.M(Index);
 				Mcm += InParticles.M(Index);
 			}
-			MXcm /= Mcm;
-			MVcm /= Mcm;
+
+			if (Mcm != 0.0f)
+			{
+				MXcm /= Mcm;
+				MVcm /= Mcm;
+			}
+
 			TVector<T, d> L = TVector<T, d>(0.f, 0.f, 0.f);
 			PMatrix<T, d, d> I(0);
 			for (const int32 Index : InActiveIndices)
