Chaos - Joint minor cleanup

diff --git a/Engine/Source/Runtime/Experimental/Chaos/Private/Chaos/PBDJointConstraints.cpp b/Engine/Source/Runtime/Experimental/Chaos/Private/Chaos/PBDJointConstraints.cpp
--- a/Engine/Source/Runtime/Experimental/Chaos/Private/Chaos/PBDJointConstraints.cpp
+++ b/Engine/Source/Runtime/Experimental/Chaos/Private/Chaos/PBDJointConstraints.cpp
@@ -423,29 +423,38 @@ namespace Chaos
 			bRequiresSort = false;
 		}
 
-		ConstraintSolvers.SetNum(NumConstraints());
-		for (int32 ConstraintIndex = 0; ConstraintIndex < NumConstraints(); ++ConstraintIndex)
+		PrepareConstraints(Dt);
+	}
+
+
+	void FPBDJointConstraints::PrepareConstraints(FReal Dt)
+	{
+		if (!ChaosJoint_UseCholeskySolver)
 		{
-			const FPBDJointSettings& JointSettings = ConstraintSettings[ConstraintIndex];
-			FJointSolverGaussSeidel& Solver = ConstraintSolvers[ConstraintIndex];
-
-			int32 Index0, Index1;
-			GetConstrainedParticleIndices(ConstraintIndex, Index0, Index1);
-			TGenericParticleHandle<FReal, 3> Particle0 = TGenericParticleHandle<FReal, 3>(ConstraintParticles[ConstraintIndex][Index0]);
-			TGenericParticleHandle<FReal, 3> Particle1 = TGenericParticleHandle<FReal, 3>(ConstraintParticles[ConstraintIndex][Index1]);
-
-			Solver.Init(
-				Dt,
-				Settings,
-				JointSettings,
-				FParticleUtilitiesXR::GetCoMWorldRotation(Particle0),	// Prev rotation
-				FParticleUtilitiesXR::GetCoMWorldRotation(Particle1),	// Prev rotation
-				Particle0->InvM(),
-				Particle0->InvI(),
-				Particle1->InvM(),
-				Particle1->InvI(),
-				FParticleUtilities::ParticleLocalToCoMLocal(Particle0, JointSettings.ConstraintFrames[Index0]),
-				FParticleUtilities::ParticleLocalToCoMLocal(Particle1, JointSettings.ConstraintFrames[Index1]));
+			ConstraintSolvers.SetNum(NumConstraints());
+			for (int32 ConstraintIndex = 0; ConstraintIndex < NumConstraints(); ++ConstraintIndex)
+			{
+				const FPBDJointSettings& JointSettings = ConstraintSettings[ConstraintIndex];
+				FJointSolverGaussSeidel& Solver = ConstraintSolvers[ConstraintIndex];
+
+				int32 Index0, Index1;
+				GetConstrainedParticleIndices(ConstraintIndex, Index0, Index1);
+				TGenericParticleHandle<FReal, 3> Particle0 = TGenericParticleHandle<FReal, 3>(ConstraintParticles[ConstraintIndex][Index0]);
+				TGenericParticleHandle<FReal, 3> Particle1 = TGenericParticleHandle<FReal, 3>(ConstraintParticles[ConstraintIndex][Index1]);
+
+				Solver.Init(
+					Dt,
+					Settings,
+					JointSettings,
+					FParticleUtilitiesXR::GetCoMWorldRotation(Particle0),	// Prev rotation
+					FParticleUtilitiesXR::GetCoMWorldRotation(Particle1),	// Prev rotation
+					Particle0->InvM(),
+					Particle0->InvI(),
+					Particle1->InvM(),
+					Particle1->InvI(),
+					FParticleUtilities::ParticleLocalToCoMLocal(Particle0, JointSettings.ConstraintFrames[Index0]),
+					FParticleUtilities::ParticleLocalToCoMLocal(Particle1, JointSettings.ConstraintFrames[Index1]));
+			}
 		}
 	}
 
