Chaos trimeshes can use small indices (saves 10% trimesh memory)

diff --git a/Engine/Source/Runtime/Core/Public/UObject/ExternalPhysicsCustomObjectVersion.h b/Engine/Source/Runtime/Core/Public/UObject/ExternalPhysicsCustomObjectVersion.h
--- a/Engine/Source/Runtime/Core/Public/UObject/ExternalPhysicsCustomObjectVersion.h
+++ b/Engine/Source/Runtime/Core/Public/UObject/ExternalPhysicsCustomObjectVersion.h
@@ -92,6 +92,9 @@ struct CORE_API FExternalPhysicsCustomObjectVersion
 
 		// Serialize whether or not an AABBTree is immutable
 		ImmutableAABBTree,
+
+		// Trimeshes can now use small indices
+		TrimeshCanUseSmallIndices,
 		
 		// -----<new versions can be added above this line>-------------------------------------------------
 		VersionPlusOne,
diff --git a/Engine/Source/Runtime/Engine/Private/PhysicsEngine/Experimental/ChaosDerivedData.cpp b/Engine/Source/Runtime/Engine/Private/PhysicsEngine/Experimental/ChaosDerivedData.cpp
--- a/Engine/Source/Runtime/Engine/Private/PhysicsEngine/Experimental/ChaosDerivedData.cpp
+++ b/Engine/Source/Runtime/Engine/Private/PhysicsEngine/Experimental/ChaosDerivedData.cpp
@@ -21,7 +21,7 @@ const TCHAR* FChaosDerivedDataCooker::GetPluginName() const
 
 const TCHAR* FChaosDerivedDataCooker::GetVersionString() const
 {
-	return TEXT("E138DCCB71AD4DBCA7AA0E469571F332");
+	return TEXT("53786FE2178C49A98EC62DD50313FE52");
 }
 
 FString FChaosDerivedDataCooker::GetPluginSpecificCacheKeySuffix() const
@@ -119,15 +119,16 @@ void FChaosDerivedDataCooker::BuildTriangleMeshes(TArray<TUniquePtr<Chaos::FTria
 	// Build chaos triangle list. #BGTODO Just make the clean function take these types instead of double copying
 	const int32 NumTriangles = FinalIndices.Num() / 3;
 	bool bHasMaterials = InParams.TriangleMeshDesc.MaterialIndices.Num() > 0;
-	TArray<Chaos::TVector<int32, 3>> Triangles;
 	TArray<uint16> MaterialIndices;
-	Triangles.Reserve(NumTriangles);
 
+	auto LambdaHelper = [&](auto& Triangles)
+	{
 	if(bHasMaterials)
 	{
 		MaterialIndices.Reserve(NumTriangles);
 	}
 
+		Triangles.Reserve(NumTriangles);
 	for(int32 TriangleIndex = 0; TriangleIndex < NumTriangles; ++TriangleIndex)
 	{
 		// Only add this triangle if it is valid
@@ -169,6 +170,18 @@ void FChaosDerivedDataCooker::BuildTriangleMeshes(TArray<TUniquePtr<Chaos::FTria
 	}
 
 	OutTriangleMeshes.Emplace(new Chaos::FTriangleMeshImplicitObject(MoveTemp(TriMeshParticles), MoveTemp(Triangles), MoveTemp(MaterialIndices)));
+	};
+
+	if(FinalVerts.Num() < TNumericLimits<uint16>::Max())
+	{
+		TArray<Chaos::TVector<uint16, 3>> TrianglesSmallIdx;
+		LambdaHelper(TrianglesSmallIdx);
+	}
+	else
+	{
+		TArray<Chaos::TVector<int32, 3>> TrianglesLargeIdx;
+		LambdaHelper(TrianglesLargeIdx);
+	}	
 }
 
 void FChaosDerivedDataCooker::BuildConvexMeshes(TArray<TUniquePtr<Chaos::FImplicitObject>>& OutConvexMeshes, const FCookBodySetupInfo& InParams)
diff --git a/Engine/Source/Runtime/Experimental/Chaos/Private/Chaos/TriangleMeshImplicitObject.cpp b/Engine/Source/Runtime/Experimental/Chaos/Private/Chaos/TriangleMeshImplicitObject.cpp
--- a/Engine/Source/Runtime/Experimental/Chaos/Private/Chaos/TriangleMeshImplicitObject.cpp
+++ b/Engine/Source/Runtime/Experimental/Chaos/Private/Chaos/TriangleMeshImplicitObject.cpp
@@ -8,23 +8,12 @@
 
 namespace Chaos
 {
-FTriangleMeshImplicitObject::FTriangleMeshImplicitObject(TParticles<FReal, 3>&& Particles, TArray<TVector<int32, 3>>&& Elements, TArray<uint16>&& InMaterialIndices)
-	: FImplicitObject(EImplicitObject::HasBoundingBox, ImplicitObjectType::TriangleMesh)
-	, MParticles(MoveTemp(Particles))
-	, MElements(MoveTemp(Elements))
-	, MLocalBoundingBox(MParticles.X(0), MParticles.X(0))
-    , MaterialIndices(MoveTemp(InMaterialIndices))
-{
-	for (uint32 Idx = 1; Idx < MParticles.Size(); ++Idx)
-	{
-		MLocalBoundingBox.GrowToInclude(MParticles.X(Idx));
-	}
-	RebuildBV();
-}
 
+
+template <typename IdxType>
 struct FTriangleMeshRaycastVisitor
 {
-	FTriangleMeshRaycastVisitor(const FVec3& InStart, const FVec3& InDir, const FReal InThickness, const TParticles<FReal,3>& InParticles, const TArray<TVector<int32, 3>>& InElements)
+	FTriangleMeshRaycastVisitor(const FVec3& InStart, const FVec3& InDir, const FReal InThickness, const TParticles<FReal,3>& InParticles, const TArray<TVector<IdxType, 3>>& InElements)
 	: Particles(InParticles)
 	, Elements(InElements)
 	, StartPoint(InStart)
@@ -188,7 +177,7 @@ struct FTriangleMeshRaycastVisitor
 	}
 
 	const TParticles<FReal, 3>& Particles;
-	const TArray<TVector<int32, 3>>& Elements;
+	const TArray<TVector<IdxType, 3>>& Elements;
 	const FVec3& StartPoint;
 	const FVec3& Dir;
 	const FReal Thickness;
@@ -204,9 +193,10 @@ FReal FTriangleMeshImplicitObject::PhiWithNormal(const FVec3& x, FVec3& Normal)
 	return 0;
 }
 
-bool FTriangleMeshImplicitObject::Raycast(const FVec3& StartPoint, const FVec3& Dir, const FReal Length, const FReal Thickness, FReal& OutTime, FVec3& OutPosition, FVec3& OutNormal, int32& OutFaceIndex) const
+template <typename IdxType>
+bool FTriangleMeshImplicitObject::RaycastImp(const TArray<TVector<IdxType, 3>>& Elements, const FVec3& StartPoint, const FVec3& Dir, const FReal Length, const FReal Thickness, FReal& OutTime, FVec3& OutPosition, FVec3& OutNormal, int32& OutFaceIndex) const
 {
-	FTriangleMeshRaycastVisitor SQVisitor(StartPoint, Dir, Thickness, MParticles, MElements);
+	FTriangleMeshRaycastVisitor<IdxType> SQVisitor(StartPoint, Dir, Thickness, MParticles, Elements);
 
 	if (Thickness > 0)
 	{
@@ -231,7 +221,20 @@ bool FTriangleMeshImplicitObject::Raycast(const FVec3& StartPoint, const FVec3&
 	}
 }
 
-bool FTriangleMeshImplicitObject::Overlap(const FVec3& Point, const FReal Thickness) const
+bool FTriangleMeshImplicitObject::Raycast(const FVec3& StartPoint, const FVec3& Dir, const FReal Length, const FReal Thickness, FReal& OutTime, FVec3& OutPosition, FVec3& OutNormal, int32& OutFaceIndex) const
+{
+	if (MElements.RequiresLargeIndices())
+	{
+		return RaycastImp(MElements.GetLargeIndexBuffer(), StartPoint, Dir, Length, Thickness, OutTime, OutPosition, OutNormal, OutFaceIndex);
+	}
+	else
+	{
+		return RaycastImp(MElements.GetSmallIndexBuffer(), StartPoint, Dir, Length, Thickness, OutTime, OutPosition, OutNormal, OutFaceIndex);
+	}
+}
+
+template <typename IdxType>
+bool FTriangleMeshImplicitObject::OverlapImp(const TArray<TVec3<IdxType>>& Elements, const FVec3& Point, const FReal Thickness) const
 {
 	TAABB<FReal, 3> QueryBounds(Point, Point);
 	QueryBounds.Thicken(Thickness);
@@ -243,9 +246,9 @@ bool FTriangleMeshImplicitObject::Overlap(const FVec3& Point, const FReal Thickn
 
 	for (int32 TriIdx : PotentialIntersections)
 	{
-		const FVec3& A = MParticles.X(MElements[TriIdx][0]);
-		const FVec3& B = MParticles.X(MElements[TriIdx][1]);
-		const FVec3& C = MParticles.X(MElements[TriIdx][2]);
+		const FVec3& A = MParticles.X(Elements[TriIdx][0]);
+		const FVec3& B = MParticles.X(Elements[TriIdx][1]);
+		const FVec3& C = MParticles.X(Elements[TriIdx][2]);
 
 		const FVec3 AB = B - A;
 		const FVec3 AC = C - A;
@@ -257,10 +260,10 @@ bool FTriangleMeshImplicitObject::Overlap(const FVec3& Point, const FReal Thickn
 			continue;
 		}
 
-		const TPlane<FReal, 3> TriPlane{ A, Normal };
+		const TPlane<FReal, 3> TriPlane{A, Normal};
 		const FVec3 ClosestPointOnTri = FindClosestPointOnTriangle(TriPlane, A, B, C, Point);
 		const FReal Distance2 = (ClosestPointOnTri - Point).SizeSquared();
-		if (Distance2 <= Thickness * Thickness)	//This really only has a hope in working if thickness is > 0
+		if (Distance2 <= Thickness * Thickness) //This really only has a hope in working if thickness is > 0
 		{
 			return true;
 		}
@@ -268,18 +271,30 @@ bool FTriangleMeshImplicitObject::Overlap(const FVec3& Point, const FReal Thickn
 	return false;
 }
 
-template <typename QueryGeomType>
+bool FTriangleMeshImplicitObject::Overlap(const FVec3& Point, const FReal Thickness) const
+{
+	if (MElements.RequiresLargeIndices())
+	{
+		return OverlapImp(MElements.GetLargeIndexBuffer(), Point, Thickness);
+	}
+	else
+	{
+		return OverlapImp(MElements.GetSmallIndexBuffer(), Point, Thickness);
+	}
+}
+
+template <typename QueryGeomType, typename IdxType>
 void TransformVertsHelper(const QueryGeomType& QueryGeom, int32 TriIdx, const TParticles<FReal, 3>& Particles,
-	const TArray<TVector<int32, 3>>& Elements, TVec3<FReal>& OutA, TVec3<FReal>& OutB, TVec3<FReal>& OutC)
+	const TArray<TVector<IdxType, 3>>& Elements, TVec3<FReal>& OutA, TVec3<FReal>& OutB, TVec3<FReal>& OutC)
 {
 	OutA = Particles.X(Elements[TriIdx][0]);
 	OutB = Particles.X(Elements[TriIdx][1]);
 	OutC = Particles.X(Elements[TriIdx][2]);
 }
 
-template <typename QueryGeomType>
+template <typename QueryGeomType, typename IdxType>
 void TransformVertsHelper(const TImplicitObjectScaled<QueryGeomType>& QueryGeom, int32 TriIdx, const TParticles<FReal, 3>& Particles,
-	const TArray<TVector<int32, 3>>& Elements, TVec3<FReal>& OutA, TVec3<FReal>& OutB, TVec3<FReal>& OutC)
+	const TArray<TVector<IdxType, 3>>& Elements, TVec3<FReal>& OutA, TVec3<FReal>& OutB, TVec3<FReal>& OutC)
 {
 	const TVec3<FReal> InvScale = QueryGeom.GetInvScale();
 	OutA = Particles.X(Elements[TriIdx][0]) * InvScale;
@@ -347,26 +362,38 @@ bool FTriangleMeshImplicitObject::OverlapGeomImp(const QueryGeomType& QueryGeom,
 
 	TRigidTransform<FReal, 3> TransformedQueryTM;
 	TransformOverlapInputsHelper(QueryGeom, QueryTM, TransformedQueryTM);
-	
-	for (int32 TriIdx : PotentialIntersections)
+
+	auto LambdaHelper = [&](const auto& Elements)
 	{
-		TVec3<FReal> A, B, C;
-		TransformVertsHelper(QueryGeom, TriIdx, MParticles, MElements, A, B, C);
+		for(int32 TriIdx : PotentialIntersections)
+		{
+			TVec3<FReal> A,B,C;
+			TransformVertsHelper(QueryGeom,TriIdx,MParticles,Elements,A,B,C);
 
-		const FVec3 AB = B - A;
-		const FVec3 AC = C - A;
+			const FVec3 AB = B - A;
+			const FVec3 AC = C - A;
 
-		//It's most likely that the query object is in front of the triangle since queries tend to be on the outside.
-		//However, maybe we should check if it's behind the triangle plane. Also, we should enforce this winding in some way
-		const FVec3 Offset = FVec3::CrossProduct(AB, AC);
+			//It's most likely that the query object is in front of the triangle since queries tend to be on the outside.
+			//However, maybe we should check if it's behind the triangle plane. Also, we should enforce this winding in some way
+			const FVec3 Offset = FVec3::CrossProduct(AB,AC);
 
-		if (GJKIntersection(TTriangle<FReal>(A, B, C), InnerQueryGeom, TransformedQueryTM, Thickness, Offset))
-		{
-			return true;
+			if(GJKIntersection(TTriangle<FReal>(A,B,C),InnerQueryGeom,TransformedQueryTM,Thickness,Offset))
+			{
+				return true;
+			}
 		}
-	}
 
-	return false;
+		return false;
+	};
+
+	if(MElements.RequiresLargeIndices())
+	{
+		return LambdaHelper(MElements.GetLargeIndexBuffer());
+	}
+	else
+	{
+		return LambdaHelper(MElements.GetSmallIndexBuffer());
+	}
 }
 
 bool FTriangleMeshImplicitObject::OverlapGeom(const TSphere<FReal, 3>& QueryGeom, const FRigidTransform3& QueryTM, const FReal Thickness) const
@@ -415,12 +442,13 @@ bool FTriangleMeshImplicitObject::OverlapGeom(const TImplicitObjectScaled<TImpli
 }
 
 
-template <typename QueryGeomType>
+template <typename QueryGeomType, typename IdxType>
 struct FTriangleMeshSweepVisitor
 {
-	FTriangleMeshSweepVisitor(const FTriangleMeshImplicitObject& InTriMesh, const QueryGeomType& InQueryGeom, const TRigidTransform<FReal,3>& InStartTM, const FVec3& InDir,
+	FTriangleMeshSweepVisitor(const FTriangleMeshImplicitObject& InTriMesh, const TArray<TVec3<IdxType>>& InElements, const QueryGeomType& InQueryGeom, const TRigidTransform<FReal,3>& InStartTM, const FVec3& InDir,
 		const FVec3& InScaledDirNormalized, const FReal InLengthScale, const FRigidTransform3& InScaledStartTM, const FReal InThickness, const bool InComputeMTD)
 	: TriMesh(InTriMesh)
+	, Elements(InElements)
 	, StartTM(InStartTM)
 	, QueryGeom(InQueryGeom)
 	, Dir(InDir)
@@ -454,7 +482,7 @@ struct FTriangleMeshSweepVisitor
 		FVec3 HitNormal;
 
 		FVec3 A, B, C;
-		TransformVertsHelper(QueryGeom, TriIdx, TriMesh.MParticles, TriMesh.MElements, A, B, C);
+		TransformVertsHelper(QueryGeom,TriIdx,TriMesh.MParticles,Elements,A,B,C);
 		TTriangle<FReal> Tri(A, B, C);
 
 		const auto& InnerQueryGeom = GetGeomHelper(QueryGeom);
@@ -483,6 +511,7 @@ struct FTriangleMeshSweepVisitor
 	}
 
 	const FTriangleMeshImplicitObject& TriMesh;
+	const TArray<TVec3<IdxType>>& Elements;
 	const FRigidTransform3 StartTM;
 	const QueryGeomType& QueryGeom;
 	const FVec3& Dir;
@@ -540,21 +569,34 @@ bool FTriangleMeshImplicitObject::SweepGeomImp(const QueryGeomType& QueryGeom, c
 	ComputeScaledSweepInputs(QueryGeom, StartTM, Dir, Length, ScaledDirNormalized, LengthScale, ScaledStartTM);
 
 	bool bHit = false;
-	FTriangleMeshSweepVisitor<QueryGeomType> SQVisitor(*this, QueryGeom, StartTM, Dir, ScaledDirNormalized, LengthScale, ScaledStartTM, Thickness, bComputeMTD);
+	auto LambdaHelper = [&](const auto& Elements)
+	{
+		using VisitorType = FTriangleMeshSweepVisitor<QueryGeomType,decltype(Elements[0][0])>;
+		VisitorType SQVisitor(*this,Elements, QueryGeom,StartTM,Dir,ScaledDirNormalized,LengthScale,ScaledStartTM,Thickness,bComputeMTD);
 
 
-	const TAABB<FReal, 3> QueryBounds = QueryGeom.BoundingBox().TransformedAABB(FRigidTransform3(FVec3::ZeroVector, StartTM.GetRotation()));
-	const FVec3 StartPoint = StartTM.TransformPositionNoScale(QueryBounds.Center());
-	const FVec3 Inflation = QueryBounds.Extents() * 0.5 + FVec3(Thickness);
-	BVH.template Sweep<FTriangleMeshSweepVisitor<QueryGeomType>>(StartPoint, Dir, Length, Inflation, SQVisitor);
+		const TAABB<FReal,3> QueryBounds = QueryGeom.BoundingBox().TransformedAABB(FRigidTransform3(FVec3::ZeroVector,StartTM.GetRotation()));
+		const FVec3 StartPoint = StartTM.TransformPositionNoScale(QueryBounds.Center());
+		const FVec3 Inflation = QueryBounds.Extents() * 0.5 + FVec3(Thickness);
+		BVH.template Sweep<VisitorType>(StartPoint,Dir,Length,Inflation,SQVisitor);
 
-	if (SQVisitor.OutTime <= Length)
+		if(SQVisitor.OutTime <= Length)
+		{
+			OutTime = SQVisitor.OutTime;
+			OutPosition = SQVisitor.OutPosition;
+			OutNormal = SQVisitor.OutNormal;
+			OutFaceIndex = SQVisitor.OutFaceIndex;
+			bHit = true;
+		}
+	};
+
+	if(MElements.RequiresLargeIndices())
 	{
-		OutTime = SQVisitor.OutTime;
-		OutPosition = SQVisitor.OutPosition;
-		OutNormal = SQVisitor.OutNormal;
-		OutFaceIndex = SQVisitor.OutFaceIndex;
-		bHit = true;
+		LambdaHelper(MElements.GetLargeIndexBuffer());
+	}
+	else
+	{
+		LambdaHelper(MElements.GetSmallIndexBuffer());
 	}
 	return bHit;
 }
@@ -599,11 +641,12 @@ bool FTriangleMeshImplicitObject::SweepGeom(const TImplicitObjectScaled<FConvex>
 	return SweepGeomImp(QueryGeom, StartTM, Dir, Length, OutTime, OutPosition, OutNormal, OutFaceIndex, Thickness, bComputeMTD);
 }
 
-int32 FTriangleMeshImplicitObject::FindMostOpposingFace(const FVec3& Position, const FVec3& UnitDir, int32 HintFaceIndex, FReal SearchDist) const
+template <typename IdxType>
+int32 FTriangleMeshImplicitObject::FindMostOpposingFace(const TArray<TVec3<IdxType>>& Elements, const FVec3& Position, const FVec3& UnitDir, int32 HintFaceIndex, FReal SearchDist) const
 {
 	//todo: this is horribly slow, need adjacency information
 	const FReal SearchDist2 = SearchDist * SearchDist;
-	
+
 	TAABB<FReal, 3> QueryBounds(Position - FVec3(SearchDist), Position + FVec3(SearchDist));
 
 	const TArray<int32> PotentialIntersections = BVH.FindAllIntersections(QueryBounds);
@@ -611,12 +654,12 @@ int32 FTriangleMeshImplicitObject::FindMostOpposingFace(const FVec3& Position, c
 
 	FReal MostOpposingDot = TNumericLimits<FReal>::Max();
 	int32 MostOpposingFace = HintFaceIndex;
-	
+
 	for (int32 TriIdx : PotentialIntersections)
 	{
-		const FVec3& A = MParticles.X(MElements[TriIdx][0]);
-		const FVec3& B = MParticles.X(MElements[TriIdx][1]);
-		const FVec3& C = MParticles.X(MElements[TriIdx][2]);
+		const FVec3& A = MParticles.X(Elements[TriIdx][0]);
+		const FVec3& B = MParticles.X(Elements[TriIdx][1]);
+		const FVec3& C = MParticles.X(Elements[TriIdx][2]);
 
 		const FVec3 AB = B - A;
 		const FVec3 AC = C - A;
@@ -628,7 +671,7 @@ int32 FTriangleMeshImplicitObject::FindMostOpposingFace(const FVec3& Position, c
 			continue;
 		}
 
-		const TPlane<FReal, 3> TriPlane{ A, Normal };
+		const TPlane<FReal, 3> TriPlane{A, Normal};
 		const FVec3 ClosestPointOnTri = FindClosestPointOnTriangle(TriPlane, A, B, C, Position);
 		const FReal Distance2 = (ClosestPointOnTri - Position).SizeSquared();
 		if (Distance2 < SearchDist2)
@@ -645,6 +688,18 @@ int32 FTriangleMeshImplicitObject::FindMostOpposingFace(const FVec3& Position, c
 	return MostOpposingFace;
 }
 
+int32 FTriangleMeshImplicitObject::FindMostOpposingFace(const FVec3& Position, const FVec3& UnitDir, int32 HintFaceIndex, FReal SearchDist) const
+{
+	if (MElements.RequiresLargeIndices())
+	{
+		return FindMostOpposingFace(MElements.GetLargeIndexBuffer(), Position, UnitDir, HintFaceIndex, SearchDist);
+	}
+	else
+	{
+		return FindMostOpposingFace(MElements.GetSmallIndexBuffer(), Position, UnitDir, HintFaceIndex, SearchDist);
+	}
+}
+
 FVec3 FTriangleMeshImplicitObject::FindGeometryOpposingNormal(const FVec3& DenormDir, int32 FaceIndex, const FVec3& OriginalNormal) const
 {
 	return GetFaceNormal(FaceIndex);
@@ -661,10 +716,22 @@ uint32 FTriangleMeshImplicitObject::GetTypeHash() const
 	uint32 Result = MParticles.GetTypeHash();
 	Result = HashCombine(Result, MLocalBoundingBox.GetTypeHash());
 
-	for(TVector<int32, 3> Tri : MElements)
+	auto LambdaHelper = [&](const auto& Elements)
 	{
-		uint32 TriHash = HashCombine(::GetTypeHash(Tri[0]), HashCombine(::GetTypeHash(Tri[1]), ::GetTypeHash(Tri[2])));
-		Result = HashCombine(Result, TriHash);
+		for (TVector<int32, 3> Tri : Elements)
+		{
+			uint32 TriHash = HashCombine(::GetTypeHash(Tri[0]), HashCombine(::GetTypeHash(Tri[1]), ::GetTypeHash(Tri[2])));
+			Result = HashCombine(Result, TriHash);
+		}
+	};
+
+	if (MElements.RequiresLargeIndices())
+	{
+		LambdaHelper(MElements.GetLargeIndexBuffer());
+	}
+	else
+	{
+		LambdaHelper(MElements.GetSmallIndexBuffer());
 	}
 
 	return Result;
@@ -674,16 +741,28 @@ FVec3 FTriangleMeshImplicitObject::GetFaceNormal(const int32 FaceIdx) const
 {
 	if (ensure(FaceIdx != INDEX_NONE))
 	{
-		const FVec3& A = MParticles.X(MElements[FaceIdx][0]);
-		const FVec3& B = MParticles.X(MElements[FaceIdx][1]);
-		const FVec3& C = MParticles.X(MElements[FaceIdx][2]);
-
-		const FVec3 AB = B - A;
-		const FVec3 AC = C - A;
-		FVec3 Normal = FVec3::CrossProduct(AB, AC);
-		const FReal Length = Normal.SafeNormalize();
-		ensure(Length);
-		return Normal;
+		auto LambdaHelper = [&](const auto& Elements)
+		{
+			const FVec3& A = MParticles.X(Elements[FaceIdx][0]);
+			const FVec3& B = MParticles.X(Elements[FaceIdx][1]);
+			const FVec3& C = MParticles.X(Elements[FaceIdx][2]);
+
+			const FVec3 AB = B - A;
+			const FVec3 AC = C - A;
+			FVec3 Normal = FVec3::CrossProduct(AB, AC);
+			const FReal Length = Normal.SafeNormalize();
+			ensure(Length);
+			return Normal;
+		};
+		
+		if (MElements.RequiresLargeIndices())
+		{
+			return LambdaHelper(MElements.GetLargeIndexBuffer());
+		}
+		else
+		{
+			return LambdaHelper(MElements.GetSmallIndexBuffer());
+		}
 	}
 
 	return FVec3(0, 0, 1);
@@ -705,22 +784,35 @@ const TParticles<FReal, 3>& FTriangleMeshImplicitObject::Particles() const
 	return MParticles;
 }
 
-const TArray<TVector<int32, 3>> FTriangleMeshImplicitObject::Elements() const
+const FTrimeshIndexBuffer& FTriangleMeshImplicitObject::Elements() const
 {
 	return MElements;
 }
 
-void Chaos::FTriangleMeshImplicitObject::RebuildBV()
+template <typename IdxType>
+void FTriangleMeshImplicitObject::RebuildBVImp(const TArray<TVec3<IdxType>>& Elements)
 {
-	const int32 NumTris = MElements.Num();
+	const int32 NumTris = Elements.Num();
+	TArray<FBvEntry<sizeof(IdxType) == sizeof(FTrimeshIndexBuffer::LargeIdxType)>> BVEntries;
 	BVEntries.Reset(NumTris);
 
-	for (int Tri = 0; Tri<NumTris; Tri++)
+	for (int Tri = 0; Tri < NumTris; Tri++)
 	{
-		BVEntries.Add({ this, Tri });
+		BVEntries.Add({this, Tri});
 	}
 	BVH.Reinitialize(BVEntries);
+}
 
+void Chaos::FTriangleMeshImplicitObject::RebuildBV()
+{
+	if (MElements.RequiresLargeIndices())
+	{
+		RebuildBVImp(MElements.GetLargeIndexBuffer());
+	}
+	else
+	{
+		RebuildBVImp(MElements.GetSmallIndexBuffer());
+	}
 }
 
 
diff --git a/Engine/Source/Runtime/Experimental/Chaos/Public/Chaos/TriangleMeshImplicitObject.h b/Engine/Source/Runtime/Experimental/Chaos/Public/Chaos/TriangleMeshImplicitObject.h
--- a/Engine/Source/Runtime/Experimental/Chaos/Public/Chaos/TriangleMeshImplicitObject.h
+++ b/Engine/Source/Runtime/Experimental/Chaos/Public/Chaos/TriangleMeshImplicitObject.h
@@ -3,41 +3,146 @@
 
 #include "Chaos/Array.h"
 #include "Chaos/GeometryParticles.h"
-#include "Chaos/SegmentMesh.h"
 #include "Chaos/ImplicitObjectScaled.h"
+#include "Chaos/SegmentMesh.h"
 
-#include "ImplicitObject.h"
-#include "Box.h"
-#include "BoundingVolumeHierarchy.h"
+#include "AABBTree.h"
 #include "BoundingVolume.h"
+#include "BoundingVolumeHierarchy.h"
+#include "Box.h"
 #include "ChaosArchive.h"
+#include "ImplicitObject.h"
 #include "UObject/ExternalPhysicsCustomObjectVersion.h"
-#include "AABBTree.h"
 
 namespace Chaos
 {
-
-	template <typename T>
+	template<typename T>
 	class TCapsule;
 
 	class FConvex;
 
+	class CHAOS_API FTrimeshIndexBuffer
+	{
+	public:
+		using LargeIdxType = int32;
+		using SmallIdxType = uint16;
+
+		FTrimeshIndexBuffer() = default;
+		FTrimeshIndexBuffer(TArray<TVector<LargeIdxType, 3>>&& Elements)
+		    : LargeIdxBuffer(MoveTemp(Elements))
+		    , bRequiresLargeIndices(true)
+		{
+		}
+
+		FTrimeshIndexBuffer(TArray<TVector<SmallIdxType, 3>>&& Elements)
+		    : SmallIdxBuffer(MoveTemp(Elements))
+		    , bRequiresLargeIndices(false)
+		{
+		}
+
+		FTrimeshIndexBuffer(const FTrimeshIndexBuffer& Other) = delete;
+		FTrimeshIndexBuffer& operator=(const FTrimeshIndexBuffer& Other) = delete;
+
+		void Serialize(FArchive& Ar)
+		{
+			Ar.UsingCustomVersion(FExternalPhysicsCustomObjectVersion::GUID);
+
+			if (Ar.CustomVer(FExternalPhysicsCustomObjectVersion::GUID) < FExternalPhysicsCustomObjectVersion::TrimeshCanUseSmallIndices)
+			{
+				Ar << LargeIdxBuffer;
+				bRequiresLargeIndices = true;
+			}
+			else
+			{
+				Ar << bRequiresLargeIndices;
+				if (bRequiresLargeIndices)
+				{
+					Ar << LargeIdxBuffer;
+				}
+				else
+				{
+					Ar << SmallIdxBuffer;
+				}
+			}
+		}
+
+		bool RequiresLargeIndices() const
+		{
+			return bRequiresLargeIndices;
+		}
+
+		const auto& GetLargeIndexBuffer() const
+		{
+			check(bRequiresLargeIndices);
+			return LargeIdxBuffer;
+		}
+
+		const auto& GetSmallIndexBuffer() const
+		{
+			check(!bRequiresLargeIndices);
+			return SmallIdxBuffer;
+		}
+		
+		template<bool TRequiresLargeIndices>
+		const typename TChooseClass<TRequiresLargeIndices, TVec3<LargeIdxType>, TVec3<SmallIdxType>>::Result& Get(int32 Idx) const;
+		
+		template<>
+		const TVec3<LargeIdxType>& Get<true>(int32 Idx) const
+		{
+			check(bRequiresLargeIndices);
+			return LargeIdxBuffer[Idx];
+		}
+		
+
+		template<>
+		const TVec3<SmallIdxType>& Get<false>(int32 Idx) const
+		{
+			check(!bRequiresLargeIndices);
+			return SmallIdxBuffer[Idx];
+		}
+
+	private:
+		TArray<TVector<LargeIdxType, 3>> LargeIdxBuffer;
+		TArray<TVector<SmallIdxType, 3>> SmallIdxBuffer;
+		bool bRequiresLargeIndices;
+	};
+
+	FORCEINLINE FArchive& operator<<(FArchive& Ar, FTrimeshIndexBuffer& Buffer)
+	{
+		Buffer.Serialize(Ar);
+		return Ar;
+	}
+
 	class CHAOS_API FTriangleMeshImplicitObject final : public FImplicitObject
 	{
 	public:
 		using FImplicitObject::GetTypeName;
 
-		FTriangleMeshImplicitObject(TParticles<FReal,3>&& Particles, TArray<TVector<int32, 3>>&& Elements, TArray<uint16>&& InMaterialIndices);
+		template <typename IdxType>
+		FTriangleMeshImplicitObject(TParticles<FReal, 3>&& Particles, TArray<TVector<IdxType, 3>>&& Elements, TArray<uint16>&& InMaterialIndices)
+		: FImplicitObject(EImplicitObject::HasBoundingBox, ImplicitObjectType::TriangleMesh)
+		, MParticles(MoveTemp(Particles))
+		, MElements(MoveTemp(Elements))
+		, MLocalBoundingBox(MParticles.X(0), MParticles.X(0))
+		, MaterialIndices(MoveTemp(InMaterialIndices))
+		{
+			for (uint32 Idx = 1; Idx < MParticles.Size(); ++Idx)
+			{
+				MLocalBoundingBox.GrowToInclude(MParticles.X(Idx));
+			}
+			RebuildBV();
+		}
+
 		FTriangleMeshImplicitObject(const FTriangleMeshImplicitObject& Other) = delete;
 		FTriangleMeshImplicitObject(FTriangleMeshImplicitObject&& Other) = delete;
 		virtual ~FTriangleMeshImplicitObject() {}
 
 		virtual FReal PhiWithNormal(const FVec3& x, FVec3& Normal) const;
-		
+
 		virtual bool Raycast(const FVec3& StartPoint, const FVec3& Dir, const FReal Length, const FReal Thickness, FReal& OutTime, FVec3& OutPosition, FVec3& OutNormal, int32& OutFaceIndex) const override;
 		virtual bool Overlap(const FVec3& Point, const FReal Thickness) const override;
 
-		bool OverlapGeom(const TSphere<FReal,3>& QueryGeom, const FRigidTransform3& QueryTM, const FReal Thickness) const;
+		bool OverlapGeom(const TSphere<FReal, 3>& QueryGeom, const FRigidTransform3& QueryTM, const FReal Thickness) const;
 		bool OverlapGeom(const TBox<FReal, 3>& QueryGeom, const FRigidTransform3& QueryTM, const FReal Thickness) const;
 		bool OverlapGeom(const TCapsule<FReal>& QueryGeom, const FRigidTransform3& QueryTM, const FReal Thickness) const;
 		bool OverlapGeom(const FConvex& QueryGeom, const FRigidTransform3& QueryTM, const FReal Thickness) const;
@@ -47,7 +152,7 @@ namespace Chaos
 		bool OverlapGeom(const TImplicitObjectScaled<FConvex>& QueryGeom, const FRigidTransform3& QueryTM, const FReal Thickness) const;
 		bool OverlapGeom(const TImplicitObjectScaled<TImplicitObjectScaled<FConvex>>& QueryGeom, const FRigidTransform3& QueryTM, const FReal Thickness) const;
 
-		bool SweepGeom(const TSphere<FReal,3>& QueryGeom, const FRigidTransform3& StartTM, const FVec3& Dir, const FReal Length, FReal& OutTime, FVec3& OutPosition, FVec3& OutNormal, int32& OutFaceIndex, const FReal Thickness = 0, const bool bComputeMTD = false) const;
+		bool SweepGeom(const TSphere<FReal, 3>& QueryGeom, const FRigidTransform3& StartTM, const FVec3& Dir, const FReal Length, FReal& OutTime, FVec3& OutPosition, FVec3& OutNormal, int32& OutFaceIndex, const FReal Thickness = 0, const bool bComputeMTD = false) const;
 		bool SweepGeom(const TBox<FReal, 3>& QueryGeom, const FRigidTransform3& StartTM, const FVec3& Dir, const FReal Length, FReal& OutTime, FVec3& OutPosition, FVec3& OutNormal, int32& OutFaceIndex, const FReal Thickness = 0, const bool bComputeMTD = false) const;
 		bool SweepGeom(const TCapsule<FReal>& QueryGeom, const FRigidTransform3& StartTM, const FVec3& Dir, const FReal Length, FReal& OutTime, FVec3& OutPosition, FVec3& OutNormal, int32& OutFaceIndex, const FReal Thickness = 0, const bool bComputeMTD = false) const;
 		bool SweepGeom(const FConvex& QueryGeom, const FRigidTransform3& StartTM, const FVec3& Dir, const FReal Length, FReal& OutTime, FVec3& OutPosition, FVec3& OutNormal, int32& OutFaceIndex, const FReal Thickness = 0, const bool bComputeMTD = false) const;
@@ -84,12 +189,12 @@ namespace Chaos
 				Ar << ConvexHulls;
 			}
 
-			if(Ar.CustomVer(FExternalPhysicsCustomObjectVersion::GUID) < FExternalPhysicsCustomObjectVersion::TrimeshSerializesBV)
+			if (Ar.CustomVer(FExternalPhysicsCustomObjectVersion::GUID) < FExternalPhysicsCustomObjectVersion::TrimeshSerializesBV)
 			{
 				// Should now only hit when loading older trimeshes
 				RebuildBV();
 			}
-			else if(Ar.CustomVer(FExternalPhysicsCustomObjectVersion::GUID) < FExternalPhysicsCustomObjectVersion::TrimeshSerializesAABBTree)
+			else if (Ar.CustomVer(FExternalPhysicsCustomObjectVersion::GUID) < FExternalPhysicsCustomObjectVersion::TrimeshSerializesAABBTree)
 			{
 				TBoundingVolume<int32, FReal, 3> Dummy;
 				Ar << Dummy;
@@ -101,7 +206,7 @@ namespace Chaos
 				Ar << BVH;
 			}
 
-			if(Ar.CustomVer(FExternalPhysicsCustomObjectVersion::GUID) >= FExternalPhysicsCustomObjectVersion::AddTrimeshMaterialIndices)
+			if (Ar.CustomVer(FExternalPhysicsCustomObjectVersion::GUID) >= FExternalPhysicsCustomObjectVersion::AddTrimeshMaterialIndices)
 			{
 				Ar << MaterialIndices;
 			}
@@ -116,14 +221,13 @@ namespace Chaos
 		virtual uint16 GetMaterialIndex(uint32 HintIndex) const override;
 
 		const TParticles<FReal, 3>& Particles() const;
-		const TArray<TVector<int32, 3>> Elements() const;
+		const FTrimeshIndexBuffer& Elements() const;
 
 	private:
-
 		void RebuildBV();
 
 		TParticles<FReal, 3> MParticles;
-		TArray<TVector<int32, 3>> MElements;
+		FTrimeshIndexBuffer MElements;
 		TAABB<FReal, 3> MLocalBoundingBox;
 		TArray<uint16> MaterialIndices;
 
@@ -133,6 +237,7 @@ namespace Chaos
 		template<typename InStorageType, typename InRealType>
 		friend struct FBvEntry;
 
+		template<bool bRequiresLargeIndex>
 		struct FBvEntry
 		{
 			FTriangleMeshImplicitObject* TmData;
@@ -142,39 +247,49 @@ namespace Chaos
 
 			TAABB<FReal, 3> BoundingBox() const
 			{
-				TAABB<FReal, 3> Bounds(TmData->MParticles.X(TmData->MElements[Index][0]), TmData->MParticles.X(TmData->MElements[Index][0]));
+				TAABB<FReal, 3> Bounds(TmData->MParticles.X(TmData->MElements.Get<bRequiresLargeIndex>(Index)[0]),
+				    TmData->MParticles.X(TmData->MElements.Get<bRequiresLargeIndex>(Index)[0]));
 
-				Bounds.GrowToInclude(TmData->MParticles.X(TmData->MElements[Index][1]));
-				Bounds.GrowToInclude(TmData->MParticles.X(TmData->MElements[Index][2]));
+				Bounds.GrowToInclude(TmData->MParticles.X(TmData->MElements.Get<bRequiresLargeIndex>(Index)[1]));
+				Bounds.GrowToInclude(TmData->MParticles.X(TmData->MElements.Get<bRequiresLargeIndex>(Index)[2]));
 
 				return Bounds;
 			}
 
-			template <typename TPayloadType>
+			template<typename TPayloadType>
 			int32 GetPayload(int32 Idx) const
 			{
 				return Idx;
 			}
 		};
-		TArray<FBvEntry> BVEntries;
 
 		BVHType BVH;
 
-		template <typename Geom>
+		template<typename Geom, typename IdxType>
 		friend struct FTriangleMeshSweepVisitor;
 
 		// Required by implicit object serialization, disabled for general use.
 		friend class FImplicitObject;
 
-		FTriangleMeshImplicitObject() 
-			: FImplicitObject(EImplicitObject::HasBoundingBox, ImplicitObjectType::TriangleMesh)
-		{};
+		FTriangleMeshImplicitObject()
+		    : FImplicitObject(EImplicitObject::HasBoundingBox, ImplicitObjectType::TriangleMesh){};
 
-		template <typename QueryGeomType>
+		template<typename QueryGeomType>
 		bool OverlapGeomImp(const QueryGeomType& QueryGeom, const FRigidTransform3& QueryTM, const FReal Thickness) const;
 
-		template <typename QueryGeomType>
+		template<typename QueryGeomType>
 		bool SweepGeomImp(const QueryGeomType& QueryGeom, const FRigidTransform3& StartTM, const FVec3& Dir, const FReal Length, FReal& OutTime, FVec3& OutPosition, FVec3& OutNormal, int32& OutFaceIndex, const FReal Thickness, const bool bComputeMTD) const;
+
+		template <typename IdxType>
+		bool RaycastImp(const TArray<TVector<IdxType, 3>>& Elements, const FVec3& StartPoint, const FVec3& Dir, const FReal Length, const FReal Thickness, FReal& OutTime, FVec3& OutPosition, FVec3& OutNormal, int32& OutFaceIndex) const;
+
+		template <typename IdxType>
+		bool OverlapImp(const TArray<TVec3<IdxType>>& Elements, const FVec3& Point, const FReal Thickness) const;
+
+		template<typename IdxType>
+		int32 FindMostOpposingFace(const TArray<TVec3<IdxType>>& Elements, const FVec3& Position, const FVec3& UnitDir, int32 HintFaceIndex, FReal SearchDist) const;
+
+		template <typename IdxType>
+		void RebuildBVImp(const TArray<TVec3<IdxType>>& Elements);
 	};
 }
-
diff --git a/Engine/Source/Runtime/NavigationSystem/Private/NavMesh/RecastNavMeshGenerator.cpp b/Engine/Source/Runtime/NavigationSystem/Private/NavMesh/RecastNavMeshGenerator.cpp
--- a/Engine/Source/Runtime/NavigationSystem/Private/NavMesh/RecastNavMeshGenerator.cpp
+++ b/Engine/Source/Runtime/NavigationSystem/Private/NavMesh/RecastNavMeshGenerator.cpp
@@ -550,51 +550,65 @@ void ExportChaosTriMesh(const Chaos::FTriangleMeshImplicitObject* const TriMesh,
 	using namespace Chaos;
 
 	int32 VertOffset = VertexBuffer.Num() / 3;
-	int32 NumTris = TriMesh->Elements().Num();
-	const TParticles<FReal, 3>& Vertices = TriMesh->Particles();
-	const TArray<TVector<int32, 3>>& Triangles = TriMesh->Elements();
 
-	VertexBuffer.Reserve(VertexBuffer.Num() + NumTris * 9);
-	IndexBuffer.Reserve(IndexBuffer.Num() + NumTris * 3);
+	auto LambdaHelper = [&](const auto& Triangles)
+	{
+		int32 NumTris = Triangles.Num();
+		const TParticles<FReal, 3>& Vertices = TriMesh->Particles();
+	
+		VertexBuffer.Reserve(VertexBuffer.Num() + NumTris * 9);
+		IndexBuffer.Reserve(IndexBuffer.Num() + NumTris * 3);
 
-	const bool bFlipCullMode = (LocalToWorld.GetDeterminant() < 0.f);
+		const bool bFlipCullMode = (LocalToWorld.GetDeterminant() < 0.f);
 
-	const int32 IndexOrder[3] = { bFlipCullMode ? 0 : 2, 1, bFlipCullMode ? 2 : 0 };
+		const int32 IndexOrder[3] = { bFlipCullMode ? 0 : 2, 1, bFlipCullMode ? 2 : 0 };
 
-#if SHOW_NAV_EXPORT_PREVIEW
-	UWorld* DebugWorld = FindEditorWorld();
-#endif // SHOW_NAV_EXPORT_PREVIEW
+	#if SHOW_NAV_EXPORT_PREVIEW
+		UWorld* DebugWorld = FindEditorWorld();
+	#endif // SHOW_NAV_EXPORT_PREVIEW
 
-	for (int32 TriIdx = 0; TriIdx < NumTris; ++TriIdx)
-	{
-		for (int32 i = 0; i < 3; i++)
+		for (int32 TriIdx = 0; TriIdx < NumTris; ++TriIdx)
 		{
-			const FVector UnrealCoords = LocalToWorld.TransformPosition(Vertices.X(Triangles[TriIdx][i]));
-			UnrealBounds += UnrealCoords;
+			for (int32 i = 0; i < 3; i++)
+			{
+				const FVector UnrealCoords = LocalToWorld.TransformPosition(Vertices.X(Triangles[TriIdx][i]));
+				UnrealBounds += UnrealCoords;
 
-			VertexBuffer.Add(UnrealCoords.X);
-			VertexBuffer.Add(UnrealCoords.Y);
-			VertexBuffer.Add(UnrealCoords.Z);
-		}
+				VertexBuffer.Add(UnrealCoords.X);
+				VertexBuffer.Add(UnrealCoords.Y);
+				VertexBuffer.Add(UnrealCoords.Z);
+			}
 
-		IndexBuffer.Add(VertOffset + IndexOrder[0]);
-		IndexBuffer.Add(VertOffset + IndexOrder[1]);
-		IndexBuffer.Add(VertOffset + IndexOrder[2]);
+			IndexBuffer.Add(VertOffset + IndexOrder[0]);
+			IndexBuffer.Add(VertOffset + IndexOrder[1]);
+			IndexBuffer.Add(VertOffset + IndexOrder[2]);
 
-#if SHOW_NAV_EXPORT_PREVIEW
-		if (DebugWorld)
-		{
-			FVector V0(VertexBuffer[(VertOffset + IndexOrder[0]) * 3 + 0], VertexBuffer[(VertOffset + IndexOrder[0]) * 3 + 1], VertexBuffer[(VertOffset + IndexOrder[0]) * 3 + 2]);
-			FVector V1(VertexBuffer[(VertOffset + IndexOrder[1]) * 3 + 0], VertexBuffer[(VertOffset + IndexOrder[1]) * 3 + 1], VertexBuffer[(VertOffset + IndexOrder[1]) * 3 + 2]);
-			FVector V2(VertexBuffer[(VertOffset + IndexOrder[2]) * 3 + 0], VertexBuffer[(VertOffset + IndexOrder[2]) * 3 + 1], VertexBuffer[(VertOffset + IndexOrder[2]) * 3 + 2]);
+	#if SHOW_NAV_EXPORT_PREVIEW
+			if (DebugWorld)
+			{
+				FVector V0(VertexBuffer[(VertOffset + IndexOrder[0]) * 3 + 0], VertexBuffer[(VertOffset + IndexOrder[0]) * 3 + 1], VertexBuffer[(VertOffset + IndexOrder[0]) * 3 + 2]);
+				FVector V1(VertexBuffer[(VertOffset + IndexOrder[1]) * 3 + 0], VertexBuffer[(VertOffset + IndexOrder[1]) * 3 + 1], VertexBuffer[(VertOffset + IndexOrder[1]) * 3 + 2]);
+				FVector V2(VertexBuffer[(VertOffset + IndexOrder[2]) * 3 + 0], VertexBuffer[(VertOffset + IndexOrder[2]) * 3 + 1], VertexBuffer[(VertOffset + IndexOrder[2]) * 3 + 2]);
 
-			DrawDebugLine(DebugWorld, V0, V1, bFlipCullMode ? FColor::Red : FColor::Blue, true);
-			DrawDebugLine(DebugWorld, V1, V2, bFlipCullMode ? FColor::Red : FColor::Blue, true);
-			DrawDebugLine(DebugWorld, V2, V0, bFlipCullMode ? FColor::Red : FColor::Blue, true);
+				DrawDebugLine(DebugWorld, V0, V1, bFlipCullMode ? FColor::Red : FColor::Blue, true);
+				DrawDebugLine(DebugWorld, V1, V2, bFlipCullMode ? FColor::Red : FColor::Blue, true);
+				DrawDebugLine(DebugWorld, V2, V0, bFlipCullMode ? FColor::Red : FColor::Blue, true);
+			}
+	#endif // SHOW_NAV_EXPORT_PREVIEW
+
+			VertOffset += 3;
 		}
-#endif // SHOW_NAV_EXPORT_PREVIEW
+	};
 
-		VertOffset += 3;
+
+	const FTrimeshIndexBuffer& IdxBuffer = TriMesh->Elements();
+	if(IdxBuffer.RequiresLargeIndices())
+	{
+		LambdaHelper(IdxBuffer.GetLargeIndexBuffer());
+	}
+	else
+	{
+		LambdaHelper(IdxBuffer.GetSmallIndexBuffer());
 	}
 }
 
