Accelerate aabb transforms.

diff --git a/Engine/Source/Runtime/Experimental/Chaos/Private/Chaos/AABB.cpp b/Engine/Source/Runtime/Experimental/Chaos/Private/Chaos/AABB.cpp
--- a/Engine/Source/Runtime/Experimental/Chaos/Private/Chaos/AABB.cpp
+++ b/Engine/Source/Runtime/Experimental/Chaos/Private/Chaos/AABB.cpp
@@ -2,6 +2,9 @@
 #include "Chaos/AABB.h"
 #include "Chaos/Sphere.h"
 #include "Chaos/Capsule.h"
+#if INTEL_ISPC
+#include "AABB.ispc.generated.h"
+#endif
 
 namespace Chaos
 {
@@ -174,6 +177,61 @@ bool TAABB<T, d>::Raycast(const TVector<T, d>& StartPoint, const TVector<T, d>&
 	OutPosition = BoxIntersection - Thickness * Normal;
 	return true;
 }
+
+template<typename T, int d>
+template<class TTRANSFORM>
+TAABB<T, d> TAABB<T, d>::TransformedAABB(const TTRANSFORM& SpaceTransform) const
+{
+	TVector<T, d> CurrentExtents = Extents();
+	int32 Idx = 0;
+	const TVector<T, d> MinToNewSpace = SpaceTransform.TransformPosition(MMin);
+	TAABB<T, d> NewAABB(MinToNewSpace, MinToNewSpace);
+	NewAABB.GrowToInclude(SpaceTransform.TransformPosition(MMax));
+
+	for (int32 j = 0; j < d; ++j)
+	{
+		NewAABB.GrowToInclude(SpaceTransform.TransformPosition(MMin + TVector<T, d>::AxisVector(j) * CurrentExtents));
+		NewAABB.GrowToInclude(SpaceTransform.TransformPosition(MMax - TVector<T, d>::AxisVector(j) * CurrentExtents));
+	}
+
+	return NewAABB;
+}
+
+template<>
+template<>
+TAABB<float, 3> TAABB<float, 3>::TransformedAABB(const FTransform& SpaceTransform) const
+{
+	if (INTEL_ISPC)
+	{
+#if INTEL_ISPC
+		TVector<float, 3> NewMin, NewMax;
+		ispc::TransformedAABB((const ispc::FTransform&)SpaceTransform, (const ispc::FVector&)MMin, (const ispc::FVector&)MMax, (ispc::FVector&)NewMin, (ispc::FVector&)NewMax);
+
+		TAABB<float, 3> NewAABB(NewMin, NewMax);
+		return NewAABB;
+#endif
+	}
+	else
+	{
+		TVector<float, 3> CurrentExtents = Extents();
+		int32 Idx = 0;
+		const TVector<float, 3> MinToNewSpace = SpaceTransform.TransformPosition(MMin);
+		TAABB<float, 3> NewAABB(MinToNewSpace, MinToNewSpace);
+		NewAABB.GrowToInclude(SpaceTransform.TransformPosition(MMax));
+
+		for (int32 j = 0; j < 3; ++j)
+		{
+			NewAABB.GrowToInclude(SpaceTransform.TransformPosition(MMin + TVector<float, 3>::AxisVector(j) * CurrentExtents));
+			NewAABB.GrowToInclude(SpaceTransform.TransformPosition(MMax - TVector<float, 3>::AxisVector(j) * CurrentExtents));
+		}
+
+		return NewAABB;
+	}
+}
 }
 
 template class Chaos::TAABB<float, 3>;
+
+template Chaos::TAABB<float, 3> Chaos::TAABB<float, 3>::TransformedAABB(const Chaos::TRigidTransform<float, 3>&) const;
+template Chaos::TAABB<float, 3> Chaos::TAABB<float, 3>::TransformedAABB(const FMatrix&) const;
+template Chaos::TAABB<float, 3> Chaos::TAABB<float, 3>::TransformedAABB(const Chaos::PMatrix<float, 4, 4>&) const;
diff --git a/Engine/Source/Runtime/Experimental/Chaos/Public/Chaos/AABB.h b/Engine/Source/Runtime/Experimental/Chaos/Public/Chaos/AABB.h
--- a/Engine/Source/Runtime/Experimental/Chaos/Public/Chaos/AABB.h
+++ b/Engine/Source/Runtime/Experimental/Chaos/Public/Chaos/AABB.h
@@ -86,22 +86,7 @@ namespace Chaos
 		}
 
 		template<class TTRANSFORM>
-		FORCEINLINE TAABB<T, d> TransformedAABB(const TTRANSFORM& SpaceTransform) const
-		{
-			TVector<T, d> CurrentExtents = Extents();
-			int32 Idx = 0;
-			const TVector<T, d> MinToNewSpace = SpaceTransform.TransformPosition(MMin);
-			TAABB<T, d> NewAABB(MinToNewSpace, MinToNewSpace);
-			NewAABB.GrowToInclude(SpaceTransform.TransformPosition(MMax));
-
-			for (int32 j = 0; j < d; ++j)
-			{
-				NewAABB.GrowToInclude(SpaceTransform.TransformPosition(MMin + TVector<T, d>::AxisVector(j) * CurrentExtents));
-				NewAABB.GrowToInclude(SpaceTransform.TransformPosition(MMax - TVector<T, d>::AxisVector(j) * CurrentExtents));
-			}
-
-			return NewAABB;
-		}
+		CHAOS_API TAABB<T, d> TransformedAABB(const TTRANSFORM& SpaceTransform) const;
 
 		FORCEINLINE bool Intersects(const TAABB<T, d>& Other) const
 		{
@@ -613,5 +598,8 @@ namespace Chaos
 		}
 	};
 
-
+	extern template CHAOS_API Chaos::TAABB<float, 3> Chaos::TAABB<float, 3>::TransformedAABB(const FTransform&) const;
+	extern template CHAOS_API Chaos::TAABB<float, 3> Chaos::TAABB<float, 3>::TransformedAABB(const Chaos::TRigidTransform<float, 3>&) const;
+	extern template CHAOS_API Chaos::TAABB<float, 3> Chaos::TAABB<float, 3>::TransformedAABB(const FMatrix&) const;
+	extern template CHAOS_API Chaos::TAABB<float, 3> Chaos::TAABB<float, 3>::TransformedAABB(const Chaos::PMatrix<float, 4, 4>&) const;
 }
\ No newline at end of file
