Chaos RBAN
- Add ether drag
- Minor joint optimization
- Fix debug draw text lifetime

diff --git a/Engine/Source/Runtime/Engine/Private/PhysicsEngine/ImmediatePhysics/ImmediatePhysicsChaos/ImmediatePhysicsActorHandle_Chaos.cpp b/Engine/Source/Runtime/Engine/Private/PhysicsEngine/ImmediatePhysics/ImmediatePhysicsChaos/ImmediatePhysicsActorHandle_Chaos.cpp
--- a/Engine/Source/Runtime/Engine/Private/PhysicsEngine/ImmediatePhysics/ImmediatePhysicsChaos/ImmediatePhysicsActorHandle_Chaos.cpp
+++ b/Engine/Source/Runtime/Engine/Private/PhysicsEngine/ImmediatePhysics/ImmediatePhysicsChaos/ImmediatePhysicsActorHandle_Chaos.cpp
@@ -284,6 +284,8 @@ namespace ImmediatePhysics_Chaos
 					Dynamic->SetInvM(MassInv);
 					Dynamic->SetI({ Inertia.X, Inertia.Y, Inertia.Z });
 					Dynamic->SetInvI({ InertiaInv.X, InertiaInv.Y, InertiaInv.Z });
+					Dynamic->SetLinearEtherDrag(BodyInstance->LinearDamping);
+					Dynamic->SetAngularEtherDrag(BodyInstance->AngularDamping);
 					Dynamic->Disabled() = true;
 				}
 			}
diff --git a/Engine/Source/Runtime/Experimental/Chaos/Private/Chaos/Evolution/PBDMinEvolution.cpp b/Engine/Source/Runtime/Experimental/Chaos/Private/Chaos/Evolution/PBDMinEvolution.cpp
--- a/Engine/Source/Runtime/Experimental/Chaos/Private/Chaos/Evolution/PBDMinEvolution.cpp
+++ b/Engine/Source/Runtime/Experimental/Chaos/Private/Chaos/Evolution/PBDMinEvolution.cpp
@@ -117,7 +117,7 @@ namespace Chaos
 		SCOPE_CYCLE_COUNTER(STAT_MinEvolution_Integrate);
 
 		TPerParticleEulerStepVelocity<FReal, 3> EulerStepVelocityRule;
-		//TPerParticleEtherDrag<FReal, 3> EtherDragRule(HackLinearDrag, HackAngularDrag);
+		TPerParticleEtherDrag<FReal, 3> EtherDragRule;
 		TPerParticlePBDEulerStep<FReal, 3> EulerStepPositionRule;
 
 		for (TTransientPBDRigidParticleHandle<FReal, 3>& Particle : Particles.GetActiveParticlesView())
@@ -131,7 +131,7 @@ namespace Chaos
 				Particle.Torque() = Particle.ExternalTorque();
 
 				EulerStepVelocityRule.Apply(Particle, Dt);
-				//EtherDragRule.Apply(Particle, Dt);
+				EtherDragRule.Apply(Particle, Dt);
 
 				EulerStepPositionRule.Apply(Particle, Dt);
 
diff --git a/Engine/Source/Runtime/Experimental/Chaos/Private/Chaos/Joint/PBDJointSolverGaussSeidel.cpp b/Engine/Source/Runtime/Experimental/Chaos/Private/Chaos/Joint/PBDJointSolverGaussSeidel.cpp
--- a/Engine/Source/Runtime/Experimental/Chaos/Private/Chaos/Joint/PBDJointSolverGaussSeidel.cpp
+++ b/Engine/Source/Runtime/Experimental/Chaos/Private/Chaos/Joint/PBDJointSolverGaussSeidel.cpp
@@ -246,10 +246,10 @@ namespace Chaos
 		const FReal M0 = FVec3::DotProduct(Axis0, Utilities::Multiply(InvI0, Axis0));
 		const FReal M1 = FVec3::DotProduct(Axis1, Utilities::Multiply(InvI1, Axis1));
 
-		//const FVec3 DR0 = Utilities::Multiply(InvI0, Axis0) * (Angle / (M0 + M1));
-		//const FVec3 DR1 = Utilities::Multiply(InvI1, Axis1) * -(Angle / (M0 + M1));
-		const FVec3 DR0 = Axis0 * (Angle * M0 / (M0 + M1));
-		const FVec3 DR1 = Axis1 * -(Angle * M1 / (M0 + M1));
+		const FVec3 DR0 = Utilities::Multiply(InvI0, Axis0) * (Angle / (M0 + M1));
+		const FVec3 DR1 = Utilities::Multiply(InvI1, Axis1) * -(Angle / (M0 + M1));
+		//const FVec3 DR0 = Axis0 * (Angle * M0 / (M0 + M1));
+		//const FVec3 DR1 = Axis1 * -(Angle * M1 / (M0 + M1));
 
 		ApplyRotationDelta(Stiffness, DR0, DR1);
 	}
@@ -368,8 +368,11 @@ namespace Chaos
 		}
 
 		// Apply twist correction
-		ApplyRotationDelta(TwistStiffness, TwistAxis0, TwistAxis1, DTwistAngle);
-		UpdateDerivedState();
+		if (DTwistAngle != 0)
+		{
+			ApplyRotationDelta(TwistStiffness, TwistAxis0, TwistAxis1, DTwistAngle);
+			UpdateDerivedState();
+		}
 	}
 
 
@@ -450,8 +453,11 @@ namespace Chaos
 		}
 
 		// Apply swing correction
-		ApplyRotationDelta(SwingStiffness, SwingAxis, SwingAxis, DSwingAngle);
-		UpdateDerivedState();
+		if (DSwingAngle != 0)
+		{
+			ApplyRotationDelta(SwingStiffness, SwingAxis, SwingAxis, DSwingAngle);
+			UpdateDerivedState();
+		}
 	}
 
 
@@ -549,8 +555,11 @@ namespace Chaos
 			}
 
 			// Apply swing correction
-			ApplyRotationDelta(SwingStiffness, SwingAxis, SwingAxis, DSwingAngle);
-			UpdateDerivedState();
+			if (DSwingAngle != 0)
+			{
+				ApplyRotationDelta(SwingStiffness, SwingAxis, SwingAxis, DSwingAngle);
+				UpdateDerivedState();
+			}
 		}
 	}
 
@@ -593,33 +602,35 @@ namespace Chaos
 		// @todo(ccaulfield): we should really be calculating axes based on joint config, rather than fixing the error components
 		// Calculate constraint error
 		const FVec3 CX = FPBDJointUtilities::GetLimitedPositionError(JointSettings, Rs[0], Xs[1] - Xs[0]);
-
-		const FMatrix33 InvI0 = Utilities::ComputeWorldSpaceInertia(Qs[0], InvILs[0]);
-		const FMatrix33 InvI1 = Utilities::ComputeWorldSpaceInertia(Qs[1], InvILs[1]);
-
-		// Calculate constraint correction
-		FMatrix33 M0 = FMatrix33(0, 0, 0);
-		FMatrix33 M1 = FMatrix33(0, 0, 0);
-		if (InvMs[0] > 0)
-		{
-			M0 = Utilities::ComputeJointFactorMatrix(Xs[0] - Ps[0], InvI0, InvMs[0]);
-		}
-		if (InvMs[1] > 0)
+		if (CX != FVec3(0))
 		{
-			M1 = Utilities::ComputeJointFactorMatrix(Xs[1] - Ps[1], InvI1, InvMs[1]);
-		}
-		const FMatrix33 MI = (M0 + M1).Inverse();
-		const FVec3 DX = Utilities::Multiply(MI, CX);
+			const FMatrix33 InvI0 = Utilities::ComputeWorldSpaceInertia(Qs[0], InvILs[0]);
+			const FMatrix33 InvI1 = Utilities::ComputeWorldSpaceInertia(Qs[1], InvILs[1]);
 
-		// Apply constraint correction
-		const FVec3 DP0 = InvMs[0] * DX;
-		const FVec3 DP1 = -InvMs[1] * DX;
-		const FVec3 DR0 = Utilities::Multiply(InvI0, FVec3::CrossProduct(Xs[0] - Ps[0], DX));
-		const FVec3 DR1 = Utilities::Multiply(InvI1, FVec3::CrossProduct(Xs[1] - Ps[1], -DX));
+			// Calculate constraint correction
+			FMatrix33 M0 = FMatrix33(0, 0, 0);
+			FMatrix33 M1 = FMatrix33(0, 0, 0);
+			if (InvMs[0] > 0)
+			{
+				M0 = Utilities::ComputeJointFactorMatrix(Xs[0] - Ps[0], InvI0, InvMs[0]);
+			}
+			if (InvMs[1] > 0)
+			{
+				M1 = Utilities::ComputeJointFactorMatrix(Xs[1] - Ps[1], InvI1, InvMs[1]);
+			}
+			const FMatrix33 MI = (M0 + M1).Inverse();
+			const FVec3 DX = Utilities::Multiply(MI, CX);
 
-		ApplyPositionDelta(LinearStiffness, DP0, DP1);
-		ApplyRotationDelta(LinearStiffness, DR0, DR1);
-		UpdateDerivedState();
+			// Apply constraint correction
+			const FVec3 DP0 = InvMs[0] * DX;
+			const FVec3 DP1 = -InvMs[1] * DX;
+			const FVec3 DR0 = Utilities::Multiply(InvI0, FVec3::CrossProduct(Xs[0] - Ps[0], DX));
+			const FVec3 DR1 = Utilities::Multiply(InvI1, FVec3::CrossProduct(Xs[1] - Ps[1], -DX));
+
+			ApplyPositionDelta(LinearStiffness, DP0, DP1);
+			ApplyRotationDelta(LinearStiffness, DR0, DR1);
+			UpdateDerivedState();
+		}
 	}
 
 
diff --git a/Engine/Source/Runtime/Experimental/Chaos/Public/Chaos/DebugDrawQueue.h b/Engine/Source/Runtime/Experimental/Chaos/Public/Chaos/DebugDrawQueue.h
--- a/Engine/Source/Runtime/Experimental/Chaos/Public/Chaos/DebugDrawQueue.h
+++ b/Engine/Source/Runtime/Experimental/Chaos/Public/Chaos/DebugDrawQueue.h
@@ -205,6 +205,7 @@ struct CHAOS_API FLatentDrawCommand
 		Command.TestBaseActor = TestBaseActor;
 		Command.Color = Color;
 		Command.Duration = Duration;
+		Command.LifeTime = Duration;
 		Command.bDrawShadow = bDrawShadow;
 		Command.FontScale = FontScale;
 		Command.Type = EDrawType::String;
