Chaos - Joint minor optimization

diff --git a/Engine/Source/Runtime/Experimental/Chaos/Private/Chaos/Joint/PBDJointSolverGaussSeidel.cpp b/Engine/Source/Runtime/Experimental/Chaos/Private/Chaos/Joint/PBDJointSolverGaussSeidel.cpp
--- a/Engine/Source/Runtime/Experimental/Chaos/Private/Chaos/Joint/PBDJointSolverGaussSeidel.cpp
+++ b/Engine/Source/Runtime/Experimental/Chaos/Private/Chaos/Joint/PBDJointSolverGaussSeidel.cpp
@@ -243,11 +243,13 @@ namespace Chaos
 		const FMatrix33 InvI0 = Utilities::ComputeWorldSpaceInertia(Qs[0], InvILs[0]);
 		const FMatrix33 InvI1 = Utilities::ComputeWorldSpaceInertia(Qs[1], InvILs[1]);
 
-		const FReal M0 = FVec3::DotProduct(Axis0, Utilities::Multiply(InvI0, Axis0));
-		const FReal M1 = FVec3::DotProduct(Axis1, Utilities::Multiply(InvI1, Axis1));
+		const FVec3 IA0 = Utilities::Multiply(InvI0, Axis0);
+		const FVec3 IA1 = Utilities::Multiply(InvI1, Axis1);
+		const FReal M0 = FVec3::DotProduct(Axis0, IA0);
+		const FReal M1 = FVec3::DotProduct(Axis1, IA1);
 
-		const FVec3 DR0 = Utilities::Multiply(InvI0, Axis0) * (Angle / (M0 + M1));
-		const FVec3 DR1 = Utilities::Multiply(InvI1, Axis1) * -(Angle / (M0 + M1));
+		const FVec3 DR0 = IA0 * (Angle / (M0 + M1));
+		const FVec3 DR1 = IA1 * -(Angle / (M0 + M1));
 		//const FVec3 DR0 = Axis0 * (Angle * M0 / (M0 + M1));
 		//const FVec3 DR1 = Axis1 * -(Angle * M1 / (M0 + M1));
 
@@ -267,8 +269,10 @@ namespace Chaos
 		const FMatrix33 InvI0 = Utilities::ComputeWorldSpaceInertia(Qs[0], InvILs[0]);
 		const FMatrix33 InvI1 = Utilities::ComputeWorldSpaceInertia(Qs[1], InvILs[1]);
 
-		const FReal IM0 = FVec3::DotProduct(Axis0, Utilities::Multiply(InvI0, Axis0));
-		const FReal IM1 = FVec3::DotProduct(Axis1, Utilities::Multiply(InvI1, Axis1));
+		const FVec3 IA0 = Utilities::Multiply(InvI0, Axis0);
+		const FVec3 IA1 = Utilities::Multiply(InvI1, Axis1);
+		const FReal IM0 = FVec3::DotProduct(Axis0, IA0);
+		const FReal IM1 = FVec3::DotProduct(Axis1, IA1);
 		const FReal IM = (IM0 + IM1);
 
 		FReal JV = 0;
@@ -298,10 +302,10 @@ namespace Chaos
 			DLambda = DLambdaNumerator / DLambdaDenominator;
 		}
 
-		//const FVec3 DR0 = Utilities::Multiply(InvI0, Axis0) * DLambda;
-		//const FVec3 DR1 = Utilities::Multiply(InvI1, Axis1) * -DLambda;
-		const FVec3 DR0 = Axis0 * (DLambda * IM0);
-		const FVec3 DR1 = Axis1 * -(DLambda * IM1);
+		const FVec3 DR0 = IA0 * DLambda;
+		const FVec3 DR1 = IA1 * -DLambda;
+		//const FVec3 DR0 = Axis0 * (DLambda * IM0);
+		//const FVec3 DR1 = Axis1 * -(DLambda * IM1);
 
 		const FRotation3 DQ0 = (FRotation3::FromElements(DR0, 0) * Qs[0]) * (FReal)0.5;
 		const FRotation3 DQ1 = (FRotation3::FromElements(DR1, 0) * Qs[1]) * (FReal)0.5;
