Avoid template specialization

diff --git a/Engine/Source/Runtime/Experimental/Chaos/Public/Chaos/TriangleMeshImplicitObject.h b/Engine/Source/Runtime/Experimental/Chaos/Public/Chaos/TriangleMeshImplicitObject.h
--- a/Engine/Source/Runtime/Experimental/Chaos/Public/Chaos/TriangleMeshImplicitObject.h
+++ b/Engine/Source/Runtime/Experimental/Chaos/Public/Chaos/TriangleMeshImplicitObject.h
@@ -82,24 +82,6 @@ namespace Chaos
 			check(!bRequiresLargeIndices);
 			return SmallIdxBuffer;
 		}
-		
-		template<bool TRequiresLargeIndices>
-		const typename TChooseClass<TRequiresLargeIndices, TVec3<LargeIdxType>, TVec3<SmallIdxType>>::Result& Get(int32 Idx) const;
-		
-		template<>
-		const TVec3<LargeIdxType>& Get<true>(int32 Idx) const
-		{
-			check(bRequiresLargeIndices);
-			return LargeIdxBuffer[Idx];
-		}
-		
-
-		template<>
-		const TVec3<SmallIdxType>& Get<false>(int32 Idx) const
-		{
-			check(!bRequiresLargeIndices);
-			return SmallIdxBuffer[Idx];
-		}
 
 	private:
 		TArray<TVector<LargeIdxType, 3>> LargeIdxBuffer;
@@ -247,13 +229,24 @@ namespace Chaos
 
 			TAABB<FReal, 3> BoundingBox() const
 			{
-				TAABB<FReal, 3> Bounds(TmData->MParticles.X(TmData->MElements.Get<bRequiresLargeIndex>(Index)[0]),
-				    TmData->MParticles.X(TmData->MElements.Get<bRequiresLargeIndex>(Index)[0]));
+				auto LambdaHelper = [&](const auto& Elements)
+				{
+					TAABB<FReal,3> Bounds(TmData->MParticles.X(Elements[Index][0]), TmData->MParticles.X(Elements[Index][0]));
+
+					Bounds.GrowToInclude(TmData->MParticles.X(Elements[Index][1]));
+					Bounds.GrowToInclude(TmData->MParticles.X(Elements[Index][2]));
 
-				Bounds.GrowToInclude(TmData->MParticles.X(TmData->MElements.Get<bRequiresLargeIndex>(Index)[1]));
-				Bounds.GrowToInclude(TmData->MParticles.X(TmData->MElements.Get<bRequiresLargeIndex>(Index)[2]));
+					return Bounds;
+				};
 
-				return Bounds;
+				if(bRequiresLargeIndex)
+				{
+					return LambdaHelper(TmData->MElements.GetLargeIndexBuffer());
+				}
+				else
+				{
+					return LambdaHelper(TmData->MElements.GetSmallIndexBuffer());
+				}
 			}
 
 			template<typename TPayloadType>
