Chaos - Fix CIS (AABB.cpp)

diff --git a/Engine/Source/Runtime/Experimental/Chaos/Private/Chaos/AABB.cpp b/Engine/Source/Runtime/Experimental/Chaos/Private/Chaos/AABB.cpp
--- a/Engine/Source/Runtime/Experimental/Chaos/Private/Chaos/AABB.cpp
+++ b/Engine/Source/Runtime/Experimental/Chaos/Private/Chaos/AABB.cpp
@@ -178,28 +178,45 @@ bool TAABB<T, d>::Raycast(const TVector<T, d>& StartPoint, const TVector<T, d>&
 	return true;
 }
 
-template<typename T, int d>
-template<class TTRANSFORM>
-TAABB<T, d> TAABB<T, d>::TransformedAABB(const TTRANSFORM& SpaceTransform) const
+template<typename T, int d, class TTRANSFORM>
+TAABB<T, d> TransformedAABBHelper(const TAABB<T, d>& AABB, const TTRANSFORM& SpaceTransform)
 {
-	TVector<T, d> CurrentExtents = Extents();
+	TVector<T, d> CurrentExtents = AABB.Extents();
 	int32 Idx = 0;
-	const TVector<T, d> MinToNewSpace = SpaceTransform.TransformPosition(MMin);
+	const TVector<T, d> MinToNewSpace = SpaceTransform.TransformPosition(AABB.Min());
 	TAABB<T, d> NewAABB(MinToNewSpace, MinToNewSpace);
-	NewAABB.GrowToInclude(SpaceTransform.TransformPosition(MMax));
+	NewAABB.GrowToInclude(SpaceTransform.TransformPosition(AABB.Max()));
 
 	for (int32 j = 0; j < d; ++j)
 	{
-		NewAABB.GrowToInclude(SpaceTransform.TransformPosition(MMin + TVector<T, d>::AxisVector(j) * CurrentExtents));
-		NewAABB.GrowToInclude(SpaceTransform.TransformPosition(MMax - TVector<T, d>::AxisVector(j) * CurrentExtents));
+		NewAABB.GrowToInclude(SpaceTransform.TransformPosition(AABB.Min() + TVector<T, d>::AxisVector(j) * CurrentExtents));
+		NewAABB.GrowToInclude(SpaceTransform.TransformPosition(AABB.Max() - TVector<T, d>::AxisVector(j) * CurrentExtents));
 	}
 
 	return NewAABB;
 }
 
-template<>
-template<>
-TAABB<float, 3> TAABB<float, 3>::TransformedAABB(const FTransform& SpaceTransform) const
+template<typename T, int d>
+TAABB<T, d> TAABB<T, d>::TransformedAABB(const Chaos::TRigidTransform<FReal, 3>& SpaceTransform) const
+{
+	return TransformedAABBHelper<T, d>(*this, SpaceTransform);
+}
+
+template<typename T, int d>
+TAABB<T, d> TAABB<T, d>::TransformedAABB(const FMatrix& SpaceTransform) const
+{
+	return TransformedAABBHelper<T, d>(*this, SpaceTransform);
+}
+
+
+template<typename T, int d>
+TAABB<T, d> TAABB<T, d>::TransformedAABB(const Chaos::PMatrix<FReal, 4, 4>& SpaceTransform) const
+{
+	return TransformedAABBHelper<T, d>(*this, SpaceTransform);
+}
+
+template<typename T, int d>
+TAABB<T, d> TAABB<T, d>::TransformedAABB(const FTransform& SpaceTransform) const
 {
 	if (INTEL_ISPC)
 	{
@@ -231,7 +248,3 @@ TAABB<float, 3> TAABB<float, 3>::TransformedAABB(const FTransform& SpaceTransfor
 }
 
 template class Chaos::TAABB<float, 3>;
-
-template Chaos::TAABB<float, 3> Chaos::TAABB<float, 3>::TransformedAABB(const Chaos::TRigidTransform<float, 3>&) const;
-template Chaos::TAABB<float, 3> Chaos::TAABB<float, 3>::TransformedAABB(const FMatrix&) const;
-template Chaos::TAABB<float, 3> Chaos::TAABB<float, 3>::TransformedAABB(const Chaos::PMatrix<float, 4, 4>&) const;
diff --git a/Engine/Source/Runtime/Experimental/Chaos/Public/Chaos/AABB.h b/Engine/Source/Runtime/Experimental/Chaos/Public/Chaos/AABB.h
--- a/Engine/Source/Runtime/Experimental/Chaos/Public/Chaos/AABB.h
+++ b/Engine/Source/Runtime/Experimental/Chaos/Public/Chaos/AABB.h
@@ -85,8 +85,10 @@ namespace Chaos
 			return TAABBSpecializeSamplingHelper<T, d>::ComputeSamplePoints(*this);
 		}
 
-		template<class TTRANSFORM>
-		CHAOS_API TAABB<T, d> TransformedAABB(const TTRANSFORM& SpaceTransform) const;
+		CHAOS_API TAABB<T, d> TransformedAABB(const FTransform&) const;
+		CHAOS_API TAABB<T, d> TransformedAABB(const Chaos::TRigidTransform<FReal, 3>&) const;
+		CHAOS_API TAABB<T, d> TransformedAABB(const FMatrix&) const;
+		CHAOS_API TAABB<T, d> TransformedAABB(const Chaos::PMatrix<FReal, 4, 4>&) const;
 
 		FORCEINLINE bool Intersects(const TAABB<T, d>& Other) const
 		{
@@ -597,9 +599,4 @@ namespace Chaos
 			return SamplePoints;
 		}
 	};
-
-	extern template CHAOS_API Chaos::TAABB<float, 3> Chaos::TAABB<float, 3>::TransformedAABB(const FTransform&) const;
-	extern template CHAOS_API Chaos::TAABB<float, 3> Chaos::TAABB<float, 3>::TransformedAABB(const Chaos::TRigidTransform<float, 3>&) const;
-	extern template CHAOS_API Chaos::TAABB<float, 3> Chaos::TAABB<float, 3>::TransformedAABB(const FMatrix&) const;
-	extern template CHAOS_API Chaos::TAABB<float, 3> Chaos::TAABB<float, 3>::TransformedAABB(const Chaos::PMatrix<float, 4, 4>&) const;
 }
\ No newline at end of file
