Chaos : Collisions

FConvex to THeightField

diff --git a/Engine/Source/Runtime/Experimental/Chaos/Private/Chaos/CollisionResolution.cpp b/Engine/Source/Runtime/Experimental/Chaos/Private/Chaos/CollisionResolution.cpp
--- a/Engine/Source/Runtime/Experimental/Chaos/Private/Chaos/CollisionResolution.cpp
+++ b/Engine/Source/Runtime/Experimental/Chaos/Private/Chaos/CollisionResolution.cpp
@@ -660,7 +660,6 @@ namespace Chaos
 			}
 		}
 
-
 		//
 		// Convex - Convex
 		//
@@ -842,6 +841,56 @@ namespace Chaos
 
 		}
 
+		//
+		// Convex - HeightField
+		//
+
+		template <typename T, int d>
+		TContactPoint<T> ConvexHeightFieldContactPoint(const FImplicitObject& A, const TRigidTransform<T, d>& ATransform, const THeightField<T>& B, const TRigidTransform<T, d>& BTransform, const T Thickness)
+		{
+			TContactPoint<T> Contact;
+
+			if(const TImplicitObjectScaled<FConvex>* ScaledConvexImplicit = A.template GetObject<const TImplicitObjectScaled<FConvex> >())
+				B.GJKContactPoint(*ScaledConvexImplicit, ATransform.GetRelativeTransform(BTransform), Thickness, Contact.Location, Contact.Normal, Contact.Phi);
+			else if (const FConvex* ConvexImplicit = A.template GetObject<const FConvex>())
+				B.GJKContactPoint(*ConvexImplicit, ATransform.GetRelativeTransform(BTransform), Thickness, Contact.Location, Contact.Normal, Contact.Phi);
+
+			Contact.Location = BTransform.TransformPosition(Contact.Location);
+			Contact.Normal = BTransform.TransformVector(Contact.Normal);
+			return Contact;
+		}
+
+		template <typename T, int d>
+		void UpdateConvexHeightFieldConstraint(const FImplicitObject& A, const TRigidTransform<T, d>& ATransform, const THeightField<T>& B, const TRigidTransform<T, d>& BTransform, const T Thickness, TRigidBodyPointContactConstraint<T, d>& Constraint)
+		{
+			UpdateContactPoint(Constraint.Manifold, ConvexHeightFieldContactPoint(A, ATransform, B, BTransform, Thickness));
+		}
+
+		template<typename T, int d>
+		void ConstructConvexHeightFieldConstraints(TGeometryParticleHandle<T, d>* Particle0, TGeometryParticleHandle<T, d>* Particle1, const FImplicitObject* Implicit0, const FImplicitObject* Implicit1, const TRigidTransform<T, d>& Transform0, const TRigidTransform<T, d>& Transform1, const T Thickness, FCollisionConstraintsArray& NewConstraints)
+		{
+
+			const THeightField<T> * Object1 = Implicit1->template GetObject<const THeightField<T> >();
+			if (ensure(Implicit0->IsConvex() && Object1))
+			{
+				TRigidTransform<T, d> ParticleImplicit0TM = Transform0.GetRelativeTransform(Collisions::GetTransform(Particle0));
+				TRigidTransform<T, d> ParticleImplicit1TM = Transform1.GetRelativeTransform(Collisions::GetTransform(Particle1));
+				FRigidBodyPointContactConstraint* Constraint = new FRigidBodyPointContactConstraint(Particle0, Implicit0, ParticleImplicit0TM, Particle1, Implicit1, ParticleImplicit1TM);
+
+				UpdateConvexHeightFieldConstraint(*Implicit0, Transform0, *Object1, Transform1, Thickness, *Constraint);
+
+				if (Constraint->GetPhi() < Thickness)
+				{
+					NewConstraints.Add(Constraint);
+				}
+				else
+				{
+					delete Constraint;
+				}
+			}
+		}
+
+
 
 		//
 		// Levelset-Levelset
@@ -1108,6 +1157,19 @@ namespace Chaos
 				//     second.
 				ensure(false);
 			}
+			else if (Implicit0Type == THeightField<T>::StaticType() && Implicit1.IsConvex())
+			{
+				//     This case should not be necessary. The height fields 
+				//     will only ever be collided against, so ideally will never 
+				//     be in index[0] position of the constraint, also the construction
+				//     of the constraint will just switch the index position so its always 
+				//     second.
+				ensure(false);
+			}
+			else if (Implicit0.IsConvex() && Implicit1Type == THeightField<T>::StaticType())
+			{
+				UpdateConvexHeightFieldConstraint(Implicit0, Transform0, *Implicit1.template GetObject< THeightField<T> >(), Transform1, Thickness, *ConstraintBase.template As<TRigidBodyPointContactConstraint<T, d>>());
+			}
 			else if(Implicit0.IsConvex() && Implicit1.IsConvex())
 			{
 				UpdateConvexConvexConstraint(Implicit0, Transform0, Implicit1, Transform1, Thickness, ConstraintBase);
@@ -1257,7 +1319,15 @@ namespace Chaos
 			}
 			else if (Implicit0Type == THeightField<T>::StaticType() && Implicit1Type == TCapsule<T>::StaticType())
 			{
-				ConstructSphereHeightFieldConstraints(Particle1, Particle0, Implicit1, Implicit0, Transform1, Transform0, Thickness, NewConstraints);
+				ConstructCapsuleHeightFieldConstraints(Particle1, Particle0, Implicit1, Implicit0, Transform1, Transform0, Thickness, NewConstraints);
+			}
+			else if (Implicit0->IsConvex() && Implicit1Type == THeightField<T>::StaticType())
+			{
+			ConstructConvexHeightFieldConstraints(Particle0, Particle1, Implicit0, Implicit1, Transform0, Transform1, Thickness, NewConstraints);
+			}
+			else if (Implicit0Type == THeightField<T>::StaticType() && Implicit1->IsConvex())
+			{
+			ConstructConvexHeightFieldConstraints(Particle1, Particle0, Implicit1, Implicit0, Transform1, Transform0, Thickness, NewConstraints);
 			}
 			else if (Implicit0->IsConvex() && Implicit1->IsConvex())
 			{
@@ -1303,6 +1373,9 @@ namespace Chaos
 		template void UpdateCapsuleHeightFieldConstraint(const TCapsule<float>& A, const TRigidTransform<float, 3>& ATransform, const THeightField<float>& B, const TRigidTransform<float, 3>& BTransform, const float Thickness, TRigidBodyPointContactConstraint<float, 3>& Constraint);
 		template void ConstructCapsuleHeightFieldConstraints(TGeometryParticleHandle<float, 3>* Particle0, TGeometryParticleHandle<float, 3>* Particle1, const FImplicitObject* Implicit0, const FImplicitObject* Implicit1, const TRigidTransform<float, 3>& Transform0, const TRigidTransform<float, 3>& Transform1, const float Thickness, FCollisionConstraintsArray& NewConstraints);
 
+		template void UpdateConvexHeightFieldConstraint(const FImplicitObject& A, const TRigidTransform<float, 3>& ATransform, const THeightField<float>& B, const TRigidTransform<float, 3>& BTransform, const float Thickness, TRigidBodyPointContactConstraint<float, 3>& Constraint);
+		template void ConstructConvexHeightFieldConstraints(TGeometryParticleHandle<float, 3>* Particle0, TGeometryParticleHandle<float, 3>* Particle1, const FImplicitObject* Implicit0, const FImplicitObject* Implicit1, const TRigidTransform<float, 3>& Transform0, const TRigidTransform<float, 3>& Transform1, const float Thickness, FCollisionConstraintsArray& NewConstraints);
+
 		template void UpdateConvexConvexConstraint<float, 3>(const FImplicitObject& A, const TRigidTransform<float, 3>& ATM, const FImplicitObject& B, const TRigidTransform<float, 3>& BTM, const float Thickness, TCollisionConstraintBase<float, 3>& Constraint);
 		template void UpdateConvexConvexManifold<float, 3>(TRigidBodyIterativeContactConstraint<float, 3>& Constraint, const TRigidTransform<float, 3>& ATM, const TRigidTransform<float, 3>& BTM, const float Thickness);
 		template void ConstructConvexConvexConstraints<float, 3>(TGeometryParticleHandle<float, 3>* Particle0, TGeometryParticleHandle<float, 3>* Particle1, const FImplicitObject* Implicit0, const FImplicitObject* Implicit1, const TRigidTransform<float, 3>& Transform0, const TRigidTransform<float, 3>& Transform1, const float Thickness, FCollisionConstraintsArray& NewConstraints);
diff --git a/Engine/Source/Runtime/Experimental/Chaos/Private/Chaos/HeightField.cpp b/Engine/Source/Runtime/Experimental/Chaos/Private/Chaos/HeightField.cpp
--- a/Engine/Source/Runtime/Experimental/Chaos/Private/Chaos/HeightField.cpp
+++ b/Engine/Source/Runtime/Experimental/Chaos/Private/Chaos/HeightField.cpp
@@ -1247,6 +1247,18 @@ namespace Chaos
 		return GJKContactPointImp(QueryGeom, QueryTM, Thickness, ContactLocation, ContactNormal, ContactPhi);
 	}
 
+	template <typename T>
+	bool THeightField<T>::GJKContactPoint(const FConvex& QueryGeom, const TRigidTransform<T, 3>& QueryTM, const T Thickness, TVector<T, 3>& ContactLocation, TVector<T, 3>& ContactNormal, T& ContactPhi) const
+	{
+		return GJKContactPointImp(QueryGeom, QueryTM, Thickness, ContactLocation, ContactNormal, ContactPhi);
+	}
+
+	template <typename T>
+	bool THeightField<T>::GJKContactPoint(const TImplicitObjectScaled<FConvex>& QueryGeom, const TRigidTransform<T, 3>& QueryTM, const T Thickness, TVector<T, 3>& ContactLocation, TVector<T, 3>& ContactNormal, T& ContactPhi) const
+	{
+		return GJKContactPointImp(QueryGeom, QueryTM, Thickness, ContactLocation, ContactNormal, ContactPhi);
+	}
+
 
 	template <typename T>
 	template <typename QueryGeomType>
diff --git a/Engine/Source/Runtime/Experimental/Chaos/Public/Chaos/CollisionResolution.h b/Engine/Source/Runtime/Experimental/Chaos/Public/Chaos/CollisionResolution.h
--- a/Engine/Source/Runtime/Experimental/Chaos/Public/Chaos/CollisionResolution.h
+++ b/Engine/Source/Runtime/Experimental/Chaos/Public/Chaos/CollisionResolution.h
@@ -15,6 +15,8 @@ namespace Chaos
 	template <typename T>
 	class TCapsule;
 
+	class FConvex;
+
 	template <typename T, int d>
 	class TPlane;
 
@@ -157,6 +159,17 @@ namespace Chaos
 			template <typename T, int d>
 			void CHAOS_API ConstructConvexConvexConstraints(TGeometryParticleHandle<T, d>* Particle0, TGeometryParticleHandle<T, d>* Particle1, const FImplicitObject* Implicit0, const FImplicitObject* Implicit1, const TRigidTransform<T, d>& Transform0, const TRigidTransform<T, d>& Transform1, const T Thickness, FCollisionConstraintsArray& NewConstraints);
 
+			//
+			// Convex-HeightField
+			//
+
+			template <typename T, int d>
+			void UpdateConvexHeightFieldConstraint(const FImplicitObject& A, const TRigidTransform<T, d>& ATransform, const THeightField<T>& B, const TRigidTransform<T, d>& BTransform, const T Thickness, TRigidBodyPointContactConstraint<T, d>& Constraint);
+
+			template<typename T, int d>
+			void ConstructConvexHeightFieldConstraints(TGeometryParticleHandle<T, d>* Particle0, TGeometryParticleHandle<T, d>* Particle1, const FImplicitObject* Implicit0, const FImplicitObject* Implicit1, const TRigidTransform<T, d>& Transform0, const TRigidTransform<T, d>& Transform1, const T Thickness, FCollisionConstraintsArray& NewConstraints);
+
+
 			//
 			// Levelset-Levelset
 			//
diff --git a/Engine/Source/Runtime/Experimental/Chaos/Public/Chaos/HeightField.h b/Engine/Source/Runtime/Experimental/Chaos/Public/Chaos/HeightField.h
--- a/Engine/Source/Runtime/Experimental/Chaos/Public/Chaos/HeightField.h
+++ b/Engine/Source/Runtime/Experimental/Chaos/Public/Chaos/HeightField.h
@@ -74,6 +74,8 @@ namespace Chaos
 		bool GJKContactPoint(const TBox<T, 3>& QueryGeom, const TRigidTransform<T, 3>& QueryTM, const T Thickness, TVector<T, 3>& ContactLocation, TVector<T, 3>& ContactNormal, T& ContactPhi) const;
 		bool GJKContactPoint(const TSphere<T, 3>& QueryGeom, const TRigidTransform<T, 3>& QueryTM, const T Thickness, TVector<T, 3>& ContactLocation, TVector<T, 3>& ContactNormal, T& ContactPhi) const;
 		bool GJKContactPoint(const TCapsule<T>& QueryGeom, const TRigidTransform<T, 3>& QueryTM, const T Thickness, TVector<T, 3>& ContactLocation, TVector<T, 3>& ContactNormal, T& ContactPhi) const;
+		bool GJKContactPoint(const FConvex& QueryGeom, const TRigidTransform<T, 3>& QueryTM, const T Thickness, TVector<T, 3>& ContactLocation, TVector<T, 3>& ContactNormal, T& ContactPhi) const;
+		bool GJKContactPoint(const TImplicitObjectScaled<FConvex>& QueryGeom, const TRigidTransform<T, 3>& QueryTM, const T Thickness, TVector<T, 3>& ContactLocation, TVector<T, 3>& ContactNormal, T& ContactPhi) const;
 
 
 		virtual int32 FindMostOpposingFace(const TVector<T, 3>& Position, const TVector<T, 3>& UnitDir, int32 HintFaceIndex, T SearchDist) const override;
