Chaos : Collisions

- Added a UpdateManiofold for specific multi-point collision constraints.

diff --git a/Engine/Source/Runtime/Experimental/Chaos/Private/Chaos/CollisionResolution.cpp b/Engine/Source/Runtime/Experimental/Chaos/Private/Chaos/CollisionResolution.cpp
--- a/Engine/Source/Runtime/Experimental/Chaos/Private/Chaos/CollisionResolution.cpp
+++ b/Engine/Source/Runtime/Experimental/Chaos/Private/Chaos/CollisionResolution.cpp
@@ -992,6 +992,32 @@ namespace Chaos
 		// Constraint API
 		//
 
+		template<typename T, int d>
+		void UpdateManifold(TCollisionConstraintBase<T, d>& ConstraintBase, const TRigidTransform<T, d>& ATM, const TRigidTransform<T, d>& BTM, const T Thickness)
+		{
+			const FImplicitObject& Implicit0 = *ConstraintBase.Manifold.Implicit[0];
+			const FImplicitObject& Implicit1 = *ConstraintBase.Manifold.Implicit[1];
+
+			const TRigidTransform<T, d>& Transform0 = ConstraintBase.ImplicitTransform[0] * ATM;
+			const TRigidTransform<T, d>& Transform1 = ConstraintBase.ImplicitTransform[1] * BTM;
+
+			//
+			// @todo(chaos): Collision Constraints (CollisionMap)
+			//    Modify Construct() and Update() use a CollisionMap indexed on 
+			//    EImplicitObjectType, instead of the if/else chain. Also, remove 
+			//    the blocks with the ensure(false), they are just for validation 
+			//    after the recent change. 
+			//
+			EImplicitObjectType Implicit0Type = Implicit0.GetType();
+			EImplicitObjectType Implicit1Type = Implicit1.GetType();
+
+			// TConvex
+			if (Implicit0.IsConvex() && Implicit1.IsConvex())
+			{
+				UpdateConvexConvexManifold(*ConstraintBase.template As<TRigidBodyIterativeContactConstraint<T, d>>(), Transform0, Transform1, Thickness);
+			}
+		}
+
 		template<ECollisionUpdateType UpdateType, typename T, int d>
 		void UpdateConstraint(TCollisionConstraintBase<T, d>& ConstraintBase, const TRigidTransform<T, d>& ParticleTransform0, const TRigidTransform<T, d>& ParticleTransform1, const T Thickness)
 		{
@@ -1386,6 +1412,7 @@ namespace Chaos
 
 		template void ConstructUnionUnionConstraints<float, 3>(TGeometryParticleHandle<float, 3>* Particle0, TGeometryParticleHandle<float, 3>* Particle1, const FImplicitObject* Implicit0, const FImplicitObject* Implicit1, const TRigidTransform<float, 3>& Transform0, const TRigidTransform<float, 3>& Transform1, const float Thickness, FCollisionConstraintsArray& NewConstraints);
 
+		template void UpdateManifold<float, 3>(TCollisionConstraintBase<float, 3>& ConstraintBase, const TRigidTransform<float, 3>& ATM, const TRigidTransform<float, 3>& BTM, const float Thickness);
 		template void UpdateConstraint<ECollisionUpdateType::Any, float, 3>(TCollisionConstraintBase<float,3>& ConstraintBase, const TRigidTransform<float,3>& Transform0, const TRigidTransform<float,3>& Transform1, const float Thickness);
 		template void UpdateConstraint<ECollisionUpdateType::Deepest, float, 3>(TCollisionConstraintBase<float, 3>& ConstraintBase, const TRigidTransform<float, 3>& Transform0, const TRigidTransform<float, 3>& Transform1, const float Thickness);
 		template void ConstructConstraints<float, 3>(TGeometryParticleHandle<float, 3>* Particle0, TGeometryParticleHandle<float, 3>* Particle1, const FImplicitObject* Implicit0, const FImplicitObject* Implicit1, const TRigidTransform<float, 3>& Transform0, const TRigidTransform<float, 3>& Transform1, const float Thickness, FCollisionConstraintsArray& NewConstraints);
diff --git a/Engine/Source/Runtime/Experimental/Chaos/Private/Chaos/PBDCollisionConstraintsContact.cpp b/Engine/Source/Runtime/Experimental/Chaos/Private/Chaos/PBDCollisionConstraintsContact.cpp
--- a/Engine/Source/Runtime/Experimental/Chaos/Private/Chaos/PBDCollisionConstraintsContact.cpp
+++ b/Engine/Source/Runtime/Experimental/Chaos/Private/Chaos/PBDCollisionConstraintsContact.cpp
@@ -29,13 +29,7 @@ namespace Chaos
 				const TRigidTransform<T, d> Transform0 = GetTransform(Constraint.Particle[0]);
 				const TRigidTransform<T, d> Transform1 = GetTransform(Constraint.Particle[1]);
 
-				const FImplicitObject& Implicit0 = *Constraint.Manifold.Implicit[0];
-				const FImplicitObject& Implicit1 = *Constraint.Manifold.Implicit[1];
-
-				if (Implicit0.IsConvex() && Implicit1.IsConvex())
-				{
-					UpdateConvexConvexManifold( *Constraint.template As<TRigidBodyIterativeContactConstraint<T, d>>(), Transform0, Transform1, Thickness);
-				}
+				UpdateManifold(Constraint, Transform0, Transform1, Thickness);
 			}
 		}
 
diff --git a/Engine/Source/Runtime/Experimental/Chaos/Public/Chaos/CollisionResolution.h b/Engine/Source/Runtime/Experimental/Chaos/Public/Chaos/CollisionResolution.h
--- a/Engine/Source/Runtime/Experimental/Chaos/Public/Chaos/CollisionResolution.h
+++ b/Engine/Source/Runtime/Experimental/Chaos/Public/Chaos/CollisionResolution.h
@@ -189,6 +189,8 @@ namespace Chaos
 			//
 			// Constraint API
 			//
+			template<typename T, int d>
+			void CHAOS_API UpdateManifold(TCollisionConstraintBase<T, d>& ConstraintBase, const TRigidTransform<T, d>& ATM, const TRigidTransform<T, d>& BTM, const T Thickness);
 
 			template<ECollisionUpdateType UpdateType, typename T, int d>
 			void CHAOS_API UpdateConstraint(TCollisionConstraintBase<T, d>& ConstraintBase, const TRigidTransform<T, d>& Transform0, const TRigidTransform<T, d>& Transform1, const T Thickness);
