Remove templates from TImplicitObjectUnion

diff --git a/Engine/Plugins/Experimental/ChaosCloth/Source/Chaos/Private/ChaosCloth/ChaosClothingSimulation.cpp b/Engine/Plugins/Experimental/ChaosCloth/Source/Chaos/Private/ChaosCloth/ChaosClothingSimulation.cpp
--- a/Engine/Plugins/Experimental/ChaosCloth/Source/Chaos/Private/ChaosCloth/ChaosClothingSimulation.cpp
+++ b/Engine/Plugins/Experimental/ChaosCloth/Source/Chaos/Private/ChaosCloth/ChaosClothingSimulation.cpp
@@ -904,7 +904,7 @@ void ClothingSimulation::AddCollisions(const FClothCollisionData& ClothCollision
 					new Chaos::TSphere<float, 3>(X1, Radius1)));
 				CollisionParticles.SetDynamicGeometry(
 					i,
-					MakeUnique<Chaos::TImplicitObjectUnion<float, 3>>(MoveTemp(Objects)));  // TODO(Kriss.Gossart): Replace this once a TTaperedCapsule implicit type is implemented (note: this tapered cylinder with spheres is an approximation of a real tapered capsule)
+					MakeUnique<Chaos::FImplicitObjectUnion>(MoveTemp(Objects)));  // TODO(Kriss.Gossart): Replace this once a TTaperedCapsule implicit type is implemented (note: this tapered cylinder with spheres is an approximation of a real tapered capsule)
 			}
 
 			// Skip spheres added as end caps for the capsule.
@@ -1579,7 +1579,7 @@ void ClothingSimulation::DebugDrawCollision(USkeletalMeshComponent* OwnerCompone
 				break;
 
 			case Chaos::ImplicitObjectType::Union:  // Union only used as collision tappered capsules
-				for (const TUniquePtr<Chaos::FImplicitObject>& SubObjectPtr : Object->GetObjectChecked<Chaos::TImplicitObjectUnion<float, 3>>().GetObjects())
+				for (const TUniquePtr<Chaos::FImplicitObject>& SubObjectPtr : Object->GetObjectChecked<Chaos::FImplicitObjectUnion>().GetObjects())
 				{
 					if (const Chaos::FImplicitObject* const SubObject = SubObjectPtr.Get())
 					{
diff --git a/Engine/Source/Runtime/Engine/Private/PhysicsEngine/BodyInstance.cpp b/Engine/Source/Runtime/Engine/Private/PhysicsEngine/BodyInstance.cpp
--- a/Engine/Source/Runtime/Engine/Private/PhysicsEngine/BodyInstance.cpp
+++ b/Engine/Source/Runtime/Engine/Private/PhysicsEngine/BodyInstance.cpp
@@ -1863,7 +1863,7 @@ bool FBodyInstance::UpdateBodyScale(const FVector& InScale3D, bool bForceUpdate)
 		// Only follow through with update if all shapes succeeded.
 		if (CHAOS_ENSURE(NewGeometry.Num() == Shapes.Num()))
 		{
-			ActorHandle->SetGeometry(MakeUnique<Chaos::TImplicitObjectUnion<FReal, 3>>(MoveTemp(NewGeometry)));
+			ActorHandle->SetGeometry(MakeUnique<Chaos::FImplicitObjectUnion>(MoveTemp(NewGeometry)));
 		}
 		else
 		{
diff --git a/Engine/Source/Runtime/Engine/Private/PhysicsEngine/ImmediatePhysics/ImmediatePhysicsChaos/ImmediatePhysicsActorHandle_Chaos.cpp b/Engine/Source/Runtime/Engine/Private/PhysicsEngine/ImmediatePhysics/ImmediatePhysicsChaos/ImmediatePhysicsActorHandle_Chaos.cpp
--- a/Engine/Source/Runtime/Engine/Private/PhysicsEngine/ImmediatePhysics/ImmediatePhysicsChaos/ImmediatePhysicsActorHandle_Chaos.cpp
+++ b/Engine/Source/Runtime/Engine/Private/PhysicsEngine/ImmediatePhysics/ImmediatePhysicsChaos/ImmediatePhysicsActorHandle_Chaos.cpp
@@ -214,7 +214,7 @@ namespace ImmediatePhysics_Chaos
 		}
 		else
 		{
-			OutGeom = MakeUnique<Chaos::TImplicitObjectUnion<float, 3>>(MoveTemp(Geoms));
+			OutGeom = MakeUnique<Chaos::FImplicitObjectUnion>(MoveTemp(Geoms));
 		}
 
 		for (auto& Shape : Shapes)
diff --git a/Engine/Source/Runtime/Engine/Public/Physics/PhysicsInterfaceTypes.h b/Engine/Source/Runtime/Engine/Public/Physics/PhysicsInterfaceTypes.h
--- a/Engine/Source/Runtime/Engine/Public/Physics/PhysicsInterfaceTypes.h
+++ b/Engine/Source/Runtime/Engine/Public/Physics/PhysicsInterfaceTypes.h
@@ -13,8 +13,7 @@
 
 namespace Chaos
 {
-	template<typename T, int d>
-	class TImplicitObjectUnion;
+	class FImplicitObjectUnion;
 
 	class FImplicitObject;
 
diff --git a/Engine/Source/Runtime/Experimental/Chaos/Private/Chaos/ChaosDebugDraw.cpp b/Engine/Source/Runtime/Experimental/Chaos/Private/Chaos/ChaosDebugDraw.cpp
--- a/Engine/Source/Runtime/Experimental/Chaos/Private/Chaos/ChaosDebugDraw.cpp
+++ b/Engine/Source/Runtime/Experimental/Chaos/Private/Chaos/ChaosDebugDraw.cpp
@@ -87,7 +87,7 @@ namespace Chaos
 			}
 			case ImplicitObjectType::Union:
 			{
-				const TImplicitObjectUnion<FReal, 3>* Union = Shape->template GetObject<TImplicitObjectUnion<FReal, 3>>();
+				const FImplicitObjectUnion* Union = Shape->template GetObject<FImplicitObjectUnion>();
 				for (auto& UnionShape : Union->GetObjects())
 				{
 					DrawShapesImpl(ShapeTransform, UnionShape.Get(), Color);
diff --git a/Engine/Source/Runtime/Experimental/Chaos/Private/Chaos/CollisionResolution.cpp b/Engine/Source/Runtime/Experimental/Chaos/Private/Chaos/CollisionResolution.cpp
--- a/Engine/Source/Runtime/Experimental/Chaos/Private/Chaos/CollisionResolution.cpp
+++ b/Engine/Source/Runtime/Experimental/Chaos/Private/Chaos/CollisionResolution.cpp
@@ -1045,15 +1045,15 @@ namespace Chaos
 			{
 				ensure(false);//should not be possible to get this type, it should already be resolved by the constraint. (see ConstructConstraints)
 			}
-			else if (Implicit0Type != TImplicitObjectUnion<T, d>::StaticType() && Implicit1Type == TImplicitObjectUnion<T, d>::StaticType())
+			else if (Implicit0Type != FImplicitObjectUnion::StaticType() && Implicit1Type == FImplicitObjectUnion::StaticType())
 			{
 				ensure(false);//should not be possible to get this type, it should already be resolved by the constraint. (see ConstructConstraints)
 			}
-			else if (Implicit0Type == TImplicitObjectUnion<T, d>::StaticType() && Implicit1Type != TImplicitObjectUnion<T, d>::StaticType())
+			else if (Implicit0Type == FImplicitObjectUnion::StaticType() && Implicit1Type != FImplicitObjectUnion::StaticType())
 			{
 				ensure(false);//should not be possible to get this type, it should already be resolved by the constraint. (see ConstructConstraints)
 			}
-			else if (Implicit0Type == TImplicitObjectUnion<T, d>::StaticType() && Implicit1Type == TImplicitObjectUnion<T, d>::StaticType())
+			else if (Implicit0Type == FImplicitObjectUnion::StaticType() && Implicit1Type == FImplicitObjectUnion::StaticType())
 			{
 				ensure(false);//should not be possible to get this type, it should already be resolved by the constraint. (see ConstructConstraints)
 			}
@@ -1246,7 +1246,7 @@ namespace Chaos
 				TRigidTransform<T, d> TransformedTransform1 = TransformedImplicit1->GetTransform() * Transform1;
 				ConstructConstraints(Particle0, Particle1, Implicit0, TransformedImplicit1->GetTransformedObject(), Transform0, TransformedTransform1, Thickness, NewConstraints);
 			}
-			else if (Implicit0Type != TImplicitObjectUnion<T, d>::StaticType() && Implicit1Type == TImplicitObjectUnion<T, d>::StaticType())
+			else if (Implicit0Type != FImplicitObjectUnion::StaticType() && Implicit1Type == FImplicitObjectUnion::StaticType())
 			{
 				const TArray<Pair<const FImplicitObject*, TRigidTransform<T, d>>> LevelsetShapes = FindRelevantShapes(Implicit0, Transform0, *Implicit1, Transform1, Thickness);
 				for (const Pair<const FImplicitObject*, TRigidTransform<T, d>>& LevelsetObjPair : LevelsetShapes)
@@ -1256,7 +1256,7 @@ namespace Chaos
 					ConstructConstraints(Particle0, Particle1, Implicit0, Implicit1InnerObj, Transform0, Implicit1InnerObjTM, Thickness, NewConstraints);
 				}
 			}
-			else if (Implicit0Type == TImplicitObjectUnion<T, d>::StaticType() && Implicit1Type != TImplicitObjectUnion<T, d>::StaticType())
+			else if (Implicit0Type == FImplicitObjectUnion::StaticType() && Implicit1Type != FImplicitObjectUnion::StaticType())
 			{
 				// [Note] forces non-unions into particle[0] position
 				const TArray<Pair<const FImplicitObject*, TRigidTransform<T, d>>> LevelsetShapes = FindRelevantShapes(Implicit1, Transform1, *Implicit0, Transform0, Thickness);
@@ -1267,7 +1267,7 @@ namespace Chaos
 					ConstructConstraints(Particle0, Particle1, Implicit0InnerObj, Implicit1, Implicit0InnerObjTM, Transform1, Thickness, NewConstraints);
 				}
 			}
-			else if (Implicit0Type == TImplicitObjectUnion<T, d>::StaticType() && Implicit1Type == TImplicitObjectUnion<T, d>::StaticType())
+			else if (Implicit0Type == FImplicitObjectUnion::StaticType() && Implicit1Type == FImplicitObjectUnion::StaticType())
 			{
 				ConstructUnionUnionConstraints(Particle0, Particle1, Implicit0, Implicit1, Transform0, Transform1, Thickness, NewConstraints);
 			}
diff --git a/Engine/Source/Runtime/Experimental/Chaos/Private/Chaos/GeometryParticles.cpp b/Engine/Source/Runtime/Experimental/Chaos/Private/Chaos/GeometryParticles.cpp
--- a/Engine/Source/Runtime/Experimental/Chaos/Private/Chaos/GeometryParticles.cpp
+++ b/Engine/Source/Runtime/Experimental/Chaos/Private/Chaos/GeometryParticles.cpp
@@ -12,7 +12,7 @@ namespace Chaos
 	{
 		if(Geometry)
 		{
-			if(const auto* Union = Geometry->template GetObject<TImplicitObjectUnion<T, d>>())
+			if(const auto* Union = Geometry->template GetObject<FImplicitObjectUnion>())
 			{
 				const int32 OldShapeNum = ShapesArray.Num();
 
diff --git a/Engine/Source/Runtime/Experimental/Chaos/Private/Chaos/ImplicitObject.cpp b/Engine/Source/Runtime/Experimental/Chaos/Private/Chaos/ImplicitObject.cpp
--- a/Engine/Source/Runtime/Experimental/Chaos/Private/Chaos/ImplicitObject.cpp
+++ b/Engine/Source/Runtime/Experimental/Chaos/Private/Chaos/ImplicitObject.cpp
@@ -391,7 +391,7 @@ FImplicitObject* FImplicitObject::SerializationFactory(FChaosArchive& Ar, FImpli
 	case ImplicitObjectType::Plane: if (Ar.IsLoading()) { return new TPlane<FReal, 3>(); } break;
 	case ImplicitObjectType::Capsule: if (Ar.IsLoading()) { return new TCapsule<FReal>(); } break;
 	case ImplicitObjectType::Transformed: if (Ar.IsLoading()) { return new TImplicitObjectTransformed<FReal, 3>(); } break;
-	case ImplicitObjectType::Union: if (Ar.IsLoading()) { return new TImplicitObjectUnion<FReal, 3>(); } break;
+	case ImplicitObjectType::Union: if (Ar.IsLoading()) { return new FImplicitObjectUnion(); } break;
 	case ImplicitObjectType::LevelSet: if (Ar.IsLoading()) { return new TLevelSet<FReal, 3>(); } break;
 	case ImplicitObjectType::Convex: if (Ar.IsLoading()) { return new FConvex(); } break;
 	case ImplicitObjectType::TaperedCylinder: if (Ar.IsLoading()) { return new TTaperedCylinder<FReal>(); } break;
diff --git a/Engine/Source/Runtime/Experimental/Chaos/Private/Chaos/ImplicitObjectUnion.cpp b/Engine/Source/Runtime/Experimental/Chaos/Private/Chaos/ImplicitObjectUnion.cpp
--- a/Engine/Source/Runtime/Experimental/Chaos/Private/Chaos/ImplicitObjectUnion.cpp
+++ b/Engine/Source/Runtime/Experimental/Chaos/Private/Chaos/ImplicitObjectUnion.cpp
@@ -4,11 +4,10 @@
 
 using namespace Chaos;
 
-template<class T, int d>
-TImplicitObjectUnion<T,d>::TImplicitObjectUnion(TArray<TUniquePtr<FImplicitObject>>&& Objects, const TArray<int32>& OriginalParticleLookupHack)
+FImplicitObjectUnion::FImplicitObjectUnion(TArray<TUniquePtr<FImplicitObject>>&& Objects, const TArray<int32>& OriginalParticleLookupHack)
 	: FImplicitObject(EImplicitObject::HasBoundingBox, ImplicitObjectType::Union)
 	, MObjects(MoveTemp(Objects))
-	, Hierarchy(new TBoundingVolumeHierarchy<TGeometryParticles<T, d>, TArray<int32>, T, d>(GeomParticles))
+	, Hierarchy(new TBoundingVolumeHierarchy<TGeometryParticles<FReal,3>, TArray<int32>, FReal, 3>(GeomParticles))
 	, MLocalBoundingBox()
 	, bHierarchyBuilt(false)
 	, MOriginalParticleLookupHack(OriginalParticleLookupHack)
@@ -28,7 +27,7 @@ TImplicitObjectUnion<T,d>::TImplicitObjectUnion(TArray<TUniquePtr<FImplicitObjec
 		if (MOriginalParticleLookupHack.Num() > 0)
 		{
 			//this whole part sucks, only needed because of how we get union children. Need to refactor and enforce no unions of unions
-			if (const TImplicitObjectTransformed<T, d>* Transformed = MObjects[i]->template GetObject<const TImplicitObjectTransformed<T, d>>())
+			if (const TImplicitObjectTransformed<FReal,3>* Transformed = MObjects[i]->template GetObject<const TImplicitObjectTransformed<FReal,3>>())
 			{
 				MCollisionParticleLookupHack.Add(Transformed->GetTransformedObject(), MOriginalParticleLookupHack[i]);
 			}
@@ -42,12 +41,11 @@ TImplicitObjectUnion<T,d>::TImplicitObjectUnion(TArray<TUniquePtr<FImplicitObjec
 	CacheAllImplicitObjects();
 }
 
-template<class T, int d>
-TImplicitObjectUnion<T,d>::TImplicitObjectUnion(TImplicitObjectUnion<T, d>&& Other)
+FImplicitObjectUnion::FImplicitObjectUnion(FImplicitObjectUnion&& Other)
 	: FImplicitObject(EImplicitObject::HasBoundingBox)
 	, MObjects(MoveTemp(Other.MObjects))
 	, GeomParticles(MoveTemp(Other.GeomParticles))
-	, Hierarchy(new TBoundingVolumeHierarchy<TGeometryParticles<T, d>, TArray<int32>, T, d>(MoveTemp(*Other.Hierarchy)))
+	, Hierarchy(new TBoundingVolumeHierarchy<TGeometryParticles<FReal,3>, TArray<int32>, FReal, 3>(MoveTemp(*Other.Hierarchy)))
 	, MLocalBoundingBox(MoveTemp(Other.MLocalBoundingBox))
 	, bHierarchyBuilt(Other.bHierarchyBuilt)
 	, MOriginalParticleLookupHack(MoveTemp(Other.MOriginalParticleLookupHack))
@@ -55,14 +53,12 @@ TImplicitObjectUnion<T,d>::TImplicitObjectUnion(TImplicitObjectUnion<T, d>&& Oth
 {
 }
 
-template<class T, int d>
-TImplicitObjectUnion<T, d>::~TImplicitObjectUnion()
+FImplicitObjectUnion::~FImplicitObjectUnion()
 {
 	delete Hierarchy;
 }
 
-template<class T, int d>
-void TImplicitObjectUnion<T,d>::FindAllIntersectingObjects(TArray < Pair<const FImplicitObject*, TRigidTransform<T, d>>>& Out, const TAABB<T, d>& LocalBounds) const
+void FImplicitObjectUnion::FindAllIntersectingObjects(TArray < Pair<const FImplicitObject*,FRigidTransform3>>& Out, const TAABB<FReal,3>& LocalBounds) const
 {
 	if (bHierarchyBuilt)
 	{
@@ -71,7 +67,7 @@ void TImplicitObjectUnion<T,d>::FindAllIntersectingObjects(TArray < Pair<const F
 		for (int32 Idx : Overlaps)
 		{
 			const FImplicitObject* Obj = GeomParticles.Geometry(Idx).Get();
-			Out.Add(MakePair(Obj, TRigidTransform<T, d>(GeomParticles.X(Idx), GeomParticles.R(Idx))));
+			Out.Add(MakePair(Obj,FRigidTransform3(GeomParticles.X(Idx), GeomParticles.R(Idx))));
 		}
 	}
 	else
@@ -83,8 +79,7 @@ void TImplicitObjectUnion<T,d>::FindAllIntersectingObjects(TArray < Pair<const F
 	}
 }
 
-template<class T, int d>
-TArray<int32> TImplicitObjectUnion<T,d>::FindAllIntersectingChildren(const TAABB<T, d>& LocalBounds) const
+TArray<int32> FImplicitObjectUnion::FindAllIntersectingChildren(const TAABB<FReal,3>& LocalBounds) const
 {
 	TArray<int32> IntersectingChildren;
 	if (bHierarchyBuilt) //todo: make this work when hierarchy is not built
@@ -115,8 +110,7 @@ TArray<int32> TImplicitObjectUnion<T,d>::FindAllIntersectingChildren(const TAABB
 	return IntersectingChildren;
 }
 
-template<class T, int d>
-TArray<int32> TImplicitObjectUnion<T,d>::FindAllIntersectingChildren(const TSpatialRay<T, d>& LocalRay) const
+TArray<int32> FImplicitObjectUnion::FindAllIntersectingChildren(const TSpatialRay<FReal,3>& LocalRay) const
 {
 	TArray<int32> IntersectingChildren;
 	if (bHierarchyBuilt) //todo: make this work when hierarchy is not built
@@ -147,11 +141,10 @@ TArray<int32> TImplicitObjectUnion<T,d>::FindAllIntersectingChildren(const TSpat
 	return IntersectingChildren;
 }
 
-template<class T, int d>
-void TImplicitObjectUnion<T,d>::CacheAllImplicitObjects()
+void FImplicitObjectUnion::CacheAllImplicitObjects()
 {
-	TArray < Pair<TSerializablePtr<FImplicitObject>, TRigidTransform<T, d>>> SubObjects;
-	AccumulateAllSerializableImplicitObjectsHelper(SubObjects, TRigidTransform<T, d>::Identity);
+	TArray < Pair<TSerializablePtr<FImplicitObject>,FRigidTransform3>> SubObjects;
+	AccumulateAllSerializableImplicitObjectsHelper(SubObjects,FRigidTransform3::Identity);
 	//build hierarchy
 	{
 		const int32 NumObjects = SubObjects.Num();
@@ -167,24 +160,20 @@ void TImplicitObjectUnion<T,d>::CacheAllImplicitObjects()
 				//check(!SubObjects[i].First->IsUnderlyingUnion());	//we don't support union of unions
 			}
 
-			TBoundingVolumeHierarchy<TGeometryParticles<T, d>, TArray<int32>, T, d> NewHierarchy(GeomParticles, 1);
+			TBoundingVolumeHierarchy<TGeometryParticles<FReal,3>, TArray<int32>, FReal, 3> NewHierarchy(GeomParticles, 1);
 			*Hierarchy = MoveTemp(NewHierarchy);
 			bHierarchyBuilt = true;
 		}
 	}
 }
 
-template<class T, int d>
-void TImplicitObjectUnion<T,d>::Serialize(FChaosArchive& Ar)
+void FImplicitObjectUnion::Serialize(FChaosArchive& Ar)
 {
 	FChaosArchiveScopedMemory ScopedMemory(Ar, GetTypeName(), false);
 	FImplicitObject::SerializeImp(Ar);
 	Ar << MObjects;
-	TBox<T, d>::SerializeAsAABB(Ar, MLocalBoundingBox);
+	TBox<FReal,3>::SerializeAsAABB(Ar, MLocalBoundingBox);
 	Ar << GeomParticles << *Hierarchy << bHierarchyBuilt;
 }
 
-template<class T, int d>
-TImplicitObjectUnion<T, d>::TImplicitObjectUnion() : FImplicitObject(EImplicitObject::HasBoundingBox, ImplicitObjectType::Union), Hierarchy(new TBoundingVolumeHierarchy<TGeometryParticles<T, d>, TArray<int32>, T, d>(GeomParticles, 1)){}
-
-template class Chaos::TImplicitObjectUnion<float, 3>;
+FImplicitObjectUnion::FImplicitObjectUnion() : FImplicitObject(EImplicitObject::HasBoundingBox, ImplicitObjectType::Union), Hierarchy(new TBoundingVolumeHierarchy<TGeometryParticles<FReal,3>, TArray<int32>, FReal, 3>(GeomParticles, 1)){}
diff --git a/Engine/Source/Runtime/Experimental/Chaos/Private/Chaos/PBDRigidClustering.cpp b/Engine/Source/Runtime/Experimental/Chaos/Private/Chaos/PBDRigidClustering.cpp
--- a/Engine/Source/Runtime/Experimental/Chaos/Private/Chaos/PBDRigidClustering.cpp
+++ b/Engine/Source/Runtime/Experimental/Chaos/Private/Chaos/PBDRigidClustering.cpp
@@ -1360,7 +1360,7 @@ namespace Chaos
 
 		{
 			QUICK_SCOPE_CYCLE_COUNTER(SpatialBVH);
-			MParticles.ChildrenSpatial(NewIndex) = Objects2.Num() ? MakeUnique<TImplicitObjectUnion<T, d>>(MoveTemp(Objects2), GeomToOriginalParticlesHack) : nullptr;
+			MParticles.ChildrenSpatial(NewIndex) = Objects2.Num() ? MakeUnique<FImplicitObjectUnion>(MoveTemp(Objects2), GeomToOriginalParticlesHack) : nullptr;
 		}
 
 		TArray<TVector<T, d>> CleanedPoints;
@@ -1394,7 +1394,7 @@ namespace Chaos
 			{
 				ensureMsgf(false, TEXT("Checking usage with no proxy and multiple ojects with levelsets"));
 
-				TImplicitObjectUnion<T, d> UnionObject(MoveTemp(Objects));
+				FImplicitObjectUnion UnionObject(MoveTemp(Objects));
 				TAABB<T, d> Bounds = UnionObject.BoundingBox();
 				const TVector<T, d> BoundsExtents = Bounds.Extents();
 				if (BoundsExtents.Min() >= MinLevelsetSize) //make sure the object is not too small
@@ -1444,7 +1444,7 @@ namespace Chaos
 					QUICK_SCOPE_CYCLE_COUNTER(UnionBVH);
 					// @coverage : { confidence tests}
 					//ensureMsgf(false, TEXT("Checking no proxy, not levelset, and multiple objects"));
-					MParticles.SetDynamicGeometry(NewIndex, MakeUnique<TImplicitObjectUnion<T, d>>(MoveTemp(Objects), GeomToOriginalParticlesHack));
+					MParticles.SetDynamicGeometry(NewIndex, MakeUnique<FImplicitObjectUnion>(MoveTemp(Objects), GeomToOriginalParticlesHack));
 				}
 			}
 		}
diff --git a/Engine/Source/Runtime/Experimental/Chaos/Public/Chaos/Capsule.h b/Engine/Source/Runtime/Experimental/Chaos/Public/Chaos/Capsule.h
--- a/Engine/Source/Runtime/Experimental/Chaos/Public/Chaos/Capsule.h
+++ b/Engine/Source/Runtime/Experimental/Chaos/Public/Chaos/Capsule.h
@@ -377,7 +377,7 @@ namespace Chaos
 			Objects.Add(MakeUnique<Chaos::TSphere<float, 3>>(X1, MRadius));
 			Objects.Add(MakeUnique<Chaos::TSphere<float, 3>>(X2, MRadius));
 
-			MUnionedObjects.Reset(new Chaos::TImplicitObjectUnion<float, 3>(std::move(Objects)));
+			MUnionedObjects.Reset(new Chaos::FImplicitObjectUnion(std::move(Objects)));
 		}
 
 		virtual Pair<TVector<T, 3>, bool> FindClosestIntersectionImp(const TVector<T, 3>& StartPoint, const TVector<T, 3>& EndPoint, const T Thickness) const override
@@ -388,7 +388,7 @@ namespace Chaos
 		TSegment<T> MSegment;
 		T MRadius;
 		TAABB<T, 3> MLocalBoundingBox;
-		TUniquePtr<TImplicitObjectUnion<T, 3>> MUnionedObjects;
+		TUniquePtr<FImplicitObjectUnion> MUnionedObjects;
 	};
 
 	template<typename T>
diff --git a/Engine/Source/Runtime/Experimental/Chaos/Public/Chaos/ImplicitObjectUnion.h b/Engine/Source/Runtime/Experimental/Chaos/Public/Chaos/ImplicitObjectUnion.h
--- a/Engine/Source/Runtime/Experimental/Chaos/Public/Chaos/ImplicitObjectUnion.h
+++ b/Engine/Source/Runtime/Experimental/Chaos/Public/Chaos/ImplicitObjectUnion.h
@@ -15,26 +15,25 @@ namespace Chaos
 template<class OBJECT_ARRAY, class LEAF_TYPE, class T, int d>
 class TBoundingVolumeHierarchy;
 
-template<class T, int d>
-class TImplicitObjectUnion : public FImplicitObject
+class FImplicitObjectUnion : public FImplicitObject
 {
   public:
 
 	using FImplicitObject::GetTypeName;
 
-	CHAOS_API TImplicitObjectUnion(TArray<TUniquePtr<FImplicitObject>>&& Objects, const TArray<int32>& OriginalParticleLookupHack = TArray<int32>());
-	TImplicitObjectUnion(const TImplicitObjectUnion<T, d>& Other) = delete;
-	TImplicitObjectUnion(TImplicitObjectUnion<T, d>&& Other);
-	virtual ~TImplicitObjectUnion();
+	CHAOS_API FImplicitObjectUnion(TArray<TUniquePtr<FImplicitObject>>&& Objects, const TArray<int32>& OriginalParticleLookupHack = TArray<int32>());
+	FImplicitObjectUnion(const FImplicitObjectUnion& Other) = delete;
+	FImplicitObjectUnion(FImplicitObjectUnion&& Other);
+	virtual ~FImplicitObjectUnion();
 
 	FORCEINLINE static constexpr EImplicitObjectType StaticType()
 	{
 		return ImplicitObjectType::Union;
 	}
 
-	virtual T PhiWithNormal(const TVector<T, d>& x, TVector<T, d>& Normal) const override
+	virtual FReal PhiWithNormal(const FVec3& x, FVec3& Normal) const override
 	{
-		T Phi = TNumericLimits<T>::Max();
+		FReal Phi = TNumericLimits<FReal>::Max();
 		bool NeedsNormalize = false;
 		for (int32 i = 0; i < MObjects.Num(); ++i)
 		{
@@ -42,8 +41,8 @@ class TImplicitObjectUnion : public FImplicitObject
 			{
 				continue;
 			}
-			TVector<T, d> NextNormal;
-			T NextPhi = MObjects[i]->PhiWithNormal(x, NextNormal);
+			FVec3 NextNormal;
+			FReal NextPhi = MObjects[i]->PhiWithNormal(x, NextNormal);
 			if (NextPhi < Phi)
 			{
 				Phi = NextPhi;
@@ -63,9 +62,9 @@ class TImplicitObjectUnion : public FImplicitObject
 		return Phi;
 	}
 
-	virtual const TAABB<T, d>& BoundingBox() const override { return MLocalBoundingBox; }
+	virtual const TAABB<FReal, 3>& BoundingBox() const override { return MLocalBoundingBox; }
 
-	virtual void AccumulateAllImplicitObjects(TArray<Pair<const FImplicitObject*, TRigidTransform<T, d>>>& Out, const TRigidTransform<T, d>& ParentTM) const
+	virtual void AccumulateAllImplicitObjects(TArray<Pair<const FImplicitObject*, FRigidTransform3>>& Out, const FRigidTransform3& ParentTM) const
 	{
 		for (const TUniquePtr<FImplicitObject>& Object : MObjects)
 		{
@@ -73,12 +72,12 @@ class TImplicitObjectUnion : public FImplicitObject
 		}
 	}
 
-	virtual void AccumulateAllSerializableImplicitObjects(TArray<Pair<TSerializablePtr<FImplicitObject>, TRigidTransform<T, d>>>& Out, const TRigidTransform<T, d>& ParentTM, TSerializablePtr<FImplicitObject> This) const
+	virtual void AccumulateAllSerializableImplicitObjects(TArray<Pair<TSerializablePtr<FImplicitObject>, FRigidTransform3>>& Out, const FRigidTransform3& ParentTM, TSerializablePtr<FImplicitObject> This) const
 	{
 		AccumulateAllSerializableImplicitObjectsHelper(Out, ParentTM);
 	}
 
-	void AccumulateAllSerializableImplicitObjectsHelper(TArray<Pair<TSerializablePtr<FImplicitObject>, TRigidTransform<T, d>>>& Out, const TRigidTransform<T, d>& ParentTM) const
+	void AccumulateAllSerializableImplicitObjectsHelper(TArray<Pair<TSerializablePtr<FImplicitObject>, FRigidTransform3>>& Out, const FRigidTransform3& ParentTM) const
 	{
 		for (const TUniquePtr<FImplicitObject>& Object : MObjects)
 		{
@@ -86,21 +85,21 @@ class TImplicitObjectUnion : public FImplicitObject
 		}
 	}
 
-	virtual void FindAllIntersectingObjects(TArray < Pair<const FImplicitObject*, TRigidTransform<T, d>>>& Out, const TAABB<T, d>& LocalBounds) const;
-	TArray<int32> FindAllIntersectingChildren(const TAABB<T, d>& LocalBounds) const;
-	TArray<int32> FindAllIntersectingChildren(const TSpatialRay<T, d>& LocalRay) const;
+	virtual void FindAllIntersectingObjects(TArray < Pair<const FImplicitObject*, FRigidTransform3>>& Out, const TAABB<FReal,3>& LocalBounds) const;
+	TArray<int32> FindAllIntersectingChildren(const TAABB<FReal,3>& LocalBounds) const;
+	TArray<int32> FindAllIntersectingChildren(const TSpatialRay<FReal,3>& LocalRay) const;
 	virtual void CacheAllImplicitObjects();
 
-	virtual bool Raycast(const TVector<T, d>& StartPoint, const TVector<T, d>& Dir, const T Length, const T Thickness, T& OutTime, TVector<T, d>& OutPosition, TVector<T, d>& OutNormal, int32& OutFaceIndex) const override
+	virtual bool Raycast(const FVec3& StartPoint, const FVec3& Dir, const FReal Length, const FReal Thickness, FReal& OutTime, FVec3& OutPosition, FVec3& OutNormal, int32& OutFaceIndex) const override
 	{
-		T MinTime = 0;	//initialization not needed, but doing it to avoid warning
+		FReal MinTime = 0;	//initialization not needed, but doing it to avoid warning
 		bool bFound = false;
 
 		for (const TUniquePtr<FImplicitObject>& Obj : MObjects)
 		{
-			TVector<T, d> Position;
-			TVector<T, d> Normal;
-			T Time;
+			FVec3 Position;
+			FVec3 Normal;
+			FReal Time;
 			int32 FaceIdx;
 			if (Obj->Raycast(StartPoint, Dir, Length, Thickness, Time, Position, Normal, FaceIdx))
 			{
@@ -119,7 +118,7 @@ class TImplicitObjectUnion : public FImplicitObject
 		return bFound;
 	}
 
-	virtual bool Overlap(const TVector<T, d>& Point, const T Thickness) const override
+	virtual bool Overlap(const FVec3& Point, const FReal Thickness) const override
 	{
 		for (const TUniquePtr<FImplicitObject>& Obj : MObjects)
 		{
@@ -157,17 +156,17 @@ class TImplicitObjectUnion : public FImplicitObject
 	}
 
 private:
-	virtual Pair<TVector<T, d>, bool> FindClosestIntersectionImp(const TVector<T, d>& StartPoint, const TVector<T, d>& EndPoint, const T Thickness) const override
+	virtual Pair<FVec3, bool> FindClosestIntersectionImp(const FVec3& StartPoint, const FVec3& EndPoint, const FReal Thickness) const override
 	{
 		check(MObjects.Num());
 		auto ClosestIntersection = MObjects[0]->FindClosestIntersection(StartPoint, EndPoint, Thickness);
-		T Length = ClosestIntersection.Second ? (ClosestIntersection.First - StartPoint).Size() : 0;
+		FReal Length = ClosestIntersection.Second ? (ClosestIntersection.First - StartPoint).Size() : 0;
 		for (int32 i = 1; i < MObjects.Num(); ++i)
 		{
 			auto NextClosestIntersection = MObjects[i]->FindClosestIntersection(StartPoint, EndPoint, Thickness);
 			if (!NextClosestIntersection.Second)
 				continue;
-			T NewLength = (NextClosestIntersection.First - StartPoint).Size();
+			FReal NewLength = (NextClosestIntersection.First - StartPoint).Size();
 			if (!ClosestIntersection.Second || NewLength < Length)
 			{
 				Length = NewLength;
@@ -179,15 +178,15 @@ private:
 
   private:
 	TArray<TUniquePtr<FImplicitObject>> MObjects;
-	TGeometryParticles<T, d> GeomParticles;
-	TBoundingVolumeHierarchy<TGeometryParticles<T, d>, TArray<int32>, T, d>* Hierarchy;
-	TAABB<T, d> MLocalBoundingBox;
+	TGeometryParticles<FReal,3> GeomParticles;
+	TBoundingVolumeHierarchy<TGeometryParticles<FReal,3>, TArray<int32>, FReal, 3>* Hierarchy;
+	TAABB<FReal,3> MLocalBoundingBox;
 	bool bHierarchyBuilt;
 	TArray<int32> MOriginalParticleLookupHack;	//temp hack for finding original particles
 
 
 	//needed for serialization
-	TImplicitObjectUnion();
+	FImplicitObjectUnion();
 	friend FImplicitObject;	//needed for serialization
 public:
 	TMap<const FImplicitObject*, int32> MCollisionParticleLookupHack;	//temp hack for finding collision particles
diff --git a/Engine/Source/Runtime/Experimental/Chaos/Public/Chaos/PBDRigidClusteredParticles.h b/Engine/Source/Runtime/Experimental/Chaos/Public/Chaos/PBDRigidClusteredParticles.h
--- a/Engine/Source/Runtime/Experimental/Chaos/Public/Chaos/PBDRigidClusteredParticles.h
+++ b/Engine/Source/Runtime/Experimental/Chaos/Public/Chaos/PBDRigidClusteredParticles.h
@@ -7,8 +7,7 @@
 
 namespace Chaos
 {
-template<class T, int d>
-class TImplicitObjectUnion; 
+	class FImplicitObjectUnion;
 
 // ClusterId 
 //    Used within the clustering system to describe the clustering hierarchy. The ID will
@@ -160,7 +159,7 @@ class TPBDRigidClusteredParticles : public TPBDRigidParticles<T, d>
 	  TArrayCollectionArray<TRigidTransform<T, d>> MChildToParent;
 	  TArrayCollectionArray<int32> MClusterGroupIndex;
 	  TArrayCollectionArray<bool> MInternalCluster;
-	  TArrayCollectionArray<TUniquePtr<TImplicitObjectUnion<T, d>>> MChildrenSpatial;
+	  TArrayCollectionArray<TUniquePtr<FImplicitObjectUnion>> MChildrenSpatial;
 	  TArrayCollectionArray<FMultiChildProxyId> MMultiChildProxyId;
 	  TArrayCollectionArray<TUniquePtr<TMultiChildProxyData<T, d>>> MMultiChildProxyData;
 
diff --git a/Engine/Source/Runtime/Experimental/Chaos/Public/Chaos/ParticleHandle.h b/Engine/Source/Runtime/Experimental/Chaos/Public/Chaos/ParticleHandle.h
--- a/Engine/Source/Runtime/Experimental/Chaos/Public/Chaos/ParticleHandle.h
+++ b/Engine/Source/Runtime/Experimental/Chaos/Public/Chaos/ParticleHandle.h
@@ -1570,7 +1570,7 @@ public:
 		FParticleData::Reset();  
 		X = TVector<T, d>(0); 
 		R = TRotation<T, d>(); 
-		Geometry = TSharedPtr<TImplicitObjectUnion<T, d>, ESPMode::ThreadSafe>();
+		Geometry = TSharedPtr<FImplicitObjectUnion, ESPMode::ThreadSafe>();
 		SpatialIdx = FSpatialAccelerationIdx{ 0,0 };
 		HashResult = 0;
 		DirtyFlags.Clear();
diff --git a/Engine/Source/Runtime/Experimental/ChaosSolvers/Public/PhysicsProxy/AnalyticImplicitGroup.h b/Engine/Source/Runtime/Experimental/ChaosSolvers/Public/PhysicsProxy/AnalyticImplicitGroup.h
--- a/Engine/Source/Runtime/Experimental/ChaosSolvers/Public/PhysicsProxy/AnalyticImplicitGroup.h
+++ b/Engine/Source/Runtime/Experimental/ChaosSolvers/Public/PhysicsProxy/AnalyticImplicitGroup.h
@@ -457,7 +457,7 @@ public:
 								Chaos::TRigidTransform<float, 3>(Xf))));
 				}
 			}
-			return new Chaos::TImplicitObjectUnion<float, 3>(MoveTemp(ImplicitObjects));
+			return new Chaos::FImplicitObjectUnion(MoveTemp(ImplicitObjects));
 		}
 	}
 
diff --git a/Engine/Source/Runtime/Experimental/GeometryCollectionEngine/Private/GeometryCollection/GeometryCollectionParticlesData.cpp b/Engine/Source/Runtime/Experimental/GeometryCollectionEngine/Private/GeometryCollection/GeometryCollectionParticlesData.cpp
--- a/Engine/Source/Runtime/Experimental/GeometryCollectionEngine/Private/GeometryCollection/GeometryCollectionParticlesData.cpp
+++ b/Engine/Source/Runtime/Experimental/GeometryCollectionEngine/Private/GeometryCollection/GeometryCollectionParticlesData.cpp
@@ -158,7 +158,7 @@ void TGeometryCollectionParticlesData<T, d>::FData::Copy(EGeometryCollectionPart
 			else if (ImplicitObject->template GetObject<Chaos::TAABB                      <T, d>>()) { return Chaos::ImplicitObjectType::Box        ; }
 			else if (ImplicitObject->template GetObject<Chaos::TPlane                    <T, d>>()) { return Chaos::ImplicitObjectType::Plane      ; }
 			else if (ImplicitObject->template GetObject<Chaos::TImplicitObjectTransformed<T, d>>()) { return Chaos::ImplicitObjectType::Transformed; }
-			else if (ImplicitObject->template GetObject<Chaos::TImplicitObjectUnion      <T, d>>()) { return Chaos::ImplicitObjectType::Union      ; }
+			else if (ImplicitObject->template GetObject<Chaos::FImplicitObjectUnion>()) { return Chaos::ImplicitObjectType::Union      ; }
 			else if (ImplicitObject->template GetObject<Chaos::TLevelSet                 <T, d>>()) { return Chaos::ImplicitObjectType::LevelSet   ; }
 		}
 		return Chaos::ImplicitObjectType::Unknown;
diff --git a/Engine/Source/Runtime/Experimental/GeometryCollectionEngine/Private/GeometryCollection/GeometryCollectionSQAccelerator.cpp b/Engine/Source/Runtime/Experimental/GeometryCollectionEngine/Private/GeometryCollection/GeometryCollectionSQAccelerator.cpp
--- a/Engine/Source/Runtime/Experimental/GeometryCollectionEngine/Private/GeometryCollection/GeometryCollectionSQAccelerator.cpp
+++ b/Engine/Source/Runtime/Experimental/GeometryCollectionEngine/Private/GeometryCollection/GeometryCollectionSQAccelerator.cpp
@@ -302,7 +302,7 @@ void FGeometryCollectionSQAccelerator::Raycast(const FVector& Start, const FVect
 				}
 				if (Object && !UseSlowSQ && Object->IsUnderlyingUnion())
 				{
-					const TImplicitObjectUnion<float, 3>* Union = static_cast<const TImplicitObjectUnion<float, 3>*>(Object);
+					const FImplicitObjectUnion* Union = static_cast<const TImplicitObjectUnion<float, 3>*>(Object);
 					//hack: this is terrible because we have no buffered transform so could be off, but most of the time these things are static
 
 					{
@@ -567,7 +567,7 @@ void FGeometryCollectionSQAccelerator::Sweep(const Chaos::FImplicitObject& Query
 				}
 				if (Object && !UseSlowSQ && Object->IsUnderlyingUnion())
 				{
-					const TImplicitObjectUnion<float, 3>* Union = static_cast<const TImplicitObjectUnion<float, 3>*>(Object);
+					const FImplicitObjectUnion* Union = static_cast<const TImplicitObjectUnion<float, 3>*>(Object);
 					//hack: this is terrible because we have no buffered transform so could be off, but most of the time these things are static
 
 					{
diff --git a/Engine/Source/Runtime/Landscape/Private/LandscapeCollision.cpp b/Engine/Source/Runtime/Landscape/Private/LandscapeCollision.cpp
--- a/Engine/Source/Runtime/Landscape/Private/LandscapeCollision.cpp
+++ b/Engine/Source/Runtime/Landscape/Private/LandscapeCollision.cpp
@@ -483,7 +483,7 @@ void ULandscapeHeightfieldCollisionComponent::OnCreatePhysicsState()
 				}
 				else
 				{
-					PhysHandle->SetGeometry(MakeUnique<Chaos::TImplicitObjectUnion<float, 3>>(MoveTemp(Geoms)));
+					PhysHandle->SetGeometry(MakeUnique<Chaos::FImplicitObjectUnion>(MoveTemp(Geoms)));
 				}
 
 				// Construct Shape Bounds
diff --git a/Engine/Source/Runtime/PhysicsCore/Private/PhysTestSerializer.cpp b/Engine/Source/Runtime/PhysicsCore/Private/PhysTestSerializer.cpp
--- a/Engine/Source/Runtime/PhysicsCore/Private/PhysTestSerializer.cpp
+++ b/Engine/Source/Runtime/PhysicsCore/Private/PhysTestSerializer.cpp
@@ -305,7 +305,7 @@ void FPhysTestSerializer::CreateChaosData()
 				}
 				else
 				{
-					GTParticle->SetGeometry(MakeUnique<TImplicitObjectUnion<float, 3>>(MoveTemp(Geoms)));
+					GTParticle->SetGeometry(MakeUnique<FImplicitObjectUnion>(MoveTemp(Geoms)));
 					Particle->SetGeometry(GTParticle->Geometry());
 				}
 
