Chaos : Collisions

- Restructuring to help with manifold callbacks.

diff --git a/Engine/Source/Runtime/Experimental/Chaos/Private/Chaos/CollisionResolution.cpp b/Engine/Source/Runtime/Experimental/Chaos/Private/Chaos/CollisionResolution.cpp
--- a/Engine/Source/Runtime/Experimental/Chaos/Private/Chaos/CollisionResolution.cpp
+++ b/Engine/Source/Runtime/Experimental/Chaos/Private/Chaos/CollisionResolution.cpp
@@ -33,6 +33,7 @@ namespace Chaos
 			TContactPoint() : Phi(TNumericLimits<T>::Max()) {}
 		};
 
+
 		template <typename T>
 		void UpdateContactPoint(TCollisionContact<T, 3>& Manifold, const TContactPoint<T>& NewContactPoint)
 		{
@@ -73,15 +74,137 @@ namespace Chaos
 			return Contact;
 		}
 
+		template <typename GeometryA, typename GeometryB, typename T, int d>
+		TContactPoint<T> GJKImplicitContactPoint(const FImplicitObject& A, const TRigidTransform<T, d>& ATransform, const GeometryB& B, const TRigidTransform<T, d>& BTransform, const T Thickness)
+		{
+			TContactPoint<T> Contact;
+			const TRigidTransform<T, d> AToBTM = ATransform.GetRelativeTransform(BTransform);
+
+			T Penetration = FLT_MAX;
+			TVec3<T> Location, Normal;
+			if (const TImplicitObjectScaled<GeometryA>* ScaledConvexImplicit = A.template GetObject<const TImplicitObjectScaled<GeometryA> >())
+			{
+				if (B.GJKContactPoint(*ScaledConvexImplicit, AToBTM, Thickness, Location, Normal, Penetration))
+				{
+					Contact.Phi = Penetration;
+					Contact.Location = BTransform.TransformPosition(Location);
+					Contact.Normal = BTransform.TransformVector(Normal);
+				}
+			}
+			else if (const GeometryA* ConvexImplicit = A.template GetObject<const GeometryA>())
+			{
+				if (B.GJKContactPoint(*ConvexImplicit, AToBTM, Thickness, Location, Normal, Penetration))
+				{
+					Contact.Phi = Penetration;
+					Contact.Location = BTransform.TransformPosition(Location);
+					Contact.Normal = BTransform.TransformVector(Normal);
+				}
+			}
+
+			return Contact;
+		}
+
 
 		template <typename T, int d>
-		TContactPoint<T> BoxBoxContactPoint(const TAABB<T, d>& Box1, const TRigidTransform<T, d>& ATM, const TAABB<T, d>& Box2, const TRigidTransform<T, d>& BTM, const T Thickness)
+		void UpdateSingleShotManifold(TRigidBodyIterativeContactConstraint<T, d>&  Constraint, const TRigidTransform<T, d>& Transform0, const TRigidTransform<T, d>& Transform1, const T Thickness)
 		{
-			return GJKContactPoint(Box1, ATM, Box2, BTM, Thickness);
+			// single shot manifolds for TConvex implicit object in the constraints implicit[0] position. 
+			TContactPoint<T> ContactPoint = ConvexConvexContactPoint(*Constraint.Manifold.Implicit[0], Transform0, *Constraint.Manifold.Implicit[1], Transform1, Thickness);
+
+			TArray<FVec3> CollisionSamples;
+			//
+			//  @todo(chaos) : Collision Manifold
+			//   Remove the dependency on the virtual calls on the Implicit. Don't use FindClosestFaceAndVertices
+			//   this relies on virtual calls on the ImplicitObject. Instead pass a parameters structures into 
+			//   ConvexConvexContactPoint that can collect the face indices during evaluation of the support functions. 
+			//   This can be implemented without virtual calls.
+			//
+			int32 FaceIndex = Constraint.Manifold.Implicit[0]->FindClosestFaceAndVertices(Transform0.InverseTransformPosition(ContactPoint.Location), CollisionSamples, 1.f);
+
+			if (!ContactPoint.Normal.Equals(Constraint.PlaneNormal) || !Constraint.NumSamples())
+			{
+				Constraint.PlaneNormal = Transform1.InverseTransformVector(ContactPoint.Normal);
+				Constraint.PlanePosition = Transform1.InverseTransformPosition(ContactPoint.Location - ContactPoint.Phi*ContactPoint.Normal);
+			}
+
+
+			if (FaceIndex != Constraint.SourceNormalIndex || !Constraint.NumSamples())
+			{
+				Constraint.ResetSamples(CollisionSamples.Num());
+				Constraint.SourceNormalIndex = FaceIndex;
+
+				//
+				// @todo(chaos) : Collision Manifold
+				//   Only save the four best samples and hard-code the size of Constraint.Samples to [len:4].
+				//   Currently this just grabs all points and uses the deepest point for resolution. 
+				//
+				for (FVec3 Sample : CollisionSamples)
+				{
+					Constraint.AddSample({ Sample,0.f });
+				}
+			}
 		}
+
+		template <typename T, int d>
+		void UpdateIterativeManifold(TRigidBodyIterativeContactConstraint<T, d>&  Constraint, const TRigidTransform<T, d>& Transform0, const TRigidTransform<T, d>& Transform1, const T Thickness)
+		{
+			// iterative manifolds for non TConvex implicit objects that require sampling 
+			TContactPoint<T> ContactPoint = ConvexConvexContactPoint(*Constraint.Manifold.Implicit[0], Transform0, *Constraint.Manifold.Implicit[1], Transform1, Thickness);
+
+			if (!ContactPoint.Normal.Equals(Constraint.PlaneNormal) || !Constraint.NumSamples())
+			{
+				Constraint.ResetSamples();
+				Constraint.PlaneNormal = Transform1.InverseTransformVector(ContactPoint.Normal);
+				Constraint.PlanePosition = Transform1.InverseTransformPosition(ContactPoint.Location - ContactPoint.Phi*ContactPoint.Normal);
+			}
+
+			TVector<T, d> SurfaceSample = Transform0.InverseTransformPosition(ContactPoint.Location);
+			if (Constraint.NumSamples() < 4)
+			{
+				Constraint.AddSample({ SurfaceSample,0.f });
+			}
+			else if (Constraint.NumSamples() == 4)
+			{
+				TVector<T, d> Center = SumSampleData(Constraint) / Constraint.NumSamples();
+				T Delta = (Center - SurfaceSample).SizeSquared();
+
+				//
+				// @todo(chaos) : Collision Manifold
+				//    The iterative manifold need to be maximized for area instead of largest 
+				//    distance from center.
+				//
+				T SmallestDelta = FLT_MAX;
+				int32 SmallestIndex = 0;
+				for (int32 idx = 0; idx < Constraint.NumSamples(); idx++)
+					if (Constraint[idx].Delta < SmallestDelta) {
+						SmallestDelta = Constraint[idx].Delta;
+						SmallestIndex = idx;
+					}
+
+				if (Delta > SmallestDelta) {
+					Constraint[SmallestIndex] = { SurfaceSample,Delta };
+				}
+			}
+			else
+			{
+				ensure(false); // max of 4 points
+			}
+
+			typedef FRigidBodyIterativeContactConstraint::FSampleData FSampleData;
+			TVector<T, d> Center = SumSampleData(Constraint) / Constraint.NumSamples();
+			for (int32 Index = 0; Index < Constraint.NumSamples(); Index++) { Constraint[Index].Delta = (Center - Constraint[Index].X).SizeSquared(); }
+		}
+
+
+
 		//
 		// Box - Box
 		//
+		template <typename T, int d>
+		TContactPoint<T> BoxBoxContactPoint(const TAABB<T, d>& Box1, const TRigidTransform<T, d>& ATM, const TAABB<T, d>& Box2, const TRigidTransform<T, d>& BTM, const T Thickness)
+		{
+			return GJKContactPoint(Box1, ATM, Box2, BTM, Thickness);
+		}
 
 		template <typename T, int d>
 		void UpdateBoxBoxConstraint(const TAABB<T, d>& Box1, const TRigidTransform<T, d>& Box1Transform, const TAABB<T, d>& Box2, const TRigidTransform<T, d>& Box2Transform, const T Thickness, TRigidBodyPointContactConstraint<T, d>& Constraint)
@@ -89,6 +212,11 @@ namespace Chaos
 			UpdateContactPoint(Constraint.Manifold, BoxBoxContactPoint(Box1, Box1Transform, Box2, Box2Transform, Thickness));
 		}
 
+		template <typename T, int d>
+		void UpdateBoxBoxManifold(const TAABB<T, d>& Box1, const TRigidTransform<T, d>& Box1Transform, const TAABB<T, d>& Box2, const TRigidTransform<T, d>& Box2Transform, const T Thickness, TRigidBodyPointContactConstraint<T, d>& Constraint)
+		{
+			// @todo(chaos) : Stub Update**Manifold
+		}
 
 		template<typename T, int d>
 		void ConstructBoxBoxConstraints(TGeometryParticleHandle<T, d>* Particle0, TGeometryParticleHandle<T, d>* Particle1, const FImplicitObject* Implicit0, const FImplicitObject* Implicit1, const TRigidTransform<T, d>& Transform0, const TRigidTransform<T, d>& Transform1, const T Thickness, FCollisionConstraintsArray& NewConstraints)
@@ -121,11 +249,7 @@ namespace Chaos
 		template <typename T, int d>
 		TContactPoint<T> BoxHeightFieldContactPoint(const TAABB<T, d>& A, const TRigidTransform<T, d>& ATransform, const THeightField<T>& B, const TRigidTransform<T, d>& BTransform, const T Thickness)
 		{
-			TContactPoint<T> Contact;
-			B.GJKContactPoint(TBox<float, 3>(A), ATransform.GetRelativeTransform(BTransform), Thickness, Contact.Location, Contact.Normal, Contact.Phi);
-			Contact.Location = BTransform.TransformPosition(Contact.Location);
-			Contact.Normal = BTransform.TransformVector(Contact.Normal);
-			return Contact;
+			return GJKImplicitContactPoint< TBox<float,3> >(TBox<float, 3>(A), ATransform, B, BTransform, Thickness);
 		}
 
 		template <typename T, int d>
@@ -313,11 +437,7 @@ namespace Chaos
 		template <typename T, int d>
 		TContactPoint<T> SphereHeightFieldContactPoint(const TSphere<T, d>& A, const TRigidTransform<T, d>& ATransform, const THeightField<T>& B, const TRigidTransform<T, d>& BTransform, const T Thickness)
 		{
-			TContactPoint<T> Contact;
-			B.GJKContactPoint(TSphere<float, 3>(A), ATransform.GetRelativeTransform(BTransform), Thickness, Contact.Location, Contact.Normal, Contact.Phi);
-			Contact.Location = BTransform.TransformPosition(Contact.Location);
-			Contact.Normal = BTransform.TransformVector(Contact.Normal);
-			return Contact;
+			return GJKImplicitContactPoint< TSphere<float,3> >(TSphere<float, 3>(A), ATransform, B, BTransform, Thickness);
 		}
 
 		template <typename T, int d>
@@ -622,11 +742,7 @@ namespace Chaos
 		template <typename T, int d>
 		TContactPoint<T> CapsuleHeightFieldContactPoint(const TCapsule<T>& A, const TRigidTransform<T, d>& ATransform, const THeightField<T>& B, const TRigidTransform<T, d>& BTransform, const T Thickness)
 		{
-			TContactPoint<T> Contact;
-			B.GJKContactPoint(TCapsule<float>(A), ATransform.GetRelativeTransform(BTransform), Thickness, Contact.Location, Contact.Normal, Contact.Phi);
-			Contact.Location = BTransform.TransformPosition(Contact.Location);
-			Contact.Normal = BTransform.TransformVector(Contact.Normal);
-			return Contact;
+			return GJKImplicitContactPoint< TCapsule<float> >(TCapsule<float>(A), ATransform, B, BTransform, Thickness);
 		}
 
 		template <typename T, int d>
@@ -725,96 +841,17 @@ namespace Chaos
 		}
 
 		template<class T, int d>
-		void UpdateConvexConvexManifold(TRigidBodyIterativeContactConstraint<T, d>& Constraint, const TRigidTransform<T, d>& Transform0, const TRigidTransform<T, d>& Transform1, const T Thickness)
+		void UpdateConvexConvexManifold(TCollisionConstraintBase<T, d>&  ConstraintBase, const TRigidTransform<T, d>& Transform0, const TRigidTransform<T, d>& Transform1, const T Thickness)
 		{
-			if (ensure(Constraint.GetType() == FRigidBodyIterativeContactConstraint::StaticType()))
+			if (ensure(ConstraintBase.GetType() == FRigidBodyIterativeContactConstraint::StaticType()))
 			{
-				TContactPoint<T> ContactPoint = ConvexConvexContactPoint(*Constraint.Manifold.Implicit[0], Transform0, *Constraint.Manifold.Implicit[1], Transform1, Thickness);
-
-				if (GetInnerType(Constraint.Manifold.Implicit[0]->GetType()) == ImplicitObjectType::Convex)
+				if (GetInnerType(ConstraintBase.Manifold.Implicit[0]->GetType()) == ImplicitObjectType::Convex)
 				{
-					// single shot manifolds for TConvex implicit object in the constraints implicit[0] position. 
-
-					TArray<FVec3> CollisionSamples;
-					//
-					//  @todo(chaos) : Collision Manifold
-					//   Remove the dependency on the virtual calls on the Implicit. Don't use FindClosestFaceAndVertices
-					//   this relies on virtual calls on the ImplicitObject. Instead pass a parameters structures into 
-					//   ConvexConvexContactPoint that can collect the face indices during evaluation of the support functions. 
-					//   This can be implemented without virtual calls.
-					//
-					int32 FaceIndex = Constraint.Manifold.Implicit[0]->FindClosestFaceAndVertices(Transform0.InverseTransformPosition(ContactPoint.Location), CollisionSamples, 1.f);
-
-					if (!ContactPoint.Normal.Equals(Constraint.PlaneNormal)
-						|| FaceIndex != Constraint.SourceNormalIndex
-						|| !Constraint.NumSamples())
-					{
-						Constraint.ResetSamples();
-						Constraint.PlaneNormal = Transform1.InverseTransformVector(ContactPoint.Normal);
-						Constraint.PlanePosition = Transform1.InverseTransformPosition(ContactPoint.Location - ContactPoint.Phi*ContactPoint.Normal);
-					}
-
-
-					if (Constraint.NumSamples() == 0)
-					{
-						//
-						// @todo(chaos) : Collision Manifold
-						//   Only save the four best samples and hard-code the size of Constraint.Samples to [len:4].
-						//   Currently this just grabs all points and uses the deepest point for resolution. 
-						//
-						Constraint.SourceNormalIndex = FaceIndex;
-						for (FVec3 Sample : CollisionSamples)
-						{
-							Constraint.AddSample({ Sample,0.f });
-						}
-					}
+					UpdateSingleShotManifold(*ConstraintBase.template As<TRigidBodyIterativeContactConstraint<T, d>>(), Transform0, Transform1, Thickness);
 				}
 				else
 				{
-					// iterative manifolds for non TConvex implicit objects that require sampling 
-
-					if (!ContactPoint.Normal.Equals(Constraint.PlaneNormal) || !Constraint.NumSamples())
-					{
-						Constraint.ResetSamples();
-						Constraint.PlaneNormal = Transform1.InverseTransformVector(ContactPoint.Normal);
-						Constraint.PlanePosition = Transform1.InverseTransformPosition(ContactPoint.Location - ContactPoint.Phi*ContactPoint.Normal);
-					}
-
-					TVector<T, d> SurfaceSample = Transform0.InverseTransformPosition(ContactPoint.Location);
-					if (Constraint.NumSamples() < 4)
-					{
-						Constraint.AddSample({ SurfaceSample,0.f });
-					}
-					else if (Constraint.NumSamples() == 4)
-					{
-						TVector<T, d> Center = SumSampleData(Constraint) / Constraint.NumSamples();
-						T Delta = (Center - SurfaceSample).SizeSquared();
-
-						//
-						// @todo(chaos) : Collision Manifold
-						//    The iterative manifold need to be maximized for area instead of largest 
-						//    distance from center.
-						//
-						T SmallestDelta = FLT_MAX;
-						int32 SmallestIndex = 0;
-						for (int32 idx = 0; idx < Constraint.NumSamples(); idx++)
-							if (Constraint[idx].Delta < SmallestDelta) {
-								SmallestDelta = Constraint[idx].Delta;
-								SmallestIndex = idx;
-							}
-
-						if (Delta > SmallestDelta) {
-							Constraint[SmallestIndex] = { SurfaceSample,Delta };
-						}
-					}
-					else
-					{
-						ensure(false); // max of 4 points
-					}
-
-					typedef FRigidBodyIterativeContactConstraint::FSampleData FSampleData;
-					TVector<T, d> Center = SumSampleData(Constraint) / Constraint.NumSamples();
-					for (int32 Index = 0; Index < Constraint.NumSamples(); Index++) { Constraint[Index].Delta = (Center - Constraint[Index].X).SizeSquared(); }
+					UpdateIterativeManifold(*ConstraintBase.template As<TRigidBodyIterativeContactConstraint<T, d>>(), Transform0, Transform1, Thickness);
 				}
 			}
 		}
@@ -848,16 +885,7 @@ namespace Chaos
 		template <typename T, int d>
 		TContactPoint<T> ConvexHeightFieldContactPoint(const FImplicitObject& A, const TRigidTransform<T, d>& ATransform, const THeightField<T>& B, const TRigidTransform<T, d>& BTransform, const T Thickness)
 		{
-			TContactPoint<T> Contact;
-
-			if(const TImplicitObjectScaled<FConvex>* ScaledConvexImplicit = A.template GetObject<const TImplicitObjectScaled<FConvex> >())
-				B.GJKContactPoint(*ScaledConvexImplicit, ATransform.GetRelativeTransform(BTransform), Thickness, Contact.Location, Contact.Normal, Contact.Phi);
-			else if (const FConvex* ConvexImplicit = A.template GetObject<const FConvex>())
-				B.GJKContactPoint(*ConvexImplicit, ATransform.GetRelativeTransform(BTransform), Thickness, Contact.Location, Contact.Normal, Contact.Phi);
-
-			Contact.Location = BTransform.TransformPosition(Contact.Location);
-			Contact.Normal = BTransform.TransformVector(Contact.Normal);
-			return Contact;
+			return GJKImplicitContactPoint< FConvex >(A, ATransform, B, BTransform, Thickness);
 		}
 
 		template <typename T, int d>
@@ -1014,7 +1042,7 @@ namespace Chaos
 			// TConvex
 			if (Implicit0.IsConvex() && Implicit1.IsConvex())
 			{
-				UpdateConvexConvexManifold(*ConstraintBase.template As<TRigidBodyIterativeContactConstraint<T, d>>(), Transform0, Transform1, Thickness);
+				UpdateConvexConvexManifold(ConstraintBase, Transform0, Transform1, Thickness);
 			}
 		}
 
@@ -1349,11 +1377,11 @@ namespace Chaos
 			}
 			else if (Implicit0->IsConvex() && Implicit1Type == THeightField<T>::StaticType())
 			{
-			ConstructConvexHeightFieldConstraints(Particle0, Particle1, Implicit0, Implicit1, Transform0, Transform1, Thickness, NewConstraints);
+				ConstructConvexHeightFieldConstraints(Particle0, Particle1, Implicit0, Implicit1, Transform0, Transform1, Thickness, NewConstraints);
 			}
 			else if (Implicit0Type == THeightField<T>::StaticType() && Implicit1->IsConvex())
 			{
-			ConstructConvexHeightFieldConstraints(Particle1, Particle0, Implicit1, Implicit0, Transform1, Transform0, Thickness, NewConstraints);
+				ConstructConvexHeightFieldConstraints(Particle1, Particle0, Implicit1, Implicit0, Transform1, Transform0, Thickness, NewConstraints);
 			}
 			else if (Implicit0->IsConvex() && Implicit1->IsConvex())
 			{
@@ -1365,7 +1393,6 @@ namespace Chaos
 			}
 		}
 
-
 		template void UpdateBoxBoxConstraint<float, 3>(const TAABB<float, 3>& Box1, const TRigidTransform<float, 3>& Box1Transform, const TAABB<float, 3>& Box2, const TRigidTransform<float, 3>& Box2Transform, const float Thickness, TRigidBodyPointContactConstraint<float, 3>& Constraint);
 		template void ConstructBoxBoxConstraints<float, 3>(TGeometryParticleHandle<float, 3>* Particle0, TGeometryParticleHandle<float, 3>* Particle1, const FImplicitObject* Implicit0, const FImplicitObject* Implicit1, const TRigidTransform<float, 3>& Transform0, const TRigidTransform<float, 3>& Transform1, const float Thickness, FCollisionConstraintsArray& NewConstraints);
 
@@ -1403,7 +1430,7 @@ namespace Chaos
 		template void ConstructConvexHeightFieldConstraints(TGeometryParticleHandle<float, 3>* Particle0, TGeometryParticleHandle<float, 3>* Particle1, const FImplicitObject* Implicit0, const FImplicitObject* Implicit1, const TRigidTransform<float, 3>& Transform0, const TRigidTransform<float, 3>& Transform1, const float Thickness, FCollisionConstraintsArray& NewConstraints);
 
 		template void UpdateConvexConvexConstraint<float, 3>(const FImplicitObject& A, const TRigidTransform<float, 3>& ATM, const FImplicitObject& B, const TRigidTransform<float, 3>& BTM, const float Thickness, TCollisionConstraintBase<float, 3>& Constraint);
-		template void UpdateConvexConvexManifold<float, 3>(TRigidBodyIterativeContactConstraint<float, 3>& Constraint, const TRigidTransform<float, 3>& ATM, const TRigidTransform<float, 3>& BTM, const float Thickness);
+		template void UpdateConvexConvexManifold<float, 3>(TCollisionConstraintBase<float, 3>& Constraint, const TRigidTransform<float, 3>& ATM, const TRigidTransform<float, 3>& BTM, const float Thickness);
 		template void ConstructConvexConvexConstraints<float, 3>(TGeometryParticleHandle<float, 3>* Particle0, TGeometryParticleHandle<float, 3>* Particle1, const FImplicitObject* Implicit0, const FImplicitObject* Implicit1, const TRigidTransform<float, 3>& Transform0, const TRigidTransform<float, 3>& Transform1, const float Thickness, FCollisionConstraintsArray& NewConstraints);
 
 		template void UpdateLevelsetLevelsetConstraint<ECollisionUpdateType::Any, float, 3>(const float Thickness, TRigidBodyPointContactConstraint<float, 3>& Constraint);
@@ -1412,6 +1439,8 @@ namespace Chaos
 
 		template void ConstructUnionUnionConstraints<float, 3>(TGeometryParticleHandle<float, 3>* Particle0, TGeometryParticleHandle<float, 3>* Particle1, const FImplicitObject* Implicit0, const FImplicitObject* Implicit1, const TRigidTransform<float, 3>& Transform0, const TRigidTransform<float, 3>& Transform1, const float Thickness, FCollisionConstraintsArray& NewConstraints);
 
+		template void UpdateSingleShotManifold<float, 3>(TRigidBodyIterativeContactConstraint<float, 3>&  Constraint, const TRigidTransform<float, 3>& Transform0, const TRigidTransform<float, 3>& Transform1, const float Thickness);
+		template void UpdateIterativeManifold<float, 3>(TRigidBodyIterativeContactConstraint<float, 3>&  Constraint, const TRigidTransform<float, 3>& Transform0, const TRigidTransform<float, 3>& Transform1, const float Thickness);
 		template void UpdateManifold<float, 3>(TCollisionConstraintBase<float, 3>& ConstraintBase, const TRigidTransform<float, 3>& ATM, const TRigidTransform<float, 3>& BTM, const float Thickness);
 		template void UpdateConstraint<ECollisionUpdateType::Any, float, 3>(TCollisionConstraintBase<float,3>& ConstraintBase, const TRigidTransform<float,3>& Transform0, const TRigidTransform<float,3>& Transform1, const float Thickness);
 		template void UpdateConstraint<ECollisionUpdateType::Deepest, float, 3>(TCollisionConstraintBase<float, 3>& ConstraintBase, const TRigidTransform<float, 3>& Transform0, const TRigidTransform<float, 3>& Transform1, const float Thickness);
diff --git a/Engine/Source/Runtime/Experimental/Chaos/Private/Chaos/HeightField.cpp b/Engine/Source/Runtime/Experimental/Chaos/Private/Chaos/HeightField.cpp
--- a/Engine/Source/Runtime/Experimental/Chaos/Private/Chaos/HeightField.cpp
+++ b/Engine/Source/Runtime/Experimental/Chaos/Private/Chaos/HeightField.cpp
@@ -1253,6 +1253,24 @@ namespace Chaos
 		return GJKContactPointImp(QueryGeom, QueryTM, Thickness, ContactLocation, ContactNormal, ContactPhi);
 	}
 
+	template <typename T>
+	bool THeightField<T>::GJKContactPoint(const TImplicitObjectScaled < TBox<T, 3>>& QueryGeom, const TRigidTransform<T, 3>& QueryTM, const T Thickness, TVector<T, 3>& ContactLocation, TVector<T, 3>& ContactNormal, T& ContactPhi) const
+	{
+		return GJKContactPointImp(QueryGeom, QueryTM, Thickness, ContactLocation, ContactNormal, ContactPhi);
+	}
+
+	template <typename T>
+	bool THeightField<T>::GJKContactPoint(const TImplicitObjectScaled < TSphere<T, 3>>& QueryGeom, const TRigidTransform<T, 3>& QueryTM, const T Thickness, TVector<T, 3>& ContactLocation, TVector<T, 3>& ContactNormal, T& ContactPhi) const
+	{
+		return GJKContactPointImp(QueryGeom, QueryTM, Thickness, ContactLocation, ContactNormal, ContactPhi);
+	}
+
+	template <typename T>
+	bool THeightField<T>::GJKContactPoint(const TImplicitObjectScaled < TCapsule<T>>& QueryGeom, const TRigidTransform<T, 3>& QueryTM, const T Thickness, TVector<T, 3>& ContactLocation, TVector<T, 3>& ContactNormal, T& ContactPhi) const
+	{
+		return GJKContactPointImp(QueryGeom, QueryTM, Thickness, ContactLocation, ContactNormal, ContactPhi);
+	}
+
 	template <typename T>
 	bool THeightField<T>::GJKContactPoint(const TImplicitObjectScaled<FConvex>& QueryGeom, const TRigidTransform<T, 3>& QueryTM, const T Thickness, TVector<T, 3>& ContactLocation, TVector<T, 3>& ContactNormal, T& ContactPhi) const
 	{
diff --git a/Engine/Source/Runtime/Experimental/Chaos/Public/Chaos/CollisionResolution.h b/Engine/Source/Runtime/Experimental/Chaos/Public/Chaos/CollisionResolution.h
--- a/Engine/Source/Runtime/Experimental/Chaos/Public/Chaos/CollisionResolution.h
+++ b/Engine/Source/Runtime/Experimental/Chaos/Public/Chaos/CollisionResolution.h
@@ -34,169 +34,178 @@ namespace Chaos
 	namespace Collisions
 	{
 
-			//
-			// Box-Box
-			//
+		//
+		// Box-Box
+		//
 
-			template <typename T, int d>
-			void CHAOS_API UpdateBoxBoxConstraint(const TAABB<T, d>& Box1, const TRigidTransform<T, d>& Box1Transform, const TAABB<T, d>& Box2, const TRigidTransform<T, d>& Box2Transform, const T Thickness, TRigidBodyPointContactConstraint<T, d>& Constraint);
+		template <typename T, int d>
+		void CHAOS_API UpdateBoxBoxConstraint(const TAABB<T, d>& Box1, const TRigidTransform<T, d>& Box1Transform, const TAABB<T, d>& Box2, const TRigidTransform<T, d>& Box2Transform, const T Thickness, TRigidBodyPointContactConstraint<T, d>& Constraint);
 
-			template<typename T, int d>
-			void CHAOS_API ConstructBoxBoxConstraints(TGeometryParticleHandle<T, d>* Particle0, TGeometryParticleHandle<T, d>* Particle1, const FImplicitObject* Implicit0, const FImplicitObject* Implicit1, const TRigidTransform<T, d>& Transform0, const TRigidTransform<T, d>& Transform1, const T Thickness, FCollisionConstraintsArray& NewConstraints);
+		template <typename T, int d>
+		void CHAOS_API UpdateBoxBoxManifold(const TAABB<T, d>& Box1, const TRigidTransform<T, d>& Box1Transform, const TAABB<T, d>& Box2, const TRigidTransform<T, d>& Box2Transform, const T Thickness, TRigidBodyPointContactConstraint<T, d>& Constraint);
 
+		template<typename T, int d>
+		void CHAOS_API ConstructBoxBoxConstraints(TGeometryParticleHandle<T, d>* Particle0, TGeometryParticleHandle<T, d>* Particle1, const FImplicitObject* Implicit0, const FImplicitObject* Implicit1, const TRigidTransform<T, d>& Transform0, const TRigidTransform<T, d>& Transform1, const T Thickness, FCollisionConstraintsArray& NewConstraints);
 
-			//
-			// Box-HeightField
-			//
 
-			template <typename T, int d>
-			void UpdateBoxHeightFieldConstraint(const TAABB<T, d>& A, const TRigidTransform<T, d>& ATransform, const THeightField<T>& B, const TRigidTransform<T, d>& BTransform, const T Thickness, TRigidBodyPointContactConstraint<T, d>& Constraint);
+		//
+		// Box-HeightField
+		//
 
-			template<typename T, int d>
-			void ConstructBoxHeightFieldConstraints(TGeometryParticleHandle<T, d>* Particle0, TGeometryParticleHandle<T, d>* Particle1, const FImplicitObject* Implicit0, const FImplicitObject* Implicit1, const TRigidTransform<T, d>& Transform0, const TRigidTransform<T, d>& Transform1, const T Thickness, FCollisionConstraintsArray& NewConstraints);
+		template <typename T, int d>
+		void UpdateBoxHeightFieldConstraint(const TAABB<T, d>& A, const TRigidTransform<T, d>& ATransform, const THeightField<T>& B, const TRigidTransform<T, d>& BTransform, const T Thickness, TRigidBodyPointContactConstraint<T, d>& Constraint);
 
-			//
-			// Box-Plane
-			//
+		template<typename T, int d>
+		void ConstructBoxHeightFieldConstraints(TGeometryParticleHandle<T, d>* Particle0, TGeometryParticleHandle<T, d>* Particle1, const FImplicitObject* Implicit0, const FImplicitObject* Implicit1, const TRigidTransform<T, d>& Transform0, const TRigidTransform<T, d>& Transform1, const T Thickness, FCollisionConstraintsArray& NewConstraints);
 
-			template <typename T, int d>
-			bool CHAOS_API UpdateBoxPlaneConstraint(const TAABB<T, d>& Box, const TRigidTransform<T, d>& BoxTransform, const TPlane<T, d>& Plane, const TRigidTransform<T, d>& PlaneTransform, const T Thickness, TRigidBodyPointContactConstraint<T, d>& Constraint);
+		//
+		// Box-Plane
+		//
 
-			template<typename T, int d>
-			void CHAOS_API ConstructBoxPlaneConstraints(TGeometryParticleHandle<T, d>* Particle0, TGeometryParticleHandle<T, d>* Particle1, const FImplicitObject* Implicit0, const FImplicitObject* Implicit1, const TRigidTransform<T, d>& Transform0, const TRigidTransform<T, d>& Transform1, const T Thickness, FCollisionConstraintsArray& NewConstraints);
+		template <typename T, int d>
+		bool CHAOS_API UpdateBoxPlaneConstraint(const TAABB<T, d>& Box, const TRigidTransform<T, d>& BoxTransform, const TPlane<T, d>& Plane, const TRigidTransform<T, d>& PlaneTransform, const T Thickness, TRigidBodyPointContactConstraint<T, d>& Constraint);
 
+		template<typename T, int d>
+		void CHAOS_API ConstructBoxPlaneConstraints(TGeometryParticleHandle<T, d>* Particle0, TGeometryParticleHandle<T, d>* Particle1, const FImplicitObject* Implicit0, const FImplicitObject* Implicit1, const TRigidTransform<T, d>& Transform0, const TRigidTransform<T, d>& Transform1, const T Thickness, FCollisionConstraintsArray& NewConstraints);
 
 
-			//
-			// Sphere-Sphere
-			//
 
-			template <typename T, int d>
-			void CHAOS_API UpdateSphereSphereConstraint(const TSphere<T, d>& Sphere1, const TRigidTransform<T, d>& Sphere1Transform, const TSphere<T, d>& Sphere2, const TRigidTransform<T, d>& Sphere2Transform, const T Thickness, TRigidBodyPointContactConstraint<T, d>& Constraint);
+		//
+		// Sphere-Sphere
+		//
 
-			template<typename T, int d>
-			void CHAOS_API ConstructSphereSphereConstraints(TGeometryParticleHandle<T, d>* Particle0, TGeometryParticleHandle<T, d>* Particle1, const FImplicitObject* Implicit0, const FImplicitObject* Implicit1, const TRigidTransform<T, d>& Transform0, const TRigidTransform<T, d>& Transform1, const T Thickness, FCollisionConstraintsArray& NewConstraints);
+		template <typename T, int d>
+		void CHAOS_API UpdateSphereSphereConstraint(const TSphere<T, d>& Sphere1, const TRigidTransform<T, d>& Sphere1Transform, const TSphere<T, d>& Sphere2, const TRigidTransform<T, d>& Sphere2Transform, const T Thickness, TRigidBodyPointContactConstraint<T, d>& Constraint);
 
-			//
-			// Sphere-HeightField
-			//
+		template<typename T, int d>
+		void CHAOS_API ConstructSphereSphereConstraints(TGeometryParticleHandle<T, d>* Particle0, TGeometryParticleHandle<T, d>* Particle1, const FImplicitObject* Implicit0, const FImplicitObject* Implicit1, const TRigidTransform<T, d>& Transform0, const TRigidTransform<T, d>& Transform1, const T Thickness, FCollisionConstraintsArray& NewConstraints);
 
-			template <typename T, int d>
-			void UpdateSphereHeightFieldConstraint(const TSphere<T, d>& A, const TRigidTransform<T, d>& ATransform, const THeightField<T>& B, const TRigidTransform<T, d>& BTransform, const T Thickness, TRigidBodyPointContactConstraint<T, d>& Constraint);
+		//
+		// Sphere-HeightField
+		//
 
-			template<typename T, int d>
-			void ConstructSphereHeightFieldConstraints(TGeometryParticleHandle<T, d>* Particle0, TGeometryParticleHandle<T, d>* Particle1, const FImplicitObject* Implicit0, const FImplicitObject* Implicit1, const TRigidTransform<T, d>& Transform0, const TRigidTransform<T, d>& Transform1, const T Thickness, FCollisionConstraintsArray& NewConstraints);
-				
-			//
-			// Sphere-Plane
-			//
+		template <typename T, int d>
+		void UpdateSphereHeightFieldConstraint(const TSphere<T, d>& A, const TRigidTransform<T, d>& ATransform, const THeightField<T>& B, const TRigidTransform<T, d>& BTransform, const T Thickness, TRigidBodyPointContactConstraint<T, d>& Constraint);
 
-			template <typename T, int d>
-			void CHAOS_API UpdateSpherePlaneConstraint(const TSphere<T, d>& Sphere, const TRigidTransform<T, d>& SphereTransform, const TPlane<T, d>& Plane, const TRigidTransform<T, d>& PlaneTransform, const T Thickness, TRigidBodyPointContactConstraint<T, d>& Constraint);
+		template<typename T, int d>
+		void ConstructSphereHeightFieldConstraints(TGeometryParticleHandle<T, d>* Particle0, TGeometryParticleHandle<T, d>* Particle1, const FImplicitObject* Implicit0, const FImplicitObject* Implicit1, const TRigidTransform<T, d>& Transform0, const TRigidTransform<T, d>& Transform1, const T Thickness, FCollisionConstraintsArray& NewConstraints);
 
-			template<typename T, int d>
-			void CHAOS_API ConstructSpherePlaneConstraints(TGeometryParticleHandle<T, d>* Particle0, TGeometryParticleHandle<T, d>* Particle1, const FImplicitObject* Implicit0, const FImplicitObject* Implicit1, const TRigidTransform<T, d>& Transform0, const TRigidTransform<T, d>& Transform1, const T Thickness, FCollisionConstraintsArray& NewConstraints);
+		//
+		// Sphere-Plane
+		//
 
-			//
-			// Sphere-Box
-			//
+		template <typename T, int d>
+		void CHAOS_API UpdateSpherePlaneConstraint(const TSphere<T, d>& Sphere, const TRigidTransform<T, d>& SphereTransform, const TPlane<T, d>& Plane, const TRigidTransform<T, d>& PlaneTransform, const T Thickness, TRigidBodyPointContactConstraint<T, d>& Constraint);
 
-			template <typename T, int d>
-			void CHAOS_API UpdateSphereBoxConstraint(const TSphere<T, d>& Sphere, const TRigidTransform<T, d>& SphereTransform, const TAABB<T, d>& Box, const TRigidTransform<T, d>& BoxTransform, const T Thickness, TRigidBodyPointContactConstraint<T, d>& Constraint);
+		template<typename T, int d>
+		void CHAOS_API ConstructSpherePlaneConstraints(TGeometryParticleHandle<T, d>* Particle0, TGeometryParticleHandle<T, d>* Particle1, const FImplicitObject* Implicit0, const FImplicitObject* Implicit1, const TRigidTransform<T, d>& Transform0, const TRigidTransform<T, d>& Transform1, const T Thickness, FCollisionConstraintsArray& NewConstraints);
 
-			template<typename T, int d>
-			void CHAOS_API ConstructSphereBoxConstraints(TGeometryParticleHandle<T, d>* Particle0, TGeometryParticleHandle<T, d>* Particle1, const FImplicitObject* Implicit0, const FImplicitObject* Implicit1, const TRigidTransform<T, d>& Transform0, const TRigidTransform<T, d>& Transform1, const T Thickness, FCollisionConstraintsArray& NewConstraints);
+		//
+		// Sphere-Box
+		//
 
-			//
-			// Sphere-Capsule
-			//
+		template <typename T, int d>
+		void CHAOS_API UpdateSphereBoxConstraint(const TSphere<T, d>& Sphere, const TRigidTransform<T, d>& SphereTransform, const TAABB<T, d>& Box, const TRigidTransform<T, d>& BoxTransform, const T Thickness, TRigidBodyPointContactConstraint<T, d>& Constraint);
 
-			template <typename T, int d>
-			void CHAOS_API UpdateSphereCapsuleConstraint(const TSphere<T, d>& Sphere, const TRigidTransform<T, d>& SphereTransform, const TCapsule<T>& Box, const TRigidTransform<T, d>& BoxTransform, const T Thickness, TRigidBodyPointContactConstraint<T, d>& Constraint);
+		template<typename T, int d>
+		void CHAOS_API ConstructSphereBoxConstraints(TGeometryParticleHandle<T, d>* Particle0, TGeometryParticleHandle<T, d>* Particle1, const FImplicitObject* Implicit0, const FImplicitObject* Implicit1, const TRigidTransform<T, d>& Transform0, const TRigidTransform<T, d>& Transform1, const T Thickness, FCollisionConstraintsArray& NewConstraints);
 
-			template<typename T, int d>
-			void CHAOS_API ConstructSphereCapsuleConstraints(TGeometryParticleHandle<T, d>* Particle0, TGeometryParticleHandle<T, d>* Particle1, const FImplicitObject* Implicit0, const FImplicitObject* Implicit1, const TRigidTransform<T, d>& Transform0, const TRigidTransform<T, d>& Transform1, const T Thickness, FCollisionConstraintsArray& NewConstraints);
+		//
+		// Sphere-Capsule
+		//
 
-			//
-			// Capsule-Capsule
-			//
+		template <typename T, int d>
+		void CHAOS_API UpdateSphereCapsuleConstraint(const TSphere<T, d>& Sphere, const TRigidTransform<T, d>& SphereTransform, const TCapsule<T>& Box, const TRigidTransform<T, d>& BoxTransform, const T Thickness, TRigidBodyPointContactConstraint<T, d>& Constraint);
 
-			template <typename T, int d>
-			void CHAOS_API UpdateCapsuleCapsuleConstraint(const TCapsule<T>& A, const TRigidTransform<T, d>& ATransform, const TCapsule<T>& B, const TRigidTransform<T, d>& BTransform, const T Thickness, TRigidBodyPointContactConstraint<T, d>& Constraint);
+		template<typename T, int d>
+		void CHAOS_API ConstructSphereCapsuleConstraints(TGeometryParticleHandle<T, d>* Particle0, TGeometryParticleHandle<T, d>* Particle1, const FImplicitObject* Implicit0, const FImplicitObject* Implicit1, const TRigidTransform<T, d>& Transform0, const TRigidTransform<T, d>& Transform1, const T Thickness, FCollisionConstraintsArray& NewConstraints);
 
-			template<typename T, int d>
-			void CHAOS_API ConstructCapsuleCapsuleConstraints(TGeometryParticleHandle<T, d>* Particle0, TGeometryParticleHandle<T, d>* Particle1, const FImplicitObject* Implicit0, const FImplicitObject* Implicit1, const TRigidTransform<T, d>& Transform0, const TRigidTransform<T, d>& Transform1, const T Thickness, FCollisionConstraintsArray& NewConstraints);
+		//
+		// Capsule-Capsule
+		//
 
-			//
-			// Capsule-Box
-			//
+		template <typename T, int d>
+		void CHAOS_API UpdateCapsuleCapsuleConstraint(const TCapsule<T>& A, const TRigidTransform<T, d>& ATransform, const TCapsule<T>& B, const TRigidTransform<T, d>& BTransform, const T Thickness, TRigidBodyPointContactConstraint<T, d>& Constraint);
 
-			template <typename T, int d>
-			void CHAOS_API UpdateCapsuleBoxConstraint(const TCapsule<T>& A, const TRigidTransform<T, d>& ATransform, const TAABB<T, d>& B, const TRigidTransform<T, d>& BTransform, const T Thickness, TRigidBodyPointContactConstraint<T, d>& Constraint);
+		template<typename T, int d>
+		void CHAOS_API ConstructCapsuleCapsuleConstraints(TGeometryParticleHandle<T, d>* Particle0, TGeometryParticleHandle<T, d>* Particle1, const FImplicitObject* Implicit0, const FImplicitObject* Implicit1, const TRigidTransform<T, d>& Transform0, const TRigidTransform<T, d>& Transform1, const T Thickness, FCollisionConstraintsArray& NewConstraints);
 
-			template<typename T, int d>
-			void CHAOS_API ConstructCapsuleBoxConstraints(TGeometryParticleHandle<T, d>* Particle0, TGeometryParticleHandle<T, d>* Particle1, const FImplicitObject* Implicit0, const FImplicitObject* Implicit1, const TRigidTransform<T, d>& Transform0, const TRigidTransform<T, d>& Transform1, const T Thickness, FCollisionConstraintsArray& NewConstraints);
+		//
+		// Capsule-Box
+		//
 
-			//
-			// Capsule-HeightField
-			//
+		template <typename T, int d>
+		void CHAOS_API UpdateCapsuleBoxConstraint(const TCapsule<T>& A, const TRigidTransform<T, d>& ATransform, const TAABB<T, d>& B, const TRigidTransform<T, d>& BTransform, const T Thickness, TRigidBodyPointContactConstraint<T, d>& Constraint);
 
-			template <typename T, int d>
-			void UpdateCapsuleHeightFieldConstraint(const TCapsule<T>& A, const TRigidTransform<T, d>& ATransform, const THeightField<T>& B, const TRigidTransform<T, d>& BTransform, const T Thickness, TRigidBodyPointContactConstraint<T, d>& Constraint);
+		template<typename T, int d>
+		void CHAOS_API ConstructCapsuleBoxConstraints(TGeometryParticleHandle<T, d>* Particle0, TGeometryParticleHandle<T, d>* Particle1, const FImplicitObject* Implicit0, const FImplicitObject* Implicit1, const TRigidTransform<T, d>& Transform0, const TRigidTransform<T, d>& Transform1, const T Thickness, FCollisionConstraintsArray& NewConstraints);
 
-			template<typename T, int d>
-			void ConstructCapsuleHeightFieldConstraints(TGeometryParticleHandle<T, d>* Particle0, TGeometryParticleHandle<T, d>* Particle1, const FImplicitObject* Implicit0, const FImplicitObject* Implicit1, const TRigidTransform<T, d>& Transform0, const TRigidTransform<T, d>& Transform1, const T Thickness, FCollisionConstraintsArray& NewConstraints);
+		//
+		// Capsule-HeightField
+		//
 
-			//
-			// Convex-Convex
-			//
-			template <typename T, int d>
-			void CHAOS_API UpdateConvexConvexConstraint(const FImplicitObject& A, const TRigidTransform<T, d>& ATM, const FImplicitObject& B, const TRigidTransform<T, d>& BTM, const T Thickness, TCollisionConstraintBase<T, d>& Constraint);
+		template <typename T, int d>
+		void UpdateCapsuleHeightFieldConstraint(const TCapsule<T>& A, const TRigidTransform<T, d>& ATransform, const THeightField<T>& B, const TRigidTransform<T, d>& BTransform, const T Thickness, TRigidBodyPointContactConstraint<T, d>& Constraint);
 
-			template<class T, int d>
-			void UpdateConvexConvexManifold(TRigidBodyIterativeContactConstraint<T, d>& Constraint, const TRigidTransform<T, d>& ATM, const TRigidTransform<T, d>& BTM, const T Thickness);
+		template<typename T, int d>
+		void ConstructCapsuleHeightFieldConstraints(TGeometryParticleHandle<T, d>* Particle0, TGeometryParticleHandle<T, d>* Particle1, const FImplicitObject* Implicit0, const FImplicitObject* Implicit1, const TRigidTransform<T, d>& Transform0, const TRigidTransform<T, d>& Transform1, const T Thickness, FCollisionConstraintsArray& NewConstraints);
 
-			template <typename T, int d>
-			void CHAOS_API ConstructConvexConvexConstraints(TGeometryParticleHandle<T, d>* Particle0, TGeometryParticleHandle<T, d>* Particle1, const FImplicitObject* Implicit0, const FImplicitObject* Implicit1, const TRigidTransform<T, d>& Transform0, const TRigidTransform<T, d>& Transform1, const T Thickness, FCollisionConstraintsArray& NewConstraints);
+		//
+		// Convex-Convex
+		//
+		template <typename T, int d>
+		void CHAOS_API UpdateConvexConvexConstraint(const FImplicitObject& A, const TRigidTransform<T, d>& ATM, const FImplicitObject& B, const TRigidTransform<T, d>& BTM, const T Thickness, TCollisionConstraintBase<T, d>& Constraint);
 
-			//
-			// Convex-HeightField
-			//
+		template<class T, int d>
+		void UpdateConvexConvexManifold(TCollisionConstraintBase<T, d>&  Constraint, const TRigidTransform<T, d>& ATM, const TRigidTransform<T, d>& BTM, const T Thickness);
 
-			template <typename T, int d>
-			void UpdateConvexHeightFieldConstraint(const FImplicitObject& A, const TRigidTransform<T, d>& ATransform, const THeightField<T>& B, const TRigidTransform<T, d>& BTransform, const T Thickness, TRigidBodyPointContactConstraint<T, d>& Constraint);
+		template <typename T, int d>
+		void CHAOS_API ConstructConvexConvexConstraints(TGeometryParticleHandle<T, d>* Particle0, TGeometryParticleHandle<T, d>* Particle1, const FImplicitObject* Implicit0, const FImplicitObject* Implicit1, const TRigidTransform<T, d>& Transform0, const TRigidTransform<T, d>& Transform1, const T Thickness, FCollisionConstraintsArray& NewConstraints);
 
-			template<typename T, int d>
-			void ConstructConvexHeightFieldConstraints(TGeometryParticleHandle<T, d>* Particle0, TGeometryParticleHandle<T, d>* Particle1, const FImplicitObject* Implicit0, const FImplicitObject* Implicit1, const TRigidTransform<T, d>& Transform0, const TRigidTransform<T, d>& Transform1, const T Thickness, FCollisionConstraintsArray& NewConstraints);
+		//
+		// Convex-HeightField
+		//
 
+		template <typename T, int d>
+		void UpdateConvexHeightFieldConstraint(const FImplicitObject& A, const TRigidTransform<T, d>& ATransform, const THeightField<T>& B, const TRigidTransform<T, d>& BTransform, const T Thickness, TRigidBodyPointContactConstraint<T, d>& Constraint);
 
-			//
-			// Levelset-Levelset
-			//
-			template<ECollisionUpdateType UpdateType, typename T = float, int d = 3>
-			void CHAOS_API UpdateLevelsetLevelsetConstraint(const T Thickness, TRigidBodyPointContactConstraint<float, 3>& Constraint);
+		template<typename T, int d>
+		void ConstructConvexHeightFieldConstraints(TGeometryParticleHandle<T, d>* Particle0, TGeometryParticleHandle<T, d>* Particle1, const FImplicitObject* Implicit0, const FImplicitObject* Implicit1, const TRigidTransform<T, d>& Transform0, const TRigidTransform<T, d>& Transform1, const T Thickness, FCollisionConstraintsArray& NewConstraints);
 
-			template<typename T, int d>
-			void CHAOS_API ConstructLevelsetLevelsetConstraints(TGeometryParticleHandle<T, d>* Particle0, TGeometryParticleHandle<T, d>* Particle1, const FImplicitObject* Implicit0, const FImplicitObject* Implicit1, const TRigidTransform<T, d>& Transform0, const TRigidTransform<T, d>& Transform1, const T Thickness, FCollisionConstraintsArray& NewConstraints);
 
+		//
+		// Levelset-Levelset
+		//
+		template<ECollisionUpdateType UpdateType, typename T = float, int d = 3>
+		void CHAOS_API UpdateLevelsetLevelsetConstraint(const T Thickness, TRigidBodyPointContactConstraint<float, 3>& Constraint);
 
-			//
-			// Union-Union
-			//
-			template<typename T, int d>
-			void CHAOS_API ConstructUnionUnionConstraints(TGeometryParticleHandle<T, d>* Particle0, TGeometryParticleHandle<T, d>* Particle1, const FImplicitObject* Implicit0, const FImplicitObject* Implicit1, const TRigidTransform<T, d>& Transform0, const TRigidTransform<T, d>& Transform1, const T Thickness, FCollisionConstraintsArray& NewConstraints);
+		template<typename T, int d>
+		void CHAOS_API ConstructLevelsetLevelsetConstraints(TGeometryParticleHandle<T, d>* Particle0, TGeometryParticleHandle<T, d>* Particle1, const FImplicitObject* Implicit0, const FImplicitObject* Implicit1, const TRigidTransform<T, d>& Transform0, const TRigidTransform<T, d>& Transform1, const T Thickness, FCollisionConstraintsArray& NewConstraints);
 
-			//
-			// Constraint API
-			//
-			template<typename T, int d>
-			void CHAOS_API UpdateManifold(TCollisionConstraintBase<T, d>& ConstraintBase, const TRigidTransform<T, d>& ATM, const TRigidTransform<T, d>& BTM, const T Thickness);
 
-			template<ECollisionUpdateType UpdateType, typename T, int d>
-			void CHAOS_API UpdateConstraint(TCollisionConstraintBase<T, d>& ConstraintBase, const TRigidTransform<T, d>& Transform0, const TRigidTransform<T, d>& Transform1, const T Thickness);
+		//
+		// Union-Union
+		//
+		template<typename T, int d>
+		void CHAOS_API ConstructUnionUnionConstraints(TGeometryParticleHandle<T, d>* Particle0, TGeometryParticleHandle<T, d>* Particle1, const FImplicitObject* Implicit0, const FImplicitObject* Implicit1, const TRigidTransform<T, d>& Transform0, const TRigidTransform<T, d>& Transform1, const T Thickness, FCollisionConstraintsArray& NewConstraints);
 
-			template<typename T, int d>
-			void CHAOS_API ConstructConstraints(TGeometryParticleHandle<T, d>* Particle0, TGeometryParticleHandle<T, d>* Particle1, const FImplicitObject* Implicit0, const FImplicitObject* Implicit1, const TRigidTransform<T, d>& Transform0, const TRigidTransform<T, d>& Transform1, const T Thickness, FCollisionConstraintsArray& NewConstraints);
+		//
+		// Constraint API
+		//
+		template <typename T, int d>
+		void UpdateSingleShotManifold(TRigidBodyIterativeContactConstraint<T, d>&  Constraint, const TRigidTransform<T, d>& Transform0, const TRigidTransform<T, d>& Transform1, const T Thickness);
+
+		template <typename T, int d>
+		void UpdateIterativeManifold(TRigidBodyIterativeContactConstraint<T, d>&  Constraint, const TRigidTransform<T, d>& Transform0, const TRigidTransform<T, d>& Transform1, const T Thickness);
+
+		template<typename T, int d>
+		void CHAOS_API UpdateManifold(TCollisionConstraintBase<T, d>& ConstraintBase, const TRigidTransform<T, d>& ATM, const TRigidTransform<T, d>& BTM, const T Thickness);
+
+		template<ECollisionUpdateType UpdateType, typename T, int d>
+		void CHAOS_API UpdateConstraint(TCollisionConstraintBase<T, d>& ConstraintBase, const TRigidTransform<T, d>& Transform0, const TRigidTransform<T, d>& Transform1, const T Thickness);
+
+		template<typename T, int d>
+		void CHAOS_API ConstructConstraints(TGeometryParticleHandle<T, d>* Particle0, TGeometryParticleHandle<T, d>* Particle1, const FImplicitObject* Implicit0, const FImplicitObject* Implicit1, const TRigidTransform<T, d>& Transform0, const TRigidTransform<T, d>& Transform1, const T Thickness, FCollisionConstraintsArray& NewConstraints);
 
 	}
 }
diff --git a/Engine/Source/Runtime/Experimental/Chaos/Public/Chaos/CollisionResolutionTypes.h b/Engine/Source/Runtime/Experimental/Chaos/Public/Chaos/CollisionResolutionTypes.h
--- a/Engine/Source/Runtime/Experimental/Chaos/Public/Chaos/CollisionResolutionTypes.h
+++ b/Engine/Source/Runtime/Experimental/Chaos/Public/Chaos/CollisionResolutionTypes.h
@@ -206,7 +206,7 @@ namespace Chaos
 		// Samples
 		int32               NumSamples()                   { return Samples.Num(); }
 		void                AddSample(FSampleData && Data) { Samples.Add(Data); };
-		void                ResetSamples()                 { Samples.Reset(); }
+		void                ResetSamples(int32 NewSize=0)  { Samples.Reset(NewSize); }
 		FSampleData &       operator[](int32 Index) { ensure(0 <= Index && Index < NumSamples()); return Samples[Index]; }
 		const FSampleData & operator[](int32 Index) const  { ensure(0 <= Index && Index < NumSamples()); return Samples[Index]; }
 
diff --git a/Engine/Source/Runtime/Experimental/Chaos/Public/Chaos/HeightField.h b/Engine/Source/Runtime/Experimental/Chaos/Public/Chaos/HeightField.h
--- a/Engine/Source/Runtime/Experimental/Chaos/Public/Chaos/HeightField.h
+++ b/Engine/Source/Runtime/Experimental/Chaos/Public/Chaos/HeightField.h
@@ -75,6 +75,9 @@ namespace Chaos
 		bool GJKContactPoint(const TSphere<T, 3>& QueryGeom, const TRigidTransform<T, 3>& QueryTM, const T Thickness, TVector<T, 3>& ContactLocation, TVector<T, 3>& ContactNormal, T& ContactPhi) const;
 		bool GJKContactPoint(const TCapsule<T>& QueryGeom, const TRigidTransform<T, 3>& QueryTM, const T Thickness, TVector<T, 3>& ContactLocation, TVector<T, 3>& ContactNormal, T& ContactPhi) const;
 		bool GJKContactPoint(const FConvex& QueryGeom, const TRigidTransform<T, 3>& QueryTM, const T Thickness, TVector<T, 3>& ContactLocation, TVector<T, 3>& ContactNormal, T& ContactPhi) const;
+		bool GJKContactPoint(const TImplicitObjectScaled<TBox<T, 3>>& QueryGeom, const TRigidTransform<T, 3>& QueryTM, const T Thickness, TVector<T, 3>& ContactLocation, TVector<T, 3>& ContactNormal, T& ContactPhi) const;
+		bool GJKContactPoint(const TImplicitObjectScaled<TSphere<T, 3>>& QueryGeom, const TRigidTransform<T, 3>& QueryTM, const T Thickness, TVector<T, 3>& ContactLocation, TVector<T, 3>& ContactNormal, T& ContactPhi) const;
+		bool GJKContactPoint(const TImplicitObjectScaled<TCapsule<T>>& QueryGeom, const TRigidTransform<T, 3>& QueryTM, const T Thickness, TVector<T, 3>& ContactLocation, TVector<T, 3>& ContactNormal, T& ContactPhi) const;
 		bool GJKContactPoint(const TImplicitObjectScaled<FConvex>& QueryGeom, const TRigidTransform<T, 3>& QueryTM, const T Thickness, TVector<T, 3>& ContactLocation, TVector<T, 3>& ContactNormal, T& ContactPhi) const;
 
 
