Update shape bounds on pull from physics state

diff --git a/Engine/Source/Runtime/Engine/Private/PhysicsEngine/Experimental/PhysInterface_Chaos.cpp b/Engine/Source/Runtime/Engine/Private/PhysicsEngine/Experimental/PhysInterface_Chaos.cpp
--- a/Engine/Source/Runtime/Engine/Private/PhysicsEngine/Experimental/PhysInterface_Chaos.cpp
+++ b/Engine/Source/Runtime/Engine/Private/PhysicsEngine/Experimental/PhysInterface_Chaos.cpp
@@ -467,6 +467,7 @@ void FPhysInterface_Chaos::SetGlobalPose_AssumesLocked(const FPhysicsActorHandle
 {
 	InActorReference->SetX(InNewPose.GetLocation());
 	InActorReference->SetR(InNewPose.GetRotation());
+	InActorReference->UpdateShapeBounds();
 
 	FPhysScene* Scene = GetCurrentScene(InActorReference);
 	Scene->GetScene().UpdateActorInAccelerationStructure(InActorReference);
diff --git a/Engine/Source/Runtime/Experimental/Chaos/Public/Chaos/ParticleHandle.h b/Engine/Source/Runtime/Experimental/Chaos/Public/Chaos/ParticleHandle.h
--- a/Engine/Source/Runtime/Experimental/Chaos/Public/Chaos/ParticleHandle.h
+++ b/Engine/Source/Runtime/Experimental/Chaos/Public/Chaos/ParticleHandle.h
@@ -1301,13 +1301,7 @@ public:
 		Ar.UsingCustomVersion(FExternalPhysicsCustomObjectVersion::GUID);
 		if (Ar.CustomVer(FExternalPhysicsCustomObjectVersion::GUID) < FExternalPhysicsCustomObjectVersion::SerializeShapeWorldSpaceBounds)
 		{
-			if (MGeometry->HasBoundingBox())
-			{
-				for (auto& Shape : MShapesArray)
-				{
-					Shape->UpdateShapeBounds(FRigidTransform3(MX, MR));
-				}
-			}
+			UpdateShapeBounds();
 		}
 
 		if(Ar.IsLoading())
@@ -1370,6 +1364,17 @@ public:
 		this->MUserData = InUserData;
 	}
 
+	void UpdateShapeBounds()
+	{
+		if (MGeometry->HasBoundingBox())
+		{
+			for (auto& Shape : MShapesArray)
+			{
+				Shape->UpdateShapeBounds(FRigidTransform3(MX, MR));
+			}
+		}
+	}
+
 #if CHAOS_CHECKED
 	const FName& DebugName() const { return this->MDebugName; }
 	void SetDebugName(const FName& InDebugName)
diff --git a/Engine/Source/Runtime/Experimental/ChaosSolvers/Private/PhysicsProxy/SingleParticlePhysicsProxy.cpp b/Engine/Source/Runtime/Experimental/ChaosSolvers/Private/PhysicsProxy/SingleParticlePhysicsProxy.cpp
--- a/Engine/Source/Runtime/Experimental/ChaosSolvers/Private/PhysicsProxy/SingleParticlePhysicsProxy.cpp
+++ b/Engine/Source/Runtime/Experimental/ChaosSolvers/Private/PhysicsProxy/SingleParticlePhysicsProxy.cpp
@@ -307,6 +307,7 @@ void FSingleParticlePhysicsProxy<Chaos::TPBDRigidParticle<float, 3>>::PullFromPh
 		Particle->SetR(Buffer->R, false);
 		Particle->SetV(Buffer->MV, false);
 		Particle->SetW(Buffer->MW, false);
+		Particle->UpdateShapeBounds();
 	}
 }
 
