Get access to the solver via GetPhysicsScene()->GetSolver() rather than GetOwner()->GetWorld()->PhysicsScene_Chaos->GetSolver(), because PhysicsScene_Chaos is no longer being initialized.

diff --git a/Engine/Source/Runtime/Experimental/GeometryCollectionEngine/Private/GeometryCollection/GeometryCollectionComponent.cpp b/Engine/Source/Runtime/Experimental/GeometryCollectionEngine/Private/GeometryCollection/GeometryCollectionComponent.cpp
--- a/Engine/Source/Runtime/Experimental/GeometryCollectionEngine/Private/GeometryCollection/GeometryCollectionComponent.cpp
+++ b/Engine/Source/Runtime/Experimental/GeometryCollectionEngine/Private/GeometryCollection/GeometryCollectionComponent.cpp
@@ -214,20 +214,25 @@ void UGeometryCollectionComponent::BeginPlay()
 	// ---------- });
 	//////////////////////////////////////////////////////////////////////////
 
-	FChaosSolversModule* ChaosModule = FModuleManager::Get().GetModulePtr<FChaosSolversModule>("ChaosSolvers");
-	if (ChaosModule != nullptr)
+	if (PhysicsProxy != nullptr)
 	{
-		Chaos::FPhysicsSolver* Solver = GetSolver(*this);
-		if (Solver != nullptr)
+		FChaosSolversModule* ChaosModule = FModuleManager::Get().GetModulePtr<FChaosSolversModule>("ChaosSolvers");
+		if (ChaosModule != nullptr)
 		{
-			if (PhysicsProxy != nullptr)
+			Chaos::FPhysicsSolver* Solver = nullptr;
+			if (ChaosSolverActor != nullptr)
+			{
+				Solver = ChaosSolverActor->GetSolver();
+			}
+			else if (FPhysScene_Chaos* Scene = GetPhysicsScene())
+			{
+				Solver = Scene->GetSolver();
+			}
+			if (Solver != nullptr)
 			{
 				ChaosModule->GetDispatcher()->EnqueueCommandImmediate(Solver, [&InPhysicsProxy = PhysicsProxy](Chaos::FPhysicsSolver* InSolver)
 				{
-					if (InPhysicsProxy)
-					{
-						InPhysicsProxy->ActivateBodies();
-					}
+					InPhysicsProxy->ActivateBodies();
 				});
 			}
 		}
