Fix CIS

diff --git a/Engine/Source/Runtime/Experimental/Chaos/Private/Chaos/CollisionResolution.cpp b/Engine/Source/Runtime/Experimental/Chaos/Private/Chaos/CollisionResolution.cpp
--- a/Engine/Source/Runtime/Experimental/Chaos/Private/Chaos/CollisionResolution.cpp
+++ b/Engine/Source/Runtime/Experimental/Chaos/Private/Chaos/CollisionResolution.cpp
@@ -160,6 +160,16 @@ namespace Chaos
 		template <typename T, int d>
 		void UpdateIterativeManifold(TRigidBodyIterativeContactConstraint<T, d>&  Constraint, const TRigidTransform<T, d>& Transform0, const TRigidTransform<T, d>& Transform1, const T Thickness)
 		{
+			auto SumSampleData = [&](FRigidBodyIterativeContactConstraint& Constraint) -> TVector<float,3>
+			{
+				TVector<float, 3> ReturnValue(0);
+				for (int i = 0; i < Constraint.NumSamples(); i++)
+				{
+					ReturnValue += Constraint[i].X;
+				}
+				return ReturnValue;
+			};
+
 			// iterative manifolds for non TConvex implicit objects that require sampling 
 			TContactPoint<T> ContactPoint = ConvexConvexContactPoint(*Constraint.Manifold.Implicit[0], Transform0, *Constraint.Manifold.Implicit[1], Transform1, Thickness);
 
@@ -792,16 +802,6 @@ namespace Chaos
 		// Convex - Convex
 		//
 
-		TVector<float, 3> SumSampleData(FRigidBodyIterativeContactConstraint& Constraint)
-		{
-			TVector<float,3> ReturnValue(0);
-			for (int i = 0; i < Constraint.NumSamples(); i++)
-			{
-				ReturnValue += Constraint[i].X;
-			}
-			return ReturnValue;
-		}
-
 		template<class T, int d>
 		void UpdateConvexConvexConstraint(const FImplicitObject& Implicit0, const TRigidTransform<T, d>& Transform0, const FImplicitObject& Implicit1, const TRigidTransform<T, d>& Transform1, const T Thickness, TCollisionConstraintBase<T,d>& ConstraintBase)
 		{
