Chaos : Collisions

Renamed RigidBodyIterativeContactConstraint to RigidBodyMultiPointContactConstraint.

diff --git a/Engine/Source/Runtime/Experimental/Chaos/Private/Chaos/CollisionResolution.cpp b/Engine/Source/Runtime/Experimental/Chaos/Private/Chaos/CollisionResolution.cpp
--- a/Engine/Source/Runtime/Experimental/Chaos/Private/Chaos/CollisionResolution.cpp
+++ b/Engine/Source/Runtime/Experimental/Chaos/Private/Chaos/CollisionResolution.cpp
@@ -118,7 +118,7 @@ namespace Chaos
 		}
 
 		template <typename T, int d>
-		void UpdateSingleShotManifold(TRigidBodyIterativeContactConstraint<T, d>&  Constraint, const TRigidTransform<T, d>& Transform0, const TRigidTransform<T, d>& Transform1, const T Thickness)
+		void UpdateSingleShotManifold(TRigidBodyMultiPointContactConstraint<T, d>&  Constraint, const TRigidTransform<T, d>& Transform0, const TRigidTransform<T, d>& Transform1, const T Thickness)
 		{
 			// single shot manifolds for TConvex implicit object in the constraints implicit[0] position. 
 			TContactPoint<T> ContactPoint = ConvexConvexContactPoint(*Constraint.Manifold.Implicit[0], Transform0, *Constraint.Manifold.Implicit[1], Transform1, Thickness);
@@ -158,9 +158,9 @@ namespace Chaos
 		}
 
 		template <typename T, int d>
-		void UpdateIterativeManifold(TRigidBodyIterativeContactConstraint<T, d>&  Constraint, const TRigidTransform<T, d>& Transform0, const TRigidTransform<T, d>& Transform1, const T Thickness)
+		void UpdateIterativeManifold(TRigidBodyMultiPointContactConstraint<T, d>&  Constraint, const TRigidTransform<T, d>& Transform0, const TRigidTransform<T, d>& Transform1, const T Thickness)
 		{
-			auto SumSampleData = [&](FRigidBodyIterativeContactConstraint& Constraint) -> TVector<float,3>
+			auto SumSampleData = [&](TRigidBodyMultiPointContactConstraint<T, d>& Constraint) -> TVector<float,3>
 			{
 				TVector<float, 3> ReturnValue(0);
 				for (int i = 0; i < Constraint.NumSamples(); i++)
@@ -212,7 +212,7 @@ namespace Chaos
 				ensure(false); // max of 4 points
 			}
 
-			typedef FRigidBodyIterativeContactConstraint::FSampleData FSampleData;
+			typedef FRigidBodyMultiPointContactConstraint::FSampleData FSampleData;
 			TVector<T, d> Center = SumSampleData(Constraint) / Constraint.NumSamples();
 			for (int32 Index = 0; Index < Constraint.NumSamples(); Index++) { Constraint[Index].Delta = (Center - Constraint[Index].X).SizeSquared(); }
 		}
@@ -811,9 +811,9 @@ namespace Chaos
 			{
 				ContactPoint = ConvexConvexContactPoint(Implicit0, Transform0, Implicit1, Transform1, Thickness);
 			}
-			else if (ConstraintBase.GetType() == FRigidBodyIterativeContactConstraint::StaticType())
+			else if (ConstraintBase.GetType() == FRigidBodyMultiPointContactConstraint::StaticType())
 			{
-				TRigidBodyIterativeContactConstraint<T, d>& Constraint = *ConstraintBase.template As<TRigidBodyIterativeContactConstraint<T, d>>();
+				TRigidBodyMultiPointContactConstraint<T, d>& Constraint = *ConstraintBase.template As<TRigidBodyMultiPointContactConstraint<T, d>>();
 				ContactPoint.Phi = FLT_MAX;
 
 				const TRigidTransform<T, d> AToBTM = Transform0.GetRelativeTransform(Transform1);
@@ -843,15 +843,15 @@ namespace Chaos
 		template<class T, int d>
 		void UpdateConvexConvexManifold(TCollisionConstraintBase<T, d>&  ConstraintBase, const TRigidTransform<T, d>& Transform0, const TRigidTransform<T, d>& Transform1, const T Thickness)
 		{
-			if (ensure(ConstraintBase.GetType() == FRigidBodyIterativeContactConstraint::StaticType()))
+			if (ensure(ConstraintBase.GetType() == FRigidBodyMultiPointContactConstraint::StaticType()))
 			{
 				if (GetInnerType(ConstraintBase.Manifold.Implicit[0]->GetType()) == ImplicitObjectType::Convex)
 				{
-					UpdateSingleShotManifold(*ConstraintBase.template As<TRigidBodyIterativeContactConstraint<T, d>>(), Transform0, Transform1, Thickness);
+					UpdateSingleShotManifold(*ConstraintBase.template As<TRigidBodyMultiPointContactConstraint<T, d>>(), Transform0, Transform1, Thickness);
 				}
 				else
 				{
-					UpdateIterativeManifold(*ConstraintBase.template As<TRigidBodyIterativeContactConstraint<T, d>>(), Transform0, Transform1, Thickness);
+					UpdateIterativeManifold(*ConstraintBase.template As<TRigidBodyMultiPointContactConstraint<T, d>>(), Transform0, Transform1, Thickness);
 				}
 			}
 		}
@@ -862,7 +862,7 @@ namespace Chaos
 		{
 			TRigidTransform<T, d> ParticleImplicit0TM  = Transform0.GetRelativeTransform(Collisions::GetTransform(Particle0));
 			TRigidTransform<T, d> ParticleImplicit1TM = Transform1.GetRelativeTransform(Collisions::GetTransform(Particle1));
-			FRigidBodyIterativeContactConstraint* Constraint = new FRigidBodyIterativeContactConstraint(Particle0, Implicit0, ParticleImplicit0TM, Particle1, Implicit1, ParticleImplicit1TM);
+			FRigidBodyMultiPointContactConstraint* Constraint = new FRigidBodyMultiPointContactConstraint(Particle0, Implicit0, ParticleImplicit0TM, Particle1, Implicit1, ParticleImplicit1TM);
 
 			UpdateConvexConvexManifold(*Constraint, Transform0, Transform1, Thickness);
 			UpdateConvexConvexConstraint(*Implicit0, Transform0, *Implicit1, Transform1, Thickness, *Constraint);
@@ -1439,8 +1439,8 @@ namespace Chaos
 
 		template void ConstructUnionUnionConstraints<float, 3>(TGeometryParticleHandle<float, 3>* Particle0, TGeometryParticleHandle<float, 3>* Particle1, const FImplicitObject* Implicit0, const FImplicitObject* Implicit1, const TRigidTransform<float, 3>& Transform0, const TRigidTransform<float, 3>& Transform1, const float Thickness, FCollisionConstraintsArray& NewConstraints);
 
-		template void UpdateSingleShotManifold<float, 3>(TRigidBodyIterativeContactConstraint<float, 3>&  Constraint, const TRigidTransform<float, 3>& Transform0, const TRigidTransform<float, 3>& Transform1, const float Thickness);
-		template void UpdateIterativeManifold<float, 3>(TRigidBodyIterativeContactConstraint<float, 3>&  Constraint, const TRigidTransform<float, 3>& Transform0, const TRigidTransform<float, 3>& Transform1, const float Thickness);
+		template void UpdateSingleShotManifold<float, 3>(TRigidBodyMultiPointContactConstraint<float, 3>&  Constraint, const TRigidTransform<float, 3>& Transform0, const TRigidTransform<float, 3>& Transform1, const float Thickness);
+		template void UpdateIterativeManifold<float, 3>(TRigidBodyMultiPointContactConstraint<float, 3>&  Constraint, const TRigidTransform<float, 3>& Transform0, const TRigidTransform<float, 3>& Transform1, const float Thickness);
 		template void UpdateManifold<float, 3>(TCollisionConstraintBase<float, 3>& ConstraintBase, const TRigidTransform<float, 3>& ATM, const TRigidTransform<float, 3>& BTM, const float Thickness);
 		template void UpdateConstraint<ECollisionUpdateType::Any, float, 3>(TCollisionConstraintBase<float,3>& ConstraintBase, const TRigidTransform<float,3>& Transform0, const TRigidTransform<float,3>& Transform1, const float Thickness);
 		template void UpdateConstraint<ECollisionUpdateType::Deepest, float, 3>(TCollisionConstraintBase<float, 3>& ConstraintBase, const TRigidTransform<float, 3>& Transform0, const TRigidTransform<float, 3>& Transform1, const float Thickness);
diff --git a/Engine/Source/Runtime/Experimental/Chaos/Private/Chaos/PBDCollisionConstraints.cpp b/Engine/Source/Runtime/Experimental/Chaos/Private/Chaos/PBDCollisionConstraints.cpp
--- a/Engine/Source/Runtime/Experimental/Chaos/Private/Chaos/PBDCollisionConstraints.cpp
+++ b/Engine/Source/Runtime/Experimental/Chaos/Private/Chaos/PBDCollisionConstraints.cpp
@@ -121,12 +121,12 @@ namespace Chaos
 
 			delete PointConstraint;
 		}
-		else if (ConstraintBase->GetType() == TRigidBodyIterativeContactConstraint<T, 3>::StaticType())
+		else if (ConstraintBase->GetType() == TRigidBodyMultiPointContactConstraint<T, 3>::StaticType())
 		{
-			TRigidBodyIterativeContactConstraint<T, d>* IterativeConstraint = ConstraintBase->template As< TRigidBodyIterativeContactConstraint<T, d> >();
+			TRigidBodyMultiPointContactConstraint<T, d>* IterativeConstraint = ConstraintBase->template As< TRigidBodyMultiPointContactConstraint<T, d> >();
 
 			int32 Idx = IterativeConstraints.Add(*IterativeConstraint);
-			Handle = HandleAllocator.template AllocHandle< TRigidBodyIterativeContactConstraint<T, d> >(this, Idx);
+			Handle = HandleAllocator.template AllocHandle< TRigidBodyMultiPointContactConstraint<T, d> >(this, Idx);
 			Handle->GetContact().Timestamp = LifespanCounter;
 
 			delete IterativeConstraint;
@@ -294,7 +294,7 @@ namespace Chaos
 				Collisions::Apply(Contact, IterationParameters, ParticleParameters);
 			}
 
-			for (FIterativeContactConstraint& Contact : IterativeConstraints)
+			for (FMultiPointContactConstraint& Contact : IterativeConstraints)
 			{
 				Collisions::TContactParticleParameters<T> ParticleParameters = { MThickness, &MCollided, &MPhysicsMaterials, CollisionFrictionOverride, MAngularFriction };
 				Collisions::TContactIterationParameters<T> IterationParameters = { Dt, Iterations, NumIterations, MApplyPairIterations, nullptr };
@@ -324,7 +324,7 @@ namespace Chaos
 				Collisions::ApplyPushOut(Contact, TempStatic, IterationParameters, ParticleParameters);
 			}
 
-			for (FIterativeContactConstraint& Contact : IterativeConstraints)
+			for (FMultiPointContactConstraint& Contact : IterativeConstraints)
 			{
 				Collisions::TContactParticleParameters<T> ParticleParameters = { MThickness, &MCollided, &MPhysicsMaterials, CollisionFrictionOverride, MAngularFriction };
 				Collisions::TContactIterationParameters<T> IterationParameters = { Dt, Iterations, NumIterations, MApplyPushOutPairIterations, &bNeedsAnotherIteration };
diff --git a/Engine/Source/Runtime/Experimental/Chaos/Public/Chaos/CollisionResolution.h b/Engine/Source/Runtime/Experimental/Chaos/Public/Chaos/CollisionResolution.h
--- a/Engine/Source/Runtime/Experimental/Chaos/Public/Chaos/CollisionResolution.h
+++ b/Engine/Source/Runtime/Experimental/Chaos/Public/Chaos/CollisionResolution.h
@@ -193,10 +193,10 @@ namespace Chaos
 		// Constraint API
 		//
 		template <typename T, int d>
-		void UpdateSingleShotManifold(TRigidBodyIterativeContactConstraint<T, d>&  Constraint, const TRigidTransform<T, d>& Transform0, const TRigidTransform<T, d>& Transform1, const T Thickness);
+		void UpdateSingleShotManifold(TRigidBodyMultiPointContactConstraint<T, d>&  Constraint, const TRigidTransform<T, d>& Transform0, const TRigidTransform<T, d>& Transform1, const T Thickness);
 
 		template <typename T, int d>
-		void UpdateIterativeManifold(TRigidBodyIterativeContactConstraint<T, d>&  Constraint, const TRigidTransform<T, d>& Transform0, const TRigidTransform<T, d>& Transform1, const T Thickness);
+		void UpdateIterativeManifold(TRigidBodyMultiPointContactConstraint<T, d>&  Constraint, const TRigidTransform<T, d>& Transform0, const TRigidTransform<T, d>& Transform1, const T Thickness);
 
 		template<typename T, int d>
 		void CHAOS_API UpdateManifold(TCollisionConstraintBase<T, d>& ConstraintBase, const TRigidTransform<T, d>& ATM, const TRigidTransform<T, d>& BTM, const T Thickness);
diff --git a/Engine/Source/Runtime/Experimental/Chaos/Public/Chaos/CollisionResolutionTypes.h b/Engine/Source/Runtime/Experimental/Chaos/Public/Chaos/CollisionResolutionTypes.h
--- a/Engine/Source/Runtime/Experimental/Chaos/Public/Chaos/CollisionResolutionTypes.h
+++ b/Engine/Source/Runtime/Experimental/Chaos/Public/Chaos/CollisionResolutionTypes.h
@@ -81,7 +81,7 @@ namespace Chaos
 		{
 			None = 0,     // default value also indicates a invalid constraint
 			SinglePoint,  //   TRigidBodyPointContactConstraint
-			MultiPoint    //   TRigidBodyIterativeContactConstraint
+			MultiPoint    //   TRigidBodyMultiPointContactConstraint
 		};
 
 		TCollisionConstraintBase(FType InType = FType::None)
@@ -180,7 +180,7 @@ namespace Chaos
 	*
 	*/
 	template<class T, int d>
-	class CHAOS_API TRigidBodyIterativeContactConstraint : public TCollisionConstraintBase<T, d>
+	class CHAOS_API TRigidBodyMultiPointContactConstraint : public TCollisionConstraintBase<T, d>
 	{
 	public:
 		using Base = TCollisionConstraintBase<T, d>;
@@ -189,8 +189,8 @@ namespace Chaos
 		using Base::Particle;
 		struct FSampleData { TVector<T,d> X; float Delta; FManifold Manifold; };
 
-		TRigidBodyIterativeContactConstraint() : Base(Base::FType::MultiPoint), SourceNormalIndex(INDEX_NONE), PlaneNormal(0), PlanePosition(0) {}
-		TRigidBodyIterativeContactConstraint(
+		TRigidBodyMultiPointContactConstraint() : Base(Base::FType::MultiPoint), SourceNormalIndex(INDEX_NONE), PlaneNormal(0), PlanePosition(0) {}
+		TRigidBodyMultiPointContactConstraint(
 			FGeometryParticleHandle* Particle0, const FImplicitObject* Implicit0, const TRigidTransform<T, d>& Transform0,
 			FGeometryParticleHandle* Particle1, const FImplicitObject* Implicit1, const TRigidTransform<T, d>& Transform1)
 			: Base(Particle0, Implicit0, Transform0, Particle1, Implicit1, Transform1, Base::FType::MultiPoint)
@@ -213,7 +213,7 @@ namespace Chaos
 	private:
 		TArray<FSampleData> Samples; // iterative samples
 	};
-	typedef TRigidBodyIterativeContactConstraint<float, 3> FRigidBodyIterativeContactConstraint;
+	typedef TRigidBodyMultiPointContactConstraint<float, 3> FRigidBodyMultiPointContactConstraint;
 
 
 	//
@@ -301,8 +301,8 @@ namespace Chaos
 		const TRigidBodyPointContactConstraint<T, d>& GetPointContact() const { check(GetType() == FConstraintBase::FType::SinglePoint); return ConstraintContainer->PointConstraints[ConstraintIndex]; }
 		TRigidBodyPointContactConstraint<T, d>& GetPointContact() { check(GetType() == FConstraintBase::FType::SinglePoint); return ConstraintContainer->PointConstraints[ConstraintIndex]; }
 
-		const TRigidBodyIterativeContactConstraint<T, d>& GetIterativeContact() const { check(GetType() == FConstraintBase::FType::MultiPoint); return ConstraintContainer->IterativeConstraints[ConstraintIndex]; }
-		TRigidBodyIterativeContactConstraint<T, d>& GetIterativeContact() { check(GetType() == FConstraintBase::FType::MultiPoint); return ConstraintContainer->IterativeConstraints[ConstraintIndex]; }
+		const TRigidBodyMultiPointContactConstraint<T, d>& GetMultiPointContact() const { check(GetType() == FConstraintBase::FType::MultiPoint); return ConstraintContainer->IterativeConstraints[ConstraintIndex]; }
+		TRigidBodyMultiPointContactConstraint<T, d>& GetMultiPointContact() { check(GetType() == FConstraintBase::FType::MultiPoint); return ConstraintContainer->IterativeConstraints[ConstraintIndex]; }
 
 		typename FConstraintBase::FType GetType() const { return ConstraintType; }
 
diff --git a/Engine/Source/Runtime/Experimental/Chaos/Public/Chaos/PBDCollisionConstraints.h b/Engine/Source/Runtime/Experimental/Chaos/Public/Chaos/PBDCollisionConstraints.h
--- a/Engine/Source/Runtime/Experimental/Chaos/Public/Chaos/PBDCollisionConstraints.h
+++ b/Engine/Source/Runtime/Experimental/Chaos/Public/Chaos/PBDCollisionConstraints.h
@@ -57,7 +57,7 @@ public:
 	using FConstraintContainerHandle = TPBDCollisionConstraintHandle<T, d>;
 	using FConstraintBase = TCollisionConstraintBase<T, d>;
  	using FPointContactConstraint = TRigidBodyPointContactConstraint<T, d>;
-	using FIterativeContactConstraint = TRigidBodyIterativeContactConstraint<T, d>;
+	using FMultiPointContactConstraint = TRigidBodyMultiPointContactConstraint<T, d>;
 	using FConstraintHandleAllocator = TConstraintHandleAllocator<TPBDCollisionConstraints<T, d>>;
 	using FConstraintContainerHandleKey = typename TPBDCollisionConstraintHandle<T, d>::FHandleKey;
 
@@ -222,7 +222,7 @@ private:
 	const TPBDRigidsSOAs<T,d>& Particles;
 
 	TArray<FPointContactConstraint> PointConstraints;
-	TArray<FIterativeContactConstraint> IterativeConstraints;
+	TArray<FMultiPointContactConstraint> IterativeConstraints;
 
 	TMap< FConstraintContainerHandleKey, FConstraintContainerHandle* > Manifolds;
 	TArray<FConstraintContainerHandle*> Handles;
