#UE4 Fix shadowed variable

diff --git a/Engine/Source/Runtime/Experimental/Chaos/Private/Chaos/CollisionResolution.cpp b/Engine/Source/Runtime/Experimental/Chaos/Private/Chaos/CollisionResolution.cpp
--- a/Engine/Source/Runtime/Experimental/Chaos/Private/Chaos/CollisionResolution.cpp
+++ b/Engine/Source/Runtime/Experimental/Chaos/Private/Chaos/CollisionResolution.cpp
@@ -160,12 +160,12 @@ namespace Chaos
 		template <typename T, int d>
 		void UpdateIterativeManifold(TRigidBodyMultiPointContactConstraint<T, d>&  Constraint, const TRigidTransform<T, d>& Transform0, const TRigidTransform<T, d>& Transform1, const T Thickness)
 		{
-			auto SumSampleData = [&](TRigidBodyMultiPointContactConstraint<T, d>& Constraint) -> TVector<float,3>
+			auto SumSampleData = [&](TRigidBodyMultiPointContactConstraint<T, d>& LambdaConstraint) -> TVector<float,3>
 			{
 				TVector<float, 3> ReturnValue(0);
-				for (int i = 0; i < Constraint.NumSamples(); i++)
+				for (int i = 0; i < LambdaConstraint.NumSamples(); i++)
 				{
-					ReturnValue += Constraint[i].X;
+					ReturnValue += LambdaConstraint[i].X;
 				}
 				return ReturnValue;
 			};
