Merge CL 10786951 by hand to avoid breaking the build.

diff --git a/Engine/Source/Runtime/Core/Private/UObject/DevObjectVersion.cpp b/Engine/Source/Runtime/Core/Private/UObject/DevObjectVersion.cpp
--- a/Engine/Source/Runtime/Core/Private/UObject/DevObjectVersion.cpp
+++ b/Engine/Source/Runtime/Core/Private/UObject/DevObjectVersion.cpp
@@ -166,7 +166,7 @@ const FGuid FDestructionObjectVersion::GUID(0x174F1F0B, 0xB4C645A5, 0xB13F2EE8,
 FDevVersionRegistration GRegisterDestructionObjectVersion(FDestructionObjectVersion::GUID, FDestructionObjectVersion::LatestVersion, TEXT("Dev-Destruction"));
 
 // Unique Physics Object version id
-const FGuid FExternalPhysicsCustomObjectVersion::GUID(0x35F94A83, 0xE258406C, 0xA31809F5, 0x9610247C);
+const FGuid FExternalPhysicsCustomObjectVersion::GUID(0x35F94A84, 0xE258406D, 0xA31809F6, 0x9610247D);
 // Register Physics custom version with Core
 FDevVersionRegistration GRegisterExternalPhysicsCustomVersion(FExternalPhysicsCustomObjectVersion::GUID, FExternalPhysicsCustomObjectVersion::LatestVersion, TEXT("Dev-Physics-Ext"));
 
diff --git a/Engine/Source/Runtime/Engine/Private/PhysicsEngine/Experimental/PhysScene_Chaos.cpp b/Engine/Source/Runtime/Engine/Private/PhysicsEngine/Experimental/PhysScene_Chaos.cpp
--- a/Engine/Source/Runtime/Engine/Private/PhysicsEngine/Experimental/PhysScene_Chaos.cpp
+++ b/Engine/Source/Runtime/Engine/Private/PhysicsEngine/Experimental/PhysScene_Chaos.cpp
@@ -1189,16 +1189,16 @@ void FPhysScene_ChaosInterface::AddForce_AssumesLocked(FBodyInstance* BodyInstan
 			{
 				Rigid->SetObjectState(EObjectStateType::Dynamic);
 
-				const Chaos::TVector<float, 3> CurrentForce = Rigid->ExternalForce();
+				const Chaos::TVector<float, 3> CurrentForce = Rigid->F();
 				if (bAccelChange)
 				{
 					const float Mass = Rigid->M();
 					const Chaos::TVector<float, 3> TotalAcceleration = CurrentForce + (Force * Mass);
-					Rigid->SetExternalForce(TotalAcceleration);
+					Rigid->SetF(TotalAcceleration);
 				}
 				else
 				{
-					Rigid->SetExternalForce(CurrentForce + Force);
+					Rigid->SetF(CurrentForce + Force);
 				}
 
 			}
@@ -1220,8 +1220,8 @@ void FPhysScene_ChaosInterface::AddForceAtPosition_AssumesLocked(FBodyInstance*
 			EObjectStateType ObjectState = Rigid->ObjectState();
 			if (CHAOS_ENSURE(ObjectState == EObjectStateType::Dynamic || ObjectState == EObjectStateType::Sleeping))
 			{
-				const Chaos::FVec3& CurrentForce = Rigid->ExternalForce();
-				const Chaos::FVec3& CurrentTorque = Rigid->ExternalTorque();
+				const Chaos::FVec3& CurrentForce = Rigid->F();
+				const Chaos::FVec3& CurrentTorque = Rigid->Torque();
 				const Chaos::FVec3 WorldCOM = FParticleUtilitiesGT::GetCoMWorldPosition(Rigid);
 
 				Rigid->SetObjectState(EObjectStateType::Dynamic);
@@ -1232,14 +1232,14 @@ void FPhysScene_ChaosInterface::AddForceAtPosition_AssumesLocked(FBodyInstance*
 					const Chaos::FVec3 WorldPosition = CurrentTransform.TransformPosition(Position);
 					const Chaos::FVec3 WorldForce = CurrentTransform.TransformVector(Force);
 					const Chaos::FVec3 WorldTorque = Chaos::FVec3::CrossProduct(WorldPosition - WorldCOM, WorldForce);
-					Rigid->SetExternalForce(CurrentForce + WorldForce);
-					Rigid->SetExternalTorque(CurrentTorque + WorldTorque);
+					Rigid->SetF(CurrentForce + WorldForce);
+					Rigid->SetTorque(CurrentTorque + WorldTorque);
 				}
 				else
 				{
 					const Chaos::FVec3 WorldTorque = Chaos::FVec3::CrossProduct(Position - WorldCOM, Force);
-					Rigid->SetExternalForce(CurrentForce + Force);
-					Rigid->SetExternalTorque(CurrentTorque + WorldTorque);
+					Rigid->SetF(CurrentForce + Force);
+					Rigid->SetTorque(CurrentTorque + WorldTorque);
 				}
 
 			}
@@ -1296,14 +1296,14 @@ void FPhysScene_ChaosInterface::AddTorque_AssumesLocked(FBodyInstance* BodyInsta
 			EObjectStateType ObjectState = Rigid->ObjectState();
 			if (CHAOS_ENSURE(ObjectState == EObjectStateType::Dynamic || ObjectState == EObjectStateType::Sleeping))
 			{
-				const Chaos::TVector<float, 3> CurrentTorque = Rigid->ExternalTorque();
+				const Chaos::TVector<float, 3> CurrentTorque = Rigid->Torque();
 				if (bAccelChange)
 				{
-					Rigid->SetExternalTorque(CurrentTorque + (Rigid->I() * Torque));
+					Rigid->SetTorque(CurrentTorque + (Rigid->I() * Torque));
 				}
 				else
 				{
-					Rigid->SetExternalTorque(CurrentTorque + Torque);
+					Rigid->SetTorque(CurrentTorque + Torque);
 				}
 			}
 		}
diff --git a/Engine/Source/Runtime/Engine/Private/PhysicsEngine/ImmediatePhysics/ImmediatePhysicsChaos/ImmediatePhysicsActorHandle_Chaos.cpp b/Engine/Source/Runtime/Engine/Private/PhysicsEngine/ImmediatePhysics/ImmediatePhysicsChaos/ImmediatePhysicsActorHandle_Chaos.cpp
--- a/Engine/Source/Runtime/Engine/Private/PhysicsEngine/ImmediatePhysics/ImmediatePhysicsChaos/ImmediatePhysicsActorHandle_Chaos.cpp
+++ b/Engine/Source/Runtime/Engine/Private/PhysicsEngine/ImmediatePhysics/ImmediatePhysicsChaos/ImmediatePhysicsActorHandle_Chaos.cpp
@@ -433,7 +433,7 @@ namespace ImmediatePhysics_Chaos
 
 		if (TPBDRigidParticleHandle<FReal, 3>* Rigid = Handle()->CastToRigidParticle())
 		{
-			Rigid->ExternalForce() += Force;
+			Rigid->F() += Force;
 		}
 	}
 
diff --git a/Engine/Source/Runtime/Experimental/Chaos/Private/Chaos/Evolution/PBDMinEvolution.cpp b/Engine/Source/Runtime/Experimental/Chaos/Private/Chaos/Evolution/PBDMinEvolution.cpp
--- a/Engine/Source/Runtime/Experimental/Chaos/Private/Chaos/Evolution/PBDMinEvolution.cpp
+++ b/Engine/Source/Runtime/Experimental/Chaos/Private/Chaos/Evolution/PBDMinEvolution.cpp
@@ -127,8 +127,7 @@ namespace Chaos
 				Particle.PreV() = Particle.V();
 				Particle.PreW() = Particle.W();
 
-				Particle.F() = Particle.ExternalForce() + Particle.M() * Gravity;
-				Particle.Torque() = Particle.ExternalTorque();
+				Particle.F() += Particle.M() * Gravity;
 
 				EulerStepVelocityRule.Apply(Particle, Dt);
 				EtherDragRule.Apply(Particle, Dt);
@@ -137,8 +136,6 @@ namespace Chaos
 
 				Particle.F() = FVec3(0);
 				Particle.Torque() = FVec3(0);
-				Particle.ExternalForce() = FVec3(0);
-				Particle.ExternalTorque() = FVec3(0);
 
 				if (Particle.HasBounds())
 				{
@@ -303,4 +300,4 @@ namespace Chaos
 		}
 	}
 
-}
\ No newline at end of file
+}
diff --git a/Engine/Source/Runtime/Experimental/Chaos/Private/Chaos/PBDRigidsEvolutionGBF.cpp b/Engine/Source/Runtime/Experimental/Chaos/Private/Chaos/PBDRigidsEvolutionGBF.cpp
--- a/Engine/Source/Runtime/Experimental/Chaos/Private/Chaos/PBDRigidsEvolutionGBF.cpp
+++ b/Engine/Source/Runtime/Experimental/Chaos/Private/Chaos/PBDRigidsEvolutionGBF.cpp
@@ -314,11 +314,6 @@ TPBDRigidsEvolutionGBF<T, d>::TPBDRigidsEvolutionGBF(TPBDRigidsSOAs<T, d>& InPar
 		});
 	});
 
-	AddForceFunction([this](TTransientPBDRigidParticleHandle<T, d>& HandleIn, const T Dt) 
-	{
-		ExternalForces.Apply(HandleIn, Dt);
-	});
-
 	AddForceFunction([this](TTransientPBDRigidParticleHandle<T, d>& HandleIn, const T Dt)
 	{
 		GravityForces.Apply(HandleIn, Dt);
diff --git a/Engine/Source/Runtime/Experimental/Chaos/Public/Chaos/PBDRigidsEvolution.h b/Engine/Source/Runtime/Experimental/Chaos/Public/Chaos/PBDRigidsEvolution.h
--- a/Engine/Source/Runtime/Experimental/Chaos/Public/Chaos/PBDRigidsEvolution.h
+++ b/Engine/Source/Runtime/Experimental/Chaos/Public/Chaos/PBDRigidsEvolution.h
@@ -269,6 +269,7 @@ class TPBDRigidsEvolutionBase
 	}
 
 	CHAOS_API void AddForceFunction(FForceRule ForceFunction) { ForceRules.Add(ForceFunction); }
+	CHAOS_API void AddImpulseFunction(FForceRule ImpulseFunction) { ImpulseRules.Add(ImpulseFunction); }
 	CHAOS_API void SetParticleUpdateVelocityFunction(FUpdateVelocityRule ParticleUpdate) { ParticleUpdateVelocity = ParticleUpdate; }
 	CHAOS_API void SetParticleUpdatePositionFunction(FUpdatePositionRule ParticleUpdate) { ParticleUpdatePosition = ParticleUpdate; }
 
@@ -623,6 +624,7 @@ protected:
 	void WaitOnAccelerationStructure();
 
 	TArray<FForceRule> ForceRules;
+	TArray<FForceRule> ImpulseRules;
 	FUpdateVelocityRule ParticleUpdateVelocity;
 	FUpdatePositionRule ParticleUpdatePosition;
 	FKinematicUpdateRule KinematicUpdate;
diff --git a/Engine/Source/Runtime/Experimental/Chaos/Public/Chaos/PBDRigidsEvolutionGBF.h b/Engine/Source/Runtime/Experimental/Chaos/Public/Chaos/PBDRigidsEvolutionGBF.h
--- a/Engine/Source/Runtime/Experimental/Chaos/Public/Chaos/PBDRigidsEvolutionGBF.h
+++ b/Engine/Source/Runtime/Experimental/Chaos/Public/Chaos/PBDRigidsEvolutionGBF.h
@@ -1,218 +1,211 @@
 // Copyright 1998-2019 Epic Games, Inc. All Rights Reserved.
 #pragma once
 
-#include "Chaos/PBDRigidsEvolution.h"
-#include "Chaos/PBDCollisionConstraints.h"
 #include "Chaos/ChaosPerfTest.h"
 #include "Chaos/Collision/CollisionDetector.h"
 #include "Chaos/Collision/CollisionReceiver.h"
 #include "Chaos/Collision/NarrowPhase.h"
 #include "Chaos/Collision/SpatialAccelerationBroadPhase.h"
-#include "Chaos/PerParticleInitForce.h"
-#include "Chaos/PerParticleEulerStepVelocity.h"
+#include "Chaos/PBDCollisionConstraints.h"
+#include "Chaos/PBDRigidsEvolution.h"
 #include "Chaos/PerParticleEtherDrag.h"
+#include "Chaos/PerParticleEulerStepVelocity.h"
+#include "Chaos/PerParticleExternalForces.h"
 #include "Chaos/PerParticleGravity.h"
+#include "Chaos/PerParticleInitForce.h"
 #include "Chaos/PerParticlePBDEulerStep.h"
-#include "Chaos/PerParticleExternalForces.h"
+#include "Chaos/PerParticleAddImpulses.h"
 
 namespace Chaos
 {
-class FChaosArchive;
+	class FChaosArchive;
 
-CHAOS_API extern float HackMaxAngularVelocity;
-CHAOS_API extern float HackMaxVelocity;
+	CHAOS_API extern float HackMaxAngularVelocity;
+	CHAOS_API extern float HackMaxVelocity;
 
-CHAOS_API extern float HackLinearDrag;
-CHAOS_API extern float HackAngularDrag;
+	CHAOS_API extern float HackLinearDrag;
+	CHAOS_API extern float HackAngularDrag;
 
-template<typename T, int d>
-class TPBDRigidsEvolutionGBF;
+	template<typename T, int d>
+	class TPBDRigidsEvolutionGBF;
 
-template<typename T, int d>
-using TPBDRigidsEvolutionCallback = TFunction<void()>;
+	template<typename T, int d>
+	using TPBDRigidsEvolutionCallback = TFunction<void()>;
 
-template<typename T, int d>
-using TPBDRigidsEvolutionIslandCallback = TFunction<void(int32 Island)>;
+	template<typename T, int d>
+	using TPBDRigidsEvolutionIslandCallback = TFunction<void(int32 Island)>;
 
-
-template<typename T, int d>
-class TPBDRigidsEvolutionGBF : public TPBDRigidsEvolutionBase<TPBDRigidsEvolutionGBF<T, d>, TPBDCollisionConstraints<T,d>, T, d>
-{
-public:
-	using Base = TPBDRigidsEvolutionBase<TPBDRigidsEvolutionGBF<T, d>, TPBDCollisionConstraints<T, d>, T, d>;
-	using Base::Particles;
-	using Base::ForceRules;
-	using Base::ParticleUpdatePosition;
-	using Base::ApplyKinematicTargets;
-	using Base::SetParticleUpdateVelocityFunction;
-	using Base::SetParticleUpdatePositionFunction;
-	using Base::AddConstraintRule;
-	using Base::AddForceFunction;
-	using Base::Clustering;
-	using typename Base::FForceRule;
-	using FGravityForces = TPerParticleGravity<T, d>;
-	using FCollisionConstraints = TPBDCollisionConstraints<T, d>;
-	using FCollisionConstraintRule = TPBDConstraintColorRule<FCollisionConstraints>;
-	using FCollisionDetector = TCollisionDetector<FSpatialAccelerationBroadPhase, FNarrowPhase, FAsyncCollisionReceiver, FCollisionConstraints>;
-	using FExternalForces = TPerParticleExternalForces<T, d>;
-
-	static constexpr int32 DefaultNumIterations = 1;
-	static constexpr int32 DefaultNumPairIterations = 1;
-	static constexpr int32 DefaultNumPushOutIterations = 5;
-	static constexpr int32 DefaultNumPushOutPairIterations = 2;
-
-	CHAOS_API TPBDRigidsEvolutionGBF(TPBDRigidsSOAs<T, d>& InParticles, int32 InNumIterations = DefaultNumIterations, bool InIsSingleThreaded = false);
-	CHAOS_API ~TPBDRigidsEvolutionGBF() {}
-
-	void SetPostIntegrateCallback(const TPBDRigidsEvolutionCallback<T, d>& Cb)
+	template<typename T, int d>
+	class TPBDRigidsEvolutionGBF : public TPBDRigidsEvolutionBase<TPBDRigidsEvolutionGBF<T, d>, TPBDCollisionConstraints<T, d>, T, d>
 	{
-		PostIntegrateCallback = Cb;
-	}
-
-	void SetPostDetectCollisionsCallback(const TPBDRigidsEvolutionCallback<T, d>& Cb)
-	{
-		PostDetectCollisionsCallback = Cb;
-	}
+	public:
+		using Base = TPBDRigidsEvolutionBase<TPBDRigidsEvolutionGBF<T, d>, TPBDCollisionConstraints<T, d>, T, d>;
+		using Base::AddConstraintRule;
+		using Base::AddForceFunction;
+		using Base::ApplyKinematicTargets;
+		using Base::Clustering;
+		using Base::ForceRules;
+		using Base::Particles;
+		using Base::ParticleUpdatePosition;
+		using Base::SetParticleUpdatePositionFunction;
+		using Base::SetParticleUpdateVelocityFunction;
+		using typename Base::FForceRule;
+		using FGravityForces = TPerParticleGravity<T, d>;
+		using FCollisionConstraints = TPBDCollisionConstraints<T, d>;
+		using FCollisionConstraintRule = TPBDConstraintColorRule<FCollisionConstraints>;
+		using FCollisionDetector = TCollisionDetector<FSpatialAccelerationBroadPhase, FNarrowPhase, FAsyncCollisionReceiver, FCollisionConstraints>;
+		using FExternalForces = TPerParticleExternalForces<T, d>;
+
+		static constexpr int32 DefaultNumIterations = 1;
+		static constexpr int32 DefaultNumPairIterations = 1;
+		static constexpr int32 DefaultNumPushOutIterations = 5;
+		static constexpr int32 DefaultNumPushOutPairIterations = 2;
+
+		CHAOS_API TPBDRigidsEvolutionGBF(TPBDRigidsSOAs<T, d>& InParticles, int32 InNumIterations = DefaultNumIterations, bool InIsSingleThreaded = false);
+		CHAOS_API ~TPBDRigidsEvolutionGBF() {}
+
+		void SetPostIntegrateCallback(const TPBDRigidsEvolutionCallback<T, d>& Cb)
+		{
+			PostIntegrateCallback = Cb;
+		}
 
-	void SetPreApplyCallback(const TPBDRigidsEvolutionCallback<T, d>& Cb)
-	{
-		PreApplyCallback = Cb;
-	}
+		void SetPostDetectCollisionsCallback(const TPBDRigidsEvolutionCallback<T, d>& Cb)
+		{
+			PostDetectCollisionsCallback = Cb;
+		}
 
-	void SetPostApplyCallback(const TPBDRigidsEvolutionIslandCallback<T, d>& Cb)
-	{
-		PostApplyCallback = Cb;
-	}
+		void SetPreApplyCallback(const TPBDRigidsEvolutionCallback<T, d>& Cb)
+		{
+			PreApplyCallback = Cb;
+		}
 
-	void SetPostApplyPushOutCallback(const TPBDRigidsEvolutionIslandCallback<T, d>& Cb)
-	{
-		PostApplyPushOutCallback = Cb;
-	}
+		void SetPostApplyCallback(const TPBDRigidsEvolutionIslandCallback<T, d>& Cb)
+		{
+			PostApplyCallback = Cb;
+		}
 
-	CHAOS_API void Advance(const T Dt, const T MaxStepDt, const int32 MaxSteps);
-	CHAOS_API void AdvanceOneTimeStep(const T dt, const T StepFraction = (T)1.0);
+		void SetPostApplyPushOutCallback(const TPBDRigidsEvolutionIslandCallback<T, d>& Cb)
+		{
+			PostApplyPushOutCallback = Cb;
+		}
 
-	using Base::ApplyConstraints;
-	using Base::ApplyPushOut;
+		CHAOS_API void Advance(const T Dt, const T MaxStepDt, const int32 MaxSteps);
+		CHAOS_API void AdvanceOneTimeStep(const T dt, const T StepFraction = (T)1.0);
 
-	FCollisionConstraints& GetCollisionConstraints() { return CollisionConstraints; }
-	const FCollisionConstraints& GetCollisionConstraints() const { return CollisionConstraints; }
+		using Base::ApplyConstraints;
+		using Base::ApplyPushOut;
 
-	FCollisionConstraintRule& GetCollisionConstraintsRule() { return CollisionRule; }
-	const FCollisionConstraintRule& GetCollisionConstraintsRule() const { return CollisionRule; }
+		FCollisionConstraints& GetCollisionConstraints() { return CollisionConstraints; }
+		const FCollisionConstraints& GetCollisionConstraints() const { return CollisionConstraints; }
 
-	FCollisionDetector& GetCollisionDetector() { return CollisionDetector; }
-	const FCollisionDetector& GetCollisionDetector() const { return CollisionDetector; }
+		FCollisionConstraintRule& GetCollisionConstraintsRule() { return CollisionRule; }
+		const FCollisionConstraintRule& GetCollisionConstraintsRule() const { return CollisionRule; }
 
-	FExternalForces& GetExternalForces() { return ExternalForces; }
-	const FExternalForces& GetExternalForces() const { return ExternalForces; }
+		FCollisionDetector& GetCollisionDetector() { return CollisionDetector; }
+		const FCollisionDetector& GetCollisionDetector() const { return CollisionDetector; }
 
-	FGravityForces& GetGravityForces() { return GravityForces; }
-	const FGravityForces& GetGravityForces() const { return GravityForces; }
+		FGravityForces& GetGravityForces() { return GravityForces; }
+		const FGravityForces& GetGravityForces() const { return GravityForces; }
 
-	CHAOS_API inline void EndFrame(T Dt)
-	{
-		Particles.GetNonDisabledDynamicView().ParallelFor([&](auto& Particle, int32 Index)
+		CHAOS_API inline void EndFrame(T Dt)
 		{
-			Particle.ExternalForce() = TVector<T, 3>(0);
-			Particle.ExternalTorque() = TVector<T, 3>(0);
-		});
-	}
-
+			Particles.GetNonDisabledDynamicView().ParallelFor([&](auto& Particle, int32 Index) {
+				Particle.F() = TVector<T, 3>(0);
+				Particle.Torque() = TVector<T, 3>(0);
+			});
+		}
 
-	template <typename TParticleView>
-	void Integrate(const TParticleView& InParticles, T Dt)
-	{
-		//SCOPE_CYCLE_COUNTER(STAT_Integrate);
-		CHAOS_SCOPED_TIMER(Integrate);
-		TPerParticleInitForce<T, d> InitForceRule;
-		TPerParticleEulerStepVelocity<T, d> EulerStepVelocityRule;
-		TPerParticleEtherDrag<T, d> EtherDragRule;
-		TPerParticlePBDEulerStep<T, d> EulerStepRule;
-
-		const T MaxAngularSpeedSq = HackMaxAngularVelocity * HackMaxAngularVelocity;
-		const T MaxSpeedSq = HackMaxVelocity * HackMaxVelocity;
-		InParticles.ParallelFor([&](auto& GeomParticle, int32 Index)
+		template<typename TParticleView>
+		void Integrate(const TParticleView& InParticles, T Dt)
 		{
-			//question: can we enforce this at the API layer? Right now islands contain non dynamic which makes this hard
-			auto PBDParticle = GeomParticle.CastToRigidParticle();
-			if(PBDParticle && PBDParticle->ObjectState() == EObjectStateType::Dynamic)
-			{
-				auto& Particle = *PBDParticle;
+			//SCOPE_CYCLE_COUNTER(STAT_Integrate);
+			CHAOS_SCOPED_TIMER(Integrate);
+			TPerParticleEulerStepVelocity<T, d> EulerStepVelocityRule;
+			TPerParticleAddImpulses<T, d> AddImpulsesRule;
+			TPerParticleEtherDrag<T, d> EtherDragRule;
+			TPerParticlePBDEulerStep<T, d> EulerStepRule;
+
+			const T MaxAngularSpeedSq = HackMaxAngularVelocity * HackMaxAngularVelocity;
+			const T MaxSpeedSq = HackMaxVelocity * HackMaxVelocity;
+			InParticles.ParallelFor([&](auto& GeomParticle, int32 Index) {
+				//question: can we enforce this at the API layer? Right now islands contain non dynamic which makes this hard
+				auto PBDParticle = GeomParticle.CastToRigidParticle();
+				if (PBDParticle && PBDParticle->ObjectState() == EObjectStateType::Dynamic)
+				{
+					auto& Particle = *PBDParticle;
 
-				//save off previous velocities
-				Particle.PreV() = Particle.V();
-				Particle.PreW() = Particle.W();
+					//save off previous velocities
+					Particle.PreV() = Particle.V();
+					Particle.PreW() = Particle.W();
 
-				InitForceRule.Apply(Particle, Dt);
-				for (FForceRule ForceRule : ForceRules)
-				{
-					ForceRule(Particle, Dt);
-				}
-				EulerStepVelocityRule.Apply(Particle, Dt);
-				EtherDragRule.Apply(Particle, Dt);
+					for (FForceRule ForceRule : ForceRules)
+					{
+						ForceRule(Particle, Dt);
+					}
+					EulerStepVelocityRule.Apply(Particle, Dt);
+					AddImpulsesRule.Apply(Particle, Dt);
+					EtherDragRule.Apply(Particle, Dt);
 
-				if (HackMaxAngularVelocity >= 0.f)
-				{
-					const T AngularSpeedSq = Particle.W().SizeSquared();
-					if (AngularSpeedSq > MaxAngularSpeedSq)
+					if (HackMaxAngularVelocity >= 0.f)
 					{
-						Particle.W() = Particle.W() * (HackMaxAngularVelocity / FMath::Sqrt(AngularSpeedSq));
+						const T AngularSpeedSq = Particle.W().SizeSquared();
+						if (AngularSpeedSq > MaxAngularSpeedSq)
+						{
+							Particle.W() = Particle.W() * (HackMaxAngularVelocity / FMath::Sqrt(AngularSpeedSq));
+						}
 					}
-				}
 
-				if (HackMaxVelocity >= 0.f)
-				{
-					const T SpeedSq = Particle.V().SizeSquared();
-					if (SpeedSq > MaxSpeedSq)
+					if (HackMaxVelocity >= 0.f)
 					{
-						Particle.V() = Particle.V() * (HackMaxVelocity / FMath::Sqrt(SpeedSq));
+						const T SpeedSq = Particle.V().SizeSquared();
+						if (SpeedSq > MaxSpeedSq)
+						{
+							Particle.V() = Particle.V() * (HackMaxVelocity / FMath::Sqrt(SpeedSq));
+						}
 					}
-				}
 
-				EulerStepRule.Apply(Particle, Dt);
+					EulerStepRule.Apply(Particle, Dt);
 
-				if (Particle.HasBounds())
-				{
-					const TAABB<T, d>& LocalBounds = Particle.LocalBounds();
-					TAABB<T, d> WorldSpaceBounds = LocalBounds.TransformedAABB(TRigidTransform<T, d>(Particle.P(), Particle.Q()));
-					WorldSpaceBounds.ThickenSymmetrically(Particle.V());
-					Particle.SetWorldSpaceInflatedBounds(WorldSpaceBounds);
+					if (Particle.HasBounds())
+					{
+						const TAABB<T, d>& LocalBounds = Particle.LocalBounds();
+						TAABB<T, d> WorldSpaceBounds = LocalBounds.TransformedAABB(TRigidTransform<T, d>(Particle.P(), Particle.Q()));
+						WorldSpaceBounds.ThickenSymmetrically(Particle.V());
+						Particle.SetWorldSpaceInflatedBounds(WorldSpaceBounds);
+					}
 				}
-			}
-		});
+			});
 
-		for (auto& Particle : InParticles)
-		{
-			Base::DirtyParticle(Particle);
+			for (auto& Particle : InParticles)
+			{
+				Base::DirtyParticle(Particle);
+			}
 		}
-	}
-
-	CHAOS_API void Serialize(FChaosArchive& Ar);
-
-protected:
-	using Base::UpdateConstraintPositionBasedState;
-	using Base::CreateConstraintGraph;
-	using Base::CreateIslands;
-	using Base::ConstraintRules;
-	using Base::GetConstraintGraph;
-	using Base::UpdateVelocities;
-	using Base::PhysicsMaterials;
-	using Base::ParticleDisableCount;
-	using Base::Collided;
-	using Base::InternalAcceleration;
-
-	FGravityForces GravityForces;
-	FExternalForces ExternalForces;
-	FCollisionConstraints CollisionConstraints;
-	FCollisionConstraintRule CollisionRule;
-	FSpatialAccelerationBroadPhase BroadPhase;
-	FCollisionDetector CollisionDetector;
-
-	TPBDRigidsEvolutionCallback<T, d> PostIntegrateCallback;
-	TPBDRigidsEvolutionCallback<T, d> PostDetectCollisionsCallback;
-	TPBDRigidsEvolutionCallback<T, d> PreApplyCallback;
-	TPBDRigidsEvolutionIslandCallback<T, d> PostApplyCallback;
-	TPBDRigidsEvolutionIslandCallback<T, d> PostApplyPushOutCallback;
-};
+
+		CHAOS_API void Serialize(FChaosArchive& Ar);
+
+	protected:
+		using Base::Collided;
+		using Base::ConstraintRules;
+		using Base::CreateConstraintGraph;
+		using Base::CreateIslands;
+		using Base::GetConstraintGraph;
+		using Base::InternalAcceleration;
+		using Base::ParticleDisableCount;
+		using Base::PhysicsMaterials;
+		using Base::UpdateConstraintPositionBasedState;
+		using Base::UpdateVelocities;
+
+		FGravityForces GravityForces;
+		FCollisionConstraints CollisionConstraints;
+		FCollisionConstraintRule CollisionRule;
+		FSpatialAccelerationBroadPhase BroadPhase;
+		FCollisionDetector CollisionDetector;
+
+		TPBDRigidsEvolutionCallback<T, d> PostIntegrateCallback;
+		TPBDRigidsEvolutionCallback<T, d> PostDetectCollisionsCallback;
+		TPBDRigidsEvolutionCallback<T, d> PreApplyCallback;
+		TPBDRigidsEvolutionIslandCallback<T, d> PostApplyCallback;
+		TPBDRigidsEvolutionIslandCallback<T, d> PostApplyPushOutCallback;
+	};
 }
diff --git a/Engine/Source/Runtime/Experimental/Chaos/Public/Chaos/ParticleDirtyFlags.h b/Engine/Source/Runtime/Experimental/Chaos/Public/Chaos/ParticleDirtyFlags.h
--- a/Engine/Source/Runtime/Experimental/Chaos/Public/Chaos/ParticleDirtyFlags.h
+++ b/Engine/Source/Runtime/Experimental/Chaos/Public/Chaos/ParticleDirtyFlags.h
@@ -31,8 +31,8 @@ namespace Chaos
 		AngularEtherDrag			= 1 << 19,
 		ObjectState				= 1 << 20,
 		Geometry				= 1 << 21,
-		ExternalForce			= 1 << 22,
-		ExternalTorque			= 1 << 23,
+		LinearImpulse			= 1 << 22,
+		AngularImpulse			= 1 << 23,
 		GravityEnabled			= 1 << 24,
 		SpatialIdx				= 1 << 25,
 		HashResult				= 1 << 26,
diff --git a/Engine/Source/Runtime/Experimental/Chaos/Public/Chaos/ParticleHandle.h b/Engine/Source/Runtime/Experimental/Chaos/Public/Chaos/ParticleHandle.h
--- a/Engine/Source/Runtime/Experimental/Chaos/Public/Chaos/ParticleHandle.h
+++ b/Engine/Source/Runtime/Experimental/Chaos/Public/Chaos/ParticleHandle.h
@@ -74,8 +74,8 @@ void PBDRigidParticleHandleImpDefaultConstruct(TPBDRigidParticleHandleImp<T, d,
 	Concrete.SetQ(Concrete.R());
 	Concrete.SetF(TVector<T, d>(0));
 	Concrete.SetTorque(TVector<T, d>(0));
-	Concrete.SetExternalForce(TVector<T, d>(0));
-	Concrete.SetExternalTorque(TVector<T, d>(0));
+	Concrete.SetLinearImpulse(TVector<T, d>(0));
+	Concrete.SetAngularImpulse(TVector<T, d>(0));
 	Concrete.SetM(1);
 	Concrete.SetInvM(1);
 	Concrete.SetI(PMatrix<T, d, d>(1, 1, 1));
@@ -97,8 +97,8 @@ void PBDRigidParticleDefaultConstruct(TPBDRigidParticle<T,d>& Concrete, const TP
 	Concrete.SetQ(Concrete.R());
 	Concrete.SetF(TVector<T, d>(0));
 	Concrete.SetTorque(TVector<T, d>(0));
-	Concrete.SetExternalForce(TVector<T, d>(0));
-	Concrete.SetExternalTorque(TVector<T, d>(0));
+	Concrete.SetLinearImpulse(TVector<T, d>(0));
+	Concrete.SetAngularImpulse(TVector<T, d>(0));
 	Concrete.SetM(1);
 	Concrete.SetInvM(1);
 	Concrete.SetI(PMatrix<T, d, d>(1, 1, 1));
@@ -651,13 +651,13 @@ public:
 	TVector<T, d>& Torque() { return PBDRigidParticles->Torque(ParticleIdx); }
 	void SetTorque(const TVector<T, d>& InTorque) { PBDRigidParticles->Torque(ParticleIdx) = InTorque; }
 
-	const TVector<T, d>& ExternalForce() const { return PBDRigidParticles->ExternalForce(ParticleIdx); }
-	TVector<T, d>& ExternalForce() { return PBDRigidParticles->ExternalForce(ParticleIdx); }
-	void SetExternalForce(const TVector<T, d>& InF) { PBDRigidParticles->ExternalForce(ParticleIdx) = InF; }
+	const TVector<T, d>& LinearImpulse() const { return PBDRigidParticles->LinearImpulse(ParticleIdx); }
+	TVector<T, d>& LinearImpulse() { return PBDRigidParticles->LinearImpulse(ParticleIdx); }
+	void SetLinearImpulse(const TVector<T, d>& InLinearImpulse) { PBDRigidParticles->LinearImpulse(ParticleIdx) = InLinearImpulse; }
 
-	const TVector<T, d>& ExternalTorque() const { return PBDRigidParticles->ExternalTorque(ParticleIdx); }
-	TVector<T, d>& ExternalTorque() { return PBDRigidParticles->ExternalTorque(ParticleIdx); }
-	void SetExternalTorque(const TVector<T, d>& InTorque) { PBDRigidParticles->ExternalTorque(ParticleIdx) = InTorque; }
+	const TVector<T, d>& AngularImpulse() const { return PBDRigidParticles->AngularImpulse(ParticleIdx); }
+	TVector<T, d>& AngularImpulse() { return PBDRigidParticles->AngularImpulse(ParticleIdx); }
+	void SetAngularImpulse(const TVector<T, d>& InAngularImpulse) { PBDRigidParticles->AngularImpulse(ParticleIdx) = InAngularImpulse; }
 
 	const PMatrix<T, d, d>& I() const { return PBDRigidParticles->I(ParticleIdx); }
 	PMatrix<T, d, d>& I() { return PBDRigidParticles->I(ParticleIdx); }
@@ -1746,8 +1746,8 @@ public:
 		Ar << MP;
 		Ar << MF;
 		Ar << MTorque;
-		Ar << MExternalForce;
-		Ar << MExternalTorque;
+		Ar << MLinearImpulse;
+		Ar << MAngularImpulse;
 		Ar << MI;
 		Ar << MInvI;
 		Ar << MCollisionParticles;
@@ -1838,18 +1838,18 @@ public:
 		this->MTorque = InTorque;
 	}
 
-	const TVector<T, d>& ExternalForce() const { return MExternalForce; }
-	void SetExternalForce(const TVector<T, d>& InExternalForce, bool bInvalidate = true)
+	const TVector<T, d>& LinearImpulse() const { return MLinearImpulse; }
+	void SetLinearImpulse(const TVector<T, d>& InLinearImpulse, bool bInvalidate = true)
 	{
-		this->MarkDirty(EParticleFlags::ExternalForce, bInvalidate);
-		this->MExternalForce = InExternalForce;
+		this->MarkDirty(EParticleFlags::LinearImpulse, bInvalidate);
+		this->MLinearImpulse = InLinearImpulse;
 	}
 
-	const TVector<T, d>& ExternalTorque() const { return MExternalTorque; }
-	void SetExternalTorque(const TVector<T, d>& InExternalTorque, bool bInvalidate = true)
+	const TVector<T, d>& AngularImpulse() const { return MAngularImpulse; }
+	void SetAngularImpulse(const TVector<T, d>& InAngularImpulse, bool bInvalidate = true)
 	{
-		this->MarkDirty(EParticleFlags::ExternalTorque, bInvalidate);
-		this->MExternalTorque = InExternalTorque;
+		this->MarkDirty(EParticleFlags::AngularImpulse, bInvalidate);
+		this->MAngularImpulse = InAngularImpulse;
 	}
 
 	const PMatrix<T, d, d>& I() const { return MI; }
@@ -1929,8 +1929,8 @@ private:
 	TVector<T, d> MP;
 	TVector<T, d> MF;
 	TVector<T, d> MTorque;
-	TVector<T, d> MExternalForce;
-	TVector<T, d> MExternalTorque;
+	TVector<T, d> MLinearImpulse;
+	TVector<T, d> MAngularImpulse;
 	PMatrix<T, d, d> MI;
 	PMatrix<T, d, d> MInvI;
 	TUniquePtr<TBVHParticles<T, d>> MCollisionParticles;
@@ -1981,8 +1981,10 @@ public:
 		, MPreV(TVector<T, d>(0))
 		, MPreW(TVector<T, d>(0))
 		, MP(TVector<T, d>(0))
-		, MExternalForce(TVector<T, d>(0))
-		, MExternalTorque(TVector<T, d>(0))
+		, MF(TVector<T, d>(0))
+		, MTorque(TVector<T, d>(0))
+		, MLinearImpulse(TVector<T, d>(0))
+		, MAngularImpulse(TVector<T, d>(0))
 		, MI(PMatrix<T, d, d>(0))
 		, MInvI(PMatrix<T, d, d>(0))
 		, MCollisionParticles(nullptr)
@@ -2004,8 +2006,10 @@ public:
 		, MPreV(InParticle.PreV())
 		, MPreW(InParticle.PreW())
 		, MP(InParticle.P())
-		, MExternalForce(InParticle.ExternalForce())
-		, MExternalTorque(InParticle.ExternalTorque())
+		, MF(InParticle.F())
+		, MTorque(InParticle.Torque())
+		, MLinearImpulse(TVector<T, d>(0))
+		, MAngularImpulse(TVector<T, d>(0))
 		, MI(InParticle.I())
 		, MInvI(InParticle.InvI())
 		, MCollisionParticles(nullptr)
@@ -2035,8 +2039,10 @@ public:
 	TVector<T, d> MPreV;
 	TVector<T, d> MPreW;
 	TVector<T, d> MP;
-	TVector<T, d> MExternalForce;
-	TVector<T, d> MExternalTorque;
+	TVector<T, d> MF;
+	TVector<T, d> MTorque;
+	TVector<T, d> MLinearImpulse;
+	TVector<T, d> MAngularImpulse;
 	PMatrix<T, d, d> MI;
 	PMatrix<T, d, d> MInvI;
 	const TBVHParticles<T, d> * MCollisionParticles;
@@ -2059,8 +2065,10 @@ public:
 		MPreV = TVector<T, d>(0);
 		MPreW = TVector<T, d>(0);
 		MP = TVector<T, d>(0);
-		MExternalForce = TVector<T, d>(0);
-		MExternalTorque = TVector<T, d>(0);
+		MF = TVector<T, d>(0);
+		MTorque = TVector<T, d>(0);
+		MLinearImpulse = TVector<T, d>(0);
+		MAngularImpulse = TVector<T, d>(0);
 		MI = PMatrix<T, d, d>(0);
 		MInvI = PMatrix<T, d, d>(0);
 		MCollisionParticles = nullptr;
diff --git a/Engine/Source/Runtime/Experimental/Chaos/Public/Chaos/PerParticleAddImpulses.h b/Engine/Source/Runtime/Experimental/Chaos/Public/Chaos/PerParticleAddImpulses.h
new file mode 100644
--- /dev/null
+++ b/Engine/Source/Runtime/Experimental/Chaos/Public/Chaos/PerParticleAddImpulses.h
@@ -0,0 +1,52 @@
+// Copyright 1998-2019 Epic Games, Inc. All Rights Reserved.
+#pragma once
+
+#include "Chaos/PerParticleRule.h"
+#include "Chaos/Utilities.h"
+
+namespace Chaos
+{
+	template<class T, int d>
+	class TPerParticleAddImpulses : public TPerParticleRule<T, d>
+	{
+	public:
+		TPerParticleAddImpulses() {}
+		virtual ~TPerParticleAddImpulses() {}
+
+		template<class T_PARTICLES>
+		inline void ApplyHelper(T_PARTICLES& InParticles, const T Dt, const int32 Index) const
+		{
+			InParticles.V(Index) += InParticles.LinearImpulse(Index) * InParticles.InvM(Index);
+		}
+
+		inline void Apply(TDynamicParticles<T, d>& InParticles, const T Dt, const int32 Index) const override //-V762
+		{
+			// @todo(mlentine): Is this something we want to support?
+			ensure(false);
+		}
+
+		inline void Apply(TRigidParticles<T, d>& InParticles, const T Dt, const int32 Index) const override //-V762
+		{
+			if (InParticles.InvM(Index) == 0 || InParticles.Disabled(Index) || InParticles.Sleeping(Index))
+				return;
+			ApplyHelper(InParticles, Dt, Index);
+
+			//
+			// TODO: This is the first-order approximation.
+			//       If needed, we might eventually want to do a second order Euler's Equation,
+			//       but if we do that we'll need to do a transform into a rotating reference frame.
+			//       Just using W += InvI * (Torque - W x (I * W)) * dt is not correct, since Torque
+			//		 and W are in an inertial frame.
+			//
+			const FMatrix33 WorldInvI = Utilities::ComputeWorldSpaceInertia(InParticles.R(Index), InParticles.InvI(Index));
+			InParticles.W(Index) += WorldInvI * InParticles.AngularImpulse(Index);
+		}
+
+		inline void Apply(TTransientPBDRigidParticleHandle<T, d>& Particle, const T Dt) const override //-V762
+		{
+			Particle.V() += Particle.LinearImpulse() * Particle.InvM();
+			const FMatrix33 WorldInvI = Utilities::ComputeWorldSpaceInertia(Particle.R(), Particle.InvI());
+			Particle.W() += WorldInvI * Particle.AngularImpulse();
+		}
+	};
+}
diff --git a/Engine/Source/Runtime/Experimental/Chaos/Public/Chaos/RigidParticles.h b/Engine/Source/Runtime/Experimental/Chaos/Public/Chaos/RigidParticles.h
--- a/Engine/Source/Runtime/Experimental/Chaos/Public/Chaos/RigidParticles.h
+++ b/Engine/Source/Runtime/Experimental/Chaos/Public/Chaos/RigidParticles.h
@@ -72,8 +72,8 @@ class TRigidParticles : public TKinematicGeometryParticles<T, d>
 	{
 		TArrayCollection::AddArray(&MF);
 		TArrayCollection::AddArray(&MT);
-		TArrayCollection::AddArray(&MExternalForce);
-		TArrayCollection::AddArray(&MExternalTorque);
+		TArrayCollection::AddArray(&MLinearImpulse);
+		TArrayCollection::AddArray(&MAngularImpulse);
 		TArrayCollection::AddArray(&MI);
 		TArrayCollection::AddArray(&MInvI);
 		TArrayCollection::AddArray(&MM);
@@ -89,12 +89,12 @@ class TRigidParticles : public TKinematicGeometryParticles<T, d>
 	}
 	TRigidParticles(const TRigidParticles<T, d>& Other) = delete;
 	CHAOS_API TRigidParticles(TRigidParticles<T, d>&& Other)
-	    : TKinematicGeometryParticles<T, d>(MoveTemp(Other)), MF(MoveTemp(Other.MF)), MT(MoveTemp(Other.MT)), MExternalForce(MoveTemp(Other.MExternalForce)), MExternalTorque(MoveTemp(Other.MExternalTorque)), MI(MoveTemp(Other.MI)), MInvI(MoveTemp(Other.MInvI)), MM(MoveTemp(Other.MM)), MInvM(MoveTemp(Other.MInvM)), MCollisionParticles(MoveTemp(Other.MCollisionParticles)), MCollisionGroup(MoveTemp(Other.MCollisionGroup)), MObjectState(MoveTemp(Other.MObjectState))
+	    : TKinematicGeometryParticles<T, d>(MoveTemp(Other)), MF(MoveTemp(Other.MF)), MT(MoveTemp(Other.MT)), MLinearImpulse(MoveTemp(Other.MLinearImpulse)), MAngularImpulse(MoveTemp(Other.MAngularImpulse)), MI(MoveTemp(Other.MI)), MInvI(MoveTemp(Other.MInvI)), MM(MoveTemp(Other.MM)), MInvM(MoveTemp(Other.MInvM)), MCollisionParticles(MoveTemp(Other.MCollisionParticles)), MCollisionGroup(MoveTemp(Other.MCollisionGroup)), MObjectState(MoveTemp(Other.MObjectState))
 	{
 		TArrayCollection::AddArray(&MF);
 		TArrayCollection::AddArray(&MT);
-		TArrayCollection::AddArray(&MExternalForce);
-		TArrayCollection::AddArray(&MExternalTorque);
+		TArrayCollection::AddArray(&MLinearImpulse);
+		TArrayCollection::AddArray(&MAngularImpulse);
 		TArrayCollection::AddArray(&MI);
 		TArrayCollection::AddArray(&MInvI);
 		TArrayCollection::AddArray(&MM);
@@ -118,11 +118,11 @@ class TRigidParticles : public TKinematicGeometryParticles<T, d>
 	FORCEINLINE const TVector<T, d>& F(const int32 Index) const { return MF[Index]; }
 	FORCEINLINE TVector<T, d>& F(const int32 Index) { return MF[Index]; }
 
-	FORCEINLINE const TVector<T, d>& ExternalTorque(const int32 Index) const { return MExternalTorque[Index]; }
-	FORCEINLINE TVector<T, d>& ExternalTorque(const int32 Index) { return MExternalTorque[Index]; }
+	FORCEINLINE const TVector<T, d>& LinearImpulse(const int32 Index) const { return MLinearImpulse[Index]; }
+	FORCEINLINE TVector<T, d>& LinearImpulse(const int32 Index) { return MLinearImpulse[Index]; }
 
-	FORCEINLINE const TVector<T, d>& ExternalForce(const int32 Index) const { return MExternalForce[Index]; }
-	FORCEINLINE TVector<T, d>& ExternalForce(const int32 Index) { return MExternalForce[Index]; }
+	FORCEINLINE const TVector<T, d>& AngularImpulse(const int32 Index) const { return MAngularImpulse[Index]; }
+	FORCEINLINE TVector<T, d>& AngularImpulse(const int32 Index) { return MAngularImpulse[Index]; }
 
 	FORCEINLINE const PMatrix<T, d, d>& I(const int32 Index) const { return MI[Index]; }
 	FORCEINLINE PMatrix<T, d, d>& I(const int32 Index) { return MI[Index]; }
@@ -188,14 +188,14 @@ class TRigidParticles : public TKinematicGeometryParticles<T, d>
 	FORCEINLINE FString ToString(int32 index) const
 	{
 		FString BaseString = TKinematicGeometryParticles<T, d>::ToString(index);
-		return FString::Printf(TEXT("%s, MF:%s, MT:%s, MExternalForce:%s, MExternalTorque:%s, MI:%s, MInvI:%s, MM:%f, MInvM:%f, MCollisionParticles(num):%d, MCollisionGroup:%d, MDisabled:%d, MSleepring:%d, MIsland:%d"), *BaseString, *F(index).ToString(), *Torque(index).ToString(), *ExternalForce(index).ToString(), *ExternalTorque(index).ToString(), *I(index).ToString(), *InvI(index).ToString(), M(index), InvM(index), CollisionParticlesSize(index), CollisionGroup(index), Disabled(index), Sleeping(index), Island(index));
+		return FString::Printf(TEXT("%s, MF:%s, MT:%s, MLinearImpulse:%s, MAngularImpulse:%s, MI:%s, MInvI:%s, MM:%f, MInvM:%f, MCollisionParticles(num):%d, MCollisionGroup:%d, MDisabled:%d, MSleepring:%d, MIsland:%d"), *BaseString, *F(index).ToString(), *Torque(index).ToString(), *LinearImpulse(index).ToString(), *AngularImpulse(index).ToString(), *I(index).ToString(), *InvI(index).ToString(), M(index), InvM(index), CollisionParticlesSize(index), CollisionGroup(index), Disabled(index), Sleeping(index), Island(index));
 	}
 
 	CHAOS_API virtual void Serialize(FChaosArchive& Ar) override
 	{
 		LLM_SCOPE(ELLMTag::ChaosParticles);
 		TKinematicGeometryParticles<T,d>::Serialize(Ar);
-		Ar << MF << MT << MExternalForce << MExternalTorque << MI << MInvI << MM << MInvM;
+		Ar << MF << MT << MLinearImpulse << MAngularImpulse << MI << MInvI << MM << MInvM;
 
 		Ar.UsingCustomVersion(FExternalPhysicsCustomObjectVersion::GUID);
 		if (Ar.CustomVer(FExternalPhysicsCustomObjectVersion::GUID) >= FExternalPhysicsCustomObjectVersion::AddDampingToRigids)
@@ -209,8 +209,8 @@ class TRigidParticles : public TKinematicGeometryParticles<T, d>
   private:
 	TArrayCollectionArray<TVector<T, d>> MF;
 	TArrayCollectionArray<TVector<T, d>> MT;
-	TArrayCollectionArray<TVector<T, d>> MExternalForce;
-	TArrayCollectionArray<TVector<T, d>> MExternalTorque;
+	TArrayCollectionArray<TVector<T, d>> MLinearImpulse;
+	TArrayCollectionArray<TVector<T, d>> MAngularImpulse;
 	TArrayCollectionArray<PMatrix<T, d, d>> MI;
 	TArrayCollectionArray<PMatrix<T, d, d>> MInvI;
 	TArrayCollectionArray<T> MM;
diff --git a/Engine/Source/Runtime/Experimental/ChaosSolvers/Private/PhysicsProxy/SingleParticlePhysicsProxy.cpp b/Engine/Source/Runtime/Experimental/ChaosSolvers/Private/PhysicsProxy/SingleParticlePhysicsProxy.cpp
--- a/Engine/Source/Runtime/Experimental/ChaosSolvers/Private/PhysicsProxy/SingleParticlePhysicsProxy.cpp
+++ b/Engine/Source/Runtime/Experimental/ChaosSolvers/Private/PhysicsProxy/SingleParticlePhysicsProxy.cpp
@@ -230,13 +230,13 @@ void FSingleParticlePhysicsProxy<Chaos::TPBDRigidParticle<float, 3>>::PushToPhys
 		RigidHandle->SetLinearEtherDrag(Data->MLinearEtherDrag);
 		RigidHandle->SetAngularEtherDrag(Data->MAngularEtherDrag);
 
-		if (Data->DirtyFlags.IsDirty(Chaos::EParticleFlags::ExternalForce))
+		if (Data->DirtyFlags.IsDirty(Chaos::EParticleFlags::F))
 		{
-			RigidHandle->SetExternalForce(Data->MExternalForce);
+			RigidHandle->SetF(Data->MF);
 		}
-		if (Data->DirtyFlags.IsDirty(Chaos::EParticleFlags::ExternalTorque))
+		if (Data->DirtyFlags.IsDirty(Chaos::EParticleFlags::Torque))
 		{
-			RigidHandle->SetExternalTorque(Data->MExternalTorque);
+			RigidHandle->SetTorque(Data->MTorque);
 		}
 		if (Data->DirtyFlags.IsDirty(Chaos::EParticleFlags::ObjectState))
 		{
@@ -277,8 +277,8 @@ void FSingleParticlePhysicsProxy<Chaos::TPBDRigidParticle<float, 3>>::PushToPhys
 template< >
 void FSingleParticlePhysicsProxy<Chaos::TPBDRigidParticle<float, 3>>::ClearAccumulatedData()
 {
-	Particle->SetExternalForce(Chaos::TVector<float, 3>(0),false);
-	Particle->SetExternalTorque(Chaos::TVector<float, 3>(0),false);
+	Particle->SetF(Chaos::TVector<float, 3>(0));
+	Particle->SetTorque(Chaos::TVector<float, 3>(0));
 	Particle->ClearDirtyFlags();
 }
 
