Make it so implicit object union only allocates hierarchy data when needed

diff --git a/Engine/Source/Runtime/Core/Public/UObject/ExternalPhysicsCustomObjectVersion.h b/Engine/Source/Runtime/Core/Public/UObject/ExternalPhysicsCustomObjectVersion.h
--- a/Engine/Source/Runtime/Core/Public/UObject/ExternalPhysicsCustomObjectVersion.h
+++ b/Engine/Source/Runtime/Core/Public/UObject/ExternalPhysicsCustomObjectVersion.h
@@ -95,6 +95,9 @@ struct CORE_API FExternalPhysicsCustomObjectVersion
 
 		// Trimeshes can now use small indices
 		TrimeshCanUseSmallIndices,
+
+		// Union objects can avoid allocating a full hierarchy
+		UnionObjectsCanAvoidHierarchy,
 		
 		// -----<new versions can be added above this line>-------------------------------------------------
 		VersionPlusOne,
diff --git a/Engine/Source/Runtime/Experimental/Chaos/Private/Chaos/ImplicitObject.cpp b/Engine/Source/Runtime/Experimental/Chaos/Private/Chaos/ImplicitObject.cpp
--- a/Engine/Source/Runtime/Experimental/Chaos/Private/Chaos/ImplicitObject.cpp
+++ b/Engine/Source/Runtime/Experimental/Chaos/Private/Chaos/ImplicitObject.cpp
@@ -392,6 +392,7 @@ FImplicitObject* FImplicitObject::SerializationFactory(FChaosArchive& Ar, FImpli
 	case ImplicitObjectType::Capsule: if (Ar.IsLoading()) { return new TCapsule<FReal>(); } break;
 	case ImplicitObjectType::Transformed: if (Ar.IsLoading()) { return new TImplicitObjectTransformed<FReal, 3>(); } break;
 	case ImplicitObjectType::Union: if (Ar.IsLoading()) { return new FImplicitObjectUnion(); } break;
+	case ImplicitObjectType::UnionClustered: if (Ar.IsLoading()) { return new FImplicitObjectUnionClustered(); } break;
 	case ImplicitObjectType::LevelSet: if (Ar.IsLoading()) { return new TLevelSet<FReal, 3>(); } break;
 	case ImplicitObjectType::Convex: if (Ar.IsLoading()) { return new FConvex(); } break;
 	case ImplicitObjectType::TaperedCylinder: if (Ar.IsLoading()) { return new TTaperedCylinder<FReal>(); } break;
diff --git a/Engine/Source/Runtime/Experimental/Chaos/Private/Chaos/ImplicitObjectUnion.cpp b/Engine/Source/Runtime/Experimental/Chaos/Private/Chaos/ImplicitObjectUnion.cpp
--- a/Engine/Source/Runtime/Experimental/Chaos/Private/Chaos/ImplicitObjectUnion.cpp
+++ b/Engine/Source/Runtime/Experimental/Chaos/Private/Chaos/ImplicitObjectUnion.cpp
@@ -2,15 +2,49 @@
 #include "Chaos/ImplicitObjectUnion.h"
 #include "Chaos/BoundingVolumeHierarchy.h"
 
-using namespace Chaos;
+namespace Chaos
+{
+struct FLargeImplicitObjectUnionData
+{
+	TGeometryParticles<FReal,3> GeomParticles;
+	TBoundingVolumeHierarchy<TGeometryParticles<FReal,3>, TArray<int32>, FReal, 3> Hierarchy;
+
+	FLargeImplicitObjectUnionData(const TArray<Pair<TSerializablePtr<FImplicitObject>,FRigidTransform3>>& SubObjects)
+	{
+		const int32 NumObjects = SubObjects.Num();
+		GeomParticles.Resize(NumObjects);
+		for (int32 i = 0; i < NumObjects; ++i)
+		{
+			GeomParticles.X(i) = SubObjects[i].Second.GetLocation();
+			GeomParticles.R(i) = SubObjects[i].Second.GetRotation();
+			GeomParticles.SetGeometry(i, SubObjects[i].First);
+			//check(!SubObjects[i].First->IsUnderlyingUnion());	//we don't support union of unions
+		}
 
-FImplicitObjectUnion::FImplicitObjectUnion(TArray<TUniquePtr<FImplicitObject>>&& Objects, const TArray<int32>& OriginalParticleLookupHack)
+		Hierarchy = TBoundingVolumeHierarchy<TGeometryParticles<FReal,3>,TArray<int32>,FReal,3> (GeomParticles,1);
+	}
+
+	void Serialize(FChaosArchive& Ar)
+	{
+		Ar << GeomParticles << Hierarchy;
+	}
+
+	FLargeImplicitObjectUnionData(){}
+
+	FLargeImplicitObjectUnionData(const FLargeImplicitObjectUnionData& Other) = delete;
+	FLargeImplicitObjectUnionData& operator=(const FLargeImplicitObjectUnionData& Other) = delete;
+};
+
+FChaosArchive& operator<<(FChaosArchive& Ar, FLargeImplicitObjectUnionData& LargeUnionData)
+{
+	LargeUnionData.Serialize(Ar);
+	return Ar;
+}
+
+FImplicitObjectUnion::FImplicitObjectUnion(TArray<TUniquePtr<FImplicitObject>>&& Objects)
 	: FImplicitObject(EImplicitObject::HasBoundingBox, ImplicitObjectType::Union)
 	, MObjects(MoveTemp(Objects))
-	, Hierarchy(new TBoundingVolumeHierarchy<TGeometryParticles<FReal,3>, TArray<int32>, FReal, 3>(GeomParticles))
 	, MLocalBoundingBox()
-	, bHierarchyBuilt(false)
-	, MOriginalParticleLookupHack(OriginalParticleLookupHack)
 {
 	ensure(MObjects.Num());
 	for (int32 i = 0; i < MObjects.Num(); ++i)
@@ -23,51 +57,31 @@ FImplicitObjectUnion::FImplicitObjectUnion(TArray<TUniquePtr<FImplicitObject>>&&
 		{
 			MLocalBoundingBox = MObjects[i]->BoundingBox();
 		}
-		check(MOriginalParticleLookupHack.Num() == 0 || MOriginalParticleLookupHack.Num() == MObjects.Num());
-		if (MOriginalParticleLookupHack.Num() > 0)
-		{
-			//this whole part sucks, only needed because of how we get union children. Need to refactor and enforce no unions of unions
-			if (const TImplicitObjectTransformed<FReal,3>* Transformed = MObjects[i]->template GetObject<const TImplicitObjectTransformed<FReal,3>>())
-			{
-				MCollisionParticleLookupHack.Add(Transformed->GetTransformedObject(), MOriginalParticleLookupHack[i]);
-			}
-			else
-			{
-				ensure(false);	//shouldn't be here
-			}
-		}
 	}
 
 	CacheAllImplicitObjects();
 }
 
 FImplicitObjectUnion::FImplicitObjectUnion(FImplicitObjectUnion&& Other)
-	: FImplicitObject(EImplicitObject::HasBoundingBox)
+	: FImplicitObject(EImplicitObject::HasBoundingBox, ImplicitObjectType::Union)
 	, MObjects(MoveTemp(Other.MObjects))
-	, GeomParticles(MoveTemp(Other.GeomParticles))
-	, Hierarchy(new TBoundingVolumeHierarchy<TGeometryParticles<FReal,3>, TArray<int32>, FReal, 3>(MoveTemp(*Other.Hierarchy)))
 	, MLocalBoundingBox(MoveTemp(Other.MLocalBoundingBox))
-	, bHierarchyBuilt(Other.bHierarchyBuilt)
-	, MOriginalParticleLookupHack(MoveTemp(Other.MOriginalParticleLookupHack))
-	, MCollisionParticleLookupHack(MoveTemp(Other.MCollisionParticleLookupHack))
+	, LargeUnionData(MoveTemp(Other.LargeUnionData))
 {
 }
 
-FImplicitObjectUnion::~FImplicitObjectUnion()
-{
-	delete Hierarchy;
-}
+FImplicitObjectUnion::~FImplicitObjectUnion() = default;
 
 void FImplicitObjectUnion::FindAllIntersectingObjects(TArray < Pair<const FImplicitObject*,FRigidTransform3>>& Out, const TAABB<FReal,3>& LocalBounds) const
 {
-	if (bHierarchyBuilt)
+	if (LargeUnionData)
 	{
-		TArray<int32> Overlaps = Hierarchy->FindAllIntersections(LocalBounds);
+		TArray<int32> Overlaps = LargeUnionData->Hierarchy.FindAllIntersections(LocalBounds);
 		Out.Reserve(Out.Num() + Overlaps.Num());
 		for (int32 Idx : Overlaps)
 		{
-			const FImplicitObject* Obj = GeomParticles.Geometry(Idx).Get();
-			Out.Add(MakePair(Obj,FRigidTransform3(GeomParticles.X(Idx), GeomParticles.R(Idx))));
+			const FImplicitObject* Obj = LargeUnionData->GeomParticles.Geometry(Idx).Get();
+			Out.Add(MakePair(Obj,FRigidTransform3(LargeUnionData->GeomParticles.X(Idx), LargeUnionData->GeomParticles.R(Idx))));
 		}
 	}
 	else
@@ -79,12 +93,12 @@ void FImplicitObjectUnion::FindAllIntersectingObjects(TArray < Pair<const FImpli
 	}
 }
 
-TArray<int32> FImplicitObjectUnion::FindAllIntersectingChildren(const TAABB<FReal,3>& LocalBounds) const
+TArray<int32> FImplicitObjectUnionClustered::FindAllIntersectingChildren(const TAABB<FReal,3>& LocalBounds) const
 {
 	TArray<int32> IntersectingChildren;
-	if (bHierarchyBuilt) //todo: make this work when hierarchy is not built
+	if (LargeUnionData) //todo: make this work when hierarchy is not built
 	{
-		IntersectingChildren = Hierarchy->FindAllIntersections(LocalBounds);
+		IntersectingChildren = LargeUnionData->Hierarchy.FindAllIntersections(LocalBounds);
 		for (int32 i = IntersectingChildren.Num() - 1; i >= 0; --i)
 		{
 			const int32 Idx = IntersectingChildren[i];
@@ -110,12 +124,12 @@ TArray<int32> FImplicitObjectUnion::FindAllIntersectingChildren(const TAABB<FRea
 	return IntersectingChildren;
 }
 
-TArray<int32> FImplicitObjectUnion::FindAllIntersectingChildren(const TSpatialRay<FReal,3>& LocalRay) const
+TArray<int32> FImplicitObjectUnionClustered::FindAllIntersectingChildren(const TSpatialRay<FReal,3>& LocalRay) const
 {
 	TArray<int32> IntersectingChildren;
-	if (bHierarchyBuilt) //todo: make this work when hierarchy is not built
+	if (LargeUnionData) //todo: make this work when hierarchy is not built
 	{
-		IntersectingChildren = Hierarchy->FindAllIntersections(LocalRay);
+		IntersectingChildren = LargeUnionData->Hierarchy.FindAllIntersections(LocalRay);
 		for (int32 i = IntersectingChildren.Num() - 1; i >= 0; --i)
 		{
 			const int32 Idx = IntersectingChildren[i];
@@ -148,32 +162,75 @@ void FImplicitObjectUnion::CacheAllImplicitObjects()
 	//build hierarchy
 	{
 		const int32 NumObjects = SubObjects.Num();
-		constexpr int32 MinSubObjectsToCache = 8;	//todo(make this tunable?)
+		constexpr int32 MinSubObjectsToCache = 32;	//todo(make this tunable?)
 		if (NumObjects > MinSubObjectsToCache)
 		{
-			GeomParticles.Resize(NumObjects);
-			for (int32 i = 0; i < NumObjects; ++i)
-			{
-				GeomParticles.X(i) = SubObjects[i].Second.GetLocation();
-				GeomParticles.R(i) = SubObjects[i].Second.GetRotation();
-				GeomParticles.SetGeometry(i, SubObjects[i].First);
-				//check(!SubObjects[i].First->IsUnderlyingUnion());	//we don't support union of unions
-			}
-
-			TBoundingVolumeHierarchy<TGeometryParticles<FReal,3>, TArray<int32>, FReal, 3> NewHierarchy(GeomParticles, 1);
-			*Hierarchy = MoveTemp(NewHierarchy);
-			bHierarchyBuilt = true;
+			LargeUnionData = MakeUnique<FLargeImplicitObjectUnionData>(SubObjects);
 		}
 	}
 }
 
 void FImplicitObjectUnion::Serialize(FChaosArchive& Ar)
 {
+	Ar.UsingCustomVersion(FExternalPhysicsCustomObjectVersion::GUID);
 	FChaosArchiveScopedMemory ScopedMemory(Ar, GetTypeName(), false);
 	FImplicitObject::SerializeImp(Ar);
 	Ar << MObjects;
 	TBox<FReal,3>::SerializeAsAABB(Ar, MLocalBoundingBox);
-	Ar << GeomParticles << *Hierarchy << bHierarchyBuilt;
+
+	bool bHierarchyBuilt = LargeUnionData.Get() != nullptr;
+	if(Ar.CustomVer(FExternalPhysicsCustomObjectVersion::GUID) < FExternalPhysicsCustomObjectVersion::UnionObjectsCanAvoidHierarchy)
+	{
+		LargeUnionData = MakeUnique<FLargeImplicitObjectUnionData>();
+		Ar << *LargeUnionData;
+		Ar << bHierarchyBuilt;
+	}
+	else
+	{
+		Ar << bHierarchyBuilt;
+		if(bHierarchyBuilt)
+		{
+			if(Ar.IsLoading())
+			{
+				LargeUnionData = MakeUnique<FLargeImplicitObjectUnionData>();
+			}
+			Ar << *LargeUnionData;
+		}
+	}
+}
+
+FImplicitObjectUnion::FImplicitObjectUnion() : FImplicitObject(EImplicitObject::HasBoundingBox, ImplicitObjectType::Union){}
+
+
+FImplicitObjectUnionClustered::FImplicitObjectUnionClustered(TArray<TUniquePtr<FImplicitObject>>&& Objects, const TArray<int32>& OriginalParticleLookupHack)
+	: FImplicitObjectUnion(MoveTemp(Objects))
+	, MOriginalParticleLookupHack(OriginalParticleLookupHack)
+{
+	Type = ImplicitObjectType::UnionClustered;
+	for (int32 i = 0; i < MObjects.Num(); ++i)
+	{
+		check(MOriginalParticleLookupHack.Num() == 0 || MOriginalParticleLookupHack.Num() == MObjects.Num());
+		if (MOriginalParticleLookupHack.Num() > 0)
+		{
+			//this whole part sucks, only needed because of how we get union children. Need to refactor and enforce no unions of unions
+			if (const TImplicitObjectTransformed<FReal,3>* Transformed = MObjects[i]->template GetObject<const TImplicitObjectTransformed<FReal,3>>())
+			{
+				MCollisionParticleLookupHack.Add(Transformed->GetTransformedObject(), MOriginalParticleLookupHack[i]);
+			}
+			else
+			{
+				ensure(false);	//shouldn't be here
+			}
+		}
+	}
 }
 
-FImplicitObjectUnion::FImplicitObjectUnion() : FImplicitObject(EImplicitObject::HasBoundingBox, ImplicitObjectType::Union), Hierarchy(new TBoundingVolumeHierarchy<TGeometryParticles<FReal,3>, TArray<int32>, FReal, 3>(GeomParticles, 1)){}
+FImplicitObjectUnionClustered::FImplicitObjectUnionClustered(FImplicitObjectUnionClustered&& Other)
+: FImplicitObjectUnion(MoveTemp(Other))
+, MOriginalParticleLookupHack(MoveTemp(MOriginalParticleLookupHack))
+, MCollisionParticleLookupHack(MoveTemp(MCollisionParticleLookupHack))
+{
+	Type = ImplicitObjectType::UnionClustered;
+}
+
+}
diff --git a/Engine/Source/Runtime/Experimental/Chaos/Private/Chaos/PBDRigidClustering.cpp b/Engine/Source/Runtime/Experimental/Chaos/Private/Chaos/PBDRigidClustering.cpp
--- a/Engine/Source/Runtime/Experimental/Chaos/Private/Chaos/PBDRigidClustering.cpp
+++ b/Engine/Source/Runtime/Experimental/Chaos/Private/Chaos/PBDRigidClustering.cpp
@@ -1360,7 +1360,7 @@ namespace Chaos
 
 		{
 			QUICK_SCOPE_CYCLE_COUNTER(SpatialBVH);
-			MParticles.ChildrenSpatial(NewIndex) = Objects2.Num() ? MakeUnique<FImplicitObjectUnion>(MoveTemp(Objects2), GeomToOriginalParticlesHack) : nullptr;
+			MParticles.ChildrenSpatial(NewIndex) = Objects2.Num() ? MakeUnique<FImplicitObjectUnionClustered>(MoveTemp(Objects2), GeomToOriginalParticlesHack) : nullptr;
 		}
 
 		TArray<TVector<T, d>> CleanedPoints;
@@ -1444,7 +1444,7 @@ namespace Chaos
 					QUICK_SCOPE_CYCLE_COUNTER(UnionBVH);
 					// @coverage : { confidence tests}
 					//ensureMsgf(false, TEXT("Checking no proxy, not levelset, and multiple objects"));
-					MParticles.SetDynamicGeometry(NewIndex, MakeUnique<FImplicitObjectUnion>(MoveTemp(Objects), GeomToOriginalParticlesHack));
+					MParticles.SetDynamicGeometry(NewIndex, MakeUnique<FImplicitObjectUnionClustered>(MoveTemp(Objects), GeomToOriginalParticlesHack));
 				}
 			}
 		}
diff --git a/Engine/Source/Runtime/Experimental/Chaos/Public/Chaos/ImplicitObject.h b/Engine/Source/Runtime/Experimental/Chaos/Public/Chaos/ImplicitObject.h
--- a/Engine/Source/Runtime/Experimental/Chaos/Public/Chaos/ImplicitObject.h
+++ b/Engine/Source/Runtime/Experimental/Chaos/Public/Chaos/ImplicitObject.h
@@ -47,6 +47,7 @@ namespace ImplicitObjectType
 		HeightField,
 		DEPRECATED_Scaled,	//needed for serialization of existing data
 		Triangle,
+		UnionClustered,
 
 		//Add entries above this line for serialization
 		IsInstanced = 1 << 6,
diff --git a/Engine/Source/Runtime/Experimental/Chaos/Public/Chaos/ImplicitObjectUnion.h b/Engine/Source/Runtime/Experimental/Chaos/Public/Chaos/ImplicitObjectUnion.h
--- a/Engine/Source/Runtime/Experimental/Chaos/Public/Chaos/ImplicitObjectUnion.h
+++ b/Engine/Source/Runtime/Experimental/Chaos/Public/Chaos/ImplicitObjectUnion.h
@@ -12,16 +12,15 @@
 namespace Chaos
 {
 
-template<class OBJECT_ARRAY, class LEAF_TYPE, class T, int d>
-class TBoundingVolumeHierarchy;
+struct FLargeImplicitObjectUnionData;
 
-class FImplicitObjectUnion : public FImplicitObject
+class CHAOS_API FImplicitObjectUnion : public FImplicitObject
 {
   public:
 
 	using FImplicitObject::GetTypeName;
 
-	CHAOS_API FImplicitObjectUnion(TArray<TUniquePtr<FImplicitObject>>&& Objects, const TArray<int32>& OriginalParticleLookupHack = TArray<int32>());
+	FImplicitObjectUnion(TArray<TUniquePtr<FImplicitObject>>&& Objects);
 	FImplicitObjectUnion(const FImplicitObjectUnion& Other) = delete;
 	FImplicitObjectUnion(FImplicitObjectUnion&& Other);
 	virtual ~FImplicitObjectUnion();
@@ -86,8 +85,6 @@ class FImplicitObjectUnion : public FImplicitObject
 	}
 
 	virtual void FindAllIntersectingObjects(TArray < Pair<const FImplicitObject*, FRigidTransform3>>& Out, const TAABB<FReal,3>& LocalBounds) const;
-	TArray<int32> FindAllIntersectingChildren(const TAABB<FReal,3>& LocalBounds) const;
-	TArray<int32> FindAllIntersectingChildren(const TSpatialRay<FReal,3>& LocalRay) const;
 	virtual void CacheAllImplicitObjects();
 
 	virtual bool Raycast(const FVec3& StartPoint, const FVec3& Dir, const FReal Length, const FReal Thickness, FReal& OutTime, FVec3& OutPosition, FVec3& OutNormal, int32& OutFaceIndex) const override
@@ -155,7 +152,7 @@ class FImplicitObjectUnion : public FImplicitObject
 		return Result;
 	}
 
-private:
+protected:
 	virtual Pair<FVec3, bool> FindClosestIntersectionImp(const FVec3& StartPoint, const FVec3& EndPoint, const FReal Thickness) const override
 	{
 		check(MObjects.Num());
@@ -176,19 +173,41 @@ private:
 		return ClosestIntersection;
 	}
 
-  private:
+  protected:
 	TArray<TUniquePtr<FImplicitObject>> MObjects;
-	TGeometryParticles<FReal,3> GeomParticles;
-	TBoundingVolumeHierarchy<TGeometryParticles<FReal,3>, TArray<int32>, FReal, 3>* Hierarchy;
 	TAABB<FReal,3> MLocalBoundingBox;
-	bool bHierarchyBuilt;
-	TArray<int32> MOriginalParticleLookupHack;	//temp hack for finding original particles
-
+	TUniquePtr<FLargeImplicitObjectUnionData> LargeUnionData;	//only needed when there are many objects
 
 	//needed for serialization
 	FImplicitObjectUnion();
 	friend FImplicitObject;	//needed for serialization
 public:
-	TMap<const FImplicitObject*, int32> MCollisionParticleLookupHack;	//temp hack for finding collision particles
+};
+
+class CHAOS_API FImplicitObjectUnionClustered: public FImplicitObjectUnion
+{
+public:
+
+	FImplicitObjectUnionClustered(TArray<TUniquePtr<FImplicitObject>>&& Objects,const TArray<int32>& OriginalParticleLookupHack = TArray<int32>());
+	FImplicitObjectUnionClustered()
+		: FImplicitObjectUnion()
+	{
+		Type = ImplicitObjectType::UnionClustered;
+	}
+	FImplicitObjectUnionClustered(const FImplicitObjectUnionClustered& Other) = delete;
+	FImplicitObjectUnionClustered(FImplicitObjectUnionClustered&& Other);
+	virtual ~FImplicitObjectUnionClustered() = default;
+
+	FORCEINLINE static constexpr EImplicitObjectType StaticType()
+	{
+		return ImplicitObjectType::UnionClustered;
+	}
+
+	TArray<int32> FindAllIntersectingChildren(const TAABB<FReal,3>& LocalBounds) const;
+	TArray<int32> FindAllIntersectingChildren(const TSpatialRay<FReal,3>& LocalRay) const;
+
+private:
+	TArray<int32> MOriginalParticleLookupHack;	//temp hack for finding original particles
+	TMap<const FImplicitObject*,int32> MCollisionParticleLookupHack;	//temp hack for finding collision particles
 };
 }
diff --git a/Engine/Source/Runtime/Experimental/Chaos/Public/Chaos/PBDRigidClusteredParticles.h b/Engine/Source/Runtime/Experimental/Chaos/Public/Chaos/PBDRigidClusteredParticles.h
--- a/Engine/Source/Runtime/Experimental/Chaos/Public/Chaos/PBDRigidClusteredParticles.h
+++ b/Engine/Source/Runtime/Experimental/Chaos/Public/Chaos/PBDRigidClusteredParticles.h
@@ -7,7 +7,7 @@
 
 namespace Chaos
 {
-	class FImplicitObjectUnion;
+	class FImplicitObjectUnionClustered;
 
 // ClusterId 
 //    Used within the clustering system to describe the clustering hierarchy. The ID will
@@ -159,7 +159,7 @@ class TPBDRigidClusteredParticles : public TPBDRigidParticles<T, d>
 	  TArrayCollectionArray<TRigidTransform<T, d>> MChildToParent;
 	  TArrayCollectionArray<int32> MClusterGroupIndex;
 	  TArrayCollectionArray<bool> MInternalCluster;
-	  TArrayCollectionArray<TUniquePtr<FImplicitObjectUnion>> MChildrenSpatial;
+	  TArrayCollectionArray<TUniquePtr<FImplicitObjectUnionClustered>> MChildrenSpatial;
 	  TArrayCollectionArray<FMultiChildProxyId> MMultiChildProxyId;
 	  TArrayCollectionArray<TUniquePtr<TMultiChildProxyData<T, d>>> MMultiChildProxyData;
 
