Merging CL#107977736

Missing file

diff --git a/Engine/Source/Runtime/Experimental/Chaos/Public/Chaos/BoundingVolumeHierarchy.h b/Engine/Source/Runtime/Experimental/Chaos/Public/Chaos/BoundingVolumeHierarchy.h
--- a/Engine/Source/Runtime/Experimental/Chaos/Public/Chaos/BoundingVolumeHierarchy.h
+++ b/Engine/Source/Runtime/Experimental/Chaos/Public/Chaos/BoundingVolumeHierarchy.h
@@ -65,6 +65,11 @@ class TBoundingVolumeHierarchy final : public ISpatialAcceleration<int32, T,d>
 	static constexpr T DefaultDt = 0;
 	using TPayloadType = typename TBVHLeafTraits<LEAF_TYPE, TModels<CComplexBVHLeaf, LEAF_TYPE>::Value>::TPayloadType;
 
+	CHAOS_API TBoundingVolumeHierarchy()
+		: MObjects(nullptr)
+	{
+	}
+
 	CHAOS_API TBoundingVolumeHierarchy(const OBJECT_ARRAY& Objects, const int32 MaxLevels = DefaultMaxLevels, const bool bUseVelocity = DefaultUseVelocity, const T Dt = DefaultDt);
 	CHAOS_API TBoundingVolumeHierarchy(const OBJECT_ARRAY& Objects, const TArray<uint32>& ActiveIndices, const int32 MaxLevels = DefaultMaxLevels, const bool bUseVelocity = DefaultUseVelocity, const T Dt = DefaultDt);
 	
