Chaos : Collisions

Hooked up Manifold callbacks for MultiPointConstraints.

diff --git a/Engine/Source/Runtime/Experimental/Chaos/Private/Chaos/CollisionResolution.cpp b/Engine/Source/Runtime/Experimental/Chaos/Private/Chaos/CollisionResolution.cpp
--- a/Engine/Source/Runtime/Experimental/Chaos/Private/Chaos/CollisionResolution.cpp
+++ b/Engine/Source/Runtime/Experimental/Chaos/Private/Chaos/CollisionResolution.cpp
@@ -914,15 +914,15 @@ namespace Chaos
 		template<class T, int d>
 		void UpdateConvexConvexManifold(TCollisionConstraintBase<T, d>&  ConstraintBase, const TRigidTransform<T, d>& Transform0, const TRigidTransform<T, d>& Transform1, const T Thickness)
 		{
-			if (ensure(ConstraintBase.GetType() == FRigidBodyMultiPointContactConstraint::StaticType()))
+			if (TRigidBodyMultiPointContactConstraint<T,d>* Constraint = ConstraintBase.template As<TRigidBodyMultiPointContactConstraint<T, d>>())
 			{
 				if (GetInnerType(ConstraintBase.Manifold.Implicit[0]->GetType()) == ImplicitObjectType::Convex)
 				{
-					UpdateSingleShotManifold(*ConstraintBase.template As<TRigidBodyMultiPointContactConstraint<T, d>>(), Transform0, Transform1, Thickness);
+					UpdateSingleShotManifold(*Constraint, Transform0, Transform1, Thickness);
 				}
 				else
 				{
-					UpdateIterativeManifold(*ConstraintBase.template As<TRigidBodyMultiPointContactConstraint<T, d>>(), Transform0, Transform1, Thickness);
+					UpdateIterativeManifold(*Constraint, Transform0, Transform1, Thickness);
 				}
 			}
 		}
@@ -1124,11 +1124,114 @@ namespace Chaos
 			EImplicitObjectType Implicit0Type = Implicit0.GetType();
 			EImplicitObjectType Implicit1Type = Implicit1.GetType();
 
-			// TConvex
-			if (Implicit0.IsConvex() && Implicit1.IsConvex())
+			if (Implicit0Type == TImplicitObjectTransformed<T, d>::StaticType())
+			{
+				ensure(false);//should not be possible to get this type, it should already be resolved by the constraint. (see ConstructConstraints)
+			}
+			else if (Implicit1Type == TImplicitObjectTransformed<T, d>::StaticType())
+			{
+				ensure(false);//should not be possible to get this type, it should already be resolved by the constraint. (see ConstructConstraints)
+			}
+			else if (Implicit0Type != FImplicitObjectUnion::StaticType() && Implicit1Type == FImplicitObjectUnion::StaticType())
+			{
+				ensure(false);//should not be possible to get this type, it should already be resolved by the constraint. (see ConstructConstraints)
+			}
+			else if (Implicit0Type == FImplicitObjectUnion::StaticType() && Implicit1Type != FImplicitObjectUnion::StaticType())
+			{
+				ensure(false);//should not be possible to get this type, it should already be resolved by the constraint. (see ConstructConstraints)
+			}
+			else if (Implicit0Type == FImplicitObjectUnion::StaticType() && Implicit1Type == FImplicitObjectUnion::StaticType())
+			{
+				ensure(false);//should not be possible to get this type, it should already be resolved by the constraint. (see ConstructConstraints)
+			}
+			else if (Implicit0Type == TBox<T, d>::StaticType() && Implicit1Type == TBox<T, d>::StaticType())
+			{
+				UpdateBoxBoxManifold(ConstraintBase, Transform0, Transform1, Thickness);
+			}
+			else if (Implicit0Type == TBox<T, d>::StaticType() && Implicit1Type == THeightField<T>::StaticType())
+			{
+				UpdateBoxHeightFieldManifold(ConstraintBase, Transform0, Transform1, Thickness);
+			}
+			else if (Implicit0Type == TSphere<T, d>::StaticType() && Implicit1Type == TSphere<T, d>::StaticType())
+			{
+				UpdateSphereSphereManifold(ConstraintBase, Transform0, Transform1, Thickness);
+			}
+			else if (Implicit0Type == TSphere<T, d>::StaticType() && Implicit1Type == THeightField<T>::StaticType())
+			{
+				UpdateSphereHeightFieldManifold(ConstraintBase, Transform0, Transform1, Thickness);
+			}
+			else if (Implicit0Type == TBox<T, d>::StaticType() && Implicit1Type == TPlane<T, d>::StaticType())
+			{
+				UpdateBoxPlaneManifold(ConstraintBase, Transform0, Transform1, Thickness);
+			}
+			else if (Implicit0Type == TSphere<T, d>::StaticType() && Implicit1Type == TPlane<T, d>::StaticType())
+			{
+				UpdateSpherePlaneManifold(ConstraintBase, Transform0, Transform1, Thickness);
+			}
+			else if (Implicit0Type == TSphere<T, d>::StaticType() && Implicit1Type == TBox<T, d>::StaticType())
+			{
+				UpdateSphereBoxManifold(ConstraintBase, Transform0, Transform1, Thickness);
+			}
+			else if (Implicit0Type == TSphere<T, d>::StaticType() && Implicit1Type == TCapsule<T>::StaticType())
+			{
+				UpdateSphereCapsuleManifold(ConstraintBase, Transform0, Transform1, Thickness);
+			}
+			else if (Implicit0Type == TCapsule<T>::StaticType() && Implicit1Type == TCapsule<T>::StaticType())
+			{
+				UpdateCapsuleCapsuleManifold(ConstraintBase, Transform0, Transform1, Thickness);
+			}
+			else if (Implicit0Type == TCapsule<T>::StaticType() && Implicit1Type == TBox<T, d>::StaticType())
+			{
+				UpdateCapsuleBoxManifold(ConstraintBase, Transform0, Transform1, Thickness);
+			}
+			else if (Implicit0Type == TCapsule<T>::StaticType() && Implicit1Type == THeightField<T>::StaticType())
+			{
+				UpdateCapsuleHeightFieldManifold(ConstraintBase, Transform0, Transform1, Thickness);
+			}
+			else if (Implicit0Type == THeightField<T>::StaticType() && Implicit1Type == TBox<T, d>::StaticType())
+			{
+				ensure(false);
+			}
+			else if (Implicit0Type == TPlane<T, d>::StaticType() && Implicit1Type == TBox<T, d>::StaticType())
+			{
+				//UpdatePlaneBoxManifold(ConstraintBase, Transform0, Transform1, Thickness);
+			}
+			else if (Implicit0Type == THeightField<T>::StaticType() && Implicit1Type == TSphere<T, d>::StaticType())
+			{
+				ensure(false);
+			}
+			else if (Implicit0Type == TPlane<T, d>::StaticType() && Implicit1Type == TSphere<T, d>::StaticType())
+			{
+				//UpdatePlaneSphereManifold(ConstraintBase, Transform0, Transform1, Thickness);
+			}
+			else if (Implicit0Type == TBox<T, d>::StaticType() && Implicit1Type == TSphere<T, d>::StaticType())
+			{
+				//UpdateBoxSphereManifold(ConstraintBase, Transform0, Transform1, Thickness);
+			}
+			else if (Implicit0Type == TBox<T, d>::StaticType() && Implicit1Type == TCapsule<T>::StaticType())
+			{
+				//UpdateBoxSphereManifold(ConstraintBase, Transform0, Transform1, Thickness);
+			}
+			else if (Implicit0Type == TCapsule<T>::StaticType() && Implicit1Type == TSphere<T, d>::StaticType())
+			{
+				//UpdateCapsuleSphereManifold(ConstraintBase, Transform0, Transform1, Thickness);
+			}
+			else if (Implicit0Type == THeightField<T>::StaticType() && Implicit1.IsConvex())
+			{
+				ensure(false);
+			}
+			else if (Implicit0.IsConvex() && Implicit1Type == THeightField<T>::StaticType())
+			{
+				UpdateConvexHeightFieldManifold(ConstraintBase, Transform0, Transform1, Thickness);
+			}
+			else if (Implicit0.IsConvex() && Implicit1.IsConvex())
 			{
 				UpdateConvexConvexManifold(ConstraintBase, Transform0, Transform1, Thickness);
 			}
+			else
+			{
+				UpdateLevelsetLevelsetManifold(ConstraintBase, Transform0, Transform1, Thickness);
+			}
 		}
 
 		template<ECollisionUpdateType UpdateType, typename T, int d>
diff --git a/Engine/Source/Runtime/Experimental/Chaos/Private/Chaos/PBDCollisionConstraints.cpp b/Engine/Source/Runtime/Experimental/Chaos/Private/Chaos/PBDCollisionConstraints.cpp
--- a/Engine/Source/Runtime/Experimental/Chaos/Private/Chaos/PBDCollisionConstraints.cpp
+++ b/Engine/Source/Runtime/Experimental/Chaos/Private/Chaos/PBDCollisionConstraints.cpp
@@ -270,13 +270,11 @@ namespace Chaos
 	template<typename T, int d>
 	void TPBDCollisionConstraints<T, d>::UpdateManifolds(T Dt)
 	{
-		PhysicsParallelFor(IterativeConstraints.Num(), [&](int32 ConstraintIndex)
+		PhysicsParallelFor(Handles.Num(), [&](int32 ConstraintHandleIndex)
 		{
-			FConstraintBase& ConstraintBase = IterativeConstraints[ConstraintIndex];
-			if (ConstraintBase.GetType() == FCollisionConstraintBase::FType::MultiPoint)
-			{
-				Collisions::UpdateManifold<float, 3>(MThickness, ConstraintBase);
-			}
+			FConstraintContainerHandle* ConstraintHandle = Handles[ConstraintHandleIndex];
+			check(ConstraintHandle != nullptr);
+			Collisions::UpdateManifold<float, 3>(MThickness, ConstraintHandle->GetContact());
 		}, bDisableCollisionParallelFor);
 	}
 
diff --git a/Engine/Source/Runtime/Experimental/Chaos/Private/Chaos/PBDCollisionConstraintsContact.cpp b/Engine/Source/Runtime/Experimental/Chaos/Private/Chaos/PBDCollisionConstraintsContact.cpp
--- a/Engine/Source/Runtime/Experimental/Chaos/Private/Chaos/PBDCollisionConstraintsContact.cpp
+++ b/Engine/Source/Runtime/Experimental/Chaos/Private/Chaos/PBDCollisionConstraintsContact.cpp
@@ -24,13 +24,10 @@ namespace Chaos
 		template<typename T, int d>
 		void UpdateManifold(const T Thickness, TCollisionConstraintBase<T, d>& Constraint)
 		{
-			if (Constraint.GetType() == TCollisionConstraintBase<T, d>::FType::MultiPoint)
-			{
-				const TRigidTransform<T, d> Transform0 = GetTransform(Constraint.Particle[0]);
-				const TRigidTransform<T, d> Transform1 = GetTransform(Constraint.Particle[1]);
+			const TRigidTransform<T, d> Transform0 = GetTransform(Constraint.Particle[0]);
+			const TRigidTransform<T, d> Transform1 = GetTransform(Constraint.Particle[1]);
 
-				UpdateManifold(Constraint, Transform0, Transform1, Thickness);
-			}
+			UpdateManifold(Constraint, Transform0, Transform1, Thickness);
 		}
 
 
