Remove sleep events when particle is destroyed. Speculative fix for bad particle access when processing sleep events during editor tick.

diff --git a/Engine/Source/Runtime/Experimental/Chaos/Public/Chaos/PBDRigidsSOAs.h b/Engine/Source/Runtime/Experimental/Chaos/Public/Chaos/PBDRigidsSOAs.h
--- a/Engine/Source/Runtime/Experimental/Chaos/Public/Chaos/PBDRigidsSOAs.h
+++ b/Engine/Source/Runtime/Experimental/Chaos/Public/Chaos/PBDRigidsSOAs.h
@@ -112,7 +112,25 @@ public:
 			RemoveFromMapAndArray(PBDRigid, ActiveParticlesToIndex, ActiveParticlesArray);
 		}
 
+
+		if (PBDRigid)
+		{
+			// Check for sleep events referencing this particle
+			// TODO think about this case more
+			auto& SleepData = GetDynamicParticles().GetSleepData();
+			for (int32 Idx = 0; Idx < SleepData.Num(); ++Idx)
+			{
+				if (SleepData[Idx].Particle == Particle)
+				{
+					SleepData.RemoveAtSwap(Idx);
+					break;
+				}
+			}
+		}
+
 		ParticleHandles.DestroyHandleSwap(Particle);
+
+		
 		UpdateViews();
 	}
 
diff --git a/Engine/Source/Runtime/Experimental/Chaos/Public/Chaos/RigidParticles.h b/Engine/Source/Runtime/Experimental/Chaos/Public/Chaos/RigidParticles.h
--- a/Engine/Source/Runtime/Experimental/Chaos/Public/Chaos/RigidParticles.h
+++ b/Engine/Source/Runtime/Experimental/Chaos/Public/Chaos/RigidParticles.h
@@ -164,13 +164,13 @@ class TRigidParticles : public TKinematicGeometryParticles<T, d>
 	FORCEINLINE const bool ToBeRemovedOnFracture(const int32 Index) const { return MToBeRemovedOnFracture[Index]; }
 	FORCEINLINE bool& ToBeRemovedOnFracture(const int32 Index) { return MToBeRemovedOnFracture[Index]; }
 
-	FORCEINLINE TQueue<TSleepData<T, d>, EQueueMode::Mpsc>& GetSleepData() { return MSleepData; }
-	FORCEINLINE void AddSleepData(TGeometryParticleHandle<T, d>* Particle, bool Sleeping)
+	FORCEINLINE TArray<TSleepData<T, d>>& GetSleepData() { return MSleepData; }
+	FORCEINLINE	void AddSleepData(TGeometryParticleHandle<T, d>* Particle, bool Sleeping)
 	{ 
 		TSleepData<T, d> SleepData;
 		SleepData.Particle = Particle;
 		SleepData.Sleeping = Sleeping;
-		MSleepData.Enqueue(SleepData); 
+		MSleepData.Add(SleepData); 
 	}
 
 	FORCEINLINE const EObjectStateType ObjectState(const int32 Index) const { return MObjectState[Index]; }
@@ -223,7 +223,7 @@ class TRigidParticles : public TKinematicGeometryParticles<T, d>
 	TArrayCollectionArray<bool> MDisabled;
 	TArrayCollectionArray<bool> MToBeRemovedOnFracture;
 	TArrayCollectionArray<EObjectStateType> MObjectState;
-	TQueue<TSleepData<T, d>, EQueueMode::Mpsc> MSleepData;
+	TArray<TSleepData<T, d>> MSleepData;
 };
 
 
diff --git a/Engine/Source/Runtime/Experimental/ChaosSolvers/Private/EventDefaults.cpp b/Engine/Source/Runtime/Experimental/ChaosSolvers/Private/EventDefaults.cpp
--- a/Engine/Source/Runtime/Experimental/ChaosSolvers/Private/EventDefaults.cpp
+++ b/Engine/Source/Runtime/Experimental/ChaosSolvers/Private/EventDefaults.cpp
@@ -365,8 +365,7 @@ namespace Chaos
 
 			Chaos::FPBDRigidsSolver* NonConstSolver = (Chaos::FPBDRigidsSolver*)(Solver);
 			auto& SolverSleepingData = NonConstSolver->Particles.GetDynamicParticles().GetSleepData();
-			TSleepData<float, 3> SleepData;
-			while(SolverSleepingData.Dequeue(SleepData))
+			for(const TSleepData<float, 3>& SleepData : SolverSleepingData)
 			{
 				if(SleepData.Particle)
 				{
@@ -384,6 +383,8 @@ namespace Chaos
 				}
 			}
 
+			SolverSleepingData.Empty();
+
 		});
 	}
 
