Chaos
- Joint soft drives using XPBD
- Fix soft stiffness being used for fixed swing joint

diff --git a/Engine/Source/Runtime/Engine/Private/PhysicsEngine/ImmediatePhysics/ImmediatePhysicsChaos/ImmediatePhysicsJointHandle_Chaos.cpp b/Engine/Source/Runtime/Engine/Private/PhysicsEngine/ImmediatePhysics/ImmediatePhysicsChaos/ImmediatePhysicsJointHandle_Chaos.cpp
--- a/Engine/Source/Runtime/Engine/Private/PhysicsEngine/ImmediatePhysics/ImmediatePhysicsChaos/ImmediatePhysicsJointHandle_Chaos.cpp
+++ b/Engine/Source/Runtime/Engine/Private/PhysicsEngine/ImmediatePhysics/ImmediatePhysicsChaos/ImmediatePhysicsJointHandle_Chaos.cpp
@@ -8,6 +8,8 @@
 
 #include "PhysicsEngine/ConstraintInstance.h"
 
+#pragma optimize("", off)
+
 static_assert((int32)Chaos::EJointMotionType::Free == (int32)EAngularConstraintMotion::ACM_Free, "Chaos::EJointMotionType and EAngularConstraintMotion mismatch");
 static_assert((int32)Chaos::EJointMotionType::Limited == (int32)EAngularConstraintMotion::ACM_Limited, "Chaos::EJointMotionType and EAngularConstraintMotion mismatch");
 static_assert((int32)Chaos::EJointMotionType::Locked == (int32)EAngularConstraintMotion::ACM_Locked, "Chaos::EJointMotionType and EAngularConstraintMotion mismatch");
@@ -28,10 +30,15 @@ namespace ImmediatePhysics_Chaos
 	float ChaosImmediate_DriveStiffnessScale = 30.0f;
 	FAutoConsoleVariableRef CVarDriveStiffnessScale(TEXT("p.Chaos.ImmPhys.DriveStiffnessScale"), ChaosImmediate_DriveStiffnessScale, TEXT("Conversion factor for drive stiffness."));
 
-	float ChaosImmediate_JointMinProjection = 0.0f;
-	float ChaosImmediate_JointMaxProjection = 0.5f;
-	FAutoConsoleVariableRef CVarJointMinProjection(TEXT("p.Chaos.ImmPhys.JointMinProjection"), ChaosImmediate_JointMinProjection, TEXT("Joint min projection (for joints with projection disabled)."));
-	FAutoConsoleVariableRef CVarJointMaxProjection(TEXT("p.Chaos.ImmPhys.JointMaxProjection"), ChaosImmediate_JointMaxProjection, TEXT("Joint max projection (for joints with projection enabled)."));
+	float ChaosImmediate_JointMinLinearProjection = 0.0f;
+	float ChaosImmediate_JointMaxLinearProjection = 0.5f;
+	FAutoConsoleVariableRef CVarJointMinLinearProjection(TEXT("p.Chaos.ImmPhys.JointMinProjection"), ChaosImmediate_JointMinLinearProjection, TEXT("Joint min projection (for joints with projection disabled)."));
+	FAutoConsoleVariableRef CVarJointMaxLinearProjection(TEXT("p.Chaos.ImmPhys.JointMaxProjection"), ChaosImmediate_JointMaxLinearProjection, TEXT("Joint max projection (for joints with projection enabled)."));
+
+	float ChaosImmediate_JointMinAngularProjection = 0.0f;
+	float ChaosImmediate_JointMaxAngularProjection = 0.5f;
+	FAutoConsoleVariableRef CVarJointMinAngularProjection(TEXT("p.Chaos.ImmPhys.JointMinProjection"), ChaosImmediate_JointMinAngularProjection, TEXT("Joint min projection (for joints with projection disabled)."));
+	FAutoConsoleVariableRef CVarJointMaxAngularProjection(TEXT("p.Chaos.ImmPhys.JointMaxProjection"), ChaosImmediate_JointMaxAngularProjection, TEXT("Joint max projection (for joints with projection enabled)."));
 
 	float ChaosImmediate_DriveStiffnessSourceMin = 100.0f;	// PhysX stiffness per inertia that we translate to a Chaos stiffness of ChaosImmediate_DriveStiffnessTargetMin
 	float ChaosImmediate_DriveStiffnessSourceMax = 2000.0f;	// PhysX stiffness per inertia that we translate to a Chaos stiffness of ChaosImmediate_DriveStiffnessTargetMax
@@ -56,7 +63,7 @@ namespace ImmediatePhysics_Chaos
 
 	// Convert the UE (PhysX) drive spring stiffness to a joint stiffness [0,1] value for use in the solver.
 	// We linearly map a range of UE stiffness values to a range of Chaos values, optionally scaling by the effective mass.
-	float ConvertDriveAngularStiffness(const float InStiffness, const Chaos::FVec3& X0, const FReal M0, const Chaos::FVec3& I0, const Chaos::FVec3& X1, const FReal M1, const Chaos::FVec3& I1)
+	float CalculateDriveAngularStiffnessScale(const Chaos::FVec3& X0, const FReal M0, const Chaos::FVec3& I0, const Chaos::FVec3& X1, const FReal M1, const Chaos::FVec3& I1)
 	{
 		using namespace Chaos;
 
@@ -81,29 +88,7 @@ namespace ImmediatePhysics_Chaos
 			}
 		}
 
-		return ChaosImmediate_DriveStiffnessScale * InStiffness / InvStiffnessScale;
-	}
-
-	float ConvertSoftLinearStiffness(const float InStiffness, const float InvM0, const float InvM1)
-	{
-		if (ChaosImmediate_SoftLinearStiffnessSourceMax > ChaosImmediate_SoftLinearStiffnessSourceMin)
-		{
-			const float InvM = ChaosImmediate_StiffnessUseMass ? (InvM0 + InvM1) : 1.0f;
-			const float F = FMath::Max(0.0f, (InStiffness * InvM - ChaosImmediate_SoftLinearStiffnessSourceMin)) / FMath::Max(SMALL_NUMBER, (ChaosImmediate_SoftLinearStiffnessSourceMax - ChaosImmediate_SoftLinearStiffnessSourceMin));
-			return FMath::Lerp(ChaosImmediate_SoftLinearStiffnessTargetMin, ChaosImmediate_SoftLinearStiffnessTargetMax, FMath::Clamp(F, 0.0f, 1.0f));
-		}
-		return 0;
-	}
-
-	float ConvertSoftAngularStiffness(const float InStiffness, const Chaos::TVector<float, 3>& InvI0, const Chaos::TVector<float, 3>& InvI1)
-	{
-		if (ChaosImmediate_SoftAngularStiffnessSourceMax > ChaosImmediate_SoftAngularStiffnessSourceMin)
-		{
-			const float InvI = ChaosImmediate_StiffnessUseMass ? (InvI0.Min() + InvI1.Min()) : 1.0f;
-			const float F = FMath::Max(0.0f, (InStiffness * InvI - ChaosImmediate_SoftAngularStiffnessSourceMin)) / FMath::Max(SMALL_NUMBER, (ChaosImmediate_SoftAngularStiffnessSourceMax - ChaosImmediate_SoftAngularStiffnessSourceMin));
-			return FMath::Lerp(ChaosImmediate_SoftAngularStiffnessTargetMin, ChaosImmediate_SoftAngularStiffnessTargetMax, FMath::Clamp(F, 0.0f, 1.0f));
-		}
-		return 0;
+		return ChaosImmediate_DriveStiffnessScale / InvStiffnessScale;
 	}
 
 
@@ -113,16 +98,27 @@ namespace ImmediatePhysics_Chaos
 	{
 		using namespace Chaos;
 
+		FPBDJointSettings ConstraintSettings;
+
 		// BodyInstance/PhysX has the constraint locations in actor-space, but we need them in Center-of-Mass space
-		// @todo(ccaulfield): support CoM in Chaos particles
-		FTransform ConstraintFrame1 = ConstraintInstance->GetRefFrame(EConstraintFrame::Frame1);
-		FTransform ConstraintFrame2 = ConstraintInstance->GetRefFrame(EConstraintFrame::Frame2);
+		FTransform ConstraintFrame1 = FParticleUtilities::ActorLocalToParticleLocal(TGenericParticleHandle<FReal, 3>(Actor1->GetParticle()), ConstraintInstance->GetRefFrame(EConstraintFrame::Frame1));
+		FTransform ConstraintFrame2 = FParticleUtilities::ActorLocalToParticleLocal(TGenericParticleHandle<FReal, 3>(Actor2->GetParticle()), ConstraintInstance->GetRefFrame(EConstraintFrame::Frame2));
+
+		// PhysX to Chaos conversions
+		const float DriveAngularStiffnessScale = CalculateDriveAngularStiffnessScale(
+			ConstraintFrame1.GetTranslation() - Actor1->GetLocalCoMTransform().GetTranslation(),
+			Actor1->GetMass(), 
+			Actor1->GetInertia(), 
+			ConstraintFrame2.GetTranslation() - Actor2->GetLocalCoMTransform().GetTranslation(),
+			Actor2->GetMass(), 
+			Actor2->GetInertia());
+
+		const float DriveAngularDampingScale = DriveAngularStiffnessScale;
 
-		FPBDJointSettings ConstraintSettings;
 		ConstraintSettings.ConstraintFrames = 
 		{ 
-			FParticleUtilities::ActorLocalToParticleLocal(TGenericParticleHandle<FReal, 3>(Actor1->GetParticle()), ConstraintFrame1),
-			FParticleUtilities::ActorLocalToParticleLocal(TGenericParticleHandle<FReal, 3>(Actor2->GetParticle()), ConstraintFrame2),
+			ConstraintFrame1,
+			ConstraintFrame2,
 		};
 		ConstraintSettings.Motion.Stiffness = ChaosImmediate_JointStiffness;
 
@@ -141,16 +137,19 @@ namespace ImmediatePhysics_Chaos
 		ConstraintSettings.Motion.AngularLimits[(int32)EJointAngularConstraintIndex::Swing1] = FMath::DegreesToRadians(ConstraintInstance->GetAngularSwing1Limit());
 		ConstraintSettings.Motion.AngularLimits[(int32)EJointAngularConstraintIndex::Swing2] = FMath::DegreesToRadians(ConstraintInstance->GetAngularSwing2Limit());
 
-		ConstraintSettings.Motion.LinearProjection = ConstraintInstance->IsProjectionEnabled() ? ChaosImmediate_JointMaxProjection : ChaosImmediate_JointMinProjection;
-		ConstraintSettings.Motion.AngularProjection = ConstraintInstance->IsProjectionEnabled() ? ChaosImmediate_JointMaxProjection : ChaosImmediate_JointMinProjection;
+		ConstraintSettings.Motion.LinearProjection = ConstraintInstance->IsProjectionEnabled() ? ChaosImmediate_JointMaxLinearProjection : ChaosImmediate_JointMinLinearProjection;
+		ConstraintSettings.Motion.AngularProjection = ConstraintInstance->IsProjectionEnabled() ? ChaosImmediate_JointMaxAngularProjection : ChaosImmediate_JointMinAngularProjection;
 
 		ConstraintSettings.Motion.bSoftLinearLimitsEnabled = ConstraintInstance->GetIsSoftLinearLimit();
 		ConstraintSettings.Motion.bSoftTwistLimitsEnabled = ConstraintInstance->GetIsSoftTwistLimit();
 		ConstraintSettings.Motion.bSoftSwingLimitsEnabled = ConstraintInstance->GetIsSoftSwingLimit();
-		ConstraintSettings.Motion.SoftLinearStiffness = ConvertSoftLinearStiffness(ConstraintInstance->GetSoftLinearLimitStiffness(), Actor1->GetInverseMass(), Actor2->GetInverseMass());
-		ConstraintSettings.Motion.SoftTwistStiffness = ConvertSoftAngularStiffness(ConstraintInstance->GetSoftTwistLimitStiffness(), Actor1->GetInverseInertia(), Actor2->GetInverseInertia());
-		ConstraintSettings.Motion.SoftSwingStiffness = ConvertSoftAngularStiffness(ConstraintInstance->GetSoftSwingLimitStiffness(), Actor1->GetInverseInertia(), Actor2->GetInverseInertia());
-		 
+		ConstraintSettings.Motion.SoftLinearStiffness = 1.0f;	// @todo(ccaulfield): xpbd soft linear constraints
+		ConstraintSettings.Motion.SoftLinearDamping = 0.0f;		// @todo(ccaulfield): xpbd soft linear constraints
+		ConstraintSettings.Motion.SoftTwistStiffness = DriveAngularStiffnessScale * ConstraintInstance->GetSoftTwistLimitStiffness();
+		ConstraintSettings.Motion.SoftTwistDamping = DriveAngularDampingScale * ConstraintInstance->GetSoftTwistLimitDamping();
+		ConstraintSettings.Motion.SoftSwingStiffness = DriveAngularStiffnessScale * ConstraintInstance->GetSoftSwingLimitStiffness();
+		ConstraintSettings.Motion.SoftSwingDamping = DriveAngularDampingScale * ConstraintInstance->GetSoftSwingLimitDamping();
+
 		ConstraintSettings.Motion.AngularDriveTarget = FQuat(ConstraintInstance->ProfileInstance.AngularDrive.OrientationTarget);
 		// NOTE: Hard dependence on EJointAngularConstraintIndex - the following will break if we change the order
 		ConstraintSettings.Motion.AngularDriveTargetAngles[(int32)EJointAngularConstraintIndex::Twist] = FMath::DegreesToRadians(ConstraintInstance->ProfileInstance.AngularDrive.OrientationTarget.Roll);
@@ -171,7 +170,8 @@ namespace ImmediatePhysics_Chaos
 			ConstraintSettings.Motion.bAngularSwingDriveEnabled = ConstraintInstance->ProfileInstance.AngularDrive.SwingDrive.bEnablePositionDrive;
 		}
 
-		ConstraintSettings.Motion.AngularDriveStiffness = ConvertDriveAngularStiffness(ConstraintInstance->ProfileInstance.AngularDrive.TwistDrive.Stiffness, ConstraintSettings.ConstraintFrames[0].GetTranslation() - Actor1->GetLocalCoMTransform().GetTranslation(), Actor1->GetMass(), Actor1->GetInertia(), ConstraintSettings.ConstraintFrames[1].GetTranslation() - Actor2->GetLocalCoMTransform().GetTranslation(), Actor2->GetMass(), Actor2->GetInertia());
+		ConstraintSettings.Motion.AngularDriveStiffness = DriveAngularStiffnessScale * ConstraintInstance->ProfileInstance.AngularDrive.TwistDrive.Stiffness;
+		ConstraintSettings.Motion.AngularDriveDamping = DriveAngularDampingScale * ConstraintInstance->ProfileInstance.AngularDrive.TwistDrive.Damping;
 
 		ConstraintSettings.Motion.Sanitize();
 
diff --git a/Engine/Source/Runtime/Engine/Private/PhysicsEngine/ImmediatePhysics/ImmediatePhysicsChaos/ImmediatePhysicsSimulation_Chaos.cpp b/Engine/Source/Runtime/Engine/Private/PhysicsEngine/ImmediatePhysics/ImmediatePhysicsChaos/ImmediatePhysicsSimulation_Chaos.cpp
--- a/Engine/Source/Runtime/Engine/Private/PhysicsEngine/ImmediatePhysics/ImmediatePhysicsChaos/ImmediatePhysicsSimulation_Chaos.cpp
+++ b/Engine/Source/Runtime/Engine/Private/PhysicsEngine/ImmediatePhysics/ImmediatePhysicsChaos/ImmediatePhysicsSimulation_Chaos.cpp
@@ -49,7 +49,10 @@ float ChaosImmediate_Joint_LinearProjection = 0.0f;
 float ChaosImmediate_Joint_AngularProjection = 0.0f;
 float ChaosImmediate_Joint_Stiffness = 1.0f;
 float ChaosImmediate_Joint_SoftLinearStiffness = 0.0f;
-float ChaosImmediate_Joint_SoftAngularStiffness = 0.0f;
+float ChaosImmediate_Joint_SoftTwistStiffness = 0.0f;
+float ChaosImmediate_Joint_SoftTwistDamping = 0.0f;
+float ChaosImmediate_Joint_SoftSwingStiffness = 0.0f;
+float ChaosImmediate_Joint_SoftSwingDamping = 0.0f;
 float ChaosImmediate_Joint_DriveStiffness = 0.0f;
 float ChaosImmediate_Joint_DriveDamping = 0.0f;
 float ChaosImmediate_Joint_MinParentMassRatio = 0.5f;
@@ -64,7 +67,10 @@ FAutoConsoleVariableRef CVarChaosImmPhysLinearProjection(TEXT("p.Chaos.ImmPhys.J
 FAutoConsoleVariableRef CVarChaosImmPhysAngularProjection(TEXT("p.Chaos.ImmPhys.Joint.AngularProjection"), ChaosImmediate_Joint_AngularProjection, TEXT("6Dof joint projection amount override (if > 0)."));
 FAutoConsoleVariableRef CVarChaosImmPhysStiffness(TEXT("p.Chaos.ImmPhys.Joint.Stiffness"), ChaosImmediate_Joint_Stiffness, TEXT("6Dof joint stiffness override (if > 0)."));
 FAutoConsoleVariableRef CVarChaosImmPhysSoftLinearStiffness(TEXT("p.Chaos.ImmPhys.Joint.SoftLinearStiffness"), ChaosImmediate_Joint_SoftLinearStiffness, TEXT("6Dof joint soft linear stiffness override (if > 0)."));
-FAutoConsoleVariableRef CVarChaosImmPhysSoftAngularStiffness(TEXT("p.Chaos.ImmPhys.Joint.SoftAngularStiffness"), ChaosImmediate_Joint_SoftAngularStiffness, TEXT("6Dof joint soft angular stiffness override (if > 0)."));
+FAutoConsoleVariableRef CVarChaosImmPhysSoftTwistStiffness(TEXT("p.Chaos.ImmPhys.Joint.SoftTwistStiffness"), ChaosImmediate_Joint_SoftTwistStiffness, TEXT("6Dof joint SoftTwist stiffness override (if > 0)."));
+FAutoConsoleVariableRef CVarChaosImmPhysSoftTwistDamping(TEXT("p.Chaos.ImmPhys.Joint.SoftTwistDamping"), ChaosImmediate_Joint_SoftTwistDamping, TEXT("6Dof joint SoftTwist damping override (if > 0)."));
+FAutoConsoleVariableRef CVarChaosImmPhysSoftSwingStiffness(TEXT("p.Chaos.ImmPhys.Joint.SoftSwingStiffness"), ChaosImmediate_Joint_SoftSwingStiffness, TEXT("6Dof joint SoftSwing stiffness override (if > 0)."));
+FAutoConsoleVariableRef CVarChaosImmPhysSoftSwingDamping(TEXT("p.Chaos.ImmPhys.Joint.SoftSwingDamping"), ChaosImmediate_Joint_SoftSwingDamping, TEXT("6Dof joint SoftSwing damping override (if > 0)."));
 FAutoConsoleVariableRef CVarChaosImmPhysDriveStiffness(TEXT("p.Chaos.ImmPhys.Joint.DriveStiffness"), ChaosImmediate_Joint_DriveStiffness, TEXT("6Dof joint drive stiffness override (if > 0)."));
 FAutoConsoleVariableRef CVarChaosImmPhysDriveDamping(TEXT("p.Chaos.ImmPhys.Joint.DriveDamping"), ChaosImmediate_Joint_DriveDamping, TEXT("6Dof joint drive damping override (if > 0)."));
 FAutoConsoleVariableRef CVarChaosImmPhysMinParentMassRatio(TEXT("p.Chaos.ImmPhys.Joint.MinParentMassRatio"), ChaosImmediate_Joint_MinParentMassRatio, TEXT("6Dof joint MinParentMassRatio (if > 0)"));
@@ -496,7 +502,10 @@ namespace ImmediatePhysics_Chaos
 			JointsSettings.AngularProjection = ChaosImmediate_Joint_AngularProjection;
 			JointsSettings.Stiffness = ChaosImmediate_Joint_Stiffness;
 			JointsSettings.SoftLinearStiffness = ChaosImmediate_Joint_SoftLinearStiffness;
-			JointsSettings.SoftAngularStiffness = ChaosImmediate_Joint_SoftAngularStiffness;
+			JointsSettings.SoftTwistStiffness = ChaosImmediate_Joint_SoftTwistStiffness;
+			JointsSettings.SoftTwistDamping = ChaosImmediate_Joint_SoftTwistDamping;
+			JointsSettings.SoftSwingStiffness = ChaosImmediate_Joint_SoftSwingStiffness;
+			JointsSettings.SoftSwingDamping = ChaosImmediate_Joint_SoftSwingDamping;
 			JointsSettings.DriveStiffness = ChaosImmediate_Joint_DriveStiffness;
 			JointsSettings.DriveDamping = ChaosImmediate_Joint_DriveDamping;
 
diff --git a/Engine/Source/Runtime/Experimental/Chaos/Private/Chaos/Joint/PBDJointSolverGaussSeidel.cpp b/Engine/Source/Runtime/Experimental/Chaos/Private/Chaos/Joint/PBDJointSolverGaussSeidel.cpp
--- a/Engine/Source/Runtime/Experimental/Chaos/Private/Chaos/Joint/PBDJointSolverGaussSeidel.cpp
+++ b/Engine/Source/Runtime/Experimental/Chaos/Private/Chaos/Joint/PBDJointSolverGaussSeidel.cpp
@@ -6,10 +6,173 @@
 #include "ChaosLog.h"
 #include "ChaosStats.h"
 
-//#pragma optimize("", off)
+#pragma optimize("", off)
 
 namespace Chaos
 {
+	void GetTwistAxisAngleLocal(
+		const FRotation3& R0,
+		const FRotation3& R1,
+		FVec3& TwistAxisLocal,
+		FReal& TwistAngle)
+	{
+		// Decompose rotation of body 1 relative to body 0 into swing and twist rotations, assuming twist is X axis
+		FRotation3 R01Twist, R01Swing;
+		FPBDJointUtilities::DecomposeSwingTwistLocal(R0, R1, R01Swing, R01Twist);
+
+		TwistAxisLocal = FJointConstants::TwistAxis();
+
+		TwistAngle = R01Twist.GetAngle();
+		if (TwistAngle > PI)
+		{
+			TwistAngle = TwistAngle - (FReal)2 * PI;
+		}
+		if (R01Twist.X < 0)
+		{
+			TwistAngle = -TwistAngle;
+		}
+	}
+
+
+	void GetConeAxisAngleLocal(
+		const FRotation3& R0,
+		const FRotation3& R1,
+		const FReal AngleTolerance,
+		FVec3& Axis,
+		FReal& Angle)
+	{
+		// Decompose rotation of body 1 relative to body 0 into swing and twist rotations, assuming twist is X axis
+		FRotation3 R01Twist, R01Swing;
+		FPBDJointUtilities::DecomposeSwingTwistLocal(R0, R1, R01Swing, R01Twist);
+
+		R01Swing.ToAxisAndAngleSafe(Axis, Angle, FJointConstants::Swing1Axis(), AngleTolerance);
+		if (Angle > PI)
+		{
+			Angle = Angle - (FReal)2 * PI;
+		}
+	}
+
+
+	FReal GetConeAngleError(
+		const FPBDJointSettings& JointSettings,
+		const FVec3& SwingAxisLocal,
+		const FReal SwingAngle)
+	{
+		// Calculate swing limit for the current swing axis
+		const FReal Swing1Limit = JointSettings.Motion.AngularLimits[(int32)EJointAngularConstraintIndex::Swing1];
+		const FReal Swing2Limit = JointSettings.Motion.AngularLimits[(int32)EJointAngularConstraintIndex::Swing2];
+
+		// Circular swing limit
+		FReal SwingAngleMax = Swing1Limit;
+
+		// Elliptical swing limit
+		// @todo(ccaulfield): do elliptical constraints properly (axis is still for circular limit)
+		if (!FMath::IsNearlyEqual(Swing1Limit, Swing2Limit, KINDA_SMALL_NUMBER))
+		{
+			// Map swing axis to ellipse and calculate limit for this swing axis
+			const FReal DotSwing1 = FMath::Abs(FVec3::DotProduct(SwingAxisLocal, FJointConstants::Swing1Axis()));
+			const FReal DotSwing2 = FMath::Abs(FVec3::DotProduct(SwingAxisLocal, FJointConstants::Swing2Axis()));
+			SwingAngleMax = FMath::Sqrt(Swing1Limit * DotSwing1 * Swing1Limit * DotSwing1 + Swing2Limit * DotSwing2 * Swing2Limit * DotSwing2);
+		}
+
+		// Calculate swing error we need to correct
+		if (SwingAngle > SwingAngleMax)
+		{
+			return SwingAngle - SwingAngleMax;
+		}
+		else if (SwingAngle < -SwingAngleMax)
+		{
+			return SwingAngle + SwingAngleMax;
+		}
+
+		return 0;
+	}
+
+
+	void GetSwingAxisAngle(
+		const FRotation3& R0,
+		const FRotation3& R1,
+		const FReal AngleTolerance,
+		const EJointAngularConstraintIndex SwingConstraintIndex,
+		const EJointAngularAxisIndex SwingAxisIndex,
+		FVec3& Axis,
+		FReal& Angle)
+	{
+#if 0
+		const FMatrix33 RM0 = R0.ToMatrix();
+		const FMatrix33 RM1 = R1.ToMatrix();
+		const FVec3& Twist0 = RM0.GetAxis((int32)EJointAngularAxisIndex::Twist);
+		const FVec3& Twist1 = RM1.GetAxis((int32)EJointAngularAxisIndex::Twist);
+
+		Axis = RM0.GetAxis((int32)SwingAxisIndex);
+		FReal SinAngle = FVec3::DotProduct(FVec3::CrossProduct(Twist1, Twist0), Axis);
+		Angle = FMath::Asin(FMath::Clamp(SinAngle, (FReal)-1, (FReal)1));
+#elif 0
+		// Decompose rotation of body 1 relative to body 0 into swing and twist rotations, assuming twist is X axis
+		FRotation3 R01Twist, R01Swing;
+		FPBDJointUtilities::DecomposeSwingTwistLocal(R0, R1, R01Swing, R01Twist);
+		const FReal R01SwingYorZ = ((int32)SwingAxisIndex == 2) ? R01Swing.Z : R01Swing.Y;	// Can't index a quat :(
+		FReal Angle = 4.0f * FMath::Atan2(R01SwingYorZ, 1.0f + R01Swing.W);
+		if (R01Twist.X < 0)
+		{
+			if (Angle > 0)
+			{
+				Angle = (FReal)PI - Angle;
+			}
+			else
+			{
+				Angle = (FReal)-PI - Angle;
+			}
+		}
+		const FVec3& AxisLocal = (SwingConstraintIndex == EJointAngularConstraintIndex::Swing1) ? FJointConstants::Swing1Axis() : FJointConstants::Swing2Axis();
+		Axis = R0 * AxisLocal;
+#else
+		// Do something better than this!!
+		FRotation3 R01Twist, R01Swing;
+		FPBDJointUtilities::DecomposeSwingTwistLocal(R0, R1, R01Swing, R01Twist);
+		FVec3 TwistAxis01 = FJointConstants::TwistAxis();
+		FReal TwistAngle = R01Twist.GetAngle();
+		if (TwistAngle > PI)
+		{
+			TwistAngle = TwistAngle - (FReal)2 * PI;
+		}
+		if (R01Twist.X < 0)
+		{
+			TwistAngle = -TwistAngle;
+		}
+		const FVec3 TwistAxis = R0 * TwistAxis01;
+
+		const FRotation3 R1NoTwist = R1 * R01Twist.Inverse();
+		const FMatrix33 Axes0 = R0.ToMatrix();
+		const FMatrix33 Axes1 = R1NoTwist.ToMatrix();
+
+		Axis = Axes0.GetAxis((int32)SwingAxisIndex);
+		Angle = 0;
+
+		const EJointAngularAxisIndex OtherSwingAxis = (SwingAxisIndex == EJointAngularAxisIndex::Swing1) ? EJointAngularAxisIndex::Swing2 : EJointAngularAxisIndex::Swing1;
+		FVec3 SwingCross = FVec3::CrossProduct(Axes0.GetAxis((int32)OtherSwingAxis), Axes1.GetAxis((int32)OtherSwingAxis));
+		SwingCross = SwingCross - FVec3::DotProduct(TwistAxis, SwingCross) * TwistAxis;
+		const FReal SwingCrossLen = SwingCross.Size();
+		if (SwingCrossLen > KINDA_SMALL_NUMBER)
+		{
+			Axis = SwingCross / SwingCrossLen;
+
+			Angle = FMath::Asin(FMath::Clamp(SwingCrossLen, (FReal)0, (FReal)1));
+			const FReal SwingDot = FVec3::DotProduct(Axes0.GetAxis((int32)OtherSwingAxis), Axes1.GetAxis((int32)OtherSwingAxis));
+			if (SwingDot < (FReal)0)
+			{
+				Angle = (FReal)PI - Angle;
+			}
+		}
+#endif
+	}
+
+	//
+	//
+	//////////////////////////////////////////////////////////////////////////
+	//
+	//
+
 	FJointSolverGaussSeidel::FJointSolverGaussSeidel()
 	{
 	}
@@ -23,6 +186,12 @@ namespace Chaos
 		Rs[1] = Qs[1] * XLs[1].GetRotation();
 	}
 
+	void FJointSolverGaussSeidel::UpdateDerivedState(const int32 BodyIndex)
+	{
+		Xs[BodyIndex] = Ps[BodyIndex] + Qs[BodyIndex] * XLs[BodyIndex].GetTranslation();
+		Rs[BodyIndex] = Qs[BodyIndex] * XLs[BodyIndex].GetRotation();
+	}
+
 
 	void FJointSolverGaussSeidel::Init(
 		const FReal Dt,
@@ -48,21 +217,10 @@ namespace Chaos
 		PrevQs[0] = PrevQ0;
 		PrevQs[1] = PrevQ1;
 
-		TwistLambda = (FReal)0;
-		SwingLambda = (FReal)0;
-
-		LinearStiffness = FPBDJointUtilities::GetLinearStiffness(SolverSettings, JointSettings);
-		TwistStiffness = FPBDJointUtilities::GetTwistStiffness(SolverSettings, JointSettings);
-		SwingStiffness = FPBDJointUtilities::GetSwingStiffness(SolverSettings, JointSettings);
-		AngularDriveStiffness = FPBDJointUtilities::GetAngularDriveStiffness(SolverSettings, JointSettings);
-		AngularDriveDamping = FPBDJointUtilities::GetAngularDriveDamping(SolverSettings, JointSettings);
-		LinearProjection = FPBDJointUtilities::GetLinearProjection(SolverSettings, JointSettings);
-		AngularProjection = FPBDJointUtilities::GetAngularProjection(SolverSettings, JointSettings);
-		SwingTwistAngleTolerance = SolverSettings.SwingTwistAngleTolerance;
-		bEnableTwistLimits = SolverSettings.bEnableTwistLimits;
-		bEnableSwingLimits = SolverSettings.bEnableSwingLimits;
-		bEnableDrives = SolverSettings.bEnableDrives;
-		bAccelerationDriveMode = true;
+		TwistSoftLambda = (FReal)0;
+		SwingSoftLambda = (FReal)0;
+		TwistDriveLambda = (FReal)0;
+		SwingDriveLambda = (FReal)0;
 	}
 
 
@@ -94,65 +252,99 @@ namespace Chaos
 
 	void FJointSolverGaussSeidel::ApplyConstraints(
 		const FReal Dt,
+		const FPBDJointSolverSettings& SolverSettings,
 		const FPBDJointSettings& JointSettings)
 	{
-		ApplyRotationConstraints(Dt, JointSettings);
-		ApplyPositionConstraints(Dt, JointSettings);
+		ApplyRotationConstraints(Dt, SolverSettings, JointSettings);
+
+		ApplyPositionConstraints(Dt, SolverSettings, JointSettings);
 	}
 
 
 	void FJointSolverGaussSeidel::ApplyDrives(
 		const FReal Dt,
+		const FPBDJointSolverSettings& SolverSettings,
 		const FPBDJointSettings& JointSettings)
 	{
 		// @todo(ccaulfield): linear and velocity drives
-		ApplyRotationDrives(Dt, JointSettings);
+		ApplyRotationDrives(Dt, SolverSettings, JointSettings);
 	}
 
+
 	void FJointSolverGaussSeidel::ApplyProjections(
 		const FReal Dt,
+		const FPBDJointSolverSettings& SolverSettings,
 		const FPBDJointSettings& JointSettings)
 	{
-		ApplyPositionProjection(Dt, JointSettings);
+		// With the parent body set to infinite mass...
+
+		// Then fix the rotations
+		ApplyRotationProjection(Dt, SolverSettings, JointSettings);
+
+		// We first apply a semi-physical linear projection
+		ApplyPositionProjection(Dt, SolverSettings, JointSettings);
 	}
 
 
 	void FJointSolverGaussSeidel::ApplyRotationConstraints(
 		const FReal Dt,
+		const FPBDJointSolverSettings& SolverSettings,
 		const FPBDJointSettings& JointSettings)
 	{
+		// @todo(ccaulfield): eliminate branches in solver loop by pre-building list of functions to call?
+
 		EJointMotionType TwistMotion = JointSettings.Motion.AngularMotionTypes[(int32)EJointAngularConstraintIndex::Twist];
 		EJointMotionType Swing1Motion = JointSettings.Motion.AngularMotionTypes[(int32)EJointAngularConstraintIndex::Swing1];
 		EJointMotionType Swing2Motion = JointSettings.Motion.AngularMotionTypes[(int32)EJointAngularConstraintIndex::Swing2];
+		bool bTwistSoft = JointSettings.Motion.bSoftTwistLimitsEnabled;
+		bool bSwingSoft = JointSettings.Motion.bSoftSwingLimitsEnabled;
 
 		// Apply twist constraint
-		if (bEnableTwistLimits)
+		if (SolverSettings.bEnableTwistLimits)
 		{
-			if (TwistMotion != EJointMotionType::Free)
+			if ((TwistMotion == EJointMotionType::Limited) && bTwistSoft)
 			{
-				ApplyTwistConstraint(Dt, JointSettings);
+				ApplyTwistConstraintSoft(Dt, SolverSettings, JointSettings);
+			}
+			else if (TwistMotion != EJointMotionType::Free)
+			{
+				ApplyTwistConstraint(Dt, SolverSettings, JointSettings);
 			}
 		}
 
 		// Apply swing constraints
-		if (bEnableSwingLimits)
+		if (SolverSettings.bEnableSwingLimits)
 		{
 			if ((Swing1Motion == EJointMotionType::Limited) && (Swing2Motion == EJointMotionType::Limited))
 			{
-				// Swing Cone
-				ApplyConeConstraint(Dt, JointSettings);
+				if (bSwingSoft)
+				{
+					ApplyConeConstraintSoft(Dt, SolverSettings, JointSettings);
+				}
+				else
+				{
+					ApplyConeConstraint(Dt, SolverSettings, JointSettings);
+				}
 			}
 			else
 			{
-				if (Swing1Motion != EJointMotionType::Free)
+				
+				if ((Swing1Motion == EJointMotionType::Limited) && bSwingSoft)
 				{
-					// Swing Arc/Lock
-					ApplySwingConstraint(Dt, JointSettings, EJointAngularConstraintIndex::Swing1, EJointAngularAxisIndex::Swing1);
+					ApplySwingConstraintSoft(Dt, SolverSettings, JointSettings, EJointAngularConstraintIndex::Swing1, EJointAngularAxisIndex::Swing1);
 				}
-				if (Swing2Motion != EJointMotionType::Free)
+				else if (Swing1Motion != EJointMotionType::Free)
 				{
-					// Swing Arc/Lock
-					ApplySwingConstraint(Dt, JointSettings, EJointAngularConstraintIndex::Swing2, EJointAngularAxisIndex::Swing2);
+					ApplySwingConstraint(Dt, SolverSettings, JointSettings, EJointAngularConstraintIndex::Swing1, EJointAngularAxisIndex::Swing1);
+				}
+				
+				if ((Swing2Motion == EJointMotionType::Limited) && bSwingSoft)
+				{
+					ApplySwingConstraintSoft(Dt, SolverSettings, JointSettings, EJointAngularConstraintIndex::Swing2, EJointAngularAxisIndex::Swing2);
+				}
+				else if (Swing2Motion != EJointMotionType::Free)
+				{
+					ApplySwingConstraint(Dt, SolverSettings, JointSettings, EJointAngularConstraintIndex::Swing2, EJointAngularAxisIndex::Swing2);
 				}
 			}
 		}
@@ -161,13 +353,14 @@ namespace Chaos
 
 	void FJointSolverGaussSeidel::ApplyRotationDrives(
 		const FReal Dt,
+		const FPBDJointSolverSettings& SolverSettings,
 		const FPBDJointSettings& JointSettings)
 	{
 		EJointMotionType TwistMotion = JointSettings.Motion.AngularMotionTypes[(int32)EJointAngularConstraintIndex::Twist];
 		EJointMotionType Swing1Motion = JointSettings.Motion.AngularMotionTypes[(int32)EJointAngularConstraintIndex::Swing1];
 		EJointMotionType Swing2Motion = JointSettings.Motion.AngularMotionTypes[(int32)EJointAngularConstraintIndex::Swing2];
 
-		if (bEnableDrives)
+		if (SolverSettings.bEnableDrives)
 		{
 			bool bTwistLocked = TwistMotion == EJointMotionType::Locked;
 			bool bSwing1Locked = Swing1Motion == EJointMotionType::Locked;
@@ -176,25 +369,67 @@ namespace Chaos
 			// No SLerp drive if we have a locked rotation (it will be grayed out in the editor in this case, but could still have been set before the rotation was locked)
 			if (JointSettings.Motion.bAngularSLerpDriveEnabled && !bTwistLocked && !bSwing1Locked && !bSwing2Locked)
 			{
-				ApplySLerpDrive(Dt, JointSettings);
+				ApplySLerpDrive(Dt, SolverSettings, JointSettings);
 			}
-
-			if (JointSettings.Motion.bAngularTwistDriveEnabled && !bTwistLocked)
+			else
 			{
-				ApplyTwistDrive(Dt, JointSettings);
+				if (JointSettings.Motion.bAngularTwistDriveEnabled && !bTwistLocked)
+				{
+					ApplyTwistDrive(Dt, SolverSettings, JointSettings);
+				}
+
+				if (JointSettings.Motion.bAngularSwingDriveEnabled && !bSwing1Locked && !bSwing2Locked)
+				{
+					ApplyConeDrive(Dt, SolverSettings, JointSettings);
+				}
+				else if (JointSettings.Motion.bAngularSwingDriveEnabled && !bSwing1Locked)
+				{
+					ApplySwingDrive(Dt, SolverSettings, JointSettings, EJointAngularConstraintIndex::Swing1, EJointAngularAxisIndex::Swing1);
+				}
+				else if (JointSettings.Motion.bAngularSwingDriveEnabled && !bSwing2Locked)
+				{
+					ApplySwingDrive(Dt, SolverSettings, JointSettings, EJointAngularConstraintIndex::Swing2, EJointAngularAxisIndex::Swing2);
+				}
 			}
+		}
+	}
+
 
-			if (JointSettings.Motion.bAngularSwingDriveEnabled && !bSwing1Locked && !bSwing2Locked)
+	void FJointSolverGaussSeidel::ApplyRotationProjection(
+		const FReal Dt,
+		const FPBDJointSolverSettings& SolverSettings,
+		const FPBDJointSettings& JointSettings)
+	{
+		EJointMotionType TwistMotion = JointSettings.Motion.AngularMotionTypes[(int32)EJointAngularConstraintIndex::Twist];
+		EJointMotionType Swing1Motion = JointSettings.Motion.AngularMotionTypes[(int32)EJointAngularConstraintIndex::Swing1];
+		EJointMotionType Swing2Motion = JointSettings.Motion.AngularMotionTypes[(int32)EJointAngularConstraintIndex::Swing2];
+		bool bTwistSoft = JointSettings.Motion.bSoftTwistLimitsEnabled;
+		bool bSwingSoft = JointSettings.Motion.bSoftSwingLimitsEnabled;
+
+		if (SolverSettings.bEnableTwistLimits)
+		{
+			if ((TwistMotion == EJointMotionType::Locked) || ((TwistMotion == EJointMotionType::Limited) && !bTwistSoft))
 			{
-				ApplyConeDrive(Dt, JointSettings);
+				ApplyTwistProjection(Dt, SolverSettings, JointSettings);
 			}
-			else if (JointSettings.Motion.bAngularSwingDriveEnabled && !bSwing1Locked)
+		}
+
+		if (SolverSettings.bEnableSwingLimits)
+		{
+			if ((Swing1Motion == EJointMotionType::Limited) && (Swing2Motion == EJointMotionType::Limited) && !bSwingSoft)
 			{
-				ApplySwingDrive(Dt, JointSettings, EJointAngularConstraintIndex::Swing1, EJointAngularAxisIndex::Swing1);
+				ApplyConeProjection(Dt, SolverSettings, JointSettings);
 			}
-			else if (JointSettings.Motion.bAngularSwingDriveEnabled && !bSwing2Locked)
+			else
 			{
-				ApplySwingDrive(Dt, JointSettings, EJointAngularConstraintIndex::Swing2, EJointAngularAxisIndex::Swing2);
+				if ((Swing1Motion == EJointMotionType::Locked) || ((Swing1Motion == EJointMotionType::Limited) && !bSwingSoft))
+				{
+					ApplySwingProjection(Dt, SolverSettings, JointSettings, EJointAngularConstraintIndex::Swing1, EJointAngularAxisIndex::Swing1);
+				}
+				if ((Swing2Motion == EJointMotionType::Locked) || ((Swing2Motion == EJointMotionType::Limited) && !bSwingSoft))
+				{
+					ApplySwingProjection(Dt, SolverSettings, JointSettings, EJointAngularConstraintIndex::Swing2, EJointAngularAxisIndex::Swing2);
+				}
 			}
 		}
 	}
@@ -202,15 +437,30 @@ namespace Chaos
 
 	void FJointSolverGaussSeidel::ApplyPositionConstraints(
 		const FReal Dt,
+		const FPBDJointSolverSettings& SolverSettings,
 		const FPBDJointSettings& JointSettings)
 	{
 		const TVector<EJointMotionType, 3>& LinearMotion = JointSettings.Motion.LinearMotionTypes;
 		if ((LinearMotion[0] != EJointMotionType::Free) || (LinearMotion[1] != EJointMotionType::Free) || (LinearMotion[2] != EJointMotionType::Free))
 		{
-			ApplyPositionConstraint(Dt, JointSettings);
+			ApplyPositionConstraint(Dt, SolverSettings, JointSettings);
 		}
 	}
 
+	//
+	//
+	//////////////////////////////////////////////////////////////////////////
+	//
+	//
+
+	void FJointSolverGaussSeidel::ApplyPositionDelta(
+		const int32 BodyIndex,
+		const FReal Stiffness,
+		const FVec3& DP)
+	{
+		Ps[BodyIndex] = Ps[BodyIndex] + Stiffness * DP;
+	}
+
 
 	void FJointSolverGaussSeidel::ApplyPositionDelta(
 		const FReal Stiffness,
@@ -222,6 +472,16 @@ namespace Chaos
 	}
 
 
+	void FJointSolverGaussSeidel::ApplyRotationDelta(
+		const int32 BodyIndex,
+		const FReal Stiffness,
+		const FVec3& DR)
+	{
+		const FRotation3 DQ = (FRotation3::FromElements(Stiffness * DR, 0) * Qs[BodyIndex]) * (FReal)0.5;
+		Qs[BodyIndex] = (Qs[BodyIndex] + DQ).GetNormalized();
+	}
+
+
 	void FJointSolverGaussSeidel::ApplyRotationDelta(
 		const FReal Stiffness,
 		const FVec3& DR0,
@@ -235,46 +495,117 @@ namespace Chaos
 	}
 
 
-	void FJointSolverGaussSeidel::ApplyRotationDelta(
+	void FJointSolverGaussSeidel::ApplyVelocityDelta(
+		const int32 BodyIndex,
+		const FReal Stiffness,
+		const FVec3& DV,
+		const FVec3& DW)
+	{
+		Vs[BodyIndex] = Vs[BodyIndex] + Stiffness * DV;
+		Ws[BodyIndex] = Ws[BodyIndex] + Stiffness * DW;
+	}
+
+
+	void FJointSolverGaussSeidel::ApplyVelocityDelta(
+		const FReal Stiffness,
+		const FVec3& DV0,
+		const FVec3& DW0,
+		const FVec3& DV1,
+		const FVec3& DW1)
+	{
+		Vs[0] = Vs[0] + Stiffness * DV0;
+		Ws[0] = Ws[0] + Stiffness * DW0;
+		Vs[1] = Vs[1] + Stiffness * DV1;
+		Ws[1] = Ws[1] + Stiffness * DW1;
+	}
+
+
+	void FJointSolverGaussSeidel::ApplyRotationCorrection(
 		const FReal Stiffness,
 		const FVec3& Axis0,
 		const FVec3& Axis1,
 		const FReal Angle)
 	{
+		// World-space inverse mass
 		const FMatrix33 InvI0 = Utilities::ComputeWorldSpaceInertia(Qs[0], InvILs[0]);
 		const FMatrix33 InvI1 = Utilities::ComputeWorldSpaceInertia(Qs[1], InvILs[1]);
-
 		const FVec3 IA0 = Utilities::Multiply(InvI0, Axis0);
 		const FVec3 IA1 = Utilities::Multiply(InvI1, Axis1);
-		const FReal M0 = FVec3::DotProduct(Axis0, IA0);
-		const FReal M1 = FVec3::DotProduct(Axis1, IA1);
-
-		const FVec3 DR0 = IA0 * (Angle / (M0 + M1));
-		const FVec3 DR1 = IA1 * -(Angle / (M0 + M1));
-		//const FVec3 DR0 = Axis0 * (Angle * M0 / (M0 + M1));
-		//const FVec3 DR1 = Axis1 * -(Angle * M1 / (M0 + M1));
 
+		// Joint-space inverse mass
+		const FReal II0 = FVec3::DotProduct(Axis0, IA0);
+		const FReal II1 = FVec3::DotProduct(Axis1, IA1);
+		const FReal IM = II0 + II1;
+
+		//FVec3 LinearAxis0 = FVec3::CrossProduct(Xs[0] - Ps[0], Axis0);
+		//FVec3 LinearAxis1 = FVec3::CrossProduct(Xs[1] - Ps[1], Axis1);
+		//FReal LinearAxis0InvLenSq = LinearAxis0.SizeSquared();
+		//FReal LinearAxis1InvLenSq = LinearAxis1.SizeSquared();
+		//if (LinearAxis0InvLenSq > KINDA_SMALL_NUMBER)
+		//{
+		//	LinearAxis0InvLenSq = (FReal)1 / LinearAxis0InvLenSq;
+		//	LinearAxis0 = LinearAxis0 * LinearAxis0InvLenSq;
+		//}
+		//if (LinearAxis1InvLenSq > KINDA_SMALL_NUMBER)
+		//{
+		//	LinearAxis1InvLenSq = (FReal)1 / LinearAxis1InvLenSq;
+		//	LinearAxis1 = LinearAxis1 * LinearAxis1InvLenSq;
+		//}
+		//const FReal IM0 = InvMs[0] * LinearAxis0InvLenSq;
+		//const FReal IM1 = InvMs[1] * LinearAxis1InvLenSq;
+		//const FReal IM = IM0 + II0 + IM1 + II1;
+
+		// This Position delta assumes locked linear constraints...correction in the direction of inactive constraints should be eliminated.
+		// However, it's pretty rare to have angular limits and position limits on the same joint. Should check for this and fix if necessary.
+		//const FVec3 DP0 = LinearAxis0 * (Angle * InvMs[0] / IM);
+		//const FVec3 DP1 = LinearAxis1 * (-Angle * InvMs[1] / IM);
+		const FVec3 DR0 = IA0 * (Angle / IM);
+		const FVec3 DR1 = IA1 * -(Angle / IM);
+		//const FVec3 DR0 = Axis0 * (Angle * II0 / IM);
+		//const FVec3 DR1 = Axis1 * -(Angle * II1 / IM);
+
+		//ApplyPositionDelta(Stiffness, DP0, DP1);
 		ApplyRotationDelta(Stiffness, DR0, DR1);
 	}
 
+
 	// See "XPBD: Position-Based Simulation of Compliant Constrained Dynamics"
-	void FJointSolverGaussSeidel::ApplyDriveRotationDelta(
+	void FJointSolverGaussSeidel::ApplyRotationCorrectionSoft(
 		const FReal Dt,
 		const FReal Stiffness,
 		const FReal Damping,
+		const bool bAccelerationMode,
 		const FVec3& Axis0,
 		const FVec3& Axis1,
 		const FReal Angle,
 		FReal& Lambda)
 	{
+		// World-space inverse mass
 		const FMatrix33 InvI0 = Utilities::ComputeWorldSpaceInertia(Qs[0], InvILs[0]);
 		const FMatrix33 InvI1 = Utilities::ComputeWorldSpaceInertia(Qs[1], InvILs[1]);
-
 		const FVec3 IA0 = Utilities::Multiply(InvI0, Axis0);
 		const FVec3 IA1 = Utilities::Multiply(InvI1, Axis1);
-		const FReal IM0 = FVec3::DotProduct(Axis0, IA0);
-		const FReal IM1 = FVec3::DotProduct(Axis1, IA1);
-		const FReal IM = (IM0 + IM1);
+
+		// Joint-space inverse mass
+		const FReal II0 = FVec3::DotProduct(Axis0, IA0);
+		const FReal II1 = FVec3::DotProduct(Axis1, IA1);
+		const FReal IM = II0 + II1;
+
+		//FVec3 LinearAxis0 = FVec3::CrossProduct(Xs[0] - Ps[0], Axis0);
+		//FVec3 LinearAxis1 = FVec3::CrossProduct(Xs[1] - Ps[1], Axis1);
+		//FReal LinearAxis0InvLenSq = LinearAxis0.SizeSquared();
+		//FReal LinearAxis1InvLenSq = LinearAxis1.SizeSquared();
+		//if (LinearAxis0InvLenSq > KINDA_SMALL_NUMBER)
+		//{
+		//	LinearAxis0InvLenSq = (FReal)1 / LinearAxis0InvLenSq;
+		//}
+		//if (LinearAxis1InvLenSq > KINDA_SMALL_NUMBER)
+		//{
+		//	LinearAxis1InvLenSq = (FReal)1 / LinearAxis1InvLenSq;
+		//}
+		//const FReal IM0 = InvMs[0] * LinearAxis0InvLenSq;
+		//const FReal IM1 = InvMs[1] * LinearAxis1InvLenSq;
+		//const FReal IM = IM0 + II0 + IM1 + II1;
 
 		// Damping angular velocity
 		FReal JW = 0;
@@ -286,7 +617,7 @@ namespace Chaos
 		}
 
 		// Cancel out the mass for acceleration drives
-		const FReal MassScale = (bAccelerationDriveMode)? IM : (FReal)1.0;
+		const FReal MassScale = (bAccelerationMode) ? IM : (FReal)1.0;
 
 		FReal DLambda = 0;
 		if (Stiffness > KINDA_SMALL_NUMBER)
@@ -299,60 +630,87 @@ namespace Chaos
 		}
 		else
 		{
-			const FReal Beta = Damping / (MassScale * Dt * Dt);
-			const FReal DLambdaNumerator = -Beta * JW;
+			const FReal Beta = MassScale * Damping / (Dt * Dt);
+			const FReal DLambdaNumerator = -Lambda - Beta * JW;
 			const FReal DLambdaDenominator = IM;
 			DLambda = DLambdaNumerator / DLambdaDenominator;
 		}
 
+		//const FVec3 DP0 = LinearAxis0 * (DLambda * IM0);
+		//const FVec3 DP1 = LinearAxis1 * (-DLambda * IM1);
 		//const FVec3 DR0 = IA0 * DLambda;
 		//const FVec3 DR1 = IA1 * -DLambda;
-		const FVec3 DR0 = Axis0 * (DLambda * IM0);
-		const FVec3 DR1 = Axis1 * -(DLambda * IM1);
-		const FRotation3 DQ0 = (FRotation3::FromElements(DR0, 0) * Qs[0]) * (FReal)0.5;
-		const FRotation3 DQ1 = (FRotation3::FromElements(DR1, 0) * Qs[1]) * (FReal)0.5;
-		Qs[0] = (Qs[0] + DQ0).GetNormalized();
-		Qs[1] = (Qs[1] + DQ1).GetNormalized();
-		Qs[1].EnforceShortestArcWith(Qs[0]);
+		const FVec3 DR0 = Axis0 * (DLambda * II0);
+		const FVec3 DR1 = Axis1 * -(DLambda * II1);
 
+		//ApplyPositionDelta(1.0f, DP0, DP1);
+		ApplyRotationDelta(1.0f, DR0, DR1);
 		Lambda += DLambda;
 	}
 
-	void FJointSolverGaussSeidel::ApplyVelocityDelta(
-		const FReal Stiffness,
-		const FVec3& DV0,
-		const FVec3& DW0,
-		const FVec3& DV1,
-		const FVec3& DW1)
-	{
-		Vs[0] = Vs[0] + Stiffness * DV0;
-		Ws[0] = Ws[0] + Stiffness * DW0;
-		Vs[1] = Vs[1] + Stiffness * DV1;
-		Ws[1] = Ws[1] + Stiffness * DW1;
-	}
+
+	//
+	//
+	//////////////////////////////////////////////////////////////////////////
+	//
+	//
 
 
 	void FJointSolverGaussSeidel::ApplyTwistConstraint(
 		const FReal Dt,
+		const FPBDJointSolverSettings& SolverSettings,
 		const FPBDJointSettings& JointSettings)
 	{
-		// Decompose rotation of body 1 relative to body 0 into swing and twist rotations, assuming twist is X axis
-		FRotation3 R01Twist, R01Swing;
-		FPBDJointUtilities::DecomposeSwingTwistLocal(Rs[0], Rs[1], R01Swing, R01Twist);
+		FVec3 TwistAxisLocal;
+		FReal TwistAngle;
+		GetTwistAxisAngleLocal(Rs[0], Rs[1], TwistAxisLocal, TwistAngle);
+		
+		const FVec3 TwistAxis0 = Rs[0] * TwistAxisLocal;
+		const FVec3 TwistAxis1 = Rs[1] * TwistAxisLocal;
 
-		FVec3 TwistAxis01 = FJointConstants::TwistAxis();
-		FReal TwistAngle = R01Twist.GetAngle();
-		if (TwistAngle > PI)
+		FReal TwistAngleMax = FLT_MAX;
+		if (JointSettings.Motion.AngularMotionTypes[(int32)EJointAngularConstraintIndex::Twist] == EJointMotionType::Limited)
 		{
-			TwistAngle = TwistAngle - (FReal)2 * PI;
+			TwistAngleMax = JointSettings.Motion.AngularLimits[(int32)EJointAngularConstraintIndex::Twist];
 		}
-		if (R01Twist.X < 0)
+		else if (JointSettings.Motion.AngularMotionTypes[(int32)EJointAngularConstraintIndex::Twist] == EJointMotionType::Locked)
 		{
-			TwistAngle = -TwistAngle;
+			TwistAngleMax = 0;
 		}
 
-		const FVec3 TwistAxis0 = Rs[0] * TwistAxis01;
-		const FVec3 TwistAxis1 = Rs[1] * TwistAxis01;
+		// Calculate the twist correction to apply to each body
+		FReal DTwistAngle = 0;
+		if (TwistAngle > TwistAngleMax)
+		{
+			DTwistAngle = TwistAngle - TwistAngleMax;
+		}
+		else if (TwistAngle < -TwistAngleMax)
+		{
+			DTwistAngle = TwistAngle + TwistAngleMax;
+		}
+
+		// Apply twist correction
+		if (DTwistAngle != 0)
+		{
+			FReal TwistStiffness = FPBDJointUtilities::GetTwistStiffness(SolverSettings, JointSettings);
+			ApplyRotationCorrection(TwistStiffness, TwistAxis0, TwistAxis1, DTwistAngle);
+			UpdateDerivedState();
+		}
+	}
+
+
+	void FJointSolverGaussSeidel::ApplyTwistConstraintSoft(
+		const FReal Dt,
+		const FPBDJointSolverSettings& SolverSettings,
+		const FPBDJointSettings& JointSettings)
+	{
+		FVec3 TwistAxisLocal;
+		FReal TwistAngle;
+		GetTwistAxisAngleLocal(Rs[0], Rs[1], TwistAxisLocal, TwistAngle);
+
+		const FVec3 TwistAxis0 = Rs[0] * TwistAxisLocal;
+		const FVec3 TwistAxis1 = Rs[1] * TwistAxisLocal;
+
 		FReal TwistAngleMax = FLT_MAX;
 		if (JointSettings.Motion.AngularMotionTypes[(int32)EJointAngularConstraintIndex::Twist] == EJointMotionType::Limited)
 		{
@@ -377,7 +735,10 @@ namespace Chaos
 		// Apply twist correction
 		if (DTwistAngle != 0)
 		{
-			ApplyRotationDelta(TwistStiffness, TwistAxis0, TwistAxis1, DTwistAngle);
+			const FReal TwistStiffness = FPBDJointUtilities::GetSoftTwistStiffness(SolverSettings, JointSettings);
+			const FReal TwistDamping = FPBDJointUtilities::GetSoftTwistDamping(SolverSettings, JointSettings);
+			const bool bAccelerationMode = FPBDJointUtilities::GetSoftAccelerationMode(SolverSettings, JointSettings);
+			ApplyRotationCorrectionSoft(Dt, TwistStiffness, TwistDamping, bAccelerationMode, TwistAxis0, TwistAxis1, DTwistAngle, TwistSoftLambda);
 			UpdateDerivedState();
 		}
 	}
@@ -385,6 +746,7 @@ namespace Chaos
 
 	void FJointSolverGaussSeidel::ApplyTwistDrive(
 		const FReal Dt,
+		const FPBDJointSolverSettings& SolverSettings,
 		const FPBDJointSettings& JointSettings)
 	{
 		// Decompose rotation of body 1 relative to body 0 into swing and twist rotations, assuming twist is X axis
@@ -408,61 +770,68 @@ namespace Chaos
 		const FReal DTwistAngle = TwistAngle - TwistAngleTarget;
 
 		// Apply twist correction
-		ApplyDriveRotationDelta(Dt, AngularDriveStiffness, AngularDriveDamping, TwistAxis0, TwistAxis1, DTwistAngle, TwistLambda);
+		const FReal AngularDriveStiffness = FPBDJointUtilities::GetAngularDriveStiffness(SolverSettings, JointSettings);
+		const FReal AngularDriveDamping = FPBDJointUtilities::GetAngularDriveDamping(SolverSettings, JointSettings);
+		const bool bAccelerationMode = FPBDJointUtilities::GetDriveAccelerationMode(SolverSettings, JointSettings);
+		ApplyRotationCorrectionSoft(Dt, AngularDriveStiffness, AngularDriveDamping, bAccelerationMode, TwistAxis0, TwistAxis1, DTwistAngle, TwistDriveLambda);
 		UpdateDerivedState();
 	}
 
 
-	void FJointSolverGaussSeidel::ApplyConeConstraint(
+	void FJointSolverGaussSeidel::ApplyTwistProjection(
 		const FReal Dt,
+		const FPBDJointSolverSettings& SolverSettings,
 		const FPBDJointSettings& JointSettings)
 	{
-		// Decompose rotation of body 1 relative to body 0 into swing and twist rotations, assuming twist is X axis
-		FRotation3 R01Twist, R01Swing;
-		FPBDJointUtilities::DecomposeSwingTwistLocal(Rs[0], Rs[1], R01Swing, R01Twist);
-
-		FVec3 SwingAxis01;
-		FReal SwingAngle;
-		R01Swing.ToAxisAndAngleSafe(SwingAxis01, SwingAngle, FJointConstants::Swing1Axis(), SwingTwistAngleTolerance);
-		if (SwingAngle > PI)
-		{
-			SwingAngle = SwingAngle - (FReal)2 * PI;
-		}
+		// @todo(ccaulfield): implement swing projection
+	}
 
-		const FVec3 SwingAxis = Rs[0] * SwingAxis01;
 
-		// Calculate swing limit for the current swing axis
-		const FReal Swing1Limit = JointSettings.Motion.AngularLimits[(int32)EJointAngularConstraintIndex::Swing1];
-		const FReal Swing2Limit = JointSettings.Motion.AngularLimits[(int32)EJointAngularConstraintIndex::Swing2];
+	void FJointSolverGaussSeidel::ApplyConeConstraint(
+		const FReal Dt,
+		const FPBDJointSolverSettings& SolverSettings,
+		const FPBDJointSettings& JointSettings)
+	{
+		// Calculate swing angle and axis
+		FReal SwingAngle;
+		FVec3 SwingAxisLocal;
+		GetConeAxisAngleLocal(Rs[0], Rs[1], SolverSettings.SwingTwistAngleTolerance, SwingAxisLocal, SwingAngle);
+		const FVec3 SwingAxis = Rs[0] * SwingAxisLocal;
 
-		// Circular swing limit
-		FReal SwingAngleMax = Swing1Limit;
+		// Calculate swing angle error
+		FReal DSwingAngle = GetConeAngleError(JointSettings, SwingAxisLocal, SwingAngle);
 
-		// Elliptical swing limit
-		// @todo(ccaulfield): do elliptical constraints properly (axis is still for circular limit)
-		if (!FMath::IsNearlyEqual(Swing1Limit, Swing2Limit, KINDA_SMALL_NUMBER))
+		// Apply swing correction to each body
+		if (DSwingAngle != 0)
 		{
-			// Map swing axis to ellipse and calculate limit for this swing axis
-			const FReal DotSwing1 = FMath::Abs(FVec3::DotProduct(SwingAxis01, FJointConstants::Swing1Axis()));
-			const FReal DotSwing2 = FMath::Abs(FVec3::DotProduct(SwingAxis01, FJointConstants::Swing2Axis()));
-			SwingAngleMax = FMath::Sqrt(Swing1Limit * DotSwing1 * Swing1Limit * DotSwing1 + Swing2Limit * DotSwing2 * Swing2Limit * DotSwing2);
+			FReal SwingStiffness = FPBDJointUtilities::GetSwingStiffness(SolverSettings, JointSettings);
+			ApplyRotationCorrection(SwingStiffness, SwingAxis, SwingAxis, DSwingAngle);
+			UpdateDerivedState();
 		}
+	}
 
-		// Calculate swing error we need to correct
-		FReal DSwingAngle = 0;
-		if (SwingAngle > SwingAngleMax)
-		{
-			DSwingAngle = SwingAngle - SwingAngleMax;
-		}
-		else if (SwingAngle < -SwingAngleMax)
-		{
-			DSwingAngle = SwingAngle + SwingAngleMax;
-		}
 
-		// Apply swing correction
+	void FJointSolverGaussSeidel::ApplyConeConstraintSoft(
+		const FReal Dt,
+		const FPBDJointSolverSettings& SolverSettings,
+		const FPBDJointSettings& JointSettings)
+	{
+		// Calculate swing angle and axis
+		FReal SwingAngle;
+		FVec3 SwingAxisLocal;
+		GetConeAxisAngleLocal(Rs[0], Rs[1], SolverSettings.SwingTwistAngleTolerance, SwingAxisLocal, SwingAngle);
+		const FVec3 SwingAxis = Rs[0] * SwingAxisLocal;
+
+		// Calculate swing angle error
+		const FReal DSwingAngle = GetConeAngleError(JointSettings, SwingAxisLocal, SwingAngle);
+
+		// Apply swing correction to each body
 		if (DSwingAngle != 0)
 		{
-			ApplyRotationDelta(SwingStiffness, SwingAxis, SwingAxis, DSwingAngle);
+			const FReal SoftSwingStiffness = FPBDJointUtilities::GetSoftSwingStiffness(SolverSettings, JointSettings);
+			const FReal SoftSwingDamping = FPBDJointUtilities::GetSoftSwingDamping(SolverSettings, JointSettings);
+			const bool bAccelerationMode = FPBDJointUtilities::GetSoftAccelerationMode(SolverSettings, JointSettings);
+			ApplyRotationCorrectionSoft(Dt, SoftSwingStiffness, SoftSwingDamping, bAccelerationMode, SwingAxis, SwingAxis, DSwingAngle, SwingSoftLambda);
 			UpdateDerivedState();
 		}
 	}
@@ -470,119 +839,173 @@ namespace Chaos
 
 	void FJointSolverGaussSeidel::ApplyConeDrive(
 		const FReal Dt,
+		const FPBDJointSolverSettings& SolverSettings,
 		const FPBDJointSettings& JointSettings)
 	{
-		// Decompose rotation of body 1 relative to body 0 into swing and twist rotations, assuming twist is X axis
-		FRotation3 R01Twist, R01Swing;
-		FPBDJointUtilities::DecomposeSwingTwistLocal(Rs[0], Rs[1], R01Swing, R01Twist);
+		// Calculate swing angle and axis
+		FReal SwingAngle;
+		FVec3 SwingAxisLocal;
+		GetConeAxisAngleLocal(Rs[0], Rs[1], SolverSettings.SwingTwistAngleTolerance, SwingAxisLocal, SwingAngle);
+		const FVec3 SwingAxis = Rs[0] * SwingAxisLocal;
 
-		FVec3 SwingAxis01;
+		// Circular swing target (max of Swing1, Swing2 targets)
+		// @todo(ccaulfield): what should cone target really do?
+		const FReal Swing1Target = JointSettings.Motion.AngularDriveTargetAngles[(int32)EJointAngularConstraintIndex::Swing1];
+		const FReal Swing2Target = JointSettings.Motion.AngularDriveTargetAngles[(int32)EJointAngularConstraintIndex::Swing2];
+		const FReal SwingAngleTarget = FMath::Max(Swing1Target, Swing2Target);
+		const FReal DSwingAngle = SwingAngle - SwingAngleTarget;
+
+		// Apply drive forces to each body
+		const FReal AngularDriveStiffness = FPBDJointUtilities::GetAngularDriveStiffness(SolverSettings, JointSettings);
+		const FReal AngularDriveDamping = FPBDJointUtilities::GetAngularDriveDamping(SolverSettings, JointSettings);
+		const bool bAccelerationMode = FPBDJointUtilities::GetDriveAccelerationMode(SolverSettings, JointSettings);
+		ApplyRotationCorrectionSoft(Dt, AngularDriveStiffness, AngularDriveDamping, bAccelerationMode, SwingAxis, SwingAxis, DSwingAngle, SwingDriveLambda);
+		UpdateDerivedState();
+	}
+
+
+	void FJointSolverGaussSeidel::ApplyConeProjection(
+		const FReal Dt,
+		const FPBDJointSolverSettings& SolverSettings,
+		const FPBDJointSettings& JointSettings)
+	{
+		// Calculate swing angle and axis
 		FReal SwingAngle;
-		R01Swing.ToAxisAndAngleSafe(SwingAxis01, SwingAngle, FJointConstants::Swing1Axis(), SwingTwistAngleTolerance);
-		if (SwingAngle > PI)
+		FVec3 SwingAxisLocal;
+		GetConeAxisAngleLocal(Rs[0], Rs[1], SolverSettings.SwingTwistAngleTolerance, SwingAxisLocal, SwingAngle);
+		const FVec3 SwingAxis = Rs[0] * SwingAxisLocal;
+
+		// Calculate swing angle error
+		const FReal DSwingAngle = GetConeAngleError(JointSettings, SwingAxisLocal, SwingAngle);
+
+		// Apply rotation to body 1, set relative angular velocity to zero, if it would increase the error
+		if (DSwingAngle != 0)
 		{
-			SwingAngle = SwingAngle - (FReal)2 * PI;
-		}
+			const FReal AngularProjection = FPBDJointUtilities::GetAngularProjection(SolverSettings, JointSettings);
 
-		const FVec3 SwingAxis = Rs[0] * SwingAxis01;
+			const FVec3 DR1 = -DSwingAngle * SwingAxis;
+			ApplyRotationDelta(1, AngularProjection, DR1);
 
-		// Circular swing target (max of Swing1, Swing2 targets)
-		// @todo(ccaulfield): what should cone target really do?
-		FReal Swing1Target = JointSettings.Motion.AngularDriveTargetAngles[(int32)EJointAngularConstraintIndex::Swing1];
-		FReal Swing2Target = JointSettings.Motion.AngularDriveTargetAngles[(int32)EJointAngularConstraintIndex::Swing2];
-		FReal SwingAngleTarget = FMath::Max(Swing1Target, Swing2Target);
-		FReal DSwingAngle = SwingAngle - SwingAngleTarget;
+			const FReal WAxis = FVec3::DotProduct(Ws[1] - Ws[0], SwingAxis);
+			if (FMath::Sign(WAxis) == FMath::Sign(DSwingAngle))
+			{
+				const FVec3 DW1 = -WAxis * SwingAxis;
+				ApplyVelocityDelta(1, AngularProjection, FVec3(0), DW1);
+			}
 
-		ApplyDriveRotationDelta(Dt, AngularDriveStiffness, AngularDriveDamping, SwingAxis, SwingAxis, DSwingAngle, SwingLambda);
-		UpdateDerivedState();
+			UpdateDerivedState(1);
+		}
 	}
 
 
 	void FJointSolverGaussSeidel::ApplySwingConstraint(
 		const FReal Dt,
+		const FPBDJointSolverSettings& SolverSettings,
 		const FPBDJointSettings& JointSettings,
 		const EJointAngularConstraintIndex SwingConstraintIndex,
 		const EJointAngularAxisIndex SwingAxisIndex)
 	{
-		// @todo(ccaulfield): this can be made much simpler I think...swing axis should just be Body0's swing axis
+		FVec3 SwingAxis;
+		FReal SwingAngle;
+		GetSwingAxisAngle(Rs[0], Rs[1], SolverSettings.SwingTwistAngleTolerance, SwingConstraintIndex, SwingAxisIndex, SwingAxis, SwingAngle);
 
-		// Decompose rotation of body 1 relative to body 0 into swing and twist rotations, assuming twist is X axis
-		FRotation3 R01Twist, R01Swing;
-		FPBDJointUtilities::DecomposeSwingTwistLocal(Rs[0], Rs[1], R01Swing, R01Twist);
+		const bool isLocked = (JointSettings.Motion.AngularMotionTypes[(int32)SwingConstraintIndex] == EJointMotionType::Locked);
+		const FReal SwingAngleMax = isLocked ? 0 : JointSettings.Motion.AngularLimits[(int32)SwingConstraintIndex];
 
-		FVec3 TwistAxis01 = FJointConstants::TwistAxis();
-		FReal TwistAngle = R01Twist.GetAngle();
-		if (TwistAngle > PI)
+		// Calculate swing error we need to correct
+		FReal DSwingAngle = 0;
+		if (SwingAngle > SwingAngleMax)
 		{
-			TwistAngle = TwistAngle - (FReal)2 * PI;
+			DSwingAngle = SwingAngle - SwingAngleMax;
 		}
-		if (R01Twist.X < 0)
+		else if (SwingAngle < -SwingAngleMax)
 		{
-			TwistAngle = -TwistAngle;
+			DSwingAngle = SwingAngle + SwingAngleMax;
 		}
-		const FVec3 TwistAxis = Rs[0] * TwistAxis01;
 
-		const FRotation3 R1NoTwist = Rs[1] * R01Twist.Inverse();
-		const FMatrix33 Axes0 = Rs[0].ToMatrix();
-		const FMatrix33 Axes1 = R1NoTwist.ToMatrix();
-		FVec3 SwingCross = FVec3::CrossProduct(Axes0.GetAxis((int32)SwingAxisIndex), Axes1.GetAxis((int32)SwingAxisIndex));
-		SwingCross = SwingCross - FVec3::DotProduct(TwistAxis, SwingCross) * TwistAxis;
-		const FReal SwingCrossLen = SwingCross.Size();
-		if (SwingCrossLen > KINDA_SMALL_NUMBER)
+		// Apply swing correction
+		if (DSwingAngle != 0)
 		{
-			const FVec3 SwingAxis = SwingCross / SwingCrossLen;
+			const FReal SwingStiffness = FPBDJointUtilities::GetSwingStiffness(SolverSettings, JointSettings);
+			ApplyRotationCorrection(SwingStiffness, SwingAxis, SwingAxis, DSwingAngle);
+			UpdateDerivedState();
+		}
+	}
 
-			FReal SwingAngle = FMath::Asin(FMath::Clamp(SwingCrossLen, (FReal)0, (FReal)1));
-			const FReal SwingDot = FVec3::DotProduct(Axes0.GetAxis((int32)SwingAxisIndex), Axes1.GetAxis((int32)SwingAxisIndex));
-			if (SwingDot < (FReal)0)
-			{
-				SwingAngle = (FReal)PI - SwingAngle;
-			}
 
-			FReal SwingAngleMax = FLT_MAX;
-			if (JointSettings.Motion.AngularMotionTypes[(int32)SwingConstraintIndex] == EJointMotionType::Limited)
-			{
-				FReal SwingLimit = JointSettings.Motion.AngularLimits[(int32)SwingConstraintIndex];
-				SwingAngleMax = SwingLimit;
-			}
-			else if (JointSettings.Motion.AngularMotionTypes[(int32)SwingConstraintIndex] == EJointMotionType::Locked)
-			{
-				SwingAngleMax = 0;
-			}
+	void FJointSolverGaussSeidel::ApplySwingConstraintSoft(
+		const FReal Dt,
+		const FPBDJointSolverSettings& SolverSettings,
+		const FPBDJointSettings& JointSettings,
+		const EJointAngularConstraintIndex SwingConstraintIndex,
+		const EJointAngularAxisIndex SwingAxisIndex)
+	{
+		FVec3 SwingAxis;
+		FReal SwingAngle;
+		GetSwingAxisAngle(Rs[0], Rs[1], SolverSettings.SwingTwistAngleTolerance, SwingConstraintIndex, SwingAxisIndex, SwingAxis, SwingAngle);
 
-			// Calculate swing error we need to correct
-			FReal DSwingAngle = 0;
-			if (SwingAngle > SwingAngleMax)
-			{
-				DSwingAngle = SwingAngle - SwingAngleMax;
-			}
-			else if (SwingAngle < -SwingAngleMax)
-			{
-				DSwingAngle = SwingAngle + SwingAngleMax;
-			}
+		const FReal SwingAngleMax = JointSettings.Motion.AngularLimits[(int32)SwingConstraintIndex];
 
-			// Apply swing correction
-			if (DSwingAngle != 0)
-			{
-				ApplyRotationDelta(SwingStiffness, SwingAxis, SwingAxis, DSwingAngle);
-				UpdateDerivedState();
-			}
+		// Calculate swing error we need to correct
+		FReal DSwingAngle = 0;
+		if (SwingAngle > SwingAngleMax)
+		{
+			DSwingAngle = SwingAngle - SwingAngleMax;
+		}
+		else if (SwingAngle < -SwingAngleMax)
+		{
+			DSwingAngle = SwingAngle + SwingAngleMax;
+		}
+
+		// Apply swing correction
+		if (DSwingAngle != 0)
+		{
+			const FReal SoftSwingStiffness = FPBDJointUtilities::GetSoftSwingStiffness(SolverSettings, JointSettings);
+			const FReal SoftSwingDamping = FPBDJointUtilities::GetSoftSwingDamping(SolverSettings, JointSettings);
+			const bool bAccelerationMode = FPBDJointUtilities::GetSoftAccelerationMode(SolverSettings, JointSettings);
+			ApplyRotationCorrectionSoft(Dt, SoftSwingStiffness, SoftSwingDamping, bAccelerationMode, SwingAxis, SwingAxis, DSwingAngle, SwingSoftLambda);
+			UpdateDerivedState();
 		}
 	}
 
 
 	void FJointSolverGaussSeidel::ApplySwingDrive(
 		const FReal Dt,
+		const FPBDJointSolverSettings& SolverSettings,
+		const FPBDJointSettings& JointSettings,
+		const EJointAngularConstraintIndex SwingConstraintIndex,
+		const EJointAngularAxisIndex SwingAxisIndex)
+	{
+		FVec3 SwingAxis;
+		FReal SwingAngle;
+		GetSwingAxisAngle(Rs[0], Rs[1], SolverSettings.SwingTwistAngleTolerance, SwingConstraintIndex, SwingAxisIndex, SwingAxis, SwingAngle);
+	
+		const FReal SwingAngleTarget = JointSettings.Motion.AngularDriveTargetAngles[(int32)SwingConstraintIndex];
+		const FReal DSwingAngle = SwingAngle - SwingAngleTarget;
+
+		// Apply drive forces to each body
+		const FReal AngularDriveStiffness = FPBDJointUtilities::GetAngularDriveStiffness(SolverSettings, JointSettings);
+		const FReal AngularDriveDamping = FPBDJointUtilities::GetAngularDriveDamping(SolverSettings, JointSettings);
+		const bool bAccelerationMode = FPBDJointUtilities::GetDriveAccelerationMode(SolverSettings, JointSettings);
+		ApplyRotationCorrectionSoft(Dt, AngularDriveStiffness, AngularDriveDamping, bAccelerationMode, SwingAxis, SwingAxis, DSwingAngle, SwingDriveLambda);
+		UpdateDerivedState();
+	}
+
+
+	void FJointSolverGaussSeidel::ApplySwingProjection(
+		const FReal Dt,
+		const FPBDJointSolverSettings& SolverSettings,
 		const FPBDJointSettings& JointSettings,
 		const EJointAngularConstraintIndex SwingConstraintIndex,
 		const EJointAngularAxisIndex SwingAxisIndex)
 	{
-		// @todo(ccaulfield): implement swing drive
+		// @todo(ccaulfield): implement swing projection
 	}
 
 
 	void FJointSolverGaussSeidel::ApplySLerpDrive(
 		const FReal Dt,
+		const FPBDJointSolverSettings& SolverSettings,
 		const FPBDJointSettings& JointSettings)
 	{
 		// Calculate the rotation we need to apply to resolve the rotation delta
@@ -597,13 +1020,18 @@ namespace Chaos
 			{
 				SLerpAngle = SLerpAngle - (FReal)2 * PI;
 			}
-			ApplyDriveRotationDelta(Dt, AngularDriveStiffness, AngularDriveDamping, SLerpAxis, SLerpAxis, -SLerpAngle, SwingLambda);
+
+			const FReal AngularDriveStiffness = FPBDJointUtilities::GetAngularDriveStiffness(SolverSettings, JointSettings);
+			const FReal AngularDriveDamping = FPBDJointUtilities::GetAngularDriveDamping(SolverSettings, JointSettings);
+			const bool bAccelerationMode = FPBDJointUtilities::GetDriveAccelerationMode(SolverSettings, JointSettings);
+			ApplyRotationCorrectionSoft(Dt, AngularDriveStiffness, AngularDriveDamping, bAccelerationMode, SLerpAxis, SLerpAxis, -SLerpAngle, SwingDriveLambda);
 		}
 	}
 
 
 	void FJointSolverGaussSeidel::ApplyPositionConstraint(
 		const FReal Dt,
+		const FPBDJointSolverSettings& SolverSettings,
 		const FPBDJointSettings& JointSettings)
 	{
 		// @todo(ccaulfield): we should really be calculating axes based on joint config, rather than fixing the error components
@@ -634,6 +1062,7 @@ namespace Chaos
 			const FVec3 DR0 = Utilities::Multiply(InvI0, FVec3::CrossProduct(Xs[0] - Ps[0], DX));
 			const FVec3 DR1 = Utilities::Multiply(InvI1, FVec3::CrossProduct(Xs[1] - Ps[1], -DX));
 
+			FReal LinearStiffness = FPBDJointUtilities::GetLinearStiffness(SolverSettings, JointSettings);
 			ApplyPositionDelta(LinearStiffness, DP0, DP1);
 			ApplyRotationDelta(LinearStiffness, DR0, DR1);
 			UpdateDerivedState();
@@ -643,6 +1072,7 @@ namespace Chaos
 
 	void FJointSolverGaussSeidel::ApplyPositionProjection(
 		const FReal Dt,
+		const FPBDJointSolverSettings& SolverSettings,
 		const FPBDJointSettings& JointSettings)
 	{
 		// Apply a position correction with the parent body set to infinite mass, then correct the velocity.
@@ -667,10 +1097,11 @@ namespace Chaos
 			const FVec3 DV1 = -InvMs[1] * DV;
 			const FVec3 DW1 = Utilities::Multiply(InvI1, FVec3::CrossProduct(Xs[1] - Ps[1], -DV));
 
-			ApplyPositionDelta(LinearProjection, FVec3(0), DP1);
-			ApplyRotationDelta(LinearProjection, FVec3(0), DR1);
-			ApplyVelocityDelta(LinearProjection, FVec3(0), FVec3(0), DV1, DW1);
-			UpdateDerivedState();
+			FReal LinearProjection = FPBDJointUtilities::GetLinearProjection(SolverSettings, JointSettings);
+			ApplyPositionDelta(1, LinearProjection, DP1);
+			ApplyRotationDelta(1, LinearProjection, DR1);
+			ApplyVelocityDelta(1, LinearProjection, DV1, DW1);
+			UpdateDerivedState(1);
 		}
 	}
 }
diff --git a/Engine/Source/Runtime/Experimental/Chaos/Private/Chaos/PBDJointConstraintUtilities.cpp b/Engine/Source/Runtime/Experimental/Chaos/Private/Chaos/PBDJointConstraintUtilities.cpp
--- a/Engine/Source/Runtime/Experimental/Chaos/Private/Chaos/PBDJointConstraintUtilities.cpp
+++ b/Engine/Source/Runtime/Experimental/Chaos/Private/Chaos/PBDJointConstraintUtilities.cpp
@@ -15,6 +15,7 @@ namespace Chaos
 		R01.ToSwingTwistX(R01Swing, R01Twist);
 	}
 
+	// @todo(ccaulfield): separate linear soft and stiff
 	FReal FPBDJointUtilities::GetLinearStiffness(
 		const FPBDJointSolverSettings& SolverSettings,
 		const FPBDJointSettings& JointSettings)
@@ -26,28 +27,46 @@ namespace Chaos
 		return Stiffness;
 	}
 
-
 	FReal FPBDJointUtilities::GetTwistStiffness(
 		const FPBDJointSolverSettings& SolverSettings,
 		const FPBDJointSettings& JointSettings)
 	{
-		const FReal SolverStiffness = (SolverSettings.Stiffness > (FReal)0) ? SolverSettings.Stiffness : JointSettings.Motion.Stiffness;
-		const FReal SoftSolverStiffness = (SolverSettings.SoftAngularStiffness > (FReal)0) ? SolverSettings.SoftAngularStiffness : JointSettings.Motion.SoftTwistStiffness;
-		const bool bIsSoft = JointSettings.Motion.bSoftTwistLimitsEnabled && (JointSettings.Motion.AngularMotionTypes[(int32)EJointAngularConstraintIndex::Twist] == EJointMotionType::Limited);
-		const FReal Stiffness = bIsSoft ? SolverStiffness * SoftSolverStiffness : SolverStiffness;
-		return Stiffness;
+		return (SolverSettings.Stiffness > (FReal)0) ? SolverSettings.Stiffness : JointSettings.Motion.Stiffness;
 	}
 
+	FReal FPBDJointUtilities::GetSoftTwistStiffness(
+		const FPBDJointSolverSettings& SolverSettings,
+		const FPBDJointSettings& JointSettings)
+	{
+		return (SolverSettings.SoftTwistStiffness > 0)? SolverSettings.SoftTwistStiffness : JointSettings.Motion.SoftTwistStiffness;
+	}
+
+	FReal FPBDJointUtilities::GetSoftTwistDamping(
+		const FPBDJointSolverSettings& SolverSettings,
+		const FPBDJointSettings& JointSettings)
+	{
+		return (SolverSettings.SoftTwistDamping > 0) ? SolverSettings.SoftTwistDamping : JointSettings.Motion.SoftTwistDamping;
+	}
 
 	FReal FPBDJointUtilities::GetSwingStiffness(
 		const FPBDJointSolverSettings& SolverSettings,
 		const FPBDJointSettings& JointSettings)
 	{
-		const FReal SolverStiffness = (SolverSettings.Stiffness > (FReal)0) ? SolverSettings.Stiffness : JointSettings.Motion.Stiffness;
-		const FReal SoftSolverStiffness = (SolverSettings.SoftAngularStiffness > (FReal)0) ? SolverSettings.SoftAngularStiffness : JointSettings.Motion.SoftSwingStiffness;
-		const bool bIsSoft = JointSettings.Motion.bSoftSwingLimitsEnabled && ((JointSettings.Motion.AngularMotionTypes[(int32)EJointAngularConstraintIndex::Swing1] == EJointMotionType::Limited) || (JointSettings.Motion.AngularMotionTypes[(int32)EJointAngularConstraintIndex::Swing2] == EJointMotionType::Limited));
-		const FReal Stiffness = bIsSoft ? SolverStiffness * SoftSolverStiffness : SolverStiffness;
-		return Stiffness;
+		return (SolverSettings.Stiffness > (FReal)0) ? SolverSettings.Stiffness : JointSettings.Motion.Stiffness;
+	}
+
+	FReal FPBDJointUtilities::GetSoftSwingStiffness(
+		const FPBDJointSolverSettings& SolverSettings,
+		const FPBDJointSettings& JointSettings)
+	{
+		return (SolverSettings.SoftSwingStiffness > 0) ? SolverSettings.SoftSwingStiffness : JointSettings.Motion.SoftSwingStiffness;
+	}
+
+	FReal FPBDJointUtilities::GetSoftSwingDamping(
+		const FPBDJointSolverSettings& SolverSettings,
+		const FPBDJointSettings& JointSettings)
+	{
+		return (SolverSettings.SoftSwingDamping > 0) ? SolverSettings.SoftSwingDamping : JointSettings.Motion.SoftSwingDamping;
 	}
 
 	FReal FPBDJointUtilities::GetAngularDriveStiffness(
@@ -79,6 +98,21 @@ namespace Chaos
 		return (SolverSettings.AngularProjection > 0.0f) ? SolverSettings.AngularProjection : JointSettings.Motion.AngularProjection;
 	}
 
+	bool FPBDJointUtilities::GetSoftAccelerationMode(
+		const FPBDJointSolverSettings& SolverSettings,
+		const FPBDJointSettings& JointSettings)
+	{
+		return true;
+	}
+
+	bool FPBDJointUtilities::GetDriveAccelerationMode(
+		const FPBDJointSolverSettings& SolverSettings,
+		const FPBDJointSettings& JointSettings)
+	{
+		return true;
+	}
+
+
 	FVec3 FPBDJointUtilities::ConditionInertia(const FVec3& InI, const FReal MaxRatio)
 	{
 		FReal IMin = InI.Min();
diff --git a/Engine/Source/Runtime/Experimental/Chaos/Private/Chaos/PBDJointConstraints.cpp b/Engine/Source/Runtime/Experimental/Chaos/Private/Chaos/PBDJointConstraints.cpp
--- a/Engine/Source/Runtime/Experimental/Chaos/Private/Chaos/PBDJointConstraints.cpp
+++ b/Engine/Source/Runtime/Experimental/Chaos/Private/Chaos/PBDJointConstraints.cpp
@@ -88,13 +88,17 @@ namespace Chaos
 		, bSoftTwistLimitsEnabled(false)
 		, bSoftSwingLimitsEnabled(false)
 		, SoftLinearStiffness(0)
+		, SoftLinearDamping(0)
 		, SoftTwistStiffness(0)
+		, SoftTwistDamping(0)
 		, SoftSwingStiffness(0)
+		, SoftSwingDamping(0)
 		, AngularDriveTarget(FRotation3::FromIdentity())
 		, AngularDriveTargetAngles(FVec3(0, 0, 0))
 		, bAngularSLerpDriveEnabled(false)
 		, bAngularTwistDriveEnabled(false)
 		, bAngularSwingDriveEnabled(false)
+		, bAngularAccelerationDriveMode(true)
 		, AngularDriveStiffness(0)
 		, AngularDriveDamping(0)
 	{
@@ -113,13 +117,17 @@ namespace Chaos
 		, bSoftTwistLimitsEnabled(false)
 		, bSoftSwingLimitsEnabled(false)
 		, SoftLinearStiffness(0)
+		, SoftLinearDamping(0)
 		, SoftTwistStiffness(0)
+		, SoftTwistDamping(0)
 		, SoftSwingStiffness(0)
+		, SoftSwingDamping(0)
 		, AngularDriveTarget(FRotation3::FromIdentity())
 		, AngularDriveTargetAngles(FVec3(0, 0, 0))
 		, bAngularSLerpDriveEnabled(false)
 		, bAngularTwistDriveEnabled(false)
 		, bAngularSwingDriveEnabled(false)
+		, bAngularAccelerationDriveMode(true)
 		, AngularDriveStiffness(0)
 		, AngularDriveDamping(0)
 	{
@@ -180,7 +188,11 @@ namespace Chaos
 		, DriveStiffness((FReal)0)
 		, DriveDamping((FReal)0)
 		, SoftLinearStiffness((FReal)0)
-		, SoftAngularStiffness((FReal)0)
+		, SoftLinearDamping((FReal)0)
+		, SoftTwistStiffness(0)
+		, SoftTwistDamping(0)
+		, SoftSwingStiffness(0)
+		, SoftSwingDamping(0)
 	{
 	}
 
@@ -762,8 +774,8 @@ namespace Chaos
 
 		for (int32 PairIt = 0; PairIt < NumPairIts; ++PairIt)
 		{
-			Solver.ApplyDrives(Dt, JointSettings);
-			Solver.ApplyConstraints(Dt, JointSettings);
+			Solver.ApplyDrives(Dt, Settings, JointSettings);
+			Solver.ApplyConstraints(Dt, Settings, JointSettings);
 		}
 
 		UpdateParticleState(Particle0->CastToRigidParticle(), Dt, Solver.GetP(0), Solver.GetQ(0));
@@ -796,7 +808,7 @@ namespace Chaos
 
 		for (int32 PairIt = 0; PairIt < NumPairIts; ++PairIt)
 		{
-			Solver.ApplyProjections(Dt, JointSettings);
+			Solver.ApplyProjections(Dt, Settings, JointSettings);
 		}
 
 		UpdateParticleState(Particle0->CastToRigidParticle(), Dt, Solver.GetP(0), Solver.GetQ(0), Solver.GetV(0), Solver.GetW(0));
diff --git a/Engine/Source/Runtime/Experimental/Chaos/Public/Chaos/Joint/PBDJointSolverGaussSeidel.h b/Engine/Source/Runtime/Experimental/Chaos/Public/Chaos/Joint/PBDJointSolverGaussSeidel.h
--- a/Engine/Source/Runtime/Experimental/Chaos/Public/Chaos/Joint/PBDJointSolverGaussSeidel.h
+++ b/Engine/Source/Runtime/Experimental/Chaos/Public/Chaos/Joint/PBDJointSolverGaussSeidel.h
@@ -76,104 +76,175 @@ namespace Chaos
 
 		void ApplyConstraints(
 			const FReal Dt,
+			const FPBDJointSolverSettings& SolverSettings,
 			const FPBDJointSettings& JointSettings);
 
 		void ApplyDrives(
 			const FReal Dt,
+			const FPBDJointSolverSettings& SolverSettings,
 			const FPBDJointSettings& JointSettings);
 
 		void ApplyProjections(
 			const FReal Dt,
+			const FPBDJointSolverSettings& SolverSettings,
 			const FPBDJointSettings& JointSettings);
 
 	private:
 
+		void UpdateDerivedState(const int32 BodyIndex);
 		void UpdateDerivedState();
 
 		void ApplyPositionConstraints(
 			const FReal Dt,
+			const FPBDJointSolverSettings& SolverSettings,
 			const FPBDJointSettings& JointSettings);
 
 		void ApplyRotationConstraints(
 			const FReal Dt,
+			const FPBDJointSolverSettings& SolverSettings,
 			const FPBDJointSettings& JointSettings);
 
 		void ApplyRotationDrives(
 			const FReal Dt,
+			const FPBDJointSolverSettings& SolverSettings,
+			const FPBDJointSettings& JointSettings);
+
+		void ApplyRotationProjection(
+			const FReal Dt,
+			const FPBDJointSolverSettings& SolverSettings,
 			const FPBDJointSettings& JointSettings);
 
 
+		void ApplyPositionDelta(
+			const int32 BodyIndex,
+			const FReal Stiffness,
+			const FVec3& DP);
+
 		void ApplyPositionDelta(
 			const FReal Stiffness,
 			const FVec3& DP0,
 			const FVec3& DP1);
 
+		void ApplyRotationDelta(
+			const int32 BodyIndex,
+			const FReal Stiffness,
+			const FVec3& DR);
+
 		void ApplyRotationDelta(
 			const FReal Stiffness,
 			const FVec3& DR0,
 			const FVec3& DR1);
 
-		void ApplyRotationDelta(
+		void ApplyVelocityDelta(
+			const int32 BodyIndex,
+			const FReal Stiffness,
+			const FVec3& DV,
+			const FVec3& DW);
+
+		void ApplyVelocityDelta(
+			const FReal Stiffness,
+			const FVec3& DV0,
+			const FVec3& DW0,
+			const FVec3& DV1,
+			const FVec3& DW1);
+
+		void ApplyRotationCorrection(
 			const FReal Stiffness,
 			const FVec3& Axis0,
 			const FVec3& Axis1,
 			const FReal Angle);
 
-		void ApplyDriveRotationDelta(
+		void ApplyRotationCorrectionSoft(
 			const FReal Dt,
 			const FReal Stiffness,
 			const FReal Damping,
+			const bool bAccelerationMode,
 			const FVec3& Axis0,
 			const FVec3& Axis1,
 			const FReal Angle,
 			FReal& Lambda);
 
-		void ApplyVelocityDelta(
-			const FReal Stiffness,
-			const FVec3& DV0,
-			const FVec3& DW0,
-			const FVec3& DV1,
-			const FVec3& DW1);
-
-
 		void ApplyTwistConstraint(
 			const FReal Dt,
+			const FPBDJointSolverSettings& SolverSettings,
+			const FPBDJointSettings& JointSettings);
+
+		void ApplyTwistConstraintSoft(
+			const FReal Dt,
+			const FPBDJointSolverSettings& SolverSettings,
 			const FPBDJointSettings& JointSettings);
 
 		void ApplyTwistDrive(
 			const FReal Dt,
+			const FPBDJointSolverSettings& SolverSettings,
+			const FPBDJointSettings& JointSettings);
+
+		void ApplyTwistProjection(
+			const FReal Dt,
+			const FPBDJointSolverSettings& SolverSettings,
 			const FPBDJointSettings& JointSettings);
 
 		void ApplyConeConstraint(
 			const FReal Dt,
+			const FPBDJointSolverSettings& SolverSettings,
+			const FPBDJointSettings& JointSettings);
+
+		void ApplyConeConstraintSoft(
+			const FReal Dt,
+			const FPBDJointSolverSettings& SolverSettings,
 			const FPBDJointSettings& JointSettings);
 
 		void ApplyConeDrive(
 			const FReal Dt,
+			const FPBDJointSolverSettings& SolverSettings,
+			const FPBDJointSettings& JointSettings);
+
+		void ApplyConeProjection(
+			const FReal Dt,
+			const FPBDJointSolverSettings& SolverSettings,
 			const FPBDJointSettings& JointSettings);
 
 		void ApplySwingConstraint(
 			const FReal Dt,
+			const FPBDJointSolverSettings& SolverSettings,
+			const FPBDJointSettings& JointSettings,
+			const EJointAngularConstraintIndex SwingConstraintIndex,
+			const EJointAngularAxisIndex SwingAxisIndex);
+
+		void ApplySwingConstraintSoft(
+			const FReal Dt,
+			const FPBDJointSolverSettings& SolverSettings,
 			const FPBDJointSettings& JointSettings,
 			const EJointAngularConstraintIndex SwingConstraintIndex,
 			const EJointAngularAxisIndex SwingAxisIndex);
 
 		void ApplySwingDrive(
 			const FReal Dt,
+			const FPBDJointSolverSettings& SolverSettings,
+			const FPBDJointSettings& JointSettings,
+			const EJointAngularConstraintIndex SwingConstraintIndex,
+			const EJointAngularAxisIndex SwingAxisIndex);
+
+		void ApplySwingProjection(
+			const FReal Dt,
+			const FPBDJointSolverSettings& SolverSettings,
 			const FPBDJointSettings& JointSettings,
 			const EJointAngularConstraintIndex SwingConstraintIndex,
 			const EJointAngularAxisIndex SwingAxisIndex);
 
 		void ApplySLerpDrive(
 			const FReal Dt,
+			const FPBDJointSolverSettings& SolverSettings,
 			const FPBDJointSettings& JointSettings);
 
 		void ApplyPositionConstraint(
 			const FReal Dt,
+			const FPBDJointSolverSettings& SolverSettings,
 			const FPBDJointSettings& JointSettings);
 
 		void ApplyPositionProjection(
 			const FReal Dt,
+			const FPBDJointSolverSettings& SolverSettings,
 			const FPBDJointSettings& JointSettings);
 
 
@@ -195,23 +266,11 @@ namespace Chaos
 		// XPBD Previous iteration world-space body state
 		FRotation3 PrevQs[MaxConstrainedBodies];	// World-space particle CoM rotations
 
-		// XPBD joint drive constraint multipliers
-		FReal TwistLambda;
-		FReal SwingLambda;
-
-		// Settings
-		FReal LinearStiffness;
-		FReal TwistStiffness;
-		FReal SwingStiffness;
-		FReal AngularDriveStiffness;
-		FReal AngularDriveDamping;
-		FReal LinearProjection;
-		FReal AngularProjection;
-		FReal SwingTwistAngleTolerance;
-		bool bEnableTwistLimits;
-		bool bEnableSwingLimits;
-		bool bEnableDrives;
-		bool bAccelerationDriveMode;
+		// XPBD constraint multipliers
+		FReal TwistSoftLambda;
+		FReal SwingSoftLambda;
+		FReal TwistDriveLambda;
+		FReal SwingDriveLambda;
 	};
 
 }
\ No newline at end of file
diff --git a/Engine/Source/Runtime/Experimental/Chaos/Public/Chaos/PBDJointConstraintTypes.h b/Engine/Source/Runtime/Experimental/Chaos/Public/Chaos/PBDJointConstraintTypes.h
--- a/Engine/Source/Runtime/Experimental/Chaos/Public/Chaos/PBDJointConstraintTypes.h
+++ b/Engine/Source/Runtime/Experimental/Chaos/Public/Chaos/PBDJointConstraintTypes.h
@@ -85,8 +85,11 @@ namespace Chaos
 		bool bSoftTwistLimitsEnabled;
 		bool bSoftSwingLimitsEnabled;
 		FReal SoftLinearStiffness;
+		FReal SoftLinearDamping;
 		FReal SoftTwistStiffness;
+		FReal SoftTwistDamping;
 		FReal SoftSwingStiffness;
+		FReal SoftSwingDamping;
 
 		// @todo(ccaulfield): remove one of these
 		FRotation3 AngularDriveTarget;
@@ -95,6 +98,7 @@ namespace Chaos
 		bool bAngularSLerpDriveEnabled;
 		bool bAngularTwistDriveEnabled;
 		bool bAngularSwingDriveEnabled;
+		bool bAngularAccelerationDriveMode;
 		FReal AngularDriveStiffness;
 		FReal AngularDriveDamping;
 	};
@@ -140,6 +144,10 @@ namespace Chaos
 		FReal DriveStiffness;
 		FReal DriveDamping;
 		FReal SoftLinearStiffness;
-		FReal SoftAngularStiffness;
+		FReal SoftLinearDamping;
+		FReal SoftTwistStiffness;
+		FReal SoftTwistDamping;
+		FReal SoftSwingStiffness;
+		FReal SoftSwingDamping;
 	};
 }
\ No newline at end of file
diff --git a/Engine/Source/Runtime/Experimental/Chaos/Public/Chaos/PBDJointConstraintUtilities.h b/Engine/Source/Runtime/Experimental/Chaos/Public/Chaos/PBDJointConstraintUtilities.h
--- a/Engine/Source/Runtime/Experimental/Chaos/Public/Chaos/PBDJointConstraintUtilities.h
+++ b/Engine/Source/Runtime/Experimental/Chaos/Public/Chaos/PBDJointConstraintUtilities.h
@@ -72,10 +72,26 @@ namespace Chaos
 			const FPBDJointSolverSettings& SolverSettings,
 			const FPBDJointSettings& JointSettings);
 
+		static FReal GetSoftTwistStiffness(
+			const FPBDJointSolverSettings& SolverSettings,
+			const FPBDJointSettings& JointSettings);
+
+		static FReal GetSoftTwistDamping(
+			const FPBDJointSolverSettings& SolverSettings,
+			const FPBDJointSettings& JointSettings);
+
 		static FReal GetSwingStiffness(
 			const FPBDJointSolverSettings& SolverSettings,
 			const FPBDJointSettings& JointSettings);
 
+		static FReal GetSoftSwingStiffness(
+			const FPBDJointSolverSettings& SolverSettings,
+			const FPBDJointSettings& JointSettings);
+
+		static FReal GetSoftSwingDamping(
+			const FPBDJointSolverSettings& SolverSettings,
+			const FPBDJointSettings& JointSettings);
+
 		static FReal GetAngularDriveStiffness(
 			const FPBDJointSolverSettings& SolverSettings,
 			const FPBDJointSettings& JointSettings);
@@ -92,6 +108,14 @@ namespace Chaos
 			const FPBDJointSolverSettings& SolverSettings,
 			const FPBDJointSettings& JointSettings);
 
+		static bool GetSoftAccelerationMode(
+			const FPBDJointSolverSettings& SolverSettings,
+			const FPBDJointSettings& JointSettings);
+
+		static bool GetDriveAccelerationMode(
+			const FPBDJointSolverSettings& SolverSettings,
+			const FPBDJointSettings& JointSettings);
+
 		static FVec3 GetSphereLimitedPositionError(const FVec3& CX, const FReal Radius);
 		static FVec3 GetSphereLimitedVelocityError(const FVec3& CX, const FReal Radius, const FVec3& CV);
 		static FVec3 GetCylinderLimitedPositionError(const FVec3& InCX, const FVec3& Axis, const FReal Limit, const EJointMotionType AxisMotion);
