Fix landscape painting.
Undo my previous heightfield GridCast alg changes and cover diagonal traversal case with additional adjacent visits as temporary solution.

diff --git a/Engine/Source/Runtime/Experimental/Chaos/Private/Chaos/HeightField.cpp b/Engine/Source/Runtime/Experimental/Chaos/Private/Chaos/HeightField.cpp
--- a/Engine/Source/Runtime/Experimental/Chaos/Private/Chaos/HeightField.cpp
+++ b/Engine/Source/Runtime/Experimental/Chaos/Private/Chaos/HeightField.cpp
@@ -719,10 +719,11 @@ namespace Chaos
 				EndCell = FlatGrid.ClampIndex(EndCell);
 
 				int32 DeltaX = FMath::Abs(EndCell[0] - StartCell[0]);
-				int32 DeltaY = FMath::Abs(EndCell[1] - StartCell[1]);
+				int32 DeltaY = -FMath::Abs(EndCell[1] - StartCell[1]);
 				int32 DirX = StartCell[0] < EndCell[0] ? 1 : -1;
 				int32 DirY = StartCell[1] < EndCell[1] ? 1 : -1;
-				int32 Error = 0;
+				int32 Error =  DeltaX + DeltaY;
+
 
 				if(StartCell == EndCell)
 				{
@@ -751,18 +752,17 @@ namespace Chaos
 
 					do 
 					{
-						const int32 ErrorX = Error + DeltaY;
-						const int32 ErrorY = Error - DeltaX;
-
+						const int32 DoubleError = Error * 2;
 
-						if(FMath::Abs(ErrorX) < FMath::Abs(ErrorY))
+						if(DoubleError >= DeltaY)
 						{
-							Error = ErrorX;
+							Error += DeltaY;
 							StartCell[0] += DirX;
 						}
-						else
+
+						if(DoubleError <= DeltaX)
 						{
-							Error = ErrorY;
+							Error += DeltaX;
 							StartCell[1] += DirY;
 						}
 
@@ -779,6 +779,40 @@ namespace Chaos
 								}
 							}
 						}
+
+						// When traversing diagonally, there is a possibility of cutting corners on adjacent cell w/ current algorithm.
+						// Visit both adjacent cells to cover this issue.
+						// To visit minimal amount of cells, need to choose cell based on whichever cell boundary (vertical/horizontal) is
+						// next intersected using real value position within cells. (MaxW has a shelf in progress, holding off to not introduce bugs).
+						if (DoubleError >= DeltaY && DoubleError <= DeltaX)
+						{
+							TVector<int32, 2> TestCell = StartCell + TVector<int32, 2>(0, -DirY); // subtract because StartCell was incremented already
+							if (GetCellBounds3DScaled(TestCell, Min, Max))
+							{
+								if (TAABB<T, 3>(Min,Max).RaycastFast(StartPoint, Dir, InvDir, bParallel, CurrentLength, InvCurrentLength, ToI, HitPoint))
+								{
+									bool bContinue = Visitor.VisitRaycast(TestCell[1] * (GeomData.NumCols - 1) + TestCell[0], CurrentLength);
+									if (!bContinue)
+									{
+										return true;
+									}
+								}
+							}
+
+							TestCell = StartCell + TVector<int32, 2>(-DirX, 0);  // subtract because StartCell was incremented already
+							if (GetCellBounds3DScaled(TestCell, Min, Max))
+							{
+								if (TAABB<T, 3>(Min,Max).RaycastFast(StartPoint, Dir, InvDir, bParallel, CurrentLength, InvCurrentLength, ToI, HitPoint))
+								{
+									bool bContinue = Visitor.VisitRaycast(TestCell[1] * (GeomData.NumCols - 1) + TestCell[0], CurrentLength);
+									if (!bContinue)
+									{
+										return true;
+									}
+								}
+							}
+						}
+
 					} while(StartCell != EndCell);
 				}
 			}
@@ -854,12 +888,12 @@ namespace Chaos
 			EndCell = FlatGrid.ClampIndex(EndCell);
 
 			const int32 DeltaX = FMath::Abs(EndCell[0] - StartCell[0]);
-			const int32 DeltaY = FMath::Abs(EndCell[1] - StartCell[1]);
+			const int32 DeltaY = -FMath::Abs(EndCell[1] - StartCell[1]);
 			const bool bSameCell = DeltaX == 0 && DeltaY == 0;
 
 			const int32 DirX = StartCell[0] < EndCell[0] ? 1 : -1;
 			const int32 DirY = StartCell[1] < EndCell[1] ? 1 : -1;
-			int32 Error = 0;
+			int32 Error = DeltaX + DeltaY;
 			const TVector<int32, 2> ThickenDir = FMath::Abs(DeltaX) > FMath::Abs(DeltaY) ? TVector<int32, 2>(0, 1) : TVector<int32, 2>(1, 0);
 
 			struct FQueueEntry
@@ -998,17 +1032,17 @@ namespace Chaos
 						// Walk the line and add to the queue
 						if(StartCell != EndCell)
 						{
-							const int32 ErrorX = Error + DeltaY;
-							const int32 ErrorY = Error - DeltaX;
+							const int32 DoubleError = Error * 2;
 
-							if(FMath::Abs(ErrorX) < FMath::Abs(ErrorY))
+							if(DoubleError >= DeltaY)
 							{
-								Error = ErrorX;
+								Error += DeltaY;
 								StartCell[0] += DirX;
 							}
-							else
+
+							if(DoubleError <= DeltaX)
 							{
-								Error = ErrorY;
+								Error += DeltaX;
 								StartCell[1] += DirY;
 							}
 
