Chaos : Collisions

- Separated the preprocessing of the wrapper types from the pairwise type function.

diff --git a/Engine/Source/Runtime/Experimental/Chaos/Private/Chaos/CollisionResolution.cpp b/Engine/Source/Runtime/Experimental/Chaos/Private/Chaos/CollisionResolution.cpp
--- a/Engine/Source/Runtime/Experimental/Chaos/Private/Chaos/CollisionResolution.cpp
+++ b/Engine/Source/Runtime/Experimental/Chaos/Private/Chaos/CollisionResolution.cpp
@@ -1114,37 +1114,48 @@ namespace Chaos
 			const TRigidTransform<T, d>& Transform0 = ConstraintBase.ImplicitTransform[0] * ATM;
 			const TRigidTransform<T, d>& Transform1 = ConstraintBase.ImplicitTransform[1] * BTM;
 
-			//
-			// @todo(chaos): Collision Constraints (CollisionMap)
-			//    Modify Construct() and Update() use a CollisionMap indexed on 
-			//    EImplicitObjectType, instead of the if/else chain. Also, remove 
-			//    the blocks with the ensure(false), they are just for validation 
-			//    after the recent change. 
-			//
-			EImplicitObjectType Implicit0Type = Implicit0.GetType();
-			EImplicitObjectType Implicit1Type = Implicit1.GetType();
+#if !UE_BUILD_SHIPPING
+			EImplicitObjectType Implicit0OuterType = Implicit0.GetType();
+			EImplicitObjectType Implicit1OuterType = Implicit1.GetType();
 
-			if (Implicit0Type == TImplicitObjectTransformed<T, d>::StaticType())
+			if (Implicit0OuterType == TImplicitObjectTransformed<T, d>::StaticType())
 			{
 				ensure(false);//should not be possible to get this type, it should already be resolved by the constraint. (see ConstructConstraints)
+				return;
 			}
-			else if (Implicit1Type == TImplicitObjectTransformed<T, d>::StaticType())
+			else if (Implicit1OuterType == TImplicitObjectTransformed<T, d>::StaticType())
 			{
 				ensure(false);//should not be possible to get this type, it should already be resolved by the constraint. (see ConstructConstraints)
+				return;
 			}
-			else if (Implicit0Type != FImplicitObjectUnion::StaticType() && Implicit1Type == FImplicitObjectUnion::StaticType())
+			else if (Implicit0OuterType != FImplicitObjectUnion::StaticType() && Implicit1OuterType == FImplicitObjectUnion::StaticType())
 			{
 				ensure(false);//should not be possible to get this type, it should already be resolved by the constraint. (see ConstructConstraints)
+				return;
 			}
-			else if (Implicit0Type == FImplicitObjectUnion::StaticType() && Implicit1Type != FImplicitObjectUnion::StaticType())
+			else if (Implicit0OuterType == FImplicitObjectUnion::StaticType() && Implicit1OuterType != FImplicitObjectUnion::StaticType())
 			{
 				ensure(false);//should not be possible to get this type, it should already be resolved by the constraint. (see ConstructConstraints)
+				return;
 			}
-			else if (Implicit0Type == FImplicitObjectUnion::StaticType() && Implicit1Type == FImplicitObjectUnion::StaticType())
+			else if (Implicit0OuterType == FImplicitObjectUnion::StaticType() && Implicit1OuterType == FImplicitObjectUnion::StaticType())
 			{
 				ensure(false);//should not be possible to get this type, it should already be resolved by the constraint. (see ConstructConstraints)
+				return;
 			}
-			else if (Implicit0Type == TBox<T, d>::StaticType() && Implicit1Type == TBox<T, d>::StaticType())
+#endif
+			//
+			// @todo(chaos): Collision Constraints (CollisionMap)
+			//    Modify Construct() and Update() use a CollisionMap indexed on 
+			//    EImplicitObjectType, instead of the if/else chain. Also, remove 
+			//    the blocks with the ensure(false), they are just for validation 
+			//    after the recent change. 
+			//
+
+			EImplicitObjectType Implicit0Type = GetInnerType(Implicit0.GetType());
+			EImplicitObjectType Implicit1Type = GetInnerType(Implicit1.GetType()); 
+			
+			if (Implicit0Type == TBox<T, d>::StaticType() && Implicit1Type == TBox<T, d>::StaticType())
 			{
 				UpdateBoxBoxManifold(ConstraintBase, Transform0, Transform1, Thickness);
 			}
@@ -1243,37 +1254,48 @@ namespace Chaos
 			const TRigidTransform<T, d>& Transform0 = ConstraintBase.ImplicitTransform[0] * ParticleTransform0;
 			const TRigidTransform<T, d>& Transform1 = ConstraintBase.ImplicitTransform[1] * ParticleTransform1;
 
-			//
-			// @todo(chaos): Collision Constraints (CollisionMap)
-			//    Modify Construct() and Update() use a CollisionMap indexed on 
-			//    EImplicitObjectType, instead of the if/else chain. Also, remove 
-			//    the blocks with the ensure(false), they are just for validation 
-			//    after the recent change. 
-			//
-			EImplicitObjectType Implicit0Type = Implicit0.GetType();
-			EImplicitObjectType Implicit1Type = Implicit1.GetType();
+#if !UE_BUILD_SHIPPING
+			EImplicitObjectType Implicit0OuterType = Implicit0.GetType();
+			EImplicitObjectType Implicit1OuterType = Implicit1.GetType();
 
-			if (Implicit0Type == TImplicitObjectTransformed<T, d>::StaticType())
+			if (Implicit0OuterType == TImplicitObjectTransformed<T, d>::StaticType())
 			{
 				ensure(false);//should not be possible to get this type, it should already be resolved by the constraint. (see ConstructConstraints)
+				return;
 			}
-			else if (Implicit1Type == TImplicitObjectTransformed<T, d>::StaticType())
+			else if (Implicit1OuterType == TImplicitObjectTransformed<T, d>::StaticType())
 			{
 				ensure(false);//should not be possible to get this type, it should already be resolved by the constraint. (see ConstructConstraints)
+				return;
 			}
-			else if (Implicit0Type != FImplicitObjectUnion::StaticType() && Implicit1Type == FImplicitObjectUnion::StaticType())
+			else if (Implicit0OuterType != FImplicitObjectUnion::StaticType() && Implicit1OuterType == FImplicitObjectUnion::StaticType())
 			{
 				ensure(false);//should not be possible to get this type, it should already be resolved by the constraint. (see ConstructConstraints)
+				return;
 			}
-			else if (Implicit0Type == FImplicitObjectUnion::StaticType() && Implicit1Type != FImplicitObjectUnion::StaticType())
+			else if (Implicit0OuterType == FImplicitObjectUnion::StaticType() && Implicit1OuterType != FImplicitObjectUnion::StaticType())
 			{
 				ensure(false);//should not be possible to get this type, it should already be resolved by the constraint. (see ConstructConstraints)
+				return;
 			}
-			else if (Implicit0Type == FImplicitObjectUnion::StaticType() && Implicit1Type == FImplicitObjectUnion::StaticType())
+			else if (Implicit0OuterType == FImplicitObjectUnion::StaticType() && Implicit1OuterType == FImplicitObjectUnion::StaticType())
 			{
 				ensure(false);//should not be possible to get this type, it should already be resolved by the constraint. (see ConstructConstraints)
+				return;
 			}
-			else if (Implicit0Type == TBox<T, d>::StaticType() && Implicit1Type == TBox<T, d>::StaticType())
+#endif
+
+			//
+			// @todo(chaos): Collision Constraints (CollisionMap)
+			//    Modify Construct() and Update() use a CollisionMap indexed on 
+			//    EImplicitObjectType, instead of the if/else chain. Also, remove 
+			//    the blocks with the ensure(false), they are just for validation 
+			//    after the recent change. 
+			//
+			EImplicitObjectType Implicit0Type = GetInnerType(Implicit0.GetType());
+			EImplicitObjectType Implicit1Type = GetInnerType(Implicit1.GetType()); 
+			
+			if (Implicit0Type == TBox<T, d>::StaticType() && Implicit1Type == TBox<T, d>::StaticType())
 			{
 				UpdateBoxBoxConstraint(Implicit0.template GetObject<TBox<T, d>>()->GetAABB(), Transform0, Implicit1.template GetObject<TBox<T, d>>()->GetAABB(), Transform1, Thickness, *ConstraintBase.template As<TRigidBodyPointContactConstraint<T, d>>());
 			}
@@ -1447,22 +1469,24 @@ namespace Chaos
 			// @todo(chaos): Collision Constraints (CollisionMap)
 			//    Modify Construct() and Update() use a CollisionMap indexed on EImplicitObjectType, instead of the if/else chain
 			//
-			EImplicitObjectType Implicit0Type = Implicit0->GetType();
-			EImplicitObjectType Implicit1Type = Implicit1->GetType();
+			EImplicitObjectType Implicit0OuterType = Implicit0->GetType();
+			EImplicitObjectType Implicit1OuterType = Implicit1->GetType();
 
-			if (Implicit0Type == TImplicitObjectTransformed<T, d>::StaticType())
+			if (Implicit0OuterType == TImplicitObjectTransformed<T, d>::StaticType())
 			{
 				const TImplicitObjectTransformed<FReal, 3>* TransformedImplicit0 = Implicit0->template GetObject<const TImplicitObjectTransformed<FReal, 3>>();
 				TRigidTransform<T, d> TransformedTransform0 = TransformedImplicit0->GetTransform() * Transform0;
 				ConstructConstraints(Particle0, Particle1, TransformedImplicit0->GetTransformedObject(), Implicit1, TransformedTransform0, Transform1, Thickness, NewConstraints);
+				return;
 			}
-			else if (Implicit1Type == TImplicitObjectTransformed<T, d>::StaticType())
+			else if (Implicit1OuterType == TImplicitObjectTransformed<T, d>::StaticType())
 			{
 				const TImplicitObjectTransformed<FReal, 3>* TransformedImplicit1 = Implicit1->template GetObject<const TImplicitObjectTransformed<FReal, 3>>();
 				TRigidTransform<T, d> TransformedTransform1 = TransformedImplicit1->GetTransform() * Transform1;
 				ConstructConstraints(Particle0, Particle1, Implicit0, TransformedImplicit1->GetTransformedObject(), Transform0, TransformedTransform1, Thickness, NewConstraints);
+				return;
 			}
-			else if (Implicit0Type != FImplicitObjectUnion::StaticType() && Implicit1Type == FImplicitObjectUnion::StaticType())
+			else if (Implicit0OuterType != FImplicitObjectUnion::StaticType() && Implicit1OuterType == FImplicitObjectUnion::StaticType())
 			{
 				const TArray<Pair<const FImplicitObject*, TRigidTransform<T, d>>> LevelsetShapes = FindRelevantShapes(Implicit0, Transform0, *Implicit1, Transform1, Thickness);
 				for (const Pair<const FImplicitObject*, TRigidTransform<T, d>>& LevelsetObjPair : LevelsetShapes)
@@ -1471,8 +1495,9 @@ namespace Chaos
 					const TRigidTransform<T, d> Implicit1InnerObjTM = LevelsetObjPair.Second * Transform1;
 					ConstructConstraints(Particle0, Particle1, Implicit0, Implicit1InnerObj, Transform0, Implicit1InnerObjTM, Thickness, NewConstraints);
 				}
+				return;
 			}
-			else if (Implicit0Type == FImplicitObjectUnion::StaticType() && Implicit1Type != FImplicitObjectUnion::StaticType())
+			else if (Implicit0OuterType == FImplicitObjectUnion::StaticType() && Implicit1OuterType != FImplicitObjectUnion::StaticType())
 			{
 				// [Note] forces non-unions into particle[0] position
 				const TArray<Pair<const FImplicitObject*, TRigidTransform<T, d>>> LevelsetShapes = FindRelevantShapes(Implicit1, Transform1, *Implicit0, Transform0, Thickness);
@@ -1482,12 +1507,18 @@ namespace Chaos
 					const TRigidTransform<T, d> Implicit0InnerObjTM = LevelsetObjPair.Second * Transform0;
 					ConstructConstraints(Particle0, Particle1, Implicit0InnerObj, Implicit1, Implicit0InnerObjTM, Transform1, Thickness, NewConstraints);
 				}
+				return;
 			}
-			else if (Implicit0Type == FImplicitObjectUnion::StaticType() && Implicit1Type == FImplicitObjectUnion::StaticType())
+			else if (Implicit0OuterType == FImplicitObjectUnion::StaticType() && Implicit1OuterType == FImplicitObjectUnion::StaticType())
 			{
 				ConstructUnionUnionConstraints(Particle0, Particle1, Implicit0, Implicit1, Transform0, Transform1, Thickness, NewConstraints);
+				return;
 			}
-			else if (Implicit0Type == TBox<T, d>::StaticType() && Implicit1Type == TBox<T, d>::StaticType())
+
+			EImplicitObjectType Implicit0Type = GetInnerType(Implicit0->GetType());
+			EImplicitObjectType Implicit1Type = GetInnerType(Implicit1->GetType());
+
+			if (Implicit0Type == TBox<T, d>::StaticType() && Implicit1Type == TBox<T, d>::StaticType())
 			{
 				ConstructBoxBoxConstraints(Particle0, Particle1, Implicit0, Implicit1, Transform0, Transform1, Thickness, NewConstraints);
 			}
