Chaos : Collisions

- Added support for GJKContactPoint to TriangleMeshImplicitObject.
- Stubbed out API in CollisionResolution.cpp

diff --git a/Engine/Source/Runtime/Experimental/Chaos/Private/Chaos/CollisionResolution.cpp b/Engine/Source/Runtime/Experimental/Chaos/Private/Chaos/CollisionResolution.cpp
--- a/Engine/Source/Runtime/Experimental/Chaos/Private/Chaos/CollisionResolution.cpp
+++ b/Engine/Source/Runtime/Experimental/Chaos/Private/Chaos/CollisionResolution.cpp
@@ -17,6 +17,7 @@
 #include "Chaos/PBDCollisionConstraints.h"
 #include "Chaos/Sphere.h"
 #include "Chaos/Transform.h"
+#include "Chaos/TriangleMeshImplicitObject.h"
 #include "Chaos/GeometryQueries.h"
 
 namespace Chaos
@@ -320,7 +321,7 @@ namespace Chaos
 		//
 
 		template <typename T, int d>
-		bool UpdateBoxPlaneConstraint(const TAABB<T, d>& Box, const TRigidTransform<T, d>& BoxTransform, const TPlane<T, d>& Plane, const TRigidTransform<T, d>& PlaneTransform, const T Thickness, TRigidBodyPointContactConstraint<T, d>& Constraint)
+		void UpdateBoxPlaneConstraint(const TAABB<T, d>& Box, const TRigidTransform<T, d>& BoxTransform, const TPlane<T, d>& Plane, const TRigidTransform<T, d>& PlaneTransform, const T Thickness, TRigidBodyPointContactConstraint<T, d>& Constraint)
 		{
 			TCollisionContact<T, d> & Contact = Constraint.Manifold;
 
@@ -329,7 +330,6 @@ namespace Chaos
 				MSVC_PRAGMA(warning(disable : ALL_CODE_ANALYSIS_WARNINGS))
 #endif	// USING_CODE_ANALYSIS
 
-				bool bApplied = false;
 			const TRigidTransform<T, d> BoxToPlaneTransform(BoxTransform.GetRelativeTransform(PlaneTransform));
 			const TVector<T, d> Extents = Box.Extents();
 			constexpr int32 NumCorners = 2 + 2 * d;
@@ -349,7 +349,7 @@ namespace Chaos
 			MSVC_PRAGMA(warning(pop))
 #endif	// USING_CODE_ANALYSIS
 
-				TVector<T, d> PotentialConstraints[NumCorners];
+			TVector<T, d> PotentialConstraints[NumCorners];
 			int32 NumConstraints = 0;
 			for (int32 i = 0; i < NumCorners; ++i)
 			{
@@ -365,7 +365,6 @@ namespace Chaos
 					Contact.Normal = PlaneTransform.TransformVector(Normal);
 					Contact.Location = PlaneTransform.TransformPosition(Corners[i]);
 					PotentialConstraints[NumConstraints++] = Contact.Location;
-					bApplied = true;
 				}
 			}
 			if (NumConstraints > 1)
@@ -377,8 +376,6 @@ namespace Chaos
 				}
 				Contact.Location = AverageLocation / NumConstraints;
 			}
-
-			return bApplied;
 		}
 
 		template<class T, int d>
@@ -414,6 +411,62 @@ namespace Chaos
 			}
 		}
 
+		//
+		// Box-TriangleMesh
+		//
+
+		template <typename T, int d>
+		TContactPoint<T> BoxTriangleMeshContactPoint(const TAABB<T, d>& A, const TRigidTransform<T, d>& ATransform, const FTriangleMeshImplicitObject& B, const TRigidTransform<T, d>& BTransform, const T Thickness)
+		{
+			return GJKImplicitContactPoint< TBox<float, 3> >(TBox<float, 3>(A), ATransform, B, BTransform, Thickness);
+		}
+
+		template <typename T, int d>
+		void UpdateBoxTriangleMeshConstraint(const TAABB<T, d>& Box0, const TRigidTransform<T, d>& Transform0, const FTriangleMeshImplicitObject& TriangleMesh1, const TRigidTransform<T, d>& Transform1, const T Thickness, TRigidBodyPointContactConstraint<T, d>& Constraint)
+		{
+			UpdateContactPoint(Constraint.Manifold, BoxTriangleMeshContactPoint(Box0, Transform0, TriangleMesh1, Transform1, Thickness));
+		}
+
+
+		template <typename T, int d>
+		void UpdateBoxTriangleMeshManifold(TCollisionConstraintBase<T, d>&  Constraint, const TRigidTransform<T, d>& ATM, const TRigidTransform<T, d>& BTM, const T Thickness)
+		{
+
+		}
+
+		template <typename TPGeometryClass>
+		const TPGeometryClass* GetUnscaledInner(const FImplicitObject& Geometry)
+		{
+			if (const TImplicitObjectScaled<TPGeometryClass>* ScaledConvexImplicit = Geometry.template GetObject<const TImplicitObjectScaled<TPGeometryClass> >())
+				return (Geometry.template GetObject<const TImplicitObjectScaled<TPGeometryClass> >())->GetUnscaledObject();
+			else if (const TPGeometryClass* ConvexImplicit = Geometry.template GetObject<const TPGeometryClass>())
+				return Geometry.template GetObject<const TPGeometryClass>();
+			return nullptr;
+		}
+
+		template<typename T, int d>
+		void ConstructBoxTriangleMeshConstraints(TGeometryParticleHandle<T, d>* Particle0, TGeometryParticleHandle<T, d>* Particle1, const FImplicitObject* Implicit0, const FImplicitObject* Implicit1, const TRigidTransform<T, d>& Transform0, const TRigidTransform<T, d>& Transform1, const T Thickness, FCollisionConstraintsArray& NewConstraints)
+		{
+			const TBox<T, d> * Object0 = Implicit0->template GetObject<const TBox<T, d> >();
+			const FTriangleMeshImplicitObject * Object1 = GetUnscaledInner<FTriangleMeshImplicitObject>(*Implicit1);
+			if (ensure(Object0 && Object1))
+			{
+				TRigidTransform<T, d> ParticleImplicit0TM = Transform0.GetRelativeTransform(Collisions::GetTransform(Particle0));
+				TRigidTransform<T, d> ParticleImplicit1TM = Transform1.GetRelativeTransform(Collisions::GetTransform(Particle1));
+				FRigidBodyPointContactConstraint* Constraint = new FRigidBodyPointContactConstraint(Particle0, Implicit0, ParticleImplicit0TM, Particle1, Implicit1, ParticleImplicit1TM);
+
+				UpdateBoxTriangleMeshConstraint(Object0->GetAABB(), Transform0, *Object1, Transform1, Thickness, *Constraint);
+
+				if (Constraint->GetPhi() < Thickness)
+				{
+					NewConstraints.Add(Constraint);
+				}
+				else
+				{
+					delete Constraint;
+				}
+			}
+		}
 
 		//
 		// Sphere - Sphere
@@ -1144,6 +1197,7 @@ namespace Chaos
 				return;
 			}
 #endif
+
 			//
 			// @todo(chaos): Collision Constraints (CollisionMap)
 			//    Modify Construct() and Update() use a CollisionMap indexed on 
@@ -1221,12 +1275,13 @@ namespace Chaos
 			}
 			else if (Implicit0Type == TBox<T, d>::StaticType() && Implicit1Type == TCapsule<T>::StaticType())
 			{
-				//UpdateBoxSphereManifold(ConstraintBase, Transform0, Transform1, Thickness);
+				//UpdateBoxCapsuleManifold(ConstraintBase, Transform0, Transform1, Thickness);
 			}
 			else if (Implicit0Type == TCapsule<T>::StaticType() && Implicit1Type == TSphere<T, d>::StaticType())
 			{
 				//UpdateCapsuleSphereManifold(ConstraintBase, Transform0, Transform1, Thickness);
 			}
+			// the generic convex bodies are last
 			else if (Implicit0Type == THeightField<T>::StaticType() && Implicit1.IsConvex())
 			{
 				ensure(false);
@@ -1243,6 +1298,18 @@ namespace Chaos
 			{
 				UpdateLevelsetLevelsetManifold(ConstraintBase, Transform0, Transform1, Thickness);
 			}
+
+			/*
+			else if (Implicit0Type == TBox<T, d>::StaticType() && Implicit1Type == FTriangleMeshImplicitObject::StaticType())
+			{
+				UpdateBoxTriangleMeshManifold(ConstraintBase, Transform0, Transform1, Thickness);
+			}
+			else if (Implicit0Type == FTriangleMeshImplicitObject::StaticType() && Implicit1Type == TBox<T, d>::StaticType())
+			{
+				ensure(false);
+			}
+			*/
+
 		}
 
 		template<ECollisionUpdateType UpdateType, typename T, int d>
@@ -1421,6 +1488,7 @@ namespace Chaos
 				//     second.
 				ensure(false);
 			}
+			// the generic convex bodies are last
 			else if (Implicit0Type == THeightField<T>::StaticType() && Implicit1.IsConvex())
 			{
 				//     This case should not be necessary. The height fields 
@@ -1442,6 +1510,23 @@ namespace Chaos
 			{
 				UpdateLevelsetLevelsetConstraint<UpdateType>(Thickness, *ConstraintBase.template As<TRigidBodyPointContactConstraint<T, d> >());
 			}
+
+			/*
+			else if (Implicit0Type == TBox<T, d>::StaticType() && Implicit1Type == FTriangleMeshImplicitObject::StaticType())
+			{
+			UpdateBoxTriangleMeshConstraint(Implicit0.template GetObject<TBox<T, d>>()->GetAABB(), Transform0, *GetUnscaledInner<FTriangleMeshImplicitObject>(Implicit1), Transform1, Thickness, *ConstraintBase.template As<TRigidBodyPointContactConstraint<T, d>>());
+			}
+			else if (Implicit0Type == TBox<T, d>::StaticType() && Implicit1Type == FTriangleMeshImplicitObject::StaticType())
+			{
+				//     This case should not be necessary. The triangle mesh
+				//     will only ever be collided against, so ideally will never
+				//     be in index[0] position of the constraint, also the construction
+				//     of the constraint will just switch the index position so its always
+				//     second.
+				ensure(false);
+			}
+			*/
+
 		}
 
 		template<typename T, int d>
@@ -1594,6 +1679,7 @@ namespace Chaos
 			{
 				ConstructCapsuleHeightFieldConstraints(Particle1, Particle0, Implicit1, Implicit0, Transform1, Transform0, Thickness, NewConstraints);
 			}
+			// the generic convex bodies are last
 			else if (Implicit0->IsConvex() && Implicit1Type == THeightField<T>::StaticType())
 			{
 				ConstructConvexHeightFieldConstraints(Particle0, Particle1, Implicit0, Implicit1, Transform0, Transform1, Thickness, NewConstraints);
@@ -1610,6 +1696,19 @@ namespace Chaos
 			{
 				ConstructLevelsetLevelsetConstraints(Particle0, Particle1, Implicit0, Implicit1, Transform0, Transform1, Thickness, NewConstraints);
 			}
+
+
+			/*
+			if (Implicit0Type == TBox<T, d>::StaticType() && Implicit1Type == FTriangleMeshImplicitObject::StaticType())
+			{
+				ConstructBoxTriangleMeshConstraints(Particle0, Particle1, Implicit0, Implicit1, Transform0, Transform1, Thickness, NewConstraints);
+			}
+			else if (Implicit0Type == FTriangleMeshImplicitObject::StaticType() && Implicit1Type == TBox<T, d>::StaticType())
+			{
+				ConstructBoxTriangleMeshConstraints(Particle1, Particle0, Implicit1, Implicit0, Transform1, Transform0, Thickness, NewConstraints);
+			}
+			*/
+
 		}
 
 		template void UpdateBoxBoxConstraint<float, 3>(const TAABB<float, 3>& Box1, const TRigidTransform<float, 3>& Box1Transform, const TAABB<float, 3>& Box2, const TRigidTransform<float, 3>& Box2Transform, const float Thickness, TRigidBodyPointContactConstraint<float, 3>& Constraint);
@@ -1620,7 +1719,7 @@ namespace Chaos
 		template void UpdateBoxHeightFieldManifold<float, 3>(TCollisionConstraintBase<float, 3>&  Constraint, const TRigidTransform<float, 3>& ATM, const TRigidTransform<float, 3>& BTM, const float Thickness);
 		template void ConstructBoxHeightFieldConstraints<float, 3>(TGeometryParticleHandle<float, 3>* Particle0, TGeometryParticleHandle<float, 3>* Particle1, const FImplicitObject* Implicit0, const FImplicitObject* Implicit1, const TRigidTransform<float, 3>& Transform0, const TRigidTransform<float, 3>& Transform1, const float Thickness, FCollisionConstraintsArray& NewConstraints);
 
-		template bool UpdateBoxPlaneConstraint<float, 3>(const TAABB<float, 3>& Box, const TRigidTransform<float, 3>& BoxTransform, const TPlane<float, 3>& Plane, const TRigidTransform<float, 3>& PlaneTransform, const float Thickness, TRigidBodyPointContactConstraint<float, 3>& Constraint);
+		template void UpdateBoxPlaneConstraint<float, 3>(const TAABB<float, 3>& Box, const TRigidTransform<float, 3>& BoxTransform, const TPlane<float, 3>& Plane, const TRigidTransform<float, 3>& PlaneTransform, const float Thickness, TRigidBodyPointContactConstraint<float, 3>& Constraint);
 		template void UpdateBoxPlaneManifold<float, 3>(TCollisionConstraintBase<float, 3>&  Constraint, const TRigidTransform<float, 3>& ATM, const TRigidTransform<float, 3>& BTM, const float Thickness);
 		template void ConstructBoxPlaneConstraints<float, 3>(TGeometryParticleHandle<float, 3>* Particle0, TGeometryParticleHandle<float, 3>* Particle1, const FImplicitObject* Implicit0, const FImplicitObject* Implicit1, const TRigidTransform<float, 3>& Transform0, const TRigidTransform<float, 3>& Transform1, const float Thickness, FCollisionConstraintsArray& NewConstraints);
 
diff --git a/Engine/Source/Runtime/Experimental/Chaos/Private/Chaos/TriangleMeshImplicitObject.cpp b/Engine/Source/Runtime/Experimental/Chaos/Private/Chaos/TriangleMeshImplicitObject.cpp
--- a/Engine/Source/Runtime/Experimental/Chaos/Private/Chaos/TriangleMeshImplicitObject.cpp
+++ b/Engine/Source/Runtime/Experimental/Chaos/Private/Chaos/TriangleMeshImplicitObject.cpp
@@ -233,6 +233,133 @@ bool FTriangleMeshImplicitObject::Raycast(const FVec3& StartPoint, const FVec3&
 	}
 }
 
+
+template <typename QueryGeomType>
+bool FTriangleMeshImplicitObject::GJKContactPointImp(const QueryGeomType& QueryGeom, const FRigidTransform3& QueryTM, const FReal Thickness, FVec3& Location, FVec3& Normal, FReal& Penetration) const
+{
+	bool bResult = false;
+
+	auto OverlapTriangle = [&](const FVec3& A, const FVec3& B, const FVec3& C,
+		FVec3& LocalContactLocation, FVec3& LocalContactNormal, FReal& LocalContactPhi) -> bool
+	{
+		const FVec3 AB = B - A;
+		const FVec3 AC = C - A;
+		TTriangle<FReal> TriangleConvex(A, B, C);
+
+		FReal LambdaPenetration;
+		FVec3 ClosestA, ClosestB, LambdaNormal;
+		if (GJKPenetration(TriangleConvex, QueryGeom, QueryTM, LambdaPenetration, ClosestA, ClosestB, LambdaNormal, (FReal)0))
+		{
+			FVec3 TestVector = QueryTM.InverseTransformVector(LambdaNormal);
+
+			LocalContactLocation = ClosestB;
+			LocalContactNormal = LambdaNormal;
+			LocalContactPhi = -LambdaPenetration;
+			return true;
+		}
+
+		return LocalContactPhi < 0.f;
+	};
+
+
+	auto LambdaHelper = [&](const auto& Elements)
+	{
+		FReal LocalContactPhi = FLT_MAX;
+		FVec3 LocalContactLocation, LocalContactNormal;
+
+		TAABB<FReal, 3> QueryBounds = QueryGeom.BoundingBox();
+		QueryBounds.Thicken(Thickness);
+		QueryBounds = QueryBounds.TransformedAABB(QueryTM);
+		const TArray<int32> PotentialIntersections = BVH.FindAllIntersections(QueryBounds);
+
+		for (int32 TriIdx : PotentialIntersections)
+		{
+			//It's most likely that the query object is in front of the triangle since queries tend to be on the outside.
+			//However, maybe we should check if it's behind the triangle plane. Also, we should enforce this winding in some way
+			TVec3<FReal> A, B, C;
+			TransformVertsHelper(QueryGeom, TriIdx, MParticles, Elements, A, B, C);
+
+			if (OverlapTriangle(A, B, C, LocalContactLocation, LocalContactNormal, LocalContactPhi))
+			{
+				if (LocalContactPhi < Penetration)
+				{
+					Penetration = LocalContactPhi;
+					Location = LocalContactLocation;
+					Normal = LocalContactNormal;
+				}
+			}
+
+			if (OverlapTriangle(A, C, B, LocalContactLocation, LocalContactNormal, LocalContactPhi))
+			{
+				if (LocalContactPhi < Penetration)
+				{
+					Penetration = LocalContactPhi;
+					Location = LocalContactLocation;
+					Normal = LocalContactNormal;
+				}
+			}
+
+		}
+
+		return Penetration < Thickness;
+	};
+
+	if (MElements.RequiresLargeIndices())
+	{
+		return LambdaHelper(MElements.GetLargeIndexBuffer());
+	}
+	else
+	{
+		return LambdaHelper(MElements.GetSmallIndexBuffer());
+	}
+
+	if (Penetration < 0)
+		return true;
+	return false;
+}
+
+bool FTriangleMeshImplicitObject::GJKContactPoint(const TSphere<FReal, 3>& QueryGeom, const FRigidTransform3& QueryTM, const FReal Thickness, FVec3& Location, FVec3& Normal, FReal& Penetration) const
+{
+	return GJKContactPointImp(QueryGeom, QueryTM, Thickness, Location, Normal, Penetration);
+}
+
+bool FTriangleMeshImplicitObject::GJKContactPoint(const TBox<FReal, 3>& QueryGeom, const FRigidTransform3& QueryTM, const FReal Thickness, FVec3& Location, FVec3& Normal, FReal& Penetration) const
+{
+	return GJKContactPointImp(QueryGeom, QueryTM, Thickness, Location, Normal, Penetration);
+}
+
+bool FTriangleMeshImplicitObject::GJKContactPoint(const TCapsule<FReal>& QueryGeom, const FRigidTransform3& QueryTM, const FReal Thickness, FVec3& Location, FVec3& Normal, FReal& Penetration) const
+{
+	return GJKContactPointImp(QueryGeom, QueryTM, Thickness, Location, Normal, Penetration);
+}
+
+bool FTriangleMeshImplicitObject::GJKContactPoint(const FConvex& QueryGeom, const FRigidTransform3& QueryTM, const FReal Thickness, FVec3& Location, FVec3& Normal, FReal& Penetration) const
+{
+	return GJKContactPointImp(QueryGeom, QueryTM, Thickness, Location, Normal, Penetration);
+}
+
+bool FTriangleMeshImplicitObject::GJKContactPoint(const TImplicitObjectScaled< TSphere<FReal, 3> >& QueryGeom, const FRigidTransform3& QueryTM, const FReal Thickness, FVec3& Location, FVec3& Normal, FReal& Penetration) const
+{
+	return GJKContactPointImp(QueryGeom, QueryTM, Thickness, Location, Normal, Penetration);
+}
+
+bool FTriangleMeshImplicitObject::GJKContactPoint(const TImplicitObjectScaled< TBox<FReal, 3> >& QueryGeom, const FRigidTransform3& QueryTM, const FReal Thickness, FVec3& Location, FVec3& Normal, FReal& Penetration) const
+{
+	return GJKContactPointImp(QueryGeom, QueryTM, Thickness, Location, Normal, Penetration);
+}
+
+bool FTriangleMeshImplicitObject::GJKContactPoint(const TImplicitObjectScaled< TCapsule<FReal> >& QueryGeom, const FRigidTransform3& QueryTM, const FReal Thickness, FVec3& Location, FVec3& Normal, FReal& Penetration) const
+{
+	return GJKContactPointImp(QueryGeom, QueryTM, Thickness, Location, Normal, Penetration);
+}
+
+bool FTriangleMeshImplicitObject::GJKContactPoint(const TImplicitObjectScaled< FConvex >& QueryGeom, const FRigidTransform3& QueryTM, const FReal Thickness, FVec3& Location, FVec3& Normal, FReal& Penetration) const
+{
+	return GJKContactPointImp(QueryGeom, QueryTM, Thickness, Location, Normal, Penetration);
+}
+
+
+
 template <typename IdxType>
 bool FTriangleMeshImplicitObject::OverlapImp(const TArray<TVec3<IdxType>>& Elements, const FVec3& Point, const FReal Thickness) const
 {
diff --git a/Engine/Source/Runtime/Experimental/Chaos/Public/Chaos/CollisionResolution.h b/Engine/Source/Runtime/Experimental/Chaos/Public/Chaos/CollisionResolution.h
--- a/Engine/Source/Runtime/Experimental/Chaos/Public/Chaos/CollisionResolution.h
+++ b/Engine/Source/Runtime/Experimental/Chaos/Public/Chaos/CollisionResolution.h
@@ -31,6 +31,9 @@ namespace Chaos
 	template <typename T, int d>
 	class TRigidTransform;
 
+	class FTriangleMeshImplicitObject;
+
+
 	namespace Collisions
 	{
 
@@ -66,7 +69,7 @@ namespace Chaos
 		//
 
 		template <typename T, int d>
-		bool CHAOS_API UpdateBoxPlaneConstraint(const TAABB<T, d>& Box, const TRigidTransform<T, d>& BoxTransform, const TPlane<T, d>& Plane, const TRigidTransform<T, d>& PlaneTransform, const T Thickness, TRigidBodyPointContactConstraint<T, d>& Constraint);
+		void CHAOS_API UpdateBoxPlaneConstraint(const TAABB<T, d>& Box, const TRigidTransform<T, d>& BoxTransform, const TPlane<T, d>& Plane, const TRigidTransform<T, d>& PlaneTransform, const T Thickness, TRigidBodyPointContactConstraint<T, d>& Constraint);
 
 		template<class T, int d>
 		void CHAOS_API UpdateBoxPlaneManifold(TCollisionConstraintBase<T, d>&  Constraint, const TRigidTransform<T, d>& ATM, const TRigidTransform<T, d>& BTM, const T Thickness);
@@ -76,6 +79,20 @@ namespace Chaos
 
 
 
+		//
+		// Box-TriangleMesh
+		//
+
+		template <typename T, int d>
+		void CHAOS_API UpdateBoxTriangleMeshConstraint(const TAABB<T, d>& Box0, const TRigidTransform<T, d>& Transform0, const FTriangleMeshImplicitObject& TriangleMesh1, const TRigidTransform<T, d>& Transformw1, const T Thickness, TRigidBodyPointContactConstraint<T, d>& Constraint);
+
+		template <typename T, int d>
+		void CHAOS_API UpdateBoxTriangleMeshManifold(TCollisionConstraintBase<T, d>&  Constraint, const TRigidTransform<T, d>& ATM, const TRigidTransform<T, d>& BTM, const T Thickness);
+
+		template<typename T, int d>
+		void CHAOS_API ConstructBoxTriangleMeshConstraints(TGeometryParticleHandle<T, d>* Particle0, TGeometryParticleHandle<T, d>* Particle1, const FImplicitObject* Implicit0, const FImplicitObject* Implicit1, const TRigidTransform<T, d>& Transform0, const TRigidTransform<T, d>& Transform1, const T Thickness, FCollisionConstraintsArray& NewConstraints);
+
+
 		//
 		// Sphere-Sphere
 		//
diff --git a/Engine/Source/Runtime/Experimental/Chaos/Public/Chaos/TriangleMeshImplicitObject.h b/Engine/Source/Runtime/Experimental/Chaos/Public/Chaos/TriangleMeshImplicitObject.h
--- a/Engine/Source/Runtime/Experimental/Chaos/Public/Chaos/TriangleMeshImplicitObject.h
+++ b/Engine/Source/Runtime/Experimental/Chaos/Public/Chaos/TriangleMeshImplicitObject.h
@@ -143,6 +143,16 @@ namespace Chaos
 		bool SweepGeom(const TImplicitObjectScaled<TCapsule<FReal>>& QueryGeom, const FRigidTransform3& StartTM, const FVec3& Dir, const FReal Length, FReal& OutTime, FVec3& OutPosition, FVec3& OutNormal, int32& OutFaceIndex, const FReal Thickness = 0, const bool bComputeMTD = false) const;
 		bool SweepGeom(const TImplicitObjectScaled<FConvex>& QueryGeom, const FRigidTransform3& StartTM, const FVec3& Dir, const FReal Length, FReal& OutTime, FVec3& OutPosition, FVec3& OutNormal, int32& OutFaceIndex, const FReal Thickness = 0, const bool bComputeMTD = false) const;
 
+		bool GJKContactPoint(const TSphere<FReal, 3>& QueryGeom, const FRigidTransform3& QueryTM, const FReal Thickness, FVec3& Location, FVec3& Normal, FReal& Penetration) const;
+		bool GJKContactPoint(const TBox<FReal, 3>& QueryGeom, const FRigidTransform3& QueryTM, const FReal Thickness, FVec3& Location, FVec3& Normal, FReal& Penetration) const;
+		bool GJKContactPoint(const TCapsule<FReal>& QueryGeom, const FRigidTransform3& QueryTM, const FReal Thickness, FVec3& Location, FVec3& Normal, FReal& Penetration) const;
+		bool GJKContactPoint(const FConvex& QueryGeom, const FRigidTransform3& QueryTM, const FReal Thickness, FVec3& Location, FVec3& Normal, FReal& Penetration) const;
+		bool GJKContactPoint(const TImplicitObjectScaled < TSphere<FReal, 3> >& QueryGeom, const FRigidTransform3& QueryTM, const FReal Thickness, FVec3& Location, FVec3& Normal, FReal& Penetration) const;
+		bool GJKContactPoint(const TImplicitObjectScaled < TBox<FReal, 3> >& QueryGeom, const FRigidTransform3& QueryTM, const FReal Thickness, FVec3& Location, FVec3& Normal, FReal& Penetration) const;
+		bool GJKContactPoint(const TImplicitObjectScaled < TCapsule<FReal> >& QueryGeom, const FRigidTransform3& QueryTM, const FReal Thickness, FVec3& Location, FVec3& Normal, FReal& Penetration) const;
+		bool GJKContactPoint(const TImplicitObjectScaled < FConvex >& QueryGeom, const FRigidTransform3& QueryTM, const FReal Thickness, FVec3& Location, FVec3& Normal, FReal& Penetration) const;
+
+
 		virtual int32 FindMostOpposingFace(const FVec3& Position, const FVec3& UnitDir, int32 HintFaceIndex, FReal SearchDistance) const override;
 		virtual FVec3 FindGeometryOpposingNormal(const FVec3& DenormDir, int32 FaceIndex, const FVec3& OriginalNormal) const override;
 
@@ -267,6 +277,9 @@ namespace Chaos
 		FTriangleMeshImplicitObject()
 		    : FImplicitObject(EImplicitObject::HasBoundingBox, ImplicitObjectType::TriangleMesh){};
 
+		template <typename QueryGeomType>
+		bool GJKContactPointImp(const QueryGeomType& QueryGeom, const FRigidTransform3& QueryTM, const FReal Thickness, FVec3& Location, FVec3& Normal, FReal& Penetration) const;
+
 		template<typename QueryGeomType>
 		bool OverlapGeomImp(const QueryGeomType& QueryGeom, const FRigidTransform3& QueryTM, const FReal Thickness) const;
 
