Chaos - hack in Poke force in PhAt

diff --git a/Engine/Source/Editor/AnimGraph/Classes/AnimPreviewInstance.h b/Engine/Source/Editor/AnimGraph/Classes/AnimPreviewInstance.h
--- a/Engine/Source/Editor/AnimGraph/Classes/AnimPreviewInstance.h
+++ b/Engine/Source/Editor/AnimGraph/Classes/AnimPreviewInstance.h
@@ -117,6 +117,8 @@ public:
 		return CurveBoneControllers;
 	}
 
+	virtual void AddImpulseAtLocation(FVector Impulse, FVector Location, FName BoneName = NAME_None) {}
+
 	/** Sets an external debug skeletal mesh component to use to debug */
 	void SetDebugSkeletalMeshComponent(USkeletalMeshComponent* InSkeletalMeshComponent);
 
@@ -300,6 +302,9 @@ public:
 	 */
 	void EnableControllers(bool bEnable);
 
+	/** Preview physics interaction */
+	void AddImpulseAtLocation(FVector Impulse, FVector Location, FName BoneName = NAME_None);
+
 	/** Sets an external debug skeletal mesh component to use to debug */
 	void SetDebugSkeletalMeshComponent(USkeletalMeshComponent* InSkeletalMeshComponent);
 
diff --git a/Engine/Source/Editor/AnimGraph/Private/AnimPreviewInstance.cpp b/Engine/Source/Editor/AnimGraph/Private/AnimPreviewInstance.cpp
--- a/Engine/Source/Editor/AnimGraph/Private/AnimPreviewInstance.cpp
+++ b/Engine/Source/Editor/AnimGraph/Private/AnimPreviewInstance.cpp
@@ -1202,6 +1202,12 @@ FAnimInstanceProxy* UAnimPreviewInstance::CreateAnimInstanceProxy()
 	return new FAnimPreviewInstanceProxy(this);
 }
 
+void UAnimPreviewInstance::AddImpulseAtLocation(FVector Impulse, FVector Location, FName BoneName)
+{
+	FAnimPreviewInstanceProxy& Proxy = GetProxyOnGameThread<FAnimPreviewInstanceProxy>();
+	Proxy.AddImpulseAtLocation(Impulse, Location, BoneName);
+}
+
 void UAnimPreviewInstance::SetDebugSkeletalMeshComponent(USkeletalMeshComponent* InSkeletalMeshComponent)
 {
 	FAnimPreviewInstanceProxy& Proxy = GetProxyOnGameThread<FAnimPreviewInstanceProxy>();
diff --git a/Engine/Source/Editor/PhysicsAssetEditor/Private/PhysicsAssetEditorAnimInstanceProxy.cpp b/Engine/Source/Editor/PhysicsAssetEditor/Private/PhysicsAssetEditorAnimInstanceProxy.cpp
--- a/Engine/Source/Editor/PhysicsAssetEditor/Private/PhysicsAssetEditorAnimInstanceProxy.cpp
+++ b/Engine/Source/Editor/PhysicsAssetEditor/Private/PhysicsAssetEditorAnimInstanceProxy.cpp
@@ -53,6 +53,11 @@ bool FPhysicsAssetEditorAnimInstanceProxy::Evaluate_WithRoot(FPoseContext& Outpu
 	}
 }
 
+void FPhysicsAssetEditorAnimInstanceProxy::AddImpulseAtLocation(FVector Impulse, FVector Location, FName BoneName)
+{
+	RagdollNode.AddImpulseAtLocation(Impulse, Location, BoneName);
+}
+
 FPhysicsAssetEditorAnimInstanceProxy::~FPhysicsAssetEditorAnimInstanceProxy()
 {
 }
\ No newline at end of file
diff --git a/Engine/Source/Editor/PhysicsAssetEditor/Private/PhysicsAssetEditorAnimInstanceProxy.h b/Engine/Source/Editor/PhysicsAssetEditor/Private/PhysicsAssetEditorAnimInstanceProxy.h
--- a/Engine/Source/Editor/PhysicsAssetEditor/Private/PhysicsAssetEditorAnimInstanceProxy.h
+++ b/Engine/Source/Editor/PhysicsAssetEditor/Private/PhysicsAssetEditorAnimInstanceProxy.h
@@ -36,6 +36,8 @@ public:
 	virtual void UpdateAnimationNode(const FAnimationUpdateContext& InContext) override;
 	virtual bool Evaluate_WithRoot(FPoseContext& Output, FAnimNode_Base* InRootNode) override;
 
+	virtual void AddImpulseAtLocation(FVector Impulse, FVector Location, FName BoneName = NAME_None) override;
+
 private:
 	void ConstructNodes();
 
diff --git a/Engine/Source/Editor/PhysicsAssetEditor/Private/PhysicsAssetEditorSkeletalMeshComponent.cpp b/Engine/Source/Editor/PhysicsAssetEditor/Private/PhysicsAssetEditorSkeletalMeshComponent.cpp
--- a/Engine/Source/Editor/PhysicsAssetEditor/Private/PhysicsAssetEditorSkeletalMeshComponent.cpp
+++ b/Engine/Source/Editor/PhysicsAssetEditor/Private/PhysicsAssetEditorSkeletalMeshComponent.cpp
@@ -12,7 +12,9 @@
 #include "PhysicsAssetEditorSkeletalMeshComponent.h"
 #include "PhysicsEngine/PhysicsConstraintTemplate.h"
 #include "PhysicsEngine/PhysicsAsset.h"
+#include "Chaos/Core.h"
 #include "SkeletalMeshTypes.h"
+#include "AnimPreviewInstance.h"
 #include "UObject/Package.h"
 #include "EditorStyleSet.h"
 
@@ -496,4 +498,16 @@ void UPhysicsAssetEditorSkeletalMeshComponent::RefreshBoneTransforms(FActorCompo
 		FinalizeBoneTransform();
 		bNeedToFlipSpaceBaseBuffers = true;
 	}
-}
\ No newline at end of file
+}
+
+void UPhysicsAssetEditorSkeletalMeshComponent::AddImpulseAtLocation(FVector Impulse, FVector Location, FName BoneName)
+{
+#if !WITH_CHAOS
+	Super::AddImpulseAtLocation(Impulse, Location, BoneName);
+#else
+	if (PreviewInstance != nullptr)
+	{
+		PreviewInstance->AddImpulseAtLocation(Impulse, Location, BoneName);
+	}
+#endif
+}
diff --git a/Engine/Source/Editor/PhysicsAssetEditor/Private/PhysicsAssetEditorSkeletalMeshComponent.h b/Engine/Source/Editor/PhysicsAssetEditor/Private/PhysicsAssetEditorSkeletalMeshComponent.h
--- a/Engine/Source/Editor/PhysicsAssetEditor/Private/PhysicsAssetEditorSkeletalMeshComponent.h
+++ b/Engine/Source/Editor/PhysicsAssetEditor/Private/PhysicsAssetEditorSkeletalMeshComponent.h
@@ -52,6 +52,7 @@ class UPhysicsAssetEditorSkeletalMeshComponent : public UDebugSkelMeshComponent
 
 	/** UPrimitiveComponent interface */
 	virtual FPrimitiveSceneProxy* CreateSceneProxy() override;
+	virtual void AddImpulseAtLocation(FVector Impulse, FVector Location, FName BoneName = NAME_None) override;
 
 	/** USkinnedMeshComponent interface */
 	virtual void RefreshBoneTransforms(FActorComponentTickFunction* TickFunction = nullptr) override;
diff --git a/Engine/Source/Runtime/AnimGraphRuntime/Private/BoneControllers/AnimNode_RigidBody.cpp b/Engine/Source/Runtime/AnimGraphRuntime/Private/BoneControllers/AnimNode_RigidBody.cpp
--- a/Engine/Source/Runtime/AnimGraphRuntime/Private/BoneControllers/AnimNode_RigidBody.cpp
+++ b/Engine/Source/Runtime/AnimGraphRuntime/Private/BoneControllers/AnimNode_RigidBody.cpp
@@ -729,6 +729,10 @@ void FAnimNode_RigidBody::InitPhysics(const UAnimInstance* InAnimInstance)
 					{
 						IgnoreCollisionActors.Add(NewBodyHandle);
 					}
+
+#if WITH_CHAOS
+					NewBodyHandle->SetName(BodySetup->BoneName);
+#endif
 				}
 			}
 		}
@@ -1144,6 +1148,22 @@ void FAnimNode_RigidBody::InitializeBoneReferences(const FBoneContainer& Require
 	}
 }
 
+void FAnimNode_RigidBody::AddImpulseAtLocation(FVector Impulse, FVector Location, FName BoneName)
+{
+#if WITH_CHAOS
+	// Find the body. This is currently only used in the editor and will need optimizing if used in game
+	for (int32 BodyIndex = 0; BodyIndex < Bodies.Num(); ++BodyIndex)
+	{
+		ImmediatePhysics::FActorHandle* Body = Bodies[BodyIndex];
+		if (Body->GetName() == BoneName)
+		{
+			Body->AddImpulseAtLocation(Impulse, Location);
+		}
+	}
+#endif
+}
+
+
 void FAnimNode_RigidBody::OnInitializeAnimInstance(const FAnimInstanceProxy* InProxy, const UAnimInstance* InAnimInstance)
 {
 	CSV_SCOPED_TIMING_STAT(Animation, RigidBodyInit);
diff --git a/Engine/Source/Runtime/AnimGraphRuntime/Public/BoneControllers/AnimNode_RigidBody.h b/Engine/Source/Runtime/AnimGraphRuntime/Public/BoneControllers/AnimNode_RigidBody.h
--- a/Engine/Source/Runtime/AnimGraphRuntime/Public/BoneControllers/AnimNode_RigidBody.h
+++ b/Engine/Source/Runtime/AnimGraphRuntime/Public/BoneControllers/AnimNode_RigidBody.h
@@ -55,6 +55,8 @@ struct ANIMGRAPHRUNTIME_API FAnimNode_RigidBody : public FAnimNode_SkeletalContr
 	virtual int32 GetLODThreshold() const override;
 	// End of FAnimNode_SkeletalControlBase interface
 
+	virtual void AddImpulseAtLocation(FVector Impulse, FVector Location, FName BoneName = NAME_None);
+
 public:
 	/** Physics asset to use. If empty use the skeletal mesh's default physics asset */
 	UPROPERTY(EditAnywhere, Category = Settings)
diff --git a/Engine/Source/Runtime/Engine/Private/PhysicsEngine/ImmediatePhysics/ImmediatePhysicsChaos/ImmediatePhysicsActorHandle_Chaos.cpp b/Engine/Source/Runtime/Engine/Private/PhysicsEngine/ImmediatePhysics/ImmediatePhysicsChaos/ImmediatePhysicsActorHandle_Chaos.cpp
--- a/Engine/Source/Runtime/Engine/Private/PhysicsEngine/ImmediatePhysics/ImmediatePhysicsChaos/ImmediatePhysicsActorHandle_Chaos.cpp
+++ b/Engine/Source/Runtime/Engine/Private/PhysicsEngine/ImmediatePhysics/ImmediatePhysicsChaos/ImmediatePhysicsActorHandle_Chaos.cpp
@@ -473,6 +473,18 @@ namespace ImmediatePhysics_Chaos
 		}
 	}
 
+	void FActorHandle::AddImpulseAtLocation(FVector Impulse, FVector Location)
+	{
+		using namespace Chaos;
+
+		if (TPBDRigidParticleHandle<FReal, 3>* Rigid = Handle()->CastToRigidParticle())
+		{
+			FVector CoM = FParticleUtilities::GetCoMWorldPosition(Rigid);
+			Rigid->LinearImpulse() += Impulse;
+			Rigid->AngularImpulse() += FVector::CrossProduct(Location - CoM, Impulse);
+		}
+	}
+
 	void FActorHandle::SetLinearDamping(float NewLinearDamping)
 	{
 		using namespace Chaos;
diff --git a/Engine/Source/Runtime/Engine/Public/Physics/ImmediatePhysics/ImmediatePhysicsChaos/ImmediatePhysicsActorHandle_Chaos.h b/Engine/Source/Runtime/Engine/Public/Physics/ImmediatePhysics/ImmediatePhysicsChaos/ImmediatePhysicsActorHandle_Chaos.h
--- a/Engine/Source/Runtime/Engine/Public/Physics/ImmediatePhysics/ImmediatePhysicsChaos/ImmediatePhysicsActorHandle_Chaos.h
+++ b/Engine/Source/Runtime/Engine/Public/Physics/ImmediatePhysics/ImmediatePhysicsChaos/ImmediatePhysicsActorHandle_Chaos.h
@@ -14,6 +14,9 @@ namespace ImmediatePhysics_Chaos
 	public:
 		~FActorHandle();
 
+		void SetName(const FName& InName) { Name = InName; }
+		const FName& GetName() const { return Name; }
+
 		void SetEnabled(bool bEnabled);
 
 		/** Sets the world transform.*/
@@ -56,6 +59,8 @@ namespace ImmediatePhysics_Chaos
 
 		void AddRadialForce(const FVector& Origin, float Strength, float Radius, ERadialImpulseFalloff Falloff, EForceType ForceType);
 
+		void AddImpulseAtLocation(FVector Impulse, FVector Location);
+
 		/** Set the linear damping*/
 		void SetLinearDamping(float NewLinearDamping);
 
@@ -125,6 +130,7 @@ namespace ImmediatePhysics_Chaos
 
 		Chaos::TGenericParticleHandle<FReal, Dimensions> Handle() const;
 
+		FName Name;
 		Chaos::TPBDRigidsSOAs<FReal, 3>& Particles;
 		Chaos::TGeometryParticleHandle<FReal, Dimensions>* ParticleHandle;
 		TUniquePtr<Chaos::FImplicitObject> Geometry;
diff --git a/Engine/Source/Runtime/Experimental/Chaos/Private/Chaos/Evolution/PBDMinEvolution.cpp b/Engine/Source/Runtime/Experimental/Chaos/Private/Chaos/Evolution/PBDMinEvolution.cpp
--- a/Engine/Source/Runtime/Experimental/Chaos/Private/Chaos/Evolution/PBDMinEvolution.cpp
+++ b/Engine/Source/Runtime/Experimental/Chaos/Private/Chaos/Evolution/PBDMinEvolution.cpp
@@ -7,6 +7,7 @@
 #include "Chaos/PBDCollisionConstraints.h"
 #include "Chaos/PBDConstraintRule.h"
 #include "Chaos/PBDRigidsSOAs.h"
+#include "Chaos/PerParticleAddImpulses.h"
 #include "Chaos/PerParticleEtherDrag.h"
 #include "Chaos/PerParticleEulerStepVelocity.h"
 #include "Chaos/PerParticleGravity.h"
@@ -103,6 +104,7 @@ namespace Chaos
 		SCOPE_CYCLE_COUNTER(STAT_MinEvolution_Integrate);
 
 		TPerParticleEulerStepVelocity<FReal, 3> EulerStepVelocityRule;
+		TPerParticleAddImpulses<FReal, 3> AddImpulsesRule;
 		TPerParticleEtherDrag<FReal, 3> EtherDragRule;
 		TPerParticlePBDEulerStep<FReal, 3> EulerStepPositionRule;
 
@@ -116,12 +118,15 @@ namespace Chaos
 				Particle.F() += Particle.M() * Gravity;
 
 				EulerStepVelocityRule.Apply(Particle, Dt);
+				AddImpulsesRule.Apply(Particle, Dt);
 				EtherDragRule.Apply(Particle, Dt);
 
 				EulerStepPositionRule.Apply(Particle, Dt);
 
 				Particle.F() = FVec3(0);
 				Particle.Torque() = FVec3(0);
+				Particle.LinearImpulse() = FVec3(0);
+				Particle.AngularImpulse() = FVec3(0);
 
 				if (Particle.HasBounds())
 				{
