Fix Landscape painting/sculpting. Fix HeightField cell bounding boxes not correctly handling negative heights. Fix a few incorrect conversions from cell indices to height indices.

diff --git a/Engine/Source/Runtime/Experimental/Chaos/Private/Chaos/HeightField.cpp b/Engine/Source/Runtime/Experimental/Chaos/Private/Chaos/HeightField.cpp
--- a/Engine/Source/Runtime/Experimental/Chaos/Private/Chaos/HeightField.cpp
+++ b/Engine/Source/Runtime/Experimental/Chaos/Private/Chaos/HeightField.cpp
@@ -539,7 +539,7 @@ namespace Chaos
 		{
 			TVector<int32, 2> CellCoord = FlatGrid.Cell(UnscaledLocation);
 
-			const int32 SingleIndex = CellCoord[1] * (GeomData.NumCols - 1) + CellCoord[0];
+			const int32 SingleIndex = CellCoord[1] * (GeomData.NumCols) + CellCoord[0];
 			TVector<T, 3> Pts[4];
 			GeomData.GetPoints(SingleIndex, Pts);
 
@@ -595,7 +595,7 @@ namespace Chaos
 	{
 		if (FlatGrid.IsValid(InCoord))
 		{
-			OutMin = TVec3<T>(InCoord[0], InCoord[1], 0);
+			OutMin = TVec3<T>(InCoord[0], InCoord[1], GeomData.GetMinHeight());
 			OutMax = TVec3<T>(InCoord[0] + 1, InCoord[1] + 1, GeomData.CellHeights[InCoord[1] * (GeomData.NumCols - 1) + InCoord[0]]);
 			OutMin = OutMin - InInflate;
 			OutMax = OutMax + InInflate;
@@ -629,7 +629,7 @@ namespace Chaos
 	{
 		if (FlatGrid.IsValid(InCoord))
 		{
-			OutMin = TVec3<T>(InCoord[0], InCoord[1], 0);
+			OutMin = TVec3<T>(InCoord[0], InCoord[1], GeomData.GetMinHeight());
 			OutMax = TVec3<T>(InCoord[0] + 1, InCoord[1] + 1, GeomData.CellHeights[InCoord[1] * (GeomData.NumCols - 1) + InCoord[0]]);
 			OutMin = OutMin * GeomData.Scale - InInflate;
 			OutMax = OutMax * GeomData.Scale + InInflate;
@@ -644,7 +644,7 @@ namespace Chaos
 	{
 		if(FlatGrid.IsValid(InCoord))
 		{
-			int32 Index = InCoord[1] * (GeomData.NumCols - 1) + InCoord[0] + InCoord[1];
+			int32 Index = InCoord[1] * (GeomData.NumCols) + InCoord[0];
 			static TVector<T, 3> Points[4];
 			GeomData.GetPoints(Index, Points);
 
@@ -1168,7 +1168,7 @@ namespace Chaos
 
 		for(const TVector<int32, 2>& Cell : Intersections)
 		{
-			const int32 SingleIndex = Cell[1] * (GeomData.NumCols - 1) + Cell[0];
+			const int32 SingleIndex = Cell[1] * (GeomData.NumCols) + Cell[0];
 			GeomData.GetPointsScaled(SingleIndex, Points);
 
 			if(OverlapTriangle(Points[0], Points[1], Points[3]))
@@ -1345,7 +1345,7 @@ namespace Chaos
 
 		for(const TVector<int32, 2>& Cell : Intersections)
 		{
-			const int32 SingleIndex = Cell[1] * (GeomData.NumCols - 1) + Cell[0];
+			const int32 SingleIndex = Cell[1] * (GeomData.NumCols) + Cell[0];
 			GeomData.GetPointsScaled(SingleIndex, Points);
 
 			if(OverlapTriangle(Points[0], Points[1], Points[3]))
diff --git a/Engine/Source/Runtime/Experimental/Chaos/Public/Chaos/HeightField.h b/Engine/Source/Runtime/Experimental/Chaos/Public/Chaos/HeightField.h
--- a/Engine/Source/Runtime/Experimental/Chaos/Public/Chaos/HeightField.h
+++ b/Engine/Source/Runtime/Experimental/Chaos/Public/Chaos/HeightField.h
@@ -254,6 +254,11 @@ namespace Chaos
 				OutPts[3] *= Scale;
 			}
 
+			FORCEINLINE T GetMinHeight() const
+			{
+				return static_cast<typename FDataType::RealType>(MinValue);
+			}
+
 			void Serialize(FChaosArchive& Ar)
 			{
 				int32 TempRealSize = RealSize;
