DoCollide flag

- Repurposed the IgnoreAnalyticCollisions for a DoCollide flag on the ImplicitObject. Used to globally disable a bodies collisions.

diff --git a/Engine/Source/Runtime/Experimental/Chaos/Private/Chaos/ImplicitObject.cpp b/Engine/Source/Runtime/Experimental/Chaos/Private/Chaos/ImplicitObject.cpp
--- a/Engine/Source/Runtime/Experimental/Chaos/Private/Chaos/ImplicitObject.cpp
+++ b/Engine/Source/Runtime/Experimental/Chaos/Private/Chaos/ImplicitObject.cpp
@@ -22,7 +22,7 @@ using namespace Chaos;
 FImplicitObject::FImplicitObject(int32 Flags, EImplicitObjectType InType)
     : Type(InType)
     , bIsConvex(!!(Flags & EImplicitObject::IsConvex))
-    , bIgnoreAnalyticCollisions(!!(Flags & EImplicitObject::IgnoreAnalyticCollisions))
+    , bDoCollide(true)
     , bHasBoundingBox(!!(Flags & EImplicitObject::HasBoundingBox))
 #if TRACK_CHAOS_GEOMETRY
     , bIsTracked(false)
@@ -58,10 +58,6 @@ void FImplicitObject::Track(TSerializablePtr<FImplicitObject> This, const FStrin
 
 EImplicitObjectType FImplicitObject::GetType(bool bGetTrueType) const
 {
-	if (bIgnoreAnalyticCollisions && !bGetTrueType)
-	{
-		return ImplicitObjectType::Unknown;
-	}
 	return Type;
 }
 
@@ -318,7 +314,7 @@ void FImplicitObject::SerializeImp(FArchive& Ar)
 	Ar.UsingCustomVersion(FDestructionObjectVersion::GUID);
 	if (Ar.CustomVer(FDestructionObjectVersion::GUID) >= FDestructionObjectVersion::ChaosArchiveAdded)
 	{
-		Ar << bIsConvex << bIgnoreAnalyticCollisions;
+		Ar << bIsConvex << bDoCollide;
 	}
 }
 
diff --git a/Engine/Source/Runtime/Experimental/Chaos/Private/Chaos/PBDRigidClustering.cpp b/Engine/Source/Runtime/Experimental/Chaos/Private/Chaos/PBDRigidClustering.cpp
--- a/Engine/Source/Runtime/Experimental/Chaos/Private/Chaos/PBDRigidClustering.cpp
+++ b/Engine/Source/Runtime/Experimental/Chaos/Private/Chaos/PBDRigidClustering.cpp
@@ -1451,7 +1451,7 @@ namespace Chaos
 
 		if (bUseParticleImplicit && MParticles.DynamicGeometry(NewIndex)) //if children are ignore analytic and this is a dynamic geom, mark it too. todo(ocohen): clean this up
 		{
-			MParticles.DynamicGeometry(NewIndex)->IgnoreAnalyticCollisions();
+			MParticles.DynamicGeometry(NewIndex)->SetDoCollide(false);
 		}
 
 		if (Parameters.CollisionParticles)
diff --git a/Engine/Source/Runtime/Experimental/Chaos/Public/Chaos/Box.h b/Engine/Source/Runtime/Experimental/Chaos/Public/Chaos/Box.h
--- a/Engine/Source/Runtime/Experimental/Chaos/Public/Chaos/Box.h
+++ b/Engine/Source/Runtime/Experimental/Chaos/Public/Chaos/Box.h
@@ -54,7 +54,7 @@ namespace Chaos
 		{
 			this->Type = Other.Type;
 			this->bIsConvex = Other.bIsConvex;
-			this->bIgnoreAnalyticCollisions = Other.bIgnoreAnalyticCollisions;
+			this->bDoCollide = Other.bDoCollide;
 			this->bHasBoundingBox = Other.bHasBoundingBox;
 
 			AABB = Other.AABB;
@@ -65,7 +65,7 @@ namespace Chaos
 		{
 			this->Type = Other.Type;
 			this->bIsConvex = Other.bIsConvex;
-			this->bIgnoreAnalyticCollisions = Other.bIgnoreAnalyticCollisions;
+			this->bDoCollide = Other.bDoCollide;
 			this->bHasBoundingBox = Other.bHasBoundingBox;
 
 			AABB = MoveTemp(Other.AABB);
diff --git a/Engine/Source/Runtime/Experimental/Chaos/Public/Chaos/Capsule.h b/Engine/Source/Runtime/Experimental/Chaos/Public/Chaos/Capsule.h
--- a/Engine/Source/Runtime/Experimental/Chaos/Public/Chaos/Capsule.h
+++ b/Engine/Source/Runtime/Experimental/Chaos/Public/Chaos/Capsule.h
@@ -59,7 +59,7 @@ namespace Chaos
 		{
 			this->Type = InSteal.Type;
 			this->bIsConvex = InSteal.bIsConvex;
-			this->bIgnoreAnalyticCollisions = InSteal.bIgnoreAnalyticCollisions;
+			this->bDoCollide = InSteal.bDoCollide;
 			this->bHasBoundingBox = InSteal.bHasBoundingBox;
 
 			MSegment = MoveTemp(InSteal.MSegment);
diff --git a/Engine/Source/Runtime/Experimental/Chaos/Public/Chaos/ImplicitObject.h b/Engine/Source/Runtime/Experimental/Chaos/Public/Chaos/ImplicitObject.h
--- a/Engine/Source/Runtime/Experimental/Chaos/Public/Chaos/ImplicitObject.h
+++ b/Engine/Source/Runtime/Experimental/Chaos/Public/Chaos/ImplicitObject.h
@@ -79,7 +79,6 @@ namespace EImplicitObject
 	{
 		IsConvex = 1,
 		HasBoundingBox = 1 << 1,
-		IgnoreAnalyticCollisions = 1 << 2,
 	};
 
 	const int32 FiniteConvex = IsConvex | HasBoundingBox;
@@ -178,8 +177,8 @@ public:
 	virtual const class TAABB<FReal, 3>& BoundingBox() const;
 	bool HasBoundingBox() const { return bHasBoundingBox; }
 	bool IsConvex() const { return bIsConvex; }
-	void IgnoreAnalyticCollisions(const bool Ignore = true) { bIgnoreAnalyticCollisions = Ignore; }
-	bool GetIgnoreAnalyticCollisions() const { return bIgnoreAnalyticCollisions; }
+	void SetDoCollide(const bool Collide ) { bDoCollide = Collide; }
+	bool GetDoCollide() const { return bDoCollide; }
 	void SetConvex(const bool Convex = true) { bIsConvex = Convex; }
 	
 #if TRACK_CHAOS_GEOMETRY
@@ -271,7 +270,7 @@ public:
 
 	virtual FString ToString() const
 	{
-		return FString::Printf(TEXT("ImplicitObject bIsConvex:%d, bIgnoreAnalyticCollision:%d, bHasBoundingBox:%d"), bIsConvex, bIgnoreAnalyticCollisions, bHasBoundingBox);
+		return FString::Printf(TEXT("ImplicitObject bIsConvex:%d, bIgnoreAnalyticCollision:%d, bHasBoundingBox:%d"), bIsConvex, bDoCollide, bHasBoundingBox);
 	}
 
 	void SerializeImp(FArchive& Ar);
@@ -301,7 +300,7 @@ public:
 protected:
 	EImplicitObjectType Type;
 	bool bIsConvex;
-	bool bIgnoreAnalyticCollisions;
+	bool bDoCollide;
 	bool bHasBoundingBox;
 
 #if TRACK_CHAOS_GEOMETRY
diff --git a/Engine/Source/Runtime/Experimental/Chaos/Public/Chaos/Sphere.h b/Engine/Source/Runtime/Experimental/Chaos/Public/Chaos/Sphere.h
--- a/Engine/Source/Runtime/Experimental/Chaos/Public/Chaos/Sphere.h
+++ b/Engine/Source/Runtime/Experimental/Chaos/Public/Chaos/Sphere.h
@@ -51,7 +51,7 @@ namespace Chaos
 		{
 			this->Type = InSteal.Type;
 			this->bIsConvex = InSteal.bIsConvex;
-			this->bIgnoreAnalyticCollisions = InSteal.bIgnoreAnalyticCollisions;
+			this->bDoCollide = InSteal.bDoCollide;
 			this->bHasBoundingBox = InSteal.bHasBoundingBox;
 			Center = MoveTemp(InSteal.Center);
 			Radius = InSteal.Radius;
diff --git a/Engine/Source/Runtime/Experimental/GeometryCollectionSimulationCore/Private/GeometryCollection/GeometryCollectionCollisionStructureManager.cpp b/Engine/Source/Runtime/Experimental/GeometryCollectionSimulationCore/Private/GeometryCollection/GeometryCollectionCollisionStructureManager.cpp
--- a/Engine/Source/Runtime/Experimental/GeometryCollectionSimulationCore/Private/GeometryCollection/GeometryCollectionCollisionStructureManager.cpp
+++ b/Engine/Source/Runtime/Experimental/GeometryCollectionSimulationCore/Private/GeometryCollection/GeometryCollectionCollisionStructureManager.cpp
@@ -177,7 +177,7 @@ void FCollisionStructureManager::UpdateImplicitFlags(FImplicit* Implicit, EColli
 {
 	if (Implicit && (CollisionType == ECollisionTypeEnum::Chaos_Surface_Volumetric))
 	{
-		Implicit->IgnoreAnalyticCollisions();
+		Implicit->SetDoCollide(false);
 		Implicit->SetConvex(false);
 	}
 }
