Fix CIS

diff --git a/Engine/Source/Runtime/Experimental/Chaos/Private/Chaos/TriangleMeshImplicitObject.cpp b/Engine/Source/Runtime/Experimental/Chaos/Private/Chaos/TriangleMeshImplicitObject.cpp
--- a/Engine/Source/Runtime/Experimental/Chaos/Private/Chaos/TriangleMeshImplicitObject.cpp
+++ b/Engine/Source/Runtime/Experimental/Chaos/Private/Chaos/TriangleMeshImplicitObject.cpp
@@ -308,14 +308,7 @@ bool FTriangleMeshImplicitObject::GJKContactPointImp(const QueryGeomType& QueryG
 	{
 		return LambdaHelper(MElements.GetLargeIndexBuffer());
 	}
-	else
-	{
-		return LambdaHelper(MElements.GetSmallIndexBuffer());
-	}
-
-	if (Penetration < 0)
-		return true;
-	return false;
+	return LambdaHelper(MElements.GetSmallIndexBuffer());
 }
 
 bool FTriangleMeshImplicitObject::GJKContactPoint(const TSphere<FReal, 3>& QueryGeom, const FRigidTransform3& QueryTM, const FReal Thickness, FVec3& Location, FVec3& Normal, FReal& Penetration) const
