Change some ImplicitObjectScaled ensures to CHAOS_ENSURE to reduce chaff in tests.

diff --git a/Engine/Source/Runtime/Experimental/Chaos/Public/Chaos/ImplicitObjectScaled.h b/Engine/Source/Runtime/Experimental/Chaos/Public/Chaos/ImplicitObjectScaled.h
--- a/Engine/Source/Runtime/Experimental/Chaos/Public/Chaos/ImplicitObjectScaled.h
+++ b/Engine/Source/Runtime/Experimental/Chaos/Public/Chaos/ImplicitObjectScaled.h
@@ -395,7 +395,7 @@ public:
 		// Compute final normal
 		const TVector<T, d> LocalNormal = MObject->FindGeometryOpposingNormal(LocalDenormDir, HintFaceIndex, LocalOriginalNormal);
 		TVector<T, d> Normal = LocalNormal * MInvScale;
-		if (ensure(Normal.SafeNormalize()) == 0)
+		if (CHAOS_ENSURE(Normal.SafeNormalize()) == 0)
 		{
 			Normal = TVector<T,3>(0,0,1);
 		}
