Fix Crash in ConstraintGraph ComputeIslands.
When disabling particle we now remove kinematics. TODO make faster.

diff --git a/Engine/Source/Runtime/Experimental/Chaos/Private/Chaos/PBDConstraintGraph.cpp b/Engine/Source/Runtime/Experimental/Chaos/Private/Chaos/PBDConstraintGraph.cpp
--- a/Engine/Source/Runtime/Experimental/Chaos/Private/Chaos/PBDConstraintGraph.cpp
+++ b/Engine/Source/Runtime/Experimental/Chaos/Private/Chaos/PBDConstraintGraph.cpp
@@ -715,6 +715,16 @@ void FPBDConstraintGraph::DisableParticle(TGeometryParticleHandle<FReal, 3>* Par
 		}
 
 	}
+	else if (PBDRigid)
+	{
+		// Kinematic particles are included in IslandToParticles, however we cannot use islands to look them up.
+		// TODO find faster removal method?
+		for (TArray<TGeometryParticleHandle<FReal, 3>*>& IslandParticles : IslandToParticles)
+		{
+			IslandParticles.RemoveSingleSwap(Particle);
+		}
+	}
+
 	ParticleRemove(Particle);
 }
 
