Reduce default angular velocity threshold for putting particles to sleep in islands.

diff --git a/Engine/Source/Runtime/Experimental/Chaos/Private/Chaos/PBDConstraintGraph.cpp b/Engine/Source/Runtime/Experimental/Chaos/Private/Chaos/PBDConstraintGraph.cpp
--- a/Engine/Source/Runtime/Experimental/Chaos/Private/Chaos/PBDConstraintGraph.cpp
+++ b/Engine/Source/Runtime/Experimental/Chaos/Private/Chaos/PBDConstraintGraph.cpp
@@ -544,7 +544,7 @@ bool FPBDConstraintGraph::SleepInactive(const int32 Island, const TArrayCollecti
 	FReal LinearSleepingThreshold = FLT_MAX;
 	FReal AngularSleepingThreshold = FLT_MAX;
 	FReal DefaultLinearSleepingThreshold = (FReal)1;
-	FReal DefaultAngularSleepingThreshold = (FReal)1;
+	FReal DefaultAngularSleepingThreshold = (FReal)0.1f;
 
 	int32 NumDynamicParticles = 0;
 
