Chaos
- fix zero-/low-stiffness drives
- use PBD constraints for soft constraints with very high stiffness
- improved angular projection
- drive position/velocity split to match PhysX
- remove dead code

diff --git a/Engine/Source/Runtime/Engine/Private/PhysicsEngine/ImmediatePhysics/ImmediatePhysicsChaos/ImmediatePhysicsJointHandle_Chaos.cpp b/Engine/Source/Runtime/Engine/Private/PhysicsEngine/ImmediatePhysics/ImmediatePhysicsChaos/ImmediatePhysicsJointHandle_Chaos.cpp
--- a/Engine/Source/Runtime/Engine/Private/PhysicsEngine/ImmediatePhysics/ImmediatePhysicsChaos/ImmediatePhysicsJointHandle_Chaos.cpp
+++ b/Engine/Source/Runtime/Engine/Private/PhysicsEngine/ImmediatePhysics/ImmediatePhysicsChaos/ImmediatePhysicsJointHandle_Chaos.cpp
@@ -24,6 +24,8 @@ namespace ImmediatePhysics_Chaos
 	FAutoConsoleVariableRef CVarDriveStiffnessScale(TEXT("p.Chaos.ImmPhys.DriveStiffnessScale"), ChaosImmediate_DriveStiffnessScale, TEXT("Conversion factor for drive stiffness."));
 	FAutoConsoleVariableRef CVarDriveDampingScale(TEXT("p.Chaos.ImmPhys.DriveDampingScale"), ChaosImmediate_DriveDampingScale, TEXT("Conversion factor for drive damping."));
 
+	// This value is from PhysX UE - in Chaos we don't really need Soft constraints with stiffness greater than about 1000
+	// (In PhysX, Soft constraints were implicitly integrated, and essentially equivalent to XPBD, which for large stiffness is just PBD).
 	float ChaosImmediate_SoftStiffnessScale = 100000.0f;
 	float ChaosImmediate_SoftDampingScale = 100000.0f;
 	FAutoConsoleVariableRef CVarSoftStiffnessScale(TEXT("p.Chaos.ImmPhys.SoftStiffnessScale"), ChaosImmediate_SoftStiffnessScale, TEXT("Conversion factor for soft-joint stiffness."));
@@ -35,7 +37,7 @@ namespace ImmediatePhysics_Chaos
 	FAutoConsoleVariableRef CVarJointMaxLinearProjection(TEXT("p.Chaos.ImmPhys.JointMaxLinearProjection"), ChaosImmediate_JointMaxLinearProjection, TEXT("Joint max projection (for joints with projection enabled)."));
 
 	float ChaosImmediate_JointMinAngularProjection = 0.0f;
-	float ChaosImmediate_JointMaxAngularProjection = 0.5f;
+	float ChaosImmediate_JointMaxAngularProjection = 0.1f;
 	FAutoConsoleVariableRef CVarJointMinAngularProjection(TEXT("p.Chaos.ImmPhys.JointMinAngularProjection"), ChaosImmediate_JointMinAngularProjection, TEXT("Joint min projection (for joints with projection disabled)."));
 	FAutoConsoleVariableRef CVarJointMaxAngularProjection(TEXT("p.Chaos.ImmPhys.JointMaxAngularProjection"), ChaosImmediate_JointMaxAngularProjection, TEXT("Joint max projection (for joints with projection enabled)."));
 
@@ -45,6 +47,8 @@ namespace ImmediatePhysics_Chaos
 	{
 		using namespace Chaos;
 
+		const FConstraintProfileProperties& Profile = ConstraintInstance->ProfileInstance;
+
 		FPBDJointSettings ConstraintSettings;
 		TVector<FRigidTransform3, 2> ConstraintFrames;
 
@@ -71,7 +75,7 @@ namespace ImmediatePhysics_Chaos
 
 		ConstraintSettings.LinearProjection = ConstraintInstance->IsProjectionEnabled() ? ChaosImmediate_JointMaxLinearProjection : ChaosImmediate_JointMinLinearProjection;
 		ConstraintSettings.AngularProjection = ConstraintInstance->IsProjectionEnabled() ? ChaosImmediate_JointMaxAngularProjection : ChaosImmediate_JointMinAngularProjection;
-		ConstraintSettings.ParentInvMassScale = ConstraintInstance->ProfileInstance.bParentDominates ? (FReal)0 : (FReal)1;
+		ConstraintSettings.ParentInvMassScale = Profile.bParentDominates ? (FReal)0 : (FReal)1;
 
 		ConstraintSettings.bSoftLinearLimitsEnabled = ConstraintInstance->GetIsSoftLinearLimit();
 		ConstraintSettings.bSoftTwistLimitsEnabled = ConstraintInstance->GetIsSoftTwistLimit();
@@ -82,36 +86,41 @@ namespace ImmediatePhysics_Chaos
 		ConstraintSettings.SoftTwistDamping = ChaosImmediate_SoftDampingScale * ConstraintInstance->GetSoftTwistLimitDamping();
 		ConstraintSettings.SoftSwingStiffness = ChaosImmediate_SoftStiffnessScale * ConstraintInstance->GetSoftSwingLimitStiffness();
 		ConstraintSettings.SoftSwingDamping = ChaosImmediate_SoftDampingScale * ConstraintInstance->GetSoftSwingLimitDamping();
-
-		ConstraintSettings.LinearDriveTarget = ConstraintInstance->ProfileInstance.LinearDrive.PositionTarget;
-		ConstraintSettings.bLinearDriveEnabled[0] = ConstraintInstance->ProfileInstance.LinearDrive.XDrive.bEnablePositionDrive;
-		ConstraintSettings.bLinearDriveEnabled[1] = ConstraintInstance->ProfileInstance.LinearDrive.YDrive.bEnablePositionDrive;
-		ConstraintSettings.bLinearDriveEnabled[2] = ConstraintInstance->ProfileInstance.LinearDrive.ZDrive.bEnablePositionDrive;
-		ConstraintSettings.LinearDriveStiffness = 0.3f;// ConstraintInstance->ProfileInstance.LinearDrive.XDrive.bEnablePositionDrive;
-		ConstraintSettings.LinearDriveDamping = 0.0f;// ConstraintInstance->ProfileInstance.LinearDrive.XDrive.bEnablePositionDrive;
-
-		ConstraintSettings.AngularDriveTarget = FQuat(ConstraintInstance->ProfileInstance.AngularDrive.OrientationTarget);
-		// NOTE: Hard dependence on EJointAngularConstraintIndex - the following will break if we change the order
-		ConstraintSettings.AngularDriveTargetAngles[(int32)EJointAngularConstraintIndex::Twist] = FMath::DegreesToRadians(ConstraintInstance->ProfileInstance.AngularDrive.OrientationTarget.Roll);
-		ConstraintSettings.AngularDriveTargetAngles[(int32)EJointAngularConstraintIndex::Swing1] = FMath::DegreesToRadians(ConstraintInstance->ProfileInstance.AngularDrive.OrientationTarget.Yaw);
-		ConstraintSettings.AngularDriveTargetAngles[(int32)EJointAngularConstraintIndex::Swing2] = FMath::DegreesToRadians(ConstraintInstance->ProfileInstance.AngularDrive.OrientationTarget.Pitch);
-
-		if (ConstraintInstance->ProfileInstance.AngularDrive.AngularDriveMode == EAngularDriveMode::SLERP)
+		ConstraintSettings.SoftForceMode = EJointForceMode::Force;
+
+		ConstraintSettings.LinearDriveTarget = Profile.LinearDrive.PositionTarget;
+		ConstraintSettings.bLinearDriveEnabled[0] = Profile.LinearDrive.XDrive.bEnablePositionDrive;
+		ConstraintSettings.bLinearDriveEnabled[1] = Profile.LinearDrive.YDrive.bEnablePositionDrive;
+		ConstraintSettings.bLinearDriveEnabled[2] = Profile.LinearDrive.ZDrive.bEnablePositionDrive;
+		ConstraintSettings.LinearDriveStiffness = 0.3f;// Profile.LinearDrive.XDrive.bEnablePositionDrive;
+		ConstraintSettings.LinearDriveDamping = 0.0f;// Profile.LinearDrive.XDrive.bEnablePositionDrive;
+		ConstraintSettings.LinearDriveForceMode = EJointForceMode::Acceleration;
+
+		ConstraintSettings.AngularDrivePositionTarget = FQuat(Profile.AngularDrive.OrientationTarget);
+		ConstraintSettings.AngularDriveVelocityTarget = Profile.AngularDrive.AngularVelocityTarget;
+		// NOTE: Hard dependence on EJointAngularConstraintIndex - the following will break if we change the order (but can be easily fixed)
+		static_assert((int32)EJointAngularConstraintIndex::Twist == 0, "Angular drive targets have hard dependency on constraint order");
+		static_assert((int32)EJointAngularConstraintIndex::Swing1 == 2, "Angular drive targets have hard dependency on constraint order");
+		ConstraintSettings.AngularDriveTargetAngles[(int32)EJointAngularConstraintIndex::Twist] = FMath::DegreesToRadians(Profile.AngularDrive.OrientationTarget.Roll);
+		ConstraintSettings.AngularDriveTargetAngles[(int32)EJointAngularConstraintIndex::Swing1] = FMath::DegreesToRadians(Profile.AngularDrive.OrientationTarget.Yaw);
+		ConstraintSettings.AngularDriveTargetAngles[(int32)EJointAngularConstraintIndex::Swing2] = FMath::DegreesToRadians(Profile.AngularDrive.OrientationTarget.Pitch);
+
+		if (Profile.AngularDrive.AngularDriveMode == EAngularDriveMode::SLERP)
 		{
-			ConstraintSettings.bAngularSLerpDriveEnabled = ConstraintInstance->ProfileInstance.AngularDrive.SlerpDrive.bEnablePositionDrive;
-			ConstraintSettings.bAngularTwistDriveEnabled = false;
-			ConstraintSettings.bAngularSwingDriveEnabled = false;
-			
+			ConstraintSettings.bAngularSLerpPositionDriveEnabled = Profile.AngularDrive.SlerpDrive.bEnablePositionDrive;
+			ConstraintSettings.bAngularSLerpVelocityDriveEnabled = Profile.AngularDrive.SlerpDrive.bEnableVelocityDrive;
 		}
 		else
 		{
-			ConstraintSettings.bAngularSLerpDriveEnabled = false;
-			ConstraintSettings.bAngularTwistDriveEnabled = ConstraintInstance->ProfileInstance.AngularDrive.TwistDrive.bEnablePositionDrive;
-			ConstraintSettings.bAngularSwingDriveEnabled = ConstraintInstance->ProfileInstance.AngularDrive.SwingDrive.bEnablePositionDrive;
+			ConstraintSettings.bAngularTwistPositionDriveEnabled = Profile.AngularDrive.TwistDrive.bEnablePositionDrive;
+			ConstraintSettings.bAngularTwistVelocityDriveEnabled = Profile.AngularDrive.TwistDrive.bEnableVelocityDrive;
+			ConstraintSettings.bAngularSwingPositionDriveEnabled = Profile.AngularDrive.SwingDrive.bEnablePositionDrive;
+			ConstraintSettings.bAngularSwingVelocityDriveEnabled = Profile.AngularDrive.SwingDrive.bEnableVelocityDrive;
 		}
-
-		ConstraintSettings.AngularDriveStiffness = ChaosImmediate_DriveStiffnessScale * ConstraintInstance->ProfileInstance.AngularDrive.TwistDrive.Stiffness;
-		ConstraintSettings.AngularDriveDamping = ChaosImmediate_DriveDampingScale * ConstraintInstance->ProfileInstance.AngularDrive.TwistDrive.Damping;
+		// All angular drives in UE have same stiffness and damping settings
+		ConstraintSettings.AngularDriveStiffness = ChaosImmediate_DriveStiffnessScale * Profile.AngularDrive.TwistDrive.Stiffness;
+		ConstraintSettings.AngularDriveDamping = ChaosImmediate_DriveDampingScale * Profile.AngularDrive.TwistDrive.Damping;
+		ConstraintSettings.AngularDriveForceMode = EJointForceMode::Acceleration;
 
 		ConstraintSettings.Sanitize();
 
diff --git a/Engine/Source/Runtime/Engine/Private/PhysicsEngine/ImmediatePhysics/ImmediatePhysicsChaos/ImmediatePhysicsSimulation_Chaos.cpp b/Engine/Source/Runtime/Engine/Private/PhysicsEngine/ImmediatePhysics/ImmediatePhysicsChaos/ImmediatePhysicsSimulation_Chaos.cpp
--- a/Engine/Source/Runtime/Engine/Private/PhysicsEngine/ImmediatePhysics/ImmediatePhysicsChaos/ImmediatePhysicsSimulation_Chaos.cpp
+++ b/Engine/Source/Runtime/Engine/Private/PhysicsEngine/ImmediatePhysics/ImmediatePhysicsChaos/ImmediatePhysicsSimulation_Chaos.cpp
@@ -45,7 +45,7 @@ FAutoConsoleVariableRef CVarChaosImmPhysCollisionPriority(TEXT("p.Chaos.ImmPhys.
 FAutoConsoleVariableRef CVarChaosImmPhysThickness(TEXT("p.Chaos.ImmPhys.Collision.Thickness"), ChaosImmediate_Collision_Thickness, TEXT("ChaosImmediateThickness"));
 
 int32 ChaosImmediate_Joint_PairIterations = 1;
-int32 ChaosImmediate_Joint_PushOutPairIterations = 2;
+int32 ChaosImmediate_Joint_PushOutPairIterations = 1;
 float ChaosImmediate_Joint_SwingTwistAngleTolerance = 1.0e-6f;
 int32 ChaosImmediate_Joint_EnableTwistLimits = 1;
 int32 ChaosImmediate_Joint_EnableSwingLimits = 1;
diff --git a/Engine/Source/Runtime/Experimental/Chaos/Private/Chaos/Joint/PBDJointSolverCholesky.cpp b/Engine/Source/Runtime/Experimental/Chaos/Private/Chaos/Joint/PBDJointSolverCholesky.cpp
--- a/Engine/Source/Runtime/Experimental/Chaos/Private/Chaos/Joint/PBDJointSolverCholesky.cpp
+++ b/Engine/Source/Runtime/Experimental/Chaos/Private/Chaos/Joint/PBDJointSolverCholesky.cpp
@@ -64,7 +64,7 @@ namespace Chaos
 		Qs[1].EnforceShortestArcWith(Qs[0]);
 
 		Stiffness = FPBDJointUtilities::GetLinearStiffness(SolverSettings, JointSettings);
-		AngularDriveStiffness = FPBDJointUtilities::GetAngularDriveStiffness(SolverSettings, JointSettings);
+		AngularDriveStiffness = 0.0f;// FPBDJointUtilities::GetAngularDriveStiffness(SolverSettings, JointSettings);
 		SwingTwistAngleTolerance = SolverSettings.SwingTwistAngleTolerance;
 		bEnableTwistLimits = SolverSettings.bEnableTwistLimits;
 		bEnableSwingLimits = SolverSettings.bEnableSwingLimits;
@@ -677,10 +677,10 @@ namespace Chaos
 		EJointMotionType Swing1Motion = JointSettings.AngularMotionTypes[(int32)EJointAngularConstraintIndex::Swing1];
 		EJointMotionType Swing2Motion = JointSettings.AngularMotionTypes[(int32)EJointAngularConstraintIndex::Swing2];
 
-		bool bTwistDriveEnabled = bEnableDrives && JointSettings.bAngularTwistDriveEnabled && (TwistMotion != EJointMotionType::Locked);
-		bool bSwing1DriveEnabled = bEnableDrives && JointSettings.bAngularSwingDriveEnabled && (Swing1Motion != EJointMotionType::Locked);
-		bool bSwing2DriveEnabled = bEnableDrives && JointSettings.bAngularSwingDriveEnabled && (Swing2Motion != EJointMotionType::Locked);
-		bool bSlerpDriveEnabled = bEnableDrives && JointSettings.bAngularSLerpDriveEnabled && (Swing1Motion != EJointMotionType::Locked) && (Swing2Motion != EJointMotionType::Locked);
+		bool bTwistDriveEnabled = false;// bEnableDrives&& JointSettings.bAngularTwistDriveEnabled && (TwistMotion != EJointMotionType::Locked);
+		bool bSwing1DriveEnabled = false;// bEnableDrives && JointSettings.bAngularSwingDriveEnabled && (Swing1Motion != EJointMotionType::Locked);
+		bool bSwing2DriveEnabled = false;// bEnableDrives && JointSettings.bAngularSwingDriveEnabled && (Swing2Motion != EJointMotionType::Locked);
+		bool bSlerpDriveEnabled = false;// bEnableDrives && JointSettings.bAngularSLerpDriveEnabled && (Swing1Motion != EJointMotionType::Locked) && (Swing2Motion != EJointMotionType::Locked);
 		bool bAnyDriveEnabled = (bTwistDriveEnabled || bSwing1DriveEnabled || bSwing2DriveEnabled || bSlerpDriveEnabled);
 
 		if (bAnyDriveEnabled)
diff --git a/Engine/Source/Runtime/Experimental/Chaos/Private/Chaos/Joint/PBDJointSolverGaussSeidel.cpp b/Engine/Source/Runtime/Experimental/Chaos/Private/Chaos/Joint/PBDJointSolverGaussSeidel.cpp
--- a/Engine/Source/Runtime/Experimental/Chaos/Private/Chaos/Joint/PBDJointSolverGaussSeidel.cpp
+++ b/Engine/Source/Runtime/Experimental/Chaos/Private/Chaos/Joint/PBDJointSolverGaussSeidel.cpp
@@ -279,11 +279,8 @@ namespace Chaos
 		const FPBDJointSettings& JointSettings)
 	{
 		// With the parent body set to infinite mass...
-
-		// Then fix the rotations
 		ApplyRotationProjection(Dt, SolverSettings, JointSettings);
 
-		// We first apply a semi-physical linear projection
 		ApplyPositionProjection(Dt, SolverSettings, JointSettings);
 	}
 
@@ -369,26 +366,27 @@ namespace Chaos
 			bool bSwing2Locked = Swing2Motion == EJointMotionType::Locked;
 
 			// No SLerp drive if we have a locked rotation (it will be grayed out in the editor in this case, but could still have been set before the rotation was locked)
-			if (JointSettings.bAngularSLerpDriveEnabled && !bTwistLocked && !bSwing1Locked && !bSwing2Locked)
+			if ((JointSettings.bAngularSLerpPositionDriveEnabled || JointSettings.bAngularSLerpVelocityDriveEnabled) && !bTwistLocked && !bSwing1Locked && !bSwing2Locked)
 			{
 				ApplySLerpDrive(Dt, SolverSettings, JointSettings);
 			}
 			else
 			{
-				if (JointSettings.bAngularTwistDriveEnabled && !bTwistLocked)
+				if ((JointSettings.bAngularTwistPositionDriveEnabled || JointSettings.bAngularTwistVelocityDriveEnabled) && !bTwistLocked)
 				{
 					ApplyTwistDrive(Dt, SolverSettings, JointSettings);
 				}
 
-				if (JointSettings.bAngularSwingDriveEnabled && !bSwing1Locked && !bSwing2Locked)
+				const bool bSwingDriveEnabled = (JointSettings.bAngularSwingPositionDriveEnabled || JointSettings.bAngularSwingVelocityDriveEnabled);
+				if (bSwingDriveEnabled && !bSwing1Locked && !bSwing2Locked)
 				{
 					ApplyConeDrive(Dt, SolverSettings, JointSettings);
 				}
-				else if (JointSettings.bAngularSwingDriveEnabled && !bSwing1Locked)
+				else if (bSwingDriveEnabled && !bSwing1Locked)
 				{
 					ApplySwingDrive(Dt, SolverSettings, JointSettings, EJointAngularConstraintIndex::Swing1, EJointAngularAxisIndex::Swing1);
 				}
-				else if (JointSettings.bAngularSwingDriveEnabled && !bSwing2Locked)
+				else if (bSwingDriveEnabled && !bSwing2Locked)
 				{
 					ApplySwingDrive(Dt, SolverSettings, JointSettings, EJointAngularConstraintIndex::Swing2, EJointAngularAxisIndex::Swing2);
 				}
@@ -578,28 +576,45 @@ namespace Chaos
 		const FReal II = (II0 + II1);
 
 		// Damping angular velocity
-		FReal AngVel = 0;
+		FReal AngVelDt = 0;
 		if (Damping > KINDA_SMALL_NUMBER)
 		{
-			const FVec3 W0 = FRotation3::CalculateAngularVelocity(PrevQs[0], Qs[0], Dt);
-			const FVec3 W1 = FRotation3::CalculateAngularVelocity(PrevQs[1], Qs[1], Dt);
-			AngVel = FVec3::DotProduct(Axis0, W0) - FVec3::DotProduct(Axis1, W1);
+			const FVec3 W0 = FRotation3::CalculateAngularVelocity(PrevQs[0], Qs[0], (FReal)1.0);
+			const FVec3 W1 = FRotation3::CalculateAngularVelocity(PrevQs[1], Qs[1], (FReal)1.0);
+			AngVelDt = FVec3::DotProduct(Axis0, W0) - FVec3::DotProduct(Axis1, W1);
 		}
 
 		FReal DLambda = 0;
 		if (Stiffness > KINDA_SMALL_NUMBER)
 		{
+			// As below, but numerically stable for large values of stiffness (same form as in XPBD paper)
 			const FReal Alpha = (FReal)1 / (Stiffness * Dt * Dt);
-			const FReal AlphaBeta = Damping / Stiffness;
+			const FReal Gamma = Damping / (Stiffness * Dt);
 			if (bAccelerationMode)
 			{
-				const FReal Multiplier = (FReal)1 / (((FReal)1 + AlphaBeta / Dt) + Alpha);
-				DLambda = Multiplier * (Angle - AlphaBeta * AngVel) / II - Multiplier * Alpha * Lambda;
+				const FReal Multiplier = (FReal)1 / (((FReal)1 + Gamma) + Alpha);
+				DLambda = Multiplier * (Angle - Gamma * AngVelDt) / II - Multiplier * Alpha * Lambda;
 			}
 			else
 			{
-				const FReal Multiplier = (FReal)1 / (((FReal)1 + AlphaBeta / Dt) * II + Alpha);
-				DLambda = Multiplier * (Angle - Alpha * Lambda - AlphaBeta * AngVel);
+				const FReal Multiplier = (FReal)1 / (((FReal)1 + Gamma) * II + Alpha);
+				DLambda = Multiplier * (Angle - Gamma * AngVelDt - Alpha * Lambda);
+			}
+		}
+		else
+		{
+			// As above, but numerically stable at low stiffness (alpha -> infinity)
+			const FReal S = Stiffness * Dt * Dt;
+			const FReal D = Damping * Dt;
+			if (bAccelerationMode)
+			{
+				const FReal Multiplier = (FReal)1 / ((S + D) + (FReal)1);
+				DLambda = Multiplier * (S * Angle - D * AngVelDt) / II - Multiplier * Lambda;
+			}
+			else
+			{
+				const FReal Multiplier = (FReal)1 / ((S + D) * II + (FReal)1);
+				DLambda = Multiplier * (S * Angle - D * AngVelDt - Lambda);
 			}
 		}
 
@@ -742,8 +757,8 @@ namespace Chaos
 		const FReal DTwistAngle = TwistAngle - TwistAngleTarget;
 
 		// Apply twist correction
-		const FReal AngularDriveStiffness = FPBDJointUtilities::GetAngularDriveStiffness(SolverSettings, JointSettings);
-		const FReal AngularDriveDamping = FPBDJointUtilities::GetAngularDriveDamping(SolverSettings, JointSettings);
+		const FReal AngularDriveStiffness = FPBDJointUtilities::GetAngularTwistDriveStiffness(SolverSettings, JointSettings);
+		const FReal AngularDriveDamping = FPBDJointUtilities::GetAngularTwistDriveDamping(SolverSettings, JointSettings);
 		const bool bAccelerationMode = FPBDJointUtilities::GetDriveAccelerationMode(SolverSettings, JointSettings);
 		ApplyRotationCorrectionSoft(Dt, AngularDriveStiffness, AngularDriveDamping, bAccelerationMode, TwistAxis0, TwistAxis1, DTwistAngle, TwistDriveLambda);
 		UpdateDerivedState();
@@ -828,8 +843,8 @@ namespace Chaos
 		const FReal DSwingAngle = SwingAngle - SwingAngleTarget;
 
 		// Apply drive forces to each body
-		const FReal AngularDriveStiffness = FPBDJointUtilities::GetAngularDriveStiffness(SolverSettings, JointSettings);
-		const FReal AngularDriveDamping = FPBDJointUtilities::GetAngularDriveDamping(SolverSettings, JointSettings);
+		const FReal AngularDriveStiffness = FPBDJointUtilities::GetAngularSwingDriveStiffness(SolverSettings, JointSettings);
+		const FReal AngularDriveDamping = FPBDJointUtilities::GetAngularSwingDriveDamping(SolverSettings, JointSettings);
 		const bool bAccelerationMode = FPBDJointUtilities::GetDriveAccelerationMode(SolverSettings, JointSettings);
 		ApplyRotationCorrectionSoft(Dt, AngularDriveStiffness, AngularDriveDamping, bAccelerationMode, SwingAxis, SwingAxis, DSwingAngle, SwingDriveLambda);
 		UpdateDerivedState();
@@ -856,15 +871,27 @@ namespace Chaos
 			const FReal AngularProjection = FPBDJointUtilities::GetAngularProjection(SolverSettings, JointSettings);
 
 			const FVec3 DR1 = -DSwingAngle * SwingAxis;
-			ApplyRotationDelta(1, AngularProjection, DR1);
+			const FVec3 DP1 = -FVec3::CrossProduct(DR1, Xs[1] - Ps[1]);
+			FVec3 DV1 = FVec3(0);
+			FVec3 DW1 = FVec3(0);
 
 			const FReal WAxis = FVec3::DotProduct(Ws[1] - Ws[0], SwingAxis);
 			if (FMath::Sign(WAxis) == FMath::Sign(DSwingAngle))
 			{
-				const FVec3 DW1 = -WAxis * SwingAxis;
-				ApplyVelocityDelta(1, AngularProjection, FVec3(0), DW1);
+				DW1 = -WAxis * SwingAxis;
 			}
 
+			FReal DP1Len = DP1.Size();
+			if (DP1Len > KINDA_SMALL_NUMBER)
+			{
+				const FVec3 DP1Dir = DP1 / DP1Len;
+				DV1 = -FVec3::DotProduct((Vs[1] - Vs[0]), DP1Dir) * DP1Dir;
+			}
+
+			ApplyRotationDelta(1, AngularProjection, DR1);
+			ApplyPositionDelta(1, AngularProjection, DP1);
+			ApplyVelocityDelta(1, AngularProjection, DV1, DW1);
+
 			UpdateDerivedState(1);
 		}
 	}
@@ -956,8 +983,8 @@ namespace Chaos
 		const FReal DSwingAngle = SwingAngle - SwingAngleTarget;
 
 		// Apply drive forces to each body
-		const FReal AngularDriveStiffness = FPBDJointUtilities::GetAngularDriveStiffness(SolverSettings, JointSettings);
-		const FReal AngularDriveDamping = FPBDJointUtilities::GetAngularDriveDamping(SolverSettings, JointSettings);
+		const FReal AngularDriveStiffness = FPBDJointUtilities::GetAngularSwingDriveStiffness(SolverSettings, JointSettings);
+		const FReal AngularDriveDamping = FPBDJointUtilities::GetAngularSwingDriveDamping(SolverSettings, JointSettings);
 		const bool bAccelerationMode = FPBDJointUtilities::GetDriveAccelerationMode(SolverSettings, JointSettings);
 		ApplyRotationCorrectionSoft(Dt, AngularDriveStiffness, AngularDriveDamping, bAccelerationMode, SwingAxis, SwingAxis, DSwingAngle, SwingDriveLambda);
 		UpdateDerivedState();
@@ -981,7 +1008,7 @@ namespace Chaos
 		const FPBDJointSettings& JointSettings)
 	{
 		// Calculate the rotation we need to apply to resolve the rotation delta
-		const FRotation3 TargetR1 = Rs[0] * JointSettings.AngularDriveTarget;
+		const FRotation3 TargetR1 = Rs[0] * JointSettings.AngularDrivePositionTarget;
 		const FRotation3 DR1 = TargetR1 * Rs[1].Inverse();
 
 		FVec3 SLerpAxis;
@@ -993,8 +1020,8 @@ namespace Chaos
 				SLerpAngle = SLerpAngle - (FReal)2 * PI;
 			}
 
-			const FReal AngularDriveStiffness = FPBDJointUtilities::GetAngularDriveStiffness(SolverSettings, JointSettings);
-			const FReal AngularDriveDamping = FPBDJointUtilities::GetAngularDriveDamping(SolverSettings, JointSettings);
+			const FReal AngularDriveStiffness = FPBDJointUtilities::GetAngularSLerpDriveStiffness(SolverSettings, JointSettings);
+			const FReal AngularDriveDamping = FPBDJointUtilities::GetAngularSLerpDriveDamping(SolverSettings, JointSettings);
 			const bool bAccelerationMode = FPBDJointUtilities::GetDriveAccelerationMode(SolverSettings, JointSettings);
 			ApplyRotationCorrectionSoft(Dt, AngularDriveStiffness, AngularDriveDamping, bAccelerationMode, SLerpAxis, SLerpAxis, -SLerpAngle, SwingDriveLambda);
 		}
diff --git a/Engine/Source/Runtime/Experimental/Chaos/Private/Chaos/PBDJointConstraintUtilities.cpp b/Engine/Source/Runtime/Experimental/Chaos/Private/Chaos/PBDJointConstraintUtilities.cpp
--- a/Engine/Source/Runtime/Experimental/Chaos/Private/Chaos/PBDJointConstraintUtilities.cpp
+++ b/Engine/Source/Runtime/Experimental/Chaos/Private/Chaos/PBDJointConstraintUtilities.cpp
@@ -83,18 +83,70 @@ namespace Chaos
 		return (SolverSettings.LinearDriveDamping > 0.0f) ? SolverSettings.LinearDriveDamping : JointSettings.LinearDriveDamping;
 	}
 
-	FReal FPBDJointUtilities::GetAngularDriveStiffness(
+	FReal FPBDJointUtilities::GetAngularTwistDriveStiffness(
 		const FPBDJointSolverSettings& SolverSettings,
 		const FPBDJointSettings& JointSettings)
 	{
-		return (SolverSettings.AngularDriveStiffness > 0.0f) ? SolverSettings.AngularDriveStiffness : JointSettings.AngularDriveStiffness;
+		if (JointSettings.bAngularTwistPositionDriveEnabled)
+		{
+			return (SolverSettings.AngularDriveStiffness > 0.0f) ? SolverSettings.AngularDriveStiffness : JointSettings.AngularDriveStiffness;
+		}
+		return 0.0f;
+	}
+
+	FReal FPBDJointUtilities::GetAngularTwistDriveDamping(
+		const FPBDJointSolverSettings& SolverSettings,
+		const FPBDJointSettings& JointSettings)
+	{
+		if (JointSettings.bAngularTwistVelocityDriveEnabled)
+		{
+			return (SolverSettings.AngularDriveDamping > 0.0f) ? SolverSettings.AngularDriveDamping : JointSettings.AngularDriveDamping;
+		}
+		return 0.0f;
+	}
+
+	FReal FPBDJointUtilities::GetAngularSwingDriveStiffness(
+		const FPBDJointSolverSettings& SolverSettings,
+		const FPBDJointSettings& JointSettings)
+	{
+		if (JointSettings.bAngularSwingPositionDriveEnabled)
+		{
+			return (SolverSettings.AngularDriveStiffness > 0.0f) ? SolverSettings.AngularDriveStiffness : JointSettings.AngularDriveStiffness;
+		}
+		return 0.0f;
+	}
+
+	FReal FPBDJointUtilities::GetAngularSwingDriveDamping(
+		const FPBDJointSolverSettings& SolverSettings,
+		const FPBDJointSettings& JointSettings)
+	{
+		if (JointSettings.bAngularSwingVelocityDriveEnabled)
+		{
+			return (SolverSettings.AngularDriveDamping > 0.0f) ? SolverSettings.AngularDriveDamping : JointSettings.AngularDriveDamping;
+		}
+		return 0.0f;
+	}
+
+	FReal FPBDJointUtilities::GetAngularSLerpDriveStiffness(
+		const FPBDJointSolverSettings& SolverSettings,
+		const FPBDJointSettings& JointSettings)
+	{
+		if (JointSettings.bAngularSLerpPositionDriveEnabled)
+		{
+			return (SolverSettings.AngularDriveStiffness > 0.0f) ? SolverSettings.AngularDriveStiffness : JointSettings.AngularDriveStiffness;
+		}
+		return 0.0f;
 	}
 
-	FReal FPBDJointUtilities::GetAngularDriveDamping(
+	FReal FPBDJointUtilities::GetAngularSLerpDriveDamping(
 		const FPBDJointSolverSettings& SolverSettings,
 		const FPBDJointSettings& JointSettings)
 	{
-		return (SolverSettings.AngularDriveDamping > 0.0f) ? SolverSettings.AngularDriveDamping : JointSettings.AngularDriveDamping;
+		if (JointSettings.bAngularSLerpVelocityDriveEnabled)
+		{
+			return (SolverSettings.AngularDriveDamping > 0.0f) ? SolverSettings.AngularDriveDamping : JointSettings.AngularDriveDamping;
+		}
+		return 0.0f;
 	}
 
 	FReal FPBDJointUtilities::GetLinearProjection(
@@ -217,7 +269,6 @@ namespace Chaos
 		}
 	}
 
-
 	
 	void FPBDJointUtilities::GetConditionedInverseMass(
 		const FReal InM0,
@@ -237,7 +288,6 @@ namespace Chaos
 	}
 
 
-	
 	FVec3 FPBDJointUtilities::GetSphereLimitedPositionError(const FVec3& CX, const FReal Radius)
 	{
 		FReal CXLen = CX.Size();
@@ -253,24 +303,7 @@ namespace Chaos
 		return CX;
 	}
 
-	
-	FVec3 FPBDJointUtilities::GetSphereLimitedVelocityError(const FVec3& CX, const FReal Radius, const FVec3& CV)
-	{
-		FReal CXLen = CX.Size();
-		if (CXLen < Radius)
-		{
-			return FVec3(0, 0, 0);
-		}
-		else if (CXLen > SMALL_NUMBER)
-		{
-			FVec3 Dir = CX / CXLen;
-			FReal CVDir = FVec3::DotProduct(CV, Dir);
-			return FMath::Max((FReal)0, CVDir) * Dir;
-		}
-		return CV;
-	}
 
-	
 	FVec3 FPBDJointUtilities::GetCylinderLimitedPositionError(const FVec3& InCX, const FVec3& Axis, const FReal Limit, const EJointMotionType AxisMotion)
 	{
 		FVec3 CXAxis = FVec3::DotProduct(InCX, Axis) * Axis;
@@ -292,34 +325,7 @@ namespace Chaos
 		return CXAxis + CXPlane;
 	}
 
-	
-	FVec3 FPBDJointUtilities::GetCylinderLimitedVelocityError(const FVec3& InCX, const FVec3& Axis, const FReal Limit, const EJointMotionType AxisMotion, const FVec3& CV)
-	{
-		FVec3 CXAxis = FVec3::DotProduct(InCX, Axis) * Axis;
-		FVec3 CXPlane = InCX - CXAxis;
-		FReal CXPlaneLen = CXPlane.Size();
-
-		FVec3 CVAxis = FVec3::DotProduct(CV, Axis) * Axis;;
-		FVec3 CVPlane = CV - CVAxis;
-
-		if (AxisMotion == EJointMotionType::Free)
-		{
-			CVAxis = FVec3(0, 0, 0);
-		}
-		if (CXPlaneLen < Limit)
-		{
-			CVPlane = FVec3(0, 0, 0);
-		}
-		else if (CXPlaneLen > KINDA_SMALL_NUMBER)
-		{
-			FVec3 Dir = CXPlane / CXPlaneLen;
-			FReal CVDir = FVec3::DotProduct(CV, Dir);
-			CVPlane = FMath::Max((FReal)0, CVDir) * Dir;
-		}
-		return CVAxis + CVPlane;
-	}
 
-	
 	FVec3 FPBDJointUtilities::GetLineLimitedPositionError(const FVec3& CX, const FVec3& Axis, const FReal Limit, const EJointMotionType AxisMotion)
 	{
 		FReal CXDist = FVec3::DotProduct(CX, Axis);
@@ -337,27 +343,7 @@ namespace Chaos
 		}
 	}
 
-	
-	FVec3 FPBDJointUtilities::GetLineLimitedVelocityError(const FVec3& CX, const FVec3& Axis, const FReal Limit, const EJointMotionType AxisMotion, const FVec3& CV)
-	{
-		FReal CXDist = FVec3::DotProduct(CX, Axis);
-		FReal CVAxis = FVec3::DotProduct(CV, Axis);
-		if ((AxisMotion == EJointMotionType::Free) || (FMath::Abs(CXDist) < Limit))
-		{
-			return CV - CVAxis * Axis;
-		}
-		else if (CXDist >= Limit)
-		{
-			return CV - FMath::Min((FReal)0, CVAxis) * Axis;
-		}
-		else
-		{
-			return CV - FMath::Max((FReal)0, CVAxis) * Axis;
-		}
-	}
-
 
-	
 	FVec3 FPBDJointUtilities::GetLimitedPositionError(const FPBDJointSettings& JointSettings, const FRotation3& R0, const FVec3& InCX)
 	{
 		const TVector<EJointMotionType, 3>& Motion = JointSettings.LinearMotionTypes;
@@ -410,1243 +396,4 @@ namespace Chaos
 			return CX;
 		}
 	}
-
-	
-	FVec3 FPBDJointUtilities::GetLimitedVelocityError(const FPBDJointSettings& JointSettings, const FRotation3& R0, const FVec3& InCX, const FVec3& InCV)
-	{
-		const TVector<EJointMotionType, 3>& Motion = JointSettings.LinearMotionTypes;
-		if ((Motion[0] == EJointMotionType::Locked) && (Motion[1] == EJointMotionType::Locked) && (Motion[2] == EJointMotionType::Locked))
-		{
-			return InCV;
-		}
-		else if ((Motion[0] == EJointMotionType::Limited) && (Motion[1] == EJointMotionType::Limited) && (Motion[2] == EJointMotionType::Limited))
-		{
-			// Spherical distance constraints
-			return GetSphereLimitedVelocityError(InCX, JointSettings.LinearLimit, InCV);
-		}
-		else if ((Motion[1] == EJointMotionType::Limited) && (Motion[2] == EJointMotionType::Limited))
-		{
-			// Circular Limit (X Axis)
-			FVec3 Axis = R0 * FVec3(1, 0, 0);
-			return GetCylinderLimitedVelocityError(InCX, Axis, JointSettings.LinearLimit, Motion[0], InCV);
-		}
-		else if ((Motion[0] == EJointMotionType::Limited) && (Motion[2] == EJointMotionType::Limited))
-		{
-			// Circular Limit (Y Axis)
-			FVec3 Axis = R0 * FVec3(0, 1, 0);
-			return GetCylinderLimitedVelocityError(InCX, Axis, JointSettings.LinearLimit, Motion[1], InCV);
-		}
-		else if ((Motion[0] == EJointMotionType::Limited) && (Motion[1] == EJointMotionType::Limited))
-		{
-			// Circular Limit (Z Axis)
-			FVec3 Axis = R0 * FVec3(0, 0, 1);
-			return GetCylinderLimitedVelocityError(InCX, Axis, JointSettings.LinearLimit, Motion[2], InCV);
-		}
-		else
-		{
-			// Line/Square/Cube Limits (no way to author square or cube limits, but would work if we wanted it)
-			FVec3 CV = InCV;
-			if (Motion[0] != EJointMotionType::Locked)
-			{
-				FVec3 Axis = R0 * FVec3(1, 0, 0);
-				CV = GetLineLimitedVelocityError(InCX, Axis, JointSettings.LinearLimit, Motion[0], CV);
-			}
-			if (Motion[1] != EJointMotionType::Locked)
-			{
-				FVec3 Axis = R0 * FVec3(0, 1, 0);
-				CV = GetLineLimitedVelocityError(InCX, Axis, JointSettings.LinearLimit, Motion[1], CV);
-			}
-			if (Motion[2] != EJointMotionType::Locked)
-			{
-				FVec3 Axis = R0 * FVec3(0, 0, 1);
-				CV = GetLineLimitedVelocityError(InCX, Axis, JointSettings.LinearLimit, Motion[2], CV);
-			}
-			return CV;
-		}
-	}
-
-	
-	void ApplyPositionDelta(
-		const FReal Dt,
-		const FPBDJointSolverSettings& SolverSettings,
-		const FPBDJointSettings& JointSettings,
-		const FReal Stiffness,
-		FVec3& P0,
-		FVec3& V0,
-		FVec3& P1,
-		FVec3& V1,
-		const FVec3& DP0,
-		const FVec3& DP1)
-	{
-		P0 = P0 + Stiffness * DP0;
-		P1 = P1 + Stiffness * DP1;
-		if (Dt > SMALL_NUMBER)
-		{
-			V0 = V0 + Stiffness * DP0 / Dt;
-			V1 = V1 + Stiffness * DP1 / Dt;
-		}
-	}
-
-	
-	void ApplyVelocityDelta(
-		const FReal Dt,
-		const FPBDJointSolverSettings& SolverSettings,
-		const FPBDJointSettings& JointSettings,
-		const FReal Stiffness,
-		FVec3& P0,
-		FVec3& V0,
-		FVec3& P1,
-		FVec3& V1,
-		const FVec3& DV0,
-		const FVec3& DV1)
-	{
-		V0 = V0 + Stiffness * DV0;
-		V1 = V1 + Stiffness * DV1;
-		P0 = P0 + Stiffness * DV0 * Dt;
-		P1 = P1 + Stiffness * DV1 * Dt;
-	}
-
-	
-	void ApplyRotationDelta(
-		const FReal Dt,
-		const FPBDJointSolverSettings& SolverSettings,
-		const FPBDJointSettings& JointSettings,
-		const FReal Stiffness,
-		FRotation3& Q0,
-		FVec3& W0,
-		FRotation3& Q1,
-		FVec3& W1,
-		const FVec3& DR0,
-		const FVec3& DR1)
-	{
-		const FRotation3 DQ0 = (FRotation3::FromElements(Stiffness * DR0, 0) * Q0) * (FReal)0.5;
-		const FRotation3 DQ1 = (FRotation3::FromElements(Stiffness * DR1, 0) * Q1) * (FReal)0.5;
-		Q0 = (Q0 + DQ0).GetNormalized();
-		Q1 = (Q1 + DQ1).GetNormalized();
-		Q1.EnforceShortestArcWith(Q0);
-		
-		if (Dt > SMALL_NUMBER)
-		{
-			W0 = W0 + Stiffness * DR0 / Dt;
-			W1 = W1 + Stiffness * DR1 / Dt;
-		}
-	}
-
-
-	
-	void ApplyRotationVelocityDelta(
-		const FReal Dt,
-		const FPBDJointSolverSettings& SolverSettings,
-		const FPBDJointSettings& JointSettings,
-		const FReal Stiffness,
-		FRotation3& Q0,
-		FVec3& W0,
-		FRotation3& Q1,
-		FVec3& W1,
-		const FVec3& DW0,
-		const FVec3& DW1)
-	{
-		W0 = W0 + Stiffness * DW0;
-		W1 = W1 + Stiffness * DW1;
-
-		const FRotation3 DQ0 = (FRotation3::FromElements(Stiffness * DW0 * Dt, 0) * Q0) * (FReal)0.5;
-		const FRotation3 DQ1 = (FRotation3::FromElements(Stiffness * DW1 * Dt, 0) * Q1) * (FReal)0.5;
-		Q0 = (Q0 + DQ0).GetNormalized();
-		Q1 = (Q1 + DQ1).GetNormalized();
-		Q1.EnforceShortestArcWith(Q0);
-	}
-
-
-	
-	void ApplyRotationVelocityDelta(
-		const FReal Dt,
-		const FPBDJointSolverSettings& SolverSettings,
-		const FPBDJointSettings& JointSettings,
-		const FReal Stiffness,
-		FRotation3& Q0,
-		FVec3& W0,
-		FRotation3& Q1,
-		FVec3& W1,
-		const FReal InvM0,
-		const FMatrix33& InvIL0,
-		const FReal InvM1,
-		const FMatrix33& InvIL1,
-		const FVec3& Axis0,
-		const FVec3& Axis1,
-		const FReal WC)
-	{
-		const FMatrix33 InvI0 = Utilities::ComputeWorldSpaceInertia(Q0, InvIL0);
-		const FMatrix33 InvI1 = Utilities::ComputeWorldSpaceInertia(Q1, InvIL1);
-		//const FReal I0 = FVec3::DotProduct(Axis0, Utilities::Multiply(InvI0, Axis0));
-		//const FReal I1 = FVec3::DotProduct(Axis1, Utilities::Multiply(InvI1, Axis1));
-		//const FVec3 DW0 = Axis0 * WC * I0 / (I0 + I1);
-		//const FVec3 DW1 = -Axis1 * WC * I1 / (I0 + I1);
-		const FReal L = (FReal)1 / (FVec3::DotProduct(Axis0, Utilities::Multiply(InvI0, Axis0)) + FVec3::DotProduct(Axis1, Utilities::Multiply(InvI1, Axis1)));
-		const FVec3 DW0 = Utilities::Multiply(InvI0, Axis0) * L * WC;
-		const FVec3 DW1 = -Utilities::Multiply(InvI1, Axis1) * L * WC;
-
-		ApplyRotationVelocityDelta(Dt, SolverSettings, JointSettings, Stiffness, Q0, W0, Q1, W1, DW0, DW1);
-	}
-
-	
-	void ApplyRotationDelta(
-		const FReal Dt,
-		const FPBDJointSolverSettings& SolverSettings,
-		const FPBDJointSettings& JointSettings,
-		const FReal Stiffness,
-		FRotation3& Q0,
-		FVec3& W0,
-		FRotation3& Q1,
-		FVec3& W1,
-		const FReal InvM0,
-		const FMatrix33& InvIL0,
-		const FReal InvM1,
-		const FMatrix33& InvIL1,
-		const FVec3& Axis0,
-		const FReal Angle0,
-		const FVec3& Axis1,
-		const FReal Angle1)
-	{
-		const FMatrix33 InvI0 = Utilities::ComputeWorldSpaceInertia(Q0, InvIL0);
-		const FMatrix33 InvI1 = Utilities::ComputeWorldSpaceInertia(Q1, InvIL1);
-
-		//const FReal I0 = FVec3::DotProduct(Axis0, Utilities::Multiply(InvI0, Axis0));
-		//const FReal I1 = FVec3::DotProduct(Axis1, Utilities::Multiply(InvI1, Axis1));
-		//const FVec3 DR0 = Axis0 * Angle0 * I0 / (I0 + I1);
-		//const FVec3 DR1 = Axis1 * Angle1 * I1 / (I0 + I1);
-		const FReal L = (FReal)1 / (FVec3::DotProduct(Axis0, Utilities::Multiply(InvI0, Axis0)) + FVec3::DotProduct(Axis1, Utilities::Multiply(InvI1, Axis1)));
-		const FVec3 DR0 = Utilities::Multiply(InvI0, Axis0) * L * Angle0;
-		const FVec3 DR1 = Utilities::Multiply(InvI1, Axis1) * L * Angle1;
-
-		ApplyRotationDelta(Dt, SolverSettings, JointSettings, Stiffness, Q0, W0, Q1, W1, DR0, DR1);
-	}
-
-	
-	void FPBDJointUtilities::CalculateSwingConstraintSpace(
-		const FPBDJointSolverSettings& SolverSettings,
-		const FPBDJointSettings& JointSettings,
-		const FRigidTransform3& XL0,
-		const FRigidTransform3& XL1,
-		const FVec3& P0,
-		const FRotation3& Q0,
-		const FVec3& P1,
-		const FRotation3& Q1,
-		FVec3& OutX0,
-		FMatrix33& OutR0,
-		FVec3& OutX1,
-		FMatrix33& OutR1,
-		FVec3& OutCR)
-	{
-		const FVec3 X0 = P0 + Q0 * XL0.GetTranslation();
-		const FVec3 X1 = P1 + Q1 * XL1.GetTranslation();
-		const FRotation3 R0 = Q0 * XL0.GetRotation();
-		const FRotation3 R1 = Q1 * XL1.GetRotation();
-
-		FRotation3 R01 = R0.Inverse() * R1;
-		FRotation3 R01Twist, R01Swing;
-		R01.ToSwingTwist(FJointConstants::TwistAxis(), R01Swing, R01Twist);
-		R01Swing = R01Swing.GetNormalized();
-		R01Twist = R01Twist.GetNormalized();
-
-		FVec3 TwistAxis01;
-		FReal TwistAngle = (FReal)0;
-		R01Twist.ToAxisAndAngleSafe(TwistAxis01, TwistAngle, FJointConstants::TwistAxis(), SolverSettings.SwingTwistAngleTolerance);
-		if (TwistAngle > PI)
-		{
-			TwistAngle = TwistAngle - (FReal)2 * PI;
-		}
-		if (FVec3::DotProduct(TwistAxis01, FJointConstants::TwistAxis()) < 0)
-		{
-			TwistAngle = -TwistAngle;
-		}
-
-		const FMatrix33 Axes0 = R0.ToMatrix();
-		const FMatrix33 Axes1 = R1.ToMatrix();
-
-		FReal Swing1Angle = (FReal)0;
-		const FVec3 SwingCross1 = FVec3::CrossProduct(Axes0.GetAxis((int32)EJointAngularConstraintIndex::Swing1), Axes1.GetAxis((int32)EJointAngularConstraintIndex::Swing1));
-		const FReal SwingCross1Len = SwingCross1.Size();
-		if (SwingCross1Len > KINDA_SMALL_NUMBER)
-		{
-			Swing1Angle = FMath::Asin(FMath::Clamp(SwingCross1Len, (FReal)0, (FReal)1));
-		}
-		const FReal Swing1Dot = FVec3::DotProduct(Axes0.GetAxis((int32)EJointAngularConstraintIndex::Swing1), Axes1.GetAxis((int32)EJointAngularConstraintIndex::Swing1));
-		if (Swing1Dot < (FReal)0)
-		{
-			Swing1Angle = (FReal)PI - Swing1Angle;
-		}
-
-		FReal Swing2Angle = (FReal)0;
-		const FVec3 SwingCross2 = FVec3::CrossProduct(Axes0.GetAxis((int32)EJointAngularConstraintIndex::Swing2), Axes1.GetAxis((int32)EJointAngularConstraintIndex::Swing2));
-		const FReal SwingCross2Len = SwingCross2.Size();
-		if (SwingCross2Len > KINDA_SMALL_NUMBER)
-		{
-			Swing2Angle = FMath::Asin(FMath::Clamp(SwingCross2Len, (FReal)0, (FReal)1));
-		}
-		const FReal Swing2Dot = FVec3::DotProduct(Axes0.GetAxis((int32)EJointAngularConstraintIndex::Swing2), Axes1.GetAxis((int32)EJointAngularConstraintIndex::Swing2));
-		if (Swing2Dot < (FReal)0)
-		{
-			Swing2Angle = (FReal)PI - Swing2Angle;
-		}
-
-		OutX0 = X0;
-		OutX1 = X1;
-		OutR0 = R0.ToMatrix();
-		OutR1 = R1.ToMatrix();
-		OutCR[(int32)EJointAngularAxisIndex::Twist] = TwistAngle;
-		OutCR[(int32)EJointAngularAxisIndex::Swing1] = Swing1Angle;
-		OutCR[(int32)EJointAngularAxisIndex::Swing2] = Swing2Angle;
-	}
-
-	
-	void FPBDJointUtilities::CalculateConeConstraintSpace(
-		const FPBDJointSolverSettings& SolverSettings,
-		const FPBDJointSettings& JointSettings,
-		const FRigidTransform3& XL0,
-		const FRigidTransform3& XL1,
-		const FVec3& P0,
-		const FRotation3& Q0,
-		const FVec3& P1,
-		const FRotation3& Q1,
-		FVec3& OutX0,
-		FMatrix33& OutR0, 
-		FVec3& OutX1, 
-		FMatrix33& OutR1, 
-		FVec3& OutCR)
-	{
-		const FVec3 X0 = P0 + Q0 * XL0.GetTranslation();
-		const FVec3 X1 = P1 + Q1 * XL1.GetTranslation();
-		const FRotation3 R0 = Q0 * XL0.GetRotation();
-		const FRotation3 R1 = Q1 * XL1.GetRotation();
-
-		FRotation3 R01 = R0.Inverse() * R1;
-		FRotation3 R01Twist, R01Swing;
-		R01.ToSwingTwist(FJointConstants::TwistAxis(), R01Swing, R01Twist);
-		R01Swing = R01Swing.GetNormalized();
-		R01Twist = R01Twist.GetNormalized();
-
-		FVec3 TwistAxis01;
-		FReal TwistAngle;
-		R01Twist.ToAxisAndAngleSafe(TwistAxis01, TwistAngle, FJointConstants::TwistAxis(), SolverSettings.SwingTwistAngleTolerance);
-		if (TwistAngle > PI)
-		{
-			TwistAngle = TwistAngle - (FReal)2 * PI;
-		}
-		if (FVec3::DotProduct(TwistAxis01, FJointConstants::TwistAxis()) < 0)
-		{
-			TwistAxis01 = -TwistAxis01;
-			TwistAngle = -TwistAngle;
-		}
-		FVec3 TwistAxis0 = R0 * TwistAxis01;
-		FVec3 TwistAxis1 = R1 * TwistAxis01;
-
-		FVec3 SwingAxis01;
-		FReal SwingAngle;
-		R01Swing.ToAxisAndAngleSafe(SwingAxis01, SwingAngle, FJointConstants::Swing1Axis(), SolverSettings.SwingTwistAngleTolerance);
-		if (SwingAngle > PI)
-		{
-			SwingAngle = SwingAngle - (FReal)2 * PI;
-		}
-		FVec3 SwingAxis0 = R0 * SwingAxis01;
-		FVec3 SwingAxis1 = SwingAxis0;
-
-		OutX0 = X0;
-		OutX1 = X1;
-		OutR0.SetAxis((int32)EJointAngularAxisIndex::Twist, TwistAxis0);
-		OutR0.SetAxis((int32)EJointAngularAxisIndex::Swing1, SwingAxis0);
-		OutR0.SetAxis((int32)EJointAngularAxisIndex::Swing2, FVec3::CrossProduct(SwingAxis0, TwistAxis0));
-		OutR1.SetAxis((int32)EJointAngularAxisIndex::Twist, TwistAxis1);
-		OutR1.SetAxis((int32)EJointAngularAxisIndex::Swing1, SwingAxis1);
-		OutR1.SetAxis((int32)EJointAngularAxisIndex::Swing2, FVec3::CrossProduct(SwingAxis1, TwistAxis1));
-		OutCR[(int32)EJointAngularAxisIndex::Twist] = TwistAngle;
-		OutCR[(int32)EJointAngularAxisIndex::Swing1] = SwingAngle;
-		OutCR[(int32)EJointAngularAxisIndex::Swing2] = (FReal)0;
-	}
-
-	
-	void FPBDJointUtilities::ApplyJointPositionConstraint(
-		const FReal Dt,
-		const FPBDJointSolverSettings& SolverSettings,
-		const FPBDJointSettings& JointSettings,
-		const FReal Stiffness,
-		const FRigidTransform3& XL0,
-		const FRigidTransform3& XL1,
-		FVec3& P0,
-		FRotation3& Q0,
-		FVec3& V0,
-		FVec3& W0,
-		FVec3& P1,
-		FRotation3& Q1,
-		FVec3& V1,
-		FVec3& W1,
-		FReal InvM0,
-		const FMatrix33& InvIL0,
-		FReal InvM1,
-		const FMatrix33& InvIL1)
-	{
-
-		const FVec3 X0 = P0 + Q0 * XL0.GetTranslation();
-		const FVec3 X1 = P1 + Q1 * XL1.GetTranslation();
-		const FRotation3 R0 = Q0 * XL0.GetRotation();
-		const FMatrix33 InvI0 = Utilities::ComputeWorldSpaceInertia(Q0, InvIL0);
-		const FMatrix33 InvI1 = Utilities::ComputeWorldSpaceInertia(Q1, InvIL1);
-
-		// Calculate constraint error
-		const FVec3 CX = GetLimitedPositionError(JointSettings, R0, X1 - X0);
-
-		// Calculate constraint correction
-		FMatrix33 M0 = FMatrix33(0, 0, 0);
-		FMatrix33 M1 = FMatrix33(0, 0, 0);
-		if (InvM0 > 0)
-		{
-			M0 = Utilities::ComputeJointFactorMatrix(X0 - P0, InvI0, InvM0);
-		}
-		if (InvM1 > 0)
-		{
-			M1 = Utilities::ComputeJointFactorMatrix(X1 - P1, InvI1, InvM1);
-		}
-		const FMatrix33 MI = (M0 + M1).Inverse();
-		const FVec3 DX = Utilities::Multiply(MI, CX);
-
-		// Apply constraint correction
-		const FVec3 DP0 = InvM0 * DX;
-		const FVec3 DP1 = -InvM1 * DX;
-		const FVec3 DR0 = Utilities::Multiply(InvI0, FVec3::CrossProduct(X0 - P0, DX));
-		const FVec3 DR1 = Utilities::Multiply(InvI1, FVec3::CrossProduct(X1 - P1, -DX));
-
-		ApplyPositionDelta(Dt, SolverSettings, JointSettings, Stiffness, P0, V0, P1, V1, DP0, DP1);
-		ApplyRotationDelta(Dt, SolverSettings, JointSettings, Stiffness, Q0, W0, Q1, W1, DR0, DR1);
-	}
-
-	
-	void FPBDJointUtilities::ApplyJointVelocityConstraint(
-		const FReal Dt,
-		const FPBDJointSolverSettings& SolverSettings,
-		const FPBDJointSettings& JointSettings,
-		const FReal Stiffness,
-		const FRigidTransform3& XL0,
-		const FRigidTransform3& XL1,
-		FVec3& P0,
-		FRotation3& Q0,
-		FVec3& V0,
-		FVec3& W0,
-		FVec3& P1,
-		FRotation3& Q1,
-		FVec3& V1,
-		FVec3& W1,
-		FReal InvM0,
-		const FMatrix33& InvIL0,
-		FReal InvM1,
-		const FMatrix33& InvIL1)
-	{
-		const FVec3 XC0 = Q0 * XL0.GetTranslation();
-		const FVec3 XC1 = Q1 * XL1.GetTranslation();
-		const FRotation3 R0 = Q0 * XL0.GetRotation();
-		const FMatrix33 InvI0 = Utilities::ComputeWorldSpaceInertia(Q0, InvIL0);
-		const FMatrix33 InvI1 = Utilities::ComputeWorldSpaceInertia(Q1, InvIL1);
-
-		const FVec3 VC0 = V0 + FVec3::CrossProduct(W0, XC0);
-		const FVec3 VC1 = V1 + FVec3::CrossProduct(W1, XC1);
-		const FVec3 VC = GetLimitedVelocityError(JointSettings, R0, (P1 + XC1) - (P0 + XC0), VC1 - VC0);
-
-		// Calculate constraint correction
-		FMatrix33 M0 = FMatrix33(0, 0, 0);
-		FMatrix33 M1 = FMatrix33(0, 0, 0);
-		if (InvM0 > 0)
-		{
-			M0 = Utilities::ComputeJointFactorMatrix(XC0, InvI0, InvM0);
-		}
-		if (InvM1 > 0)
-		{
-			M1 = Utilities::ComputeJointFactorMatrix(XC1, InvI1, InvM1);
-		}
-		const FMatrix33 MI = (M0 + M1).Inverse();
-		const FVec3 DL = Utilities::Multiply(MI, VC);
-
-		// Apply constraint correction
-		const FVec3 DV0 = InvM0 * DL;
-		const FVec3 DV1 = -InvM1 * DL;
-		const FVec3 DW0 = Utilities::Multiply(InvI0, FVec3::CrossProduct(XC0, DL));
-		const FVec3 DW1 = -Utilities::Multiply(InvI1, FVec3::CrossProduct(XC1, DL));
-
-		ApplyVelocityDelta(Dt, SolverSettings, JointSettings, Stiffness, P0, V0, P1, V1, DV0, DV1);
-		ApplyRotationVelocityDelta(Dt, SolverSettings, JointSettings, Stiffness, Q0, W0, Q1, W1, DW0, DW1);
-	}
-
-
-	
-	void FPBDJointUtilities::ApplyJointTwistConstraint(
-		const FReal Dt,
-		const FPBDJointSolverSettings& SolverSettings,
-		const FPBDJointSettings& JointSettings,
-		const FReal Stiffness,
-		const FRigidTransform3& XL0,
-		const FRigidTransform3& XL1,
-		FVec3& P0,
-		FRotation3& Q0,
-		FVec3& V0,
-		FVec3& W0,
-		FVec3& P1,
-		FRotation3& Q1,
-		FVec3& V1,
-		FVec3& W1,
-		FReal InvM0,
-		const FMatrix33& InvIL0,
-		FReal InvM1,
-		const FMatrix33& InvIL1)
-	{
-		const FVec3 X0 = P0 + Q0 * XL0.GetTranslation();
-		const FVec3 X1 = P1 + Q1 * XL1.GetTranslation();
-		const FRotation3 R0 = Q0 * XL0.GetRotation();
-		const FRotation3 R1 = Q1 * XL1.GetRotation();
-
-		// Calculate the Twist Axis and Angle for each body
-		const FRotation3 R01 = R0.Inverse() * R1;
-		FRotation3 R01Twist, R01Swing;
-		R01.ToSwingTwist(FJointConstants::TwistAxis(), R01Swing, R01Twist);
-		R01Swing = R01Swing.GetNormalized();
-		R01Twist = R01Twist.GetNormalized();
-
-		FVec3 TwistAxis01;
-		FReal TwistAngle;
-		R01Twist.ToAxisAndAngleSafe(TwistAxis01, TwistAngle, FJointConstants::TwistAxis(), SolverSettings.SwingTwistAngleTolerance);
-		if (TwistAngle > PI)
-		{
-			TwistAngle = TwistAngle - (FReal)2 * PI;
-		}
-		if (FVec3::DotProduct(TwistAxis01, FJointConstants::TwistAxis()) < 0)
-		{
-			TwistAxis01 = -TwistAxis01;
-			TwistAngle = -TwistAngle;
-		}
-
-		const FVec3 TwistAxis0 = R0 * TwistAxis01;
-		const FVec3 TwistAxis1 = R1 * TwistAxis01;
-		FReal TwistAngleMax = FLT_MAX;
-		if (JointSettings.AngularMotionTypes[(int32)EJointAngularConstraintIndex::Twist] == EJointMotionType::Limited)
-		{
-			TwistAngleMax = JointSettings.AngularLimits[(int32)EJointAngularConstraintIndex::Twist];
-		}
-		else if (JointSettings.AngularMotionTypes[(int32)EJointAngularConstraintIndex::Twist] == EJointMotionType::Locked)
-		{
-			TwistAngleMax = 0;
-		}
-
-		// Calculate the twist correction to apply to each body
-		FReal DTwistAngle = 0;
-		if (TwistAngle > TwistAngleMax)
-		{
-			DTwistAngle = TwistAngle - TwistAngleMax;
-		}
-		else if (TwistAngle < -TwistAngleMax)
-		{
-			DTwistAngle = TwistAngle + TwistAngleMax;
-		}
-		const FReal DTwistAngle0 = DTwistAngle;
-		const FReal DTwistAngle1 = -DTwistAngle;
-
-		// Apply twist correction
-		ApplyRotationDelta(Dt, SolverSettings, JointSettings, Stiffness, Q0, W0, Q1, W1, InvM0, InvIL0, InvM1, InvIL1, TwistAxis0, DTwistAngle0, TwistAxis1, DTwistAngle1);
-	}
-
-	
-	void FPBDJointUtilities::ApplyJointTwistVelocityConstraint(
-		const FReal Dt,
-		const FPBDJointSolverSettings& SolverSettings,
-		const FPBDJointSettings& JointSettings,
-		const FReal Stiffness,
-		const FRigidTransform3& XL0,
-		const FRigidTransform3& XL1,
-		FVec3& P0,
-		FRotation3& Q0,
-		FVec3& V0,
-		FVec3& W0,
-		FVec3& P1,
-		FRotation3& Q1,
-		FVec3& V1,
-		FVec3& W1,
-		FReal InvM0,
-		const FMatrix33& InvIL0,
-		FReal InvM1,
-		const FMatrix33& InvIL1)
-	{
-		const FVec3 X0 = P0 + Q0 * XL0.GetTranslation();
-		const FVec3 X1 = P1 + Q1 * XL1.GetTranslation();
-		const FRotation3 R0 = Q0 * XL0.GetRotation();
-		const FRotation3 R1 = Q1 * XL1.GetRotation();
-
-		// Calculate the Twist Axis and Angle for each body
-		const FRotation3 R01 = R0.Inverse() * R1;
-		FRotation3 R01Twist, R01Swing;
-		R01.ToSwingTwist(FJointConstants::TwistAxis(), R01Swing, R01Twist);
-		R01Swing = R01Swing.GetNormalized();
-		R01Twist = R01Twist.GetNormalized();
-
-		FVec3 TwistAxis01;
-		FReal TwistAngle;
-		R01Twist.ToAxisAndAngleSafe(TwistAxis01, TwistAngle, FJointConstants::TwistAxis(), SolverSettings.SwingTwistAngleTolerance);
-		if (TwistAngle > PI)
-		{
-			TwistAngle = TwistAngle - (FReal)2 * PI;
-		}
-		if (FVec3::DotProduct(TwistAxis01, FJointConstants::TwistAxis()) < 0)
-		{
-			TwistAxis01 = -TwistAxis01;
-			TwistAngle = -TwistAngle;
-		}
-
-		const FVec3 TwistAxis0 = R0 * TwistAxis01;
-		const FVec3 TwistAxis1 = R1 * TwistAxis01;
-
-		FReal TwistAngleMax = FLT_MAX;
-		if (JointSettings.AngularMotionTypes[(int32)EJointAngularConstraintIndex::Twist] == EJointMotionType::Limited)
-		{
-			TwistAngleMax = JointSettings.AngularLimits[(int32)EJointAngularConstraintIndex::Twist];
-		}
-		else if (JointSettings.AngularMotionTypes[(int32)EJointAngularConstraintIndex::Twist] == EJointMotionType::Locked)
-		{
-			TwistAngleMax = 0;
-		}
-
-		const FReal WC0 = FVec3::DotProduct(W0, TwistAxis0);
-		const FReal WC1 = FVec3::DotProduct(W1, TwistAxis1);
-		FReal DW = (FReal)0;
-		if (TwistAngle > TwistAngleMax)
-		{
-			DW = FMath::Max((FReal)0, WC1 - WC0);
-		}
-		else if (TwistAngle < -TwistAngleMax)
-		{
-			DW = FMath::Min((FReal)0, WC1 - WC0);
-		}
-
-		ApplyRotationVelocityDelta(Dt, SolverSettings, JointSettings, Stiffness, Q0, W0, Q1, W1, InvM0, InvIL0, InvM1, InvIL1, TwistAxis0, TwistAxis1, DW);
-	}
-
-
-	
-	void FPBDJointUtilities::ApplyJointConeConstraint(
-		const FReal Dt,
-		const FPBDJointSolverSettings& SolverSettings,
-		const FPBDJointSettings& JointSettings,
-		const FReal Stiffness,
-		const FRigidTransform3& XL0,
-		const FRigidTransform3& XL1,
-		FVec3& P0,
-		FRotation3& Q0,
-		FVec3& V0,
-		FVec3& W0,
-		FVec3& P1,
-		FRotation3& Q1,
-		FVec3& V1,
-		FVec3& W1,
-		FReal InvM0,
-		const FMatrix33& InvIL0,
-		FReal InvM1,
-		const FMatrix33& InvIL1)
-	{
-		const FVec3 X0 = P0 + Q0 * XL0.GetTranslation();
-		const FVec3 X1 = P1 + Q1 * XL1.GetTranslation();
-		const FRotation3 R0 = Q0 * XL0.GetRotation();
-		const FRotation3 R1 = Q1 * XL1.GetRotation();
-
-		// Calculate Swing axis for each body
-		const FRotation3 R01 = R0.Inverse() * R1;
-		FRotation3 R01Twist, R01Swing;
-		R01.ToSwingTwist(FJointConstants::TwistAxis(), R01Swing, R01Twist);
-		R01Swing = R01Swing.GetNormalized();
-		R01Twist = R01Twist.GetNormalized();
-
-		FVec3 SwingAxis01;
-		FReal SwingAngle;
-		R01Swing.ToAxisAndAngleSafe(SwingAxis01, SwingAngle, FJointConstants::Swing1Axis(), SolverSettings.SwingTwistAngleTolerance);
-		if (SwingAngle > PI)
-		{
-			SwingAngle = SwingAngle - (FReal)2 * PI;
-		}
-
-		const FVec3 SwingAxis0 = R0 * SwingAxis01;
-		const FVec3 SwingAxis1 = SwingAxis0;
-
-		// Calculate swing limit for the current swing axis
-		FReal SwingAngleMax = FLT_MAX;
-		const FReal Swing1Limit = JointSettings.AngularLimits[(int32)EJointAngularConstraintIndex::Swing1];
-		const FReal Swing2Limit = JointSettings.AngularLimits[(int32)EJointAngularConstraintIndex::Swing2];
-
-		// Circular swing limit
-		SwingAngleMax = Swing1Limit;
-
-		// Elliptical swing limit
-		if (!FMath::IsNearlyEqual(Swing1Limit, Swing2Limit, KINDA_SMALL_NUMBER))
-		{
-			// Map swing axis to ellipse and calculate limit for this swing axis
-			const FReal DotSwing1 = FMath::Abs(FVec3::DotProduct(SwingAxis01, FJointConstants::Swing1Axis()));
-			const FReal DotSwing2 = FMath::Abs(FVec3::DotProduct(SwingAxis01, FJointConstants::Swing2Axis()));
-			SwingAngleMax = FMath::Sqrt(Swing1Limit * DotSwing2 * Swing1Limit * DotSwing2 + Swing2Limit * DotSwing1 * Swing2Limit * DotSwing1);
-		}
-
-		// Calculate swing error we need to correct
-		FReal DSwingAngle = 0;
-		if (SwingAngle > SwingAngleMax)
-		{
-			DSwingAngle = SwingAngle - SwingAngleMax;
-		}
-		else if (SwingAngle < -SwingAngleMax)
-		{
-			DSwingAngle = SwingAngle + SwingAngleMax;
-		}
-		FReal DSwingAngle0 = DSwingAngle;
-		FReal DSwingAngle1 = -DSwingAngle;
-
-		// Apply swing correction
-		ApplyRotationDelta(Dt, SolverSettings, JointSettings, Stiffness, Q0, W0, Q1, W1, InvM0, InvIL0, InvM1, InvIL1, SwingAxis0, DSwingAngle0, SwingAxis1, DSwingAngle1);
-	}
-
-	
-	void FPBDJointUtilities::ApplyJointConeVelocityConstraint(
-		const FReal Dt,
-		const FPBDJointSolverSettings& SolverSettings,
-		const FPBDJointSettings& JointSettings,
-		const FReal Stiffness,
-		const FRigidTransform3& XL0,
-		const FRigidTransform3& XL1,
-		FVec3& P0,
-		FRotation3& Q0,
-		FVec3& V0,
-		FVec3& W0,
-		FVec3& P1,
-		FRotation3& Q1,
-		FVec3& V1,
-		FVec3& W1,
-		FReal InvM0,
-		const FMatrix33& InvIL0,
-		FReal InvM1,
-		const FMatrix33& InvIL1)
-	{
-		const FVec3 X0 = P0 + Q0 * XL0.GetTranslation();
-		const FVec3 X1 = P1 + Q1 * XL1.GetTranslation();
-		const FRotation3 R0 = Q0 * XL0.GetRotation();
-		const FRotation3 R1 = Q1 * XL1.GetRotation();
-
-		// Calculate Swing axis for each body
-		const FRotation3 R01 = R0.Inverse() * R1;
-		FRotation3 R01Twist, R01Swing;
-		R01.ToSwingTwist(FJointConstants::TwistAxis(), R01Swing, R01Twist);
-		R01Swing = R01Swing.GetNormalized();
-		R01Twist = R01Twist.GetNormalized();
-
-		FVec3 SwingAxis01;
-		FReal SwingAngle;
-		R01Swing.ToAxisAndAngleSafe(SwingAxis01, SwingAngle, FJointConstants::Swing1Axis(), SolverSettings.SwingTwistAngleTolerance);
-		if (SwingAngle > PI)
-		{
-			SwingAngle = SwingAngle - (FReal)2 * PI;
-		}
-
-		const FVec3 SwingAxis0 = R0 * SwingAxis01;
-		const FVec3 SwingAxis1 = SwingAxis0;
-
-		// Calculate swing limit for the current swing axis
-		FReal SwingAngleMax = FLT_MAX;
-		const FReal Swing1Limit = JointSettings.AngularLimits[(int32)EJointAngularConstraintIndex::Swing1];
-		const FReal Swing2Limit = JointSettings.AngularLimits[(int32)EJointAngularConstraintIndex::Swing2];
-
-		// Circular swing limit
-		SwingAngleMax = Swing1Limit;
-
-		// Elliptical swing limit
-		if (!FMath::IsNearlyEqual(Swing1Limit, Swing2Limit, KINDA_SMALL_NUMBER))
-		{
-			// Map swing axis to ellipse and calculate limit for this swing axis
-			const FReal DotSwing1 = FMath::Abs(FVec3::DotProduct(SwingAxis01, FJointConstants::Swing1Axis()));
-			const FReal DotSwing2 = FMath::Abs(FVec3::DotProduct(SwingAxis01, FJointConstants::Swing2Axis()));
-			SwingAngleMax = FMath::Sqrt(Swing1Limit * DotSwing2 * Swing1Limit * DotSwing2 + Swing2Limit * DotSwing1 * Swing2Limit * DotSwing1);
-		}
-
-		// Only clamp veloicity if we are outside the limits and moving to increase the error
-		const FReal WC0 = FVec3::DotProduct(W0, SwingAxis0);
-		const FReal WC1 = FVec3::DotProduct(W1, SwingAxis1);
-		FReal DW = 0;
-		if (SwingAngle > SwingAngleMax)
-		{
-			DW = FMath::Max((FReal)0, WC1 - WC0);
-		}
-		else if (SwingAngle < -SwingAngleMax)
-		{
-			DW = FMath::Min((FReal)0, WC1 - WC0);
-		}
-
-		ApplyRotationVelocityDelta(Dt, SolverSettings, JointSettings, Stiffness, Q0, W0, Q1, W1, InvM0, InvIL0, InvM1, InvIL1, SwingAxis0, SwingAxis1, DW);
-	}
-
-
-	
-	void FPBDJointUtilities::ApplyJointSwingConstraint(
-		const FReal Dt,
-		const FPBDJointSolverSettings& SolverSettings,
-		const FPBDJointSettings& JointSettings,
-		const FReal Stiffness,
-		const FRigidTransform3& XL0,
-		const FRigidTransform3& XL1,
-		const EJointAngularConstraintIndex SwingConstraintIndex,
-		const EJointAngularAxisIndex SwingAxisIndex,
-		FVec3& P0,
-		FRotation3& Q0,
-		FVec3& V0,
-		FVec3& W0,
-		FVec3& P1,
-		FRotation3& Q1,
-		FVec3& V1,
-		FVec3& W1,
-		FReal InvM0,
-		const FMatrix33& InvIL0,
-		FReal InvM1,
-		const FMatrix33& InvIL1)
-	{
-		const FVec3 X0 = P0 + Q0 * XL0.GetTranslation();
-		const FVec3 X1 = P1 + Q1 * XL1.GetTranslation();
-		const FRotation3 R0 = Q0 * XL0.GetRotation();
-		const FRotation3 R1 = Q1 * XL1.GetRotation();
-
-		// Calculate the swing axis for each body
-		const FRotation3 R01 = R0.Inverse() * R1;
-		FRotation3 R01Twist, R01Swing;
-		R01.ToSwingTwist(FJointConstants::TwistAxis(), R01Swing, R01Twist);
-		R01Swing = R01Swing.GetNormalized();
-		R01Twist = R01Twist.GetNormalized();
-
-		FVec3 TwistAxis01;
-		FReal TwistAngle;
-		R01Twist.ToAxisAndAngleSafe(TwistAxis01, TwistAngle, FJointConstants::TwistAxis(), SolverSettings.SwingTwistAngleTolerance);
-		if (TwistAngle > PI)
-		{
-			TwistAngle = TwistAngle - (FReal)2 * PI;
-		}
-		if (FVec3::DotProduct(TwistAxis01, FJointConstants::TwistAxis()) < 0)
-		{
-			TwistAxis01 = -TwistAxis01;
-			TwistAngle = -TwistAngle;
-		}
-		const FVec3 TwistAxis = R0 * TwistAxis01;
-
-		const FRotation3 R1NoTwist = R1 * R01Twist.Inverse();
-		const FMatrix33 Axes0 = R0.ToMatrix();
-		const FMatrix33 Axes1 = R1NoTwist.ToMatrix();
-		FVec3 SwingCross = FVec3::CrossProduct(Axes0.GetAxis((int32)SwingAxisIndex), Axes1.GetAxis((int32)SwingAxisIndex));
-		SwingCross = SwingCross - FVec3::DotProduct(TwistAxis, SwingCross) * TwistAxis;
-		const FReal SwingCrossLen = SwingCross.Size();
-		if (SwingCrossLen > KINDA_SMALL_NUMBER)
-		{
-			const FVec3 SwingAxis = SwingCross / SwingCrossLen;
-			FVec3 SwingAxis0 = SwingAxis;
-			FVec3 SwingAxis1 = SwingAxis;
-
-			FReal SwingAngle = FMath::Asin(FMath::Clamp(SwingCrossLen, (FReal)0, (FReal)1));
-			const FReal SwingDot = FVec3::DotProduct(Axes0.GetAxis((int32)SwingAxisIndex), Axes1.GetAxis((int32)SwingAxisIndex));
-			if (SwingDot < (FReal)0)
-			{
-				SwingAngle = (FReal)PI - SwingAngle;
-			}
-
-			FReal SwingAngleMax = FLT_MAX;
-			if (JointSettings.AngularMotionTypes[(int32)SwingConstraintIndex] == EJointMotionType::Limited)
-			{
-				FReal SwingLimit = JointSettings.AngularLimits[(int32)SwingConstraintIndex];
-				SwingAngleMax = SwingLimit;
-			}
-			else if (JointSettings.AngularMotionTypes[(int32)SwingConstraintIndex] == EJointMotionType::Locked)
-			{
-				SwingAngleMax = 0;
-			}
-
-			// Calculate swing error we need to correct
-			FReal DSwingAngle = 0;
-			if (SwingAngle > SwingAngleMax)
-			{
-				DSwingAngle = SwingAngle - SwingAngleMax;
-			}
-			else if (SwingAngle < -SwingAngleMax)
-			{
-				DSwingAngle = SwingAngle + SwingAngleMax;
-			}
-			FReal DSwingAngle0 = DSwingAngle;
-			FReal DSwingAngle1 = -DSwingAngle;
-
-			// Apply swing correction
-			ApplyRotationDelta(Dt, SolverSettings, JointSettings, Stiffness, Q0, W0, Q1, W1, InvM0, InvIL0, InvM1, InvIL1, SwingAxis0, DSwingAngle0, SwingAxis1, DSwingAngle1);
-		}
-	}
-
-	
-	void FPBDJointUtilities::ApplyJointSwingVelocityConstraint(
-		const FReal Dt,
-		const FPBDJointSolverSettings& SolverSettings,
-		const FPBDJointSettings& JointSettings,
-		const FReal Stiffness,
-		const FRigidTransform3& XL0,
-		const FRigidTransform3& XL1,
-		const EJointAngularConstraintIndex SwingConstraintIndex,
-		const EJointAngularAxisIndex SwingAxisIndex,
-		FVec3& P0,
-		FRotation3& Q0,
-		FVec3& V0,
-		FVec3& W0,
-		FVec3& P1,
-		FRotation3& Q1,
-		FVec3& V1,
-		FVec3& W1,
-		FReal InvM0,
-		const FMatrix33& InvIL0,
-		FReal InvM1,
-		const FMatrix33& InvIL1)
-	{
-		const FVec3 X0 = P0 + Q0 * XL0.GetTranslation();
-		const FVec3 X1 = P1 + Q1 * XL1.GetTranslation();
-		const FRotation3 R0 = Q0 * XL0.GetRotation();
-		const FRotation3 R1 = Q1 * XL1.GetRotation();
-
-		// Calculate the swing axis for each body
-		const FRotation3 R01 = R0.Inverse() * R1;
-		FRotation3 R01Twist, R01Swing;
-		R01.ToSwingTwist(FJointConstants::TwistAxis(), R01Swing, R01Twist);
-		R01Swing = R01Swing.GetNormalized();
-		R01Twist = R01Twist.GetNormalized();
-
-		FVec3 TwistAxis01;
-		FReal TwistAngle;
-		R01Twist.ToAxisAndAngleSafe(TwistAxis01, TwistAngle, FJointConstants::TwistAxis(), SolverSettings.SwingTwistAngleTolerance);
-		if (TwistAngle > PI)
-		{
-			TwistAngle = TwistAngle - (FReal)2 * PI;
-		}
-		if (FVec3::DotProduct(TwistAxis01, FJointConstants::TwistAxis()) < 0)
-		{
-			TwistAxis01 = -TwistAxis01;
-			TwistAngle = -TwistAngle;
-		}
-		const FVec3 TwistAxis = R0 * TwistAxis01;
-
-		const FRotation3 R1NoTwist = R1 * R01Twist.Inverse();
-		const FMatrix33 Axes0 = R0.ToMatrix();
-		const FMatrix33 Axes1 = R1NoTwist.ToMatrix();
-		FVec3 SwingCross = FVec3::CrossProduct(Axes0.GetAxis((int32)SwingAxisIndex), Axes1.GetAxis((int32)SwingAxisIndex));
-		SwingCross = SwingCross - FVec3::DotProduct(TwistAxis, SwingCross) * TwistAxis;
-		const FReal SwingCrossLen = SwingCross.Size();
-		if (SwingCrossLen > KINDA_SMALL_NUMBER)
-		{
-			const FVec3 SwingAxis = SwingCross / SwingCrossLen;
-			FVec3 SwingAxis0 = SwingAxis;
-			FVec3 SwingAxis1 = SwingAxis;
-
-			FReal SwingAngle = FMath::Asin(FMath::Clamp(SwingCrossLen, (FReal)0, (FReal)1));
-			const FReal SwingDot = FVec3::DotProduct(Axes0.GetAxis((int32)SwingAxisIndex), Axes1.GetAxis((int32)SwingAxisIndex));
-			if (SwingDot < (FReal)0)
-			{
-				SwingAngle = (FReal)PI - SwingAngle;
-			}
-
-			FReal SwingAngleMax = FLT_MAX;
-			if (JointSettings.AngularMotionTypes[(int32)SwingConstraintIndex] == EJointMotionType::Limited)
-			{
-				FReal Swing1Limit = JointSettings.AngularLimits[(int32)SwingConstraintIndex];
-				SwingAngleMax = Swing1Limit;
-			}
-			else if (JointSettings.AngularMotionTypes[(int32)SwingConstraintIndex] == EJointMotionType::Locked)
-			{
-				SwingAngleMax = 0;
-			}
-
-			// Only clamp veloicity if we are outside the limits and moving to increase the error
-			const FReal WC0 = FVec3::DotProduct(W0, SwingAxis0);
-			const FReal WC1 = FVec3::DotProduct(W1, SwingAxis1);
-			FReal DW = 0;
-			if (SwingAngle > SwingAngleMax)
-			{
-				DW = FMath::Max((FReal)0, WC1 - WC0);
-			}
-			else if (SwingAngle < -SwingAngleMax)
-			{
-				DW = FMath::Min((FReal)0, WC1 - WC0);
-			}
-
-			ApplyRotationVelocityDelta(Dt, SolverSettings, JointSettings, Stiffness, Q0, W0, Q1, W1, InvM0, InvIL0, InvM1, InvIL1, SwingAxis0, SwingAxis1, DW);
-		}
-	}
-
-	
-	void FPBDJointUtilities::ApplyJointTwistDrive(
-		const FReal Dt,
-		const FPBDJointSolverSettings& SolverSettings,
-		const FPBDJointSettings& JointSettings,
-		const FRigidTransform3& XL0,
-		const FRigidTransform3& XL1,
-		FVec3& P0,
-		FRotation3& Q0,
-		FVec3& V0,
-		FVec3& W0,
-		FVec3& P1,
-		FRotation3& Q1,
-		FVec3& V1,
-		FVec3& W1,
-		FReal InvM0,
-		const FMatrix33& InvIL0,
-		FReal InvM1,
-		const FMatrix33& InvIL1)
-	{
-		const FVec3 X0 = P0 + Q0 * XL0.GetTranslation();
-		const FVec3 X1 = P1 + Q1 * XL1.GetTranslation();
-		const FRotation3 R0 = Q0 * XL0.GetRotation();
-		const FRotation3 R1 = Q1 * XL1.GetRotation();
-
-		const FRotation3 R01 = R0.Inverse() * R1;
-		FRotation3 R01Twist, R01Swing;
-		R01.ToSwingTwist(FJointConstants::TwistAxis(), R01Swing, R01Twist);
-		R01Swing = R01Swing.GetNormalized();
-		R01Twist = R01Twist.GetNormalized();
-
-		FVec3 TwistAxis01;
-		FReal TwistAngle;
-		R01Twist.ToAxisAndAngleSafe(TwistAxis01, TwistAngle, FJointConstants::TwistAxis(), SolverSettings.SwingTwistAngleTolerance);
-		if (TwistAngle > PI)
-		{
-			TwistAngle = TwistAngle - (FReal)2 * PI;
-		}
-		if (FVec3::DotProduct(TwistAxis01, FJointConstants::TwistAxis()) < 0)
-		{
-			TwistAxis01 = -TwistAxis01;
-			TwistAngle = -TwistAngle;
-		}
-
-		const FVec3 TwistAxis0 = R0 * TwistAxis01;
-		const FVec3 TwistAxis1 = R1 * TwistAxis01;
-		const FReal TwistAngleTarget = JointSettings.AngularDriveTargetAngles[(int32)EJointAngularConstraintIndex::Twist];
-		const FReal DriveStiffness = GetAngularDriveStiffness(SolverSettings, JointSettings);
-		const FReal DTwistAngle = TwistAngle - TwistAngleTarget;
-		const FReal DTwistAngle0 = DTwistAngle;
-		const FReal DTwistAngle1 = -DTwistAngle;
-
-		ApplyRotationDelta(Dt, SolverSettings, JointSettings, DriveStiffness, Q0, W0, Q1, W1, InvM0, InvIL0, InvM1, InvIL1, TwistAxis0, DTwistAngle0, TwistAxis1, DTwistAngle1);
-	}
-
-	
-	void FPBDJointUtilities::ApplyJointConeDrive(
-		const FReal Dt,
-		const FPBDJointSolverSettings& SolverSettings,
-		const FPBDJointSettings& JointSettings,
-		const FRigidTransform3& XL0,
-		const FRigidTransform3& XL1,
-		FVec3& P0,
-		FRotation3& Q0,
-		FVec3& V0,
-		FVec3& W0,
-		FVec3& P1,
-		FRotation3& Q1,
-		FVec3& V1,
-		FVec3& W1,
-		FReal InvM0,
-		const FMatrix33& InvIL0,
-		FReal InvM1,
-		const FMatrix33& InvIL1)
-	{
-		const FVec3 X0 = P0 + Q0 * XL0.GetTranslation();
-		const FVec3 X1 = P1 + Q1 * XL1.GetTranslation();
-		const FRotation3 R0 = Q0 * XL0.GetRotation();
-		const FRotation3 R1 = Q1 * XL1.GetRotation();
-
-		// Calculate Swing axis for each body
-		const FRotation3 R01 = R0.Inverse() * R1;
-		FRotation3 R01Twist, R01Swing;
-		R01.ToSwingTwist(FJointConstants::TwistAxis(), R01Swing, R01Twist);
-		R01Swing = R01Swing.GetNormalized();
-		R01Twist = R01Twist.GetNormalized();
-
-		FVec3 SwingAxis01;
-		FReal SwingAngle;
-		R01Swing.ToAxisAndAngleSafe(SwingAxis01, SwingAngle, FJointConstants::Swing1Axis(), SolverSettings.SwingTwistAngleTolerance);
-		if (SwingAngle > PI)
-		{
-			SwingAngle = SwingAngle - (FReal)2 * PI;
-		}
-
-		const FVec3 SwingAxis0 = R0 * SwingAxis01;
-		const FVec3 SwingAxis1 = SwingAxis0;
-
-		// Circular swing target (max of Swing1, Swing2 targets)
-		FReal Swing1Target = JointSettings.AngularDriveTargetAngles[(int32)EJointAngularConstraintIndex::Swing1];
-		FReal Swing2Target = JointSettings.AngularDriveTargetAngles[(int32)EJointAngularConstraintIndex::Swing2];
-		FReal SwingAngleTarget = FMath::Max(Swing1Target, Swing2Target);
-
-		const FReal DriveStiffness = GetAngularDriveStiffness(SolverSettings, JointSettings);
-		FReal DSwingAngle = SwingAngle - SwingAngleTarget;
-		FReal DSwingAngle0 = DSwingAngle;
-		FReal DSwingAngle1 = -DSwingAngle;
-
-		// Apply swing correction
-		ApplyRotationDelta(Dt, SolverSettings, JointSettings, DriveStiffness, Q0, W0, Q1, W1, InvM0, InvIL0, InvM1, InvIL1, SwingAxis0, DSwingAngle0, SwingAxis1, DSwingAngle1);
-	}
-
-	
-	void FPBDJointUtilities::ApplyJointSLerpDrive(
-		const FReal Dt,
-		const FPBDJointSolverSettings& SolverSettings,
-		const FPBDJointSettings& JointSettings,
-		const FRigidTransform3& XL0,
-		const FRigidTransform3& XL1,
-		FVec3& P0,
-		FRotation3& Q0,
-		FVec3& V0,
-		FVec3& W0,
-		FVec3& P1,
-		FRotation3& Q1,
-		FVec3& V1,
-		FVec3& W1,
-		FReal InvM0,
-		const FMatrix33& InvIL0,
-		FReal InvM1,
-		const FMatrix33& InvIL1)
-	{
-		const FVec3 DX0 = Q0 * XL0.GetTranslation();
-		const FVec3 DX1 = Q1 * XL1.GetTranslation();
-		const FVec3 X0 = P0 + DX0;
-		const FVec3 X1 = P1 + DX1;
-		const FRotation3 R0 = Q0 * XL0.GetRotation();
-		const FRotation3 R1 = Q1 * XL1.GetRotation();
-
-		// Calculate the rotation we need to apply to resolve the rotation delta
-		const FRotation3 TargetR1 = R0 * JointSettings.AngularDriveTarget;
-		const FRotation3 DR1 = TargetR1 * R1.Inverse();
-		const FRotation3 TargetQ0 = DR1.Inverse() * Q0;
-		const FRotation3 TargetQ1 = DR1 * Q1;
-
-		const FReal DriveStiffness = GetAngularDriveStiffness(SolverSettings, JointSettings);
-
-		FVec3 SLerpAxis;
-		FReal SLerpAngle;
-		if (DR1.ToAxisAndAngleSafe(SLerpAxis, SLerpAngle, FVec3(1, 0, 0)))
-		{
-			const FMatrix33 InvI0 = Utilities::ComputeWorldSpaceInertia(Q0, InvIL0);
-			const FMatrix33 InvI1 = Utilities::ComputeWorldSpaceInertia(Q1, InvIL1);
-			const FReal I0 = FVec3::DotProduct(SLerpAxis, Utilities::Multiply(InvI0, SLerpAxis));
-			const FReal I1 = FVec3::DotProduct(SLerpAxis, Utilities::Multiply(InvI1, SLerpAxis));
-			const FReal F0 = DriveStiffness * I0 / (I0 + I1);
-			const FReal F1 = DriveStiffness * I1 / (I0 + I1);
-
-			// Apply the rotation delta about the connector
-			Q0 = FRotation3::Slerp(Q0, TargetQ0, F0);
-			Q1 = FRotation3::Slerp(Q1, TargetQ1, F1);
-			Q1.EnforceShortestArcWith(Q0);
-
-			// @todo(ccaulfield): this does not take into account the fact that some linear dofs may be inactive
-			const FVec3 X0_2 = P0 + Q0 * XL0.GetTranslation();
-			const FVec3 X1_2 = P1 + Q1 * XL1.GetTranslation();
-			const FVec3 Delta = (X1_2 - X0_2) - (X1 - X0);
-			P0 += (InvM0 / (InvM0 + InvM1)) * Delta;
-			P1 -= (InvM1 / (InvM0 + InvM1)) * Delta;
-		}
-	}
-
-	
-	void FPBDJointUtilities::ApplyJointPositionProjection(
-		const FReal Dt,
-		const FPBDJointSolverSettings& SolverSettings,
-		const FPBDJointSettings& JointSettings,
-		const FReal Stiffness,
-		const FRigidTransform3& XL0,
-		const FRigidTransform3& XL1,
-		FVec3& P0,
-		FRotation3& Q0,
-		FVec3& P1,
-		FRotation3& Q1,
-		FReal InvM0,
-		const FMatrix33& InvIL0,
-		FReal InvM1,
-		const FMatrix33& InvIL1)
-	{
-		const FVec3 X0 = P0 + Q0 * XL0.GetTranslation();
-		const FVec3 X1 = P1 + Q1 * XL1.GetTranslation();
-		const FRotation3 R0 = Q0 * XL0.GetRotation();
-		const FVec3 CX = GetLimitedPositionError(JointSettings, R0, X1 - X0);
-
-		const FVec3 DP0 = CX * (Stiffness * InvM0 / (InvM0 + InvM1));
-		const FVec3 DP1 = -CX * (Stiffness* InvM1 / (InvM0 + InvM1));
-		P0 = P0 + DP0;
-		P1 = P1 + DP1;
-	}
-
-	
-	void FPBDJointUtilities::ApplyJointTwistProjection(
-		const FReal Dt,
-		const FPBDJointSolverSettings& SolverSettings,
-		const FPBDJointSettings& JointSettings,
-		const FReal Stiffness,
-		const FRigidTransform3& XL0,
-		const FRigidTransform3& XL1,
-		FVec3& P0,
-		FRotation3& Q0,
-		FVec3& P1,
-		FRotation3& Q1,
-		FReal InvM0,
-		const FMatrix33& InvIL0,
-		FReal InvM1,
-		const FMatrix33& InvIL1)
-	{
-		FVec3 V0(0, 0, 0);
-		FVec3 W0(0, 0, 0);
-		FVec3 V1(0, 0, 0);
-		FVec3 W1(0, 0, 0);
-		ApplyJointTwistConstraint(Dt, SolverSettings, JointSettings, Stiffness, XL0, XL1, P0, Q0, V0, W0, P1, Q1, V1, W1, InvM0, InvIL0, InvM1, InvIL1);
-	}
-
-	
-	void FPBDJointUtilities::ApplyJointConeProjection(
-		const FReal Dt,
-		const FPBDJointSolverSettings& SolverSettings,
-		const FPBDJointSettings& JointSettings,
-		const FReal Stiffness,
-		const FRigidTransform3& XL0,
-		const FRigidTransform3& XL1,
-		FVec3& P0,
-		FRotation3& Q0,
-		FVec3& P1,
-		FRotation3& Q1,
-		FReal InvM0,
-		const FMatrix33& InvIL0,
-		FReal InvM1,
-		const FMatrix33& InvIL1)
-	{
-		FVec3 V0(0, 0, 0);
-		FVec3 W0(0, 0, 0);
-		FVec3 V1(0, 0, 0);
-		FVec3 W1(0, 0, 0);
-		ApplyJointConeConstraint(Dt, SolverSettings, JointSettings, Stiffness, XL0, XL1, P0, Q0, V0, W0, P1, Q1, V1, W1, InvM0, InvIL0, InvM1, InvIL1);
-	}
-
-	
-	void FPBDJointUtilities::ApplyJointSwingProjection(
-		const FReal Dt,
-		const FPBDJointSolverSettings& SolverSettings,
-		const FPBDJointSettings& JointSettings,
-		const FReal Stiffness,
-		const FRigidTransform3& XL0,
-		const FRigidTransform3& XL1,
-		const EJointAngularConstraintIndex SwingConstraintIndex,
-		const EJointAngularAxisIndex SwingAxisIndex,
-		FVec3& P0,
-		FRotation3& Q0,
-		FVec3& P1,
-		FRotation3& Q1,
-		FReal InvM0,
-		const FMatrix33& InvIL0,
-		FReal InvM1,
-		const FMatrix33& InvIL1)
-	{
-		FVec3 V0(0, 0, 0);
-		FVec3 W0(0, 0, 0);
-		FVec3 V1(0, 0, 0);
-		FVec3 W1(0, 0, 0);
-		ApplyJointSwingConstraint(Dt, SolverSettings, JointSettings, Stiffness, XL0, XL1, SwingConstraintIndex, SwingAxisIndex, P0, Q0, V0, W0, P1, Q1, V1, W1, InvM0, InvIL0, InvM1, InvIL1);
-	}
-
-
-
 }
diff --git a/Engine/Source/Runtime/Experimental/Chaos/Private/Chaos/PBDJointConstraints.cpp b/Engine/Source/Runtime/Experimental/Chaos/Private/Chaos/PBDJointConstraints.cpp
--- a/Engine/Source/Runtime/Experimental/Chaos/Private/Chaos/PBDJointConstraints.cpp
+++ b/Engine/Source/Runtime/Experimental/Chaos/Private/Chaos/PBDJointConstraints.cpp
@@ -100,47 +100,15 @@ namespace Chaos
 		, LinearDriveForceMode(EJointForceMode::Acceleration)
 		, LinearDriveStiffness(0)
 		, LinearDriveDamping(0)
-		, AngularDriveTarget(FRotation3::FromIdentity())
+		, AngularDrivePositionTarget(FRotation3::FromIdentity())
 		, AngularDriveTargetAngles(FVec3(0, 0, 0))
-		, bAngularSLerpDriveEnabled(false)
-		, bAngularTwistDriveEnabled(false)
-		, bAngularSwingDriveEnabled(false)
-		, AngularDriveForceMode(EJointForceMode::Acceleration)
-		, AngularDriveStiffness(0)
-		, AngularDriveDamping(0)
-	{
-	}
-
-	
-	FPBDJointSettings::FPBDJointSettings(const TVector<EJointMotionType, 3>& InLinearMotionTypes, const TVector<EJointMotionType, 3>& InAngularMotionTypes)
-		: Stiffness(1)
-		, LinearProjection(0)
-		, AngularProjection(0)
-		, ParentInvMassScale(1)
-		, LinearMotionTypes(InLinearMotionTypes)
-		, LinearLimit(FLT_MAX)
-		, AngularMotionTypes({ EJointMotionType::Free, EJointMotionType::Free, EJointMotionType::Free })
-		, AngularLimits(FVec3(FLT_MAX, FLT_MAX, FLT_MAX))
-		, bSoftLinearLimitsEnabled(false)
-		, bSoftTwistLimitsEnabled(false)
-		, bSoftSwingLimitsEnabled(false)
-		, SoftForceMode(EJointForceMode::Acceleration)
-		, SoftLinearStiffness(0)
-		, SoftLinearDamping(0)
-		, SoftTwistStiffness(0)
-		, SoftTwistDamping(0)
-		, SoftSwingStiffness(0)
-		, SoftSwingDamping(0)
-		, LinearDriveTarget(FVec3(0))
-		, bLinearDriveEnabled(TVector<bool, 3>(false, false, false))
-		, LinearDriveForceMode(EJointForceMode::Acceleration)
-		, LinearDriveStiffness(0)
-		, LinearDriveDamping(0)
-		, AngularDriveTarget(FRotation3::FromIdentity())
-		, AngularDriveTargetAngles(FVec3(0, 0, 0))
-		, bAngularSLerpDriveEnabled(false)
-		, bAngularTwistDriveEnabled(false)
-		, bAngularSwingDriveEnabled(false)
+		, AngularDriveVelocityTarget(FVec3(0, 0, 0))
+		, bAngularSLerpPositionDriveEnabled(false)
+		, bAngularSLerpVelocityDriveEnabled(false)
+		, bAngularTwistPositionDriveEnabled(false)
+		, bAngularTwistVelocityDriveEnabled(false)
+		, bAngularSwingPositionDriveEnabled(false)
+		, bAngularSwingVelocityDriveEnabled(false)
 		, AngularDriveForceMode(EJointForceMode::Acceleration)
 		, AngularDriveStiffness(0)
 		, AngularDriveDamping(0)
@@ -168,6 +136,21 @@ namespace Chaos
 		{
 			AngularLimits[(int32)EJointAngularConstraintIndex::Swing2] = 0;
 		}
+
+		// Don't use soft constraints if the stiffness is very high
+		if (bSoftLinearLimitsEnabled && SoftLinearStiffness > 1000)
+		{
+			bSoftLinearLimitsEnabled = false;
+		}
+		if (bSoftTwistLimitsEnabled && SoftTwistStiffness > 1000)
+		{
+			bSoftTwistLimitsEnabled = false;
+		}
+		if (bSoftSwingLimitsEnabled && SoftSwingStiffness > 1000)
+		{
+			bSoftSwingLimitsEnabled = false;
+		}
+
 	}
 
 	
diff --git a/Engine/Source/Runtime/Experimental/Chaos/Public/Chaos/PBDJointConstraintTypes.h b/Engine/Source/Runtime/Experimental/Chaos/Public/Chaos/PBDJointConstraintTypes.h
--- a/Engine/Source/Runtime/Experimental/Chaos/Public/Chaos/PBDJointConstraintTypes.h
+++ b/Engine/Source/Runtime/Experimental/Chaos/Public/Chaos/PBDJointConstraintTypes.h
@@ -66,7 +66,6 @@ namespace Chaos
 	{
 	public:
 		FPBDJointSettings();
-		FPBDJointSettings(const TVector<EJointMotionType, 3>& InLinearMotionTypes, const TVector<EJointMotionType, 3>& InAngularMotionTypes);
 
 		void Sanitize();
 
@@ -99,12 +98,16 @@ namespace Chaos
 		FReal LinearDriveDamping;
 
 		// @todo(ccaulfield): remove one of these
-		FRotation3 AngularDriveTarget;
+		FRotation3 AngularDrivePositionTarget;
 		FVec3 AngularDriveTargetAngles;
-
-		bool bAngularSLerpDriveEnabled;
-		bool bAngularTwistDriveEnabled;
-		bool bAngularSwingDriveEnabled;
+		FVec3 AngularDriveVelocityTarget;
+
+		bool bAngularSLerpPositionDriveEnabled;
+		bool bAngularSLerpVelocityDriveEnabled;
+		bool bAngularTwistPositionDriveEnabled;
+		bool bAngularTwistVelocityDriveEnabled;
+		bool bAngularSwingPositionDriveEnabled;
+		bool bAngularSwingVelocityDriveEnabled;
 		EJointForceMode AngularDriveForceMode;
 		FReal AngularDriveStiffness;
 		FReal AngularDriveDamping;
diff --git a/Engine/Source/Runtime/Experimental/Chaos/Public/Chaos/PBDJointConstraintUtilities.h b/Engine/Source/Runtime/Experimental/Chaos/Public/Chaos/PBDJointConstraintUtilities.h
--- a/Engine/Source/Runtime/Experimental/Chaos/Public/Chaos/PBDJointConstraintUtilities.h
+++ b/Engine/Source/Runtime/Experimental/Chaos/Public/Chaos/PBDJointConstraintUtilities.h
@@ -100,355 +100,49 @@ namespace Chaos
 			const FPBDJointSolverSettings& SolverSettings,
 			const FPBDJointSettings& JointSettings);
 
-		static FReal GetAngularDriveStiffness(
+		static FReal GetAngularTwistDriveStiffness(
 			const FPBDJointSolverSettings& SolverSettings,
 			const FPBDJointSettings& JointSettings);
 
-		static FReal GetAngularDriveDamping(
+		static FReal GetAngularTwistDriveDamping(
 			const FPBDJointSolverSettings& SolverSettings,
 			const FPBDJointSettings& JointSettings);
 
-		static FReal GetLinearProjection(
+		static FReal GetAngularSwingDriveStiffness(
 			const FPBDJointSolverSettings& SolverSettings,
 			const FPBDJointSettings& JointSettings);
 
-		static FReal GetAngularProjection(
+		static FReal GetAngularSwingDriveDamping(
 			const FPBDJointSolverSettings& SolverSettings,
 			const FPBDJointSettings& JointSettings);
 
-		static bool GetSoftAccelerationMode(
+		static FReal GetAngularSLerpDriveStiffness(
 			const FPBDJointSolverSettings& SolverSettings,
 			const FPBDJointSettings& JointSettings);
 
-		static bool GetDriveAccelerationMode(
+		static FReal GetAngularSLerpDriveDamping(
 			const FPBDJointSolverSettings& SolverSettings,
 			const FPBDJointSettings& JointSettings);
 
-		static FVec3 GetSphereLimitedPositionError(const FVec3& CX, const FReal Radius);
-		static FVec3 GetSphereLimitedVelocityError(const FVec3& CX, const FReal Radius, const FVec3& CV);
-		static FVec3 GetCylinderLimitedPositionError(const FVec3& InCX, const FVec3& Axis, const FReal Limit, const EJointMotionType AxisMotion);
-		static FVec3 GetCylinderLimitedVelocityError(const FVec3& InCX, const FVec3& Axis, const FReal Limit, const EJointMotionType AxisMotion, const FVec3& CV);
-		static FVec3 GetLineLimitedPositionError(const FVec3& CX, const FVec3& Axis, const FReal Limit, const EJointMotionType AxisMotion);
-		static FVec3 GetLineLimitedVelocityError(const FVec3& CX, const FVec3& Axis, const FReal Limit, const EJointMotionType AxisMotion, const FVec3& CV);
-		static FVec3 GetLimitedPositionError(const FPBDJointSettings& JointSettings, const FRotation3& R0, const FVec3& InCX);
-		static FVec3 GetLimitedVelocityError(const FPBDJointSettings& JointSettings, const FRotation3& R0, const FVec3& InCX, const FVec3& InCV);
-
-		static CHAOS_API void CalculateSwingConstraintSpace(
-			const FPBDJointSolverSettings& SolverSettings,
-			const FPBDJointSettings& JointSettings,
-			const FRigidTransform3& XL0,
-			const FRigidTransform3& XL1,
-			const FVec3& P0,
-			const FRotation3& Q0,
-			const FVec3& P1,
-			const FRotation3& Q1,
-			FVec3& OutX0,
-			FMatrix33& OutR0,
-			FVec3& OutX1,
-			FMatrix33& OutR1,
-			FVec3& OutCR);
-
-		static CHAOS_API void CalculateConeConstraintSpace(
-			const FPBDJointSolverSettings& SolverSettings,
-			const FPBDJointSettings& JointSettings,
-			const FRigidTransform3& XL0,
-			const FRigidTransform3& XL1,
-			const FVec3& P0,
-			const FRotation3& Q0,
-			const FVec3& P1,
-			const FRotation3& Q1,
-			FVec3& OutX0,
-			FMatrix33& OutR0,
-			FVec3& OutX1,
-			FMatrix33& OutR1,
-			FVec3& OutCR);
-
-		static CHAOS_API void ApplyJointPositionConstraint(
-			const FReal Dt,
-			const FPBDJointSolverSettings& SolverSettings,
-			const FPBDJointSettings& JointSettings,
-			const FReal Stiffness,
-			const FRigidTransform3& XL0,
-			const FRigidTransform3& XL1,
-			FVec3& P0,
-			FRotation3& Q0,
-			FVec3& V0,
-			FVec3& W0,
-			FVec3& P1,
-			FRotation3& Q1,
-			FVec3& V1,
-			FVec3& W1,
-			float InvM0,
-			const FMatrix33& InvIL0,
-			float InvM1,
-			const FMatrix33& InvIL1);
-
-		static CHAOS_API void ApplyJointVelocityConstraint(
-			const FReal Dt,
-			const FPBDJointSolverSettings& SolverSettings,
-			const FPBDJointSettings& JointSettings,
-			const FReal Stiffness,
-			const FRigidTransform3& XL0,
-			const FRigidTransform3& XL1,
-			FVec3& P0,
-			FRotation3& Q0,
-			FVec3& V0,
-			FVec3& W0,
-			FVec3& P1,
-			FRotation3& Q1,
-			FVec3& V1,
-			FVec3& W1,
-			float InvM0,
-			const FMatrix33& InvIL0,
-			float InvM1,
-			const FMatrix33& InvIL1);
-
-		static CHAOS_API void ApplyJointTwistConstraint(
-			const FReal Dt,
-			const FPBDJointSolverSettings& SolverSettings,
-			const FPBDJointSettings& JointSettings,
-			const FReal Stiffness,
-			const FRigidTransform3& XL0,
-			const FRigidTransform3& XL1,
-			FVec3& P0,
-			FRotation3& Q0,
-			FVec3& V0,
-			FVec3& W0,
-			FVec3& P1,
-			FRotation3& Q1,
-			FVec3& V1,
-			FVec3& W1,
-			float InvM0,
-			const FMatrix33& InvIL0,
-			float InvM1,
-			const FMatrix33& InvIL1);
-
-		static CHAOS_API void ApplyJointTwistVelocityConstraint(
-			const FReal Dt,
-			const FPBDJointSolverSettings& SolverSettings,
-			const FPBDJointSettings& JointSettings,
-			const FReal Stiffness,
-			const FRigidTransform3& XL0,
-			const FRigidTransform3& XL1,
-			FVec3& P0,
-			FRotation3& Q0,
-			FVec3& V0,
-			FVec3& W0,
-			FVec3& P1,
-			FRotation3& Q1,
-			FVec3& V1,
-			FVec3& W1,
-			float InvM0,
-			const FMatrix33& InvIL0,
-			float InvM1,
-			const FMatrix33& InvIL1);
-
-		static CHAOS_API void ApplyJointConeConstraint(
-			const FReal Dt,
-			const FPBDJointSolverSettings& SolverSettings,
-			const FPBDJointSettings& JointSettings,
-			const FReal Stiffness,
-			const FRigidTransform3& XL0,
-			const FRigidTransform3& XL1,
-			FVec3& P0,
-			FRotation3& Q0,
-			FVec3& V0,
-			FVec3& W0,
-			FVec3& P1,
-			FRotation3& Q1,
-			FVec3& V1,
-			FVec3& W1,
-			float InvM0,
-			const FMatrix33& InvIL0,
-			float InvM1,
-			const FMatrix33& InvIL1);
-
-		static CHAOS_API void ApplyJointConeVelocityConstraint(
-			const FReal Dt,
-			const FPBDJointSolverSettings& SolverSettings,
-			const FPBDJointSettings& JointSettings,
-			const FReal Stiffness,
-			const FRigidTransform3& XL0,
-			const FRigidTransform3& XL1,
-			FVec3& P0,
-			FRotation3& Q0,
-			FVec3& V0,
-			FVec3& W0,
-			FVec3& P1,
-			FRotation3& Q1,
-			FVec3& V1,
-			FVec3& W1,
-			float InvM0,
-			const FMatrix33& InvIL0,
-			float InvM1,
-			const FMatrix33& InvIL1);
-
-		static CHAOS_API void ApplyJointSwingConstraint(
-			const FReal Dt,
-			const FPBDJointSolverSettings& SolverSettings,
-			const FPBDJointSettings& JointSettings,
-			const FReal Stiffness,
-			const FRigidTransform3& XL0,
-			const FRigidTransform3& XL1,
-			const EJointAngularConstraintIndex SwingConstraintIndex,
-			const EJointAngularAxisIndex SwingAxisIndex,
-			FVec3& P0,
-			FRotation3& Q0,
-			FVec3& V0,
-			FVec3& W0,
-			FVec3& P1,
-			FRotation3& Q1,
-			FVec3& V1,
-			FVec3& W1,
-			float InvM0,
-			const FMatrix33& InvIL0,
-			float InvM1,
-			const FMatrix33& InvIL1);
-
-		static CHAOS_API void ApplyJointSwingVelocityConstraint(
-			const FReal Dt,
-			const FPBDJointSolverSettings& SolverSettings,
-			const FPBDJointSettings& JointSettings,
-			const FReal Stiffness,
-			const FRigidTransform3& XL0,
-			const FRigidTransform3& XL1,
-			const EJointAngularConstraintIndex SwingConstraintIndex,
-			const EJointAngularAxisIndex SwingAxisIndex,
-			FVec3& P0,
-			FRotation3& Q0,
-			FVec3& V0,
-			FVec3& W0,
-			FVec3& P1,
-			FRotation3& Q1,
-			FVec3& V1,
-			FVec3& W1,
-			float InvM0,
-			const FMatrix33& InvIL0,
-			float InvM1,
-			const FMatrix33& InvIL1);
-
-		static CHAOS_API void ApplyJointTwistDrive(
-			const FReal Dt,
-			const FPBDJointSolverSettings& SolverSettings,
-			const FPBDJointSettings& JointSettings,
-			const FRigidTransform3& XL0,
-			const FRigidTransform3& XL1,
-			FVec3& P0,
-			FRotation3& Q0,
-			FVec3& V0,
-			FVec3& W0,
-			FVec3& P1,
-			FRotation3& Q1,
-			FVec3& V1,
-			FVec3& W1,
-			float InvM0,
-			const FMatrix33& InvIL0,
-			float InvM1,
-			const FMatrix33& InvIL1);
-
-		static CHAOS_API void ApplyJointConeDrive(
-			const FReal Dt,
-			const FPBDJointSolverSettings& SolverSettings,
-			const FPBDJointSettings& JointSettings,
-			const FRigidTransform3& XL0,
-			const FRigidTransform3& XL1,
-			FVec3& P0,
-			FRotation3& Q0,
-			FVec3& V0,
-			FVec3& W0,
-			FVec3& P1,
-			FRotation3& Q1,
-			FVec3& V1,
-			FVec3& W1,
-			float InvM0,
-			const FMatrix33& InvIL0,
-			float InvM1,
-			const FMatrix33& InvIL1);
-
-		static CHAOS_API void ApplyJointSLerpDrive(
-			const FReal Dt,
-			const FPBDJointSolverSettings& SolverSettings,
-			const FPBDJointSettings& JointSettings,
-			const FRigidTransform3& XL0,
-			const FRigidTransform3& XL1,
-			FVec3& P0,
-			FRotation3& Q0,
-			FVec3& V0,
-			FVec3& W0,
-			FVec3& P1,
-			FRotation3& Q1,
-			FVec3& V1,
-			FVec3& W1,
-			float InvM0,
-			const FMatrix33& InvIL0,
-			float InvM1,
-			const FMatrix33& InvIL1);
-
-		static CHAOS_API void ApplyJointPositionProjection(
-			const FReal Dt,
+		static FReal GetLinearProjection(
 			const FPBDJointSolverSettings& SolverSettings,
-			const FPBDJointSettings& JointSettings,
-			const FReal Stiffness,
-			const FRigidTransform3& XL0,
-			const FRigidTransform3& XL1,
-			FVec3& P0,
-			FRotation3& Q0,
-			FVec3& P1,
-			FRotation3& Q1,
-			float InvM0,
-			const FMatrix33& InvIL0,
-			float InvM1,
-			const FMatrix33& InvIL1);
+			const FPBDJointSettings& JointSettings);
 
-		static CHAOS_API void ApplyJointTwistProjection(
-			const FReal Dt,
+		static FReal GetAngularProjection(
 			const FPBDJointSolverSettings& SolverSettings,
-			const FPBDJointSettings& JointSettings,
-			const FReal Stiffness,
-			const FRigidTransform3& XL0,
-			const FRigidTransform3& XL1,
-			FVec3& P0,
-			FRotation3& Q0,
-			FVec3& P1,
-			FRotation3& Q1,
-			float InvM0,
-			const FMatrix33& InvIL0,
-			float InvM1,
-			const FMatrix33& InvIL1);
+			const FPBDJointSettings& JointSettings);
 
-		static CHAOS_API void ApplyJointConeProjection(
-			const FReal Dt,
+		static bool GetSoftAccelerationMode(
 			const FPBDJointSolverSettings& SolverSettings,
-			const FPBDJointSettings& JointSettings,
-			const FReal Stiffness,
-			const FRigidTransform3& XL0,
-			const FRigidTransform3& XL1,
-			FVec3& P0,
-			FRotation3& Q0,
-			FVec3& P1,
-			FRotation3& Q1,
-			float InvM0,
-			const FMatrix33& InvIL0,
-			float InvM1,
-			const FMatrix33& InvIL1);
+			const FPBDJointSettings& JointSettings);
 
-		static CHAOS_API void ApplyJointSwingProjection(
-			const FReal Dt,
+		static bool GetDriveAccelerationMode(
 			const FPBDJointSolverSettings& SolverSettings,
-			const FPBDJointSettings& JointSettings,
-			const FReal Stiffness,
-			const FRigidTransform3& XL0,
-			const FRigidTransform3& XL1,
-			const EJointAngularConstraintIndex SwingConstraintIndex,
-			const EJointAngularAxisIndex SwingAxisIndex,
-			FVec3& P0,
-			FRotation3& Q0,
-			FVec3& P1,
-			FRotation3& Q1,
-			float InvM0,
-			const FMatrix33& InvIL0,
-			float InvM1,
-			const FMatrix33& InvIL1);
+			const FPBDJointSettings& JointSettings);
 
+		static FVec3 GetSphereLimitedPositionError(const FVec3& CX, const FReal Radius);
+		static FVec3 GetCylinderLimitedPositionError(const FVec3& CX, const FVec3& Axis, const FReal Limit, const EJointMotionType AxisMotion);
+		static FVec3 GetLineLimitedPositionError(const FVec3& CX, const FVec3& Axis, const FReal Limit, const EJointMotionType AxisMotion);
+		static FVec3 GetLimitedPositionError(const FPBDJointSettings& JointSettings, const FRotation3& R0, const FVec3& CX);
 	};
 }
\ No newline at end of file
