Updating copyright for Engine Runtime.

#rnx

diff --git a/Engine/Source/Runtime/AIModule/AIModule.Build.cs b/Engine/Source/Runtime/AIModule/AIModule.Build.cs
--- a/Engine/Source/Runtime/AIModule/AIModule.Build.cs
+++ b/Engine/Source/Runtime/AIModule/AIModule.Build.cs
@@ -1,4 +1,4 @@
-// Copyright 1998-2019 Epic Games, Inc. All Rights Reserved.
+// Copyright Epic Games, Inc. All Rights Reserved.
 
 namespace UnrealBuildTool.Rules
 {
diff --git a/Engine/Source/Runtime/AIModule/Classes/AIController.h b/Engine/Source/Runtime/AIModule/Classes/AIController.h
--- a/Engine/Source/Runtime/AIModule/Classes/AIController.h
+++ b/Engine/Source/Runtime/AIModule/Classes/AIController.h
@@ -1,4 +1,4 @@
-// Copyright 1998-2019 Epic Games, Inc. All Rights Reserved.
+// Copyright Epic Games, Inc. All Rights Reserved.
 
 #pragma once
 
diff --git a/Engine/Source/Runtime/ApplicationCore/Private/Linux/LinuxPlatformSplash.cpp b/Engine/Source/Runtime/ApplicationCore/Private/Linux/LinuxPlatformSplash.cpp
--- a/Engine/Source/Runtime/ApplicationCore/Private/Linux/LinuxPlatformSplash.cpp
+++ b/Engine/Source/Runtime/ApplicationCore/Private/Linux/LinuxPlatformSplash.cpp
@@ -1,4 +1,4 @@
-// Copyright 1998-2019 Epic Games, Inc. All Rights Reserved.
+// Copyright Epic Games, Inc. All Rights Reserved.
 
 #include "Linux/LinuxPlatformSplash.h"
 
@@ -766,7 +766,7 @@ void FLinuxPlatformSplash::Show( )
 
 		// Display copyright information in editor splash screen
 		{
-			const FText CopyrightInfo = NSLOCTEXT( "UnrealEd", "SplashScreen_CopyrightInfo", "Copyright \x00a9 1998-2019   Epic Games, Inc.   All rights reserved." );
+			const FText CopyrightInfo = NSLOCTEXT( "UnrealEd", "SplashScreen_CopyrightInfo", "Copyright \x00a9   Epic Games, Inc.   All rights reserved." );
 			GSplashState->SetSplashText( SplashTextType::CopyrightInfo, CopyrightInfo );
 		}
 	}
diff --git a/Engine/Source/Runtime/ApplicationCore/Private/Mac/MacPlatformSplash.cpp b/Engine/Source/Runtime/ApplicationCore/Private/Mac/MacPlatformSplash.cpp
--- a/Engine/Source/Runtime/ApplicationCore/Private/Mac/MacPlatformSplash.cpp
+++ b/Engine/Source/Runtime/ApplicationCore/Private/Mac/MacPlatformSplash.cpp
@@ -1,4 +1,4 @@
-// Copyright 1998-2019 Epic Games, Inc. All Rights Reserved.
+// Copyright Epic Games, Inc. All Rights Reserved.
 
 #include "Mac/MacPlatformSplash.h"
 #include "Misc/App.h"
@@ -233,7 +233,7 @@ void FMacPlatformSplash::Show()
 
 				// Display copyright information in editor splash screen
 				{
-					const FText CopyrightInfo = NSLOCTEXT( "UnrealEd", "SplashScreen_CopyrightInfo", "Copyright \x00a9 1998-2019   Epic Games, Inc.   All rights reserved." );
+					const FText CopyrightInfo = NSLOCTEXT( "UnrealEd", "SplashScreen_CopyrightInfo", "Copyright \x00a9   Epic Games, Inc.   All rights reserved." );
 					StartSetSplashText( SplashTextType::CopyrightInfo, CopyrightInfo );
 				}
 			}
diff --git a/Engine/Source/Runtime/ApplicationCore/Private/Windows/WindowsPlatformSplash.cpp b/Engine/Source/Runtime/ApplicationCore/Private/Windows/WindowsPlatformSplash.cpp
--- a/Engine/Source/Runtime/ApplicationCore/Private/Windows/WindowsPlatformSplash.cpp
+++ b/Engine/Source/Runtime/ApplicationCore/Private/Windows/WindowsPlatformSplash.cpp
@@ -1,4 +1,4 @@
-// Copyright 1998-2019 Epic Games, Inc. All Rights Reserved.
+// Copyright Epic Games, Inc. All Rights Reserved.
 
 #include "Windows/WindowsPlatformSplash.h"
 #include "HAL/PlatformMemory.h"
@@ -658,7 +658,7 @@ void FWindowsPlatformSplash::Show()
 
 				// Display copyright information in editor splash screen
 				{
-					const FString CopyrightInfo = NSLOCTEXT( "UnrealEd", "SplashScreen_CopyrightInfo", "Copyright \x00a9 1998-2019   Epic Games, Inc.   All rights reserved." ).ToString();
+					const FString CopyrightInfo = NSLOCTEXT( "UnrealEd", "SplashScreen_CopyrightInfo", "Copyright \x00a9   Epic Games, Inc.   All rights reserved." ).ToString();
 					StartSetSplashText( SplashTextType::CopyrightInfo, *CopyrightInfo );
 				}
 			}
diff --git a/Engine/Source/Runtime/Core/Private/Misc/Base64.cpp b/Engine/Source/Runtime/Core/Private/Misc/Base64.cpp
--- a/Engine/Source/Runtime/Core/Private/Misc/Base64.cpp
+++ b/Engine/Source/Runtime/Core/Private/Misc/Base64.cpp
@@ -1,4 +1,4 @@
-﻿// Copyright 1998-2019 Epic Games, Inc. All Rights Reserved.
+﻿// Copyright Epic Games, Inc. All Rights Reserved.
 
 #include "Misc/Base64.h"
 #include "Containers/StringConv.h"
diff --git a/Engine/Source/Runtime/Core/Public/Misc/Base64.h b/Engine/Source/Runtime/Core/Public/Misc/Base64.h
--- a/Engine/Source/Runtime/Core/Public/Misc/Base64.h
+++ b/Engine/Source/Runtime/Core/Public/Misc/Base64.h
@@ -1,4 +1,4 @@
-﻿// Copyright 1998-2019 Epic Games, Inc. All Rights Reserved.
+﻿// Copyright Epic Games, Inc. All Rights Reserved.
 
 #pragma once
 
diff --git a/Engine/Source/Runtime/Engine/Private/Streaming/Texture2DMipDataProvider_IO.cpp b/Engine/Source/Runtime/Engine/Private/Streaming/Texture2DMipDataProvider_IO.cpp
--- a/Engine/Source/Runtime/Engine/Private/Streaming/Texture2DMipDataProvider_IO.cpp
+++ b/Engine/Source/Runtime/Engine/Private/Streaming/Texture2DMipDataProvider_IO.cpp
@@ -1,4 +1,4 @@
-// Copyright 1998-2019 Epic Games, Inc. All Rights Reserved.FTexture2DMipDataProvider_IO
+// Copyright Epic Games, Inc. All Rights Reserved.FTexture2DMipDataProvider_IO
 
 /*=============================================================================
 Texture2DMipDataProvider_IO.cpp : Implementation of FTextureMipDataProvider using cooked file IO.
diff --git a/Engine/Source/Runtime/Engine/Public/StereoRenderTargetManager.h b/Engine/Source/Runtime/Engine/Public/StereoRenderTargetManager.h
--- a/Engine/Source/Runtime/Engine/Public/StereoRenderTargetManager.h
+++ b/Engine/Source/Runtime/Engine/Public/StereoRenderTargetManager.h
@@ -1,4 +1,4 @@
-// Copyright 1998-2019 Epic Games, Inc.All Rights Reserved.
+// Copyright Epic Games, Inc.All Rights Reserved.
 
 /*=============================================================================
 StereoRenderTargetManager.h: Abstract interface returned from IStereoRendering to support rendering into a texture
diff --git a/Engine/Source/Runtime/Experimental/GeometryCollectionCore/Private/GeometryCollection/GeometryCollection.cpp b/Engine/Source/Runtime/Experimental/GeometryCollectionCore/Private/GeometryCollection/GeometryCollection.cpp
--- a/Engine/Source/Runtime/Experimental/GeometryCollectionCore/Private/GeometryCollection/GeometryCollection.cpp
+++ b/Engine/Source/Runtime/Experimental/GeometryCollectionCore/Private/GeometryCollection/GeometryCollection.cpp
@@ -1,4 +1,4 @@
-// Copyright 1998-2019 Epic Games, Inc. All Rights Reserved.
+// Copyright Epic Games, Inc. All Rights Reserved.
 
 /*=============================================================================
 	GeometryCollection.cpp: FGeometryCollection methods.
@@ -1142,7 +1142,7 @@ void FGeometryCollection::WriteDataToHeaderFile(const FString &Name, const FStri
 
 	ofstream DataFile;
 	DataFile.open(string(TCHAR_TO_UTF8(*FullPath)));
-	DataFile << "// Copyright 1998-2019 Epic Games, Inc. All Rights Reserved." << endl << endl;
+	DataFile << "// Copyright Epic Games, Inc. All Rights Reserved." << endl << endl;
 	DataFile << "#pragma once" << endl << endl;
 	DataFile << "class " << TCHAR_TO_UTF8(*Name) << endl;
 	DataFile << "{" << endl;
diff --git a/Engine/Source/Runtime/Renderer/Private/MobileDistortionPass.h b/Engine/Source/Runtime/Renderer/Private/MobileDistortionPass.h
--- a/Engine/Source/Runtime/Renderer/Private/MobileDistortionPass.h
+++ b/Engine/Source/Runtime/Renderer/Private/MobileDistortionPass.h
@@ -1,4 +1,4 @@
-// Copyright 1998-2019 Epic Games, Inc.All Rights Reserved.
+// Copyright Epic Games, Inc.All Rights Reserved.
 
 #pragma once
 
diff --git a/Engine/Source/Runtime/Slate/Public/Widgets/Views/STreeView.h b/Engine/Source/Runtime/Slate/Public/Widgets/Views/STreeView.h
--- a/Engine/Source/Runtime/Slate/Public/Widgets/Views/STreeView.h
+++ b/Engine/Source/Runtime/Slate/Public/Widgets/Views/STreeView.h
@@ -1,4 +1,4 @@
-// Copyright 1998-2019 Epic Games, Inc. All Rights Reserved
+// Copyright Epic Games, Inc. All Rights Reserved
  
 #pragma once
 
diff --git a/Engine/Source/Runtime/Sockets/Private/Unix/SocketsUnix.cpp b/Engine/Source/Runtime/Sockets/Private/Unix/SocketsUnix.cpp
--- a/Engine/Source/Runtime/Sockets/Private/Unix/SocketsUnix.cpp
+++ b/Engine/Source/Runtime/Sockets/Private/Unix/SocketsUnix.cpp
@@ -1,4 +1,4 @@
-// Copyright 1998-2019 Epic Games, Inc. All Rights Reserved
+// Copyright Epic Games, Inc. All Rights Reserved
 
 
 #include "SocketsUnix.h"
diff --git a/Engine/Source/Runtime/VulkanRHI/Private/VulkanCommandBuffer.h b/Engine/Source/Runtime/VulkanRHI/Private/VulkanCommandBuffer.h
--- a/Engine/Source/Runtime/VulkanRHI/Private/VulkanCommandBuffer.h
+++ b/Engine/Source/Runtime/VulkanRHI/Private/VulkanCommandBuffer.h
@@ -1,4 +1,4 @@
-// Copyright 1998-2019 Epic Games, Inc. All Rights Reserved..
+// Copyright Epic Games, Inc. All Rights Reserved..
 
 /*=============================================================================
 	VulkanCommandBuffer.h: Private Vulkan RHI definitions.
diff --git a/Engine/Source/Runtime/VulkanRHI/Private/VulkanDevice.h b/Engine/Source/Runtime/VulkanRHI/Private/VulkanDevice.h
--- a/Engine/Source/Runtime/VulkanRHI/Private/VulkanDevice.h
+++ b/Engine/Source/Runtime/VulkanRHI/Private/VulkanDevice.h
@@ -1,4 +1,4 @@
-// Copyright 1998-2019 Epic Games, Inc. All Rights Reserved..
+// Copyright Epic Games, Inc. All Rights Reserved..
 
 /*=============================================================================
 	VulkanDevice.h: Private Vulkan RHI definitions.
diff --git a/Engine/Source/Runtime/Windows/XAudio2/Private/XAudio2Effects.cpp b/Engine/Source/Runtime/Windows/XAudio2/Private/XAudio2Effects.cpp
--- a/Engine/Source/Runtime/Windows/XAudio2/Private/XAudio2Effects.cpp
+++ b/Engine/Source/Runtime/Windows/XAudio2/Private/XAudio2Effects.cpp
@@ -1,4 +1,4 @@
-// Copyright 1998-2019 Epic Games, Inc. All Rights Reserved.
+// Copyright Epic Games, Inc. All Rights Reserved.
 
 /*=============================================================================
 	XeAudioDevice.cpp: Unreal XAudio2 Audio interface object.
@@ -295,7 +295,7 @@ XAPO_REGISTRATION_PROPERTIES FXAudio2RadioEffect::Registration =
 {
 	__uuidof( FXAudio2RadioEffect ),
 	TEXT( "FXAudio2RadioEffect" ), 
-	TEXT( "Copyright 1998-2019 Epic Games, Inc. All Rights Reserved." ),
+	TEXT( "Copyright Epic Games, Inc. All Rights Reserved." ),
 	1, 0,
 	XAPO_FLAG_INPLACE_REQUIRED	| XAPO_FLAG_CHANNELS_MUST_MATCH
 								| XAPO_FLAG_FRAMERATE_MUST_MATCH
