Chaos
- temporarily disable combo of linear and angular limits on same joint (the angular stiffness fixes don't work with this)
- fix angular projection velocity correction

diff --git a/Engine/Source/Runtime/Experimental/Chaos/Private/Chaos/Joint/PBDJointSolverGaussSeidel.cpp b/Engine/Source/Runtime/Experimental/Chaos/Private/Chaos/Joint/PBDJointSolverGaussSeidel.cpp
--- a/Engine/Source/Runtime/Experimental/Chaos/Private/Chaos/Joint/PBDJointSolverGaussSeidel.cpp
+++ b/Engine/Source/Runtime/Experimental/Chaos/Private/Chaos/Joint/PBDJointSolverGaussSeidel.cpp
@@ -875,17 +875,25 @@ namespace Chaos
 			FVec3 DV1 = FVec3(0);
 			FVec3 DW1 = FVec3(0);
 
+			// Remove relative angular velocity if it would increase the constraint error
 			const FReal WAxis = FVec3::DotProduct(Ws[1] - Ws[0], SwingAxis);
 			if (FMath::Sign(WAxis) == FMath::Sign(DSwingAngle))
 			{
 				DW1 = -WAxis * SwingAxis;
 			}
 
+			// Remove relative velocity if it would increase the constraint error
 			FReal DP1Len = DP1.Size();
 			if (DP1Len > KINDA_SMALL_NUMBER)
 			{
 				const FVec3 DP1Dir = DP1 / DP1Len;
-				DV1 = -FVec3::DotProduct((Vs[1] - Vs[0]), DP1Dir) * DP1Dir;
+				const FVec3 V0 = Vs[0] + FVec3::CrossProduct(Ws[0], Xs[0] - Ps[0] + Ps[1] - Xs[1]);
+				const FVec3 V1 = Vs[1];
+				const FReal VAxis = FVec3::DotProduct(V1 - V0, DP1Dir);
+				if (VAxis < 0)
+				{
+					DV1 = -VAxis * DP1Dir;
+				}
 			}
 
 			ApplyRotationDelta(1, AngularProjection, DR1);
diff --git a/Engine/Source/Runtime/Experimental/Chaos/Private/Chaos/PBDJointConstraints.cpp b/Engine/Source/Runtime/Experimental/Chaos/Private/Chaos/PBDJointConstraints.cpp
--- a/Engine/Source/Runtime/Experimental/Chaos/Private/Chaos/PBDJointConstraints.cpp
+++ b/Engine/Source/Runtime/Experimental/Chaos/Private/Chaos/PBDJointConstraints.cpp
@@ -151,6 +151,14 @@ namespace Chaos
 			bSoftSwingLimitsEnabled = false;
 		}
 
+		// Currently we do not support angular limits if linear limits are set up
+		// @todo(ccaulfield): fix angular constraint position correction in FJointSolverGaussSeidel::ApplyRotationCorrection and FJointSolverGaussSeidel::ApplyRotationCorrectionSoft
+		if ((LinearMotionTypes[0] == EJointMotionType::Limited) || (LinearMotionTypes[1] == EJointMotionType::Limited) || (LinearMotionTypes[2] == EJointMotionType::Limited))
+		{
+			AngularMotionTypes[(int32)EJointAngularConstraintIndex::Twist] = EJointMotionType::Free;
+			AngularMotionTypes[(int32)EJointAngularConstraintIndex::Swing1] = EJointMotionType::Free;
+			AngularMotionTypes[(int32)EJointAngularConstraintIndex::Swing2] = EJointMotionType::Free;
+		}
 	}
 
 	
