Chaos
- position XPBD with force or acceleration mode
- position drives and soft joints now match PhysX
- option to solve each constraint in a joint using state at beginning of iteration. Could be significantly faster, but is not very stable - checking in only for possible future reference.

diff --git a/Engine/Source/Runtime/Engine/Private/PhysicsEngine/ImmediatePhysics/ImmediatePhysicsChaos/ImmediatePhysicsJointHandle_Chaos.cpp b/Engine/Source/Runtime/Engine/Private/PhysicsEngine/ImmediatePhysics/ImmediatePhysicsChaos/ImmediatePhysicsJointHandle_Chaos.cpp
--- a/Engine/Source/Runtime/Engine/Private/PhysicsEngine/ImmediatePhysics/ImmediatePhysicsChaos/ImmediatePhysicsJointHandle_Chaos.cpp
+++ b/Engine/Source/Runtime/Engine/Private/PhysicsEngine/ImmediatePhysics/ImmediatePhysicsChaos/ImmediatePhysicsJointHandle_Chaos.cpp
@@ -14,6 +14,10 @@ static_assert((int32)Chaos::EJointMotionType::Free == (int32)EAngularConstraintM
 static_assert((int32)Chaos::EJointMotionType::Limited == (int32)EAngularConstraintMotion::ACM_Limited, "Chaos::EJointMotionType and EAngularConstraintMotion mismatch");
 static_assert((int32)Chaos::EJointMotionType::Locked == (int32)EAngularConstraintMotion::ACM_Locked, "Chaos::EJointMotionType and EAngularConstraintMotion mismatch");
 
+// NOTE: Hard dependence on EJointAngularConstraintIndex - the following will break if we change the order (but can be easily fixed). See FJointHandle::FJointHandle
+static_assert((int32)Chaos::EJointAngularConstraintIndex::Twist == 0, "Angular drive targets have hard dependency on constraint order");
+static_assert((int32)Chaos::EJointAngularConstraintIndex::Swing1 == 2, "Angular drive targets have hard dependency on constraint order");
+
 namespace ImmediatePhysics_Chaos
 {
 	float ChaosImmediate_JointStiffness = 1.0f;
@@ -24,12 +28,17 @@ namespace ImmediatePhysics_Chaos
 	FAutoConsoleVariableRef CVarDriveStiffnessScale(TEXT("p.Chaos.ImmPhys.DriveStiffnessScale"), ChaosImmediate_DriveStiffnessScale, TEXT("Conversion factor for drive stiffness."));
 	FAutoConsoleVariableRef CVarDriveDampingScale(TEXT("p.Chaos.ImmPhys.DriveDampingScale"), ChaosImmediate_DriveDampingScale, TEXT("Conversion factor for drive damping."));
 
+	float ChaosImmediate_SoftLinearStiffnessScale = 1.0f;
+	float ChaosImmediate_SoftLinearDampingScale = 1.0f;
+	FAutoConsoleVariableRef CVarSoftLinearStiffnessScale(TEXT("p.Chaos.ImmPhys.SoftLinearStiffnessScale"), ChaosImmediate_SoftLinearStiffnessScale, TEXT("Conversion factor for soft-joint stiffness."));
+	FAutoConsoleVariableRef CVarSoftLinearDampingScale(TEXT("p.Chaos.ImmPhys.SoftLinearDampingScale"), ChaosImmediate_SoftLinearDampingScale, TEXT("Conversion factor for soft-joint damping."));
+
 	// This value is from PhysX UE - in Chaos we don't really need Soft constraints with stiffness greater than about 1000
 	// (In PhysX, Soft constraints were implicitly integrated, and essentially equivalent to XPBD, which for large stiffness is just PBD).
-	float ChaosImmediate_SoftStiffnessScale = 100000.0f;
-	float ChaosImmediate_SoftDampingScale = 100000.0f;
-	FAutoConsoleVariableRef CVarSoftStiffnessScale(TEXT("p.Chaos.ImmPhys.SoftStiffnessScale"), ChaosImmediate_SoftStiffnessScale, TEXT("Conversion factor for soft-joint stiffness."));
-	FAutoConsoleVariableRef CVarSoftDampingScale(TEXT("p.Chaos.ImmPhys.SoftDampingScale"), ChaosImmediate_SoftDampingScale, TEXT("Conversion factor for soft-joint damping."));
+	float ChaosImmediate_SoftAngularStiffnessScale = 100000.0f;
+	float ChaosImmediate_SoftAngularDampingScale = 100000.0f;
+	FAutoConsoleVariableRef CVarSoftAngularStiffnessScale(TEXT("p.Chaos.ImmPhys.SoftAngularStiffnessScale"), ChaosImmediate_SoftAngularStiffnessScale, TEXT("Conversion factor for soft-joint stiffness."));
+	FAutoConsoleVariableRef CVarSoftAngularDampingScale(TEXT("p.Chaos.ImmPhys.SoftAngularDampingScale"), ChaosImmediate_SoftAngularDampingScale, TEXT("Conversion factor for soft-joint damping."));
 
 	float ChaosImmediate_JointMinLinearProjection = 0.0f;
 	float ChaosImmediate_JointMaxLinearProjection = 0.5f;
@@ -80,27 +89,28 @@ namespace ImmediatePhysics_Chaos
 		ConstraintSettings.bSoftLinearLimitsEnabled = ConstraintInstance->GetIsSoftLinearLimit();
 		ConstraintSettings.bSoftTwistLimitsEnabled = ConstraintInstance->GetIsSoftTwistLimit();
 		ConstraintSettings.bSoftSwingLimitsEnabled = ConstraintInstance->GetIsSoftSwingLimit();
-		ConstraintSettings.SoftLinearStiffness = 1.0f;	// @todo(ccaulfield): xpbd soft linear constraints
-		ConstraintSettings.SoftLinearDamping = 0.0f;		// @todo(ccaulfield): xpbd soft linear constraints
-		ConstraintSettings.SoftTwistStiffness = ChaosImmediate_SoftStiffnessScale * ConstraintInstance->GetSoftTwistLimitStiffness();
-		ConstraintSettings.SoftTwistDamping = ChaosImmediate_SoftDampingScale * ConstraintInstance->GetSoftTwistLimitDamping();
-		ConstraintSettings.SoftSwingStiffness = ChaosImmediate_SoftStiffnessScale * ConstraintInstance->GetSoftSwingLimitStiffness();
-		ConstraintSettings.SoftSwingDamping = ChaosImmediate_SoftDampingScale * ConstraintInstance->GetSoftSwingLimitDamping();
-		ConstraintSettings.SoftForceMode = EJointForceMode::Force;
+		ConstraintSettings.SoftLinearStiffness = ChaosImmediate_SoftLinearStiffnessScale * ConstraintInstance->GetSoftLinearLimitStiffness();
+		ConstraintSettings.SoftLinearDamping = ChaosImmediate_SoftLinearDampingScale * ConstraintInstance->GetSoftLinearLimitDamping();
+		ConstraintSettings.SoftTwistStiffness = ChaosImmediate_SoftAngularStiffnessScale * ConstraintInstance->GetSoftTwistLimitStiffness();
+		ConstraintSettings.SoftTwistDamping = ChaosImmediate_SoftAngularDampingScale * ConstraintInstance->GetSoftTwistLimitDamping();
+		ConstraintSettings.SoftSwingStiffness = ChaosImmediate_SoftAngularStiffnessScale * ConstraintInstance->GetSoftSwingLimitStiffness();
+		ConstraintSettings.SoftSwingDamping = ChaosImmediate_SoftAngularDampingScale * ConstraintInstance->GetSoftSwingLimitDamping();
+		ConstraintSettings.LinearSoftForceMode = EJointForceMode::Acceleration;
+		ConstraintSettings.AngularSoftForceMode = EJointForceMode::Acceleration;
 
 		ConstraintSettings.LinearDriveTarget = Profile.LinearDrive.PositionTarget;
-		ConstraintSettings.bLinearDriveEnabled[0] = Profile.LinearDrive.XDrive.bEnablePositionDrive;
-		ConstraintSettings.bLinearDriveEnabled[1] = Profile.LinearDrive.YDrive.bEnablePositionDrive;
-		ConstraintSettings.bLinearDriveEnabled[2] = Profile.LinearDrive.ZDrive.bEnablePositionDrive;
-		ConstraintSettings.LinearDriveStiffness = 0.3f;// Profile.LinearDrive.XDrive.bEnablePositionDrive;
-		ConstraintSettings.LinearDriveDamping = 0.0f;// Profile.LinearDrive.XDrive.bEnablePositionDrive;
+		ConstraintSettings.bLinearPositionDriveEnabled[0] = Profile.LinearDrive.XDrive.bEnablePositionDrive;
+		ConstraintSettings.bLinearPositionDriveEnabled[1] = Profile.LinearDrive.YDrive.bEnablePositionDrive;
+		ConstraintSettings.bLinearPositionDriveEnabled[2] = Profile.LinearDrive.ZDrive.bEnablePositionDrive;
+		ConstraintSettings.bLinearVelocityDriveEnabled[0] = Profile.LinearDrive.XDrive.bEnableVelocityDrive;
+		ConstraintSettings.bLinearVelocityDriveEnabled[1] = Profile.LinearDrive.YDrive.bEnableVelocityDrive;
+		ConstraintSettings.bLinearVelocityDriveEnabled[2] = Profile.LinearDrive.ZDrive.bEnableVelocityDrive;
+		ConstraintSettings.LinearDriveStiffness = Profile.LinearDrive.XDrive.Stiffness;
+		ConstraintSettings.LinearDriveDamping = Profile.LinearDrive.XDrive.Damping;
 		ConstraintSettings.LinearDriveForceMode = EJointForceMode::Acceleration;
 
 		ConstraintSettings.AngularDrivePositionTarget = FQuat(Profile.AngularDrive.OrientationTarget);
 		ConstraintSettings.AngularDriveVelocityTarget = Profile.AngularDrive.AngularVelocityTarget;
-		// NOTE: Hard dependence on EJointAngularConstraintIndex - the following will break if we change the order (but can be easily fixed)
-		static_assert((int32)EJointAngularConstraintIndex::Twist == 0, "Angular drive targets have hard dependency on constraint order");
-		static_assert((int32)EJointAngularConstraintIndex::Swing1 == 2, "Angular drive targets have hard dependency on constraint order");
 		ConstraintSettings.AngularDriveTargetAngles[(int32)EJointAngularConstraintIndex::Twist] = FMath::DegreesToRadians(Profile.AngularDrive.OrientationTarget.Roll);
 		ConstraintSettings.AngularDriveTargetAngles[(int32)EJointAngularConstraintIndex::Swing1] = FMath::DegreesToRadians(Profile.AngularDrive.OrientationTarget.Yaw);
 		ConstraintSettings.AngularDriveTargetAngles[(int32)EJointAngularConstraintIndex::Swing2] = FMath::DegreesToRadians(Profile.AngularDrive.OrientationTarget.Pitch);
@@ -117,12 +127,11 @@ namespace ImmediatePhysics_Chaos
 			ConstraintSettings.bAngularSwingPositionDriveEnabled = Profile.AngularDrive.SwingDrive.bEnablePositionDrive;
 			ConstraintSettings.bAngularSwingVelocityDriveEnabled = Profile.AngularDrive.SwingDrive.bEnableVelocityDrive;
 		}
-		// All angular drives in UE have same stiffness and damping settings
 		ConstraintSettings.AngularDriveStiffness = ChaosImmediate_DriveStiffnessScale * Profile.AngularDrive.TwistDrive.Stiffness;
 		ConstraintSettings.AngularDriveDamping = ChaosImmediate_DriveDampingScale * Profile.AngularDrive.TwistDrive.Damping;
 		ConstraintSettings.AngularDriveForceMode = EJointForceMode::Acceleration;
 
-		ConstraintSettings.Sanitize();
+		ConstraintSettings.Sanitize(Actor1->GetMass(), Actor2->GetMass());
 
 		ConstraintHandle = Constraints->AddConstraint({ Actor1->ParticleHandle, Actor2->ParticleHandle }, ConstraintFrames, ConstraintSettings);
 	}
diff --git a/Engine/Source/Runtime/Experimental/Chaos/Private/Chaos/Joint/PBDJointSolverGaussSeidel.cpp b/Engine/Source/Runtime/Experimental/Chaos/Private/Chaos/Joint/PBDJointSolverGaussSeidel.cpp
--- a/Engine/Source/Runtime/Experimental/Chaos/Private/Chaos/Joint/PBDJointSolverGaussSeidel.cpp
+++ b/Engine/Source/Runtime/Experimental/Chaos/Private/Chaos/Joint/PBDJointSolverGaussSeidel.cpp
@@ -8,8 +8,102 @@
 
 //#pragma optimize("", off)
 
+// Set to true to calculate all constraint corrections based on state at the beginning of the iteration. This would eliminate the need to recalculate the inertia and a few other
+// things in each constraint's Apply method. However, it probably introduces instability and requires more iterations to solve...
+bool bChaos_Joint_Accumulate = false;
+FAutoConsoleVariableRef CVarChaosImmPhysStepTime(TEXT("p.Chaos.Joint.Accumulate"), bChaos_Joint_Accumulate, TEXT("Whether to accumulate forces for each constraint over an iteration and apply all at once"));
+
 namespace Chaos
 {
+	void GetSphericalAxisDelta(
+		const FVec3& X0,
+		const FVec3& X1,
+		FVec3& Axis,
+		FReal& Delta)
+	{
+		Axis = FVec3(0);
+		Delta = 0;
+
+		const FVec3 DX = X1 - X0;
+		const FReal DXLen = DX.Size();
+		if (DXLen > KINDA_SMALL_NUMBER)
+		{
+			Axis = DX / DXLen;
+			Delta = DXLen;
+		}
+	}
+
+
+	FReal GetSphericalError(
+		const FReal& Delta,
+		const FReal& Limit)
+	{
+		return FMath::Max((FReal)0, Delta - Limit);
+	}
+
+
+	void GetCylindricalAxisDelta(
+		const FRotation3& R0,
+		const FVec3& X0,
+		const FVec3& X1,
+		const int32 CylinderAxisIndex,
+		FVec3& Axis,
+		FReal& Delta)
+	{
+		Axis = FVec3(0);
+		Delta = 0;
+
+		const FMatrix33 R0M = R0.ToMatrix();
+		const FVec3 CylinderAxis = R0M.GetAxis(CylinderAxisIndex);
+		FVec3 DX = X1 - X0;
+		DX = DX - FVec3::DotProduct(DX, CylinderAxis) * CylinderAxis;
+		const FReal DXLen = DX.Size();
+		if (DXLen > KINDA_SMALL_NUMBER)
+		{
+			Axis = DX / DXLen;
+			Delta = DXLen;
+		}
+	}
+
+
+	FReal GetCylindricalError(
+		const FReal& Delta,
+		const FReal& Limit)
+	{
+		return FMath::Max((FReal)0, Delta - Limit);
+	}
+
+
+	void GetPlanarAxisDelta(
+		const FRotation3& R0,
+		const FVec3& X0,
+		const FVec3& X1,
+		const int32 PlaneAxisIndex,
+		FVec3& Axis,
+		FReal& Delta)
+	{
+		const FMatrix33 R0M = R0.ToMatrix();
+		Axis = R0M.GetAxis(PlaneAxisIndex);
+		Delta = FVec3::DotProduct(X1 - X0, Axis);
+	}
+
+
+	FReal GetPlanarError(
+		const FReal& Delta,
+		const FReal& Limit)
+	{
+		if (Delta > Limit)
+		{
+			return Delta - Limit;
+		}
+		else if (Delta < Limit)
+		{
+			return Delta + Limit;
+		}
+		return 0;
+	}
+
+
 	void GetTwistAxisAngleLocal(
 		const FRotation3& R0,
 		const FRotation3& R1,
@@ -197,6 +291,8 @@ namespace Chaos
 		const FReal Dt,
 		const FPBDJointSolverSettings& SolverSettings,
 		const FPBDJointSettings& JointSettings,
+		const FVec3& PrevP0,
+		const FVec3& PrevP1,
 		const FRotation3& PrevQ0,
 		const FRotation3& PrevQ1,
 		const FReal InvM0,
@@ -215,13 +311,23 @@ namespace Chaos
 		InvMs[0] = JointSettings.ParentInvMassScale * InvM0;
 		InvMs[1] = InvM1;
 
+		PrevPs[0] = PrevP0;
+		PrevPs[1] = PrevP1;
 		PrevQs[0] = PrevQ0;
 		PrevQs[1] = PrevQ1;
+		PrevXs[0] = PrevP0 + PrevQ0 * XL0.GetTranslation();
+		PrevXs[1] = PrevP1 + PrevQ1 * XL1.GetTranslation();
 
+		LinearSoftLambda = (FReal)0;
+		LinearDriveLambda = (FReal)0;
 		TwistSoftLambda = (FReal)0;
 		SwingSoftLambda = (FReal)0;
 		TwistDriveLambda = (FReal)0;
 		SwingDriveLambda = (FReal)0;
+
+		// Disable the angular limit hardness improvement if linear limits are set up
+		// @todo(ccaulfield): fix angular constraint position correction in ApplyRotationCorrection and ApplyRotationCorrectionSoft
+		bEnableAngularConstraintPositionCorrection = ((JointSettings.LinearMotionTypes[0] == EJointMotionType::Locked) && (JointSettings.LinearMotionTypes[1] == EJointMotionType::Locked) && (JointSettings.LinearMotionTypes[2] == EJointMotionType::Locked));
 	}
 
 
@@ -251,14 +357,31 @@ namespace Chaos
 	}
 
 
+	void FJointSolverGaussSeidel::InitAccumulatedDeltas()
+	{
+		DPs[0] = FVec3(0);
+		DPs[1] = FVec3(0);
+		DRs[0] = FVec3(0);
+		DRs[1] = FVec3(0);
+		DVs[0] = FVec3(0);
+		DVs[1] = FVec3(0);
+		DWs[0] = FVec3(0);
+		DWs[1] = FVec3(0);
+	}
+
+
 	void FJointSolverGaussSeidel::ApplyConstraints(
 		const FReal Dt,
 		const FPBDJointSolverSettings& SolverSettings,
 		const FPBDJointSettings& JointSettings)
 	{
+		InitAccumulatedDeltas();
+
 		ApplyRotationConstraints(Dt, SolverSettings, JointSettings);
 
 		ApplyPositionConstraints(Dt, SolverSettings, JointSettings);
+	
+		ApplyAccumulatedDeltas();
 	}
 
 
@@ -267,9 +390,13 @@ namespace Chaos
 		const FPBDJointSolverSettings& SolverSettings,
 		const FPBDJointSettings& JointSettings)
 	{
+		InitAccumulatedDeltas();
+
 		ApplyRotationDrives(Dt, SolverSettings, JointSettings);
 
 		ApplyPositionDrives(Dt, SolverSettings, JointSettings);
+
+		ApplyAccumulatedDeltas();
 	}
 
 
@@ -278,10 +405,14 @@ namespace Chaos
 		const FPBDJointSolverSettings& SolverSettings,
 		const FPBDJointSettings& JointSettings)
 	{
+		InitAccumulatedDeltas();
+
 		// With the parent body set to infinite mass...
 		ApplyRotationProjection(Dt, SolverSettings, JointSettings);
 
 		ApplyPositionProjection(Dt, SolverSettings, JointSettings);
+
+		ApplyAccumulatedDeltas();
 	}
 
 
@@ -441,9 +572,132 @@ namespace Chaos
 		const FPBDJointSettings& JointSettings)
 	{
 		const TVector<EJointMotionType, 3>& LinearMotion = JointSettings.LinearMotionTypes;
-		if ((LinearMotion[0] != EJointMotionType::Free) || (LinearMotion[1] != EJointMotionType::Free) || (LinearMotion[2] != EJointMotionType::Free))
+
+		// All 3 axes locked -> hard point constraint
+		if ((LinearMotion[0] == EJointMotionType::Locked) && (LinearMotion[1] == EJointMotionType::Locked) && (LinearMotion[2] == EJointMotionType::Locked))
+		{
+			ApplyPointPositionConstraint(Dt, SolverSettings, JointSettings);
+			return;
+		}
+
+		// At least one axes is limited or free
+		const TVector<bool, 3> bLinearLimitSoft =
+		{
+			(LinearMotion[0] == EJointMotionType::Limited) && JointSettings.bSoftLinearLimitsEnabled,
+			(LinearMotion[1] == EJointMotionType::Limited) && JointSettings.bSoftLinearLimitsEnabled,
+			(LinearMotion[2] == EJointMotionType::Limited) && JointSettings.bSoftLinearLimitsEnabled,
+		};
+		const TVector<bool, 3> bLinearLimitHard =
+		{
+			(LinearMotion[0] == EJointMotionType::Locked) || ((LinearMotion[0] == EJointMotionType::Limited) && !JointSettings.bSoftLinearLimitsEnabled),
+			(LinearMotion[1] == EJointMotionType::Locked) || ((LinearMotion[1] == EJointMotionType::Limited) && !JointSettings.bSoftLinearLimitsEnabled),
+			(LinearMotion[2] == EJointMotionType::Locked) || ((LinearMotion[2] == EJointMotionType::Limited) && !JointSettings.bSoftLinearLimitsEnabled),
+		};
+
+		// Handle hard limits
+		if (bLinearLimitHard[0] && bLinearLimitHard[1] && bLinearLimitHard[2])
+		{
+			ApplySphericalPositionConstraint(Dt, SolverSettings, JointSettings);
+		}
+		else if (bLinearLimitHard[1] && bLinearLimitHard[2])
+		{
+			ApplyCylindricalPositionConstraint(Dt, 0, SolverSettings, JointSettings);
+		}
+		else if (bLinearLimitHard[0] && bLinearLimitHard[2])
+		{
+			ApplyCylindricalPositionConstraint(Dt, 1, SolverSettings, JointSettings);
+		}
+		else if (bLinearLimitHard[0] && bLinearLimitHard[1])
+		{
+			ApplyCylindricalPositionConstraint(Dt, 2, SolverSettings, JointSettings);
+		}
+		else if (bLinearLimitHard[0])
+		{
+			ApplyPlanarPositionConstraint(Dt, 0, SolverSettings, JointSettings);
+		}
+		else if (bLinearLimitHard[1])
+		{
+			ApplyPlanarPositionConstraint(Dt, 1, SolverSettings, JointSettings);
+		}
+		else if (bLinearLimitHard[2])
+		{
+			ApplyPlanarPositionConstraint(Dt, 2, SolverSettings, JointSettings);
+		}
+
+		// Handle soft limits
+		if (bLinearLimitSoft[0] && bLinearLimitSoft[1] && bLinearLimitSoft[2])
+		{
+			ApplySphericalPositionConstraintSoft(Dt, SolverSettings, JointSettings);
+		}
+		else if (bLinearLimitSoft[1] && bLinearLimitSoft[2])
+		{
+			ApplyCylindricalPositionConstraintSoft(Dt, 0, SolverSettings, JointSettings);
+		}
+		else if (bLinearLimitSoft[0] && bLinearLimitSoft[2])
+		{
+			ApplyCylindricalPositionConstraintSoft(Dt, 1, SolverSettings, JointSettings);
+		}
+		else if (bLinearLimitSoft[0] && bLinearLimitSoft[1])
+		{
+			ApplyCylindricalPositionConstraintSoft(Dt, 2, SolverSettings, JointSettings);
+		}
+		else if (bLinearLimitSoft[0])
+		{
+			ApplyPlanarPositionConstraintSoft(Dt, 0, SolverSettings, JointSettings);
+		}
+		else if (bLinearLimitSoft[1])
+		{
+			ApplyPlanarPositionConstraintSoft(Dt, 1, SolverSettings, JointSettings);
+		}
+		else if (bLinearLimitSoft[2])
 		{
-			ApplyPositionConstraint(Dt, SolverSettings, JointSettings);
+			ApplyPlanarPositionConstraintSoft(Dt, 2, SolverSettings, JointSettings);
+		}
+	}
+
+
+	void FJointSolverGaussSeidel::ApplyPositionDrives(
+		const FReal Dt,
+		const FPBDJointSolverSettings& SolverSettings,
+		const FPBDJointSettings& JointSettings)
+	{
+		if (SolverSettings.bEnableDrives)
+		{
+			TVector<bool, 3> bDriven =
+			{
+				(JointSettings.bLinearPositionDriveEnabled[0] || JointSettings.bLinearVelocityDriveEnabled[0]) && (JointSettings.LinearMotionTypes[0] != EJointMotionType::Locked),
+				(JointSettings.bLinearPositionDriveEnabled[1] || JointSettings.bLinearVelocityDriveEnabled[1]) && (JointSettings.LinearMotionTypes[1] != EJointMotionType::Locked),
+				(JointSettings.bLinearPositionDriveEnabled[2] || JointSettings.bLinearVelocityDriveEnabled[2]) && (JointSettings.LinearMotionTypes[2] != EJointMotionType::Locked),
+			};
+
+			if (bDriven[0] && bDriven[1] && bDriven[2])
+			{
+				ApplySphericalPositionDrive(Dt, SolverSettings, JointSettings);
+			}
+			else if (bDriven[1] && bDriven[2])
+			{
+				ApplyCylindricalPositionDrive(Dt, 0, SolverSettings, JointSettings);
+			}
+			else if (bDriven[0] && bDriven[2])
+			{
+				ApplyCylindricalPositionDrive(Dt, 1, SolverSettings, JointSettings);
+			}
+			else if (bDriven[0] && bDriven[1])
+			{
+				ApplyCylindricalPositionDrive(Dt, 2, SolverSettings, JointSettings);
+			}
+			else if (bDriven[0])
+			{
+				ApplyPlanarPositionDrive(Dt, 0, SolverSettings, JointSettings);
+			}
+			else if (bDriven[1])
+			{
+				ApplyPlanarPositionDrive(Dt, 1, SolverSettings, JointSettings);
+			}
+			else if (bDriven[2])
+			{
+				ApplyPlanarPositionDrive(Dt, 2, SolverSettings, JointSettings);
+			}
 		}
 	}
 
@@ -453,12 +707,42 @@ namespace Chaos
 	//
 	//
 
+	void FJointSolverGaussSeidel::ApplyAccumulatedDeltas()
+	{
+		if (bChaos_Joint_Accumulate)
+		{
+			Ps[0] += DPs[0];
+			Ps[1] += DPs[1];
+
+			const FRotation3 DQ0 = (FRotation3::FromElements(DRs[0], 0) * Qs[0]) * (FReal)0.5;
+			const FRotation3 DQ1 = (FRotation3::FromElements(DRs[1], 0) * Qs[1]) * (FReal)0.5;
+			Qs[0] = (Qs[0] + DQ0).GetNormalized();
+			Qs[1] = (Qs[1] + DQ1).GetNormalized();
+			Qs[1].EnforceShortestArcWith(Qs[0]);
+
+			Vs[0] += DVs[0];
+			Vs[1] += DVs[1];
+
+			Ws[0] += DWs[0];
+			Ws[1] += DWs[1];
+
+			UpdateDerivedState();
+		}
+	}
+
 	void FJointSolverGaussSeidel::ApplyPositionDelta(
 		const int32 BodyIndex,
 		const FReal Stiffness,
 		const FVec3& DP)
 	{
-		Ps[BodyIndex] = Ps[BodyIndex] + Stiffness * DP;
+		if (bChaos_Joint_Accumulate)
+		{
+			DPs[BodyIndex] += Stiffness * DP;
+		}
+		else
+		{
+			Ps[BodyIndex] += Stiffness * DP;
+		}
 	}
 
 
@@ -467,8 +751,16 @@ namespace Chaos
 		const FVec3& DP0,
 		const FVec3& DP1)
 	{
-		Ps[0] = Ps[0] + Stiffness * DP0;
-		Ps[1] = Ps[1] + Stiffness * DP1;
+		if (bChaos_Joint_Accumulate)
+		{
+			DPs[0] += Stiffness * DP0;
+			DPs[1] += Stiffness * DP1;
+		}
+		else
+		{
+			Ps[0] += Stiffness * DP0;
+			Ps[1] += Stiffness * DP1;
+		}
 	}
 
 
@@ -477,8 +769,15 @@ namespace Chaos
 		const FReal Stiffness,
 		const FVec3& DR)
 	{
-		const FRotation3 DQ = (FRotation3::FromElements(Stiffness * DR, 0) * Qs[BodyIndex]) * (FReal)0.5;
-		Qs[BodyIndex] = (Qs[BodyIndex] + DQ).GetNormalized();
+		if (bChaos_Joint_Accumulate)
+		{
+			DRs[BodyIndex] = Stiffness * DR;
+		}
+		else
+		{
+			const FRotation3 DQ = (FRotation3::FromElements(Stiffness * DR, 0) * Qs[BodyIndex]) * (FReal)0.5;
+			Qs[BodyIndex] = (Qs[BodyIndex] + DQ).GetNormalized();
+		}
 	}
 
 
@@ -487,11 +786,19 @@ namespace Chaos
 		const FVec3& DR0,
 		const FVec3& DR1)
 	{
-		const FRotation3 DQ0 = (FRotation3::FromElements(Stiffness * DR0, 0) * Qs[0]) * (FReal)0.5;
-		const FRotation3 DQ1 = (FRotation3::FromElements(Stiffness * DR1, 0) * Qs[1]) * (FReal)0.5;
-		Qs[0] = (Qs[0] + DQ0).GetNormalized();
-		Qs[1] = (Qs[1] + DQ1).GetNormalized();
-		Qs[1].EnforceShortestArcWith(Qs[0]);
+		if (bChaos_Joint_Accumulate)
+		{
+			DRs[0] += Stiffness * DR0;
+			DRs[1] += Stiffness * DR1;
+		}
+		else
+		{
+			const FRotation3 DQ0 = (FRotation3::FromElements(Stiffness * DR0, 0) * Qs[0]) * (FReal)0.5;
+			const FRotation3 DQ1 = (FRotation3::FromElements(Stiffness * DR1, 0) * Qs[1]) * (FReal)0.5;
+			Qs[0] = (Qs[0] + DQ0).GetNormalized();
+			Qs[1] = (Qs[1] + DQ1).GetNormalized();
+			Qs[1].EnforceShortestArcWith(Qs[0]);
+		}
 	}
 
 
@@ -501,8 +808,16 @@ namespace Chaos
 		const FVec3& DV,
 		const FVec3& DW)
 	{
-		Vs[BodyIndex] = Vs[BodyIndex] + Stiffness * DV;
-		Ws[BodyIndex] = Ws[BodyIndex] + Stiffness * DW;
+		if (bChaos_Joint_Accumulate)
+		{
+			DVs[BodyIndex] += Stiffness * DV;
+			DWs[BodyIndex] += Stiffness * DW;
+		}
+		else
+		{
+			Vs[BodyIndex] = Vs[BodyIndex] + Stiffness * DV;
+			Ws[BodyIndex] = Ws[BodyIndex] + Stiffness * DW;
+		}
 	}
 
 
@@ -513,14 +828,131 @@ namespace Chaos
 		const FVec3& DV1,
 		const FVec3& DW1)
 	{
-		Vs[0] = Vs[0] + Stiffness * DV0;
-		Ws[0] = Ws[0] + Stiffness * DW0;
-		Vs[1] = Vs[1] + Stiffness * DV1;
-		Ws[1] = Ws[1] + Stiffness * DW1;
+		if (bChaos_Joint_Accumulate)
+		{
+			DVs[0] += Stiffness * DV0;
+			DVs[1] += Stiffness * DV1;
+			DWs[0] += Stiffness * DW0;
+			DWs[1] += Stiffness * DW1;
+		}
+		else
+		{
+			Vs[0] += Stiffness * DV0;
+			Vs[1] += Stiffness * DV1;
+			Ws[0] += Stiffness * DW0;
+			Ws[1] += Stiffness * DW1;
+		}
+	}
+
+
+	void FJointSolverGaussSeidel::ApplyPositionConstraint(
+		const FReal Stiffness,
+		const FVec3& Axis,
+		const FReal Delta)
+	{
+		const FVec3 AngularAxis0 = FVec3::CrossProduct(Xs[0] - Ps[0], Axis);
+		const FVec3 AngularAxis1 = FVec3::CrossProduct(Xs[1] - Ps[1], Axis);
+		const FMatrix33 InvI0 = Utilities::ComputeWorldSpaceInertia(Qs[0], InvILs[0]);
+		const FMatrix33 InvI1 = Utilities::ComputeWorldSpaceInertia(Qs[1], InvILs[1]);
+		const FVec3 IA0 = Utilities::Multiply(InvI0, AngularAxis0);
+		const FVec3 IA1 = Utilities::Multiply(InvI1, AngularAxis1);
+
+		// Joint-space inverse mass
+		const FReal II0 = FVec3::DotProduct(AngularAxis0, IA0);
+		const FReal II1 = FVec3::DotProduct(AngularAxis1, IA1);
+		const FReal IM = InvMs[0] + II0 + InvMs[1] + II1;
+
+		const FVec3 DX = Axis * Delta / IM;
+
+		// Apply constraint correction
+		const FVec3 DP0 = InvMs[0] * DX;
+		const FVec3 DP1 = -InvMs[1] * DX;
+		const FVec3 DR0 = Utilities::Multiply(InvI0, FVec3::CrossProduct(Xs[0] - Ps[0], DX));
+		const FVec3 DR1 = Utilities::Multiply(InvI1, FVec3::CrossProduct(Xs[1] - Ps[1], -DX));
+
+		ApplyPositionDelta(Stiffness, DP0, DP1);
+		ApplyRotationDelta(Stiffness, DR0, DR1);
+		UpdateDerivedState();
 	}
 
 
-	void FJointSolverGaussSeidel::ApplyRotationCorrection(
+	void FJointSolverGaussSeidel::ApplyPositionConstraintSoft(
+		const FReal Dt,
+		const FReal Stiffness,
+		const FReal Damping,
+		const bool bAccelerationMode,
+		const FVec3& Axis,
+		const FReal Delta,
+		FReal& Lambda)
+	{
+		// World-space inverse mass
+		const FMatrix33 InvI0 = Utilities::ComputeWorldSpaceInertia(Qs[0], InvILs[0]);
+		const FMatrix33 InvI1 = Utilities::ComputeWorldSpaceInertia(Qs[1], InvILs[1]);
+
+		// Joint-space inverse mass
+		const FVec3 AngularAxis0 = FVec3::CrossProduct(Xs[0] - Ps[0], Axis);
+		const FVec3 AngularAxis1 = FVec3::CrossProduct(Xs[1] - Ps[1], Axis);
+		const FVec3 IA0 = Utilities::Multiply(InvI0, AngularAxis0);
+		const FVec3 IA1 = Utilities::Multiply(InvI1, AngularAxis1);
+		const FReal II0 = FVec3::DotProduct(AngularAxis0, IA0);
+		const FReal II1 = FVec3::DotProduct(AngularAxis1, IA1);
+		const FReal II = (InvMs[0] + II0 + InvMs[1] + II1);
+
+		FReal VelDt = 0;
+		if (Damping > KINDA_SMALL_NUMBER)
+		{
+			const FVec3 V0 = FVec3::CalculateVelocity(PrevXs[0], Xs[0], (FReal)1.0);
+			const FVec3 V1 = FVec3::CalculateVelocity(PrevXs[1], Xs[1], (FReal)1.0);
+			VelDt = FVec3::DotProduct(V0 - V1, Axis);
+		}
+
+		FReal DLambda = 0;
+		if (Stiffness > KINDA_SMALL_NUMBER)
+		{
+			// As below, but numerically stable for large values of stiffness (same form as in XPBD paper)
+			const FReal Alpha = (FReal)1 / (Stiffness * Dt * Dt);
+			const FReal Gamma = Damping / (Stiffness * Dt);
+			if (bAccelerationMode)
+			{
+				const FReal Multiplier = (FReal)1 / (((FReal)1 + Gamma) + Alpha);
+				DLambda = Multiplier * (Delta - Gamma * VelDt) / II - Multiplier * Alpha * Lambda;
+			}
+			else
+			{
+				const FReal Multiplier = (FReal)1 / (((FReal)1 + Gamma) * II + Alpha);
+				DLambda = Multiplier * (Delta - Gamma * VelDt - Alpha * Lambda);
+			}
+		}
+		else
+		{
+			// As above, but numerically stable at low stiffness (alpha -> infinity)
+			const FReal S = Stiffness * Dt * Dt;
+			const FReal D = Damping * Dt;
+			if (bAccelerationMode)
+			{
+				const FReal Multiplier = (FReal)1 / ((S + D) + (FReal)1);
+				DLambda = Multiplier * (S * Delta - D * VelDt) / II - Multiplier * Lambda;
+			}
+			else
+			{
+				const FReal Multiplier = (FReal)1 / ((S + D) * II + (FReal)1);
+				DLambda = Multiplier * (S * Delta - D * VelDt - Lambda);
+			}
+		}
+	
+		const FVec3 DP0 = (InvMs[0] * DLambda) * Axis;
+		const FVec3 DP1 = (-InvMs[1] * DLambda) * Axis;
+		const FVec3 DR0 = DLambda * Utilities::Multiply(InvI0, AngularAxis0);
+		const FVec3 DR1 = -DLambda * Utilities::Multiply(InvI1, AngularAxis1);
+
+		Lambda += DLambda;
+		ApplyPositionDelta((FReal)1, DP0, DP1);
+		ApplyRotationDelta((FReal)1, DR0, DR1);
+		UpdateDerivedState();
+	}
+	
+	
+	void FJointSolverGaussSeidel::ApplyRotationConstraint(
 		const FReal Stiffness,
 		const FVec3& Axis0,
 		const FVec3& Axis1,
@@ -542,19 +974,23 @@ namespace Chaos
 		//const FVec3 DR1 = Axis1 * -(Angle * II1 / (II0 + II1));
 
 		// @todo(ccaulfield): this position correction needs to have components in direction of inactive position constraints removed
-		const FReal IM0 = InvMs[0];
-		const FReal IM1 = InvMs[1];
-		const FVec3 DX = FVec3::CrossProduct(DR0, Xs[0] - Ps[0]) + FVec3::CrossProduct(DR1, Xs[1] - Ps[1]);
-		const FVec3 DP0 = DX * (IM0 / (IM0 + IM1));
-		const FVec3 DP1 = DX * (-IM1 / (IM0 + IM1));
+		if (bEnableAngularConstraintPositionCorrection)
+		{
+			const FReal IM0 = InvMs[0];
+			const FReal IM1 = InvMs[1];
+			const FVec3 DX = FVec3::CrossProduct(DR0, Xs[0] - Ps[0]) + FVec3::CrossProduct(DR1, Xs[1] - Ps[1]);
+			const FVec3 DP0 = DX * (IM0 / (IM0 + IM1));
+			const FVec3 DP1 = DX * (-IM1 / (IM0 + IM1));
+			ApplyPositionDelta(Stiffness, DP0, DP1);
+		}
 
-		ApplyPositionDelta(Stiffness, DP0, DP1);
 		ApplyRotationDelta(Stiffness, DR0, DR1);
+		UpdateDerivedState();
 	}
 
 
 	// See "XPBD: Position-Based Simulation of Compliant Constrained Dynamics"
-	void FJointSolverGaussSeidel::ApplyRotationCorrectionSoft(
+	void FJointSolverGaussSeidel::ApplyRotationConstraintSoft(
 		const FReal Dt,
 		const FReal Stiffness,
 		const FReal Damping,
@@ -624,15 +1060,19 @@ namespace Chaos
 		const FVec3 DR1 = Axis1 * -(DLambda * II1);
 
 		// @todo(ccaulfield): this position correction needs to have components in direction of inactive position constraints removed
-		const FReal IM0 = InvMs[0];
-		const FReal IM1 = InvMs[1];
-		const FVec3 DX = FVec3::CrossProduct(DR0, Xs[0] - Ps[0]) + FVec3::CrossProduct(DR1, Xs[1] - Ps[1]);
-		const FVec3 DP0 = DX * (IM0 / (IM0 + IM1));
-		const FVec3 DP1 = DX * (-IM1 / (IM0 + IM1));
+		if (bEnableAngularConstraintPositionCorrection)
+		{
+			const FReal IM0 = InvMs[0];
+			const FReal IM1 = InvMs[1];
+			const FVec3 DX = FVec3::CrossProduct(DR0, Xs[0] - Ps[0]) + FVec3::CrossProduct(DR1, Xs[1] - Ps[1]);
+			const FVec3 DP0 = DX * (IM0 / (IM0 + IM1));
+			const FVec3 DP1 = DX * (-IM1 / (IM0 + IM1));
+			ApplyPositionDelta(1.0f, DP0, DP1);
+		}
 
-		ApplyPositionDelta(1.0f, DP0, DP1);
-		ApplyRotationDelta(1.0f, DR0, DR1);
 		Lambda += DLambda;
+		ApplyRotationDelta(1.0f, DR0, DR1);
+		UpdateDerivedState();
 	}
 
 
@@ -680,8 +1120,7 @@ namespace Chaos
 		if (DTwistAngle != 0)
 		{
 			FReal TwistStiffness = FPBDJointUtilities::GetTwistStiffness(SolverSettings, JointSettings);
-			ApplyRotationCorrection(TwistStiffness, TwistAxis0, TwistAxis1, DTwistAngle);
-			UpdateDerivedState();
+			ApplyRotationConstraint(TwistStiffness, TwistAxis0, TwistAxis1, DTwistAngle);
 		}
 	}
 
@@ -724,9 +1163,8 @@ namespace Chaos
 		{
 			const FReal TwistStiffness = FPBDJointUtilities::GetSoftTwistStiffness(SolverSettings, JointSettings);
 			const FReal TwistDamping = FPBDJointUtilities::GetSoftTwistDamping(SolverSettings, JointSettings);
-			const bool bAccelerationMode = FPBDJointUtilities::GetSoftAccelerationMode(SolverSettings, JointSettings);
-			ApplyRotationCorrectionSoft(Dt, TwistStiffness, TwistDamping, bAccelerationMode, TwistAxis0, TwistAxis1, DTwistAngle, TwistSoftLambda);
-			UpdateDerivedState();
+			const bool bAccelerationMode = FPBDJointUtilities::GetAngularSoftAccelerationMode(SolverSettings, JointSettings);
+			ApplyRotationConstraintSoft(Dt, TwistStiffness, TwistDamping, bAccelerationMode, TwistAxis0, TwistAxis1, DTwistAngle, TwistSoftLambda);
 		}
 	}
 
@@ -760,8 +1198,7 @@ namespace Chaos
 		const FReal AngularDriveStiffness = FPBDJointUtilities::GetAngularTwistDriveStiffness(SolverSettings, JointSettings);
 		const FReal AngularDriveDamping = FPBDJointUtilities::GetAngularTwistDriveDamping(SolverSettings, JointSettings);
 		const bool bAccelerationMode = FPBDJointUtilities::GetDriveAccelerationMode(SolverSettings, JointSettings);
-		ApplyRotationCorrectionSoft(Dt, AngularDriveStiffness, AngularDriveDamping, bAccelerationMode, TwistAxis0, TwistAxis1, DTwistAngle, TwistDriveLambda);
-		UpdateDerivedState();
+		ApplyRotationConstraintSoft(Dt, AngularDriveStiffness, AngularDriveDamping, bAccelerationMode, TwistAxis0, TwistAxis1, DTwistAngle, TwistDriveLambda);
 	}
 
 
@@ -770,7 +1207,7 @@ namespace Chaos
 		const FPBDJointSolverSettings& SolverSettings,
 		const FPBDJointSettings& JointSettings)
 	{
-		// @todo(ccaulfield): implement swing projection
+		// @todo(ccaulfield): implement twist projection
 	}
 
 
@@ -792,8 +1229,7 @@ namespace Chaos
 		if (DSwingAngle != 0)
 		{
 			FReal SwingStiffness = FPBDJointUtilities::GetSwingStiffness(SolverSettings, JointSettings);
-			ApplyRotationCorrection(SwingStiffness, SwingAxis, SwingAxis, DSwingAngle);
-			UpdateDerivedState();
+			ApplyRotationConstraint(SwingStiffness, SwingAxis, SwingAxis, DSwingAngle);
 		}
 	}
 
@@ -817,9 +1253,8 @@ namespace Chaos
 		{
 			const FReal SoftSwingStiffness = FPBDJointUtilities::GetSoftSwingStiffness(SolverSettings, JointSettings);
 			const FReal SoftSwingDamping = FPBDJointUtilities::GetSoftSwingDamping(SolverSettings, JointSettings);
-			const bool bAccelerationMode = FPBDJointUtilities::GetSoftAccelerationMode(SolverSettings, JointSettings);
-			ApplyRotationCorrectionSoft(Dt, SoftSwingStiffness, SoftSwingDamping, bAccelerationMode, SwingAxis, SwingAxis, DSwingAngle, SwingSoftLambda);
-			UpdateDerivedState();
+			const bool bAccelerationMode = FPBDJointUtilities::GetAngularSoftAccelerationMode(SolverSettings, JointSettings);
+			ApplyRotationConstraintSoft(Dt, SoftSwingStiffness, SoftSwingDamping, bAccelerationMode, SwingAxis, SwingAxis, DSwingAngle, SwingSoftLambda);
 		}
 	}
 
@@ -846,8 +1281,7 @@ namespace Chaos
 		const FReal AngularDriveStiffness = FPBDJointUtilities::GetAngularSwingDriveStiffness(SolverSettings, JointSettings);
 		const FReal AngularDriveDamping = FPBDJointUtilities::GetAngularSwingDriveDamping(SolverSettings, JointSettings);
 		const bool bAccelerationMode = FPBDJointUtilities::GetDriveAccelerationMode(SolverSettings, JointSettings);
-		ApplyRotationCorrectionSoft(Dt, AngularDriveStiffness, AngularDriveDamping, bAccelerationMode, SwingAxis, SwingAxis, DSwingAngle, SwingDriveLambda);
-		UpdateDerivedState();
+		ApplyRotationConstraintSoft(Dt, AngularDriveStiffness, AngularDriveDamping, bAccelerationMode, SwingAxis, SwingAxis, DSwingAngle, SwingDriveLambda);
 	}
 
 
@@ -899,7 +1333,6 @@ namespace Chaos
 			ApplyRotationDelta(1, AngularProjection, DR1);
 			ApplyPositionDelta(1, AngularProjection, DP1);
 			ApplyVelocityDelta(1, AngularProjection, DV1, DW1);
-
 			UpdateDerivedState(1);
 		}
 	}
@@ -934,8 +1367,7 @@ namespace Chaos
 		if (DSwingAngle != 0)
 		{
 			const FReal SwingStiffness = FPBDJointUtilities::GetSwingStiffness(SolverSettings, JointSettings);
-			ApplyRotationCorrection(SwingStiffness, SwingAxis, SwingAxis, DSwingAngle);
-			UpdateDerivedState();
+			ApplyRotationConstraint(SwingStiffness, SwingAxis, SwingAxis, DSwingAngle);
 		}
 	}
 
@@ -969,9 +1401,8 @@ namespace Chaos
 		{
 			const FReal SoftSwingStiffness = FPBDJointUtilities::GetSoftSwingStiffness(SolverSettings, JointSettings);
 			const FReal SoftSwingDamping = FPBDJointUtilities::GetSoftSwingDamping(SolverSettings, JointSettings);
-			const bool bAccelerationMode = FPBDJointUtilities::GetSoftAccelerationMode(SolverSettings, JointSettings);
-			ApplyRotationCorrectionSoft(Dt, SoftSwingStiffness, SoftSwingDamping, bAccelerationMode, SwingAxis, SwingAxis, DSwingAngle, SwingSoftLambda);
-			UpdateDerivedState();
+			const bool bAccelerationMode = FPBDJointUtilities::GetAngularSoftAccelerationMode(SolverSettings, JointSettings);
+			ApplyRotationConstraintSoft(Dt, SoftSwingStiffness, SoftSwingDamping, bAccelerationMode, SwingAxis, SwingAxis, DSwingAngle, SwingSoftLambda);
 		}
 	}
 
@@ -994,8 +1425,7 @@ namespace Chaos
 		const FReal AngularDriveStiffness = FPBDJointUtilities::GetAngularSwingDriveStiffness(SolverSettings, JointSettings);
 		const FReal AngularDriveDamping = FPBDJointUtilities::GetAngularSwingDriveDamping(SolverSettings, JointSettings);
 		const bool bAccelerationMode = FPBDJointUtilities::GetDriveAccelerationMode(SolverSettings, JointSettings);
-		ApplyRotationCorrectionSoft(Dt, AngularDriveStiffness, AngularDriveDamping, bAccelerationMode, SwingAxis, SwingAxis, DSwingAngle, SwingDriveLambda);
-		UpdateDerivedState();
+		ApplyRotationConstraintSoft(Dt, AngularDriveStiffness, AngularDriveDamping, bAccelerationMode, SwingAxis, SwingAxis, DSwingAngle, SwingDriveLambda);
 	}
 
 
@@ -1031,19 +1461,17 @@ namespace Chaos
 			const FReal AngularDriveStiffness = FPBDJointUtilities::GetAngularSLerpDriveStiffness(SolverSettings, JointSettings);
 			const FReal AngularDriveDamping = FPBDJointUtilities::GetAngularSLerpDriveDamping(SolverSettings, JointSettings);
 			const bool bAccelerationMode = FPBDJointUtilities::GetDriveAccelerationMode(SolverSettings, JointSettings);
-			ApplyRotationCorrectionSoft(Dt, AngularDriveStiffness, AngularDriveDamping, bAccelerationMode, SLerpAxis, SLerpAxis, -SLerpAngle, SwingDriveLambda);
+			ApplyRotationConstraintSoft(Dt, AngularDriveStiffness, AngularDriveDamping, bAccelerationMode, SLerpAxis, SLerpAxis, -SLerpAngle, SwingDriveLambda);
 		}
 	}
 
 
-	void FJointSolverGaussSeidel::ApplyPositionConstraint(
+	void FJointSolverGaussSeidel::ApplyPointPositionConstraint(
 		const FReal Dt,
 		const FPBDJointSolverSettings& SolverSettings,
 		const FPBDJointSettings& JointSettings)
 	{
-		// @todo(ccaulfield): we should really be calculating axes based on joint config, rather than fixing the error components
-		// Calculate constraint error
-		const FVec3 CX = FPBDJointUtilities::GetLimitedPositionError(JointSettings, Rs[0], Xs[1] - Xs[0]);
+		const FVec3 CX = Xs[1] - Xs[0];
 		if (CX != FVec3(0))
 		{
 			const FMatrix33 InvI0 = Utilities::ComputeWorldSpaceInertia(Qs[0], InvILs[0]);
@@ -1077,55 +1505,173 @@ namespace Chaos
 	}
 
 
-	void FJointSolverGaussSeidel::ApplyPositionDrives(
+	void FJointSolverGaussSeidel::ApplySphericalPositionConstraint(
 		const FReal Dt,
 		const FPBDJointSolverSettings& SolverSettings,
 		const FPBDJointSettings& JointSettings)
 	{
-		// @todo(ccaulfield): XPBD for linear drives
-		if (JointSettings.bLinearDriveEnabled[0] || JointSettings.bLinearDriveEnabled[1] || JointSettings.bLinearDriveEnabled[2])
+		FVec3 Axis;
+		FReal Delta;
+		GetSphericalAxisDelta(Xs[0], Xs[1], Axis, Delta);
+		const FReal Error = GetSphericalError(Delta, JointSettings.LinearLimit);
+		if (Error != 0)
 		{
-			const FVec3 TargetX = Xs[0] + Rs[0] * JointSettings.LinearDriveTarget;
-			FVec3 CX = Xs[1] - TargetX;
+			const FReal Stiffness = FPBDJointUtilities::GetLinearStiffness(SolverSettings, JointSettings);
+			ApplyPositionConstraint(Stiffness, Axis, Error);
+		}
+	}
 
-			const FMatrix33 RM0 = Rs[0].ToMatrix();
-			for (int32 IAxis = 0; IAxis < 3; ++IAxis)
-			{
-				if (!JointSettings.bLinearDriveEnabled[IAxis])
-				{
-					const FVec3 Axis = RM0.GetAxis(IAxis);
-					CX = CX - FVec3::DotProduct(Axis, CX) * Axis;
-				}
-			}
 
-			const FMatrix33 InvI0 = Utilities::ComputeWorldSpaceInertia(Qs[0], InvILs[0]);
-			const FMatrix33 InvI1 = Utilities::ComputeWorldSpaceInertia(Qs[1], InvILs[1]);
+	void FJointSolverGaussSeidel::ApplySphericalPositionConstraintSoft(
+		const FReal Dt,
+		const FPBDJointSolverSettings& SolverSettings,
+		const FPBDJointSettings& JointSettings)
+	{
+		FVec3 Axis;
+		FReal Delta;
+		GetSphericalAxisDelta(Xs[0], Xs[1], Axis, Delta);
+		const FReal Error = GetSphericalError(Delta, JointSettings.LinearLimit);
+		if (Error != 0)
+		{
+			const FReal Stiffness = FPBDJointUtilities::GetSoftLinearStiffness(SolverSettings, JointSettings);
+			const FReal Damping = FPBDJointUtilities::GetSoftLinearDamping(SolverSettings, JointSettings);
+			const bool bAccelerationMode = FPBDJointUtilities::GetLinearSoftAccelerationMode(SolverSettings, JointSettings);
+			ApplyPositionConstraintSoft(Dt, Stiffness, Damping, bAccelerationMode, Axis, Error, LinearSoftLambda);
+		}
+	}
 
-			// Calculate constraint correction
-			FMatrix33 M0 = FMatrix33(0, 0, 0);
-			FMatrix33 M1 = FMatrix33(0, 0, 0);
-			if (InvMs[0] > 0)
-			{
-				M0 = Utilities::ComputeJointFactorMatrix(Xs[0] - Ps[0], InvI0, InvMs[0]);
-			}
-			if (InvMs[1] > 0)
-			{
-				M1 = Utilities::ComputeJointFactorMatrix(Xs[1] - Ps[1], InvI1, InvMs[1]);
-			}
-			const FMatrix33 MI = (M0 + M1).Inverse();
-			const FVec3 DX = Utilities::Multiply(MI, CX);
 
-			// Apply constraint correction
-			const FVec3 DP0 = InvMs[0] * DX;
-			const FVec3 DP1 = -InvMs[1] * DX;
-			const FVec3 DR0 = Utilities::Multiply(InvI0, FVec3::CrossProduct(Xs[0] - Ps[0], DX));
-			const FVec3 DR1 = Utilities::Multiply(InvI1, FVec3::CrossProduct(Xs[1] - Ps[1], -DX));
+	void FJointSolverGaussSeidel::ApplySphericalPositionDrive(
+		const FReal Dt,
+		const FPBDJointSolverSettings& SolverSettings,
+		const FPBDJointSettings& JointSettings)
+	{
+		const FVec3 XTarget = Xs[0] + Rs[0] * JointSettings.LinearDriveTarget;
+		FVec3 Axis;
+		FReal Delta;
+		GetSphericalAxisDelta(XTarget, Xs[1], Axis, Delta);
+		if (Delta != 0)
+		{
+			const FReal Stiffness = FPBDJointUtilities::GetLinearDriveStiffness(SolverSettings, JointSettings);
+			const FReal Damping = FPBDJointUtilities::GetLinearDriveDamping(SolverSettings, JointSettings);
+			const bool bAccelerationMode = FPBDJointUtilities::GetDriveAccelerationMode(SolverSettings, JointSettings);
+			ApplyPositionConstraintSoft(Dt, Stiffness, Damping, bAccelerationMode, Axis, Delta, LinearDriveLambda);
+		}
+	}
 
-			const FReal LinearStiffness = FPBDJointUtilities::GetLinearDriveStiffness(SolverSettings, JointSettings);
-			const FReal LinearDamping = FPBDJointUtilities::GetLinearDriveDamping(SolverSettings, JointSettings);
-			ApplyPositionDelta(LinearStiffness, DP0, DP1);
-			ApplyRotationDelta(LinearStiffness, DR0, DR1);
-			UpdateDerivedState();
+
+	void FJointSolverGaussSeidel::ApplyCylindricalPositionConstraint(
+		const FReal Dt,
+		const int32 AxisIndex,
+		const FPBDJointSolverSettings& SolverSettings,
+		const FPBDJointSettings& JointSettings)
+	{
+		FVec3 Axis;
+		FReal Delta;
+		GetCylindricalAxisDelta(Rs[0], Xs[0], Xs[1], AxisIndex, Axis, Delta);
+		const FReal Error = GetCylindricalError(Delta, JointSettings.LinearLimit);
+		if (Error != 0)
+		{
+			const FReal Stiffness = FPBDJointUtilities::GetLinearStiffness(SolverSettings, JointSettings);
+			ApplyPositionConstraint(Stiffness, Axis, Error);
+		}
+	}
+
+
+	void FJointSolverGaussSeidel::ApplyCylindricalPositionConstraintSoft(
+		const FReal Dt,
+		const int32 AxisIndex,
+		const FPBDJointSolverSettings& SolverSettings,
+		const FPBDJointSettings& JointSettings)
+	{
+		FVec3 Axis;
+		FReal Delta;
+		GetCylindricalAxisDelta(Rs[0], Xs[0], Xs[1], AxisIndex, Axis, Delta);
+		const FReal Error = GetCylindricalError(Delta, JointSettings.LinearLimit);
+		if (Error != 0)
+		{
+			const FReal Stiffness = FPBDJointUtilities::GetSoftLinearStiffness(SolverSettings, JointSettings);
+			const FReal Damping = FPBDJointUtilities::GetSoftLinearDamping(SolverSettings, JointSettings);
+			const bool bAccelerationMode = FPBDJointUtilities::GetLinearSoftAccelerationMode(SolverSettings, JointSettings);
+			ApplyPositionConstraintSoft(Dt, Stiffness, Damping, bAccelerationMode, Axis, Error, LinearSoftLambda);
+		}
+	}
+
+
+	void FJointSolverGaussSeidel::ApplyCylindricalPositionDrive(
+		const FReal Dt,
+		const int32 AxisIndex,
+		const FPBDJointSolverSettings& SolverSettings,
+		const FPBDJointSettings& JointSettings)
+	{
+		const FVec3 XTarget = Xs[0] + Rs[0] * JointSettings.LinearDriveTarget;
+		FVec3 Axis;
+		FReal Delta;
+		GetCylindricalAxisDelta(Rs[0], XTarget, Xs[1], AxisIndex, Axis, Delta);
+		if (Delta != 0)
+		{
+			const FReal Stiffness = FPBDJointUtilities::GetLinearDriveStiffness(SolverSettings, JointSettings);
+			const FReal Damping = FPBDJointUtilities::GetLinearDriveDamping(SolverSettings, JointSettings);
+			const bool bAccelerationMode = FPBDJointUtilities::GetDriveAccelerationMode(SolverSettings, JointSettings);
+			ApplyPositionConstraintSoft(Dt, Stiffness, Damping, bAccelerationMode, Axis, Delta, LinearDriveLambda);
+		}
+	}
+
+
+	void FJointSolverGaussSeidel::ApplyPlanarPositionConstraint(
+		const FReal Dt,
+		const int32 AxisIndex,
+		const FPBDJointSolverSettings& SolverSettings,
+		const FPBDJointSettings& JointSettings)
+	{
+		FVec3 Axis;
+		FReal Delta;
+		GetPlanarAxisDelta(Rs[0], Xs[0], Xs[1], AxisIndex, Axis, Delta);
+		const FReal Error = GetPlanarError(Delta, JointSettings.LinearLimit);
+		if (Error != 0)
+		{
+			const FReal Stiffness = FPBDJointUtilities::GetLinearStiffness(SolverSettings, JointSettings);
+			ApplyPositionConstraint(Stiffness, Axis, Error);
+		}
+	}
+
+
+	void FJointSolverGaussSeidel::ApplyPlanarPositionConstraintSoft(
+		const FReal Dt,
+		const int32 AxisIndex,
+		const FPBDJointSolverSettings& SolverSettings,
+		const FPBDJointSettings& JointSettings)
+	{
+		FVec3 Axis;
+		FReal Delta;
+		GetPlanarAxisDelta(Rs[0], Xs[0], Xs[1], AxisIndex, Axis, Delta);
+		const FReal Error = GetPlanarError(Delta, JointSettings.LinearLimit);
+		if (Error != 0)
+		{
+			const FReal Stiffness = FPBDJointUtilities::GetSoftLinearStiffness(SolverSettings, JointSettings);
+			const FReal Damping = FPBDJointUtilities::GetSoftLinearDamping(SolverSettings, JointSettings);
+			const bool bAccelerationMode = FPBDJointUtilities::GetLinearSoftAccelerationMode(SolverSettings, JointSettings);
+			ApplyPositionConstraintSoft(Dt, Stiffness, Damping, bAccelerationMode, Axis, Error, LinearSoftLambda);
+		}
+	}
+
+
+	void FJointSolverGaussSeidel::ApplyPlanarPositionDrive(
+		const FReal Dt,
+		const int32 AxisIndex,
+		const FPBDJointSolverSettings& SolverSettings,
+		const FPBDJointSettings& JointSettings)
+	{
+		const FVec3 XTarget = Xs[0] + Rs[0] * JointSettings.LinearDriveTarget;
+		FVec3 Axis;
+		FReal Delta;
+		GetPlanarAxisDelta(Rs[0], XTarget, Xs[1], AxisIndex, Axis, Delta);
+		if (Delta != 0)
+		{
+			const FReal Stiffness = FPBDJointUtilities::GetLinearDriveStiffness(SolverSettings, JointSettings);
+			const FReal Damping = FPBDJointUtilities::GetLinearDriveDamping(SolverSettings, JointSettings);
+			const bool bAccelerationMode = FPBDJointUtilities::GetDriveAccelerationMode(SolverSettings, JointSettings);
+			ApplyPositionConstraintSoft(Dt, Stiffness, Damping, bAccelerationMode, Axis, Delta, LinearDriveLambda);
 		}
 	}
 
diff --git a/Engine/Source/Runtime/Experimental/Chaos/Private/Chaos/PBDJointConstraintUtilities.cpp b/Engine/Source/Runtime/Experimental/Chaos/Private/Chaos/PBDJointConstraintUtilities.cpp
--- a/Engine/Source/Runtime/Experimental/Chaos/Private/Chaos/PBDJointConstraintUtilities.cpp
+++ b/Engine/Source/Runtime/Experimental/Chaos/Private/Chaos/PBDJointConstraintUtilities.cpp
@@ -20,11 +20,21 @@ namespace Chaos
 		const FPBDJointSolverSettings& SolverSettings,
 		const FPBDJointSettings& JointSettings)
 	{
-		const FReal SolverStiffness = (SolverSettings.Stiffness > (FReal)0) ? SolverSettings.Stiffness : JointSettings.Stiffness;
-		const FReal SoftSolverStiffness = (SolverSettings.SoftLinearStiffness > (FReal)0) ? SolverSettings.SoftLinearStiffness : JointSettings.SoftLinearStiffness;
-		const bool bIsSoft = JointSettings.bSoftLinearLimitsEnabled && ((JointSettings.LinearMotionTypes[0] == EJointMotionType::Limited) || (JointSettings.LinearMotionTypes[1] == EJointMotionType::Limited) || (JointSettings.LinearMotionTypes[2] == EJointMotionType::Limited));
-		const FReal Stiffness = bIsSoft ? SolverStiffness * SoftSolverStiffness : SolverStiffness;
-		return Stiffness;
+		return (SolverSettings.Stiffness > (FReal)0) ? SolverSettings.Stiffness : JointSettings.Stiffness;
+	}
+
+	FReal FPBDJointUtilities::GetSoftLinearStiffness(
+		const FPBDJointSolverSettings& SolverSettings,
+		const FPBDJointSettings& JointSettings)
+	{
+		return (SolverSettings.SoftLinearStiffness > (FReal)0) ? SolverSettings.SoftLinearStiffness : JointSettings.SoftLinearStiffness;
+	}
+
+	FReal FPBDJointUtilities::GetSoftLinearDamping(
+		const FPBDJointSolverSettings& SolverSettings,
+		const FPBDJointSettings& JointSettings)
+	{
+		return (SolverSettings.SoftLinearDamping > (FReal)0) ? SolverSettings.SoftLinearDamping : JointSettings.SoftLinearDamping;
 	}
 
 	FReal FPBDJointUtilities::GetTwistStiffness(
@@ -163,11 +173,18 @@ namespace Chaos
 		return (SolverSettings.AngularProjection > 0.0f) ? SolverSettings.AngularProjection : JointSettings.AngularProjection;
 	}
 
-	bool FPBDJointUtilities::GetSoftAccelerationMode(
+	bool FPBDJointUtilities::GetLinearSoftAccelerationMode(
 		const FPBDJointSolverSettings& SolverSettings,
 		const FPBDJointSettings& JointSettings)
 	{
-		return JointSettings.SoftForceMode == EJointForceMode::Acceleration;
+		return JointSettings.LinearSoftForceMode == EJointForceMode::Acceleration;
+	}
+
+	bool FPBDJointUtilities::GetAngularSoftAccelerationMode(
+		const FPBDJointSolverSettings& SolverSettings,
+		const FPBDJointSettings& JointSettings)
+	{
+		return JointSettings.AngularSoftForceMode == EJointForceMode::Acceleration;
 	}
 
 	bool FPBDJointUtilities::GetDriveAccelerationMode(
@@ -346,7 +363,15 @@ namespace Chaos
 
 	FVec3 FPBDJointUtilities::GetLimitedPositionError(const FPBDJointSettings& JointSettings, const FRotation3& R0, const FVec3& InCX)
 	{
-		const TVector<EJointMotionType, 3>& Motion = JointSettings.LinearMotionTypes;
+		// This function is only used for projection and is only relevant for hard limits.
+		// Treat soft-limits as free for error calculation.
+		const TVector<EJointMotionType, 3>& Motion =
+		{
+			((JointSettings.LinearMotionTypes[0] == EJointMotionType::Limited) && JointSettings.bSoftLinearLimitsEnabled) ? EJointMotionType::Free : JointSettings.LinearMotionTypes[0],
+			((JointSettings.LinearMotionTypes[1] == EJointMotionType::Limited) && JointSettings.bSoftLinearLimitsEnabled) ? EJointMotionType::Free : JointSettings.LinearMotionTypes[1],
+			((JointSettings.LinearMotionTypes[2] == EJointMotionType::Limited) && JointSettings.bSoftLinearLimitsEnabled) ? EJointMotionType::Free : JointSettings.LinearMotionTypes[2],
+		};
+
 		if ((Motion[0] == EJointMotionType::Locked) && (Motion[1] == EJointMotionType::Locked) && (Motion[2] == EJointMotionType::Locked))
 		{
 			return InCX;
diff --git a/Engine/Source/Runtime/Experimental/Chaos/Private/Chaos/PBDJointConstraints.cpp b/Engine/Source/Runtime/Experimental/Chaos/Private/Chaos/PBDJointConstraints.cpp
--- a/Engine/Source/Runtime/Experimental/Chaos/Private/Chaos/PBDJointConstraints.cpp
+++ b/Engine/Source/Runtime/Experimental/Chaos/Private/Chaos/PBDJointConstraints.cpp
@@ -88,7 +88,8 @@ namespace Chaos
 		, bSoftLinearLimitsEnabled(false)
 		, bSoftTwistLimitsEnabled(false)
 		, bSoftSwingLimitsEnabled(false)
-		, SoftForceMode(EJointForceMode::Acceleration)
+		, LinearSoftForceMode(EJointForceMode::Acceleration)
+		, AngularSoftForceMode(EJointForceMode::Acceleration)
 		, SoftLinearStiffness(0)
 		, SoftLinearDamping(0)
 		, SoftTwistStiffness(0)
@@ -96,7 +97,8 @@ namespace Chaos
 		, SoftSwingStiffness(0)
 		, SoftSwingDamping(0)
 		, LinearDriveTarget(FVec3(0, 0, 0))
-		, bLinearDriveEnabled(TVector<bool, 3>( false, false, false ))
+		, bLinearPositionDriveEnabled(TVector<bool, 3>(false, false, false))
+		, bLinearVelocityDriveEnabled(TVector<bool, 3>(false, false, false))
 		, LinearDriveForceMode(EJointForceMode::Acceleration)
 		, LinearDriveStiffness(0)
 		, LinearDriveDamping(0)
@@ -118,7 +120,7 @@ namespace Chaos
 
 	void FPBDJointSettings::Sanitize()
 	{
-		// Reset limits if they won;t be used (means we don't have to check if limited/locked in a few cases).
+		// Reset limits if they won't be used (means we don't have to check if limited/locked in a few cases).
 		// A side effect: if we enable a constraint, we need to reset the value of the limit.
 		if ((LinearMotionTypes[0] != EJointMotionType::Limited) && (LinearMotionTypes[1] != EJointMotionType::Limited) && (LinearMotionTypes[2] != EJointMotionType::Limited))
 		{
@@ -138,6 +140,7 @@ namespace Chaos
 		}
 
 		// Don't use soft constraints if the stiffness is very high
+		// (Should really use mass if Force Mode is enabled - which is isn't for anything atm).
 		if (bSoftLinearLimitsEnabled && SoftLinearStiffness > 1000)
 		{
 			bSoftLinearLimitsEnabled = false;
@@ -150,15 +153,6 @@ namespace Chaos
 		{
 			bSoftSwingLimitsEnabled = false;
 		}
-
-		// Currently we do not support angular limits if linear limits are set up
-		// @todo(ccaulfield): fix angular constraint position correction in FJointSolverGaussSeidel::ApplyRotationCorrection and FJointSolverGaussSeidel::ApplyRotationCorrectionSoft
-		if ((LinearMotionTypes[0] == EJointMotionType::Limited) || (LinearMotionTypes[1] == EJointMotionType::Limited) || (LinearMotionTypes[2] == EJointMotionType::Limited))
-		{
-			AngularMotionTypes[(int32)EJointAngularConstraintIndex::Twist] = EJointMotionType::Free;
-			AngularMotionTypes[(int32)EJointAngularConstraintIndex::Swing1] = EJointMotionType::Free;
-			AngularMotionTypes[(int32)EJointAngularConstraintIndex::Swing2] = EJointMotionType::Free;
-		}
 	}
 
 	
@@ -460,6 +454,8 @@ namespace Chaos
 					Dt,
 					Settings,
 					JointSettings,
+					FParticleUtilitiesXR::GetCoMWorldPosition(Particle0),	// Prev position
+					FParticleUtilitiesXR::GetCoMWorldPosition(Particle1),	// Prev position
 					FParticleUtilitiesXR::GetCoMWorldRotation(Particle0),	// Prev rotation
 					FParticleUtilitiesXR::GetCoMWorldRotation(Particle1),	// Prev rotation
 					Particle0->InvM(),
diff --git a/Engine/Source/Runtime/Experimental/Chaos/Public/Chaos/Joint/PBDJointSolverGaussSeidel.h b/Engine/Source/Runtime/Experimental/Chaos/Public/Chaos/Joint/PBDJointSolverGaussSeidel.h
--- a/Engine/Source/Runtime/Experimental/Chaos/Public/Chaos/Joint/PBDJointSolverGaussSeidel.h
+++ b/Engine/Source/Runtime/Experimental/Chaos/Public/Chaos/Joint/PBDJointSolverGaussSeidel.h
@@ -54,6 +54,8 @@ namespace Chaos
 			const FReal Dt,
 			const FPBDJointSolverSettings& SolverSettings,
 			const FPBDJointSettings& JointSettings,
+			const FVec3& PrevP0,
+			const FVec3& PrevP1,
 			const FRotation3& PrevQ0,
 			const FRotation3& PrevQ1,
 			const FReal InvM0,
@@ -93,6 +95,8 @@ namespace Chaos
 
 		void UpdateDerivedState(const int32 BodyIndex);
 		void UpdateDerivedState();
+		void InitAccumulatedDeltas();
+		void ApplyAccumulatedDeltas();
 
 		void ApplyPositionConstraints(
 			const FReal Dt,
@@ -153,13 +157,27 @@ namespace Chaos
 			const FVec3& DV1,
 			const FVec3& DW1);
 
-		void ApplyRotationCorrection(
+		void ApplyPositionConstraint(
+			const FReal Stiffness,
+			const FVec3& Axis,
+			const FReal Delta);
+
+		void ApplyPositionConstraintSoft(
+			const FReal Dt,
+			const FReal Stiffness,
+			const FReal Damping,
+			const bool bAccelerationMode,
+			const FVec3& Axis,
+			const FReal Delta,
+			FReal& Lambda);
+
+		void ApplyRotationConstraint(
 			const FReal Stiffness,
 			const FVec3& Axis0,
 			const FVec3& Axis1,
 			const FReal Angle);
 
-		void ApplyRotationCorrectionSoft(
+		void ApplyRotationConstraintSoft(
 			const FReal Dt,
 			const FReal Stiffness,
 			const FReal Damping,
@@ -242,8 +260,59 @@ namespace Chaos
 			const FPBDJointSolverSettings& SolverSettings,
 			const FPBDJointSettings& JointSettings);
 
-		void ApplyPositionConstraint(
+		void ApplyPointPositionConstraint(
+			const FReal Dt,
+			const FPBDJointSolverSettings& SolverSettings,
+			const FPBDJointSettings& JointSettings);
+
+		void ApplySphericalPositionConstraint(
+			const FReal Dt,
+			const FPBDJointSolverSettings& SolverSettings,
+			const FPBDJointSettings& JointSettings);
+
+		void ApplySphericalPositionConstraintSoft(
+			const FReal Dt,
+			const FPBDJointSolverSettings& SolverSettings,
+			const FPBDJointSettings& JointSettings);
+
+		void ApplySphericalPositionDrive(
+			const FReal Dt,
+			const FPBDJointSolverSettings& SolverSettings,
+			const FPBDJointSettings& JointSettings);
+
+		void ApplyCylindricalPositionConstraint(
+			const FReal Dt,
+			const int32 AxisIndex,
+			const FPBDJointSolverSettings& SolverSettings,
+			const FPBDJointSettings& JointSettings);
+
+		void ApplyCylindricalPositionConstraintSoft(
+			const FReal Dt,
+			const int32 AxisIndex,
+			const FPBDJointSolverSettings& SolverSettings,
+			const FPBDJointSettings& JointSettings);
+
+		void ApplyCylindricalPositionDrive(
+			const FReal Dt,
+			const int32 AxisIndex,
+			const FPBDJointSolverSettings& SolverSettings,
+			const FPBDJointSettings& JointSettings);
+
+		void ApplyPlanarPositionConstraint(
+			const FReal Dt,
+			const int32 AxisIndex,
+			const FPBDJointSolverSettings& SolverSettings,
+			const FPBDJointSettings& JointSettings);
+
+		void ApplyPlanarPositionConstraintSoft(
 			const FReal Dt,
+			const int32 AxisIndex,
+			const FPBDJointSolverSettings& SolverSettings,
+			const FPBDJointSettings& JointSettings);
+
+		void ApplyPlanarPositionDrive(
+			const FReal Dt,
+			const int32 AxisIndex,
 			const FPBDJointSolverSettings& SolverSettings,
 			const FPBDJointSettings& JointSettings);
 
@@ -268,14 +337,26 @@ namespace Chaos
 		FVec3 Vs[MaxConstrainedBodies];				// World-space particle CoM velocities
 		FVec3 Ws[MaxConstrainedBodies];				// World-space particle CoM angular velocities
 
+		// Accumulated world-body state deltas
+		FVec3 DPs[MaxConstrainedBodies];			// World-space particle CoM positions
+		FVec3 DRs[MaxConstrainedBodies];			// World-space particle CoM rotations (in axis-angle form)
+		FVec3 DVs[MaxConstrainedBodies];			// World-space particle CoM velocities
+		FVec3 DWs[MaxConstrainedBodies];			// World-space particle CoM angular velocities
+
 		// XPBD Previous iteration world-space body state
+		FVec3 PrevPs[MaxConstrainedBodies];			// World-space particle CoM positions
 		FRotation3 PrevQs[MaxConstrainedBodies];	// World-space particle CoM rotations
+		FVec3 PrevXs[MaxConstrainedBodies];			// World-space joint connector positions
 
-		// XPBD constraint multipliers
+		// XPBD constraint multipliers (net applied constraint-space deltas)
+		FReal LinearSoftLambda;
+		FReal LinearDriveLambda;
 		FReal TwistSoftLambda;
 		FReal SwingSoftLambda;
 		FReal TwistDriveLambda;
 		FReal SwingDriveLambda;
+
+		bool bEnableAngularConstraintPositionCorrection;
 	};
 
 }
\ No newline at end of file
diff --git a/Engine/Source/Runtime/Experimental/Chaos/Public/Chaos/PBDJointConstraintTypes.h b/Engine/Source/Runtime/Experimental/Chaos/Public/Chaos/PBDJointConstraintTypes.h
--- a/Engine/Source/Runtime/Experimental/Chaos/Public/Chaos/PBDJointConstraintTypes.h
+++ b/Engine/Source/Runtime/Experimental/Chaos/Public/Chaos/PBDJointConstraintTypes.h
@@ -83,7 +83,8 @@ namespace Chaos
 		bool bSoftLinearLimitsEnabled;
 		bool bSoftTwistLimitsEnabled;
 		bool bSoftSwingLimitsEnabled;
-		EJointForceMode SoftForceMode;
+		EJointForceMode LinearSoftForceMode;
+		EJointForceMode AngularSoftForceMode;
 		FReal SoftLinearStiffness;
 		FReal SoftLinearDamping;
 		FReal SoftTwistStiffness;
@@ -92,7 +93,8 @@ namespace Chaos
 		FReal SoftSwingDamping;
 
 		FVec3 LinearDriveTarget;
-		TVector<bool, 3> bLinearDriveEnabled;
+		TVector<bool, 3> bLinearPositionDriveEnabled;
+		TVector<bool, 3> bLinearVelocityDriveEnabled;
 		EJointForceMode LinearDriveForceMode;
 		FReal LinearDriveStiffness;
 		FReal LinearDriveDamping;
diff --git a/Engine/Source/Runtime/Experimental/Chaos/Public/Chaos/PBDJointConstraintUtilities.h b/Engine/Source/Runtime/Experimental/Chaos/Public/Chaos/PBDJointConstraintUtilities.h
--- a/Engine/Source/Runtime/Experimental/Chaos/Public/Chaos/PBDJointConstraintUtilities.h
+++ b/Engine/Source/Runtime/Experimental/Chaos/Public/Chaos/PBDJointConstraintUtilities.h
@@ -67,7 +67,15 @@ namespace Chaos
 		static FReal GetLinearStiffness(
 			const FPBDJointSolverSettings& SolverSettings,
 			const FPBDJointSettings& JointSettings);
-		
+
+		static FReal GetSoftLinearStiffness(
+			const FPBDJointSolverSettings& SolverSettings,
+			const FPBDJointSettings& JointSettings);
+
+		static FReal GetSoftLinearDamping(
+			const FPBDJointSolverSettings& SolverSettings,
+			const FPBDJointSettings& JointSettings);
+
 		static FReal GetTwistStiffness(
 			const FPBDJointSolverSettings& SolverSettings,
 			const FPBDJointSettings& JointSettings);
@@ -132,7 +140,11 @@ namespace Chaos
 			const FPBDJointSolverSettings& SolverSettings,
 			const FPBDJointSettings& JointSettings);
 
-		static bool GetSoftAccelerationMode(
+		static bool GetLinearSoftAccelerationMode(
+			const FPBDJointSolverSettings& SolverSettings,
+			const FPBDJointSettings& JointSettings);
+
+		static bool GetAngularSoftAccelerationMode(
 			const FPBDJointSolverSettings& SolverSettings,
 			const FPBDJointSettings& JointSettings);
 
