CIS Fix (dupe symbol in PBDJointSolverGaussSeidel)

diff --git a/Engine/Source/Runtime/Experimental/Chaos/Private/Chaos/Joint/PBDJointSolverGaussSeidel.cpp b/Engine/Source/Runtime/Experimental/Chaos/Private/Chaos/Joint/PBDJointSolverGaussSeidel.cpp
--- a/Engine/Source/Runtime/Experimental/Chaos/Private/Chaos/Joint/PBDJointSolverGaussSeidel.cpp
+++ b/Engine/Source/Runtime/Experimental/Chaos/Private/Chaos/Joint/PBDJointSolverGaussSeidel.cpp
@@ -11,7 +11,7 @@
 // Set to true to calculate all constraint corrections based on state at the beginning of the iteration. This would eliminate the need to recalculate the inertia and a few other
 // things in each constraint's Apply method. However, it probably introduces instability and requires more iterations to solve...
 bool bChaos_Joint_Accumulate = false;
-FAutoConsoleVariableRef CVarChaosImmPhysStepTime(TEXT("p.Chaos.Joint.Accumulate"), bChaos_Joint_Accumulate, TEXT("Whether to accumulate forces for each constraint over an iteration and apply all at once"));
+FAutoConsoleVariableRef CVarChaosJointAccumulate(TEXT("p.Chaos.Joint.Accumulate"), bChaos_Joint_Accumulate, TEXT("Whether to accumulate forces for each constraint over an iteration and apply all at once"));
 
 namespace Chaos
 {
