CIS fix required after recetn addition of the Chaos Cloth XPBD constraints.

diff --git a/Engine/Source/Runtime/Experimental/Chaos/Public/Chaos/XPBDAxialSpringConstraints.h b/Engine/Source/Runtime/Experimental/Chaos/Public/Chaos/XPBDAxialSpringConstraints.h
--- a/Engine/Source/Runtime/Experimental/Chaos/Public/Chaos/XPBDAxialSpringConstraints.h
+++ b/Engine/Source/Runtime/Experimental/Chaos/Public/Chaos/XPBDAxialSpringConstraints.h
@@ -21,6 +21,8 @@ class TXPBDAxialSpringConstraints : public TParticleRule<T, d>, public TPBDAxial
 	typedef TPBDAxialSpringConstraintsBase<T, d> Base;
 	using Base::MBarys;
 	using Base::MConstraints;
+	using Base::MDists;
+	using Base::MStiffness;
 
 public:
 	TXPBDAxialSpringConstraints(const TDynamicParticles<T, d>& InParticles, TArray<TVector<int32, 3>>&& Constraints, const T Stiffness = (T)1.)
diff --git a/Engine/Source/Runtime/Experimental/Chaos/Public/Chaos/XPBDLongRangeConstraints.h b/Engine/Source/Runtime/Experimental/Chaos/Public/Chaos/XPBDLongRangeConstraints.h
--- a/Engine/Source/Runtime/Experimental/Chaos/Public/Chaos/XPBDLongRangeConstraints.h
+++ b/Engine/Source/Runtime/Experimental/Chaos/Public/Chaos/XPBDLongRangeConstraints.h
@@ -18,6 +18,8 @@ class TXPBDLongRangeConstraints : public TPBDLongRangeConstraintsBase<T, d>, pub
 {
 	typedef TPBDLongRangeConstraintsBase<T, d> Base;
 	using Base::MConstraints;
+	using Base::MDists;
+	using Base::MStiffness;
 
 public:
 	TXPBDLongRangeConstraints(const TDynamicParticles<T, d>& InParticles, const TMap<int32, TSet<uint32>>& PointToNeighbors, const int32 NumberOfAttachments = 1, const T Stiffness = (T)1)
diff --git a/Engine/Source/Runtime/Experimental/Chaos/Public/Chaos/XPBDSpringConstraints.h b/Engine/Source/Runtime/Experimental/Chaos/Public/Chaos/XPBDSpringConstraints.h
--- a/Engine/Source/Runtime/Experimental/Chaos/Public/Chaos/XPBDSpringConstraints.h
+++ b/Engine/Source/Runtime/Experimental/Chaos/Public/Chaos/XPBDSpringConstraints.h
@@ -21,6 +21,8 @@ class TXPBDSpringConstraints : public TPBDSpringConstraintsBase<T, d>, public FP
 {
 	typedef TPBDSpringConstraintsBase<T, d> Base;
 	using Base::MConstraints;
+	using Base::MDists;
+	using Base::MStiffness;
 
 public:
 	TXPBDSpringConstraints(const T Stiffness = (T)1)
