CIS fix for error in XPBDSpringConstraint constructor initialization.

diff --git a/Engine/Source/Runtime/Experimental/Chaos/Public/Chaos/XPBDSpringConstraints.h b/Engine/Source/Runtime/Experimental/Chaos/Public/Chaos/XPBDSpringConstraints.h
--- a/Engine/Source/Runtime/Experimental/Chaos/Public/Chaos/XPBDSpringConstraints.h
+++ b/Engine/Source/Runtime/Experimental/Chaos/Public/Chaos/XPBDSpringConstraints.h
@@ -26,7 +26,7 @@ class TXPBDSpringConstraints : public TPBDSpringConstraintsBase<T, d>, public FP
 
 public:
 	TXPBDSpringConstraints(const T Stiffness = (T)1)
-	    : TPBDSpringConstraintsBase<T, d>(Stiffness * Stiffness * (T)XPBDSpringMaxStiffness)
+	    : TPBDSpringConstraintsBase<T, d>(Stiffness)
 	{}
 
 	TXPBDSpringConstraints(const TDynamicParticles<T, d>& InParticles, TArray<TVector<int32, 2>>&& Constraints, const T Stiffness = (T)1.)
