CIS Fix - Linux builds, add missing qualifier to find function in base class.

diff --git a/Engine/Source/Runtime/Experimental/Chaos/Public/Chaos/XPBDLongRangeConstraints.h b/Engine/Source/Runtime/Experimental/Chaos/Public/Chaos/XPBDLongRangeConstraints.h
--- a/Engine/Source/Runtime/Experimental/Chaos/Public/Chaos/XPBDLongRangeConstraints.h
+++ b/Engine/Source/Runtime/Experimental/Chaos/Public/Chaos/XPBDLongRangeConstraints.h
@@ -66,7 +66,7 @@ private:
 		const uint32 i2m1 = Constraint[Constraint.Num() - 2];
 		check(InParticles.InvM(i1) == 0);
 		check(InParticles.InvM(i2) > 0);
-		const T Distance = ComputeGeodesicDistance(InParticles, Constraint); // This function is used for either Euclidean or Geodesic distances
+		const T Distance = Base::ComputeGeodesicDistance(InParticles, Constraint); // This function is used for either Euclidean or Geodesic distances
 		if (Distance < MDists[InConstraintIndices]) { return TVector<T, d>(0); }
 		const T Offset = Distance - MDists[InConstraintIndices];
 
