Chaos - Joint position constraints logic fixes

diff --git a/Engine/Source/Runtime/Experimental/Chaos/Private/Chaos/Joint/PBDJointSolverGaussSeidel.cpp b/Engine/Source/Runtime/Experimental/Chaos/Private/Chaos/Joint/PBDJointSolverGaussSeidel.cpp
--- a/Engine/Source/Runtime/Experimental/Chaos/Private/Chaos/Joint/PBDJointSolverGaussSeidel.cpp
+++ b/Engine/Source/Runtime/Experimental/Chaos/Private/Chaos/Joint/PBDJointSolverGaussSeidel.cpp
@@ -34,46 +34,37 @@ namespace Chaos
 	}
 
 
-	FReal GetSphericalError(
-		const FReal& Delta,
-		const FReal& Limit)
-	{
-		return FMath::Max((FReal)0, Delta - Limit);
-	}
-
-
-	void GetCylindricalAxisDelta(
+	void GetCylindricalAxesDeltas(
 		const FRotation3& R0,
 		const FVec3& X0,
 		const FVec3& X1,
 		const int32 CylinderAxisIndex,
-		FVec3& Axis,
-		FReal& Delta)
+		FVec3& CylinderAxis,
+		FReal& CylinderDelta,
+		FVec3& RadialAxis,
+		FReal& RadialDelta)
 	{
-		Axis = FVec3(0);
-		Delta = 0;
+		CylinderAxis = FVec3(0);
+		CylinderDelta = 0;
+		RadialAxis = FVec3(0);
+		RadialDelta = 0;
 
-		const FMatrix33 R0M = R0.ToMatrix();
-		const FVec3 CylinderAxis = R0M.GetAxis(CylinderAxisIndex);
 		FVec3 DX = X1 - X0;
-		DX = DX - FVec3::DotProduct(DX, CylinderAxis) * CylinderAxis;
+
+		const FMatrix33 R0M = R0.ToMatrix();
+		CylinderAxis = R0M.GetAxis(CylinderAxisIndex);
+		CylinderDelta = FVec3::DotProduct(DX, CylinderAxis);
+
+		DX = DX - CylinderDelta * CylinderAxis;
 		const FReal DXLen = DX.Size();
 		if (DXLen > KINDA_SMALL_NUMBER)
 		{
-			Axis = DX / DXLen;
-			Delta = DXLen;
+			RadialAxis = DX / DXLen;
+			RadialDelta = DXLen;
 		}
 	}
 
 
-	FReal GetCylindricalError(
-		const FReal& Delta,
-		const FReal& Limit)
-	{
-		return FMath::Max((FReal)0, Delta - Limit);
-	}
-
-
 	void GetPlanarAxisDelta(
 		const FRotation3& R0,
 		const FVec3& X0,
@@ -88,22 +79,6 @@ namespace Chaos
 	}
 
 
-	FReal GetPlanarError(
-		const FReal& Delta,
-		const FReal& Limit)
-	{
-		if (Delta > Limit)
-		{
-			return Delta - Limit;
-		}
-		else if (Delta < Limit)
-		{
-			return Delta + Limit;
-		}
-		return 0;
-	}
-
-
 	void GetTwistAxisAngleLocal(
 		const FRotation3& R0,
 		const FRotation3& R1,
@@ -571,87 +546,76 @@ namespace Chaos
 		const FPBDJointSolverSettings& SolverSettings,
 		const FPBDJointSettings& JointSettings)
 	{
-		const TVector<EJointMotionType, 3>& LinearMotion = JointSettings.LinearMotionTypes;
+		// @todo(ccaulfield): the branch logic is all constant for a joint - move it to initialization and turn it into a table or something (also for Hard/Soft login in apply functions)
 
-		// All 3 axes locked -> hard point constraint
-		if ((LinearMotion[0] == EJointMotionType::Locked) && (LinearMotion[1] == EJointMotionType::Locked) && (LinearMotion[2] == EJointMotionType::Locked))
-		{
-			ApplyPointPositionConstraint(Dt, SolverSettings, JointSettings);
-			return;
-		}
-
-		// At least one axes is limited or free
-		const TVector<bool, 3> bLinearLimitSoft =
+		const TVector<EJointMotionType, 3>& LinearMotion = JointSettings.LinearMotionTypes;
+		const TVector<bool, 3> bLinearLocked =
 		{
-			(LinearMotion[0] == EJointMotionType::Limited) && JointSettings.bSoftLinearLimitsEnabled,
-			(LinearMotion[1] == EJointMotionType::Limited) && JointSettings.bSoftLinearLimitsEnabled,
-			(LinearMotion[2] == EJointMotionType::Limited) && JointSettings.bSoftLinearLimitsEnabled,
+			(LinearMotion[0] == EJointMotionType::Locked),
+			(LinearMotion[1] == EJointMotionType::Locked),
+			(LinearMotion[2] == EJointMotionType::Locked),
 		};
-		const TVector<bool, 3> bLinearLimitHard =
+		const TVector<bool, 3> bLinearLimted =
 		{
-			(LinearMotion[0] == EJointMotionType::Locked) || ((LinearMotion[0] == EJointMotionType::Limited) && !JointSettings.bSoftLinearLimitsEnabled),
-			(LinearMotion[1] == EJointMotionType::Locked) || ((LinearMotion[1] == EJointMotionType::Limited) && !JointSettings.bSoftLinearLimitsEnabled),
-			(LinearMotion[2] == EJointMotionType::Locked) || ((LinearMotion[2] == EJointMotionType::Limited) && !JointSettings.bSoftLinearLimitsEnabled),
+			(LinearMotion[0] == EJointMotionType::Limited),
+			(LinearMotion[1] == EJointMotionType::Limited),
+			(LinearMotion[2] == EJointMotionType::Limited),
 		};
 
-		// Handle hard limits
-		if (bLinearLimitHard[0] && bLinearLimitHard[1] && bLinearLimitHard[2])
-		{
-			ApplySphericalPositionConstraint(Dt, SolverSettings, JointSettings);
-		}
-		else if (bLinearLimitHard[1] && bLinearLimitHard[2])
-		{
-			ApplyCylindricalPositionConstraint(Dt, 0, SolverSettings, JointSettings);
-		}
-		else if (bLinearLimitHard[0] && bLinearLimitHard[2])
+		if (bLinearLocked[0] && bLinearLocked[1] && bLinearLocked[2])
 		{
-			ApplyCylindricalPositionConstraint(Dt, 1, SolverSettings, JointSettings);
+			// Hard point constraint (most common case)
+			ApplyPointPositionConstraint(Dt, SolverSettings, JointSettings);
 		}
-		else if (bLinearLimitHard[0] && bLinearLimitHard[1])
+		else if (bLinearLimted[0] && bLinearLimted[1] && bLinearLimted[2])
 		{
-			ApplyCylindricalPositionConstraint(Dt, 2, SolverSettings, JointSettings);
+			// Spherical constraint
+			ApplySphericalPositionConstraint(Dt, SolverSettings, JointSettings);
 		}
-		else if (bLinearLimitHard[0])
+		else if (bLinearLocked[1] && bLinearLocked[2] && !bLinearLocked[0])
 		{
-			ApplyPlanarPositionConstraint(Dt, 0, SolverSettings, JointSettings);
+			// Line constraint along X axis
+			ApplyCylindricalPositionConstraint(Dt, 0, LinearMotion[0], EJointMotionType::Locked, SolverSettings, JointSettings);
 		}
-		else if (bLinearLimitHard[1])
+		else if (bLinearLocked[0] && bLinearLocked[2] && !bLinearLocked[1])
 		{
-			ApplyPlanarPositionConstraint(Dt, 1, SolverSettings, JointSettings);
+			// Line constraint along Y axis
+			ApplyCylindricalPositionConstraint(Dt, 1, LinearMotion[1], EJointMotionType::Locked, SolverSettings, JointSettings);
 		}
-		else if (bLinearLimitHard[2])
+		else if (bLinearLocked[0] && bLinearLocked[1] && !bLinearLocked[2])
 		{
-			ApplyPlanarPositionConstraint(Dt, 2, SolverSettings, JointSettings);
+			// Line constraint along Z axis
+			ApplyCylindricalPositionConstraint(Dt, 2, LinearMotion[2], EJointMotionType::Locked, SolverSettings, JointSettings);
 		}
-
-		// Handle soft limits
-		if (bLinearLimitSoft[0] && bLinearLimitSoft[1] && bLinearLimitSoft[2])
+		else if (bLinearLimted[1] && bLinearLimted[2] && !bLinearLimted[0])
 		{
-			ApplySphericalPositionConstraintSoft(Dt, SolverSettings, JointSettings);
+			// Cylindrical constraint along X axis
+			ApplyCylindricalPositionConstraint(Dt, 0, LinearMotion[0], EJointMotionType::Limited, SolverSettings, JointSettings);
 		}
-		else if (bLinearLimitSoft[1] && bLinearLimitSoft[2])
+		else if (bLinearLimted[0] && bLinearLimted[2] && !bLinearLimted[1])
 		{
-			ApplyCylindricalPositionConstraintSoft(Dt, 0, SolverSettings, JointSettings);
+			// Cylindrical constraint along Y axis
+			ApplyCylindricalPositionConstraint(Dt, 1, LinearMotion[1], EJointMotionType::Limited, SolverSettings, JointSettings);
 		}
-		else if (bLinearLimitSoft[0] && bLinearLimitSoft[2])
+		else if (bLinearLimted[0] && bLinearLimted[1] && !bLinearLimted[2])
 		{
-			ApplyCylindricalPositionConstraintSoft(Dt, 1, SolverSettings, JointSettings);
+			// Cylindrical constraint along Z axis
+			ApplyCylindricalPositionConstraint(Dt, 2, LinearMotion[2], EJointMotionType::Limited, SolverSettings, JointSettings);
 		}
-		else if (bLinearLimitSoft[0] && bLinearLimitSoft[1])
+		else if (bLinearLocked[0] || bLinearLimted[0])
 		{
-			ApplyCylindricalPositionConstraintSoft(Dt, 2, SolverSettings, JointSettings);
+			// Planar constraint along X axis
+			ApplyPlanarPositionConstraint(Dt, 0, LinearMotion[0], SolverSettings, JointSettings);
 		}
-		else if (bLinearLimitSoft[0])
+		else if (bLinearLocked[1] || bLinearLimted[1])
 		{
-			ApplyPlanarPositionConstraintSoft(Dt, 0, SolverSettings, JointSettings);
+			// Planar constraint along Y axis
+			ApplyPlanarPositionConstraint(Dt, 1, LinearMotion[1], SolverSettings, JointSettings);
 		}
-		else if (bLinearLimitSoft[1])
+		else if (bLinearLocked[2] || bLinearLimted[2])
 		{
-			ApplyPlanarPositionConstraintSoft(Dt, 1, SolverSettings, JointSettings);
-		}
-		else if (bLinearLimitSoft[2])
-		{
-			ApplyPlanarPositionConstraintSoft(Dt, 2, SolverSettings, JointSettings);
+			// Planar constraint along Z axis
+			ApplyPlanarPositionConstraint(Dt, 2, LinearMotion[2], SolverSettings, JointSettings);
 		}
 	}
 
@@ -661,6 +625,10 @@ namespace Chaos
 		const FPBDJointSolverSettings& SolverSettings,
 		const FPBDJointSettings& JointSettings)
 	{
+		// @todo(ccaulfield): This logic is broken if different axes have drives position/velocity enabled - fix it. 
+		// E.g., if X/Y have position drives and Z is a velocity drive, it should be a circular position-drive and axial velocity-drive rather than a spherical one.
+		// E.g., if X/Y/Z have position drives and only 1 or 2 axes have velocity drives, we need to apply 2 drives. Etc etc.
+		// Basically we need to split the axes by those that have the same drive settings.
 		if (SolverSettings.bEnableDrives)
 		{
 			TVector<bool, 3> bDriven =
@@ -676,27 +644,27 @@ namespace Chaos
 			}
 			else if (bDriven[1] && bDriven[2])
 			{
-				ApplyCylindricalPositionDrive(Dt, 0, SolverSettings, JointSettings);
+				ApplyCircularPositionDrive(Dt, 0, SolverSettings, JointSettings);
 			}
 			else if (bDriven[0] && bDriven[2])
 			{
-				ApplyCylindricalPositionDrive(Dt, 1, SolverSettings, JointSettings);
+				ApplyCircularPositionDrive(Dt, 1, SolverSettings, JointSettings);
 			}
 			else if (bDriven[0] && bDriven[1])
 			{
-				ApplyCylindricalPositionDrive(Dt, 2, SolverSettings, JointSettings);
+				ApplyCircularPositionDrive(Dt, 2, SolverSettings, JointSettings);
 			}
 			else if (bDriven[0])
 			{
-				ApplyPlanarPositionDrive(Dt, 0, SolverSettings, JointSettings);
+				ApplyAxialPositionDrive(Dt, 0, SolverSettings, JointSettings);
 			}
 			else if (bDriven[1])
 			{
-				ApplyPlanarPositionDrive(Dt, 1, SolverSettings, JointSettings);
+				ApplyAxialPositionDrive(Dt, 1, SolverSettings, JointSettings);
 			}
 			else if (bDriven[2])
 			{
-				ApplyPlanarPositionDrive(Dt, 2, SolverSettings, JointSettings);
+				ApplyAxialPositionDrive(Dt, 2, SolverSettings, JointSettings);
 			}
 		}
 	}
@@ -1513,30 +1481,22 @@ namespace Chaos
 		FVec3 Axis;
 		FReal Delta;
 		GetSphericalAxisDelta(Xs[0], Xs[1], Axis, Delta);
-		const FReal Error = GetSphericalError(Delta, JointSettings.LinearLimit);
-		if (Error != 0)
-		{
-			const FReal Stiffness = FPBDJointUtilities::GetLinearStiffness(SolverSettings, JointSettings);
-			ApplyPositionConstraint(Stiffness, Axis, Error);
-		}
-	}
-
 
-	void FJointSolverGaussSeidel::ApplySphericalPositionConstraintSoft(
-		const FReal Dt,
-		const FPBDJointSolverSettings& SolverSettings,
-		const FPBDJointSettings& JointSettings)
-	{
-		FVec3 Axis;
-		FReal Delta;
-		GetSphericalAxisDelta(Xs[0], Xs[1], Axis, Delta);
-		const FReal Error = GetSphericalError(Delta, JointSettings.LinearLimit);
+		const FReal Error = FMath::Max((FReal)0, Delta - JointSettings.LinearLimit);
 		if (Error != 0)
 		{
-			const FReal Stiffness = FPBDJointUtilities::GetSoftLinearStiffness(SolverSettings, JointSettings);
-			const FReal Damping = FPBDJointUtilities::GetSoftLinearDamping(SolverSettings, JointSettings);
-			const bool bAccelerationMode = FPBDJointUtilities::GetLinearSoftAccelerationMode(SolverSettings, JointSettings);
-			ApplyPositionConstraintSoft(Dt, Stiffness, Damping, bAccelerationMode, Axis, Error, LinearSoftLambda);
+			if (!JointSettings.bSoftLinearLimitsEnabled)
+			{
+				const FReal Stiffness = FPBDJointUtilities::GetLinearStiffness(SolverSettings, JointSettings);
+				ApplyPositionConstraint(Stiffness, Axis, Error);
+			}
+			else
+			{
+				const FReal Stiffness = FPBDJointUtilities::GetSoftLinearStiffness(SolverSettings, JointSettings);
+				const FReal Damping = FPBDJointUtilities::GetSoftLinearDamping(SolverSettings, JointSettings);
+				const bool bAccelerationMode = FPBDJointUtilities::GetLinearSoftAccelerationMode(SolverSettings, JointSettings);
+				ApplyPositionConstraintSoft(Dt, Stiffness, Damping, bAccelerationMode, Axis, Error, LinearSoftLambda);
+			}
 		}
 	}
 
@@ -1563,57 +1523,85 @@ namespace Chaos
 	void FJointSolverGaussSeidel::ApplyCylindricalPositionConstraint(
 		const FReal Dt,
 		const int32 AxisIndex,
+		const EJointMotionType AxialMotion,
+		const EJointMotionType RadialMotion,
 		const FPBDJointSolverSettings& SolverSettings,
 		const FPBDJointSettings& JointSettings)
 	{
-		FVec3 Axis;
-		FReal Delta;
-		GetCylindricalAxisDelta(Rs[0], Xs[0], Xs[1], AxisIndex, Axis, Delta);
-		const FReal Error = GetCylindricalError(Delta, JointSettings.LinearLimit);
-		if (Error != 0)
+		check(AxialMotion != RadialMotion);
+
+		FVec3 Axis, RadialAxis;
+		FReal AxialDelta, RadialDelta;
+		GetCylindricalAxesDeltas(Rs[0], Xs[0], Xs[1], AxisIndex, Axis, AxialDelta, RadialAxis, RadialDelta);
+		
+		if ((AxialMotion == EJointMotionType::Limited) && JointSettings.bSoftLinearLimitsEnabled)
 		{
-			const FReal Stiffness = FPBDJointUtilities::GetLinearStiffness(SolverSettings, JointSettings);
-			ApplyPositionConstraint(Stiffness, Axis, Error);
+			// Soft Axial constraint
+			const FReal AxialLimit = JointSettings.LinearLimit;
+			if (FMath::Abs(AxialDelta) > AxialLimit)
+			{
+				const FReal AxialError = (AxialDelta > 0) ? AxialDelta - AxialLimit : AxialDelta + AxialLimit;
+				const FReal Stiffness = FPBDJointUtilities::GetSoftLinearStiffness(SolverSettings, JointSettings);
+				const FReal Damping = FPBDJointUtilities::GetSoftLinearDamping(SolverSettings, JointSettings);
+				const bool bAccelerationMode = FPBDJointUtilities::GetLinearSoftAccelerationMode(SolverSettings, JointSettings);
+				ApplyPositionConstraintSoft(Dt, Stiffness, Damping, bAccelerationMode, Axis, AxialError, LinearSoftLambda);
+			}
+		}
+		else if (AxialMotion != EJointMotionType::Free)
+		{
+			// Hard Axial constraint
+			const FReal AxialLimit = (AxialMotion == EJointMotionType::Locked) ? 0 : JointSettings.LinearLimit;
+			if (FMath::Abs(AxialDelta) > AxialLimit)
+			{
+				const FReal AxialError = (AxialDelta > 0) ? AxialDelta - AxialLimit : AxialDelta + AxialLimit;
+				const FReal Stiffness = FPBDJointUtilities::GetLinearStiffness(SolverSettings, JointSettings);
+				ApplyPositionConstraint(Stiffness, Axis, AxialError);
+			}
 		}
-	}
-
 
-	void FJointSolverGaussSeidel::ApplyCylindricalPositionConstraintSoft(
-		const FReal Dt,
-		const int32 AxisIndex,
-		const FPBDJointSolverSettings& SolverSettings,
-		const FPBDJointSettings& JointSettings)
-	{
-		FVec3 Axis;
-		FReal Delta;
-		GetCylindricalAxisDelta(Rs[0], Xs[0], Xs[1], AxisIndex, Axis, Delta);
-		const FReal Error = GetCylindricalError(Delta, JointSettings.LinearLimit);
-		if (Error != 0)
+		if ((RadialMotion == EJointMotionType::Limited) && JointSettings.bSoftLinearLimitsEnabled)
 		{
-			const FReal Stiffness = FPBDJointUtilities::GetSoftLinearStiffness(SolverSettings, JointSettings);
-			const FReal Damping = FPBDJointUtilities::GetSoftLinearDamping(SolverSettings, JointSettings);
-			const bool bAccelerationMode = FPBDJointUtilities::GetLinearSoftAccelerationMode(SolverSettings, JointSettings);
-			ApplyPositionConstraintSoft(Dt, Stiffness, Damping, bAccelerationMode, Axis, Error, LinearSoftLambda);
+			// Soft Radial constraint
+			const FReal RadialLimit = JointSettings.LinearLimit;
+			if (RadialDelta > RadialLimit)
+			{
+				const FReal RadialError = FMath::Max((FReal)0, RadialDelta - RadialLimit);
+				const FReal Stiffness = FPBDJointUtilities::GetSoftLinearStiffness(SolverSettings, JointSettings);
+				const FReal Damping = FPBDJointUtilities::GetSoftLinearDamping(SolverSettings, JointSettings);
+				const bool bAccelerationMode = FPBDJointUtilities::GetLinearSoftAccelerationMode(SolverSettings, JointSettings);
+				ApplyPositionConstraintSoft(Dt, Stiffness, Damping, bAccelerationMode, RadialAxis, RadialError, LinearSoftLambda);
+			}
+		}
+		else if (RadialMotion != EJointMotionType::Free)
+		{
+			// Hard Radial constraint
+			const FReal RadialLimit = (RadialMotion == EJointMotionType::Locked) ? 0 : JointSettings.LinearLimit;
+			if (RadialDelta > RadialLimit)
+			{
+				const FReal RadialError = FMath::Max((FReal)0, RadialDelta - RadialLimit);
+				const FReal Stiffness = FPBDJointUtilities::GetLinearStiffness(SolverSettings, JointSettings);
+				ApplyPositionConstraint(Stiffness, RadialAxis, RadialError);
+			}
 		}
 	}
 
 
-	void FJointSolverGaussSeidel::ApplyCylindricalPositionDrive(
+	void FJointSolverGaussSeidel::ApplyCircularPositionDrive(
 		const FReal Dt,
 		const int32 AxisIndex,
 		const FPBDJointSolverSettings& SolverSettings,
 		const FPBDJointSettings& JointSettings)
 	{
 		const FVec3 XTarget = Xs[0] + Rs[0] * JointSettings.LinearDriveTarget;
-		FVec3 Axis;
-		FReal Delta;
-		GetCylindricalAxisDelta(Rs[0], XTarget, Xs[1], AxisIndex, Axis, Delta);
-		if (Delta != 0)
+		FVec3 Axis, RadialAxis;
+		FReal AxialDelta, RadialDelta;
+		GetCylindricalAxesDeltas(Rs[0], XTarget, Xs[1], AxisIndex, Axis, AxialDelta, RadialAxis, RadialDelta);
+		if (RadialDelta != 0)
 		{
 			const FReal Stiffness = FPBDJointUtilities::GetLinearDriveStiffness(SolverSettings, JointSettings);
 			const FReal Damping = FPBDJointUtilities::GetLinearDriveDamping(SolverSettings, JointSettings);
 			const bool bAccelerationMode = FPBDJointUtilities::GetDriveAccelerationMode(SolverSettings, JointSettings);
-			ApplyPositionConstraintSoft(Dt, Stiffness, Damping, bAccelerationMode, Axis, Delta, LinearDriveLambda);
+			ApplyPositionConstraintSoft(Dt, Stiffness, Damping, bAccelerationMode, RadialAxis, RadialDelta, LinearDriveLambda);
 		}
 	}
 
@@ -1621,42 +1609,35 @@ namespace Chaos
 	void FJointSolverGaussSeidel::ApplyPlanarPositionConstraint(
 		const FReal Dt,
 		const int32 AxisIndex,
+		const EJointMotionType AxialMotion,
 		const FPBDJointSolverSettings& SolverSettings,
 		const FPBDJointSettings& JointSettings)
 	{
 		FVec3 Axis;
 		FReal Delta;
 		GetPlanarAxisDelta(Rs[0], Xs[0], Xs[1], AxisIndex, Axis, Delta);
-		const FReal Error = GetPlanarError(Delta, JointSettings.LinearLimit);
-		if (Error != 0)
-		{
-			const FReal Stiffness = FPBDJointUtilities::GetLinearStiffness(SolverSettings, JointSettings);
-			ApplyPositionConstraint(Stiffness, Axis, Error);
-		}
-	}
-
 
-	void FJointSolverGaussSeidel::ApplyPlanarPositionConstraintSoft(
-		const FReal Dt,
-		const int32 AxisIndex,
-		const FPBDJointSolverSettings& SolverSettings,
-		const FPBDJointSettings& JointSettings)
-	{
-		FVec3 Axis;
-		FReal Delta;
-		GetPlanarAxisDelta(Rs[0], Xs[0], Xs[1], AxisIndex, Axis, Delta);
-		const FReal Error = GetPlanarError(Delta, JointSettings.LinearLimit);
-		if (Error != 0)
+		const FReal Limit = (AxialMotion == EJointMotionType::Locked) ? 0 : JointSettings.LinearLimit;
+		if (FMath::Abs(Delta) > Limit)
 		{
-			const FReal Stiffness = FPBDJointUtilities::GetSoftLinearStiffness(SolverSettings, JointSettings);
-			const FReal Damping = FPBDJointUtilities::GetSoftLinearDamping(SolverSettings, JointSettings);
-			const bool bAccelerationMode = FPBDJointUtilities::GetLinearSoftAccelerationMode(SolverSettings, JointSettings);
-			ApplyPositionConstraintSoft(Dt, Stiffness, Damping, bAccelerationMode, Axis, Error, LinearSoftLambda);
+			const FReal Error = (Delta > 0) ? Delta - Limit : Delta + Limit;
+			if ((AxialMotion == EJointMotionType::Limited) && JointSettings.bSoftLinearLimitsEnabled)
+			{
+				const FReal Stiffness = FPBDJointUtilities::GetSoftLinearStiffness(SolverSettings, JointSettings);
+				const FReal Damping = FPBDJointUtilities::GetSoftLinearDamping(SolverSettings, JointSettings);
+				const bool bAccelerationMode = FPBDJointUtilities::GetLinearSoftAccelerationMode(SolverSettings, JointSettings);
+				ApplyPositionConstraintSoft(Dt, Stiffness, Damping, bAccelerationMode, Axis, Error, LinearSoftLambda);
+			}
+			else
+			{
+				const FReal Stiffness = FPBDJointUtilities::GetLinearStiffness(SolverSettings, JointSettings);
+				ApplyPositionConstraint(Stiffness, Axis, Error);
+			}
 		}
 	}
 
 
-	void FJointSolverGaussSeidel::ApplyPlanarPositionDrive(
+	void FJointSolverGaussSeidel::ApplyAxialPositionDrive(
 		const FReal Dt,
 		const int32 AxisIndex,
 		const FPBDJointSolverSettings& SolverSettings,
diff --git a/Engine/Source/Runtime/Experimental/Chaos/Public/Chaos/Joint/PBDJointSolverGaussSeidel.h b/Engine/Source/Runtime/Experimental/Chaos/Public/Chaos/Joint/PBDJointSolverGaussSeidel.h
--- a/Engine/Source/Runtime/Experimental/Chaos/Public/Chaos/Joint/PBDJointSolverGaussSeidel.h
+++ b/Engine/Source/Runtime/Experimental/Chaos/Public/Chaos/Joint/PBDJointSolverGaussSeidel.h
@@ -270,11 +270,6 @@ namespace Chaos
 			const FPBDJointSolverSettings& SolverSettings,
 			const FPBDJointSettings& JointSettings);
 
-		void ApplySphericalPositionConstraintSoft(
-			const FReal Dt,
-			const FPBDJointSolverSettings& SolverSettings,
-			const FPBDJointSettings& JointSettings);
-
 		void ApplySphericalPositionDrive(
 			const FReal Dt,
 			const FPBDJointSolverSettings& SolverSettings,
@@ -283,16 +278,12 @@ namespace Chaos
 		void ApplyCylindricalPositionConstraint(
 			const FReal Dt,
 			const int32 AxisIndex,
+			const EJointMotionType AxialMotion,
+			const EJointMotionType RadialMotion,
 			const FPBDJointSolverSettings& SolverSettings,
 			const FPBDJointSettings& JointSettings);
 
-		void ApplyCylindricalPositionConstraintSoft(
-			const FReal Dt,
-			const int32 AxisIndex,
-			const FPBDJointSolverSettings& SolverSettings,
-			const FPBDJointSettings& JointSettings);
-
-		void ApplyCylindricalPositionDrive(
+		void ApplyCircularPositionDrive(
 			const FReal Dt,
 			const int32 AxisIndex,
 			const FPBDJointSolverSettings& SolverSettings,
@@ -301,16 +292,11 @@ namespace Chaos
 		void ApplyPlanarPositionConstraint(
 			const FReal Dt,
 			const int32 AxisIndex,
+			const EJointMotionType AxialMotion,
 			const FPBDJointSolverSettings& SolverSettings,
 			const FPBDJointSettings& JointSettings);
 
-		void ApplyPlanarPositionConstraintSoft(
-			const FReal Dt,
-			const int32 AxisIndex,
-			const FPBDJointSolverSettings& SolverSettings,
-			const FPBDJointSettings& JointSettings);
-
-		void ApplyPlanarPositionDrive(
+		void ApplyAxialPositionDrive(
 			const FReal Dt,
 			const int32 AxisIndex,
 			const FPBDJointSolverSettings& SolverSettings,
