Chaos
- Inertia principle axis calculation fix
- Inertia world space transform unit test

diff --git a/Engine/Source/Runtime/Experimental/Chaos/Private/Chaos/MassProperties.cpp b/Engine/Source/Runtime/Experimental/Chaos/Private/Chaos/MassProperties.cpp
--- a/Engine/Source/Runtime/Experimental/Chaos/Private/Chaos/MassProperties.cpp
+++ b/Engine/Source/Runtime/Experimental/Chaos/Private/Chaos/MassProperties.cpp
@@ -64,7 +64,8 @@ namespace Chaos
 		// Return results
 		Inertia = PMatrix<T, d, d>(m00, 0, 0, m11, 0, m22);
 		PMatrix<T, d, d> RotationMatrix = DoSwap ? PMatrix<T, d, d>(Eigenvector2, Eigenvector1, -Eigenvector0) : PMatrix<T, d, d>(Eigenvector0, Eigenvector1, Eigenvector2);
-		FinalRotation = TRotation<T,d>(RotationMatrix);
+		// NOTE: UE Matrix are column-major, so the PMatrix constructor is not setting eigenvectors - we need to transpose it to get a UE rotation matrix.
+		FinalRotation = TRotation<T,d>(RotationMatrix.GetTransposed());
 		check(FMath::IsNearlyEqual(FinalRotation.Size(), 1.0f, KINDA_SMALL_NUMBER));
 		
 		return FinalRotation;
diff --git a/Engine/Source/Runtime/Experimental/Chaos/Public/Chaos/Utilities.h b/Engine/Source/Runtime/Experimental/Chaos/Public/Chaos/Utilities.h
--- a/Engine/Source/Runtime/Experimental/Chaos/Public/Chaos/Utilities.h
+++ b/Engine/Source/Runtime/Experimental/Chaos/Public/Chaos/Utilities.h
@@ -235,10 +235,15 @@ namespace Chaos
 		static FMatrix33 ComputeWorldSpaceInertia(const FRotation3& Q, const FMatrix33& I)
 		{
 			FMatrix33 QM = Q.ToMatrix();
-			//return QM * I * QM.GetTransposed();
 			return MultiplyAB(QM, MultiplyABt(I, QM));
 		}
 
+		static FMatrix33 ComputeWorldSpaceInertia(const FRotation3& Q, const FVec3& I)
+		{
+			// @todo(ccaulfield): optimize ComputeWorldSpaceInertia
+			return ComputeWorldSpaceInertia(Q, FMatrix33(I.X, I.Y, I.Z));
+		}
+
 		/**
 		 * Calculate the matrix that maps a constraint position error to constraint position and rotation corrections.
 		 */
