Chaos
- Mass conditioning WIP

diff --git a/Engine/Source/Runtime/Engine/Private/PhysicsEngine/ImmediatePhysics/ImmediatePhysicsChaos/ImmediatePhysicsActorHandle_Chaos.cpp b/Engine/Source/Runtime/Engine/Private/PhysicsEngine/ImmediatePhysics/ImmediatePhysicsChaos/ImmediatePhysicsActorHandle_Chaos.cpp
--- a/Engine/Source/Runtime/Engine/Private/PhysicsEngine/ImmediatePhysics/ImmediatePhysicsChaos/ImmediatePhysicsActorHandle_Chaos.cpp
+++ b/Engine/Source/Runtime/Engine/Private/PhysicsEngine/ImmediatePhysics/ImmediatePhysicsChaos/ImmediatePhysicsActorHandle_Chaos.cpp
@@ -25,6 +25,9 @@ namespace ImmediatePhysics_Chaos
 	// Utils
 	//
 
+	// @todo(ccaulfield): Max mass conditioning an option (or push it into the constraint which is where it is really needed)
+	const FReal Chaos_MaxDimensionRatio = 5.0f;
+
 	Chaos::FMatrix33 CalculateInertia_Solid(const FReal Mass, const FKSphereElem& SphereElem)
 	{
 		return Chaos::TSphere<FReal, 3>::GetInertiaTensor(Mass, SphereElem.Radius);
@@ -32,12 +35,27 @@ namespace ImmediatePhysics_Chaos
 
 	Chaos::FMatrix33 CalculateInertia_Solid(const FReal Mass, const FKSphylElem& SphylElem)
 	{
-		return Chaos::TCapsule<FReal>::GetInertiaTensor(Mass, SphylElem.Length, SphylElem.Radius);
+		float Len = SphylElem.Length;
+		float Rad = SphylElem.Radius;
+		if (Len > Chaos_MaxDimensionRatio * Rad)
+		{
+			Rad = Len / Chaos_MaxDimensionRatio;
+		}
+		return Chaos::TCapsule<FReal>::GetInertiaTensor(Mass, Len, Rad);
 	}
 
 	Chaos::FMatrix33 CalculateInertia_Solid(const FReal Mass, const FKBoxElem& BoxElem)
 	{
-		return Chaos::TBox<FReal, 3>::GetInertiaTensor(Mass, FVector(BoxElem.X, BoxElem.Y, BoxElem.Z));
+		FVector Dim = FVector(BoxElem.X, BoxElem.Y, BoxElem.Z);
+		float MaxDim = Dim.GetAbsMax();
+		float MinDim = Dim.GetAbsMin();
+		if (MaxDim > Chaos_MaxDimensionRatio* MinDim)
+		{
+			Dim.X = FMath::Lerp(MinDim, MaxDim, (Dim.X - MinDim) / (MaxDim - MinDim));
+			Dim.Y = FMath::Lerp(MinDim, MaxDim, (Dim.Y - MinDim) / (MaxDim - MinDim));
+			Dim.Z = FMath::Lerp(MinDim, MaxDim, (Dim.Z - MinDim) / (MaxDim - MinDim));
+		}
+		return Chaos::TBox<FReal, 3>::GetInertiaTensor(Mass, Dim);
 	}
 
 	void CalculateMassProperties(const FVector& Scale, const FTransform& LocalTransform, const FKAggregateGeom& AggGeom, Chaos::TMassProperties<FReal, 3>& OutMassProperties)
diff --git a/Engine/Source/Runtime/Engine/Private/PhysicsEngine/ImmediatePhysics/ImmediatePhysicsChaos/ImmediatePhysicsJointHandle_Chaos.cpp b/Engine/Source/Runtime/Engine/Private/PhysicsEngine/ImmediatePhysics/ImmediatePhysicsChaos/ImmediatePhysicsJointHandle_Chaos.cpp
--- a/Engine/Source/Runtime/Engine/Private/PhysicsEngine/ImmediatePhysics/ImmediatePhysicsChaos/ImmediatePhysicsJointHandle_Chaos.cpp
+++ b/Engine/Source/Runtime/Engine/Private/PhysicsEngine/ImmediatePhysics/ImmediatePhysicsChaos/ImmediatePhysicsJointHandle_Chaos.cpp
@@ -148,6 +148,22 @@ namespace ImmediatePhysics_Chaos
 			}
 		}
 
+		// Disable Soft Limits when stiffness is extremely high
+		//const FReal MaxLinearStiffness = 10000;
+		//const FReal MaxAngularStiffness = 10000;
+		//if (ConstraintSettings.bSoftLinearLimitsEnabled && ConstraintSettings.SoftLinearStiffness > MaxLinearStiffness)
+		//{
+		//	ConstraintSettings.bSoftLinearLimitsEnabled = false;
+		//}
+		//if (ConstraintSettings.bSoftTwistLimitsEnabled && ConstraintSettings.SoftTwistStiffness > MaxAngularStiffness)
+		//{
+		//	ConstraintSettings.bSoftTwistLimitsEnabled = false;
+		//}
+		//if (ConstraintSettings.bSoftSwingLimitsEnabled && ConstraintSettings.SoftSwingStiffness > MaxAngularStiffness)
+		//{
+		//	ConstraintSettings.bSoftSwingLimitsEnabled = false;
+		//}
+
 		ConstraintSettings.Sanitize();
 
 		ConstraintHandle = Constraints->AddConstraint({ Actor1->ParticleHandle, Actor2->ParticleHandle }, ConstraintFrames, ConstraintSettings);
diff --git a/Engine/Source/Runtime/Engine/Private/PhysicsEngine/ImmediatePhysics/ImmediatePhysicsChaos/ImmediatePhysicsSimulation_Chaos.cpp b/Engine/Source/Runtime/Engine/Private/PhysicsEngine/ImmediatePhysics/ImmediatePhysicsChaos/ImmediatePhysicsSimulation_Chaos.cpp
--- a/Engine/Source/Runtime/Engine/Private/PhysicsEngine/ImmediatePhysics/ImmediatePhysicsChaos/ImmediatePhysicsSimulation_Chaos.cpp
+++ b/Engine/Source/Runtime/Engine/Private/PhysicsEngine/ImmediatePhysics/ImmediatePhysicsChaos/ImmediatePhysicsSimulation_Chaos.cpp
@@ -62,7 +62,7 @@ float ChaosImmediate_Joint_LinearDriveStiffness = 0.0f;
 float ChaosImmediate_Joint_LinearDriveDamping = 0.0f;
 float ChaosImmediate_Joint_AngularDriveStiffness = 0.0f;
 float ChaosImmediate_Joint_AngularDriveDamping = 0.0f;
-float ChaosImmediate_Joint_MinParentMassRatio = 0.5f;
+float ChaosImmediate_Joint_MinParentMassRatio = 0.2f;
 float ChaosImmediate_Joint_MaxInertiaRatio = 5.0f;
 FAutoConsoleVariableRef CVarChaosImmPhysPairIterations(TEXT("p.Chaos.ImmPhys.Joint.PairIterations"), ChaosImmediate_Joint_PairIterations, TEXT("PairIterations."));
 FAutoConsoleVariableRef CVarChaosImmPhysPushOutPairIterations(TEXT("p.Chaos.ImmPhys.Joint.PushOutPairIterations"), ChaosImmediate_Joint_PushOutPairIterations, TEXT("PushOutPairIterations."));
diff --git a/Engine/Source/Runtime/Experimental/Chaos/Private/Chaos/Joint/PBDJointSolverGaussSeidel.cpp b/Engine/Source/Runtime/Experimental/Chaos/Private/Chaos/Joint/PBDJointSolverGaussSeidel.cpp
--- a/Engine/Source/Runtime/Experimental/Chaos/Private/Chaos/Joint/PBDJointSolverGaussSeidel.cpp
+++ b/Engine/Source/Runtime/Experimental/Chaos/Private/Chaos/Joint/PBDJointSolverGaussSeidel.cpp
@@ -271,21 +271,27 @@ namespace Chaos
 		const FRotation3& PrevQ0,
 		const FRotation3& PrevQ1,
 		const FReal InvM0,
-		const FMatrix33& InvIL0,
+		const FVec3& InvIL0,
 		const FReal InvM1,
-		const FMatrix33& InvIL1,
+		const FVec3& InvIL1,
 		const FRigidTransform3& XL0,
 		const FRigidTransform3& XL1)
 	{
 		XLs[0] = XL0;
 		XLs[1] = XL1;
 
-		// @todo(ccaulfield): mass conditioning
 		InvILs[0] = JointSettings.ParentInvMassScale * InvIL0;
 		InvILs[1] = InvIL1;
 		InvMs[0] = JointSettings.ParentInvMassScale * InvM0;
 		InvMs[1] = InvM1;
 
+		// @todo(ccaulfield): mass conditioning
+		//FReal M0 = (JointSettings.ParentInvMassScale * InvM0 > KINDA_SMALL_NUMBER) ? 1.0f / (JointSettings.ParentInvMassScale * InvM0) : 0.0f;
+		//FReal M1 = (InvM1 > 0) ? 1.0f / InvM1 : 0.0f;
+		//FVec3 IL0 = (JointSettings.ParentInvMassScale * InvM0 > KINDA_SMALL_NUMBER) ? FVec3(1.0f / (JointSettings.ParentInvMassScale * InvIL0.X), 1.0f / (JointSettings.ParentInvMassScale * InvIL0.Y), 1.0f / (JointSettings.ParentInvMassScale * InvIL0.Z)) : FVec3(0);
+		//FVec3 IL1 = (InvM1 > 0) ? FVec3(1.0f / InvIL1.X, 1.0f / InvIL1.Y, 1.0f / InvIL1.Z) : FVec3(0);
+		//FPBDJointUtilities::GetConditionedInverseMass(M0, IL0, M1, IL1, InvMs[0], InvMs[1], InvILs[0], InvILs[1], SolverSettings.MinParentMassRatio, SolverSettings.MaxInertiaRatio);
+
 		PrevPs[0] = PrevP0;
 		PrevPs[1] = PrevP1;
 		PrevQs[0] = PrevQ0;
@@ -1678,8 +1684,8 @@ namespace Chaos
 
 			const FReal IM0 = ParentMassScale * InvMs[0];
 			const FReal IM1 = InvMs[1];
-			const FMatrix33 IIL0 = ParentMassScale * InvILs[0];
-			const FMatrix33& IIL1 = InvILs[1];
+			const FVec3 IIL0 = ParentMassScale * InvILs[0];
+			const FVec3& IIL1 = InvILs[1];
 			FMatrix33 II0 = Utilities::ComputeWorldSpaceInertia(Qs[0], IIL0);
 			FMatrix33 II1 = Utilities::ComputeWorldSpaceInertia(Qs[1], IIL1);
 			FMatrix33 J0 = (IM0 > 0) ? Utilities::ComputeJointFactorMatrix(Xs[0] - Ps[0], II0, IM0) : FMatrix33(0, 0, 0);
diff --git a/Engine/Source/Runtime/Experimental/Chaos/Private/Chaos/PBDJointConstraintUtilities.cpp b/Engine/Source/Runtime/Experimental/Chaos/Private/Chaos/PBDJointConstraintUtilities.cpp
--- a/Engine/Source/Runtime/Experimental/Chaos/Private/Chaos/PBDJointConstraintUtilities.cpp
+++ b/Engine/Source/Runtime/Experimental/Chaos/Private/Chaos/PBDJointConstraintUtilities.cpp
@@ -257,8 +257,8 @@ namespace Chaos
 		const FVec3 InIChild, 
 		FReal& OutInvMParent, 
 		FReal& OutInvMChild, 
-		FMatrix33& OutInvIParent, 
-		FMatrix33& OutInvIChild, 
+		FVec3& OutInvIParent,
+		FVec3& OutInvIChild,
 		const FReal MinParentMassRatio, 
 		const FReal MaxInertiaRatio)
 	{
@@ -270,19 +270,19 @@ namespace Chaos
 		IParent = ConditionParentInertia(IParent, IChild, MinParentMassRatio);
 
 		OutInvMParent = 0;
-		OutInvIParent = FMatrix33(0, 0, 0);
+		OutInvIParent = FVec3(0, 0, 0);
 		if (MParent > 0)
 		{
 			OutInvMParent = (FReal)1 / MParent;
-			OutInvIParent = FMatrix33((FReal)1 / IParent.X, (FReal)1 / IParent.Y, (FReal)1 / IParent.Z);
+			OutInvIParent = FVec3((FReal)1 / IParent.X, (FReal)1 / IParent.Y, (FReal)1 / IParent.Z);
 		}
 
 		OutInvMChild = 0;
-		OutInvIChild = FMatrix33(0, 0, 0);
+		OutInvIChild = FVec3(0, 0, 0);
 		if (MChild > 0)
 		{
 			OutInvMChild = (FReal)1 / MChild;
-			OutInvIChild = FMatrix33((FReal)1 / IChild.X, (FReal)1 / IChild.Y, (FReal)1 / IChild.Z);
+			OutInvIChild = FVec3((FReal)1 / IChild.X, (FReal)1 / IChild.Y, (FReal)1 / IChild.Z);
 		}
 	}
 
@@ -291,16 +291,16 @@ namespace Chaos
 		const FReal InM0,
 		const FVec3 InI0,
 		FReal& OutInvM0, 
-		FMatrix33& OutInvI0, 
+		FVec3& OutInvI0,
 		const FReal MaxInertiaRatio)
 	{
 		OutInvM0 = 0;
-		OutInvI0 = FMatrix33(0, 0, 0);
+		OutInvI0 = FVec3(0, 0, 0);
 		if (InM0 > 0)
 		{
 			FVec3 I0 = ConditionInertia(InI0, MaxInertiaRatio);
 			OutInvM0 = (FReal)1 / InM0;
-			OutInvI0 = FMatrix33((FReal)1 / I0.X, (FReal)1 / I0.Y, (FReal)1 / I0.Z);
+			OutInvI0 = FVec3((FReal)1 / I0.X, (FReal)1 / I0.Y, (FReal)1 / I0.Z);
 		}
 	}
 
diff --git a/Engine/Source/Runtime/Experimental/Chaos/Private/Chaos/PBDJointConstraints.cpp b/Engine/Source/Runtime/Experimental/Chaos/Private/Chaos/PBDJointConstraints.cpp
--- a/Engine/Source/Runtime/Experimental/Chaos/Private/Chaos/PBDJointConstraints.cpp
+++ b/Engine/Source/Runtime/Experimental/Chaos/Private/Chaos/PBDJointConstraints.cpp
@@ -444,9 +444,9 @@ namespace Chaos
 					FParticleUtilitiesXR::GetCoMWorldRotation(Particle0),	// Prev rotation
 					FParticleUtilitiesXR::GetCoMWorldRotation(Particle1),	// Prev rotation
 					Particle0->InvM(),
-					Particle0->InvI(),
+					Particle0->InvI().GetDiagonal(),
 					Particle1->InvM(),
-					Particle1->InvI(),
+					Particle1->InvI().GetDiagonal(),
 					FParticleUtilities::ParticleLocalToCoMLocal(Particle0, JointFrames[Index0]),
 					FParticleUtilities::ParticleLocalToCoMLocal(Particle1, JointFrames[Index1]));
 			}
diff --git a/Engine/Source/Runtime/Experimental/Chaos/Public/Chaos/Joint/PBDJointSolverGaussSeidel.h b/Engine/Source/Runtime/Experimental/Chaos/Public/Chaos/Joint/PBDJointSolverGaussSeidel.h
--- a/Engine/Source/Runtime/Experimental/Chaos/Public/Chaos/Joint/PBDJointSolverGaussSeidel.h
+++ b/Engine/Source/Runtime/Experimental/Chaos/Public/Chaos/Joint/PBDJointSolverGaussSeidel.h
@@ -59,9 +59,9 @@ namespace Chaos
 			const FRotation3& PrevQ0,
 			const FRotation3& PrevQ1,
 			const FReal InvM0,
-			const FMatrix33& InvIL0,
+			const FVec3& InvIL0,
 			const FReal InvM1,
-			const FMatrix33& InvIL1,
+			const FVec3& InvIL1,
 			const FRigidTransform3& XL0,
 			const FRigidTransform3& XL1);
 
@@ -310,7 +310,7 @@ namespace Chaos
 
 		// Local-space constraint settings
 		FRigidTransform3 XLs[MaxConstrainedBodies];	// Local-space joint connector transforms
-		FMatrix33 InvILs[MaxConstrainedBodies];		// Local-space inverse inertias
+		FVec3 InvILs[MaxConstrainedBodies];			// Local-space inverse inertias
 		FReal InvMs[MaxConstrainedBodies];			// Inverse masses
 
 		// World-space constraint state
diff --git a/Engine/Source/Runtime/Experimental/Chaos/Public/Chaos/PBDJointConstraintUtilities.h b/Engine/Source/Runtime/Experimental/Chaos/Public/Chaos/PBDJointConstraintUtilities.h
--- a/Engine/Source/Runtime/Experimental/Chaos/Public/Chaos/PBDJointConstraintUtilities.h
+++ b/Engine/Source/Runtime/Experimental/Chaos/Public/Chaos/PBDJointConstraintUtilities.h
@@ -51,8 +51,8 @@ namespace Chaos
 			const FVec3 IChild,
 			FReal& OutInvMParent,
 			FReal& OutInvMChild,
-			FMatrix33& OutInvIParent,
-			FMatrix33& OutInvIChild,
+			FVec3& OutInvIParent,
+			FVec3& OutInvIChild,
 			const FReal MinParentMassRatio,
 			const FReal MaxInertiaRatio);
 
@@ -60,7 +60,7 @@ namespace Chaos
 			const float M,
 			const FVec3 I,
 			FReal& OutInvM0,
-			FMatrix33& OutInvI0, 
+			FVec3& OutInvI0,
 			const FReal MaxInertiaRatio);
 
 
