Remove unions from capsules.
Compute bounds on the fly when appropriate

diff --git a/Engine/Source/Runtime/Core/Public/UObject/ExternalPhysicsCustomObjectVersion.h b/Engine/Source/Runtime/Core/Public/UObject/ExternalPhysicsCustomObjectVersion.h
--- a/Engine/Source/Runtime/Core/Public/UObject/ExternalPhysicsCustomObjectVersion.h
+++ b/Engine/Source/Runtime/Core/Public/UObject/ExternalPhysicsCustomObjectVersion.h
@@ -98,6 +98,9 @@ struct CORE_API FExternalPhysicsCustomObjectVersion
 
 		// Union objects can avoid allocating a full hierarchy
 		UnionObjectsCanAvoidHierarchy,
+
+		// Capsules no longer have a union inside them or stored aabbs
+		CapsulesNoUnionOrAABBs,
 		
 		// -----<new versions can be added above this line>-------------------------------------------------
 		VersionPlusOne,
diff --git a/Engine/Source/Runtime/Engine/Private/PhysicsEngine/Experimental/ChaosDerivedData.cpp b/Engine/Source/Runtime/Engine/Private/PhysicsEngine/Experimental/ChaosDerivedData.cpp
--- a/Engine/Source/Runtime/Engine/Private/PhysicsEngine/Experimental/ChaosDerivedData.cpp
+++ b/Engine/Source/Runtime/Engine/Private/PhysicsEngine/Experimental/ChaosDerivedData.cpp
@@ -21,7 +21,7 @@ const TCHAR* FChaosDerivedDataCooker::GetPluginName() const
 
 const TCHAR* FChaosDerivedDataCooker::GetVersionString() const
 {
-	return TEXT("53786FE2178C49A98EC62DD50313FE52");
+	return TEXT("385FFCA8CED24783AD05635A6835CF95");
 }
 
 FString FChaosDerivedDataCooker::GetPluginSpecificCacheKeySuffix() const
@@ -123,53 +123,53 @@ void FChaosDerivedDataCooker::BuildTriangleMeshes(TArray<TUniquePtr<Chaos::FTria
 
 	auto LambdaHelper = [&](auto& Triangles)
 	{
-	if(bHasMaterials)
-	{
-		MaterialIndices.Reserve(NumTriangles);
-	}
+		if(bHasMaterials)
+		{
+			MaterialIndices.Reserve(NumTriangles);
+		}
 
 		Triangles.Reserve(NumTriangles);
-	for(int32 TriangleIndex = 0; TriangleIndex < NumTriangles; ++TriangleIndex)
-	{
-		// Only add this triangle if it is valid
-		const int32 BaseIndex = TriangleIndex * 3;
-		const bool bIsValidTriangle = Chaos::FConvexBuilder::IsValidTriangle(
-			FinalVerts[FinalIndices[BaseIndex]],
-			FinalVerts[FinalIndices[BaseIndex + 1]],
-			FinalVerts[FinalIndices[BaseIndex + 2]]);
-
-		// TODO: Figure out a proper way to handle this. Could these edges get sewn together? Is this important?
-		//if (ensureMsgf(bIsValidTriangle, TEXT("FChaosDerivedDataCooker::BuildTriangleMeshes(): Trimesh attempted cooked with invalid triangle!")));
-		if (bIsValidTriangle)
+		for(int32 TriangleIndex = 0; TriangleIndex < NumTriangles; ++TriangleIndex)
 		{
-			Triangles.Add(Chaos::TVector<int32, 3>(FinalIndices[BaseIndex], FinalIndices[BaseIndex + 1], FinalIndices[BaseIndex + 2]));
-
-			if(bHasMaterials)
+			// Only add this triangle if it is valid
+			const int32 BaseIndex = TriangleIndex * 3;
+			const bool bIsValidTriangle = Chaos::FConvexBuilder::IsValidTriangle(
+				FinalVerts[FinalIndices[BaseIndex]],
+				FinalVerts[FinalIndices[BaseIndex + 1]],
+				FinalVerts[FinalIndices[BaseIndex + 2]]);
+
+			// TODO: Figure out a proper way to handle this. Could these edges get sewn together? Is this important?
+			//if (ensureMsgf(bIsValidTriangle, TEXT("FChaosDerivedDataCooker::BuildTriangleMeshes(): Trimesh attempted cooked with invalid triangle!")));
+			if (bIsValidTriangle)
 			{
-				if(!ensure(OutFaceRemap.IsValidIndex(TriangleIndex)))
-				{
-					MaterialIndices.Empty();
-					bHasMaterials = false;
-				}
-				else
-				{
-					const int32 OriginalIndex = OutFaceRemap[TriangleIndex];
+				Triangles.Add(Chaos::TVector<int32, 3>(FinalIndices[BaseIndex], FinalIndices[BaseIndex + 1], FinalIndices[BaseIndex + 2]));
 
-					if(ensure(InParams.TriangleMeshDesc.MaterialIndices.IsValidIndex(OriginalIndex)))
+				if(bHasMaterials)
+				{
+				if(!ensure(OutFaceRemap.IsValidIndex(TriangleIndex)))
 					{
-						MaterialIndices.Add(InParams.TriangleMeshDesc.MaterialIndices[OriginalIndex]);
+						MaterialIndices.Empty();
+						bHasMaterials = false;
 					}
 					else
 					{
-						MaterialIndices.Empty();
-						bHasMaterials = false;
+					const int32 OriginalIndex = OutFaceRemap[TriangleIndex];
+
+						if(ensure(InParams.TriangleMeshDesc.MaterialIndices.IsValidIndex(OriginalIndex)))
+						{
+							MaterialIndices.Add(InParams.TriangleMeshDesc.MaterialIndices[OriginalIndex]);
+						}
+						else
+						{
+							MaterialIndices.Empty();
+							bHasMaterials = false;
+						}
 					}
 				}
 			}
 		}
-	}
 
-	OutTriangleMeshes.Emplace(new Chaos::FTriangleMeshImplicitObject(MoveTemp(TriMeshParticles), MoveTemp(Triangles), MoveTemp(MaterialIndices)));
+		OutTriangleMeshes.Emplace(new Chaos::FTriangleMeshImplicitObject(MoveTemp(TriMeshParticles), MoveTemp(Triangles), MoveTemp(MaterialIndices)));
 	};
 
 	if(FinalVerts.Num() < TNumericLimits<uint16>::Max())
diff --git a/Engine/Source/Runtime/Experimental/Chaos/Private/Chaos/ImplicitObject.cpp b/Engine/Source/Runtime/Experimental/Chaos/Private/Chaos/ImplicitObject.cpp
--- a/Engine/Source/Runtime/Experimental/Chaos/Private/Chaos/ImplicitObject.cpp
+++ b/Engine/Source/Runtime/Experimental/Chaos/Private/Chaos/ImplicitObject.cpp
@@ -90,7 +90,7 @@ FVec3 FImplicitObject::Normal(const FVec3& x) const
 	return Normal;
 }
 
-const TAABB<FReal, 3>& FImplicitObject::BoundingBox() const
+const TAABB<FReal, 3> FImplicitObject::BoundingBox() const
 {
 	check(false);
 	static const TAABB<FReal, 3> Unbounded(FVec3(-FLT_MAX), FVec3(FLT_MAX));
diff --git a/Engine/Source/Runtime/Experimental/Chaos/Public/Chaos/BoundingVolumeUtilities.h b/Engine/Source/Runtime/Experimental/Chaos/Public/Chaos/BoundingVolumeUtilities.h
--- a/Engine/Source/Runtime/Experimental/Chaos/Public/Chaos/BoundingVolumeUtilities.h
+++ b/Engine/Source/Runtime/Experimental/Chaos/Public/Chaos/BoundingVolumeUtilities.h
@@ -89,7 +89,7 @@ bool HasBoundingBox(const TPBDRigidParticleHandleImp<T, d, bPersistent>& Handle)
 }
 
 template<class OBJECT_ARRAY, class T, int d>
-const TAABB<T, d>& GetWorldSpaceBoundingBox(const OBJECT_ARRAY& Objects, const int32 i, const TMap<int32, TAABB<T, d>>& WorldSpaceBoxes)
+const TAABB<T, d> GetWorldSpaceBoundingBox(const OBJECT_ARRAY& Objects, const int32 i, const TMap<int32, TAABB<T, d>>& WorldSpaceBoxes)
 {
 	return Objects[i]->BoundingBox();
 }
@@ -107,7 +107,7 @@ const TAABB<T, d>& GetWorldSpaceBoundingBox(const TGeometryParticles<T, d>& Obje
 }
 
 template<class T, int d>
-const TAABB<T, d>& GetWorldSpaceBoundingBox(const TPBDRigidParticles<T, d>& Objects, const int32 i, const TMap<int32, TAABB<T, d>>& WorldSpaceBoxes)
+const TAABB<T, d> GetWorldSpaceBoundingBox(const TPBDRigidParticles<T, d>& Objects, const int32 i, const TMap<int32, TAABB<T, d>>& WorldSpaceBoxes)
 {
 	return GetWorldSpaceBoundingBox(static_cast<const TParticles<T, d>&>(Objects), i, WorldSpaceBoxes);
 }
diff --git a/Engine/Source/Runtime/Experimental/Chaos/Public/Chaos/Box.h b/Engine/Source/Runtime/Experimental/Chaos/Public/Chaos/Box.h
--- a/Engine/Source/Runtime/Experimental/Chaos/Public/Chaos/Box.h
+++ b/Engine/Source/Runtime/Experimental/Chaos/Public/Chaos/Box.h
@@ -138,7 +138,7 @@ namespace Chaos
 
 		FORCEINLINE static constexpr EImplicitObjectType StaticType() { return ImplicitObjectType::Box; }
 
-		const TAABB<T, d>& BoundingBox() const { return AABB; }
+		const TAABB<T, d> BoundingBox() const { return AABB; }
 
 		virtual T PhiWithNormal(const TVector<T, d>& x, TVector<T, d>& Normal) const override
 		{
diff --git a/Engine/Source/Runtime/Experimental/Chaos/Public/Chaos/Capsule.h b/Engine/Source/Runtime/Experimental/Chaos/Public/Chaos/Capsule.h
--- a/Engine/Source/Runtime/Experimental/Chaos/Public/Chaos/Capsule.h
+++ b/Engine/Source/Runtime/Experimental/Chaos/Public/Chaos/Capsule.h
@@ -27,32 +27,21 @@ namespace Chaos
 		    : FImplicitObject(EImplicitObject::FiniteConvex, ImplicitObjectType::Capsule)
 			, MSegment(x1, x2)
 		    , MRadius(Radius)
-		    , MLocalBoundingBox(x1, x1)
-		    , MUnionedObjects(nullptr)
 		{
-			MLocalBoundingBox.GrowToInclude(x2);
-			MLocalBoundingBox = TAABB<T, 3>(MLocalBoundingBox.Min() - TVector<T, 3>(MRadius), MLocalBoundingBox.Max() + TVector<T, 3>(MRadius));
-			InitUnionedObjects();
 		}
 
 		TCapsule(const TCapsule<T>& Other)
 		    : FImplicitObject(EImplicitObject::FiniteConvex, ImplicitObjectType::Capsule)
 			, MSegment(Other.MSegment)
 		    , MRadius(Other.MRadius)
-		    , MLocalBoundingBox(Other.MLocalBoundingBox)
-		    , MUnionedObjects(nullptr)
 		{
-			InitUnionedObjects();
 		}
 
 		TCapsule(TCapsule<T>&& Other)
 		    : FImplicitObject(EImplicitObject::FiniteConvex, ImplicitObjectType::Capsule)
 			, MSegment(MoveTemp(Other.MSegment))
 		    , MRadius(Other.MRadius)
-		    , MLocalBoundingBox(MoveTemp(Other.MLocalBoundingBox))
-		    , MUnionedObjects(MoveTemp(Other.MUnionedObjects))
 		{
-			InitUnionedObjects();
 		}
 
 		TCapsule& operator=(TCapsule<T>&& InSteal)
@@ -64,10 +53,6 @@ namespace Chaos
 
 			MSegment = MoveTemp(InSteal.MSegment);
 			MRadius = InSteal.MRadius;
-			MLocalBoundingBox = MoveTemp(InSteal.MLocalBoundingBox);
-			MUnionedObjects = MoveTemp(InSteal.MUnionedObjects);
-
-			InitUnionedObjects();
 
 			return *this;
 		}
@@ -131,7 +116,12 @@ namespace Chaos
 			return Normal.SafeNormalize() - MRadius;
 		}
 
-		virtual const TAABB<T, 3>& BoundingBox() const override { return MLocalBoundingBox; }
+		virtual const TAABB<T, 3> BoundingBox() const override
+		{
+			TAABB<T,3> Box = MSegment.BoundingBox();
+			Box.Thicken(MRadius);
+			return Box;
+		}
 
 		static bool RaycastFast(T MRadius, T MHeight, const TVector<T,3>& MVector, const TVector<T,3>& X1, const TVector<T,3>& X2, const TVector<T, 3>& StartPoint, const TVector<T, 3>& Dir, const T Length, const T Thickness, T& OutTime, TVector<T, 3>& OutPosition, TVector<T, 3>& OutNormal, int32& OutFaceIndex)
 		{
@@ -301,17 +291,30 @@ namespace Chaos
 
 		FORCEINLINE void SerializeImp(FArchive& Ar)
 		{
+			Ar.UsingCustomVersion(FExternalPhysicsCustomObjectVersion::GUID);
 			FImplicitObject::SerializeImp(Ar);
 			MSegment.Serialize(Ar);
 			Ar << MRadius;
-			TBox<FReal, 3>::SerializeAsAABB(Ar, MLocalBoundingBox);
+
+			
+			if(Ar.CustomVer(FExternalPhysicsCustomObjectVersion::GUID) < FExternalPhysicsCustomObjectVersion::CapsulesNoUnionOrAABBs)
+			{
+				TAABB<T,3> DummyBox;	//no longer store this, computed on demand
+				TBox<FReal,3>::SerializeAsAABB(Ar,DummyBox);
+			}
 		}
 
 		virtual void Serialize(FChaosArchive& Ar) override
 		{
+			Ar.UsingCustomVersion(FExternalPhysicsCustomObjectVersion::GUID);
 			FChaosArchiveScopedMemory ScopedMemory(Ar, GetTypeName());
 			SerializeImp(Ar);
-			Ar << MUnionedObjects;
+
+			if(Ar.CustomVer(FExternalPhysicsCustomObjectVersion::GUID) < FExternalPhysicsCustomObjectVersion::CapsulesNoUnionOrAABBs)
+			{
+				TUniquePtr<FImplicitObjectUnion> TmpUnion;
+				Ar << TmpUnion;
+			}
 		}
 
 		virtual TUniquePtr<FImplicitObject> Copy() const override
@@ -367,28 +370,9 @@ namespace Chaos
 		}
 
 	private:
-		void InitUnionedObjects()
-		{
-			TArray<TUniquePtr<FImplicitObject>> Objects;
-
-			const TVector<T, 3>& X1 = GetX1();
-			const TVector<T, 3>& X2 = GetX2();
-			Objects.Add(MakeUnique<Chaos::TCylinder<float>>(X1, X2, MRadius));
-			Objects.Add(MakeUnique<Chaos::TSphere<float, 3>>(X1, MRadius));
-			Objects.Add(MakeUnique<Chaos::TSphere<float, 3>>(X2, MRadius));
-
-			MUnionedObjects.Reset(new Chaos::FImplicitObjectUnion(std::move(Objects)));
-		}
-
-		virtual Pair<TVector<T, 3>, bool> FindClosestIntersectionImp(const TVector<T, 3>& StartPoint, const TVector<T, 3>& EndPoint, const T Thickness) const override
-		{
-			return MUnionedObjects->FindClosestIntersection(StartPoint, EndPoint, Thickness);
-		}
 
 		TSegment<T> MSegment;
 		T MRadius;
-		TAABB<T, 3> MLocalBoundingBox;
-		TUniquePtr<FImplicitObjectUnion> MUnionedObjects;
 	};
 
 	template<typename T>
diff --git a/Engine/Source/Runtime/Experimental/Chaos/Public/Chaos/Convex.h b/Engine/Source/Runtime/Experimental/Chaos/Public/Chaos/Convex.h
--- a/Engine/Source/Runtime/Experimental/Chaos/Public/Chaos/Convex.h
+++ b/Engine/Source/Runtime/Experimental/Chaos/Public/Chaos/Convex.h
@@ -72,7 +72,7 @@ namespace Chaos
 			return ImplicitObjectType::Convex;
 		}
 
-		virtual const TAABB<FReal, 3>& BoundingBox() const override
+		virtual const TAABB<FReal, 3> BoundingBox() const override
 		{
 			return LocalBoundingBox;
 		}
diff --git a/Engine/Source/Runtime/Experimental/Chaos/Public/Chaos/Cylinder.h b/Engine/Source/Runtime/Experimental/Chaos/Public/Chaos/Cylinder.h
--- a/Engine/Source/Runtime/Experimental/Chaos/Public/Chaos/Cylinder.h
+++ b/Engine/Source/Runtime/Experimental/Chaos/Public/Chaos/Cylinder.h
@@ -169,7 +169,7 @@ namespace Chaos
 			return SideDistance;
 		}
 
-		virtual const TAABB<T, 3>& BoundingBox() const override { return MLocalBoundingBox; }
+		virtual const TAABB<T, 3> BoundingBox() const override { return MLocalBoundingBox; }
 
 		T GetRadius() const { return MRadius; }
 		T GetHeight() const { return MHeight; }
diff --git a/Engine/Source/Runtime/Experimental/Chaos/Public/Chaos/HeightField.h b/Engine/Source/Runtime/Experimental/Chaos/Public/Chaos/HeightField.h
--- a/Engine/Source/Runtime/Experimental/Chaos/Public/Chaos/HeightField.h
+++ b/Engine/Source/Runtime/Experimental/Chaos/Public/Chaos/HeightField.h
@@ -109,7 +109,7 @@ namespace Chaos
 			return 0;
 		}
 
-		virtual const TAABB<T, 3>& BoundingBox() const
+		virtual const TAABB<T, 3> BoundingBox() const
 		{
 			CachedBounds = TAABB<T, 3>(LocalBounds.Min() * GeomData.Scale, LocalBounds.Max() * GeomData.Scale);
 			return CachedBounds;
diff --git a/Engine/Source/Runtime/Experimental/Chaos/Public/Chaos/ImplicitObject.h b/Engine/Source/Runtime/Experimental/Chaos/Public/Chaos/ImplicitObject.h
--- a/Engine/Source/Runtime/Experimental/Chaos/Public/Chaos/ImplicitObject.h
+++ b/Engine/Source/Runtime/Experimental/Chaos/Public/Chaos/ImplicitObject.h
@@ -174,7 +174,7 @@ public:
 	// Explicitly non-virtual.  Must cast to derived types to target their implementation.
 	FVec3 Normal(const FVec3& x) const;
 	virtual FReal PhiWithNormal(const FVec3& x, FVec3& Normal) const = 0;
-	virtual const class TAABB<FReal, 3>& BoundingBox() const;
+	virtual const class TAABB<FReal, 3> BoundingBox() const;
 	bool HasBoundingBox() const { return bHasBoundingBox; }
 	bool IsConvex() const { return bIsConvex; }
 	void SetDoCollide(const bool Collide ) { bDoCollide = Collide; }
diff --git a/Engine/Source/Runtime/Experimental/Chaos/Public/Chaos/ImplicitObjectIntersection.h b/Engine/Source/Runtime/Experimental/Chaos/Public/Chaos/ImplicitObjectIntersection.h
--- a/Engine/Source/Runtime/Experimental/Chaos/Public/Chaos/ImplicitObjectIntersection.h
+++ b/Engine/Source/Runtime/Experimental/Chaos/Public/Chaos/ImplicitObjectIntersection.h
@@ -55,7 +55,7 @@ class TImplicitObjectIntersection : public FImplicitObject
 		return Phi;
 	}
 
-	virtual const TAABB<T,d>& BoundingBox() const { return MLocalBoundingBox; }
+	virtual const TAABB<T,d> BoundingBox() const { return MLocalBoundingBox; }
 
 
 	virtual uint32 GetTypeHash() const override
diff --git a/Engine/Source/Runtime/Experimental/Chaos/Public/Chaos/ImplicitObjectScaled.h b/Engine/Source/Runtime/Experimental/Chaos/Public/Chaos/ImplicitObjectScaled.h
--- a/Engine/Source/Runtime/Experimental/Chaos/Public/Chaos/ImplicitObjectScaled.h
+++ b/Engine/Source/Runtime/Experimental/Chaos/Public/Chaos/ImplicitObjectScaled.h
@@ -95,7 +95,7 @@ public:
 	FORCEINLINE TVector<T, d> Support(const TVector<T, d>& Direction, const T Thickness) const { return MObject->Support(Direction, Thickness); }
 	FORCEINLINE TVector<T, d> Support2(const TVector<T, d>& Direction, const T Thickness) const { return MObject->Support2(Direction); }
 
-	virtual const TAABB<T, d>& BoundingBox() const override { return MObject->BoundingBox(); }
+	virtual const TAABB<T, d> BoundingBox() const override { return MObject->BoundingBox(); }
 
 	const ObjectType Object() const { return MObject; }
 
@@ -487,7 +487,7 @@ public:
 		UpdateBounds();
 	}
 
-	virtual const TAABB<T, d>& BoundingBox() const override { return MLocalBoundingBox; }
+	virtual const TAABB<T, d> BoundingBox() const override { return MLocalBoundingBox; }
 
 	const FReal GetVolume() const
 	{
diff --git a/Engine/Source/Runtime/Experimental/Chaos/Public/Chaos/ImplicitObjectTransformed.h b/Engine/Source/Runtime/Experimental/Chaos/Public/Chaos/ImplicitObjectTransformed.h
--- a/Engine/Source/Runtime/Experimental/Chaos/Public/Chaos/ImplicitObjectTransformed.h
+++ b/Engine/Source/Runtime/Experimental/Chaos/Public/Chaos/ImplicitObjectTransformed.h
@@ -198,7 +198,7 @@ public:
 		}
 	}
 
-	virtual const TAABB<T, d>& BoundingBox() const override { return MLocalBoundingBox; }
+	virtual const TAABB<T, d> BoundingBox() const override { return MLocalBoundingBox; }
 
 	const FReal GetVolume() const
 	{
diff --git a/Engine/Source/Runtime/Experimental/Chaos/Public/Chaos/ImplicitObjectUnion.h b/Engine/Source/Runtime/Experimental/Chaos/Public/Chaos/ImplicitObjectUnion.h
--- a/Engine/Source/Runtime/Experimental/Chaos/Public/Chaos/ImplicitObjectUnion.h
+++ b/Engine/Source/Runtime/Experimental/Chaos/Public/Chaos/ImplicitObjectUnion.h
@@ -61,7 +61,7 @@ class CHAOS_API FImplicitObjectUnion : public FImplicitObject
 		return Phi;
 	}
 
-	virtual const TAABB<FReal, 3>& BoundingBox() const override { return MLocalBoundingBox; }
+	virtual const TAABB<FReal, 3> BoundingBox() const override { return MLocalBoundingBox; }
 
 	virtual void AccumulateAllImplicitObjects(TArray<Pair<const FImplicitObject*, FRigidTransform3>>& Out, const FRigidTransform3& ParentTM) const
 	{
diff --git a/Engine/Source/Runtime/Experimental/Chaos/Public/Chaos/Levelset.h b/Engine/Source/Runtime/Experimental/Chaos/Public/Chaos/Levelset.h
--- a/Engine/Source/Runtime/Experimental/Chaos/Public/Chaos/Levelset.h
+++ b/Engine/Source/Runtime/Experimental/Chaos/Public/Chaos/Levelset.h
@@ -35,7 +35,7 @@ class CHAOS_API TLevelSet final : public FImplicitObject
 	virtual T PhiWithNormal(const TVector<T, d>& x, TVector<T, d>& Normal) const override;
 	T SignedDistance(const TVector<T, d>& x) const;
 
-	virtual const TAABB<T, d>& BoundingBox() const override { return MOriginalLocalBoundingBox; }
+	virtual const TAABB<T, d> BoundingBox() const override { return MOriginalLocalBoundingBox; }
 
 	// Returns a const ref to the underlying phi grid
 	const TArrayND<T, d>& GetPhiArray() const { return MPhi; }
diff --git a/Engine/Source/Runtime/Experimental/Chaos/Public/Chaos/Segment.h b/Engine/Source/Runtime/Experimental/Chaos/Public/Chaos/Segment.h
--- a/Engine/Source/Runtime/Experimental/Chaos/Public/Chaos/Segment.h
+++ b/Engine/Source/Runtime/Experimental/Chaos/Public/Chaos/Segment.h
@@ -53,6 +53,13 @@ namespace Chaos
 			Ar << MPoint << MAxis << MLength;
 		}
 
+		FORCEINLINE TAABB<T, 3> BoundingBox() const
+		{
+			TAABB<T,3> Box(MPoint,MPoint);
+			Box.GrowToInclude(GetX2());
+			return Box;
+		}
+
 	private:
 		TVector<T, 3> MPoint;
 		TVector<T, 3> MAxis;
diff --git a/Engine/Source/Runtime/Experimental/Chaos/Public/Chaos/Sphere.h b/Engine/Source/Runtime/Experimental/Chaos/Public/Chaos/Sphere.h
--- a/Engine/Source/Runtime/Experimental/Chaos/Public/Chaos/Sphere.h
+++ b/Engine/Source/Runtime/Experimental/Chaos/Public/Chaos/Sphere.h
@@ -27,7 +27,6 @@ namespace Chaos
 		    : FImplicitObject(EImplicitObject::IsConvex | EImplicitObject::HasBoundingBox, ImplicitObjectType::Sphere)
 		    , Center(InCenter)
 		    , Radius(InRadius)
-		    , LocalBoundingBox(Center - Radius, Center + Radius)
 		{
 		}
 
@@ -35,7 +34,6 @@ namespace Chaos
 		    : FImplicitObject(EImplicitObject::IsConvex | EImplicitObject::HasBoundingBox, ImplicitObjectType::Sphere)
 		    , Center(Other.Center)
 		    , Radius(Other.Radius)
-		    , LocalBoundingBox(Other.LocalBoundingBox)
 		{
 		}
 
@@ -43,7 +41,6 @@ namespace Chaos
 		    : FImplicitObject(EImplicitObject::IsConvex | EImplicitObject::HasBoundingBox, ImplicitObjectType::Sphere)
 		    , Center(MoveTemp(Other.Center))
 		    , Radius(Other.Radius)
-		    , LocalBoundingBox(MoveTemp(Other.LocalBoundingBox))
 		{
 		}
 
@@ -55,7 +52,6 @@ namespace Chaos
 			this->bHasBoundingBox = InSteal.bHasBoundingBox;
 			Center = MoveTemp(InSteal.Center);
 			Radius = InSteal.Radius;
-			LocalBoundingBox = MoveTemp(InSteal.LocalBoundingBox);
 
 			return *this;
 		}
@@ -198,9 +194,9 @@ namespace Chaos
 			return Radius;
 		}
 
-		virtual const TAABB<T, d>& BoundingBox() const 
+		virtual const TAABB<T, d> BoundingBox() const 
 		{
-			return LocalBoundingBox; 
+			return TAABB<T,d>(Center - TVector<T,d>(Radius),Center + TVector<T,d>(Radius));
 		}
 
 		T GetArea() const 
@@ -251,10 +247,6 @@ namespace Chaos
 		{
 			FImplicitObject::SerializeImp(Ar);
 			Ar << Center << Radius;
-			if (Ar.IsLoading())
-			{
-				LocalBoundingBox = TAABB<T, d>(Center - Radius, Center + Radius);
-			}
 		}
 
 		virtual void Serialize(FChaosArchive& Ar) override 
@@ -321,8 +313,7 @@ namespace Chaos
 		{
 			const uint32 CenterHash = ::GetTypeHash(Center);
 			const uint32 RadiusHash = ::GetTypeHash(Radius);
-			const uint32 BoundsHash = LocalBoundingBox.GetTypeHash();
-			return HashCombine(CenterHash, HashCombine(RadiusHash, BoundsHash));
+			return HashCombine(CenterHash, RadiusHash);
 		}
 
 		virtual TUniquePtr<FImplicitObject> Copy() const override
@@ -333,7 +324,6 @@ namespace Chaos
 	private:
 		TVector<T, d> Center;
 		T Radius;
-		TAABB<T, d> LocalBoundingBox;
 
 	private:
 
diff --git a/Engine/Source/Runtime/Experimental/Chaos/Public/Chaos/TaperedCylinder.h b/Engine/Source/Runtime/Experimental/Chaos/Public/Chaos/TaperedCylinder.h
--- a/Engine/Source/Runtime/Experimental/Chaos/Public/Chaos/TaperedCylinder.h
+++ b/Engine/Source/Runtime/Experimental/Chaos/Public/Chaos/TaperedCylinder.h
@@ -112,7 +112,7 @@ namespace Chaos
 		TArray<TVector<T, 3>> ComputeSamplePoints(const T PointsPerUnitArea, const bool IncludeEndCaps = true, const int32 MinPoints = 0, const int32 MaxPoints = 1000) const
 		{ return ComputeSamplePoints(FMath::Clamp(static_cast<int32>(ceil(PointsPerUnitArea * GetArea(IncludeEndCaps))), MinPoints, MaxPoints), IncludeEndCaps); }
 
-		virtual const TAABB<T, 3>& BoundingBox() const override { return MLocalBoundingBox; }
+		virtual const TAABB<T, 3> BoundingBox() const override { return MLocalBoundingBox; }
 
 		T PhiWithNormal(const TVector<T, 3>& x, TVector<T, 3>& Normal) const
 		{
diff --git a/Engine/Source/Runtime/Experimental/Chaos/Public/Chaos/Triangle.h b/Engine/Source/Runtime/Experimental/Chaos/Public/Chaos/Triangle.h
--- a/Engine/Source/Runtime/Experimental/Chaos/Public/Chaos/Triangle.h
+++ b/Engine/Source/Runtime/Experimental/Chaos/Public/Chaos/Triangle.h
@@ -155,9 +155,6 @@ namespace Chaos
 			: FImplicitObject(EImplicitObject::IsConvex | EImplicitObject::HasBoundingBox, ImplicitObjectType::Triangle)
 			, Tri(InA, InB, InC)
 		{
-			Bounds = TAABB<T, 3>(Tri[0], Tri[0]);
-			Bounds.GrowToInclude(Tri[1]);
-			Bounds.GrowToInclude(Tri[2]);
 		}
 
 		TVector<T, 3>& operator[](uint32 InIndex)
@@ -190,8 +187,12 @@ namespace Chaos
 			return Tri.PhiWithNormal(InSamplePoint, OutNormal);
 		}
 
-		virtual const class TAABB<T, 3>& BoundingBox() const override
+		virtual const class TAABB<T, 3> BoundingBox() const override
 		{
+			TAABB<T,3> Bounds(Tri[0],Tri[0]);
+			Bounds.GrowToInclude(Tri[1]);
+			Bounds.GrowToInclude(Tri[2]);
+
 			return Bounds;
 		}
 
@@ -227,7 +228,7 @@ namespace Chaos
 
 		virtual uint32 GetTypeHash() const override
 		{
-			return HashCombine(Bounds.GetTypeHash(), HashCombine(::GetTypeHash(Tri[0]), HashCombine(::GetTypeHash(Tri[1]), ::GetTypeHash(Tri[2]))));
+			return HashCombine(::GetTypeHash(Tri[0]), HashCombine(::GetTypeHash(Tri[1]), ::GetTypeHash(Tri[2])));
 		}
 
 		virtual FName GetTypeName() const override
@@ -238,6 +239,5 @@ namespace Chaos
 	private:
 
 		TTriangle<T> Tri;
-		TAABB<T, 3> Bounds;
 	};
 }
\ No newline at end of file
diff --git a/Engine/Source/Runtime/Experimental/Chaos/Public/Chaos/TriangleMeshImplicitObject.h b/Engine/Source/Runtime/Experimental/Chaos/Public/Chaos/TriangleMeshImplicitObject.h
--- a/Engine/Source/Runtime/Experimental/Chaos/Public/Chaos/TriangleMeshImplicitObject.h
+++ b/Engine/Source/Runtime/Experimental/Chaos/Public/Chaos/TriangleMeshImplicitObject.h
@@ -156,7 +156,7 @@ namespace Chaos
 		virtual int32 FindMostOpposingFace(const FVec3& Position, const FVec3& UnitDir, int32 HintFaceIndex, FReal SearchDistance) const override;
 		virtual FVec3 FindGeometryOpposingNormal(const FVec3& DenormDir, int32 FaceIndex, const FVec3& OriginalNormal) const override;
 
-		virtual const TAABB<FReal, 3>& BoundingBox() const
+		virtual const TAABB<FReal, 3> BoundingBox() const
 		{
 			return MLocalBoundingBox;
 		}
