Chaos : Collisions

- Setup ImplicitObjectScaled and ImplicitObjectTransformed to mirror the collision flag of its owned MObject.

diff --git a/Engine/Source/Runtime/Core/Public/UObject/DestructionObjectVersion.h b/Engine/Source/Runtime/Core/Public/UObject/DestructionObjectVersion.h
--- a/Engine/Source/Runtime/Core/Public/UObject/DestructionObjectVersion.h
+++ b/Engine/Source/Runtime/Core/Public/UObject/DestructionObjectVersion.h
@@ -39,6 +39,10 @@ struct CORE_API FDestructionObjectVersion
 		// bulk serialize arrays
 		GroupAndAttributeNameRemapping,
 
+		// bulk serialize arrays
+		ImplicitObjectDoCollideAttribute,
+
+
 		// -----<new versions can be added above this line>-------------------------------------------------
 		VersionPlusOne,
 		LatestVersion = VersionPlusOne - 1
diff --git a/Engine/Source/Runtime/Engine/Private/PhysicsEngine/BodySetup.cpp b/Engine/Source/Runtime/Engine/Private/PhysicsEngine/BodySetup.cpp
--- a/Engine/Source/Runtime/Engine/Private/PhysicsEngine/BodySetup.cpp
+++ b/Engine/Source/Runtime/Engine/Private/PhysicsEngine/BodySetup.cpp
@@ -742,6 +742,14 @@ void UBodySetup::FinishCreatingPhysicsMeshes_Chaos(FChaosDerivedDataReader<float
 	}
 #endif
 
+#if WITH_CHAOS
+	// Force trimesh collisions off
+	for (auto& TriMesh : ChaosTriMeshes)
+	{
+		TriMesh->SetDoCollide(false);
+	}
+#endif
+
 	// Clear the cooked data
 	if (!GIsEditor && !bSharedCookedData)
 	{
diff --git a/Engine/Source/Runtime/Engine/Private/PhysicsEngine/Experimental/PhysInterface_Chaos.cpp b/Engine/Source/Runtime/Engine/Private/PhysicsEngine/Experimental/PhysInterface_Chaos.cpp
--- a/Engine/Source/Runtime/Engine/Private/PhysicsEngine/Experimental/PhysInterface_Chaos.cpp
+++ b/Engine/Source/Runtime/Engine/Private/PhysicsEngine/Experimental/PhysInterface_Chaos.cpp
@@ -1523,7 +1523,6 @@ bool CalculateMassPropertiesOfImplicitType(
 	if (ImplicitObject)
 	{
 		// Hack to handle Transformed and Scaled<ImplicitObjectTriangleMesh> until CastHelper can properly support transformed
-		Chaos::FRigidTransform3 Transform(FMatrix::Identity);
 		// Commenting this out temporarily as it breaks vehicles
 		/*	if (Chaos::IsScaled(ImplicitObject->GetType(true)) && Chaos::GetInnerType(ImplicitObject->GetType(true)) & Chaos::ImplicitObjectType::TriangleMesh)
 			{
diff --git a/Engine/Source/Runtime/Experimental/Chaos/Private/Chaos/ImplicitObject.cpp b/Engine/Source/Runtime/Experimental/Chaos/Private/Chaos/ImplicitObject.cpp
--- a/Engine/Source/Runtime/Experimental/Chaos/Private/Chaos/ImplicitObject.cpp
+++ b/Engine/Source/Runtime/Experimental/Chaos/Private/Chaos/ImplicitObject.cpp
@@ -22,7 +22,7 @@ using namespace Chaos;
 FImplicitObject::FImplicitObject(int32 Flags, EImplicitObjectType InType)
     : Type(InType)
     , bIsConvex(!!(Flags & EImplicitObject::IsConvex))
-    , bDoCollide(true)
+    , bDoCollide(!(Flags & EImplicitObject::DisableCollisions))
     , bHasBoundingBox(!!(Flags & EImplicitObject::HasBoundingBox))
 #if TRACK_CHAOS_GEOMETRY
     , bIsTracked(false)
@@ -316,6 +316,11 @@ void FImplicitObject::SerializeImp(FArchive& Ar)
 	{
 		Ar << bIsConvex << bDoCollide;
 	}
+
+	if (Ar.CustomVer(FDestructionObjectVersion::GUID) <= FDestructionObjectVersion::ImplicitObjectDoCollideAttribute)
+	{
+		bDoCollide = true;
+	}
 }
 
 void FImplicitObject::Serialize(FChaosArchive& Ar)
diff --git a/Engine/Source/Runtime/Experimental/Chaos/Public/Chaos/ImplicitObject.h b/Engine/Source/Runtime/Experimental/Chaos/Public/Chaos/ImplicitObject.h
--- a/Engine/Source/Runtime/Experimental/Chaos/Public/Chaos/ImplicitObject.h
+++ b/Engine/Source/Runtime/Experimental/Chaos/Public/Chaos/ImplicitObject.h
@@ -79,6 +79,7 @@ namespace EImplicitObject
 	{
 		IsConvex = 1,
 		HasBoundingBox = 1 << 1,
+		DisableCollisions = 1 << 2
 	};
 
 	const int32 FiniteConvex = IsConvex | HasBoundingBox;
@@ -176,10 +177,12 @@ public:
 	virtual FReal PhiWithNormal(const FVec3& x, FVec3& Normal) const = 0;
 	virtual const class TAABB<FReal, 3> BoundingBox() const;
 	bool HasBoundingBox() const { return bHasBoundingBox; }
+
 	bool IsConvex() const { return bIsConvex; }
+	void SetConvex(const bool Convex = true) { bIsConvex = Convex; }
+
 	void SetDoCollide(const bool Collide ) { bDoCollide = Collide; }
 	bool GetDoCollide() const { return bDoCollide; }
-	void SetConvex(const bool Convex = true) { bIsConvex = Convex; }
 	
 #if TRACK_CHAOS_GEOMETRY
 	//Turn on memory tracking. Must pass object itself as a serializable ptr so we can save it out
diff --git a/Engine/Source/Runtime/Experimental/Chaos/Public/Chaos/ImplicitObjectScaled.h b/Engine/Source/Runtime/Experimental/Chaos/Public/Chaos/ImplicitObjectScaled.h
--- a/Engine/Source/Runtime/Experimental/Chaos/Public/Chaos/ImplicitObjectScaled.h
+++ b/Engine/Source/Runtime/Experimental/Chaos/Public/Chaos/ImplicitObjectScaled.h
@@ -31,6 +31,7 @@ public:
 	{
 		ensure(IsInstanced(MObject->GetType(true)) == false);	//cannot have an instance of an instance
 		this->bIsConvex = MObject->IsConvex();
+		this->bDoCollide = MObject->GetDoCollide();
 	}
 
 	static constexpr EImplicitObjectType StaticType()
@@ -43,6 +44,11 @@ public:
 		return MObject.Get();
 	}
 
+	bool GetDoCollide() const
+	{
+		return MObject->GetDoCollide();
+	}
+
 	virtual T PhiWithNormal(const TVector<T, d>& X, TVector<T, d>& Normal) const override
 	{
 		return MObject->PhiWithNormal(X, Normal);
@@ -153,6 +159,7 @@ public:
 			break;
 		}
 		this->bIsConvex = MObject->IsConvex();
+		this->bDoCollide = MObject->GetDoCollide();
 		SetScale(Scale);
 	}
 
@@ -173,6 +180,7 @@ public:
 			break;
 		}
 		this->bIsConvex = MObject->IsConvex();
+		this->bDoCollide = MObject->GetDoCollide();
 		SetScale(Scale);
 	}
 
@@ -185,6 +193,7 @@ public:
 		ensureMsgf((IsScaled(MObject->GetType(true)) == false), TEXT("Scaled objects should not contain each other."));
 		ensureMsgf((IsInstanced(MObject->GetType(true)) == false), TEXT("Scaled objects should not contain instances."));
 		this->bIsConvex = Object->IsConvex();
+		this->bDoCollide = MObject->GetDoCollide();
 		SetScale(Scale);
 	}
 
@@ -201,6 +210,7 @@ public:
 		ensureMsgf((IsScaled(MObject->GetType(true)) == false), TEXT("Scaled objects should not contain each other."));
 		ensureMsgf((IsInstanced(MObject->GetType(true)) == false), TEXT("Scaled objects should not contain instances."));
 		this->bIsConvex = Other.MObject->IsConvex();
+		this->bDoCollide = Other.MObject->GetDoCollide();
 	}
 	~TImplicitObjectScaled() {}
 
diff --git a/Engine/Source/Runtime/Experimental/Chaos/Public/Chaos/ImplicitObjectTransformed.h b/Engine/Source/Runtime/Experimental/Chaos/Public/Chaos/ImplicitObjectTransformed.h
--- a/Engine/Source/Runtime/Experimental/Chaos/Public/Chaos/ImplicitObjectTransformed.h
+++ b/Engine/Source/Runtime/Experimental/Chaos/Public/Chaos/ImplicitObjectTransformed.h
@@ -54,6 +54,7 @@ public:
 	    , MLocalBoundingBox(Object->BoundingBox().TransformedAABB(InTransform))
 	{
 		this->bIsConvex = Object->IsConvex();
+		this->bDoCollide = MObject->GetDoCollide();
 	}
 
 	/**
@@ -68,6 +69,7 @@ public:
 		this->MObject = FStorage::Convert(MObjectOwner);
 		this->MLocalBoundingBox = MObject->BoundingBox().TransformedAABB(InTransform);
 		this->bIsConvex = MObject->IsConvex();
+		this->bDoCollide = MObject->GetDoCollide();
 	}
 
 	TImplicitObjectTransformed(const TImplicitObjectTransformed<T, d, bSerializable>& Other) = delete;
@@ -79,6 +81,7 @@ public:
 	    , MLocalBoundingBox(MoveTemp(Other.MLocalBoundingBox))
 	{
 		this->bIsConvex = Other.MObject->IsConvex();
+		this->bDoCollide = Other.MObject->GetDoCollide();
 	}
 	~TImplicitObjectTransformed() {}
 
@@ -92,6 +95,11 @@ public:
 		return MObject.Get();
 	}
 
+	bool GetDoCollide() const
+	{
+		return MObject->GetDoCollide();
+	}
+
 	virtual T PhiWithNormal(const TVector<T, d>& x, TVector<T, d>& Normal) const override
 	{
 		auto TransformedX = MTransform.InverseTransformPosition(x);
diff --git a/Engine/Source/Runtime/Experimental/Chaos/Public/Chaos/TriangleMeshImplicitObject.h b/Engine/Source/Runtime/Experimental/Chaos/Public/Chaos/TriangleMeshImplicitObject.h
--- a/Engine/Source/Runtime/Experimental/Chaos/Public/Chaos/TriangleMeshImplicitObject.h
+++ b/Engine/Source/Runtime/Experimental/Chaos/Public/Chaos/TriangleMeshImplicitObject.h
@@ -102,7 +102,7 @@ namespace Chaos
 
 		template <typename IdxType>
 		FTriangleMeshImplicitObject(TParticles<FReal, 3>&& Particles, TArray<TVector<IdxType, 3>>&& Elements, TArray<uint16>&& InMaterialIndices)
-		: FImplicitObject(EImplicitObject::HasBoundingBox, ImplicitObjectType::TriangleMesh)
+		: FImplicitObject(EImplicitObject::HasBoundingBox | EImplicitObject::DisableCollisions, ImplicitObjectType::TriangleMesh)
 		, MParticles(MoveTemp(Particles))
 		, MElements(MoveTemp(Elements))
 		, MLocalBoundingBox(MParticles.X(0), MParticles.X(0))
