Avoid marking force or torque dirty if they are identical values.
This seems like a workaround for the problem that ClearAccumulatedData is called after marking proxy as clean, only to have the force and torque set it back to dirty.
I assume there's a better solution than this, but it fixes the issue of 4k objects being dirty all the time

diff --git a/Engine/Source/Runtime/Experimental/Chaos/Public/Chaos/ParticleHandle.h b/Engine/Source/Runtime/Experimental/Chaos/Public/Chaos/ParticleHandle.h
--- a/Engine/Source/Runtime/Experimental/Chaos/Public/Chaos/ParticleHandle.h
+++ b/Engine/Source/Runtime/Experimental/Chaos/Public/Chaos/ParticleHandle.h
@@ -1827,15 +1827,24 @@ public:
 	const TVector<T, d>& F() const { return MF; }
 	void SetF(const TVector<T, d>& InF)
 	{
-		this->MarkDirty(EParticleFlags::F);
-		this->MF = InF;
+		//question: should we do this check? only adding because we clear forces after removing from dirty list, but this marks dirty
+		if(InF != this->MF)
+		{
+			this->MarkDirty(EParticleFlags::F);
+			this->MF = InF;
+		}
+		
 	}
 
 	const TVector<T, d>& Torque() const { return MTorque; }
 	void SetTorque(const TVector<T, d>& InTorque)
 	{
-		this->MarkDirty(EParticleFlags::Torque);
-		this->MTorque = InTorque;
+		//question: should we do this check? only adding because we clear forces after removing from dirty list, but this marks dirty
+		if(InTorque != this->MTorque)
+		{
+			this->MarkDirty(EParticleFlags::Torque);
+			this->MTorque = InTorque;
+		}
 	}
 
 	const TVector<T, d>& LinearImpulse() const { return MLinearImpulse; }
