Chaos - Joint angular constraint stiffness improvement stability fix

diff --git a/Engine/Source/Runtime/Engine/Private/PhysicsEngine/ImmediatePhysics/ImmediatePhysicsChaos/ImmediatePhysicsSimulation_Chaos.cpp b/Engine/Source/Runtime/Engine/Private/PhysicsEngine/ImmediatePhysics/ImmediatePhysicsChaos/ImmediatePhysicsSimulation_Chaos.cpp
--- a/Engine/Source/Runtime/Engine/Private/PhysicsEngine/ImmediatePhysics/ImmediatePhysicsChaos/ImmediatePhysicsSimulation_Chaos.cpp
+++ b/Engine/Source/Runtime/Engine/Private/PhysicsEngine/ImmediatePhysics/ImmediatePhysicsChaos/ImmediatePhysicsSimulation_Chaos.cpp
@@ -17,6 +17,8 @@
 #include "PhysicsEngine/BodyInstance.h"
 #include "PhysicsEngine/ConstraintInstance.h"
 
+//#pragma optimize("", off)
+
 //////////////////////////////////////////////////////////////////////////
 // @todo(ccaulfield): remove when finished
 //
@@ -64,7 +66,7 @@ float ChaosImmediate_Joint_AngularDriveStiffness = 0.0f;
 float ChaosImmediate_Joint_AngularDriveDamping = 0.0f;
 float ChaosImmediate_Joint_MinParentMassRatio = 0.2f;
 float ChaosImmediate_Joint_MaxInertiaRatio = 5.0f;
-float ChaosImmediate_Joint_AngularPositionCorrection = 0.7f;
+float ChaosImmediate_Joint_AngularPositionCorrection = 1.0f;
 FAutoConsoleVariableRef CVarChaosImmPhysPairIterations(TEXT("p.Chaos.ImmPhys.Joint.PairIterations"), ChaosImmediate_Joint_PairIterations, TEXT("PairIterations."));
 FAutoConsoleVariableRef CVarChaosImmPhysPushOutPairIterations(TEXT("p.Chaos.ImmPhys.Joint.PushOutPairIterations"), ChaosImmediate_Joint_PushOutPairIterations, TEXT("PushOutPairIterations."));
 FAutoConsoleVariableRef CVarChaosImmPhysSwingTwistAngleTolerance(TEXT("p.Chaos.ImmPhys.Joint.SwingTwistAngleTolerance"), ChaosImmediate_Joint_SwingTwistAngleTolerance, TEXT("SwingTwistAngleTolerance."));
diff --git a/Engine/Source/Runtime/Experimental/Chaos/Private/Chaos/Joint/PBDJointSolverGaussSeidel.cpp b/Engine/Source/Runtime/Experimental/Chaos/Private/Chaos/Joint/PBDJointSolverGaussSeidel.cpp
--- a/Engine/Source/Runtime/Experimental/Chaos/Private/Chaos/Joint/PBDJointSolverGaussSeidel.cpp
+++ b/Engine/Source/Runtime/Experimental/Chaos/Private/Chaos/Joint/PBDJointSolverGaussSeidel.cpp
@@ -943,19 +943,26 @@ namespace Chaos
 		const FVec3 DR1 = IA1 * -(Angle / (II0 + II1));
 		//const FVec3 DR0 = Axis0 * (Angle * II0 / (II0 + II1));
 		//const FVec3 DR1 = Axis1 * -(Angle * II1 / (II0 + II1));
+		
+		ApplyRotationDelta(Stiffness, DR0, DR1);
 
+		// Correct the positional error that was introduced by the rotation correction. This is
+		// correct when the position dofs are all locked, but not otherwise (fixable).
+		// This significantly improves angular stiffness at lower iterations, although the same effect
+		// is achieved by increasing iterations.
 		// @todo(ccaulfield): this position correction needs to have components in direction of inactive position constraints removed
 		if (AngularPositionCorrection > 0)
 		{
-			const FReal IM0 = InvMs[0];
-			const FReal IM1 = InvMs[1];
-			const FVec3 DX = FVec3::CrossProduct(DR0, Xs[0] - Ps[0]) + FVec3::CrossProduct(DR1, Xs[1] - Ps[1]);
-			const FVec3 DP0 = DX * (IM0 / (IM0 + IM1));
-			const FVec3 DP1 = DX * (-IM1 / (IM0 + IM1));
+			const FVec3 PrevX0 = Xs[0];
+			const FVec3 PrevX1 = Xs[1];
+			UpdateDerivedState();
+
+			const FVec3 DX = (Xs[1] - Xs[0]) - (PrevX1 - PrevX0);
+			const FVec3 DP0 = DX * (InvMs[0] / (InvMs[0] + InvMs[1]));
+			const FVec3 DP1 = DX * (-InvMs[1] / (InvMs[0] + InvMs[1]));
 			ApplyPositionDelta(Stiffness * AngularPositionCorrection, DP0, DP1);
 		}
 
-		ApplyRotationDelta(Stiffness, DR0, DR1);
 		UpdateDerivedState();
 	}
 
@@ -1030,19 +1037,26 @@ namespace Chaos
 		const FVec3 DR0 = Axis0 * (DLambda * II0);
 		const FVec3 DR1 = Axis1 * -(DLambda * II1);
 
+		Lambda += DLambda;
+		ApplyRotationDelta(1.0f, DR0, DR1);
+
+		// Correct the positional error that was introduced by the rotation correction. This is
+		// correct when the position dofs are all locked, but not otherwise (fixable).
+		// This significantly improves angular stiffness at lower iterations, although the same effect
+		// is achieved by increasing iterations.
 		// @todo(ccaulfield): this position correction needs to have components in direction of inactive position constraints removed
 		if (AngularPositionCorrection > 0)
 		{
-			const FReal IM0 = InvMs[0];
-			const FReal IM1 = InvMs[1];
-			const FVec3 DX = FVec3::CrossProduct(DR0, Xs[0] - Ps[0]) + FVec3::CrossProduct(DR1, Xs[1] - Ps[1]);
-			const FVec3 DP0 = DX * (IM0 / (IM0 + IM1));
-			const FVec3 DP1 = DX * (-IM1 / (IM0 + IM1));
+			const FVec3 PrevX0 = Xs[0];
+			const FVec3 PrevX1 = Xs[1];
+			UpdateDerivedState();
+
+			const FVec3 DX = (Xs[1] - Xs[0]) - (PrevX1 - PrevX0);
+			const FVec3 DP0 = DX * (InvMs[0] / (InvMs[0] + InvMs[1]));
+			const FVec3 DP1 = DX * (-InvMs[1] / (InvMs[0] + InvMs[1]));
 			ApplyPositionDelta(AngularPositionCorrection, DP0, DP1);
 		}
 
-		Lambda += DLambda;
-		ApplyRotationDelta(1.0f, DR0, DR1);
 		UpdateDerivedState();
 	}
 
