Convexes use TPlaneConcrete instead of TPlane. Saves 20% memory

diff --git a/Engine/Plugins/Experimental/ChaosCloth/Source/Chaos/Private/ChaosCloth/ChaosClothingSimulation.cpp b/Engine/Plugins/Experimental/ChaosCloth/Source/Chaos/Private/ChaosCloth/ChaosClothingSimulation.cpp
--- a/Engine/Plugins/Experimental/ChaosCloth/Source/Chaos/Private/ChaosCloth/ChaosClothingSimulation.cpp
+++ b/Engine/Plugins/Experimental/ChaosCloth/Source/Chaos/Private/ChaosCloth/ChaosClothingSimulation.cpp
@@ -843,9 +843,9 @@ void ClothingSimulation::ExtractPhysicsAssetCollisions(const UClothingAssetCommo
 				const Chaos::FConvex& ChaosConvex = ChaosConvexMesh.GetObjectChecked<Chaos::FConvex>();
 
 				// Copy planes
-				const TArray<TPlane<float, 3>>& Planes = ChaosConvex.GetFaces();
+				const TArray<TPlaneConcrete<float, 3>>& Planes = ChaosConvex.GetFaces();
 				Convex.Planes.Reserve(Planes.Num());
-				for (const TPlane<float, 3>& Plane : Planes)
+				for (const TPlaneConcrete<float, 3>& Plane : Planes)
 				{
 					Convex.Planes.Add(FPlane(Plane.X(), Plane.Normal()));
 				}
@@ -1040,7 +1040,7 @@ void ClothingSimulation::AddCollisions(const FClothCollisionData& ClothCollision
 			else
 			{
 				// Retrieve convex planes
-				TArray<TPlane<float, 3>> Planes;
+				TArray<TPlaneConcrete<float, 3>> Planes;
 				Planes.Reserve(Convex.Planes.Num());
 				for (const FPlane& Plane : Convex.Planes)
 				{
@@ -1050,7 +1050,7 @@ void ClothingSimulation::AddCollisions(const FClothCollisionData& ClothCollision
 						const Chaos::TVector<float, 3> Normal(static_cast<FVector>(NormalizedPlane));
 						const Chaos::TVector<float, 3> Base = Normal * NormalizedPlane.W;
 
-						Planes.Add(Chaos::TPlane<float, 3>(Base, Normal));
+						Planes.Add(Chaos::TPlaneConcrete<float, 3>(Base, Normal));
 					}
 					else
 					{
@@ -1585,14 +1585,14 @@ void ClothingSimulation::DebugDrawCollision(USkeletalMeshComponent* OwnerCompone
 
 	auto DrawConvex = [&PDI](const Chaos::FConvex& Convex, const TRotation<float, 3>& Rotation, const Chaos::TVector<float, 3>& Position, const FLinearColor& Color)
 	{
-		const TArray<Chaos::TPlane<float, 3>>& Planes = Convex.GetFaces();
+		const TArray<Chaos::TPlaneConcrete<float, 3>>& Planes = Convex.GetFaces();
 		for (int32 PlaneIndex1 = 0; PlaneIndex1 < Planes.Num(); ++PlaneIndex1)
 		{
-			const Chaos::TPlane<float, 3>& Plane1 = Planes[PlaneIndex1];
+			const Chaos::TPlaneConcrete<float, 3>& Plane1 = Planes[PlaneIndex1];
 
 			for (int32 PlaneIndex2 = PlaneIndex1 + 1; PlaneIndex2 < Planes.Num(); ++PlaneIndex2)
 			{
-				const Chaos::TPlane<float, 3>& Plane2 = Planes[PlaneIndex2];
+				const Chaos::TPlaneConcrete<float, 3>& Plane2 = Planes[PlaneIndex2];
 
 				// Find the two surface points that belong to both Plane1 and Plane2
 				uint32 ParticleIndex1 = INDEX_NONE;
diff --git a/Engine/Source/Runtime/Core/Public/UObject/ExternalPhysicsCustomObjectVersion.h b/Engine/Source/Runtime/Core/Public/UObject/ExternalPhysicsCustomObjectVersion.h
--- a/Engine/Source/Runtime/Core/Public/UObject/ExternalPhysicsCustomObjectVersion.h
+++ b/Engine/Source/Runtime/Core/Public/UObject/ExternalPhysicsCustomObjectVersion.h
@@ -101,6 +101,9 @@ struct CORE_API FExternalPhysicsCustomObjectVersion
 
 		// Capsules no longer have a union inside them or stored aabbs
 		CapsulesNoUnionOrAABBs,
+
+		// Convexes use concrete planes
+		ConvexUsesTPlaneConcrete,
 		
 		// -----<new versions can be added above this line>-------------------------------------------------
 		VersionPlusOne,
diff --git a/Engine/Source/Runtime/Engine/Private/PhysicsEngine/Experimental/ChaosDerivedData.cpp b/Engine/Source/Runtime/Engine/Private/PhysicsEngine/Experimental/ChaosDerivedData.cpp
--- a/Engine/Source/Runtime/Engine/Private/PhysicsEngine/Experimental/ChaosDerivedData.cpp
+++ b/Engine/Source/Runtime/Engine/Private/PhysicsEngine/Experimental/ChaosDerivedData.cpp
@@ -1,4 +1,4 @@
-// Copyright Epic Games, Inc. All Rights Reserved.
+// Copyright 1998-2019 Epic Games, Inc. All Rights Reserved.
 
 #include "ChaosDerivedData.h"
 
@@ -21,7 +21,7 @@ const TCHAR* FChaosDerivedDataCooker::GetPluginName() const
 
 const TCHAR* FChaosDerivedDataCooker::GetVersionString() const
 {
-	return TEXT("41954A2DFA074735B0B409CF01F3598A");
+	return TEXT("6249C3BD908145F3B3AC3048F3BF1C01");
 }
 
 FString FChaosDerivedDataCooker::GetPluginSpecificCacheKeySuffix() const
@@ -84,7 +84,7 @@ FChaosDerivedDataCooker::FChaosDerivedDataCooker(UBodySetup* InSetup, FName InFo
 
 }
 
-void FChaosDerivedDataCooker::BuildTriangleMeshes(TArray<TUniquePtr<Chaos::FTriangleMeshImplicitObject>>& OutTriangleMeshes, TArray<int32>& OutFaceRemap, const FCookBodySetupInfo& InParams)
+void FChaosDerivedDataCooker::BuildTriangleMeshes(TArray<TUniquePtr<Chaos::FTriangleMeshImplicitObject>>& OutTriangleMeshes, const FCookBodySetupInfo& InParams)
 {
 	if(!InParams.bCookTriMesh)
 	{
@@ -104,7 +104,8 @@ void FChaosDerivedDataCooker::BuildTriangleMeshes(TArray<TUniquePtr<Chaos::FTria
 		FinalIndices.Add(Tri.v2);
 	}
 
-	Chaos::CleanTrimesh(FinalVerts, FinalIndices, &OutFaceRemap);
+	TArray<int32> Remap;
+	Chaos::CleanTrimesh(FinalVerts, FinalIndices, &Remap);
 
 	// Build particle list #BG Maybe allow TParticles to copy vectors?
 	Chaos::TParticles<Chaos::FReal, 3> TriMeshParticles;
@@ -146,14 +147,14 @@ void FChaosDerivedDataCooker::BuildTriangleMeshes(TArray<TUniquePtr<Chaos::FTria
 
 				if(bHasMaterials)
 				{
-				if(!ensure(OutFaceRemap.IsValidIndex(TriangleIndex)))
+					if(!ensure(Remap.IsValidIndex(TriangleIndex)))
 					{
 						MaterialIndices.Empty();
 						bHasMaterials = false;
 					}
 					else
 					{
-					const int32 OriginalIndex = OutFaceRemap[TriangleIndex];
+						const int32 OriginalIndex = Remap[TriangleIndex];
 
 						if(ensure(InParams.TriangleMeshDesc.MaterialIndices.IsValidIndex(OriginalIndex)))
 						{
@@ -257,10 +258,9 @@ void FChaosDerivedDataCooker::BuildInternal(Chaos::FChaosArchive& Ar, FCookBodyS
 	TArray<TUniquePtr<Chaos::FImplicitObject>> SimpleImplicits;
 	TArray<TUniquePtr<Chaos::FTriangleMeshImplicitObject>> ComplexImplicits;
 
-	TArray<int32> FaceRemap;
 	//BuildSimpleShapes(SimpleImplicits, Setup);
 	BuildConvexMeshes(SimpleImplicits, InInfo);
-	BuildTriangleMeshes(ComplexImplicits, FaceRemap, InInfo);
+	BuildTriangleMeshes(ComplexImplicits, InInfo);
 
 	//TUniquePtr<Chaos::TImplicitObjectUnion<Precision, 3>> SimpleUnion = MakeUnique<Chaos::TImplicitObjectUnion<Precision, 3>>(MoveTemp(SimpleImplicits));
 	//TUniquePtr<Chaos::TImplicitObjectUnion<Precision, 3>> ComplexUnion = MakeUnique<Chaos::TImplicitObjectUnion<Precision, 3>>(MoveTemp(ComplexImplicits));
@@ -270,12 +270,8 @@ void FChaosDerivedDataCooker::BuildInternal(Chaos::FChaosArchive& Ar, FCookBodyS
 	{
 		UVInfo.FillFromTriMesh(InInfo.TriangleMeshDesc);
 	}
-	if (!InInfo.bSupportFaceRemap)
-	{
-		FaceRemap.Empty();
-	}
 
-	Ar << SimpleImplicits << ComplexImplicits << UVInfo << FaceRemap;
+	Ar << SimpleImplicits << ComplexImplicits << UVInfo;
 }
 
 #endif
diff --git a/Engine/Source/Runtime/Experimental/Chaos/Private/Chaos/Convex.cpp b/Engine/Source/Runtime/Experimental/Chaos/Private/Chaos/Convex.cpp
--- a/Engine/Source/Runtime/Experimental/Chaos/Private/Chaos/Convex.cpp
+++ b/Engine/Source/Runtime/Experimental/Chaos/Private/Chaos/Convex.cpp
@@ -13,7 +13,7 @@ namespace Chaos
 		FReal MostOpposingDot = TNumericLimits<FReal>::Max();
 		for (int32 Idx = 0; Idx < Planes.Num(); ++Idx)
 		{
-			const TPlane<FReal, 3>& Plane = Planes[Idx];
+			const TPlaneConcrete<FReal, 3>& Plane = Planes[Idx];
 			const FReal Distance = Plane.SignedDistance(Position);
 			if (FMath::Abs(Distance) < SearchDist)
 			{
@@ -43,7 +43,7 @@ namespace Chaos
 		TSet<int32> IncludedParticles;
 		for (int32 Idx = 0; Idx < Planes.Num(); ++Idx)
 		{
-			const TPlane<FReal, 3>& Plane = Planes[Idx];
+			const TPlaneConcrete<FReal, 3>& Plane = Planes[Idx];
 			FReal AbsOfSignedDistance = FMath::Abs(Plane.SignedDistance(Position));
 			if (AbsOfSignedDistance < SearchDist)
 			{
diff --git a/Engine/Source/Runtime/Experimental/Chaos/Public/Chaos/CollisionConvexMesh.h b/Engine/Source/Runtime/Experimental/Chaos/Public/Chaos/CollisionConvexMesh.h
--- a/Engine/Source/Runtime/Experimental/Chaos/Public/Chaos/CollisionConvexMesh.h
+++ b/Engine/Source/Runtime/Experimental/Chaos/Public/Chaos/CollisionConvexMesh.h
@@ -20,7 +20,7 @@ namespace Chaos
 			return Cross.Size() > 1e-4;
 		}
 
-		static void Build(const TParticles<FReal, 3>& InParticles, TArray <TPlane<FReal, 3>>& OutPlanes, TArray<TArray<int32>>& OutFaceIndices, TParticles<FReal, 3>& OutSurfaceParticles, TAABB<FReal, 3>& OutLocalBounds)
+		static void Build(const TParticles<FReal, 3>& InParticles, TArray <TPlaneConcrete<FReal, 3>>& OutPlanes, TArray<TArray<int32>>& OutFaceIndices, TParticles<FReal, 3>& OutSurfaceParticles, TAABB<FReal, 3>& OutLocalBounds)
 		{
 			OutPlanes.Reset();
 			OutSurfaceParticles.Resize(0);
@@ -60,7 +60,7 @@ namespace Chaos
 				{
 					FVec3 Vs[3] = {InParticles.X(Idx[0]), InParticles.X(Idx[1]), InParticles.X(Idx[2])};
 					const FVec3 Normal = FVec3::CrossProduct(Vs[1] - Vs[0], Vs[2] - Vs[0]).GetUnsafeNormal();
-					OutPlanes.Add(TPlane<FReal, 3>(Vs[0], Normal));
+					OutPlanes.Add(TPlaneConcrete<FReal, 3>(Vs[0], Normal));
 					TArray<int32> FaceIndices;
 					FaceIndices.SetNum(3);
 					FaceIndices[0] = AddIndex(Idx[0]);
@@ -84,7 +84,7 @@ namespace Chaos
 				if (ensureMsgf(bIsValidTriangle, TEXT("FConvexBuilder::Build(): Generated invalid triangle!")))
 				{
 					FVec3 Normal = FVec3::CrossProduct(InParticles.X(1) - InParticles.X(0), InParticles.X(2) - InParticles.X(0)).GetSafeNormal();
-					OutPlanes.Add(TPlane<FReal, 3>(InParticles.X(0), Normal));
+					OutPlanes.Add(TPlaneConcrete<FReal, 3>(InParticles.X(0), Normal));
 					OutSurfaceParticles.AddParticles(3);
 					OutSurfaceParticles.X(0) = InParticles.X(0);
 					OutSurfaceParticles.X(1) = InParticles.X(1);
@@ -161,7 +161,7 @@ namespace Chaos
 			return FString::Printf(TEXT("Planes %d, SurfaceParticles %d"), NumPlanes, NumParticles);
 		}
 
-		static CHAOS_API void Simplify(TArray <TPlane<FReal, 3>>& InOutPlanes, TArray<TArray<int32>>& InOutFaces, TParticles<FReal, 3>& InOutParticles, TAABB<FReal, 3>& InOutLocalBounds)
+		static CHAOS_API void Simplify(TArray <TPlaneConcrete<FReal, 3>>& InOutPlanes, TArray<TArray<int32>>& InOutFaces, TParticles<FReal, 3>& InOutParticles, TAABB<FReal, 3>& InOutLocalBounds)
 		{
 			struct TPair
 			{
@@ -254,7 +254,7 @@ namespace Chaos
 		struct FHalfEdge;
 		struct FConvexFace
 		{
-			FConvexFace(const TPlane<FReal, 3>& FacePlane)
+			FConvexFace(const TPlaneConcrete<FReal, 3>& FacePlane)
 				: ConflictList(nullptr)
 				, Plane(FacePlane)
 			{
@@ -262,7 +262,7 @@ namespace Chaos
 
 			FHalfEdge* FirstEdge;
 			FHalfEdge* ConflictList; //Note that these half edges are really just free verts grouped together
-			TPlane<FReal, 3> Plane;
+			TPlaneConcrete<FReal, 3> Plane;
 			FConvexFace* Prev;
 			FConvexFace* Next; //these have no geometric meaning, just used for book keeping
 		};
@@ -296,7 +296,7 @@ namespace Chaos
 			const FReal RSTNormalSize = RSTNormal.Size();
 			check(RSTNormalSize > 1e-4);
 			RSTNormal = RSTNormal * (1 / RSTNormalSize);
-			FConvexFace* RST = new FConvexFace(TPlane<FReal, 3>(InParticles.X(RS->Vertex), RSTNormal));
+			FConvexFace* RST = new FConvexFace(TPlaneConcrete<FReal, 3>(InParticles.X(RS->Vertex), RSTNormal));
 			RST->FirstEdge = RS;
 			RS->Face = RST;
 			ST->Face = RST;
diff --git a/Engine/Source/Runtime/Experimental/Chaos/Public/Chaos/Convex.h b/Engine/Source/Runtime/Experimental/Chaos/Public/Chaos/Convex.h
--- a/Engine/Source/Runtime/Experimental/Chaos/Public/Chaos/Convex.h
+++ b/Engine/Source/Runtime/Experimental/Chaos/Public/Chaos/Convex.h
@@ -36,7 +36,7 @@ namespace Chaos
 		// NOTE: This constructor will result in approximate COM and volume calculations, since it does
 		// not have face indices for surface particles.
 		// TODO: Keep track of invalid state and ensure on volume or COM access?
-		FConvex(TArray<TPlane<FReal, 3>>&& InPlanes, TParticles<FReal, 3>&& InSurfaceParticles)
+		FConvex(TArray<TPlaneConcrete<FReal, 3>>&& InPlanes, TParticles<FReal, 3>&& InSurfaceParticles)
 		    : FImplicitObject(EImplicitObject::IsConvex | EImplicitObject::HasBoundingBox, ImplicitObjectType::Convex)
 			, Planes(MoveTemp(InPlanes))
 		    , SurfaceParticles(MoveTemp(InSurfaceParticles))
@@ -148,7 +148,7 @@ namespace Chaos
 			// passes in a valid face index.
 			if (CHAOS_ENSURE(FaceIndex != INDEX_NONE))
 			{
-				const TPlane<FReal, 3>& OpposingFace = GetFaces()[FaceIndex];
+				const TPlaneConcrete<FReal, 3>& OpposingFace = GetFaces()[FaceIndex];
 				return OpposingFace.Normal();
 			}
 			return FVec3(0.f, 0.f, 1.f);
@@ -195,7 +195,7 @@ namespace Chaos
 			return SurfaceParticles;
 		}
 
-		const TArray<TPlane<FReal, 3>>& GetFaces() const
+		const TArray<TPlaneConcrete<FReal, 3>>& GetFaces() const
 		{
 			return Planes;
 		}
@@ -222,7 +222,7 @@ namespace Chaos
 
 			Result = HashCombine(Result, SurfaceParticles.GetTypeHash());
 
-			for(const TPlane<FReal, 3>& Plane : Planes)
+			for(const TPlaneConcrete<FReal, 3>& Plane : Planes)
 			{
 				Result = HashCombine(Result, Plane.GetTypeHash());
 			}
@@ -232,11 +232,28 @@ namespace Chaos
 
 		FORCEINLINE void SerializeImp(FArchive& Ar)
 		{
+			Ar.UsingCustomVersion(FExternalPhysicsCustomObjectVersion::GUID);
 			FImplicitObject::SerializeImp(Ar);
-			Ar << Planes << SurfaceParticles;
+
+			if (Ar.CustomVer(FExternalPhysicsCustomObjectVersion::GUID) < FExternalPhysicsCustomObjectVersion::ConvexUsesTPlaneConcrete)
+			{
+				TArray<TPlane<FReal, 3>> TmpPlanes;
+				Ar << TmpPlanes;
+
+				Planes.SetNum(TmpPlanes.Num());
+				for(int32 Idx = 0; Idx < Planes.Num(); ++Idx)
+				{
+					Planes[Idx] = TmpPlanes[Idx].PlaneConcrete();
+				}
+			}
+			else
+			{
+				Ar << Planes;
+			}
+
+			Ar << SurfaceParticles;
 			TBox<FReal,3>::SerializeAsAABB(Ar, LocalBoundingBox);
 
-			Ar.UsingCustomVersion(FExternalPhysicsCustomObjectVersion::GUID);
 			if (Ar.CustomVer(FExternalPhysicsCustomObjectVersion::GUID) >= FExternalPhysicsCustomObjectVersion::AddConvexCenterOfMassAndVolume)
 			{
 				Ar << Volume;
@@ -299,7 +316,7 @@ namespace Chaos
 		}
 
 	private:
-		TArray<TPlane<FReal, 3>> Planes;
+		TArray<TPlaneConcrete<FReal, 3>> Planes;
 		TParticles<FReal, 3> SurfaceParticles;	//copy of the vertices that are just on the convex hull boundary
 		TAABB<FReal, 3> LocalBoundingBox;
 		float Volume;
diff --git a/Engine/Source/Runtime/Experimental/Chaos/Public/Chaos/Plane.h b/Engine/Source/Runtime/Experimental/Chaos/Public/Chaos/Plane.h
--- a/Engine/Source/Runtime/Experimental/Chaos/Public/Chaos/Plane.h
+++ b/Engine/Source/Runtime/Experimental/Chaos/Public/Chaos/Plane.h
@@ -6,38 +6,17 @@
 
 namespace Chaos
 {
-template<class T, int d>
-class TPlane final : public FImplicitObject
+
+template <typename T, int d>
+class TPlaneConcrete
 {
-  public:
-	using FImplicitObject::GetTypeName;
-
-
-	TPlane() : FImplicitObject(0, ImplicitObjectType::Plane) {}	//needed for serialization
-	TPlane(const TVector<T, d>& InX, const TVector<T, d>& InNormal)
-	    : FImplicitObject(0, ImplicitObjectType::Plane)
-	    , MX(InX)
+public:
+	TPlaneConcrete() = default;
+	TPlaneConcrete(const TVector<T, d>& InX, const TVector<T, d>& InNormal)
+	    : MX(InX)
 	    , MNormal(InNormal)
 	{
 	}
-	TPlane(const TPlane<T, d>& Other)
-	    : FImplicitObject(0, ImplicitObjectType::Plane)
-	    , MX(Other.MX)
-	    , MNormal(Other.MNormal)
-	{
-	}
-	TPlane(TPlane<T, d>&& Other)
-	    : FImplicitObject(0, ImplicitObjectType::Plane)
-	    , MX(MoveTemp(Other.MX))
-	    , MNormal(MoveTemp(Other.MNormal))
-	{
-	}
-	virtual ~TPlane() {}
 
-	static constexpr EImplicitObjectType StaticType()
-	{
-		return ImplicitObjectType::Plane;
-	}
-
 	/**
 	 * Phi is positive on the side of the normal, and negative otherwise.
@@ -50,7 +29,7 @@
 	/**
 	 * Phi is positive on the side of the normal, and negative otherwise.
 	 */
-	virtual T PhiWithNormal(const TVector<T, d>& x, TVector<T, d>& Normal) const override
+	T PhiWithNormal(const TVector<T, d>& x, TVector<T, d>& Normal) const
 	{
 		Normal = MNormal;
 		return TVector<T, d>::DotProduct(x - MX, MNormal);
@@ -62,7 +41,7 @@
 		return x - TVector<T, d>(Dist * MNormal);
 	}
 
-	virtual bool Raycast(const TVector<T, d>& StartPoint, const TVector<T, d>& Dir, const T Length, const T Thickness, T& OutTime, TVector<T, d>& OutPosition, TVector<T, d>& OutNormal, int32& OutFaceIndex) const override
+	bool Raycast(const TVector<T, d>& StartPoint, const TVector<T, d>& Dir, const T Length, const T Thickness, T& OutTime, TVector<T, d>& OutPosition, TVector<T, d>& OutNormal, int32& OutFaceIndex) const
 	{
 		ensure(FMath::IsNearlyEqual(Dir.SizeSquared(),1, KINDA_SMALL_NUMBER));
 		ensure(Length > 0);
@@ -101,7 +80,7 @@
 		return true;
 	}
 
-	virtual Pair<TVector<T, d>, bool> FindClosestIntersectionImp(const TVector<T, d>& StartPoint, const TVector<T, d>& EndPoint, const T Thickness) const override
+	Pair<TVector<T, d>, bool> FindClosestIntersection(const TVector<T, d>& StartPoint, const TVector<T, d>& EndPoint, const T Thickness) const
  	{
 		TVector<T, d> Direction = EndPoint - StartPoint;
 		T Length = Direction.Size();
@@ -131,10 +110,97 @@
 	const TVector<T, d>& Normal() const { return MNormal; }
 	const TVector<T, d>& Normal(const TVector<T, d>&) const { return MNormal; }
 	
+	FORCEINLINE void Serialize(FArchive& Ar)
+	{
+		Ar << MX << MNormal;
+	}
+
+	uint32 GetTypeHash() const
+	{
+		return HashCombine(::GetTypeHash(MX), ::GetTypeHash(MNormal));
+	}
+
+  private:
+	TVector<T, d> MX;
+	TVector<T, d> MNormal;
+};
+
+template <typename T, int d>
+FArchive& operator<<(FArchive& Ar, TPlaneConcrete<T,d>& PlaneConcrete)
+{
+	PlaneConcrete.Serialize(Ar);
+	return Ar;
+}
+
+template<class T, int d>
+class TPlane final : public FImplicitObject
+{
+  public:
+	using FImplicitObject::GetTypeName;
+
+
+	TPlane() : FImplicitObject(0, ImplicitObjectType::Plane) {}	//needed for serialization
+	TPlane(const TVector<T, d>& InX, const TVector<T, d>& InNormal)
+	    : FImplicitObject(0, ImplicitObjectType::Plane)
+		, MPlaneConcrete(InX, InNormal)
+	{
+	}
+	TPlane(const TPlane<T, d>& Other)
+	    : FImplicitObject(0, ImplicitObjectType::Plane)
+	    , MPlaneConcrete(Other.MPlaneConcrete)
+	{
+	}
+	TPlane(TPlane<T, d>&& Other)
+	    : FImplicitObject(0, ImplicitObjectType::Plane)
+	    , MPlaneConcrete(MoveTemp(Other.MPlaneConcrete))
+	{
+	}
+	virtual ~TPlane() {}
+
+	static constexpr EImplicitObjectType StaticType()
+	{
+		return ImplicitObjectType::Plane;
+	}
+
+	/**
+	 * Phi is positive on the side of the normal, and negative otherwise.
+	 */
+	T SignedDistance(const TVector<T, d>& x) const
+	{
+		return MPlaneConcrete.SignedDistance(x);
+	}
+
+	/**
+	 * Phi is positive on the side of the normal, and negative otherwise.
+	 */
+	virtual T PhiWithNormal(const TVector<T, d>& x, TVector<T, d>& Normal) const override
+	{
+		return MPlaneConcrete.PhiWithNormal(x,Normal);
+	}
+
+	TVector<T, d> FindClosestPoint(const TVector<T, d>& x, const T Thickness = (T)0) const
+	{
+		return MPlaneConcrete.FindClosestPoint(x,Thickness);
+	}
+
+	virtual bool Raycast(const TVector<T, d>& StartPoint, const TVector<T, d>& Dir, const T Length, const T Thickness, T& OutTime, TVector<T, d>& OutPosition, TVector<T, d>& OutNormal, int32& OutFaceIndex) const override
+	{
+		return MPlaneConcrete.Raycast(StartPoint,Dir,Length,Thickness,OutTime,OutPosition,OutNormal,OutFaceIndex);
+	}
+
+	virtual Pair<TVector<T, d>, bool> FindClosestIntersectionImp(const TVector<T, d>& StartPoint, const TVector<T, d>& EndPoint, const T Thickness) const override
+ 	{
+		return MPlaneConcrete.FindClosestIntersection(StartPoint,EndPoint,Thickness);
+	}
+
+	const TVector<T,d>& X() const { return MPlaneConcrete.X(); }
+	const TVector<T,d>& Normal() const { return MPlaneConcrete.Normal(); }
+	const TVector<T, d>& Normal(const TVector<T, d>&) const { return MPlaneConcrete.Normal(); }
+	
 	FORCEINLINE void SerializeImp(FArchive& Ar)
 	{
 		FImplicitObject::SerializeImp(Ar);
-		Ar << MX << MNormal;
+		MPlaneConcrete.Serialize(Ar);
 	}
 
 	virtual void Serialize(FChaosArchive& Ar) override
@@ -150,12 +216,13 @@
 
 	virtual uint32 GetTypeHash() const override
 	{
-		return HashCombine(::GetTypeHash(MX), ::GetTypeHash(MNormal));
+		return MPlaneConcrete.GetTypeHash();
 	}
 
+	const TPlaneConcrete<T,d>& PlaneConcrete() const { return MPlaneConcrete; }
+
   private:
-	TVector<T, d> MX;
-	TVector<T, d> MNormal;
+	  TPlaneConcrete<T,d> MPlaneConcrete;
 };
 
 template<typename T, int d>
