Support transformed geometry in overlap tests

diff --git a/Engine/Source/Runtime/Engine/Private/PhysicsEngine/Experimental/PhysInterface_Chaos.cpp b/Engine/Source/Runtime/Engine/Private/PhysicsEngine/Experimental/PhysInterface_Chaos.cpp
--- a/Engine/Source/Runtime/Engine/Private/PhysicsEngine/Experimental/PhysInterface_Chaos.cpp
+++ b/Engine/Source/Runtime/Engine/Private/PhysicsEngine/Experimental/PhysInterface_Chaos.cpp
@@ -1559,7 +1559,8 @@ bool CalculateMassPropertiesOfImplicitType(
 		}
 		else
 		{
-			Chaos::CastHelper(*ImplicitObject, [&OutMassProperties, InDensityKGPerCM](const auto& Object)
+			//question: is LocalTM enough to merge these two branches?
+			Chaos::CastHelper(*ImplicitObject, FTransform::Identity, [&OutMassProperties, InDensityKGPerCM](const auto& Object, const auto& LocalTM)
 			{
 				OutMassProperties.Volume = Object.GetVolume();
 				OutMassProperties.Mass = OutMassProperties.Volume * InDensityKGPerCM;
@@ -1781,7 +1782,7 @@ bool FPhysInterface_Chaos::Sweep_Geom(FHitResult& OutHit, const FBodyInstance* I
 							Chaos::TVector<float, 3> WorldPosition;
 							Chaos::TVector<float, 3> WorldNormal;
 							int32 FaceIdx;
-							if (Chaos::CastHelper(ShapeAdapter.GetGeometry(), [&](const auto& Downcast) { return Chaos::SweepQuery(*Shape->Geometry, ActorTM, Downcast, StartTM, Dir, DeltaMag, Hit.Distance, WorldPosition, WorldNormal, FaceIdx, 0.f, false); }))
+							if (Chaos::CastHelper(ShapeAdapter.GetGeometry(), ActorTM, [&](const auto& Downcast, const auto& FullActorTM) { return Chaos::SweepQuery(*Shape->Geometry, FullActorTM, Downcast, StartTM, Dir, DeltaMag, Hit.Distance, WorldPosition, WorldNormal, FaceIdx, 0.f, false); }))
 							{
 								// we just like to make sure if the hit is made
 								FCollisionFilterData QueryFilter;
@@ -1841,7 +1842,7 @@ bool Overlap_GeomInternal(const FBodyInstance* InInstance, const Chaos::FImplici
 			if (OutOptResult)
 			{
 				Chaos::FMTDInfo MTDInfo;
-				if (Chaos::CastHelper(InGeom, [&](const auto& Downcast) { return Chaos::OverlapQuery(*Shape->Geometry, ActorTM, Downcast, GeomTransform, /*Thickness=*/0, &MTDInfo); }))
+				if (Chaos::CastHelper(InGeom, ActorTM, [&](const auto& Downcast, const auto& FullActorTM) { return Chaos::OverlapQuery(*Shape->Geometry, FullActorTM, Downcast, GeomTransform, /*Thickness=*/0, &MTDInfo); }))
 				{
 					OutOptResult->Distance = MTDInfo.Penetration;
 					OutOptResult->Direction = MTDInfo.Normal;
@@ -1850,7 +1851,7 @@ bool Overlap_GeomInternal(const FBodyInstance* InInstance, const Chaos::FImplici
 			}
 			else	//question: why do we even allow user to not pass in MTD info?
 			{
-				if (Chaos::CastHelper(InGeom, [&](const auto& Downcast) { return Chaos::OverlapQuery(*Shape->Geometry, ActorTM, Downcast, GeomTransform); }))
+				if (Chaos::CastHelper(InGeom, ActorTM, [&](const auto& Downcast, const auto& FullActorTM) { return Chaos::OverlapQuery(*Shape->Geometry, FullActorTM, Downcast, GeomTransform); }))
 				{
 					return true;
 				}
diff --git a/Engine/Source/Runtime/Experimental/Chaos/Private/Chaos/CollisionResolution.cpp b/Engine/Source/Runtime/Experimental/Chaos/Private/Chaos/CollisionResolution.cpp
--- a/Engine/Source/Runtime/Experimental/Chaos/Private/Chaos/CollisionResolution.cpp
+++ b/Engine/Source/Runtime/Experimental/Chaos/Private/Chaos/CollisionResolution.cpp
@@ -109,11 +109,11 @@ namespace Chaos
 		template<class T, int d>
 		TContactPoint<T> ConvexConvexContactPoint(const FImplicitObject& A, const TRigidTransform<T, d>& ATM, const FImplicitObject& B, const TRigidTransform<T, d>& BTM, const T Thickness)
 		{
-			return CastHelper(A, [&](const auto& ADowncast)
+			return CastHelper(A, ATM, [&](const auto& ADowncast, const TRigidTransform<T,d>& AFullTM)
 			{
-				return CastHelper(B, [&](const auto& BDowncast)
+				return CastHelper(B, BTM, [&](const auto& BDowncast, const TRigidTransform<T,d>& BFullTM)
 				{
-					return GJKContactPoint(ADowncast, ATM, BDowncast, BTM, Thickness);
+					return GJKContactPoint(ADowncast, AFullTM, BDowncast, BFullTM, Thickness);
 				});
 			});
 		}
diff --git a/Engine/Source/Runtime/Experimental/Chaos/Public/Chaos/GeometryQueries.h b/Engine/Source/Runtime/Experimental/Chaos/Public/Chaos/GeometryQueries.h
--- a/Engine/Source/Runtime/Experimental/Chaos/Public/Chaos/GeometryQueries.h
+++ b/Engine/Source/Runtime/Experimental/Chaos/Public/Chaos/GeometryQueries.h
@@ -20,27 +20,28 @@
 namespace Chaos
 {
 	template <typename Lambda>
-	FORCEINLINE_DEBUGGABLE auto CastHelper(const FImplicitObject& Geom, const Lambda& Func)
+	FORCEINLINE_DEBUGGABLE auto CastHelper(const FImplicitObject& Geom, const FRigidTransform3& TM, const Lambda& Func)
 	{
 		const EImplicitObjectType Type = Geom.GetType(true);
 		switch (Type)
 		{
-		case ImplicitObjectType::Sphere: return Func(Geom.template GetObjectChecked<TSphere<FReal, 3>>());
-		case ImplicitObjectType::Box: return Func(Geom.template GetObjectChecked<TBox<FReal, 3>>());
-		case ImplicitObjectType::Capsule: return Func(Geom.template GetObjectChecked<TCapsule<FReal>>());
-		case ImplicitObjectType::Convex: return Func(Geom.template GetObjectChecked<FConvex>());
-		case ImplicitObjectType::IsScaled | ImplicitObjectType::Sphere: return Func(Geom.template GetObjectChecked<TImplicitObjectScaled<TSphere<FReal, 3>>>());
-		case ImplicitObjectType::IsScaled | ImplicitObjectType::Box: return Func(Geom.template GetObjectChecked< TImplicitObjectScaled<TBox<FReal, 3>>>());
-		case ImplicitObjectType::IsScaled | ImplicitObjectType::Capsule: return Func(Geom.template GetObjectChecked< TImplicitObjectScaled<TCapsule<FReal>>>());
-		case ImplicitObjectType::IsScaled | ImplicitObjectType::Convex: return Func(Geom.template GetObjectChecked< TImplicitObjectScaled<FConvex>>());
+		case ImplicitObjectType::Sphere: return Func(Geom.template GetObjectChecked<TSphere<FReal, 3>>(), TM);
+		case ImplicitObjectType::Box: return Func(Geom.template GetObjectChecked<TBox<FReal, 3>>(), TM);
+		case ImplicitObjectType::Capsule: return Func(Geom.template GetObjectChecked<TCapsule<FReal>>(), TM);
+		case ImplicitObjectType::Convex: return Func(Geom.template GetObjectChecked<FConvex>(), TM);
+		case ImplicitObjectType::IsScaled | ImplicitObjectType::Sphere: return Func(Geom.template GetObjectChecked<TImplicitObjectScaled<TSphere<FReal, 3>>>(), TM);
+		case ImplicitObjectType::IsScaled | ImplicitObjectType::Box: return Func(Geom.template GetObjectChecked< TImplicitObjectScaled<TBox<FReal, 3>>>(), TM);
+		case ImplicitObjectType::IsScaled | ImplicitObjectType::Capsule: return Func(Geom.template GetObjectChecked< TImplicitObjectScaled<TCapsule<FReal>>>(), TM);
+		case ImplicitObjectType::IsScaled | ImplicitObjectType::Convex: return Func(Geom.template GetObjectChecked< TImplicitObjectScaled<FConvex>>(), TM);
 		case ImplicitObjectType::Transformed:
-			ensure(false); // We are drilling down to concrete implicit inside transformed, this is disregarding transform data. Caller must specially handle transform.
-			// TODO: Refactor Transformed implicit to use same structure as scaled.
-			return CastHelper(*(Geom.template GetObjectChecked<TImplicitObjectTransformed<FReal, 3>>().GetTransformedObject()), Func);
+		{
+			const auto& ImplicitObjectTransformed = (Geom.template GetObjectChecked<TImplicitObjectTransformed<FReal,3>>());
+			return CastHelper(*ImplicitObjectTransformed.GetTransformedObject(), ImplicitObjectTransformed.GetTransform() * TM, Func);
+		}
 
 		default: check(false);
 		}
-		return Func(Geom.template GetObjectChecked<TSphere<FReal, 3>>());	//needed for return type
+		return Func(Geom.template GetObjectChecked<TSphere<FReal, 3>>(), TM);	//needed for return type
 	}
 
 	struct FMTDInfo
@@ -79,11 +80,11 @@ namespace Chaos
 			{
 				FVec3 ClosestA;
 				FVec3 ClosestB;
-				return CastHelper(A, [&](const auto& AConcrete) { return GJKPenetration<FReal>(AConcrete, B, BToATM, OutMTD->Penetration, ClosestA, ClosestB, OutMTD->Normal, Thickness, Offset.SizeSquared() < 1e-4 ? FVec3(1, 0, 0) : Offset); });
+				return CastHelper(A, BToATM, [&](const auto& AConcrete, const auto& BToAFullTM) { return GJKPenetration<FReal>(AConcrete, B, BToAFullTM, OutMTD->Penetration, ClosestA, ClosestB, OutMTD->Normal, Thickness, Offset.SizeSquared() < 1e-4 ? FVec3(1, 0, 0) : Offset); });
 			}
 			else
 			{
-				return CastHelper(A, [&](const auto& AConcrete) { return GJKIntersection<FReal>(AConcrete, B, BToATM, Thickness, Offset.SizeSquared() < 1e-4 ? FVec3(1, 0, 0) : Offset); });
+				return CastHelper(A, BToATM, [&](const auto& AConcrete, const auto& BToAFullTM) { return GJKIntersection<FReal>(AConcrete, B, BToAFullTM, Thickness, Offset.SizeSquared() < 1e-4 ? FVec3(1, 0, 0) : Offset); });
 			}
 		}
 		else
@@ -176,7 +177,7 @@ namespace Chaos
 			}
 
 			const FVec3 Offset = ATM.GetLocation() - BTM.GetLocation();
-			bResult = CastHelper(A, [&](const auto& ADowncast) { return GJKRaycast2(ADowncast, B, BToATM, LocalDir, Length, OutTime, LocalPosition, LocalNormal, Thickness, bComputeMTD, Offset, Thickness); });
+			bResult = CastHelper(A, BToATM, [&](const auto& ADowncast, const auto& BToAFullTM){ return GJKRaycast2(ADowncast, B, BToAFullTM, LocalDir, Length, OutTime, LocalPosition, LocalNormal, Thickness, bComputeMTD, Offset, Thickness); });
 			if (AType == ImplicitObjectType::Convex)
 			{
 				//todo: find face index
diff --git a/Engine/Source/Runtime/PhysicsSQ/Private/SQAccelerator.cpp b/Engine/Source/Runtime/PhysicsSQ/Private/SQAccelerator.cpp
--- a/Engine/Source/Runtime/PhysicsSQ/Private/SQAccelerator.cpp
+++ b/Engine/Source/Runtime/PhysicsSQ/Private/SQAccelerator.cpp
@@ -360,7 +360,7 @@ void SweepHelper(const QueryGeomType& QueryGeom, const Chaos::ISpatialAccelerati
 
 void FChaosSQAccelerator::Sweep(const Chaos::FImplicitObject& QueryGeom, const FTransform& StartTM, const FVector& Dir, const float DeltaMagnitude, ChaosInterface::FSQHitBuffer<ChaosInterface::FSweepHit>& HitBuffer, EHitFlags OutputFlags, const FQueryFilterData& QueryFilterData, ICollisionQueryFilterCallbackBase& QueryCallback, const FQueryDebugParams& DebugParams) const
 {
-	return Chaos::CastHelper(QueryGeom, [&](const auto& Downcast) { return SweepHelper(Downcast, SpatialAcceleration, StartTM, Dir, DeltaMagnitude, HitBuffer, OutputFlags, QueryFilterData, QueryCallback, DebugParams); });
+	return Chaos::CastHelper(QueryGeom, StartTM, [&](const auto& Downcast, const FTransform& StartFullTM) { return SweepHelper(Downcast, SpatialAcceleration, StartFullTM, Dir, DeltaMagnitude, HitBuffer, OutputFlags, QueryFilterData, QueryCallback, DebugParams); });
 }
 
 template <typename QueryGeomType>
@@ -378,7 +378,7 @@ void OverlapHelper(const QueryGeomType& QueryGeom, const Chaos::ISpatialAccelera
 
 void FChaosSQAccelerator::Overlap(const Chaos::FImplicitObject& QueryGeom, const FTransform& GeomPose, ChaosInterface::FSQHitBuffer<ChaosInterface::FOverlapHit>& HitBuffer, const FQueryFilterData& QueryFilterData, ICollisionQueryFilterCallbackBase& QueryCallback, const FQueryDebugParams& DebugParams) const
 {
-	return Chaos::CastHelper(QueryGeom, [&](const auto& Downcast) { return OverlapHelper(Downcast, SpatialAcceleration, GeomPose, HitBuffer, QueryFilterData, QueryCallback, DebugParams); });
+	return Chaos::CastHelper(QueryGeom, GeomPose, [&](const auto& Downcast, const FTransform& GeomFullPose) { return OverlapHelper(Downcast, SpatialAcceleration, GeomFullPose, HitBuffer, QueryFilterData, QueryCallback, DebugParams); });
 }
 
 #if WITH_PHYSX
