Chaos - Joint debug draw tweaks

diff --git a/Engine/Source/Runtime/Experimental/Chaos/Private/Chaos/ChaosDebugDraw.cpp b/Engine/Source/Runtime/Experimental/Chaos/Private/Chaos/ChaosDebugDraw.cpp
--- a/Engine/Source/Runtime/Experimental/Chaos/Private/Chaos/ChaosDebugDraw.cpp
+++ b/Engine/Source/Runtime/Experimental/Chaos/Private/Chaos/ChaosDebugDraw.cpp
@@ -25,6 +25,7 @@ namespace Chaos
 		float ContactPhiWidth = 1.5f;
 		float ContactOwnerWidth = 0.0f;
 		float ConstraintAxisLen = 5.0f;
+		float JointComSize = 2.0f;
 		float LineThickness = 0.15f;
 		float DrawScale = 1.0f;
 		float FontHeight = 10.0f;
@@ -214,10 +215,30 @@ namespace Chaos
 
 			if (FeatureMask & (uint32)EDebugDrawJointFeature::Connector)
 			{
-				FDebugDrawQueue::GetInstance().DrawDebugLine(Pa, Xa, R, false, KINDA_SMALL_NUMBER, DrawPriority, LineThickness);
-				FDebugDrawQueue::GetInstance().DrawDebugLine(Pb, Xb, C, false, KINDA_SMALL_NUMBER, DrawPriority, LineThickness);
-				FDebugDrawQueue::GetInstance().DrawDebugPoint(Xa, R, false, KINDA_SMALL_NUMBER, DrawPriority, DrawScale * PointSize);
-				FDebugDrawQueue::GetInstance().DrawDebugPoint(Xb, C, false, KINDA_SMALL_NUMBER, DrawPriority, DrawScale * PointSize);
+				const FReal ConnectorThickness = 1.5f * LineThickness;
+				const FReal CoMSize = DrawScale * JointComSize;
+				// Leave a gap around the body position so we can see where the center is
+				FVec3 Sa = Pa;
+				const FReal Lena = (Xa - Pa).Size();
+				if (Lena > KINDA_SMALL_NUMBER)
+				{
+					Sa = FMath::Lerp(Pa, Xa, FMath::Clamp(CoMSize / Lena, 0.0f, 1.0f));
+				}
+				FVec3 Sb = Pb;
+				const FReal Lenb = (Xb - Pb).Size();
+				if (Lenb > KINDA_SMALL_NUMBER)
+				{
+					Sb = FMath::Lerp(Pb, Xb, FMath::Clamp(CoMSize / Lena, 0.0f, 1.0f));
+				}
+				FDebugDrawQueue::GetInstance().DrawDebugLine(Pa, Sa, FColor::Black, false, KINDA_SMALL_NUMBER, DrawPriority, ConnectorThickness);
+				FDebugDrawQueue::GetInstance().DrawDebugLine(Pb, Sb, FColor::Black, false, KINDA_SMALL_NUMBER, DrawPriority, ConnectorThickness);
+				FDebugDrawQueue::GetInstance().DrawDebugLine(Sa, Xa, R, false, KINDA_SMALL_NUMBER, DrawPriority, ConnectorThickness);
+				FDebugDrawQueue::GetInstance().DrawDebugLine(Sb, Xb, C, false, KINDA_SMALL_NUMBER, DrawPriority, ConnectorThickness);
+			}
+			if (FeatureMask & (uint32)EDebugDrawJointFeature::Stretch)
+			{
+				const FReal StretchThickness = 3.0f * LineThickness;
+				FDebugDrawQueue::GetInstance().DrawDebugLine(Xa, Xb, M, false, KINDA_SMALL_NUMBER, DrawPriority, StretchThickness);
 			}
 			if (FeatureMask & (uint32)EDebugDrawJointFeature::Axes)
 			{
diff --git a/Engine/Source/Runtime/Experimental/Chaos/Public/Chaos/ChaosDebugDraw.h b/Engine/Source/Runtime/Experimental/Chaos/Public/Chaos/ChaosDebugDraw.h
--- a/Engine/Source/Runtime/Experimental/Chaos/Public/Chaos/ChaosDebugDraw.h
+++ b/Engine/Source/Runtime/Experimental/Chaos/Public/Chaos/ChaosDebugDraw.h
@@ -17,11 +17,12 @@ namespace Chaos
 		{
 			None = 0,
 			Connector = 1 << 0,
-			Axes = 1 << 1,
-			Level = 1 << 2,
-			Index = 1 << 3,
+			Stretch = 1 << 1,
+			Axes = 1 << 2,
+			Level = 1 << 3,
+			Index = 1 << 4,
 
-			Default = Connector | Index,
+			Default = Connector | Stretch,
 			All = Connector | Axes | Level
 		};
 
