Chaos - debugdraw joint bone connectors

diff --git a/Engine/Source/Runtime/Experimental/Chaos/Private/Chaos/ChaosDebugDraw.cpp b/Engine/Source/Runtime/Experimental/Chaos/Private/Chaos/ChaosDebugDraw.cpp
--- a/Engine/Source/Runtime/Experimental/Chaos/Private/Chaos/ChaosDebugDraw.cpp
+++ b/Engine/Source/Runtime/Experimental/Chaos/Private/Chaos/ChaosDebugDraw.cpp
@@ -199,7 +199,7 @@ namespace Chaos
 			}
 		}
 
-		void DrawJointConstraintImpl(const FRigidTransform3& SpaceTransform, const FVec3& InPa, const FVec3& InXa, const FMatrix33& Ra, const FVec3& InPb, const FVec3& InXb, const FMatrix33& Rb, int32 Level, int32 Index, FReal ColorScale, uint32 FeatureMask)
+		void DrawJointConstraintImpl(const FRigidTransform3& SpaceTransform, const FVec3& InPa, const FVec3& InCa, const FVec3& InXa, const FMatrix33& Ra, const FVec3& InPb, const FVec3& InCb, const FVec3& InXb, const FMatrix33& Rb, int32 Level, int32 Index, FReal ColorScale, uint32 FeatureMask)
 		{
 			using namespace Chaos::DebugDraw;
 			FColor R = (ColorScale * FColor::Red).ToFColor(false);
@@ -210,14 +210,16 @@ namespace Chaos
 			FColor Y = (ColorScale * FColor::Yellow).ToFColor(false);
 			FVec3 Pa = SpaceTransform.TransformPosition(InPa);
 			FVec3 Pb = SpaceTransform.TransformPosition(InPb);
+			FVec3 Ca = SpaceTransform.TransformPosition(InCa);
+			FVec3 Cb = SpaceTransform.TransformPosition(InCb);
 			FVec3 Xa = SpaceTransform.TransformPosition(InXa);
 			FVec3 Xb = SpaceTransform.TransformPosition(InXb);
 
-			if (FeatureMask & (uint32)EDebugDrawJointFeature::Connector)
+			if (FeatureMask & (uint32)EDebugDrawJointFeature::ActorConnector)
 			{
 				const FReal ConnectorThickness = 1.5f * LineThickness;
 				const FReal CoMSize = DrawScale * JointComSize;
-				// Leave a gap around the body position so we can see where the center is
+				// Leave a gap around the actor position so we can see where the center is
 				FVec3 Sa = Pa;
 				const FReal Lena = (Xa - Pa).Size();
 				if (Lena > KINDA_SMALL_NUMBER)
@@ -230,8 +232,30 @@ namespace Chaos
 				{
 					Sb = FMath::Lerp(Pb, Xb, FMath::Clamp(CoMSize / Lena, 0.0f, 1.0f));
 				}
-				FDebugDrawQueue::GetInstance().DrawDebugLine(Pa, Sa, FColor::Black, false, KINDA_SMALL_NUMBER, DrawPriority, ConnectorThickness);
-				FDebugDrawQueue::GetInstance().DrawDebugLine(Pb, Sb, FColor::Black, false, KINDA_SMALL_NUMBER, DrawPriority, ConnectorThickness);
+				FDebugDrawQueue::GetInstance().DrawDebugLine(Pa, Sa, FColor::White, false, KINDA_SMALL_NUMBER, DrawPriority, ConnectorThickness);
+				FDebugDrawQueue::GetInstance().DrawDebugLine(Pb, Sb, FColor::White, false, KINDA_SMALL_NUMBER, DrawPriority, ConnectorThickness);
+				FDebugDrawQueue::GetInstance().DrawDebugLine(Sa, Xa, R, false, KINDA_SMALL_NUMBER, DrawPriority, ConnectorThickness);
+				FDebugDrawQueue::GetInstance().DrawDebugLine(Sb, Xb, C, false, KINDA_SMALL_NUMBER, DrawPriority, ConnectorThickness);
+			}
+			if (FeatureMask & (uint32)EDebugDrawJointFeature::CoMConnector)
+			{
+				const FReal ConnectorThickness = 1.5f * LineThickness;
+				const FReal CoMSize = DrawScale * JointComSize;
+				// Leave a gap around the body position so we can see where the center is
+				FVec3 Sa = Ca;
+				const FReal Lena = (Xa - Ca).Size();
+				if (Lena > KINDA_SMALL_NUMBER)
+				{
+					Sa = FMath::Lerp(Ca, Xa, FMath::Clamp(CoMSize / Lena, 0.0f, 1.0f));
+				}
+				FVec3 Sb = Cb;
+				const FReal Lenb = (Xb - Cb).Size();
+				if (Lenb > KINDA_SMALL_NUMBER)
+				{
+					Sb = FMath::Lerp(Cb, Xb, FMath::Clamp(CoMSize / Lena, 0.0f, 1.0f));
+				}
+				FDebugDrawQueue::GetInstance().DrawDebugLine(Ca, Sa, FColor::Black, false, KINDA_SMALL_NUMBER, DrawPriority, ConnectorThickness);
+				FDebugDrawQueue::GetInstance().DrawDebugLine(Cb, Sb, FColor::Black, false, KINDA_SMALL_NUMBER, DrawPriority, ConnectorThickness);
 				FDebugDrawQueue::GetInstance().DrawDebugLine(Sa, Xa, R, false, KINDA_SMALL_NUMBER, DrawPriority, ConnectorThickness);
 				FDebugDrawQueue::GetInstance().DrawDebugLine(Sb, Xb, C, false, KINDA_SMALL_NUMBER, DrawPriority, ConnectorThickness);
 			}
@@ -266,12 +290,14 @@ namespace Chaos
 			auto RigidParticle1 = ConstrainedParticles[1]->CastToRigidParticle();
 			if ((RigidParticle0 && RigidParticle0->ObjectState() == EObjectStateType::Dynamic) || (RigidParticle1 && RigidParticle1->ObjectState() == EObjectStateType::Dynamic))
 			{
-				FVec3 Pa = FParticleUtilities::GetCoMWorldPosition(TConstGenericParticleHandle<FReal, 3>(ConstraintHandle->GetConstrainedParticles()[1]));
-				FVec3 Pb = FParticleUtilities::GetCoMWorldPosition(TConstGenericParticleHandle<FReal, 3>(ConstraintHandle->GetConstrainedParticles()[0]));
+				FVec3 Pa = FParticleUtilities::GetActorWorldTransform(TConstGenericParticleHandle<FReal, 3>(ConstraintHandle->GetConstrainedParticles()[1])).GetTranslation();
+				FVec3 Pb = FParticleUtilities::GetActorWorldTransform(TConstGenericParticleHandle<FReal, 3>(ConstraintHandle->GetConstrainedParticles()[0])).GetTranslation();
+				FVec3 Ca = FParticleUtilities::GetCoMWorldPosition(TConstGenericParticleHandle<FReal, 3>(ConstraintHandle->GetConstrainedParticles()[1]));
+				FVec3 Cb = FParticleUtilities::GetCoMWorldPosition(TConstGenericParticleHandle<FReal, 3>(ConstraintHandle->GetConstrainedParticles()[0]));
 				FVec3 Xa, Xb;
 				FMatrix33 Ra, Rb;
 				ConstraintHandle->CalculateConstraintSpace(Xa, Ra, Xb, Rb);
-				DrawJointConstraintImpl(SpaceTransform, Pa, Xa, Ra, Pb, Xb, Rb, ConstraintHandle->GetConstraintLevel(), ConstraintHandle->GetConstraintIndex(), ColorScale, FeatureMask);
+				DrawJointConstraintImpl(SpaceTransform, Pa, Ca, Xa, Ra, Pb, Cb, Xb, Rb, ConstraintHandle->GetConstraintLevel(), ConstraintHandle->GetConstraintIndex(), ColorScale, FeatureMask);
 			}
 		}
 
diff --git a/Engine/Source/Runtime/Experimental/Chaos/Public/Chaos/ChaosDebugDraw.h b/Engine/Source/Runtime/Experimental/Chaos/Public/Chaos/ChaosDebugDraw.h
--- a/Engine/Source/Runtime/Experimental/Chaos/Public/Chaos/ChaosDebugDraw.h
+++ b/Engine/Source/Runtime/Experimental/Chaos/Public/Chaos/ChaosDebugDraw.h
@@ -16,14 +16,14 @@ namespace Chaos
 		enum class EDebugDrawJointFeature
 		{
 			None = 0,
-			Connector = 1 << 0,
-			Stretch = 1 << 1,
-			Axes = 1 << 2,
-			Level = 1 << 3,
-			Index = 1 << 4,
+			CoMConnector = 1 << 0,
+			ActorConnector = 1 << 1,
+			Stretch = 1 << 2,
+			Axes = 1 << 3,
+			Level = 1 << 4,
+			Index = 1 << 5,
 
-			Default = Connector | Stretch,
-			All = Connector | Axes | Level
+			Default = ActorConnector | Stretch
 		};
 
 		CHAOS_API void DrawParticleShapes(const FRigidTransform3& SpaceTransform, const TParticleView<TGeometryParticles<float, 3>>& ParticlesView, const FColor& Color);
