Chaos Cloth - Add basic wind uniform velocity field to the evolution solver.

diff --git a/Engine/Plugins/Experimental/ChaosCloth/Source/Chaos/Private/ChaosCloth/ChaosClothConfig.cpp b/Engine/Plugins/Experimental/ChaosCloth/Source/Chaos/Private/ChaosCloth/ChaosClothConfig.cpp
--- a/Engine/Plugins/Experimental/ChaosCloth/Source/Chaos/Private/ChaosCloth/ChaosClothConfig.cpp
+++ b/Engine/Plugins/Experimental/ChaosCloth/Source/Chaos/Private/ChaosCloth/ChaosClothConfig.cpp
@@ -6,7 +6,6 @@
 #include "ClothingSimulationInteractor.h"
 
 // Legacy parameters not yet migrated to Chaos parameters:
-//  WindDragCoefficient
 //  WindMethod
 //  VerticalConstraintConfig.CompressionLimit
 //  VerticalConstraintConfig.StretchLimit
@@ -88,6 +87,8 @@ void UChaosClothSharedSimConfig::MigrateFrom(const FClothConfig_Legacy& ClothCon
 	GravityScale = ClothConfig.GravityScale;
 
 	Gravity = ClothConfig.GravityOverride;
+
+	WindDrag = (ClothConfig.WindMethod == EClothingWindMethod_Legacy::Accurate) ? ClothConfig.WindDragCoefficient: 0.5f;  // Only Accurate wind uses the WindDragCoefficient
 }
 
 void UChaosClothSharedSimConfig::Serialize(FArchive& Ar)
diff --git a/Engine/Plugins/Experimental/ChaosCloth/Source/Chaos/Private/ChaosCloth/ChaosClothingSimulation.cpp b/Engine/Plugins/Experimental/ChaosCloth/Source/Chaos/Private/ChaosCloth/ChaosClothingSimulation.cpp
--- a/Engine/Plugins/Experimental/ChaosCloth/Source/Chaos/Private/ChaosCloth/ChaosClothingSimulation.cpp
+++ b/Engine/Plugins/Experimental/ChaosCloth/Source/Chaos/Private/ChaosCloth/ChaosClothingSimulation.cpp
@@ -60,6 +60,7 @@ namespace ClothingSimulationDefault
 	static const float CoefficientOfFriction = 0.f;
 	static const float Damping = 0.01f;
 	static const float SolverFrequency = 60.f;
+	static const float WindDrag = 0.f;
 }
 
 ClothingSimulation::ClothingSimulation()
@@ -92,6 +93,7 @@ void ClothingSimulation::Initialize()
 			ClothingSimulationDefault::Damping));
     Evolution->CollisionParticles().AddArray(&BoneIndices);
 	Evolution->CollisionParticles().AddArray(&BaseTransforms);
+	Evolution->GetVelocityField().SetDrag(ClothingSimulationDefault::WindDrag);  // TODO: Per particle drag
     Evolution->GetGravityForces().SetAcceleration(Gravity);
 
     Evolution->SetKinematicUpdateFunction(
@@ -289,6 +291,7 @@ void ClothingSimulation::UpdateSimulationFromSharedSimConfig()
 		Evolution->SetSelfCollisionThickness(ClothSharedSimConfig->SelfCollisionThickness);
 		Evolution->SetCollisionThickness(ClothSharedSimConfig->CollisionThickness);
 		Evolution->SetDamping(ClothSharedSimConfig->Damping);
+		Evolution->GetVelocityField().SetDrag(ClothSharedSimConfig->WindDrag);  // TODO: Per particle drag
 		Gravity = ClothSharedSimConfig->Gravity;
 	}
 }
@@ -1121,6 +1124,10 @@ void ClothingSimulation::Simulate(IClothingSimulationContext* InContext)
 			Context->WorldGravity * ClothSharedSimConfig->GravityScale:
 			Gravity));
 
+	// Set wind velocity
+	static const float UnitConversionScale = 100.f;  // It would seem that in game wind velocity is in m/s and not cm/s as seen in NvClothSupport::Constants::UnitConversionScale
+	Evolution->GetVelocityField().SetVelocity(Context->WindVelocity * UnitConversionScale);
+
 	// Get New Animation Positions and Normals
 	{
 		SCOPE_CYCLE_COUNTER(STAT_ChaosClothGetAnimationData);
diff --git a/Engine/Plugins/Experimental/ChaosCloth/Source/Chaos/Public/ChaosCloth/ChaosClothConfig.h b/Engine/Plugins/Experimental/ChaosCloth/Source/Chaos/Public/ChaosCloth/ChaosClothConfig.h
--- a/Engine/Plugins/Experimental/ChaosCloth/Source/Chaos/Public/ChaosCloth/ChaosClothConfig.h
+++ b/Engine/Plugins/Experimental/ChaosCloth/Source/Chaos/Public/ChaosCloth/ChaosClothConfig.h
@@ -165,6 +165,10 @@ public:
 	UPROPERTY(EditAnywhere, Category = Simulation, meta = (EditCondition = "bUseGravityOverride"))
 	FVector Gravity = { 0.f, 0.f, -980.665f };
 
+	// The wind drag coefficient applying on each particle
+	UPROPERTY(EditAnywhere, Category = Simulation, meta = (UIMin = "0", UIMax = "1", ClampMin = "0", ClampMax = "10"))
+	float WindDrag = 0.5f;
+
 	// Enable the XPBD constraints that resolve stiffness independently from the number of iterations
 	// Experimental, this feature might be removed without warning, not for production use
 	UPROPERTY(EditAnywhere, Category = "Experimental", meta = (DisplayName="Use XPBD Constraints (Experimental)"))
diff --git a/Engine/Source/Runtime/Experimental/Chaos/Private/Chaos/PBDEvolution.cpp b/Engine/Source/Runtime/Experimental/Chaos/Private/Chaos/PBDEvolution.cpp
--- a/Engine/Source/Runtime/Experimental/Chaos/Private/Chaos/PBDEvolution.cpp
+++ b/Engine/Source/Runtime/Experimental/Chaos/Private/Chaos/PBDEvolution.cpp
@@ -83,6 +83,7 @@ void TPBDEvolution<T, d>::AdvanceOneTimeStep(const T Dt)
 		{
 			InitForceRule.Apply(MParticles, Dt, Index); // F = TV(0)
 			GravityForces.Apply(MParticles, Dt, Index);
+			VelocityField.Apply(MParticles, Dt, Index);
 			for (TFunction<void(TPBDParticles<T, d>&, const T, const int32)>& ForceRule : MForceRules)
 			{
 				ForceRule(MParticles, Dt, Index); // F += M * A
diff --git a/Engine/Source/Runtime/Experimental/Chaos/Public/Chaos/PBDEvolution.h b/Engine/Source/Runtime/Experimental/Chaos/Public/Chaos/PBDEvolution.h
--- a/Engine/Source/Runtime/Experimental/Chaos/Public/Chaos/PBDEvolution.h
+++ b/Engine/Source/Runtime/Experimental/Chaos/Public/Chaos/PBDEvolution.h
@@ -3,6 +3,7 @@
 
 #include "Chaos/KinematicGeometryParticles.h"
 #include "Chaos/PerParticleGravity.h"
+#include "Chaos/PerParticleVelocityField.h"
 #include "Chaos/PBDParticles.h"
 #include "Chaos/Vector.h"
 
@@ -13,6 +14,7 @@ class CHAOS_API TPBDEvolution
 {
   public:
 	using FGravityForces = TPerParticleGravity<T, d>;
+	using FVelocityField = TPerParticleVelocityField<T, d>;
 
 	// TODO(mlentine): Init particles from some type of input
 	TPBDEvolution(TPBDParticles<T, d>&& InParticles, TKinematicGeometryClothParticles<T, d>&& InGeometryParticles, TArray<TVector<int32, 3>>&& CollisionTriangles, int32 NumIterations = 1, T CollisionThickness = 0, T SelfCollisionsThickness = 0, T CoefficientOfFriction = 0, T Damping = 0.04);
@@ -33,6 +35,9 @@ class CHAOS_API TPBDEvolution
 	FGravityForces& GetGravityForces() { return GravityForces; }
 	const FGravityForces& GetGravityForces() const { return GravityForces; }
 
+	FVelocityField& GetVelocityField() { return VelocityField; }
+	const FVelocityField& GetVelocityField() const { return VelocityField; }
+
 	const TGeometryClothParticles<T, d>& CollisionParticles() const { return MCollisionParticles; }
 	TGeometryClothParticles<T, d>& CollisionParticles() { return MCollisionParticles; }
 	const bool Collided(int32 index) { return MCollided[index]; }
@@ -74,6 +79,7 @@ class CHAOS_API TPBDEvolution
 	T MTime;
 
 	FGravityForces GravityForces;
+	FVelocityField VelocityField;
 
 	TArray<TFunction<void(TPBDParticles<T, d>&, const T, const int32)>> MForceRules;
 	TArray<TFunction<void()>> MInitConstraintRules;
diff --git a/Engine/Source/Runtime/Experimental/Chaos/Public/Chaos/PerParticleVelocityField.h b/Engine/Source/Runtime/Experimental/Chaos/Public/Chaos/PerParticleVelocityField.h
new file mode 100644
--- /dev/null
+++ b/Engine/Source/Runtime/Experimental/Chaos/Public/Chaos/PerParticleVelocityField.h
@@ -0,0 +1,91 @@
+// Copyright Epic Games, Inc. All Rights Reserved.
+#pragma once
+
+#include "Chaos/DynamicParticles.h"
+#include "Chaos/PerParticleRule.h"
+#include "Chaos/ParticleHandle.h"
+
+namespace Chaos
+{
+	// Velocity field basic implementation
+	// TODO:
+	//  - Turn this into a base class, with inherited uniform and other type of (non-uniform) fields.
+	//  - Add drag per particle instead of per field.
+	//  - Calculate field effect on the geometry area.
+	//  - Expose fluid density (currently using air density).
+	template<class T, int d>
+	class TPerParticleVelocityField : public TPerParticleRule<T, d>
+	{
+	public:
+		TPerParticleVelocityField()
+			: MVelocity(TVector<T, d>((T)0.)), MHalfRhoDragArea((T)0.) {}
+		TPerParticleVelocityField(const TVector<T, d>& InVelocity, const T InDrag)
+			: MVelocity(InVelocity)
+		{
+			SetDrag(InDrag);
+		}
+
+		virtual ~TPerParticleVelocityField() {}
+
+		// TODO: Remove this - we should no longer be using indices directly.
+		//       This has been kept around for cloth, which uses it in
+		//       PBDEvolution.
+		template<class T_PARTICLES>
+		inline void ApplyHelper(T_PARTICLES& InParticles, const T Dt, const int Index) const
+		{
+			TVector<T, d> Direction = MVelocity - InParticles.V(Index);
+			const T RelVelocity = Direction.SafeNormalize();
+			InParticles.F(Index) += Direction * MHalfRhoDragArea * FMath::Square(RelVelocity);
+		}
+
+		inline void Apply(TPBDParticles<T, d>& InParticles, const T Dt, const int Index) const override //-V762
+		{
+			ApplyHelper(InParticles, Dt, Index);
+		}
+
+		void Apply(TTransientPBDRigidParticleHandle<T, d>& Particle, const T Dt) const override //-V762
+		{
+			if (!UsesVelocityField.Contains(Particle.Handle()))
+			{
+				TVector<T, d> Direction = MVelocity - Particle.V();
+				const T RelVelocity = Direction.SafeNormalize();
+				Particle.F() += Direction * MHalfRhoDragArea * FMath::Square(RelVelocity);
+			}
+		}
+
+		void SetEnabled(const TGeometryParticleHandleImp<T, d, true>& Handle, bool bEnabled)
+		{
+			if (bEnabled)
+			{
+				if (UsesVelocityField.Contains(&Handle))
+				{
+					UsesVelocityField.Remove(&Handle);
+				}
+			}
+			else
+			{
+				if (!UsesVelocityField.Contains(&Handle))
+				{
+					UsesVelocityField.Add(&Handle);
+				}
+			}
+		}
+
+		void SetVelocity(const TVector<T, d>& InVelocity)
+		{
+			MVelocity = InVelocity;
+		}
+
+		void SetDrag(const T InDrag)
+		{
+			static const T Area = (T)0.1;  // TODO: Work out a correct calculation of the area
+			static const T AirDensity = (T)1.225;  // TODO: Expose fluid density for other fluid effects
+			MHalfRhoDragArea = (T)0.5 * AirDensity * InDrag * Area;
+		}
+
+	private:
+		TSet<const TGeometryParticleHandleImp<T, d, true>*> UsesVelocityField;
+		TVector<T, d> MVelocity;
+		float MHalfRhoDragArea;
+	};
+}
