Another build fix

#rnx

diff --git a/Engine/Source/Runtime/Experimental/ChaosCore/Public/Chaos/Vector.h b/Engine/Source/Runtime/Experimental/ChaosCore/Public/Chaos/Vector.h
--- a/Engine/Source/Runtime/Experimental/ChaosCore/Public/Chaos/Vector.h
+++ b/Engine/Source/Runtime/Experimental/ChaosCore/Public/Chaos/Vector.h
@@ -791,12 +791,6 @@ namespace Chaos
 
 } // namespace Chaos
 
-template<class T, int d>
-struct TIsContiguousContainer<Chaos::TVector<T, d>>
-{
-	static constexpr bool Value = TIsContiguousContainer<TArray<T>>::Value;
-};
-
 //template<>
 //uint32 GetTypeHash(const Chaos::TVector<int32, 2>& V)
 //{
