More CIS fixing

#rnx

diff --git a/Engine/Source/Runtime/Experimental/ChaosCore/Public/Chaos/ArrayCollectionArray.h b/Engine/Source/Runtime/Experimental/ChaosCore/Public/Chaos/ArrayCollectionArray.h
--- a/Engine/Source/Runtime/Experimental/ChaosCore/Public/Chaos/ArrayCollectionArray.h
+++ b/Engine/Source/Runtime/Experimental/ChaosCore/Public/Chaos/ArrayCollectionArray.h
@@ -71,9 +71,3 @@ class TArrayCollectionArray : public TArrayCollectionArrayBase, public TArray<T>
 	}
 };
 }
-
-template<class T>
-struct TIsContiguousContainer<Chaos::TArrayCollectionArrayView<T>>
-{
-	static constexpr bool Value = TIsContiguousContainer<TArrayView<T>>::Value;
-};
