Another CIS fix because I'm having such a great time and I don't want it to end!

#rnx

diff --git a/Engine/Source/Runtime/Experimental/ChaosCore/Public/Chaos/ArrayCollectionArray.h b/Engine/Source/Runtime/Experimental/ChaosCore/Public/Chaos/ArrayCollectionArray.h
--- a/Engine/Source/Runtime/Experimental/ChaosCore/Public/Chaos/ArrayCollectionArray.h
+++ b/Engine/Source/Runtime/Experimental/ChaosCore/Public/Chaos/ArrayCollectionArray.h
@@ -71,3 +71,9 @@ class TArrayCollectionArray : public TArrayCollectionArrayBase, public TArray<T>
 	}
 };
 }
+
+template<class T>
+struct TIsContiguousContainer<Chaos::TArrayCollectionArray<T>>
+{
+	static constexpr bool Value = TIsContiguousContainer<TArrayView<T>>::Value;
+};
\ No newline at end of file
