Fix memory leak caused by missing virtual destructor

diff --git a/Engine/Source/Runtime/Experimental/Chaos/Public/Chaos/ParticleHandle.h b/Engine/Source/Runtime/Experimental/Chaos/Public/Chaos/ParticleHandle.h
--- a/Engine/Source/Runtime/Experimental/Chaos/Public/Chaos/ParticleHandle.h
+++ b/Engine/Source/Runtime/Experimental/Chaos/Public/Chaos/ParticleHandle.h
@@ -1244,6 +1244,8 @@ public:
 
 	void Reset() { Type = EParticleType::Static; }
 
+	virtual ~FParticleData() = default;
+
 	EParticleType Type;
 };
 
