#chaos - Fix incorrect Chaos heightfield height caused by wrong Barycentric coordinate calculation.

diff --git a/Engine/Source/Runtime/Experimental/Chaos/Private/Chaos/HeightField.cpp b/Engine/Source/Runtime/Experimental/Chaos/Private/Chaos/HeightField.cpp
--- a/Engine/Source/Runtime/Experimental/Chaos/Private/Chaos/HeightField.cpp
+++ b/Engine/Source/Runtime/Experimental/Chaos/Private/Chaos/HeightField.cpp
@@ -563,7 +563,7 @@ namespace Chaos
 				
 				FVector Bary = FMath::GetBaryCentric2D({ FractionX, FractionY, 0.0f }, Tri[0], Tri[1], Tri[2]);
 
-				return Pts[0].Z * Bary[0] + Pts[1].Z * Bary[1] * Pts[3].Z * Bary[2];
+				return Pts[0].Z * Bary[0] + Pts[1].Z * Bary[1] + Pts[3].Z * Bary[2];
 			}
 			else
 			{
@@ -583,7 +583,7 @@ namespace Chaos
 
 				FVector Bary = FMath::GetBaryCentric2D({FractionX, FractionY, 0.0f}, Tri[0], Tri[1], Tri[2]);
 
-				return Pts[0].Z * Bary[0] + Pts[1].Z * Bary[1] * Pts[3].Z * Bary[2];
+				return Pts[0].Z * Bary[0] + Pts[1].Z * Bary[1] + Pts[3].Z * Bary[2];
 			}
 		}
 
