Implemented holes for sim + query Chaos landscapes

diff --git a/Engine/Source/Runtime/Experimental/Chaos/Private/Chaos/HeightField.cpp b/Engine/Source/Runtime/Experimental/Chaos/Private/Chaos/HeightField.cpp
--- a/Engine/Source/Runtime/Experimental/Chaos/Private/Chaos/HeightField.cpp
+++ b/Engine/Source/Runtime/Experimental/Chaos/Private/Chaos/HeightField.cpp
@@ -143,11 +143,22 @@ namespace Chaos
 				{
 					if(Time < OutTime)
 					{
-						OutPosition = ResultPosition;
-						OutNormal = ResultNormal;
-						OutTime = Time;
-						OutFaceIndex = FaceIndex;
-						CurrentLength = Time;
+						bool bHole = false;
+
+						const int32 CellIndex = FaceIndex / 2;
+						if(GeomData->MaterialIndices.IsValidIndex(CellIndex))
+						{
+							bHole = GeomData->MaterialIndices[CellIndex] == TNumericLimits<uint8>::Max();
+						}
+
+						if(!bHole)
+						{
+							OutPosition = ResultPosition;
+							OutNormal = ResultNormal;
+							OutTime = Time;
+							OutFaceIndex = FaceIndex;
+							CurrentLength = Time;
+						}
 					}
 				}
 
@@ -217,6 +228,7 @@ namespace Chaos
 				{
 					return false;
 				}
+
 				TTriangle<T> Convex(A, B, C);
 
 				T Time;
@@ -226,18 +238,29 @@ namespace Chaos
 				{
 					if(Time < OutTime)
 					{
-						OutNormal = HitNormal;
-						OutPosition = HitPosition;
-						OutTime = Time;
-						OutFaceIndex = FaceIndex;
+						bool bHole = false;
 
-						if(Time <= 0)	//initial overlap or MTD, so stop
+						const int32 CellIndex = FaceIndex / 2;
+						if(HfData->MaterialIndices.IsValidIndex(CellIndex))
 						{
-							CurrentLength = 0;
-							return false;
+							bHole = HfData->MaterialIndices[CellIndex] == TNumericLimits<uint8>::Max();
 						}
 
-						CurrentLength = Time;
+						if(!bHole)
+						{
+							OutNormal = HitNormal;
+							OutPosition = HitPosition;
+							OutTime = Time;
+							OutFaceIndex = FaceIndex;
+
+							if(Time <= 0) //initial overlap or MTD, so stop
+							{
+								CurrentLength = 0;
+								return false;
+							}
+
+							CurrentLength = Time;
+						}
 					}
 				}
 
@@ -1235,6 +1258,15 @@ namespace Chaos
 		for (const TVector<int32, 2>& Cell : Intersections)
 		{
 			const int32 SingleIndex = Cell[1] * GeomData.NumCols + Cell[0];
+			const int32 CellIndex = Cell[1] * (GeomData.NumCols - 1) + Cell[0];
+			
+			// Check for holes and skip checking if we'll never collide
+			if(GeomData.MaterialIndices.IsValidIndex(CellIndex) && GeomData.MaterialIndices[CellIndex] == TNumericLimits<uint8>::Max())
+			{
+				continue;
+			}
+
+			// The triangle is solid so proceed to test it
 			GeomData.GetPointsScaled(SingleIndex, Points);
 
 			if (OverlapTriangle(Points[0], Points[1], Points[3], LocalContactLocation, LocalContactNormal, LocalContactPhi))
@@ -1258,7 +1290,7 @@ namespace Chaos
 			}
 		}
 
-		if (ContactPhi < 0)
+		if(ContactPhi < 0)
 			return true;
 		return false;
 	}
diff --git a/Engine/Source/Runtime/Landscape/Private/LandscapeCollision.cpp b/Engine/Source/Runtime/Landscape/Private/LandscapeCollision.cpp
--- a/Engine/Source/Runtime/Landscape/Private/LandscapeCollision.cpp
+++ b/Engine/Source/Runtime/Landscape/Private/LandscapeCollision.cpp
@@ -129,7 +129,7 @@ ULandscapeMeshCollisionComponent::FPhysXMeshRef::~FPhysXMeshRef()
 }
 
 // Generate a new guid to force a recache of landscape collison derived data
-#define LANDSCAPE_COLLISION_DERIVEDDATA_VER	TEXT("2265A519785B402D90EF2E726E6D8759")
+#define LANDSCAPE_COLLISION_DERIVEDDATA_VER	TEXT("43B2B22117524EAB8FB46062C2AF0BDA")
 
 static FString GetHFDDCKeyString(const FName& Format, bool bDefMaterial, const FGuid& StateId, const TArray<UPhysicalMaterial*>& PhysicalMaterials)
 {
@@ -895,7 +895,7 @@ bool ULandscapeHeightfieldCollisionComponent::CookCollisionData(const FName& For
 	{
 		for(int32 ColIndex = 0; ColIndex < CollisionSizeVerts; ColIndex++)
 		{
-			const int32 SrcSampleIndex = (ColIndex * CollisionSizeVerts) + (bIsMirrored ? RowIndex : (CollisionSizeVerts - RowIndex - 1));
+			const int32 SrcSampleIndex = (RowIndex * CollisionSizeVerts) + (bIsMirrored ? (CollisionSizeVerts - ColIndex - 1) : ColIndex);
 
 			// Materials are not relevant on the last row/column because they are per-triangle and the last row/column don't own any
 			if(RowIndex < CollisionSizeVerts - 1 &&
