Chaos - minor Joint Solver simplification

diff --git a/Engine/Source/Runtime/Experimental/Chaos/Private/Chaos/Joint/PBDJointSolverGaussSeidel.cpp b/Engine/Source/Runtime/Experimental/Chaos/Private/Chaos/Joint/PBDJointSolverGaussSeidel.cpp
--- a/Engine/Source/Runtime/Experimental/Chaos/Private/Chaos/Joint/PBDJointSolverGaussSeidel.cpp
+++ b/Engine/Source/Runtime/Experimental/Chaos/Private/Chaos/Joint/PBDJointSolverGaussSeidel.cpp
@@ -869,48 +869,30 @@ namespace Chaos
 		const FReal II1 = FVec3::DotProduct(AngularAxis1, IA1);
 		const FReal II = (InvMs[0] + II0 + InvMs[1] + II1);
 
-		FReal VelDt = 0;
+		FReal Vel = 0;
 		if (Damping > KINDA_SMALL_NUMBER)
 		{
-			const FVec3 V0 = FVec3::CalculateVelocity(PrevXs[0], Xs[0], (FReal)1.0);
-			const FVec3 V1 = FVec3::CalculateVelocity(PrevXs[1], Xs[1], (FReal)1.0);
-			VelDt = FVec3::DotProduct(V0 - V1, Axis);
+			const FVec3 V0 = FVec3::CalculateVelocity(PrevXs[0], Xs[0], Dt);
+			const FVec3 V1 = FVec3::CalculateVelocity(PrevXs[1], Xs[1], Dt);
+			Vel = FVec3::DotProduct(V0 - V1, Axis);
 		}
 
+		// XPBD, but numerically stable at low stiffness
 		FReal DLambda = 0;
-		if (Stiffness > KINDA_SMALL_NUMBER)
+		const FReal S = Stiffness * Dt * Dt;
+		const FReal D = Damping * Dt * Dt;
+		const FReal D1 = Damping * Dt;
+		if (bAccelerationMode)
 		{
-			// As below, but numerically stable for large values of stiffness (same form as in XPBD paper)
-			const FReal Alpha = (FReal)1 / (Stiffness * Dt * Dt);
-			const FReal Gamma = Damping / (Stiffness * Dt);
-			if (bAccelerationMode)
-			{
-				const FReal Multiplier = (FReal)1 / (((FReal)1 + Gamma) + Alpha);
-				DLambda = Multiplier * (Delta - Gamma * VelDt) / II - Multiplier * Alpha * Lambda;
-			}
-			else
-			{
-				const FReal Multiplier = (FReal)1 / (((FReal)1 + Gamma) * II + Alpha);
-				DLambda = Multiplier * (Delta - Gamma * VelDt - Alpha * Lambda);
-			}
+			const FReal Multiplier = (FReal)1 / ((S + D1) + (FReal)1);
+			DLambda = Multiplier * (S * Delta - D * Vel) / II - Multiplier * Lambda;
 		}
 		else
 		{
-			// As above, but numerically stable at low stiffness (alpha -> infinity)
-			const FReal S = Stiffness * Dt * Dt;
-			const FReal D = Damping * Dt;
-			if (bAccelerationMode)
-			{
-				const FReal Multiplier = (FReal)1 / ((S + D) + (FReal)1);
-				DLambda = Multiplier * (S * Delta - D * VelDt) / II - Multiplier * Lambda;
-			}
-			else
-			{
-				const FReal Multiplier = (FReal)1 / ((S + D) * II + (FReal)1);
-				DLambda = Multiplier * (S * Delta - D * VelDt - Lambda);
-			}
+			const FReal Multiplier = (FReal)1 / ((S + D1) * II + (FReal)1);
+			DLambda = Multiplier * (S * Delta - D * Vel - Lambda);
 		}
-	
+
 		const FVec3 DP0 = (InvMs[0] * DLambda) * Axis;
 		const FVec3 DP1 = (-InvMs[1] * DLambda) * Axis;
 		const FVec3 DR0 = DLambda * Utilities::Multiply(InvI0, AngularAxis0);
@@ -990,46 +972,28 @@ namespace Chaos
 		const FReal II = (II0 + II1);
 
 		// Damping angular velocity
-		FReal AngVelDt = 0;
+		FReal AngVel = 0;
 		if (Damping > KINDA_SMALL_NUMBER)
 		{
-			const FVec3 W0 = FRotation3::CalculateAngularVelocity(PrevQs[0], Qs[0], (FReal)1.0);
-			const FVec3 W1 = FRotation3::CalculateAngularVelocity(PrevQs[1], Qs[1], (FReal)1.0);
-			AngVelDt = FVec3::DotProduct(Axis0, W0) - FVec3::DotProduct(Axis1, W1);
+			const FVec3 W0 = FRotation3::CalculateAngularVelocity(PrevQs[0], Qs[0], Dt);
+			const FVec3 W1 = FRotation3::CalculateAngularVelocity(PrevQs[1], Qs[1], Dt);
+			AngVel = FVec3::DotProduct(Axis0, W0) - FVec3::DotProduct(Axis1, W1);
 		}
 
+		// XPBD, but numerically stable at low stiffness
 		FReal DLambda = 0;
-		if (Stiffness > KINDA_SMALL_NUMBER)
+		const FReal S = Stiffness * Dt * Dt;
+		const FReal D = Damping * Dt * Dt;
+		const FReal D1 = Damping * Dt;
+		if (bAccelerationMode)
 		{
-			// As below, but numerically stable for large values of stiffness (same form as in XPBD paper)
-			const FReal Alpha = (FReal)1 / (Stiffness * Dt * Dt);
-			const FReal Gamma = Damping / (Stiffness * Dt);
-			if (bAccelerationMode)
-			{
-				const FReal Multiplier = (FReal)1 / (((FReal)1 + Gamma) + Alpha);
-				DLambda = Multiplier * (Angle - Gamma * AngVelDt) / II - Multiplier * Alpha * Lambda;
-			}
-			else
-			{
-				const FReal Multiplier = (FReal)1 / (((FReal)1 + Gamma) * II + Alpha);
-				DLambda = Multiplier * (Angle - Gamma * AngVelDt - Alpha * Lambda);
-			}
+			const FReal Multiplier = (FReal)1 / ((S + D1) + (FReal)1);
+			DLambda = Multiplier * (S * Angle - D * AngVel) / II - Multiplier * Lambda;
 		}
 		else
 		{
-			// As above, but numerically stable at low stiffness (alpha -> infinity)
-			const FReal S = Stiffness * Dt * Dt;
-			const FReal D = Damping * Dt;
-			if (bAccelerationMode)
-			{
-				const FReal Multiplier = (FReal)1 / ((S + D) + (FReal)1);
-				DLambda = Multiplier * (S * Angle - D * AngVelDt) / II - Multiplier * Lambda;
-			}
-			else
-			{
-				const FReal Multiplier = (FReal)1 / ((S + D) * II + (FReal)1);
-				DLambda = Multiplier * (S * Angle - D * AngVelDt - Lambda);
-			}
+			const FReal Multiplier = (FReal)1 / ((S + D1) * II + (FReal)1);
+			DLambda = Multiplier * (S * Angle - D * AngVel - Lambda);
 		}
 
 		//const FVec3 DR0 = IA0 * DLambda;
