Add option to PhysScene and PhysicsTickTask to tick a single solver, rather than all solvers in the SolversModule

diff --git a/Engine/Source/Runtime/Engine/Private/PhysicsEngine/Experimental/PhysScene_Chaos.cpp b/Engine/Source/Runtime/Engine/Private/PhysicsEngine/Experimental/PhysScene_Chaos.cpp
--- a/Engine/Source/Runtime/Engine/Private/PhysicsEngine/Experimental/PhysScene_Chaos.cpp
+++ b/Engine/Source/Runtime/Engine/Private/PhysicsEngine/Experimental/PhysScene_Chaos.cpp
@@ -1434,7 +1434,7 @@ void FPhysScene_ChaosInterface::StartFrame()
 			// Need to fire off a parallel task to handle running physics commands and
 			// ticking the scene while the engine continues on until TG_EndPhysics
 			// (this should happen in TG_StartPhysics)
-			PhysicsTickTask = TGraphTask<FPhysicsTickTask>::CreateTask(nullptr, ENamedThreads::GameThread).ConstructAndDispatchWhenReady(SimulationCompleteEvent, Dt);
+			PhysicsTickTask = TGraphTask<FPhysicsTickTask>::CreateTask(nullptr, ENamedThreads::GameThread).ConstructAndDispatchWhenReady(SimulationCompleteEvent, GetSolver(), Dt);
 
 			// Setup post simulate tasks
 			if (PhysicsTickTask.GetReference())
@@ -1495,8 +1495,14 @@ void FPhysScene_ChaosInterface::EndFrame(ULineBatchComponent* InLineBatcher)
 		//flush queue so we can merge the two threads
 		Dispatcher->Execute();
 
+		// Make a list of solvers to process. This is just our solver if we have one,
+		// and if not then it's all solvers in the solvers module.
+		const TArray<FPhysicsSolver*>& SolverList
+			= (GetSolver() == nullptr)
+			? SolverModule->GetSolvers()
+			: [&]() { TArray<FPhysicsSolver*> Solvers; Solvers.Init(GetSolver(), 1); return Solvers; }();
+
 		// flush solver queues
-		const TArray<FPhysicsSolver*>& SolverList = SolverModule->GetSolvers();
 		for (FPhysicsSolver* Solver : SolverList)
 		{
 			TQueue<TFunction<void(Chaos::FPhysicsSolver*)>, EQueueMode::Mpsc>& Queue = Solver->GetCommandQueue();
diff --git a/Engine/Source/Runtime/Experimental/ChaosSolvers/Private/Framework/PhysicsTickTask.cpp b/Engine/Source/Runtime/Experimental/ChaosSolvers/Private/Framework/PhysicsTickTask.cpp
--- a/Engine/Source/Runtime/Experimental/ChaosSolvers/Private/Framework/PhysicsTickTask.cpp
+++ b/Engine/Source/Runtime/Experimental/ChaosSolvers/Private/Framework/PhysicsTickTask.cpp
@@ -16,9 +16,11 @@ FAutoConsoleTaskPriority CPrio_FPhysicsTickTask(
 	ENamedThreads::HighTaskPriority // if we don't have hi pri threads, then use normal priority threads at high task priority instead
 );
 
-FPhysicsTickTask::FPhysicsTickTask(FGraphEventRef& InCompletionEvent, float InDt)
+// If InSolver is null, then the tick task will tick all solvers in the solvers module
+FPhysicsTickTask::FPhysicsTickTask(FGraphEventRef& InCompletionEvent, Chaos::FPhysicsSolver* InPhysicsSolver, float InDt)
 	: CompletionEvent(InCompletionEvent)
 	, Module(nullptr)
+	, PhysicsSolver(InPhysicsSolver)
 	, Dt(InDt)
 {
 	Module = FChaosSolversModule::GetModule();
@@ -50,7 +52,12 @@ void FPhysicsTickTask::DoTask(ENamedThreads::Type CurrentThread, const FGraphEve
 	// per-solver commands and the solver advance
 	FGraphEventRef CommandsTask = TGraphTask<FPhysicsCommandsTask>::CreateTask().ConstructAndDispatchWhenReady();
 
-	const TArray<FPhysicsSolver*>& SolverList = Module->GetSolvers();
+	// If the stored solver is non-null, create a temporary solver list with just that element.
+	// Otherwise get a full list of solvers from the solvers module.
+	const TArray<FPhysicsSolver*>& SolverList
+		= (PhysicsSolver == nullptr)
+		? Module->GetSolvers()
+		: [&]() { TArray<FPhysicsSolver*> Solvers; Solvers.Init(PhysicsSolver, 1); return Solvers; }();
 
 	// List of active solvers (assume all are active for single alloc)
 	TArray<FPhysicsSolver*> ActiveSolvers;
diff --git a/Engine/Source/Runtime/Experimental/ChaosSolvers/Public/Framework/PhysicsTickTask.h b/Engine/Source/Runtime/Experimental/ChaosSolvers/Public/Framework/PhysicsTickTask.h
--- a/Engine/Source/Runtime/Experimental/ChaosSolvers/Public/Framework/PhysicsTickTask.h
+++ b/Engine/Source/Runtime/Experimental/ChaosSolvers/Public/Framework/PhysicsTickTask.h
@@ -17,7 +17,7 @@ class CHAOSSOLVERS_API FPhysicsTickTask
 {
 public:
 
-	FPhysicsTickTask(FGraphEventRef& InCompletionEvent, float InDt);
+	FPhysicsTickTask(FGraphEventRef& InCompletionEvent, Chaos::FPhysicsSolver* InPhysicsSolver, float InDt);
 
 	TStatId GetStatId() const;
 	static ENamedThreads::Type GetDesiredThread();
@@ -28,6 +28,7 @@ private:
 
 	FGraphEventRef CompletionEvent;
 	FChaosSolversModule* Module;
+	Chaos::FPhysicsSolver* PhysicsSolver;
 	float Dt;
 };
 
