Chaos - Soft joint fixes

diff --git a/Engine/Source/Runtime/Engine/Private/PhysicsEngine/ImmediatePhysics/ImmediatePhysicsChaos/ImmediatePhysicsJointHandle_Chaos.cpp b/Engine/Source/Runtime/Engine/Private/PhysicsEngine/ImmediatePhysics/ImmediatePhysicsChaos/ImmediatePhysicsJointHandle_Chaos.cpp
--- a/Engine/Source/Runtime/Engine/Private/PhysicsEngine/ImmediatePhysics/ImmediatePhysicsChaos/ImmediatePhysicsJointHandle_Chaos.cpp
+++ b/Engine/Source/Runtime/Engine/Private/PhysicsEngine/ImmediatePhysics/ImmediatePhysicsChaos/ImmediatePhysicsJointHandle_Chaos.cpp
@@ -28,8 +28,8 @@ namespace ImmediatePhysics_Chaos
 	FAutoConsoleVariableRef CVarDriveStiffnessScale(TEXT("p.Chaos.ImmPhys.DriveStiffnessScale"), ChaosImmediate_DriveStiffnessScale, TEXT("Conversion factor for drive stiffness."));
 	FAutoConsoleVariableRef CVarDriveDampingScale(TEXT("p.Chaos.ImmPhys.DriveDampingScale"), ChaosImmediate_DriveDampingScale, TEXT("Conversion factor for drive damping."));
 
-	float ChaosImmediate_SoftLinearStiffnessScale = 1.0f;
-	float ChaosImmediate_SoftLinearDampingScale = 1.0f;
+	float ChaosImmediate_SoftLinearStiffnessScale = 1.75f;
+	float ChaosImmediate_SoftLinearDampingScale = 1.75f;
 	FAutoConsoleVariableRef CVarSoftLinearStiffnessScale(TEXT("p.Chaos.ImmPhys.SoftLinearStiffnessScale"), ChaosImmediate_SoftLinearStiffnessScale, TEXT("Conversion factor for soft-joint stiffness."));
 	FAutoConsoleVariableRef CVarSoftLinearDampingScale(TEXT("p.Chaos.ImmPhys.SoftLinearDampingScale"), ChaosImmediate_SoftLinearDampingScale, TEXT("Conversion factor for soft-joint damping."));
 
@@ -38,7 +38,7 @@ namespace ImmediatePhysics_Chaos
 	FAutoConsoleVariableRef CVarSoftLinearForceMode(TEXT("p.Chaos.ImmPhys.SoftLinearForceMode"), ChaosImmediate_SoftLinearForceMode, TEXT("Soft Linear constraint force mode (0: Acceleration; 1: Force"));
 	FAutoConsoleVariableRef CVarSoftAngularForceMode(TEXT("p.Chaos.ImmPhys.SoftAngularForceMode"), ChaosImmediate_SoftAngularForceMode, TEXT("Soft Angular constraint force mode (0: Acceleration; 1: Force"));
 
-	// This value is from PhysX UE - in Chaos we don't really need Soft constraints with stiffness greater than about 1000
+	// This value US passed to PhysX was 100000 UE
 	// (In PhysX, Soft constraints were implicitly integrated, and essentially equivalent to XPBD, which for large stiffness is just PBD).
 	float ChaosImmediate_SoftAngularStiffnessScale = 100000.0f;
 	float ChaosImmediate_SoftAngularDampingScale = 100000.0f;
diff --git a/Engine/Source/Runtime/Experimental/Chaos/Private/Chaos/Joint/PBDJointSolverGaussSeidel.cpp b/Engine/Source/Runtime/Experimental/Chaos/Private/Chaos/Joint/PBDJointSolverGaussSeidel.cpp
--- a/Engine/Source/Runtime/Experimental/Chaos/Private/Chaos/Joint/PBDJointSolverGaussSeidel.cpp
+++ b/Engine/Source/Runtime/Experimental/Chaos/Private/Chaos/Joint/PBDJointSolverGaussSeidel.cpp
@@ -356,10 +356,10 @@ namespace Chaos
 	{
 		InitAccumulatedDeltas();
 
+		ApplyPositionConstraints(Dt, SolverSettings, JointSettings);
+
 		ApplyRotationConstraints(Dt, SolverSettings, JointSettings);
 
-		ApplyPositionConstraints(Dt, SolverSettings, JointSettings);
-	
 		ApplyAccumulatedDeltas();
 	}
 
@@ -869,29 +869,19 @@ namespace Chaos
 		const FReal II1 = FVec3::DotProduct(AngularAxis1, IA1);
 		const FReal II = (InvMs[0] + II0 + InvMs[1] + II1);
 
-		FReal Vel = 0;
+		FReal VelDt = 0;
 		if (Damping > KINDA_SMALL_NUMBER)
 		{
-			const FVec3 V0 = FVec3::CalculateVelocity(PrevXs[0], Xs[0], Dt);
-			const FVec3 V1 = FVec3::CalculateVelocity(PrevXs[1], Xs[1], Dt);
-			Vel = FVec3::DotProduct(V0 - V1, Axis);
+			const FVec3 V0 = FVec3::CalculateVelocity(PrevXs[0], Xs[0], 1.0f);
+			const FVec3 V1 = FVec3::CalculateVelocity(PrevXs[1], Xs[1], 1.0f);
+			VelDt = FVec3::DotProduct(V0 - V1, Axis);
 		}
 
-		// XPBD, but numerically stable at low stiffness
-		FReal DLambda = 0;
-		const FReal S = Stiffness * Dt * Dt;
-		const FReal D = Damping * Dt * Dt;
-		const FReal D1 = Damping * Dt;
-		if (bAccelerationMode)
-		{
-			const FReal Multiplier = (FReal)1 / ((S + D1) + (FReal)1);
-			DLambda = Multiplier * (S * Delta - D * Vel) / II - Multiplier * Lambda;
-		}
-		else
-		{
-			const FReal Multiplier = (FReal)1 / ((S + D1) * II + (FReal)1);
-			DLambda = Multiplier * (S * Delta - D * Vel - Lambda);
-		}
+		const FReal SpringMassScale = (bAccelerationMode) ? 1.0f / (InvMs[0] + InvMs[1]) : 1.0f;
+		const FReal S = SpringMassScale * Stiffness * Dt * Dt;
+		const FReal D = SpringMassScale * Damping * Dt;
+		const FReal Multiplier = (FReal)1 / ((S + D) * II + (FReal)1);
+		const FReal DLambda = Multiplier * (S * Delta - D * VelDt - Lambda);
 
 		const FVec3 DP0 = (InvMs[0] * DLambda) * Axis;
 		const FVec3 DP1 = (-InvMs[1] * DLambda) * Axis;
@@ -972,29 +962,19 @@ namespace Chaos
 		const FReal II = (II0 + II1);
 
 		// Damping angular velocity
-		FReal AngVel = 0;
+		FReal AngVelDt = 0;
 		if (Damping > KINDA_SMALL_NUMBER)
 		{
-			const FVec3 W0 = FRotation3::CalculateAngularVelocity(PrevQs[0], Qs[0], Dt);
-			const FVec3 W1 = FRotation3::CalculateAngularVelocity(PrevQs[1], Qs[1], Dt);
-			AngVel = FVec3::DotProduct(Axis0, W0) - FVec3::DotProduct(Axis1, W1);
+			const FVec3 W0 = FRotation3::CalculateAngularVelocity(PrevQs[0], Qs[0], 1.0f);
+			const FVec3 W1 = FRotation3::CalculateAngularVelocity(PrevQs[1], Qs[1], 1.0f);
+			AngVelDt = FVec3::DotProduct(Axis0, W0) - FVec3::DotProduct(Axis1, W1);
 		}
 
-		// XPBD, but numerically stable at low stiffness
-		FReal DLambda = 0;
-		const FReal S = Stiffness * Dt * Dt;
-		const FReal D = Damping * Dt * Dt;
-		const FReal D1 = Damping * Dt;
-		if (bAccelerationMode)
-		{
-			const FReal Multiplier = (FReal)1 / ((S + D1) + (FReal)1);
-			DLambda = Multiplier * (S * Angle - D * AngVel) / II - Multiplier * Lambda;
-		}
-		else
-		{
-			const FReal Multiplier = (FReal)1 / ((S + D1) * II + (FReal)1);
-			DLambda = Multiplier * (S * Angle - D * AngVel - Lambda);
-		}
+		const FReal SpringMassScale = (bAccelerationMode) ? 1.0f / II : 1.0f;
+		const FReal S = SpringMassScale * Stiffness * Dt * Dt;
+		const FReal D = SpringMassScale * Damping * Dt;
+		const FReal Multiplier = (FReal)1 / ((S + D) * II + (FReal)1);
+		const FReal DLambda = Multiplier * (S * Angle - D * AngVelDt - Lambda);
 
 		//const FVec3 DR0 = IA0 * DLambda;
 		//const FVec3 DR1 = IA1 * -DLambda;
