Chaos - RBAN fix force initialization when substepping

diff --git a/Engine/Source/Runtime/Experimental/Chaos/Private/Chaos/Evolution/PBDMinEvolution.cpp b/Engine/Source/Runtime/Experimental/Chaos/Private/Chaos/Evolution/PBDMinEvolution.cpp
--- a/Engine/Source/Runtime/Experimental/Chaos/Private/Chaos/Evolution/PBDMinEvolution.cpp
+++ b/Engine/Source/Runtime/Experimental/Chaos/Private/Chaos/Evolution/PBDMinEvolution.cpp
@@ -44,6 +44,14 @@ namespace Chaos
 
 	void FPBDMinEvolution::Advance(const FReal StepDt, const int32 NumSteps)
 	{
+		for (TTransientPBDRigidParticleHandle<FReal, 3>& Particle : Particles.GetActiveParticlesView())
+		{
+			if (Particle.ObjectState() == EObjectStateType::Dynamic)
+			{
+				Particle.F() += Particle.M() * Gravity;
+			}
+		}
+
 		for (int32 Step = 0; Step < NumSteps; ++Step)
 		{
 			// StepFraction: how much of the remaining time this step represents, used to interpolate kinematic targets
@@ -54,6 +62,15 @@ namespace Chaos
 
 			AdvanceOneTimeStep(StepDt, StepFraction);
 		}
+
+		for (TTransientPBDRigidParticleHandle<FReal, 3>& Particle : Particles.GetActiveParticlesView())
+		{
+			if (Particle.ObjectState() == EObjectStateType::Dynamic)
+			{
+				Particle.F() = FVec3(0);
+				Particle.Torque() = FVec3(0);
+			}
+		}
 	}
 
 	void FPBDMinEvolution::AdvanceOneTimeStep(const FReal Dt, const FReal StepFraction)
@@ -115,18 +132,16 @@ namespace Chaos
 				Particle.PreV() = Particle.V();
 				Particle.PreW() = Particle.W();
 
-				Particle.F() += Particle.M() * Gravity;
-
 				EulerStepVelocityRule.Apply(Particle, Dt);
+
 				AddImpulsesRule.Apply(Particle, Dt);
+				Particle.LinearImpulse() = FVec3(0);
+				Particle.AngularImpulse() = FVec3(0);
+
 				EtherDragRule.Apply(Particle, Dt);
 
 				EulerStepPositionRule.Apply(Particle, Dt);
 
-				Particle.F() = FVec3(0);
-				Particle.Torque() = FVec3(0);
-				Particle.LinearImpulse() = FVec3(0);
-				Particle.AngularImpulse() = FVec3(0);
 
 				if (Particle.HasBounds())
 				{
