Chaos - alternate method for calculating angvel from change in quaternion (more accurate at small dq)

diff --git a/Engine/Source/Runtime/Experimental/Chaos/Public/Chaos/Rotation.h b/Engine/Source/Runtime/Experimental/Chaos/Public/Chaos/Rotation.h
--- a/Engine/Source/Runtime/Experimental/Chaos/Public/Chaos/Rotation.h
+++ b/Engine/Source/Runtime/Experimental/Chaos/Public/Chaos/Rotation.h
@@ -264,7 +264,7 @@ namespace Chaos
 		 */
 		static TVector<float, 3> CalculateAngularVelocity(const TRotation<float, 3>& InR0, const TRotation<float, 3>& InR1, const float InDt)
 		{
-			return CalculateAngularVelocity2(InR0, InR1, InDt);
+			return CalculateAngularVelocity1(InR0, InR1, InDt);
 		}
 
 		/**
