Chaos - Joints ISPC optimizations

diff --git a/Engine/Source/Runtime/Core/Public/Math/Matrix.isph b/Engine/Source/Runtime/Core/Public/Math/Matrix.isph
--- a/Engine/Source/Runtime/Core/Public/Math/Matrix.isph
+++ b/Engine/Source/Runtime/Core/Public/Math/Matrix.isph
@@ -31,6 +31,8 @@ struct FMatrix3x4
 #endif
 };
 
+static const uniform struct FMatrix MatrixIdentity = {{ 1, 0, 0, 0, 0, 1, 0, 0, 0, 0, 1, 0, 0, 0, 0, 1}};
+
 inline uniform FMatrix SetMatrix(const uniform FVector4 &R0, const uniform FVector4 &R1, const uniform FVector4 &R2, const uniform FVector4 &R3)
 {
 	const uniform FMatrix Result = {{ R0.V[0], R0.V[1], R0.V[2], R0.V[3],
@@ -150,21 +152,9 @@ inline uniform FVector MatrixGetOrigin(const uniform FMatrix &M)
 
 inline void MatrixGetScaledAxes(const uniform FMatrix &M, uniform FVector &X, uniform FVector &Y, uniform FVector &Z)
 {
-	varying float S0, S1, S2;
-	*((uniform FVector *uniform)&S0) = *((uniform FVector *uniform)&M.M[0]);
-	*((uniform FVector *uniform)&S1) = *((uniform FVector *uniform)&M.M[4]);
-	*((uniform FVector *uniform)&S2) = *((uniform FVector *uniform)&M.M[8]);
-
-#if EXPLICIT_VECTOR == 1
-	// Using explicit vector aligns to 16 bytes. Zero the 4th components when pulling from Matrix type
-	S0 = insert(S0, 3, 0.0f);
-	S1 = insert(S1, 3, 0.0f);
-	S2 = insert(S2, 3, 0.0f);
-#endif
-
-	*((uniform FVector *uniform)&X) = *((uniform FVector *uniform)&S0);
-	*((uniform FVector *uniform)&Y) = *((uniform FVector *uniform)&S1);
-	*((uniform FVector *uniform)&Z) = *((uniform FVector *uniform)&S2);
+	X = SetVector(M.M[0], M.M[1], M.M[2]);
+	Y = SetVector(M.M[4], M.M[5], M.M[6]);
+	Z = SetVector(M.M[8], M.M[9], M.M[10]);
 }
 
 inline void MatrixScaleAxis(uniform FMatrix &M, const uniform int Axis, const uniform float ScalingFactor)
@@ -188,12 +178,25 @@ inline void MatrixScaleAxis(uniform FMatrix &M, const uniform int Axis, const un
 	*((uniform FVector *uniform)&M.M[(4*Axis)]) = *((uniform FVector *uniform)&S0);
 }
 
-inline uniform bool IsMatrixScaledAxesNearlyZero(const uniform FMatrix &M)
+inline uniform bool IsAnyMatrixScaledAxesNearlyZero(const uniform FMatrix &M, const uniform float Tolerance)
+{
+	uniform FVector X, Y, Z;
+	MatrixGetScaledAxes(M, X, Y, Z);
+
+	if(VectorIsAnyNearlyZero(X, Tolerance) || VectorIsAnyNearlyZero(Y, Tolerance) || VectorIsAnyNearlyZero(Z, Tolerance))
+	{
+		return true;
+	}
+
+	return false;
+}
+
+inline uniform bool IsAllMatrixScaledAxesNearlyZero(const uniform FMatrix &M, const uniform float Tolerance)
 {
 	uniform FVector X, Y, Z;
 	MatrixGetScaledAxes(M, X, Y, Z);
 
-	if(VectorIsNearlyZero(X) || VectorIsNearlyZero(Y) || VectorIsNearlyZero(Z))
+	if(VectorIsAllNearlyZero(X, Tolerance) && VectorIsAllNearlyZero(Y, Tolerance) && VectorIsAllNearlyZero(Z, Tolerance))
 	{
 		return true;
 	}
@@ -201,4 +204,177 @@ inline uniform bool IsMatrixScaledAxesNearlyZero(const uniform FMatrix &M)
 	return false;
 }
 
+inline uniform FMatrix MatrixTranspose(const uniform FMatrix& M)
+{
+	uniform FMatrix Result;
+
+	const uniform FVector4 R1 = *((uniform FVector4 *uniform)&M.M[0]);
+	const uniform FVector4 R2 = *((uniform FVector4 *uniform)&M.M[4]);
+	const uniform FVector4 R3 = *((uniform FVector4 *uniform)&M.M[8]);
+	const uniform FVector4 R4 = *((uniform FVector4 *uniform)&M.M[12]);
+
+	*((uniform FVector4 *uniform)&Result.M[0]) = SetVector4(R1.V[0], R2.V[0], R3.V[0], R4.V[0]);
+	*((uniform FVector4 *uniform)&Result.M[4]) = SetVector4(R1.V[1], R2.V[1], R3.V[1], R4.V[1]);
+	*((uniform FVector4 *uniform)&Result.M[8]) = SetVector4(R1.V[2], R2.V[2], R3.V[2], R4.V[2]);
+	*((uniform FVector4 *uniform)&Result.M[12]) = SetVector4(R1.V[3], R2.V[3], R3.V[3], R4.V[3]);
+
+	return Result;
+}
+
+// for row major matrix
+// we use __m128 to represent 2x2 matrix as A = | A0  A1 |
+//                                              | A2  A3 |
+// 2x2 row major Matrix multiply A*B
+static inline uniform FVector4 Mat2Mul(const uniform FVector4& vec1, const uniform FVector4& vec2)
+{
+	return
+		VectorAdd(VectorMultiply( vec1, VectorSwizzle(vec2, 0,3,0,3)),
+			VectorMultiply(VectorSwizzle(vec1, 1,0,3,2), VectorSwizzle(vec2, 2,1,2,1)));
+}
+// 2x2 row major Matrix adjugate multiply (A#)*B
+static inline uniform FVector4 Mat2AdjMul(const uniform FVector4& vec1, const uniform FVector4& vec2)
+{
+	return
+		VectorSubtract(VectorMultiply(VectorSwizzle(vec1, 3,3,0,0), vec2),
+			VectorMultiply(VectorSwizzle(vec1, 1,1,2,2), VectorSwizzle(vec2, 2,3,0,1)));
+
+}
+// 2x2 row major Matrix multiply adjugate A*(B#)
+static inline uniform FVector4 Mat2MulAdj(const uniform FVector4& vec1, const uniform FVector4& vec2)
+{
+	return
+		VectorSubtract(VectorMultiply( vec1, VectorSwizzle(vec2, 3,0,3,0)),
+			VectorMultiply(VectorSwizzle(vec1, 1,0,3,2), VectorSwizzle(vec2, 2,1,2,1)));
+}
+
+inline uniform float MatrixDeterminant(const uniform FMatrix& M)
+{
+	const uniform FVector4 *uniform MVec = (const uniform FVector4 *uniform) &M;
+
+	// sub matrices
+	const uniform FVector4 A = VectorShuffle(MVec[0], MVec[1], 0,1,0,1);
+	const uniform FVector4 B = VectorShuffle(MVec[0], MVec[1], 2,3,2,3);
+	const uniform FVector4 C = VectorShuffle(MVec[2], MVec[3], 0,1,0,1);
+	const uniform FVector4 D = VectorShuffle(MVec[2], MVec[3], 2,3,2,3);
+
+	// determinant as (|A| |B| |C| |D|)
+	const uniform FVector4 detSub = VectorSubtract(
+		VectorMultiply(VectorShuffle(MVec[0], MVec[2], 0,2,0,2), VectorShuffle(MVec[1], MVec[3], 1,3,1,3)),
+		VectorMultiply(VectorShuffle(MVec[0], MVec[2], 1,3,1,3), VectorShuffle(MVec[1], MVec[3], 0,2,0,2))
+	);
+
+	// D#C
+	const uniform FVector4 D_C = Mat2AdjMul(D, C);
+	// A#B
+	const uniform FVector4 A_B = Mat2AdjMul(A, B);
+
+	// |M| = |A|*|D| + |B|*|C|
+	uniform float detM = (detSub.V[0] * detSub.V[3]) + (detSub.V[1] * detSub.V[2]);
+
+	// tr((A#B)(D#C))
+	uniform FVector4 tr = VectorMultiply(A_B, VectorSwizzle(D_C, 0,2,1,3));
+	tr = tr + VectorSwizzle(tr, 2,3,0,1);
+	tr = tr + VectorSwizzle(tr, 1,2,3,0);
+	// |M| = |A|*|D| + |B|*|C| - tr((A#B)(D#C)
+	detM = detM - tr.V[0];
+
+	return detM;
+}
+
+inline uniform FMatrix VectorMatrixInverse(const uniform FMatrix& M)
+{
+	// use block matrix method
+	// A is a matrix, then i(A) or iA means inverse of A, A# (or A_ in code) means adjugate of A, |A| (or detA in code) is determinant, tr(A) is trace
+	const uniform FVector4 *uniform MVec = (const uniform FVector4 *uniform) &M;
+
+	// sub matrices
+	const uniform FVector4 A = VectorShuffle(MVec[0], MVec[1], 0,1,0,1);
+	const uniform FVector4 B = VectorShuffle(MVec[0], MVec[1], 2,3,2,3);
+	const uniform FVector4 C = VectorShuffle(MVec[2], MVec[3], 0,1,0,1);
+	const uniform FVector4 D = VectorShuffle(MVec[2], MVec[3], 2,3,2,3);
+
+	// determinant as (|A| |B| |C| |D|)
+	const uniform FVector4 detSub = VectorSubtract(
+		VectorMultiply(VectorShuffle(MVec[0], MVec[2], 0,2,0,2), VectorShuffle(MVec[1], MVec[3], 1,3,1,3)),
+		VectorMultiply(VectorShuffle(MVec[0], MVec[2], 1,3,1,3), VectorShuffle(MVec[1], MVec[3], 0,2,0,2))
+	);
+
+	const uniform FVector4 detA = VectorSwizzle(detSub, 0,0,0,0);
+	const uniform FVector4 detB = VectorSwizzle(detSub, 1,1,1,1);
+	const uniform FVector4 detC = VectorSwizzle(detSub, 2,2,2,2);
+	const uniform FVector4 detD = VectorSwizzle(detSub, 3,3,3,3);
+
+	// let iM = 1/|M| * | X  Y |
+	//                  | Z  W |
+
+	// D#C
+	const uniform FVector4 D_C = Mat2AdjMul(D, C);
+	// A#B
+	const uniform FVector4 A_B = Mat2AdjMul(A, B);
+	// X# = |D|A - B(D#C)
+	uniform FVector4 X_ = VectorSubtract(VectorMultiply(detD, A), Mat2Mul(B, D_C));
+	// W# = |A|D - C(A#B)
+	uniform FVector4 W_ = VectorSubtract(VectorMultiply(detA, D), Mat2Mul(C, A_B));
+
+	// |M| = |A|*|D| + ... (continue later)
+	uniform FVector4 detM = VectorMultiply(detA, detD);
+
+	// Y# = |B|C - D(A#B)#
+	uniform FVector4 Y_ = VectorSubtract(VectorMultiply(detB, C), Mat2MulAdj(D, A_B));
+	// Z# = |C|B - A(D#C)#
+	uniform FVector4 Z_ = VectorSubtract(VectorMultiply(detC, B), Mat2MulAdj(A, D_C));
+
+	// |M| = |A|*|D| + |B|*|C| ... (continue later)
+	detM = VectorAdd(detM, VectorMultiply(detB, detC));
+
+	// tr((A#B)(D#C))
+	uniform FVector4 tr = VectorMultiply(A_B, VectorSwizzle(D_C, 0,2,1,3));
+	tr = tr + VectorSwizzle(tr, 2,3,0,1);
+	tr = tr + VectorSwizzle(tr, 1,2,3,0);
+	// |M| = |A|*|D| + |B|*|C| - tr((A#B)(D#C)
+	detM = VectorSubtract(detM, tr);
+
+	uniform FMatrix R = MatrixIdentity;
+
+	if(detM.V[0] != 0.0f)
+	{
+		const uniform FVector4 adjSignMask = SetVector4(1.f, -1.f, -1.f, 1.f);
+		// (1/|M|, -1/|M|, -1/|M|, 1/|M|)
+		const uniform FVector4 rDetM = VectorDivide(adjSignMask, detM);
+
+		X_ = VectorMultiply(X_, rDetM);
+		Y_ = VectorMultiply(Y_, rDetM);
+		Z_ = VectorMultiply(Z_, rDetM);
+		W_ = VectorMultiply(W_, rDetM);
+
+		uniform FVector4 *uniform RVec = (uniform FVector4 *uniform) &R;
+
+		// apply adjugate and store, here we combine adjugate shuffle and store shuffle
+		RVec[0] = VectorShuffle(X_, Y_, 3,1,3,1);
+		RVec[1] = VectorShuffle(X_, Y_, 2,0,2,0);
+		RVec[2] = VectorShuffle(Z_, W_, 3,1,3,1);
+		RVec[3] = VectorShuffle(Z_, W_, 2,0,2,0);	
+	}
+
+	return R;
+}
+
+inline uniform FMatrix MatrixInverse(const uniform FMatrix& M)
+{
+	uniform FMatrix Result;
+
+	// Check for zero scale matrix to invert
+	if(	IsAllMatrixScaledAxesNearlyZero(M, SMALL_NUMBER) ) 
+	{
+		// just set to zero - avoids unsafe inverse of zero and duplicates what QNANs were resulting in before (scaling away all children)
+		Result = MatrixIdentity;
+	}
+	else
+	{
+		Result = VectorMatrixInverse(M);
+	}
+
+	return Result;
+}
+
 #endif
diff --git a/Engine/Source/Runtime/Core/Public/Math/Quat.isph b/Engine/Source/Runtime/Core/Public/Math/Quat.isph
--- a/Engine/Source/Runtime/Core/Public/Math/Quat.isph
+++ b/Engine/Source/Runtime/Core/Public/Math/Quat.isph
@@ -11,7 +11,7 @@ inline uniform FVector4 MatrixToQuat(const uniform FMatrix &M)
 	// If Matrix is NULL, return Identity quaternion. If any of them is 0, you won't be able to construct rotation
 	// if you have two plane at least, we can reconstruct the frame using cross product, but that's a bit expensive op to do here
 	// for now, if you convert to matrix from 0 scale and convert back, you'll lose rotation. Don't do that.
-	if (IsMatrixScaledAxesNearlyZero(M))
+	if (IsAnyMatrixScaledAxesNearlyZero(M, KINDA_SMALL_NUMBER))
 	{
 		return Float0001;
 	}
@@ -108,19 +108,17 @@ inline FVector4 VectorNormalizeQuaternion(const FVector4& UnnormalizedQuat)
 
 inline uniform FVector4 VectorNormalizeQuaternion(const uniform FVector4& UnnormalizedQuat)
 {
-	return VectorNormalizeSafe(UnnormalizedQuat, Float0001);
+	return VectorNormalizeSafe2(UnnormalizedQuat, Float0001);
 }
 
 inline FVector4 QuatInverse(const FVector4 &Quat)
 {
-    static const uniform FVector4 Mul = {{-1.0f, -1.0f, -1.0f, 1.0f}};
-	return Quat * Mul;
+	return Quat * QINV_SIGN_MASK;
 }
 
 inline uniform FVector4 QuatInverse(const uniform FVector4 &Quat)
 {
-    static const uniform FVector4 Mul = {{-1.0f, -1.0f, -1.0f, 1.0f}};
-	return Quat * Mul;
+	return Quat * QINV_SIGN_MASK;
 }
 
 inline FVector4 QuatFastLerp(const FVector4& A, const FVector4& B, const float Alpha)
@@ -184,6 +182,21 @@ inline uniform FVector4 VectorQuaternionRotateVector(const uniform FVector4& Qua
 	return Rotated;
 }
 
+inline uniform FVector VectorQuaternionRotateVector(const uniform FVector4& Quat, const uniform FVector& V)
+{	
+	// http://people.csail.mit.edu/bkph/articles/Quaternions.pdf
+	// V' = V + 2w(Q x V) + (2Q x (Q x V))
+	// refactor:
+	// V' = V + w(2(Q x V)) + (Q x (2(Q x V)))
+	// T = 2(Q x V);
+	// V' = V + w*(T) + (Q x T)
+
+	const uniform FVector Q = SetVector(Quat.V[0], Quat.V[1], Quat.V[2]);
+	const uniform FVector T = VectorCross(Q, V) * 2.f;
+	const uniform FVector Result = V + (T * Quat.V[3]) + VectorCross(Q, T);
+	return Result;
+}
+
 inline uniform FVector4 VectorLerpQuat(const uniform FVector4& A, const uniform FVector4& B, const uniform float Alpha)
 {
 	// Blend rotation
@@ -202,4 +215,12 @@ inline uniform FVector4 VectorLerpQuat(const uniform FVector4& A, const uniform
 	return UnnormalizedResult;
 }
 
+inline uniform FVector4 VectorQuaternionEnforceShortestArcWith(const uniform FVector4& A, const uniform FVector4& B)
+{
+	const uniform float DotResult = VectorDot(B, A);
+	const uniform float Bias = select(DotResult >= 0.f, 1.0f, -1.0f);
+
+	return A * Bias;
+}
+
 #endif
diff --git a/Engine/Source/Runtime/Core/Public/Math/Vector.isph b/Engine/Source/Runtime/Core/Public/Math/Vector.isph
--- a/Engine/Source/Runtime/Core/Public/Math/Vector.isph
+++ b/Engine/Source/Runtime/Core/Public/Math/Vector.isph
@@ -46,10 +46,13 @@
 #define SMALL_NUMBER (1.e-8f)
 #define KINDA_SMALL_NUMBER	(1.e-4f)
 #define SIGN_BIT (0x80000000)
+#define FLT_MAX 3.402823466e+38F
+#define FLT_MIN 1.175494351e-38F
 
 static const uniform struct FVector4 QMULTI_SIGN_MASK0 = { {1.f, -1.f, 1.f, -1.f} };
 static const uniform struct FVector4 QMULTI_SIGN_MASK1 = { {1.f, 1.f, -1.f, -1.f} };
 static const uniform struct FVector4 QMULTI_SIGN_MASK2 = { {-1.f, 1.f, 1.f, -1.f} };
+static const uniform struct FVector4 QINV_SIGN_MASK = { {-1.f, -1.f, -1.f, 1.f} };
 static const uniform struct FVector4 FloatZero = { {0.f, 0.f, 0.f, 0.f} };
 static const uniform struct FVector4 FloatOne = { {1.f, 1.f, 1.f, 1.f} };
 static const uniform struct FVector4 FloatMinusOne = { {-1.f, -1.f, -1.f, -1.f} };
@@ -72,8 +75,6 @@
 static const uniform struct FVector4 Float180 = { {180.f, 180.f, 180.f, 180.f} };
 static const uniform struct FVector4 FloatNonFractional = { {8388608.f, 8388608.f, 8388608.f, 8388608.f} };
 
-static const uniform struct FVector4 SignBit = { {(unsigned int)SIGN_BIT, (unsigned int)SIGN_BIT, (unsigned int)SIGN_BIT, (unsigned int)SIGN_BIT} };
-
 static const uniform struct FVector4 DEG_TO_RAD ={ {PI/(180.f), PI/(180.f), PI/(180.f), PI/(180.f) } };
 static const uniform struct FVector4 DEG_TO_RAD_HALF = { {(PI/180.f)*0.5f, (PI/180.f)*0.5f, (PI/180.f)*0.5f, (PI/180.f)*0.5f} };
 static const uniform struct FVector4 RAD_TO_DEG = { {(180.f)/PI, (180.f)/PI, (180.f)/PI, (180.f)/PI} };
@@ -136,12 +137,8 @@
 
 inline uniform FVector4 operator+(const uniform FVector4 &A, const uniform float F)
 {
-	varying float S0, Result;
-	*((uniform FVector4 *uniform)&S0) = *((uniform FVector4 *uniform)&A);
-
-	Result = S0 + F;
-
-	return *((uniform FVector4 *uniform)&Result);
+	const uniform FVector4 FVec = {{F,F,F,F}};
+	return A + FVec;
 }
 
 inline FVector4 operator*(const FVector4 &A, const FVector4 &B)
@@ -205,12 +202,14 @@
 
 inline uniform FVector4 operator*(const uniform FVector4 &A, const uniform float F)
 {
-	varying float S0, Result;
-	*((uniform FVector4 *uniform)&S0) = *((uniform FVector4 *uniform)&A);
-
-	Result = S0 * F;
-
-	return *((uniform FVector4 *uniform)&Result);
+	const uniform FVector4 FVec = {{F,F,F,F}};
+	return A * FVec;
 }
 
+inline uniform FVector4 operator*(const uniform float F, const uniform FVector4 &A)
+{
+	const uniform FVector4 FVec = {{F,F,F,F}};
+	return FVec * A;
+}
+
 inline FVector4 operator/(const FVector4 &A, const FVector4 &B)
@@ -236,6 +235,12 @@
 	return *((uniform FVector4 *uniform)&Result);
 }
 
+inline uniform FVector4 operator/(const uniform FVector4 &A, const uniform float F)
+{
+	const uniform FVector4 FVec = {{F,F,F,F}};
+	return A / FVec;
+}
+
 inline FVector4 operator-(const FVector4 &A, const FVector4 &B)
 {
 	FVector4 Result;
@@ -297,12 +302,14 @@
 
 inline uniform FVector4 operator-(const uniform FVector4 &A, const uniform float F)
 {
-	varying float S0, Result;
-	*((uniform FVector4 *uniform)&S0) = *((uniform FVector4 *uniform)&A);
-
-	Result = S0 - F;
-
-	return *((uniform FVector4 *uniform)&Result);
+	const uniform FVector4 FVec = {{F,F,F,F}};
+	return A - FVec;
 }
 
+inline uniform FVector4 operator-(const uniform float F, const uniform FVector4 &A)
+{
+	const uniform FVector4 FVec = {{F,F,F,F}};
+	return FVec - A;
+}
+
 inline FVector4 SetVector4(const float X, const float Y, const float Z, const float W)
@@ -317,6 +324,18 @@
 	return Result;
 }
 
+inline uniform FVector4 SetVector4(const uniform FVector &V, const uniform float W)
+{
+	const uniform FVector4 Result = {{V.V[0], V.V[1], V.V[2], W}};
+	return Result;
+}
+
+inline uniform FVector4 SetVector4(const uniform float F)
+{
+	const uniform FVector4 Result = {{F, F, F, F}};
+	return Result;
+}
+
 inline FVector operator+(const FVector &A, const FVector &B)
 {
 	FVector Result;
@@ -330,13 +349,13 @@
 
 inline uniform FVector operator+(const uniform FVector &A, const uniform FVector &B)
 {
-	uniform FVector Result;
+	uniform FVector4 S0, S1, Result;
+	*((uniform FVector *uniform)&S0) = *((uniform FVector *uniform)&A);
+	*((uniform FVector *uniform)&S1) = *((uniform FVector *uniform)&B);
 
-	Result.V[0] = A.V[0] + B.V[0];
-	Result.V[1] = A.V[1] + B.V[1];
-	Result.V[2] = A.V[2] + B.V[2];
+	Result = S0 + S1;
 
-	return Result;
+	return *((uniform FVector *uniform)&Result);
 }
 
 inline FVector operator+(const FVector &A, const uniform FVector &B)
@@ -363,13 +382,8 @@
 
 inline uniform FVector operator+(const uniform FVector &A, const uniform float F)
 {
-	uniform FVector Result;
-
-	Result.V[0] = A.V[0] + F;
-	Result.V[1] = A.V[1] + F;
-	Result.V[2] = A.V[2] + F;
-
-	return Result;
+	const uniform FVector FVec = {{F,F,F}};
+	return A + FVec;
 }
 
 inline FVector operator*(const FVector &A, const FVector &B)
@@ -385,13 +399,13 @@
 
 inline uniform FVector operator*(const uniform FVector &A, const uniform FVector &B)
 {
-	uniform FVector Result;
+	uniform FVector4 S0, S1, Result;
+	*((uniform FVector *uniform)&S0) = *((uniform FVector *uniform)&A);
+	*((uniform FVector *uniform)&S1) = *((uniform FVector *uniform)&B);
 
-	Result.V[0] = A.V[0] * B.V[0];
-	Result.V[1] = A.V[1] * B.V[1];
-	Result.V[2] = A.V[2] * B.V[2];
+	Result = S0 * S1;
 
-	return Result;
+	return *((uniform FVector *uniform)&Result);
 }
 
 inline FVector operator*(const FVector &A, const uniform FVector &B)
@@ -429,13 +443,14 @@
 
 inline uniform FVector operator*(const uniform FVector &A, const uniform float F)
 {
-	uniform FVector Result;
-
-	Result.V[0] = A.V[0] * F;
-	Result.V[1] = A.V[1] * F;
-	Result.V[2] = A.V[2] * F;
-
-	return Result;
+	const uniform FVector FVec = {{F,F,F}};
+	return A * FVec;
 }
 
+inline uniform FVector operator*(const uniform float F, const uniform FVector &A)
+{
+	const uniform FVector FVec = {{F,F,F}};
+	return FVec * A;
+}
+
 inline FVector operator-(const FVector &A, const FVector &B)
@@ -451,13 +466,13 @@
 
 inline uniform FVector operator-(const uniform FVector &A, const uniform FVector &B)
 {
-	uniform FVector Result;
+	uniform FVector4 S0, S1, Result;
+	*((uniform FVector *uniform)&S0) = *((uniform FVector *uniform)&A);
+	*((uniform FVector *uniform)&S1) = *((uniform FVector *uniform)&B);
 
-	Result.V[0] = A.V[0] - B.V[0];
-	Result.V[1] = A.V[1] - B.V[1];
-	Result.V[2] = A.V[2] - B.V[2];
+	Result = S0 - S1;
 
-	return Result;
+	return *((uniform FVector *uniform)&Result);
 }
 
 inline FVector operator-(const FVector &A, const uniform FVector &B)
@@ -495,13 +510,14 @@
 
 inline uniform FVector operator-(const uniform FVector &A, const uniform float F)
 {
-	uniform FVector Result;
-
-	Result.V[0] = A.V[0] - F;
-	Result.V[1] = A.V[1] - F;
-	Result.V[2] = A.V[2] - F;
-
-	return Result;
+	const uniform FVector FVec = {{F,F,F}};
+	return A - FVec;
 }
 
+inline uniform FVector operator-(const uniform float F, const uniform FVector &A)
+{
+	const uniform FVector FVec = {{F,F,F}};
+	return FVec - A;
+}
+
 inline FVector operator-(const float F, const FVector &A)
@@ -528,13 +544,13 @@
 
 inline uniform FVector operator/(const uniform FVector &A, const uniform FVector &B)
 {
-	uniform FVector Result;
+	uniform FVector4 S0, S1, Result;
+	*((uniform FVector *uniform)&S0) = *((uniform FVector *uniform)&A);
+	*((uniform FVector *uniform)&S1) = *((uniform FVector *uniform)&B);
 
-	Result.V[0] = A.V[0] / B.V[0];
-	Result.V[1] = A.V[1] / B.V[1];
-	Result.V[2] = A.V[2] / B.V[2];
+	Result = S0 / S1;
 
-	return Result;
+	return *((uniform FVector *uniform)&Result);
 }
 
 inline FVector operator/(const FVector &A, const uniform FVector &B)
@@ -561,13 +577,8 @@
 
 inline uniform FVector operator/(const uniform FVector &A, const uniform float F)
 {
-	uniform FVector Result;
-
-	Result.V[0] = A.V[0] / F;
-	Result.V[1] = A.V[1] / F;
-	Result.V[2] = A.V[2] / F;
-
-	return Result;
+	const uniform FVector FVec = {{F,F,F}};
+	return A / FVec;
 }
 
 inline FVector SetVector(const float X, const float Y, const float Z)
@@ -588,17 +599,6 @@
 	return Result;
 }
 
-inline uniform FVector VectorAbs(const uniform FVector& A)
-{
-	uniform FVector Result;
-
-	Result.V[0] = abs(A.V[0]);
-	Result.V[1] = abs(A.V[1]);
-	Result.V[2] = abs(A.V[2]);
-
-	return Result;
-}
-
 inline FIntVector operator+(const FIntVector &A, const FIntVector &B)
 {
 	FIntVector Result;
@@ -808,70 +808,28 @@
 	return Result;
 }
 
-inline bool IsVectorLessEqual(const FVector &A, const uniform FVector &B)
-{
-	return A.V[0] <= B.V[0] && A.V[1] <= B.V[1] && A.V[2] <= B.V[2];
-}
-
-inline bool IsVectorGreaterEqual(const FVector &A, const uniform FVector &B)
-{
-	return A.V[0] >= B.V[0] && A.V[1] >= B.V[1] && A.V[2] >= B.V[2];
-}
-
-inline FVector VectorMax(const FVector& V1, const FVector& V2)
-{
-	FVector Result;
-
-	Result.V[0] = max(V1.V[0], V2.V[0]);
-	Result.V[1] = max(V1.V[1], V2.V[1]);
-	Result.V[2] = max(V1.V[2], V2.V[2]);
-
-	return Result;
-}
-
-inline uniform FVector VectorMax(const uniform FVector& V1, const uniform FVector& V2)
-{
-	uniform FVector Result;
-
-	Result.V[0] = max(V1.V[0], V2.V[0]);
-	Result.V[1] = max(V1.V[1], V2.V[1]);
-	Result.V[2] = max(V1.V[2], V2.V[2]);
-
-	return Result;
-}
-
-inline FVector VectorMin(const FVector& V1, const FVector& V2)
-{
-	FVector Result;
-
-	Result.V[0] = min(V1.V[0], V2.V[0]);
-	Result.V[1] = min(V1.V[1], V2.V[1]);
-	Result.V[2] = min(V1.V[2], V2.V[2]);
-
-	return Result;
-}
-
-inline uniform FVector VectorMin(const uniform FVector& V1, const uniform FVector& V2)
-{
-	uniform FVector Result;
-
-	Result.V[0] = min(V1.V[0], V2.V[0]);
-	Result.V[1] = min(V1.V[1], V2.V[1]);
-	Result.V[2] = min(V1.V[2], V2.V[2]);
-
-	return Result;
-}
-
-inline uniform bool VectorIsNearlyZero(const uniform FVector &V)
+inline uniform bool VectorIsAnyNearlyZero(const uniform FVector &V, const uniform float Tolerance)
 {
 	varying bool NearlyZero = false;
-	varying float S0;
+	varying float S0 = FLT_MAX;
 
 	*((uniform FVector *uniform)&S0) = *((uniform FVector *uniform)&V);
 
-	NearlyZero = abs(S0) <= KINDA_SMALL_NUMBER ? true : false;
+	NearlyZero = abs(S0) <= Tolerance ? true : false;
 
 	return any(NearlyZero);
 }
 
+inline uniform bool VectorIsAllNearlyZero(const uniform FVector &V, const uniform float Tolerance)
+{
+	varying bool NearlyZero = false;
+	varying float S0 = FLT_MIN;
+
+	*((uniform FVector *uniform)&S0) = *((uniform FVector *uniform)&V);
+
+	NearlyZero = abs(S0) <= Tolerance ? true : false;
+
+	return all(NearlyZero);
+}
+
 inline FVector VectorLerp(const FVector& A, const FVector& B, const float Alpha)
@@ -1031,11 +989,26 @@
 	return *((uniform FVector4 *uniform)&Result);
 }
 
+inline uniform FVector VectorAbs(const uniform FVector& A)
+{
+	uniform FVector4 S0, Result;
+	*((uniform FVector *uniform)&S0) = *((uniform FVector *uniform)&A);
+
+	Result = VectorAbs(S0);
+
+	return *((uniform FVector *uniform)&Result);
+}
+
 inline bool IsVectorLessEqual(const FVector4 &A, const uniform FVector4 &B)
 {
 	return A.V[0] <= B.V[0] && A.V[1] <= B.V[1] && A.V[2] <= B.V[2] && A.V[3] <= B.V[3];
 }
 
+inline bool IsVectorLessEqual(const FVector &A, const uniform FVector &B)
+{
+	return A.V[0] <= B.V[0] && A.V[1] <= B.V[1] && A.V[2] <= B.V[2];
+}
+
 inline uniform bool IsVectorGreaterEqual(const uniform FVector4 &A, const uniform FVector4 &B)
 {
 	varying bool b = true;
@@ -1048,6 +1021,11 @@
 	return all(b);
 }
 
+inline bool IsVectorGreaterEqual(const FVector &A, const uniform FVector &B)
+{
+	return A.V[0] >= B.V[0] && A.V[1] >= B.V[1] && A.V[2] >= B.V[2];
+}
+
 inline FVector4 VectorCompareGE(const FVector4 &A, const uniform FVector4 &B)
 {
 	FVector4 Result;
@@ -1083,6 +1061,11 @@
 	return *((uniform FVector4 *uniform)&Result);
 }
 
+inline uniform float CompareGE(const uniform float &A, const uniform float &B)
+{
+	return A >= B ? floatbits(0xFFFFFFFF) : floatbits(0);
+}
+
 inline uniform FVector4 VectorCompareGT(const uniform FVector4 &A, const uniform FVector4 &B)
 {
 	varying float S0, S1, Result;
@@ -1117,6 +1100,28 @@
 	return *((uniform FVector4 *uniform)&Result);
 }
 
+inline FVector VectorMax(const FVector& V1, const FVector& V2)
+{
+	FVector Result;
+
+	Result.V[0] = max(V1.V[0], V2.V[0]);
+	Result.V[1] = max(V1.V[1], V2.V[1]);
+	Result.V[2] = max(V1.V[2], V2.V[2]);
+
+	return Result;
+}
+
+inline uniform FVector VectorMax(const uniform FVector& A, const uniform FVector& B)
+{
+	uniform FVector4 S0, S1, Result;
+	*((uniform FVector *uniform)&S0) = *((uniform FVector *uniform)&A);
+	*((uniform FVector *uniform)&S1) = *((uniform FVector *uniform)&B);
+
+	Result = VectorMax(S0, S1);
+
+	return *((uniform FVector *uniform)&Result);
+}
+
 inline FVector4 VectorMin(const FVector4& V1, const FVector4& V2)
 {
 	FVector4 Result;
@@ -1140,6 +1145,28 @@
 	return *((uniform FVector4 *uniform)&Result);
 }
 
+inline FVector VectorMin(const FVector& V1, const FVector& V2)
+{
+	FVector Result;
+
+	Result.V[0] = min(V1.V[0], V2.V[0]);
+	Result.V[1] = min(V1.V[1], V2.V[1]);
+	Result.V[2] = min(V1.V[2], V2.V[2]);
+
+	return Result;
+}
+
+inline uniform FVector VectorMin(const uniform FVector& A, const uniform FVector& B)
+{
+	uniform FVector4 S0, S1, Result;
+	*((uniform FVector *uniform)&S0) = *((uniform FVector *uniform)&A);
+	*((uniform FVector *uniform)&S1) = *((uniform FVector *uniform)&B);
+
+	Result = VectorMin(S0, S1);
+
+	return *((uniform FVector *uniform)&Result);
+}
+
 inline FVector4 VectorReciprocal(const FVector4& V)
 {
 	FVector4 Result;
@@ -1380,6 +1407,38 @@
 	return x2;
 }
 
+inline uniform float ReciprocalSqrtAccurate(const uniform float V)
+{
+	// Perform two passes of Newton-Raphson iteration on the hardware estimate
+	//    v^-0.5 = x
+	// => x^2 = v^-1
+	// => 1/(x^2) = v
+	// => F(x) = x^-2 - v
+	//    F'(x) = -2x^-3
+
+	//    x1 = x0 - F(x0)/F'(x0)
+	// => x1 = x0 + 0.5 * (x0^-2 - Vec) * x0^3
+	// => x1 = x0 + 0.5 * (x0 - Vec * x0^3)
+	// => x1 = x0 + x0 * (0.5 - 0.5 * Vec * x0^2)
+
+	const uniform float DivBy2 = V * 0.5f;
+
+	// Initial estimate
+	const uniform float x0 = rsqrt_fast(V);
+
+	// First iteration
+	uniform float x1 = x0 * x0;
+	x1 = 0.5f - (DivBy2 * x1);
+	x1 = x0 * x1 + x0;
+
+	// Second iteration
+	uniform float x2 = x1 * x1;
+	x2 = 0.5f - (DivBy2 * x2);
+	x2 = x1 * x2 + x1;
+
+	return x2;
+}
+
 /**
  * Replicates one element into all four elements and returns the new vector.
  *
@@ -1417,6 +1476,11 @@
 	return SetVector4(Vec.V[X], Vec.V[Y], Vec.V[Z], Vec.V[W]);
 }
 
+inline uniform FVector VectorSwizzle(const uniform FVector &Vec, const uniform int X, const uniform int Y, const uniform int Z)
+{
+	return SetVector(Vec.V[X], Vec.V[Y], Vec.V[Z]);
+}
+
 /**
  * Creates a vector through selecting two components from each vector via a shuffle mask. 
  *
@@ -1454,6 +1518,15 @@
 	return (A_YZXW * B_ZXYW) - (A_ZXYW * B_YZXW);
 }
 
+inline uniform FVector VectorCross(const uniform FVector& Vec1, const uniform FVector& Vec2)
+{
+	const uniform FVector A_YZXW = VectorSwizzle(Vec1, 1,2,0);
+	const uniform FVector B_ZXYW = VectorSwizzle(Vec2, 2,0,1);
+	const uniform FVector A_ZXYW = VectorSwizzle(Vec1, 2,0,1);
+	const uniform FVector B_YZXW = VectorSwizzle(Vec2, 1,2,0);
+	return (A_YZXW * B_ZXYW) - (A_ZXYW * B_YZXW);
+}
+
 /**
  * Calculates the dot4 product of two vectors and returns a vector with the result in all 4 components.
  * Only really efficient on Xbox 360.
@@ -1482,6 +1555,22 @@
 	return Temp1 + Temp2; // (X*X + Z*Z + Y*Y + W*W, Y*Y + W*W + Z*Z + X*X, Z*Z + X*X + W*W + Y*Y, W*W + Y*Y + X*X + Z*Z)
 }
 
+inline uniform float VectorDot( const uniform FVector4& Vec1, const uniform FVector4& Vec2 )
+{
+	uniform FVector4 Temp1, Temp2;
+	Temp1 = Vec1 * Vec2;
+	Temp2 = VectorSwizzle(Temp1, 2,3,0,1);	// (Z,W,X,Y).
+	Temp1 = Temp1 + Temp2; // (X*X + Z*Z, Y*Y + W*W, Z*Z + X*X, W*W + Y*Y)
+	Temp2 = VectorSwizzle(Temp1, 1,2,3,0); // Rotate left 4 bytes (Y,Z,W,X).
+	return Temp1.V[0] + Temp2.V[0]; // (X*X + Z*Z + Y*Y + W*W)
+}
+
+inline uniform float VectorDot( const uniform FVector &Vec1, const uniform FVector &Vec2 )
+{
+	const uniform FVector Temp = Vec1 * Vec2;
+	return Temp.V[0] + Temp.V[1] + Temp.V[2];
+}
+
 /**
  * Does a bitwise vector selection based on a mask (e.g., created from VectorCompareXX)
  *
@@ -1565,6 +1654,16 @@
 	return VectorSelect(NonZeroMask, NormalizedVector, DefaultValue);
 }
 
+inline uniform FVector4 VectorNormalizeSafe2( const uniform FVector4& Vector, const uniform FVector4& DefaultValue)
+{
+	const uniform float SquareSum = VectorDot(Vector, Vector);
+	const uniform float NonZeroMask = CompareGE(SquareSum, SMALL_NUMBER);
+	const uniform float InvLength = ReciprocalSqrtAccurate(SquareSum);
+	const uniform FVector4 NormalizedVector = Vector * InvLength;
+	const uniform FVector4 VNonZeroMask = SetVector4(NonZeroMask, NonZeroMask, NonZeroMask, NonZeroMask);
+	return VectorSelect(VNonZeroMask, NormalizedVector, DefaultValue);
+}
+
 inline uniform FVector4 VectorTruncate(const uniform FVector4 &X)
 {
 	varying float S0, Result;
@@ -1586,6 +1685,16 @@
 	return *((uniform FVector4 *uniform)&Result);
 }
 
+inline uniform FVector4 VectorBitwiseAnd(const uniform FVector4 &Vec1, const uniform unsigned int I)
+{
+	varying float S0, Result;
+	*((uniform FVector4 *uniform)&S0) = *((uniform FVector4 *uniform)&Vec1);
+
+	Result = floatbits(intbits(S0) & I);
+
+	return *((uniform FVector4 *uniform)&Result);
+}
+
 inline uniform FVector4 VectorBitwiseOr(const uniform FVector4 &Vec1, const uniform FVector4 &Vec2)
 {
 	varying float S0, S1, Result;
@@ -1637,7 +1746,7 @@
 	uniform FVector4 X = VectorSubtract(VAngles, VectorMultiply(TwoPi, Quotient));
 
 	// Map in [-pi/2,pi/2]
-	uniform FVector4 sign = VectorBitwiseAnd(X, SignBit);
+	uniform FVector4 sign = VectorBitwiseAnd(X, SIGN_BIT);
 	const uniform FVector4 c = VectorBitwiseOr(Pi, sign);  // pi when x >= 0, -pi when x < 0
 	const uniform FVector4 absx = VectorAbs(X);
 	const uniform FVector4 rflx = VectorSubtract(c, X);
@@ -1682,4 +1791,24 @@
 	return VectorMultiplyAdd(Alpha, Delta, A);
 }
 
+inline uniform FVector4 VectorNegate(const uniform FVector4 &A)
+{
+	varying float S0, Result;
+	*((uniform FVector4 *uniform)&S0) = *((uniform FVector4 *uniform)&A);
+
+	Result = floatbits(intbits(S0) ^ SIGN_BIT);
+
+	return *((uniform FVector4 *uniform)&Result);
+}
+
+inline uniform FVector VectorNegate(const uniform FVector &A)
+{
+	varying float S0, Result;
+	*((uniform FVector *uniform)&S0) = *((uniform FVector *uniform)&A);
+
+	Result = floatbits(intbits(S0) ^ SIGN_BIT);
+
+	return *((uniform FVector *uniform)&Result);
+}
+
 #endif
diff --git a/Engine/Source/Runtime/Core/Public/Math/WideVector.isph b/Engine/Source/Runtime/Core/Public/Math/WideVector.isph
--- a/Engine/Source/Runtime/Core/Public/Math/WideVector.isph
+++ b/Engine/Source/Runtime/Core/Public/Math/WideVector.isph
@@ -28,6 +28,15 @@ unmasked inline uniform WideFVector4 operator+(const uniform WideFVector4 &A, co
 	return Result;
 }
 
+unmasked inline uniform WideFVector4 operator-(const uniform WideFVector4 &A, const uniform WideFVector4 &B)
+{
+	uniform WideFVector4 Result;
+
+	Result.V[programIndex] = A.V[programIndex] - B.V[programIndex];
+
+	return Result;
+}
+
 unmasked inline uniform WideFVector4 operator*(const uniform WideFVector4 &A, const uniform WideFVector4 &B)
 {
 	uniform WideFVector4 Result;
@@ -55,6 +64,24 @@ unmasked inline uniform WideFVector4 VectorSwizzle(const uniform WideFVector4 &V
 	return Result;
 }
 
+unmasked inline uniform WideFVector4 VectorReplicate(const uniform WideFVector4 &Vec, const uniform int R)
+{
+#if TARGET_WIDTH == 4
+	const varying int vPerm = { R, R, R, R };
+#elif TARGET_WIDTH == 8
+	const varying int vPerm = { R, R, R, R, R+4, R+4, R+4, R+4 };
+#elif TARGET_WIDTH == 16
+	const varying int vPerm = { R, R, R, R, R+4, R+4, R+4, R+4, R+8, R+8, R+8, R+8, R+12, R+12, R+12, R+12 };
+#endif
+
+	const float V = Vec.V[programIndex];
+	const float S = shuffle(V, vPerm);
+	uniform WideFVector4 Result;
+	Result.V[programIndex] = S;
+
+	return Result;
+}
+
 unmasked inline uniform WideFVector4 VectorCompareGE(const uniform WideFVector4 &A, const uniform WideFVector4 &B)
 {
 	uniform WideFVector4 Result;
@@ -106,6 +133,39 @@ unmasked inline uniform WideFVector4 VectorNormalizeQuaternion(const uniform Wid
 	return VectorNormalizeSafe(UnnormalizedQuat, WFloat0001);
 }
 
+unmasked inline uniform WideFVector4 VectorCross(const uniform WideFVector4& Vec1, const uniform WideFVector4& Vec2)
+{
+	const uniform WideFVector4 A_YZXW = VectorSwizzle(Vec1, 1,2,0,3);
+	const uniform WideFVector4 B_ZXYW = VectorSwizzle(Vec2, 2,0,1,3);
+	const uniform WideFVector4 A_ZXYW = VectorSwizzle(Vec1, 2,0,1,3);
+	const uniform WideFVector4 B_YZXW = VectorSwizzle(Vec2, 1,2,0,3);
+	return (A_YZXW * B_ZXYW) - (A_ZXYW * B_YZXW);
+}
+
+unmasked inline uniform WideFVector4 VectorQuaternionRotateVector(const uniform WideFVector4& Quat, const uniform WideFVector4& VectorW0)
+{
+	// Q * V * Q.Inverse
+	//const VectorRegister InverseRotation = VectorQuaternionInverse(Quat);
+	//const VectorRegister Temp = VectorQuaternionMultiply2(Quat, VectorW0);
+	//const VectorRegister Rotated = VectorQuaternionMultiply2(Temp, InverseRotation);
+
+	// Equivalence of above can be shown to be:
+	// http://people.csail.mit.edu/bkph/articles/Quaternions.pdf
+	// V' = V + 2w(Q x V) + (2Q x (Q x V))
+	// refactor:
+	// V' = V + w(2(Q x V)) + (Q x (2(Q x V)))
+	// T = 2(Q x V);
+	// V' = V + w*(T) + (Q x T)
+
+	const uniform WideFVector4 QW = VectorReplicate(Quat, 3);
+	uniform WideFVector4 T = VectorCross(Quat, VectorW0);
+	T = T + T;
+	const uniform WideFVector4 VTemp0 = QW * T + VectorW0;
+	const uniform WideFVector4 VTemp1 = VectorCross(Quat, T);
+	const uniform WideFVector4 Rotated = VTemp0 + VTemp1;
+	return Rotated;
+}
+
 unmasked inline void LoadStridedWideFVector4(uniform FVector4 * uniform DstPtr, const uniform FVector4* uniform SrcPtr, const uniform int SrcStride)
 {
 	*DstPtr = *SrcPtr;
diff --git a/Engine/Source/Runtime/Experimental/Chaos/Private/Chaos/Joint/PBDJointSolverGaussSeidel.cpp b/Engine/Source/Runtime/Experimental/Chaos/Private/Chaos/Joint/PBDJointSolverGaussSeidel.cpp
--- a/Engine/Source/Runtime/Experimental/Chaos/Private/Chaos/Joint/PBDJointSolverGaussSeidel.cpp
+++ b/Engine/Source/Runtime/Experimental/Chaos/Private/Chaos/Joint/PBDJointSolverGaussSeidel.cpp
@@ -5,6 +5,9 @@
 #include "Chaos/Utilities.h"
 #include "ChaosLog.h"
 #include "ChaosStats.h"
+#if INTEL_ISPC
+#include "PBDJointSolverGaussSeidel.ispc.generated.h"
+#endif
 
 //#pragma optimize("", off)
 
@@ -244,15 +247,28 @@ namespace Chaos
 
 	FJointSolverGaussSeidel::FJointSolverGaussSeidel()
 	{
+#if INTEL_ISPC
+		check(bChaos_Joint_Accumulate == false);
+		check(sizeof(FJointSolverGaussSeidel) == ispc::SizeofFJointSolverGaussSeidel());
+#endif
 	}
 
 
 	void FJointSolverGaussSeidel::UpdateDerivedState()
 	{
-		Xs[0] = Ps[0] + Qs[0] * XLs[0].GetTranslation();
-		Xs[1] = Ps[1] + Qs[1] * XLs[1].GetTranslation();
-		Rs[0] = Qs[0] * XLs[0].GetRotation();
-		Rs[1] = Qs[1] * XLs[1].GetRotation();
+		if (INTEL_ISPC)
+		{
+#if INTEL_ISPC
+			ispc::UpdateDerivedState((ispc::FJointSolverGaussSeidel*)this);
+#endif
+		}
+		else
+		{
+			Xs[0] = Ps[0] + Qs[0] * XLs[0].GetTranslation();
+			Xs[1] = Ps[1] + Qs[1] * XLs[1].GetTranslation();
+			Rs[0] = Qs[0] * XLs[0].GetRotation();
+			Rs[1] = Qs[1] * XLs[1].GetRotation();
+		}
 	}
 
 	void FJointSolverGaussSeidel::UpdateDerivedState(const int32 BodyIndex)
@@ -722,15 +738,24 @@ namespace Chaos
 		const FVec3& DP0,
 		const FVec3& DP1)
 	{
-		if (bChaos_Joint_Accumulate)
+		if (INTEL_ISPC)
 		{
-			DPs[0] += Stiffness * DP0;
-			DPs[1] += Stiffness * DP1;
+#if INTEL_ISPC
+			ispc::ApplyPositionDelta((ispc::FJointSolverGaussSeidel*)this, Stiffness, (ispc::FVector&)DP0, (ispc::FVector&)DP1);
+#endif
 		}
 		else
 		{
-			Ps[0] += Stiffness * DP0;
-			Ps[1] += Stiffness * DP1;
+			if (bChaos_Joint_Accumulate)
+			{
+				DPs[0] += Stiffness * DP0;
+				DPs[1] += Stiffness * DP1;
+			}
+			else
+			{
+				Ps[0] += Stiffness * DP0;
+				Ps[1] += Stiffness * DP1;
+			}
 		}
 	}
 
@@ -757,18 +782,27 @@ namespace Chaos
 		const FVec3& DR0,
 		const FVec3& DR1)
 	{
-		if (bChaos_Joint_Accumulate)
+		if (INTEL_ISPC)
 		{
-			DRs[0] += Stiffness * DR0;
-			DRs[1] += Stiffness * DR1;
+#if INTEL_ISPC
+			ispc::ApplyRotationDelta((ispc::FJointSolverGaussSeidel*)this, Stiffness, (ispc::FVector&)DR0, (ispc::FVector&)DR1);
+#endif
 		}
 		else
 		{
-			const FRotation3 DQ0 = (FRotation3::FromElements(Stiffness * DR0, 0) * Qs[0]) * (FReal)0.5;
-			const FRotation3 DQ1 = (FRotation3::FromElements(Stiffness * DR1, 0) * Qs[1]) * (FReal)0.5;
-			Qs[0] = (Qs[0] + DQ0).GetNormalized();
-			Qs[1] = (Qs[1] + DQ1).GetNormalized();
-			Qs[1].EnforceShortestArcWith(Qs[0]);
+			if (bChaos_Joint_Accumulate)
+			{
+				DRs[0] += Stiffness * DR0;
+				DRs[1] += Stiffness * DR1;
+			}
+			else
+			{
+				const FRotation3 DQ0 = (FRotation3::FromElements(Stiffness * DR0, 0) * Qs[0]) * (FReal)0.5;
+				const FRotation3 DQ1 = (FRotation3::FromElements(Stiffness * DR1, 0) * Qs[1]) * (FReal)0.5;
+				Qs[0] = (Qs[0] + DQ0).GetNormalized();
+				Qs[1] = (Qs[1] + DQ1).GetNormalized();
+				Qs[1].EnforceShortestArcWith(Qs[0]);
+			}
 		}
 	}
 
@@ -799,19 +833,28 @@ namespace Chaos
 		const FVec3& DV1,
 		const FVec3& DW1)
 	{
-		if (bChaos_Joint_Accumulate)
+		if (INTEL_ISPC)
 		{
-			DVs[0] += Stiffness * DV0;
-			DVs[1] += Stiffness * DV1;
-			DWs[0] += Stiffness * DW0;
-			DWs[1] += Stiffness * DW1;
+#if INTEL_ISPC
+			ispc::ApplyVelocityDelta((ispc::FJointSolverGaussSeidel*)this, Stiffness, (ispc::FVector&)DV0, (ispc::FVector&)DW0, (ispc::FVector&)DV1, (ispc::FVector&)DW1);
+#endif
 		}
 		else
 		{
-			Vs[0] += Stiffness * DV0;
-			Vs[1] += Stiffness * DV1;
-			Ws[0] += Stiffness * DW0;
-			Ws[1] += Stiffness * DW1;
+			if (bChaos_Joint_Accumulate)
+			{
+				DVs[0] += Stiffness * DV0;
+				DVs[1] += Stiffness * DV1;
+				DWs[0] += Stiffness * DW0;
+				DWs[1] += Stiffness * DW1;
+			}
+			else
+			{
+				Vs[0] += Stiffness * DV0;
+				Vs[1] += Stiffness * DV1;
+				Ws[0] += Stiffness * DW0;
+				Ws[1] += Stiffness * DW1;
+			}
 		}
 	}
 
@@ -901,41 +944,50 @@ namespace Chaos
 		const FVec3& Axis1,
 		const FReal Angle)
 	{
-		// World-space inverse mass
-		const FMatrix33 InvI0 = Utilities::ComputeWorldSpaceInertia(Qs[0], InvILs[0]);
-		const FMatrix33 InvI1 = Utilities::ComputeWorldSpaceInertia(Qs[1], InvILs[1]);
-		const FVec3 IA0 = Utilities::Multiply(InvI0, Axis0);
-		const FVec3 IA1 = Utilities::Multiply(InvI1, Axis1);
+		if (INTEL_ISPC)
+		{
+#if INTEL_ISPC
+			ispc::ApplyRotationConstraint((ispc::FJointSolverGaussSeidel*)this, Stiffness, (ispc::FVector&)Axis0, (ispc::FVector&)Axis1, Angle);
+#endif
+		}
+		else
+		{
+			// World-space inverse mass
+			const FMatrix33 InvI0 = Utilities::ComputeWorldSpaceInertia(Qs[0], InvILs[0]);
+			const FMatrix33 InvI1 = Utilities::ComputeWorldSpaceInertia(Qs[1], InvILs[1]);
+			const FVec3 IA0 = Utilities::Multiply(InvI0, Axis0);
+			const FVec3 IA1 = Utilities::Multiply(InvI1, Axis1);
 
-		// Joint-space inverse mass
-		const FReal II0 = FVec3::DotProduct(Axis0, IA0);
-		const FReal II1 = FVec3::DotProduct(Axis1, IA1);
+			// Joint-space inverse mass
+			const FReal II0 = FVec3::DotProduct(Axis0, IA0);
+			const FReal II1 = FVec3::DotProduct(Axis1, IA1);
 
-		const FVec3 DR0 = IA0 * (Angle / (II0 + II1));
-		const FVec3 DR1 = IA1 * -(Angle / (II0 + II1));
-		//const FVec3 DR0 = Axis0 * (Angle * II0 / (II0 + II1));
-		//const FVec3 DR1 = Axis1 * -(Angle * II1 / (II0 + II1));
+			const FVec3 DR0 = IA0 * (Angle / (II0 + II1));
+			const FVec3 DR1 = IA1 * -(Angle / (II0 + II1));
+			//const FVec3 DR0 = Axis0 * (Angle * II0 / (II0 + II1));
+			//const FVec3 DR1 = Axis1 * -(Angle * II1 / (II0 + II1));
 		
-		ApplyRotationDelta(Stiffness, DR0, DR1);
+			ApplyRotationDelta(Stiffness, DR0, DR1);
 
-		// Correct the positional error that was introduced by the rotation correction. This is
-		// correct when the position dofs are all locked, but not otherwise (fixable).
-		// This significantly improves angular stiffness at lower iterations, although the same effect
-		// is achieved by increasing iterations.
-		// @todo(ccaulfield): this position correction needs to have components in direction of inactive position constraints removed
-		if (AngularPositionCorrection > 0)
-		{
-			const FVec3 PrevX0 = Xs[0];
-			const FVec3 PrevX1 = Xs[1];
-			UpdateDerivedState();
+			// Correct the positional error that was introduced by the rotation correction. This is
+			// correct when the position dofs are all locked, but not otherwise (fixable).
+			// This significantly improves angular stiffness at lower iterations, although the same effect
+			// is achieved by increasing iterations.
+			// @todo(ccaulfield): this position correction needs to have components in direction of inactive position constraints removed
+			if (AngularPositionCorrection > 0)
+			{
+				const FVec3 PrevX0 = Xs[0];
+				const FVec3 PrevX1 = Xs[1];
+				UpdateDerivedState();
+
+				const FVec3 DX = (Xs[1] - Xs[0]) - (PrevX1 - PrevX0);
+				const FVec3 DP0 = DX * (InvMs[0] / (InvMs[0] + InvMs[1]));
+				const FVec3 DP1 = DX * (-InvMs[1] / (InvMs[0] + InvMs[1]));
+				ApplyPositionDelta(Stiffness * AngularPositionCorrection, DP0, DP1);
+			}
 
-			const FVec3 DX = (Xs[1] - Xs[0]) - (PrevX1 - PrevX0);
-			const FVec3 DP0 = DX * (InvMs[0] / (InvMs[0] + InvMs[1]));
-			const FVec3 DP1 = DX * (-InvMs[1] / (InvMs[0] + InvMs[1]));
-			ApplyPositionDelta(Stiffness * AngularPositionCorrection, DP0, DP1);
+			UpdateDerivedState();
 		}
-
-		UpdateDerivedState();
 	}
 
 
@@ -950,58 +1002,67 @@ namespace Chaos
 		const FReal Angle,
 		FReal& Lambda)
 	{
-		// World-space inverse mass
-		const FMatrix33 InvI0 = Utilities::ComputeWorldSpaceInertia(Qs[0], InvILs[0]);
-		const FMatrix33 InvI1 = Utilities::ComputeWorldSpaceInertia(Qs[1], InvILs[1]);
-		const FVec3 IA0 = Utilities::Multiply(InvI0, Axis0);
-		const FVec3 IA1 = Utilities::Multiply(InvI1, Axis1);
-
-		// Joint-space inverse mass
-		const FReal II0 = FVec3::DotProduct(Axis0, IA0);
-		const FReal II1 = FVec3::DotProduct(Axis1, IA1);
-		const FReal II = (II0 + II1);
-
-		// Damping angular velocity
-		FReal AngVelDt = 0;
-		if (Damping > KINDA_SMALL_NUMBER)
+		if (INTEL_ISPC)
 		{
-			const FVec3 W0 = FRotation3::CalculateAngularVelocity(PrevQs[0], Qs[0], 1.0f);
-			const FVec3 W1 = FRotation3::CalculateAngularVelocity(PrevQs[1], Qs[1], 1.0f);
-			AngVelDt = FVec3::DotProduct(Axis0, W0) - FVec3::DotProduct(Axis1, W1);
+#if INTEL_ISPC
+			ispc::ApplyRotationConstraintSoft((ispc::FJointSolverGaussSeidel*)this, Dt, Stiffness, Damping, bAccelerationMode, (ispc::FVector&) Axis0, (ispc::FVector&)Axis1, Angle, Lambda);
+#endif
 		}
+		else
+		{
+			// World-space inverse mass
+			const FMatrix33 InvI0 = Utilities::ComputeWorldSpaceInertia(Qs[0], InvILs[0]);
+			const FMatrix33 InvI1 = Utilities::ComputeWorldSpaceInertia(Qs[1], InvILs[1]);
+			const FVec3 IA0 = Utilities::Multiply(InvI0, Axis0);
+			const FVec3 IA1 = Utilities::Multiply(InvI1, Axis1);
 
-		const FReal SpringMassScale = (bAccelerationMode) ? 1.0f / II : 1.0f;
-		const FReal S = SpringMassScale * Stiffness * Dt * Dt;
-		const FReal D = SpringMassScale * Damping * Dt;
-		const FReal Multiplier = (FReal)1 / ((S + D) * II + (FReal)1);
-		const FReal DLambda = Multiplier * (S * Angle - D * AngVelDt - Lambda);
+			// Joint-space inverse mass
+			const FReal II0 = FVec3::DotProduct(Axis0, IA0);
+			const FReal II1 = FVec3::DotProduct(Axis1, IA1);
+			const FReal II = (II0 + II1);
 
-		//const FVec3 DR0 = IA0 * DLambda;
-		//const FVec3 DR1 = IA1 * -DLambda;
-		const FVec3 DR0 = Axis0 * (DLambda * II0);
-		const FVec3 DR1 = Axis1 * -(DLambda * II1);
+			// Damping angular velocity
+			FReal AngVelDt = 0;
+			if (Damping > KINDA_SMALL_NUMBER)
+			{
+				const FVec3 W0 = FRotation3::CalculateAngularVelocity(PrevQs[0], Qs[0], 1.0f);
+				const FVec3 W1 = FRotation3::CalculateAngularVelocity(PrevQs[1], Qs[1], 1.0f);
+				AngVelDt = FVec3::DotProduct(Axis0, W0) - FVec3::DotProduct(Axis1, W1);
+			}
 
-		Lambda += DLambda;
-		ApplyRotationDelta(1.0f, DR0, DR1);
+			const FReal SpringMassScale = (bAccelerationMode) ? 1.0f / II : 1.0f;
+			const FReal S = SpringMassScale * Stiffness * Dt * Dt;
+			const FReal D = SpringMassScale * Damping * Dt;
+			const FReal Multiplier = (FReal)1 / ((S + D) * II + (FReal)1);
+			const FReal DLambda = Multiplier * (S * Angle - D * AngVelDt - Lambda);
+
+			//const FVec3 DR0 = IA0 * DLambda;
+			//const FVec3 DR1 = IA1 * -DLambda;
+			const FVec3 DR0 = Axis0 * (DLambda * II0);
+			const FVec3 DR1 = Axis1 * -(DLambda * II1);
+
+			Lambda += DLambda;
+			ApplyRotationDelta(1.0f, DR0, DR1);
+
+			// Correct the positional error that was introduced by the rotation correction. This is
+			// correct when the position dofs are all locked, but not otherwise (fixable).
+			// This significantly improves angular stiffness at lower iterations, although the same effect
+			// is achieved by increasing iterations.
+			// @todo(ccaulfield): this position correction needs to have components in direction of inactive position constraints removed
+			if (AngularPositionCorrection > 0)
+			{
+				const FVec3 PrevX0 = Xs[0];
+				const FVec3 PrevX1 = Xs[1];
+				UpdateDerivedState();
+
+				const FVec3 DX = (Xs[1] - Xs[0]) - (PrevX1 - PrevX0);
+				const FVec3 DP0 = DX * (InvMs[0] / (InvMs[0] + InvMs[1]));
+				const FVec3 DP1 = DX * (-InvMs[1] / (InvMs[0] + InvMs[1]));
+				ApplyPositionDelta(AngularPositionCorrection, DP0, DP1);
+			}
 
-		// Correct the positional error that was introduced by the rotation correction. This is
-		// correct when the position dofs are all locked, but not otherwise (fixable).
-		// This significantly improves angular stiffness at lower iterations, although the same effect
-		// is achieved by increasing iterations.
-		// @todo(ccaulfield): this position correction needs to have components in direction of inactive position constraints removed
-		if (AngularPositionCorrection > 0)
-		{
-			const FVec3 PrevX0 = Xs[0];
-			const FVec3 PrevX1 = Xs[1];
 			UpdateDerivedState();
-
-			const FVec3 DX = (Xs[1] - Xs[0]) - (PrevX1 - PrevX0);
-			const FVec3 DP0 = DX * (InvMs[0] / (InvMs[0] + InvMs[1]));
-			const FVec3 DP1 = DX * (-InvMs[1] / (InvMs[0] + InvMs[1]));
-			ApplyPositionDelta(AngularPositionCorrection, DP0, DP1);
 		}
-
-		UpdateDerivedState();
 	}
 
 
@@ -1403,33 +1464,43 @@ namespace Chaos
 		const FVec3 CX = Xs[1] - Xs[0];
 		if (CX != FVec3(0))
 		{
-			const FMatrix33 InvI0 = Utilities::ComputeWorldSpaceInertia(Qs[0], InvILs[0]);
-			const FMatrix33 InvI1 = Utilities::ComputeWorldSpaceInertia(Qs[1], InvILs[1]);
-
-			// Calculate constraint correction
-			FMatrix33 M0 = FMatrix33(0, 0, 0);
-			FMatrix33 M1 = FMatrix33(0, 0, 0);
-			if (InvMs[0] > 0)
+			if (INTEL_ISPC)
 			{
-				M0 = Utilities::ComputeJointFactorMatrix(Xs[0] - Ps[0], InvI0, InvMs[0]);
+#if INTEL_ISPC
+				FReal LinearStiffness = FPBDJointUtilities::GetLinearStiffness(SolverSettings, JointSettings);
+				ispc::ApplyPositionConstraint((ispc::FJointSolverGaussSeidel*)this, (ispc::FVector&)CX, LinearStiffness);
+#endif
 			}
-			if (InvMs[1] > 0)
+			else
 			{
-				M1 = Utilities::ComputeJointFactorMatrix(Xs[1] - Ps[1], InvI1, InvMs[1]);
+				const FMatrix33 InvI0 = Utilities::ComputeWorldSpaceInertia(Qs[0], InvILs[0]);
+				const FMatrix33 InvI1 = Utilities::ComputeWorldSpaceInertia(Qs[1], InvILs[1]);
+
+				// Calculate constraint correction
+				FMatrix33 M0 = FMatrix33(0, 0, 0);
+				FMatrix33 M1 = FMatrix33(0, 0, 0);
+				if (InvMs[0] > 0)
+				{
+					M0 = Utilities::ComputeJointFactorMatrix(Xs[0] - Ps[0], InvI0, InvMs[0]);
+				}
+				if (InvMs[1] > 0)
+				{
+					M1 = Utilities::ComputeJointFactorMatrix(Xs[1] - Ps[1], InvI1, InvMs[1]);
+				}
+				const FMatrix33 MI = (M0 + M1).Inverse();
+				const FVec3 DX = Utilities::Multiply(MI, CX);
+
+				// Apply constraint correction
+				const FVec3 DP0 = InvMs[0] * DX;
+				const FVec3 DP1 = -InvMs[1] * DX;
+				const FVec3 DR0 = Utilities::Multiply(InvI0, FVec3::CrossProduct(Xs[0] - Ps[0], DX));
+				const FVec3 DR1 = Utilities::Multiply(InvI1, FVec3::CrossProduct(Xs[1] - Ps[1], -DX));
+
+				FReal LinearStiffness = FPBDJointUtilities::GetLinearStiffness(SolverSettings, JointSettings);
+				ApplyPositionDelta(LinearStiffness, DP0, DP1);
+				ApplyRotationDelta(LinearStiffness, DR0, DR1);
+				UpdateDerivedState();
 			}
-			const FMatrix33 MI = (M0 + M1).Inverse();
-			const FVec3 DX = Utilities::Multiply(MI, CX);
-
-			// Apply constraint correction
-			const FVec3 DP0 = InvMs[0] * DX;
-			const FVec3 DP1 = -InvMs[1] * DX;
-			const FVec3 DR0 = Utilities::Multiply(InvI0, FVec3::CrossProduct(Xs[0] - Ps[0], DX));
-			const FVec3 DR1 = Utilities::Multiply(InvI1, FVec3::CrossProduct(Xs[1] - Ps[1], -DX));
-
-			FReal LinearStiffness = FPBDJointUtilities::GetLinearStiffness(SolverSettings, JointSettings);
-			ApplyPositionDelta(LinearStiffness, DP0, DP1);
-			ApplyRotationDelta(LinearStiffness, DR0, DR1);
-			UpdateDerivedState();
 		}
 	}
 
@@ -1632,40 +1703,49 @@ namespace Chaos
 		{
 			FReal LinearProjection = FPBDJointUtilities::GetLinearProjection(SolverSettings, JointSettings);
 			const FReal ParentMassScale = FMath::Max(0.0f, 1.0f - LinearProjection);
+			if (INTEL_ISPC)
+			{
+#if INTEL_ISPC
+				const FReal Stiffness = FPBDJointUtilities::GetLinearStiffness(SolverSettings, JointSettings);
+				ispc::ApplyPositionProjection((ispc::FJointSolverGaussSeidel*)this, (ispc::FVector&)CX, CXLen, ParentMassScale, Stiffness);
+#endif
+			}
+			else
+			{
+				const FVec3 CXDir = CX / CXLen;
+				const FVec3 V0 = Vs[0] + FVec3::CrossProduct(Ws[0], Xs[0] - Ps[0]);
+				const FVec3 V1 = Vs[1] + FVec3::CrossProduct(Ws[1], Xs[1] - Ps[1]);
+				FVec3 CV = FVec3::DotProduct(V1 - V0, CXDir) * CXDir;
+
+				const FReal IM0 = ParentMassScale * InvMs[0];
+				const FReal IM1 = InvMs[1];
+				const FVec3 IIL0 = ParentMassScale * InvILs[0];
+				const FVec3& IIL1 = InvILs[1];
+				FMatrix33 II0 = Utilities::ComputeWorldSpaceInertia(Qs[0], IIL0);
+				FMatrix33 II1 = Utilities::ComputeWorldSpaceInertia(Qs[1], IIL1);
+				FMatrix33 J0 = (IM0 > 0) ? Utilities::ComputeJointFactorMatrix(Xs[0] - Ps[0], II0, IM0) : FMatrix33(0, 0, 0);
+				FMatrix33 J1 = Utilities::ComputeJointFactorMatrix(Xs[1] - Ps[1], II1, IM1);
+				FMatrix33 IJ = (J0 + J1).Inverse();
+
+				const FVec3 DX = Utilities::Multiply(IJ, CX);
+				const FVec3 DV = Utilities::Multiply(IJ, CV);
+
+				const FVec3 DP0 = IM0 * DX;
+				const FVec3 DP1 = -IM1 * DX;
+				const FVec3 DR0 = Utilities::Multiply(II0, FVec3::CrossProduct(Xs[0] - Ps[0], DX));
+				const FVec3 DR1 = Utilities::Multiply(II1, FVec3::CrossProduct(Xs[1] - Ps[1], -DX));
+
+				const FVec3 DV0 = IM0 * DV;
+				const FVec3 DV1 = -IM1 * DV;
+				const FVec3 DW0 = Utilities::Multiply(II0, FVec3::CrossProduct(Xs[0] - Ps[0], DV));
+				const FVec3 DW1 = Utilities::Multiply(II1, FVec3::CrossProduct(Xs[1] - Ps[1], -DV));
 
-			const FVec3 CXDir = CX / CXLen;
-			const FVec3 V0 = Vs[0] + FVec3::CrossProduct(Ws[0], Xs[0] - Ps[0]);
-			const FVec3 V1 = Vs[1] + FVec3::CrossProduct(Ws[1], Xs[1] - Ps[1]);
-			FVec3 CV = FVec3::DotProduct(V1 - V0, CXDir) * CXDir;
-
-			const FReal IM0 = ParentMassScale * InvMs[0];
-			const FReal IM1 = InvMs[1];
-			const FVec3 IIL0 = ParentMassScale * InvILs[0];
-			const FVec3& IIL1 = InvILs[1];
-			FMatrix33 II0 = Utilities::ComputeWorldSpaceInertia(Qs[0], IIL0);
-			FMatrix33 II1 = Utilities::ComputeWorldSpaceInertia(Qs[1], IIL1);
-			FMatrix33 J0 = (IM0 > 0) ? Utilities::ComputeJointFactorMatrix(Xs[0] - Ps[0], II0, IM0) : FMatrix33(0, 0, 0);
-			FMatrix33 J1 = Utilities::ComputeJointFactorMatrix(Xs[1] - Ps[1], II1, IM1);
-			FMatrix33 IJ = (J0 + J1).Inverse();
-
-			const FVec3 DX = Utilities::Multiply(IJ, CX);
-			const FVec3 DV = Utilities::Multiply(IJ, CV);
-
-			const FVec3 DP0 = IM0 * DX;
-			const FVec3 DP1 = -IM1 * DX;
-			const FVec3 DR0 = Utilities::Multiply(II0, FVec3::CrossProduct(Xs[0] - Ps[0], DX));
-			const FVec3 DR1 = Utilities::Multiply(II1, FVec3::CrossProduct(Xs[1] - Ps[1], -DX));
-
-			const FVec3 DV0 = IM0 * DV;
-			const FVec3 DV1 = -IM1 * DV;
-			const FVec3 DW0 = Utilities::Multiply(II0, FVec3::CrossProduct(Xs[0] - Ps[0], DV));
-			const FVec3 DW1 = Utilities::Multiply(II1, FVec3::CrossProduct(Xs[1] - Ps[1], -DV));
-
-			const FReal Stiffness = FPBDJointUtilities::GetLinearStiffness(SolverSettings, JointSettings);
-			ApplyPositionDelta(Stiffness, DP0, DP1);
-			ApplyRotationDelta(Stiffness, DR0, DR1);
-			ApplyVelocityDelta(Stiffness, DV0, DW0, DV1, DW1);
-			UpdateDerivedState();
+				const FReal Stiffness = FPBDJointUtilities::GetLinearStiffness(SolverSettings, JointSettings);
+				ApplyPositionDelta(Stiffness, DP0, DP1);
+				ApplyRotationDelta(Stiffness, DR0, DR1);
+				ApplyVelocityDelta(Stiffness, DV0, DW0, DV1, DW1);
+				UpdateDerivedState();
+			}
 		}
 	}
 }
diff --git a/Engine/Source/Runtime/Experimental/Chaos/Private/Chaos/Joint/PBDJointSolverGaussSeidel.ispc b/Engine/Source/Runtime/Experimental/Chaos/Private/Chaos/Joint/PBDJointSolverGaussSeidel.ispc
new file mode 100644
--- /dev/null
+++ b/Engine/Source/Runtime/Experimental/Chaos/Private/Chaos/Joint/PBDJointSolverGaussSeidel.ispc
@@ -0,0 +1,515 @@
+// Copyright 1998-2019 Epic Games, Inc. All Rights Reserved.
+
+#define EXPLICIT_VECTOR4 1
+
+#include "Math/Vector.isph"
+#include "Math/Quat.isph"
+
+struct FRigidTransform
+{
+	FVector4 Rotation;
+	FVector Translation;
+	FVector Scale3D;
+};
+
+static const uniform struct FMatrix Matrix33Zero = {{ 0, 0, 0, 0, 0, 0, 0, 0, 0, 0, 0, 0, 0, 0, 0, 1}};
+
+static const uniform int32 MaxConstrainedBodies = 2;
+
+struct FJointSolverGaussSeidel
+{
+	// Local-space constraint settings
+	FRigidTransform XLs[MaxConstrainedBodies];	// Local-space joint connector transforms
+	FVector InvILs[MaxConstrainedBodies];		// Local-space inverse inertias
+	float InvMs[MaxConstrainedBodies];			// Inverse masses
+
+	// World-space constraint state
+	FVector Xs[MaxConstrainedBodies];			// World-space joint connector positions
+	FVector4 Rs[MaxConstrainedBodies];			// World-space joint connector rotations
+
+	// World-space body state
+	FVector Ps[MaxConstrainedBodies];			// World-space particle CoM positions
+	FVector4 Qs[MaxConstrainedBodies];			// World-space particle CoM rotations
+	FVector Vs[MaxConstrainedBodies];			// World-space particle CoM velocities
+	FVector Ws[MaxConstrainedBodies];			// World-space particle CoM angular velocities
+
+	// Accumulated world-body state deltas
+	FVector DPs[MaxConstrainedBodies];			// World-space particle CoM positions
+	FVector DRs[MaxConstrainedBodies];			// World-space particle CoM rotations (in axis-angle form)
+	FVector DVs[MaxConstrainedBodies];			// World-space particle CoM velocities
+	FVector DWs[MaxConstrainedBodies];			// World-space particle CoM angular velocities
+
+	// XPBD Previous iteration world-space body state
+	FVector PrevPs[MaxConstrainedBodies];		// World-space particle CoM positions
+	FVector4 PrevQs[MaxConstrainedBodies];		// World-space particle CoM rotations
+	FVector PrevXs[MaxConstrainedBodies];		// World-space joint connector positions
+
+	// XPBD constraint multipliers (net applied constraint-space deltas)
+	float LinearSoftLambda;
+	float LinearDriveLambda;
+	float TwistSoftLambda;
+	float SwingSoftLambda;
+	float TwistDriveLambda;
+	float SwingDriveLambda;
+
+	// Post-angular constraint position fixup (to reduce iterations required for stiff angular constraints)
+	float AngularPositionCorrection;
+};
+
+static inline uniform FMatrix SetMatrix33(
+	const uniform float x00, const uniform float x10, const uniform float x20, 
+	const uniform float x01, const uniform float x11, const uniform float x21, 
+	const uniform float x02, const uniform float x12, const uniform float x22)
+{
+	return SetMatrix(SetVector4(x00, x01, x02, 0), SetVector4(x10, x11, x12, 0), SetVector4(x20, x21, x22, 0), Float0001);
+}
+
+static inline uniform FMatrix SetMatrix33(
+	const uniform float x00, const uniform float x10, const uniform float x20, 
+	const uniform float x11, const uniform float x21, const uniform float x22)
+{
+	return SetMatrix(SetVector4(x00, x10, x20, 0), SetVector4(x10, x11, x21, 0), SetVector4(x20, x21, x22, 0), Float0001);
+}
+
+static inline uniform FMatrix SetMatrix33(
+	const uniform float x00, const uniform float x11, const uniform float x22)
+{
+	return SetMatrix(SetVector4(x00, 0, 0, 0), SetVector4(0, x11, 0, 0), SetVector4(0, 0, x22, 0), Float0001);
+}
+
+static inline uniform FMatrix SetMatrix33(const uniform FVector4& x0, const uniform FVector4& x1, const uniform FVector4& x2)
+{
+	return SetMatrix(x0, x1, x2, Float0001);
+}
+
+static inline uniform FMatrix MultiplyAB(const uniform FMatrix& L, const uniform FMatrix& R)
+{
+	const uniform FVector4 L0 = *((uniform FVector4 *uniform)&L.M[0]);
+	const uniform FVector4 L1 = *((uniform FVector4 *uniform)&L.M[4]);
+	const uniform FVector4 L2 = *((uniform FVector4 *uniform)&L.M[8]);
+
+	const uniform FVector4 x0 = SetVector4(R.M[0]) * L0 + SetVector4(R.M[1]) * L1 + SetVector4(R.M[2]) * L2;
+	const uniform FVector4 x1 = SetVector4(R.M[4]) * L0 + SetVector4(R.M[5]) * L1 + SetVector4(R.M[6]) * L2;
+	const uniform FVector4 x2 = SetVector4(R.M[8]) * L0 + SetVector4(R.M[9]) * L1 + SetVector4(R.M[10]) * L2;
+
+	return SetMatrix33(x0, x1, x2);
+}
+
+static inline uniform FMatrix MultiplyABt(const uniform FMatrix& L, const uniform FMatrix& R)
+{
+	const uniform FVector4 L0 = *((uniform FVector4 *uniform)&L.M[0]);
+	const uniform FVector4 L1 = *((uniform FVector4 *uniform)&L.M[4]);
+	const uniform FVector4 L2 = *((uniform FVector4 *uniform)&L.M[8]);
+
+	const uniform FVector4 x0 = SetVector4(R.M[0]) * L0 + SetVector4(R.M[4]) * L1 + SetVector4(R.M[8]) * L2;
+	const uniform FVector4 x1 = SetVector4(R.M[1]) * L0 + SetVector4(R.M[5]) * L1 + SetVector4(R.M[9]) * L2;
+	const uniform FVector4 x2 = SetVector4(R.M[2]) * L0 + SetVector4(R.M[6]) * L1 + SetVector4(R.M[10]) * L2;
+
+	return SetMatrix33(x0, x1, x2);
+}
+
+static inline uniform FVector Multiply(const uniform FMatrix& LIn, const uniform FVector& R)
+{
+	const uniform FVector4 L0 = *((uniform FVector4 *uniform)&LIn.M[0]);
+	const uniform FVector4 L1 = *((uniform FVector4 *uniform)&LIn.M[4]);
+	const uniform FVector4 L2 = *((uniform FVector4 *uniform)&LIn.M[8]);
+
+	const uniform FVector4 Result = L0 * SetVector4(R.V[0]) + L1 * SetVector4(R.V[1]) + L2 * SetVector4(R.V[2]);
+
+	return SetVector(Result.V[0], Result.V[1], Result.V[2]);
+}
+
+static inline uniform FMatrix AddAB(const uniform FMatrix& L, const uniform FMatrix& R)
+{
+	const uniform FVector4 *uniform A	= (const uniform FVector4 *uniform) &L;
+	const uniform FVector4 *uniform B	= (const uniform FVector4 *uniform) &R;
+
+	return SetMatrix33(A[0] + B[0], A[1] + B[1], A[2] + B[2]);
+}
+
+export uniform int32 SizeofFJointSolverGaussSeidel()
+{
+	return sizeof(uniform FJointSolverGaussSeidel);
+}
+
+static inline void UpdateDerivedStateInner(uniform FJointSolverGaussSeidel * uniform M)
+{
+	M->Xs[0] = M->Ps[0] + VectorQuaternionRotateVector(M->Qs[0], M->XLs[0].Translation);
+	M->Xs[1] = M->Ps[1] + VectorQuaternionRotateVector(M->Qs[1], M->XLs[1].Translation);
+	M->Rs[0] = VectorQuaternionMultiply2(M->Qs[0], M->XLs[0].Rotation);
+	M->Rs[1] = VectorQuaternionMultiply2(M->Qs[1], M->XLs[1].Rotation);
+}
+
+export void UpdateDerivedState(uniform FJointSolverGaussSeidel * uniform M)
+{
+	UpdateDerivedStateInner(M);
+}
+
+static inline void ApplyPositionDeltaInner(
+	uniform FJointSolverGaussSeidel * uniform M,
+	const uniform float Stiffness,
+	const uniform FVector& DP0,
+	const uniform FVector& DP1)
+{
+	M->Ps[0] = M->Ps[0] + (Stiffness * DP0);
+	M->Ps[1] = M->Ps[1] + (Stiffness * DP1);
+}
+
+export void ApplyPositionDelta(
+	uniform FJointSolverGaussSeidel * uniform M,
+	const uniform float Stiffness,
+	const uniform FVector& DP0,
+	const uniform FVector& DP1)
+{
+	ApplyPositionDeltaInner(M, Stiffness, DP0, DP1);
+}
+
+static inline void ApplyRotationDeltaInner(
+	uniform FJointSolverGaussSeidel * uniform M,
+	const uniform float Stiffness,
+	const uniform FVector& DR0,
+	const uniform FVector& DR1)
+{
+	const uniform FVector4 Qs0 = M->Qs[0];
+	const uniform FVector4 Qs1 = M->Qs[1];
+	const uniform FVector4 DR0Stiffness = SetVector4(DR0 * Stiffness, 0);
+	const uniform FVector4 DR1Stiffness = SetVector4(DR1 * Stiffness, 0);
+
+	const uniform FVector4 DQ0 = VectorQuaternionMultiply2(DR0Stiffness, Qs0) * 0.5f;
+	const uniform FVector4 DQ1 = VectorQuaternionMultiply2(DR1Stiffness, Qs1) * 0.5f;
+
+	const uniform FVector4 Qs0Norm = VectorNormalizeQuaternion(Qs0 + DQ0);
+	const uniform FVector4 Qs1Norm = VectorNormalizeQuaternion(Qs1 + DQ1);
+
+	const uniform FVector4 Qs1ShortestArcWith = VectorQuaternionEnforceShortestArcWith(Qs1Norm, Qs0Norm);
+
+	M->Qs[0] = Qs0Norm;
+	M->Qs[1] = Qs1ShortestArcWith;
+}
+
+export void ApplyRotationDelta(
+	uniform FJointSolverGaussSeidel * uniform M,
+	const uniform float Stiffness,
+	const uniform FVector& DR0,
+	const uniform FVector& DR1)
+{
+	ApplyRotationDeltaInner(M, Stiffness, DR0, DR1);
+}
+
+static inline uniform FMatrix MakeQuatRotationTranslationMatrix(const uniform FVector4& Q)
+{
+	// Assumes origin is 0,0,0
+	const uniform FVector4 xyz2 = Q + Q; // x2, y2, z2
+	const uniform FVector4 x = xyz2 * Q.V[0]; // xx, xy, xz
+	const uniform FVector4 y = xyz2 * Q.V[1]; // yx, yy, yz
+	const uniform FVector4 w = xyz2 * Q.V[3]; // wx, wy, wz
+	const uniform float zz = Q.V[2] * xyz2.V[2];
+
+	const uniform FVector4 yx = SetVector4(y.V[2], x.V[2], x.V[1], 0); // yz, xz, xy
+	const uniform FVector4 Adds = yx + w; // M[6], M[8], M[1]
+	const uniform FVector4 Minuses = yx - w; // M[9], M[2], M[4]
+	
+	const uniform FVector4 OneMinus0 = SetVector4(y.V[1], x.V[0], x.V[0], 0); // yy, xx, xx
+	const uniform FVector4 Oneminus1 = SetVector4(zz, zz, y.V[1], 0); // zz, zz, yy
+	const uniform FVector4 OneMinuses = FloatOne - (OneMinus0 + Oneminus1); // M[0], M[5], M[10]
+	
+	return SetMatrix33(
+		SetVector4(OneMinuses.V[0], Adds.V[2], Minuses.V[1], 0),
+		SetVector4(Minuses.V[2], OneMinuses.V[1], Adds.V[0], 0),
+		SetVector4(Adds.V[1], Minuses.V[0], OneMinuses.V[2], 0)
+	);
+}
+
+static inline uniform FMatrix ComputeWorldSpaceInertia(const uniform FVector4 &Q, const uniform FVector &I)
+{
+	const uniform FMatrix QM = MakeQuatRotationTranslationMatrix(Q);
+	return MultiplyAB(QM, MultiplyABt(SetMatrix33(I.V[0], I.V[1], I.V[2]), QM));
+}
+
+export void ApplyRotationConstraint(
+	uniform FJointSolverGaussSeidel * uniform M,
+	const uniform float Stiffness,
+	const uniform FVector& Axis0,
+	const uniform FVector& Axis1,
+	const uniform float Angle)
+{
+	// World-space inverse mass
+	const uniform FMatrix InvI0 = ComputeWorldSpaceInertia(M->Qs[0], M->InvILs[0]);
+	const uniform FMatrix InvI1 = ComputeWorldSpaceInertia(M->Qs[1], M->InvILs[1]);
+	const uniform FVector IA0 = Multiply(InvI0, Axis0);
+	const uniform FVector IA1 = Multiply(InvI1, Axis1);
+
+	// Joint-space inverse mass
+	const uniform float II0 = VectorDot(Axis0, IA0);
+	const uniform float II1 = VectorDot(Axis1, IA1);
+	const uniform float AngleM = Angle / (II0 + II1);
+
+	const uniform FVector DR0 = IA0 * (AngleM);
+	const uniform FVector DR1 = IA1 * -(AngleM);
+
+	ApplyRotationDeltaInner(M, Stiffness, DR0, DR1);
+
+	if (M->AngularPositionCorrection > 0)
+	{
+		const uniform FVector PrevX0 = M->Xs[0];
+		const uniform FVector PrevX1 = M->Xs[1];
+		UpdateDerivedStateInner(M);
+
+		const uniform FVector DX = (M->Xs[1] - M->Xs[0]) - (PrevX1 - PrevX0);
+		const uniform FVector DP0 = DX * (M->InvMs[0] / (M->InvMs[0] + M->InvMs[1]));
+		const uniform FVector DP1 = DX * (-M->InvMs[1] / (M->InvMs[0] + M->InvMs[1]));
+		ApplyPositionDeltaInner(M, Stiffness * M->AngularPositionCorrection, DP0, DP1);
+	}
+
+	UpdateDerivedStateInner(M);
+}
+
+static inline uniform FMatrix ComputeJointFactorMatrix(const uniform FVector& V, const uniform FMatrix& M, const uniform float Im)
+{
+	// Rigid objects rotational contribution to the impulse.
+	// Vx*M*VxT+Im
+	const uniform FVector V221 = SetVector(-V.V[2], V.V[2], -V.V[1]);
+	const uniform FVector V222 = SetVector(-V.V[2], -V.V[2], -V.V[2]);
+	const uniform FVector M544 = SetVector(M.M[5], M.M[4], M.M[4]);
+	const uniform FVector V111 = SetVector(V.V[1], V.V[1], V.V[1]);
+	const uniform FVector M988 = SetVector(M.M[9], M.M[8], M.M[8]);
+	const uniform FVector V100 = SetVector(V.V[1], -V.V[0], V.V[0]);
+	const uniform FVector M995 = SetVector(M.M[9], M.M[9], M.M[5]);
+	const uniform FVector M10109 = SetVector(M.M[10], M.M[10], M.M[9]);
+	const uniform FVector IM = SetVector(Im, 0, 0);
+
+	const uniform FVector R1 = (V221 * (V222 * M544 + V111 * M988)) + (V100 * (V222 * M995 + V111 * M10109)) + IM;
+
+	const uniform FVector V211 = SetVector(V.V[2], -V.V[1], -V.V[1]);
+	const uniform FVector V221_2 = SetVector(V.V[2], V.V[2], -V.V[1]);
+	const uniform FVector M000 = SetVector(M.M[0], M.M[0], M.M[0]);
+	const uniform FVector V000 = SetVector(-V.V[0], -V.V[0], V.V[0]);
+	const uniform FVector M884 = SetVector(M.M[8], M.M[8], M.M[4]);
+	const uniform FVector V000_2 = SetVector(-V.V[0], V.V[0], V.V[0]);
+	const uniform FVector M844 = SetVector(M.M[8], M.M[4], M.M[4]);
+	const uniform FVector M1095 = SetVector(M.M[10], M.M[9], M.M[5]);
+	const uniform FVector IM_2 = SetVector(Im, 0, Im);
+
+	const uniform FVector R2 = (V211 * (V221_2 * M000 + V000 * M884)) + (V000_2 * (V221_2 * M844 + V000 * M1095)) + IM_2;
+
+	return SetMatrix33(R1.V[0], R1.V[1], R1.V[2], R2.V[0], R2.V[1], R2.V[2]);
+}
+
+export void ApplyPositionConstraint(
+	uniform FJointSolverGaussSeidel * uniform M,
+	const uniform FVector &CX,
+	const uniform float LinearStiffness)
+{
+	const uniform FMatrix InvI0 = ComputeWorldSpaceInertia(M->Qs[0], M->InvILs[0]);
+	const uniform FMatrix InvI1 = ComputeWorldSpaceInertia(M->Qs[1], M->InvILs[1]);
+
+	// Calculate constraint correction
+	uniform FMatrix M0 = Matrix33Zero;
+	uniform FMatrix M1 = M0;
+	if (M->InvMs[0] > 0)
+	{
+		M0 = ComputeJointFactorMatrix(M->Xs[0] - M->Ps[0], InvI0, M->InvMs[0]);
+	}
+	if (M->InvMs[1] > 0)
+	{
+		M1 = ComputeJointFactorMatrix(M->Xs[1] - M->Ps[1], InvI1, M->InvMs[1]);
+	}
+	const uniform FMatrix MI = MatrixInverse(AddAB(M0, M1));
+	const uniform FVector DX = Multiply(MI, CX);
+
+	// Apply constraint correction
+	const uniform FVector DP0 = DX * M->InvMs[0];
+	const uniform FVector DP1 = DX * -M->InvMs[1];
+	const uniform FVector DR0 = Multiply(InvI0, VectorCross(M->Xs[0] - M->Ps[0], DX));
+	const uniform FVector DR1 = Multiply(InvI1, VectorCross(M->Xs[1] - M->Ps[1], VectorNegate(DX)));
+
+	ApplyPositionDeltaInner(M, LinearStiffness, DP0, DP1);
+	ApplyRotationDeltaInner(M, LinearStiffness, DR0, DR1);
+	UpdateDerivedStateInner(M);
+}
+
+static inline uniform FVector CalculateAngularVelocity(const uniform FVector4 &R0, const uniform FVector4 &R1, const uniform float Dt)
+{
+	const uniform FVector4 R1Shortest = VectorQuaternionEnforceShortestArcWith(R1, R0);
+
+	// W = 2 * dQ/dT * Qinv
+	const uniform FVector4 DRDt = (R1Shortest - R0) * (1.0f / Dt);
+	const uniform FVector4 RInv = QuatInverse(R0);
+	const uniform FVector4 W = VectorQuaternionMultiply2(DRDt, RInv) * 2.0f;
+
+	return SetVector(W);
+}
+
+static inline void ApplyPositionDeltaInner(
+	uniform FJointSolverGaussSeidel * uniform M,
+	const uniform int32 BodyIndex,
+	const uniform float Stiffness,
+	const uniform FVector& DP)
+{
+	M->Ps[BodyIndex] = M->Ps[BodyIndex] + Stiffness * DP;
+}
+
+static inline void ApplyRotationDeltaInner(
+	uniform FJointSolverGaussSeidel * uniform M,
+	const uniform int32 BodyIndex,
+	const uniform float Stiffness,
+	const uniform FVector& DR)
+{
+	const uniform FVector4 Qs = M->Qs[BodyIndex];
+	const uniform FVector4 DRStiffness = SetVector4(DR * Stiffness, 0);
+	const uniform FVector4 DQ = VectorQuaternionMultiply2(DRStiffness, Qs) * 0.5f;
+	M->Qs[BodyIndex] = VectorNormalizeQuaternion(Qs + DQ);
+}
+
+static inline void ApplyVelocityDeltaInner(
+	uniform FJointSolverGaussSeidel * uniform M,
+	const uniform int32 BodyIndex,
+	const uniform float Stiffness,
+	const uniform FVector& DV,
+	const uniform FVector& DW)
+{
+	M->Vs[BodyIndex] = M->Vs[BodyIndex] + Stiffness * DV;
+	M->Ws[BodyIndex] = M->Ws[BodyIndex] + Stiffness * DW;
+}
+
+static inline void ApplyVelocityDeltaInner(
+	uniform FJointSolverGaussSeidel * uniform M,
+	const uniform float Stiffness,
+	const uniform FVector& DV0,
+	const uniform FVector& DW0,
+	const uniform FVector& DV1,
+	const uniform FVector& DW1)
+{
+	M->Vs[0] = M->Vs[0] + Stiffness * DV0;
+	M->Vs[1] = M->Vs[1] + Stiffness * DV1;
+	M->Ws[0] = M->Ws[0] + Stiffness * DW0;
+	M->Ws[1] = M->Ws[1] + Stiffness * DW1;
+}
+
+export void ApplyVelocityDelta(
+	uniform FJointSolverGaussSeidel * uniform M,
+	const uniform float Stiffness,
+	const uniform FVector& DV0,
+	const uniform FVector& DW0,
+	const uniform FVector& DV1,
+	const uniform FVector& DW1)
+{
+	ApplyVelocityDeltaInner(M, Stiffness, DV0, DW0, DV1, DW1);
+}
+
+static inline void UpdateDerivedStateInner(uniform FJointSolverGaussSeidel * uniform M, const uniform int32 BodyIndex)
+{
+	M->Xs[BodyIndex] = M->Ps[BodyIndex] + VectorQuaternionRotateVector(M->Qs[BodyIndex], M->XLs[BodyIndex].Translation);
+	M->Rs[BodyIndex] = VectorQuaternionMultiply2(M->Qs[BodyIndex], M->XLs[BodyIndex].Rotation);
+}
+
+export void ApplyPositionProjection(
+	uniform FJointSolverGaussSeidel * uniform M,
+	const uniform FVector &CX,
+	const uniform float CXLen,
+	const uniform float ParentMassScale,
+	const uniform float Stiffness)
+{
+	const uniform FVector CXDir = CX / CXLen;
+	const uniform FVector V0 = M->Vs[0] + VectorCross(M->Ws[0], M->Xs[0] - M->Ps[0]);
+	const uniform FVector V1 = M->Vs[1] + VectorCross(M->Ws[1], M->Xs[1] - M->Ps[1]);
+	const uniform FVector CV = VectorDot(V1 - V0, CXDir) * CXDir;
+
+	const uniform float IM0 = ParentMassScale * M->InvMs[0];
+	const uniform float IM1 = M->InvMs[1];
+	const uniform FVector IIL0 = ParentMassScale * M->InvILs[0];
+	const uniform FVector& IIL1 = M->InvILs[1];
+	const uniform FMatrix II0 = ComputeWorldSpaceInertia(M->Qs[0], IIL0);
+	const uniform FMatrix II1 = ComputeWorldSpaceInertia(M->Qs[1], IIL1);
+	uniform FMatrix J0 = Matrix33Zero;
+	if(IM0 > 0)
+	{
+		J0 = ComputeJointFactorMatrix(M->Xs[0] - M->Ps[0], II0, IM0);
+	}
+	const uniform FMatrix J1 = ComputeJointFactorMatrix(M->Xs[1] - M->Ps[1], II1, IM1);
+	const uniform FMatrix IJ = MatrixInverse(AddAB(J0, J1));
+
+	const uniform FVector DX = Multiply(IJ, CX);
+	const uniform FVector DV = Multiply(IJ, CV);
+
+	const uniform FVector DP0 = IM0 * DX;
+	const uniform FVector DP1 = -IM1 * DX;
+	const uniform FVector DR0 = Multiply(II0, VectorCross(M->Xs[0] - M->Ps[0], DX));
+	const uniform FVector DR1 = Multiply(II1, VectorCross(M->Xs[1] - M->Ps[1], VectorNegate(DX)));
+
+	const uniform FVector DV0 = IM0 * DV;
+	const uniform FVector DV1 = -IM1 * DV;
+	const uniform FVector DW0 = Multiply(II0, VectorCross(M->Xs[0] - M->Ps[0], DV));
+	const uniform FVector DW1 = Multiply(II1, VectorCross(M->Xs[1] - M->Ps[1], VectorNegate(DV)));
+
+	ApplyPositionDeltaInner(M, Stiffness, DP0, DP1);
+	ApplyRotationDeltaInner(M, Stiffness, DR0, DR1);
+	ApplyVelocityDeltaInner(M, Stiffness, DV0, DW0, DV1, DW1);
+	UpdateDerivedStateInner(M);
+}
+
+// See "XPBD: Position-Based Simulation of Compliant Constrained Dynamics"
+export void ApplyRotationConstraintSoft(
+	uniform FJointSolverGaussSeidel * uniform M,
+	const uniform float Dt,
+	const uniform float Stiffness,
+	const uniform float Damping,
+	const uniform bool bAccelerationMode,
+	const uniform FVector& Axis0,
+	const uniform FVector& Axis1,
+	const uniform float Angle,
+	uniform float& Lambda)
+{
+	// World-space inverse mass
+	const uniform FMatrix InvI0 = ComputeWorldSpaceInertia(M->Qs[0], M->InvILs[0]);
+	const uniform FMatrix InvI1 = ComputeWorldSpaceInertia(M->Qs[1], M->InvILs[1]);
+	const uniform FVector IA0 = Multiply(InvI0, Axis0);
+	const uniform FVector IA1 = Multiply(InvI1, Axis1);
+
+	// Joint-space inverse mass
+	const uniform float II0 = VectorDot(Axis0, IA0);
+	const uniform float II1 = VectorDot(Axis1, IA1);
+	const uniform float II = (II0 + II1);
+
+	// Damping angular velocity
+	uniform float AngVelDt = 0;
+	if (Damping > KINDA_SMALL_NUMBER)
+	{
+		const uniform FVector W0 = CalculateAngularVelocity(M->PrevQs[0], M->Qs[0], 1.0f);
+		const uniform FVector W1 = CalculateAngularVelocity(M->PrevQs[1], M->Qs[1], 1.0f);
+		AngVelDt = VectorDot(Axis0, W0) - VectorDot(Axis1, W1);
+	}
+
+	const uniform float SpringMassScale = (bAccelerationMode) ? 1.0f / II : 1.0f;
+	const uniform float S = SpringMassScale * Stiffness * Dt * Dt;
+	const uniform float D = SpringMassScale * Damping * Dt;
+	const uniform float Multiplier = 1.0f / ((S + D) * II + 1.0f);
+	const uniform float DLambda = Multiplier * (S * Angle - D * AngVelDt - Lambda);
+
+	//const FVec3 DR0 = IA0 * DLambda;
+	//const FVec3 DR1 = IA1 * -DLambda;
+	const uniform FVector DR0 = Axis0 * (DLambda * II0);
+	const uniform FVector DR1 = Axis1 * -(DLambda * II1);
+
+	Lambda += DLambda;
+	ApplyRotationDeltaInner(M, 1.0f, DR0, DR1);
+
+	// Correct the positional error that was introduced by the rotation correction. This is
+	// correct when the position dofs are all locked, but not otherwise (fixable).
+	// This significantly improves angular stiffness at lower iterations, although the same effect
+	// is achieved by increasing iterations.
+	// @todo(ccaulfield): this position correction needs to have components in direction of inactive position constraints removed
+	if (M->AngularPositionCorrection > 0)
+	{
+		const uniform FVector PrevX0 = M->Xs[0];
+		const uniform FVector PrevX1 = M->Xs[1];
+		UpdateDerivedStateInner(M);
+
+		const uniform FVector DX = (M->Xs[1] - M->Xs[0]) - (PrevX1 - PrevX0);
+		const uniform FVector DP0 = DX * (M->InvMs[0] / (M->InvMs[0] + M->InvMs[1]));
+		const uniform FVector DP1 = DX * (-M->InvMs[1] / (M->InvMs[0] + M->InvMs[1]));
+		ApplyPositionDeltaInner(M, M->AngularPositionCorrection, DP0, DP1);
+	}
+
+	UpdateDerivedStateInner(M);
+}
