#UE4 Fix memory corruption due to an overflow/underflow in heightfield. Added protection against the corruption in the future

diff --git a/Engine/Source/Runtime/Experimental/Chaos/Private/Chaos/HeightField.cpp b/Engine/Source/Runtime/Experimental/Chaos/Private/Chaos/HeightField.cpp
--- a/Engine/Source/Runtime/Experimental/Chaos/Private/Chaos/HeightField.cpp
+++ b/Engine/Source/Runtime/Experimental/Chaos/Private/Chaos/HeightField.cpp
@@ -851,9 +851,9 @@ namespace Chaos
 			, NumY(Size[1])
 		{
 			int32 BitsNeeded = NumX * NumY;
-			int32 BytesNeeded = 1 + (BitsNeeded) / 8;
-			Data = MakeUnique<uint8[]>(BytesNeeded);
-			FMemory::Memzero(Data.Get(), BytesNeeded);
+			DataSize = 1 + (BitsNeeded) / 8;
+			Data = MakeUnique<uint8[]>(DataSize);
+			FMemory::Memzero(Data.Get(), DataSize);
 		}
 
 		bool Contains(const TVector<int32, 2>& Coordinate)
@@ -861,6 +861,7 @@ namespace Chaos
 			int32 Idx = Coordinate[1] * NumX + Coordinate[0];
 			int32 ByteIdx = Idx / 8;
 			int32 BitIdx = Idx % 8;
+			check(ByteIdx >= 0 && ByteIdx < DataSize);
 			bool bContains = (Data[ByteIdx] >> BitIdx) & 0x1;
 			return bContains;
 		}
@@ -871,6 +872,7 @@ namespace Chaos
 			int32 ByteIdx = Idx / 8;
 			int32 BitIdx = Idx % 8;
 			uint8 Mask = 1 << BitIdx;
+			check(ByteIdx >= 0 && ByteIdx < DataSize);
 			Data[ByteIdx] |= Mask;
 		}
 
@@ -878,6 +880,7 @@ namespace Chaos
 		int32 NumX;
 		int32 NumY;
 		TUniquePtr<uint8[]> Data;
+		int32 DataSize;
 	};
 
 	template<typename T>
@@ -1005,7 +1008,7 @@ namespace Chaos
 							CurrentCell += Direction;
 
 							// Fail if we leave the grid
-							if(Begin[0] < 0 || Begin[1] < 0 || Begin[0] > FlatGrid.Counts()[0] - 1 || Begin[1] > FlatGrid.Counts()[1] - 1)
+							if(CurrentCell[0] < 0 || CurrentCell[1] < 0 || CurrentCell[0] > FlatGrid.Counts()[0] - 1 || CurrentCell[1] > FlatGrid.Counts()[1] - 1)
 							{
 								break;
 							}
