Remove cell height caching from Heightfield, saves significant memory

diff --git a/Engine/Source/Runtime/Core/Public/UObject/ExternalPhysicsCustomObjectVersion.h b/Engine/Source/Runtime/Core/Public/UObject/ExternalPhysicsCustomObjectVersion.h
--- a/Engine/Source/Runtime/Core/Public/UObject/ExternalPhysicsCustomObjectVersion.h
+++ b/Engine/Source/Runtime/Core/Public/UObject/ExternalPhysicsCustomObjectVersion.h
@@ -104,6 +104,9 @@ struct CORE_API FExternalPhysicsCustomObjectVersion
 
 		// Convexes use concrete planes
 		ConvexUsesTPlaneConcrete,
+
+		// Heightfield uses uint16 heights directly
+		HeightfieldUsesHeightsDirectly,
 		
 		// -----<new versions can be added above this line>-------------------------------------------------
 		VersionPlusOne,
diff --git a/Engine/Source/Runtime/Experimental/Chaos/Private/Chaos/HeightField.cpp b/Engine/Source/Runtime/Experimental/Chaos/Private/Chaos/HeightField.cpp
--- a/Engine/Source/Runtime/Experimental/Chaos/Private/Chaos/HeightField.cpp
+++ b/Engine/Source/Runtime/Experimental/Chaos/Private/Chaos/HeightField.cpp
@@ -618,8 +618,12 @@ namespace Chaos
 	{
 		if (FlatGrid.IsValid(InCoord))
 		{
+			//todo: just compute max height, avoid extra work since this is called from tight loop
+			TVec3<T> Min,Max;
+			CalcCellBounds3D(InCoord,Min,Max);
+
 			OutMin = TVec3<T>(InCoord[0], InCoord[1], GeomData.GetMinHeight());
-			OutMax = TVec3<T>(InCoord[0] + 1, InCoord[1] + 1, GeomData.CellHeights[InCoord[1] * (GeomData.NumCols - 1) + InCoord[0]]);
+			OutMax = TVec3<T>(InCoord[0] + 1, InCoord[1] + 1, Max[2]);
 			OutMin = OutMin - InInflate;
 			OutMax = OutMax + InInflate;
 
@@ -652,8 +656,12 @@ namespace Chaos
 	{
 		if (FlatGrid.IsValid(InCoord))
 		{
+			//todo: just compute max height, avoid extra work since this is called from tight loop
+			TVec3<T> Min,Max;
+			CalcCellBounds3D(InCoord,Min,Max);
+
 			OutMin = TVec3<T>(InCoord[0], InCoord[1], GeomData.GetMinHeight());
-			OutMax = TVec3<T>(InCoord[0] + 1, InCoord[1] + 1, GeomData.CellHeights[InCoord[1] * (GeomData.NumCols - 1) + InCoord[0]]);
+			OutMax = TVec3<T>(InCoord[0] + 1, InCoord[1] + 1, Max[2]);
 			OutMin = OutMin * GeomData.Scale - InInflate;
 			OutMax = OutMax * GeomData.Scale + InInflate;
 			return true;
@@ -1629,19 +1637,6 @@ namespace Chaos
 		// Cache per-cell bounds
 		const int32 NumX = GeomData.NumCols - 1;
 		const int32 NumY = GeomData.NumRows - 1;
-		//GeomData.CellBounds.SetNum(NumX * NumY);
-		GeomData.CellHeights.SetNum(NumX*NumY);
-		for(int32 XIndex = 0; XIndex < NumX; ++XIndex)
-		{
-			for(int32 YIndex = 0; YIndex < NumY; ++YIndex)
-			{
-				const TVector<int32, 2> Cell(XIndex, YIndex);
-				TVector<T, 3> Min, Max;
-				CalcCellBounds3D(Cell, Min, Max);
-				//GeomData.CellBounds[YIndex * NumX + XIndex] = TAABB<T, 3>(Min, Max);
-				GeomData.CellHeights[YIndex * NumX + XIndex] = Max.Z;
-			}
-		}
 	}
 
 	template<typename T>
diff --git a/Engine/Source/Runtime/Experimental/Chaos/Public/Chaos/HeightField.h b/Engine/Source/Runtime/Experimental/Chaos/Public/Chaos/HeightField.h
--- a/Engine/Source/Runtime/Experimental/Chaos/Public/Chaos/HeightField.h
+++ b/Engine/Source/Runtime/Experimental/Chaos/Public/Chaos/HeightField.h
@@ -201,7 +201,6 @@ namespace Chaos
 			uint16 NumCols;
 			RealType Range;
 			RealType HeightPerUnit;
-			TArray<RealType> CellHeights; //todo: remove this and use heights directly
 
 			constexpr float GetCellWidth() const
 			{
@@ -290,14 +289,11 @@ namespace Chaos
 					{
 						TArray<TBox<RealType, 3>> CellBounds;
 						Ar << CellBounds;
-						for (const TBox<RealType, 3>& Box : CellBounds)
-						{
-							CellHeights.Add(Box.Extents().Z);
-						}
 					}
-					else
+					else if(Ar.CustomVer(FExternalPhysicsCustomObjectVersion::GUID) < FExternalPhysicsCustomObjectVersion::HeightfieldUsesHeightsDirectly)
 					{
-						Ar << CellHeights;
+						TArray<RealType> OldHeights;
+						Ar << OldHeights;
 					}
 				}
 
diff --git a/Engine/Source/Runtime/Landscape/Private/LandscapeCollision.cpp b/Engine/Source/Runtime/Landscape/Private/LandscapeCollision.cpp
--- a/Engine/Source/Runtime/Landscape/Private/LandscapeCollision.cpp
+++ b/Engine/Source/Runtime/Landscape/Private/LandscapeCollision.cpp
@@ -129,7 +129,7 @@ ULandscapeMeshCollisionComponent::FPhysXMeshRef::~FPhysXMeshRef()
 }
 
 // Generate a new guid to force a recache of landscape collison derived data
-#define LANDSCAPE_COLLISION_DERIVEDDATA_VER	TEXT("43B2B22117524EAB8FB46062C2AF0BDA")
+#define LANDSCAPE_COLLISION_DERIVEDDATA_VER	TEXT("CC58B9FA08AD47E3BF06976E60B693C3")
 
 static FString GetHFDDCKeyString(const FName& Format, bool bDefMaterial, const FGuid& StateId, const TArray<UPhysicalMaterial*>& PhysicalMaterials)
 {
