Forced GJK to always provide a touch or a penetration once we decide a hit happened. Combats sweeping touching objects very slightly together resulting in a hit time greater than zero.

diff --git a/Engine/Source/Runtime/Experimental/Chaos/Public/Chaos/GJK.h b/Engine/Source/Runtime/Experimental/Chaos/Public/Chaos/GJK.h
--- a/Engine/Source/Runtime/Experimental/Chaos/Public/Chaos/GJK.h
+++ b/Engine/Source/Runtime/Experimental/Chaos/Public/Chaos/GJK.h
@@ -574,7 +574,7 @@ namespace Chaos
 				const TVec3<T> ClosestBInA = StartTM.TransformPosition(ClosestB);
 				const T InGJKPreDist = FMath::Sqrt(InGJKPreDist2);
 				OutNormal = (ClosestBInA - ClosestA).GetUnsafeNormal();	//question: should we just use InGJKPreDist2?
-				const T Penetration = ThicknessA + ThicknessB - InGJKPreDist;
+				const T Penetration = FMath::Clamp<T>(ThicknessA + ThicknessB - InGJKPreDist, 0, TNumericLimits<T>::Max());
 				const TVector<T, 3> ClosestLocal = ClosestB - OutNormal * ThicknessB;
 
 				OutPosition = StartPoint + ClosestLocal + OutNormal * Penetration;
