Fix ensure caused by numeric precision

diff --git a/Engine/Source/Runtime/Experimental/Chaos/Public/Chaos/ImplicitObjectScaled.h b/Engine/Source/Runtime/Experimental/Chaos/Public/Chaos/ImplicitObjectScaled.h
--- a/Engine/Source/Runtime/Experimental/Chaos/Public/Chaos/ImplicitObjectScaled.h
+++ b/Engine/Source/Runtime/Experimental/Chaos/Public/Chaos/ImplicitObjectScaled.h
@@ -350,11 +350,15 @@ public:
 			
 			if (MObject->SweepGeom(ScaledB, BToATMNoScale, UnscaledDir, UnscaledLength, UnscaledTime, UnscaledPosition, UnscaledNormal, OutFaceIndex, MInternalThickness + Thickness, bComputeMTD))
 			{
-				OutTime = LengthScaleInv * UnscaledTime;
-				LocalPosition = MScale * UnscaledPosition;
-				LocalNormal = (MInvScale * UnscaledNormal).GetSafeNormal();
-				ensure(OutTime <= Length);
-				return true;
+				const T NewTime = LengthScaleInv * UnscaledTime;
+				//We double check that NewTime < Length because of potential precision issues. When that happens we always keep the shortest hit first
+				if (NewTime < Length)
+				{
+					OutTime = NewTime;
+					LocalPosition = MScale * UnscaledPosition;
+					LocalNormal = (MInvScale * UnscaledNormal).GetSafeNormal();
+					return true;
+				}
 			}
 		}
 
