Fix memory read access issue on chaos constraint graph

diff --git a/Engine/Source/Runtime/Experimental/Chaos/Private/Chaos/PBDConstraintGraph.cpp b/Engine/Source/Runtime/Experimental/Chaos/Private/Chaos/PBDConstraintGraph.cpp
--- a/Engine/Source/Runtime/Experimental/Chaos/Private/Chaos/PBDConstraintGraph.cpp
+++ b/Engine/Source/Runtime/Experimental/Chaos/Private/Chaos/PBDConstraintGraph.cpp
@@ -14,8 +14,6 @@
 
 using namespace Chaos;
 
-
-
 FPBDConstraintGraph::FPBDConstraintGraph() : VisitToken(0)
 {
 }
@@ -715,7 +713,7 @@ void FPBDConstraintGraph::DisableParticle(TGeometryParticleHandle<FReal, 3>* Par
 		}
 
 	}
-	else if (PBDRigid)
+	else
 	{
 		// Kinematic particles are included in IslandToParticles, however we cannot use islands to look them up.
 		// TODO find faster removal method?
@@ -784,3 +782,4 @@ bool FPBDConstraintGraph::CheckIslands(const TArray<TGeometryParticleHandle<FRea
 
 	return bIsValid;
 }
+
