FQueryFastData from 40 bytes 9 padding down to 32 bytes 1 padding
TPBDAnimDriveConstraint from 40 bytes 8 padding down to 32 bytes 0 padding
FPhysicsSceneGuard down from 24 bytes 8 padding to 16 bytes 0 padding
FJointSolverGaussSeidel down from 528 bytes 36 padding to 496 bytes 4 padding
FPBDJointConstraints down from 400 bytes 22 padding to 384 bytes 6 padding

diff --git a/Engine/Source/Runtime/Experimental/Chaos/Public/Chaos/ISpatialAcceleration.h b/Engine/Source/Runtime/Experimental/Chaos/Public/Chaos/ISpatialAcceleration.h
--- a/Engine/Source/Runtime/Experimental/Chaos/Public/Chaos/ISpatialAcceleration.h
+++ b/Engine/Source/Runtime/Experimental/Chaos/Public/Chaos/ISpatialAcceleration.h
@@ -10,21 +10,22 @@ namespace Chaos
 struct CHAOS_API FQueryFastData
 {
 	FQueryFastData(const FVec3& InDir, const FReal InLength)
-		: bParallel{ InDir[0] == 0, InDir[1] == 0, InDir[2] == 0 }
-		, Dir(InDir)
+		: Dir(InDir)
 		, InvDir(bParallel[0] ? 0 : 1 / Dir[0], bParallel[1] ? 0 : 1 / Dir[1], bParallel[2] ? 0 : 1 / Dir[2])
+		, bParallel{ InDir[0] == 0, InDir[1] == 0, InDir[2] == 0 }
 	{
 		ensure(InLength);
 		SetLength(InLength);
 	}
 
-	const bool bParallel[3];
 	const FVec3& Dir;
 	const FVec3 InvDir;
 
 	FReal CurrentLength;
 	FReal InvCurrentLength;
 
+	const bool bParallel[3];
+
 #if PLATFORM_WINDOWS || PLATFORM_XBOXONE
 	#pragma warning (push)
 	#pragma warning(disable:4723)
@@ -49,9 +50,9 @@ struct CHAOS_API FQueryFastData
 
 protected:
 	FQueryFastData(const FVec3& InDummyDir)
-		: bParallel{}
-		, Dir(InDummyDir)
+		: Dir(InDummyDir)
 		, InvDir()
+		, bParallel{}
 	{}
 };
 
diff --git a/Engine/Source/Runtime/Experimental/Chaos/Public/Chaos/Joint/PBDJointSolverGaussSeidel.h b/Engine/Source/Runtime/Experimental/Chaos/Public/Chaos/Joint/PBDJointSolverGaussSeidel.h
--- a/Engine/Source/Runtime/Experimental/Chaos/Public/Chaos/Joint/PBDJointSolverGaussSeidel.h
+++ b/Engine/Source/Runtime/Experimental/Chaos/Public/Chaos/Joint/PBDJointSolverGaussSeidel.h
@@ -315,10 +315,14 @@ namespace Chaos
 
 		// World-space constraint state
 		FVec3 Xs[MaxConstrainedBodies];				// World-space joint connector positions
+		FReal LinearSoftLambda;						// XPBD constraint multipliers (net applied constraint-space deltas)
+		FReal LinearDriveLambda;					// XPBD constraint multipliers (net applied constraint-space deltas)
 		FRotation3 Rs[MaxConstrainedBodies];		// World-space joint connector rotations
 
 		// World-space body state
 		FVec3 Ps[MaxConstrainedBodies];				// World-space particle CoM positions
+		FReal TwistSoftLambda;						// XPBD constraint multipliers (net applied constraint-space deltas)
+		FReal SwingSoftLambda;						// XPBD constraint multipliers (net applied constraint-space deltas)
 		FRotation3 Qs[MaxConstrainedBodies];		// World-space particle CoM rotations
 		FVec3 Vs[MaxConstrainedBodies];				// World-space particle CoM velocities
 		FVec3 Ws[MaxConstrainedBodies];				// World-space particle CoM angular velocities
@@ -331,17 +335,11 @@ namespace Chaos
 
 		// XPBD Previous iteration world-space body state
 		FVec3 PrevPs[MaxConstrainedBodies];			// World-space particle CoM positions
+		FReal TwistDriveLambda;						// XPBD constraint multipliers (net applied constraint-space deltas)
+		FReal SwingDriveLambda;						// XPBD constraint multipliers (net applied constraint-space deltas)
 		FRotation3 PrevQs[MaxConstrainedBodies];	// World-space particle CoM rotations
 		FVec3 PrevXs[MaxConstrainedBodies];			// World-space joint connector positions
-
-		// XPBD constraint multipliers (net applied constraint-space deltas)
-		FReal LinearSoftLambda;
-		FReal LinearDriveLambda;
-		FReal TwistSoftLambda;
-		FReal SwingSoftLambda;
-		FReal TwistDriveLambda;
-		FReal SwingDriveLambda;
-
+		
 		// Post-angular constraint position fixup (to reduce iterations required for stiff angular constraints)
 		FReal AngularPositionCorrection;
 	};
