Chaos - RBAN Joint memory reduction

diff --git a/Engine/Source/Runtime/Experimental/Chaos/Private/Chaos/Evolution/PBDMinEvolution.cpp b/Engine/Source/Runtime/Experimental/Chaos/Private/Chaos/Evolution/PBDMinEvolution.cpp
--- a/Engine/Source/Runtime/Experimental/Chaos/Private/Chaos/Evolution/PBDMinEvolution.cpp
+++ b/Engine/Source/Runtime/Experimental/Chaos/Private/Chaos/Evolution/PBDMinEvolution.cpp
@@ -22,6 +22,8 @@ namespace Chaos
 	DECLARE_CYCLE_STAT(TEXT("FPBDMinEvolution::AdvanceOneTimeStep"), STAT_MinEvolution_AdvanceOneTimeStep, STATGROUP_Chaos);
 	DECLARE_CYCLE_STAT(TEXT("FPBDMinEvolution::Integrate"), STAT_MinEvolution_Integrate, STATGROUP_Chaos);
 	DECLARE_CYCLE_STAT(TEXT("FPBDMinEvolution::KinematicTargets"), STAT_MinEvolution_KinematicTargets, STATGROUP_Chaos);
+	DECLARE_CYCLE_STAT(TEXT("FPBDMinEvolution::PrepareConstraints"), STAT_MinEvolution_PrepareConstraints, STATGROUP_Chaos);
+	DECLARE_CYCLE_STAT(TEXT("FPBDMinEvolution::UnprepareConstraints"), STAT_MinEvolution_UnprepareConstraints, STATGROUP_Chaos);
 	DECLARE_CYCLE_STAT(TEXT("FPBDMinEvolution::ApplyConstraints"), STAT_MinEvolution_ApplyConstraints, STATGROUP_Chaos);
 	DECLARE_CYCLE_STAT(TEXT("FPBDMinEvolution::UpdateVelocities"), STAT_MinEvolution_UpdateVelocites, STATGROUP_Chaos);
 	DECLARE_CYCLE_STAT(TEXT("FPBDMinEvolution::ApplyPushOut"), STAT_MinEvolution_ApplyPushOut, STATGROUP_Chaos);
@@ -95,6 +97,8 @@ namespace Chaos
 
 		if (Dt > 0)
 		{
+			PrepareConstraints(Dt);
+
 			ApplyConstraints(Dt);
 
 			if (PostApplyCallback != nullptr)
@@ -111,6 +115,8 @@ namespace Chaos
 				PostApplyPushOutCallback();
 			}
 
+			UnprepareConstraints(Dt);
+
 			UpdatePositions(Dt);
 		}
 	}
@@ -253,6 +259,26 @@ namespace Chaos
 		CollisionDetector.DetectCollisions(Dt);
 	}
 
+	void FPBDMinEvolution::PrepareConstraints(FReal Dt)
+	{
+		SCOPE_CYCLE_COUNTER(STAT_MinEvolution_PrepareConstraints);
+
+		for (FSimpleConstraintRule* ConstraintRule : PrioritizedConstraintRules)
+		{
+			ConstraintRule->PrepareConstraints(Dt);
+		}
+	}
+
+	void FPBDMinEvolution::UnprepareConstraints(FReal Dt)
+	{
+		SCOPE_CYCLE_COUNTER(STAT_MinEvolution_PrepareConstraints);
+
+		for (FSimpleConstraintRule* ConstraintRule : PrioritizedConstraintRules)
+		{
+			ConstraintRule->UnprepareConstraints(Dt);
+		}
+	}
+
 	void FPBDMinEvolution::ApplyConstraints(FReal Dt)
 	{
 		SCOPE_CYCLE_COUNTER(STAT_MinEvolution_ApplyConstraints);
diff --git a/Engine/Source/Runtime/Experimental/Chaos/Private/Chaos/Joint/PBDJointSolverGaussSeidel.cpp b/Engine/Source/Runtime/Experimental/Chaos/Private/Chaos/Joint/PBDJointSolverGaussSeidel.cpp
--- a/Engine/Source/Runtime/Experimental/Chaos/Private/Chaos/Joint/PBDJointSolverGaussSeidel.cpp
+++ b/Engine/Source/Runtime/Experimental/Chaos/Private/Chaos/Joint/PBDJointSolverGaussSeidel.cpp
@@ -11,10 +11,12 @@
 
 //#pragma optimize("", off)
 
-// Set to true to calculate all constraint corrections based on state at the beginning of the iteration. This would eliminate the need to recalculate the inertia and a few other
-// things in each constraint's Apply method. However, it probably introduces instability and requires more iterations to solve...
-bool bChaos_Joint_Accumulate = false;
-FAutoConsoleVariableRef CVarChaosJointAccumulate(TEXT("p.Chaos.Joint.Accumulate"), bChaos_Joint_Accumulate, TEXT("Whether to accumulate forces for each constraint over an iteration and apply all at once"));
+#if UE_BUILD_SHIPPING
+const bool bChaos_Joint_ISPC_Enabled = true;
+#else
+bool bChaos_Joint_ISPC_Enabled = true;
+FAutoConsoleVariableRef CVarChaosJointISPCEnabled(TEXT("p.Chaos.Joint.ISPC"), bChaos_Joint_ISPC_Enabled, TEXT("Whether to use ISPC optimizations in the Joint Solver"));
+#endif
 
 namespace Chaos
 {
@@ -247,10 +249,9 @@ namespace Chaos
 
 	FJointSolverGaussSeidel::FJointSolverGaussSeidel()
 	{
-		if (INTEL_ISPC)
+		if (bChaos_Joint_ISPC_Enabled)
 		{
 #if INTEL_ISPC
-			check(bChaos_Joint_Accumulate == false);
 			check(sizeof(FJointSolverGaussSeidel) == ispc::SizeofFJointSolverGaussSeidel());
 #endif
 		}
@@ -346,31 +347,14 @@ namespace Chaos
 	}
 
 
-	void FJointSolverGaussSeidel::InitAccumulatedDeltas()
-	{
-		DPs[0] = FVec3(0);
-		DPs[1] = FVec3(0);
-		DRs[0] = FVec3(0);
-		DRs[1] = FVec3(0);
-		DVs[0] = FVec3(0);
-		DVs[1] = FVec3(0);
-		DWs[0] = FVec3(0);
-		DWs[1] = FVec3(0);
-	}
-
-
 	void FJointSolverGaussSeidel::ApplyConstraints(
 		const FReal Dt,
 		const FPBDJointSolverSettings& SolverSettings,
 		const FPBDJointSettings& JointSettings)
 	{
-		InitAccumulatedDeltas();
-
 		ApplyPositionConstraints(Dt, SolverSettings, JointSettings);
 
 		ApplyRotationConstraints(Dt, SolverSettings, JointSettings);
-
-		ApplyAccumulatedDeltas();
 	}
 
 
@@ -379,13 +363,9 @@ namespace Chaos
 		const FPBDJointSolverSettings& SolverSettings,
 		const FPBDJointSettings& JointSettings)
 	{
-		InitAccumulatedDeltas();
-
 		ApplyRotationDrives(Dt, SolverSettings, JointSettings);
 
 		ApplyPositionDrives(Dt, SolverSettings, JointSettings);
-
-		ApplyAccumulatedDeltas();
 	}
 
 
@@ -394,13 +374,9 @@ namespace Chaos
 		const FPBDJointSolverSettings& SolverSettings,
 		const FPBDJointSettings& JointSettings)
 	{
-		InitAccumulatedDeltas();
-
 		//ApplyRotationProjection(Dt, SolverSettings, JointSettings);
 
 		ApplyPositionProjection(Dt, SolverSettings, JointSettings);
-
-		ApplyAccumulatedDeltas();
 	}
 
 
@@ -682,48 +658,20 @@ namespace Chaos
 		}
 	}
 
+
 	//
 	//
 	//////////////////////////////////////////////////////////////////////////
 	//
 	//
 
-	void FJointSolverGaussSeidel::ApplyAccumulatedDeltas()
-	{
-		if (bChaos_Joint_Accumulate)
-		{
-			Ps[0] += DPs[0];
-			Ps[1] += DPs[1];
-
-			const FRotation3 DQ0 = (FRotation3::FromElements(DRs[0], 0) * Qs[0]) * (FReal)0.5;
-			const FRotation3 DQ1 = (FRotation3::FromElements(DRs[1], 0) * Qs[1]) * (FReal)0.5;
-			Qs[0] = (Qs[0] + DQ0).GetNormalized();
-			Qs[1] = (Qs[1] + DQ1).GetNormalized();
-			Qs[1].EnforceShortestArcWith(Qs[0]);
-
-			Vs[0] += DVs[0];
-			Vs[1] += DVs[1];
-
-			Ws[0] += DWs[0];
-			Ws[1] += DWs[1];
-
-			UpdateDerivedState();
-		}
-	}
 
 	void FJointSolverGaussSeidel::ApplyPositionDelta(
 		const int32 BodyIndex,
 		const FReal Stiffness,
 		const FVec3& DP)
 	{
-		if (bChaos_Joint_Accumulate)
-		{
-			DPs[BodyIndex] += Stiffness * DP;
-		}
-		else
-		{
-			Ps[BodyIndex] += Stiffness * DP;
-		}
+		Ps[BodyIndex] += Stiffness * DP;
 	}
 
 
@@ -732,16 +680,8 @@ namespace Chaos
 		const FVec3& DP0,
 		const FVec3& DP1)
 	{
-		if (bChaos_Joint_Accumulate)
-		{
-			DPs[0] += Stiffness * DP0;
-			DPs[1] += Stiffness * DP1;
-		}
-		else
-		{
-			Ps[0] += Stiffness * DP0;
-			Ps[1] += Stiffness * DP1;
-		}
+		Ps[0] += Stiffness * DP0;
+		Ps[1] += Stiffness * DP1;
 	}
 
 
@@ -750,15 +690,8 @@ namespace Chaos
 		const FReal Stiffness,
 		const FVec3& DR)
 	{
-		if (bChaos_Joint_Accumulate)
-		{
-			DRs[BodyIndex] = Stiffness * DR;
-		}
-		else
-		{
-			const FRotation3 DQ = (FRotation3::FromElements(Stiffness * DR, 0) * Qs[BodyIndex]) * (FReal)0.5;
-			Qs[BodyIndex] = (Qs[BodyIndex] + DQ).GetNormalized();
-		}
+		const FRotation3 DQ = (FRotation3::FromElements(Stiffness * DR, 0) * Qs[BodyIndex]) * (FReal)0.5;
+		Qs[BodyIndex] = (Qs[BodyIndex] + DQ).GetNormalized();
 	}
 
 
@@ -767,19 +700,11 @@ namespace Chaos
 		const FVec3& DR0,
 		const FVec3& DR1)
 	{
-		if (bChaos_Joint_Accumulate)
-		{
-			DRs[0] += Stiffness * DR0;
-			DRs[1] += Stiffness * DR1;
-		}
-		else
-		{
-			const FRotation3 DQ0 = (FRotation3::FromElements(Stiffness * DR0, 0) * Qs[0]) * (FReal)0.5;
-			const FRotation3 DQ1 = (FRotation3::FromElements(Stiffness * DR1, 0) * Qs[1]) * (FReal)0.5;
-			Qs[0] = (Qs[0] + DQ0).GetNormalized();
-			Qs[1] = (Qs[1] + DQ1).GetNormalized();
-			Qs[1].EnforceShortestArcWith(Qs[0]);
-		}
+		const FRotation3 DQ0 = (FRotation3::FromElements(Stiffness * DR0, 0) * Qs[0]) * (FReal)0.5;
+		const FRotation3 DQ1 = (FRotation3::FromElements(Stiffness * DR1, 0) * Qs[1]) * (FReal)0.5;
+		Qs[0] = (Qs[0] + DQ0).GetNormalized();
+		Qs[1] = (Qs[1] + DQ1).GetNormalized();
+		Qs[1].EnforceShortestArcWith(Qs[0]);
 	}
 
 
@@ -789,16 +714,8 @@ namespace Chaos
 		const FVec3& DV,
 		const FVec3& DW)
 	{
-		if (bChaos_Joint_Accumulate)
-		{
-			DVs[BodyIndex] += Stiffness * DV;
-			DWs[BodyIndex] += Stiffness * DW;
-		}
-		else
-		{
-			Vs[BodyIndex] = Vs[BodyIndex] + Stiffness * DV;
-			Ws[BodyIndex] = Ws[BodyIndex] + Stiffness * DW;
-		}
+		Vs[BodyIndex] = Vs[BodyIndex] + Stiffness * DV;
+		Ws[BodyIndex] = Ws[BodyIndex] + Stiffness * DW;
 	}
 
 
@@ -809,20 +726,10 @@ namespace Chaos
 		const FVec3& DV1,
 		const FVec3& DW1)
 	{
-		if (bChaos_Joint_Accumulate)
-		{
-			DVs[0] += Stiffness * DV0;
-			DVs[1] += Stiffness * DV1;
-			DWs[0] += Stiffness * DW0;
-			DWs[1] += Stiffness * DW1;
-		}
-		else
-		{
-			Vs[0] += Stiffness * DV0;
-			Vs[1] += Stiffness * DV1;
-			Ws[0] += Stiffness * DW0;
-			Ws[1] += Stiffness * DW1;
-		}
+		Vs[0] += Stiffness * DV0;
+		Vs[1] += Stiffness * DV1;
+		Ws[0] += Stiffness * DW0;
+		Ws[1] += Stiffness * DW1;
 	}
 
 
@@ -866,7 +773,7 @@ namespace Chaos
 		const FReal Delta,
 		FReal& Lambda)
 	{
-		if (INTEL_ISPC)
+		if (bChaos_Joint_ISPC_Enabled)
 		{
 #if INTEL_ISPC
 			ispc::ApplyPositionConstraintSoft((ispc::FJointSolverGaussSeidel*)this, Dt, Stiffness, Damping, bAccelerationMode, (ispc::FVector&)Axis, Delta, Lambda);
@@ -969,7 +876,7 @@ namespace Chaos
 		const FReal Angle,
 		FReal& Lambda)
 	{
-		if (INTEL_ISPC)
+		if (bChaos_Joint_ISPC_Enabled)
 		{
 #if INTEL_ISPC
 			ispc::ApplyRotationConstraintSoft((ispc::FJointSolverGaussSeidel*)this, Dt, Stiffness, Damping, bAccelerationMode, (ispc::FVector&) Axis0, (ispc::FVector&)Axis1, Angle, Lambda);
@@ -1431,7 +1338,7 @@ namespace Chaos
 		const FVec3 CX = Xs[1] - Xs[0];
 		if (CX != FVec3(0))
 		{
-			if (INTEL_ISPC)
+			if (bChaos_Joint_ISPC_Enabled)
 			{
 #if INTEL_ISPC
 				FReal LinearStiffness = FPBDJointUtilities::GetLinearStiffness(SolverSettings, JointSettings);
@@ -1670,7 +1577,7 @@ namespace Chaos
 		{
 			FReal LinearProjection = FPBDJointUtilities::GetLinearProjection(SolverSettings, JointSettings);
 			const FReal ParentMassScale = FMath::Max(0.0f, 1.0f - LinearProjection);
-			if (INTEL_ISPC)
+			if (bChaos_Joint_ISPC_Enabled)
 			{
 #if INTEL_ISPC
 				const FReal Stiffness = FPBDJointUtilities::GetLinearStiffness(SolverSettings, JointSettings);
diff --git a/Engine/Source/Runtime/Experimental/Chaos/Private/Chaos/Joint/PBDJointSolverGaussSeidel.ispc b/Engine/Source/Runtime/Experimental/Chaos/Private/Chaos/Joint/PBDJointSolverGaussSeidel.ispc
--- a/Engine/Source/Runtime/Experimental/Chaos/Private/Chaos/Joint/PBDJointSolverGaussSeidel.ispc
+++ b/Engine/Source/Runtime/Experimental/Chaos/Private/Chaos/Joint/PBDJointSolverGaussSeidel.ispc
@@ -36,12 +36,6 @@ struct FJointSolverGaussSeidel
 	FVector Vs[MaxConstrainedBodies];			// World-space particle CoM velocities
 	FVector Ws[MaxConstrainedBodies];			// World-space particle CoM angular velocities
 
-	// Accumulated world-body state deltas
-	FVector DPs[MaxConstrainedBodies];			// World-space particle CoM positions
-	FVector DRs[MaxConstrainedBodies];			// World-space particle CoM rotations (in axis-angle form)
-	FVector DVs[MaxConstrainedBodies];			// World-space particle CoM velocities
-	FVector DWs[MaxConstrainedBodies];			// World-space particle CoM angular velocities
-
 	// XPBD Previous iteration world-space body state
 	FVector PrevPs[MaxConstrainedBodies];		// World-space particle CoM positions
 	float TwistDriveLambda;						// XPBD constraint multipliers (net applied constraint-space deltas)
diff --git a/Engine/Source/Runtime/Experimental/Chaos/Private/Chaos/PBDJointConstraints.cpp b/Engine/Source/Runtime/Experimental/Chaos/Private/Chaos/PBDJointConstraints.cpp
--- a/Engine/Source/Runtime/Experimental/Chaos/Private/Chaos/PBDJointConstraints.cpp
+++ b/Engine/Source/Runtime/Experimental/Chaos/Private/Chaos/PBDJointConstraints.cpp
@@ -406,7 +406,7 @@ namespace Chaos
 		ConstraintStates[ConstraintIndex].ParticleLevels = ParticleLevels;
 		bRequiresSort = bRequiresSort || (NewLevel != PreviousLevel);
 	}
-
+	
 
 	void FPBDJointConstraints::UpdatePositionBasedState(const FReal Dt)
 	{
@@ -415,8 +415,6 @@ namespace Chaos
 			SortConstraints();
 			bRequiresSort = false;
 		}
-
-		PrepareConstraints(Dt);
 	}
 
 
@@ -454,6 +452,15 @@ namespace Chaos
 		}
 	}
 
+
+	void FPBDJointConstraints::UnprepareConstraints(FReal Dt)
+	{
+		if (!ChaosJoint_UseCholeskySolver)
+		{
+			ConstraintSolvers.Empty();
+		}
+	}
+
 	
 	void FPBDJointConstraints::CalculateConstraintSpace(int32 ConstraintIndex, FVec3& OutX0, FMatrix33& OutR0, FVec3& OutX1, FMatrix33& OutR1) const
 	{
diff --git a/Engine/Source/Runtime/Experimental/Chaos/Private/Chaos/PBDRigidsEvolutionGBF.cpp b/Engine/Source/Runtime/Experimental/Chaos/Private/Chaos/PBDRigidsEvolutionGBF.cpp
--- a/Engine/Source/Runtime/Experimental/Chaos/Private/Chaos/PBDRigidsEvolutionGBF.cpp
+++ b/Engine/Source/Runtime/Experimental/Chaos/Private/Chaos/PBDRigidsEvolutionGBF.cpp
@@ -54,6 +54,8 @@ FAutoConsoleVariableRef CVarBoundsThicknessVelocityMultiplier(TEXT("p.CollisionB
 DECLARE_CYCLE_STAT(TEXT("TPBDRigidsEvolutionGBF::AdvanceOneTimeStep"), STAT_Evolution_AdvanceOneTimeStep, STATGROUP_Chaos);
 DECLARE_CYCLE_STAT(TEXT("TPBDRigidsEvolutionGBF::Integrate"), STAT_Evolution_Integrate, STATGROUP_Chaos);
 DECLARE_CYCLE_STAT(TEXT("TPBDRigidsEvolutionGBF::KinematicTargets"), STAT_Evolution_KinematicTargets, STATGROUP_Chaos);
+DECLARE_CYCLE_STAT(TEXT("TPBDRigidsEvolutionGBF::PrepareConstraints"), STAT_Evolution_PrepareConstraints, STATGROUP_Chaos);
+DECLARE_CYCLE_STAT(TEXT("TPBDRigidsEvolutionGBF::UnprepareConstraints"), STAT_Evolution_UnprepareConstraints, STATGROUP_Chaos);
 DECLARE_CYCLE_STAT(TEXT("TPBDRigidsEvolutionGBF::ApplyConstraints"), STAT_Evolution_ApplyConstraints, STATGROUP_Chaos);
 DECLARE_CYCLE_STAT(TEXT("TPBDRigidsEvolutionGBF::UpdateVelocities"), STAT_Evolution_UpdateVelocites, STATGROUP_Chaos);
 DECLARE_CYCLE_STAT(TEXT("TPBDRigidsEvolutionGBF::ApplyPushOut"), STAT_Evolution_ApplyPushOut, STATGROUP_Chaos);
@@ -177,6 +179,11 @@ void TPBDRigidsEvolutionGBF<T, d>::AdvanceOneTimeStep(const T Dt, const T StepFr
 		PostDetectCollisionsCallback();
 	}
 
+	{
+		SCOPE_CYCLE_COUNTER(STAT_Evolution_PrepareConstraints);
+		PrepareConstraints(Dt);
+	}
+
 	{
 		SCOPE_CYCLE_COUNTER(STAT_Evolution_CreateConstraintGraph);
 		CreateConstraintGraph();
@@ -264,6 +271,12 @@ void TPBDRigidsEvolutionGBF<T, d>::AdvanceOneTimeStep(const T Dt, const T StepFr
 			SleepedIslands[Island] = GetConstraintGraph().SleepInactive(Island, PhysicsMaterials);
 		});
 	}
+
+	{
+		SCOPE_CYCLE_COUNTER(STAT_Evolution_UnprepareConstraints);
+		UnprepareConstraints(Dt);
+	}
+
 	{
 		SCOPE_CYCLE_COUNTER(STAT_Evolution_DeactivateSleep);
 		for (int32 Island = 0; Island < GetConstraintGraph().NumIslands(); ++Island)
diff --git a/Engine/Source/Runtime/Experimental/Chaos/Public/Chaos/Evolution/PBDMinEvolution.h b/Engine/Source/Runtime/Experimental/Chaos/Public/Chaos/Evolution/PBDMinEvolution.h
--- a/Engine/Source/Runtime/Experimental/Chaos/Public/Chaos/Evolution/PBDMinEvolution.h
+++ b/Engine/Source/Runtime/Experimental/Chaos/Public/Chaos/Evolution/PBDMinEvolution.h
@@ -83,6 +83,8 @@ namespace Chaos
 		void Integrate(FReal Dt);
 		void ApplyKinematicTargets(FReal Dt, FReal StepFraction);
 		void DetectCollisions(FReal Dt);
+		void PrepareConstraints(FReal Dt);
+		void UnprepareConstraints(FReal Dt);
 		void ApplyConstraints(FReal Dt);
 		void UpdateVelocities(FReal Dt);
 		void ApplyPushOutConstraints(FReal Dt);
diff --git a/Engine/Source/Runtime/Experimental/Chaos/Public/Chaos/Joint/PBDJointSolverGaussSeidel.h b/Engine/Source/Runtime/Experimental/Chaos/Public/Chaos/Joint/PBDJointSolverGaussSeidel.h
--- a/Engine/Source/Runtime/Experimental/Chaos/Public/Chaos/Joint/PBDJointSolverGaussSeidel.h
+++ b/Engine/Source/Runtime/Experimental/Chaos/Public/Chaos/Joint/PBDJointSolverGaussSeidel.h
@@ -95,8 +95,6 @@ namespace Chaos
 
 		void UpdateDerivedState(const int32 BodyIndex);
 		void UpdateDerivedState();
-		void InitAccumulatedDeltas();
-		void ApplyAccumulatedDeltas();
 
 		void ApplyPositionConstraints(
 			const FReal Dt,
@@ -327,12 +325,6 @@ namespace Chaos
 		FVec3 Vs[MaxConstrainedBodies];				// World-space particle CoM velocities
 		FVec3 Ws[MaxConstrainedBodies];				// World-space particle CoM angular velocities
 
-		// Accumulated world-body state deltas
-		FVec3 DPs[MaxConstrainedBodies];			// World-space particle CoM positions
-		FVec3 DRs[MaxConstrainedBodies];			// World-space particle CoM rotations (in axis-angle form)
-		FVec3 DVs[MaxConstrainedBodies];			// World-space particle CoM velocities
-		FVec3 DWs[MaxConstrainedBodies];			// World-space particle CoM angular velocities
-
 		// XPBD Previous iteration world-space body state
 		FVec3 PrevPs[MaxConstrainedBodies];			// World-space particle CoM positions
 		FReal TwistDriveLambda;						// XPBD constraint multipliers (net applied constraint-space deltas)
diff --git a/Engine/Source/Runtime/Experimental/Chaos/Public/Chaos/PBDCollisionConstraints.h b/Engine/Source/Runtime/Experimental/Chaos/Public/Chaos/PBDCollisionConstraints.h
--- a/Engine/Source/Runtime/Experimental/Chaos/Public/Chaos/PBDCollisionConstraints.h
+++ b/Engine/Source/Runtime/Experimental/Chaos/Public/Chaos/PBDCollisionConstraints.h
@@ -127,6 +127,10 @@ public:
 	// General Rule API
 	//
 
+	void PrepareConstraints(FReal Dt) {}
+
+	void UnprepareConstraints(FReal Dt) {}
+
 	/**
 	 * Generate all contact constraints.
 	 */
diff --git a/Engine/Source/Runtime/Experimental/Chaos/Public/Chaos/PBDConstraintRule.h b/Engine/Source/Runtime/Experimental/Chaos/Public/Chaos/PBDConstraintRule.h
--- a/Engine/Source/Runtime/Experimental/Chaos/Public/Chaos/PBDConstraintRule.h
+++ b/Engine/Source/Runtime/Experimental/Chaos/Public/Chaos/PBDConstraintRule.h
@@ -31,6 +31,12 @@ namespace Chaos
 			return L.GetPriority() < R.GetPriority();
 		}
 
+		/** Called once per frame before Apply. Can be used to prepare caches etc. */
+		virtual void PrepareConstraints(FReal Dt) {}
+
+		/** Called once per frame after Apply. Should be used to release any transient stores created in PrepareConstraints. */
+		virtual void UnprepareConstraints(FReal Dt) {}
+
 	protected:
 		int32 Priority;
 	};
@@ -60,6 +66,16 @@ namespace Chaos
 		{
 		}
 
+		virtual void PrepareConstraints(FReal Dt) override
+		{
+			Constraints.PrepareConstraints(Dt);
+		}
+
+		virtual void UnprepareConstraints(FReal Dt) override
+		{
+			Constraints.UnprepareConstraints(Dt);
+		}
+
 		virtual void UpdatePositionBasedState(const FReal Dt) override
 		{
 			return Constraints.UpdatePositionBasedState(Dt);
@@ -133,6 +149,16 @@ namespace Chaos
 
 		TPBDConstraintGraphRuleImpl(FConstraints& InConstraints, int32 InPriority);
 
+		virtual void PrepareConstraints(FReal Dt) override
+		{
+			Constraints.PrepareConstraints(Dt);
+		}
+
+		virtual void UnprepareConstraints(FReal Dt) override
+		{
+			Constraints.UnprepareConstraints(Dt);
+		}
+
 		virtual void BindToGraph(FPBDConstraintGraph& InContactGraph, uint32 InContainerId) override;
 
 		virtual void UpdatePositionBasedState(const FReal Dt) override;
diff --git a/Engine/Source/Runtime/Experimental/Chaos/Public/Chaos/PBDJointConstraints.h b/Engine/Source/Runtime/Experimental/Chaos/Public/Chaos/PBDJointConstraints.h
--- a/Engine/Source/Runtime/Experimental/Chaos/Public/Chaos/PBDJointConstraints.h
+++ b/Engine/Source/Runtime/Experimental/Chaos/Public/Chaos/PBDJointConstraints.h
@@ -131,6 +131,9 @@ namespace Chaos
 		// General Rule API
 		//
 
+		void PrepareConstraints(FReal Dt);
+		void UnprepareConstraints(FReal Dt);
+
 		void UpdatePositionBasedState(const FReal Dt);
 
 		//
@@ -155,8 +158,6 @@ namespace Chaos
 	private:
 		friend class FPBDJointConstraintHandle;
 
-		void PrepareConstraints(FReal Dt);
-
 		void GetConstrainedParticleIndices(const int32 ConstraintIndex, int32& Index0, int32& Index1) const;
 		void CalculateConstraintSpace(int32 ConstraintIndex, FVec3& OutX0, FMatrix33& OutR0, FVec3& OutX1, FMatrix33& OutR1) const;
 		void UpdateParticleState(TPBDRigidParticleHandle<FReal, 3>* Rigid, const FReal Dt, const FVec3& P, const FRotation3& Q, const bool bUpdateVelocity = true);
diff --git a/Engine/Source/Runtime/Experimental/Chaos/Public/Chaos/PBDPositionConstraints.h b/Engine/Source/Runtime/Experimental/Chaos/Public/Chaos/PBDPositionConstraints.h
--- a/Engine/Source/Runtime/Experimental/Chaos/Public/Chaos/PBDPositionConstraints.h
+++ b/Engine/Source/Runtime/Experimental/Chaos/Public/Chaos/PBDPositionConstraints.h
@@ -163,6 +163,14 @@ namespace Chaos
 		// Island Rule API
 		//
 
+		void PrepareConstraints(FReal Dt)
+		{
+		}
+
+		void UnprepareConstraints(FReal Dt)
+		{
+		}
+
 		void UpdatePositionBasedState(const T Dt)
 		{
 		}
diff --git a/Engine/Source/Runtime/Experimental/Chaos/Public/Chaos/PBDRigidDynamicSpringConstraints.h b/Engine/Source/Runtime/Experimental/Chaos/Public/Chaos/PBDRigidDynamicSpringConstraints.h
--- a/Engine/Source/Runtime/Experimental/Chaos/Public/Chaos/PBDRigidDynamicSpringConstraints.h
+++ b/Engine/Source/Runtime/Experimental/Chaos/Public/Chaos/PBDRigidDynamicSpringConstraints.h
@@ -173,6 +173,10 @@ namespace Chaos
 		// Island Rule API
 		//
 
+		void PrepareConstraints(FReal Dt) {}
+
+		void UnprepareConstraints(FReal Dt) {}
+
 		void UpdatePositionBasedState(const T Dt);
 
 		void Apply(const T Dt, const TArray<FConstraintContainerHandle*>& InConstraintHandles, const int32 It, const int32 NumIts)
diff --git a/Engine/Source/Runtime/Experimental/Chaos/Public/Chaos/PBDRigidSpringConstraints.h b/Engine/Source/Runtime/Experimental/Chaos/Public/Chaos/PBDRigidSpringConstraints.h
--- a/Engine/Source/Runtime/Experimental/Chaos/Public/Chaos/PBDRigidSpringConstraints.h
+++ b/Engine/Source/Runtime/Experimental/Chaos/Public/Chaos/PBDRigidSpringConstraints.h
@@ -176,6 +176,14 @@ namespace Chaos
 		// Island Rule API
 		//
 
+		void PrepareConstraints(FReal Dt)
+		{
+		}
+
+		void UnprepareConstraints(FReal Dt)
+		{
+		}
+
 		void UpdatePositionBasedState(const T Dt)
 		{
 		}
diff --git a/Engine/Source/Runtime/Experimental/Chaos/Public/Chaos/PBDRigidsEvolution.h b/Engine/Source/Runtime/Experimental/Chaos/Public/Chaos/PBDRigidsEvolution.h
--- a/Engine/Source/Runtime/Experimental/Chaos/Public/Chaos/PBDRigidsEvolution.h
+++ b/Engine/Source/Runtime/Experimental/Chaos/Public/Chaos/PBDRigidsEvolution.h
@@ -414,6 +414,22 @@ class TPBDRigidsEvolutionBase
 		}
 	}
 
+	void PrepareConstraints(const T Dt)
+	{
+		for (FPBDConstraintGraphRule* ConstraintRule : ConstraintRules)
+		{
+			ConstraintRule->PrepareConstraints(Dt);
+		}
+	}
+
+	void UnprepareConstraints(const T Dt)
+	{
+		for (FPBDConstraintGraphRule* ConstraintRule : ConstraintRules)
+		{
+			ConstraintRule->UnprepareConstraints(Dt);
+		}
+	}
+
 	void UpdateAccelerationStructures(int32 Island)
 	{
 		for (FPBDConstraintGraphRule* ConstraintRule : ConstraintRules)
