Fix CIS (PBDRigidsEvolutionGBF).
Fix ISPC cvar use.

diff --git a/Engine/Source/Runtime/Experimental/Chaos/Private/Chaos/Joint/PBDJointSolverGaussSeidel.cpp b/Engine/Source/Runtime/Experimental/Chaos/Private/Chaos/Joint/PBDJointSolverGaussSeidel.cpp
--- a/Engine/Source/Runtime/Experimental/Chaos/Private/Chaos/Joint/PBDJointSolverGaussSeidel.cpp
+++ b/Engine/Source/Runtime/Experimental/Chaos/Private/Chaos/Joint/PBDJointSolverGaussSeidel.cpp
@@ -11,7 +11,9 @@
 
 //#pragma optimize("", off)
 
-#if UE_BUILD_SHIPPING
+#if !INTEL_ISPC
+const bool bChaos_Joint_ISPC_Enabled = false;
+#elif UE_BUILD_SHIPPING
 const bool bChaos_Joint_ISPC_Enabled = true;
 #else
 bool bChaos_Joint_ISPC_Enabled = true;
diff --git a/Engine/Source/Runtime/Experimental/Chaos/Public/Chaos/PBDRigidsEvolutionGBF.h b/Engine/Source/Runtime/Experimental/Chaos/Public/Chaos/PBDRigidsEvolutionGBF.h
--- a/Engine/Source/Runtime/Experimental/Chaos/Public/Chaos/PBDRigidsEvolutionGBF.h
+++ b/Engine/Source/Runtime/Experimental/Chaos/Public/Chaos/PBDRigidsEvolutionGBF.h
@@ -92,6 +92,8 @@ namespace Chaos
 		CHAOS_API void Advance(const T Dt, const T MaxStepDt, const int32 MaxSteps);
 		CHAOS_API void AdvanceOneTimeStep(const T dt, const T StepFraction = (T)1.0);
 
+		using Base::PrepareConstraints;
+		using Base::UnprepareConstraints;
 		using Base::ApplyConstraints;
 		using Base::ApplyPushOut;
 
