No code changes! Rearranges PBDRigidClustering.cpp functions to match header order, to make a subsequent refactoring more managable and diffable.

diff --git a/Engine/Source/Runtime/Experimental/Chaos/Private/Chaos/PBDRigidClustering.cpp b/Engine/Source/Runtime/Experimental/Chaos/Private/Chaos/PBDRigidClustering.cpp
--- a/Engine/Source/Runtime/Experimental/Chaos/Private/Chaos/PBDRigidClustering.cpp
+++ b/Engine/Source/Runtime/Experimental/Chaos/Private/Chaos/PBDRigidClustering.cpp
@@ -70,6 +70,387 @@
 	int32 DeactivateClusterChildren = 0;
 	FAutoConsoleVariableRef CVarDeactivateClusterChildren(TEXT("p.DeactivateClusterChildren"), DeactivateClusterChildren, TEXT("If children should be decativated when broken and put into another cluster."));
 
+	//==========================================================================
+	// Free helper functions
+	//==========================================================================
+
+	template<class T, int d>
+	TVector<T, d> GetContactLocation(const TRigidBodyPointContactConstraint<T, d>& Contact)
+	{
+		return Contact.GetLocation();
+	}
+
+	template<class T, int d>
+	TVector<T, d> GetContactLocation(const TRigidBodyContactConstraintPGS<T, d>& Contact)
+	{
+		// @todo(mlentine): Does the exact point matter?
+		T MinPhi = FLT_MAX;
+		TVector<T, d> MinLoc;
+		for (int32 i = 0; i < Contact.Phi.Num(); ++i)
+		{
+			if (Contact.Phi[i] < MinPhi)
+			{
+				MinPhi = Contact.Phi[i];
+				MinLoc = Contact.Location[i];
+			}
+		}
+		return MinLoc;
+	}
+
+	template<class T, int d>
+	T CalculatePseudoMomentum(const TPBDRigidClusteredParticles<T, d>& InParticles, const uint32 Index)
+	{
+		TVector<T, d> LinearPseudoMomentum = (InParticles.X(Index) - InParticles.P(Index)) * InParticles.M(Index);
+		TRotation<T, d> Delta = InParticles.R(Index) * InParticles.Q(Index).Inverse();
+		TVector<T, d> Axis;
+		T Angle;
+		Delta.ToAxisAndAngle(Axis, Angle);
+		TVector<T, d> AngularPseudoMomentum = InParticles.I(Index) * (Axis * Angle);
+		return LinearPseudoMomentum.Size() + AngularPseudoMomentum.Size();
+	}
+
+	DECLARE_CYCLE_STAT(TEXT("TPBDRigidClustering<>::RewindAndEvolve<BGF>()"), STAT_RewindAndEvolve_BGF, STATGROUP_Chaos);
+	template<class T, int d>
+	void RewindAndEvolve(TPBDRigidsEvolutionGBF<T, d>& Evolution, TPBDRigidClusteredParticles<T, d>& InParticles, const TSet<int32>& IslandsToRecollide, const TSet<uint32>& AllActivatedChildren, const T Dt, TPBDCollisionConstraints<T, d>& CollisionRule)
+	{
+		SCOPE_CYCLE_COUNTER(STAT_RewindAndEvolve_BGF);
+		// Rewind active particles
+		const TArray<int32> IslandsToRecollideArray = IslandsToRecollide.Array();
+		PhysicsParallelFor(IslandsToRecollideArray.Num(), [&](int32 Idx) {
+			int32 Island = IslandsToRecollideArray[Idx];
+			auto Particles = Evolution.GetIslandParticles(Island); // copy
+			for (int32 ArrayIdx = Particles.Num() - 1; ArrayIdx >= 0; --ArrayIdx)
+			{
+				auto PBDRigid = Particles[ArrayIdx]->CastToRigidParticle();
+				if(PBDRigid && PBDRigid->ObjectState() == EObjectStateType::Dynamic)
+				{
+					if (!PBDRigid->Sleeping() && !PBDRigid->Disabled())
+					{
+						PBDRigid->P() = PBDRigid->X();
+						PBDRigid->Q() = PBDRigid->R();
+						PBDRigid->V() = PBDRigid->PreV();
+						PBDRigid->W() = PBDRigid->PreW();
+						continue;
+					}
+				}
+				Particles.RemoveAtSwap(ArrayIdx);
+			}
+			Evolution.Integrate(MakeHandleView(Particles), Dt);
+		});
+
+		TSet<TGeometryParticleHandle<T, d>*> AllIslandParticles;
+		for (int32 Island = 0; Island < Evolution.NumIslands(); ++Island)
+		{
+			const auto& ParticleIndices = Evolution.GetIslandParticles(Island);
+			for (const auto Particle : ParticleIndices)
+			{
+				auto PBDRigid = Particle->CastToRigidParticle();
+				if(PBDRigid && PBDRigid->ObjectState() == EObjectStateType::Dynamic)
+				{
+					bool bDisabled = PBDRigid->Disabled();
+
+					// #TODO - Have to repeat checking out whether the particle is disabled matching the PFor above.
+					// Move these into shared array so we only process it once
+					if (!AllIslandParticles.Contains(Particle) && !bDisabled)
+					{
+						AllIslandParticles.Add(Particle);
+					}
+				}
+			}
+		}
+
+		const bool bRewindOnDeclusterSolve = ChildrenInheritVelocity < 1.f;
+		if (bRewindOnDeclusterSolve)
+		{
+			// @todo(mlentine): We can precompute internal constraints which can filter some from the narrow phase tests but may not help much
+
+#if CHAOS_PARTICLEHANDLE_TODO
+			CollisionRule.UpdateConstraints(InParticles, Evolution.GetNonDisabledIndices(), Dt, AllActivatedChildren, AllIslandParticles.Array());
+#else
+			CollisionRule.UpdateConstraints(Dt, AllIslandParticles);	//this seems wrong
+#endif
+
+			Evolution.InitializeAccelerationStructures();
+
+			// Resolve collisions
+			PhysicsParallelFor(IslandsToRecollide.Num(), [&](int32 Island) {
+				// @todo(mlentine): This is heavy handed and probably can be simplified as we know only a little bit changed.
+				Evolution.UpdateAccelerationStructures(Island);
+				Evolution.ApplyConstraints(Dt, Island);
+				// @todo(ccaulfield): should we also update velocities here? Evolution does...
+				Evolution.ApplyPushOut(Dt, Island);
+				// @todo(ccaulfield): support sleep state update on evolution
+				//Evolution.UpdateSleepState(Island);
+			});
+		}
+	}
+
+#if CHAOS_PARTICLEHANDLE_TODO
+	DECLARE_CYCLE_STAT(TEXT("TPBDRigidClustering<>::RewindAndEvolve<PSG>()"), STAT_RewindAndEvolve_PSG, STATGROUP_Chaos);
+	template<class T, int d>
+	void RewindAndEvolve(TPBDRigidsEvolutionPGS<T, d>& Evolution, TPBDRigidClusteredParticles<T, d>& InParticles, const TSet<int32>& IslandsToRecollide, const TSet<uint32>& AllActivatedChildren, const T Dt, TPBDCollisionConstraintPGS<T, d>& CollisionRule)
+	{
+		SCOPE_CYCLE_COUNTER(STAT_RewindAndEvolve_PSG);
+		// Rewind active particles
+		TArray<int32>& NonDisabledIndices = Evolution.GetNonDisabledIndices();
+		PhysicsParallelFor(IslandsToRecollide.Num(), [&](int32 Island) {
+			TArray<int32> ParticleIndices = Evolution.GetIslandParticles(Island); // copy
+			for (int32 ArrayIdx = ParticleIndices.Num() - 1; ArrayIdx >= 0; --ArrayIdx)
+			{
+				int32 Index = ParticleIndices[ArrayIdx];
+				if (InParticles.Sleeping(Index) || InParticles.Disabled(Index))
+				{
+					ParticleIndices.RemoveAtSwap(ArrayIdx);
+				}
+				else
+				{
+					InParticles.P(Index) = InParticles.X(Index);
+					InParticles.Q(Index) = InParticles.R(Index);
+					InParticles.V(Index) = InParticles.PreV(Index);
+					InParticles.W(Index) = InParticles.PreW(Index);
+				}
+			}
+			Evolution.IntegrateV(ParticleIndices, Dt);
+		});
+
+		TSet<uint32> AllIslandParticles;
+		for (const auto& Island : IslandsToRecollide)
+		{
+			const TArray<int32>& ParticleIndices = Evolution.GetIslandParticles(Island);
+			for (const auto& Index : ParticleIndices)
+			{
+				if (InParticles.Disabled(Index) == false) //HACK: cluster code is incorrectly adding disabled children
+				{
+					if (!AllIslandParticles.Contains(Index))
+					{
+						AllIslandParticles.Add(Index);
+						NonDisabledIndices.Add(Index);
+					}
+				}
+				else
+				{
+					//FPlatformMisc::DebugBreak();
+				}
+			}
+		}
+
+		// @todo(mlentine): We can precompute internal constraints which can filter some from the narrow phase tests but may not help much
+		CollisionRule.UpdateConstraints(InParticles, Dt, AllActivatedChildren, AllIslandParticles.Array());
+
+		Evolution.InitializeAccelerationStructures();
+
+		PhysicsParallelFor(Evolution.NumIslands(), [&](int32 Island) {
+			Evolution.UpdateAccelerationStructures(Island);
+			Evolution.ApplyConstraints(Dt, Island);
+		});
+		PhysicsParallelFor(Evolution.NumIslands(), [&](int32 Island) {
+			const TArray<int32>& ParticleIndices = Evolution.GetIslandParticles(Island);
+			Evolution.IntegrateX(ParticleIndices, Dt);
+		});
+
+		// @todo(mlentine): Need to enforce constraints
+		PhysicsParallelFor(Evolution.NumIslands(), [&](int32 Island) {
+			Evolution.ApplyPushOut(Dt, Island);
+		});
+	}
+#endif
+
+	DECLARE_CYCLE_STAT(TEXT("TPBDRigidClustering<>::UpdateClusterMassProperties()"), STAT_UpdateClusterMassProperties, STATGROUP_Chaos);
+	template<class T, int d>
+	void UpdateClusterMassProperties(TPBDRigidClusteredParticles<T, d>& Particles, const TArray<uint32>& Children, const uint32 NewIndex, const TRigidTransform<T, d>* ForceMassOrientation, const TArrayCollectionArray<TUniquePtr<TMultiChildProxyData<T, d>>>* MMultiChildProxyData, const TArrayCollectionArray<FMultiChildProxyId>* MMultiChildProxyId)
+	{
+		SCOPE_CYCLE_COUNTER(STAT_UpdateClusterMassProperties);
+		check(0 <= NewIndex && NewIndex < Particles.Size());
+		check(Children.Num());
+
+		Particles.X(NewIndex) = TVector<T, d>(0);
+		Particles.R(NewIndex) = TRotation<T, d>(FQuat::MakeFromEuler(TVector<T, d>(0)));
+		Particles.V(NewIndex) = TVector<T, d>(0);
+		Particles.W(NewIndex) = TVector<T, d>(0);
+		Particles.M(NewIndex) = 0;
+		Particles.I(NewIndex) = PMatrix<T, d, d>(0);
+
+		auto GetMultiChildProxy = [MMultiChildProxyId, MMultiChildProxyData](const int32 OriginalChild, const TMultiChildProxyData<T, d>*& ProxyData) -> int32
+		{
+			const int32 MultiChildProxyId = MassPropertiesFromMultiChildProxy && MMultiChildProxyId ? (*MMultiChildProxyId)[OriginalChild].Id : INDEX_NONE;
+			if (MultiChildProxyId != INDEX_NONE)
+			{
+				ProxyData = (*MMultiChildProxyData)[MultiChildProxyId].Get();
+				return ProxyData ? MultiChildProxyId : INDEX_NONE;
+			}
+			return INDEX_NONE;
+		};
+
+		bool bHasChild = false;
+		bool bHasProxyChild = false;
+		for (const auto OriginalChild : Children)
+		{
+			const TMultiChildProxyData<T, d>* ProxyData;
+			const int32 MultiChildProxyId = GetMultiChildProxy(OriginalChild, ProxyData);
+			TVector<T, d> ChildPosition;
+			TRotation<T, d> ChildRotation;
+			int32 Child;
+
+			if (MultiChildProxyId == INDEX_NONE)
+			{
+				Child = OriginalChild;
+				ChildPosition = Particles.X(Child);
+				ChildRotation = Particles.R(Child);
+			}
+			else if (ProxyData->KeyChild == OriginalChild)
+			{
+				Child = MultiChildProxyId;
+
+				const TRigidTransform<T, d> ProxyWorldTM = ProxyData->RelativeToKeyChild * TRigidTransform<T, d>(Particles.X(OriginalChild), Particles.R(OriginalChild));
+				ChildPosition = ProxyWorldTM.GetLocation();
+				ChildRotation = ProxyWorldTM.GetRotation();
+			}
+			else
+			{
+				continue; //using a proxy but we are not the key child
+			}
+
+			const auto& ChildMass = Particles.M(Child);
+			const PMatrix<T, d, d> ChildWorldSpaceI = (ChildRotation * FMatrix::Identity) * Particles.I(Child) * (ChildRotation * FMatrix::Identity).GetTransposed();
+			if (ChildWorldSpaceI.ContainsNaN())
+			{
+				continue;
+			}
+			bHasProxyChild = true;
+			bHasChild = true;
+			Particles.I(NewIndex) += ChildWorldSpaceI;
+			Particles.M(NewIndex) += ChildMass;
+			Particles.X(NewIndex) += ChildPosition * ChildMass;
+			Particles.V(NewIndex) += Particles.V(OriginalChild) * ChildMass; //use original child for veocities because we don't simulate proxy
+			Particles.W(NewIndex) += Particles.W(OriginalChild) * ChildMass;
+		}
+		if (!ensure(bHasProxyChild))
+		{
+			for (const auto OriginalChild : Children)
+			{
+				TVector<T, d> ChildPosition;
+				TRotation<T, d> ChildRotation;
+				int32 Child;
+
+				Child = OriginalChild;
+				ChildPosition = Particles.X(Child);
+				ChildRotation = Particles.R(Child);
+
+				const auto& ChildMass = Particles.M(Child);
+				const PMatrix<T, d, d> ChildWorldSpaceI = (ChildRotation * FMatrix::Identity) * Particles.I(Child) * (ChildRotation * FMatrix::Identity).GetTransposed();
+				if (ChildWorldSpaceI.ContainsNaN())
+				{
+					continue;
+				}
+				bHasChild = true;
+				Particles.I(NewIndex) += ChildWorldSpaceI;
+				Particles.M(NewIndex) += ChildMass;
+				Particles.X(NewIndex) += ChildPosition * ChildMass;
+				Particles.V(NewIndex) += Particles.V(OriginalChild) * ChildMass;	//use original child for veocities because we don't simulate proxy
+				Particles.W(NewIndex) += Particles.W(OriginalChild) * ChildMass;
+			}
+		}
+		for (int32 i = 0; i < d; i++)
+		{
+			const PMatrix<T, d, d>& InertiaTensor = Particles.I(NewIndex);
+			if (InertiaTensor.GetColumn(i)[i] < SMALL_NUMBER)
+			{
+				Particles.I(NewIndex) = PMatrix<T, d, d>(1.f, 1.f, 1.f);
+				break;
+			}
+		}
+
+		if (!ensure(bHasChild) || !ensure(Particles.M(NewIndex) > SMALL_NUMBER))
+		{
+			Particles.M(NewIndex) = 1;
+			Particles.X(NewIndex) = TVector<T, d>(0);
+			Particles.V(NewIndex) = TVector<T, d>(0);
+			Particles.PreV(NewIndex) = Particles.V(NewIndex);
+			Particles.InvM(NewIndex) = 1;
+			Particles.P(NewIndex) = Particles.X(NewIndex);
+			Particles.W(NewIndex) = TVector<T, d>(0);
+			Particles.PreW(NewIndex) = Particles.W(NewIndex);
+			Particles.R(NewIndex) = TRotation<T, d>(FMatrix::Identity);
+			Particles.Q(NewIndex) = Particles.R(NewIndex);
+			Particles.I(NewIndex) = FMatrix::Identity;
+			Particles.InvI(NewIndex) = FMatrix::Identity;
+			return;
+		}
+
+		check(Particles.M(NewIndex) > SMALL_NUMBER);
+
+		Particles.X(NewIndex) /= Particles.M(NewIndex);
+		Particles.V(NewIndex) /= Particles.M(NewIndex);
+		Particles.PreV(NewIndex) = Particles.V(NewIndex);
+		Particles.InvM(NewIndex) = 1. / Particles.M(NewIndex);
+		if (ForceMassOrientation)
+		{
+			Particles.X(NewIndex) = ForceMassOrientation->GetLocation();
+		}
+		Particles.P(NewIndex) = Particles.X(NewIndex);
+		for (const auto OriginalChild : Children)
+		{
+			int32 Child;
+			const TMultiChildProxyData<T, d>* ProxyData = nullptr;
+			const int32 MultiChildProxyId = bHasProxyChild ? GetMultiChildProxy(OriginalChild, ProxyData) : INDEX_NONE;
+
+			TVector<T, d> ChildPosition;
+			if (MultiChildProxyId == INDEX_NONE)
+			{
+				Child = OriginalChild;
+				ChildPosition = Particles.X(Child);
+			}
+			else if (ProxyData->KeyChild == OriginalChild)
+			{
+				Child = MultiChildProxyId;
+				const TRigidTransform<T, d> ProxyWorldTM = ProxyData->RelativeToKeyChild * TRigidTransform<T, d>(Particles.X(OriginalChild), Particles.R(OriginalChild));
+				ChildPosition = ProxyWorldTM.GetLocation();
+			}
+			else
+			{
+				continue; //using a proxy but we are not the key child
+			}
+
+			TVector<T, d> ParentToChild = ChildPosition - Particles.X(NewIndex);
+
+			const auto& ChildMass = Particles.M(Child);
+			Particles.W(NewIndex) += TVector<T, d>::CrossProduct(ParentToChild, Particles.V(OriginalChild) * ChildMass); //taking v from original child since we are not actually simulating the proxy child
+			{
+				const T& p0 = ParentToChild[0];
+				const T& p1 = ParentToChild[1];
+				const T& p2 = ParentToChild[2];
+				const T& m = Particles.M(Child);
+				Particles.I(NewIndex) += PMatrix<T, d, d>(m * (p1 * p1 + p2 * p2), -m * p1 * p0, -m * p2 * p0, m * (p2 * p2 + p0 * p0), -m * p2 * p1, m * (p1 * p1 + p0 * p0));
+			}
+		}
+		if (Particles.I(NewIndex).ContainsNaN())
+		{
+			Particles.I(NewIndex) = PMatrix<T, d, d>(1.f, 1.f, 1.f);
+		}
+		for (int32 i = 0; i < d; i++)
+		{
+			const PMatrix<T, d, d>& InertiaTensor = Particles.I(NewIndex);
+			if (InertiaTensor.GetColumn(i)[i] < SMALL_NUMBER)
+			{
+				Particles.I(NewIndex) = PMatrix<T, d, d>(1.f, 1.f, 1.f);
+				break;
+			}
+		}
+		Particles.W(NewIndex) /= Particles.M(NewIndex);
+		Particles.PreW(NewIndex) = Particles.W(NewIndex);
+		Particles.R(NewIndex) = Chaos::TransformToLocalSpace<T, d>(Particles.I(NewIndex));
+		if (ForceMassOrientation)
+		{
+			Particles.R(NewIndex) = ForceMassOrientation->GetRotation();
+		}
+		Particles.Q(NewIndex) = Particles.R(NewIndex);
+		Particles.InvI(NewIndex) = Particles.I(NewIndex).Inverse();
+	}
+
+	//==========================================================================
+	// TPBDRigidClustering
+	//==========================================================================
+
 	template<class FPBDRigidsEvolution, class FPBDCollisionConstraint, class T, int d>
 	TPBDRigidClustering<FPBDRigidsEvolution, FPBDCollisionConstraint, T, d>::TPBDRigidClustering(FPBDRigidsEvolution& InEvolution, TPBDRigidClusteredParticles<T, d>& InParticles)
 		: MEvolution(InEvolution)
@@ -340,54 +721,11 @@
 		return ActivatedChildren;
 	}
 
-	DECLARE_CYCLE_STAT(TEXT("TPBDRigidClustering<>::ReleaseClusterParticles(LIST)"), STAT_ReleaseClusterParticles_LIST, STATGROUP_Chaos);
-	template<class FPBDRigidsEvolution, class FPBDCollisionConstraint, class T, int d>
-	TSet<uint32> TPBDRigidClustering<FPBDRigidsEvolution, FPBDCollisionConstraint, T, d>::ReleaseClusterParticles(const TArray<uint32>& ChildrenParticles)
-	{
-		SCOPE_CYCLE_COUNTER(STAT_ReleaseClusterParticles_LIST);
-		TSet<uint32> ActivatedBodies;
-		bool bFound = false;
-		if (ChildrenParticles.Num())
-		{
-			uint32 ClusterIdx = 0;
-			//todo(ocohen): refactor incoming, for now just assume these all belong to same cluster and hack strain array
-			TArray<float> FakeStrain;
-			FakeStrain.Init(0, MParticles.Size()); //this part especially sucks
-
-			bool bPreDoGenerateData = DoGenerateBreakingData;
-			DoGenerateBreakingData = false;
-
-			for (uint32 ChildIdx : ChildrenParticles)
-			{
-				if (MParticles.Disabled(ChildIdx) && MParticles.ClusterIds(ChildIdx).Id != INDEX_NONE)
-				{
-					if (ensure(!bFound || MParticles.ClusterIds(ChildIdx).Id == ClusterIdx))
-					{
-						bFound = true;
-						FakeStrain[ChildIdx] = FLT_MAX;
-						ClusterIdx = MParticles.ClusterIds(ChildIdx).Id;
-					}
-					else
-					{
-						break; //shouldn't be here
-					}
-				}
-			}
-
-			if (bFound)
-			{
-				ActivatedBodies.Append(ReleaseClusterParticles(MParticles.Handle(ClusterIdx), FakeStrain));
-			}
-
-			DoGenerateBreakingData = bPreDoGenerateData;
-		}
-
-		return ActivatedBodies;
-	}
-
 	DECLARE_CYCLE_STAT(TEXT("TPBDRigidClustering<>::ReleaseClusterParticles(STRAIN)"), STAT_ReleaseClusterParticles_STRAIN, STATGROUP_Chaos);
 	template<class FPBDRigidsEvolution, class FPBDCollisionConstraint, class T, int d>
-	TSet<uint32> TPBDRigidClustering<FPBDRigidsEvolution, FPBDCollisionConstraint, T, d>::ReleaseClusterParticles(TPBDRigidClusteredParticleHandle<T,d>* ClusteredParticle, const TArrayView<T>& StrainArray)
+	TSet<uint32> TPBDRigidClustering<FPBDRigidsEvolution, FPBDCollisionConstraint, T, d>::ReleaseClusterParticles(
+		TPBDRigidClusteredParticleHandle<T,d>* ClusteredParticle, 
+		const TArrayView<T>& StrainArray)
 	{
 		SCOPE_CYCLE_COUNTER(STAT_ReleaseClusterParticles_STRAIN);
 		FClusterMap& MParentToChildren = GetChildrenMap();
@@ -564,287 +902,51 @@
 		return ActivatedChildren;
 	}
 
+	DECLARE_CYCLE_STAT(TEXT("TPBDRigidClustering<>::ReleaseClusterParticles(LIST)"), STAT_ReleaseClusterParticles_LIST, STATGROUP_Chaos);
 	template<class FPBDRigidsEvolution, class FPBDCollisionConstraint, class T, int d>
-	void TPBDRigidClustering<FPBDRigidsEvolution, FPBDCollisionConstraint, T, d>::DisableCluster(TPBDRigidClusteredParticleHandle<T,d>* ClusteredParticle)
+	TSet<uint32> TPBDRigidClustering<FPBDRigidsEvolution, FPBDCollisionConstraint, T, d>::ReleaseClusterParticles(
+		const TArray<uint32>& ChildrenParticles)
 	{
-		// #note: we don't recursively descend to the children
-		FClusterMap& ParentToChildren = GetChildrenMap();
+		SCOPE_CYCLE_COUNTER(STAT_ReleaseClusterParticles_LIST);
+		TSet<uint32> ActivatedBodies;
+		bool bFound = false;
+		if (ChildrenParticles.Num())
+		{
+			uint32 ClusterIdx = 0;
+			//todo(ocohen): refactor incoming, for now just assume these all belong to same cluster and hack strain array
+			TArray<float> FakeStrain;
+			FakeStrain.Init(0, MParticles.Size()); //this part especially sucks
 
-		MEvolution.DisableParticle(ClusteredParticle);
+			bool bPreDoGenerateData = DoGenerateBreakingData;
+			DoGenerateBreakingData = false;
 
-		if (MoveClustersWhenDeactivated)
-		{
-			ClusteredParticle->P() -= FVector(0.f, 0.f, -10000.f); // HACK : Move them away to prevent reactivation. 
-			ClusteredParticle->X() -= FVector(0.f, 0.f, -10000.f); // HACK : Move them away to prevent reactivation. 
-			ClusteredParticle->V() = FVector(0.f);            // HACK : Move them away to prevent reactivation.
-		}
-
-		TopLevelClusterParents.Remove(ClusteredParticle->TransientParticleIndex());
-		ParentToChildren.Remove(ClusteredParticle);
-		ClusteredParticle->ClusterIds() = ClusterId();
-		ClusteredParticle->ClusterGroupIndex() = 0;
-		MActiveRemovalIndices.Remove(ClusteredParticle->TransientParticleIndex());
-	}
-
-	template<class FPBDRigidsEvolution, class FPBDCollisionConstraint, class T, int d>
-	void TPBDRigidClustering<FPBDRigidsEvolution, FPBDCollisionConstraint, T, d>::DisableParticleWithBreakEvent(uint32 ClusterIndex)
-	{
-		DisableCluster(MParticles.Handle(ClusterIndex));
-
-		if (DoGenerateBreakingData)
-		{
-			int32 NewIdx = MAllClusterBreakings.Add(TBreakingData<float, 3>());
-
-			TBreakingData<float, 3>& ClusterBreakingItem = MAllClusterBreakings[NewIdx];
-			ClusterBreakingItem.ParticleIndex = ClusterIndex;
-			ClusterBreakingItem.Location = MParticles.X(ClusterIndex);
-			ClusterBreakingItem.Velocity = MParticles.V(ClusterIndex);
-			ClusterBreakingItem.AngularVelocity = MParticles.W(ClusterIndex);
-			ClusterBreakingItem.Mass = MParticles.M(ClusterIndex);
-		}
-	}
-
-	DECLARE_CYCLE_STAT(TEXT("ResetCollisionImpulseArray"), STAT_ResetCollisionImpulseArray, STATGROUP_Chaos);
-	template<class FPBDRigidsEvolution, class FPBDCollisionConstraint, class T, int d>
-	void TPBDRigidClustering<FPBDRigidsEvolution, FPBDCollisionConstraint, T, d>::ResetCollisionImpulseArray()
-	{
-		SCOPE_CYCLE_COUNTER(STAT_ResetCollisionImpulseArray);
-		if (MCollisionImpulseArrayDirty)
-		{
-			for (uint32 i = 0; i < MParticles.Size(); i++)
+			for (uint32 ChildIdx : ChildrenParticles)
 			{
-				MParticles.CollisionImpulses(i) = (T)0.0;
-			}
-			MCollisionImpulseArrayDirty = false;
-		}
-	}
-
-	float MinImpulseForStrainEval = 980 * 2 * 1.f / 30.f; //ignore impulses caused by just keeping object on ground. This is a total hack, we should not use accumulated impulse directly. Instead we need to look at delta v along constraint normal
-	FAutoConsoleVariableRef CVarMinImpulseForStrainEval(TEXT("p.chaos.MinImpulseForStrainEval"), MinImpulseForStrainEval, TEXT("Minimum accumulated impulse before accumulating for strain eval "));
-
-	DECLARE_CYCLE_STAT(TEXT("ComputeStrainFromCollision"), STAT_ComputeStrainFromCollision, STATGROUP_Chaos);
-	template<class FPBDRigidsEvolution, class FPBDCollisionConstraint, class T, int d>
-	void TPBDRigidClustering<FPBDRigidsEvolution, FPBDCollisionConstraint, T, d>::ComputeStrainFromCollision(const FPBDCollisionConstraint& CollisionRule)
-	{
-		SCOPE_CYCLE_COUNTER(STAT_ComputeStrainFromCollision);
-		FClusterMap& MParentToChildren = GetChildrenMap();
-
-		ResetCollisionImpulseArray();
-
-		for (const Chaos::TPBDCollisionConstraintHandle<T, 3> * ContactHandle : CollisionRule.GetConstConstraintHandles())
-		{
-			if (ContactHandle->GetAccumulatedImpulse().Size() < MinImpulseForStrainEval)
-			{
-				continue;
-			}
-
-			auto ComputeStrainLambda = [&](const TPBDRigidClusteredParticleHandle<T, d>* Cluster, const TArray<uint32>& ParentToChildren) {
-				const TRigidTransform<T, d> WorldToClusterTM = TRigidTransform<T, d>(Cluster->P(), Cluster->Q());
-				const TVector<T, d> ContactLocationClusterLocal = WorldToClusterTM.InverseTransformPosition(ContactHandle->GetContactLocation());
-				TAABB<T, d> ContactBox(ContactLocationClusterLocal, ContactLocationClusterLocal);
-				ContactBox.Thicken(ClusterDistanceThreshold);
-				if (Cluster->ChildrenSpatial())
+				if (MParticles.Disabled(ChildIdx) && MParticles.ClusterIds(ChildIdx).Id != INDEX_NONE)
 				{
-					const TArray<int32> Intersections = Cluster->ChildrenSpatial()->FindAllIntersectingChildren(ContactBox);
-					for (int32 Child : Intersections)
+					if (ensure(!bFound || MParticles.ClusterIds(ChildIdx).Id == ClusterIdx))
 					{
-						if (const TUniquePtr<TMultiChildProxyData<T, d>>& ProxyData = MParticles.MultiChildProxyData(Child))
-						{
-							//multi child so get its children
-							const int32 KeyChild = ProxyData->KeyChild;
-							const TRigidTransform<T, d> ProxyToCluster = ProxyData->RelativeToKeyChild * MParticles.ChildToParent(KeyChild);
-							const TVector<T, d> ContactLocationProxyLocal = ProxyToCluster.InverseTransformPosition(ContactLocationClusterLocal);
-							TAABB<T, d> ContactBoxProxy(ContactLocationProxyLocal, ContactLocationProxyLocal);
-							ContactBoxProxy.Thicken(ClusterDistanceThreshold);
-							if (MParticles.ChildrenSpatial(Child))
-							{
-								const TArray<int32> SubIntersections = MParticles.ChildrenSpatial(Child)->FindAllIntersectingChildren(ContactBoxProxy);
-								for (int32 SubChild : SubIntersections)
-								{
-									MParticles.CollisionImpulses(SubChild) += ContactHandle->GetAccumulatedImpulse().Size();
-								}
-							}
-						}
-						else
-						{
-							MParticles.CollisionImpulses(Child) += ContactHandle->GetAccumulatedImpulse().Size();
-						}
+						bFound = true;
+						FakeStrain[ChildIdx] = FLT_MAX;
+						ClusterIdx = MParticles.ClusterIds(ChildIdx).Id;
 					}
-				}
-			};
-
-			TVector<const TGeometryParticleHandle<T, d>*, 2> ConstrainedParticles = ContactHandle->GetConstrainedParticles();
-			if (auto ChildrenPtr = MParentToChildren.Find(ConstrainedParticles[0]))
-			{
-				ComputeStrainLambda(ConstrainedParticles[0]->CastToClustered(), *ChildrenPtr);
-			}
-
-			if (auto ChildrenPtr = MParentToChildren.Find(ConstrainedParticles[1]))
-			{
-				ComputeStrainLambda(ConstrainedParticles[1]->CastToClustered(), *ChildrenPtr);
-			}
-
-			MCollisionImpulseArrayDirty = true;
-		}
-	}
-
-	template<class T, int d>
-	T CalculatePseudoMomentum(const TPBDRigidClusteredParticles<T, d>& InParticles, const uint32 Index)
-	{
-		TVector<T, d> LinearPseudoMomentum = (InParticles.X(Index) - InParticles.P(Index)) * InParticles.M(Index);
-		TRotation<T, d> Delta = InParticles.R(Index) * InParticles.Q(Index).Inverse();
-		TVector<T, d> Axis;
-		T Angle;
-		Delta.ToAxisAndAngle(Axis, Angle);
-		TVector<T, d> AngularPseudoMomentum = InParticles.I(Index) * (Axis * Angle);
-		return LinearPseudoMomentum.Size() + AngularPseudoMomentum.Size();
-	}
-
-	DECLARE_CYCLE_STAT(TEXT("TPBDRigidClustering<>::RewindAndEvolve<BGF>()"), STAT_RewindAndEvolve_BGF, STATGROUP_Chaos);
-	template<class T, int d>
-	void RewindAndEvolve(TPBDRigidsEvolutionGBF<T, d>& Evolution, TPBDRigidClusteredParticles<T, d>& InParticles, const TSet<int32>& IslandsToRecollide, const TSet<uint32>& AllActivatedChildren, const T Dt, TPBDCollisionConstraints<T, d>& CollisionRule)
-	{
-		SCOPE_CYCLE_COUNTER(STAT_RewindAndEvolve_BGF);
-		// Rewind active particles
-		const TArray<int32> IslandsToRecollideArray = IslandsToRecollide.Array();
-		PhysicsParallelFor(IslandsToRecollideArray.Num(), [&](int32 Idx) {
-			int32 Island = IslandsToRecollideArray[Idx];
-			auto Particles = Evolution.GetIslandParticles(Island); // copy
-			for (int32 ArrayIdx = Particles.Num() - 1; ArrayIdx >= 0; --ArrayIdx)
-			{
-				auto PBDRigid = Particles[ArrayIdx]->CastToRigidParticle();
-				if(PBDRigid && PBDRigid->ObjectState() == EObjectStateType::Dynamic)
-				{
-					if (!PBDRigid->Sleeping() && !PBDRigid->Disabled())
+					else
 					{
-						PBDRigid->P() = PBDRigid->X();
-						PBDRigid->Q() = PBDRigid->R();
-						PBDRigid->V() = PBDRigid->PreV();
-						PBDRigid->W() = PBDRigid->PreW();
-						continue;
-					}
-				}
-				Particles.RemoveAtSwap(ArrayIdx);
-			}
-			Evolution.Integrate(MakeHandleView(Particles), Dt);
-		});
-
-		TSet<TGeometryParticleHandle<T, d>*> AllIslandParticles;
-		for (int32 Island = 0; Island < Evolution.NumIslands(); ++Island)
-		{
-			const auto& ParticleIndices = Evolution.GetIslandParticles(Island);
-			for (const auto Particle : ParticleIndices)
-			{
-				auto PBDRigid = Particle->CastToRigidParticle();
-				if(PBDRigid && PBDRigid->ObjectState() == EObjectStateType::Dynamic)
-				{
-					bool bDisabled = PBDRigid->Disabled();
-
-					// #TODO - Have to repeat checking out whether the particle is disabled matching the PFor above.
-					// Move these into shared array so we only process it once
-					if (!AllIslandParticles.Contains(Particle) && !bDisabled)
-					{
-						AllIslandParticles.Add(Particle);
+						break; //shouldn't be here
 					}
 				}
 			}
+
+			if (bFound)
+			{
+				ActivatedBodies.Append(ReleaseClusterParticles(MParticles.Handle(ClusterIdx), FakeStrain));
+			}
+
+			DoGenerateBreakingData = bPreDoGenerateData;
 		}
 
-		const bool bRewindOnDeclusterSolve = ChildrenInheritVelocity < 1.f;
-		if (bRewindOnDeclusterSolve)
-		{
-			// @todo(mlentine): We can precompute internal constraints which can filter some from the narrow phase tests but may not help much
-
-#if CHAOS_PARTICLEHANDLE_TODO
-			CollisionRule.UpdateConstraints(InParticles, Evolution.GetNonDisabledIndices(), Dt, AllActivatedChildren, AllIslandParticles.Array());
-#else
-			CollisionRule.UpdateConstraints(Dt, AllIslandParticles);	//this seems wrong
-#endif
-
-			Evolution.InitializeAccelerationStructures();
-
-			// Resolve collisions
-			PhysicsParallelFor(IslandsToRecollide.Num(), [&](int32 Island) {
-				// @todo(mlentine): This is heavy handed and probably can be simplified as we know only a little bit changed.
-				Evolution.UpdateAccelerationStructures(Island);
-				Evolution.ApplyConstraints(Dt, Island);
-				// @todo(ccaulfield): should we also update velocities here? Evolution does...
-				Evolution.ApplyPushOut(Dt, Island);
-				// @todo(ccaulfield): support sleep state update on evolution
-				//Evolution.UpdateSleepState(Island);
-			});
-		}
+		return ActivatedBodies;
 	}
 
-#if CHAOS_PARTICLEHANDLE_TODO
-	DECLARE_CYCLE_STAT(TEXT("TPBDRigidClustering<>::RewindAndEvolve<PSG>()"), STAT_RewindAndEvolve_PSG, STATGROUP_Chaos);
-	template<class T, int d>
-	void RewindAndEvolve(TPBDRigidsEvolutionPGS<T, d>& Evolution, TPBDRigidClusteredParticles<T, d>& InParticles, const TSet<int32>& IslandsToRecollide, const TSet<uint32>& AllActivatedChildren, const T Dt, TPBDCollisionConstraintPGS<T, d>& CollisionRule)
-	{
-		SCOPE_CYCLE_COUNTER(STAT_RewindAndEvolve_PSG);
-		// Rewind active particles
-		TArray<int32>& NonDisabledIndices = Evolution.GetNonDisabledIndices();
-		PhysicsParallelFor(IslandsToRecollide.Num(), [&](int32 Island) {
-			TArray<int32> ParticleIndices = Evolution.GetIslandParticles(Island); // copy
-			for (int32 ArrayIdx = ParticleIndices.Num() - 1; ArrayIdx >= 0; --ArrayIdx)
-			{
-				int32 Index = ParticleIndices[ArrayIdx];
-				if (InParticles.Sleeping(Index) || InParticles.Disabled(Index))
-				{
-					ParticleIndices.RemoveAtSwap(ArrayIdx);
-				}
-				else
-				{
-					InParticles.P(Index) = InParticles.X(Index);
-					InParticles.Q(Index) = InParticles.R(Index);
-					InParticles.V(Index) = InParticles.PreV(Index);
-					InParticles.W(Index) = InParticles.PreW(Index);
-				}
-			}
-			Evolution.IntegrateV(ParticleIndices, Dt);
-		});
-
-		TSet<uint32> AllIslandParticles;
-		for (const auto& Island : IslandsToRecollide)
-		{
-			const TArray<int32>& ParticleIndices = Evolution.GetIslandParticles(Island);
-			for (const auto& Index : ParticleIndices)
-			{
-				if (InParticles.Disabled(Index) == false) //HACK: cluster code is incorrectly adding disabled children
-				{
-					if (!AllIslandParticles.Contains(Index))
-					{
-						AllIslandParticles.Add(Index);
-						NonDisabledIndices.Add(Index);
-					}
-				}
-				else
-				{
-					//FPlatformMisc::DebugBreak();
-				}
-			}
-		}
-
-		// @todo(mlentine): We can precompute internal constraints which can filter some from the narrow phase tests but may not help much
-		CollisionRule.UpdateConstraints(InParticles, Dt, AllActivatedChildren, AllIslandParticles.Array());
-
-		Evolution.InitializeAccelerationStructures();
-
-		PhysicsParallelFor(Evolution.NumIslands(), [&](int32 Island) {
-			Evolution.UpdateAccelerationStructures(Island);
-			Evolution.ApplyConstraints(Dt, Island);
-		});
-		PhysicsParallelFor(Evolution.NumIslands(), [&](int32 Island) {
-			const TArray<int32>& ParticleIndices = Evolution.GetIslandParticles(Island);
-			Evolution.IntegrateX(ParticleIndices, Dt);
-		});
-
-		// @todo(mlentine): Need to enforce constraints
-		PhysicsParallelFor(Evolution.NumIslands(), [&](int32 Island) {
-			Evolution.ApplyPushOut(Dt, Island);
-		});
-	}
-#endif
-
 	DECLARE_CYCLE_STAT(TEXT("TPBDRigidClustering<>::AdvanceClustering"), STAT_AdvanceClustering, STATGROUP_Chaos);
 	DECLARE_CYCLE_STAT(TEXT("TPBDRigidClustering<>::Update Impulse from Strain"), STAT_UpdateImpulseStrain, STATGROUP_Chaos);
@@ -970,29 +1072,6 @@
 		UE_LOG(LogChaos, Verbose, TEXT("Cluster Break Update Time is %f"), Time);
 	}
 
-	DECLARE_CYCLE_STAT(TEXT("TPBDRigidClustering<>::PromoteStrains()"), STAT_PromoteStrains, STATGROUP_Chaos);
-	template<class FPBDRigidsEvolution, class FPBDCollisionConstraint, class T, int d>
-	T TPBDRigidClustering<FPBDRigidsEvolution, FPBDCollisionConstraint, T, d>::PromoteStrains(uint32 CurrentNode, TArrayView<T>& ExternalStrains)
-	{
-		SCOPE_CYCLE_COUNTER(STAT_PromoteStrains);
-		FClusterMap& MParentToChildren = GetChildrenMap();
-
-		T Result = 0;
-		if (MParentToChildren.Contains(MParticles.Handle(CurrentNode)))
-		{
-			for (uint32 Child : MParentToChildren[MParticles.Handle(CurrentNode)])
-			{
-				Result += PromoteStrains(Child, ExternalStrains);
-			}
-		}
-		else
-		{
-			return ExternalStrains[CurrentNode];
-		}
-		ExternalStrains[CurrentNode] += Result;
-		return Result;
-	}
-
 	DECLARE_CYCLE_STAT(TEXT("TPBDRigidClustering<>::BreakingModel()"), STAT_BreakingModel, STATGROUP_Chaos);
 	template<class FPBDRigidsEvolution, class FPBDCollisionConstraint, class T, int d>
 	TMap<TPBDRigidClusteredParticleHandle<T, d>*, TSet<uint32>> TPBDRigidClustering<FPBDRigidsEvolution, FPBDCollisionConstraint, T, d>::BreakingModel(TArrayView<T> ExternalStrain)
@@ -1040,6 +1119,29 @@
 		return AllActivatedChildren;
 	}
 
+	DECLARE_CYCLE_STAT(TEXT("TPBDRigidClustering<>::PromoteStrains()"), STAT_PromoteStrains, STATGROUP_Chaos);
+	template<class FPBDRigidsEvolution, class FPBDCollisionConstraint, class T, int d>
+	T TPBDRigidClustering<FPBDRigidsEvolution, FPBDCollisionConstraint, T, d>::PromoteStrains(uint32 CurrentNode, TArrayView<T>& ExternalStrains)
+	{
+		SCOPE_CYCLE_COUNTER(STAT_PromoteStrains);
+		FClusterMap& MParentToChildren = GetChildrenMap();
+
+		T Result = 0;
+		if (MParentToChildren.Contains(MParticles.Handle(CurrentNode)))
+		{
+			for (uint32 Child : MParentToChildren[MParticles.Handle(CurrentNode)])
+			{
+				Result += PromoteStrains(Child, ExternalStrains);
+			}
+		}
+		else
+		{
+			return ExternalStrains[CurrentNode];
+		}
+		ExternalStrains[CurrentNode] += Result;
+		return Result;
+	}
+
 	DECLARE_CYCLE_STAT(TEXT("TPBDRigidClustering<>::UpdateKinematicProperties()"), STAT_UpdateKinematicProperties, STATGROUP_Chaos);
 	template<class FPBDRigidsEvolution, class FPBDCollisionConstraint, class T, int d>
 	void TPBDRigidClustering<FPBDRigidsEvolution, FPBDCollisionConstraint, T, d>::UpdateKinematicProperties(uint32 ClusterIndex)
@@ -1084,202 +1186,122 @@
 		}
 	}
 
+	DECLARE_CYCLE_STAT(TEXT("TPBDRigidClustering<>::SwapBufferedData"), STAT_SwapBufferedData, STATGROUP_Chaos);
+	template<class FPBDRigidsEvolution, class FPBDCollisionConstraint, class T, int d>
+	void TPBDRigidClustering<FPBDRigidsEvolution, FPBDCollisionConstraint, T, d>::SwapBufferedData()
+	{
+		SCOPE_CYCLE_COUNTER(STAT_SwapBufferedData);
+		ResourceLock.WriteLock();
+		//BufferResource.MChildren.Reset();
+		//BufferResource.ClusterParentTransforms.Reset();	//todo: once everything is atomic this should get reset
+		const TArray<TSerializablePtr<FImplicitObject>>& AllGeom = MParticles.GetAllGeometry();
+		BufferResource.GeometryPtrs.SetNum(AllGeom.Num());
 
+		const auto& NonDisabledClusteredParticles = MEvolution.GetNonDisabledClusteredArray();
+		for (auto& ClusteredParticle : NonDisabledClusteredParticles)
+		{
+			const TArray<uint32>* Children = MChildren.Find(ClusteredParticle);
+			if (Children && ClusteredParticle->ClusterIds().Id == INDEX_NONE) //root cluster so copy children
+			{
+				//TODO: record GT particle pointer instead
+				BufferResource.MChildren.Add(ClusteredParticle->TransientParticleIndex(), *Children);
+				BufferResource.ClusterParentTransforms.Add(ClusteredParticle->TransientParticleIndex(), TRigidTransform<float, 3>(ClusteredParticle->X(), ClusteredParticle->R()));
+			}
+		}
+
+		BufferResource.GeometryPtrs = AllGeom; //in future this should be sparse. SQ has fallback that relies on potentially all geom so can't do it yet
+
+		ResourceLock.WriteUnlock();
+	}
+
+	DECLARE_CYCLE_STAT(TEXT("TPBDRigidClustering<>::GetActiveClusterIndex"), STAT_GetActiveClusterIndex, STATGROUP_Chaos);
+	template<class FPBDRigidsEvolution, class FPBDCollisionConstraint, class T, int d>
+	int32 TPBDRigidClustering<FPBDRigidsEvolution, FPBDCollisionConstraint, T, d>::GetActiveClusterIndex(uint32 ChildIndex)
+	{
+		SCOPE_CYCLE_COUNTER(STAT_GetActiveClusterIndex);
+		int32 NumParticles = MParticles.Size();
+		int32 CurrentIndex = (int32)ChildIndex;
+		while (0 <= CurrentIndex && CurrentIndex < NumParticles)
+		{
+			if (!MParticles.Disabled(CurrentIndex))
+			{
+				return CurrentIndex;
+			}
+			CurrentIndex = MParticles.ClusterIds(CurrentIndex).Id;
+		}
+		return INDEX_NONE;
+	}
+
+	template<class FPBDRigidsEvolution, class FPBDCollisionConstraint, class T, int d>
+	void TPBDRigidClustering<FPBDRigidsEvolution, FPBDCollisionConstraint, T, d>::IncrementPendingClusterCounter(uint32 ClusterGroupID)
+	{
+		if (!PendingClusterCounter.Contains(ClusterGroupID))
+		{
+			PendingClusterCounter.Add(ClusterGroupID, 0);
+		}
+		PendingClusterCounter[ClusterGroupID]++;
+	}
+
+	template<class FPBDRigidsEvolution, class FPBDCollisionConstraint, class T, int d>
+	void TPBDRigidClustering<FPBDRigidsEvolution, FPBDCollisionConstraint, T, d>::DecrementPendingClusterCounter(uint32 ClusterGroupID)
+	{
+		PendingClusterCounter[ClusterGroupID]--;
+		ensure(0 <= PendingClusterCounter[ClusterGroupID]);
+	}
+
+	DECLARE_CYCLE_STAT(TEXT("TPBDRigidClustering<>::GenerateConnectionGraph"), STAT_GenerateConnectionGraph, STATGROUP_Chaos);
+	template<class FPBDRigidsEvolution, class FPBDCollisionConstraint, class T, int d>
+	void TPBDRigidClustering<FPBDRigidsEvolution, FPBDCollisionConstraint, T, d>::GenerateConnectionGraph(int32 ClusterIndex, const FClusterCreationParameters<T>& Parameters)
+	{
+		SCOPE_CYCLE_COUNTER(STAT_GenerateConnectionGraph);
+		if (!GetChildrenMap().Contains(MParticles.Handle(ClusterIndex))) return;
+
+		// Connectivity Graph
+		//    Build a connectivity graph for the cluster. If the PointImplicit is specified
+		//    and the ClusterIndex has collision particles then use the expensive connection
+		//    method. Otherwise try the DelaunayTriangulation if not none.
+		//
+		if (Parameters.bGenerateConnectionGraph)
+		{
+			typename FClusterCreationParameters<T>::EConnectionMethod LocalConnectionMethod = Parameters.ConnectionMethod;
+
+			if (LocalConnectionMethod == FClusterCreationParameters<T>::EConnectionMethod::None ||
+				(LocalConnectionMethod == FClusterCreationParameters<T>::EConnectionMethod::PointImplicit && !MParticles.CollisionParticles(ClusterIndex)))
+			{
+				LocalConnectionMethod = FClusterCreationParameters<T>::EConnectionMethod::MinimalSpanningSubsetDelaunayTriangulation; // default method
+			}
+
+			if (LocalConnectionMethod == FClusterCreationParameters<T>::EConnectionMethod::PointImplicit ||
+				LocalConnectionMethod == FClusterCreationParameters<T>::EConnectionMethod::PointImplicitAugmentedWithMinimalDelaunay)
+			{
+				UpdateConnectivityGraphUsingPointImplicit(ClusterIndex, Parameters);
+			}
+
+			if (LocalConnectionMethod == FClusterCreationParameters<T>::EConnectionMethod::DelaunayTriangulation)
+			{
+				UpdateConnectivityGraphUsingDelaunayTriangulation(ClusterIndex, Parameters); // not thread safe
+			}
+
+			if (LocalConnectionMethod == FClusterCreationParameters<T>::EConnectionMethod::PointImplicitAugmentedWithMinimalDelaunay ||
+				LocalConnectionMethod == FClusterCreationParameters<T>::EConnectionMethod::MinimalSpanningSubsetDelaunayTriangulation)
+			{
+				FixConnectivityGraphUsingDelaunayTriangulation(ClusterIndex, Parameters);
+			}
+		}
+	}
+
+	DECLARE_CYCLE_STAT(TEXT("TPBDRigidClustering<>::UpdateMassProperties"), STAT_UpdateMassProperties, STATGROUP_Chaos);
+	template<class FPBDRigidsEvolution, class FPBDCollisionConstraint, class T, int d>
+	void TPBDRigidClustering<FPBDRigidsEvolution, FPBDCollisionConstraint, T, d>::UpdateMassProperties(const TArray<uint32>& Children, const uint32 ClusterIndex, const TRigidTransform<T, d>* ForceMassOrientation)
+	{
+		SCOPE_CYCLE_COUNTER(STAT_UpdateMassProperties);
+		UpdateClusterMassProperties(MParticles, Children, ClusterIndex, ForceMassOrientation, &MParticles.MultiChildProxyDataArray(), &MParticles.MultiChildProxyIdArray());
+		UpdateKinematicProperties(ClusterIndex);
+	}
+
 	int32 MassPropertiesFromMultiChildProxy = 1;
 	FAutoConsoleVariableRef CVarMassPropertiesFromMultiChildProxy(TEXT("p.MassPropertiesFromMultiChildProxy"), MassPropertiesFromMultiChildProxy, TEXT(""));
 
-	DECLARE_CYCLE_STAT(TEXT("TPBDRigidClustering<>::UpdateClusterMassProperties()"), STAT_UpdateClusterMassProperties, STATGROUP_Chaos);
-	template<class T, int d>
-	void UpdateClusterMassProperties(TPBDRigidClusteredParticles<T, d>& Particles, const TArray<uint32>& Children, const uint32 NewIndex, const TRigidTransform<T, d>* ForceMassOrientation, const TArrayCollectionArray<TUniquePtr<TMultiChildProxyData<T, d>>>* MMultiChildProxyData, const TArrayCollectionArray<FMultiChildProxyId>* MMultiChildProxyId)
-	{
-		SCOPE_CYCLE_COUNTER(STAT_UpdateClusterMassProperties);
-		check(0 <= NewIndex && NewIndex < Particles.Size());
-		check(Children.Num());
-
-		Particles.X(NewIndex) = TVector<T, d>(0);
-		Particles.R(NewIndex) = TRotation<T, d>(FQuat::MakeFromEuler(TVector<T, d>(0)));
-		Particles.V(NewIndex) = TVector<T, d>(0);
-		Particles.W(NewIndex) = TVector<T, d>(0);
-		Particles.M(NewIndex) = 0;
-		Particles.I(NewIndex) = PMatrix<T, d, d>(0);
-
-		auto GetMultiChildProxy = [MMultiChildProxyId, MMultiChildProxyData](const int32 OriginalChild, const TMultiChildProxyData<T, d>*& ProxyData) -> int32
-		{
-			const int32 MultiChildProxyId = MassPropertiesFromMultiChildProxy && MMultiChildProxyId ? (*MMultiChildProxyId)[OriginalChild].Id : INDEX_NONE;
-			if (MultiChildProxyId != INDEX_NONE)
-			{
-				ProxyData = (*MMultiChildProxyData)[MultiChildProxyId].Get();
-				return ProxyData ? MultiChildProxyId : INDEX_NONE;
-			}
-			return INDEX_NONE;
-		};
-
-		bool bHasChild = false;
-		bool bHasProxyChild = false;
-		for (const auto OriginalChild : Children)
-		{
-			const TMultiChildProxyData<T, d>* ProxyData;
-			const int32 MultiChildProxyId = GetMultiChildProxy(OriginalChild, ProxyData);
-			TVector<T, d> ChildPosition;
-			TRotation<T, d> ChildRotation;
-			int32 Child;
-
-			if (MultiChildProxyId == INDEX_NONE)
-			{
-				Child = OriginalChild;
-				ChildPosition = Particles.X(Child);
-				ChildRotation = Particles.R(Child);
-			}
-			else if (ProxyData->KeyChild == OriginalChild)
-			{
-				Child = MultiChildProxyId;
-
-				const TRigidTransform<T, d> ProxyWorldTM = ProxyData->RelativeToKeyChild * TRigidTransform<T, d>(Particles.X(OriginalChild), Particles.R(OriginalChild));
-				ChildPosition = ProxyWorldTM.GetLocation();
-				ChildRotation = ProxyWorldTM.GetRotation();
-			}
-			else
-			{
-				continue; //using a proxy but we are not the key child
-			}
-
-			const auto& ChildMass = Particles.M(Child);
-			const PMatrix<T, d, d> ChildWorldSpaceI = (ChildRotation * FMatrix::Identity) * Particles.I(Child) * (ChildRotation * FMatrix::Identity).GetTransposed();
-			if (ChildWorldSpaceI.ContainsNaN())
-			{
-				continue;
-			}
-			bHasProxyChild = true;
-			bHasChild = true;
-			Particles.I(NewIndex) += ChildWorldSpaceI;
-			Particles.M(NewIndex) += ChildMass;
-			Particles.X(NewIndex) += ChildPosition * ChildMass;
-			Particles.V(NewIndex) += Particles.V(OriginalChild) * ChildMass; //use original child for veocities because we don't simulate proxy
-			Particles.W(NewIndex) += Particles.W(OriginalChild) * ChildMass;
-		}
-		if (!ensure(bHasProxyChild))
-		{
-			for (const auto OriginalChild : Children)
-			{
-				TVector<T, d> ChildPosition;
-				TRotation<T, d> ChildRotation;
-				int32 Child;
-
-				Child = OriginalChild;
-				ChildPosition = Particles.X(Child);
-				ChildRotation = Particles.R(Child);
-
-				const auto& ChildMass = Particles.M(Child);
-				const PMatrix<T, d, d> ChildWorldSpaceI = (ChildRotation * FMatrix::Identity) * Particles.I(Child) * (ChildRotation * FMatrix::Identity).GetTransposed();
-				if (ChildWorldSpaceI.ContainsNaN())
-				{
-					continue;
-				}
-				bHasChild = true;
-				Particles.I(NewIndex) += ChildWorldSpaceI;
-				Particles.M(NewIndex) += ChildMass;
-				Particles.X(NewIndex) += ChildPosition * ChildMass;
-				Particles.V(NewIndex) += Particles.V(OriginalChild) * ChildMass;	//use original child for veocities because we don't simulate proxy
-				Particles.W(NewIndex) += Particles.W(OriginalChild) * ChildMass;
-			}
-		}
-		for (int32 i = 0; i < d; i++)
-		{
-			const PMatrix<T, d, d>& InertiaTensor = Particles.I(NewIndex);
-			if (InertiaTensor.GetColumn(i)[i] < SMALL_NUMBER)
-			{
-				Particles.I(NewIndex) = PMatrix<T, d, d>(1.f, 1.f, 1.f);
-				break;
-			}
-		}
-
-		if (!ensure(bHasChild) || !ensure(Particles.M(NewIndex) > SMALL_NUMBER))
-		{
-			Particles.M(NewIndex) = 1;
-			Particles.X(NewIndex) = TVector<T, d>(0);
-			Particles.V(NewIndex) = TVector<T, d>(0);
-			Particles.PreV(NewIndex) = Particles.V(NewIndex);
-			Particles.InvM(NewIndex) = 1;
-			Particles.P(NewIndex) = Particles.X(NewIndex);
-			Particles.W(NewIndex) = TVector<T, d>(0);
-			Particles.PreW(NewIndex) = Particles.W(NewIndex);
-			Particles.R(NewIndex) = TRotation<T, d>(FMatrix::Identity);
-			Particles.Q(NewIndex) = Particles.R(NewIndex);
-			Particles.I(NewIndex) = FMatrix::Identity;
-			Particles.InvI(NewIndex) = FMatrix::Identity;
-			return;
-		}
-
-		check(Particles.M(NewIndex) > SMALL_NUMBER);
-
-		Particles.X(NewIndex) /= Particles.M(NewIndex);
-		Particles.V(NewIndex) /= Particles.M(NewIndex);
-		Particles.PreV(NewIndex) = Particles.V(NewIndex);
-		Particles.InvM(NewIndex) = 1. / Particles.M(NewIndex);
-		if (ForceMassOrientation)
-		{
-			Particles.X(NewIndex) = ForceMassOrientation->GetLocation();
-		}
-		Particles.P(NewIndex) = Particles.X(NewIndex);
-		for (const auto OriginalChild : Children)
-		{
-			int32 Child;
-			const TMultiChildProxyData<T, d>* ProxyData = nullptr;
-			const int32 MultiChildProxyId = bHasProxyChild ? GetMultiChildProxy(OriginalChild, ProxyData) : INDEX_NONE;
-
-			TVector<T, d> ChildPosition;
-			if (MultiChildProxyId == INDEX_NONE)
-			{
-				Child = OriginalChild;
-				ChildPosition = Particles.X(Child);
-			}
-			else if (ProxyData->KeyChild == OriginalChild)
-			{
-				Child = MultiChildProxyId;
-				const TRigidTransform<T, d> ProxyWorldTM = ProxyData->RelativeToKeyChild * TRigidTransform<T, d>(Particles.X(OriginalChild), Particles.R(OriginalChild));
-				ChildPosition = ProxyWorldTM.GetLocation();
-			}
-			else
-			{
-				continue; //using a proxy but we are not the key child
-			}
-
-			TVector<T, d> ParentToChild = ChildPosition - Particles.X(NewIndex);
-
-			const auto& ChildMass = Particles.M(Child);
-			Particles.W(NewIndex) += TVector<T, d>::CrossProduct(ParentToChild, Particles.V(OriginalChild) * ChildMass); //taking v from original child since we are not actually simulating the proxy child
-			{
-				const T& p0 = ParentToChild[0];
-				const T& p1 = ParentToChild[1];
-				const T& p2 = ParentToChild[2];
-				const T& m = Particles.M(Child);
-				Particles.I(NewIndex) += PMatrix<T, d, d>(m * (p1 * p1 + p2 * p2), -m * p1 * p0, -m * p2 * p0, m * (p2 * p2 + p0 * p0), -m * p2 * p1, m * (p1 * p1 + p0 * p0));
-			}
-		}
-		if (Particles.I(NewIndex).ContainsNaN())
-		{
-			Particles.I(NewIndex) = PMatrix<T, d, d>(1.f, 1.f, 1.f);
-		}
-		for (int32 i = 0; i < d; i++)
-		{
-			const PMatrix<T, d, d>& InertiaTensor = Particles.I(NewIndex);
-			if (InertiaTensor.GetColumn(i)[i] < SMALL_NUMBER)
-			{
-				Particles.I(NewIndex) = PMatrix<T, d, d>(1.f, 1.f, 1.f);
-				break;
-			}
-		}
-		Particles.W(NewIndex) /= Particles.M(NewIndex);
-		Particles.PreW(NewIndex) = Particles.W(NewIndex);
-		Particles.R(NewIndex) = Chaos::TransformToLocalSpace<T, d>(Particles.I(NewIndex));
-		if (ForceMassOrientation)
-		{
-			Particles.R(NewIndex) = ForceMassOrientation->GetRotation();
-		}
-		Particles.Q(NewIndex) = Particles.R(NewIndex);
-		Particles.InvI(NewIndex) = Particles.I(NewIndex).Inverse();
-	}
-
 	DECLARE_CYCLE_STAT(TEXT("TPBDRigidClustering<>::UpdateGeometry"), STAT_UpdateGeometry, STATGROUP_Chaos);
 	DECLARE_CYCLE_STAT(TEXT("TPBDRigidClustering<>::UpdateGeometry_GatherObjects"), STAT_UpdateGeometry_GatherObjects, STATGROUP_Chaos);
 	DECLARE_CYCLE_STAT(TEXT("TPBDRigidClustering<>::UpdateGeometry_GatherPoints"), STAT_UpdateGeometry_GatherPoints, STATGROUP_Chaos);
@@ -1481,43 +1503,196 @@
 		}
 	}
 
-	DECLARE_CYCLE_STAT(TEXT("TPBDRigidClustering<>::GenerateConnectionGraph"), STAT_GenerateConnectionGraph, STATGROUP_Chaos);
+	float MinImpulseForStrainEval = 980 * 2 * 1.f / 30.f; //ignore impulses caused by just keeping object on ground. This is a total hack, we should not use accumulated impulse directly. Instead we need to look at delta v along constraint normal
+	FAutoConsoleVariableRef CVarMinImpulseForStrainEval(TEXT("p.chaos.MinImpulseForStrainEval"), MinImpulseForStrainEval, TEXT("Minimum accumulated impulse before accumulating for strain eval "));
+
+	DECLARE_CYCLE_STAT(TEXT("ComputeStrainFromCollision"), STAT_ComputeStrainFromCollision, STATGROUP_Chaos);
 	template<class FPBDRigidsEvolution, class FPBDCollisionConstraint, class T, int d>
-	void TPBDRigidClustering<FPBDRigidsEvolution, FPBDCollisionConstraint, T, d>::GenerateConnectionGraph(int32 ClusterIndex, const FClusterCreationParameters<T>& Parameters)
+	void TPBDRigidClustering<FPBDRigidsEvolution, FPBDCollisionConstraint, T, d>::ComputeStrainFromCollision(const FPBDCollisionConstraint& CollisionRule)
 	{
-		SCOPE_CYCLE_COUNTER(STAT_GenerateConnectionGraph);
-		if (!GetChildrenMap().Contains(MParticles.Handle(ClusterIndex))) return;
+		SCOPE_CYCLE_COUNTER(STAT_ComputeStrainFromCollision);
+		FClusterMap& MParentToChildren = GetChildrenMap();
 
-		// Connectivity Graph
-		//    Build a connectivity graph for the cluster. If the PointImplicit is specified
-		//    and the ClusterIndex has collision particles then use the expensive connection
-		//    method. Otherwise try the DelaunayTriangulation if not none.
-		//
-		if (Parameters.bGenerateConnectionGraph)
+		ResetCollisionImpulseArray();
+
+		for (const Chaos::TPBDCollisionConstraintHandle<T, 3> * ContactHandle : CollisionRule.GetConstConstraintHandles())
 		{
-			typename FClusterCreationParameters<T>::EConnectionMethod LocalConnectionMethod = Parameters.ConnectionMethod;
-
-			if (LocalConnectionMethod == FClusterCreationParameters<T>::EConnectionMethod::None ||
-				(LocalConnectionMethod == FClusterCreationParameters<T>::EConnectionMethod::PointImplicit && !MParticles.CollisionParticles(ClusterIndex)))
+			if (ContactHandle->GetAccumulatedImpulse().Size() < MinImpulseForStrainEval)
 			{
-				LocalConnectionMethod = FClusterCreationParameters<T>::EConnectionMethod::MinimalSpanningSubsetDelaunayTriangulation; // default method
+				continue;
 			}
 
-			if (LocalConnectionMethod == FClusterCreationParameters<T>::EConnectionMethod::PointImplicit ||
-				LocalConnectionMethod == FClusterCreationParameters<T>::EConnectionMethod::PointImplicitAugmentedWithMinimalDelaunay)
+			auto ComputeStrainLambda = [&](const TPBDRigidClusteredParticleHandle<T, d>* Cluster, const TArray<uint32>& ParentToChildren) {
+				const TRigidTransform<T, d> WorldToClusterTM = TRigidTransform<T, d>(Cluster->P(), Cluster->Q());
+				const TVector<T, d> ContactLocationClusterLocal = WorldToClusterTM.InverseTransformPosition(ContactHandle->GetContactLocation());
+				TAABB<T, d> ContactBox(ContactLocationClusterLocal, ContactLocationClusterLocal);
+				ContactBox.Thicken(ClusterDistanceThreshold);
+				if (Cluster->ChildrenSpatial())
+				{
+					const TArray<int32> Intersections = Cluster->ChildrenSpatial()->FindAllIntersectingChildren(ContactBox);
+					for (int32 Child : Intersections)
+					{
+						if (const TUniquePtr<TMultiChildProxyData<T, d>>& ProxyData = MParticles.MultiChildProxyData(Child))
+						{
+							//multi child so get its children
+							const int32 KeyChild = ProxyData->KeyChild;
+							const TRigidTransform<T, d> ProxyToCluster = ProxyData->RelativeToKeyChild * MParticles.ChildToParent(KeyChild);
+							const TVector<T, d> ContactLocationProxyLocal = ProxyToCluster.InverseTransformPosition(ContactLocationClusterLocal);
+							TAABB<T, d> ContactBoxProxy(ContactLocationProxyLocal, ContactLocationProxyLocal);
+							ContactBoxProxy.Thicken(ClusterDistanceThreshold);
+							if (MParticles.ChildrenSpatial(Child))
+							{
+								const TArray<int32> SubIntersections = MParticles.ChildrenSpatial(Child)->FindAllIntersectingChildren(ContactBoxProxy);
+								for (int32 SubChild : SubIntersections)
+								{
+									MParticles.CollisionImpulses(SubChild) += ContactHandle->GetAccumulatedImpulse().Size();
+								}
+							}
+						}
+						else
+						{
+							MParticles.CollisionImpulses(Child) += ContactHandle->GetAccumulatedImpulse().Size();
+						}
+					}
+				}
+			};
+
+			TVector<const TGeometryParticleHandle<T, d>*, 2> ConstrainedParticles = ContactHandle->GetConstrainedParticles();
+			if (auto ChildrenPtr = MParentToChildren.Find(ConstrainedParticles[0]))
 			{
-				UpdateConnectivityGraphUsingPointImplicit(ClusterIndex, Parameters);
+				ComputeStrainLambda(ConstrainedParticles[0]->CastToClustered(), *ChildrenPtr);
 			}
 
-			if (LocalConnectionMethod == FClusterCreationParameters<T>::EConnectionMethod::DelaunayTriangulation)
+			if (auto ChildrenPtr = MParentToChildren.Find(ConstrainedParticles[1]))
 			{
-				UpdateConnectivityGraphUsingDelaunayTriangulation(ClusterIndex, Parameters); // not thread safe
+				ComputeStrainLambda(ConstrainedParticles[1]->CastToClustered(), *ChildrenPtr);
 			}
 
-			if (LocalConnectionMethod == FClusterCreationParameters<T>::EConnectionMethod::PointImplicitAugmentedWithMinimalDelaunay ||
-				LocalConnectionMethod == FClusterCreationParameters<T>::EConnectionMethod::MinimalSpanningSubsetDelaunayTriangulation)
+			MCollisionImpulseArrayDirty = true;
+		}
+	}
+
+	DECLARE_CYCLE_STAT(TEXT("ResetCollisionImpulseArray"), STAT_ResetCollisionImpulseArray, STATGROUP_Chaos);
+	template<class FPBDRigidsEvolution, class FPBDCollisionConstraint, class T, int d>
+	void TPBDRigidClustering<FPBDRigidsEvolution, FPBDCollisionConstraint, T, d>::ResetCollisionImpulseArray()
+	{
+		SCOPE_CYCLE_COUNTER(STAT_ResetCollisionImpulseArray);
+		if (MCollisionImpulseArrayDirty)
+		{
+			for (uint32 i = 0; i < MParticles.Size(); i++)
 			{
-				FixConnectivityGraphUsingDelaunayTriangulation(ClusterIndex, Parameters);
+				MParticles.CollisionImpulses(i) = (T)0.0;
+			}
+			MCollisionImpulseArrayDirty = false;
+		}
+	}
+
+	template<class FPBDRigidsEvolution, class FPBDCollisionConstraint, class T, int d>
+	void TPBDRigidClustering<FPBDRigidsEvolution, FPBDCollisionConstraint, T, d>::DisableCluster(TPBDRigidClusteredParticleHandle<T,d>* ClusteredParticle)
+	{
+		// #note: we don't recursively descend to the children
+		FClusterMap& ParentToChildren = GetChildrenMap();
+
+		MEvolution.DisableParticle(ClusteredParticle);
+
+		if (MoveClustersWhenDeactivated)
+		{
+			ClusteredParticle->P() -= FVector(0.f, 0.f, -10000.f); // HACK : Move them away to prevent reactivation. 
+			ClusteredParticle->X() -= FVector(0.f, 0.f, -10000.f); // HACK : Move them away to prevent reactivation. 
+			ClusteredParticle->V() = FVector(0.f);            // HACK : Move them away to prevent reactivation.
+		}
+
+		TopLevelClusterParents.Remove(ClusteredParticle->TransientParticleIndex());
+		ParentToChildren.Remove(ClusteredParticle);
+		ClusteredParticle->ClusterIds() = ClusterId();
+		ClusteredParticle->ClusterGroupIndex() = 0;
+		MActiveRemovalIndices.Remove(ClusteredParticle->TransientParticleIndex());
+	}
+
+	template<class FPBDRigidsEvolution, class FPBDCollisionConstraint, class T, int d>
+	void TPBDRigidClustering<FPBDRigidsEvolution, FPBDCollisionConstraint, T, d>::DisableParticleWithBreakEvent(uint32 ClusterIndex)
+	{
+		DisableCluster(MParticles.Handle(ClusterIndex));
+
+		if (DoGenerateBreakingData)
+		{
+			int32 NewIdx = MAllClusterBreakings.Add(TBreakingData<float, 3>());
+
+			TBreakingData<float, 3>& ClusterBreakingItem = MAllClusterBreakings[NewIdx];
+			ClusterBreakingItem.ParticleIndex = ClusterIndex;
+			ClusterBreakingItem.Location = MParticles.X(ClusterIndex);
+			ClusterBreakingItem.Velocity = MParticles.V(ClusterIndex);
+			ClusterBreakingItem.AngularVelocity = MParticles.W(ClusterIndex);
+			ClusterBreakingItem.Mass = MParticles.M(ClusterIndex);
+		}
+	}
+
+	DECLARE_CYCLE_STAT(TEXT("TPBDRigidClustering<>::UpdateConnectivityGraphUsingPointImplicit"), STAT_UpdateConnectivityGraphUsingPointImplicit, STATGROUP_Chaos);
+	template<class FPBDRigidsEvolution, class FPBDCollisionConstraint, class T, int d>
+	void TPBDRigidClustering<FPBDRigidsEvolution, FPBDCollisionConstraint, T, d>::UpdateConnectivityGraphUsingPointImplicit(uint32 ClusterIndex, const FClusterCreationParameters<T>& Parameters)
+	{
+		SCOPE_CYCLE_COUNTER(STAT_UpdateConnectivityGraphUsingPointImplicit);
+		FClusterMap& MParentToChildren = GetChildrenMap();
+
+		if (UseConnectivity)
+		{
+			T Delta = FMath::Min(FMath::Max(Parameters.CoillisionThicknessPercent, (T)0), T(1));
+			const TArray<uint32>& Children = MParentToChildren[MParticles.Handle(ClusterIndex)];
+			for (int32 i = 0; i < Children.Num(); ++i)
+			{
+				const uint32 Child1 = Children[i];
+				if (!MParticles.Geometry(Child1) || !MParticles.Geometry(Child1)->HasBoundingBox())
+				{
+					continue;
+				}
+				if (!(ensure(!FMath::IsNaN(MParticles.X(Child1)[0])) && ensure(!FMath::IsNaN(MParticles.X(Child1)[1])) && ensure(!FMath::IsNaN(MParticles.X(Child1)[2]))))
+				{
+					continue;
+				}
+				TRigidTransform<T, d> TM1 = TRigidTransform<T, d>(MParticles.X(Child1), MParticles.R(Child1));
+				TAABB<T, d> Box1 = MParticles.Geometry(Child1)->BoundingBox();
+
+				TArray<TArray<TPair<uint32, uint32>>> Connections;
+				Connections.Init(TArray<TPair<uint32, uint32>>(), Children.Num() - (i + 1));
+				PhysicsParallelFor(Children.Num() - (i + 1), [&](int32 Idx) {
+					const uint32 Child2 = Children[Idx + i + 1];
+					if (!MParticles.CollisionParticles(Child2))
+					{
+						return;
+					}
+					if (!(ensure(!FMath::IsNaN(MParticles.X(Child2)[0])) && ensure(!FMath::IsNaN(MParticles.X(Child2)[1])) && ensure(!FMath::IsNaN(MParticles.X(Child2)[2]))))
+					{
+						return;
+					}
+					TRigidTransform<T, d> TM = TM1.GetRelativeTransform(TRigidTransform<T, d>(MParticles.X(Child2), MParticles.R(Child2)));
+
+					bool bCollided = false;
+					for (uint32 CollisionIdx = 0; !bCollided && CollisionIdx < MParticles.CollisionParticles(Child2)->Size(); ++CollisionIdx)
+					{
+						TVector<T, d> Normal;
+						TVector<T, d> LocalPoint = TM.TransformPositionNoScale(MParticles.CollisionParticles(Child2)->X(CollisionIdx));
+						T Phi = MParticles.Geometry(Child1)->PhiWithNormal(LocalPoint - (LocalPoint * Delta), Normal);
+						if (Phi < 0.0)
+							bCollided = true;
+					}
+					if (bCollided)
+					{
+						Connections[Idx].Add(TPair<uint32, uint32>((uint32)Child1, (uint32)Child2));
+					}
+				});
+
+				// join results and make connections
+				for (const auto& ConnectionList : Connections)
+				{
+					for (const auto& Elem : ConnectionList)
+					{
+						if (MParticles.ConnectivityEdges(Elem.Key).Num() < Parameters.MaxNumConnections)
+						{
+							const T AvgStrain = (MParticles.Strains(Elem.Key) + MParticles.Strains(Elem.Value)) * (T)0.5f;
+							ConnectNodes(Elem.Key, Elem.Value, AvgStrain);
+						}
+					}
+				}
 			}
 		}
 	}
@@ -1688,126 +1863,6 @@
 		}
 	}
 
-	DECLARE_CYCLE_STAT(TEXT("TPBDRigidClustering<>::UpdateConnectivityGraphUsingPointImplicit"), STAT_UpdateConnectivityGraphUsingPointImplicit, STATGROUP_Chaos);
 	template<class FPBDRigidsEvolution, class FPBDCollisionConstraint, class T, int d>
-	void TPBDRigidClustering<FPBDRigidsEvolution, FPBDCollisionConstraint, T, d>::UpdateConnectivityGraphUsingPointImplicit(uint32 ClusterIndex, const FClusterCreationParameters<T>& Parameters)
-	{
-		SCOPE_CYCLE_COUNTER(STAT_UpdateConnectivityGraphUsingPointImplicit);
-		FClusterMap& MParentToChildren = GetChildrenMap();
-
-		if (UseConnectivity)
-		{
-			T Delta = FMath::Min(FMath::Max(Parameters.CoillisionThicknessPercent, (T)0), T(1));
-			const TArray<uint32>& Children = MParentToChildren[MParticles.Handle(ClusterIndex)];
-			for (int32 i = 0; i < Children.Num(); ++i)
-			{
-				const uint32 Child1 = Children[i];
-				if (!MParticles.Geometry(Child1) || !MParticles.Geometry(Child1)->HasBoundingBox())
-				{
-					continue;
-				}
-				if (!(ensure(!FMath::IsNaN(MParticles.X(Child1)[0])) && ensure(!FMath::IsNaN(MParticles.X(Child1)[1])) && ensure(!FMath::IsNaN(MParticles.X(Child1)[2]))))
-				{
-					continue;
-				}
-				TRigidTransform<T, d> TM1 = TRigidTransform<T, d>(MParticles.X(Child1), MParticles.R(Child1));
-				TAABB<T, d> Box1 = MParticles.Geometry(Child1)->BoundingBox();
-
-				TArray<TArray<TPair<uint32, uint32>>> Connections;
-				Connections.Init(TArray<TPair<uint32, uint32>>(), Children.Num() - (i + 1));
-				PhysicsParallelFor(Children.Num() - (i + 1), [&](int32 Idx) {
-					const uint32 Child2 = Children[Idx + i + 1];
-					if (!MParticles.CollisionParticles(Child2))
-					{
-						return;
-					}
-					if (!(ensure(!FMath::IsNaN(MParticles.X(Child2)[0])) && ensure(!FMath::IsNaN(MParticles.X(Child2)[1])) && ensure(!FMath::IsNaN(MParticles.X(Child2)[2]))))
-					{
-						return;
-					}
-					TRigidTransform<T, d> TM = TM1.GetRelativeTransform(TRigidTransform<T, d>(MParticles.X(Child2), MParticles.R(Child2)));
-
-					bool bCollided = false;
-					for (uint32 CollisionIdx = 0; !bCollided && CollisionIdx < MParticles.CollisionParticles(Child2)->Size(); ++CollisionIdx)
-					{
-						TVector<T, d> Normal;
-						TVector<T, d> LocalPoint = TM.TransformPositionNoScale(MParticles.CollisionParticles(Child2)->X(CollisionIdx));
-						T Phi = MParticles.Geometry(Child1)->PhiWithNormal(LocalPoint - (LocalPoint * Delta), Normal);
-						if (Phi < 0.0)
-							bCollided = true;
-					}
-					if (bCollided)
-					{
-						Connections[Idx].Add(TPair<uint32, uint32>((uint32)Child1, (uint32)Child2));
-					}
-				});
-
-				// join results and make connections
-				for (const auto& ConnectionList : Connections)
-				{
-					for (const auto& Elem : ConnectionList)
-					{
-						if (MParticles.ConnectivityEdges(Elem.Key).Num() < Parameters.MaxNumConnections)
-						{
-							const T AvgStrain = (MParticles.Strains(Elem.Key) + MParticles.Strains(Elem.Value)) * (T)0.5f;
-							ConnectNodes(Elem.Key, Elem.Value, AvgStrain);
-						}
-					}
-				}
-			}
-		}
-	}
-
-	template<class T, int d>
-	TVector<T, d> GetContactLocation(const TRigidBodyPointContactConstraint<T, d>& Contact)
-	{
-		return Contact.GetLocation();
-	}
-
-	template<class T, int d>
-	TVector<T, d> GetContactLocation(const TRigidBodyContactConstraintPGS<T, d>& Contact)
-	{
-		// @todo(mlentine): Does the exact point matter?
-		T MinPhi = FLT_MAX;
-		TVector<T, d> MinLoc;
-		for (int32 i = 0; i < Contact.Phi.Num(); ++i)
-		{
-			if (Contact.Phi[i] < MinPhi)
-			{
-				MinPhi = Contact.Phi[i];
-				MinLoc = Contact.Location[i];
-			}
-		}
-		return MinLoc;
-	}
-
-	DECLARE_CYCLE_STAT(TEXT("TPBDRigidClustering<>::UpdateMassProperties"), STAT_UpdateMassProperties, STATGROUP_Chaos);
-	template<class FPBDRigidsEvolution, class FPBDCollisionConstraint, class T, int d>
-	void TPBDRigidClustering<FPBDRigidsEvolution, FPBDCollisionConstraint, T, d>::UpdateMassProperties(const TArray<uint32>& Children, const uint32 ClusterIndex, const TRigidTransform<T, d>* ForceMassOrientation)
-	{
-		SCOPE_CYCLE_COUNTER(STAT_UpdateMassProperties);
-		UpdateClusterMassProperties(MParticles, Children, ClusterIndex, ForceMassOrientation, &MParticles.MultiChildProxyDataArray(), &MParticles.MultiChildProxyIdArray());
-		UpdateKinematicProperties(ClusterIndex);
-	}
-
-	DECLARE_CYCLE_STAT(TEXT("TPBDRigidClustering<>::GetActiveClusterIndex"), STAT_GetActiveClusterIndex, STATGROUP_Chaos);
-	template<class FPBDRigidsEvolution, class FPBDCollisionConstraint, class T, int d>
-	int32 TPBDRigidClustering<FPBDRigidsEvolution, FPBDCollisionConstraint, T, d>::GetActiveClusterIndex(uint32 ChildIndex)
-	{
-		SCOPE_CYCLE_COUNTER(STAT_GetActiveClusterIndex);
-		int32 NumParticles = MParticles.Size();
-		int32 CurrentIndex = (int32)ChildIndex;
-		while (0 <= CurrentIndex && CurrentIndex < NumParticles)
-		{
-			if (!MParticles.Disabled(CurrentIndex))
-			{
-				return CurrentIndex;
-			}
-			CurrentIndex = MParticles.ClusterIds(CurrentIndex).Id;
-		}
-		return INDEX_NONE;
-	}
-
-	template<class FPBDRigidsEvolution, class FPBDCollisionConstraint, class T, int d>
 	void TPBDRigidClustering<FPBDRigidsEvolution, FPBDCollisionConstraint, T, d>::AddUniqueConnection(uint32 Index1, uint32 Index2, T Strain)
 	{
@@ -1859,51 +1914,6 @@
 		MParticles.ConnectivityEdges(ParticleIndex).SetNum(0);
 	}
 
-	DECLARE_CYCLE_STAT(TEXT("TPBDRigidClustering<>::SwapBufferedData"), STAT_SwapBufferedData, STATGROUP_Chaos);
-	template<class FPBDRigidsEvolution, class FPBDCollisionConstraint, class T, int d>
-	void TPBDRigidClustering<FPBDRigidsEvolution, FPBDCollisionConstraint, T, d>::SwapBufferedData()
-	{
-		SCOPE_CYCLE_COUNTER(STAT_SwapBufferedData);
-		ResourceLock.WriteLock();
-		//BufferResource.MChildren.Reset();
-		//BufferResource.ClusterParentTransforms.Reset();	//todo: once everything is atomic this should get reset
-		const TArray<TSerializablePtr<FImplicitObject>>& AllGeom = MParticles.GetAllGeometry();
-		BufferResource.GeometryPtrs.SetNum(AllGeom.Num());
-
-		const auto& NonDisabledClusteredParticles = MEvolution.GetNonDisabledClusteredArray();
-		for (auto& ClusteredParticle : NonDisabledClusteredParticles)
-		{
-			const TArray<uint32>* Children = MChildren.Find(ClusteredParticle);
-			if (Children && ClusteredParticle->ClusterIds().Id == INDEX_NONE) //root cluster so copy children
-			{
-				//TODO: record GT particle pointer instead
-				BufferResource.MChildren.Add(ClusteredParticle->TransientParticleIndex(), *Children);
-				BufferResource.ClusterParentTransforms.Add(ClusteredParticle->TransientParticleIndex(), TRigidTransform<float, 3>(ClusteredParticle->X(), ClusteredParticle->R()));
-			}
-		}
-
-		BufferResource.GeometryPtrs = AllGeom; //in future this should be sparse. SQ has fallback that relies on potentially all geom so can't do it yet
-
-		ResourceLock.WriteUnlock();
-	}
-
-	template<class FPBDRigidsEvolution, class FPBDCollisionConstraint, class T, int d>
-	void TPBDRigidClustering<FPBDRigidsEvolution, FPBDCollisionConstraint, T, d>::IncrementPendingClusterCounter(uint32 ClusterGroupID)
-	{
-		if (!PendingClusterCounter.Contains(ClusterGroupID))
-		{
-			PendingClusterCounter.Add(ClusterGroupID, 0);
-		}
-		PendingClusterCounter[ClusterGroupID]++;
-	}
-
-	template<class FPBDRigidsEvolution, class FPBDCollisionConstraint, class T, int d>
-	void TPBDRigidClustering<FPBDRigidsEvolution, FPBDCollisionConstraint, T, d>::DecrementPendingClusterCounter(uint32 ClusterGroupID)
-	{
-		PendingClusterCounter[ClusterGroupID]--;
-		ensure(0 <= PendingClusterCounter[ClusterGroupID]);
-	}
-
 	/*
 	template<class FPBDRigidsEvolution, class FPBDCollisionConstraint, class T, int d>
 	void TPBDRigidClustering<FPBDRigidsEvolution, FPBDCollisionConstraint, T, d>::ClearPendingClusterCounter(uint32 ClusterGroupID)
@@ -1912,7 +1922,7 @@
 		//ensure(0 <= PendingClusterCounter[ClusterGroupID]);
 	}
 */
-}
+} // namespace Chaos
 
 using namespace Chaos;
 template class CHAOS_API Chaos::TPBDRigidClustering<TPBDRigidsEvolutionGBF<float, 3>, TPBDCollisionConstraints<float, 3>, float, 3>;
