Chaos - Add Inertia Scaling utility method

diff --git a/Engine/Source/Runtime/Experimental/Chaos/Public/Chaos/Utilities.h b/Engine/Source/Runtime/Experimental/Chaos/Public/Chaos/Utilities.h
--- a/Engine/Source/Runtime/Experimental/Chaos/Public/Chaos/Utilities.h
+++ b/Engine/Source/Runtime/Experimental/Chaos/Public/Chaos/Utilities.h
@@ -286,5 +286,17 @@ namespace Chaos
 			return OutTA >= 0 && OutTA <= 1 && OutTB > 0 && OutTB < 1;
 		}
 
+		/**
+		 * Given the local-space inertia for an unscaled object, return an inertia as if generated from a non-uniformly scaled shape with the specified scale.
+		 */
+		inline Chaos::FVec3 ScaleInertia(const Chaos::FVec3& Inertia, const Chaos::FVec3& Scale)
+		{
+			using namespace Chaos;
+
+			FReal ScaleVol = Scale.X * Scale.Y * Scale.Z;
+			FVec3 ScaleSq = Scale * Scale;
+			return 0.5f * ScaleVol * (FVec3(Inertia.Y + Inertia.Z, Inertia.X + Inertia.Z, Inertia.X + Inertia.Y) * ScaleSq);
+		}
+
 	} // namespace Utilities
 } // namespace Chaos
\ No newline at end of file
