Chaos: Fix FQueryFastData using uniniitalized values in bParallel.

diff --git a/Engine/Source/Runtime/Experimental/Chaos/Public/Chaos/ISpatialAcceleration.h b/Engine/Source/Runtime/Experimental/Chaos/Public/Chaos/ISpatialAcceleration.h
--- a/Engine/Source/Runtime/Experimental/Chaos/Public/Chaos/ISpatialAcceleration.h
+++ b/Engine/Source/Runtime/Experimental/Chaos/Public/Chaos/ISpatialAcceleration.h
@@ -11,7 +11,7 @@ struct CHAOS_API FQueryFastData
 {
 	FQueryFastData(const FVec3& InDir, const FReal InLength)
 		: Dir(InDir)
-		, InvDir(bParallel[0] ? 0 : 1 / Dir[0], bParallel[1] ? 0 : 1 / Dir[1], bParallel[2] ? 0 : 1 / Dir[2])
+		, InvDir( (InDir[0] == 0) ? 0 : 1 / Dir[0], (InDir[1] == 0) ? 0 : 1 / Dir[1], (InDir[2] == 0) ? 0 : 1 / Dir[2])
 		, bParallel{ InDir[0] == 0, InDir[1] == 0, InDir[2] == 0 }
 	{
 		ensure(InLength);
