Change FQueryFastData ensure(InLength) to Chaos ensure.

diff --git a/Engine/Source/Runtime/Experimental/Chaos/Public/Chaos/ISpatialAcceleration.h b/Engine/Source/Runtime/Experimental/Chaos/Public/Chaos/ISpatialAcceleration.h
--- a/Engine/Source/Runtime/Experimental/Chaos/Public/Chaos/ISpatialAcceleration.h
+++ b/Engine/Source/Runtime/Experimental/Chaos/Public/Chaos/ISpatialAcceleration.h
@@ -3,6 +3,8 @@
 #include "Chaos/Vector.h"
 #include "Chaos/Box.h"
 #include "GeometryParticlesfwd.h"
+#include "ChaosCheck.h"
+
 
 namespace Chaos
 {
@@ -14,7 +16,7 @@ struct CHAOS_API FQueryFastData
 		, InvDir( (InDir[0] == 0) ? 0 : 1 / Dir[0], (InDir[1] == 0) ? 0 : 1 / Dir[1], (InDir[2] == 0) ? 0 : 1 / Dir[2])
 		, bParallel{ InDir[0] == 0, InDir[1] == 0, InDir[2] == 0 }
 	{
-		ensure(InLength);
+		CHAOS_ENSURE(InLength);
 		SetLength(InLength);
 	}
 
