Track chaos broadphase memory

diff --git a/Engine/Source/Runtime/Experimental/Chaos/Private/Chaos/PBDRigidsEvolution.cpp b/Engine/Source/Runtime/Experimental/Chaos/Private/Chaos/PBDRigidsEvolution.cpp
--- a/Engine/Source/Runtime/Experimental/Chaos/Private/Chaos/PBDRigidsEvolution.cpp
+++ b/Engine/Source/Runtime/Experimental/Chaos/Private/Chaos/PBDRigidsEvolution.cpp
@@ -224,7 +224,7 @@ namespace Chaos
 	template<class FPBDRigidsEvolution, class FPBDCollisionConstraint, class T, int d>
 	void TPBDRigidsEvolutionBase<FPBDRigidsEvolution, FPBDCollisionConstraint, T, d>::FChaosAccelerationStructureTask::DoTask(ENamedThreads::Type CurrentThread, const FGraphEventRef& MyCompletionGraphEvent)
 	{
-		LLM_SCOPE(ELLMTag::Chaos);
+		LLM_SCOPE(ELLMTag::ChaosAcceleration);
 
 		uint8 ActiveBucketsMask = SpatialCollectionFactory.GetActiveBucketsMask();
 		TArray<TSOAView<TSpatialAccelerationCache<T, d>>> ViewsPerBucket[8];
@@ -416,6 +416,7 @@ namespace Chaos
 	template<class FPBDRigidsEvolution, class FPBDCollisionConstraint, class T, int d>
 	void TPBDRigidsEvolutionBase<FPBDRigidsEvolution, FPBDCollisionConstraint, T, d>::ComputeIntermediateSpatialAcceleration(bool bBlock)
 	{
+		LLM_SCOPE(ELLMTag::ChaosAcceleration);
 		SCOPE_CYCLE_COUNTER(STAT_ComputeIntermediateSpatialAcceleration);
 		CHAOS_SCOPED_TIMER(ComputeIntermediateSpatialAcceleration);
 
@@ -484,7 +485,7 @@ namespace Chaos
 	void TPBDRigidsEvolutionBase<FPBDRigidsEvolution, FPBDCollisionConstraint, T, d>::UpdateExternalAccelerationStructure(TUniquePtr<ISpatialAccelerationCollection<TAccelerationStructureHandle<T, d>, T, d>>& StructToUpdate)
 	{
 		DECLARE_SCOPE_CYCLE_COUNTER(TEXT("CreateExternalAccelerationStructure"), STAT_CreateExternalAccelerationStructure, STATGROUP_Physics);
-
+		LLM_SCOPE(ELLMTag::ChaosAcceleration);
 		check(!bIsSingleThreaded);
 
 		if (bExternalReady)
diff --git a/Engine/Source/Runtime/Experimental/Chaos/Public/Chaos/BoundingVolume.h b/Engine/Source/Runtime/Experimental/Chaos/Public/Chaos/BoundingVolume.h
--- a/Engine/Source/Runtime/Experimental/Chaos/Public/Chaos/BoundingVolume.h
+++ b/Engine/Source/Runtime/Experimental/Chaos/Public/Chaos/BoundingVolume.h
@@ -208,7 +208,6 @@ public:
 
 	virtual void RemoveElement(const TPayloadType& Payload) override
 	{
-		LLM_SCOPE(ELLMTag::ChaosAcceleration);
 		SCOPE_CYCLE_COUNTER(STAT_BoundingVolumeRemoveElement);
 		if (const FPayloadInfo* PayloadInfo = MPayloadInfo.Find(Payload))
 		{
@@ -233,7 +232,6 @@ public:
 
 	virtual void UpdateElement(const TPayloadType& Payload, const TAABB<T,d>& NewBounds, bool bHasBounds) override
 	{
-		LLM_SCOPE(ELLMTag::ChaosAcceleration);
 		SCOPE_CYCLE_COUNTER(STAT_BoundingVolumeUpdateElement);
 		if (FPayloadInfo* PayloadInfo = MPayloadInfo.Find(Payload))
 		{
@@ -317,7 +315,6 @@ public:
 
 	virtual void Serialize(FChaosArchive& Ar) override
 	{
-		LLM_SCOPE(ELLMTag::ChaosAcceleration);
 		Ar.UsingCustomVersion(FExternalPhysicsCustomObjectVersion::GUID);
 		if (Ar.CustomVer(FExternalPhysicsCustomObjectVersion::GUID) < FExternalPhysicsCustomObjectVersion::GlobalElementsHaveBounds)
 		{
@@ -817,7 +814,6 @@ private:
 	template <typename ParticleView>
 	void GenerateTree(const ParticleView& Particles, const bool bUseVelocity, const T Dt, const int32 MaxCells)
 	{
-		LLM_SCOPE(ELLMTag::ChaosAcceleration);
 		MGlobalPayloads.Reset();
 		MPayloadInfo.Reset();
 		bIsEmpty = true;
diff --git a/Engine/Source/Runtime/Experimental/Chaos/Public/Chaos/KinematicGeometryParticles.h b/Engine/Source/Runtime/Experimental/Chaos/Public/Chaos/KinematicGeometryParticles.h
--- a/Engine/Source/Runtime/Experimental/Chaos/Public/Chaos/KinematicGeometryParticles.h
+++ b/Engine/Source/Runtime/Experimental/Chaos/Public/Chaos/KinematicGeometryParticles.h
@@ -115,7 +115,6 @@ class TKinematicGeometryParticlesImp : public TGeometryParticlesImp<T, d, SimTyp
 
 	CHAOS_API virtual void Serialize(FChaosArchive& Ar) override
 	{
-		LLM_SCOPE(ELLMTag::ChaosParticles);
 		TGeometryParticlesImp<T, d, SimType>::Serialize(Ar);
 		Ar << MV << MW;
 		
diff --git a/Engine/Source/Runtime/Experimental/Chaos/Public/Chaos/PBDRigidParticles.h b/Engine/Source/Runtime/Experimental/Chaos/Public/Chaos/PBDRigidParticles.h
--- a/Engine/Source/Runtime/Experimental/Chaos/Public/Chaos/PBDRigidParticles.h
+++ b/Engine/Source/Runtime/Experimental/Chaos/Public/Chaos/PBDRigidParticles.h
@@ -162,7 +162,6 @@ class TPBDRigidParticles : public TRigidParticles<T, d>
 
 	CHAOS_API virtual void Serialize(FChaosArchive& Ar) override
 	{
-		LLM_SCOPE(ELLMTag::ChaosParticles);
 		TRigidParticles<T, d>::Serialize(Ar);
 		Ar << MP << MQ << MPreV << MPreW;
 	}
diff --git a/Engine/Source/Runtime/Experimental/Chaos/Public/Chaos/Particles.h b/Engine/Source/Runtime/Experimental/Chaos/Public/Chaos/Particles.h
--- a/Engine/Source/Runtime/Experimental/Chaos/Public/Chaos/Particles.h
+++ b/Engine/Source/Runtime/Experimental/Chaos/Public/Chaos/Particles.h
@@ -67,14 +67,12 @@ namespace Chaos
 
 		void AddParticles(const int32 Num)
 		{
-			LLM_SCOPE(ELLMTag::ChaosParticles);
 			AddElementsHelper(Num);
 			IncrementDirtyValidation();
 		}
 
 		void DestroyParticle(const int32 Idx)
 		{
-			LLM_SCOPE(ELLMTag::ChaosParticles);
 			if (MRemoveParticleBehavior == ERemoveParticleBehavior::RemoveAtSwap)
 			{
 				RemoveAtSwapHelper(Idx);
@@ -121,7 +119,6 @@ namespace Chaos
 		
 		void Serialize(FArchive& Ar)
 		{
-			LLM_SCOPE(ELLMTag::ChaosParticles);
 			bool bSerialize = true;	//leftover from when we had view support
 			Ar << bSerialize;
 			if (ensureMsgf(bSerialize, TEXT("Cannot serialize shared views. Refactor needed to reduce memory")))
diff --git a/Engine/Source/Runtime/Experimental/Chaos/Public/Chaos/RigidParticles.h b/Engine/Source/Runtime/Experimental/Chaos/Public/Chaos/RigidParticles.h
--- a/Engine/Source/Runtime/Experimental/Chaos/Public/Chaos/RigidParticles.h
+++ b/Engine/Source/Runtime/Experimental/Chaos/Public/Chaos/RigidParticles.h
@@ -193,7 +193,6 @@ class TRigidParticles : public TKinematicGeometryParticles<T, d>
 
 	CHAOS_API virtual void Serialize(FChaosArchive& Ar) override
 	{
-		LLM_SCOPE(ELLMTag::ChaosParticles);
 		TKinematicGeometryParticles<T,d>::Serialize(Ar);
 		Ar << MF << MT << MLinearImpulse << MAngularImpulse << MI << MInvI << MM << MInvM;
 
