Use instanced when scale is 1,1,1

diff --git a/Engine/Source/Runtime/Engine/Private/PhysicsEngine/Experimental/ChaosInterfaceUtils.cpp b/Engine/Source/Runtime/Engine/Private/PhysicsEngine/Experimental/ChaosInterfaceUtils.cpp
--- a/Engine/Source/Runtime/Engine/Private/PhysicsEngine/Experimental/ChaosInterfaceUtils.cpp
+++ b/Engine/Source/Runtime/Engine/Private/PhysicsEngine/Experimental/ChaosInterfaceUtils.cpp
@@ -262,7 +262,16 @@ namespace ChaosInterface
 				//}
 				//else
 				{
-					auto Implicit = MakeUnique<Chaos::TImplicitObjectScaled<Chaos::FConvex>>(MakeSerializable(ConvexImplicit), Scale);
+					TUniquePtr<Chaos::FImplicitObject> Implicit;
+					if(Scale == FVector(1))
+					{
+						Implicit = TUniquePtr<Chaos::FImplicitObject>(new Chaos::TImplicitObjectInstanced<Chaos::FConvex>(ConvexImplicit));
+					}
+					else
+					{
+						Implicit = TUniquePtr<Chaos::FImplicitObject>(new Chaos::TImplicitObjectScaled<Chaos::FConvex>(ConvexImplicit,Scale));
+					}
+
 					auto NewShape = NewShapeHelper(MakeSerializable(Implicit), (void*)CollisionBody.GetUserData());
 					Shapes.Emplace(MoveTemp(NewShape));
 					Geoms.Add(MoveTemp(Implicit));
@@ -272,7 +281,15 @@ namespace ChaosInterface
 
 		for (const auto& ChaosTriMesh : InParams.ChaosTriMeshes)
 		{
-			auto Implicit = MakeUnique<Chaos::TImplicitObjectScaled<Chaos::FTriangleMeshImplicitObject>>(ChaosTriMesh, Scale);
+			TUniquePtr<Chaos::FImplicitObject> Implicit;
+			if(Scale == FVector(1))
+			{
+				Implicit = TUniquePtr<Chaos::FImplicitObject>(new Chaos::TImplicitObjectInstanced<Chaos::FTriangleMeshImplicitObject>(ChaosTriMesh));
+			} else
+			{
+				Implicit = TUniquePtr<Chaos::FImplicitObject>(new Chaos::TImplicitObjectScaled<Chaos::FTriangleMeshImplicitObject>(ChaosTriMesh,Scale));
+			}
+			
 			auto NewShape = NewShapeHelper(MakeSerializable(Implicit), nullptr, true);
 			Shapes.Emplace(MoveTemp(NewShape));
 			Geoms.Add(MoveTemp(Implicit));
diff --git a/Engine/Source/Runtime/Experimental/Chaos/Public/Chaos/GeometryQueries.h b/Engine/Source/Runtime/Experimental/Chaos/Public/Chaos/GeometryQueries.h
--- a/Engine/Source/Runtime/Experimental/Chaos/Public/Chaos/GeometryQueries.h
+++ b/Engine/Source/Runtime/Experimental/Chaos/Public/Chaos/GeometryQueries.h
@@ -33,6 +33,22 @@ namespace Chaos
 		case ImplicitObjectType::IsScaled | ImplicitObjectType::Box: return Func(Geom.template GetObjectChecked< TImplicitObjectScaled<TBox<FReal, 3>>>(), TM);
 		case ImplicitObjectType::IsScaled | ImplicitObjectType::Capsule: return Func(Geom.template GetObjectChecked< TImplicitObjectScaled<TCapsule<FReal>>>(), TM);
 		case ImplicitObjectType::IsScaled | ImplicitObjectType::Convex: return Func(Geom.template GetObjectChecked< TImplicitObjectScaled<FConvex>>(), TM);
+		case ImplicitObjectType::IsInstanced | ImplicitObjectType::Sphere:
+		{
+			return Func(Geom.template GetObjectChecked< TImplicitObjectInstanced<TSphere<FReal,3>>>().GetInstancedObject()->template GetObjectChecked<TSphere<FReal,3>>(),TM);
+		}
+		case ImplicitObjectType::IsInstanced | ImplicitObjectType::Box:
+		{
+			return Func(Geom.template GetObjectChecked< TImplicitObjectInstanced<TBox<FReal,3>>>().GetInstancedObject()->template GetObjectChecked<TBox<FReal,3>>(),TM);
+		}
+		case ImplicitObjectType::IsInstanced | ImplicitObjectType::Capsule:
+		{
+			return Func(Geom.template GetObjectChecked< TImplicitObjectInstanced<TCapsule<FReal>>>().GetInstancedObject()->template GetObjectChecked<TCapsule<FReal>>(),TM);
+		}
+		case ImplicitObjectType::IsInstanced | ImplicitObjectType::Convex:
+		{
+			return Func(Geom.template GetObjectChecked< TImplicitObjectInstanced<FConvex>>().GetInstancedObject()->template GetObjectChecked<FConvex>(),TM);
+		}
 		case ImplicitObjectType::Transformed:
 		{
 			const auto& ImplicitObjectTransformed = (Geom.template GetObjectChecked<TImplicitObjectTransformed<FReal,3>>());
@@ -107,6 +123,11 @@ namespace Chaos
 					const auto& AScaled = TImplicitObjectScaled<FTriangleMeshImplicitObject>::AsScaledChecked(A);
 					return AScaled.LowLevelOverlapGeom(B, BToATM, Thickness);
 				}
+				else if(IsInstanced(AType))
+				{
+					const auto& AInstanced = TImplicitObjectInstanced<FTriangleMeshImplicitObject>::AsInstancedChecked(A);
+					return AInstanced.LowLevelOverlapGeom(B,BToATM,Thickness);
+				}
 				else
 				{
 					check(false);	//unsupported query type
@@ -212,6 +233,12 @@ namespace Chaos
 					bResult = AScaled.LowLevelSweepGeom(B, BToATM, LocalDir, Length, OutTime, LocalPosition, LocalNormal, OutFaceIndex, Thickness, bComputeMTD);
 					break;
 				}
+				else if(IsInstanced(AType))
+				{
+					const auto& Instanced = TImplicitObjectInstanced<FTriangleMeshImplicitObject>::AsInstancedChecked(A);
+					bResult = Instanced.LowLevelSweepGeom(B,BToATM,LocalDir,Length,OutTime,LocalPosition,LocalNormal,OutFaceIndex,Thickness,bComputeMTD);
+					break;
+				}
 				else
 				{
 					check(false);	//unsupported query type
diff --git a/Engine/Source/Runtime/Experimental/Chaos/Public/Chaos/ImplicitObjectScaled.h b/Engine/Source/Runtime/Experimental/Chaos/Public/Chaos/ImplicitObjectScaled.h
--- a/Engine/Source/Runtime/Experimental/Chaos/Public/Chaos/ImplicitObjectScaled.h
+++ b/Engine/Source/Runtime/Experimental/Chaos/Public/Chaos/ImplicitObjectScaled.h
@@ -13,7 +13,7 @@
 namespace Chaos
 {
 
-template <typename TConcrete, bool bInstanced = true>
+template <typename TConcrete>
 class TImplicitObjectInstanced final : public FImplicitObject
 {
 public:
@@ -21,11 +21,11 @@ public:
 	using TType = T;
 	static constexpr int d = TConcrete::D;
 	static constexpr int D = d;
+	using ObjectType = TSharedPtr<TConcrete,ESPMode::ThreadSafe>;
 
-	using ObjectType = typename TChooseClass<bInstanced, TSerializablePtr<TConcrete>, TUniquePtr<TConcrete>>::Result;
 	using FImplicitObject::GetTypeName;
 
-	TImplicitObjectInstanced(ObjectType Object)
+	TImplicitObjectInstanced(const ObjectType&& Object)
 		: FImplicitObject(EImplicitObject::HasBoundingBox, Object->GetType() | ImplicitObjectType::IsInstanced)
 		, MObject(MoveTemp(Object))
 	{
@@ -34,6 +34,15 @@ public:
 		this->bDoCollide = MObject->GetDoCollide();
 	}
 
+	TImplicitObjectInstanced(const ObjectType& Object)
+		: FImplicitObject(EImplicitObject::HasBoundingBox,Object->GetType() | ImplicitObjectType::IsInstanced)
+		,MObject(Object)
+	{
+		ensure(IsInstanced(MObject->GetType(true)) == false);	//cannot have an instance of an instance
+		this->bIsConvex = MObject->IsConvex();
+		this->bDoCollide = MObject->GetDoCollide();
+	}
+
 	static constexpr EImplicitObjectType StaticType()
 	{
 		return TConcrete::StaticType() | ImplicitObjectType::IsInstanced;
@@ -66,21 +75,6 @@ public:
 		Ar << MObject;
 	}
 
-	/** This is a low level function and assumes the internal object has a SweepGeom function. Should not be called directly. See GeometryQueries.h : SweepQuery */
-	template <typename TGeometry>
-	static bool LowLevelSweepGeom(const TImplicitObjectInstanced<TConcrete, bInstanced>& OwningInstanced, const FImplicitObject& B, const TRigidTransform<T, d>& BToATM, const TVector<T, d>& LocalDir, const T Length, T& OutTime, TVector<T, d>& LocalPosition, TVector<T, d>& LocalNormal, int32& OutFaceIndex, T Thickness = 0)
-	{
-		ensure(false);
-		return false;
-	}
-
-	/** This is a low level function and assumes the internal object has a OverlapGeom function. Should not be called directly. See GeometryQueries.h : OverlapQuery */
-	bool LowLevelOverlapGeom(const TImplicitObjectInstanced<TConcrete, bInstanced>& OwningInstanced, const FImplicitObject& B, const TRigidTransform<T, d>& BToATM, T Thickness = 0) const
-	{
-		ensure(false);
-		return false;
-	}
-
 	virtual int32 FindMostOpposingFace(const TVector<T, d>& Position, const TVector<T, d>& UnitDir, int32 HintFaceIndex, T SearchDist) const override
 	{
 		return MObject->FindMostOpposingFace(Position, UnitDir, HintFaceIndex, SearchDist);
@@ -115,17 +109,39 @@ public:
 		return TUniquePtr<FImplicitObject>(CopyHelper(this));
 	}
 
-protected:
-	ObjectType MObject;
+	static const TImplicitObjectInstanced<TConcrete>& AsInstancedChecked(const FImplicitObject& Obj)
+	{
+		if(TIsSame<TConcrete,FImplicitObject>::Value)
+		{
+			//can cast any instanced to ImplicitObject base
+			check(IsInstanced(Obj.GetType()));
+		} else
+		{
+			check(StaticType() == Obj.GetType());
+		}
+		return static_cast<const TImplicitObjectInstanced<TConcrete>&>(Obj);
+	}
+
+	/** This is a low level function and assumes the internal object has a SweepGeom function. Should not be called directly. See GeometryQueries.h : SweepQuery */
+	template <typename QueryGeomType>
+	bool LowLevelSweepGeom(const QueryGeomType& B,const TRigidTransform<T,d>& BToATM,const TVector<T,d>& LocalDir,const T Length,T& OutTime,TVector<T,d>& LocalPosition,TVector<T,d>& LocalNormal,int32& OutFaceIndex,T Thickness = 0,bool bComputeMTD = false) const
+	{
+		return MObject->SweepGeom(B,BToATM,LocalDir,Length,OutTime,LocalPosition,LocalNormal,OutFaceIndex,Thickness,bComputeMTD);
+	}
 
-	static TImplicitObjectInstanced<TConcrete, true>* CopyHelper(const TImplicitObjectInstanced<TConcrete, true>* Obj)
+	/** This is a low level function and assumes the internal object has a OverlapGeom function. Should not be called directly. See GeometryQueries.h : OverlapQuery */
+	template <typename QueryGeomType>
+	bool LowLevelOverlapGeom(const QueryGeomType& B,const TRigidTransform<T,d>& BToATM,T Thickness = 0) const
 	{
-		return new TImplicitObjectInstanced<TConcrete, true>(Obj->MObject);
+		return MObject->OverlapGeom(B,BToATM,Thickness);
 	}
 
-	static TImplicitObjectInstanced<TConcrete, false>* CopyHelper(const TImplicitObjectInstanced<TConcrete, false>* Obj)
+protected:
+	ObjectType MObject;
+
+	static TImplicitObjectInstanced<TConcrete>* CopyHelper(const TImplicitObjectInstanced<TConcrete>* Obj)
 	{
-		return new TImplicitObjectInstanced<TConcrete, false>(Obj->MObject->Copy());
+		return new TImplicitObjectInstanced<TConcrete>(Obj->MObject);
 	}
 };
 
