Chaos - Joint ISPC fixed and re-enabled

diff --git a/Engine/Source/Runtime/Core/Public/Math/Vector.isph b/Engine/Source/Runtime/Core/Public/Math/Vector.isph
--- a/Engine/Source/Runtime/Core/Public/Math/Vector.isph
+++ b/Engine/Source/Runtime/Core/Public/Math/Vector.isph
@@ -1811,4 +1811,14 @@ inline uniform FVector VectorNegate(const uniform FVector &A)
 	return *((uniform FVector *uniform)&Result);
 }
 
+inline uniform float VectorSizeSquared(const uniform FVector4 &A)
+{
+	return VectorDot(A, A);
+}
+
+inline uniform float VectorSizeSquared(const uniform FVector &A)
+{
+	return VectorDot(A, A);
+}
+
 #endif
diff --git a/Engine/Source/Runtime/Experimental/Chaos/Private/Chaos/Joint/PBDJointSolverGaussSeidel.cpp b/Engine/Source/Runtime/Experimental/Chaos/Private/Chaos/Joint/PBDJointSolverGaussSeidel.cpp
--- a/Engine/Source/Runtime/Experimental/Chaos/Private/Chaos/Joint/PBDJointSolverGaussSeidel.cpp
+++ b/Engine/Source/Runtime/Experimental/Chaos/Private/Chaos/Joint/PBDJointSolverGaussSeidel.cpp
@@ -14,9 +14,9 @@
 #if !INTEL_ISPC
 const bool bChaos_Joint_ISPC_Enabled = false;
 #elif UE_BUILD_SHIPPING
-const bool bChaos_Joint_ISPC_Enabled = false;
+const bool bChaos_Joint_ISPC_Enabled = true;
 #else
-bool bChaos_Joint_ISPC_Enabled = false;
+bool bChaos_Joint_ISPC_Enabled = true;
 FAutoConsoleVariableRef CVarChaosJointISPCEnabled(TEXT("p.Chaos.Joint.ISPC"), bChaos_Joint_ISPC_Enabled, TEXT("Whether to use ISPC optimizations in the Joint Solver"));
 #endif
 
diff --git a/Engine/Source/Runtime/Experimental/Chaos/Private/Chaos/Joint/PBDJointSolverGaussSeidel.ispc b/Engine/Source/Runtime/Experimental/Chaos/Private/Chaos/Joint/PBDJointSolverGaussSeidel.ispc
--- a/Engine/Source/Runtime/Experimental/Chaos/Private/Chaos/Joint/PBDJointSolverGaussSeidel.ispc
+++ b/Engine/Source/Runtime/Experimental/Chaos/Private/Chaos/Joint/PBDJointSolverGaussSeidel.ispc
@@ -50,7 +50,8 @@ struct FJointSolverGaussSeidel
 static inline uniform FMatrix33 ComputeWorldSpaceInertia(const uniform FVector4 &Q, const uniform FVector &I)
 {
 	const uniform FMatrix33 QM = MakeQuatRotationTranslationMatrix(Q);
-	return MultiplyAB(QM, MultiplyABt(SetMatrix33(I.V[0], I.V[1], I.V[2]), QM));
+	const uniform FMatrix33 L = SetMatrix33(I.V[0], I.V[1], I.V[2]);
+	return MultiplyAB(QM, MultiplyABt(L, QM));
 }
 
 static inline uniform FMatrix33 ComputeJointFactorMatrix(const uniform FVector& V, const uniform FMatrix33& M, const uniform float Im)
@@ -228,8 +229,34 @@ export void ApplyRotationConstraintSoft(
 	const uniform FVector IA1 = Multiply(InvI1, Axis1);
 
 	// Joint-space inverse mass
-	const uniform float II0 = VectorDot(Axis0, IA0);
-	const uniform float II1 = VectorDot(Axis1, IA1);
+	uniform float II0 = VectorDot(Axis0, IA0);
+	uniform float II1 = VectorDot(Axis1, IA1);
+
+	// If we are correcting the position, we need to adjust the constraint effective mass using parallel axis theorem
+	// @todo(ccaulfield): the IMs here are constant per constraint...pre-build and cache it
+	if (M->AngularPositionCorrection > 0)
+	{
+		if (M->InvMs[0] > 0)
+		{
+			const uniform FVector LinearAxis0 = VectorCross(M->Xs[0] - M->Ps[0], Axis0);
+			const uniform float LinearAxisLen0Sq = VectorSizeSquared(LinearAxis0);
+			if (LinearAxisLen0Sq > KINDA_SMALL_NUMBER)
+			{
+				const uniform float IM0 = M->InvMs[0] / LinearAxisLen0Sq;
+				II0 = II0 * IM0 / (II0 + IM0);
+			}
+		}
+		if (M->InvMs[1] > 0)
+		{
+			const uniform FVector LinearAxis1 = VectorCross(M->Xs[1] - M->Ps[1], Axis1);
+			const uniform float LinearAxisLen1Sq = VectorSizeSquared(LinearAxis1);
+			if (LinearAxisLen1Sq > KINDA_SMALL_NUMBER)
+			{
+				const uniform float IM1 = M->InvMs[1] / LinearAxisLen1Sq;
+				II1 = II1 * IM1 / (II1 + IM1);
+			}
+		}
+	}
 	const uniform float II = (II0 + II1);
 
 	// Damping angular velocity
diff --git a/Engine/Source/Runtime/Experimental/ChaosCore/Public/Chaos/Matrix33.isph b/Engine/Source/Runtime/Experimental/ChaosCore/Public/Chaos/Matrix33.isph
--- a/Engine/Source/Runtime/Experimental/ChaosCore/Public/Chaos/Matrix33.isph
+++ b/Engine/Source/Runtime/Experimental/ChaosCore/Public/Chaos/Matrix33.isph
@@ -38,9 +38,9 @@ inline uniform FMatrix33 SetMatrix33(const uniform FVector4& x0, const uniform F
 
 inline uniform FMatrix33 MultiplyAB(const uniform FMatrix33& L, const uniform FMatrix33& R)
 {
-	const uniform FVector4 L0 = *((uniform FVector4 *uniform)&L.M[0]);
-	const uniform FVector4 L1 = *((uniform FVector4 *uniform)&L.M[4]);
-	const uniform FVector4 L2 = *((uniform FVector4 *uniform)&L.M[8]);
+	const uniform FVector4 L0 = SetVector4(L.M[0], L.M[1], L.M[2], L.M[3]);
+	const uniform FVector4 L1 = SetVector4(L.M[4], L.M[5], L.M[6], L.M[7]);
+	const uniform FVector4 L2 = SetVector4(L.M[8], L.M[9], L.M[10], L.M[11]);
 
 	uniform FVector4 x0 = SetVector4(R.M[0]) * L0 + SetVector4(R.M[1]) * L1 + SetVector4(R.M[2]) * L2;
 	uniform FVector4 x1 = SetVector4(R.M[4]) * L0 + SetVector4(R.M[5]) * L1 + SetVector4(R.M[6]) * L2;
@@ -55,9 +55,9 @@ inline uniform FMatrix33 MultiplyAB(const uniform FMatrix33& L, const uniform FM
 
 inline uniform FMatrix33 MultiplyABt(const uniform FMatrix33& L, const uniform FMatrix33& R)
 {
-	const uniform FVector4 L0 = *((uniform FVector4 *uniform)&L.M[0]);
-	const uniform FVector4 L1 = *((uniform FVector4 *uniform)&L.M[4]);
-	const uniform FVector4 L2 = *((uniform FVector4 *uniform)&L.M[8]);
+	const uniform FVector4 L0 = SetVector4(L.M[0], L.M[1], L.M[2], L.M[3]);
+	const uniform FVector4 L1 = SetVector4(L.M[4], L.M[5], L.M[6], L.M[7]);
+	const uniform FVector4 L2 = SetVector4(L.M[8], L.M[9], L.M[10], L.M[11]);
 
 	uniform FVector4 x0 = SetVector4(R.M[0]) * L0 + SetVector4(R.M[4]) * L1 + SetVector4(R.M[8]) * L2;
 	uniform FVector4 x1 = SetVector4(R.M[1]) * L0 + SetVector4(R.M[5]) * L1 + SetVector4(R.M[9]) * L2;
@@ -70,11 +70,11 @@ inline uniform FMatrix33 MultiplyABt(const uniform FMatrix33& L, const uniform F
 	return SetMatrix33(x0, x1, x2);
 }
 
-inline uniform FVector Multiply(const uniform FMatrix33& LIn, const uniform FVector& R)
+inline uniform FVector Multiply(const uniform FMatrix33& L, const uniform FVector& R)
 {
-	const uniform FVector4 L0 = *((uniform FVector4 *uniform)&LIn.M[0]);
-	const uniform FVector4 L1 = *((uniform FVector4 *uniform)&LIn.M[4]);
-	const uniform FVector4 L2 = *((uniform FVector4 *uniform)&LIn.M[8]);
+	const uniform FVector4 L0 = SetVector4(L.M[0], L.M[1], L.M[2], L.M[3]);
+	const uniform FVector4 L1 = SetVector4(L.M[4], L.M[5], L.M[6], L.M[7]);
+	const uniform FVector4 L2 = SetVector4(L.M[8], L.M[9], L.M[10], L.M[11]);
 
 	const uniform FVector4 Result = L0 * SetVector4(R.V[0]) + L1 * SetVector4(R.V[1]) + L2 * SetVector4(R.V[2]);
 
@@ -83,13 +83,13 @@ inline uniform FVector Multiply(const uniform FMatrix33& LIn, const uniform FVec
 
 inline uniform FMatrix33 AddAB(const uniform FMatrix33& L, const uniform FMatrix33& R)
 {
-	const uniform FVector4 L0 = *((uniform FVector4 *uniform)&L.M[0]);
-	const uniform FVector4 L1 = *((uniform FVector4 *uniform)&L.M[4]);
-	const uniform FVector4 L2 = *((uniform FVector4 *uniform)&L.M[8]);
+	const uniform FVector4 L0 = SetVector4(L.M[0], L.M[1], L.M[2], L.M[3]);
+	const uniform FVector4 L1 = SetVector4(L.M[4], L.M[5], L.M[6], L.M[7]);
+	const uniform FVector4 L2 = SetVector4(L.M[8], L.M[9], L.M[10], L.M[11]);
 
-	const uniform FVector4 R0 = *((uniform FVector4 *uniform)&R.M[0]);
-	const uniform FVector4 R1 = *((uniform FVector4 *uniform)&R.M[4]);
-	const uniform FVector4 R2 = *((uniform FVector4 *uniform)&R.M[8]);
+	const uniform FVector4 R0 = SetVector4(R.M[0], R.M[1], R.M[2], R.M[3]);
+	const uniform FVector4 R1 = SetVector4(R.M[4], R.M[5], R.M[6], R.M[7]);
+	const uniform FVector4 R2 = SetVector4(R.M[8], R.M[9], R.M[10], R.M[11]);
 
 	uniform FVector4 x0 = R0 + L0;
 	uniform FVector4 x1 = R1 + L1;
