Add Lock for adding to SleepData array.

diff --git a/Engine/Source/Runtime/Experimental/Chaos/Public/Chaos/RigidParticles.h b/Engine/Source/Runtime/Experimental/Chaos/Public/Chaos/RigidParticles.h
--- a/Engine/Source/Runtime/Experimental/Chaos/Public/Chaos/RigidParticles.h
+++ b/Engine/Source/Runtime/Experimental/Chaos/Public/Chaos/RigidParticles.h
@@ -170,7 +170,10 @@ class TRigidParticles : public TKinematicGeometryParticles<T, d>
 		TSleepData<T, d> SleepData;
 		SleepData.Particle = Particle;
 		SleepData.Sleeping = Sleeping;
-		MSleepData.Add(SleepData); 
+
+		SleepDataLock.Lock();
+		MSleepData.Add(SleepData);
+		SleepDataLock.Unlock();
 	}
 
 	FORCEINLINE const EObjectStateType ObjectState(const int32 Index) const { return MObjectState[Index]; }
@@ -222,7 +225,9 @@ class TRigidParticles : public TKinematicGeometryParticles<T, d>
 	TArrayCollectionArray<bool> MDisabled;
 	TArrayCollectionArray<bool> MToBeRemovedOnFracture;
 	TArrayCollectionArray<EObjectStateType> MObjectState;
+
 	TArray<TSleepData<T, d>> MSleepData;
+	FCriticalSection SleepDataLock;
 };
 
 
