Fix inefficient array allocation in ChaosArchive serialization code. Call TArray::Reserve() before calling SetNum() to avoid unnecessary additional array slack
Saved 36% on heightfield memory and 23% trimesh

diff --git a/Engine/Source/Runtime/Experimental/Chaos/Public/Chaos/ChaosArchive.h b/Engine/Source/Runtime/Experimental/Chaos/Public/Chaos/ChaosArchive.h
--- a/Engine/Source/Runtime/Experimental/Chaos/Public/Chaos/ChaosArchive.h
+++ b/Engine/Source/Runtime/Experimental/Chaos/Public/Chaos/ChaosArchive.h
@@ -328,6 +328,7 @@ FChaosArchive& operator<<(FChaosArchive& Ar, TArray<T, TAllocator>& Array)
 {
 	int32 ArrayNum = Array.Num();
 	Ar << ArrayNum;
+	Array.Reserve(ArrayNum);
 	Array.SetNum(ArrayNum);
 
 	for (int32 Idx = 0; Idx < ArrayNum; ++Idx)
@@ -376,6 +377,7 @@ typename TEnableIf<IsSerializablePtr<T>(), FChaosArchive& > ::Type operator<<(FC
 {
 	int32 ArrayNum = Array.Num();
 	Ar << ArrayNum;
+	Array.Reserve(ArrayNum);
 	Array.SetNum(ArrayNum);
 
 	for (int32 Idx = 0; Idx < ArrayNum; ++Idx)
@@ -398,6 +400,7 @@ typename TEnableIf<IsSerializablePtr<T>(), FChaosArchive& > ::Type operator<<(FC
 {
 	int32 ArrayNum = Array.Num();
 	Ar << ArrayNum;
+	Array.Reserve(ArrayNum);
 	Array.SetNum(ArrayNum);
 
 	for (int32 Idx = 0; Idx < ArrayNum; ++Idx)
@@ -413,6 +416,7 @@ typename TEnableIf<IsSerializablePtr<T>(), FChaosArchive& > ::Type operator<<(FC
 {
 	int32 ArrayNum = Array.Num();
 	Ar << ArrayNum;
+	Array.Reserve(ArrayNum);
 	Array.SetNum(ArrayNum);
 
 	for (int32 Idx = 0; Idx < ArrayNum; ++Idx)
@@ -428,6 +432,7 @@ typename TEnableIf<IsSerializablePtr<T>(), FChaosArchive& > ::Type operator<<(FC
 {
 	int32 ArrayNum = Array.Num();
 	Ar << ArrayNum;
+	Array.Reserve(ArrayNum);
 	Array.SetNum(ArrayNum);
 
 	for (int32 Idx = 0; Idx < ArrayNum; ++Idx)
@@ -443,6 +448,7 @@ typename TEnableIf<IsSerializablePtr<T>(), FChaosArchive& > ::Type operator<<(FC
 {
 	int32 ArrayNum = Array.Num();
 	Ar << ArrayNum;
+	Array.Reserve(ArrayNum);
 	Array.SetNum(ArrayNum);
 
 	for (int32 Idx = 0; Idx < ArrayNum; ++Idx)
