#UE4 Fix read overflow in TriangleSimplexFindOrigin

diff --git a/Engine/Source/Runtime/Experimental/Chaos/Public/Chaos/Simplex.h b/Engine/Source/Runtime/Experimental/Chaos/Public/Chaos/Simplex.h
--- a/Engine/Source/Runtime/Experimental/Chaos/Public/Chaos/Simplex.h
+++ b/Engine/Source/Runtime/Experimental/Chaos/Public/Chaos/Simplex.h
@@ -177,7 +177,7 @@ namespace Chaos
 		bool bSignMatch[3];
 		FSimplex SubSimplices[3] = { {Idx1,Idx2}, {Idx0,Idx2}, {Idx0,Idx1} };
 		TVector<T, 3> ClosestPointSub[3];
-		T SubBarycentric[3][3];
+		T SubBarycentric[3][4];
 		int32 ClosestSubIdx = INDEX_NONE;
 		T MinSubDist2 = 0;	//not needed
 		bool bInside = true;
@@ -214,6 +214,9 @@ namespace Chaos
 		}
 		else
 		{
+			check(Idx0 >= 0 && Idx0 < 4);
+			check(Idx1 >= 0 && Idx1 < 4);
+			check(Idx2 >= 0 && Idx2 < 4);
 			Idxs = SubSimplices[ClosestSubIdx];
 			OutBarycentric[Idx0] = SubBarycentric[ClosestSubIdx][Idx0];
 			OutBarycentric[Idx1] = SubBarycentric[ClosestSubIdx][Idx1];
