Chaos
- remove Cholesky solver calling code
- rename Chaos collision and joint stats
- add drag handles to RBAN editing in PhAt

diff --git a/Engine/Source/Editor/PhysicsAssetEditor/Private/PhysicsAssetEditor.cpp b/Engine/Source/Editor/PhysicsAssetEditor/Private/PhysicsAssetEditor.cpp
--- a/Engine/Source/Editor/PhysicsAssetEditor/Private/PhysicsAssetEditor.cpp
+++ b/Engine/Source/Editor/PhysicsAssetEditor/Private/PhysicsAssetEditor.cpp
@@ -64,7 +64,7 @@
 #include "SPhysicsAssetGraph.h"
 #include "PhysicsAssetEditorEditMode.h"
 #include "AssetEditorModeManager.h"
-#include "PhysicsEngine/PhysicsHandleComponent.h"
+#include "PhysicsAssetEditorPhysicsHandleComponent.h"
 #include "ISkeletonTreeItem.h"
 #include "Algo/Transform.h"
 #include "SkeletonTreeSelection.h"
diff --git a/Engine/Source/Editor/PhysicsAssetEditor/Private/PhysicsAssetEditorAnimInstance.cpp b/Engine/Source/Editor/PhysicsAssetEditor/Private/PhysicsAssetEditorAnimInstance.cpp
--- a/Engine/Source/Editor/PhysicsAssetEditor/Private/PhysicsAssetEditorAnimInstance.cpp
+++ b/Engine/Source/Editor/PhysicsAssetEditor/Private/PhysicsAssetEditorAnimInstance.cpp
@@ -17,3 +17,28 @@ FAnimInstanceProxy* UPhysicsAssetEditorAnimInstance::CreateAnimInstanceProxy()
 {
 	return new FPhysicsAssetEditorAnimInstanceProxy(this);
 }
+
+void UPhysicsAssetEditorAnimInstance::Grab(FName InBoneName, const FVector& Location, const FRotator& Rotation, bool bRotationConstrained)
+{
+	FPhysicsAssetEditorAnimInstanceProxy& Proxy = GetProxyOnGameThread<FPhysicsAssetEditorAnimInstanceProxy>();
+	Proxy.Grab(InBoneName, Location, Rotation, bRotationConstrained);
+}
+
+void UPhysicsAssetEditorAnimInstance::Ungrab()
+{
+	FPhysicsAssetEditorAnimInstanceProxy& Proxy = GetProxyOnGameThread<FPhysicsAssetEditorAnimInstanceProxy>();
+	Proxy.Ungrab();
+}
+
+void UPhysicsAssetEditorAnimInstance::UpdateHandleTransform(const FTransform& NewTransform)
+{
+	FPhysicsAssetEditorAnimInstanceProxy& Proxy = GetProxyOnGameThread<FPhysicsAssetEditorAnimInstanceProxy>();
+	Proxy.UpdateHandleTransform(NewTransform);
+}
+
+void UPhysicsAssetEditorAnimInstance::UpdateDriveSettings(bool bLinearSoft, float LinearStiffness, float LinearDamping)
+{
+	FPhysicsAssetEditorAnimInstanceProxy& Proxy = GetProxyOnGameThread<FPhysicsAssetEditorAnimInstanceProxy>();
+	Proxy.UpdateDriveSettings(bLinearSoft, LinearStiffness, LinearDamping);
+}
+
diff --git a/Engine/Source/Editor/PhysicsAssetEditor/Private/PhysicsAssetEditorAnimInstance.h b/Engine/Source/Editor/PhysicsAssetEditor/Private/PhysicsAssetEditorAnimInstance.h
--- a/Engine/Source/Editor/PhysicsAssetEditor/Private/PhysicsAssetEditorAnimInstance.h
+++ b/Engine/Source/Editor/PhysicsAssetEditor/Private/PhysicsAssetEditorAnimInstance.h
@@ -16,6 +16,11 @@ class UPhysicsAssetEditorAnimInstance : public UAnimPreviewInstance
 {
 	GENERATED_UCLASS_BODY()
 
+	virtual void Grab(FName InBoneName, const FVector& Location, const FRotator& Rotation, bool bRotationConstrained);
+	virtual void Ungrab();
+	virtual void UpdateHandleTransform(const FTransform& NewTransform);
+	virtual void UpdateDriveSettings(bool bLinearSoft, float LinearStiffness, float LinearDamping);
+
 protected:
 	virtual FAnimInstanceProxy* CreateAnimInstanceProxy() override;
 };
diff --git a/Engine/Source/Editor/PhysicsAssetEditor/Private/PhysicsAssetEditorAnimInstanceProxy.cpp b/Engine/Source/Editor/PhysicsAssetEditor/Private/PhysicsAssetEditorAnimInstanceProxy.cpp
--- a/Engine/Source/Editor/PhysicsAssetEditor/Private/PhysicsAssetEditorAnimInstanceProxy.cpp
+++ b/Engine/Source/Editor/PhysicsAssetEditor/Private/PhysicsAssetEditorAnimInstanceProxy.cpp
@@ -3,6 +3,30 @@
 #include "PhysicsAssetEditorAnimInstanceProxy.h"
 #include "PhysicsAssetEditorAnimInstance.h"
 
+#include "Physics/ImmediatePhysics/ImmediatePhysicsActorHandle.h"
+#include "Physics/ImmediatePhysics/ImmediatePhysicsJointHandle.h"
+#include "Physics/ImmediatePhysics/ImmediatePhysicsSimulation.h"
+
+
+FPhysicsAssetEditorAnimInstanceProxy::FPhysicsAssetEditorAnimInstanceProxy()
+	: TargetActor(nullptr)
+	, HandleActor(nullptr)
+	, HandleJoint(nullptr)
+{
+}
+
+FPhysicsAssetEditorAnimInstanceProxy::FPhysicsAssetEditorAnimInstanceProxy(UAnimInstance* InAnimInstance)
+	: FAnimPreviewInstanceProxy(InAnimInstance)
+	, TargetActor(nullptr)
+	, HandleActor(nullptr)
+	, HandleJoint(nullptr)
+{
+}
+
+FPhysicsAssetEditorAnimInstanceProxy::~FPhysicsAssetEditorAnimInstanceProxy()
+{
+}
+
 void FPhysicsAssetEditorAnimInstanceProxy::Initialize(UAnimInstance* InAnimInstance)
 {
 	FAnimPreviewInstanceProxy::Initialize(InAnimInstance);
@@ -58,6 +82,70 @@ void FPhysicsAssetEditorAnimInstanceProxy::AddImpulseAtLocation(FVector Impulse,
 	RagdollNode.AddImpulseAtLocation(Impulse, Location, BoneName);
 }
 
-FPhysicsAssetEditorAnimInstanceProxy::~FPhysicsAssetEditorAnimInstanceProxy()
+void FPhysicsAssetEditorAnimInstanceProxy::Grab(FName InBoneName, const FVector& Location, const FRotator& Rotation, bool bRotationConstrained)
+{
+#if WITH_CHAOS && !PHYSICS_INTERFACE_PHYSX
+	ImmediatePhysics::FSimulation* Simulation = RagdollNode.GetSimulation();
+
+	if (TargetActor != nullptr)
+	{
+		Ungrab();
+	}
+
+	for (int ActorIndex = 0; ActorIndex < Simulation->NumActors(); ++ActorIndex)
+	{
+		if (Simulation->GetActorHandle(ActorIndex)->GetName() == InBoneName)
+		{
+			TargetActor = Simulation->GetActorHandle(ActorIndex);
+			break;
+		}
+	}
+
+	if (TargetActor != nullptr)
+	{
+		FTransform HandleTransform = FTransform(Rotation, Location);
+		HandleActor = Simulation->CreateKinematicActor(nullptr, HandleTransform);
+		HandleActor->SetWorldTransform(HandleTransform);
+		HandleActor->SetKinematicTarget(HandleTransform);
+
+		HandleJoint = Simulation->CreateJoint(nullptr, TargetActor, HandleActor);
+	}
+#endif
+}
+
+void FPhysicsAssetEditorAnimInstanceProxy::Ungrab()
+{
+#if WITH_CHAOS && !PHYSICS_INTERFACE_PHYSX
+	if (TargetActor != nullptr)
+	{
+		ImmediatePhysics::FSimulation* Simulation = RagdollNode.GetSimulation();
+		Simulation->DestroyJoint(HandleJoint);
+		Simulation->DestroyActor(HandleActor);
+		TargetActor = nullptr;
+		HandleActor = nullptr;
+		HandleJoint = nullptr;
+	}
+#endif
+}
+
+void FPhysicsAssetEditorAnimInstanceProxy::UpdateHandleTransform(const FTransform& NewTransform)
+{
+#if WITH_CHAOS && !PHYSICS_INTERFACE_PHYSX
+	if (HandleActor != nullptr)
+	{
+		HandleActor->SetKinematicTarget(NewTransform);
+	}
+#endif
+}
+
+void FPhysicsAssetEditorAnimInstanceProxy::UpdateDriveSettings(bool bLinearSoft, float LinearStiffness, float LinearDamping)
 {
-}
\ No newline at end of file
+#if WITH_CHAOS && !PHYSICS_INTERFACE_PHYSX
+	using namespace Chaos;
+	if (HandleJoint != nullptr)
+	{
+		HandleJoint->SetSoftLinearSettings(bLinearSoft, LinearStiffness, LinearDamping);
+	}
+#endif
+}
+
diff --git a/Engine/Source/Editor/PhysicsAssetEditor/Private/PhysicsAssetEditorAnimInstanceProxy.h b/Engine/Source/Editor/PhysicsAssetEditor/Private/PhysicsAssetEditorAnimInstanceProxy.h
--- a/Engine/Source/Editor/PhysicsAssetEditor/Private/PhysicsAssetEditorAnimInstanceProxy.h
+++ b/Engine/Source/Editor/PhysicsAssetEditor/Private/PhysicsAssetEditorAnimInstanceProxy.h
@@ -7,6 +7,8 @@
 #include "Animation/AnimNode_SequencePlayer.h"
 #include "Animation/AnimNodeSpaceConversions.h"
 #include "BoneControllers/AnimNode_RigidBody.h"
+#include "Physics/ImmediatePhysics/ImmediatePhysicsDeclares.h"
+
 #include "PhysicsAssetEditorAnimInstanceProxy.generated.h"
 
 class UAnimSequence;
@@ -18,14 +20,9 @@ struct FPhysicsAssetEditorAnimInstanceProxy : public FAnimPreviewInstanceProxy
 	GENERATED_BODY()
 
 public:
-	FPhysicsAssetEditorAnimInstanceProxy()
-	{
-	}
+	FPhysicsAssetEditorAnimInstanceProxy();
 
-	FPhysicsAssetEditorAnimInstanceProxy(UAnimInstance* InAnimInstance)
-		: FAnimPreviewInstanceProxy(InAnimInstance)
-	{
-	}
+	FPhysicsAssetEditorAnimInstanceProxy(UAnimInstance* InAnimInstance);
 
 	virtual ~FPhysicsAssetEditorAnimInstanceProxy();
 
@@ -38,9 +35,18 @@ public:
 
 	virtual void AddImpulseAtLocation(FVector Impulse, FVector Location, FName BoneName = NAME_None) override;
 
+	virtual void Grab(FName InBoneName, const FVector& Location, const FRotator& Rotation, bool bRotationConstrained);
+	virtual void Ungrab();
+	virtual void UpdateHandleTransform(const FTransform& NewTransform);
+	virtual void UpdateDriveSettings(bool bLinearSoft, float LinearStiffness, float LinearDamping);
+
 private:
 	void ConstructNodes();
 
 	FAnimNode_RigidBody RagdollNode;
  	FAnimNode_ConvertComponentToLocalSpace ComponentToLocalSpace;
+
+	ImmediatePhysics::FActorHandle* TargetActor;
+	ImmediatePhysics::FActorHandle* HandleActor;
+	ImmediatePhysics::FJointHandle* HandleJoint;
 };
\ No newline at end of file
diff --git a/Engine/Source/Editor/PhysicsAssetEditor/Private/PhysicsAssetEditorEditMode.cpp b/Engine/Source/Editor/PhysicsAssetEditor/Private/PhysicsAssetEditorEditMode.cpp
--- a/Engine/Source/Editor/PhysicsAssetEditor/Private/PhysicsAssetEditorEditMode.cpp
+++ b/Engine/Source/Editor/PhysicsAssetEditor/Private/PhysicsAssetEditorEditMode.cpp
@@ -14,7 +14,7 @@
 #include "IPersonaPreviewScene.h"
 #include "PhysicsAssetEditor.h"
 #include "PhysicsAssetEditorHitProxies.h"
-#include "PhysicsEngine/PhysicsHandleComponent.h"
+#include "PhysicsAssetEditorPhysicsHandleComponent.h"
 #include "DrawDebugHelpers.h"
 #include "SEditorViewport.h"
 #include "IPersonaToolkit.h"
diff --git a/Engine/Source/Editor/PhysicsAssetEditor/Private/PhysicsAssetEditorPhysicsHandleComponent.cpp b/Engine/Source/Editor/PhysicsAssetEditor/Private/PhysicsAssetEditorPhysicsHandleComponent.cpp
new file mode 100644
--- /dev/null
+++ b/Engine/Source/Editor/PhysicsAssetEditor/Private/PhysicsAssetEditorPhysicsHandleComponent.cpp
@@ -0,0 +1,73 @@
+// Copyright Epic Games, Inc. All Rights Reserved.
+#include "PhysicsAssetEditorPhysicsHandleComponent.h"
+
+#include "PhysicsAssetEditorSkeletalMeshComponent.h"
+
+UPhysicsAssetEditorPhysicsHandleComponent::UPhysicsAssetEditorPhysicsHandleComponent(const FObjectInitializer& ObjectInitializer)
+	: Super(ObjectInitializer)
+	, bAnimInstanceMode(false)
+{
+}
+
+void UPhysicsAssetEditorPhysicsHandleComponent::UpdateHandleTransform(const FTransform& NewTransform)
+{
+	Super::UpdateHandleTransform(NewTransform);
+
+	if (bAnimInstanceMode)
+	{
+		UPhysicsAssetEditorSkeletalMeshComponent* PhatComponent = Cast<UPhysicsAssetEditorSkeletalMeshComponent>(GrabbedComponent);
+		if (PhatComponent != nullptr)
+		{
+			PhatComponent->UpdateHandleTransform(NewTransform);
+		}
+	}
+}
+
+void UPhysicsAssetEditorPhysicsHandleComponent::UpdateDriveSettings()
+{
+	Super::UpdateDriveSettings();
+
+	if (bAnimInstanceMode)
+	{
+		UPhysicsAssetEditorSkeletalMeshComponent* PhatComponent = Cast<UPhysicsAssetEditorSkeletalMeshComponent>(GrabbedComponent);
+		if (PhatComponent != nullptr)
+		{
+			PhatComponent->UpdateDriveSettings(bSoftLinearConstraint, LinearStiffness, LinearDamping);
+		}
+	}
+}
+
+void UPhysicsAssetEditorPhysicsHandleComponent::SetAnimInstanceMode(bool bInAnimInstanceMode)
+{
+	bAnimInstanceMode = bInAnimInstanceMode;
+}
+
+void UPhysicsAssetEditorPhysicsHandleComponent::GrabComponentImp(class UPrimitiveComponent* Component, FName InBoneName, const FVector& Location, const FRotator& Rotation, bool InbRotationConstrained)
+{
+	Super::GrabComponentImp(Component, InBoneName, Location, Rotation, bRotationConstrained);
+
+	if (bAnimInstanceMode)
+	{
+		TargetTransform = CurrentTransform = FTransform(Rotation, Location);
+
+		UPhysicsAssetEditorSkeletalMeshComponent* PhatComponent = Cast<UPhysicsAssetEditorSkeletalMeshComponent>(Component);
+		if (PhatComponent != nullptr)
+		{
+			PhatComponent->Grab(InBoneName, Location, Rotation, InbRotationConstrained);
+		}
+	}
+}
+
+void UPhysicsAssetEditorPhysicsHandleComponent::ReleaseComponent()
+{
+	if (bAnimInstanceMode)
+	{
+		UPhysicsAssetEditorSkeletalMeshComponent* PhatComponent = Cast<UPhysicsAssetEditorSkeletalMeshComponent>(GrabbedComponent);
+		if (PhatComponent != nullptr)
+		{
+			PhatComponent->Ungrab();
+		}
+	}
+
+	Super::ReleaseComponent();
+}
diff --git a/Engine/Source/Editor/PhysicsAssetEditor/Private/PhysicsAssetEditorPhysicsHandleComponent.h b/Engine/Source/Editor/PhysicsAssetEditor/Private/PhysicsAssetEditorPhysicsHandleComponent.h
new file mode 100644
--- /dev/null
+++ b/Engine/Source/Editor/PhysicsAssetEditor/Private/PhysicsAssetEditorPhysicsHandleComponent.h
@@ -0,0 +1,27 @@
+// Copyright Epic Games, Inc. All Rights Reserved.
+
+#include "PhysicsEngine/PhysicsHandleComponent.h"
+
+#include "PhysicsAssetEditorPhysicsHandleComponent.generated.h"
+
+/**
+ * Extend the Physics Handle for PhAt. This adds support for manipulating the physics
+ * if it is running in a RigidBody AnimNode (which is always is with Chaos at the moment).
+ */
+UCLASS()
+class UPhysicsAssetEditorPhysicsHandleComponent : public UPhysicsHandleComponent
+{
+	GENERATED_UCLASS_BODY()
+
+	bool bAnimInstanceMode;
+
+public:
+	void SetAnimInstanceMode(bool bInAnimInstanceMode);
+
+	virtual void ReleaseComponent() override;
+
+protected:
+	virtual void UpdateHandleTransform(const FTransform& NewTransform) override;
+	virtual void UpdateDriveSettings() override;
+	virtual void GrabComponentImp(class UPrimitiveComponent* Component, FName InBoneName, const FVector& Location, const FRotator& Rotation, bool bRotationConstrained) override;
+};
diff --git a/Engine/Source/Editor/PhysicsAssetEditor/Private/PhysicsAssetEditorSharedData.cpp b/Engine/Source/Editor/PhysicsAssetEditor/Private/PhysicsAssetEditorSharedData.cpp
--- a/Engine/Source/Editor/PhysicsAssetEditor/Private/PhysicsAssetEditorSharedData.cpp
+++ b/Engine/Source/Editor/PhysicsAssetEditor/Private/PhysicsAssetEditorSharedData.cpp
@@ -1,7 +1,7 @@
 // Copyright Epic Games, Inc. All Rights Reserved.
 
 #include "PhysicsAssetEditorSharedData.h"
-#include "PhysicsEngine/PhysicsHandleComponent.h"
+#include "PhysicsAssetEditorPhysicsHandleComponent.h"
 #include "PhysicsEngine/RigidBodyIndexPair.h"
 #include "Misc/MessageDialog.h"
 #include "Modules/ModuleManager.h"
@@ -68,7 +68,12 @@ FPhysicsAssetEditorSharedData::FPhysicsAssetEditorSharedData()
 	bManipulating = false;
 	
 	// Construct mouse handle
-	MouseHandle = NewObject<UPhysicsHandleComponent>();
+	MouseHandle = NewObject<UPhysicsAssetEditorPhysicsHandleComponent>();
+
+	// in Chaos, we have to manipulate the RBAN node in the Anim Instance (at least until we get SkelMeshComp implemented)
+#if PHAT_USE_RBAN_SIMULATION
+	MouseHandle->SetAnimInstanceMode(true);
+#endif
 
 	// Construct sim options.
 	EditorOptions = NewObject<UPhysicsAssetEditorOptions>(GetTransientPackage(), MakeUniqueObjectName(GetTransientPackage(), UPhysicsAssetEditorOptions::StaticClass(), FName(TEXT("EditorOptions"))), RF_Transactional);
diff --git a/Engine/Source/Editor/PhysicsAssetEditor/Private/PhysicsAssetEditorSharedData.h b/Engine/Source/Editor/PhysicsAssetEditor/Private/PhysicsAssetEditorSharedData.h
--- a/Engine/Source/Editor/PhysicsAssetEditor/Private/PhysicsAssetEditorSharedData.h
+++ b/Engine/Source/Editor/PhysicsAssetEditor/Private/PhysicsAssetEditorSharedData.h
@@ -12,7 +12,7 @@ class UBodySetup;
 class UPhysicsAssetEditorSkeletalMeshComponent;
 class UPhysicsAsset;
 class UPhysicsConstraintTemplate;
-class UPhysicsHandleComponent;
+class UPhysicsAssetEditorPhysicsHandleComponent;
 class USkeletalMesh;
 class UStaticMeshComponent;
 struct FBoneVertInfo;
@@ -247,7 +247,7 @@ public:
 	EAppReturnType::Type NewBodyResponse;
 
 	/** Helps define how the asset behaves given user interaction in simulation mode*/
-	UPhysicsHandleComponent* MouseHandle;
+	UPhysicsAssetEditorPhysicsHandleComponent* MouseHandle;
 
 	/** Draw color for center of mass debug strings */
 	const FColor COMRenderColor;
diff --git a/Engine/Source/Editor/PhysicsAssetEditor/Private/PhysicsAssetEditorSkeletalMeshComponent.cpp b/Engine/Source/Editor/PhysicsAssetEditor/Private/PhysicsAssetEditorSkeletalMeshComponent.cpp
--- a/Engine/Source/Editor/PhysicsAssetEditor/Private/PhysicsAssetEditorSkeletalMeshComponent.cpp
+++ b/Engine/Source/Editor/PhysicsAssetEditor/Private/PhysicsAssetEditorSkeletalMeshComponent.cpp
@@ -10,6 +10,7 @@
 #include "PhysicsAssetEditorSharedData.h"
 #include "PhysicsAssetEditorHitProxies.h"
 #include "PhysicsAssetEditorSkeletalMeshComponent.h"
+#include "PhysicsAssetEditorAnimInstance.h"
 #include "PhysicsEngine/PhysicsConstraintTemplate.h"
 #include "PhysicsEngine/PhysicsAsset.h"
 #include "Chaos/Core.h"
@@ -511,3 +512,39 @@ void UPhysicsAssetEditorSkeletalMeshComponent::AddImpulseAtLocation(FVector Impu
 	}
 #endif
 }
+
+void UPhysicsAssetEditorSkeletalMeshComponent::Grab(FName InBoneName, const FVector& Location, const FRotator& Rotation, bool bRotationConstrained)
+{
+	UPhysicsAssetEditorAnimInstance* PhatPreviewInstance = Cast<UPhysicsAssetEditorAnimInstance>(PreviewInstance);
+	if (PhatPreviewInstance != nullptr)
+	{
+		PhatPreviewInstance->Grab(InBoneName, Location, Rotation, bRotationConstrained);
+	}
+}
+
+void UPhysicsAssetEditorSkeletalMeshComponent::Ungrab()
+{
+	UPhysicsAssetEditorAnimInstance* PhatPreviewInstance = Cast<UPhysicsAssetEditorAnimInstance>(PreviewInstance);
+	if (PhatPreviewInstance != nullptr)
+	{
+		PhatPreviewInstance->Ungrab();
+	}
+}
+
+void UPhysicsAssetEditorSkeletalMeshComponent::UpdateHandleTransform(const FTransform& NewTransform)
+{
+	UPhysicsAssetEditorAnimInstance* PhatPreviewInstance = Cast<UPhysicsAssetEditorAnimInstance>(PreviewInstance);
+	if (PhatPreviewInstance != nullptr)
+	{
+		PhatPreviewInstance->UpdateHandleTransform(NewTransform);
+	}
+}
+
+void UPhysicsAssetEditorSkeletalMeshComponent::UpdateDriveSettings(bool bLinearSoft, float LinearStiffness, float LinearDamping)
+{
+	UPhysicsAssetEditorAnimInstance* PhatPreviewInstance = Cast<UPhysicsAssetEditorAnimInstance>(PreviewInstance);
+	if (PhatPreviewInstance != nullptr)
+	{
+		PhatPreviewInstance->UpdateDriveSettings(bLinearSoft, LinearStiffness, LinearDamping);
+	}
+}
diff --git a/Engine/Source/Editor/PhysicsAssetEditor/Private/PhysicsAssetEditorSkeletalMeshComponent.h b/Engine/Source/Editor/PhysicsAssetEditor/Private/PhysicsAssetEditorSkeletalMeshComponent.h
--- a/Engine/Source/Editor/PhysicsAssetEditor/Private/PhysicsAssetEditorSkeletalMeshComponent.h
+++ b/Engine/Source/Editor/PhysicsAssetEditor/Private/PhysicsAssetEditorSkeletalMeshComponent.h
@@ -70,6 +70,13 @@ class UPhysicsAssetEditorSkeletalMeshComponent : public UDebugSkelMeshComponent
 	FTransform GetPrimitiveTransform(FTransform& BoneTM, int32 BodyIndex, EAggCollisionShape::Type PrimType, int32 PrimIndex, float Scale);
 	FColor GetPrimitiveColor(int32 BodyIndex, EAggCollisionShape::Type PrimitiveType, int32 PrimitiveIndex);
 	UMaterialInterface* GetPrimitiveMaterial(int32 BodyIndex, EAggCollisionShape::Type PrimitiveType, int32 PrimitiveIndex);
+
+	/** Manipulator methods */
+	virtual void Grab(FName InBoneName, const FVector& Location, const FRotator& Rotation, bool bRotationConstrained);
+	virtual void Ungrab();
+	virtual void UpdateHandleTransform(const FTransform& NewTransform);
+	virtual void UpdateDriveSettings(bool bLinearSoft, float LinearStiffness, float LinearDamping);
+
 public:
 	virtual bool CanOverrideCollisionProfile() const override { return false;  }
 };
diff --git a/Engine/Source/Runtime/AnimGraphRuntime/Public/BoneControllers/AnimNode_RigidBody.h b/Engine/Source/Runtime/AnimGraphRuntime/Public/BoneControllers/AnimNode_RigidBody.h
--- a/Engine/Source/Runtime/AnimGraphRuntime/Public/BoneControllers/AnimNode_RigidBody.h
+++ b/Engine/Source/Runtime/AnimGraphRuntime/Public/BoneControllers/AnimNode_RigidBody.h
@@ -57,6 +57,9 @@ struct ANIMGRAPHRUNTIME_API FAnimNode_RigidBody : public FAnimNode_SkeletalContr
 
 	virtual void AddImpulseAtLocation(FVector Impulse, FVector Location, FName BoneName = NAME_None);
 
+	// TEMP: Exposed for use in PhAt as a quick way to get drag handles working with Chaos
+	virtual ImmediatePhysics::FSimulation* GetSimulation() { return PhysicsSimulation; }
+
 public:
 	/** Physics asset to use. If empty use the skeletal mesh's default physics asset */
 	UPROPERTY(EditAnywhere, Category = Settings)
diff --git a/Engine/Source/Runtime/Engine/Classes/PhysicsEngine/PhysicsHandleComponent.h b/Engine/Source/Runtime/Engine/Classes/PhysicsEngine/PhysicsHandleComponent.h
--- a/Engine/Source/Runtime/Engine/Classes/PhysicsEngine/PhysicsHandleComponent.h
+++ b/Engine/Source/Runtime/Engine/Classes/PhysicsEngine/PhysicsHandleComponent.h
@@ -139,7 +139,7 @@ public:
 
 protected:
 	/** Move the kinematic handle to the specified */
-	void UpdateHandleTransform(const FTransform& NewTransform);
+	virtual void UpdateHandleTransform(const FTransform& NewTransform);
 
 	/** Update the underlying constraint drive settings from the params in this component */
 	virtual void UpdateDriveSettings();
diff --git a/Engine/Source/Runtime/Engine/Private/PhysicsEngine/ImmediatePhysics/ImmediatePhysicsChaos/ImmediatePhysicsActorHandle_Chaos.cpp b/Engine/Source/Runtime/Engine/Private/PhysicsEngine/ImmediatePhysics/ImmediatePhysicsChaos/ImmediatePhysicsActorHandle_Chaos.cpp
--- a/Engine/Source/Runtime/Engine/Private/PhysicsEngine/ImmediatePhysics/ImmediatePhysicsChaos/ImmediatePhysicsActorHandle_Chaos.cpp
+++ b/Engine/Source/Runtime/Engine/Private/PhysicsEngine/ImmediatePhysics/ImmediatePhysicsChaos/ImmediatePhysicsActorHandle_Chaos.cpp
@@ -154,11 +154,38 @@ namespace ImmediatePhysics_Chaos
 		}
 	}
 
+	bool CreateDefaultGeometry(const FVector& Scale, float& OutMass, Chaos::TVector<float, 3>& OutInertia, Chaos::TRigidTransform<float, 3>& OutCoMTransform, TUniquePtr<Chaos::FImplicitObject>& OutGeom, TArray<TUniquePtr<Chaos::TPerShapeData<float, 3>>>& OutShapes)
+	{
+		using namespace Chaos;
+
+		const FReal Mass = 1.0f;
+		const FReal Radius = 1.0f * Scale.GetMax();
+
+		auto ImplicitSphere = MakeUnique<Chaos::TSphere<float, 3>>(FVec3(0), Radius);
+		auto NewShape = Chaos::TPerShapeData<float, 3>::CreatePerShapeData();
+		NewShape->Geometry = MakeSerializable(ImplicitSphere);
+		NewShape->UpdateShapeBounds(FTransform::Identity);
+		NewShape->UserData = 0;
+		NewShape->bSimulate = false;
+
+		OutMass = Mass;
+		OutInertia = TSphere<FReal, 3>::GetInertiaTensor(Mass, Radius).GetDiagonal();
+		OutCoMTransform = FTransform::Identity;
+		OutShapes.Emplace(MoveTemp(NewShape));
+		OutGeom = MoveTemp(ImplicitSphere);
+
+		return true;
+	}
 
 	bool CreateGeometry(FBodyInstance* BodyInstance, const FVector& Scale, float& OutMass, Chaos::TVector<float, 3>& OutInertia, Chaos::TRigidTransform<float, 3>& OutCoMTransform, TUniquePtr<Chaos::FImplicitObject>& OutGeom, TArray<TUniquePtr<Chaos::TPerShapeData<float, 3>>>& OutShapes)
 	{
 		using namespace Chaos;
 
+		if (BodyInstance == nullptr)
+		{
+			return CreateDefaultGeometry(Scale, OutMass, OutInertia, OutCoMTransform, OutGeom, OutShapes);
+		}
+
 		UBodySetup* BodySetup = BodyInstance->BodySetup.Get();
 
 #if WITH_CHAOS && !PHYSICS_INTERFACE_PHYSX
@@ -289,8 +316,11 @@ namespace ImmediatePhysics_Chaos
 					Dynamic->SetInvM(MassInv);
 					Dynamic->SetI({ Inertia.X, Inertia.Y, Inertia.Z });
 					Dynamic->SetInvI({ InertiaInv.X, InertiaInv.Y, InertiaInv.Z });
-					Dynamic->SetLinearEtherDrag(BodyInstance->LinearDamping);
-					Dynamic->SetAngularEtherDrag(BodyInstance->AngularDamping);
+					if (BodyInstance != nullptr)
+					{
+						Dynamic->SetLinearEtherDrag(BodyInstance->LinearDamping);
+						Dynamic->SetAngularEtherDrag(BodyInstance->AngularDamping);
+					}
 					Dynamic->Disabled() = true;
 				}
 			}
@@ -654,5 +684,4 @@ namespace ImmediatePhysics_Chaos
 	{
 		Level = InLevel;
 	}
-
 }
diff --git a/Engine/Source/Runtime/Engine/Private/PhysicsEngine/ImmediatePhysics/ImmediatePhysicsChaos/ImmediatePhysicsJointHandle_Chaos.cpp b/Engine/Source/Runtime/Engine/Private/PhysicsEngine/ImmediatePhysics/ImmediatePhysicsChaos/ImmediatePhysicsJointHandle_Chaos.cpp
--- a/Engine/Source/Runtime/Engine/Private/PhysicsEngine/ImmediatePhysics/ImmediatePhysicsChaos/ImmediatePhysicsJointHandle_Chaos.cpp
+++ b/Engine/Source/Runtime/Engine/Private/PhysicsEngine/ImmediatePhysics/ImmediatePhysicsChaos/ImmediatePhysicsJointHandle_Chaos.cpp
@@ -52,24 +52,15 @@ namespace ImmediatePhysics_Chaos
 	FAutoConsoleVariableRef CVarJointMinAngularProjection(TEXT("p.Chaos.ImmPhys.JointMinAngularProjection"), ChaosImmediate_JointMinAngularProjection, TEXT("Joint min projection (for joints with projection disabled)."));
 	FAutoConsoleVariableRef CVarJointMaxAngularProjection(TEXT("p.Chaos.ImmPhys.JointMaxAngularProjection"), ChaosImmediate_JointMaxAngularProjection, TEXT("Joint max projection (for joints with projection enabled)."));
 
-	FJointHandle::FJointHandle(FChaosConstraintContainer* InConstraints, FConstraintInstance* ConstraintInstance, FActorHandle* Actor1, FActorHandle* Actor2)
-		: ActorHandles({ Actor1, Actor2 })
-		, Constraints(InConstraints)
+	void TransferJointSettings(FConstraintInstance* ConstraintInstance, Chaos::FPBDJointSettings& ConstraintSettings)
 	{
 		using namespace Chaos;
 
 		const FConstraintProfileProperties& Profile = ConstraintInstance->ProfileInstance;
 
-		FPBDJointSettings ConstraintSettings;
-		TVector<FRigidTransform3, 2> ConstraintFrames;
-
-		// BodyInstance/PhysX has the constraint locations in actor-space, but we need them in Center-of-Mass space
-		ConstraintFrames[0] = FParticleUtilities::ActorLocalToParticleLocal(TGenericParticleHandle<FReal, 3>(Actor1->GetParticle()), ConstraintInstance->GetRefFrame(EConstraintFrame::Frame1));
-		ConstraintFrames[1] = FParticleUtilities::ActorLocalToParticleLocal(TGenericParticleHandle<FReal, 3>(Actor2->GetParticle()), ConstraintInstance->GetRefFrame(EConstraintFrame::Frame2));
-
 		ConstraintSettings.Stiffness = ChaosImmediate_JointStiffness;
 
-		ConstraintSettings.LinearMotionTypes = 
+		ConstraintSettings.LinearMotionTypes =
 		{
 			static_cast<EJointMotionType>(ConstraintInstance->GetLinearXMotion()),
 			static_cast<EJointMotionType>(ConstraintInstance->GetLinearYMotion()),
@@ -97,8 +88,8 @@ namespace ImmediatePhysics_Chaos
 		ConstraintSettings.SoftTwistDamping = ChaosImmediate_SoftAngularDampingScale * ConstraintInstance->GetSoftTwistLimitDamping();
 		ConstraintSettings.SoftSwingStiffness = ChaosImmediate_SoftAngularStiffnessScale * ConstraintInstance->GetSoftSwingLimitStiffness();
 		ConstraintSettings.SoftSwingDamping = ChaosImmediate_SoftAngularDampingScale * ConstraintInstance->GetSoftSwingLimitDamping();
-		ConstraintSettings.LinearSoftForceMode = (ChaosImmediate_SoftLinearForceMode == 0)? EJointForceMode::Acceleration : EJointForceMode::Force;
-		ConstraintSettings.AngularSoftForceMode = (ChaosImmediate_SoftAngularForceMode == 0)? EJointForceMode::Acceleration : EJointForceMode::Force;
+		ConstraintSettings.LinearSoftForceMode = (ChaosImmediate_SoftLinearForceMode == 0) ? EJointForceMode::Acceleration : EJointForceMode::Force;
+		ConstraintSettings.AngularSoftForceMode = (ChaosImmediate_SoftAngularForceMode == 0) ? EJointForceMode::Acceleration : EJointForceMode::Force;
 
 		ConstraintSettings.LinearDriveTarget = Profile.LinearDrive.PositionTarget;
 		ConstraintSettings.bLinearPositionDriveEnabled[0] = Profile.LinearDrive.XDrive.bEnablePositionDrive;
@@ -160,6 +151,35 @@ namespace ImmediatePhysics_Chaos
 		//{
 		//	ConstraintSettings.bSoftSwingLimitsEnabled = false;
 		//}
+	}
+
+	FJointHandle::FJointHandle(FChaosConstraintContainer* InConstraints, FConstraintInstance* ConstraintInstance, FActorHandle* Actor1, FActorHandle* Actor2)
+		: ActorHandles({ Actor1, Actor2 })
+		, Constraints(InConstraints)
+	{
+		using namespace Chaos;
+
+		FPBDJointSettings ConstraintSettings;
+		TVector<FRigidTransform3, 2> ConstraintFrames;
+
+		if (ConstraintInstance != nullptr)
+		{
+			// BodyInstance/PhysX has the constraint locations in actor-space, but we need them in Center-of-Mass space
+			ConstraintFrames[0] = FParticleUtilities::ActorLocalToParticleLocal(TGenericParticleHandle<FReal, 3>(Actor1->GetParticle()), ConstraintInstance->GetRefFrame(EConstraintFrame::Frame1));
+			ConstraintFrames[1] = FParticleUtilities::ActorLocalToParticleLocal(TGenericParticleHandle<FReal, 3>(Actor2->GetParticle()), ConstraintInstance->GetRefFrame(EConstraintFrame::Frame2));
+			TransferJointSettings(ConstraintInstance, ConstraintSettings);
+		}
+		else
+		{
+			// TEMP: all creation with null ConstraintIndex for PhAt handles
+			ConstraintFrames[0] = Actor2->GetWorldTransform().GetRelativeTransform(Actor1->GetWorldTransform());
+			ConstraintFrames[1] = FRigidTransform3();
+			ConstraintSettings.LinearMotionTypes = { EJointMotionType::Limited, EJointMotionType::Limited, EJointMotionType::Limited };
+			ConstraintSettings.LinearLimit = 0.0f;
+			ConstraintSettings.SoftLinearStiffness = 10.0f;
+			ConstraintSettings.SoftLinearDamping = 1.0f;
+			ConstraintSettings.LinearSoftForceMode = EJointForceMode::Acceleration;
+		}
 
 		ConstraintSettings.Sanitize();
 
@@ -195,4 +215,15 @@ namespace ImmediatePhysics_Chaos
 	{
 		ConstraintHandle->SetParticleLevels({ ActorHandles[0]->GetLevel(), ActorHandles[1]->GetLevel() });
 	}
+
+	void FJointHandle::SetSoftLinearSettings(bool bLinearSoft, FReal LinearStiffness, FReal LinearDamping)
+	{
+		using namespace Chaos;
+		FPBDJointSettings JointSettings = ConstraintHandle->GetSettings();
+		JointSettings.bSoftLinearLimitsEnabled = bLinearSoft;
+		JointSettings.SoftLinearStiffness = bLinearSoft ? LinearStiffness : 0.0f;
+		JointSettings.SoftLinearDamping = bLinearSoft ? LinearDamping : 0.0f;
+		ConstraintHandle->SetSettings(JointSettings);
+	}
 }
+
diff --git a/Engine/Source/Runtime/Engine/Public/Physics/ImmediatePhysics/ImmediatePhysicsChaos/ImmediatePhysicsJointHandle_Chaos.h b/Engine/Source/Runtime/Engine/Public/Physics/ImmediatePhysics/ImmediatePhysicsChaos/ImmediatePhysicsJointHandle_Chaos.h
--- a/Engine/Source/Runtime/Engine/Public/Physics/ImmediatePhysics/ImmediatePhysicsChaos/ImmediatePhysicsJointHandle_Chaos.h
+++ b/Engine/Source/Runtime/Engine/Public/Physics/ImmediatePhysics/ImmediatePhysicsChaos/ImmediatePhysicsJointHandle_Chaos.h
@@ -29,6 +29,8 @@ namespace ImmediatePhysics_Chaos
 
 		void UpdateLevels();
 
+		void SetSoftLinearSettings(bool bLinearSoft, FReal LinearStiffness, FReal LinearDamping);
+
 	private:
 		Chaos::TVector<FActorHandle*, 2> ActorHandles;
 		FChaosConstraintContainer* Constraints;
diff --git a/Engine/Source/Runtime/Engine/Public/Physics/ImmediatePhysics/ImmediatePhysicsDeclares.h b/Engine/Source/Runtime/Engine/Public/Physics/ImmediatePhysics/ImmediatePhysicsDeclares.h
--- a/Engine/Source/Runtime/Engine/Public/Physics/ImmediatePhysics/ImmediatePhysicsDeclares.h
+++ b/Engine/Source/Runtime/Engine/Public/Physics/ImmediatePhysics/ImmediatePhysicsDeclares.h
@@ -28,9 +28,11 @@ namespace ImmediatePhysics
 #if WITH_PHYSX && PHYSICS_INTERFACE_PHYSX
 	using FSimulation = ImmediatePhysics_PhysX::FSimulation;
 	using FActorHandle = ImmediatePhysics_PhysX::FActorHandle;
+	using FJointHandle = ImmediatePhysics_PhysX::FJointHandle;
 #elif WITH_CHAOS
 	using FSimulation = ImmediatePhysics_Chaos::FSimulation;
 	using FActorHandle = ImmediatePhysics_Chaos::FActorHandle;
+	using FJointHandle = ImmediatePhysics_Chaos::FJointHandle;
 #else
 #error Global ImmediatePhysics is not defined
 #endif
diff --git a/Engine/Source/Runtime/Engine/Public/Physics/ImmediatePhysics/ImmediatePhysicsJointHandle.h b/Engine/Source/Runtime/Engine/Public/Physics/ImmediatePhysics/ImmediatePhysicsJointHandle.h
new file mode 100644
--- /dev/null
+++ b/Engine/Source/Runtime/Engine/Public/Physics/ImmediatePhysics/ImmediatePhysicsJointHandle.h
@@ -0,0 +1,8 @@
+// Copyright 1998-2019 Epic Games, Inc. All Rights Reserved.
+
+#pragma once
+
+#include "Physics/ImmediatePhysics/ImmediatePhysicsPhysX/ImmediatePhysicsJointHandle_PhysX.h"
+#include "Physics/ImmediatePhysics/ImmediatePhysicsChaos/ImmediatePhysicsJointHandle_Chaos.h"
+
+#include "Physics/ImmediatePhysics/ImmediatePhysicsDeclares.h"
diff --git a/Engine/Source/Runtime/Engine/Public/Physics/ImmediatePhysics/ImmediatePhysicsPhysX/ImmediatePhysicsDeclares_PhysX.h b/Engine/Source/Runtime/Engine/Public/Physics/ImmediatePhysics/ImmediatePhysicsPhysX/ImmediatePhysicsDeclares_PhysX.h
--- a/Engine/Source/Runtime/Engine/Public/Physics/ImmediatePhysics/ImmediatePhysicsPhysX/ImmediatePhysicsDeclares_PhysX.h
+++ b/Engine/Source/Runtime/Engine/Public/Physics/ImmediatePhysics/ImmediatePhysicsPhysX/ImmediatePhysicsDeclares_PhysX.h
@@ -7,5 +7,6 @@ namespace ImmediatePhysics_PhysX
 {
 	struct FSimulation;
 	struct FActorHandle;
+	struct FJointHandle;
 }
 #endif
diff --git a/Engine/Source/Runtime/Experimental/Chaos/Private/Chaos/PBDCollisionConstraints.cpp b/Engine/Source/Runtime/Experimental/Chaos/Private/Chaos/PBDCollisionConstraints.cpp
--- a/Engine/Source/Runtime/Experimental/Chaos/Private/Chaos/PBDCollisionConstraints.cpp
+++ b/Engine/Source/Runtime/Experimental/Chaos/Private/Chaos/PBDCollisionConstraints.cpp
@@ -49,9 +49,12 @@ namespace Chaos
 	CHAOS_API int32 EnableCollisions = 1;
 	FAutoConsoleVariableRef CVarEnableCollisions(TEXT("p.EnableCollisions"), EnableCollisions, TEXT("Enable/Disable collisions on the Chaos solver."));
 
-	DECLARE_CYCLE_STAT(TEXT("TPBDCollisionConstraints::Reset"), STAT_Collisions_Reset, STATGROUP_Chaos);
-	DECLARE_CYCLE_STAT(TEXT("TPBDCollisionConstraints::Apply"), STAT_Collisions_Apply, STATGROUP_Chaos);
-	DECLARE_CYCLE_STAT(TEXT("TPBDCollisionConstraints::ApplyPushOut"), STAT_Collisions_ApplyPushOut, STATGROUP_Chaos);
+	DECLARE_CYCLE_STAT(TEXT("Collisions::Reset"), STAT_Collisions_Reset, STATGROUP_ChaosCollision);
+	DECLARE_CYCLE_STAT(TEXT("Collisions::Apply"), STAT_Collisions_Apply, STATGROUP_ChaosCollision);
+	DECLARE_CYCLE_STAT(TEXT("Collisions::ApplyPushOut"), STAT_Collisions_ApplyPushOut, STATGROUP_ChaosCollision);
+
+	DECLARE_DWORD_COUNTER_STAT(TEXT("Collisions::NumPointConstraints"), STAT_NumPointConstraints, STATGROUP_ChaosCollision);
+	DECLARE_DWORD_COUNTER_STAT(TEXT("Collisions::NumManifoldConstraints"), STAT_NumManifoldConstraints, STATGROUP_ChaosCollision);
 
 	//
 	// Collision Constraint Container
@@ -282,6 +285,8 @@ namespace Chaos
 	void TPBDCollisionConstraints<T, d>::Apply(const T Dt, const int32 Iterations, const int32 NumIterations)
 	{
 		SCOPE_CYCLE_COUNTER(STAT_Collisions_Apply);
+		SET_DWORD_STAT(STAT_NumPointConstraints, PointConstraints.Num());
+		SET_DWORD_STAT(STAT_NumManifoldConstraints, IterativeConstraints.Num());
 
 		if (MApplyPairIterations > 0)
 		{
@@ -310,6 +315,8 @@ namespace Chaos
 	bool TPBDCollisionConstraints<T, d>::ApplyPushOut(const T Dt, const int32 Iterations, const int32 NumIterations)
 	{
 		SCOPE_CYCLE_COUNTER(STAT_Collisions_ApplyPushOut);
+		SET_DWORD_STAT(STAT_NumPointConstraints, PointConstraints.Num());
+		SET_DWORD_STAT(STAT_NumManifoldConstraints, IterativeConstraints.Num());
 
 		TSet<const TGeometryParticleHandle<T, d>*> TempStatic;
 		bool bNeedsAnotherIteration = false;
@@ -343,6 +350,9 @@ namespace Chaos
 	void TPBDCollisionConstraints<T, d>::Apply(const T Dt, const TArray<FConstraintContainerHandle*>& InConstraintHandles, const int32 Iterations, const int32 NumIterations)
 	{
 		SCOPE_CYCLE_COUNTER(STAT_Collisions_Apply);
+		SET_DWORD_STAT(STAT_NumPointConstraints, PointConstraints.Num());
+		SET_DWORD_STAT(STAT_NumManifoldConstraints, IterativeConstraints.Num());
+
 		if (MApplyPairIterations > 0)
 		{
 			PhysicsParallelFor(InConstraintHandles.Num(), [&](int32 ConstraintHandleIndex) {
@@ -367,6 +377,8 @@ namespace Chaos
 	bool TPBDCollisionConstraints<T, d>::ApplyPushOut(const T Dt, const TArray<FConstraintContainerHandle*>& InConstraintHandles, const TSet< const TGeometryParticleHandle<T, d>*>& IsTemporarilyStatic, int32 Iteration, int32 NumIterations)
 	{
 		SCOPE_CYCLE_COUNTER(STAT_Collisions_ApplyPushOut);
+		SET_DWORD_STAT(STAT_NumPointConstraints, PointConstraints.Num());
+		SET_DWORD_STAT(STAT_NumManifoldConstraints, IterativeConstraints.Num());
 
 		bool bNeedsAnotherIteration = false;
 		if (MApplyPushOutPairIterations > 0)
diff --git a/Engine/Source/Runtime/Experimental/Chaos/Private/Chaos/PBDJointConstraints.cpp b/Engine/Source/Runtime/Experimental/Chaos/Private/Chaos/PBDJointConstraints.cpp
--- a/Engine/Source/Runtime/Experimental/Chaos/Private/Chaos/PBDJointConstraints.cpp
+++ b/Engine/Source/Runtime/Experimental/Chaos/Private/Chaos/PBDJointConstraints.cpp
@@ -2,7 +2,6 @@
 #include "Chaos/PBDJointConstraints.h"
 #include "Chaos/ChaosDebugDraw.h"
 #include "Chaos/DebugDrawQueue.h"
-#include "Chaos/Joint/PBDJointSolverCholesky.h"
 #include "Chaos/Joint/PBDJointSolverGaussSeidel.h"
 #include "Chaos/Particle/ParticleUtilities.h"
 #include "Chaos/ParticleHandle.h"
@@ -17,15 +16,11 @@
 
 namespace Chaos
 {
-	DECLARE_CYCLE_STAT(TEXT("TPBDJointConstraints::Sort"), STAT_Joints_Sort, STATGROUP_Chaos);
-	DECLARE_CYCLE_STAT(TEXT("TPBDJointConstraints::Apply"), STAT_Joints_Apply, STATGROUP_Chaos);
-	DECLARE_CYCLE_STAT(TEXT("TPBDJointConstraints::ApplyPushOut"), STAT_Joints_ApplyPushOut, STATGROUP_Chaos);
-	DECLARE_CYCLE_STAT(TEXT("TPBDJointConstraints::SolveCholesky"), STAT_Joints_Solve_Cholesky, STATGROUP_Chaos);
-	DECLARE_CYCLE_STAT(TEXT("TPBDJointConstraints::SolveGaussSeidel"), STAT_Joints_Solve_GaussSeidel, STATGROUP_Chaos);
-
-	bool ChaosJoint_UseCholeskySolver = false;
-	FAutoConsoleVariableRef CVarChaosImmPhysDeltaTime(TEXT("p.Chaos.Joint.UseCholeskySolver"), ChaosJoint_UseCholeskySolver, TEXT("Whether to use the new solver"));
+	DECLARE_CYCLE_STAT(TEXT("Joints::Sort"), STAT_Joints_Sort, STATGROUP_ChaosJoint);
+	DECLARE_CYCLE_STAT(TEXT("Joints::Apply"), STAT_Joints_Apply, STATGROUP_ChaosJoint);
+	DECLARE_CYCLE_STAT(TEXT("Joints::ApplyPushOut"), STAT_Joints_ApplyPushOut, STATGROUP_ChaosJoint);
 
+	DECLARE_DWORD_COUNTER_STAT(TEXT("Joints::NumConstraints"), STAT_NumJointConstraints, STATGROUP_ChaosJoint);
 
 	//
 	// Constraint Handle
@@ -66,6 +61,11 @@ namespace Chaos
 		return ConstraintContainer->GetConstraintSettings(ConstraintIndex);
 	}
 
+	void FPBDJointConstraintHandle::SetSettings(const FPBDJointSettings& Settings)
+	{
+		ConstraintContainer->SetConstraintSettings(ConstraintIndex, Settings);
+	}
+
 	TVector<TGeometryParticleHandle<float,3>*, 2> FPBDJointConstraintHandle::GetConstrainedParticles() const 
 	{ 
 		return ConstraintContainer->GetConstrainedParticles(ConstraintIndex); 
@@ -391,6 +391,12 @@ namespace Chaos
 		return ConstraintSettings[ConstraintIndex];
 	}
 
+
+	void FPBDJointConstraints::SetConstraintSettings(int32 ConstraintIndex, const FPBDJointSettings& InConstraintSettings)
+	{
+		ConstraintSettings[ConstraintIndex] = InConstraintSettings;
+	}
+
 	
 	int32 FPBDJointConstraints::GetConstraintLevel(int32 ConstraintIndex) const
 	{
@@ -420,45 +426,39 @@ namespace Chaos
 
 	void FPBDJointConstraints::PrepareConstraints(FReal Dt)
 	{
-		if (!ChaosJoint_UseCholeskySolver)
+		ConstraintSolvers.SetNum(NumConstraints());
+		for (int32 ConstraintIndex = 0; ConstraintIndex < NumConstraints(); ++ConstraintIndex)
 		{
-			ConstraintSolvers.SetNum(NumConstraints());
-			for (int32 ConstraintIndex = 0; ConstraintIndex < NumConstraints(); ++ConstraintIndex)
-			{
-				const FPBDJointSettings& JointSettings = ConstraintSettings[ConstraintIndex];
-				const FTransformPair& JointFrames = ConstraintFrames[ConstraintIndex];
-				FJointSolverGaussSeidel& Solver = ConstraintSolvers[ConstraintIndex];
-
-				int32 Index0, Index1;
-				GetConstrainedParticleIndices(ConstraintIndex, Index0, Index1);
-				TGenericParticleHandle<FReal, 3> Particle0 = TGenericParticleHandle<FReal, 3>(ConstraintParticles[ConstraintIndex][Index0]);
-				TGenericParticleHandle<FReal, 3> Particle1 = TGenericParticleHandle<FReal, 3>(ConstraintParticles[ConstraintIndex][Index1]);
-
-				Solver.Init(
-					Dt,
-					Settings,
-					JointSettings,
-					FParticleUtilitiesXR::GetCoMWorldPosition(Particle0),	// Prev position
-					FParticleUtilitiesXR::GetCoMWorldPosition(Particle1),	// Prev position
-					FParticleUtilitiesXR::GetCoMWorldRotation(Particle0),	// Prev rotation
-					FParticleUtilitiesXR::GetCoMWorldRotation(Particle1),	// Prev rotation
-					Particle0->InvM(),
-					Particle0->InvI().GetDiagonal(),
-					Particle1->InvM(),
-					Particle1->InvI().GetDiagonal(),
-					FParticleUtilities::ParticleLocalToCoMLocal(Particle0, JointFrames[Index0]),
-					FParticleUtilities::ParticleLocalToCoMLocal(Particle1, JointFrames[Index1]));
-			}
+			const FPBDJointSettings& JointSettings = ConstraintSettings[ConstraintIndex];
+			const FTransformPair& JointFrames = ConstraintFrames[ConstraintIndex];
+			FJointSolverGaussSeidel& Solver = ConstraintSolvers[ConstraintIndex];
+
+			int32 Index0, Index1;
+			GetConstrainedParticleIndices(ConstraintIndex, Index0, Index1);
+			TGenericParticleHandle<FReal, 3> Particle0 = TGenericParticleHandle<FReal, 3>(ConstraintParticles[ConstraintIndex][Index0]);
+			TGenericParticleHandle<FReal, 3> Particle1 = TGenericParticleHandle<FReal, 3>(ConstraintParticles[ConstraintIndex][Index1]);
+
+			Solver.Init(
+				Dt,
+				Settings,
+				JointSettings,
+				FParticleUtilitiesXR::GetCoMWorldPosition(Particle0),	// Prev position
+				FParticleUtilitiesXR::GetCoMWorldPosition(Particle1),	// Prev position
+				FParticleUtilitiesXR::GetCoMWorldRotation(Particle0),	// Prev rotation
+				FParticleUtilitiesXR::GetCoMWorldRotation(Particle1),	// Prev rotation
+				Particle0->InvM(),
+				Particle0->InvI().GetDiagonal(),
+				Particle1->InvM(),
+				Particle1->InvI().GetDiagonal(),
+				FParticleUtilities::ParticleLocalToCoMLocal(Particle0, JointFrames[Index0]),
+				FParticleUtilities::ParticleLocalToCoMLocal(Particle1, JointFrames[Index1]));
 		}
 	}
 
 
 	void FPBDJointConstraints::UnprepareConstraints(FReal Dt)
 	{
-		if (!ChaosJoint_UseCholeskySolver)
-		{
-			ConstraintSolvers.Empty();
-		}
+		ConstraintSolvers.Empty();
 	}
 
 	
@@ -490,6 +490,7 @@ namespace Chaos
 	void FPBDJointConstraints::Apply(const FReal Dt, const int32 It, const int32 NumIts)
 	{
 		SCOPE_CYCLE_COUNTER(STAT_Joints_Apply);
+		SET_DWORD_STAT(STAT_NumJointConstraints, NumConstraints());
 
 		if (PreApplyCallback != nullptr)
 		{
@@ -498,21 +499,9 @@ namespace Chaos
 
 		if (Settings.ApplyPairIterations > 0)
 		{
-			if (ChaosJoint_UseCholeskySolver)
-			{
-				SCOPE_CYCLE_COUNTER(STAT_Joints_Solve_Cholesky);
-				for (int32 ConstraintIndex = 0; ConstraintIndex < NumConstraints(); ++ConstraintIndex)
-				{
-					SolvePosition_Cholesky(Dt, ConstraintIndex, Settings.ApplyPairIterations, It, NumIts);
-				}
-			}
-			else
+			for (int32 ConstraintIndex = 0; ConstraintIndex < NumConstraints(); ++ConstraintIndex)
 			{
-				SCOPE_CYCLE_COUNTER(STAT_Joints_Solve_GaussSeidel);
-				for (int32 ConstraintIndex = 0; ConstraintIndex < NumConstraints(); ++ConstraintIndex)
-				{
-					SolvePosition_GaussSiedel(Dt, ConstraintIndex, Settings.ApplyPairIterations, It, NumIts);
-				}
+				SolvePosition_GaussSiedel(Dt, ConstraintIndex, Settings.ApplyPairIterations, It, NumIts);
 			}
 		}
 
@@ -525,22 +514,16 @@ namespace Chaos
 	bool FPBDJointConstraints::ApplyPushOut(const FReal Dt, const int32 It, const int32 NumIts)
 	{
 		SCOPE_CYCLE_COUNTER(STAT_Joints_ApplyPushOut);
+		SET_DWORD_STAT(STAT_NumJointConstraints, NumConstraints());
 
 		// @todo(ccaulfield): track whether we are sufficiently solved
 		bool bNeedsAnotherIteration = true;
 
 		if (Settings.ApplyPushOutPairIterations > 0)
 		{
-			if (ChaosJoint_UseCholeskySolver)
-			{
-				// TODO
-			}
-			else
+			for (int32 ConstraintIndex = 0; ConstraintIndex < NumConstraints(); ++ConstraintIndex)
 			{
-				for (int32 ConstraintIndex = 0; ConstraintIndex < NumConstraints(); ++ConstraintIndex)
-				{
-					ProjectPosition_GaussSiedel(Dt, ConstraintIndex, Settings.ApplyPushOutPairIterations, It, NumIts);
-				}
+				ProjectPosition_GaussSiedel(Dt, ConstraintIndex, Settings.ApplyPushOutPairIterations, It, NumIts);
 			}
 		}
 
@@ -589,21 +572,9 @@ namespace Chaos
 
 		if (Settings.ApplyPairIterations > 0)
 		{
-			if (ChaosJoint_UseCholeskySolver)
-			{
-				SCOPE_CYCLE_COUNTER(STAT_Joints_Solve_Cholesky);
-				for (FConstraintContainerHandle* ConstraintHandle : SortedConstraintHandles)
-				{
-					SolvePosition_Cholesky(Dt, ConstraintHandle->GetConstraintIndex(), Settings.ApplyPairIterations, It, NumIts);
-				}
-			}
-			else
+			for (FConstraintContainerHandle* ConstraintHandle : SortedConstraintHandles)
 			{
-				SCOPE_CYCLE_COUNTER(STAT_Joints_Solve_GaussSeidel);
-				for (FConstraintContainerHandle* ConstraintHandle : SortedConstraintHandles)
-				{
-					SolvePosition_GaussSiedel(Dt, ConstraintHandle->GetConstraintIndex(), Settings.ApplyPairIterations, It, NumIts);
-				}
+				SolvePosition_GaussSiedel(Dt, ConstraintHandle->GetConstraintIndex(), Settings.ApplyPairIterations, It, NumIts);
 			}
 		}
 
@@ -630,16 +601,9 @@ namespace Chaos
 
 		if (Settings.ApplyPushOutPairIterations > 0)
 		{
-			if (ChaosJoint_UseCholeskySolver)
+			for (FConstraintContainerHandle* ConstraintHandle : SortedConstraintHandles)
 			{
-				// TODO
-			}
-			else
-			{
-				for (FConstraintContainerHandle* ConstraintHandle : SortedConstraintHandles)
-				{
-					ProjectPosition_GaussSiedel(Dt, ConstraintHandle->GetConstraintIndex(), Settings.ApplyPushOutPairIterations, It, NumIts);
-				}
+				ProjectPosition_GaussSiedel(Dt, ConstraintHandle->GetConstraintIndex(), Settings.ApplyPushOutPairIterations, It, NumIts);
 			}
 		}
 
@@ -693,49 +657,6 @@ namespace Chaos
 	}
 
 
-	// This position solver solves all (active) inner position and angular constraints simultaneously by building the Jacobian and solving [JMJt].DX = C
-	// where DX(6x1) are the unknown position and rotation corrections, C(Nx1) is the current constraint error, J(Nx6) the Jacobian, M(6x6) the inverse mass matrix, 
-	// and N the number of active constraints. "Active constraints" are all bilateral constraints plus any violated unilateral constraints.
-	void FPBDJointConstraints::SolvePosition_Cholesky(const FReal Dt, const int32 ConstraintIndex, const int32 NumPairIts, const int32 It, const int32 NumIts)
-	{
-		const TVector<TGeometryParticleHandle<FReal, 3>*, 2>& Constraint = ConstraintParticles[ConstraintIndex];
-		UE_LOG(LogChaosJoint, VeryVerbose, TEXT("Solve Joint Constraint %d %s %s (dt = %f; it = %d / %d)"), ConstraintIndex, *Constraint[0]->ToString(), *Constraint[1]->ToString(), Dt, It, NumIts);
-
-		const FPBDJointSettings& JointSettings = ConstraintSettings[ConstraintIndex];
-		const FTransformPair& JointFrames = ConstraintFrames[ConstraintIndex];
-
-		int32 Index0, Index1;
-		GetConstrainedParticleIndices(ConstraintIndex, Index0, Index1);
-		TGenericParticleHandle<FReal, 3> Particle0 = TGenericParticleHandle<FReal, 3>(ConstraintParticles[ConstraintIndex][Index0]);
-		TGenericParticleHandle<FReal, 3> Particle1 = TGenericParticleHandle<FReal, 3>(ConstraintParticles[ConstraintIndex][Index1]);
-
-		FJointSolverCholesky Solver;
-		Solver.InitConstraints(
-			Dt, 
-			Settings, 
-			JointSettings, 
-			FParticleUtilities::GetCoMWorldPosition(Particle0),
-			FParticleUtilities::GetCoMWorldRotation(Particle0),
-			FParticleUtilities::GetCoMWorldPosition(Particle1),
-			FParticleUtilities::GetCoMWorldRotation(Particle1),
-			Particle0->InvM(), 
-			Particle0->InvI(), 
-			Particle1->InvM(), 
-			Particle1->InvI(), 
-			FParticleUtilities::ParticleLocalToCoMLocal(Particle0, JointFrames[Index0]),
-			FParticleUtilities::ParticleLocalToCoMLocal(Particle1, JointFrames[Index1]));
-		
-		for (int32 PairIt = 0; PairIt < NumPairIts; ++PairIt)
-		{
-			Solver.ApplyDrives(Dt, JointSettings);
-			Solver.ApplyConstraints(Dt, JointSettings);
-		}
-
-		UpdateParticleState(Particle0->CastToRigidParticle(), Dt, Solver.GetP(0), Solver.GetQ(0));
-		UpdateParticleState(Particle1->CastToRigidParticle(), Dt, Solver.GetP(1), Solver.GetQ(1));
-	}
-
-
 	// This position solver iterates over each of the inner constraints (position, twist, swing) and solves them independently.
 	// This will converge slowly in some cases, particularly where resolving angular constraints violates position constraints and vice versa.
 	void FPBDJointConstraints::SolvePosition_GaussSiedel(const FReal Dt, const int32 ConstraintIndex, const int32 NumPairIts, const int32 It, const int32 NumIts)
diff --git a/Engine/Source/Runtime/Experimental/Chaos/Public/Chaos/PBDJointConstraints.h b/Engine/Source/Runtime/Experimental/Chaos/Public/Chaos/PBDJointConstraints.h
--- a/Engine/Source/Runtime/Experimental/Chaos/Public/Chaos/PBDJointConstraints.h
+++ b/Engine/Source/Runtime/Experimental/Chaos/Public/Chaos/PBDJointConstraints.h
@@ -30,6 +30,7 @@ namespace Chaos
 		CHAOS_API void SetParticleLevels(const TVector<int32, 2>& ParticleLevels);
 		CHAOS_API int32 GetConstraintLevel() const;
 		CHAOS_API const FPBDJointSettings& GetSettings() const;
+		CHAOS_API void SetSettings(const FPBDJointSettings& Settings);
 		CHAOS_API TVector<TGeometryParticleHandle<float,3>*, 2> GetConstrainedParticles() const;
 
 	protected:
@@ -123,6 +124,7 @@ namespace Chaos
 		const FParticlePair& GetConstrainedParticles(int32 ConstraintIndex) const;
 
 		const FPBDJointSettings& GetConstraintSettings(int32 ConstraintIndex) const;
+		void SetConstraintSettings(int32 ConstraintIndex, const FPBDJointSettings& InConstraintSettings);
 
 		int32 GetConstraintLevel(int32 ConstraintIndex) const;
 		void SetParticleLevels(int32 ConstraintIndex, const TVector<int32, 2>& ParticleLevels);
diff --git a/Engine/Source/Runtime/Experimental/Chaos/Public/ChaosStats.h b/Engine/Source/Runtime/Experimental/Chaos/Public/ChaosStats.h
--- a/Engine/Source/Runtime/Experimental/Chaos/Public/ChaosStats.h
+++ b/Engine/Source/Runtime/Experimental/Chaos/Public/ChaosStats.h
@@ -10,6 +10,8 @@ DECLARE_STATS_GROUP_VERBOSE(TEXT("ChaosWide"), STATGROUP_ChaosWide, STATCAT_Adva
 DECLARE_STATS_GROUP(TEXT("ChaosThread"), STATGROUP_ChaosThread, STATCAT_Advanced);
 DECLARE_STATS_GROUP(TEXT("ChaosDedicated"), STATGROUP_ChaosDedicated, STATCAT_Advanced);
 DECLARE_STATS_GROUP(TEXT("ChaosEngine"), STATGROUP_ChaosEngine, STATCAT_Advanced);
+DECLARE_STATS_GROUP(TEXT("ChaosCollision"), STATGROUP_ChaosCollision, STATCAT_Advanced);
+DECLARE_STATS_GROUP(TEXT("ChaosJoint"), STATGROUP_ChaosJoint, STATCAT_Advanced);
 
 DECLARE_CYCLE_STAT_EXTERN(TEXT("Physics Advance"), STAT_PhysicsAdvance, STATGROUP_Chaos, CHAOS_API);
 DECLARE_CYCLE_STAT_EXTERN(TEXT("Solver Advance"), STAT_SolverAdvance, STATGROUP_Chaos, CHAOS_API);
