Chaos : Protect against potential null pointer

diff --git a/Engine/Source/Runtime/Experimental/Chaos/Private/Chaos/PBDConstraintGraph.cpp b/Engine/Source/Runtime/Experimental/Chaos/Private/Chaos/PBDConstraintGraph.cpp
--- a/Engine/Source/Runtime/Experimental/Chaos/Private/Chaos/PBDConstraintGraph.cpp
+++ b/Engine/Source/Runtime/Experimental/Chaos/Private/Chaos/PBDConstraintGraph.cpp
@@ -389,20 +389,24 @@ void FPBDConstraintGraph::ComputeIslands(const TParticleView<TPBDRigidParticles<
 
 			for (TGeometryParticleHandle<FReal, 3>* Particle : IslandToParticles[Island])
 			{
-				TPBDRigidParticleHandle<FReal, 3>* PBDRigid = Particle->CastToRigidParticle();
-				const bool bIsDynamic = PBDRigid && PBDRigid->ObjectState() != EObjectStateType::Kinematic;
-				int32 TmpIsland = bIsDynamic ? PBDRigid->Island() : INDEX_NONE; //question: should we even store non dynamics in this array?
-
-				if (OtherIsland == INDEX_NONE && TmpIsland >= 0)
+				if (CHAOS_ENSURE(Particle))
 				{
-					OtherIsland = TmpIsland;
-				}
-				else
-				{
-					if (TmpIsland >= 0 && OtherIsland != TmpIsland)
+					TPBDRigidParticleHandle<FReal, 3>* PBDRigid = Particle->CastToRigidParticle();
+
+					const bool bIsDynamic = PBDRigid && PBDRigid->ObjectState() != EObjectStateType::Kinematic;
+					int32 TmpIsland = bIsDynamic ? PBDRigid->Island() : INDEX_NONE; //question: should we even store non dynamics in this array?
+
+					if (OtherIsland == INDEX_NONE && TmpIsland >= 0)
 					{
-						bIsSameIsland = false;
-						break;
+						OtherIsland = TmpIsland;
+					}
+					else
+					{
+						if (TmpIsland >= 0 && OtherIsland != TmpIsland)
+						{
+							bIsSameIsland = false;
+							break;
+						}
 					}
 				}
 			}
@@ -422,19 +426,22 @@ void FPBDConstraintGraph::ComputeIslands(const TParticleView<TPBDRigidParticles<
 			{
 				for (TGeometryParticleHandle<FReal, 3>* Particle : IslandToParticles[Island])
 				{
-					TPBDRigidParticleHandle<FReal, 3>* PBDRigid = Particle->CastToRigidParticle();
-					if (PBDRigid && PBDRigid->ObjectState() != EObjectStateType::Kinematic)
+					if (CHAOS_ENSURE(Particle))
 					{
-						if (!PBDRigid->Disabled())	// todo: why is this needed? [we aren't handling enable/disable state changes properly so disabled particles end up in the graph.]
+						TPBDRigidParticleHandle<FReal, 3>* PBDRigid = Particle->CastToRigidParticle();
+						if (PBDRigid && PBDRigid->ObjectState() != EObjectStateType::Kinematic)
+						{
+							if (!PBDRigid->Disabled())	// todo: why is this needed? [we aren't handling enable/disable state changes properly so disabled particles end up in the graph.]
+							{
+								PBDRigid->SetSleeping(false);
+								Particles.ActivateParticle(Particle);
+							}
+						}
+						else
 						{
-							PBDRigid->SetSleeping(false);
 							Particles.ActivateParticle(Particle);
 						}
 					}
-					else
-					{
-						Particles.ActivateParticle(Particle);
-					}
 				}
 			}
 
