#chaos - Apply similar clamp as in 10942094 to GJKPenetration when Penetration is very slightly below zero, so PushOut doesn't pull things in instead.

diff --git a/Engine/Source/Runtime/Experimental/Chaos/Public/Chaos/GJK.h b/Engine/Source/Runtime/Experimental/Chaos/Public/Chaos/GJK.h
--- a/Engine/Source/Runtime/Experimental/Chaos/Public/Chaos/GJK.h
+++ b/Engine/Source/Runtime/Experimental/Chaos/Public/Chaos/GJK.h
@@ -179,7 +179,7 @@ namespace Chaos
 			
 			const T PreDist = FMath::Sqrt(PrevDist2);
 			OutNormal = (ClosestBInA - ClosestA).GetUnsafeNormal();	//question: should we just use PreDist2?
-			const T Penetration = ThicknessA + ThicknessB - PreDist;
+			const T Penetration = FMath::Clamp<T>(ThicknessA + ThicknessB - PreDist, 0, TNumericLimits<T>::Max());
 			OutPenetration = Penetration;
 			OutClosestA = ClosestA + OutNormal * ThicknessA;
 			OutClosestB = ClosestBInA - OutNormal * ThicknessB;
