Fix TParticleView iteration such that the EParticleType of the handle is preserved.  Make parallel and non-parallel iteration function in the same way in this regard.

diff --git a/Engine/Source/Runtime/Experimental/Chaos/Public/Chaos/ParticleHandle.h b/Engine/Source/Runtime/Experimental/Chaos/Public/Chaos/ParticleHandle.h
--- a/Engine/Source/Runtime/Experimental/Chaos/Public/Chaos/ParticleHandle.h
+++ b/Engine/Source/Runtime/Experimental/Chaos/Public/Chaos/ParticleHandle.h
@@ -243,6 +243,7 @@ public:
 		GeometryParticles = InParticles;
 	}
 
+	EParticleType GetParticleType() const { return Type; }
 protected:
 	union
 	{
@@ -692,7 +693,10 @@ public:
 	void SetToBeRemovedOnFracture(const bool bToBeRemovedOnFracture) { PBDRigidParticles->ToBeRemovedOnFracture(ParticleIdx) = bToBeRemovedOnFracture; }
 
 	EObjectStateType ObjectState() const { return PBDRigidParticles->ObjectState(ParticleIdx); }
+	void SetObjectState(EObjectStateType InState) { PBDRigidParticles->SetObjectState(ParticleIdx, InState); }
 	void SetObjectStateLowLevel(EObjectStateType InState) { PBDRigidParticles->SetObjectState(ParticleIdx, InState); }
+	
+	bool Sleeping() const { return PBDRigidParticles->Sleeping(ParticleIdx); }
 	void SetSleeping(bool bSleeping) { PBDRigidParticles->SetSleeping(ParticleIdx, bSleeping); }
 
 	//Really only useful when using a transient handle
@@ -739,32 +743,60 @@ public:
 		return Serializable;
 	}
 
-	const auto& ClusterIds() const { return PBDRigidClusteredParticles->ClusterIds(ParticleIdx); }
-	auto& ClusterIds() { return PBDRigidClusteredParticles->ClusterIds(ParticleIdx); }
-
-	const auto& ChildToParent() const { return PBDRigidClusteredParticles->ChildToParent(ParticleIdx); }
-	auto& ChildToParent() { return PBDRigidClusteredParticles->ChildToParent(ParticleIdx); }
-
-	const auto& ClusterGroupIndex() const { return PBDRigidClusteredParticles->ClusterGroupIndex(ParticleIdx); }
-	auto& ClusterGroupIndex() { return PBDRigidClusteredParticles->ClusterGroupIndex(ParticleIdx); }
-
-	const auto& InternalCluster() const { return PBDRigidClusteredParticles->InternalCluster(ParticleIdx); }
-	auto& InternalCluster() { return PBDRigidClusteredParticles->InternalCluster(ParticleIdx); }
-
-	const auto& ChildrenSpatial() const { return PBDRigidClusteredParticles->ChildrenSpatial(ParticleIdx); }
-	auto& ChildrenSpatial() { return PBDRigidClusteredParticles->ChildrenSpatial(ParticleIdx); }
-
-	//const auto& MultiChildProxyId() const { return PBDRigidClusteredParticles->MultiChildProxyId(ParticleIdx); }
-	//auto& MultiChildProxyId() { return PBDRigidClusteredParticles->MultiChildProxyId(ParticleIdx); }
-
-	//const auto& MultiChildProxyData() const { return PBDRigidClusteredParticles->MultiChildProxyData(ParticleIdx); }
-	//auto& MultiChildProxyData() { return PBDRigidClusteredParticles->MultiChildProxyData(ParticleIdx); }
-
-	const auto& CollisionImpulses() const { return PBDRigidClusteredParticles->CollisionImpulses(ParticleIdx); }
-	auto& CollisionImpulses() { return PBDRigidClusteredParticles->CollisionImpulses(ParticleIdx); }
-
-	const auto& Strains() const { return PBDRigidClusteredParticles->Strains(ParticleIdx); }
-	auto& Strains() { return PBDRigidClusteredParticles->Strains(ParticleIdx); }
+//	const ClusterId& ClusterId() const { return PBDRigidClusteredParticles->ClusterIds(ParticleIdx); }
+//	ClusterId& ClusterId() { return PBDRigidClusteredParticles->ClusterIds(ParticleIdx); }
+	void SetClusterId(const ClusterId& Id) { PBDRigidClusteredParticles->ClusterIds(ParticleIdx) = Id; }
+	const ClusterId& ClusterIds() const { return PBDRigidClusteredParticles->ClusterIds(ParticleIdx); }
+	ClusterId& ClusterIds() { return PBDRigidClusteredParticles->ClusterIds(ParticleIdx); }
+
+	const TRigidTransform<T,d>& ChildToParent() const { return PBDRigidClusteredParticles->ChildToParent(ParticleIdx); }
+	TRigidTransform<T,d>& ChildToParent() { return PBDRigidClusteredParticles->ChildToParent(ParticleIdx); }
+	void SetChildToParent(const TRigidTransform<T, d>& Xf) { PBDRigidClusteredParticles->ChildToParent(ParticleIdx) = Xf; }
+
+	const int32& ClusterGroupIndex() const { return PBDRigidClusteredParticles->ClusterGroupIndex(ParticleIdx); }
+	int32& ClusterGroupIndex() { return PBDRigidClusteredParticles->ClusterGroupIndex(ParticleIdx); }
+	void SetClusterGroupIndex(const int32 Idx) { PBDRigidClusteredParticles->ClusterGroupIndex(ParticleIdx) = Idx; }
+
+	const bool& InternalCluster() const { return PBDRigidClusteredParticles->InternalCluster(ParticleIdx); }
+	bool& InternalCluster() { return PBDRigidClusteredParticles->InternalCluster(ParticleIdx); }
+	void SetInternalCluster(const bool Value) { PBDRigidClusteredParticles->InternalCluster(ParticleIdx) = Value; }
+
+	const TUniquePtr<FImplicitObjectUnionClustered>& ChildrenSpatial() const { return PBDRigidClusteredParticles->ChildrenSpatial(ParticleIdx); }
+	TUniquePtr<FImplicitObjectUnionClustered>& ChildrenSpatial() { return PBDRigidClusteredParticles->ChildrenSpatial(ParticleIdx); }
+	void SetChildrenSpatial(TUniquePtr<FImplicitObjectUnion>& Obj) { PBDRigidClusteredParticles->ChildrenSpatial(ParticleIdx) = Obj; }
+
+	const FMultiChildProxyId& MultiChildProxyId() const { return PBDRigidClusteredParticles->MultiChildProxyId(ParticleIdx); }
+	FMultiChildProxyId& MultiChildProxyId() { return PBDRigidClusteredParticles->MultiChildProxyId(ParticleIdx); }
+	void SetMultiChildProxyId(const FMultiChildProxyId& Id) { PBDRigidClusteredParticles->MultiChildProxyId(ParticleIdx) = Id; }
+
+	const TUniquePtr<TMultiChildProxyData<T, d>>& MultiChildProxyData() const { return PBDRigidClusteredParticles->MultiChildProxyData(ParticleIdx); }
+	TUniquePtr<TMultiChildProxyData<T, d>>& MultiChildProxyData() { return PBDRigidClusteredParticles->MultiChildProxyData(ParticleIdx); }
+	void SetMultiChildProxyData(TUniquePtr<TMultiChildProxyData<T, d>>&& ProxyData) { PBDRigidClusteredParticles->MultiChildProxyData(ParticleIdx) = MoveTemp(ProxyData); }
+
+	const T& CollisionImpulse() const { return PBDRigidClusteredParticles->CollisionImpulses(ParticleIdx); }
+	T& CollisionImpulse() { return PBDRigidClusteredParticles->CollisionImpulses(ParticleIdx); }
+	void SetCollisionImpulse(const T Value) { PBDRigidClusteredParticles->CollisionImpulses(ParticleIdx) = Value; }
+	// Deprecated
+	const T& CollisionImpulses() const { return PBDRigidClusteredParticles->CollisionImpulses(ParticleIdx); }
+	// Deprecated
+	T& CollisionImpulses() { return PBDRigidClusteredParticles->CollisionImpulses(ParticleIdx); }
+	// Deprecated
+	void SetCollisionImpulses(const T Value) { PBDRigidClusteredParticles->CollisionImpulses(ParticleIdx) = Value; }
+
+	const T& Strain() const { return PBDRigidClusteredParticles->Strains(ParticleIdx); }
+	T& Strain() { return PBDRigidClusteredParticles->Strains(ParticleIdx); }
+	void SetStrain(const T Value) { PBDRigidClusteredParticles->Strains(ParticleIdx) = Value; }
+	// Deprecated
+	const T& Strains() const { return PBDRigidClusteredParticles->Strains(ParticleIdx); }
+	// Deprecated
+	T& Strains() { return PBDRigidClusteredParticles->Strains(ParticleIdx); }
+	// Deprecated
+	void SetStrains(const T Value) { PBDRigidClusteredParticles->Strains(ParticleIdx) = Value; }
+
+	const TArray<TConnectivityEdge<T>>& ConnectivityEdges() const { return PBDRigidClusteredParticles->ConnectivityEdges(ParticleIdx); }
+	TArray<TConnectivityEdge<T>>& ConnectivityEdges() { return PBDRigidClusteredParticles->ConnectivityEdges(ParticleIdx); }
+	void SetConnectivityEdges(const TArray<TConnectivityEdge<T>>& Edges) { PBDRigidClusteredParticles->ConnectivityEdges(ParticleIdx) = Edges; }
+	void SetConnectivityEdges(TArray<TConnectivityEdge<T>>&& Edges) { PBDRigidClusteredParticles->ConnectivityEdges(ParticleIdx) = MoveTemp(Edges); }
 
 	const TPBDRigidClusteredParticleHandleImp<T, d, true>* Handle() const { return PBDRigidClusteredParticles->Handle(ParticleIdx); }
 	TPBDRigidClusteredParticleHandleImp<T, d, true>* Handle() { return PBDRigidClusteredParticles->Handle(ParticleIdx); }
@@ -775,11 +807,11 @@ public:
 };
 
 template <typename T, int d, bool bPersistent = true>
-class TPBDGeometryCollectionParticleHandleImp : public TPBDRigidParticleHandleImp<T, d, bPersistent>
+class TPBDGeometryCollectionParticleHandleImp : public TPBDRigidClusteredParticleHandleImp<T, d, bPersistent>
 {
 public:
 	using TGeometryParticleHandleImp<T, d, bPersistent>::ParticleIdx;
-	using TGeometryParticleHandleImp<T, d, bPersistent>::PBDRigidParticles;
+	using TGeometryParticleHandleImp<T, d, bPersistent>::PBDRigidClusteredParticles;
 	using TGeometryParticleHandleImp<T, d, bPersistent>::Type;
 	using TTransientHandle = TTransientPBDGeometryCollectionParticleHandle<T, d>;
 	using TSOAType = TPBDGeometryCollectionParticles<T, d>;
@@ -789,17 +821,21 @@ protected:
 
 	//needed for serialization
 	TPBDGeometryCollectionParticleHandleImp()
-		: TPBDRigidParticleHandleImp<T, d, bPersistent>()
-	{}
+		: TPBDRigidClusteredParticleHandleImp<T, d, bPersistent>()
+	{
+		Type = EParticleType::GeometryCollection;
+	}
 
 	TPBDGeometryCollectionParticleHandleImp<T, d, bPersistent>(
 		TSerializablePtr<TPBDGeometryCollectionParticles<T, d>> Particles, 
 		int32 InIdx, 
 		int32 InGlobalIdx, 
 		const TPBDRigidParticleParameters<T, d>& Params = TPBDRigidParticleParameters<T, d>())
-		: TPBDRigidParticleHandleImp<T, d, bPersistent>(
-			TSerializablePtr<TPBDRigidParticles<T, d>>(Particles), InIdx, InGlobalIdx, Params)
-	{}
+		: TPBDRigidClusteredParticleHandleImp<T, d, bPersistent>(
+			TSerializablePtr<TPBDRigidClusteredParticles<T, d>>(Particles), InIdx, InGlobalIdx, Params)
+	{
+		Type = EParticleType::GeometryCollection;
+	}
 public:
 
 	static TUniquePtr<TPBDGeometryCollectionParticleHandleImp<T, d, bPersistent>> CreateParticleHandle(
@@ -819,8 +855,8 @@ public:
 		return Serializable;
 	}
 
-	const TPBDGeometryCollectionParticleHandleImp<T, d, true>* Handle() const { return PBDRigidParticles->Handle(ParticleIdx); }
-	TPBDGeometryCollectionParticleHandleImp<T, d, true>* Handle() { return PBDRigidParticles->Handle(ParticleIdx); }
+	const TPBDGeometryCollectionParticleHandleImp<T, d, true>* Handle() const { return PBDRigidClusteredParticles->Handle(ParticleIdx); }
+	TPBDGeometryCollectionParticleHandleImp<T, d, true>* Handle() { return PBDRigidClusteredParticles->Handle(ParticleIdx); }
 
 	static constexpr EParticleType StaticType() { return EParticleType::GeometryCollection; }
 };
diff --git a/Engine/Source/Runtime/Experimental/Chaos/Public/Chaos/ParticleIterator.h b/Engine/Source/Runtime/Experimental/Chaos/Public/Chaos/ParticleIterator.h
--- a/Engine/Source/Runtime/Experimental/Chaos/Public/Chaos/ParticleIterator.h
+++ b/Engine/Source/Runtime/Experimental/Chaos/Public/Chaos/ParticleIterator.h
@@ -326,6 +326,10 @@ public:
 	TConstParticleIterator(const TConstParticleIterator& Rhs) = default;
 
 	operator bool() const { return TransientHandle.GeometryParticles != nullptr; }
+
+	// The non-parallel implementation of iteration should not deviate in behavior
+	// from the parallel implementation in \c ParticleViewParallelForImp().  They
+	// must be kept in sync.
 	TConstParticleIterator<TSOA>& operator++()
 	{
 		RangedForValidation();
@@ -343,18 +347,20 @@ public:
 			++CurHandleIdx;
 			if (CurHandleIdx < CurHandlesArray->Num())
 			{
-				TransientHandle.ParticleIdx = (*CurHandlesArray)[CurHandleIdx]->ParticleIdx;
-				TransientHandle.GeometryParticles = (*CurHandlesArray)[CurHandleIdx]->GeometryParticles;
+				// Reconstruct the TransientHandle so that it has a chance to update
+				// other data members, like the particle type in the case of geometry 
+				// particles.
+				THandle* HandlePtr = (*CurHandlesArray)[CurHandleIdx];
+				THandleBase Handle(HandlePtr->GeometryParticles, HandlePtr->ParticleIdx);
+				TransientHandle = static_cast<TTransientHandle&>(Handle);
 			}
 			else
 			{
 				IncSOAIdx();
 			}
 		}
-
 		return *this;
 	}
-
 	
 	const TTransientHandle& operator*() const
 	{
@@ -602,6 +608,8 @@ TParticleView<TSOA> MakeParticleView(TSOA* SOA)
 	return TParticleView<TSOA>(MoveTemp(SOAs));
 }
 
+// The non-parallel implementation of iteration should not deviate in behavior from
+// this parallel implementation.  They must be kept in sync.
 template <typename TParticleView, typename Lambda>
 void ParticleViewParallelForImp(const TParticleView& Particles, const Lambda& Func)
 {
@@ -630,8 +638,11 @@ void ParticleViewParallelForImp(const TParticleView& Particles, const Lambda& Fu
 			const int32 HandleCount = CurHandlesArray->Num();
 			::ParallelFor(HandleCount, [&Func, CurHandlesArray, ParticleIdxOff](const int32 HandleIdx)
 			{
+				// Reconstruct the TransientHandle so that it has a chance to update
+				// other data members, like the particle type in the case of geometry 
+				// particles.
 				THandle* HandlePtr = (*CurHandlesArray)[HandleIdx];
-				THandleBase Handle = THandleBase(HandlePtr->GeometryParticles, HandlePtr->ParticleIdx);
+				THandleBase Handle(HandlePtr->GeometryParticles, HandlePtr->ParticleIdx);
 				Func(static_cast<TTransientHandle&>(Handle), ParticleIdxOff + HandleIdx);
 			});
 			ParticleIdxOff += HandleCount;
@@ -641,7 +652,10 @@ void ParticleViewParallelForImp(const TParticleView& Particles, const Lambda& Fu
 			// Do a regular parallel for over the particles in this SOA view
 			::ParallelFor(ParticleCount, [&Func, &SOAView, ParticleIdxOff](const int32 ParticleIdx)
 			{
-				THandleBase Handle = THandleBase(SOAView.SOA, ParticleIdx);
+				// Reconstruct the TransientHandle so that it has a chance to update
+				// other data members, like the particle type in the case of geometry 
+				// particles.
+				THandleBase Handle(SOAView.SOA, ParticleIdx);
 				Func(static_cast<TTransientHandle&>(Handle), ParticleIdxOff + ParticleIdx);
 			});
 			ParticleIdxOff += ParticleCount;
diff --git a/Engine/Source/Runtime/Experimental/ChaosSolvers/Private/PBDRigidsSolver.cpp b/Engine/Source/Runtime/Experimental/ChaosSolvers/Private/PBDRigidsSolver.cpp
--- a/Engine/Source/Runtime/Experimental/ChaosSolvers/Private/PBDRigidsSolver.cpp
+++ b/Engine/Source/Runtime/Experimental/ChaosSolvers/Private/PBDRigidsSolver.cpp
@@ -570,9 +570,10 @@ namespace Chaos
 		TArray<FGeometryCollectionPhysicsProxy*> ActiveGC;
 		ActiveGC.Reserve(GeometryCollectionPhysicsProxies.Num());
 
-		for (Chaos::TGeometryParticleHandleImp<float, 3, false>& ActiveObject : GetParticles().GetActiveParticlesView())
+		TParticleView<TPBDRigidParticles<float, 3>>& ActiveParticles = GetParticles().GetActiveParticlesView();
+		for (Chaos::TPBDRigidParticleHandleImp<float, 3, false>& ActiveObject : ActiveParticles)
 		{
-			switch (ActiveObject.Type)
+			switch (ActiveObject.GetParticleType())
 			{
 			case Chaos::EParticleType::Rigid:
 				((FRigidParticlePhysicsProxy*)(ActiveObject.GTGeometryParticle()->Proxy))->BufferPhysicsResults();
@@ -603,9 +604,10 @@ namespace Chaos
 		TArray<FGeometryCollectionPhysicsProxy*> ActiveGC;
 		ActiveGC.Reserve(GeometryCollectionPhysicsProxies.Num());
 
-		for (Chaos::TGeometryParticleHandleImp < float, 3, false>& ActiveObject : GetParticles().GetActiveParticlesView())
+		TParticleView<TPBDRigidParticles<float, 3>>& ActiveParticles = GetParticles().GetActiveParticlesView();
+		for (Chaos::TPBDRigidParticleHandleImp<float, 3, false>& ActiveObject : ActiveParticles)
 		{
-			switch (ActiveObject.Type)
+			switch (ActiveObject.GetParticleType())
 			{
 			case Chaos::EParticleType::Rigid:
 				((FRigidParticlePhysicsProxy*)(ActiveObject.GTGeometryParticle()->Proxy))->FlipBuffer();
@@ -636,9 +638,10 @@ namespace Chaos
 		TArray<FGeometryCollectionPhysicsProxy*> ActiveGC;
 		ActiveGC.Reserve(GeometryCollectionPhysicsProxies.Num());
 
-		for (Chaos::TGeometryParticleHandleImp < float, 3, false>& ActiveObject : GetParticles().GetActiveParticlesView())
+		TParticleView<TPBDRigidParticles<float, 3>>& ActiveParticles = GetParticles().GetActiveParticlesView();
+		for (Chaos::TPBDRigidParticleHandleImp<float, 3, false>& ActiveObject : ActiveParticles)
 		{
-			switch (ActiveObject.Type)
+			switch (ActiveObject.GetParticleType())
 			{
 			case Chaos::EParticleType::Rigid:
 				((FRigidParticlePhysicsProxy*)(ActiveObject.GTGeometryParticle()->Proxy))->PullFromPhysicsState();
