Correct ImplicitObject's ToString() string.

diff --git a/Engine/Source/Runtime/Experimental/Chaos/Public/Chaos/ImplicitObject.h b/Engine/Source/Runtime/Experimental/Chaos/Public/Chaos/ImplicitObject.h
--- a/Engine/Source/Runtime/Experimental/Chaos/Public/Chaos/ImplicitObject.h
+++ b/Engine/Source/Runtime/Experimental/Chaos/Public/Chaos/ImplicitObject.h
@@ -273,7 +273,7 @@ public:
 
 	virtual FString ToString() const
 	{
-		return FString::Printf(TEXT("ImplicitObject bIsConvex:%d, bIgnoreAnalyticCollision:%d, bHasBoundingBox:%d"), bIsConvex, bDoCollide, bHasBoundingBox);
+		return FString::Printf(TEXT("ImplicitObject bIsConvex:%d, bDoCollide:%d, bHasBoundingBox:%d"), bIsConvex, bDoCollide, bHasBoundingBox);
 	}
 
 	void SerializeImp(FArchive& Ar);
