Remove changes that went into change 11000863, that shouldn't have.  Serves me right for being fancy and moving shelves around...

diff --git a/Engine/Source/Runtime/Experimental/Chaos/Public/Chaos/ParticleHandle.h b/Engine/Source/Runtime/Experimental/Chaos/Public/Chaos/ParticleHandle.h
--- a/Engine/Source/Runtime/Experimental/Chaos/Public/Chaos/ParticleHandle.h
+++ b/Engine/Source/Runtime/Experimental/Chaos/Public/Chaos/ParticleHandle.h
@@ -693,10 +693,7 @@ public:
 	void SetToBeRemovedOnFracture(const bool bToBeRemovedOnFracture) { PBDRigidParticles->ToBeRemovedOnFracture(ParticleIdx) = bToBeRemovedOnFracture; }
 
 	EObjectStateType ObjectState() const { return PBDRigidParticles->ObjectState(ParticleIdx); }
-	void SetObjectState(EObjectStateType InState) { PBDRigidParticles->SetObjectState(ParticleIdx, InState); }
 	void SetObjectStateLowLevel(EObjectStateType InState) { PBDRigidParticles->SetObjectState(ParticleIdx, InState); }
-	
-	bool Sleeping() const { return PBDRigidParticles->Sleeping(ParticleIdx); }
 	void SetSleeping(bool bSleeping) { PBDRigidParticles->SetSleeping(ParticleIdx, bSleeping); }
 
 	//Really only useful when using a transient handle
@@ -743,60 +740,32 @@ public:
 		return Serializable;
 	}
 
-//	const ClusterId& ClusterId() const { return PBDRigidClusteredParticles->ClusterIds(ParticleIdx); }
-//	ClusterId& ClusterId() { return PBDRigidClusteredParticles->ClusterIds(ParticleIdx); }
-	void SetClusterId(const ClusterId& Id) { PBDRigidClusteredParticles->ClusterIds(ParticleIdx) = Id; }
-	const ClusterId& ClusterIds() const { return PBDRigidClusteredParticles->ClusterIds(ParticleIdx); }
-	ClusterId& ClusterIds() { return PBDRigidClusteredParticles->ClusterIds(ParticleIdx); }
-
-	const TRigidTransform<T,d>& ChildToParent() const { return PBDRigidClusteredParticles->ChildToParent(ParticleIdx); }
-	TRigidTransform<T,d>& ChildToParent() { return PBDRigidClusteredParticles->ChildToParent(ParticleIdx); }
-	void SetChildToParent(const TRigidTransform<T, d>& Xf) { PBDRigidClusteredParticles->ChildToParent(ParticleIdx) = Xf; }
-
-	const int32& ClusterGroupIndex() const { return PBDRigidClusteredParticles->ClusterGroupIndex(ParticleIdx); }
-	int32& ClusterGroupIndex() { return PBDRigidClusteredParticles->ClusterGroupIndex(ParticleIdx); }
-	void SetClusterGroupIndex(const int32 Idx) { PBDRigidClusteredParticles->ClusterGroupIndex(ParticleIdx) = Idx; }
-
-	const bool& InternalCluster() const { return PBDRigidClusteredParticles->InternalCluster(ParticleIdx); }
-	bool& InternalCluster() { return PBDRigidClusteredParticles->InternalCluster(ParticleIdx); }
-	void SetInternalCluster(const bool Value) { PBDRigidClusteredParticles->InternalCluster(ParticleIdx) = Value; }
-
-	const TUniquePtr<FImplicitObjectUnionClustered>& ChildrenSpatial() const { return PBDRigidClusteredParticles->ChildrenSpatial(ParticleIdx); }
-	TUniquePtr<FImplicitObjectUnionClustered>& ChildrenSpatial() { return PBDRigidClusteredParticles->ChildrenSpatial(ParticleIdx); }
-	void SetChildrenSpatial(TUniquePtr<FImplicitObjectUnion>& Obj) { PBDRigidClusteredParticles->ChildrenSpatial(ParticleIdx) = Obj; }
-
-	const FMultiChildProxyId& MultiChildProxyId() const { return PBDRigidClusteredParticles->MultiChildProxyId(ParticleIdx); }
-	FMultiChildProxyId& MultiChildProxyId() { return PBDRigidClusteredParticles->MultiChildProxyId(ParticleIdx); }
-	void SetMultiChildProxyId(const FMultiChildProxyId& Id) { PBDRigidClusteredParticles->MultiChildProxyId(ParticleIdx) = Id; }
-
-	const TUniquePtr<TMultiChildProxyData<T, d>>& MultiChildProxyData() const { return PBDRigidClusteredParticles->MultiChildProxyData(ParticleIdx); }
-	TUniquePtr<TMultiChildProxyData<T, d>>& MultiChildProxyData() { return PBDRigidClusteredParticles->MultiChildProxyData(ParticleIdx); }
-	void SetMultiChildProxyData(TUniquePtr<TMultiChildProxyData<T, d>>&& ProxyData) { PBDRigidClusteredParticles->MultiChildProxyData(ParticleIdx) = MoveTemp(ProxyData); }
-
-	const T& CollisionImpulse() const { return PBDRigidClusteredParticles->CollisionImpulses(ParticleIdx); }
-	T& CollisionImpulse() { return PBDRigidClusteredParticles->CollisionImpulses(ParticleIdx); }
-	void SetCollisionImpulse(const T Value) { PBDRigidClusteredParticles->CollisionImpulses(ParticleIdx) = Value; }
-	// Deprecated
-	const T& CollisionImpulses() const { return PBDRigidClusteredParticles->CollisionImpulses(ParticleIdx); }
-	// Deprecated
-	T& CollisionImpulses() { return PBDRigidClusteredParticles->CollisionImpulses(ParticleIdx); }
-	// Deprecated
-	void SetCollisionImpulses(const T Value) { PBDRigidClusteredParticles->CollisionImpulses(ParticleIdx) = Value; }
-
-	const T& Strain() const { return PBDRigidClusteredParticles->Strains(ParticleIdx); }
-	T& Strain() { return PBDRigidClusteredParticles->Strains(ParticleIdx); }
-	void SetStrain(const T Value) { PBDRigidClusteredParticles->Strains(ParticleIdx) = Value; }
-	// Deprecated
-	const T& Strains() const { return PBDRigidClusteredParticles->Strains(ParticleIdx); }
-	// Deprecated
-	T& Strains() { return PBDRigidClusteredParticles->Strains(ParticleIdx); }
-	// Deprecated
-	void SetStrains(const T Value) { PBDRigidClusteredParticles->Strains(ParticleIdx) = Value; }
-
-	const TArray<TConnectivityEdge<T>>& ConnectivityEdges() const { return PBDRigidClusteredParticles->ConnectivityEdges(ParticleIdx); }
-	TArray<TConnectivityEdge<T>>& ConnectivityEdges() { return PBDRigidClusteredParticles->ConnectivityEdges(ParticleIdx); }
-	void SetConnectivityEdges(const TArray<TConnectivityEdge<T>>& Edges) { PBDRigidClusteredParticles->ConnectivityEdges(ParticleIdx) = Edges; }
-	void SetConnectivityEdges(TArray<TConnectivityEdge<T>>&& Edges) { PBDRigidClusteredParticles->ConnectivityEdges(ParticleIdx) = MoveTemp(Edges); }
+	const auto& ClusterIds() const { return PBDRigidClusteredParticles->ClusterIds(ParticleIdx); }
+	auto& ClusterIds() { return PBDRigidClusteredParticles->ClusterIds(ParticleIdx); }
+
+	const auto& ChildToParent() const { return PBDRigidClusteredParticles->ChildToParent(ParticleIdx); }
+	auto& ChildToParent() { return PBDRigidClusteredParticles->ChildToParent(ParticleIdx); }
+
+	const auto& ClusterGroupIndex() const { return PBDRigidClusteredParticles->ClusterGroupIndex(ParticleIdx); }
+	auto& ClusterGroupIndex() { return PBDRigidClusteredParticles->ClusterGroupIndex(ParticleIdx); }
+
+	const auto& InternalCluster() const { return PBDRigidClusteredParticles->InternalCluster(ParticleIdx); }
+	auto& InternalCluster() { return PBDRigidClusteredParticles->InternalCluster(ParticleIdx); }
+
+	const auto& ChildrenSpatial() const { return PBDRigidClusteredParticles->ChildrenSpatial(ParticleIdx); }
+	auto& ChildrenSpatial() { return PBDRigidClusteredParticles->ChildrenSpatial(ParticleIdx); }
+
+	//const auto& MultiChildProxyId() const { return PBDRigidClusteredParticles->MultiChildProxyId(ParticleIdx); }
+	//auto& MultiChildProxyId() { return PBDRigidClusteredParticles->MultiChildProxyId(ParticleIdx); }
+
+	//const auto& MultiChildProxyData() const { return PBDRigidClusteredParticles->MultiChildProxyData(ParticleIdx); }
+	//auto& MultiChildProxyData() { return PBDRigidClusteredParticles->MultiChildProxyData(ParticleIdx); }
+
+	const auto& CollisionImpulses() const { return PBDRigidClusteredParticles->CollisionImpulses(ParticleIdx); }
+	auto& CollisionImpulses() { return PBDRigidClusteredParticles->CollisionImpulses(ParticleIdx); }
+
+	const auto& Strains() const { return PBDRigidClusteredParticles->Strains(ParticleIdx); }
+	auto& Strains() { return PBDRigidClusteredParticles->Strains(ParticleIdx); }
 
 	const TPBDRigidClusteredParticleHandleImp<T, d, true>* Handle() const { return PBDRigidClusteredParticles->Handle(ParticleIdx); }
 	TPBDRigidClusteredParticleHandleImp<T, d, true>* Handle() { return PBDRigidClusteredParticles->Handle(ParticleIdx); }
@@ -807,11 +776,11 @@ public:
 };
 
 template <typename T, int d, bool bPersistent = true>
-class TPBDGeometryCollectionParticleHandleImp : public TPBDRigidClusteredParticleHandleImp<T, d, bPersistent>
+class TPBDGeometryCollectionParticleHandleImp : public TPBDRigidParticleHandleImp<T, d, bPersistent>
 {
 public:
 	using TGeometryParticleHandleImp<T, d, bPersistent>::ParticleIdx;
-	using TGeometryParticleHandleImp<T, d, bPersistent>::PBDRigidClusteredParticles;
+	using TGeometryParticleHandleImp<T, d, bPersistent>::PBDRigidParticles;
 	using TGeometryParticleHandleImp<T, d, bPersistent>::Type;
 	using TTransientHandle = TTransientPBDGeometryCollectionParticleHandle<T, d>;
 	using TSOAType = TPBDGeometryCollectionParticles<T, d>;
@@ -821,21 +790,17 @@ protected:
 
 	//needed for serialization
 	TPBDGeometryCollectionParticleHandleImp()
-		: TPBDRigidClusteredParticleHandleImp<T, d, bPersistent>()
-	{
-		Type = EParticleType::GeometryCollection;
-	}
+		: TPBDRigidParticleHandleImp<T, d, bPersistent>()
+	{}
 
 	TPBDGeometryCollectionParticleHandleImp<T, d, bPersistent>(
 		TSerializablePtr<TPBDGeometryCollectionParticles<T, d>> Particles, 
 		int32 InIdx, 
 		int32 InGlobalIdx, 
 		const TPBDRigidParticleParameters<T, d>& Params = TPBDRigidParticleParameters<T, d>())
-		: TPBDRigidClusteredParticleHandleImp<T, d, bPersistent>(
-			TSerializablePtr<TPBDRigidClusteredParticles<T, d>>(Particles), InIdx, InGlobalIdx, Params)
-	{
-		Type = EParticleType::GeometryCollection;
-	}
+		: TPBDRigidParticleHandleImp<T, d, bPersistent>(
+			TSerializablePtr<TPBDRigidParticles<T, d>>(Particles), InIdx, InGlobalIdx, Params)
+	{}
 public:
 
 	static TUniquePtr<TPBDGeometryCollectionParticleHandleImp<T, d, bPersistent>> CreateParticleHandle(
@@ -855,8 +820,8 @@ public:
 		return Serializable;
 	}
 
-	const TPBDGeometryCollectionParticleHandleImp<T, d, true>* Handle() const { return PBDRigidClusteredParticles->Handle(ParticleIdx); }
-	TPBDGeometryCollectionParticleHandleImp<T, d, true>* Handle() { return PBDRigidClusteredParticles->Handle(ParticleIdx); }
+	const TPBDGeometryCollectionParticleHandleImp<T, d, true>* Handle() const { return PBDRigidParticles->Handle(ParticleIdx); }
+	TPBDGeometryCollectionParticleHandleImp<T, d, true>* Handle() { return PBDRigidParticles->Handle(ParticleIdx); }
 
 	static constexpr EParticleType StaticType() { return EParticleType::GeometryCollection; }
 };
