Implement PhiWithNormal() for FTriangleMeshImplicitObject.

diff --git a/Engine/Source/Runtime/Experimental/Chaos/Private/Chaos/TriangleMeshImplicitObject.cpp b/Engine/Source/Runtime/Experimental/Chaos/Private/Chaos/TriangleMeshImplicitObject.cpp
--- a/Engine/Source/Runtime/Experimental/Chaos/Private/Chaos/TriangleMeshImplicitObject.cpp
+++ b/Engine/Source/Runtime/Experimental/Chaos/Private/Chaos/TriangleMeshImplicitObject.cpp
@@ -189,8 +189,12 @@ struct FTriangleMeshRaycastVisitor
 
 FReal FTriangleMeshImplicitObject::PhiWithNormal(const FVec3& x, FVec3& Normal) const
 {
-	ensure(false);	//not supported yet - might support it in the future or we may change the interface
-	return 0;
+	TSphere<FReal, 3> TestSphere(x, 0.0f);
+	FRigidTransform3 TestXf(TVector<float, 3>(0.0), TRotation<float, 3>::FromIdentity());
+	FVec3 TestLocation = x;
+	FReal Depth;
+	GJKContactPointImp(TestSphere, TestXf, 0.0f, TestLocation, Normal, Depth);
+	return Depth;
 }
 
 template <typename IdxType>
