Enable support for Scene Queries mapping from triangle mesh internal face index to external. Fixes procedural foliage.

diff --git a/Engine/Source/Runtime/Core/Public/UObject/ExternalPhysicsCustomObjectVersion.h b/Engine/Source/Runtime/Core/Public/UObject/ExternalPhysicsCustomObjectVersion.h
--- a/Engine/Source/Runtime/Core/Public/UObject/ExternalPhysicsCustomObjectVersion.h
+++ b/Engine/Source/Runtime/Core/Public/UObject/ExternalPhysicsCustomObjectVersion.h
@@ -107,6 +107,9 @@ struct CORE_API FExternalPhysicsCustomObjectVersion
 
 		// Heightfield uses uint16 heights directly
 		HeightfieldUsesHeightsDirectly,
+
+		// TriangleMesh has map from internal face indices to external.
+		TriangleMeshHasFaceIndexMap,
 		
 		// -----<new versions can be added above this line>-------------------------------------------------
 		VersionPlusOne,
diff --git a/Engine/Source/Runtime/Engine/Private/PhysicsEngine/Experimental/ChaosDerivedData.cpp b/Engine/Source/Runtime/Engine/Private/PhysicsEngine/Experimental/ChaosDerivedData.cpp
--- a/Engine/Source/Runtime/Engine/Private/PhysicsEngine/Experimental/ChaosDerivedData.cpp
+++ b/Engine/Source/Runtime/Engine/Private/PhysicsEngine/Experimental/ChaosDerivedData.cpp
@@ -24,7 +24,7 @@ const TCHAR* FChaosDerivedDataCooker::GetPluginName() const
 
 const TCHAR* FChaosDerivedDataCooker::GetVersionString() const
 {
-	return TEXT("6249C3BD908145F3B3AC3048F3BF1C01");
+	return TEXT("2C093F83EC7F49A382170F519868BBC6");
 }
 
 FString FChaosDerivedDataCooker::GetPluginSpecificCacheKeySuffix() const
@@ -152,7 +152,7 @@ void FChaosDerivedDataCooker::BuildTriangleMeshes(TArray<TUniquePtr<Chaos::FTria
 
 				if(bHasMaterials)
 				{
-					if (EnableMeshClean)				
+					if (EnableMeshClean)
 					{
 						if(!ensure(OutFaceRemap.IsValidIndex(TriangleIndex)))
 						{
@@ -191,7 +191,7 @@ void FChaosDerivedDataCooker::BuildTriangleMeshes(TArray<TUniquePtr<Chaos::FTria
 			}
 		}
 
-		OutTriangleMeshes.Emplace(new Chaos::FTriangleMeshImplicitObject(MoveTemp(TriMeshParticles), MoveTemp(Triangles), MoveTemp(MaterialIndices)));
+		OutTriangleMeshes.Emplace(new Chaos::FTriangleMeshImplicitObject(MoveTemp(TriMeshParticles), MoveTemp(Triangles), MoveTemp(MaterialIndices), MoveTemp(MakeUnique<TArray<int32>>(OutFaceRemap))));
 	};
 
 	if(FinalVerts.Num() < TNumericLimits<uint16>::Max())
diff --git a/Engine/Source/Runtime/Engine/Private/PhysicsEngine/Experimental/PhysInterface_Chaos.cpp b/Engine/Source/Runtime/Engine/Private/PhysicsEngine/Experimental/PhysInterface_Chaos.cpp
--- a/Engine/Source/Runtime/Engine/Private/PhysicsEngine/Experimental/PhysInterface_Chaos.cpp
+++ b/Engine/Source/Runtime/Engine/Private/PhysicsEngine/Experimental/PhysInterface_Chaos.cpp
@@ -14,6 +14,7 @@
 #include "Chaos/Capsule.h"
 #include "Chaos/ImplicitObjectTransformed.h"
 #include "Chaos/ImplicitObjectUnion.h"
+#include "Chaos/TriangleMeshImplicitObject.h"
 #include "Chaos/Levelset.h"
 #include "Chaos/PBDRigidParticles.h"
 #include "Chaos/Sphere.h"
@@ -1949,6 +1950,30 @@ bool FPhysInterface_Chaos::GetSquaredDistanceToBody(const FBodyInstance* InInsta
 	return bFoundValidBody;
 }
 
+uint32 GetTriangleMeshExternalFaceIndex(const FPhysicsShape& Shape, uint32 InternalFaceIndex)
+{
+	using namespace Chaos;
+	uint8 Type = Shape.Geometry->GetType();
+	if (ensure(Type | ImplicitObjectType::TriangleMesh))
+	{
+		const FTriangleMeshImplicitObject* TriangleMesh = nullptr;
+
+		if (Type | ImplicitObjectType::IsScaled)
+		{
+			const TImplicitObjectScaled<FTriangleMeshImplicitObject>* ScaledTriangleMesh = static_cast<const TImplicitObjectScaled<FTriangleMeshImplicitObject>*>(Shape.Geometry.Get());
+			TriangleMesh = ScaledTriangleMesh->GetUnscaledObject();
+		}
+		else
+		{
+			TriangleMesh = static_cast<const FTriangleMeshImplicitObject*>(Shape.Geometry.Get());
+		}
+
+		return TriangleMesh->GetExternalFaceIndexFromInternal(InternalFaceIndex);
+	}
+
+	return -1;
+}
+
 template<typename AllocatorType>
 int32 GetAllShapesInternal_AssumedLocked(const FPhysicsActorHandle& InActorHandle, TArray<FPhysicsShapeReference_Chaos, AllocatorType>& OutShapes)
 {
diff --git a/Engine/Source/Runtime/Engine/Public/Physics/Experimental/PhysInterface_Chaos.h b/Engine/Source/Runtime/Engine/Public/Physics/Experimental/PhysInterface_Chaos.h
--- a/Engine/Source/Runtime/Engine/Public/Physics/Experimental/PhysInterface_Chaos.h
+++ b/Engine/Source/Runtime/Engine/Public/Physics/Experimental/PhysInterface_Chaos.h
@@ -447,10 +447,7 @@ FORCEINLINE void ComputeZeroDistanceImpactNormalAndPenetration(const UWorld* Wor
 
 Chaos::FChaosPhysicsMaterial* GetMaterialFromInternalFaceIndex(const FPhysicsShape& Shape, const FPhysicsActor& Actor, uint32 InternalFaceIndex);
 
-inline uint32 GetTriangleMeshExternalFaceIndex(const FPhysicsShape& Shape, uint32 InternalFaceIndex)
-{
-	return GetInvalidPhysicsFaceIndex();
-}
+uint32 GetTriangleMeshExternalFaceIndex(const FPhysicsShape& Shape, uint32 InternalFaceIndex);
 
 inline void GetShapes(const FPhysActorDummy& RigidActor, FPhysTypeDummy** ShapesBuffer, uint32 NumShapes)
 {
diff --git a/Engine/Source/Runtime/Experimental/Chaos/Private/Chaos/TriangleMeshImplicitObject.cpp b/Engine/Source/Runtime/Experimental/Chaos/Private/Chaos/TriangleMeshImplicitObject.cpp
--- a/Engine/Source/Runtime/Experimental/Chaos/Private/Chaos/TriangleMeshImplicitObject.cpp
+++ b/Engine/Source/Runtime/Experimental/Chaos/Private/Chaos/TriangleMeshImplicitObject.cpp
@@ -355,7 +355,18 @@ bool FTriangleMeshImplicitObject::GJKContactPoint(const TImplicitObjectScaled< F
 	return GJKContactPointImp(QueryGeom, QueryTM, Thickness, Location, Normal, Penetration);
 }
 
+int32 FTriangleMeshImplicitObject::GetExternalFaceIndexFromInternal(int32 InternalFaceIndex) const
+{
+	if (InternalFaceIndex > -1 && ExternalFaceIndexMap.Get())
+	{
+		if (ensure(InternalFaceIndex >= 0 && InternalFaceIndex < ExternalFaceIndexMap->Num()))
+		{
+			return (*ExternalFaceIndexMap)[InternalFaceIndex];
+		}
+	}
 
+	return -1;
+}
 
 template <typename IdxType>
 bool FTriangleMeshImplicitObject::OverlapImp(const TArray<TVec3<IdxType>>& Elements, const FVec3& Point, const FReal Thickness) const
diff --git a/Engine/Source/Runtime/Experimental/Chaos/Public/Chaos/TriangleMeshImplicitObject.h b/Engine/Source/Runtime/Experimental/Chaos/Public/Chaos/TriangleMeshImplicitObject.h
--- a/Engine/Source/Runtime/Experimental/Chaos/Public/Chaos/TriangleMeshImplicitObject.h
+++ b/Engine/Source/Runtime/Experimental/Chaos/Public/Chaos/TriangleMeshImplicitObject.h
@@ -101,12 +101,13 @@ namespace Chaos
 		using FImplicitObject::GetTypeName;
 
 		template <typename IdxType>
-		FTriangleMeshImplicitObject(TParticles<FReal, 3>&& Particles, TArray<TVector<IdxType, 3>>&& Elements, TArray<uint16>&& InMaterialIndices)
+		FTriangleMeshImplicitObject(TParticles<FReal, 3>&& Particles, TArray<TVector<IdxType, 3>>&& Elements, TArray<uint16>&& InMaterialIndices, TUniquePtr<TArray<int32>>&& InExternalFaceIndexMap = nullptr)
 		: FImplicitObject(EImplicitObject::HasBoundingBox | EImplicitObject::DisableCollisions, ImplicitObjectType::TriangleMesh)
 		, MParticles(MoveTemp(Particles))
 		, MElements(MoveTemp(Elements))
 		, MLocalBoundingBox(MParticles.X(0), MParticles.X(0))
 		, MaterialIndices(MoveTemp(InMaterialIndices))
+		, ExternalFaceIndexMap(MoveTemp(InExternalFaceIndexMap))
 		{
 			for (uint32 Idx = 1; Idx < MParticles.Size(); ++Idx)
 			{
@@ -152,6 +153,9 @@ namespace Chaos
 		bool GJKContactPoint(const TImplicitObjectScaled < TCapsule<FReal> >& QueryGeom, const FRigidTransform3& QueryTM, const FReal Thickness, FVec3& Location, FVec3& Normal, FReal& Penetration) const;
 		bool GJKContactPoint(const TImplicitObjectScaled < FConvex >& QueryGeom, const FRigidTransform3& QueryTM, const FReal Thickness, FVec3& Location, FVec3& Normal, FReal& Penetration) const;
 
+		// Returns -1 if InternalFaceIndex is not in map, or map is invalid.
+		int32 GetExternalFaceIndexFromInternal(int32 InternalFaceIndex) const;
+
 
 		virtual int32 FindMostOpposingFace(const FVec3& Position, const FVec3& UnitDir, int32 HintFaceIndex, FReal SearchDistance) const override;
 		virtual FVec3 FindGeometryOpposingNormal(const FVec3& DenormDir, int32 FaceIndex, const FVec3& OriginalNormal) const override;
@@ -202,6 +206,29 @@ namespace Chaos
 			{
 				Ar << MaterialIndices;
 			}
+
+			if (Ar.CustomVer(FExternalPhysicsCustomObjectVersion::GUID) >= FExternalPhysicsCustomObjectVersion::TriangleMeshHasFaceIndexMap)
+			{
+				// TODO: This data is only needed in editor unless project configuration requests this for gameplay. We should not serialize this when cooking
+				// unless it is required for gameplay, as we are wasting disk space.
+				if (Ar.IsLoading())
+				{
+					ExternalFaceIndexMap = MakeUnique<TArray<int32>>(TArray<int32>());
+					Ar << *ExternalFaceIndexMap;
+				}
+				else
+				{
+					if (ExternalFaceIndexMap == nullptr)
+					{
+						TArray<int32> EmptyArray;
+						Ar << EmptyArray;
+					}
+					else
+					{
+						Ar << *ExternalFaceIndexMap;
+					}
+				}
+			}
 		}
 
 		virtual void Serialize(FChaosArchive& Ar) override;
@@ -227,6 +254,7 @@ namespace Chaos
 		FTrimeshIndexBuffer MElements;
 		TAABB<FReal, 3> MLocalBoundingBox;
 		TArray<uint16> MaterialIndices;
+		TUniquePtr<TArray<int32>> ExternalFaceIndexMap;
 
 		//using BVHType = TBoundingVolume<int32, T, 3>;
 		using BVHType = TAABBTree<int32, TAABBTreeLeafArray<int32, FReal>, FReal, /*bMutable=*/false>;
