Repair the field component so that it accesses the physics scene(s) without tripping over null pointers in dead code.

diff --git a/Engine/Source/Runtime/Experimental/FieldSystem/Source/FieldSystemEngine/Private/Field/FieldSystemComponent.cpp b/Engine/Source/Runtime/Experimental/FieldSystem/Source/FieldSystemEngine/Private/Field/FieldSystemComponent.cpp
--- a/Engine/Source/Runtime/Experimental/FieldSystem/Source/FieldSystemEngine/Private/Field/FieldSystemComponent.cpp
+++ b/Engine/Source/Runtime/Experimental/FieldSystem/Source/FieldSystemEngine/Private/Field/FieldSystemComponent.cpp
@@ -27,7 +27,6 @@ UFieldSystemComponent::UFieldSystemComponent(const FObjectInitializer& ObjectIni
 	SetGenerateOverlapEvents(false);
 }
 
-
 FPrimitiveSceneProxy* UFieldSystemComponent::CreateSceneProxy()
 {
 	UE_LOG(FSC_Log, Log, TEXT("FieldSystemComponent[%p]::CreateSceneProxy()"), this);
@@ -35,6 +34,35 @@ FPrimitiveSceneProxy* UFieldSystemComponent::CreateSceneProxy()
 	return new FFieldSystemSceneProxy(this);
 }
 
+TSet<FPhysScene_Chaos*> UFieldSystemComponent::GetPhysicsScenes() const
+{
+	TSet<FPhysScene_Chaos*> Scenes;
+	if (SupportedSolvers.Num())
+	{
+		for (const TSoftObjectPtr<AChaosSolverActor>& Actor : SupportedSolvers)
+		{
+			if (!Actor.IsValid())
+				continue;
+			Scenes.Add(Actor->GetPhysicsScene().Get());
+		}
+	}
+	else
+	{
+#if INCLUDE_CHAOS
+		if (ensure(GetOwner()) && ensure(GetOwner()->GetWorld()))
+		{
+			FPhysScene_ChaosInterface* WorldPhysScene = GetOwner()->GetWorld()->GetPhysicsScene();
+			Scenes.Add(&WorldPhysScene->GetScene());
+		}
+		else
+		{
+			check(GWorld);
+			Scenes.Add(&GWorld->GetPhysicsScene()->GetScene());
+		}
+#endif
+	}
+	return Scenes;
+}
 
 void UFieldSystemComponent::OnCreatePhysicsState()
 {
@@ -49,10 +77,13 @@ void UFieldSystemComponent::OnCreatePhysicsState()
 
 		PhysicsProxy = new FFieldSystemPhysicsProxy(this);
 #if INCLUDE_CHAOS
-		TSharedPtr<FPhysScene_Chaos> Scene = GetOwner()->GetWorld()->PhysicsScene_Chaos;
-		Scene->AddObject(this, PhysicsProxy);
+		TSet<FPhysScene_Chaos*> Scenes = GetPhysicsScenes();
+		for (auto* Scene : Scenes)
+		{
+			// Does each scene need its own proxy?
+			Scene->AddObject(this, PhysicsProxy);
+		}
 #endif
-
 		bHasPhysicsState = true;
 
 		if(FieldSystem)
@@ -75,15 +106,19 @@ void UFieldSystemComponent::OnDestroyPhysicsState()
 	}
 
 #if INCLUDE_CHAOS
-	TSharedPtr<FPhysScene_Chaos> Scene = GetOwner()->GetWorld()->PhysicsScene_Chaos;
-	Scene->RemoveObject(PhysicsProxy);
+	//TSharedPtr<FPhysScene_Chaos> Scene = GetOwner()->GetWorld()->PhysicsScene_Chaos;
+	TSet<FPhysScene_Chaos*> Scenes = GetPhysicsScenes();
+	for (auto* Scene : Scenes)
+	{
+		Scene->RemoveObject(PhysicsProxy);
+	}
 #endif
 
 	ChaosModule = nullptr;
+	// Discard the pointer (cleanup happens through the scene or dedicated thread)
 	PhysicsProxy = nullptr;
 
 	bHasPhysicsState = false;
-	
 }
 
 bool UFieldSystemComponent::ShouldCreatePhysicsState() const
diff --git a/Engine/Source/Runtime/Experimental/FieldSystem/Source/FieldSystemEngine/Public/Field/FieldSystemComponent.h b/Engine/Source/Runtime/Experimental/FieldSystem/Source/FieldSystemEngine/Public/Field/FieldSystemComponent.h
--- a/Engine/Source/Runtime/Experimental/FieldSystem/Source/FieldSystemEngine/Public/Field/FieldSystemComponent.h
+++ b/Engine/Source/Runtime/Experimental/FieldSystem/Source/FieldSystemEngine/Public/Field/FieldSystemComponent.h
@@ -25,18 +25,15 @@ class FIELDSYSTEMENGINE_API UFieldSystemComponent : public UPrimitiveComponent
 	friend class FFieldSystemEditorCommands;
 
 public:
-
-
-
 	//~ Begin USceneComponent Interface.
 	virtual bool HasAnySockets() const override { return false; }
 	//~ Begin USceneComponent Interface.
 
-
 	//~ Begin UPrimitiveComponent Interface.
 	virtual FPrimitiveSceneProxy* CreateSceneProxy() override;
 	//~ End UPrimitiveComponent Interface.
 
+	TSet<FPhysScene_Chaos*> GetPhysicsScenes() const;
 
 	/** FieldSystem @todo(remove the field system, we dont need the asset*/
 	void SetFieldSystem(UFieldSystem * FieldSystemIn) { FieldSystem = FieldSystemIn; }
@@ -45,7 +42,6 @@ public:
 	UPROPERTY(EditAnywhere, BlueprintReadOnly, Category = "Field")
 	UFieldSystem* FieldSystem;
 
-
 	//
 	// Blueprint based field interface
 	//
@@ -63,7 +59,6 @@ public:
 	UFUNCTION(BlueprintCallable, Category = "Field")
 	void ApplyLinearForce(bool Enabled, FVector Direction, float Magnitude);
 
-
 	/**
 	*  ApplyStayDynamicField
 	*    This function will dispatch a command to the physics thread to apply
@@ -155,8 +150,6 @@ public:
 	UFUNCTION(BlueprintCallable, Category = "Field")
 	void ApplyPhysicsField(bool Enabled, EFieldPhysicsType Target, UFieldSystemMetaData* MetaData, UFieldNodeBase* Field);
 
-
-
 	//
 	// Blueprint Construction based field interface
 	//
@@ -191,5 +184,4 @@ protected:
 	FChaosSolversModule* ChaosModule;
 
 	bool bHasPhysicsState;
-
 };
