Use RemoveSingleSwap instead of Remove when uregsitering physics proxies.

diff --git a/Engine/Source/Runtime/Experimental/ChaosSolvers/Private/PBDRigidsSolver.cpp b/Engine/Source/Runtime/Experimental/ChaosSolvers/Private/PBDRigidsSolver.cpp
--- a/Engine/Source/Runtime/Experimental/ChaosSolvers/Private/PBDRigidsSolver.cpp
+++ b/Engine/Source/Runtime/Experimental/ChaosSolvers/Private/PBDRigidsSolver.cpp
@@ -282,11 +282,11 @@ namespace Chaos
 		// Remove the proxy from the GT proxy map
 		if (InParticleType == EParticleType::Rigid)
 		{
-			RigidParticlePhysicsProxies.Remove((FRigidParticlePhysicsProxy*)InProxy);
+			RigidParticlePhysicsProxies.RemoveSingleSwap((FRigidParticlePhysicsProxy*)InProxy);
 		}
 		else if (InParticleType == EParticleType::Kinematic)
 		{
-			KinematicGeometryParticlePhysicsProxies.Remove((FKinematicGeometryParticlePhysicsProxy*)InProxy);
+			KinematicGeometryParticlePhysicsProxies.RemoveSingleSwap((FKinematicGeometryParticlePhysicsProxy*)InProxy);
 		}
 		else if (InParticleType == EParticleType::GeometryCollection)
 		{
@@ -294,7 +294,7 @@ namespace Chaos
 		}
 		else
 		{
-			GeometryParticlePhysicsProxies.Remove((FGeometryParticlePhysicsProxy*)InProxy);
+			GeometryParticlePhysicsProxies.RemoveSingleSwap((FGeometryParticlePhysicsProxy*)InProxy);
 		}
 
 		// Enqueue a command to remove the particle and delete the proxy
