Chaos - RBAN
- reduce iterations and pair iterations to 1 for joints and collisions
- re-enable joint projection
- collision stats

diff --git a/Engine/Source/Runtime/Engine/Private/PhysicsEngine/ImmediatePhysics/ImmediatePhysicsChaos/ImmediatePhysicsSimulation_Chaos.cpp b/Engine/Source/Runtime/Engine/Private/PhysicsEngine/ImmediatePhysics/ImmediatePhysicsChaos/ImmediatePhysicsSimulation_Chaos.cpp
--- a/Engine/Source/Runtime/Engine/Private/PhysicsEngine/ImmediatePhysics/ImmediatePhysicsChaos/ImmediatePhysicsSimulation_Chaos.cpp
+++ b/Engine/Source/Runtime/Engine/Private/PhysicsEngine/ImmediatePhysics/ImmediatePhysicsChaos/ImmediatePhysicsSimulation_Chaos.cpp
@@ -26,8 +26,8 @@ float ChaosImmediate_Evolution_StepTime = 0.0f;
 int32 ChaosImmediate_Evolution_NumSteps = 0;
 float ChaosImmediate_Evolution_InitialStepTime = 0.033f;
 int32 ChaosImmediate_Evolution_DeltaTimeCount = 100;
-int32 ChaosImmediate_Evolution_Iterations = 4;
-int32 ChaosImmediate_Evolution_PushOutIterations = 2;
+int32 ChaosImmediate_Evolution_Iterations = 1;
+int32 ChaosImmediate_Evolution_PushOutIterations = 1;
 FAutoConsoleVariableRef CVarChaosImmPhysStepTime(TEXT("p.Chaos.ImmPhys.StepTime"), ChaosImmediate_Evolution_StepTime, TEXT("Override step time (if not zero)"));
 FAutoConsoleVariableRef CVarChaosImmPhysNumSteps(TEXT("p.Chaos.ImmPhys.NumSteps"), ChaosImmediate_Evolution_NumSteps, TEXT("Override num steps (if not zero)"));
 FAutoConsoleVariableRef CVarChaosImmPhysInitialStepTime(TEXT("p.Chaos.ImmPhys.InitialStepTime"), ChaosImmediate_Evolution_InitialStepTime, TEXT("Initial step time (then calculated from rolling average)"));
@@ -36,8 +36,8 @@ FAutoConsoleVariableRef CVarChaosImmPhysIterations(TEXT("p.Chaos.ImmPhys.Iterati
 FAutoConsoleVariableRef CVarChaosImmPhysPushOutIterations(TEXT("p.Chaos.ImmPhys.PushOutIterations"), ChaosImmediate_Evolution_PushOutIterations, TEXT("Set the ApplyPushOut() (position correction) iteration count"));
 
 int32 ChaosImmediate_Collision_Enabled = 1;
-int32 ChaosImmediate_Collision_PairIterations = 2;
-int32 ChaosImmediate_Collision_PushOutPairIterations = 2;
+int32 ChaosImmediate_Collision_PairIterations = 1;
+int32 ChaosImmediate_Collision_PushOutPairIterations = 1;
 int32 ChaosImmediate_Collision_Priority = 1;
 float ChaosImmediate_Collision_Thickness = 0;
 FAutoConsoleVariableRef CVarChaosImmPhysDisableCollisions(TEXT("p.Chaos.ImmPhys.Collision.Enabled"), ChaosImmediate_Collision_Enabled, TEXT("Enable/Disable collisions in Immediate Physics."));
@@ -46,8 +46,8 @@ FAutoConsoleVariableRef CVarChaosImmPhysCollisionPushOutPairIterations(TEXT("p.C
 FAutoConsoleVariableRef CVarChaosImmPhysCollisionPriority(TEXT("p.Chaos.ImmPhys.Collision.Priority"), ChaosImmediate_Collision_Priority, TEXT("Set the Collision constraint sort order (Joints have priority 0)"));
 FAutoConsoleVariableRef CVarChaosImmPhysThickness(TEXT("p.Chaos.ImmPhys.Collision.Thickness"), ChaosImmediate_Collision_Thickness, TEXT("ChaosImmediateThickness"));
 
-int32 ChaosImmediate_Joint_PairIterations = 2;
-int32 ChaosImmediate_Joint_PushOutPairIterations = 0;// 1; Until I fix projection for soft joints
+int32 ChaosImmediate_Joint_PairIterations = 1;
+int32 ChaosImmediate_Joint_PushOutPairIterations = 1;
 float ChaosImmediate_Joint_SwingTwistAngleTolerance = 1.0e-6f;
 int32 ChaosImmediate_Joint_EnableTwistLimits = 1;
 int32 ChaosImmediate_Joint_EnableSwingLimits = 1;
diff --git a/Engine/Source/Runtime/Experimental/Chaos/Private/Chaos/Collision/CollisionDetector.cpp b/Engine/Source/Runtime/Experimental/Chaos/Private/Chaos/Collision/CollisionDetector.cpp
--- a/Engine/Source/Runtime/Experimental/Chaos/Private/Chaos/Collision/CollisionDetector.cpp
+++ b/Engine/Source/Runtime/Experimental/Chaos/Private/Chaos/Collision/CollisionDetector.cpp
@@ -1,11 +1,14 @@
 // Copyright Epic Games, Inc. All Rights Reserved.
 #include "Chaos/Collision/CollisionDetector.h"
+#include "Chaos/Collision/ParticlePairBroadPhase.h"
 #include "ChaosLog.h"
 #include "ChaosStats.h"
 
 namespace Chaos
 {
-	DEFINE_STAT(STAT_DetectCollisions);
+	DEFINE_STAT(STAT_Collisions_Detect);
+	DEFINE_STAT(STAT_Collisions_BroadPhase);
+	DEFINE_STAT(STAT_Collisions_NarrowPhase);
 
 	template class TCollisionDetector<FParticlePairBroadPhase, FNarrowPhase, FSyncCollisionReceiver, TPBDCollisionConstraints<FReal, 3>>;
 	template class TCollisionDetector<FSpatialAccelerationBroadPhase, FNarrowPhase, FAsyncCollisionReceiver, TPBDCollisionConstraints<FReal, 3>>;
diff --git a/Engine/Source/Runtime/Experimental/Chaos/Private/Chaos/Joint/PBDJointSolverGaussSeidel.cpp b/Engine/Source/Runtime/Experimental/Chaos/Private/Chaos/Joint/PBDJointSolverGaussSeidel.cpp
--- a/Engine/Source/Runtime/Experimental/Chaos/Private/Chaos/Joint/PBDJointSolverGaussSeidel.cpp
+++ b/Engine/Source/Runtime/Experimental/Chaos/Private/Chaos/Joint/PBDJointSolverGaussSeidel.cpp
@@ -1596,57 +1596,58 @@ namespace Chaos
 		const FPBDJointSolverSettings& SolverSettings,
 		const FPBDJointSettings& JointSettings)
 	{
-		// Apply a position correction with the parent body set to infinite mass, then correct the velocity.
-		// Note: Velocity correction does not set parent to infinite mass (if we do, then the velocities
-		// get amplified down the chain).
-		FVec3 CX = FPBDJointUtilities::GetLimitedPositionError(JointSettings, Rs[0], Xs[1] - Xs[0]);
-		const FReal CXLen = CX.Size();
-		if (CXLen > KINDA_SMALL_NUMBER)
-		{
-			FReal LinearProjection = FPBDJointUtilities::GetLinearProjection(SolverSettings, JointSettings);
-			const FReal ParentMassScale = FMath::Max(0.0f, 1.0f - LinearProjection);
-			if (bChaos_Joint_ISPC_Enabled)
+		FReal LinearProjection = FPBDJointUtilities::GetLinearProjection(SolverSettings, JointSettings);
+		if ((LinearProjection > 0) && !JointSettings.bSoftLinearLimitsEnabled)
+		{
+			// Apply a position correction with the parent body set to infinite mass, then correct the velocity.
+			FVec3 CX = FPBDJointUtilities::GetLimitedPositionError(JointSettings, Rs[0], Xs[1] - Xs[0]);
+			const FReal CXLen = CX.Size();
+			if (CXLen > KINDA_SMALL_NUMBER)
 			{
+				const FReal ParentMassScale = FMath::Max(0.0f, 1.0f - LinearProjection);
+				if (bChaos_Joint_ISPC_Enabled)
+				{
 #if INTEL_ISPC
-				const FReal Stiffness = FPBDJointUtilities::GetLinearStiffness(SolverSettings, JointSettings);
-				ispc::ApplyPositionProjection((ispc::FJointSolverGaussSeidel*)this, (ispc::FVector&)CX, CXLen, ParentMassScale, Stiffness);
+					const FReal Stiffness = FPBDJointUtilities::GetLinearStiffness(SolverSettings, JointSettings);
+					ispc::ApplyPositionProjection((ispc::FJointSolverGaussSeidel*)this, (ispc::FVector&)CX, CXLen, ParentMassScale, Stiffness);
 #endif
-			}
-			else
-			{
-				const FVec3 CXDir = CX / CXLen;
-				const FVec3 V0 = Vs[0] + FVec3::CrossProduct(Ws[0], Xs[0] - Ps[0]);
-				const FVec3 V1 = Vs[1] + FVec3::CrossProduct(Ws[1], Xs[1] - Ps[1]);
-				FVec3 CV = FVec3::DotProduct(V1 - V0, CXDir) * CXDir;
-
-				const FReal IM0 = ParentMassScale * InvMs[0];
-				const FReal IM1 = InvMs[1];
-				const FVec3 IIL0 = ParentMassScale * InvILs[0];
-				const FVec3& IIL1 = InvILs[1];
-				FMatrix33 II0 = Utilities::ComputeWorldSpaceInertia(Qs[0], IIL0);
-				FMatrix33 II1 = Utilities::ComputeWorldSpaceInertia(Qs[1], IIL1);
-				FMatrix33 J0 = (IM0 > 0) ? Utilities::ComputeJointFactorMatrix(Xs[0] - Ps[0], II0, IM0) : FMatrix33(0, 0, 0);
-				FMatrix33 J1 = Utilities::ComputeJointFactorMatrix(Xs[1] - Ps[1], II1, IM1);
-				FMatrix33 IJ = (J0 + J1).Inverse();
-
-				const FVec3 DX = Utilities::Multiply(IJ, CX);
-				const FVec3 DV = Utilities::Multiply(IJ, CV);
-
-				const FVec3 DP0 = IM0 * DX;
-				const FVec3 DP1 = -IM1 * DX;
-				const FVec3 DR0 = Utilities::Multiply(II0, FVec3::CrossProduct(Xs[0] - Ps[0], DX));
-				const FVec3 DR1 = Utilities::Multiply(II1, FVec3::CrossProduct(Xs[1] - Ps[1], -DX));
-
-				const FVec3 DV0 = IM0 * DV;
-				const FVec3 DV1 = -IM1 * DV;
-				const FVec3 DW0 = Utilities::Multiply(II0, FVec3::CrossProduct(Xs[0] - Ps[0], DV));
-				const FVec3 DW1 = Utilities::Multiply(II1, FVec3::CrossProduct(Xs[1] - Ps[1], -DV));
-
-				const FReal Stiffness = FPBDJointUtilities::GetLinearStiffness(SolverSettings, JointSettings);
-				ApplyPositionDelta(Stiffness, DP0, DP1);
-				ApplyRotationDelta(Stiffness, DR0, DR1);
-				ApplyVelocityDelta(Stiffness, DV0, DW0, DV1, DW1);
-				UpdateDerivedState();
+				}
+				else
+				{
+					const FVec3 CXDir = CX / CXLen;
+					const FVec3 V0 = Vs[0] + FVec3::CrossProduct(Ws[0], Xs[0] - Ps[0]);
+					const FVec3 V1 = Vs[1] + FVec3::CrossProduct(Ws[1], Xs[1] - Ps[1]);
+					FVec3 CV = FVec3::DotProduct(V1 - V0, CXDir) * CXDir;
+
+					const FReal IM0 = ParentMassScale * InvMs[0];
+					const FReal IM1 = InvMs[1];
+					const FVec3 IIL0 = ParentMassScale * InvILs[0];
+					const FVec3& IIL1 = InvILs[1];
+					FMatrix33 II0 = Utilities::ComputeWorldSpaceInertia(Qs[0], IIL0);
+					FMatrix33 II1 = Utilities::ComputeWorldSpaceInertia(Qs[1], IIL1);
+					FMatrix33 J0 = (IM0 > 0) ? Utilities::ComputeJointFactorMatrix(Xs[0] - Ps[0], II0, IM0) : FMatrix33(0, 0, 0);
+					FMatrix33 J1 = Utilities::ComputeJointFactorMatrix(Xs[1] - Ps[1], II1, IM1);
+					FMatrix33 IJ = (J0 + J1).Inverse();
+
+					const FVec3 DX = Utilities::Multiply(IJ, CX);
+					const FVec3 DV = Utilities::Multiply(IJ, CV);
+
+					const FVec3 DP0 = IM0 * DX;
+					const FVec3 DP1 = -IM1 * DX;
+					const FVec3 DR0 = Utilities::Multiply(II0, FVec3::CrossProduct(Xs[0] - Ps[0], DX));
+					const FVec3 DR1 = Utilities::Multiply(II1, FVec3::CrossProduct(Xs[1] - Ps[1], -DX));
+
+					const FVec3 DV0 = IM0 * DV;
+					const FVec3 DV1 = -IM1 * DV;
+					const FVec3 DW0 = Utilities::Multiply(II0, FVec3::CrossProduct(Xs[0] - Ps[0], DV));
+					const FVec3 DW1 = Utilities::Multiply(II1, FVec3::CrossProduct(Xs[1] - Ps[1], -DV));
+
+					const FReal Stiffness = FPBDJointUtilities::GetLinearStiffness(SolverSettings, JointSettings);
+					ApplyPositionDelta(Stiffness, DP0, DP1);
+					ApplyRotationDelta(Stiffness, DR0, DR1);
+					ApplyVelocityDelta(Stiffness, DV0, DW0, DV1, DW1);
+					UpdateDerivedState();
+				}
 			}
 		}
 	}
diff --git a/Engine/Source/Runtime/Experimental/Chaos/Private/Chaos/PBDCollisionConstraints.cpp b/Engine/Source/Runtime/Experimental/Chaos/Private/Chaos/PBDCollisionConstraints.cpp
--- a/Engine/Source/Runtime/Experimental/Chaos/Private/Chaos/PBDCollisionConstraints.cpp
+++ b/Engine/Source/Runtime/Experimental/Chaos/Private/Chaos/PBDCollisionConstraints.cpp
@@ -50,12 +50,11 @@ namespace Chaos
 	FAutoConsoleVariableRef CVarEnableCollisions(TEXT("p.EnableCollisions"), EnableCollisions, TEXT("Enable/Disable collisions on the Chaos solver."));
 
 	DECLARE_CYCLE_STAT(TEXT("Collisions::Reset"), STAT_Collisions_Reset, STATGROUP_ChaosCollision);
+	DECLARE_CYCLE_STAT(TEXT("Collisions::UpdatePointConstraints"), STAT_Collisions_UpdatePointConstraints, STATGROUP_ChaosCollision);
+	DECLARE_CYCLE_STAT(TEXT("Collisions::UpdateManifoldConstraints"), STAT_Collisions_UpdateManifoldConstraints, STATGROUP_ChaosCollision);
 	DECLARE_CYCLE_STAT(TEXT("Collisions::Apply"), STAT_Collisions_Apply, STATGROUP_ChaosCollision);
 	DECLARE_CYCLE_STAT(TEXT("Collisions::ApplyPushOut"), STAT_Collisions_ApplyPushOut, STATGROUP_ChaosCollision);
 
-	DECLARE_DWORD_COUNTER_STAT(TEXT("Collisions::NumPointConstraints"), STAT_NumPointConstraints, STATGROUP_ChaosCollision);
-	DECLARE_DWORD_COUNTER_STAT(TEXT("Collisions::NumManifoldConstraints"), STAT_NumManifoldConstraints, STATGROUP_ChaosCollision);
-
 	//
 	// Collision Constraint Container
 	//
@@ -107,6 +106,43 @@ namespace Chaos
 		PostApplyPushOutCallback = nullptr;
 	}
 
+	template<typename T, int d>
+	void TPBDCollisionConstraints<T, d>::UpdateConstraintMaterialProperties(FConstraintBase& Constraint)
+	{
+		TSerializablePtr<FChaosPhysicsMaterial> PhysicsMaterial0 = Constraint.Particle[0]->AuxilaryValue(MPhysicsMaterials);
+		TSerializablePtr<FChaosPhysicsMaterial> PhysicsMaterial1 = Constraint.Particle[1]->AuxilaryValue(MPhysicsMaterials);
+
+		TCollisionContact<T, d>& Contact = Constraint.Manifold;
+		if (PhysicsMaterial0 && PhysicsMaterial1)
+		{
+			// @todo(ccaulfield): support different friction/restitution combining algorithms
+			Contact.Restitution = FMath::Min(PhysicsMaterial0->Restitution, PhysicsMaterial1->Restitution);
+			Contact.Friction = FMath::Max(PhysicsMaterial0->Friction, PhysicsMaterial1->Friction);
+		}
+		else if (PhysicsMaterial0)
+		{
+			Contact.Restitution = PhysicsMaterial0->Restitution;
+			Contact.Friction = PhysicsMaterial0->Friction;
+		}
+		else if (PhysicsMaterial1)
+		{
+			Contact.Restitution = PhysicsMaterial1->Restitution;
+			Contact.Friction = PhysicsMaterial1->Friction;
+		}
+		else
+		{
+			Contact.Friction = 0;
+			Contact.Restitution = 0;
+		}
+		Contact.AngularFriction = MAngularFriction;
+
+		// Overrides for testing
+		if (CollisionFrictionOverride >= 0)
+		{
+			Contact.Friction = CollisionFrictionOverride;
+		}
+	}
+
 	template<typename T, int d>
 	void TPBDCollisionConstraints<T, d>::AddConstraint(FConstraintBase* ConstraintBase)
 	{
@@ -114,6 +150,8 @@ namespace Chaos
 
 		FConstraintContainerHandle* Handle = nullptr;
 
+		UpdateConstraintMaterialProperties(*ConstraintBase);
+
 		if (ConstraintBase->GetType() == TRigidBodyPointContactConstraint<T, 3>::StaticType())
 		{
 			TRigidBodyPointContactConstraint<T, d>* PointConstraint = ConstraintBase->template As< TRigidBodyPointContactConstraint<T, d> >();
@@ -257,50 +295,69 @@ namespace Chaos
 	template<typename T, int d>
 	void TPBDCollisionConstraints<T, d>::UpdateConstraints(T Dt)
 	{
-		PhysicsParallelFor(Handles.Num(), [&](int32 ConstraintHandleIndex)
-		{
-			FConstraintContainerHandle* ConstraintHandle = Handles[ConstraintHandleIndex];
-			check(ConstraintHandle != nullptr);
-			Collisions::Update<ECollisionUpdateType::Deepest, float, 3>(MThickness, ConstraintHandle->GetContact());
+		SCOPE_CYCLE_COUNTER(STAT_Collisions_UpdatePointConstraints);
+
+		//PhysicsParallelFor(Handles.Num(), [&](int32 ConstraintHandleIndex)
+		//{
+		//	FConstraintContainerHandle* ConstraintHandle = Handles[ConstraintHandleIndex];
+		//	check(ConstraintHandle != nullptr);
+		//	Collisions::Update<ECollisionUpdateType::Deepest, float, 3>(MThickness, ConstraintHandle->GetContact());
 
-			if (ConstraintHandle->GetContact().GetPhi() < MThickness) 
+		//	if (ConstraintHandle->GetContact().GetPhi() < MThickness) 
+		//	{
+		//		ConstraintHandle->GetContact().Timestamp = LifespanCounter;
+		//	}
+		//}, bDisableCollisionParallelFor);
+
+		for (FPointContactConstraint& Contact : PointConstraints)
+		{
+			Collisions::Update<ECollisionUpdateType::Deepest, float, 3>(MThickness, Contact);
+			if (Contact.GetPhi() < MThickness)
 			{
-				ConstraintHandle->GetContact().Timestamp = LifespanCounter;
+				Contact.Timestamp = LifespanCounter;
 			}
-		}, bDisableCollisionParallelFor);
+		}
 	}
 
 	template<typename T, int d>
 	void TPBDCollisionConstraints<T, d>::UpdateManifolds(T Dt)
 	{
-		PhysicsParallelFor(Handles.Num(), [&](int32 ConstraintHandleIndex)
+		SCOPE_CYCLE_COUNTER(STAT_Collisions_UpdateManifoldConstraints);
+
+		//PhysicsParallelFor(Handles.Num(), [&](int32 ConstraintHandleIndex)
+		//{
+		//	FConstraintContainerHandle* ConstraintHandle = Handles[ConstraintHandleIndex];
+		//	check(ConstraintHandle != nullptr);
+		//	Collisions::UpdateManifold<float, 3>(MThickness, ConstraintHandle->GetContact());
+		//}, bDisableCollisionParallelFor);
+
+		for (FMultiPointContactConstraint& Contact : IterativeConstraints)
 		{
-			FConstraintContainerHandle* ConstraintHandle = Handles[ConstraintHandleIndex];
-			check(ConstraintHandle != nullptr);
-			Collisions::UpdateManifold<float, 3>(MThickness, ConstraintHandle->GetContact());
-		}, bDisableCollisionParallelFor);
+			Collisions::UpdateManifold<float, 3>(MThickness, Contact);
+			if (Contact.GetPhi() < MThickness)
+			{
+				Contact.Timestamp = LifespanCounter;
+			}
+		}
 	}
 
 	template<typename T, int d>
 	void TPBDCollisionConstraints<T, d>::Apply(const T Dt, const int32 Iterations, const int32 NumIterations)
 	{
 		SCOPE_CYCLE_COUNTER(STAT_Collisions_Apply);
-		SET_DWORD_STAT(STAT_NumPointConstraints, PointConstraints.Num());
-		SET_DWORD_STAT(STAT_NumManifoldConstraints, IterativeConstraints.Num());
 
 		if (MApplyPairIterations > 0)
 		{
+			const Collisions::TContactParticleParameters<T> ParticleParameters = { MThickness, &MCollided };
+			const Collisions::TContactIterationParameters<T> IterationParameters = { Dt, Iterations, NumIterations, MApplyPairIterations, nullptr };
+
 			for (FPointContactConstraint& Contact : PointConstraints)
 			{
-				Collisions::TContactParticleParameters<T> ParticleParameters = { MThickness, &MCollided, &MPhysicsMaterials, CollisionFrictionOverride, MAngularFriction };
-				Collisions::TContactIterationParameters<T> IterationParameters = { Dt, Iterations, NumIterations, MApplyPairIterations, nullptr };
 				Collisions::Apply(Contact, IterationParameters, ParticleParameters);
 			}
 
 			for (FMultiPointContactConstraint& Contact : IterativeConstraints)
 			{
-				Collisions::TContactParticleParameters<T> ParticleParameters = { MThickness, &MCollided, &MPhysicsMaterials, CollisionFrictionOverride, MAngularFriction };
-				Collisions::TContactIterationParameters<T> IterationParameters = { Dt, Iterations, NumIterations, MApplyPairIterations, nullptr };
 				Collisions::Apply(Contact, IterationParameters, ParticleParameters);
 			}
 		}
@@ -315,24 +372,21 @@ namespace Chaos
 	bool TPBDCollisionConstraints<T, d>::ApplyPushOut(const T Dt, const int32 Iterations, const int32 NumIterations)
 	{
 		SCOPE_CYCLE_COUNTER(STAT_Collisions_ApplyPushOut);
-		SET_DWORD_STAT(STAT_NumPointConstraints, PointConstraints.Num());
-		SET_DWORD_STAT(STAT_NumManifoldConstraints, IterativeConstraints.Num());
 
 		TSet<const TGeometryParticleHandle<T, d>*> TempStatic;
 		bool bNeedsAnotherIteration = false;
 		if (MApplyPushOutPairIterations > 0)
 		{
+			const Collisions::TContactParticleParameters<T> ParticleParameters = { MThickness, &MCollided };
+			const Collisions::TContactIterationParameters<T> IterationParameters = { Dt, Iterations, NumIterations, MApplyPushOutPairIterations, &bNeedsAnotherIteration };
+
 			for (FPointContactConstraint& Contact : PointConstraints)
 			{
-				Collisions::TContactParticleParameters<T> ParticleParameters = { MThickness, &MCollided, &MPhysicsMaterials, CollisionFrictionOverride, MAngularFriction };
-				Collisions::TContactIterationParameters<T> IterationParameters = { Dt, Iterations, NumIterations, MApplyPushOutPairIterations, &bNeedsAnotherIteration };
 				Collisions::ApplyPushOut(Contact, TempStatic, IterationParameters, ParticleParameters);
 			}
 
 			for (FMultiPointContactConstraint& Contact : IterativeConstraints)
 			{
-				Collisions::TContactParticleParameters<T> ParticleParameters = { MThickness, &MCollided, &MPhysicsMaterials, CollisionFrictionOverride, MAngularFriction };
-				Collisions::TContactIterationParameters<T> IterationParameters = { Dt, Iterations, NumIterations, MApplyPushOutPairIterations, &bNeedsAnotherIteration };
 				Collisions::ApplyPushOut(Contact, TempStatic, IterationParameters, ParticleParameters);
 			}
 		}
@@ -350,8 +404,6 @@ namespace Chaos
 	void TPBDCollisionConstraints<T, d>::Apply(const T Dt, const TArray<FConstraintContainerHandle*>& InConstraintHandles, const int32 Iterations, const int32 NumIterations)
 	{
 		SCOPE_CYCLE_COUNTER(STAT_Collisions_Apply);
-		SET_DWORD_STAT(STAT_NumPointConstraints, PointConstraints.Num());
-		SET_DWORD_STAT(STAT_NumManifoldConstraints, IterativeConstraints.Num());
 
 		if (MApplyPairIterations > 0)
 		{
@@ -359,7 +411,7 @@ namespace Chaos
 				FConstraintContainerHandle* ConstraintHandle = InConstraintHandles[ConstraintHandleIndex];
 				check(ConstraintHandle != nullptr);
 
-				Collisions::TContactParticleParameters<T> ParticleParameters = { MThickness, &MCollided, &MPhysicsMaterials, CollisionFrictionOverride, MAngularFriction };
+				Collisions::TContactParticleParameters<T> ParticleParameters = { MThickness, &MCollided };
 				Collisions::TContactIterationParameters<T> IterationParameters = { Dt, Iterations, NumIterations, MApplyPairIterations, nullptr };
 				Collisions::Apply(ConstraintHandle->GetContact(), IterationParameters, ParticleParameters);
 
@@ -377,8 +429,6 @@ namespace Chaos
 	bool TPBDCollisionConstraints<T, d>::ApplyPushOut(const T Dt, const TArray<FConstraintContainerHandle*>& InConstraintHandles, const TSet< const TGeometryParticleHandle<T, d>*>& IsTemporarilyStatic, int32 Iteration, int32 NumIterations)
 	{
 		SCOPE_CYCLE_COUNTER(STAT_Collisions_ApplyPushOut);
-		SET_DWORD_STAT(STAT_NumPointConstraints, PointConstraints.Num());
-		SET_DWORD_STAT(STAT_NumManifoldConstraints, IterativeConstraints.Num());
 
 		bool bNeedsAnotherIteration = false;
 		if (MApplyPushOutPairIterations > 0)
@@ -388,7 +438,7 @@ namespace Chaos
 				FConstraintContainerHandle* ConstraintHandle = InConstraintHandles[ConstraintHandleIndex];
 				check(ConstraintHandle != nullptr);
 
-				Collisions::TContactParticleParameters<T> ParticleParameters = { MThickness, &MCollided, &MPhysicsMaterials, CollisionFrictionOverride, MAngularFriction };
+				Collisions::TContactParticleParameters<T> ParticleParameters = { MThickness, &MCollided };
 				Collisions::TContactIterationParameters<T> IterationParameters = { Dt, Iteration, NumIterations, MApplyPushOutPairIterations, &bNeedsAnotherIteration };
 				Collisions::ApplyPushOut(ConstraintHandle->GetContact(), IsTemporarilyStatic, IterationParameters, ParticleParameters);
 
diff --git a/Engine/Source/Runtime/Experimental/Chaos/Private/Chaos/PBDCollisionConstraintsContact.cpp b/Engine/Source/Runtime/Experimental/Chaos/Private/Chaos/PBDCollisionConstraintsContact.cpp
--- a/Engine/Source/Runtime/Experimental/Chaos/Private/Chaos/PBDCollisionConstraintsContact.cpp
+++ b/Engine/Source/Runtime/Experimental/Chaos/Private/Chaos/PBDCollisionConstraintsContact.cpp
@@ -35,18 +35,10 @@ namespace Chaos
 		TVector<T,d> ApplyContact(TCollisionContact<T, d>& Contact,
 			TGenericParticleHandle<T, d> Particle0, 
 			TGenericParticleHandle<T, d> Particle1,
-			TContactIterationParameters<T> & IterationParameters, TContactParticleParameters<T> & ParticleParameters)
+			const TContactIterationParameters<T> & IterationParameters, const TContactParticleParameters<T> & ParticleParameters)
 		{
 			TVector<T, d> AccumulatedImpulse(0);
 
-			TSerializablePtr<FChaosPhysicsMaterial> PhysicsMaterial0, PhysicsMaterial1;
-			if (ParticleParameters.PhysicsMaterials)
-			{
-				PhysicsMaterial0 = Particle0->AuxilaryValue(*ParticleParameters.PhysicsMaterials);
-				PhysicsMaterial1 = Particle1->AuxilaryValue(*ParticleParameters.PhysicsMaterials);
-			}
-
-
 			TPBDRigidParticleHandle<T, d>* PBDRigid0 = Particle0->CastToRigidParticle();
 			TPBDRigidParticleHandle<T, d>* PBDRigid1 = Particle1->CastToRigidParticle();
 
@@ -77,50 +69,13 @@ namespace Chaos
 				TVector<T, d> AngularImpulse(0);
 
 				// Resting contact if very close to the surface
-				T Restitution = (T)0;
-				T Friction = (T)0;
-				T AngularFriction = (T)0;
 				bool bApplyRestitution = (RelativeVelocity.Size() > (2 * 980 * IterationParameters.Dt));
-				if (PhysicsMaterial0 && PhysicsMaterial1)
-				{
-					if (bApplyRestitution)
-					{
-						Restitution = FMath::Min(PhysicsMaterial0->Restitution, PhysicsMaterial1->Restitution);
-					}
-					Friction = FMath::Max(PhysicsMaterial0->Friction, PhysicsMaterial1->Friction);
-				}
-				else if (PhysicsMaterial0)
-				{
-					if (bApplyRestitution)
-					{
-						Restitution = PhysicsMaterial0->Restitution;
-					}
-					Friction = PhysicsMaterial0->Friction;
-				}
-				else if (PhysicsMaterial1)
-				{
-					if (bApplyRestitution)
-					{
-						Restitution = PhysicsMaterial1->Restitution;
-					}
-					Friction = PhysicsMaterial1->Friction;
-				}
+				T Restitution = (bApplyRestitution) ? Contact.Restitution : (T)0;
+				T Friction = Contact.Friction;
+				T AngularFriction = Contact.AngularFriction;
 
-				if (ParticleParameters.FrictionOverride >= 0)
-				{
-					Friction = ParticleParameters.FrictionOverride;
-				}
-				if (ParticleParameters.AngularFrictionOverride >= 0)
+				if (Friction > 0)
 				{
-					AngularFriction = ParticleParameters.AngularFrictionOverride;
-				}
-
-				if (Friction)
-				{
-					if (RelativeNormalVelocity > 0)
-					{
-						RelativeNormalVelocity = 0;
-					}
 					TVector<T, d> VelocityChange = -(Restitution * RelativeNormalVelocity * Contact.Normal + RelativeVelocity);
 					T NormalVelocityChange = TVector<T, d>::DotProduct(VelocityChange, Contact.Normal);
 					PMatrix<T, d, d> FactorInverse = Factor.Inverse();
@@ -205,9 +160,9 @@ namespace Chaos
 					Impulse = ImpulseNumerator / ImpulseDenominator;
 				}
 
-
 				Impulse = GetEnergyClampedImpulse(Particle0->CastToRigidParticle(), Particle1->CastToRigidParticle(), Impulse, VectorToPoint1, VectorToPoint2, Body1Velocity, Body2Velocity);
 				AccumulatedImpulse += Impulse;
+
 				if (bIsRigidDynamic0)
 				{
 					// Velocity update for next step
@@ -241,14 +196,22 @@ namespace Chaos
 		}
 
 		template<typename T, int d>
-		void Apply(TCollisionConstraintBase<T, d>& Constraint, TContactIterationParameters<T> & IterationParameters, TContactParticleParameters<T> & ParticleParameters)
+		void Apply(TCollisionConstraintBase<T, d>& Constraint, const TContactIterationParameters<T> & IterationParameters, const TContactParticleParameters<T> & ParticleParameters)
 		{
 			TGenericParticleHandle<T, d> Particle0 = TGenericParticleHandle<T, d>(Constraint.Particle[0]);
 			TGenericParticleHandle<T, d> Particle1 = TGenericParticleHandle<T, d>(Constraint.Particle[1]);
 
 			for (int32 PairIt = 0; PairIt < IterationParameters.NumPairIterations; ++PairIt)
 			{
-				Collisions::Update<ECollisionUpdateType::Deepest>(ParticleParameters.Thickness, Constraint);
+				// Collision is already up-to-date on first iteration (we either just detected it, or updated it in DetectCollisions)
+				// @todo(ccaulfield): this is not great - try to do something nicer like a dirty flag on the constraint?
+				// In particular it is not right if the Collisions are not the first constraints to be solved...
+				const bool bNeedCollisionUpdate = (PairIt > 0) || (IterationParameters.Iteration > 0);
+				if (bNeedCollisionUpdate)
+				{
+					Collisions::Update<ECollisionUpdateType::Deepest>(ParticleParameters.Thickness, Constraint);
+				}
+
 				if (Constraint.GetPhi() >= ParticleParameters.Thickness)
 				{
 					return;
@@ -280,7 +243,7 @@ namespace Chaos
 			TGenericParticleHandle<T, d> Particle0, 
 			TGenericParticleHandle<T, d> Particle1,
 			const TSet<const TGeometryParticleHandle<T, d>*>& IsTemporarilyStatic,
-			TContactIterationParameters<T> & IterationParameters, TContactParticleParameters<T> & ParticleParameters)
+			const TContactIterationParameters<T> & IterationParameters, const TContactParticleParameters<T> & ParticleParameters)
 		{
 			TVector<T, d> AccumulatedImpulse(0);
 
@@ -297,13 +260,6 @@ namespace Chaos
 			const bool IsTemporarilyStatic0 = IsTemporarilyStatic.Contains(Particle0->GeometryParticleHandle());
 			const bool IsTemporarilyStatic1 = IsTemporarilyStatic.Contains(Particle1->GeometryParticleHandle());
 
-			TSerializablePtr<FChaosPhysicsMaterial> PhysicsMaterial0, PhysicsMaterial1;
-			if (ParticleParameters.PhysicsMaterials)
-			{
-				PhysicsMaterial0 = Particle0->AuxilaryValue(*ParticleParameters.PhysicsMaterials);
-				PhysicsMaterial1 = Particle1->AuxilaryValue(*ParticleParameters.PhysicsMaterials);
-			}
-
 			if (Contact.Phi >= ParticleParameters.Thickness)
 			{
 				return AccumulatedImpulse;
@@ -314,7 +270,11 @@ namespace Chaos
 				return AccumulatedImpulse;
 			}
 
-			*IterationParameters.NeedsAnotherIteration = true;
+			if (IterationParameters.NeedsAnotherIteration)
+			{
+				*IterationParameters.NeedsAnotherIteration = true;
+			}
+
 			PMatrix<T, d, d> WorldSpaceInvI1 = bIsRigidDynamic0 ? Utilities::ComputeWorldSpaceInertia(Q0, PBDRigid0->InvI()) : PMatrix<T, d, d>(0);
 			PMatrix<T, d, d> WorldSpaceInvI2 = bIsRigidDynamic1 ? Utilities::ComputeWorldSpaceInertia(Q1, PBDRigid1->InvI()) : PMatrix<T, d, d>(0);
 			TVector<T, d> VectorToPoint1 = Contact.Location - P0;
@@ -388,12 +348,12 @@ namespace Chaos
 
 		template<typename T, int d>
 		void ApplyPushOut(TCollisionConstraintBase<T, d>& Constraint, const TSet<const TGeometryParticleHandle<T, d>*>& IsTemporarilyStatic,
-			TContactIterationParameters<T> & IterationParameters, TContactParticleParameters<T> & ParticleParameters)
+			const TContactIterationParameters<T> & IterationParameters, const TContactParticleParameters<T> & ParticleParameters)
 		{
 			TGenericParticleHandle<T, d> Particle0 = TGenericParticleHandle<T, d>(Constraint.Particle[0]);
 			TGenericParticleHandle<T, d> Particle1 = TGenericParticleHandle<T, d>(Constraint.Particle[1]);
 
-			for (int32 PairIteration = 0; PairIteration < IterationParameters.NumPairIterations; ++PairIteration)
+			for (int32 PairIt = 0; PairIt < IterationParameters.NumPairIterations; ++PairIt)
 			{
 				Update<ECollisionUpdateType::Deepest>(ParticleParameters.Thickness, Constraint);
 
@@ -407,9 +367,9 @@ namespace Chaos
 		template void Update<ECollisionUpdateType::Deepest, float, 3>(const float, TCollisionConstraintBase<float, 3>&);
 		template void UpdateManifold<float, 3>(const float, TCollisionConstraintBase<float, 3>&);
 		//
-		template void Apply<float, 3>(TCollisionConstraintBase<float, 3>&, TContactIterationParameters<float> &, TContactParticleParameters<float> &);
+		template void Apply<float, 3>(TCollisionConstraintBase<float, 3>&, const TContactIterationParameters<float> &, const TContactParticleParameters<float> &);
 		template void ApplyPushOut<float, 3>(TCollisionConstraintBase<float,3>& , const TSet<const TGeometryParticleHandle<float,3>*>&,
-			TContactIterationParameters<float> & IterationParameters, TContactParticleParameters<float> & ParticleParameters);
+			const TContactIterationParameters<float> & IterationParameters, const TContactParticleParameters<float> & ParticleParameters);
 
 	} // Collisions
 
diff --git a/Engine/Source/Runtime/Experimental/Chaos/Private/Chaos/PBDJointConstraints.cpp b/Engine/Source/Runtime/Experimental/Chaos/Private/Chaos/PBDJointConstraints.cpp
--- a/Engine/Source/Runtime/Experimental/Chaos/Private/Chaos/PBDJointConstraints.cpp
+++ b/Engine/Source/Runtime/Experimental/Chaos/Private/Chaos/PBDJointConstraints.cpp
@@ -20,8 +20,6 @@ namespace Chaos
 	DECLARE_CYCLE_STAT(TEXT("Joints::Apply"), STAT_Joints_Apply, STATGROUP_ChaosJoint);
 	DECLARE_CYCLE_STAT(TEXT("Joints::ApplyPushOut"), STAT_Joints_ApplyPushOut, STATGROUP_ChaosJoint);
 
-	DECLARE_DWORD_COUNTER_STAT(TEXT("Joints::NumConstraints"), STAT_NumJointConstraints, STATGROUP_ChaosJoint);
-
 	//
 	// Constraint Handle
 	//
@@ -490,7 +488,6 @@ namespace Chaos
 	void FPBDJointConstraints::Apply(const FReal Dt, const int32 It, const int32 NumIts)
 	{
 		SCOPE_CYCLE_COUNTER(STAT_Joints_Apply);
-		SET_DWORD_STAT(STAT_NumJointConstraints, NumConstraints());
 
 		if (PreApplyCallback != nullptr)
 		{
@@ -514,7 +511,6 @@ namespace Chaos
 	bool FPBDJointConstraints::ApplyPushOut(const FReal Dt, const int32 It, const int32 NumIts)
 	{
 		SCOPE_CYCLE_COUNTER(STAT_Joints_ApplyPushOut);
-		SET_DWORD_STAT(STAT_NumJointConstraints, NumConstraints());
 
 		// @todo(ccaulfield): track whether we are sufficiently solved
 		bool bNeedsAnotherIteration = true;
diff --git a/Engine/Source/Runtime/Experimental/Chaos/Public/Chaos/Collision/CollisionDetector.h b/Engine/Source/Runtime/Experimental/Chaos/Public/Chaos/Collision/CollisionDetector.h
--- a/Engine/Source/Runtime/Experimental/Chaos/Public/Chaos/Collision/CollisionDetector.h
+++ b/Engine/Source/Runtime/Experimental/Chaos/Public/Chaos/Collision/CollisionDetector.h
@@ -11,8 +11,7 @@
 
 namespace Chaos
 {
-	DECLARE_CYCLE_STAT_EXTERN(TEXT("DetectCollisions"), STAT_DetectCollisions, STATGROUP_Chaos, CHAOS_API);
-
+	DECLARE_CYCLE_STAT_EXTERN(TEXT("Collisions::Detect"), STAT_Collisions_Detect, STATGROUP_ChaosCollision, CHAOS_API);
 
 	template<typename T_BROADPHASE, typename T_NARROWPHASE, typename T_RECEIVER, typename T_CONTAINER>
 	class CHAOS_API TCollisionDetector
@@ -40,7 +39,7 @@ namespace Chaos
 
 		void DetectCollisions(const FReal Dt, CollisionStats::FStatData& StatData)
 		{
-			SCOPE_CYCLE_COUNTER(STAT_DetectCollisions);
+			SCOPE_CYCLE_COUNTER(STAT_Collisions_Detect);
 			CHAOS_SCOPED_TIMER(DetectCollisions);
 
 			if (!GetCollisionContainer().GetCollisionsEnabled())
diff --git a/Engine/Source/Runtime/Experimental/Chaos/Public/Chaos/Collision/NarrowPhase.h b/Engine/Source/Runtime/Experimental/Chaos/Public/Chaos/Collision/NarrowPhase.h
--- a/Engine/Source/Runtime/Experimental/Chaos/Public/Chaos/Collision/NarrowPhase.h
+++ b/Engine/Source/Runtime/Experimental/Chaos/Public/Chaos/Collision/NarrowPhase.h
@@ -7,9 +7,12 @@
 #include "Chaos/ParticleHandle.h"
 #include "Chaos/PBDCollisionConstraints.h"
 #include "Chaos/PBDRigidsSOAs.h"
+#include "ChaosStats.h"
 
 namespace Chaos
 {
+	DECLARE_CYCLE_STAT_EXTERN(TEXT("Collisions::NarrowPhase"), STAT_Collisions_NarrowPhase, STATGROUP_ChaosCollision, CHAOS_API);
+
 	/**
 	 * Generate contact manifolds for particle pairs and pass them on to the consumer. Can be composed with a
 	 * multi-threaded BroadPhase as long as the collisions receiver type can handle multi-threaded calls to ReceiveCollisions.
@@ -26,6 +29,8 @@ namespace Chaos
 		template<typename T_RECEIVER>
 		void GenerateCollisions(FReal Dt, T_RECEIVER& Receiver, TGeometryParticleHandle<FReal, 3>* Particle1, TGeometryParticleHandle<FReal, 3>* Particle2, const FReal BoundsThickness, CollisionStats::FStatData& StatData)
 		{
+			SCOPE_CYCLE_COUNTER(STAT_Collisions_NarrowPhase);
+
 			// @todo(ccaulfield): COLLISION - Thickness: add shape padding (BoundsThickness is the distance within which we speculatively create constraints)
 
 			FCollisionConstraintsArray NewConstraints;
diff --git a/Engine/Source/Runtime/Experimental/Chaos/Public/Chaos/Collision/ParticlePairBroadPhase.h b/Engine/Source/Runtime/Experimental/Chaos/Public/Chaos/Collision/ParticlePairBroadPhase.h
--- a/Engine/Source/Runtime/Experimental/Chaos/Public/Chaos/Collision/ParticlePairBroadPhase.h
+++ b/Engine/Source/Runtime/Experimental/Chaos/Public/Chaos/Collision/ParticlePairBroadPhase.h
@@ -7,9 +7,12 @@
 #include "Chaos/ParticleHandle.h"
 #include "Chaos/PBDCollisionConstraints.h"
 #include "Chaos/PBDRigidsSOAs.h"
+#include "ChaosStats.h"
 
 namespace Chaos
 {
+	DECLARE_CYCLE_STAT_EXTERN(TEXT("Collisions::BroadPhase"), STAT_Collisions_BroadPhase, STATGROUP_ChaosCollision, CHAOS_API);
+
 	/**
 	 * Run through a list of particle pairs and pass them onto the collision detector if their AABBs overlap.
 	 * No spatial acceleration, and the order is assumed to be already optimized for cache efficiency.
@@ -33,6 +36,8 @@ namespace Chaos
 			FSyncCollisionReceiver& Receiver,
 			CollisionStats::FStatData& StatData)
 		{
+			SCOPE_CYCLE_COUNTER(STAT_Collisions_BroadPhase);
+
 			const FReal BoundsThickness = 1.0f;
 			const FReal BoundsThicknessVelocityInflation = 2.0f;
 
diff --git a/Engine/Source/Runtime/Experimental/Chaos/Public/Chaos/CollisionResolutionTypes.h b/Engine/Source/Runtime/Experimental/Chaos/Public/Chaos/CollisionResolutionTypes.h
--- a/Engine/Source/Runtime/Experimental/Chaos/Public/Chaos/CollisionResolutionTypes.h
+++ b/Engine/Source/Runtime/Experimental/Chaos/Public/Chaos/CollisionResolutionTypes.h
@@ -46,7 +46,7 @@ namespace Chaos
 	{
 	public:
 		TCollisionContact(const FImplicitObject* InImplicit0 = nullptr, const FImplicitObject* InImplicit1 = nullptr)
-			: bDisabled(true), Normal(0), Location(0), Phi(FLT_MAX)
+			: bDisabled(true), Normal(0), Location(0), Phi(FLT_MAX), Friction(0), AngularFriction(0), Restitution(0)
 		{
 			Implicit[0] = InImplicit0;
 			Implicit[1] = InImplicit1;
@@ -57,6 +57,10 @@ namespace Chaos
 		TVector<T, d> Location;
 		T Phi;
 
+		T Friction;
+		T AngularFriction;
+		T Restitution;
+
 
 		FString ToString() const
 		{
diff --git a/Engine/Source/Runtime/Experimental/Chaos/Public/Chaos/PBDCollisionConstraints.h b/Engine/Source/Runtime/Experimental/Chaos/Public/Chaos/PBDCollisionConstraints.h
--- a/Engine/Source/Runtime/Experimental/Chaos/Public/Chaos/PBDCollisionConstraints.h
+++ b/Engine/Source/Runtime/Experimental/Chaos/Public/Chaos/PBDCollisionConstraints.h
@@ -218,8 +218,10 @@ public:
 
 
 protected:
-		using Base::GetConstraintIndex;
-		using Base::SetConstraintIndex;
+	using Base::GetConstraintIndex;
+	using Base::SetConstraintIndex;
+
+	void UpdateConstraintMaterialProperties(FConstraintBase& Contact);
 
 private:
 
@@ -240,6 +242,7 @@ private:
 	T MAngularFriction;
 	bool bUseCCD;
 	bool bEnableCollisions;
+	bool bEnableParallelFor;
 
 	int32 LifespanCounter;
 
diff --git a/Engine/Source/Runtime/Experimental/Chaos/Public/Chaos/PBDCollisionConstraintsContact.h b/Engine/Source/Runtime/Experimental/Chaos/Public/Chaos/PBDCollisionConstraintsContact.h
--- a/Engine/Source/Runtime/Experimental/Chaos/Public/Chaos/PBDCollisionConstraintsContact.h
+++ b/Engine/Source/Runtime/Experimental/Chaos/Public/Chaos/PBDCollisionConstraintsContact.h
@@ -13,9 +13,6 @@ namespace Chaos
 		struct TContactParticleParameters {
 			T Thickness;
 			TArrayCollectionArray<bool>* Collided;
-			const TArrayCollectionArray<TSerializablePtr<FChaosPhysicsMaterial>>* PhysicsMaterials;
-			T FrictionOverride;
-			T AngularFrictionOverride;
 		};
 
 		template<class T = float>
@@ -34,11 +31,11 @@ namespace Chaos
 		void UpdateManifold(const T Thickness, TCollisionConstraintBase<T, d>& Constraint);
 
 		template<typename T, int d>
-		void Apply(TCollisionConstraintBase<T, d>& Constraint, TContactIterationParameters<T> & IterationParameters, TContactParticleParameters<T> & ParticleParameters);
+		void Apply(TCollisionConstraintBase<T, d>& Constraint, const TContactIterationParameters<T> & IterationParameters, const TContactParticleParameters<T> & ParticleParameters);
 
 		template<typename T, int d>
 		void ApplyPushOut(TCollisionConstraintBase<T, d>& Constraint, const TSet<const TGeometryParticleHandle<T, d>*>& IsTemporarilyStatic,
-			TContactIterationParameters<T> & IterationParameters, TContactParticleParameters<T> & ParticleParameters);
+			const TContactIterationParameters<T> & IterationParameters, const TContactParticleParameters<T> & ParticleParameters);
 
 
 	}
