Chaos : Solver Awake Events

- Propagate the state of the particle into the solver so that sleeping vehicles can be woken up when their external state changes.
- Expose internal simulation state changes from sleeping to dynamic, then broadcast the events on the associated BodyInstance.

diff --git a/Engine/Source/Runtime/Engine/Private/PhysicsEngine/Experimental/PhysInterface_Chaos.cpp b/Engine/Source/Runtime/Engine/Private/PhysicsEngine/Experimental/PhysInterface_Chaos.cpp
--- a/Engine/Source/Runtime/Engine/Private/PhysicsEngine/Experimental/PhysInterface_Chaos.cpp
+++ b/Engine/Source/Runtime/Engine/Private/PhysicsEngine/Experimental/PhysInterface_Chaos.cpp
@@ -405,6 +405,7 @@ void FPhysInterface_Chaos::WakeUp_AssumesLocked(const FPhysicsActorHandle& InAct
 	if(Particle && Particle->ObjectState() == Chaos::EObjectStateType::Sleeping)
 	{
 		Particle->SetObjectState(Chaos::EObjectStateType::Dynamic);
+		Particle->ClearEvents();
 	}
 }
 
diff --git a/Engine/Source/Runtime/Engine/Private/PhysicsEngine/Experimental/PhysScene_Chaos.cpp b/Engine/Source/Runtime/Engine/Private/PhysicsEngine/Experimental/PhysScene_Chaos.cpp
--- a/Engine/Source/Runtime/Engine/Private/PhysicsEngine/Experimental/PhysScene_Chaos.cpp
+++ b/Engine/Source/Runtime/Engine/Private/PhysicsEngine/Experimental/PhysScene_Chaos.cpp
@@ -1644,27 +1644,36 @@ void FPhysScene_ChaosInterface::SyncBodies(Chaos::FPhysicsSolver* Solver)
 
 				if (FBodyInstance* BodyInstance = FPhysicsUserData::Get<FBodyInstance>(ActiveParticle->UserData()))
 				{
-					if (BodyInstance->InstanceBodyIndex == INDEX_NONE && BodyInstance->OwnerComponent.IsValid())
+					if (BodyInstance->OwnerComponent.IsValid())
 					{
 						UPrimitiveComponent* OwnerComponent = BodyInstance->OwnerComponent.Get();
 						if (OwnerComponent != nullptr)
 						{
 							AActor* Owner = OwnerComponent->GetOwner();
 
-							Chaos::TRigidTransform<float, 3> NewTransform(ActiveParticle->X(), ActiveParticle->R());
-
-							if (!NewTransform.EqualsNoScale(OwnerComponent->GetComponentTransform()))
+							if (BodyInstance->InstanceBodyIndex == INDEX_NONE)
 							{
-								const FVector MoveBy = NewTransform.GetLocation() - OwnerComponent->GetComponentTransform().GetLocation();
-								const FQuat NewRotation = NewTransform.GetRotation();
+								Chaos::TRigidTransform<float, 3> NewTransform(ActiveParticle->X(), ActiveParticle->R());
+
+								if (!NewTransform.EqualsNoScale(OwnerComponent->GetComponentTransform()))
+								{
+									const FVector MoveBy = NewTransform.GetLocation() - OwnerComponent->GetComponentTransform().GetLocation();
+									const FQuat NewRotation = NewTransform.GetRotation();
 
-								OwnerComponent->MoveComponent(MoveBy, NewRotation, false, NULL, MOVECOMP_SkipPhysicsMove);
+									OwnerComponent->MoveComponent(MoveBy, NewRotation, false, NULL, MOVECOMP_SkipPhysicsMove);
+								}
+
+								if (Owner != NULL && !Owner->IsPendingKill())
+								{
+									Owner->CheckStillInWorld();
+								}
 							}
 
-							if (Owner != NULL && !Owner->IsPendingKill())
+							if (Proxy->HasAwakeEvent())
 							{
-								Owner->CheckStillInWorld();
+								OwnerComponent->DispatchWakeEvents(ESleepEvent::SET_Wakeup, NAME_None);
 							}
+							Proxy->ClearEvents();
 						}
 					}
 				}
diff --git a/Engine/Source/Runtime/Experimental/Chaos/Private/Chaos/ParticleHandle.cpp b/Engine/Source/Runtime/Experimental/Chaos/Private/Chaos/ParticleHandle.cpp
--- a/Engine/Source/Runtime/Experimental/Chaos/Private/Chaos/ParticleHandle.cpp
+++ b/Engine/Source/Runtime/Experimental/Chaos/Private/Chaos/ParticleHandle.cpp
@@ -49,4 +49,18 @@ namespace Chaos
 		return Idx;
 	}
 
+	/*
+	template <>
+	void TPBDRigidParticle<float, 3>::SetObjectState(const EObjectStateType InState, bool bAllowEvents)
+	{
+		//todo: look at physics thread logic
+		if (bAllowEvents && MObjectState != EObjectStateType::Dynamic && InState == EObjectStateType::Dynamic)
+		{
+			MAwakeEvent |= true;
+		}
+		MObjectState = InState;
+		this->MarkDirty(EParticleFlags::ObjectState);
+	}
+	*/
+
 }
diff --git a/Engine/Source/Runtime/Experimental/Chaos/Public/Chaos/ParticleHandle.h b/Engine/Source/Runtime/Experimental/Chaos/Public/Chaos/ParticleHandle.h
--- a/Engine/Source/Runtime/Experimental/Chaos/Public/Chaos/ParticleHandle.h
+++ b/Engine/Source/Runtime/Experimental/Chaos/Public/Chaos/ParticleHandle.h
@@ -107,6 +107,7 @@ void PBDRigidParticleDefaultConstruct(TPBDRigidParticle<T,d>& Concrete, const TP
 	Concrete.SetAngularEtherDrag(0.f);
 	Concrete.SetObjectState(Params.bStartSleeping ? EObjectStateType::Sleeping : EObjectStateType::Dynamic);
 	Concrete.SetGravityEnabled(Params.bGravityEnabled);
+	Concrete.ClearEvents();
 }
 
 
@@ -1408,6 +1409,7 @@ public:
 	const TShapesArray<T,d>& ShapesArray() const { return MShapesArray; }
 
 	EObjectStateType ObjectState() const;
+	void SetObjectState(const EObjectStateType InState, bool bAllowEvents = false);
 
 	EParticleType ObjectType() const
 	{
@@ -1737,7 +1739,7 @@ public:
 protected:
 	friend TGeometryParticle<T, d>* TGeometryParticle<T, d>::SerializationFactory(FChaosArchive& Ar, TGeometryParticle<T, d>* Serializable);
 	TPBDRigidParticle<T, d>(const TPBDRigidParticleParameters<T, d>& DynamicParams = TPBDRigidParticleParameters<T, d>())
-		: TKinematicGeometryParticle<T, d>(DynamicParams)
+		: TKinematicGeometryParticle<T, d>(DynamicParams), MAwakeEvent(false)
 	{
 		Type = EParticleType::Rigid;
 		MIsland = INDEX_NONE;
@@ -1932,13 +1934,19 @@ public:
 	}
 
 	EObjectStateType ObjectState() const { return MObjectState; }
-	void SetObjectState(EObjectStateType InState)
+	void SetObjectState(const EObjectStateType InState, bool bAllowEvents = false)
 	{
-		//todo: look at physics thread logic
+		if (bAllowEvents && MObjectState != EObjectStateType::Dynamic && InState == EObjectStateType::Dynamic)
+		{
+			MAwakeEvent |= true;
+		}
 		MObjectState = InState;
 		this->MarkDirty(EParticleFlags::ObjectState);
 	}
 
+	void ClearEvents() { MAwakeEvent = false; }
+	bool HasAwakeEvent() { return MAwakeEvent; }
+
 	FParticleData* NewData() const
 	{
 		return new TPBDRigidParticleData<T, d>(*this);
@@ -1967,6 +1975,7 @@ private:
 	bool MDisabled;
 	bool MToBeRemovedOnFracture;
 	bool MGravityEnabled;
+	bool MAwakeEvent;
 };
 
 template <typename T, int d>
@@ -2146,6 +2155,15 @@ const TPBDRigidParticle<T, d>* TGeometryParticle<T, d>::CastToRigidParticle()  c
 	return nullptr;
 }
 
+template <typename T, int d>
+void TGeometryParticle<T, d>::SetObjectState(const EObjectStateType InState, bool bAllowEvents)
+{
+	TPBDRigidParticle<T, d>* Dyn = CastToRigidParticle();
+	if (Dyn)
+	{
+		Dyn->SetObjectState(InState, bAllowEvents);
+	}
+}
 
 template <typename T, int d>
 EObjectStateType TGeometryParticle<T, d>::ObjectState() const
diff --git a/Engine/Source/Runtime/Experimental/ChaosSolvers/Private/PhysicsProxy/SingleParticlePhysicsProxy.cpp b/Engine/Source/Runtime/Experimental/ChaosSolvers/Private/PhysicsProxy/SingleParticlePhysicsProxy.cpp
--- a/Engine/Source/Runtime/Experimental/ChaosSolvers/Private/PhysicsProxy/SingleParticlePhysicsProxy.cpp
+++ b/Engine/Source/Runtime/Experimental/ChaosSolvers/Private/PhysicsProxy/SingleParticlePhysicsProxy.cpp
@@ -110,6 +110,17 @@ EPhysicsProxyType FSingleParticlePhysicsProxy<Chaos::TGeometryParticle<float, 3>
 	return EPhysicsProxyType::SingleGeometryParticleType;
 }
 
+template< >
+bool FSingleParticlePhysicsProxy<Chaos::TGeometryParticle<float, 3>>::HasAwakeEvent() const
+{
+	return false;
+}
+
+template< >
+void FSingleParticlePhysicsProxy<Chaos::TGeometryParticle<float, 3>>::ClearEvents()
+{
+
+}
 //
 // TKinematicGeometryParticle template specialization 
 //
@@ -187,6 +198,17 @@ EPhysicsProxyType FSingleParticlePhysicsProxy<Chaos::TKinematicGeometryParticle<
 	return EPhysicsProxyType::SingleKinematicParticleType;
 }
 
+template< >
+bool FSingleParticlePhysicsProxy<Chaos::TKinematicGeometryParticle<float, 3>>::HasAwakeEvent() const
+{
+	return false;
+}
+
+template< >
+void FSingleParticlePhysicsProxy<Chaos::TKinematicGeometryParticle<float, 3>>::ClearEvents()
+{
+
+}
 
 //
 // TPBDRigidParticle template specialization 
@@ -217,15 +239,22 @@ void FSingleParticlePhysicsProxy<Chaos::TPBDRigidParticle<float, 3>>::PushToPhys
 		RigidHandle->SetDebugName(Data->DebugName);
 #endif
 
+		//  DynamicProperties are anything that would wake up a sleeping
+		//  particle. For example, if we set position on a sleeping particle
+		//  we will need to wake it up
+		bool bDynamicPropertyUpdated = false;
+
 		if (Data->DirtyFlags.IsDirty(Chaos::EParticleFlags::X))
 		{
 			RigidHandle->SetX(Data->X);
 			RigidHandle->SetP(Data->X);
+			bDynamicPropertyUpdated = true;
 		}
 		if (Data->DirtyFlags.IsDirty(Chaos::EParticleFlags::R))
 		{
 			RigidHandle->SetR(Data->R);
 			RigidHandle->SetQ(Data->R);
+			bDynamicPropertyUpdated = true;
 		}
 		if (Data->DirtyFlags.IsDirty(Chaos::EParticleFlags::Geometry))
 		{
@@ -234,10 +263,12 @@ void FSingleParticlePhysicsProxy<Chaos::TPBDRigidParticle<float, 3>>::PushToPhys
 		if (Data->DirtyFlags.IsDirty(Chaos::EParticleFlags::V))
 		{
 			RigidHandle->SetV(Data->MV);
+			bDynamicPropertyUpdated = true;
 		}
 		if (Data->DirtyFlags.IsDirty(Chaos::EParticleFlags::W))
 		{
 			RigidHandle->SetW(Data->MW);
+			bDynamicPropertyUpdated = true;
 		}
 
 		RigidHandle->SetCenterOfMass(Data->MCenterOfMass);
@@ -252,10 +283,12 @@ void FSingleParticlePhysicsProxy<Chaos::TPBDRigidParticle<float, 3>>::PushToPhys
 		if (Data->DirtyFlags.IsDirty(Chaos::EParticleFlags::F))
 		{
 			RigidHandle->SetF(Data->MF);
+			bDynamicPropertyUpdated = true;
 		}
 		if (Data->DirtyFlags.IsDirty(Chaos::EParticleFlags::Torque))
 		{
 			RigidHandle->SetTorque(Data->MTorque);
+			bDynamicPropertyUpdated = true;
 		}
 		if (Data->DirtyFlags.IsDirty(Chaos::EParticleFlags::ObjectState))
 		{
@@ -282,15 +315,29 @@ void FSingleParticlePhysicsProxy<Chaos::TPBDRigidParticle<float, 3>>::PushToPhys
 			}
 		}
 
-		if(Data->DirtyFlags.IsDirty(Chaos::EParticleFlags::ShapeDisableCollision))
+		if (Data->DirtyFlags.IsDirty(Chaos::EParticleFlags::ShapeDisableCollision))
 		{
 			int32 CurrShape = 0;
-			for(const TUniquePtr<Chaos::TPerShapeData<Chaos::FReal, 3>>& Shape : RigidHandle->ShapesArray())
+			for (const TUniquePtr<Chaos::TPerShapeData<Chaos::FReal, 3>>& Shape : RigidHandle->ShapesArray())
 			{
 				Shape->bDisable = Data->ShapeCollisionDisableFlags[CurrShape++];
 			}
 		}
+
+		if (bDynamicPropertyUpdated)
+		{
+			if (bInitialized && RigidHandle->ObjectState() == Chaos::EObjectStateType::Sleeping)
+			{
+				GetSolver()->GetEvolution()->SetParticleObjectState(RigidHandle, Chaos::EObjectStateType::Dynamic);
+			}
+		}
+		else
+		{
+			// wait for the first pass with nothing updated to claim its initialized
+			bInitialized = true;
+		}
 	}
+
 }
 
 template< >
@@ -312,6 +359,7 @@ void FSingleParticlePhysicsProxy<Chaos::TPBDRigidParticle<float, 3>>::BufferPhys
 		Buffer->R = RigidHandle->R();
 		Buffer->MV = RigidHandle->V();
 		Buffer->MW = RigidHandle->W();
+		Buffer->MObjectState = RigidHandle->ObjectState();
 	}
 }
 
@@ -327,6 +375,7 @@ void FSingleParticlePhysicsProxy<Chaos::TPBDRigidParticle<float, 3>>::PullFromPh
 		Particle->SetV(Buffer->MV, false);
 		Particle->SetW(Buffer->MW, false);
 		Particle->UpdateShapeBounds();
+		Particle->SetObjectState(Buffer->MObjectState, true);
 	}
 }
 
@@ -336,6 +385,18 @@ bool FSingleParticlePhysicsProxy<Chaos::TPBDRigidParticle<float, 3>>::IsDirty()
 	return Particle->IsDirty();
 }
 
+template< >
+bool FSingleParticlePhysicsProxy< Chaos::TPBDRigidParticle<float, 3> >::HasAwakeEvent() const
+{
+	return Particle->HasAwakeEvent();
+}
+
+template< >
+void FSingleParticlePhysicsProxy< Chaos::TPBDRigidParticle<float, 3> >::ClearEvents()
+{
+	Particle->ClearEvents();
+}
+
 template< >
 EPhysicsProxyType  FSingleParticlePhysicsProxy<Chaos::TPBDRigidParticle<float, 3>>::ConcreteType() 
 { 
diff --git a/Engine/Source/Runtime/Experimental/ChaosSolvers/Public/PhysicsProxy/SingleParticlePhysicsProxy.h b/Engine/Source/Runtime/Experimental/ChaosSolvers/Public/PhysicsProxy/SingleParticlePhysicsProxy.h
--- a/Engine/Source/Runtime/Experimental/ChaosSolvers/Public/PhysicsProxy/SingleParticlePhysicsProxy.h
+++ b/Engine/Source/Runtime/Experimental/ChaosSolvers/Public/PhysicsProxy/SingleParticlePhysicsProxy.h
@@ -136,9 +136,15 @@ public:
 	/**/
 	void PullFromPhysicsState();
 
-
+	/**/
 	bool IsDirty();
 
+	/**/
+	bool HasAwakeEvent() const;
+
+	/**/
+	void ClearEvents();
+
 private:
 	bool bInitialized;
 	TArray<int32> InitializedIndices;
@@ -171,6 +177,11 @@ void FSingleParticlePhysicsProxy<Chaos::TGeometryParticle<float, 3>>::PushToPhys
 template< >
 EPhysicsProxyType FSingleParticlePhysicsProxy<Chaos::TGeometryParticle<float, 3>>::ConcreteType();
 
+template< >
+CHAOSSOLVERS_API bool FSingleParticlePhysicsProxy<Chaos::TGeometryParticle<float, 3>>::HasAwakeEvent() const;
+
+template< >
+CHAOSSOLVERS_API void FSingleParticlePhysicsProxy<Chaos::TGeometryParticle<float, 3>>::ClearEvents();
 
 // TKinematicGeometryParticle specialization prototypes
 
@@ -192,6 +203,11 @@ bool FSingleParticlePhysicsProxy<Chaos::TKinematicGeometryParticle<float, 3>>::I
 template< >
 EPhysicsProxyType FSingleParticlePhysicsProxy<Chaos::TKinematicGeometryParticle<float, 3>>::ConcreteType();
 
+template< >
+CHAOSSOLVERS_API bool FSingleParticlePhysicsProxy<Chaos::TKinematicGeometryParticle<float, 3>>::HasAwakeEvent() const;
+
+template< >
+CHAOSSOLVERS_API void FSingleParticlePhysicsProxy<Chaos::TKinematicGeometryParticle<float, 3>>::ClearEvents();
 
 // TPBDRigidParticles specialization prototypes
 
@@ -212,6 +228,14 @@ void FSingleParticlePhysicsProxy<Chaos::TPBDRigidParticle<float, 3>>::PushToPhys
 
 template< >
 EPhysicsProxyType FSingleParticlePhysicsProxy<Chaos::TPBDRigidParticle<float, 3>>::ConcreteType();
+
+template< >
+CHAOSSOLVERS_API bool FSingleParticlePhysicsProxy<Chaos::TPBDRigidParticle<float, 3>>::HasAwakeEvent() const;
+
+template< >
+CHAOSSOLVERS_API void FSingleParticlePhysicsProxy<Chaos::TPBDRigidParticle<float, 3>>::ClearEvents();
+
+
 extern template class FSingleParticlePhysicsProxy< Chaos::TGeometryParticle<float, 3> >;
 extern template class FSingleParticlePhysicsProxy< Chaos::TKinematicGeometryParticle<float, 3> >;
 extern template class FSingleParticlePhysicsProxy< Chaos::TPBDRigidParticle<float,3> >;
