Chaos : Cleanup
 - removed commented out code

diff --git a/Engine/Source/Runtime/Experimental/Chaos/Private/Chaos/ParticleHandle.cpp b/Engine/Source/Runtime/Experimental/Chaos/Private/Chaos/ParticleHandle.cpp
--- a/Engine/Source/Runtime/Experimental/Chaos/Private/Chaos/ParticleHandle.cpp
+++ b/Engine/Source/Runtime/Experimental/Chaos/Private/Chaos/ParticleHandle.cpp
@@ -49,18 +49,4 @@ namespace Chaos
 		return Idx;
 	}
 
-	/*
-	template <>
-	void TPBDRigidParticle<float, 3>::SetObjectState(const EObjectStateType InState, bool bAllowEvents)
-	{
-		//todo: look at physics thread logic
-		if (bAllowEvents && MObjectState != EObjectStateType::Dynamic && InState == EObjectStateType::Dynamic)
-		{
-			MAwakeEvent |= true;
-		}
-		MObjectState = InState;
-		this->MarkDirty(EParticleFlags::ObjectState);
-	}
-	*/
-
 }
