Change acceleration structures TMap to be a dense array. This requires the concept of a unique idx that is dense and thread safe. Helps with GT stalls when we do many Find operations on the map

diff --git a/Engine/Source/Programs/HeadlessChaos/HeadlessChaos.Target.cs b/Engine/Source/Programs/HeadlessChaos/HeadlessChaos.Target.cs
--- a/Engine/Source/Programs/HeadlessChaos/HeadlessChaos.Target.cs
+++ b/Engine/Source/Programs/HeadlessChaos/HeadlessChaos.Target.cs
@@ -24,5 +24,7 @@ public class HeadlessChaosTarget : TargetRules
 
         // UnrealHeaderTool is a console application, not a Windows app (sets entry point to main(), instead of WinMain())
         bIsBuildingConsoleApplication = true;
-    }
+
+		GlobalDefinitions.Add("CHAOS_SERIALIZE_OUT=1");
+	}
 }
diff --git a/Engine/Source/Programs/HeadlessChaos/Private/HeadlessChaosTestParticleHandle.cpp b/Engine/Source/Programs/HeadlessChaos/Private/HeadlessChaosTestParticleHandle.cpp
--- a/Engine/Source/Programs/HeadlessChaos/Private/HeadlessChaosTestParticleHandle.cpp
+++ b/Engine/Source/Programs/HeadlessChaos/Private/HeadlessChaosTestParticleHandle.cpp
@@ -524,12 +524,16 @@ namespace ChaosTest {
 		// because of this, we must be able to compare (external, null) == (external, null) 
 		// and also (external, null) == (external, internal)
 
+		TPBDRigidsSOAs<FReal,3> SOAs;
 
 		auto GTParticle = TGeometryParticle<FReal, 3>::CreateParticle();
+		//fake unique assignment like we would for solver
+		GTParticle->SetUniqueIdx(SOAs.GetUniqueIndices().GenerateUniqueIdx());
+
 		TAccelerationStructureHandle<FReal, 3> ExternalOnlyHandle(GTParticle.Get());
 
-		TPBDRigidsSOAs<FReal, 3> SOAs;
-		auto Particles = SOAs.CreateStaticParticles(1);
+		FUniqueIdx Idx = GTParticle->UniqueIdx();
+		auto Particles = SOAs.CreateStaticParticles(1, &Idx);
 		TAccelerationStructureHandle<FReal, 3> ExternalInternalHandle(Particles[0], GTParticle.Get());
 
 		TAccelerationStructureHandle<FReal, 3> InternalOnlyHandle(Particles[0], nullptr);
diff --git a/Engine/Source/Runtime/Core/Public/UObject/ExternalPhysicsCustomObjectVersion.h b/Engine/Source/Runtime/Core/Public/UObject/ExternalPhysicsCustomObjectVersion.h
--- a/Engine/Source/Runtime/Core/Public/UObject/ExternalPhysicsCustomObjectVersion.h
+++ b/Engine/Source/Runtime/Core/Public/UObject/ExternalPhysicsCustomObjectVersion.h
@@ -110,6 +110,9 @@ struct CORE_API FExternalPhysicsCustomObjectVersion
 
 		// TriangleMesh has map from internal face indices to external.
 		TriangleMeshHasFaceIndexMap,
+
+		// Acceleration structures use unique payload idx
+		UniquePayloadIdx,
 		
 		// -----<new versions can be added above this line>-------------------------------------------------
 		VersionPlusOne,
diff --git a/Engine/Source/Runtime/Engine/Private/PhysicsEngine/Experimental/ChaosDerivedData.cpp b/Engine/Source/Runtime/Engine/Private/PhysicsEngine/Experimental/ChaosDerivedData.cpp
--- a/Engine/Source/Runtime/Engine/Private/PhysicsEngine/Experimental/ChaosDerivedData.cpp
+++ b/Engine/Source/Runtime/Engine/Private/PhysicsEngine/Experimental/ChaosDerivedData.cpp
@@ -24,7 +24,7 @@ const TCHAR* FChaosDerivedDataCooker::GetPluginName() const
 
 const TCHAR* FChaosDerivedDataCooker::GetVersionString() const
 {
-	return TEXT("2C093F83EC7F49A382170F519868BBC6");
+	return TEXT("6D0B86653D9846D6A8760B75E1034322");
 }
 
 FString FChaosDerivedDataCooker::GetPluginSpecificCacheKeySuffix() const
diff --git a/Engine/Source/Runtime/Engine/Private/PhysicsEngine/Experimental/PhysScene_Chaos.cpp b/Engine/Source/Runtime/Engine/Private/PhysicsEngine/Experimental/PhysScene_Chaos.cpp
--- a/Engine/Source/Runtime/Engine/Private/PhysicsEngine/Experimental/PhysScene_Chaos.cpp
+++ b/Engine/Source/Runtime/Engine/Private/PhysicsEngine/Experimental/PhysScene_Chaos.cpp
@@ -613,7 +613,7 @@ void FPhysScene_Chaos::AddObject(UPrimitiveComponent* Component, FFieldSystemPhy
 void FPhysScene_Chaos::RemoveActorFromAccelerationStructure(FPhysicsActorHandle& Actor)
 {
 #if WITH_CHAOS
-	if (GetSpacialAcceleration())
+	if (GetSpacialAcceleration() && Actor->UniqueIdx().IsValid())
 	{
 		ExternalDataLock.WriteLock();
 		Chaos::TAccelerationStructureHandle<float, 3> AccelerationHandle(Actor);
diff --git a/Engine/Source/Runtime/Engine/Private/PhysicsEngine/ImmediatePhysics/ImmediatePhysicsChaos/ImmediatePhysicsActorHandle_Chaos.cpp b/Engine/Source/Runtime/Engine/Private/PhysicsEngine/ImmediatePhysics/ImmediatePhysicsChaos/ImmediatePhysicsActorHandle_Chaos.cpp
--- a/Engine/Source/Runtime/Engine/Private/PhysicsEngine/ImmediatePhysics/ImmediatePhysicsChaos/ImmediatePhysicsActorHandle_Chaos.cpp
+++ b/Engine/Source/Runtime/Engine/Private/PhysicsEngine/ImmediatePhysics/ImmediatePhysicsChaos/ImmediatePhysicsActorHandle_Chaos.cpp
@@ -276,13 +276,13 @@ namespace ImmediatePhysics_Chaos
 			switch (ActorType)
 			{
 			case EActorType::StaticActor:
-				ParticleHandle = Particles.CreateStaticParticles(1, TGeometryParticleParameters<FReal, Dimensions>())[0];
+				ParticleHandle = Particles.CreateStaticParticles(1, nullptr, TGeometryParticleParameters<FReal, Dimensions>())[0];
 				break;
 			case EActorType::KinematicActor:
-				ParticleHandle = Particles.CreateKinematicParticles(1, TKinematicGeometryParticleParameters<FReal, Dimensions>())[0];
+				ParticleHandle = Particles.CreateKinematicParticles(1, nullptr, TKinematicGeometryParticleParameters<FReal, Dimensions>())[0];
 				break;
 			case EActorType::DynamicActor:
-				ParticleHandle = Particles.CreateDynamicParticles(1, TPBDRigidParticleParameters<FReal, Dimensions>())[0];
+				ParticleHandle = Particles.CreateDynamicParticles(1, nullptr, TPBDRigidParticleParameters<FReal, Dimensions>())[0];
 				break;
 			}
 
diff --git a/Engine/Source/Runtime/Experimental/Chaos/Private/Chaos/GeometryParticles.cpp b/Engine/Source/Runtime/Experimental/Chaos/Private/Chaos/GeometryParticles.cpp
--- a/Engine/Source/Runtime/Experimental/Chaos/Private/Chaos/GeometryParticles.cpp
+++ b/Engine/Source/Runtime/Experimental/Chaos/Private/Chaos/GeometryParticles.cpp
@@ -168,19 +168,6 @@ namespace Chaos
 		Ar << SerializableGeometryParticles;
 	}
 
-	template <typename T, int d, EGeometryParticlesSimType SimType>
-	void TGeometryParticlesImp<T, d, SimType>::SerializeHashResultHelper(FChaosArchive& Ar, TGeometryParticle<T, d>* Particle)
-	{
-		if (Particle)
-		{
-			MHashResult.Add(Particle->GetHashResultLowLevel());
-		}
-		else
-		{
-			MHashResult.Add(FMath::RandHelper(TNumericLimits<uint32>::Max()));
-		}
-	}
-
 	
 	template class TGeometryParticlesImp<float, 3, EGeometryParticlesSimType::RigidBodySim>;
 	template class TGeometryParticlesImp<float, 3, EGeometryParticlesSimType::Other>;
diff --git a/Engine/Source/Runtime/Experimental/Chaos/Public/Chaos/AABBTree.h b/Engine/Source/Runtime/Experimental/Chaos/Public/Chaos/AABBTree.h
--- a/Engine/Source/Runtime/Experimental/Chaos/Public/Chaos/AABBTree.h
+++ b/Engine/Source/Runtime/Experimental/Chaos/Public/Chaos/AABBTree.h
@@ -518,18 +518,8 @@ public:
 
 		if (bSerializePayloadToInfo)
 		{
-			TArray<TPayloadType> Payloads;
-			if (!Ar.IsLoading())
-			{
-				PayloadToInfo.GenerateKeyArray(Payloads);
-			}
-			Ar << Payloads;
+			Ar << PayloadToInfo;
 
-			for (auto Payload : Payloads)
-			{
-				auto& Info = PayloadToInfo.FindOrAdd(Payload);
-				Ar << Info;
-			}
 			if (!bMutable)	//if immutable empty this even if we had to serialize it in for backwards compat
 			{
 				PayloadToInfo.Empty();
@@ -977,7 +967,8 @@ private:
 	TArray<TLeafType> Leaves;
 	TArray<FElement> DirtyElements;
 	TArray<FElement> GlobalPayloads;
-	TMap<TPayloadType, FAABBTreePayloadInfo> PayloadToInfo;
+	TArrayAsMap<TPayloadType, FAABBTreePayloadInfo> PayloadToInfo;
+
 	int32 MaxChildrenInLeaf;
 	int32 MaxTreeDepth;
 	T MaxPayloadBounds;
diff --git a/Engine/Source/Runtime/Experimental/Chaos/Public/Chaos/BoundingVolume.h b/Engine/Source/Runtime/Experimental/Chaos/Public/Chaos/BoundingVolume.h
--- a/Engine/Source/Runtime/Experimental/Chaos/Public/Chaos/BoundingVolume.h
+++ b/Engine/Source/Runtime/Experimental/Chaos/Public/Chaos/BoundingVolume.h
@@ -338,18 +338,8 @@ public:
 		Ar << MDirtyElements;
 		Ar << bIsEmpty;
 
-		TArray<TPayloadType> Payloads;
-		if (!Ar.IsLoading())
-		{
-			MPayloadInfo.GenerateKeyArray(Payloads);
-		}
-		Ar << Payloads;
-
-		for (auto Payload : Payloads)
-		{
-			auto& Info = MPayloadInfo.FindOrAdd(Payload);
-			Ar << Info;
-		}
+		Ar << MPayloadInfo;
+		
 
 	}
 
@@ -775,7 +765,7 @@ private:
 
 		const TVector<int32, d> StartIndex = MGrid.ClampIndex(MGrid.Cell(QueryBounds.Min()));
 		const TVector<int32, d> EndIndex = MGrid.ClampIndex(MGrid.Cell(QueryBounds.Max()));
-		TSet<TPayloadType> InstancesSeen;
+		TSet<FUniqueIdx> InstancesSeen;
 
 		for (int32 X = StartIndex[0]; X <= EndIndex[0]; ++X)
 		{
@@ -788,11 +778,11 @@ private:
 					{
 						if (bPruneDuplicates)
 						{
-							if (InstancesSeen.Contains(Elem.Payload))
+							if (InstancesSeen.Contains(GetUniqueIdx(Elem.Payload)))
 							{
 								continue;
 							}
-							InstancesSeen.Add(Elem.Payload);
+							InstancesSeen.Add(GetUniqueIdx(Elem.Payload));
 						}
 						const TAABB<T, d>& InstanceBounds = Elem.Bounds;
 						if (QueryBounds.Intersects(InstanceBounds))
@@ -1110,8 +1100,10 @@ private:
 			}
 		}
 
-
-		Algo::Sort(Intersections);
+		Algo::Sort(Intersections,[](const TPayloadType& A,const TPayloadType& B)
+		{
+			return GetUniqueIdx(A) < GetUniqueIdx(B);
+		});
 
 		for (int32 i = Intersections.Num() - 1; i > 0; i--)
 		{
@@ -1174,7 +1166,7 @@ private:
 	TUniformGrid<T, d> MGrid;
 	TArrayND<TArray<FCellElement>, d> MElements;
 	TArray<FCellElement> MDirtyElements;
-	TMap<TPayloadType, FPayloadInfo> MPayloadInfo;
+	TArrayAsMap<TPayloadType, FPayloadInfo> MPayloadInfo;
 	T MaxPayloadBounds;
 	bool bIsEmpty;
 };
diff --git a/Engine/Source/Runtime/Experimental/Chaos/Public/Chaos/GeometryParticles.h b/Engine/Source/Runtime/Experimental/Chaos/Public/Chaos/GeometryParticles.h
--- a/Engine/Source/Runtime/Experimental/Chaos/Public/Chaos/GeometryParticles.h
+++ b/Engine/Source/Runtime/Experimental/Chaos/Public/Chaos/GeometryParticles.h
@@ -95,6 +95,7 @@ namespace Chaos
 		    : TParticles<T, d>()
 		{
 			MParticleType = EParticleType::Static;
+			TArrayCollection::AddArray(&MUniqueIdx);
 			TArrayCollection::AddArray(&MR);
 			TArrayCollection::AddArray(&MGeometry);
 			TArrayCollection::AddArray(&MSharedGeometry);
@@ -106,7 +107,6 @@ namespace Chaos
 			TArrayCollection::AddArray(&MWorldSpaceInflatedBounds);
 			TArrayCollection::AddArray(&MHasBounds);
 			TArrayCollection::AddArray(&MSpatialIdx);
-			TArrayCollection::AddArray(&MHashResult);
 #if CHAOS_CHECKED
 			TArrayCollection::AddArray(&MDebugName);
 #endif
@@ -121,6 +121,7 @@ namespace Chaos
 		TGeometryParticlesImp(const TGeometryParticlesImp<T, d, SimType>& Other) = delete;
 		CHAOS_API TGeometryParticlesImp(TGeometryParticlesImp<T, d, SimType>&& Other)
 		    : TParticles<T, d>(MoveTemp(Other))
+			, MUniqueIdx(MoveTemp(Other.MUniqueIdx))
 			, MR(MoveTemp(Other.MR))
 			, MGeometry(MoveTemp(Other.MGeometry))
 			, MSharedGeometry(MoveTemp(Other.MSharedGeometry))
@@ -133,12 +134,12 @@ namespace Chaos
 			, MWorldSpaceInflatedBounds(MoveTemp(Other.MWorldSpaceInflatedBounds))
 			, MHasBounds(MoveTemp(Other.MHasBounds))
 			, MSpatialIdx(MoveTemp(Other.MSpatialIdx))
-			, MHashResult(MoveTemp(Other.MHashResult))
 #if CHAOS_DETERMINISTIC
 			, MParticleIDs(MoveTemp(Other.MParticleIDs))
 #endif
 		{
 			MParticleType = EParticleType::Static;
+			TArrayCollection::AddArray(&MUniqueIdx);
 			TArrayCollection::AddArray(&MR);
 			TArrayCollection::AddArray(&MGeometry);
 			TArrayCollection::AddArray(&MSharedGeometry);
@@ -149,7 +150,6 @@ namespace Chaos
 			TArrayCollection::AddArray(&MWorldSpaceInflatedBounds);
 			TArrayCollection::AddArray(&MHasBounds);
 			TArrayCollection::AddArray(&MSpatialIdx);
-			TArrayCollection::AddArray(&MHashResult);
 #if CHAOS_DETERMINISTIC
 			TArrayCollection::AddArray(&MParticleIDs);
 #endif
@@ -170,6 +170,7 @@ namespace Chaos
 		    : TParticles<T, d>(MoveTemp(Other))
 		{
 			MParticleType = EParticleType::Static;
+			TArrayCollection::AddArray(&MUniqueIdx);
 			TArrayCollection::AddArray(&MR);
 			TArrayCollection::AddArray(&MGeometry);
 			TArrayCollection::AddArray(&MSharedGeometry);
@@ -180,7 +181,6 @@ namespace Chaos
 			TArrayCollection::AddArray(&MWorldSpaceInflatedBounds);
 			TArrayCollection::AddArray(&MHasBounds);
 			TArrayCollection::AddArray(&MSpatialIdx);
-			TArrayCollection::AddArray(&MHashResult);
 #if CHAOS_DETERMINISTIC
 			TArrayCollection::AddArray(&MParticleIDs);
 #endif
@@ -200,7 +200,8 @@ namespace Chaos
 
 		CHAOS_API const TRotation<T, d>& R(const int32 Index) const { return MR[Index]; }
 		CHAOS_API TRotation<T, d>& R(const int32 Index) { return MR[Index]; }
-
+		FUniqueIdx UniqueIdx(const int32 Index) const { return MUniqueIdx[Index]; }
+		FUniqueIdx& UniqueIdx(const int32 Index) { return MUniqueIdx[Index]; }
 		CHAOS_API TSerializablePtr<FImplicitObject> Geometry(const int32 Index) const { return MGeometry[Index]; }
 
 		CHAOS_API const TUniquePtr<FImplicitObject>& DynamicGeometry(const int32 Index) const { return MDynamicGeometry[Index]; }
@@ -271,16 +272,6 @@ namespace Chaos
 			return MSpatialIdx[Index];
 		}
 
-		CHAOS_API uint32 HashResultLowLevel(const int32 Index) const
-		{
-			return MHashResult[Index];
-		}
-
-		CHAOS_API uint32& HashResultLowLevel(const int32 Index)
-		{
-			return MHashResult[Index];
-		}
-
 #if CHAOS_CHECKED
 		const FName& DebugName(const int32 Index) const
 		{
@@ -326,7 +317,7 @@ namespace Chaos
 		FString ToString(int32 index) const
 		{
 			FString BaseString = TParticles<T, d>::ToString(index);
-			return FString::Printf(TEXT("%s, MR:%s, MGeometry:%s, IsDynamic:%d"), *BaseString, *R(index).ToString(), (Geometry(index) ? *(Geometry(index)->ToString()) : TEXT("none")), (DynamicGeometry(index) != nullptr));
+			return FString::Printf(TEXT("%s, MUniqueIdx:%d MR:%s, MGeometry:%s, IsDynamic:%d"), *BaseString, UniqueIdx(index).Idx, *R(index).ToString(), (Geometry(index) ? *(Geometry(index)->ToString()) : TEXT("none")), (DynamicGeometry(index) != nullptr));
 		}
 
 		CHAOS_API virtual void Serialize(FChaosArchive& Ar)
@@ -391,14 +382,7 @@ namespace Chaos
 
 			if (Ar.CustomVer(FExternalPhysicsCustomObjectVersion::GUID) < FExternalPhysicsCustomObjectVersion::SerializeHashResult)
 			{
-				for (const auto& Particle : MGeometryParticle)
-				{
-					SerializeHashResultHelper(Ar, Particle);
-				}
-			}
-			else
-			{
-				Ar << MHashResult;
+				//no longer care about hash so don't read it and don't do anything
 			}
 		}
 
@@ -423,6 +407,7 @@ namespace Chaos
 		EParticleType MParticleType;
 
 	private:
+		TArrayCollectionArray<FUniqueIdx> MUniqueIdx;
 		TArrayCollectionArray<TRotation<T, d>> MR;
 		// MGeometry contains raw ptrs to every entry in both MSharedGeometry and MDynamicGeometry.
 		// It may also contain raw ptrs to geometry which is managed outside of Chaos.
@@ -441,7 +426,6 @@ namespace Chaos
 		TArrayCollectionArray<TAABB<T, d>> MWorldSpaceInflatedBounds;
 		TArrayCollectionArray<bool> MHasBounds;
 		TArrayCollectionArray<FSpatialAccelerationIdx> MSpatialIdx;
-		TArrayCollectionArray<uint32> MHashResult;
 
 		void UpdateShapesArray(const int32 Index)
 		{
@@ -478,7 +462,6 @@ namespace Chaos
 		friend class TGeometryParticlesImp;
 
 		CHAOS_API void SerializeGeometryParticleHelper(FChaosArchive& Ar, TGeometryParticlesImp<T, d, EGeometryParticlesSimType::RigidBodySim>* GeometryParticles);
-		CHAOS_API void SerializeHashResultHelper(FChaosArchive& Ar, TGeometryParticle<T, d>* Particle);
 		
 		void SerializeGeometryParticleHelper(FChaosArchive& Ar, TGeometryParticlesImp<T, d, EGeometryParticlesSimType::Other>* GeometryParticles)
 		{
diff --git a/Engine/Source/Runtime/Experimental/Chaos/Public/Chaos/GeometryParticlesfwd.h b/Engine/Source/Runtime/Experimental/Chaos/Public/Chaos/GeometryParticlesfwd.h
--- a/Engine/Source/Runtime/Experimental/Chaos/Public/Chaos/GeometryParticlesfwd.h
+++ b/Engine/Source/Runtime/Experimental/Chaos/Public/Chaos/GeometryParticlesfwd.h
@@ -43,4 +43,21 @@ namespace Chaos
 	{
 		return Ar << (uint16&)Idx;
 	}
+
+
+	struct FUniqueIdx
+	{
+		int32 Idx;
+		FUniqueIdx(): Idx(INDEX_NONE){}
+		explicit FUniqueIdx(int32 InIdx): Idx(InIdx){}
+
+		bool IsValid() const { return Idx != INDEX_NONE; }
+		bool operator<(const FUniqueIdx& Other) const { return Idx < Other.Idx; }
+		bool operator==(const FUniqueIdx& Other) const { return Idx == Other.Idx; }
+	};
+
+	FORCEINLINE uint32 GetTypeHash(const FUniqueIdx& Unique)
+	{
+		return ::GetTypeHash(Unique.Idx);
+	}
 }
\ No newline at end of file
diff --git a/Engine/Source/Runtime/Experimental/Chaos/Public/Chaos/ISpatialAcceleration.h b/Engine/Source/Runtime/Experimental/Chaos/Public/Chaos/ISpatialAcceleration.h
--- a/Engine/Source/Runtime/Experimental/Chaos/Public/Chaos/ISpatialAcceleration.h
+++ b/Engine/Source/Runtime/Experimental/Chaos/Public/Chaos/ISpatialAcceleration.h
@@ -172,6 +172,28 @@ enum ESpatialAcceleration
 
 using SpatialAccelerationType = uint8;	//see ESpatialAcceleration. Projects can add their own custom types by using enum values higher than ESpatialAcceleration::Unknown
 
+template <typename TPayload>
+typename TEnableIf<!TIsPointer<TPayload>::Value, FUniqueIdx>::Type GetUniqueIdx(const TPayload& Payload)
+{
+	const FUniqueIdx Idx = Payload.UniqueIdx();
+	ensure(Idx.IsValid());
+	return Idx;
+}
+
+template <typename TPayload>
+typename TEnableIf<TIsPointer<TPayload>::Value,FUniqueIdx>::Type GetUniqueIdx(const TPayload& Payload)
+{
+	const FUniqueIdx Idx = Payload->UniqueIdx();
+	ensure(Idx.IsValid());
+	return Idx;
+}
+
+FORCEINLINE FUniqueIdx GetUniqueIdx(const int32 Payload)
+{
+	ensure(Payload >=0);	//-1 idx implies it was never set
+	return FUniqueIdx(Payload);
+}
+
 
 template <typename TPayloadType, typename T>
 struct TPayloadBoundsElement
@@ -191,6 +213,11 @@ struct TPayloadBoundsElement
 	bool HasBoundingBox() const { return true; }
 
 	const TAABB<T, 3>& BoundingBox() const { return Bounds; }
+
+	FUniqueIdx UniqueIdx() const
+	{
+		return ::Chaos::GetUniqueIdx(Payload);
+	}
 };
 
 template <typename TPayloadType, typename T>
@@ -337,4 +364,146 @@ private:
 	ISpatialVisitor<TPayloadType, T>& Visitor;
 };
 
+#ifndef CHAOS_SERIALIZE_OUT
+#define CHAOS_SERIALIZE_OUT WITH_EDITOR
+#endif
+
+//Provides a TMap like API but backed by a dense array. The keys provided must implement GetUniqueIdx
+template <typename TKey, typename TValue>
+class TArrayAsMap
+{
+public:
+	TValue* Find(const TKey& Key)
+	{
+		const int32 Idx = GetUniqueIdx(Key).Idx;
+		if(Idx < Entries.Num() && Entries[Idx].bSet)
+		{
+			return &Entries[Idx].Value;
+		}
+		return nullptr;
+	}
+
+	TValue& FindChecked(const TKey& Key)
+	{
+		return Entries[GetUniqueIdx(Key).Idx].Value;
+	}
+
+	TValue& FindOrAdd(const TKey& Key)
+	{
+		if(TValue* Elem = Find(Key))
+		{
+			return *Elem;
+		}
+
+		return Add(Key);
+	}
+
+	void Empty()
+	{
+		Entries.Empty();
+	}
+
+	TValue& Add(const TKey& Key)
+	{
+		const int32 Idx = GetUniqueIdx(Key).Idx;
+		if(Idx >= Entries.Num())
+		{
+			const int32 NumToAdd = Idx + 1 - Entries.Num();
+			Entries.AddDefaulted(NumToAdd);
+#if CHAOS_SERIALIZE_OUT
+			KeysToSerializeOut.AddDefaulted(NumToAdd);
+#endif
+		}
+
+		ensure(Entries[Idx].bSet == false);	//element already added
+		Entries[Idx].bSet = true;
+
+#if CHAOS_SERIALIZE_OUT
+		KeysToSerializeOut[Idx] = Key;
+#endif
+
+		return Entries[Idx].Value;
+	}
+
+	void Add(const TKey& Key, const TValue& Value)
+	{
+		Add(Key) = Value;
+	}
+
+	void Remove(const TKey& Key)
+	{
+		const int32 Idx = GetUniqueIdx(Key).Idx;
+		Entries[Idx].bSet = false;
+#if CHAOS_SERIALIZE_OUT
+		KeysToSerializeOut[Idx] = TKey();
+#endif
+	}
+
+	void Reset()
+	{
+		Entries.Reset();
+#if CHAOS_SERIALIZE_OUT 
+		KeysToSerializeOut.Reset();
+#endif
+	}
+
+	void Serialize(FChaosArchive& Ar)
+	{
+		bool bCanSerialize = Ar.IsLoading();
+#if CHAOS_SERIALIZE_OUT 
+		bCanSerialize = true;
+#endif
+
+		if(bCanSerialize)
+		{
+			TArray<TKey> DirectKeys;
+			Ar << DirectKeys;
+
+			for(auto& Key : DirectKeys)
+			{
+				TValue& Value = Add(Key);
+				Ar << Value;
+			}
+		}
+		else
+		{
+			ensure(false);	//can't serialize out, if you are trying to serialize for perf/debug set CHAOS_SERIALIZE_OUT to 1 
+		}
+	}
+
+private:
+
+	struct FEntry
+	{
+#if CHAOS_SERIALIZE_OUT
+		//The indices are generated at runtime, so there's no way to serialize them directly
+		//Because of that we serialize the actual key which we can find, and then at runtime we use its transient index
+		TKey KeyToSerializeOut;
+#endif
+		TValue Value;
+		bool bSet;
+
+		FEntry()
+			: bSet(false)
+		{
+
+		}
+	};
+
+	TArray<FEntry> Entries;
+
+#if CHAOS_SERIALIZE_OUT
+	//The indices are generated at runtime, so there's no way to serialize them directly
+	//Because of that we serialize the actual key which we can find, and then at runtime we use its transient index
+	TArray<TKey> KeysToSerializeOut;
+#endif
+};
+
+template <typename TKey, typename TValue>
+FChaosArchive& operator<< (FChaosArchive& Ar, TArrayAsMap<TKey, TValue>& Map)
+{
+	Map.Serialize(Ar);
+	return Ar;
+}
+
 }
diff --git a/Engine/Source/Runtime/Experimental/Chaos/Public/Chaos/PBDRigidsEvolution.h b/Engine/Source/Runtime/Experimental/Chaos/Public/Chaos/PBDRigidsEvolution.h
--- a/Engine/Source/Runtime/Experimental/Chaos/Public/Chaos/PBDRigidsEvolution.h
+++ b/Engine/Source/Runtime/Experimental/Chaos/Public/Chaos/PBDRigidsEvolution.h
@@ -228,9 +228,9 @@ class TPBDRigidsEvolutionBase
 	CHAOS_API TPBDRigidsEvolutionBase(TPBDRigidsSOAs<T, d>& InParticles, int32 InNumIterations = 1, int32 InNumPushOutIterations = 1, bool InIsSingleThreaded = false);
 	CHAOS_API virtual ~TPBDRigidsEvolutionBase();
 
-	CHAOS_API TArray<TGeometryParticleHandle<T, d>*> CreateStaticParticles(int32 NumParticles, const TGeometryParticleParameters<T, d>& Params = TGeometryParticleParameters<T, d>())
+	CHAOS_API TArray<TGeometryParticleHandle<T, d>*> CreateStaticParticles(int32 NumParticles, const FUniqueIdx* ExistingIndices = nullptr, const TGeometryParticleParameters<T, d>& Params = TGeometryParticleParameters<T, d>())
 	{
-		auto NewParticles = Particles.CreateStaticParticles(NumParticles, Params);
+		auto NewParticles = Particles.CreateStaticParticles(NumParticles, ExistingIndices, Params);
 		for (auto& Particle : NewParticles)
 		{
 			DirtyParticle(*Particle);
@@ -238,9 +238,9 @@ class TPBDRigidsEvolutionBase
 		return NewParticles;
 	}
 
-	CHAOS_API TArray<TKinematicGeometryParticleHandle<T, d>*> CreateKinematicParticles(int32 NumParticles, const TKinematicGeometryParticleParameters<T, d>& Params = TKinematicGeometryParticleParameters<T, d>())
+	CHAOS_API TArray<TKinematicGeometryParticleHandle<T, d>*> CreateKinematicParticles(int32 NumParticles, const FUniqueIdx* ExistingIndices = nullptr, const TKinematicGeometryParticleParameters<T, d>& Params = TKinematicGeometryParticleParameters<T, d>())
 	{
-		auto NewParticles = Particles.CreateKinematicParticles(NumParticles, Params);
+		auto NewParticles = Particles.CreateKinematicParticles(NumParticles, ExistingIndices, Params);
 		for (auto& Particle : NewParticles)
 		{
 			DirtyParticle(*Particle);
@@ -248,9 +248,9 @@ class TPBDRigidsEvolutionBase
 		return NewParticles;
 	}
 
-	CHAOS_API TArray<TPBDRigidParticleHandle<T, d>*> CreateDynamicParticles(int32 NumParticles, const TPBDRigidParticleParameters<T, d>& Params = TPBDRigidParticleParameters<T, d>())
+	CHAOS_API TArray<TPBDRigidParticleHandle<T, d>*> CreateDynamicParticles(int32 NumParticles, const FUniqueIdx* ExistingIndices = nullptr, const TPBDRigidParticleParameters<T, d>& Params = TPBDRigidParticleParameters<T, d>())
 	{
-		auto NewParticles = Particles.CreateDynamicParticles(NumParticles, Params);
+		auto NewParticles = Particles.CreateDynamicParticles(NumParticles, ExistingIndices, Params);
 		for (auto& Particle : NewParticles)
 		{
 			DirtyParticle(*Particle);
@@ -258,9 +258,9 @@ class TPBDRigidsEvolutionBase
 		return NewParticles;
 	}
 
-	CHAOS_API TArray<TPBDRigidClusteredParticleHandle<T, d>*> CreateClusteredParticles(int32 NumParticles, const TPBDRigidParticleParameters<T, d>& Params = TPBDRigidParticleParameters<T, d>())
+	CHAOS_API TArray<TPBDRigidClusteredParticleHandle<T, d>*> CreateClusteredParticles(int32 NumParticles,const FUniqueIdx* ExistingIndices = nullptr,  const TPBDRigidParticleParameters<T, d>& Params = TPBDRigidParticleParameters<T, d>())
 	{
-		auto NewParticles = Particles.CreateClusteredParticles(NumParticles, Params);
+		auto NewParticles = Particles.CreateClusteredParticles(NumParticles, ExistingIndices, Params);
 		for (auto& Particle : NewParticles)
 		{
 			DirtyParticle(*Particle);
@@ -537,6 +537,12 @@ class TPBDRigidsEvolutionBase
 
 	void Serialize(FChaosArchive& Ar);
 
+	FUniqueIdx GenerateUniqueIdx()
+	{
+		//NOTE: this should be thread safe since evolution has already been initialized on GT
+		return Particles.GetUniqueIndices().GenerateUniqueIdx();
+	}
+
 protected:
 	int32 NumConstraints() const
 	{
diff --git a/Engine/Source/Runtime/Experimental/Chaos/Public/Chaos/PBDRigidsEvolutionGBF.h b/Engine/Source/Runtime/Experimental/Chaos/Public/Chaos/PBDRigidsEvolutionGBF.h
--- a/Engine/Source/Runtime/Experimental/Chaos/Public/Chaos/PBDRigidsEvolutionGBF.h
+++ b/Engine/Source/Runtime/Experimental/Chaos/Public/Chaos/PBDRigidsEvolutionGBF.h
@@ -49,6 +49,7 @@ namespace Chaos
 		using Base::ParticleUpdatePosition;
 		using Base::SetParticleUpdatePositionFunction;
 		using Base::SetParticleUpdateVelocityFunction;
+		using Base::GenerateUniqueIdx;
 		using typename Base::FForceRule;
 		using FGravityForces = TPerParticleGravity<T, d>;
 		using FCollisionConstraints = TPBDCollisionConstraints<T, d>;
diff --git a/Engine/Source/Runtime/Experimental/Chaos/Public/Chaos/PBDRigidsSOAs.h b/Engine/Source/Runtime/Experimental/Chaos/Public/Chaos/PBDRigidsSOAs.h
--- a/Engine/Source/Runtime/Experimental/Chaos/Public/Chaos/PBDRigidsSOAs.h
+++ b/Engine/Source/Runtime/Experimental/Chaos/Public/Chaos/PBDRigidsSOAs.h
@@ -6,6 +6,73 @@
 
 namespace Chaos
 {
+class FParticleUniqueIndices
+{
+public:
+	FParticleUniqueIndices()
+	{
+		//Note: tune this so that all allocation is done at initialization
+		AddPageAndAcquireNextId(/*bAcquireNextId = */ false);
+	}
+
+	FUniqueIdx GenerateUniqueIdx()
+	{
+		if(FUniqueIdx* Idx = FreeIndices.Pop())
+		{
+			return *Idx;
+		}
+		else
+		{
+			//Should never really hit this, but need fallback
+			FScopeLock CSLock(&CS);
+			if(FUniqueIdx* TryIdx = FreeIndices.Pop())
+			{
+				//just in case a page was added before we got the lock
+				return *TryIdx;
+			}
+			return FUniqueIdx(AddPageAndAcquireNextId(/*bAcquireNextId = */ true));
+		}
+	}
+
+	void ReleaseIdx(FUniqueIdx Unique)
+	{
+		ensure(Unique.IsValid());
+		int32 PageIdx = Unique.Idx / IndicesPerPage;
+		int32 Entry = Unique.Idx % IndicesPerPage;
+		FUniqueIdxPage& Page = *Pages[PageIdx];
+		FreeIndices.Push(&Page.Indices[Entry]);
+	}
+
+private:
+
+	int32 AddPageAndAcquireNextId(bool bAcquireNextIdx)
+	{
+		//Note: this should never really be called post initialization
+		TUniquePtr<FUniqueIdxPage> Page = MakeUnique<FUniqueIdxPage>();
+		const int32 PageIdx = Pages.Num();
+		int32 FirstIdxInPage = PageIdx * IndicesPerPage;
+		Page->Indices[0] = FUniqueIdx(FirstIdxInPage);
+
+		//If we acquire next id we avoid pushing it into the queue
+		for(int32 Count = bAcquireNextIdx ? 1 : 0; Count < IndicesPerPage; Count++)
+		{
+			Page->Indices[Count] = FUniqueIdx(FirstIdxInPage + Count);
+			FreeIndices.Push(&Page->Indices[Count]);
+		}
+
+		Pages.Emplace(MoveTemp(Page));
+		return bAcquireNextIdx ? FirstIdxInPage : INDEX_NONE;
+	}
+
+	static constexpr int32 IndicesPerPage = 100000;
+	struct FUniqueIdxPage
+	{
+		FUniqueIdx Indices[IndicesPerPage];
+	};
+	TArray<TUniquePtr<FUniqueIdxPage>> Pages;
+	TLockFreePointerListFIFO<FUniqueIdx,0> FreeIndices;
+	FCriticalSection CS;
+};
 
 template <typename T, int d>
 class TPBDRigidsSOAs
@@ -44,21 +111,21 @@ public:
 		check(0);
 	}
 	
-	TArray<TGeometryParticleHandle<T, d>*> CreateStaticParticles(int32 NumParticles, const TGeometryParticleParameters<T, d>& Params = TGeometryParticleParameters<T, d>())
+	TArray<TGeometryParticleHandle<T, d>*> CreateStaticParticles(int32 NumParticles, const FUniqueIdx* ExistingIndices = nullptr, const TGeometryParticleParameters<T, d>& Params = TGeometryParticleParameters<T, d>())
 	{
-		auto Results =  CreateParticlesHelper<TGeometryParticleHandle<T, d>>(NumParticles, Params.bDisabled ? StaticDisabledParticles : StaticParticles, Params);
+		auto Results =  CreateParticlesHelper<TGeometryParticleHandle<T, d>>(NumParticles, ExistingIndices, Params.bDisabled ? StaticDisabledParticles : StaticParticles, Params);
 		UpdateViews();
 		return Results;
 	}
-	TArray<TKinematicGeometryParticleHandle<T, d>*> CreateKinematicParticles(int32 NumParticles, const TKinematicGeometryParticleParameters<T, d>& Params = TKinematicGeometryParticleParameters<T, d>())
+	TArray<TKinematicGeometryParticleHandle<T, d>*> CreateKinematicParticles(int32 NumParticles, const FUniqueIdx* ExistingIndices = nullptr,  const TKinematicGeometryParticleParameters<T, d>& Params = TKinematicGeometryParticleParameters<T, d>())
 	{
-		auto Results = CreateParticlesHelper<TKinematicGeometryParticleHandle<T, d>>(NumParticles, Params.bDisabled ? KinematicDisabledParticles : KinematicParticles, Params);
+		auto Results = CreateParticlesHelper<TKinematicGeometryParticleHandle<T, d>>(NumParticles, ExistingIndices, Params.bDisabled ? KinematicDisabledParticles : KinematicParticles, Params);
 		UpdateViews();
 		return Results;
 	}
-	TArray<TPBDRigidParticleHandle<T, d>*> CreateDynamicParticles(int32 NumParticles, const TPBDRigidParticleParameters<T, d>& Params = TPBDRigidParticleParameters<T, d>())
+	TArray<TPBDRigidParticleHandle<T, d>*> CreateDynamicParticles(int32 NumParticles, const FUniqueIdx* ExistingIndices = nullptr,  const TPBDRigidParticleParameters<T, d>& Params = TPBDRigidParticleParameters<T, d>())
 	{
-		auto Results = CreateParticlesHelper<TPBDRigidParticleHandle<T, d>>(NumParticles, Params.bDisabled ? DynamicDisabledParticles : DynamicParticles, Params);
+		auto Results = CreateParticlesHelper<TPBDRigidParticleHandle<T, d>>(NumParticles, ExistingIndices, Params.bDisabled ? DynamicDisabledParticles : DynamicParticles, Params);;
 
 		if (!Params.bStartSleeping)
 		{
@@ -82,9 +149,9 @@ public:
 	}
 
 	/** Used specifically by PBDRigidClustering. These have special properties for maintaining relative order, efficiently switching from kinematic to dynamic, disable to enable, etc... */
-	TArray<TPBDRigidClusteredParticleHandle<T, d>*> CreateClusteredParticles(int32 NumParticles, const TPBDRigidParticleParameters<T, d>& Params = TPBDRigidParticleParameters<T, d>())
+	TArray<TPBDRigidClusteredParticleHandle<T, d>*> CreateClusteredParticles(int32 NumParticles, const FUniqueIdx* ExistingIndices = nullptr,  const TPBDRigidParticleParameters<T, d>& Params = TPBDRigidParticleParameters<T, d>())
 	{
-		auto NewClustered = CreateParticlesHelper<TPBDRigidClusteredParticleHandle<T, d>>(NumParticles, ClusteredParticles, Params);
+		auto NewClustered = CreateParticlesHelper<TPBDRigidClusteredParticleHandle<T, d>>(NumParticles, ExistingIndices, ClusteredParticles, Params);
 		
 		if (!Params.bDisabled)
 		{
@@ -128,6 +195,11 @@ public:
 			}
 		}
 
+		//NOTE: This assumes that we are never creating a PT particle that is replicated to GT
+		//At the moment that is true, and it seems like we have enough mechanisms to avoid this direction
+		//If we want to support that, the UniqueIndex must be kept around until GT goes away
+		//This is hard to do, but would probably mean the ownership of the index is in the proxy
+		UniqueIndices.ReleaseIdx(Particle->UniqueIdx());
 		ParticleHandles.DestroyHandleSwap(Particle);
 
 		
@@ -303,6 +375,26 @@ public:
 			Ar << DynamicKinematicParticles;
 		}
 
+		{
+			//need to assign indices to everything
+			auto AssignIdxHelper = [&](const auto& Particles)
+			{
+				for(uint32 ParticleIdx = 0; ParticleIdx < Particles->Size(); ++ParticleIdx)
+				{
+					FUniqueIdx Unique = UniqueIndices.GenerateUniqueIdx();
+					Particles->UniqueIdx(ParticleIdx) = Unique;
+					Particles->GTGeometryParticle(ParticleIdx)->SetUniqueIdx(Unique);
+				}
+			};
+
+			AssignIdxHelper(StaticParticles);
+			AssignIdxHelper(StaticDisabledParticles);
+			AssignIdxHelper(KinematicParticles);
+			AssignIdxHelper(DynamicParticles);
+			AssignIdxHelper(DynamicParticles);
+			AssignIdxHelper(DynamicDisabledParticles);
+		}
+
 		ensure(ClusteredParticles->Size() == 0);	//not supported yet
 		//Ar << ClusteredParticles;
 		Ar << GeometryCollectionParticles;
@@ -460,9 +552,11 @@ public:
 	const auto& GetClusteredParticles() const { return *ClusteredParticles; }
 	auto& GetClusteredParticles() { return *ClusteredParticles; }
 
+	auto& GetUniqueIndices() { return UniqueIndices; }
+
 private:
 	template <typename TParticleHandleType, typename TParticles>
-	TArray<TParticleHandleType*> CreateParticlesHelper(int32 NumParticles, TUniquePtr<TParticles>& Particles, const TGeometryParticleParameters<T, d>& Params)
+	TArray<TParticleHandleType*> CreateParticlesHelper(int32 NumParticles, const FUniqueIdx* ExistingIndices,  TUniquePtr<TParticles>& Particles, const TGeometryParticleParameters<T, d>& Params)
 	{
 		const int32 ParticlesStartIdx = Particles->Size();
 		Particles->AddParticles(NumParticles);
@@ -482,6 +576,15 @@ private:
 			NewParticleHandle->ParticleID() = BiggestParticleID++;
 #endif
 			ReturnHandles[Count] = NewParticleHandle.Get();
+			//If unique indices are null it means there is no GT particle that already registered an ID, so create one
+			if(ExistingIndices)
+			{
+				ReturnHandles[Count]->SetUniqueIdx(ExistingIndices[Count]);
+			}
+			else
+			{
+				ReturnHandles[Count]->SetUniqueIdx(UniqueIndices.GenerateUniqueIdx());
+			}
 			ParticleHandles.Handle(HandleIdx) = MoveTemp(NewParticleHandle);
 		}
 
@@ -637,6 +740,8 @@ private:
 	//Auxiliary data synced with particle handles
 	TGeometryParticleHandles<T, d> ParticleHandles;
 
+	FParticleUniqueIndices UniqueIndices;
+
 #if CHAOS_DETERMINISTIC
 	int32 BiggestParticleID;
 #endif
diff --git a/Engine/Source/Runtime/Experimental/Chaos/Public/Chaos/ParticleDirtyFlags.h b/Engine/Source/Runtime/Experimental/Chaos/Public/Chaos/ParticleDirtyFlags.h
--- a/Engine/Source/Runtime/Experimental/Chaos/Public/Chaos/ParticleDirtyFlags.h
+++ b/Engine/Source/Runtime/Experimental/Chaos/Public/Chaos/ParticleDirtyFlags.h
@@ -35,8 +35,8 @@ namespace Chaos
 		AngularImpulse			= 1 << 23,
 		GravityEnabled			= 1 << 24,
 		SpatialIdx				= 1 << 25,
-		HashResult				= 1 << 26,
-		ShapeDisableCollision	= 1 << 27
+		UniqueIdx				= 1 << 26,
+		ShapeDisableCollision	= 1 << 27,
 #if CHAOS_CHECKED
 		, DebugName				= 1 << 28
 #endif
diff --git a/Engine/Source/Runtime/Experimental/Chaos/Public/Chaos/ParticleHandle.h b/Engine/Source/Runtime/Experimental/Chaos/Public/Chaos/ParticleHandle.h
--- a/Engine/Source/Runtime/Experimental/Chaos/Public/Chaos/ParticleHandle.h
+++ b/Engine/Source/Runtime/Experimental/Chaos/Public/Chaos/ParticleHandle.h
@@ -50,7 +50,6 @@ void GeometryParticleDefaultConstruct(FConcrete& Concrete, const TGeometryPartic
 	Concrete.SetX(TVector<T, d>(0));
 	Concrete.SetR(TRotation<T, d>::Identity);
 	Concrete.SetSpatialIdx(FSpatialAccelerationIdx{ 0,0 });
-	Concrete.SetHashResultLowLevel(FMath::RandRange(0, FMath::RandHelper(TNumericLimits<int32>::Max())));
 }
 
 template <typename T, int d, typename FConcrete>
@@ -165,31 +164,20 @@ public:
 		return !(*this == Rhs);
 	}
 
-	bool operator<(const TAccelerationStructureHandle<T, d>& Rhs) const
-	{
-		return HashResult < Rhs.HashResult;
-	}
+	void Serialize(FChaosArchive& Ar);
 
-	uint32 GetTypeHash() const
+	FUniqueIdx UniqueIdx() const
 	{
-		return HashResult;
+		return CachedUniqueIdx;
 	}
 
-	void Serialize(FChaosArchive& Ar);
-
 private:
 	TGeometryParticle<T, d>* ExternalGeometryParticle;
 	TGeometryParticleHandle<T, d>* GeometryParticleHandle;
 
-	uint32 HashResult;
+	FUniqueIdx CachedUniqueIdx;
 };
 
-template <typename T, int d>
-uint32 GetTypeHash(const TAccelerationStructureHandle<T, d>& A)
-{
-	return A.GetTypeHash();
-}
-
 template <typename T, int d>
 class TParticleHandleBase
 {
@@ -370,6 +358,9 @@ public:
 	TVector<T, d>& X() { return GeometryParticles->X(ParticleIdx); }
 	void SetX(const TVector<T, d>& InX) { GeometryParticles->X(ParticleIdx) = InX; }
 
+	FUniqueIdx UniqueIdx() const { return GeometryParticles->UniqueIdx(ParticleIdx); }
+	void SetUniqueIdx(const FUniqueIdx UniqueIdx) const { GeometryParticles->UniqueIdx(ParticleIdx) = UniqueIdx; }
+
 	const TRotation<T, d>& R() const { return GeometryParticles->R(ParticleIdx); }
 	TRotation<T, d>& R() { return GeometryParticles->R(ParticleIdx); }
 	void SetR(const TRotation<T, d>& InR) { GeometryParticles->R(ParticleIdx) = InR; }
@@ -400,9 +391,6 @@ public:
 	FSpatialAccelerationIdx SpatialIdx() const { return GeometryParticles->SpatialIdx(ParticleIdx); }
 	void SetSpatialIdx(FSpatialAccelerationIdx Idx) { GeometryParticles->SpatialIdx(ParticleIdx) = Idx; }
 
-	void SetHashResultLowLevel(uint32 Value) { GeometryParticles->HashResultLowLevel(ParticleIdx) = Value; }
-	uint32 GetHashResultLowLevel() const { return GeometryParticles->HashResultLowLevel(ParticleIdx); }
-
 #if CHAOS_CHECKED
 	const FName& DebugName() const { return GeometryParticles->DebugName(ParticleIdx); }
 	void SetDebugName(const FName& InDebugName) { GeometryParticles->DebugName(ParticleIdx) = InDebugName; }
@@ -1328,6 +1316,13 @@ public:
 		this->MX = InX;
 	}
 
+	FUniqueIdx UniqueIdx() const { return this->MUniqueIdx; }
+	void SetUniqueIdx(const FUniqueIdx UniqueIdx, bool bInvalidate = true)
+	{
+		this->MarkDirty(EParticleFlags::UniqueIdx,bInvalidate);
+		this->MUniqueIdx = UniqueIdx;
+	}
+
 	const TRotation<T, d>& R() const { return this->MR; }
 	void SetR(const TRotation<T, d>& InR, bool bInvalidate = true)
 	{
@@ -1388,14 +1383,6 @@ public:
 	}
 #endif
 
-	void SetHashResultLowLevel(uint32 Value)
-	{
-		MarkDirty(EParticleFlags::HashResult);
-		MHashResult = Value;
-	}
-
-	uint32 GetHashResultLowLevel() const { return MHashResult; }
-
 	//Note: this must be called after setting geometry. This API seems bad. Should probably be part of setting geometry
 	void SetShapesArray(TShapesArray<T, d>&& InShapesArray)
 	{
@@ -1490,16 +1477,13 @@ public:
 
 private:
 	TVector<T, d> MX;
+	FUniqueIdx MUniqueIdx;
 	TRotation<T, d> MR;
 	TSharedPtr<FImplicitObject, ESPMode::ThreadSafe> MGeometry;	//TODO: geometry should live in bodysetup
 	TShapesArray<T,d> MShapesArray;
 	TMap<const FImplicitObject*, int32> ImplicitShapeMap;
 	FSpatialAccelerationIdx MSpatialIdx;
 
-	// This value is generated and used in AccelerationStructureHandle for a hash, as that handle
-	// sometimes only has a particle pointer from game thread or phys thread, no easy way to get stable hash from that.
-	uint32 MHashResult;
-
 	// Pointer to some arbitrary data associated with the particle, but not used by Chaos. External systems may use this for whatever.
 	void* MUserData;
 
@@ -1555,7 +1539,6 @@ public:
 		, X(TVector<T, d>(0))
 		, R(TRotation<T, d>())
 		, SpatialIdx(FSpatialAccelerationIdx{ 0,0 })
-		, HashResult(0)
 		, DirtyFlags()
 #if CHAOS_CHECKED
 		, DebugName(NAME_None)
@@ -1568,7 +1551,7 @@ public:
 		, R(InParticle.R())
 		, Geometry(InParticle.GeometrySharedLowLevel())
 		, SpatialIdx(InParticle.SpatialIdx())
-		, HashResult(InParticle.GetHashResultLowLevel())
+		, UniqueIdx(InParticle.UniqueIdx())
 		, DirtyFlags(InParticle.DirtyFlags())
 #if CHAOS_CHECKED
 		, DebugName(InParticle.DebugName())
@@ -1590,7 +1573,7 @@ public:
 		R = TRotation<T, d>(); 
 		Geometry = TSharedPtr<FImplicitObjectUnion, ESPMode::ThreadSafe>();
 		SpatialIdx = FSpatialAccelerationIdx{ 0,0 };
-		HashResult = 0;
+		UniqueIdx = FUniqueIdx();
 		DirtyFlags.Clear();
 		ShapeCollisionDisableFlags.Reset();
 #if CHAOS_CHECKED
@@ -1602,7 +1585,7 @@ public:
 	TRotation<T, d> R;
 	TSharedPtr<FImplicitObject, ESPMode::ThreadSafe> Geometry;
 	FSpatialAccelerationIdx SpatialIdx;
-	uint32 HashResult;
+	FUniqueIdx UniqueIdx;
 	FParticleDirtyFlags DirtyFlags;
 	TBitArray<> ShapeCollisionDisableFlags;
 #if CHAOS_CHECKED
@@ -2203,16 +2186,18 @@ template <typename T, int d>
 TAccelerationStructureHandle<T,d>::TAccelerationStructureHandle(TGeometryParticleHandle<T, d>* InHandle)
 	: ExternalGeometryParticle(InHandle->GTGeometryParticle())
 	, GeometryParticleHandle(InHandle)
-	, HashResult(InHandle->GetHashResultLowLevel())
+	, CachedUniqueIdx(InHandle->UniqueIdx())
 {
+	ensure(CachedUniqueIdx.IsValid());
 }
 
 template <typename T, int d>
 TAccelerationStructureHandle<T,d>::TAccelerationStructureHandle(TGeometryParticle<T, d>* InGeometryParticle)
 	: ExternalGeometryParticle(InGeometryParticle)
 	, GeometryParticleHandle(InGeometryParticle ? InGeometryParticle->Handle() : nullptr)
-	, HashResult(InGeometryParticle ? InGeometryParticle->GetHashResultLowLevel() : 0)
+	, CachedUniqueIdx(InGeometryParticle ? InGeometryParticle->UniqueIdx() : FUniqueIdx())
 {
+	ensure(!InGeometryParticle || CachedUniqueIdx.IsValid());
 }
 
 template <typename T, int d>
@@ -2220,8 +2205,9 @@ template <bool bPersistent>
 TAccelerationStructureHandle<T, d>::TAccelerationStructureHandle(TGeometryParticleHandleImp<T, d, bPersistent>& InHandle)
 	: ExternalGeometryParticle(InHandle.GTGeometryParticle())
 	, GeometryParticleHandle(InHandle.Handle())
-	, HashResult(InHandle.GetHashResultLowLevel())
+	, CachedUniqueIdx(InHandle.UniqueIdx())
 {
+	ensure(CachedUniqueIdx.IsValid());
 }
 
 template <typename T, int d>
@@ -2229,21 +2215,23 @@ template <bool bPersistent>
 TAccelerationStructureHandle<T, d>::TAccelerationStructureHandle(TGeometryParticleHandleImp<T, d, bPersistent>* InHandle, TGeometryParticle<T, d>* InGeometryParticle)
 	: ExternalGeometryParticle(InGeometryParticle)
 	, GeometryParticleHandle(InHandle)
-	, HashResult(0)
 {
 	if (GeometryParticleHandle)
 	{
-		HashResult = GeometryParticleHandle->GetHashResultLowLevel();
+		CachedUniqueIdx = GeometryParticleHandle->UniqueIdx();
 	}
 	else if (ExternalGeometryParticle)
 	{
-		HashResult = ExternalGeometryParticle->GetHashResultLowLevel();
+		CachedUniqueIdx = ExternalGeometryParticle->UniqueIdx();
 	}
 
 	if (GeometryParticleHandle && ExternalGeometryParticle)
 	{
-		CHAOS_ENSURE(GeometryParticleHandle->GetHashResultLowLevel() == ExternalGeometryParticle->GetHashResultLowLevel());
+		ensure(GeometryParticleHandle->UniqueIdx() == ExternalGeometryParticle->UniqueIdx());
 	}
+
+	//If either handle or particle are valid, we must have a valid idx
+	ensure((!GeometryParticleHandle && !ExternalGeometryParticle) || CachedUniqueIdx.IsValid());
 }
 
 template <typename T, int d>
@@ -2253,30 +2241,28 @@ void TAccelerationStructureHandle<T, d>::Serialize(FChaosArchive& Ar)
 	Ar << AsAlwaysSerializable(GeometryParticleHandle);
 
 	Ar.UsingCustomVersion(FExternalPhysicsCustomObjectVersion::GUID);
-	if (Ar.CustomVer(FExternalPhysicsCustomObjectVersion::GUID) < FExternalPhysicsCustomObjectVersion::SerializeHashResult)
+	if (Ar.CustomVer(FExternalPhysicsCustomObjectVersion::GUID) >= FExternalPhysicsCustomObjectVersion::SerializeHashResult && Ar.CustomVer(FExternalPhysicsCustomObjectVersion::GUID) < FExternalPhysicsCustomObjectVersion::UniquePayloadIdx)
 	{
-		if (GeometryParticleHandle)
-		{
-			HashResult = GeometryParticleHandle->GetHashResultLowLevel();
-		}
-		else if (ExternalGeometryParticle)
-		{
-			HashResult = ExternalGeometryParticle->GetHashResultLowLevel();
-		}
-		else
-		{
-			HashResult = 0;
-		}
-
-		if (GeometryParticleHandle && ExternalGeometryParticle)
-		{
-			CHAOS_ENSURE(GeometryParticleHandle->GetHashResultLowLevel() == ExternalGeometryParticle->GetHashResultLowLevel());
-		}
+		uint32 DummyHash;
+		Ar << DummyHash;
 	}
-	else
+	
+	if(GeometryParticleHandle)
 	{
-		Ar << HashResult;
+		CachedUniqueIdx = GeometryParticleHandle->UniqueIdx();
 	}
+	else if(ExternalGeometryParticle)
+	{
+		CachedUniqueIdx = ExternalGeometryParticle->UniqueIdx();
+	}
+	
+	if (GeometryParticleHandle && ExternalGeometryParticle)
+	{
+		ensure(GeometryParticleHandle->UniqueIdx() == ExternalGeometryParticle->UniqueIdx());
+	}
+	
+	ensure(!GeometryParticleHandle || CachedUniqueIdx.IsValid());
+	ensure(!ExternalGeometryParticle || CachedUniqueIdx.IsValid());
 }
 
 template <typename T, int d>
diff --git a/Engine/Source/Runtime/Experimental/Chaos/Public/Chaos/TriangleMeshImplicitObject.h b/Engine/Source/Runtime/Experimental/Chaos/Public/Chaos/TriangleMeshImplicitObject.h
--- a/Engine/Source/Runtime/Experimental/Chaos/Public/Chaos/TriangleMeshImplicitObject.h
+++ b/Engine/Source/Runtime/Experimental/Chaos/Public/Chaos/TriangleMeshImplicitObject.h
@@ -297,6 +297,11 @@ namespace Chaos
 			{
 				return Idx;
 			}
+
+			FUniqueIdx UniqueIdx() const
+			{
+				return FUniqueIdx(Index);
+			}
 		};
 
 		BVHType BVH;
diff --git a/Engine/Source/Runtime/Experimental/ChaosSolvers/Private/PBDRigidsSolver.cpp b/Engine/Source/Runtime/Experimental/ChaosSolvers/Private/PBDRigidsSolver.cpp
--- a/Engine/Source/Runtime/Experimental/ChaosSolvers/Private/PBDRigidsSolver.cpp
+++ b/Engine/Source/Runtime/Experimental/ChaosSolvers/Private/PBDRigidsSolver.cpp
@@ -189,6 +189,8 @@ namespace Chaos
 		Chaos::FParticleData* ProxyData;
 		IPhysicsProxyBase* ProxyBase;
 
+		GTParticle->SetUniqueIdx(GetEvolution()->GenerateUniqueIdx());
+
 		// Make a physics proxy, giving it our particle and particle handle
 		const EParticleType InParticleType = GTParticle->ObjectType();
 		if (InParticleType == EParticleType::Rigid)
@@ -232,23 +234,25 @@ namespace Chaos
 			//       Proxy should be able to map back to Particle.
 			//
 
+			const auto& GeomData = static_cast<const Chaos::TGeometryParticleData<FReal,3>*>(ProxyData);
+
 			if (InParticleType == EParticleType::Rigid)
 			{
-				Handle = Solver->Particles.CreateDynamicParticles(1)[0];
+				Handle = Solver->Particles.CreateDynamicParticles(1, &GeomData->UniqueIdx)[0];
 				auto Proxy = static_cast<FRigidParticlePhysicsProxy*>(ProxyBase);
 				Proxy->SetHandle(Handle->CastToRigidParticle());
 				Proxy->PushToPhysicsState(ProxyData);
 			}
 			else if (InParticleType == EParticleType::Kinematic)
 			{
-				Handle = Solver->Particles.CreateKinematicParticles(1)[0];
+				Handle = Solver->Particles.CreateKinematicParticles(1, &GeomData->UniqueIdx)[0];
 				auto Proxy = static_cast<FKinematicGeometryParticlePhysicsProxy*>(ProxyBase);
 				Proxy->SetHandle(Handle->CastToKinematicParticle());
 				Proxy->PushToPhysicsState(ProxyData);
 			}
 			else // Assume it's a static (geometry) if it's not dynamic or kinematic
 			{
-				Handle = Solver->Particles.CreateStaticParticles(1)[0];
+				Handle = Solver->Particles.CreateStaticParticles(1, &GeomData->UniqueIdx)[0];
 
 				auto Proxy = static_cast<FGeometryParticlePhysicsProxy*>(ProxyBase);
 				Proxy->SetHandle(Handle);
diff --git a/Engine/Source/Runtime/Experimental/ChaosSolvers/Private/PhysicsProxy/SingleParticlePhysicsProxy.cpp b/Engine/Source/Runtime/Experimental/ChaosSolvers/Private/PhysicsProxy/SingleParticlePhysicsProxy.cpp
--- a/Engine/Source/Runtime/Experimental/ChaosSolvers/Private/PhysicsProxy/SingleParticlePhysicsProxy.cpp
+++ b/Engine/Source/Runtime/Experimental/ChaosSolvers/Private/PhysicsProxy/SingleParticlePhysicsProxy.cpp
@@ -60,7 +60,7 @@ void FSingleParticlePhysicsProxy<Chaos::TGeometryParticle<float, 3>>::PushToPhys
 			RigidHandle->SetWorldSpaceInflatedBounds(Data->Geometry->BoundingBox().TransformedAABB(Chaos::TRigidTransform<float, 3>(Data->X, Data->R)));
 		}
 		RigidHandle->SetSpatialIdx(Data->SpatialIdx);	//todo: this needs to only happen once during initialization
-		RigidHandle->SetHashResultLowLevel(Data->HashResult);
+		RigidHandle->SetUniqueIdx(Data->UniqueIdx);
 #if CHAOS_CHECKED
 		RigidHandle->SetDebugName(Data->DebugName);
 #endif
@@ -142,7 +142,7 @@ void FSingleParticlePhysicsProxy<Chaos::TKinematicGeometryParticle<float, 3>>::P
 		RigidHandle->SetCenterOfMass(Data->MCenterOfMass);
 		RigidHandle->SetRotationOfMass(Data->MRotationOfMass);
 		RigidHandle->SetSpatialIdx(Data->SpatialIdx);	//todo: this needs to only happen once during initialization
-		RigidHandle->SetHashResultLowLevel(Data->HashResult);	//todo: this needs to only happen once during initialization
+		RigidHandle->SetUniqueIdx(Data->UniqueIdx);
 #if CHAOS_CHECKED
 		RigidHandle->SetDebugName(Data->DebugName);
 #endif
@@ -234,7 +234,7 @@ void FSingleParticlePhysicsProxy<Chaos::TPBDRigidParticle<float, 3>>::PushToPhys
 		// TODO: Since this can't change, it doesn't have a dirty flag.
 		// This should be moved to initialization.
 		RigidHandle->SetSpatialIdx(Data->SpatialIdx);	//todo: this needs to only happen once during initialization
-		RigidHandle->SetHashResultLowLevel(Data->HashResult);	//todo: this needs to only happen once during initialization
+		RigidHandle->SetUniqueIdx(Data->UniqueIdx);	//todo: this needs to only happen once during initialization
 #if CHAOS_CHECKED
 		RigidHandle->SetDebugName(Data->DebugName);
 #endif
