Chaos : Instanced Collision Resolution

- Added supprot for Instance Implicit Object.

diff --git a/Engine/Source/Runtime/Experimental/Chaos/Private/Chaos/CollisionResolution.cpp b/Engine/Source/Runtime/Experimental/Chaos/Private/Chaos/CollisionResolution.cpp
--- a/Engine/Source/Runtime/Experimental/Chaos/Private/Chaos/CollisionResolution.cpp
+++ b/Engine/Source/Runtime/Experimental/Chaos/Private/Chaos/CollisionResolution.cpp
@@ -92,6 +92,18 @@ namespace Chaos
 					Contact.Normal = BTransform.TransformVector(Normal);
 				}
 			}
+			else if (const TImplicitObjectInstanced<GeometryA>* InstancedConvexImplicit = A.template GetObject<const TImplicitObjectInstanced<GeometryA> >())
+			{
+				if (const GeometryA * InstancedInnerObject = static_cast<const GeometryA*>(InstancedConvexImplicit->GetInstancedObject()))
+				{
+					if (B.GJKContactPoint(*InstancedInnerObject, AToBTM, Thickness, Location, Normal, Penetration))
+					{
+						Contact.Phi = Penetration;
+						Contact.Location = BTransform.TransformPosition(Location);
+						Contact.Normal = BTransform.TransformVector(Normal);
+					}
+				}
+			}
 			else if (const GeometryA* ConvexImplicit = A.template GetObject<const GeometryA>())
 			{
 				if (B.GJKContactPoint(*ConvexImplicit, AToBTM, Thickness, Location, Normal, Penetration))
@@ -435,10 +447,12 @@ namespace Chaos
 		}
 
 		template <typename TPGeometryClass>
-		const TPGeometryClass* GetUnscaledInner(const FImplicitObject& Geometry)
+		const TPGeometryClass* GetInnerObject(const FImplicitObject& Geometry)
 		{
 			if (const TImplicitObjectScaled<TPGeometryClass>* ScaledConvexImplicit = Geometry.template GetObject<const TImplicitObjectScaled<TPGeometryClass> >())
 				return (Geometry.template GetObject<const TImplicitObjectScaled<TPGeometryClass> >())->GetUnscaledObject();
+			else if (const TImplicitObjectInstanced<TPGeometryClass>* InstancedImplicit = Geometry.template GetObject<const TImplicitObjectInstanced<TPGeometryClass> >())
+				return (Geometry.template GetObject<const TImplicitObjectInstanced<TPGeometryClass> >())->GetInstancedObject();
 			else if (const TPGeometryClass* ConvexImplicit = Geometry.template GetObject<const TPGeometryClass>())
 				return Geometry.template GetObject<const TPGeometryClass>();
 			return nullptr;
@@ -448,7 +462,7 @@ namespace Chaos
 		void ConstructBoxTriangleMeshConstraints(TGeometryParticleHandle<T, d>* Particle0, TGeometryParticleHandle<T, d>* Particle1, const FImplicitObject* Implicit0, const FImplicitObject* Implicit1, const TRigidTransform<T, d>& Transform0, const TRigidTransform<T, d>& Transform1, const T Thickness, FCollisionConstraintsArray& NewConstraints)
 		{
 			const TBox<T, d> * Object0 = Implicit0->template GetObject<const TBox<T, d> >();
-			const FTriangleMeshImplicitObject * Object1 = GetUnscaledInner<FTriangleMeshImplicitObject>(*Implicit1);
+			const FTriangleMeshImplicitObject * Object1 = GetInnerObject<FTriangleMeshImplicitObject>(*Implicit1);
 			if (ensure(Object0 && Object1))
 			{
 				TRigidTransform<T, d> ParticleImplicit0TM = Transform0.GetRelativeTransform(Collisions::GetTransform(Particle0));
@@ -1571,6 +1585,73 @@ namespace Chaos
 				ConstructConstraints(Particle0, Particle1, Implicit0, TransformedImplicit1->GetTransformedObject(), Transform0, TransformedTransform1, Thickness, NewConstraints);
 				return;
 			}
+
+
+			else if (Implicit0OuterType == TImplicitObjectInstanced<FConvex>::StaticType())
+			{
+				const TImplicitObjectInstanced<FConvex>* TransformedImplicit0 = Implicit0->template GetObject<const TImplicitObjectInstanced<FConvex>>();
+				ConstructConstraints(Particle0, Particle1, TransformedImplicit0->GetInstancedObject(), Implicit1, Transform0, Transform1, Thickness, NewConstraints);
+				return;
+			}
+			else if (Implicit1OuterType == TImplicitObjectInstanced<FConvex>::StaticType())
+			{
+				const TImplicitObjectInstanced<FConvex>* TransformedImplicit1 = Implicit1->template GetObject<const TImplicitObjectInstanced<FConvex>>();
+				ConstructConstraints(Particle0, Particle1, Implicit0, TransformedImplicit1->GetInstancedObject(), Transform0, Transform1, Thickness, NewConstraints);
+				return;
+			}
+
+
+			else if (Implicit0OuterType == TImplicitObjectInstanced<TBox<FReal,3>>::StaticType())
+			{
+				const TImplicitObjectInstanced<TBox<FReal, 3>>* TransformedImplicit0 = Implicit0->template GetObject<const TImplicitObjectInstanced<TBox<FReal, 3>>>();
+				ConstructConstraints(Particle0, Particle1, TransformedImplicit0->GetInstancedObject(), Implicit1, Transform0, Transform1, Thickness, NewConstraints);
+				return;
+			}
+			else if (Implicit1OuterType == TImplicitObjectInstanced<TBox<FReal, 3>>::StaticType())
+			{
+				const TImplicitObjectInstanced<TBox<FReal, 3>>* TransformedImplicit1 = Implicit1->template GetObject<const TImplicitObjectInstanced<TBox<FReal, 3>>>();
+				ConstructConstraints(Particle0, Particle1, Implicit0, TransformedImplicit1->GetInstancedObject(), Transform0, Transform1, Thickness, NewConstraints);
+				return;
+			}
+
+			else if (Implicit0OuterType == TImplicitObjectInstanced<TCapsule<FReal>>::StaticType())
+			{
+				const TImplicitObjectInstanced<TCapsule<FReal>>* TransformedImplicit0 = Implicit0->template GetObject<const TImplicitObjectInstanced<TCapsule<FReal>>>();
+				ConstructConstraints(Particle0, Particle1, TransformedImplicit0->GetInstancedObject(), Implicit1, Transform0, Transform1, Thickness, NewConstraints);
+				return;
+			}
+			else if (Implicit1OuterType == TImplicitObjectInstanced<TCapsule<FReal>>::StaticType())
+			{
+				const TImplicitObjectInstanced<TCapsule<FReal>>* TransformedImplicit1 = Implicit1->template GetObject<const TImplicitObjectInstanced<TCapsule<FReal>>>();
+				ConstructConstraints(Particle0, Particle1, Implicit0, TransformedImplicit1->GetInstancedObject(), Transform0, Transform1, Thickness, NewConstraints);
+				return;
+			}
+
+			else if (Implicit0OuterType == TImplicitObjectInstanced<TSphere<FReal, 3>>::StaticType())
+			{
+				const TImplicitObjectInstanced<TSphere<FReal, 3>>* TransformedImplicit0 = Implicit0->template GetObject<const TImplicitObjectInstanced<TSphere<FReal, 3>>>();
+				ConstructConstraints(Particle0, Particle1, TransformedImplicit0->GetInstancedObject(), Implicit1, Transform0, Transform1, Thickness, NewConstraints);
+				return;
+			}
+			else if (Implicit1OuterType == TImplicitObjectInstanced<TSphere<FReal, 3>>::StaticType())
+			{
+				const TImplicitObjectInstanced<TSphere<FReal, 3>>* TransformedImplicit1 = Implicit1->template GetObject<const TImplicitObjectInstanced<TSphere<FReal, 3>>>();
+				ConstructConstraints(Particle0, Particle1, Implicit0, TransformedImplicit1->GetInstancedObject(), Transform0, Transform1, Thickness, NewConstraints);
+				return;
+			}
+
+			else if (Implicit0OuterType == TImplicitObjectInstanced<FConvex>::StaticType())
+			{
+				const TImplicitObjectInstanced<FConvex>* TransformedImplicit0 = Implicit0->template GetObject<const TImplicitObjectInstanced<FConvex>>();
+				ConstructConstraints(Particle0, Particle1, TransformedImplicit0->GetInstancedObject(), Implicit1, Transform0, Transform1, Thickness, NewConstraints);
+				return;
+			}
+			else if (Implicit1OuterType == TImplicitObjectInstanced<FConvex>::StaticType())
+			{
+				const TImplicitObjectInstanced<FConvex>* TransformedImplicit1 = Implicit1->template GetObject<const TImplicitObjectInstanced<FConvex>>();
+				ConstructConstraints(Particle0, Particle1, Implicit0, TransformedImplicit1->GetInstancedObject(), Transform0, Transform1, Thickness, NewConstraints);
+				return;
+			}
 			else if (Implicit0OuterType != FImplicitObjectUnion::StaticType() && Implicit1OuterType == FImplicitObjectUnion::StaticType())
 			{
 				const TArray<Pair<const FImplicitObject*, TRigidTransform<T, d>>> LevelsetShapes = FindRelevantShapes(Implicit0, Transform0, *Implicit1, Transform1, Thickness);
