Fix cis

diff --git a/Engine/Source/Runtime/Experimental/Chaos/Public/Chaos/PBDRigidsSOAs.h b/Engine/Source/Runtime/Experimental/Chaos/Public/Chaos/PBDRigidsSOAs.h
--- a/Engine/Source/Runtime/Experimental/Chaos/Public/Chaos/PBDRigidsSOAs.h
+++ b/Engine/Source/Runtime/Experimental/Chaos/Public/Chaos/PBDRigidsSOAs.h
@@ -134,10 +134,10 @@ public:
 		UpdateViews();
 		return Results;
 	}
-	TArray<TPBDGeometryCollectionParticleHandle<T, d>*> CreateGeometryCollectionParticles(int32 NumParticles, const TPBDRigidParticleParameters<T, d>& Params = TPBDRigidParticleParameters<T, d>())
+	TArray<TPBDGeometryCollectionParticleHandle<T, d>*> CreateGeometryCollectionParticles(int32 NumParticles, const FUniqueIdx* ExistingIndices = nullptr, const TPBDRigidParticleParameters<T, d>& Params = TPBDRigidParticleParameters<T, d>())
 	{
 		TArray<TPBDGeometryCollectionParticleHandle<T, d>*> Results = CreateParticlesHelper<TPBDGeometryCollectionParticleHandle<T, d>>(
-			NumParticles, GeometryCollectionParticles, Params);
+			NumParticles, ExistingIndices, GeometryCollectionParticles, Params);
 		//TArray<TPBDRigidParticleHandle<T, d>*>& RigidHandles = (TArray<TPBDRigidParticleHandle<T, d>*>*)&Results;//*static_cast<TArray<TPBDRigidParticleHandle<T, d>*>*>(&Results);
 		if (!Params.bStartSleeping)
 		{
