Fix GetHeight function.

diff --git a/Engine/Source/Runtime/Experimental/Chaos/Private/Chaos/HeightField.cpp b/Engine/Source/Runtime/Experimental/Chaos/Private/Chaos/HeightField.cpp
--- a/Engine/Source/Runtime/Experimental/Chaos/Private/Chaos/HeightField.cpp
+++ b/Engine/Source/Runtime/Experimental/Chaos/Private/Chaos/HeightField.cpp
@@ -586,7 +586,7 @@ namespace Chaos
 				
 				FVector Bary = FMath::GetBaryCentric2D({ FractionX, FractionY, 0.0f }, Tri[0], Tri[1], Tri[2]);
 
-				return Pts[0].Z * Bary[0] + Pts[1].Z * Bary[1] + Pts[3].Z * Bary[2];
+				return Pts[0].Z * Bary[0] + Pts[3].Z * Bary[1] + Pts[2].Z * Bary[2];
 			}
 			else
 			{
