Chaos: fix crash bug in ComputeIslands, Static Particles were not being removed from IslandToParticles array

diff --git a/Engine/Source/Runtime/Experimental/Chaos/Private/Chaos/PBDConstraintGraph.cpp b/Engine/Source/Runtime/Experimental/Chaos/Private/Chaos/PBDConstraintGraph.cpp
--- a/Engine/Source/Runtime/Experimental/Chaos/Private/Chaos/PBDConstraintGraph.cpp
+++ b/Engine/Source/Runtime/Experimental/Chaos/Private/Chaos/PBDConstraintGraph.cpp
@@ -722,7 +722,7 @@ void FPBDConstraintGraph::DisableParticle(TGeometryParticleHandle<FReal, 3>* Par
 	}
 	else
 	{
-		// Kinematic particles are included in IslandToParticles, however we cannot use islands to look them up.
+		// Kinematic & Static particles are included in IslandToParticles, however we cannot use islands to look them up.
 		// TODO find faster removal method?
 		for (TArray<TGeometryParticleHandle<FReal, 3>*>& IslandParticles : IslandToParticles)
 		{
