Chaos - Jaco's GJKPenetration returns distance as well as penetration

#fui jaco.vandyk,ori.cohen

diff --git a/Engine/Source/Runtime/Experimental/Chaos/Private/Chaos/CollisionResolution.cpp b/Engine/Source/Runtime/Experimental/Chaos/Private/Chaos/CollisionResolution.cpp
--- a/Engine/Source/Runtime/Experimental/Chaos/Private/Chaos/CollisionResolution.cpp
+++ b/Engine/Source/Runtime/Experimental/Chaos/Private/Chaos/CollisionResolution.cpp
@@ -20,6 +20,8 @@
 #include "Chaos/TriangleMeshImplicitObject.h"
 #include "Chaos/GeometryQueries.h"
 
+DECLARE_CYCLE_STAT(TEXT("Collisions::GJK"), STAT_Collisions_GJK, STATGROUP_ChaosCollision);
+
 namespace Chaos
 {
 	namespace Collisions
@@ -50,27 +52,21 @@ namespace Chaos
 		template <typename T, int d, typename GeometryA, typename GeometryB>
 		TContactPoint<T> GJKContactPoint(const GeometryA& A, const TRigidTransform<T, d>& ATM, const GeometryB& B, const TRigidTransform<T, d>& BTM, const T Thickness)
 		{
+			SCOPE_CYCLE_COUNTER(STAT_Collisions_GJK);
+
 			TContactPoint<T> Contact;
 			const TRigidTransform<T, d> BToATM = BTM.GetRelativeTransform(ATM);
+			const TVector<T, 3> InitialDirection = TVector<T, 3>(1, 0, 0);
 
 			T Penetration;
 			TVec3<T> ClosestA, ClosestB, Normal;
-			if (GJKPenetration(A, B, BToATM, Penetration, ClosestA, ClosestB, Normal, (T)0))
+
+			if (ensure(GJKPenetration(A, B, BToATM, Penetration, ClosestA, ClosestB, Normal, (T)0, InitialDirection, (T)0, true)))
 			{
 				Contact.Location = ATM.TransformPosition(ClosestA);
 				Contact.Normal = -ATM.TransformVector(Normal);
 				Contact.Phi = -Penetration;
 			}
-			else if (ensure(GJKDistance(A, B, BToATM, Penetration, ClosestA, ClosestB)))
-			{
-				//todo: make GJKPenetration support no penetration case
-				TVector<T, d> NearPointAWorld = ATM.TransformPosition(ClosestA);
-				TVector<T, d> NearPointBWorld = BTM.TransformPosition(ClosestB);
-				TVector<T, d> NearPointBtoAWorld = NearPointAWorld - NearPointBWorld;
-				Contact.Phi = Penetration;
-				Contact.Normal = NearPointBtoAWorld.GetSafeNormal();
-				Contact.Location = NearPointAWorld;
-			}
 
 			return Contact;
 		}
diff --git a/Engine/Source/Runtime/Experimental/Chaos/Public/Chaos/GJK.h b/Engine/Source/Runtime/Experimental/Chaos/Public/Chaos/GJK.h
--- a/Engine/Source/Runtime/Experimental/Chaos/Public/Chaos/GJK.h
+++ b/Engine/Source/Runtime/Experimental/Chaos/Public/Chaos/GJK.h
@@ -7,6 +7,7 @@
 #include "Chaos/Sphere.h"
 #include "Chaos/EPA.h"
 #include "ChaosCheck.h"
+#include "ChaosLog.h"
 
 namespace Chaos
 {
@@ -82,10 +83,112 @@ namespace Chaos
 		} while (!bTerminate);
 
 		return bNearZero;
+		
 	}
 
+
+	/** Determines the distance between two convex bodies
+	// Improvement over previous function: calculates distances using core shapes allowing for better performance and negative distances
+	// The distances can be negative
+	 @A The first geometry
+	 @B The second geometry
+	 @BToATM The transform of B in A's local space
+	 @ThicknessA The amount of geometry inflation for Geometry A(for example if the surface distance of two geometries with thickness 0 would be 2, a thickness of 0.5 would give a distance of 1.5)
+	 @InitialDir The first direction we use to search the CSO
+	 @ThicknessB The amount of geometry inflation for Geometry B(for example if the surface distance of two geometries with thickness 0 would be 2, a thickness of 0.5 would give a distance of 1.5)
+	 @OutDistance if returns true, the minimum distance between A and B, otherwise not modified.
+	 @OutNearestA if returns true, the near point on A in local-space, otherwise not modified.
+	 @OutNearestB if returns true, the near point on B in local-space, otherwise not modified.
+	 @return false if the geometries overlap sop much that a distance cannot be calculated by GJK*/
+
 	template <typename T, typename TGeometryA, typename TGeometryB>
-	bool GJKPenetration(const TGeometryA& A, const TGeometryB& B, const TRigidTransform<T, 3>& BToATM, T& OutPenetration, TVec3<T>& OutClosestA, TVec3<T>& OutClosestB, TVec3<T>& OutNormal, const T InThicknessA = 0, const TVector<T, 3>& InitialDir = TVector<T, 3>(1, 0, 0), const T InThicknessB = 0)
+	bool GJKDistanceImproved(const TGeometryA& A, const TGeometryB& B, const TRigidTransform<T, 3>& BToATM, const T InThicknessA = 0, const TVector<T, 3>& InitialDir = TVector<T, 3>(1, 0, 0), const T InThicknessB = 0,
+		T* OutDistance = nullptr, TVector<T, 3>* OutNearestA = nullptr, TVector<T, 3>* OutNearestB = nullptr)
+	{
+		const T Epsilon = 1e-3;
+		TVector<T, 3> V = -InitialDir;
+		if (V.SafeNormalize() == 0)
+		{
+			V = TVec3<T>(-1, 0, 0);
+		}
+
+		FSimplex SimplexIDs;
+		TVector<T, 3> Simplex[4];
+		TVector<T, 3> SimplexA[4];
+		TVector<T, 3> SimplexB[4];
+
+		T Barycentric[4] = { -1,-1,-1,-1 };	//not needed, but compiler warns
+		const TRotation<T, 3> AToBRotation = BToATM.GetRotation().Inverse();
+		bool bTerminate;
+		bool bNearZero = false;
+		int NumIterations = 0;
+		T PrevDist2 = FLT_MAX;
+		const T ThicknessA = A.GetMargin() + InThicknessA;
+		const T ThicknessB = B.GetMargin() + InThicknessB;
+		const T Inflation = ThicknessA + ThicknessB + Epsilon;
+		const T Inflation2 = Inflation * Inflation;
+		do
+		{
+			if (!ensure(NumIterations++ < 32))	//todo: take this out
+			{
+				break;	//if taking too long just stop. This should never happen
+			}
+			const TVector<T, 3> NegV = -V;
+			const TVector<T, 3> SupportA = A.Support2(NegV);
+			const TVector<T, 3> VInB = AToBRotation * V;
+			const TVector<T, 3> SupportBLocal = B.Support2(VInB);
+			const TVector<T, 3> SupportB = BToATM.TransformPositionNoScale(SupportBLocal);
+			const TVector<T, 3> W = SupportA - SupportB;
+			
+
+			SimplexIDs[SimplexIDs.NumVerts] = SimplexIDs.NumVerts;
+			SimplexA[SimplexIDs.NumVerts] = SupportA;
+			SimplexB[SimplexIDs.NumVerts] = SupportBLocal;
+			Simplex[SimplexIDs.NumVerts++] = W;
+
+			V = SimplexFindClosestToOrigin(Simplex, SimplexIDs, Barycentric);
+
+			T NewDist2 = V.SizeSquared();
+			bNearZero = NewDist2 < Epsilon * Epsilon;//Inflation2;
+
+			//as simplices become degenerate we will stop making progress. This is a side-effect of precision, in that case take V as the current best approximation
+			//question: should we take previous v in case it's better?
+			const bool bMadeProgress = NewDist2 < PrevDist2;
+			bTerminate = bNearZero || !bMadeProgress;
+
+			PrevDist2 = NewDist2;
+
+			if (!bTerminate)
+			{
+				V /= FMath::Sqrt(NewDist2);
+			}
+
+			if (bTerminate && !bNearZero)
+			{
+				// Geometries do not overlap, so we can populate the output parameters
+				*OutDistance = V.Size() - Inflation + Epsilon;
+
+				// The simplex vertices are the nearest point/line/face
+				*OutNearestA = TVector<T, 3>(0, 0, 0);
+				*OutNearestB = TVector<T, 3>(0, 0, 0);
+				for (int32 VertIndex = 0; VertIndex < SimplexIDs.NumVerts; ++VertIndex)
+				{
+					int32 WIndex = SimplexIDs[VertIndex];
+					*OutNearestA += Barycentric[WIndex] * SimplexA[WIndex];
+					*OutNearestB += Barycentric[WIndex] * SimplexB[WIndex];
+				}
+				*OutNearestA -= V * ThicknessA;
+				*OutNearestB += AToBRotation * V * ThicknessB;
+
+			}
+		} while (!bTerminate);
+
+		return !bNearZero;
+
+	}
+
+	template <typename T, typename TGeometryA, typename TGeometryB>
+	bool GJKPenetration(const TGeometryA& A, const TGeometryB& B, const TRigidTransform<T, 3>& BToATM, T& OutPenetration, TVec3<T>& OutClosestA, TVec3<T>& OutClosestB, TVec3<T>& OutNormal, const T InThicknessA = 0, const TVector<T, 3>& InitialDir = TVector<T, 3>(1, 0, 0), const T InThicknessB = 0, bool bNegativePenetrationSupport = false)
 	{
 		auto SupportAFunc = [&A](const TVec3<T>& V)
 		{
@@ -135,7 +238,7 @@ namespace Chaos
 			const TVector<T, 3> SupportB = SupportBFunc(V);
 			const TVector<T, 3> W = SupportA - SupportB;
 
-			if (TVector<T, 3>::DotProduct(V, W) > Inflation)
+			if (!bNegativePenetrationSupport && TVector<T, 3>::DotProduct(V, W) > Inflation)
 			{
 				return false;
 			}
@@ -179,7 +282,11 @@ namespace Chaos
 			
 			const T PreDist = FMath::Sqrt(PrevDist2);
 			OutNormal = (ClosestBInA - ClosestA).GetUnsafeNormal();	//question: should we just use PreDist2?
-			const T Penetration = FMath::Clamp<T>(ThicknessA + ThicknessB - PreDist, 0, TNumericLimits<T>::Max());
+			T Penetration = ThicknessA + ThicknessB - PreDist;
+			if (!bNegativePenetrationSupport)
+			{
+				Penetration = FMath::Clamp<T>(Penetration, 0, TNumericLimits<T>::Max());
+			}
 			OutPenetration = Penetration;
 			OutClosestA = ClosestA + OutNormal * ThicknessA;
 			OutClosestB = ClosestBInA - OutNormal * ThicknessB;
