Make sure proxy exists before using it.

diff --git a/Engine/Source/Runtime/Experimental/ChaosSolvers/Private/PBDRigidsSolver.cpp b/Engine/Source/Runtime/Experimental/ChaosSolvers/Private/PBDRigidsSolver.cpp
--- a/Engine/Source/Runtime/Experimental/ChaosSolvers/Private/PBDRigidsSolver.cpp
+++ b/Engine/Source/Runtime/Experimental/ChaosSolvers/Private/PBDRigidsSolver.cpp
@@ -573,22 +573,25 @@ namespace Chaos
 		TParticleView<TPBDRigidParticles<float, 3>>& ActiveParticles = GetParticles().GetActiveParticlesView();
 		for (Chaos::TPBDRigidParticleHandleImp<float, 3, false>& ActiveObject : ActiveParticles)
 		{
-			switch (ActiveObject.GetParticleType())
+			if (ensure(ActiveObject.GTGeometryParticle()->Proxy))
 			{
-			case Chaos::EParticleType::Rigid:
-				((FRigidParticlePhysicsProxy*)(ActiveObject.GTGeometryParticle()->Proxy))->BufferPhysicsResults();
-				break;
-			case Chaos::EParticleType::Kinematic:
-				((FKinematicGeometryParticlePhysicsProxy*)(ActiveObject.GTGeometryParticle()->Proxy))->BufferPhysicsResults();
-				break;
-			case Chaos::EParticleType::Static:
-				((FGeometryParticlePhysicsProxy*)(ActiveObject.GTGeometryParticle()->Proxy))->BufferPhysicsResults();
-				break;
-			case Chaos::EParticleType::GeometryCollection:
-				ActiveGC.AddUnique((FGeometryCollectionPhysicsProxy*)(ActiveObject.GTGeometryParticle()->Proxy));
-				break;
-			default:
-				check(false);
+				switch (ActiveObject.GetParticleType())
+				{
+				case Chaos::EParticleType::Rigid:
+					((FRigidParticlePhysicsProxy*)(ActiveObject.GTGeometryParticle()->Proxy))->BufferPhysicsResults();
+					break;
+				case Chaos::EParticleType::Kinematic:
+					((FKinematicGeometryParticlePhysicsProxy*)(ActiveObject.GTGeometryParticle()->Proxy))->BufferPhysicsResults();
+					break;
+				case Chaos::EParticleType::Static:
+					((FGeometryParticlePhysicsProxy*)(ActiveObject.GTGeometryParticle()->Proxy))->BufferPhysicsResults();
+					break;
+				case Chaos::EParticleType::GeometryCollection:
+					ActiveGC.AddUnique((FGeometryCollectionPhysicsProxy*)(ActiveObject.GTGeometryParticle()->Proxy));
+					break;
+				default:
+					check(false);
+				}
 			}
 		}
 
