Potentialy fix for BufferPhysicsResult crash

diff --git a/Engine/Source/Runtime/Experimental/Chaos/Public/Chaos/PBDRigidsSOAs.h b/Engine/Source/Runtime/Experimental/Chaos/Public/Chaos/PBDRigidsSOAs.h
--- a/Engine/Source/Runtime/Experimental/Chaos/Public/Chaos/PBDRigidsSOAs.h
+++ b/Engine/Source/Runtime/Experimental/Chaos/Public/Chaos/PBDRigidsSOAs.h
@@ -190,7 +190,7 @@ public:
 		check(Particle->CastToClustered() == nullptr);	//not supported
 
 		auto PBDRigid = Particle->CastToRigidParticle();
-		if(PBDRigid && PBDRigid->ObjectState() == EObjectStateType::Dynamic)
+		if(PBDRigid)
 		{
 			RemoveFromMapAndArray(PBDRigid, ActiveParticlesToIndex, ActiveParticlesArray);
 		}
