Chaos : Memory

- Delete the ProxyData object if the handle is null

diff --git a/Engine/Source/Runtime/Experimental/ChaosSolvers/Private/PBDRigidsSolver.cpp b/Engine/Source/Runtime/Experimental/ChaosSolvers/Private/PBDRigidsSolver.cpp
--- a/Engine/Source/Runtime/Experimental/ChaosSolvers/Private/PBDRigidsSolver.cpp
+++ b/Engine/Source/Runtime/Experimental/ChaosSolvers/Private/PBDRigidsSolver.cpp
@@ -493,6 +493,10 @@ namespace Chaos
 
 				Proxy->ClearAccumulatedData();
 			}
+			else
+			{
+				delete ProxyData;
+			}
 		}
 	}
 
