Make sure TArrayMap resets slots. Fixes crash

diff --git a/Engine/Source/Runtime/Experimental/Chaos/Public/Chaos/ISpatialAcceleration.h b/Engine/Source/Runtime/Experimental/Chaos/Public/Chaos/ISpatialAcceleration.h
--- a/Engine/Source/Runtime/Experimental/Chaos/Public/Chaos/ISpatialAcceleration.h
+++ b/Engine/Source/Runtime/Experimental/Chaos/Public/Chaos/ISpatialAcceleration.h
@@ -433,7 +433,7 @@ public:
 	void Remove(const TKey& Key)
 	{
 		const int32 Idx = GetUniqueIdx(Key).Idx;
-		Entries[Idx].bSet = false;
+		Entries[Idx] = FEntry();	//Mark as free, also resets default values for next use of value
 #if CHAOS_SERIALIZE_OUT
 		KeysToSerializeOut[Idx] = TKey();
 #endif
@@ -475,11 +475,6 @@ private:
 
 	struct FEntry
 	{
-#if CHAOS_SERIALIZE_OUT
-		//The indices are generated at runtime, so there's no way to serialize them directly
-		//Because of that we serialize the actual key which we can find, and then at runtime we use its transient index
-		TKey KeyToSerializeOut;
-#endif
 		TValue Value;
 		bool bSet;
 
