Chaos: speculative fix for sleeping event crash

diff --git a/Engine/Source/Runtime/Experimental/Chaos/Public/Chaos/PBDRigidsSOAs.h b/Engine/Source/Runtime/Experimental/Chaos/Public/Chaos/PBDRigidsSOAs.h
--- a/Engine/Source/Runtime/Experimental/Chaos/Public/Chaos/PBDRigidsSOAs.h
+++ b/Engine/Source/Runtime/Experimental/Chaos/Public/Chaos/PBDRigidsSOAs.h
@@ -200,6 +200,7 @@ public:
 		{
 			// Check for sleep events referencing this particle
 			// TODO think about this case more
+			GetDynamicParticles().GetSleepDataLock().WriteLock();
 			auto& SleepData = GetDynamicParticles().GetSleepData();
 			for (int32 Idx = 0; Idx < SleepData.Num(); ++Idx)
 			{
@@ -209,6 +210,7 @@ public:
 					break;
 				}
 			}
+			GetDynamicParticles().GetSleepDataLock().WriteUnlock();
 		}
 
 		//NOTE: This assumes that we are never creating a PT particle that is replicated to GT
diff --git a/Engine/Source/Runtime/Experimental/Chaos/Public/Chaos/RigidParticles.h b/Engine/Source/Runtime/Experimental/Chaos/Public/Chaos/RigidParticles.h
--- a/Engine/Source/Runtime/Experimental/Chaos/Public/Chaos/RigidParticles.h
+++ b/Engine/Source/Runtime/Experimental/Chaos/Public/Chaos/RigidParticles.h
@@ -171,10 +171,17 @@ class TRigidParticles : public TKinematicGeometryParticles<T, d>
 		SleepData.Particle = Particle;
 		SleepData.Sleeping = Sleeping;
 
-		SleepDataLock.Lock();
+		SleepDataLock.WriteLock();
 		MSleepData.Add(SleepData);
-		SleepDataLock.Unlock();
+		SleepDataLock.WriteUnlock();
 	}
+	void ClearSleepData()
+	{
+		SleepDataLock.WriteLock();
+		MSleepData.Empty();
+		SleepDataLock.WriteUnlock();
+	}
+	FRWLock& GetSleepDataLock() { return SleepDataLock; }
 
 	FORCEINLINE const EObjectStateType ObjectState(const int32 Index) const { return MObjectState[Index]; }
 	FORCEINLINE EObjectStateType& ObjectState(const int32 Index) { return MObjectState[Index]; }
@@ -227,7 +234,7 @@ class TRigidParticles : public TKinematicGeometryParticles<T, d>
 	TArrayCollectionArray<EObjectStateType> MObjectState;
 
 	TArray<TSleepData<T, d>> MSleepData;
-	FCriticalSection SleepDataLock;
+	FRWLock SleepDataLock;
 };
 
 
diff --git a/Engine/Source/Runtime/Experimental/ChaosSolvers/Private/EventDefaults.cpp b/Engine/Source/Runtime/Experimental/ChaosSolvers/Private/EventDefaults.cpp
--- a/Engine/Source/Runtime/Experimental/ChaosSolvers/Private/EventDefaults.cpp
+++ b/Engine/Source/Runtime/Experimental/ChaosSolvers/Private/EventDefaults.cpp
@@ -364,13 +364,15 @@ namespace Chaos
 			AllSleepIndicesByPhysicsProxy.Reset();
 
 			Chaos::FPBDRigidsSolver* NonConstSolver = (Chaos::FPBDRigidsSolver*)(Solver);
+
+			NonConstSolver->Particles.GetDynamicParticles().GetSleepDataLock().ReadLock();
 			auto& SolverSleepingData = NonConstSolver->Particles.GetDynamicParticles().GetSleepData();
-			for(const TSleepData<float, 3>& SleepData : SolverSleepingData)
+			for (const TSleepData<float, 3>& SleepData : SolverSleepingData)
 			{
-				if(SleepData.Particle)
+				if (SleepData.Particle)
 				{
 					TGeometryParticle<float, 3>* Particle = SleepData.Particle->GTGeometryParticle();
-					if(Particle->Proxy != nullptr)
+					if (Particle->Proxy != nullptr)
 					{
 						int32 NewIdx = EventSleepDataArray.Add(TSleepingData<float, 3>());
 						TSleepingData<float, 3>& SleepingDataArrayItem = EventSleepDataArray[NewIdx];
@@ -382,8 +384,10 @@ namespace Chaos
 					}
 				}
 			}
+			NonConstSolver->Particles.GetDynamicParticles().GetSleepDataLock().ReadUnlock();
+
+			NonConstSolver->Particles.GetDynamicParticles().ClearSleepData();
 
-			SolverSleepingData.Empty();
 
 		});
 	}
