Chaos - eliminate allocations in collision narrowphase

diff --git a/Engine/Source/Runtime/Experimental/Chaos/Private/Chaos/CollisionResolution.cpp b/Engine/Source/Runtime/Experimental/Chaos/Private/Chaos/CollisionResolution.cpp
--- a/Engine/Source/Runtime/Experimental/Chaos/Private/Chaos/CollisionResolution.cpp
+++ b/Engine/Source/Runtime/Experimental/Chaos/Private/Chaos/CollisionResolution.cpp
@@ -259,18 +259,9 @@ namespace Chaos
 			{
 				TRigidTransform<T, d> ParticleImplicit0TM = Transform0.GetRelativeTransform(Collisions::GetTransform(Particle0));
 				TRigidTransform<T, d> ParticleImplicit1TM = Transform1.GetRelativeTransform(Collisions::GetTransform(Particle1));
-				FRigidBodyPointContactConstraint* Constraint = new FRigidBodyPointContactConstraint(Particle0, Implicit0, ParticleImplicit0TM, Particle1, Implicit1, ParticleImplicit1TM);
-
-				UpdateBoxBoxConstraint(Object0->BoundingBox(), Transform0, Object1->BoundingBox(), Transform1, Thickness, *Constraint);
-
-				if (Constraint->GetPhi() < Thickness)
-				{
-					NewConstraints.Add(Constraint);
-				}
-				else
-				{
-					delete Constraint;
-				}
+				FRigidBodyPointContactConstraint Constraint = FRigidBodyPointContactConstraint(Particle0, Implicit0, ParticleImplicit0TM, Particle1, Implicit1, ParticleImplicit1TM);
+				UpdateBoxBoxConstraint(Object0->BoundingBox(), Transform0, Object1->BoundingBox(), Transform1, Thickness, Constraint);
+				NewConstraints.TryAdd(Thickness, Constraint);
 			}
 		}
 
@@ -307,18 +298,9 @@ namespace Chaos
 			{
 				TRigidTransform<T, d> ParticleImplicit0TM = Transform0.GetRelativeTransform(Collisions::GetTransform(Particle0));
 				TRigidTransform<T, d> ParticleImplicit1TM = Transform1.GetRelativeTransform(Collisions::GetTransform(Particle1));
-				FRigidBodyPointContactConstraint* Constraint = new FRigidBodyPointContactConstraint(Particle0, Implicit0, ParticleImplicit0TM, Particle1, Implicit1, ParticleImplicit1TM);
-
-				UpdateBoxHeightFieldConstraint(Object0->BoundingBox(), Transform0, *Object1, Transform1, Thickness, *Constraint);
-
-				if (Constraint->GetPhi() < Thickness)
-				{
-					NewConstraints.Add(Constraint);
-				}
-				else
-				{
-					delete Constraint;
-				}
+				FRigidBodyPointContactConstraint Constraint = FRigidBodyPointContactConstraint(Particle0, Implicit0, ParticleImplicit0TM, Particle1, Implicit1, ParticleImplicit1TM);
+				UpdateBoxHeightFieldConstraint(Object0->BoundingBox(), Transform0, *Object1, Transform1, Thickness, Constraint);
+				NewConstraints.TryAdd(Thickness, Constraint);
 			}
 		}
 
@@ -404,18 +386,9 @@ namespace Chaos
 			{
 				TRigidTransform<T, d> ParticleImplicit0TM = Transform0.GetRelativeTransform(Collisions::GetTransform(Particle0));
 				TRigidTransform<T, d> ParticleImplicit1TM = Transform1.GetRelativeTransform(Collisions::GetTransform(Particle1));
-				FRigidBodyPointContactConstraint* Constraint = new FRigidBodyPointContactConstraint(Particle0, Implicit0, ParticleImplicit0TM, Particle1, Implicit1, ParticleImplicit1TM);
-
-				UpdateBoxPlaneConstraint(Object0->BoundingBox(), Transform0, *Object1, Transform1, Thickness, *Constraint);
-
-				if (Constraint->GetPhi() < Thickness)
-				{
-					NewConstraints.Add(Constraint);
-				}
-				else
-				{
-					delete Constraint;
-				}
+				FRigidBodyPointContactConstraint Constraint = FRigidBodyPointContactConstraint(Particle0, Implicit0, ParticleImplicit0TM, Particle1, Implicit1, ParticleImplicit1TM);
+				UpdateBoxPlaneConstraint(Object0->BoundingBox(), Transform0, *Object1, Transform1, Thickness, Constraint);
+				NewConstraints.TryAdd(Thickness, Constraint);
 			}
 		}
 
@@ -463,18 +436,9 @@ namespace Chaos
 			{
 				TRigidTransform<T, d> ParticleImplicit0TM = Transform0.GetRelativeTransform(Collisions::GetTransform(Particle0));
 				TRigidTransform<T, d> ParticleImplicit1TM = Transform1.GetRelativeTransform(Collisions::GetTransform(Particle1));
-				FRigidBodyPointContactConstraint* Constraint = new FRigidBodyPointContactConstraint(Particle0, Implicit0, ParticleImplicit0TM, Particle1, Implicit1, ParticleImplicit1TM);
-
-				UpdateBoxTriangleMeshConstraint(Object0->GetAABB(), Transform0, *Object1, Transform1, Thickness, *Constraint);
-
-				if (Constraint->GetPhi() < Thickness)
-				{
-					NewConstraints.Add(Constraint);
-				}
-				else
-				{
-					delete Constraint;
-				}
+				FRigidBodyPointContactConstraint Constraint = FRigidBodyPointContactConstraint(Particle0, Implicit0, ParticleImplicit0TM, Particle1, Implicit1, ParticleImplicit1TM);
+				UpdateBoxTriangleMeshConstraint(Object0->GetAABB(), Transform0, *Object1, Transform1, Thickness, Constraint);
+				NewConstraints.TryAdd(Thickness, Constraint);
 			}
 		}
 
@@ -522,18 +486,9 @@ namespace Chaos
 			{
 				TRigidTransform<T, d> ParticleImplicit0TM = Transform0.GetRelativeTransform(Collisions::GetTransform(Particle0));
 				TRigidTransform<T, d> ParticleImplicit1TM = Transform1.GetRelativeTransform(Collisions::GetTransform(Particle1));
-				FRigidBodyPointContactConstraint* Constraint = new FRigidBodyPointContactConstraint(Particle0, Implicit0, ParticleImplicit0TM, Particle1, Implicit1, ParticleImplicit1TM);
-
-				UpdateSphereSphereConstraint(*Object0, Transform0, *Object1, Transform1, Thickness, *Constraint);
-
-				if (Constraint->GetPhi() < Thickness)
-				{
-					NewConstraints.Add(Constraint);
-				}
-				else
-				{
-					delete Constraint;
-				}
+				FRigidBodyPointContactConstraint Constraint = FRigidBodyPointContactConstraint(Particle0, Implicit0, ParticleImplicit0TM, Particle1, Implicit1, ParticleImplicit1TM);
+				UpdateSphereSphereConstraint(*Object0, Transform0, *Object1, Transform1, Thickness, Constraint);
+				NewConstraints.TryAdd(Thickness, Constraint);
 			}
 		}
 
@@ -570,18 +525,9 @@ namespace Chaos
 			{
 				TRigidTransform<T, d> ParticleImplicit0TM = Transform0.GetRelativeTransform(Collisions::GetTransform(Particle0));
 				TRigidTransform<T, d> ParticleImplicit1TM = Transform1.GetRelativeTransform(Collisions::GetTransform(Particle1));
-				FRigidBodyPointContactConstraint* Constraint = new FRigidBodyPointContactConstraint(Particle0, Implicit0, ParticleImplicit0TM, Particle1, Implicit1, ParticleImplicit1TM);
-
-				UpdateSphereHeightFieldConstraint(*Object0, Transform0, *Object1, Transform1, Thickness, *Constraint);
-
-				if (Constraint->GetPhi() < Thickness)
-				{
-					NewConstraints.Add(Constraint);
-				}
-				else
-				{
-					delete Constraint;
-				}
+				FRigidBodyPointContactConstraint Constraint = FRigidBodyPointContactConstraint(Particle0, Implicit0, ParticleImplicit0TM, Particle1, Implicit1, ParticleImplicit1TM);
+				UpdateSphereHeightFieldConstraint(*Object0, Transform0, *Object1, Transform1, Thickness, Constraint);
+				NewConstraints.TryAdd(Thickness, Constraint);
 			}
 		}
 
@@ -626,18 +572,9 @@ namespace Chaos
 			{
 				TRigidTransform<T, d> ParticleImplicit0TM = Transform0.GetRelativeTransform(Collisions::GetTransform(Particle0));
 				TRigidTransform<T, d> ParticleImplicit1TM = Transform1.GetRelativeTransform(Collisions::GetTransform(Particle1));
-				FRigidBodyPointContactConstraint* Constraint = new FRigidBodyPointContactConstraint(Particle0, Implicit0, ParticleImplicit0TM, Particle1, Implicit1, ParticleImplicit1TM);
-
-				UpdateSpherePlaneConstraint(*Object0, Transform0, *Object1, Transform1, Thickness, *Constraint);
-
-				if (Constraint->GetPhi() < Thickness)
-				{
-					NewConstraints.Add(Constraint);
-				}
-				else
-				{
-					delete Constraint;
-				}
+				FRigidBodyPointContactConstraint Constraint = FRigidBodyPointContactConstraint(Particle0, Implicit0, ParticleImplicit0TM, Particle1, Implicit1, ParticleImplicit1TM);
+				UpdateSpherePlaneConstraint(*Object0, Transform0, *Object1, Transform1, Thickness, Constraint);
+				NewConstraints.TryAdd(Thickness, Constraint);
 			}
 		}
 
@@ -686,18 +623,9 @@ namespace Chaos
 			{
 				TRigidTransform<T, d> ParticleImplicit0TM = Transform0.GetRelativeTransform(Collisions::GetTransform(Particle0));
 				TRigidTransform<T, d> ParticleImplicit1TM = Transform1.GetRelativeTransform(Collisions::GetTransform(Particle1));
-				FRigidBodyPointContactConstraint* Constraint = new FRigidBodyPointContactConstraint(Particle0, Implicit0, ParticleImplicit0TM, Particle1, Implicit1, ParticleImplicit1TM);
-
-				UpdateSphereBoxConstraint(*Object0, Transform0, Object1->BoundingBox(), Transform1, Thickness, *Constraint);
-
-				if (Constraint->GetPhi() < Thickness)
-				{
-					NewConstraints.Add(Constraint);
-				}
-				else
-				{
-					delete Constraint;
-				}
+				FRigidBodyPointContactConstraint Constraint = FRigidBodyPointContactConstraint(Particle0, Implicit0, ParticleImplicit0TM, Particle1, Implicit1, ParticleImplicit1TM);
+				UpdateSphereBoxConstraint(*Object0, Transform0, Object1->BoundingBox(), Transform1, Thickness, Constraint);
+				NewConstraints.TryAdd(Thickness, Constraint);
 			}
 		}
 
@@ -752,18 +680,9 @@ namespace Chaos
 			{
 				TRigidTransform<T, d> ParticleImplicit0TM = Transform0.GetRelativeTransform(Collisions::GetTransform(Particle0));
 				TRigidTransform<T, d> ParticleImplicit1TM = Transform1.GetRelativeTransform(Collisions::GetTransform(Particle1));
-				FRigidBodyPointContactConstraint* Constraint = new FRigidBodyPointContactConstraint(Particle0, Implicit0, ParticleImplicit0TM, Particle1, Implicit1, ParticleImplicit1TM);
-
-				UpdateSphereCapsuleConstraint(*Object0, Transform0, *Object1, Transform1, Thickness, *Constraint);
-
-				if (Constraint->GetPhi() < Thickness)
-				{
-					NewConstraints.Add(Constraint);
-				}
-				else
-				{
-					delete Constraint;
-				}
+				FRigidBodyPointContactConstraint Constraint = FRigidBodyPointContactConstraint(Particle0, Implicit0, ParticleImplicit0TM, Particle1, Implicit1, ParticleImplicit1TM);
+				UpdateSphereCapsuleConstraint(*Object0, Transform0, *Object1, Transform1, Thickness, Constraint);
+				NewConstraints.TryAdd(Thickness, Constraint);
 			}
 		}
 
@@ -821,18 +740,9 @@ namespace Chaos
 			{
 				TRigidTransform<T, d> ParticleImplicit0TM = Transform0.GetRelativeTransform(Collisions::GetTransform(Particle0));
 				TRigidTransform<T, d> ParticleImplicit1TM = Transform1.GetRelativeTransform(Collisions::GetTransform(Particle1));
-				FRigidBodyPointContactConstraint* Constraint = new FRigidBodyPointContactConstraint(Particle0, Implicit0, ParticleImplicit0TM, Particle1, Implicit1, ParticleImplicit1TM);
-
-				UpdateCapsuleCapsuleConstraint(*Object0, Transform0, *Object1, Transform1, Thickness, *Constraint);
-
-				if (Constraint->GetPhi() < Thickness)
-				{
-					NewConstraints.Add(Constraint);
-				}
-				else
-				{
-					delete Constraint;
-				}
+				FRigidBodyPointContactConstraint Constraint = FRigidBodyPointContactConstraint(Particle0, Implicit0, ParticleImplicit0TM, Particle1, Implicit1, ParticleImplicit1TM);
+				UpdateCapsuleCapsuleConstraint(*Object0, Transform0, *Object1, Transform1, Thickness, Constraint);
+				NewConstraints.TryAdd(Thickness, Constraint);
 			}
 		}
 
@@ -869,18 +779,9 @@ namespace Chaos
 			{
 				TRigidTransform<T, d> ParticleImplicit0TM = Transform0.GetRelativeTransform(Collisions::GetTransform(Particle0));
 				TRigidTransform<T, d> ParticleImplicit1TM = Transform1.GetRelativeTransform(Collisions::GetTransform(Particle1));
-				FRigidBodyPointContactConstraint* Constraint = new FRigidBodyPointContactConstraint(Particle0, Implicit0, ParticleImplicit0TM, Particle1, Implicit1, ParticleImplicit1TM);
-
-				UpdateCapsuleBoxConstraint(*Object0, Transform0, Object1->BoundingBox(), Transform1, Thickness, *Constraint);
-
-				if (Constraint->GetPhi() < Thickness)
-				{
-					NewConstraints.Add(Constraint);
-				}
-				else
-				{
-					delete Constraint;
-				}
+				FRigidBodyPointContactConstraint Constraint = FRigidBodyPointContactConstraint(Particle0, Implicit0, ParticleImplicit0TM, Particle1, Implicit1, ParticleImplicit1TM);
+				UpdateCapsuleBoxConstraint(*Object0, Transform0, Object1->BoundingBox(), Transform1, Thickness, Constraint);
+				NewConstraints.TryAdd(Thickness, Constraint);
 			}
 		}
 
@@ -917,18 +818,9 @@ namespace Chaos
 			{
 				TRigidTransform<T, d> ParticleImplicit0TM = Transform0.GetRelativeTransform(Collisions::GetTransform(Particle0));
 				TRigidTransform<T, d> ParticleImplicit1TM = Transform1.GetRelativeTransform(Collisions::GetTransform(Particle1));
-				FRigidBodyPointContactConstraint* Constraint = new FRigidBodyPointContactConstraint(Particle0, Implicit0, ParticleImplicit0TM, Particle1, Implicit1, ParticleImplicit1TM);
-
-				UpdateCapsuleHeightFieldConstraint(*Object0, Transform0, *Object1, Transform1, Thickness, *Constraint);
-
-				if (Constraint->GetPhi() < Thickness)
-				{
-					NewConstraints.Add(Constraint);
-				}
-				else
-				{
-					delete Constraint;
-				}
+				FRigidBodyPointContactConstraint Constraint = FRigidBodyPointContactConstraint(Particle0, Implicit0, ParticleImplicit0TM, Particle1, Implicit1, ParticleImplicit1TM);
+				UpdateCapsuleHeightFieldConstraint(*Object0, Transform0, *Object1, Transform1, Thickness, Constraint);
+				NewConstraints.TryAdd(Thickness, Constraint);
 			}
 		}
 
@@ -996,20 +888,10 @@ namespace Chaos
 		{
 			TRigidTransform<T, d> ParticleImplicit0TM = Transform0.GetRelativeTransform(Collisions::GetTransform(Particle0));
 			TRigidTransform<T, d> ParticleImplicit1TM = Transform1.GetRelativeTransform(Collisions::GetTransform(Particle1));
-			FRigidBodyMultiPointContactConstraint* Constraint = new FRigidBodyMultiPointContactConstraint(Particle0, Implicit0, ParticleImplicit0TM, Particle1, Implicit1, ParticleImplicit1TM);
-
-			UpdateConvexConvexManifold(*Constraint, Transform0, Transform1, Thickness);
-			UpdateConvexConvexConstraint(*Implicit0, Transform0, *Implicit1, Transform1, Thickness, *Constraint);
-
-			if (Constraint->GetPhi() < Thickness)
-			{
-				NewConstraints.Add(Constraint);
-			}
-			else
-			{
-				delete Constraint;
-			}
-
+			FRigidBodyMultiPointContactConstraint Constraint = FRigidBodyMultiPointContactConstraint(Particle0, Implicit0, ParticleImplicit0TM, Particle1, Implicit1, ParticleImplicit1TM);
+			UpdateConvexConvexManifold(Constraint, Transform0, Transform1, Thickness);
+			UpdateConvexConvexConstraint(*Implicit0, Transform0, *Implicit1, Transform1, Thickness, Constraint);
+			NewConstraints.TryAdd(Thickness, Constraint);
 		}
 
 		//
@@ -1044,18 +926,9 @@ namespace Chaos
 			{
 				TRigidTransform<T, d> ParticleImplicit0TM = Transform0.GetRelativeTransform(Collisions::GetTransform(Particle0));
 				TRigidTransform<T, d> ParticleImplicit1TM = Transform1.GetRelativeTransform(Collisions::GetTransform(Particle1));
-				FRigidBodyPointContactConstraint* Constraint = new FRigidBodyPointContactConstraint(Particle0, Implicit0, ParticleImplicit0TM, Particle1, Implicit1, ParticleImplicit1TM);
-
-				UpdateConvexHeightFieldConstraint(*Implicit0, Transform0, *Object1, Transform1, Thickness, *Constraint);
-
-				if (Constraint->GetPhi() < Thickness)
-				{
-					NewConstraints.Add(Constraint);
-				}
-				else
-				{
-					delete Constraint;
-				}
+				FRigidBodyPointContactConstraint Constraint = FRigidBodyPointContactConstraint(Particle0, Implicit0, ParticleImplicit0TM, Particle1, Implicit1, ParticleImplicit1TM);
+				UpdateConvexHeightFieldConstraint(*Implicit0, Transform0, *Object1, Transform1, Thickness, Constraint);
+				NewConstraints.TryAdd(Thickness, Constraint);
 			}
 		}
 
@@ -1107,32 +980,25 @@ namespace Chaos
 		{
 			TRigidTransform<T, d> ParticleImplicit0TM = Transform0.GetRelativeTransform(Collisions::GetTransform(Particle0));
 			TRigidTransform<T, d> ParticleImplicit1TM = Transform1.GetRelativeTransform(Collisions::GetTransform(Particle1));
-			FRigidBodyPointContactConstraint* Constraint = new FRigidBodyPointContactConstraint(nullptr, nullptr, ParticleImplicit0TM, nullptr, nullptr, ParticleImplicit1TM);
+			FRigidBodyPointContactConstraint Constraint = FRigidBodyPointContactConstraint(Particle0, Implicit0, ParticleImplicit0TM, Particle1, Implicit1, ParticleImplicit1TM);
 
 			bool bIsParticleDynamic0 = Particle0->CastToRigidParticle() && Particle0->ObjectState() == EObjectStateType::Dynamic;
 			if (!Particle1->Geometry() || (bIsParticleDynamic0 && !Particle0->CastToRigidParticle()->CollisionParticlesSize() && Particle0->Geometry() && !Particle0->Geometry()->IsUnderlyingUnion()))
 			{
-				Constraint->Particle[0] = Particle1;
-				Constraint->Particle[1] = Particle0;
-				Constraint->SetManifold(Implicit1, Implicit0);
+				Constraint.Particle[0] = Particle1;
+				Constraint.Particle[1] = Particle0;
+				Constraint.SetManifold(Implicit1, Implicit0);
 			}
 			else
 			{
-				Constraint->Particle[0] = Particle0;
-				Constraint->Particle[1] = Particle1;
-				Constraint->SetManifold(Implicit0, Implicit1);
+				Constraint.Particle[0] = Particle0;
+				Constraint.Particle[1] = Particle1;
+				Constraint.SetManifold(Implicit0, Implicit1);
 			}
 
-			UpdateLevelsetLevelsetConstraint<ECollisionUpdateType::Any>(Thickness, *Constraint);
+			UpdateLevelsetLevelsetConstraint<ECollisionUpdateType::Any>(Thickness, Constraint);
 
-			if (Constraint->GetPhi() < Thickness)
-			{
-				NewConstraints.Add(Constraint);
-			}
-			else
-			{
-				delete Constraint;
-			}
+			NewConstraints.TryAdd(Thickness, Constraint);
 		}
 
 		//
diff --git a/Engine/Source/Runtime/Experimental/Chaos/Private/Chaos/PBDCollisionConstraints.cpp b/Engine/Source/Runtime/Experimental/Chaos/Private/Chaos/PBDCollisionConstraints.cpp
--- a/Engine/Source/Runtime/Experimental/Chaos/Private/Chaos/PBDCollisionConstraints.cpp
+++ b/Engine/Source/Runtime/Experimental/Chaos/Private/Chaos/PBDCollisionConstraints.cpp
@@ -144,38 +144,31 @@ namespace Chaos
 	}
 
 	template<typename T, int d>
-	void TPBDCollisionConstraints<T, d>::AddConstraint(FConstraintBase* ConstraintBase)
+	void TPBDCollisionConstraints<T, d>::AddConstraint(const TRigidBodyPointContactConstraint<FReal, 3>& InConstraint)
 	{
-		// WARNING : ConstraintBase is about to be deleted!
+		int32 Idx = PointConstraints.Add(InConstraint);
+		FConstraintContainerHandle* Handle = HandleAllocator.template AllocHandle< TRigidBodyPointContactConstraint<T, d> >(this, Idx);
+		Handle->GetContact().Timestamp = -INT_MAX; // force point constraints to be deleted.
 
-		FConstraintContainerHandle* Handle = nullptr;
-
-		UpdateConstraintMaterialProperties(*ConstraintBase);
-
-		if (ConstraintBase->GetType() == TRigidBodyPointContactConstraint<T, 3>::StaticType())
-		{
-			TRigidBodyPointContactConstraint<T, d>* PointConstraint = ConstraintBase->template As< TRigidBodyPointContactConstraint<T, d> >();
-
-			int32 Idx = PointConstraints.Add(*PointConstraint);
-			Handle = HandleAllocator.template AllocHandle< TRigidBodyPointContactConstraint<T, d> >(this, Idx);
-			Handle->GetContact().Timestamp = -INT_MAX; // force point constraints to be deleted.
-
-			delete PointConstraint;
-		}
-		else if (ConstraintBase->GetType() == TRigidBodyMultiPointContactConstraint<T, 3>::StaticType())
-		{
-			TRigidBodyMultiPointContactConstraint<T, d>* IterativeConstraint = ConstraintBase->template As< TRigidBodyMultiPointContactConstraint<T, d> >();
+		check(Handle != nullptr);
+		Handles.Add(Handle);
+		Manifolds.Add(Handle->GetKey(), Handle);
 
-			int32 Idx = IterativeConstraints.Add(*IterativeConstraint);
-			Handle = HandleAllocator.template AllocHandle< TRigidBodyMultiPointContactConstraint<T, d> >(this, Idx);
-			Handle->GetContact().Timestamp = LifespanCounter;
+		UpdateConstraintMaterialProperties(PointConstraints[Idx]);
+	}
 
-			delete IterativeConstraint;
-		}
+	template<typename T, int d>
+	void TPBDCollisionConstraints<T, d>::AddConstraint(const TRigidBodyMultiPointContactConstraint<FReal, 3>& InConstraint)
+	{
+		int32 Idx = IterativeConstraints.Add(InConstraint);
+		FConstraintContainerHandle* Handle = HandleAllocator.template AllocHandle< TRigidBodyMultiPointContactConstraint<T, d> >(this, Idx);
+		Handle->GetContact().Timestamp = LifespanCounter;
 
 		check(Handle != nullptr);
 		Handles.Add(Handle);
 		Manifolds.Add(Handle->GetKey(), Handle);
+
+		UpdateConstraintMaterialProperties(IterativeConstraints[Idx]);
 	}
 
 	template<typename T, int d>
diff --git a/Engine/Source/Runtime/Experimental/Chaos/Public/Chaos/Collision/CollisionReceiver.h b/Engine/Source/Runtime/Experimental/Chaos/Public/Chaos/Collision/CollisionReceiver.h
--- a/Engine/Source/Runtime/Experimental/Chaos/Public/Chaos/Collision/CollisionReceiver.h
+++ b/Engine/Source/Runtime/Experimental/Chaos/Public/Chaos/Collision/CollisionReceiver.h
@@ -24,11 +24,15 @@ namespace Chaos
 		/**
 		 * Called by a CollisionDetector (possibly in a task) when it finds collisions.
 		 */
-		void ReceiveCollisions(const TArray<TCollisionConstraintBase<float, 3>*>& Constraints)
+		void ReceiveCollisions(const FCollisionConstraintsArray& Constraints)
 		{
-			for (TCollisionConstraintBase<FReal, 3>* ConstraintBase : Constraints)
+			for (const TRigidBodyPointContactConstraint<FReal, 3>& Constraint : Constraints.SinglePointConstraints)
 			{
-				Queue.Enqueue(ConstraintBase);
+				SinglePointQueue.Enqueue(Constraint);
+			}
+			for (const TRigidBodyMultiPointContactConstraint<FReal, 3>& Constraint : Constraints.MultiPointConstraints)
+			{
+				MultiPointQueue.Enqueue(Constraint);
 			}
 		}
 
@@ -37,24 +41,33 @@ namespace Chaos
 		 */
 		void ProcessCollisions()
 		{
-			FCollisionConstraintBase* ConstraintBase = nullptr;
-			while (Queue.Dequeue(ConstraintBase))
+			// @todo(chaos) : Collision Constraints
+			//    This needs to be moved within the MidPhase (which does not exist yet). The 
+			//    NarrowPhase is creating temporary constraints for all interacting bodies 
+			//    even though we might already have those bodies within the constraint.
+			//    So there is a bunch of wasted computation when the CreateConstraints is called. 
+			TRigidBodyPointContactConstraint<FReal, 3> SinglePointConstraint;
+			while (SinglePointQueue.Dequeue(SinglePointConstraint))
+			{
+				if (!CollisionConstraints.Contains(&SinglePointConstraint))
+				{
+					CollisionConstraints.AddConstraint(SinglePointConstraint);
+				}
+			}
+			TRigidBodyMultiPointContactConstraint<FReal, 3> MultiPointConstraint;
+			while (MultiPointQueue.Dequeue(MultiPointConstraint))
 			{
-				// @todo(chaos) : Collision Constraints
-				//    This needs to be moved within the MidPhase (which does not exist yet). The 
-				//    NarrowPhase is creating temporary constraints for all interacting bodies 
-				//    even though we might already have those bodies within the constraint.
-				//    So there is a bunch of wasted computation when the CreateConstraints is called. 
-				if (!CollisionConstraints.Contains(ConstraintBase))
+				if (!CollisionConstraints.Contains(&MultiPointConstraint))
 				{
-					CollisionConstraints.AddConstraint(ConstraintBase);
+					CollisionConstraints.AddConstraint(MultiPointConstraint);
 				}
 			}
 		}
 
 	private:
 		//todo(ocohen): use per thread buffer instead, need better support than ParallelFor for this
-		TQueue<FCollisionConstraintBase*, EQueueMode::Mpsc> Queue;
+		TQueue<TRigidBodyPointContactConstraint<FReal, 3>, EQueueMode::Mpsc> SinglePointQueue;
+		TQueue<TRigidBodyMultiPointContactConstraint<FReal, 3>, EQueueMode::Mpsc> MultiPointQueue;
 		FCollisionConstraints& CollisionConstraints;
 	};
 
@@ -76,11 +89,15 @@ namespace Chaos
 		/**
 		 * Called by a CollisionDetector (possibly in a task) when it finds collisions.
 		 */
-		void ReceiveCollisions(const TArray<TCollisionConstraintBase<float, 3>*>& Constraints)
+		void ReceiveCollisions(const FCollisionConstraintsArray& Constraints)
 		{
-			for (TCollisionConstraintBase<FReal, 3>* ConstraintBase : Constraints)
+			for (const TRigidBodyPointContactConstraint<FReal, 3>& Constraint : Constraints.SinglePointConstraints)
+			{
+				CollisionConstraints.AddConstraint(Constraint);
+			}
+			for (const TRigidBodyMultiPointContactConstraint<FReal, 3>& Constraint : Constraints.MultiPointConstraints)
 			{
-				CollisionConstraints.AddConstraint(ConstraintBase);
+				CollisionConstraints.AddConstraint(Constraint);
 			}
 		}
 
diff --git a/Engine/Source/Runtime/Experimental/Chaos/Public/Chaos/Collision/NarrowPhase.h b/Engine/Source/Runtime/Experimental/Chaos/Public/Chaos/Collision/NarrowPhase.h
--- a/Engine/Source/Runtime/Experimental/Chaos/Public/Chaos/Collision/NarrowPhase.h
+++ b/Engine/Source/Runtime/Experimental/Chaos/Public/Chaos/Collision/NarrowPhase.h
@@ -14,34 +14,18 @@ namespace Chaos
 	DECLARE_CYCLE_STAT_EXTERN(TEXT("Collisions::NarrowPhase"), STAT_Collisions_NarrowPhase, STATGROUP_ChaosCollision, CHAOS_API);
 
 	/**
-	 * Generate contact manifolds for particle pairs and pass them on to the consumer. Can be composed with a
-	 * multi-threaded BroadPhase as long as the collisions receiver type can handle multi-threaded calls to ReceiveCollisions.
+	 * Generate contact manifolds for particle pairs.
 	 *
-	 * /template T_RECEIVER The object type that will take the detected collisions. Generally this will depend on the type of BroadPhase.
 	 * /see FAsyncCollisionReceiver, FSyncCollisionReceiver.
 	 */
 	class CHAOS_API FNarrowPhase
 	{
 	public:
-		using FCollisionConstraintsArray = TArray<TCollisionConstraintBase<FReal, 3>*>;
-
-		// @todo(ccaulfield): COLLISION Transient Handle version
-		template<typename T_RECEIVER>
-		void GenerateCollisions(FReal Dt, T_RECEIVER& Receiver, TGeometryParticleHandle<FReal, 3>* Particle1, TGeometryParticleHandle<FReal, 3>* Particle2, const FReal BoundsThickness, CollisionStats::FStatData& StatData)
+		// @todo(chaos): COLLISION Transient Handle version
+		// @todo(chaos): COLLISION - Thickness: add shape padding (BoundsThickness is the distance within which we speculatively create constraints)
+		void GenerateCollisions(FCollisionConstraintsArray& NewConstraints, FReal Dt, TGeometryParticleHandle<FReal, 3>* Particle0, TGeometryParticleHandle<FReal, 3>* Particle1, const FReal AABBExpansion, CollisionStats::FStatData& StatData)
 		{
 			SCOPE_CYCLE_COUNTER(STAT_Collisions_NarrowPhase);
-
-			// @todo(ccaulfield): COLLISION - Thickness: add shape padding (BoundsThickness is the distance within which we speculatively create constraints)
-
-			FCollisionConstraintsArray NewConstraints;
-			ConstructConstraints(Particle1, Particle2, BoundsThickness, NewConstraints, StatData);
-
-			Receiver.ReceiveCollisions(NewConstraints);
-		}
-
-	private:
-		void ConstructConstraints(TGeometryParticleHandle<FReal, 3>* Particle0, TGeometryParticleHandle<FReal, 3>* Particle1, const FReal Thickness, FCollisionConstraintsArray& NewConstraints, CollisionStats::FStatData& StatData)
-		{
 			if (ensure(Particle0 && Particle1))
 			{
 				//
@@ -51,7 +35,7 @@ namespace Chaos
 				//   determine if the constraint is already defined, and then opt out of 
 				//   the creation process. 
 				//
-				Collisions::ConstructConstraints<FReal, 3>(Particle0, Particle1, Particle0->Geometry().Get(), Particle1->Geometry().Get(), Collisions::GetTransform(Particle0), Collisions::GetTransform(Particle1), Thickness, NewConstraints);
+				Collisions::ConstructConstraints<FReal, 3>(Particle0, Particle1, Particle0->Geometry().Get(), Particle1->Geometry().Get(), Collisions::GetTransform(Particle0), Collisions::GetTransform(Particle1), AABBExpansion, NewConstraints);
 
 				CHAOS_COLLISION_STAT(if (NewConstraints.Num()) { StatData.IncrementCountNP(NewConstraints.Num()); });
 				CHAOS_COLLISION_STAT(if (!NewConstraints.Num()) { StatData.IncrementRejectedNP(); });
diff --git a/Engine/Source/Runtime/Experimental/Chaos/Public/Chaos/Collision/ParticlePairBroadPhase.h b/Engine/Source/Runtime/Experimental/Chaos/Public/Chaos/Collision/ParticlePairBroadPhase.h
--- a/Engine/Source/Runtime/Experimental/Chaos/Public/Chaos/Collision/ParticlePairBroadPhase.h
+++ b/Engine/Source/Runtime/Experimental/Chaos/Public/Chaos/Collision/ParticlePairBroadPhase.h
@@ -38,10 +38,10 @@ namespace Chaos
 		{
 			SCOPE_CYCLE_COUNTER(STAT_Collisions_BroadPhase);
 
-			const FReal BoundsThickness = 1.0f;
-			const FReal BoundsThicknessVelocityInflation = 2.0f;
+			const FReal AABBExpansion = 1.0f;		// WRONG!!!!
 
-			for (const FParticlePair& ParticlePair : ParticlePairs)	
+			FCollisionConstraintsArray NewConstraints;
+			for (const FParticlePair& ParticlePair : ParticlePairs)
 			{
 				// Array is const, things in it are not...
 				TGeometryParticleHandle<FReal, 3>* Particle0 = const_cast<TGeometryParticleHandle<FReal, 3>*>(ParticlePair[0]);
@@ -56,7 +56,9 @@ namespace Chaos
 					const TAABB<FReal, 3>& Box1 = Particle1->WorldSpaceInflatedBounds();
 					if (Box0.Intersects(Box1))
 					{
-						NarrowPhase.GenerateCollisions(Dt, Receiver, Particle0, Particle1, BoundsThickness, StatData);
+						NarrowPhase.GenerateCollisions(NewConstraints, Dt,Particle0, Particle1, AABBExpansion, StatData);
+						Receiver.ReceiveCollisions(NewConstraints);
+						NewConstraints.Empty();
 					}
 				}
 			}
diff --git a/Engine/Source/Runtime/Experimental/Chaos/Public/Chaos/Collision/SpatialAccelerationBroadPhase.h b/Engine/Source/Runtime/Experimental/Chaos/Public/Chaos/Collision/SpatialAccelerationBroadPhase.h
--- a/Engine/Source/Runtime/Experimental/Chaos/Public/Chaos/Collision/SpatialAccelerationBroadPhase.h
+++ b/Engine/Source/Runtime/Experimental/Chaos/Public/Chaos/Collision/SpatialAccelerationBroadPhase.h
@@ -211,7 +211,9 @@ namespace Chaos
 
 					const FReal Box2Thickness = bIsParticle2Dynamic ? ComputeBoundsThickness(*Particle2.CastToRigidParticle(), Dt, BoundsThickness, BoundsThicknessVelocityInflation).Size() : (FReal)0;
 
-					NarrowPhase.GenerateCollisions(Dt, Receiver, Particle1.Handle(), Particle2.Handle(), FMath::Max(Box1Thickness, Box2Thickness), StatData);
+					FCollisionConstraintsArray NewConstraints;
+					NarrowPhase.GenerateCollisions(NewConstraints, Dt, Particle1.Handle(), Particle2.Handle(), FMath::Max(Box1Thickness, Box2Thickness), StatData);
+					Receiver.ReceiveCollisions(NewConstraints);
 				}
 			}
 
diff --git a/Engine/Source/Runtime/Experimental/Chaos/Public/Chaos/CollisionResolutionTypes.h b/Engine/Source/Runtime/Experimental/Chaos/Public/Chaos/CollisionResolutionTypes.h
--- a/Engine/Source/Runtime/Experimental/Chaos/Public/Chaos/CollisionResolutionTypes.h
+++ b/Engine/Source/Runtime/Experimental/Chaos/Public/Chaos/CollisionResolutionTypes.h
@@ -155,7 +155,6 @@ namespace Chaos
 		FType Type;
 	};
 	typedef TCollisionConstraintBase<float, 3> FCollisionConstraintBase;
-	typedef TArray<TCollisionConstraintBase<float, 3>*> FCollisionConstraintsArray;
 
 
 	/*
@@ -345,4 +344,42 @@ namespace Chaos
 
 	};
 	typedef TPBDCollisionConstraintHandle<float, 3> FPBDCollisionConstraintHandle;
+
+
+	template<int T_MAXCONSTRAINTS>
+	struct TCollisionConstraintsStore
+	{
+		TArray<TRigidBodyPointContactConstraint<FReal, 3>, TInlineAllocator<T_MAXCONSTRAINTS>> SinglePointConstraints;
+		TArray<TRigidBodyMultiPointContactConstraint<FReal, 3>, TInlineAllocator<T_MAXCONSTRAINTS>> MultiPointConstraints;
+
+		int32 Num() const { return SinglePointConstraints.Num() + MultiPointConstraints.Num(); }
+
+		void Empty()
+		{
+			SinglePointConstraints.Empty();
+			MultiPointConstraints.Empty();
+		}
+
+		TRigidBodyPointContactConstraint<FReal, 3>* TryAdd(FReal MaxPhi, const TRigidBodyPointContactConstraint<FReal, 3>& C)
+		{
+			if ((SinglePointConstraints.Num() < T_MAXCONSTRAINTS) && (C.GetPhi() < MaxPhi))
+			{
+				int32 ConstraintIndex = SinglePointConstraints.Add(C);
+				return &SinglePointConstraints[ConstraintIndex];
+			}
+			return nullptr;
+		}
+
+		TRigidBodyMultiPointContactConstraint<FReal, 3>* TryAdd(FReal MaxPhi, const TRigidBodyMultiPointContactConstraint<FReal, 3>& C)
+		{
+			if ((MultiPointConstraints.Num() < T_MAXCONSTRAINTS) && (C.GetPhi() < MaxPhi))
+			{
+				int32 ConstraintIndex = MultiPointConstraints.Add(C);
+				return &MultiPointConstraints[ConstraintIndex];
+			}
+			return nullptr;
+		}
+	};
+
+	using FCollisionConstraintsArray = TCollisionConstraintsStore<8>;
 }
diff --git a/Engine/Source/Runtime/Experimental/Chaos/Public/Chaos/PBDCollisionConstraints.h b/Engine/Source/Runtime/Experimental/Chaos/Public/Chaos/PBDCollisionConstraints.h
--- a/Engine/Source/Runtime/Experimental/Chaos/Public/Chaos/PBDCollisionConstraints.h
+++ b/Engine/Source/Runtime/Experimental/Chaos/Public/Chaos/PBDCollisionConstraints.h
@@ -77,7 +77,8 @@ public:
 	*  This function will make a deep copy of the constraint and 
 	*  then delete the InConstraint. 
 	*/
-	void AddConstraint(FConstraintBase* InConstraint);
+	void AddConstraint(const TRigidBodyPointContactConstraint<FReal, 3>& InConstraint);
+	void AddConstraint(const TRigidBodyMultiPointContactConstraint<FReal, 3>& InConstraint);
 
 	/**
 	*  Reset the constraint frame. 
