Chaos
- add collision update fucntion table (WIP)
- add box-space bounds check to capsule-box
- warm-start GJK with existing normal
- rename collision Thickness to CullDistance in CollisionResponse.cpp
- fix world-space particle bounds in MinEvolution

diff --git a/Engine/Source/Runtime/Engine/Private/PhysicsEngine/ImmediatePhysics/ImmediatePhysicsChaos/ImmediatePhysicsSimulation_Chaos.cpp b/Engine/Source/Runtime/Engine/Private/PhysicsEngine/ImmediatePhysics/ImmediatePhysicsChaos/ImmediatePhysicsSimulation_Chaos.cpp
--- a/Engine/Source/Runtime/Engine/Private/PhysicsEngine/ImmediatePhysics/ImmediatePhysicsChaos/ImmediatePhysicsSimulation_Chaos.cpp
+++ b/Engine/Source/Runtime/Engine/Private/PhysicsEngine/ImmediatePhysics/ImmediatePhysicsChaos/ImmediatePhysicsSimulation_Chaos.cpp
@@ -31,23 +31,27 @@ float ChaosImmediate_Evolution_InitialStepTime = 0.033f;
 int32 ChaosImmediate_Evolution_DeltaTimeCount = 100;
 int32 ChaosImmediate_Evolution_Iterations = -1;
 int32 ChaosImmediate_Evolution_PushOutIterations = -1;
+float ChaosImmediate_Evolution_BoundsExtension = 0.0f;
 FAutoConsoleVariableRef CVarChaosImmPhysStepTime(TEXT("p.Chaos.ImmPhys.StepTime"), ChaosImmediate_Evolution_StepTime, TEXT("Override step time (if not zero)"));
 FAutoConsoleVariableRef CVarChaosImmPhysNumSteps(TEXT("p.Chaos.ImmPhys.NumSteps"), ChaosImmediate_Evolution_NumSteps, TEXT("Override num steps (if not zero)"));
 FAutoConsoleVariableRef CVarChaosImmPhysInitialStepTime(TEXT("p.Chaos.ImmPhys.InitialStepTime"), ChaosImmediate_Evolution_InitialStepTime, TEXT("Initial step time (then calculated from rolling average)"));
 FAutoConsoleVariableRef CVarChaosImmPhysDeltaTimeCount(TEXT("p.Chaos.ImmPhys.DeltaTimeCount"), ChaosImmediate_Evolution_DeltaTimeCount, TEXT("The number of ticks over which the moving average is calculated"));
 FAutoConsoleVariableRef CVarChaosImmPhysIterations(TEXT("p.Chaos.ImmPhys.Iterations"), ChaosImmediate_Evolution_Iterations, TEXT("Override number of constraint solver loops in immediate physics (if >= 0)"));
 FAutoConsoleVariableRef CVarChaosImmPhysPushOutIterations(TEXT("p.Chaos.ImmPhys.PushOutIterations"), ChaosImmediate_Evolution_PushOutIterations, TEXT("Override number of solver push-out loops (if >= 0)"));
+FAutoConsoleVariableRef CVarChaosImmPhysBoundsExtension(TEXT("p.Chaos.ImmPhys.BoundsExtension"), ChaosImmediate_Evolution_BoundsExtension, TEXT("Bounds are grown by this fraction of their size (should be >= 0.0)"));
 
 int32 ChaosImmediate_Collision_Enabled = 1;
 int32 ChaosImmediate_Collision_PairIterations = -1;
 int32 ChaosImmediate_Collision_PushOutPairIterations = -1;
 int32 ChaosImmediate_Collision_Priority = 1;
-float ChaosImmediate_Collision_Thickness = 0;
+float ChaosImmediate_Collision_CullDistance = 1.0f;
+float ChaosImmediate_Collision_ShapePadding = 0;
 FAutoConsoleVariableRef CVarChaosImmPhysCollisionDisable(TEXT("p.Chaos.ImmPhys.Collision.Enabled"), ChaosImmediate_Collision_Enabled, TEXT("Enable/Disable collisions in Immediate Physics."));
 FAutoConsoleVariableRef CVarChaosImmPhysCollisionPairIterations(TEXT("p.Chaos.ImmPhys.Collision.PairIterations"), ChaosImmediate_Collision_PairIterations, TEXT("Override collision pair iterations (if >= 0)"));
 FAutoConsoleVariableRef CVarChaosImmPhysCollisionPushOutPairIterations(TEXT("p.Chaos.ImmPhys.Collision.PushOutPairIterations"), ChaosImmediate_Collision_PushOutPairIterations, TEXT("Override collision push-out pair iterations (if >= 0)"));
 FAutoConsoleVariableRef CVarChaosImmPhysCollisionPriority(TEXT("p.Chaos.ImmPhys.Collision.Priority"), ChaosImmediate_Collision_Priority, TEXT("Set the Collision constraint sort order (Joints have priority 0)"));
-FAutoConsoleVariableRef CVarChaosImmPhysCollisionThickness(TEXT("p.Chaos.ImmPhys.Collision.Thickness"), ChaosImmediate_Collision_Thickness, TEXT("ChaosImmediateThickness"));
+FAutoConsoleVariableRef CVarChaosImmPhysCollisionCullDistance(TEXT("p.Chaos.ImmPhys.Collision.CullDistance"), ChaosImmediate_Collision_CullDistance, TEXT("CullDistance"));
+FAutoConsoleVariableRef CVarChaosImmPhysCollisionShapePadding(TEXT("p.Chaos.ImmPhys.Collision.ShapePadding"), ChaosImmediate_Collision_ShapePadding, TEXT("ShapePadding"));
 
 int32 ChaosImmediate_Joint_PairIterations = -1;
 int32 ChaosImmediate_Joint_PushOutPairIterations = -1;
@@ -162,12 +166,12 @@ namespace ImmediatePhysics_Chaos
 		, PerParticleMaterials()
 		, Particles()
 		, Joints()
-		, Collisions(Particles, CollidedParticles, ParticleMaterials, 0, 0)
-		, BroadPhase(PotentiallyCollidingPairs)
+		, Collisions(Particles, CollidedParticles, ParticleMaterials, 0, 0, ChaosImmediate_Collision_CullDistance, ChaosImmediate_Collision_ShapePadding)
+		, BroadPhase(PotentiallyCollidingPairs, ChaosImmediate_Collision_CullDistance)
 		, CollisionDetector(BroadPhase, Collisions)
 		, JointsRule(0, Joints)
 		, CollisionsRule(1, Collisions)
-		, Evolution(Particles, CollisionDetector)
+		, Evolution(Particles, CollisionDetector, ChaosImmediate_Evolution_BoundsExtension)
 		, NumActiveDynamicActorHandles(0)
 		, SimulationSpaceTransform(FTransform::Identity)
 		, RollingAverageStepTime(ChaosImmediate_Evolution_InitialStepTime)
@@ -535,10 +539,14 @@ namespace ImmediatePhysics_Chaos
 			JointsSettings.AngularDriveDamping = ChaosImmediate_Joint_AngularDriveDamping;
 			Joints.SetSettings(JointsSettings);
 
-			Collisions.SetThickness(ChaosImmediate_Collision_Thickness);
+			Collisions.SetShapePadding(ChaosImmediate_Collision_ShapePadding);
 			Collisions.SetCollisionsEnabled(ChaosImmediate_Collision_Enabled != 0);
 			CollisionsRule.SetPriority(ChaosImmediate_Collision_Priority);
 
+			Collisions.SetCullDistance(ChaosImmediate_Collision_CullDistance);
+			BroadPhase.SetCullDustance(ChaosImmediate_Collision_CullDistance);
+			Evolution.SetBoundsExtension(ChaosImmediate_Evolution_BoundsExtension);
+
 			SetSolverIterations(
 				ChaosImmediate_Evolution_Iterations,
 				ChaosImmediate_Joint_PairIterations,
diff --git a/Engine/Source/Runtime/Experimental/Chaos/Private/Chaos/CollisionResolution.cpp b/Engine/Source/Runtime/Experimental/Chaos/Private/Chaos/CollisionResolution.cpp
--- a/Engine/Source/Runtime/Experimental/Chaos/Private/Chaos/CollisionResolution.cpp
+++ b/Engine/Source/Runtime/Experimental/Chaos/Private/Chaos/CollisionResolution.cpp
@@ -22,6 +22,8 @@
 
 DECLARE_CYCLE_STAT(TEXT("Collisions::GJK"), STAT_Collisions_GJK, STATGROUP_ChaosCollision);
 
+//#pragma optimize("", off)
+
 namespace Chaos
 {
 	namespace Collisions
@@ -33,7 +35,9 @@ namespace Chaos
 			TVec3<T> Location;
 			T Phi;
 
-			TContactPoint() : Phi(TNumericLimits<T>::Max()) {}
+			TContactPoint() 
+			: Normal(1, 0, 0)
+			, Phi(TNumericLimits<T>::Max()) {}
 		};
 
 
@@ -50,29 +54,33 @@ namespace Chaos
 		}
 
 		template <typename T, int d, typename GeometryA, typename GeometryB>
-		TContactPoint<T> GJKContactPoint(const GeometryA& A, const TRigidTransform<T, d>& ATM, const GeometryB& B, const TRigidTransform<T, d>& BTM, const T Thickness)
+		TContactPoint<T> GJKContactPoint(const GeometryA& A, const TRigidTransform<T, d>& ATM, const GeometryB& B, const TRigidTransform<T, d>& BTM, const TVector<T, 3>& InitialDir)
 		{
 			SCOPE_CYCLE_COUNTER(STAT_Collisions_GJK);
 
 			TContactPoint<T> Contact;
 			const TRigidTransform<T, d> BToATM = BTM.GetRelativeTransform(ATM);
-			const TVector<T, 3> InitialDirection = TVector<T, 3>(1, 0, 0);
 
 			T Penetration;
 			TVec3<T> ClosestA, ClosestB, Normal;
+			int32 NumIterations = 0;
 
-			if (ensure(GJKPenetration(A, B, BToATM, Penetration, ClosestA, ClosestB, Normal, (T)0, InitialDirection, (T)0, true)))
+			if (ensure(GJKPenetration(A, B, BToATM, Penetration, ClosestA, ClosestB, Normal, (T)0, InitialDir, (T)0, true, &NumIterations)))
 			{
 				Contact.Location = ATM.TransformPosition(ClosestA);
 				Contact.Normal = -ATM.TransformVector(Normal);
 				Contact.Phi = -Penetration;
 			}
 
+			//static float AverageIterations = 0;
+			//AverageIterations = AverageIterations + ((float)NumIterations - AverageIterations) / 1000.0f;
+			//UE_LOG(LogChaos, Warning, TEXT("GJK Its: %f"), AverageIterations);
+
 			return Contact;
 		}
 
 		template <typename GeometryA, typename GeometryB, typename T, int d>
-		TContactPoint<T> GJKImplicitContactPoint(const FImplicitObject& A, const TRigidTransform<T, d>& ATransform, const GeometryB& B, const TRigidTransform<T, d>& BTransform, const T Thickness)
+		TContactPoint<T> GJKImplicitContactPoint(const FImplicitObject& A, const TRigidTransform<T, d>& ATransform, const GeometryB& B, const TRigidTransform<T, d>& BTransform, const T CullDistance)
 		{
 			TContactPoint<T> Contact;
 			const TRigidTransform<T, d> AToBTM = ATransform.GetRelativeTransform(BTransform);
@@ -81,7 +89,7 @@ namespace Chaos
 			TVec3<T> Location, Normal;
 			if (const TImplicitObjectScaled<GeometryA>* ScaledConvexImplicit = A.template GetObject<const TImplicitObjectScaled<GeometryA> >())
 			{
-				if (B.GJKContactPoint(*ScaledConvexImplicit, AToBTM, Thickness, Location, Normal, Penetration))
+				if (B.GJKContactPoint(*ScaledConvexImplicit, AToBTM, CullDistance, Location, Normal, Penetration))
 				{
 					Contact.Phi = Penetration;
 					Contact.Location = BTransform.TransformPosition(Location);
@@ -92,7 +100,7 @@ namespace Chaos
 			{
 				if (const GeometryA * InstancedInnerObject = static_cast<const GeometryA*>(InstancedConvexImplicit->GetInstancedObject()))
 				{
-					if (B.GJKContactPoint(*InstancedInnerObject, AToBTM, Thickness, Location, Normal, Penetration))
+					if (B.GJKContactPoint(*InstancedInnerObject, AToBTM, CullDistance, Location, Normal, Penetration))
 					{
 						Contact.Phi = Penetration;
 						Contact.Location = BTransform.TransformPosition(Location);
@@ -102,7 +110,7 @@ namespace Chaos
 			}
 			else if (const GeometryA* ConvexImplicit = A.template GetObject<const GeometryA>())
 			{
-				if (B.GJKContactPoint(*ConvexImplicit, AToBTM, Thickness, Location, Normal, Penetration))
+				if (B.GJKContactPoint(*ConvexImplicit, AToBTM, CullDistance, Location, Normal, Penetration))
 				{
 					Contact.Phi = Penetration;
 					Contact.Location = BTransform.TransformPosition(Location);
@@ -115,22 +123,22 @@ namespace Chaos
 
 
 		template<class T, int d>
-		TContactPoint<T> ConvexConvexContactPoint(const FImplicitObject& A, const TRigidTransform<T, d>& ATM, const FImplicitObject& B, const TRigidTransform<T, d>& BTM, const T Thickness)
+		TContactPoint<T> ConvexConvexContactPoint(const FImplicitObject& A, const TRigidTransform<T, d>& ATM, const FImplicitObject& B, const TRigidTransform<T, d>& BTM, const T CullDistance)
 		{
 			return CastHelper(A, ATM, [&](const auto& ADowncast, const TRigidTransform<T,d>& AFullTM)
 			{
 				return CastHelper(B, BTM, [&](const auto& BDowncast, const TRigidTransform<T,d>& BFullTM)
 				{
-					return GJKContactPoint(ADowncast, AFullTM, BDowncast, BFullTM, Thickness);
+					return GJKContactPoint(ADowncast, AFullTM, BDowncast, BFullTM, TVector<T, d>(1, 0, 0));
 				});
 			});
 		}
 
 		template <typename T, int d>
-		void UpdateSingleShotManifold(TRigidBodyMultiPointContactConstraint<T, d>&  Constraint, const TRigidTransform<T, d>& Transform0, const TRigidTransform<T, d>& Transform1, const T Thickness)
+		void UpdateSingleShotManifold(TRigidBodyMultiPointContactConstraint<T, d>&  Constraint, const TRigidTransform<T, d>& Transform0, const TRigidTransform<T, d>& Transform1, const T CullDistance)
 		{
 			// single shot manifolds for TConvex implicit object in the constraints implicit[0] position. 
-			TContactPoint<T> ContactPoint = ConvexConvexContactPoint(*Constraint.Manifold.Implicit[0], Transform0, *Constraint.Manifold.Implicit[1], Transform1, Thickness);
+			TContactPoint<T> ContactPoint = ConvexConvexContactPoint(*Constraint.Manifold.Implicit[0], Transform0, *Constraint.Manifold.Implicit[1], Transform1, CullDistance);
 
 			TArray<FVec3> CollisionSamples;
 			//
@@ -167,7 +175,7 @@ namespace Chaos
 		}
 
 		template <typename T, int d>
-		void UpdateIterativeManifold(TRigidBodyMultiPointContactConstraint<T, d>&  Constraint, const TRigidTransform<T, d>& Transform0, const TRigidTransform<T, d>& Transform1, const T Thickness)
+		void UpdateIterativeManifold(TRigidBodyMultiPointContactConstraint<T, d>&  Constraint, const TRigidTransform<T, d>& Transform0, const TRigidTransform<T, d>& Transform1, const T CullDistance)
 		{
 			auto SumSampleData = [&](TRigidBodyMultiPointContactConstraint<T, d>& LambdaConstraint) -> TVector<float, 3>
 			{
@@ -180,7 +188,7 @@ namespace Chaos
 			};
 
 			// iterative manifolds for non TConvex implicit objects that require sampling 
-			TContactPoint<T> ContactPoint = ConvexConvexContactPoint(*Constraint.Manifold.Implicit[0], Transform0, *Constraint.Manifold.Implicit[1], Transform1, Thickness);
+			TContactPoint<T> ContactPoint = ConvexConvexContactPoint(*Constraint.Manifold.Implicit[0], Transform0, *Constraint.Manifold.Implicit[1], Transform1, CullDistance);
 
 			if (!ContactPoint.Normal.Equals(Constraint.PlaneNormal) || !Constraint.NumSamples())
 			{
@@ -232,26 +240,26 @@ namespace Chaos
 		// Box - Box
 		//
 		template <typename T, int d>
-		TContactPoint<T> BoxBoxContactPoint(const TAABB<T, d>& Box1, const TRigidTransform<T, d>& ATM, const TAABB<T, d>& Box2, const TRigidTransform<T, d>& BTM, const T Thickness)
+		TContactPoint<T> BoxBoxContactPoint(const TAABB<T, d>& Box1, const TRigidTransform<T, d>& ATM, const TAABB<T, d>& Box2, const TRigidTransform<T, d>& BTM, const T CullDistance)
 		{
-			return GJKContactPoint(Box1, ATM, Box2, BTM, Thickness);
+			return GJKContactPoint(Box1, ATM, Box2, BTM, TVector<T, d>(1, 0, 0));
 		}
 
 		template <typename T, int d>
-		void UpdateBoxBoxConstraint(const TAABB<T, d>& Box1, const TRigidTransform<T, d>& Box1Transform, const TAABB<T, d>& Box2, const TRigidTransform<T, d>& Box2Transform, const T Thickness, TRigidBodyPointContactConstraint<T, d>& Constraint)
+		void UpdateBoxBoxConstraint(const TAABB<T, d>& Box1, const TRigidTransform<T, d>& Box1Transform, const TAABB<T, d>& Box2, const TRigidTransform<T, d>& Box2Transform, const T CullDistance, TRigidBodyPointContactConstraint<T, d>& Constraint)
 		{
-			UpdateContactPoint(Constraint.Manifold, BoxBoxContactPoint(Box1, Box1Transform, Box2, Box2Transform, Thickness));
+			UpdateContactPoint(Constraint.Manifold, BoxBoxContactPoint(Box1, Box1Transform, Box2, Box2Transform, CullDistance));
 		}
 
 		template <typename T, int d>
-		void UpdateBoxBoxManifold(TCollisionConstraintBase<T, d>&  Constraint, const TRigidTransform<T, d>& ATM, const TRigidTransform<T, d>& BTM, const T Thickness)
+		void UpdateBoxBoxManifold(TCollisionConstraintBase<T, d>&  Constraint, const TRigidTransform<T, d>& ATM, const TRigidTransform<T, d>& BTM, const T CullDistance)
 		{
 			// @todo(chaos) : Stub Update Manifold
 			//   Stub function for updating the manifold prior to the Apply and ApplyPushOut
 		}
 
 		template<typename T, int d>
-		void ConstructBoxBoxConstraints(TGeometryParticleHandle<T, d>* Particle0, TGeometryParticleHandle<T, d>* Particle1, const FImplicitObject* Implicit0, const FImplicitObject* Implicit1, const TRigidTransform<T, d>& Transform0, const TRigidTransform<T, d>& Transform1, const T Thickness, FCollisionConstraintsArray& NewConstraints)
+		void ConstructBoxBoxConstraints(TGeometryParticleHandle<T, d>* Particle0, TGeometryParticleHandle<T, d>* Particle1, const FImplicitObject* Implicit0, const FImplicitObject* Implicit1, const TRigidTransform<T, d>& Transform0, const TRigidTransform<T, d>& Transform1, const T CullDistance, FCollisionConstraintsArray& NewConstraints)
 		{
 			const TBox<T, d> * Object0 = Implicit0->template GetObject<const TBox<T, d> >();
 			const TBox<T, d> * Object1 = Implicit1->template GetObject<const TBox<T, d> >();
@@ -260,8 +268,8 @@ namespace Chaos
 				TRigidTransform<T, d> ParticleImplicit0TM = Transform0.GetRelativeTransform(Collisions::GetTransform(Particle0));
 				TRigidTransform<T, d> ParticleImplicit1TM = Transform1.GetRelativeTransform(Collisions::GetTransform(Particle1));
 				FRigidBodyPointContactConstraint Constraint = FRigidBodyPointContactConstraint(Particle0, Implicit0, ParticleImplicit0TM, Particle1, Implicit1, ParticleImplicit1TM);
-				UpdateBoxBoxConstraint(Object0->BoundingBox(), Transform0, Object1->BoundingBox(), Transform1, Thickness, Constraint);
-				NewConstraints.TryAdd(Thickness, Constraint);
+				UpdateBoxBoxConstraint(Object0->BoundingBox(), Transform0, Object1->BoundingBox(), Transform1, CullDistance, Constraint);
+				NewConstraints.TryAdd(CullDistance, Constraint);
 			}
 		}
 
@@ -270,26 +278,26 @@ namespace Chaos
 		//
 
 		template <typename T, int d>
-		TContactPoint<T> BoxHeightFieldContactPoint(const TAABB<T, d>& A, const TRigidTransform<T, d>& ATransform, const THeightField<T>& B, const TRigidTransform<T, d>& BTransform, const T Thickness)
+		TContactPoint<T> BoxHeightFieldContactPoint(const TAABB<T, d>& A, const TRigidTransform<T, d>& ATransform, const THeightField<T>& B, const TRigidTransform<T, d>& BTransform, const T CullDistance)
 		{
-			return GJKImplicitContactPoint< TBox<float, 3> >(TBox<float, 3>(A), ATransform, B, BTransform, Thickness);
+			return GJKImplicitContactPoint< TBox<float, 3> >(TBox<float, 3>(A), ATransform, B, BTransform, CullDistance);
 		}
 
 		template <typename T, int d>
-		void UpdateBoxHeightFieldConstraint(const TAABB<T, d>& A, const TRigidTransform<T, d>& ATransform, const THeightField<T>& B, const TRigidTransform<T, d>& BTransform, const T Thickness, TRigidBodyPointContactConstraint<T, d>& Constraint)
+		void UpdateBoxHeightFieldConstraint(const TAABB<T, d>& A, const TRigidTransform<T, d>& ATransform, const THeightField<T>& B, const TRigidTransform<T, d>& BTransform, const T CullDistance, TRigidBodyPointContactConstraint<T, d>& Constraint)
 		{
-			UpdateContactPoint(Constraint.Manifold, BoxHeightFieldContactPoint(A, ATransform, B, BTransform, Thickness));
+			UpdateContactPoint(Constraint.Manifold, BoxHeightFieldContactPoint(A, ATransform, B, BTransform, CullDistance));
 		}
 
 		template<class T, int d>
-		void UpdateBoxHeightFieldManifold(TCollisionConstraintBase<T, d>&  Constraint, const TRigidTransform<T, d>& ATM, const TRigidTransform<T, d>& BTM, const T Thickness)
+		void UpdateBoxHeightFieldManifold(TCollisionConstraintBase<T, d>&  Constraint, const TRigidTransform<T, d>& ATM, const TRigidTransform<T, d>& BTM, const T CullDistance)
 		{
 			// @todo(chaos) : Stub Update Manifold
 			//   Stub function for updating the manifold prior to the Apply and ApplyPushOut
 		}
 
 		template<typename T, int d>
-		void ConstructBoxHeightFieldConstraints(TGeometryParticleHandle<T, d>* Particle0, TGeometryParticleHandle<T, d>* Particle1, const FImplicitObject* Implicit0, const FImplicitObject* Implicit1, const TRigidTransform<T, d>& Transform0, const TRigidTransform<T, d>& Transform1, const T Thickness, FCollisionConstraintsArray& NewConstraints)
+		void ConstructBoxHeightFieldConstraints(TGeometryParticleHandle<T, d>* Particle0, TGeometryParticleHandle<T, d>* Particle1, const FImplicitObject* Implicit0, const FImplicitObject* Implicit1, const TRigidTransform<T, d>& Transform0, const TRigidTransform<T, d>& Transform1, const T CullDistance, FCollisionConstraintsArray& NewConstraints)
 		{
 
 			const TBox<T, d> * Object0 = Implicit0->template GetObject<const TBox<T, d> >();
@@ -299,8 +307,8 @@ namespace Chaos
 				TRigidTransform<T, d> ParticleImplicit0TM = Transform0.GetRelativeTransform(Collisions::GetTransform(Particle0));
 				TRigidTransform<T, d> ParticleImplicit1TM = Transform1.GetRelativeTransform(Collisions::GetTransform(Particle1));
 				FRigidBodyPointContactConstraint Constraint = FRigidBodyPointContactConstraint(Particle0, Implicit0, ParticleImplicit0TM, Particle1, Implicit1, ParticleImplicit1TM);
-				UpdateBoxHeightFieldConstraint(Object0->BoundingBox(), Transform0, *Object1, Transform1, Thickness, Constraint);
-				NewConstraints.TryAdd(Thickness, Constraint);
+				UpdateBoxHeightFieldConstraint(Object0->BoundingBox(), Transform0, *Object1, Transform1, CullDistance, Constraint);
+				NewConstraints.TryAdd(CullDistance, Constraint);
 			}
 		}
 
@@ -311,7 +319,7 @@ namespace Chaos
 		//
 
 		template <typename T, int d>
-		void UpdateBoxPlaneConstraint(const TAABB<T, d>& Box, const TRigidTransform<T, d>& BoxTransform, const TPlane<T, d>& Plane, const TRigidTransform<T, d>& PlaneTransform, const T Thickness, TRigidBodyPointContactConstraint<T, d>& Constraint)
+		void UpdateBoxPlaneConstraint(const TAABB<T, d>& Box, const TRigidTransform<T, d>& BoxTransform, const TPlane<T, d>& Plane, const TRigidTransform<T, d>& PlaneTransform, const T CullDistance, TRigidBodyPointContactConstraint<T, d>& Constraint)
 		{
 			TCollisionContact<T, d> & Contact = Constraint.Manifold;
 
@@ -369,7 +377,7 @@ namespace Chaos
 		}
 
 		template<class T, int d>
-		void UpdateBoxPlaneManifold(TCollisionConstraintBase<T, d>&  Constraint, const TRigidTransform<T, d>& ATM, const TRigidTransform<T, d>& BTM, const T Thickness)
+		void UpdateBoxPlaneManifold(TCollisionConstraintBase<T, d>&  Constraint, const TRigidTransform<T, d>& ATM, const TRigidTransform<T, d>& BTM, const T CullDistance)
 		{
 			// @todo(chaos) : Stub Update Manifold
 			//   Stub function for updating the manifold prior to the Apply and ApplyPushOut
@@ -377,7 +385,7 @@ namespace Chaos
 
 
 		template<typename T, int d>
-		void ConstructBoxPlaneConstraints(TGeometryParticleHandle<T, d>* Particle0, TGeometryParticleHandle<T, d>* Particle1, const FImplicitObject* Implicit0, const FImplicitObject* Implicit1, const TRigidTransform<T, d>& Transform0, const TRigidTransform<T, d>& Transform1, const T Thickness, FCollisionConstraintsArray& NewConstraints)
+		void ConstructBoxPlaneConstraints(TGeometryParticleHandle<T, d>* Particle0, TGeometryParticleHandle<T, d>* Particle1, const FImplicitObject* Implicit0, const FImplicitObject* Implicit1, const TRigidTransform<T, d>& Transform0, const TRigidTransform<T, d>& Transform1, const T CullDistance, FCollisionConstraintsArray& NewConstraints)
 		{
 
 			const TBox<T, d> * Object0 = Implicit0->template GetObject<const TBox<T, d> >();
@@ -387,8 +395,8 @@ namespace Chaos
 				TRigidTransform<T, d> ParticleImplicit0TM = Transform0.GetRelativeTransform(Collisions::GetTransform(Particle0));
 				TRigidTransform<T, d> ParticleImplicit1TM = Transform1.GetRelativeTransform(Collisions::GetTransform(Particle1));
 				FRigidBodyPointContactConstraint Constraint = FRigidBodyPointContactConstraint(Particle0, Implicit0, ParticleImplicit0TM, Particle1, Implicit1, ParticleImplicit1TM);
-				UpdateBoxPlaneConstraint(Object0->BoundingBox(), Transform0, *Object1, Transform1, Thickness, Constraint);
-				NewConstraints.TryAdd(Thickness, Constraint);
+				UpdateBoxPlaneConstraint(Object0->BoundingBox(), Transform0, *Object1, Transform1, CullDistance, Constraint);
+				NewConstraints.TryAdd(CullDistance, Constraint);
 			}
 		}
 
@@ -397,20 +405,20 @@ namespace Chaos
 		//
 
 		template <typename T, int d>
-		TContactPoint<T> BoxTriangleMeshContactPoint(const TAABB<T, d>& A, const TRigidTransform<T, d>& ATransform, const FTriangleMeshImplicitObject& B, const TRigidTransform<T, d>& BTransform, const T Thickness)
+		TContactPoint<T> BoxTriangleMeshContactPoint(const TAABB<T, d>& A, const TRigidTransform<T, d>& ATransform, const FTriangleMeshImplicitObject& B, const TRigidTransform<T, d>& BTransform, const T CullDistance)
 		{
-			return GJKImplicitContactPoint< TBox<float, 3> >(TBox<float, 3>(A), ATransform, B, BTransform, Thickness);
+			return GJKImplicitContactPoint< TBox<float, 3> >(TBox<float, 3>(A), ATransform, B, BTransform, CullDistance);
 		}
 
 		template <typename T, int d>
-		void UpdateBoxTriangleMeshConstraint(const TAABB<T, d>& Box0, const TRigidTransform<T, d>& Transform0, const FTriangleMeshImplicitObject& TriangleMesh1, const TRigidTransform<T, d>& Transform1, const T Thickness, TRigidBodyPointContactConstraint<T, d>& Constraint)
+		void UpdateBoxTriangleMeshConstraint(const TAABB<T, d>& Box0, const TRigidTransform<T, d>& Transform0, const FTriangleMeshImplicitObject& TriangleMesh1, const TRigidTransform<T, d>& Transform1, const T CullDistance, TRigidBodyPointContactConstraint<T, d>& Constraint)
 		{
-			UpdateContactPoint(Constraint.Manifold, BoxTriangleMeshContactPoint(Box0, Transform0, TriangleMesh1, Transform1, Thickness));
+			UpdateContactPoint(Constraint.Manifold, BoxTriangleMeshContactPoint(Box0, Transform0, TriangleMesh1, Transform1, CullDistance));
 		}
 
 
 		template <typename T, int d>
-		void UpdateBoxTriangleMeshManifold(TCollisionConstraintBase<T, d>&  Constraint, const TRigidTransform<T, d>& ATM, const TRigidTransform<T, d>& BTM, const T Thickness)
+		void UpdateBoxTriangleMeshManifold(TCollisionConstraintBase<T, d>&  Constraint, const TRigidTransform<T, d>& ATM, const TRigidTransform<T, d>& BTM, const T CullDistance)
 		{
 
 		}
@@ -428,7 +436,7 @@ namespace Chaos
 		}
 
 		template<typename T, int d>
-		void ConstructBoxTriangleMeshConstraints(TGeometryParticleHandle<T, d>* Particle0, TGeometryParticleHandle<T, d>* Particle1, const FImplicitObject* Implicit0, const FImplicitObject* Implicit1, const TRigidTransform<T, d>& Transform0, const TRigidTransform<T, d>& Transform1, const T Thickness, FCollisionConstraintsArray& NewConstraints)
+		void ConstructBoxTriangleMeshConstraints(TGeometryParticleHandle<T, d>* Particle0, TGeometryParticleHandle<T, d>* Particle1, const FImplicitObject* Implicit0, const FImplicitObject* Implicit1, const TRigidTransform<T, d>& Transform0, const TRigidTransform<T, d>& Transform1, const T CullDistance, FCollisionConstraintsArray& NewConstraints)
 		{
 			const TBox<T, d> * Object0 = Implicit0->template GetObject<const TBox<T, d> >();
 			const FTriangleMeshImplicitObject * Object1 = GetInnerObject<FTriangleMeshImplicitObject>(*Implicit1);
@@ -437,8 +445,8 @@ namespace Chaos
 				TRigidTransform<T, d> ParticleImplicit0TM = Transform0.GetRelativeTransform(Collisions::GetTransform(Particle0));
 				TRigidTransform<T, d> ParticleImplicit1TM = Transform1.GetRelativeTransform(Collisions::GetTransform(Particle1));
 				FRigidBodyPointContactConstraint Constraint = FRigidBodyPointContactConstraint(Particle0, Implicit0, ParticleImplicit0TM, Particle1, Implicit1, ParticleImplicit1TM);
-				UpdateBoxTriangleMeshConstraint(Object0->GetAABB(), Transform0, *Object1, Transform1, Thickness, Constraint);
-				NewConstraints.TryAdd(Thickness, Constraint);
+				UpdateBoxTriangleMeshConstraint(Object0->GetAABB(), Transform0, *Object1, Transform1, CullDistance, Constraint);
+				NewConstraints.TryAdd(CullDistance, Constraint);
 			}
 		}
 
@@ -447,7 +455,7 @@ namespace Chaos
 		//
 
 		template <typename T, int d>
-		TContactPoint<T> SphereSphereContactPoint(const TSphere<T, d>& Sphere1, const TRigidTransform<T, d>& Sphere1Transform, const TSphere<T, d>& Sphere2, const TRigidTransform<T, d>& Sphere2Transform, const T Thickness)
+		TContactPoint<T> SphereSphereContactPoint(const TSphere<T, d>& Sphere1, const TRigidTransform<T, d>& Sphere1Transform, const TSphere<T, d>& Sphere2, const TRigidTransform<T, d>& Sphere2Transform, const T CullDistance)
 		{
 			TContactPoint<T> Result;
 
@@ -464,20 +472,20 @@ namespace Chaos
 		}
 
 		template <typename T, int d>
-		void UpdateSphereSphereConstraint(const TSphere<T, d>& Sphere1, const TRigidTransform<T, d>& Sphere1Transform, const TSphere<T, d>& Sphere2, const TRigidTransform<T, d>& Sphere2Transform, const T Thickness, TRigidBodyPointContactConstraint<T, d>& Constraint)
+		void UpdateSphereSphereConstraint(const TSphere<T, d>& Sphere1, const TRigidTransform<T, d>& Sphere1Transform, const TSphere<T, d>& Sphere2, const TRigidTransform<T, d>& Sphere2Transform, const T CullDistance, TRigidBodyPointContactConstraint<T, d>& Constraint)
 		{
-			UpdateContactPoint(Constraint.Manifold, SphereSphereContactPoint(Sphere1, Sphere1Transform, Sphere2, Sphere2Transform, Thickness));
+			UpdateContactPoint(Constraint.Manifold, SphereSphereContactPoint(Sphere1, Sphere1Transform, Sphere2, Sphere2Transform, CullDistance));
 		}
 
 		template<class T, int d>
-		void UpdateSphereSphereManifold(TCollisionConstraintBase<T, d>&  Constraint, const TRigidTransform<T, d>& ATM, const TRigidTransform<T, d>& BTM, const T Thickness)
+		void UpdateSphereSphereManifold(TCollisionConstraintBase<T, d>&  Constraint, const TRigidTransform<T, d>& ATM, const TRigidTransform<T, d>& BTM, const T CullDistance)
 		{
 			// @todo(chaos) : Stub Update Manifold
 			//   Stub function for updating the manifold prior to the Apply and ApplyPushOut
 		}
 
 		template<typename T, int d>
-		void ConstructSphereSphereConstraints(TGeometryParticleHandle<T, d>* Particle0, TGeometryParticleHandle<T, d>* Particle1, const FImplicitObject* Implicit0, const FImplicitObject* Implicit1, const TRigidTransform<T, d>& Transform0, const TRigidTransform<T, d>& Transform1, const T Thickness, FCollisionConstraintsArray& NewConstraints)
+		void ConstructSphereSphereConstraints(TGeometryParticleHandle<T, d>* Particle0, TGeometryParticleHandle<T, d>* Particle1, const FImplicitObject* Implicit0, const FImplicitObject* Implicit1, const TRigidTransform<T, d>& Transform0, const TRigidTransform<T, d>& Transform1, const T CullDistance, FCollisionConstraintsArray& NewConstraints)
 		{
 
 			const TSphere<T, d> * Object0 = Implicit0->template GetObject<const TSphere<T, d> >();
@@ -487,8 +495,8 @@ namespace Chaos
 				TRigidTransform<T, d> ParticleImplicit0TM = Transform0.GetRelativeTransform(Collisions::GetTransform(Particle0));
 				TRigidTransform<T, d> ParticleImplicit1TM = Transform1.GetRelativeTransform(Collisions::GetTransform(Particle1));
 				FRigidBodyPointContactConstraint Constraint = FRigidBodyPointContactConstraint(Particle0, Implicit0, ParticleImplicit0TM, Particle1, Implicit1, ParticleImplicit1TM);
-				UpdateSphereSphereConstraint(*Object0, Transform0, *Object1, Transform1, Thickness, Constraint);
-				NewConstraints.TryAdd(Thickness, Constraint);
+				UpdateSphereSphereConstraint(*Object0, Transform0, *Object1, Transform1, CullDistance, Constraint);
+				NewConstraints.TryAdd(CullDistance, Constraint);
 			}
 		}
 
@@ -497,26 +505,26 @@ namespace Chaos
 		//
 
 		template <typename T, int d>
-		TContactPoint<T> SphereHeightFieldContactPoint(const TSphere<T, d>& A, const TRigidTransform<T, d>& ATransform, const THeightField<T>& B, const TRigidTransform<T, d>& BTransform, const T Thickness)
+		TContactPoint<T> SphereHeightFieldContactPoint(const TSphere<T, d>& A, const TRigidTransform<T, d>& ATransform, const THeightField<T>& B, const TRigidTransform<T, d>& BTransform, const T CullDistance)
 		{
-			return GJKImplicitContactPoint< TSphere<float, 3> >(TSphere<float, 3>(A), ATransform, B, BTransform, Thickness);
+			return GJKImplicitContactPoint< TSphere<float, 3> >(TSphere<float, 3>(A), ATransform, B, BTransform, CullDistance);
 		}
 
 		template <typename T, int d>
-		void UpdateSphereHeightFieldConstraint(const TSphere<T, d>& A, const TRigidTransform<T, d>& ATransform, const THeightField<T>& B, const TRigidTransform<T, d>& BTransform, const T Thickness, TRigidBodyPointContactConstraint<T, d>& Constraint)
+		void UpdateSphereHeightFieldConstraint(const TSphere<T, d>& A, const TRigidTransform<T, d>& ATransform, const THeightField<T>& B, const TRigidTransform<T, d>& BTransform, const T CullDistance, TRigidBodyPointContactConstraint<T, d>& Constraint)
 		{
-			UpdateContactPoint(Constraint.Manifold, SphereHeightFieldContactPoint(A, ATransform, B, BTransform, Thickness));
+			UpdateContactPoint(Constraint.Manifold, SphereHeightFieldContactPoint(A, ATransform, B, BTransform, CullDistance));
 		}
 
 		template<class T, int d>
-		void UpdateSphereHeightFieldManifold(TCollisionConstraintBase<T, d>&  Constraint, const TRigidTransform<T, d>& ATM, const TRigidTransform<T, d>& BTM, const T Thickness)
+		void UpdateSphereHeightFieldManifold(TCollisionConstraintBase<T, d>&  Constraint, const TRigidTransform<T, d>& ATM, const TRigidTransform<T, d>& BTM, const T CullDistance)
 		{
 			// @todo(chaos) : Stub Update Manifold
 			//   Stub function for updating the manifold prior to the Apply and ApplyPushOut
 		}
 
 		template<typename T, int d>
-		void ConstructSphereHeightFieldConstraints(TGeometryParticleHandle<T, d>* Particle0, TGeometryParticleHandle<T, d>* Particle1, const FImplicitObject* Implicit0, const FImplicitObject* Implicit1, const TRigidTransform<T, d>& Transform0, const TRigidTransform<T, d>& Transform1, const T Thickness, FCollisionConstraintsArray& NewConstraints)
+		void ConstructSphereHeightFieldConstraints(TGeometryParticleHandle<T, d>* Particle0, TGeometryParticleHandle<T, d>* Particle1, const FImplicitObject* Implicit0, const FImplicitObject* Implicit1, const TRigidTransform<T, d>& Transform0, const TRigidTransform<T, d>& Transform1, const T CullDistance, FCollisionConstraintsArray& NewConstraints)
 		{
 
 			const TSphere<T, d> * Object0 = Implicit0->template GetObject<const TSphere<T, d> >();
@@ -526,8 +534,8 @@ namespace Chaos
 				TRigidTransform<T, d> ParticleImplicit0TM = Transform0.GetRelativeTransform(Collisions::GetTransform(Particle0));
 				TRigidTransform<T, d> ParticleImplicit1TM = Transform1.GetRelativeTransform(Collisions::GetTransform(Particle1));
 				FRigidBodyPointContactConstraint Constraint = FRigidBodyPointContactConstraint(Particle0, Implicit0, ParticleImplicit0TM, Particle1, Implicit1, ParticleImplicit1TM);
-				UpdateSphereHeightFieldConstraint(*Object0, Transform0, *Object1, Transform1, Thickness, Constraint);
-				NewConstraints.TryAdd(Thickness, Constraint);
+				UpdateSphereHeightFieldConstraint(*Object0, Transform0, *Object1, Transform1, CullDistance, Constraint);
+				NewConstraints.TryAdd(CullDistance, Constraint);
 			}
 		}
 
@@ -536,7 +544,7 @@ namespace Chaos
 		//
 
 		template <typename T, int d>
-		void UpdateSpherePlaneConstraint(const TSphere<T, d>& Sphere, const TRigidTransform<T, d>& SphereTransform, const TPlane<T, d>& Plane, const TRigidTransform<T, d>& PlaneTransform, const T Thickness, TRigidBodyPointContactConstraint<T, d>& Constraint)
+		void UpdateSpherePlaneConstraint(const TSphere<T, d>& Sphere, const TRigidTransform<T, d>& SphereTransform, const TPlane<T, d>& Plane, const TRigidTransform<T, d>& PlaneTransform, const T CullDistance, TRigidBodyPointContactConstraint<T, d>& Constraint)
 		{
 			TCollisionContact<T, d> & Contact = Constraint.Manifold;
 
@@ -556,14 +564,14 @@ namespace Chaos
 		}
 
 		template<class T, int d>
-		void UpdateSpherePlaneManifold(TCollisionConstraintBase<T, d>&  Constraint, const TRigidTransform<T, d>& ATM, const TRigidTransform<T, d>& BTM, const T Thickness)
+		void UpdateSpherePlaneManifold(TCollisionConstraintBase<T, d>&  Constraint, const TRigidTransform<T, d>& ATM, const TRigidTransform<T, d>& BTM, const T CullDistance)
 		{
 			// @todo(chaos) : Stub Update Manifold
 			//   Stub function for updating the manifold prior to the Apply and ApplyPushOut
 		}
 
 		template<typename T, int d>
-		void ConstructSpherePlaneConstraints(TGeometryParticleHandle<T, d>* Particle0, TGeometryParticleHandle<T, d>* Particle1, const FImplicitObject* Implicit0, const FImplicitObject* Implicit1, const TRigidTransform<T, d>& Transform0, const TRigidTransform<T, d>& Transform1, const T Thickness, FCollisionConstraintsArray& NewConstraints)
+		void ConstructSpherePlaneConstraints(TGeometryParticleHandle<T, d>* Particle0, TGeometryParticleHandle<T, d>* Particle1, const FImplicitObject* Implicit0, const FImplicitObject* Implicit1, const TRigidTransform<T, d>& Transform0, const TRigidTransform<T, d>& Transform1, const T CullDistance, FCollisionConstraintsArray& NewConstraints)
 		{
 
 			const TSphere<T, d> * Object0 = Implicit0->template GetObject<const TSphere<T, d> >();
@@ -573,8 +581,8 @@ namespace Chaos
 				TRigidTransform<T, d> ParticleImplicit0TM = Transform0.GetRelativeTransform(Collisions::GetTransform(Particle0));
 				TRigidTransform<T, d> ParticleImplicit1TM = Transform1.GetRelativeTransform(Collisions::GetTransform(Particle1));
 				FRigidBodyPointContactConstraint Constraint = FRigidBodyPointContactConstraint(Particle0, Implicit0, ParticleImplicit0TM, Particle1, Implicit1, ParticleImplicit1TM);
-				UpdateSpherePlaneConstraint(*Object0, Transform0, *Object1, Transform1, Thickness, Constraint);
-				NewConstraints.TryAdd(Thickness, Constraint);
+				UpdateSpherePlaneConstraint(*Object0, Transform0, *Object1, Transform1, CullDistance, Constraint);
+				NewConstraints.TryAdd(CullDistance, Constraint);
 			}
 		}
 
@@ -583,7 +591,7 @@ namespace Chaos
 		//
 
 		template <typename T, int d>
-		TContactPoint<T> SphereBoxContactPoint(const TSphere<T, d>& Sphere, const TRigidTransform<T, d>& SphereTransform, const TAABB<T, d>& Box, const TRigidTransform<T, d>& BoxTransform, const T Thickness)
+		TContactPoint<T> SphereBoxContactPoint(const TSphere<T, d>& Sphere, const TRigidTransform<T, d>& SphereTransform, const TAABB<T, d>& Box, const TRigidTransform<T, d>& BoxTransform, const T CullDistance)
 		{
 			TContactPoint<T> Result;
 
@@ -601,20 +609,20 @@ namespace Chaos
 		}
 
 		template <typename T, int d>
-		void UpdateSphereBoxConstraint(const TSphere<T, d>& Sphere, const TRigidTransform<T, d>& SphereTransform, const TAABB<T, d>& Box, const TRigidTransform<T, d>& BoxTransform, const T Thickness, TRigidBodyPointContactConstraint<T, d>& Constraint)
+		void UpdateSphereBoxConstraint(const TSphere<T, d>& Sphere, const TRigidTransform<T, d>& SphereTransform, const TAABB<T, d>& Box, const TRigidTransform<T, d>& BoxTransform, const T CullDistance, TRigidBodyPointContactConstraint<T, d>& Constraint)
 		{
-			UpdateContactPoint(Constraint.Manifold, SphereBoxContactPoint(Sphere, SphereTransform, Box, BoxTransform, Thickness));
+			UpdateContactPoint(Constraint.Manifold, SphereBoxContactPoint(Sphere, SphereTransform, Box, BoxTransform, CullDistance));
 		}
 
 		template<class T, int d>
-		void UpdateSphereBoxManifold(TCollisionConstraintBase<T, d>&  Constraint, const TRigidTransform<T, d>& ATM, const TRigidTransform<T, d>& BTM, const T Thickness)
+		void UpdateSphereBoxManifold(TCollisionConstraintBase<T, d>&  Constraint, const TRigidTransform<T, d>& ATM, const TRigidTransform<T, d>& BTM, const T CullDistance)
 		{
 			// @todo(chaos) : Stub Update Manifold
 			//   Stub function for updating the manifold prior to the Apply and ApplyPushOut
 		}
 
 		template<typename T, int d>
-		void ConstructSphereBoxConstraints(TGeometryParticleHandle<T, d>* Particle0, TGeometryParticleHandle<T, d>* Particle1, const FImplicitObject* Implicit0, const FImplicitObject* Implicit1, const TRigidTransform<T, d>& Transform0, const TRigidTransform<T, d>& Transform1, const T Thickness, FCollisionConstraintsArray& NewConstraints)
+		void ConstructSphereBoxConstraints(TGeometryParticleHandle<T, d>* Particle0, TGeometryParticleHandle<T, d>* Particle1, const FImplicitObject* Implicit0, const FImplicitObject* Implicit1, const TRigidTransform<T, d>& Transform0, const TRigidTransform<T, d>& Transform1, const T CullDistance, FCollisionConstraintsArray& NewConstraints)
 		{
 
 			const TSphere<T, d> * Object0 = Implicit0->template GetObject<const TSphere<T, d> >();
@@ -624,8 +632,8 @@ namespace Chaos
 				TRigidTransform<T, d> ParticleImplicit0TM = Transform0.GetRelativeTransform(Collisions::GetTransform(Particle0));
 				TRigidTransform<T, d> ParticleImplicit1TM = Transform1.GetRelativeTransform(Collisions::GetTransform(Particle1));
 				FRigidBodyPointContactConstraint Constraint = FRigidBodyPointContactConstraint(Particle0, Implicit0, ParticleImplicit0TM, Particle1, Implicit1, ParticleImplicit1TM);
-				UpdateSphereBoxConstraint(*Object0, Transform0, Object1->BoundingBox(), Transform1, Thickness, Constraint);
-				NewConstraints.TryAdd(Thickness, Constraint);
+				UpdateSphereBoxConstraint(*Object0, Transform0, Object1->BoundingBox(), Transform1, CullDistance, Constraint);
+				NewConstraints.TryAdd(CullDistance, Constraint);
 			}
 		}
 
@@ -635,7 +643,7 @@ namespace Chaos
 		//
 
 		template <typename T, int d>
-		TContactPoint<T> SphereCapsuleContactPoint(const TSphere<T, d>& A, const TRigidTransform<T, d>& ATransform, const TCapsule<T>& B, const TRigidTransform<T, d>& BTransform, const T Thickness)
+		TContactPoint<T> SphereCapsuleContactPoint(const TSphere<T, d>& A, const TRigidTransform<T, d>& ATransform, const TCapsule<T>& B, const TRigidTransform<T, d>& BTransform, const T CullDistance)
 		{
 			TContactPoint<T> Result;
 
@@ -659,20 +667,20 @@ namespace Chaos
 		}
 
 		template <typename T, int d>
-		void UpdateSphereCapsuleConstraint(const TSphere<T, d>& A, const TRigidTransform<T, d>& ATransform, const TCapsule<T>& B, const TRigidTransform<T, d>& BTransform, const T Thickness, TRigidBodyPointContactConstraint<T, d>& Constraint)
+		void UpdateSphereCapsuleConstraint(const TSphere<T, d>& A, const TRigidTransform<T, d>& ATransform, const TCapsule<T>& B, const TRigidTransform<T, d>& BTransform, const T CullDistance, TRigidBodyPointContactConstraint<T, d>& Constraint)
 		{
-			UpdateContactPoint(Constraint.Manifold, SphereCapsuleContactPoint(A, ATransform, B, BTransform, Thickness));
+			UpdateContactPoint(Constraint.Manifold, SphereCapsuleContactPoint(A, ATransform, B, BTransform, CullDistance));
 		}
 
 		template<class T, int d>
-		void UpdateSphereCapsuleManifold(TCollisionConstraintBase<T, d>&  Constraint, const TRigidTransform<T, d>& ATM, const TRigidTransform<T, d>& BTM, const T Thickness)
+		void UpdateSphereCapsuleManifold(TCollisionConstraintBase<T, d>&  Constraint, const TRigidTransform<T, d>& ATM, const TRigidTransform<T, d>& BTM, const T CullDistance)
 		{
 			// @todo(chaos) : Stub Update Manifold
 			//   Stub function for updating the manifold prior to the Apply and ApplyPushOut
 		}
 
 		template<typename T, int d>
-		void ConstructSphereCapsuleConstraints(TGeometryParticleHandle<T, d>* Particle0, TGeometryParticleHandle<T, d>* Particle1, const FImplicitObject* Implicit0, const FImplicitObject* Implicit1, const TRigidTransform<T, d>& Transform0, const TRigidTransform<T, d>& Transform1, const T Thickness, FCollisionConstraintsArray& NewConstraints)
+		void ConstructSphereCapsuleConstraints(TGeometryParticleHandle<T, d>* Particle0, TGeometryParticleHandle<T, d>* Particle1, const FImplicitObject* Implicit0, const FImplicitObject* Implicit1, const TRigidTransform<T, d>& Transform0, const TRigidTransform<T, d>& Transform1, const T CullDistance, FCollisionConstraintsArray& NewConstraints)
 		{
 			const TSphere<T, d>* Object0 = Implicit0->template GetObject<const TSphere<T, d> >();
 			const TCapsule<T>* Object1 = Implicit1->template GetObject<const TCapsule<T> >();
@@ -681,8 +689,8 @@ namespace Chaos
 				TRigidTransform<T, d> ParticleImplicit0TM = Transform0.GetRelativeTransform(Collisions::GetTransform(Particle0));
 				TRigidTransform<T, d> ParticleImplicit1TM = Transform1.GetRelativeTransform(Collisions::GetTransform(Particle1));
 				FRigidBodyPointContactConstraint Constraint = FRigidBodyPointContactConstraint(Particle0, Implicit0, ParticleImplicit0TM, Particle1, Implicit1, ParticleImplicit1TM);
-				UpdateSphereCapsuleConstraint(*Object0, Transform0, *Object1, Transform1, Thickness, Constraint);
-				NewConstraints.TryAdd(Thickness, Constraint);
+				UpdateSphereCapsuleConstraint(*Object0, Transform0, *Object1, Transform1, CullDistance, Constraint);
+				NewConstraints.TryAdd(CullDistance, Constraint);
 			}
 		}
 
@@ -692,7 +700,7 @@ namespace Chaos
 		//
 
 		template <typename T, int d>
-		TContactPoint<T> CapsuleCapsuleContactPoint(const TCapsule<T>& A, const TRigidTransform<T, d>& ATransform, const TCapsule<T>& B, const TRigidTransform<T, d>& BTransform, const T Thickness)
+		TContactPoint<T> CapsuleCapsuleContactPoint(const TCapsule<T>& A, const TRigidTransform<T, d>& ATransform, const TCapsule<T>& B, const TRigidTransform<T, d>& BTransform, const T CullDistance)
 		{
 			TContactPoint<T> Result;
 
@@ -718,20 +726,20 @@ namespace Chaos
 		}
 
 		template <typename T, int d>
-		void UpdateCapsuleCapsuleConstraint(const TCapsule<T>& A, const TRigidTransform<T, d>& ATransform, const TCapsule<T>& B, const TRigidTransform<T, d>& BTransform, const T Thickness, TRigidBodyPointContactConstraint<T, d>& Constraint)
+		void UpdateCapsuleCapsuleConstraint(const TCapsule<T>& A, const TRigidTransform<T, d>& ATransform, const TCapsule<T>& B, const TRigidTransform<T, d>& BTransform, const T CullDistance, TRigidBodyPointContactConstraint<T, d>& Constraint)
 		{
-			UpdateContactPoint(Constraint.Manifold, CapsuleCapsuleContactPoint(A, ATransform, B, BTransform, Thickness));
+			UpdateContactPoint(Constraint.Manifold, CapsuleCapsuleContactPoint(A, ATransform, B, BTransform, CullDistance));
 		}
 
 		template<class T, int d>
-		void UpdateCapsuleCapsuleManifold(TCollisionConstraintBase<T, d>&  Constraint, const TRigidTransform<T, d>& ATM, const TRigidTransform<T, d>& BTM, const T Thickness)
+		void UpdateCapsuleCapsuleManifold(TCollisionConstraintBase<T, d>&  Constraint, const TRigidTransform<T, d>& ATM, const TRigidTransform<T, d>& BTM, const T CullDistance)
 		{
 			// @todo(chaos) : Stub Update Manifold
 			//   Stub function for updating the manifold prior to the Apply and ApplyPushOut
 		}
 
 		template<typename T, int d>
-		void ConstructCapsuleCapsuleConstraints(TGeometryParticleHandle<T, d>* Particle0, TGeometryParticleHandle<T, d>* Particle1, const FImplicitObject* Implicit0, const FImplicitObject* Implicit1, const TRigidTransform<T, d>& Transform0, const TRigidTransform<T, d>& Transform1, const T Thickness, FCollisionConstraintsArray& NewConstraints)
+		void ConstructCapsuleCapsuleConstraints(TGeometryParticleHandle<T, d>* Particle0, TGeometryParticleHandle<T, d>* Particle1, const FImplicitObject* Implicit0, const FImplicitObject* Implicit1, const TRigidTransform<T, d>& Transform0, const TRigidTransform<T, d>& Transform1, const T CullDistance, FCollisionConstraintsArray& NewConstraints)
 		{
 
 			const TCapsule<T> * Object0 = Implicit0->template GetObject<const TCapsule<T> >();
@@ -741,8 +749,8 @@ namespace Chaos
 				TRigidTransform<T, d> ParticleImplicit0TM = Transform0.GetRelativeTransform(Collisions::GetTransform(Particle0));
 				TRigidTransform<T, d> ParticleImplicit1TM = Transform1.GetRelativeTransform(Collisions::GetTransform(Particle1));
 				FRigidBodyPointContactConstraint Constraint = FRigidBodyPointContactConstraint(Particle0, Implicit0, ParticleImplicit0TM, Particle1, Implicit1, ParticleImplicit1TM);
-				UpdateCapsuleCapsuleConstraint(*Object0, Transform0, *Object1, Transform1, Thickness, Constraint);
-				NewConstraints.TryAdd(Thickness, Constraint);
+				UpdateCapsuleCapsuleConstraint(*Object0, Transform0, *Object1, Transform1, CullDistance, Constraint);
+				NewConstraints.TryAdd(CullDistance, Constraint);
 			}
 		}
 
@@ -751,37 +759,47 @@ namespace Chaos
 		//
 
 		template <typename T, int d>
-		TContactPoint<T> CapsuleBoxContactPoint(const TCapsule<T>& A, const TRigidTransform<T, d>& ATransform, const TAABB<T, d>& B, const TRigidTransform<T, d>& BTransform, const T Thickness)
+		TContactPoint<T> CapsuleBoxContactPoint(const TCapsule<T>& A, const TRigidTransform<T, d>& ATransform, const TAABB<T, d>& B, const TRigidTransform<T, d>& BTransform, const TVector<T, 3>& InitialDir, const T CullDistance)
 		{
-			return GJKContactPoint(A, ATransform, B, BTransform, Thickness);
+			return GJKContactPoint(A, ATransform, B, BTransform, InitialDir);
 		}
 
 		template <typename T, int d>
-		void UpdateCapsuleBoxConstraint(const TCapsule<T>& A, const TRigidTransform<T, d>& ATransform, const TAABB<T, d>& B, const TRigidTransform<T, d>& BTransform, const T Thickness, TRigidBodyPointContactConstraint<T, d>& Constraint)
+		void UpdateCapsuleBoxConstraint(const TCapsule<T>& A, const TRigidTransform<T, d>& ATransform, const TAABB<T, d>& B, const TRigidTransform<T, d>& BTransform, const T CullDistance, TRigidBodyPointContactConstraint<T, d>& Constraint)
 		{
-			UpdateContactPoint(Constraint.Manifold, CapsuleBoxContactPoint(A, ATransform, B, BTransform, Thickness));
+			const TVector<T, d> InitialDir = ATransform.GetRotation().Inverse() * -Constraint.GetNormal();
+			UpdateContactPoint(Constraint.Manifold, CapsuleBoxContactPoint(A, ATransform, B, BTransform, InitialDir, CullDistance));
 		}
 
 		template<class T, int d>
-		void UpdateCapsuleBoxManifold(TCollisionConstraintBase<T, d>&  Constraint, const TRigidTransform<T, d>& ATM, const TRigidTransform<T, d>& BTM, const T Thickness)
+		void UpdateCapsuleBoxManifold(TCollisionConstraintBase<T, d>&  Constraint, const TRigidTransform<T, d>& ATM, const TRigidTransform<T, d>& BTM, const T CullDistance)
 		{
 			// @todo(chaos) : Stub Update Manifold
 			//   Stub function for updating the manifold prior to the Apply and ApplyPushOut
 		}
 
 		template<typename T, int d>
-		void ConstructCapsuleBoxConstraints(TGeometryParticleHandle<T, d>* Particle0, TGeometryParticleHandle<T, d>* Particle1, const FImplicitObject* Implicit0, const FImplicitObject* Implicit1, const TRigidTransform<T, d>& Transform0, const TRigidTransform<T, d>& Transform1, const T Thickness, FCollisionConstraintsArray& NewConstraints)
+		void ConstructCapsuleBoxConstraints(TGeometryParticleHandle<T, d>* Particle0, TGeometryParticleHandle<T, d>* Particle1, const FImplicitObject* Implicit0, const FImplicitObject* Implicit1, const TRigidTransform<T, d>& Transform0, const TRigidTransform<T, d>& Transform1, const T CullDistance, FCollisionConstraintsArray& NewConstraints)
 		{
-
 			const TCapsule<T> * Object0 = Implicit0->template GetObject<const TCapsule<T> >();
 			const TBox<T, d> * Object1 = Implicit1->template GetObject<const TBox<T, d> >();
 			if (ensure(Object0 && Object1))
 			{
-				TRigidTransform<T, d> ParticleImplicit0TM = Transform0.GetRelativeTransform(Collisions::GetTransform(Particle0));
-				TRigidTransform<T, d> ParticleImplicit1TM = Transform1.GetRelativeTransform(Collisions::GetTransform(Particle1));
-				FRigidBodyPointContactConstraint Constraint = FRigidBodyPointContactConstraint(Particle0, Implicit0, ParticleImplicit0TM, Particle1, Implicit1, ParticleImplicit1TM);
-				UpdateCapsuleBoxConstraint(*Object0, Transform0, Object1->BoundingBox(), Transform1, Thickness, Constraint);
-				NewConstraints.TryAdd(Thickness, Constraint);
+				// Box-space AABB check
+				// @todo(chaos): avoid recalc of relative transform in GJK (will have to switch order of box/capsule)
+				const FRigidTransform3 CapsuleToBoxTM = Transform0.GetRelativeTransform(Transform1);
+				const FVec3 P1 = CapsuleToBoxTM.TransformPosition(Object0->GetX1());
+				const FVec3 P2 = CapsuleToBoxTM.TransformPosition(Object0->GetX2());
+				TAABB<T, 3> CapsuleAABB(P1.ComponentMin(P2), P1.ComponentMax(P2));
+				CapsuleAABB.Thicken(Object0->GetRadius() + CullDistance);
+				if (CapsuleAABB.Intersects(Object1->GetAABB()))
+				{
+					TRigidTransform<T, d> ParticleImplicit0TM = Transform0.GetRelativeTransform(Collisions::GetTransform(Particle0));
+					TRigidTransform<T, d> ParticleImplicit1TM = Transform1.GetRelativeTransform(Collisions::GetTransform(Particle1));
+					FRigidBodyPointContactConstraint Constraint = FRigidBodyPointContactConstraint(Particle0, Implicit0, ParticleImplicit0TM, Particle1, Implicit1, ParticleImplicit1TM, EContactShapesType::CapsuleBox);
+					UpdateCapsuleBoxConstraint(*Object0, Transform0, Object1->BoundingBox(), Transform1, CullDistance, Constraint);
+					NewConstraints.TryAdd(CullDistance, Constraint);
+				}
 			}
 		}
 
@@ -790,26 +808,26 @@ namespace Chaos
 		//
 
 		template <typename T, int d>
-		TContactPoint<T> CapsuleHeightFieldContactPoint(const TCapsule<T>& A, const TRigidTransform<T, d>& ATransform, const THeightField<T>& B, const TRigidTransform<T, d>& BTransform, const T Thickness)
+		TContactPoint<T> CapsuleHeightFieldContactPoint(const TCapsule<T>& A, const TRigidTransform<T, d>& ATransform, const THeightField<T>& B, const TRigidTransform<T, d>& BTransform, const T CullDistance)
 		{
-			return GJKImplicitContactPoint< TCapsule<float> >(TCapsule<float>(A), ATransform, B, BTransform, Thickness);
+			return GJKImplicitContactPoint< TCapsule<float> >(TCapsule<float>(A), ATransform, B, BTransform, CullDistance);
 		}
 
 		template <typename T, int d>
-		void UpdateCapsuleHeightFieldConstraint(const TCapsule<T>& A, const TRigidTransform<T, d>& ATransform, const THeightField<T>& B, const TRigidTransform<T, d>& BTransform, const T Thickness, TRigidBodyPointContactConstraint<T, d>& Constraint)
+		void UpdateCapsuleHeightFieldConstraint(const TCapsule<T>& A, const TRigidTransform<T, d>& ATransform, const THeightField<T>& B, const TRigidTransform<T, d>& BTransform, const T CullDistance, TRigidBodyPointContactConstraint<T, d>& Constraint)
 		{
-			UpdateContactPoint(Constraint.Manifold, CapsuleHeightFieldContactPoint(A, ATransform, B, BTransform, Thickness));
+			UpdateContactPoint(Constraint.Manifold, CapsuleHeightFieldContactPoint(A, ATransform, B, BTransform, CullDistance));
 		}
 
 		template<class T, int d>
-		void UpdateCapsuleHeightFieldManifold(TCollisionConstraintBase<T, d>&  Constraint, const TRigidTransform<T, d>& ATM, const TRigidTransform<T, d>& BTM, const T Thickness)
+		void UpdateCapsuleHeightFieldManifold(TCollisionConstraintBase<T, d>&  Constraint, const TRigidTransform<T, d>& ATM, const TRigidTransform<T, d>& BTM, const T CullDistance)
 		{
 			// @todo(chaos) : Stub Update Manifold
 			//   Stub function for updating the manifold prior to the Apply and ApplyPushOut
 		}
 
 		template<typename T, int d>
-		void ConstructCapsuleHeightFieldConstraints(TGeometryParticleHandle<T, d>* Particle0, TGeometryParticleHandle<T, d>* Particle1, const FImplicitObject* Implicit0, const FImplicitObject* Implicit1, const TRigidTransform<T, d>& Transform0, const TRigidTransform<T, d>& Transform1, const T Thickness, FCollisionConstraintsArray& NewConstraints)
+		void ConstructCapsuleHeightFieldConstraints(TGeometryParticleHandle<T, d>* Particle0, TGeometryParticleHandle<T, d>* Particle1, const FImplicitObject* Implicit0, const FImplicitObject* Implicit1, const TRigidTransform<T, d>& Transform0, const TRigidTransform<T, d>& Transform1, const T CullDistance, FCollisionConstraintsArray& NewConstraints)
 		{
 
 			const TCapsule<T> * Object0 = Implicit0->template GetObject<const TCapsule<T> >();
@@ -819,8 +837,8 @@ namespace Chaos
 				TRigidTransform<T, d> ParticleImplicit0TM = Transform0.GetRelativeTransform(Collisions::GetTransform(Particle0));
 				TRigidTransform<T, d> ParticleImplicit1TM = Transform1.GetRelativeTransform(Collisions::GetTransform(Particle1));
 				FRigidBodyPointContactConstraint Constraint = FRigidBodyPointContactConstraint(Particle0, Implicit0, ParticleImplicit0TM, Particle1, Implicit1, ParticleImplicit1TM);
-				UpdateCapsuleHeightFieldConstraint(*Object0, Transform0, *Object1, Transform1, Thickness, Constraint);
-				NewConstraints.TryAdd(Thickness, Constraint);
+				UpdateCapsuleHeightFieldConstraint(*Object0, Transform0, *Object1, Transform1, CullDistance, Constraint);
+				NewConstraints.TryAdd(CullDistance, Constraint);
 			}
 		}
 
@@ -829,13 +847,13 @@ namespace Chaos
 		//
 
 		template<class T, int d>
-		void UpdateConvexConvexConstraint(const FImplicitObject& Implicit0, const TRigidTransform<T, d>& Transform0, const FImplicitObject& Implicit1, const TRigidTransform<T, d>& Transform1, const T Thickness, TCollisionConstraintBase<T, d>& ConstraintBase)
+		void UpdateConvexConvexConstraint(const FImplicitObject& Implicit0, const TRigidTransform<T, d>& Transform0, const FImplicitObject& Implicit1, const TRigidTransform<T, d>& Transform1, const T CullDistance, TCollisionConstraintBase<T, d>& ConstraintBase)
 		{
 			TContactPoint<T> ContactPoint;
 
 			if (ConstraintBase.GetType() == FRigidBodyPointContactConstraint::StaticType())
 			{
-				ContactPoint = ConvexConvexContactPoint(Implicit0, Transform0, Implicit1, Transform1, Thickness);
+				ContactPoint = ConvexConvexContactPoint(Implicit0, Transform0, Implicit1, Transform1, CullDistance);
 			}
 			else if (ConstraintBase.GetType() == FRigidBodyMultiPointContactConstraint::StaticType())
 			{
@@ -867,31 +885,31 @@ namespace Chaos
 		}
 
 		template<class T, int d>
-		void UpdateConvexConvexManifold(TCollisionConstraintBase<T, d>&  ConstraintBase, const TRigidTransform<T, d>& Transform0, const TRigidTransform<T, d>& Transform1, const T Thickness)
+		void UpdateConvexConvexManifold(TCollisionConstraintBase<T, d>&  ConstraintBase, const TRigidTransform<T, d>& Transform0, const TRigidTransform<T, d>& Transform1, const T CullDistance)
 		{
 			if (TRigidBodyMultiPointContactConstraint<T,d>* Constraint = ConstraintBase.template As<TRigidBodyMultiPointContactConstraint<T, d>>())
 			{
 				if (GetInnerType(ConstraintBase.Manifold.Implicit[0]->GetType()) == ImplicitObjectType::Convex)
 				{
-					UpdateSingleShotManifold(*Constraint, Transform0, Transform1, Thickness);
+					UpdateSingleShotManifold(*Constraint, Transform0, Transform1, CullDistance);
 				}
 				else
 				{
-					UpdateIterativeManifold(*Constraint, Transform0, Transform1, Thickness);
+					UpdateIterativeManifold(*Constraint, Transform0, Transform1, CullDistance);
 				}
 			}
 		}
 
 
 		template<class T, int d>
-		void ConstructConvexConvexConstraints(TGeometryParticleHandle<T, d>* Particle0, TGeometryParticleHandle<T, d>* Particle1, const FImplicitObject* Implicit0, const FImplicitObject* Implicit1, const TRigidTransform<T, d>& Transform0, const TRigidTransform<T, d>& Transform1, const T Thickness, FCollisionConstraintsArray& NewConstraints)
+		void ConstructConvexConvexConstraints(TGeometryParticleHandle<T, d>* Particle0, TGeometryParticleHandle<T, d>* Particle1, const FImplicitObject* Implicit0, const FImplicitObject* Implicit1, const TRigidTransform<T, d>& Transform0, const TRigidTransform<T, d>& Transform1, const T CullDistance, FCollisionConstraintsArray& NewConstraints)
 		{
 			TRigidTransform<T, d> ParticleImplicit0TM = Transform0.GetRelativeTransform(Collisions::GetTransform(Particle0));
 			TRigidTransform<T, d> ParticleImplicit1TM = Transform1.GetRelativeTransform(Collisions::GetTransform(Particle1));
 			FRigidBodyMultiPointContactConstraint Constraint = FRigidBodyMultiPointContactConstraint(Particle0, Implicit0, ParticleImplicit0TM, Particle1, Implicit1, ParticleImplicit1TM);
-			UpdateConvexConvexManifold(Constraint, Transform0, Transform1, Thickness);
-			UpdateConvexConvexConstraint(*Implicit0, Transform0, *Implicit1, Transform1, Thickness, Constraint);
-			NewConstraints.TryAdd(Thickness, Constraint);
+			UpdateConvexConvexManifold(Constraint, Transform0, Transform1, CullDistance);
+			UpdateConvexConvexConstraint(*Implicit0, Transform0, *Implicit1, Transform1, CullDistance, Constraint);
+			NewConstraints.TryAdd(CullDistance, Constraint);
 		}
 
 		//
@@ -899,26 +917,26 @@ namespace Chaos
 		//
 
 		template <typename T, int d>
-		TContactPoint<T> ConvexHeightFieldContactPoint(const FImplicitObject& A, const TRigidTransform<T, d>& ATransform, const THeightField<T>& B, const TRigidTransform<T, d>& BTransform, const T Thickness)
+		TContactPoint<T> ConvexHeightFieldContactPoint(const FImplicitObject& A, const TRigidTransform<T, d>& ATransform, const THeightField<T>& B, const TRigidTransform<T, d>& BTransform, const T CullDistance)
 		{
-			return GJKImplicitContactPoint< FConvex >(A, ATransform, B, BTransform, Thickness);
+			return GJKImplicitContactPoint< FConvex >(A, ATransform, B, BTransform, CullDistance);
 		}
 
 		template <typename T, int d>
-		void UpdateConvexHeightFieldConstraint(const FImplicitObject& A, const TRigidTransform<T, d>& ATransform, const THeightField<T>& B, const TRigidTransform<T, d>& BTransform, const T Thickness, TRigidBodyPointContactConstraint<T, d>& Constraint)
+		void UpdateConvexHeightFieldConstraint(const FImplicitObject& A, const TRigidTransform<T, d>& ATransform, const THeightField<T>& B, const TRigidTransform<T, d>& BTransform, const T CullDistance, TRigidBodyPointContactConstraint<T, d>& Constraint)
 		{
-			UpdateContactPoint(Constraint.Manifold, ConvexHeightFieldContactPoint(A, ATransform, B, BTransform, Thickness));
+			UpdateContactPoint(Constraint.Manifold, ConvexHeightFieldContactPoint(A, ATransform, B, BTransform, CullDistance));
 		}
 
 		template<class T, int d>
-		void UpdateConvexHeightFieldManifold(TCollisionConstraintBase<T, d>&  Constraint, const TRigidTransform<T, d>& ATM, const TRigidTransform<T, d>& BTM, const T Thickness)
+		void UpdateConvexHeightFieldManifold(TCollisionConstraintBase<T, d>&  Constraint, const TRigidTransform<T, d>& ATM, const TRigidTransform<T, d>& BTM, const T CullDistance)
 		{
 			// @todo(chaos) : Stub Update Manifold
 			//   Stub function for updating the manifold prior to the Apply and ApplyPushOut
 		}
 
 		template<typename T, int d>
-		void ConstructConvexHeightFieldConstraints(TGeometryParticleHandle<T, d>* Particle0, TGeometryParticleHandle<T, d>* Particle1, const FImplicitObject* Implicit0, const FImplicitObject* Implicit1, const TRigidTransform<T, d>& Transform0, const TRigidTransform<T, d>& Transform1, const T Thickness, FCollisionConstraintsArray& NewConstraints)
+		void ConstructConvexHeightFieldConstraints(TGeometryParticleHandle<T, d>* Particle0, TGeometryParticleHandle<T, d>* Particle1, const FImplicitObject* Implicit0, const FImplicitObject* Implicit1, const TRigidTransform<T, d>& Transform0, const TRigidTransform<T, d>& Transform1, const T CullDistance, FCollisionConstraintsArray& NewConstraints)
 		{
 
 			const THeightField<T> * Object1 = Implicit1->template GetObject<const THeightField<T> >();
@@ -927,8 +945,8 @@ namespace Chaos
 				TRigidTransform<T, d> ParticleImplicit0TM = Transform0.GetRelativeTransform(Collisions::GetTransform(Particle0));
 				TRigidTransform<T, d> ParticleImplicit1TM = Transform1.GetRelativeTransform(Collisions::GetTransform(Particle1));
 				FRigidBodyPointContactConstraint Constraint = FRigidBodyPointContactConstraint(Particle0, Implicit0, ParticleImplicit0TM, Particle1, Implicit1, ParticleImplicit1TM);
-				UpdateConvexHeightFieldConstraint(*Implicit0, Transform0, *Object1, Transform1, Thickness, Constraint);
-				NewConstraints.TryAdd(Thickness, Constraint);
+				UpdateConvexHeightFieldConstraint(*Implicit0, Transform0, *Object1, Transform1, CullDistance, Constraint);
+				NewConstraints.TryAdd(CullDistance, Constraint);
 			}
 		}
 
@@ -940,7 +958,7 @@ namespace Chaos
 
 		DECLARE_CYCLE_STAT(TEXT("TPBDCollisionConstraints::UpdateLevelsetLevelsetConstraint"), STAT_UpdateLevelsetLevelsetConstraint, STATGROUP_ChaosWide);
 		template<ECollisionUpdateType UpdateType, typename T, int d>
-		void UpdateLevelsetLevelsetConstraint(const T Thickness, TRigidBodyPointContactConstraint<T, d>& Constraint)
+		void UpdateLevelsetLevelsetConstraint(const T CullDistance, TRigidBodyPointContactConstraint<T, d>& Constraint)
 		{
 			SCOPE_CYCLE_COUNTER(STAT_UpdateLevelsetLevelsetConstraint);
 
@@ -963,12 +981,12 @@ namespace Chaos
 
 			if (SampleParticles)
 			{
-				SampleObject<UpdateType>(*Particle1->Geometry(), LevelsetTM, *SampleParticles, ParticlesTM, Thickness, Constraint);
+				SampleObject<UpdateType>(*Particle1->Geometry(), LevelsetTM, *SampleParticles, ParticlesTM, CullDistance, Constraint);
 			}
 		}
 
 		template<class T, int d>
-		void UpdateLevelsetLevelsetManifold(TCollisionConstraintBase<T, d>&  Constraint, const TRigidTransform<T, d>& ATM, const TRigidTransform<T, d>& BTM, const T Thickness)
+		void UpdateLevelsetLevelsetManifold(TCollisionConstraintBase<T, d>&  Constraint, const TRigidTransform<T, d>& ATM, const TRigidTransform<T, d>& BTM, const T CullDistance)
 		{
 			// @todo(chaos) : Stub Update Manifold
 			//   Stub function for updating the manifold prior to the Apply and ApplyPushOut
@@ -976,7 +994,7 @@ namespace Chaos
 
 
 		template<typename T, int d>
-		void ConstructLevelsetLevelsetConstraints(TGeometryParticleHandle<T, d>* Particle0, TGeometryParticleHandle<T, d>* Particle1, const FImplicitObject* Implicit0, const FImplicitObject* Implicit1, const TRigidTransform<T, d>& Transform0, const TRigidTransform<T, d>& Transform1, const T Thickness, FCollisionConstraintsArray& NewConstraints)
+		void ConstructLevelsetLevelsetConstraints(TGeometryParticleHandle<T, d>* Particle0, TGeometryParticleHandle<T, d>* Particle1, const FImplicitObject* Implicit0, const FImplicitObject* Implicit1, const TRigidTransform<T, d>& Transform0, const TRigidTransform<T, d>& Transform1, const T CullDistance, FCollisionConstraintsArray& NewConstraints)
 		{
 			TRigidTransform<T, d> ParticleImplicit0TM = Transform0.GetRelativeTransform(Collisions::GetTransform(Particle0));
 			TRigidTransform<T, d> ParticleImplicit1TM = Transform1.GetRelativeTransform(Collisions::GetTransform(Particle1));
@@ -996,9 +1014,9 @@ namespace Chaos
 				Constraint.SetManifold(Implicit0, Implicit1);
 			}
 
-			UpdateLevelsetLevelsetConstraint<ECollisionUpdateType::Any>(Thickness, Constraint);
+			UpdateLevelsetLevelsetConstraint<ECollisionUpdateType::Any>(CullDistance, Constraint);
 
-			NewConstraints.TryAdd(Thickness, Constraint);
+			NewConstraints.TryAdd(CullDistance, Constraint);
 		}
 
 		//
@@ -1007,11 +1025,11 @@ namespace Chaos
 
 		DECLARE_CYCLE_STAT(TEXT("TPBDCollisionConstraints::ConstructUnionUnionConstraints"), STAT_ConstructUnionUnionConstraints, STATGROUP_ChaosWide);
 		template<typename T, int d>
-		void ConstructUnionUnionConstraints(TGeometryParticleHandle<T, d>* Particle0, TGeometryParticleHandle<T, d>* Particle1, const FImplicitObject* Implicit0, const FImplicitObject* Implicit1, const TRigidTransform<T, d>& Transform0, const TRigidTransform<T, d>& Transform1, const T Thickness, FCollisionConstraintsArray& NewConstraints)
+		void ConstructUnionUnionConstraints(TGeometryParticleHandle<T, d>* Particle0, TGeometryParticleHandle<T, d>* Particle1, const FImplicitObject* Implicit0, const FImplicitObject* Implicit1, const TRigidTransform<T, d>& Transform0, const TRigidTransform<T, d>& Transform1, const T CullDistance, FCollisionConstraintsArray& NewConstraints)
 		{
 			SCOPE_CYCLE_COUNTER(STAT_ConstructUnionUnionConstraints);
 
-			const TArray<Pair<const FImplicitObject*, TRigidTransform<T, d>>> LevelsetShapes = FindRelevantShapes(Implicit0, Transform0, *Implicit1, Transform1, Thickness);
+			const TArray<Pair<const FImplicitObject*, TRigidTransform<T, d>>> LevelsetShapes = FindRelevantShapes(Implicit0, Transform0, *Implicit1, Transform1, CullDistance);
 
 			for (const Pair<const FImplicitObject*, TRigidTransform<T, d>>& LevelsetObjPair : LevelsetShapes)
 			{
@@ -1019,13 +1037,13 @@ namespace Chaos
 				const TRigidTransform<T, d>& LevelsetInnerObjTM = LevelsetObjPair.Second * Transform1;
 
 				//now find all particle inner objects
-				const TArray<Pair<const FImplicitObject*, TRigidTransform<T, d>>> ParticleShapes = FindRelevantShapes(LevelsetInnerObj, LevelsetInnerObjTM, *Implicit0, Transform0, Thickness);
+				const TArray<Pair<const FImplicitObject*, TRigidTransform<T, d>>> ParticleShapes = FindRelevantShapes(LevelsetInnerObj, LevelsetInnerObjTM, *Implicit0, Transform0, CullDistance);
 
 				//for each inner obj pair, update constraint
 				for (const Pair<const FImplicitObject*, TRigidTransform<T, d>>& ParticlePair : ParticleShapes)
 				{
 					const FImplicitObject* ParticleInnerObj = ParticlePair.First;
-					ConstructConstraints<T, d>(Particle0, Particle1, ParticleInnerObj, LevelsetInnerObj, Transform0, Transform1, Thickness, NewConstraints);
+					ConstructConstraints<T, d>(Particle0, Particle1, ParticleInnerObj, LevelsetInnerObj, Transform0, Transform1, CullDistance, NewConstraints);
 				}
 			}
 		}
@@ -1035,7 +1053,7 @@ namespace Chaos
 		//
 
 		template<typename T, int d>
-		void UpdateManifold(TCollisionConstraintBase<T, d>& ConstraintBase, const TRigidTransform<T, d>& ATM, const TRigidTransform<T, d>& BTM, const T Thickness)
+		void UpdateManifold(TCollisionConstraintBase<T, d>& ConstraintBase, const TRigidTransform<T, d>& ATM, const TRigidTransform<T, d>& BTM, const T CullDistance)
 		{
 			const FImplicitObject& Implicit0 = *ConstraintBase.Manifold.Implicit[0];
 			const FImplicitObject& Implicit1 = *ConstraintBase.Manifold.Implicit[1];
@@ -1087,47 +1105,47 @@ namespace Chaos
 			
 			if (Implicit0Type == TBox<T, d>::StaticType() && Implicit1Type == TBox<T, d>::StaticType())
 			{
-				UpdateBoxBoxManifold(ConstraintBase, Transform0, Transform1, Thickness);
+				UpdateBoxBoxManifold(ConstraintBase, Transform0, Transform1, CullDistance);
 			}
 			else if (Implicit0Type == TBox<T, d>::StaticType() && Implicit1Type == THeightField<T>::StaticType())
 			{
-				UpdateBoxHeightFieldManifold(ConstraintBase, Transform0, Transform1, Thickness);
+				UpdateBoxHeightFieldManifold(ConstraintBase, Transform0, Transform1, CullDistance);
 			}
 			else if (Implicit0Type == TSphere<T, d>::StaticType() && Implicit1Type == TSphere<T, d>::StaticType())
 			{
-				UpdateSphereSphereManifold(ConstraintBase, Transform0, Transform1, Thickness);
+				UpdateSphereSphereManifold(ConstraintBase, Transform0, Transform1, CullDistance);
 			}
 			else if (Implicit0Type == TSphere<T, d>::StaticType() && Implicit1Type == THeightField<T>::StaticType())
 			{
-				UpdateSphereHeightFieldManifold(ConstraintBase, Transform0, Transform1, Thickness);
+				UpdateSphereHeightFieldManifold(ConstraintBase, Transform0, Transform1, CullDistance);
 			}
 			else if (Implicit0Type == TBox<T, d>::StaticType() && Implicit1Type == TPlane<T, d>::StaticType())
 			{
-				UpdateBoxPlaneManifold(ConstraintBase, Transform0, Transform1, Thickness);
+				UpdateBoxPlaneManifold(ConstraintBase, Transform0, Transform1, CullDistance);
 			}
 			else if (Implicit0Type == TSphere<T, d>::StaticType() && Implicit1Type == TPlane<T, d>::StaticType())
 			{
-				UpdateSpherePlaneManifold(ConstraintBase, Transform0, Transform1, Thickness);
+				UpdateSpherePlaneManifold(ConstraintBase, Transform0, Transform1, CullDistance);
 			}
 			else if (Implicit0Type == TSphere<T, d>::StaticType() && Implicit1Type == TBox<T, d>::StaticType())
 			{
-				UpdateSphereBoxManifold(ConstraintBase, Transform0, Transform1, Thickness);
+				UpdateSphereBoxManifold(ConstraintBase, Transform0, Transform1, CullDistance);
 			}
 			else if (Implicit0Type == TSphere<T, d>::StaticType() && Implicit1Type == TCapsule<T>::StaticType())
 			{
-				UpdateSphereCapsuleManifold(ConstraintBase, Transform0, Transform1, Thickness);
+				UpdateSphereCapsuleManifold(ConstraintBase, Transform0, Transform1, CullDistance);
 			}
 			else if (Implicit0Type == TCapsule<T>::StaticType() && Implicit1Type == TCapsule<T>::StaticType())
 			{
-				UpdateCapsuleCapsuleManifold(ConstraintBase, Transform0, Transform1, Thickness);
+				UpdateCapsuleCapsuleManifold(ConstraintBase, Transform0, Transform1, CullDistance);
 			}
 			else if (Implicit0Type == TCapsule<T>::StaticType() && Implicit1Type == TBox<T, d>::StaticType())
 			{
-				UpdateCapsuleBoxManifold(ConstraintBase, Transform0, Transform1, Thickness);
+				UpdateCapsuleBoxManifold(ConstraintBase, Transform0, Transform1, CullDistance);
 			}
 			else if (Implicit0Type == TCapsule<T>::StaticType() && Implicit1Type == THeightField<T>::StaticType())
 			{
-				UpdateCapsuleHeightFieldManifold(ConstraintBase, Transform0, Transform1, Thickness);
+				UpdateCapsuleHeightFieldManifold(ConstraintBase, Transform0, Transform1, CullDistance);
 			}
 			else if (Implicit0Type == THeightField<T>::StaticType() && Implicit1Type == TBox<T, d>::StaticType())
 			{
@@ -1135,7 +1153,7 @@ namespace Chaos
 			}
 			else if (Implicit0Type == TPlane<T, d>::StaticType() && Implicit1Type == TBox<T, d>::StaticType())
 			{
-				//UpdatePlaneBoxManifold(ConstraintBase, Transform0, Transform1, Thickness);
+				//UpdatePlaneBoxManifold(ConstraintBase, Transform0, Transform1, CullDistance);
 			}
 			else if (Implicit0Type == THeightField<T>::StaticType() && Implicit1Type == TSphere<T, d>::StaticType())
 			{
@@ -1143,19 +1161,19 @@ namespace Chaos
 			}
 			else if (Implicit0Type == TPlane<T, d>::StaticType() && Implicit1Type == TSphere<T, d>::StaticType())
 			{
-				//UpdatePlaneSphereManifold(ConstraintBase, Transform0, Transform1, Thickness);
+				//UpdatePlaneSphereManifold(ConstraintBase, Transform0, Transform1, CullDistance);
 			}
 			else if (Implicit0Type == TBox<T, d>::StaticType() && Implicit1Type == TSphere<T, d>::StaticType())
 			{
-				//UpdateBoxSphereManifold(ConstraintBase, Transform0, Transform1, Thickness);
+				//UpdateBoxSphereManifold(ConstraintBase, Transform0, Transform1, CullDistance);
 			}
 			else if (Implicit0Type == TBox<T, d>::StaticType() && Implicit1Type == TCapsule<T>::StaticType())
 			{
-				//UpdateBoxCapsuleManifold(ConstraintBase, Transform0, Transform1, Thickness);
+				//UpdateBoxCapsuleManifold(ConstraintBase, Transform0, Transform1, CullDistance);
 			}
 			else if (Implicit0Type == TCapsule<T>::StaticType() && Implicit1Type == TSphere<T, d>::StaticType())
 			{
-				//UpdateCapsuleSphereManifold(ConstraintBase, Transform0, Transform1, Thickness);
+				//UpdateCapsuleSphereManifold(ConstraintBase, Transform0, Transform1, CullDistance);
 			}
 			// the generic convex bodies are last
 			else if (Implicit0Type == THeightField<T>::StaticType() && Implicit1.IsConvex())
@@ -1164,21 +1182,21 @@ namespace Chaos
 			}
 			else if (Implicit0.IsConvex() && Implicit1Type == THeightField<T>::StaticType())
 			{
-				UpdateConvexHeightFieldManifold(ConstraintBase, Transform0, Transform1, Thickness);
+				UpdateConvexHeightFieldManifold(ConstraintBase, Transform0, Transform1, CullDistance);
 			}
 			else if (Implicit0.IsConvex() && Implicit1.IsConvex())
 			{
-				UpdateConvexConvexManifold(ConstraintBase, Transform0, Transform1, Thickness);
+				UpdateConvexConvexManifold(ConstraintBase, Transform0, Transform1, CullDistance);
 			}
 			else
 			{
-				UpdateLevelsetLevelsetManifold(ConstraintBase, Transform0, Transform1, Thickness);
+				UpdateLevelsetLevelsetManifold(ConstraintBase, Transform0, Transform1, CullDistance);
 			}
 
 			/*
 			else if (Implicit0Type == TBox<T, d>::StaticType() && Implicit1Type == FTriangleMeshImplicitObject::StaticType())
 			{
-				UpdateBoxTriangleMeshManifold(ConstraintBase, Transform0, Transform1, Thickness);
+				UpdateBoxTriangleMeshManifold(ConstraintBase, Transform0, Transform1, CullDistance);
 			}
 			else if (Implicit0Type == FTriangleMeshImplicitObject::StaticType() && Implicit1Type == TBox<T, d>::StaticType())
 			{
@@ -1189,14 +1207,8 @@ namespace Chaos
 		}
 
 		template<ECollisionUpdateType UpdateType, typename T, int d>
-		void UpdateConstraint(TCollisionConstraintBase<T, d>& ConstraintBase, const TRigidTransform<T, d>& ParticleTransform0, const TRigidTransform<T, d>& ParticleTransform1, const T Thickness)
+		void UpdateConstraintAny(TCollisionConstraintBase<T, d>& ConstraintBase, const FImplicitObject& Implicit0, const FImplicitObject& Implicit1, const TRigidTransform<T, d>& Transform0, const TRigidTransform<T, d>& Transform1, const T CullDistance)
 		{
-			const FImplicitObject& Implicit0 = *ConstraintBase.Manifold.Implicit[0];
-			const FImplicitObject& Implicit1 = *ConstraintBase.Manifold.Implicit[1];
-
-			const TRigidTransform<T, d>& Transform0 = ConstraintBase.ImplicitTransform[0] * ParticleTransform0;
-			const TRigidTransform<T, d>& Transform1 = ConstraintBase.ImplicitTransform[1] * ParticleTransform1;
-
 #if !UE_BUILD_SHIPPING
 			EImplicitObjectType Implicit0OuterType = Implicit0.GetType();
 			EImplicitObjectType Implicit1OuterType = Implicit1.GetType();
@@ -1240,47 +1252,47 @@ namespace Chaos
 			
 			if (Implicit0Type == TBox<T, d>::StaticType() && Implicit1Type == TBox<T, d>::StaticType())
 			{
-				UpdateBoxBoxConstraint(Implicit0.template GetObject<TBox<T, d>>()->GetAABB(), Transform0, Implicit1.template GetObject<TBox<T, d>>()->GetAABB(), Transform1, Thickness, *ConstraintBase.template As<TRigidBodyPointContactConstraint<T, d>>());
+				UpdateBoxBoxConstraint(Implicit0.template GetObject<TBox<T, d>>()->GetAABB(), Transform0, Implicit1.template GetObject<TBox<T, d>>()->GetAABB(), Transform1, CullDistance, *ConstraintBase.template As<TRigidBodyPointContactConstraint<T, d>>());
 			}
 			else if (Implicit0Type == TBox<T, d>::StaticType() && Implicit1Type == THeightField<T>::StaticType())
 			{
-				UpdateBoxHeightFieldConstraint(Implicit0.template GetObject<TBox<T, d>>()->GetAABB(), Transform0, *Implicit1.template GetObject< THeightField<T> >(), Transform1, Thickness, *ConstraintBase.template As<TRigidBodyPointContactConstraint<T, d>>());
+				UpdateBoxHeightFieldConstraint(Implicit0.template GetObject<TBox<T, d>>()->GetAABB(), Transform0, *Implicit1.template GetObject< THeightField<T> >(), Transform1, CullDistance, *ConstraintBase.template As<TRigidBodyPointContactConstraint<T, d>>());
 			}
 			else if (Implicit0Type == TSphere<T, d>::StaticType() && Implicit1Type == TSphere<T, d>::StaticType())
 			{
-				UpdateSphereSphereConstraint(*Implicit0.template GetObject<TSphere<T, d>>(), Transform0, *Implicit1.template GetObject<TSphere<T, d>>(), Transform1, Thickness, *ConstraintBase.template As<TRigidBodyPointContactConstraint<T, d>>());
+				UpdateSphereSphereConstraint(*Implicit0.template GetObject<TSphere<T, d>>(), Transform0, *Implicit1.template GetObject<TSphere<T, d>>(), Transform1, CullDistance, *ConstraintBase.template As<TRigidBodyPointContactConstraint<T, d>>());
 			}
 			else if (Implicit0Type == TSphere<T, d>::StaticType() && Implicit1Type == THeightField<T>::StaticType())
 			{
-				UpdateSphereHeightFieldConstraint(*Implicit0.template GetObject<TSphere<T, d>>(), Transform0, *Implicit1.template GetObject< THeightField<T> >(), Transform1, Thickness, *ConstraintBase.template As<TRigidBodyPointContactConstraint<T, d>>());
+				UpdateSphereHeightFieldConstraint(*Implicit0.template GetObject<TSphere<T, d>>(), Transform0, *Implicit1.template GetObject< THeightField<T> >(), Transform1, CullDistance, *ConstraintBase.template As<TRigidBodyPointContactConstraint<T, d>>());
 			}
 			else if (Implicit0Type == TBox<T, d>::StaticType() && Implicit1Type == TPlane<T, d>::StaticType())
 			{
-				UpdateBoxPlaneConstraint(Implicit0.template GetObject<TBox<T, d>>()->GetAABB(), Transform0, *Implicit1.template GetObject<TPlane<T, d>>(), Transform1, Thickness, *ConstraintBase.template As<TRigidBodyPointContactConstraint<T, d>>());
+				UpdateBoxPlaneConstraint(Implicit0.template GetObject<TBox<T, d>>()->GetAABB(), Transform0, *Implicit1.template GetObject<TPlane<T, d>>(), Transform1, CullDistance, *ConstraintBase.template As<TRigidBodyPointContactConstraint<T, d>>());
 			}
 			else if (Implicit0Type == TSphere<T, d>::StaticType() && Implicit1Type == TPlane<T, d>::StaticType())
 			{
-				UpdateSpherePlaneConstraint(*Implicit0.template GetObject<TSphere<T, d>>(), Transform0, *Implicit1.template GetObject<TPlane<T, d>>(), Transform1, Thickness, *ConstraintBase.template As<TRigidBodyPointContactConstraint<T, d>>());
+				UpdateSpherePlaneConstraint(*Implicit0.template GetObject<TSphere<T, d>>(), Transform0, *Implicit1.template GetObject<TPlane<T, d>>(), Transform1, CullDistance, *ConstraintBase.template As<TRigidBodyPointContactConstraint<T, d>>());
 			}
 			else if (Implicit0Type == TSphere<T, d>::StaticType() && Implicit1Type == TBox<T, d>::StaticType())
 			{
-				UpdateSphereBoxConstraint(*Implicit0.template GetObject<TSphere<T, d>>(), Transform0, Implicit1.template GetObject<TBox<T, d>>()->GetAABB(), Transform1, Thickness, *ConstraintBase.template As<TRigidBodyPointContactConstraint<T, d>>());
+				UpdateSphereBoxConstraint(*Implicit0.template GetObject<TSphere<T, d>>(), Transform0, Implicit1.template GetObject<TBox<T, d>>()->GetAABB(), Transform1, CullDistance, *ConstraintBase.template As<TRigidBodyPointContactConstraint<T, d>>());
 			}
 			else if (Implicit0Type == TSphere<T, d>::StaticType() && Implicit1Type == TCapsule<T>::StaticType())
 			{
-				UpdateSphereCapsuleConstraint(*Implicit0.template GetObject<TSphere<T, d>>(), Transform0, *Implicit1.template GetObject<TCapsule<T>>(), Transform1, Thickness, *ConstraintBase.template As<TRigidBodyPointContactConstraint<T, d>>());
+				UpdateSphereCapsuleConstraint(*Implicit0.template GetObject<TSphere<T, d>>(), Transform0, *Implicit1.template GetObject<TCapsule<T>>(), Transform1, CullDistance, *ConstraintBase.template As<TRigidBodyPointContactConstraint<T, d>>());
 			}
 			else if (Implicit0Type == TCapsule<T>::StaticType() && Implicit1Type == TCapsule<T>::StaticType())
 			{
-				UpdateCapsuleCapsuleConstraint(*Implicit0.template GetObject<TCapsule<T>>(), Transform0, *Implicit1.template GetObject<TCapsule<T>>(), Transform1, Thickness, *ConstraintBase.template As<TRigidBodyPointContactConstraint<T, d>>());
+				UpdateCapsuleCapsuleConstraint(*Implicit0.template GetObject<TCapsule<T>>(), Transform0, *Implicit1.template GetObject<TCapsule<T>>(), Transform1, CullDistance, *ConstraintBase.template As<TRigidBodyPointContactConstraint<T, d>>());
 			}
 			else if (Implicit0Type == TCapsule<T>::StaticType() && Implicit1Type == TBox<T, d>::StaticType())
 			{
-				UpdateCapsuleBoxConstraint(*Implicit0.template GetObject<TCapsule<T>>(), Transform0, Implicit1.template GetObject<TBox<T, d>>()->GetAABB(), Transform1, Thickness, *ConstraintBase.template As<TRigidBodyPointContactConstraint<T, d>>());
+				UpdateCapsuleBoxConstraint(*Implicit0.template GetObject<TCapsule<T>>(), Transform0, Implicit1.template GetObject<TBox<T, d>>()->GetAABB(), Transform1, CullDistance, *ConstraintBase.template As<TRigidBodyPointContactConstraint<T, d>>());
 			}
 			else if (Implicit0Type == TCapsule<T>::StaticType() && Implicit1Type == THeightField<T>::StaticType())
 			{
-				UpdateCapsuleHeightFieldConstraint(*Implicit0.template GetObject<TCapsule<T>>(), Transform0, *Implicit1.template GetObject< THeightField<T> >(), Transform1, Thickness, *ConstraintBase.template As<TRigidBodyPointContactConstraint<T, d>>());
+				UpdateCapsuleHeightFieldConstraint(*Implicit0.template GetObject<TCapsule<T>>(), Transform0, *Implicit1.template GetObject< THeightField<T> >(), Transform1, CullDistance, *ConstraintBase.template As<TRigidBodyPointContactConstraint<T, d>>());
 			}
 			else if (Implicit0Type == THeightField<T>::StaticType() && Implicit1Type == TBox<T, d>::StaticType())
 			{
@@ -1295,7 +1307,7 @@ namespace Chaos
 			{
 				TRigidBodyPointContactConstraint<T, d>& Constraint = *ConstraintBase.template As<TRigidBodyPointContactConstraint<T, d> >();
 				TRigidBodyPointContactConstraint<T, d> TmpConstraint = Constraint;
-				UpdateBoxPlaneConstraint(Implicit1.template GetObject<TBox<T, d>>()->GetAABB(), Transform1, *Implicit0.template GetObject<TPlane<T, d>>(), Transform0, Thickness, TmpConstraint);
+				UpdateBoxPlaneConstraint(Implicit1.template GetObject<TBox<T, d>>()->GetAABB(), Transform1, *Implicit0.template GetObject<TPlane<T, d>>(), Transform0, CullDistance, TmpConstraint);
 				if (TmpConstraint.GetPhi() < Constraint.GetPhi())
 				{
 					Constraint = TmpConstraint;
@@ -1315,7 +1327,7 @@ namespace Chaos
 			{
 				TRigidBodyPointContactConstraint<T, d>& Constraint = *ConstraintBase.template As<TRigidBodyPointContactConstraint<T, d> >();
 				TRigidBodyPointContactConstraint<T, d> TmpConstraint = Constraint;
-				UpdateSpherePlaneConstraint(*Implicit1.template GetObject<TSphere<T, d>>(), Transform1, *Implicit0.template GetObject<TPlane<T, d>>(), Transform0, Thickness, TmpConstraint);
+				UpdateSpherePlaneConstraint(*Implicit1.template GetObject<TSphere<T, d>>(), Transform1, *Implicit0.template GetObject<TPlane<T, d>>(), Transform0, CullDistance, TmpConstraint);
 				if (TmpConstraint.GetPhi() < Constraint.GetPhi())
 				{
 					Constraint = TmpConstraint;
@@ -1326,7 +1338,7 @@ namespace Chaos
 			{
 				TRigidBodyPointContactConstraint<T, d>& Constraint = *ConstraintBase.template As<TRigidBodyPointContactConstraint<T, d> >();
 				TRigidBodyPointContactConstraint<T, d> TmpConstraint = Constraint;
-				UpdateSphereBoxConstraint(*Implicit1.template GetObject<TSphere<T, d>>(), Transform1, Implicit0.template GetObject<TBox<T, d>>()->GetAABB(), Transform0, Thickness, TmpConstraint);
+				UpdateSphereBoxConstraint(*Implicit1.template GetObject<TSphere<T, d>>(), Transform1, Implicit0.template GetObject<TBox<T, d>>()->GetAABB(), Transform0, CullDistance, TmpConstraint);
 				if (TmpConstraint.GetPhi() < Constraint.GetPhi())
 				{
 					Constraint = TmpConstraint;
@@ -1337,7 +1349,7 @@ namespace Chaos
 			{
 				TRigidBodyPointContactConstraint<T, d>& Constraint = *ConstraintBase.template As<TRigidBodyPointContactConstraint<T, d> >();
 				TRigidBodyPointContactConstraint<T, d> TmpConstraint = Constraint;
-				UpdateCapsuleBoxConstraint(*Implicit1.template GetObject<TCapsule<T>>(), Transform1, Implicit0.template GetObject<TBox<T, d>>()->GetAABB(), Transform0, Thickness, TmpConstraint);
+				UpdateCapsuleBoxConstraint(*Implicit1.template GetObject<TCapsule<T>>(), Transform1, Implicit0.template GetObject<TBox<T, d>>()->GetAABB(), Transform0, CullDistance, TmpConstraint);
 				if (TmpConstraint.GetPhi() < Constraint.GetPhi())
 				{
 					Constraint = TmpConstraint;
@@ -1348,7 +1360,7 @@ namespace Chaos
 			{
 				TRigidBodyPointContactConstraint<T, d>& Constraint = *ConstraintBase.template As<TRigidBodyPointContactConstraint<T, d> >();
 				TRigidBodyPointContactConstraint<T, d> TmpConstraint = Constraint;
-				UpdateSphereCapsuleConstraint(*Implicit1.template GetObject<TSphere<T, d>>(), Transform1, *Implicit0.template GetObject<TCapsule<T>>(), Transform0, Thickness, TmpConstraint);
+				UpdateSphereCapsuleConstraint(*Implicit1.template GetObject<TSphere<T, d>>(), Transform1, *Implicit0.template GetObject<TCapsule<T>>(), Transform0, CullDistance, TmpConstraint);
 				if (TmpConstraint.GetPhi() < Constraint.GetPhi())
 				{
 					Constraint = TmpConstraint;
@@ -1376,21 +1388,21 @@ namespace Chaos
 			}
 			else if (Implicit0.IsConvex() && Implicit1Type == THeightField<T>::StaticType())
 			{
-				UpdateConvexHeightFieldConstraint(Implicit0, Transform0, *Implicit1.template GetObject< THeightField<T> >(), Transform1, Thickness, *ConstraintBase.template As<TRigidBodyPointContactConstraint<T, d>>());
+				UpdateConvexHeightFieldConstraint(Implicit0, Transform0, *Implicit1.template GetObject< THeightField<T> >(), Transform1, CullDistance, *ConstraintBase.template As<TRigidBodyPointContactConstraint<T, d>>());
 			}
 			else if (Implicit0.IsConvex() && Implicit1.IsConvex())
 			{
-				UpdateConvexConvexConstraint(Implicit0, Transform0, Implicit1, Transform1, Thickness, ConstraintBase);
+				UpdateConvexConvexConstraint(Implicit0, Transform0, Implicit1, Transform1, CullDistance, ConstraintBase);
 			}
 			else
 			{
-				UpdateLevelsetLevelsetConstraint<UpdateType>(Thickness, *ConstraintBase.template As<TRigidBodyPointContactConstraint<T, d> >());
+				UpdateLevelsetLevelsetConstraint<UpdateType>(CullDistance, *ConstraintBase.template As<TRigidBodyPointContactConstraint<T, d> >());
 			}
 
 			/*
 			else if (Implicit0Type == TBox<T, d>::StaticType() && Implicit1Type == FTriangleMeshImplicitObject::StaticType())
 			{
-			UpdateBoxTriangleMeshConstraint(Implicit0.template GetObject<TBox<T, d>>()->GetAABB(), Transform0, *GetUnscaledInner<FTriangleMeshImplicitObject>(Implicit1), Transform1, Thickness, *ConstraintBase.template As<TRigidBodyPointContactConstraint<T, d>>());
+			UpdateBoxTriangleMeshConstraint(Implicit0.template GetObject<TBox<T, d>>()->GetAABB(), Transform0, *GetUnscaledInner<FTriangleMeshImplicitObject>(Implicit1), Transform1, CullDistance, *ConstraintBase.template As<TRigidBodyPointContactConstraint<T, d>>());
 			}
 			else if (Implicit0Type == TBox<T, d>::StaticType() && Implicit1Type == FTriangleMeshImplicitObject::StaticType())
 			{
@@ -1402,11 +1414,36 @@ namespace Chaos
 				ensure(false);
 			}
 			*/
+		}
+
+		template<ECollisionUpdateType UpdateType, typename T, int d>
+		void UpdateConstraint(TCollisionConstraintBase<T, d>& ConstraintBase, const TRigidTransform<T, d>& ParticleTransform0, const TRigidTransform<T, d>& ParticleTransform1, const T CullDistance)
+		{
+			const FImplicitObject& Implicit0 = *ConstraintBase.Manifold.Implicit[0];
+			const FImplicitObject& Implicit1 = *ConstraintBase.Manifold.Implicit[1];
+
+			const TRigidTransform<T, d>& Transform0 = ConstraintBase.ImplicitTransform[0] * ParticleTransform0;
+			const TRigidTransform<T, d>& Transform1 = ConstraintBase.ImplicitTransform[1] * ParticleTransform1;
 
+			switch (ConstraintBase.Manifold.ShapesType)
+			{
+			case EContactShapesType::CapsuleCapsule:
+				UpdateCapsuleCapsuleConstraint(*Implicit0.template GetObject<TCapsule<T>>(), Transform0, *Implicit1.template GetObject<TCapsule<T>>(), Transform1, CullDistance, *ConstraintBase.template As<TRigidBodyPointContactConstraint<T, d>>());
+				break;
+			case EContactShapesType::CapsuleBox:
+				UpdateCapsuleBoxConstraint(*Implicit0.template GetObject<TCapsule<T>>(), Transform0, Implicit1.template GetObject<TBox<T, d>>()->GetAABB(), Transform1, CullDistance, *ConstraintBase.template As<TRigidBodyPointContactConstraint<T, d>>());
+				break;
+			case EContactShapesType::BoxBox:
+				UpdateBoxBoxConstraint(Implicit0.template GetObject<TBox<T, d>>()->GetAABB(), Transform0, Implicit1.template GetObject<TBox<T, d>>()->GetAABB(), Transform1, CullDistance, *ConstraintBase.template As<TRigidBodyPointContactConstraint<T, d>>());
+				break;
+			default:
+				UpdateConstraintAny<UpdateType, T, d>(ConstraintBase, Implicit0, Implicit1, Transform0, Transform1, CullDistance);
+				break;
+			}
 		}
 
 		template<typename T, int d>
-		void ConstructConstraints(TGeometryParticleHandle<T, d>* Particle0, TGeometryParticleHandle<T, d>* Particle1, const FImplicitObject* Implicit0, const FImplicitObject* Implicit1, const TRigidTransform<T, d>& Transform0, const TRigidTransform<T, d>& Transform1, const T Thickness, FCollisionConstraintsArray& NewConstraints)
+		void ConstructConstraints(TGeometryParticleHandle<T, d>* Particle0, TGeometryParticleHandle<T, d>* Particle1, const FImplicitObject* Implicit0, const FImplicitObject* Implicit1, const TRigidTransform<T, d>& Transform0, const TRigidTransform<T, d>& Transform1, const T CullDistance, FCollisionConstraintsArray& NewConstraints)
 		{
 			//  TriMesh Culling
 			if (Implicit0 && Implicit0->GetType() == ImplicitObjectType::TriangleMesh) return;
@@ -1422,7 +1459,7 @@ namespace Chaos
 
 			if (!Implicit0 || !Implicit1)
 			{
-				ConstructLevelsetLevelsetConstraints(Particle0, Particle1, Implicit0, Implicit1, Transform0, Transform1, Thickness, NewConstraints);
+				ConstructLevelsetLevelsetConstraints(Particle0, Particle1, Implicit0, Implicit1, Transform0, Transform1, CullDistance, NewConstraints);
 				return;
 			}
 
@@ -1437,14 +1474,14 @@ namespace Chaos
 			{
 				const TImplicitObjectTransformed<FReal, 3>* TransformedImplicit0 = Implicit0->template GetObject<const TImplicitObjectTransformed<FReal, 3>>();
 				TRigidTransform<T, d> TransformedTransform0 = TransformedImplicit0->GetTransform() * Transform0;
-				ConstructConstraints(Particle0, Particle1, TransformedImplicit0->GetTransformedObject(), Implicit1, TransformedTransform0, Transform1, Thickness, NewConstraints);
+				ConstructConstraints(Particle0, Particle1, TransformedImplicit0->GetTransformedObject(), Implicit1, TransformedTransform0, Transform1, CullDistance, NewConstraints);
 				return;
 			}
 			else if (Implicit1OuterType == TImplicitObjectTransformed<T, d>::StaticType())
 			{
 				const TImplicitObjectTransformed<FReal, 3>* TransformedImplicit1 = Implicit1->template GetObject<const TImplicitObjectTransformed<FReal, 3>>();
 				TRigidTransform<T, d> TransformedTransform1 = TransformedImplicit1->GetTransform() * Transform1;
-				ConstructConstraints(Particle0, Particle1, Implicit0, TransformedImplicit1->GetTransformedObject(), Transform0, TransformedTransform1, Thickness, NewConstraints);
+				ConstructConstraints(Particle0, Particle1, Implicit0, TransformedImplicit1->GetTransformedObject(), Transform0, TransformedTransform1, CullDistance, NewConstraints);
 				return;
 			}
 
@@ -1452,13 +1489,13 @@ namespace Chaos
 			else if (Implicit0OuterType == TImplicitObjectInstanced<FConvex>::StaticType())
 			{
 				const TImplicitObjectInstanced<FConvex>* TransformedImplicit0 = Implicit0->template GetObject<const TImplicitObjectInstanced<FConvex>>();
-				ConstructConstraints(Particle0, Particle1, TransformedImplicit0->GetInstancedObject(), Implicit1, Transform0, Transform1, Thickness, NewConstraints);
+				ConstructConstraints(Particle0, Particle1, TransformedImplicit0->GetInstancedObject(), Implicit1, Transform0, Transform1, CullDistance, NewConstraints);
 				return;
 			}
 			else if (Implicit1OuterType == TImplicitObjectInstanced<FConvex>::StaticType())
 			{
 				const TImplicitObjectInstanced<FConvex>* TransformedImplicit1 = Implicit1->template GetObject<const TImplicitObjectInstanced<FConvex>>();
-				ConstructConstraints(Particle0, Particle1, Implicit0, TransformedImplicit1->GetInstancedObject(), Transform0, Transform1, Thickness, NewConstraints);
+				ConstructConstraints(Particle0, Particle1, Implicit0, TransformedImplicit1->GetInstancedObject(), Transform0, Transform1, CullDistance, NewConstraints);
 				return;
 			}
 
@@ -1466,80 +1503,80 @@ namespace Chaos
 			else if (Implicit0OuterType == TImplicitObjectInstanced<TBox<FReal,3>>::StaticType())
 			{
 				const TImplicitObjectInstanced<TBox<FReal, 3>>* TransformedImplicit0 = Implicit0->template GetObject<const TImplicitObjectInstanced<TBox<FReal, 3>>>();
-				ConstructConstraints(Particle0, Particle1, TransformedImplicit0->GetInstancedObject(), Implicit1, Transform0, Transform1, Thickness, NewConstraints);
+				ConstructConstraints(Particle0, Particle1, TransformedImplicit0->GetInstancedObject(), Implicit1, Transform0, Transform1, CullDistance, NewConstraints);
 				return;
 			}
 			else if (Implicit1OuterType == TImplicitObjectInstanced<TBox<FReal, 3>>::StaticType())
 			{
 				const TImplicitObjectInstanced<TBox<FReal, 3>>* TransformedImplicit1 = Implicit1->template GetObject<const TImplicitObjectInstanced<TBox<FReal, 3>>>();
-				ConstructConstraints(Particle0, Particle1, Implicit0, TransformedImplicit1->GetInstancedObject(), Transform0, Transform1, Thickness, NewConstraints);
+				ConstructConstraints(Particle0, Particle1, Implicit0, TransformedImplicit1->GetInstancedObject(), Transform0, Transform1, CullDistance, NewConstraints);
 				return;
 			}
 
 			else if (Implicit0OuterType == TImplicitObjectInstanced<TCapsule<FReal>>::StaticType())
 			{
 				const TImplicitObjectInstanced<TCapsule<FReal>>* TransformedImplicit0 = Implicit0->template GetObject<const TImplicitObjectInstanced<TCapsule<FReal>>>();
-				ConstructConstraints(Particle0, Particle1, TransformedImplicit0->GetInstancedObject(), Implicit1, Transform0, Transform1, Thickness, NewConstraints);
+				ConstructConstraints(Particle0, Particle1, TransformedImplicit0->GetInstancedObject(), Implicit1, Transform0, Transform1, CullDistance, NewConstraints);
 				return;
 			}
 			else if (Implicit1OuterType == TImplicitObjectInstanced<TCapsule<FReal>>::StaticType())
 			{
 				const TImplicitObjectInstanced<TCapsule<FReal>>* TransformedImplicit1 = Implicit1->template GetObject<const TImplicitObjectInstanced<TCapsule<FReal>>>();
-				ConstructConstraints(Particle0, Particle1, Implicit0, TransformedImplicit1->GetInstancedObject(), Transform0, Transform1, Thickness, NewConstraints);
+				ConstructConstraints(Particle0, Particle1, Implicit0, TransformedImplicit1->GetInstancedObject(), Transform0, Transform1, CullDistance, NewConstraints);
 				return;
 			}
 
 			else if (Implicit0OuterType == TImplicitObjectInstanced<TSphere<FReal, 3>>::StaticType())
 			{
 				const TImplicitObjectInstanced<TSphere<FReal, 3>>* TransformedImplicit0 = Implicit0->template GetObject<const TImplicitObjectInstanced<TSphere<FReal, 3>>>();
-				ConstructConstraints(Particle0, Particle1, TransformedImplicit0->GetInstancedObject(), Implicit1, Transform0, Transform1, Thickness, NewConstraints);
+				ConstructConstraints(Particle0, Particle1, TransformedImplicit0->GetInstancedObject(), Implicit1, Transform0, Transform1, CullDistance, NewConstraints);
 				return;
 			}
 			else if (Implicit1OuterType == TImplicitObjectInstanced<TSphere<FReal, 3>>::StaticType())
 			{
 				const TImplicitObjectInstanced<TSphere<FReal, 3>>* TransformedImplicit1 = Implicit1->template GetObject<const TImplicitObjectInstanced<TSphere<FReal, 3>>>();
-				ConstructConstraints(Particle0, Particle1, Implicit0, TransformedImplicit1->GetInstancedObject(), Transform0, Transform1, Thickness, NewConstraints);
+				ConstructConstraints(Particle0, Particle1, Implicit0, TransformedImplicit1->GetInstancedObject(), Transform0, Transform1, CullDistance, NewConstraints);
 				return;
 			}
 
 			else if (Implicit0OuterType == TImplicitObjectInstanced<FConvex>::StaticType())
 			{
 				const TImplicitObjectInstanced<FConvex>* TransformedImplicit0 = Implicit0->template GetObject<const TImplicitObjectInstanced<FConvex>>();
-				ConstructConstraints(Particle0, Particle1, TransformedImplicit0->GetInstancedObject(), Implicit1, Transform0, Transform1, Thickness, NewConstraints);
+				ConstructConstraints(Particle0, Particle1, TransformedImplicit0->GetInstancedObject(), Implicit1, Transform0, Transform1, CullDistance, NewConstraints);
 				return;
 			}
 			else if (Implicit1OuterType == TImplicitObjectInstanced<FConvex>::StaticType())
 			{
 				const TImplicitObjectInstanced<FConvex>* TransformedImplicit1 = Implicit1->template GetObject<const TImplicitObjectInstanced<FConvex>>();
-				ConstructConstraints(Particle0, Particle1, Implicit0, TransformedImplicit1->GetInstancedObject(), Transform0, Transform1, Thickness, NewConstraints);
+				ConstructConstraints(Particle0, Particle1, Implicit0, TransformedImplicit1->GetInstancedObject(), Transform0, Transform1, CullDistance, NewConstraints);
 				return;
 			}
 			else if (Implicit0OuterType != FImplicitObjectUnion::StaticType() && Implicit1OuterType == FImplicitObjectUnion::StaticType())
 			{
-				const TArray<Pair<const FImplicitObject*, TRigidTransform<T, d>>> LevelsetShapes = FindRelevantShapes(Implicit0, Transform0, *Implicit1, Transform1, Thickness);
+				const TArray<Pair<const FImplicitObject*, TRigidTransform<T, d>>> LevelsetShapes = FindRelevantShapes(Implicit0, Transform0, *Implicit1, Transform1, CullDistance);
 				for (const Pair<const FImplicitObject*, TRigidTransform<T, d>>& LevelsetObjPair : LevelsetShapes)
 				{
 					const FImplicitObject* Implicit1InnerObj = LevelsetObjPair.First;
 					const TRigidTransform<T, d> Implicit1InnerObjTM = LevelsetObjPair.Second * Transform1;
-					ConstructConstraints(Particle0, Particle1, Implicit0, Implicit1InnerObj, Transform0, Implicit1InnerObjTM, Thickness, NewConstraints);
+					ConstructConstraints(Particle0, Particle1, Implicit0, Implicit1InnerObj, Transform0, Implicit1InnerObjTM, CullDistance, NewConstraints);
 				}
 				return;
 			}
 			else if (Implicit0OuterType == FImplicitObjectUnion::StaticType() && Implicit1OuterType != FImplicitObjectUnion::StaticType())
 			{
 				// [Note] forces non-unions into particle[0] position
-				const TArray<Pair<const FImplicitObject*, TRigidTransform<T, d>>> LevelsetShapes = FindRelevantShapes(Implicit1, Transform1, *Implicit0, Transform0, Thickness);
+				const TArray<Pair<const FImplicitObject*, TRigidTransform<T, d>>> LevelsetShapes = FindRelevantShapes(Implicit1, Transform1, *Implicit0, Transform0, CullDistance);
 				for (const Pair<const FImplicitObject*, TRigidTransform<T, d>>& LevelsetObjPair : LevelsetShapes)
 				{
 					const FImplicitObject* Implicit0InnerObj = LevelsetObjPair.First;
 					const TRigidTransform<T, d> Implicit0InnerObjTM = LevelsetObjPair.Second * Transform0;
-					ConstructConstraints(Particle0, Particle1, Implicit0InnerObj, Implicit1, Implicit0InnerObjTM, Transform1, Thickness, NewConstraints);
+					ConstructConstraints(Particle0, Particle1, Implicit0InnerObj, Implicit1, Implicit0InnerObjTM, Transform1, CullDistance, NewConstraints);
 				}
 				return;
 			}
 			else if (Implicit0OuterType == FImplicitObjectUnion::StaticType() && Implicit1OuterType == FImplicitObjectUnion::StaticType())
 			{
-				ConstructUnionUnionConstraints(Particle0, Particle1, Implicit0, Implicit1, Transform0, Transform1, Thickness, NewConstraints);
+				ConstructUnionUnionConstraints(Particle0, Particle1, Implicit0, Implicit1, Transform0, Transform1, CullDistance, NewConstraints);
 				return;
 			}
 
@@ -1548,177 +1585,177 @@ namespace Chaos
 
 			if (Implicit0Type == TBox<T, d>::StaticType() && Implicit1Type == TBox<T, d>::StaticType())
 			{
-				ConstructBoxBoxConstraints(Particle0, Particle1, Implicit0, Implicit1, Transform0, Transform1, Thickness, NewConstraints);
+				ConstructBoxBoxConstraints(Particle0, Particle1, Implicit0, Implicit1, Transform0, Transform1, CullDistance, NewConstraints);
 			}
 			else if (Implicit0Type == TBox<T, d>::StaticType() && Implicit1Type == THeightField<T>::StaticType())
 			{
-				ConstructBoxHeightFieldConstraints(Particle0, Particle1, Implicit0, Implicit1, Transform0, Transform1, Thickness, NewConstraints);
+				ConstructBoxHeightFieldConstraints(Particle0, Particle1, Implicit0, Implicit1, Transform0, Transform1, CullDistance, NewConstraints);
 			}
 			else if (Implicit0Type == THeightField<T>::StaticType() && Implicit1Type == TBox<T, d>::StaticType())
 			{
-				ConstructBoxHeightFieldConstraints(Particle1, Particle0, Implicit1, Implicit0, Transform1, Transform0, Thickness, NewConstraints);
+				ConstructBoxHeightFieldConstraints(Particle1, Particle0, Implicit1, Implicit0, Transform1, Transform0, CullDistance, NewConstraints);
 			}
 			else if (Implicit0Type == TBox<T, d>::StaticType() && Implicit1Type == TPlane<T, d>::StaticType())
 			{
-				ConstructBoxPlaneConstraints(Particle0, Particle1, Implicit0, Implicit1, Transform0, Transform1, Thickness, NewConstraints);
+				ConstructBoxPlaneConstraints(Particle0, Particle1, Implicit0, Implicit1, Transform0, Transform1, CullDistance, NewConstraints);
 			}
 			else if (Implicit0Type == TPlane<T, d>::StaticType() && Implicit1Type == TBox<T, d>::StaticType())
 			{
-				ConstructBoxPlaneConstraints(Particle1, Particle0, Implicit1, Implicit0, Transform1, Transform0, Thickness, NewConstraints);
+				ConstructBoxPlaneConstraints(Particle1, Particle0, Implicit1, Implicit0, Transform1, Transform0, CullDistance, NewConstraints);
 			}
 			else if (Implicit0Type == TSphere<T, d>::StaticType() && Implicit1Type == TSphere<T, d>::StaticType())
 			{
-				ConstructSphereSphereConstraints(Particle0, Particle1, Implicit0, Implicit1, Transform0, Transform1, Thickness, NewConstraints);
+				ConstructSphereSphereConstraints(Particle0, Particle1, Implicit0, Implicit1, Transform0, Transform1, CullDistance, NewConstraints);
 			}
 			else if (Implicit0Type == TSphere<T, d>::StaticType() && Implicit1Type == THeightField<T>::StaticType())
 			{
-				ConstructSphereHeightFieldConstraints(Particle0, Particle1, Implicit0, Implicit1, Transform0, Transform1, Thickness, NewConstraints);
+				ConstructSphereHeightFieldConstraints(Particle0, Particle1, Implicit0, Implicit1, Transform0, Transform1, CullDistance, NewConstraints);
 			}
 			else if (Implicit0Type == THeightField<T>::StaticType() && Implicit1Type == TSphere<T, d>::StaticType())
 			{
-				ConstructSphereHeightFieldConstraints(Particle1, Particle0, Implicit1, Implicit0, Transform1, Transform0, Thickness, NewConstraints);
+				ConstructSphereHeightFieldConstraints(Particle1, Particle0, Implicit1, Implicit0, Transform1, Transform0, CullDistance, NewConstraints);
 			}
 			else if (Implicit0Type == TSphere<T, d>::StaticType() && Implicit1Type == TPlane<T, d>::StaticType())
 			{
-				ConstructSpherePlaneConstraints(Particle0, Particle1, Implicit0, Implicit1, Transform0, Transform1, Thickness, NewConstraints);
+				ConstructSpherePlaneConstraints(Particle0, Particle1, Implicit0, Implicit1, Transform0, Transform1, CullDistance, NewConstraints);
 			}
 			else if (Implicit0Type == TPlane<T, d>::StaticType() && Implicit1Type == TSphere<T, d>::StaticType())
 			{
-				ConstructSpherePlaneConstraints(Particle1, Particle0, Implicit1, Implicit0, Transform1, Transform0, Thickness, NewConstraints);
+				ConstructSpherePlaneConstraints(Particle1, Particle0, Implicit1, Implicit0, Transform1, Transform0, CullDistance, NewConstraints);
 			}
 			else if (Implicit0Type == TSphere<T, d>::StaticType() && Implicit1Type == TBox<T, d>::StaticType())
 			{
-				ConstructSphereBoxConstraints(Particle0, Particle1, Implicit0, Implicit1, Transform0, Transform1, Thickness, NewConstraints);
+				ConstructSphereBoxConstraints(Particle0, Particle1, Implicit0, Implicit1, Transform0, Transform1, CullDistance, NewConstraints);
 			}
 			else if (Implicit0Type == TBox<T, d>::StaticType() && Implicit1Type == TSphere<T, d>::StaticType())
 			{
-				ConstructSphereBoxConstraints(Particle1, Particle0, Implicit1, Implicit0, Transform1, Transform0, Thickness, NewConstraints);
+				ConstructSphereBoxConstraints(Particle1, Particle0, Implicit1, Implicit0, Transform1, Transform0, CullDistance, NewConstraints);
 			}
 			else if (Implicit0Type == TSphere<T, d>::StaticType() && Implicit1Type == TCapsule<T>::StaticType())
 			{
-				ConstructSphereCapsuleConstraints(Particle0, Particle1, Implicit0, Implicit1, Transform0, Transform1, Thickness, NewConstraints);
+				ConstructSphereCapsuleConstraints(Particle0, Particle1, Implicit0, Implicit1, Transform0, Transform1, CullDistance, NewConstraints);
 			}
 			else if (Implicit0Type == TCapsule<T>::StaticType() && Implicit1Type == TSphere<T, d>::StaticType())
 			{
-				ConstructSphereCapsuleConstraints(Particle1, Particle0, Implicit1, Implicit0, Transform1, Transform0, Thickness, NewConstraints);
+				ConstructSphereCapsuleConstraints(Particle1, Particle0, Implicit1, Implicit0, Transform1, Transform0, CullDistance, NewConstraints);
 			}
 			else if (Implicit0Type == TCapsule<T>::StaticType() && Implicit1Type == TCapsule<T>::StaticType())
 			{
-				ConstructCapsuleCapsuleConstraints(Particle0, Particle1, Implicit0, Implicit1, Transform0, Transform1, Thickness, NewConstraints);
+				ConstructCapsuleCapsuleConstraints(Particle0, Particle1, Implicit0, Implicit1, Transform0, Transform1, CullDistance, NewConstraints);
 			}
 			else if (Implicit0Type == TCapsule<T>::StaticType() && Implicit1Type == TBox<T, d>::StaticType())
 			{
-				ConstructCapsuleBoxConstraints(Particle0, Particle1, Implicit0, Implicit1, Transform0, Transform1, Thickness, NewConstraints);
+				ConstructCapsuleBoxConstraints(Particle0, Particle1, Implicit0, Implicit1, Transform0, Transform1, CullDistance, NewConstraints);
 			}
 			else if (Implicit0Type == TBox<T, d>::StaticType() && Implicit1Type == TCapsule<T>::StaticType())
 			{
-				ConstructCapsuleBoxConstraints(Particle1, Particle0, Implicit1, Implicit0, Transform1, Transform0, Thickness, NewConstraints);
+				ConstructCapsuleBoxConstraints(Particle1, Particle0, Implicit1, Implicit0, Transform1, Transform0, CullDistance, NewConstraints);
 			}
 			else if (Implicit0Type == TCapsule<T>::StaticType() && Implicit1Type == THeightField<T>::StaticType())
 			{
-				ConstructCapsuleHeightFieldConstraints(Particle0, Particle1, Implicit0, Implicit1, Transform0, Transform1, Thickness, NewConstraints);
+				ConstructCapsuleHeightFieldConstraints(Particle0, Particle1, Implicit0, Implicit1, Transform0, Transform1, CullDistance, NewConstraints);
 			}
 			else if (Implicit0Type == THeightField<T>::StaticType() && Implicit1Type == TCapsule<T>::StaticType())
 			{
-				ConstructCapsuleHeightFieldConstraints(Particle1, Particle0, Implicit1, Implicit0, Transform1, Transform0, Thickness, NewConstraints);
+				ConstructCapsuleHeightFieldConstraints(Particle1, Particle0, Implicit1, Implicit0, Transform1, Transform0, CullDistance, NewConstraints);
 			}
 			// the generic convex bodies are last
 			else if (Implicit0->IsConvex() && Implicit1Type == THeightField<T>::StaticType())
 			{
-				ConstructConvexHeightFieldConstraints(Particle0, Particle1, Implicit0, Implicit1, Transform0, Transform1, Thickness, NewConstraints);
+				ConstructConvexHeightFieldConstraints(Particle0, Particle1, Implicit0, Implicit1, Transform0, Transform1, CullDistance, NewConstraints);
 			}
 			else if (Implicit0Type == THeightField<T>::StaticType() && Implicit1->IsConvex())
 			{
-				ConstructConvexHeightFieldConstraints(Particle1, Particle0, Implicit1, Implicit0, Transform1, Transform0, Thickness, NewConstraints);
+				ConstructConvexHeightFieldConstraints(Particle1, Particle0, Implicit1, Implicit0, Transform1, Transform0, CullDistance, NewConstraints);
 			}
 			else if (Implicit0->IsConvex() && Implicit1->IsConvex())
 			{
-				ConstructConvexConvexConstraints(Particle0, Particle1, Implicit0, Implicit1, Transform0, Transform1, Thickness, NewConstraints);
+				ConstructConvexConvexConstraints(Particle0, Particle1, Implicit0, Implicit1, Transform0, Transform1, CullDistance, NewConstraints);
 			}
 			else
 			{
-				ConstructLevelsetLevelsetConstraints(Particle0, Particle1, Implicit0, Implicit1, Transform0, Transform1, Thickness, NewConstraints);
+				ConstructLevelsetLevelsetConstraints(Particle0, Particle1, Implicit0, Implicit1, Transform0, Transform1, CullDistance, NewConstraints);
 			}
 
 
 			/*
 			if (Implicit0Type == TBox<T, d>::StaticType() && Implicit1Type == FTriangleMeshImplicitObject::StaticType())
 			{
-				ConstructBoxTriangleMeshConstraints(Particle0, Particle1, Implicit0, Implicit1, Transform0, Transform1, Thickness, NewConstraints);
+				ConstructBoxTriangleMeshConstraints(Particle0, Particle1, Implicit0, Implicit1, Transform0, Transform1, CullDistance, NewConstraints);
 			}
 			else if (Implicit0Type == FTriangleMeshImplicitObject::StaticType() && Implicit1Type == TBox<T, d>::StaticType())
 			{
-				ConstructBoxTriangleMeshConstraints(Particle1, Particle0, Implicit1, Implicit0, Transform1, Transform0, Thickness, NewConstraints);
+				ConstructBoxTriangleMeshConstraints(Particle1, Particle0, Implicit1, Implicit0, Transform1, Transform0, CullDistance, NewConstraints);
 			}
 			*/
 
 		}
 
-		template void UpdateBoxBoxConstraint<float, 3>(const TAABB<float, 3>& Box1, const TRigidTransform<float, 3>& Box1Transform, const TAABB<float, 3>& Box2, const TRigidTransform<float, 3>& Box2Transform, const float Thickness, TRigidBodyPointContactConstraint<float, 3>& Constraint);
-		template void UpdateBoxBoxManifold<float, 3>(TCollisionConstraintBase<float, 3>&  Constraint, const TRigidTransform<float, 3>& ATM, const TRigidTransform<float, 3>& BTM, const float Thickness);
-		template void ConstructBoxBoxConstraints<float, 3>(TGeometryParticleHandle<float, 3>* Particle0, TGeometryParticleHandle<float, 3>* Particle1, const FImplicitObject* Implicit0, const FImplicitObject* Implicit1, const TRigidTransform<float, 3>& Transform0, const TRigidTransform<float, 3>& Transform1, const float Thickness, FCollisionConstraintsArray& NewConstraints);
+		template void UpdateBoxBoxConstraint<float, 3>(const TAABB<float, 3>& Box1, const TRigidTransform<float, 3>& Box1Transform, const TAABB<float, 3>& Box2, const TRigidTransform<float, 3>& Box2Transform, const float CullDistance, TRigidBodyPointContactConstraint<float, 3>& Constraint);
+		template void UpdateBoxBoxManifold<float, 3>(TCollisionConstraintBase<float, 3>&  Constraint, const TRigidTransform<float, 3>& ATM, const TRigidTransform<float, 3>& BTM, const float CullDistance);
+		template void ConstructBoxBoxConstraints<float, 3>(TGeometryParticleHandle<float, 3>* Particle0, TGeometryParticleHandle<float, 3>* Particle1, const FImplicitObject* Implicit0, const FImplicitObject* Implicit1, const TRigidTransform<float, 3>& Transform0, const TRigidTransform<float, 3>& Transform1, const float CullDistance, FCollisionConstraintsArray& NewConstraints);
 
-		template void UpdateBoxHeightFieldConstraint<float, 3>(const TAABB<float, 3>& A, const TRigidTransform<float, 3>& ATransform, const THeightField<float>& B, const TRigidTransform<float, 3>& BTransform, const float Thickness, TRigidBodyPointContactConstraint<float, 3>& Constraint);
-		template void UpdateBoxHeightFieldManifold<float, 3>(TCollisionConstraintBase<float, 3>&  Constraint, const TRigidTransform<float, 3>& ATM, const TRigidTransform<float, 3>& BTM, const float Thickness);
-		template void ConstructBoxHeightFieldConstraints<float, 3>(TGeometryParticleHandle<float, 3>* Particle0, TGeometryParticleHandle<float, 3>* Particle1, const FImplicitObject* Implicit0, const FImplicitObject* Implicit1, const TRigidTransform<float, 3>& Transform0, const TRigidTransform<float, 3>& Transform1, const float Thickness, FCollisionConstraintsArray& NewConstraints);
+		template void UpdateBoxHeightFieldConstraint<float, 3>(const TAABB<float, 3>& A, const TRigidTransform<float, 3>& ATransform, const THeightField<float>& B, const TRigidTransform<float, 3>& BTransform, const float CullDistance, TRigidBodyPointContactConstraint<float, 3>& Constraint);
+		template void UpdateBoxHeightFieldManifold<float, 3>(TCollisionConstraintBase<float, 3>&  Constraint, const TRigidTransform<float, 3>& ATM, const TRigidTransform<float, 3>& BTM, const float CullDistance);
+		template void ConstructBoxHeightFieldConstraints<float, 3>(TGeometryParticleHandle<float, 3>* Particle0, TGeometryParticleHandle<float, 3>* Particle1, const FImplicitObject* Implicit0, const FImplicitObject* Implicit1, const TRigidTransform<float, 3>& Transform0, const TRigidTransform<float, 3>& Transform1, const float CullDistance, FCollisionConstraintsArray& NewConstraints);
 
-		template void UpdateBoxPlaneConstraint<float, 3>(const TAABB<float, 3>& Box, const TRigidTransform<float, 3>& BoxTransform, const TPlane<float, 3>& Plane, const TRigidTransform<float, 3>& PlaneTransform, const float Thickness, TRigidBodyPointContactConstraint<float, 3>& Constraint);
-		template void UpdateBoxPlaneManifold<float, 3>(TCollisionConstraintBase<float, 3>&  Constraint, const TRigidTransform<float, 3>& ATM, const TRigidTransform<float, 3>& BTM, const float Thickness);
-		template void ConstructBoxPlaneConstraints<float, 3>(TGeometryParticleHandle<float, 3>* Particle0, TGeometryParticleHandle<float, 3>* Particle1, const FImplicitObject* Implicit0, const FImplicitObject* Implicit1, const TRigidTransform<float, 3>& Transform0, const TRigidTransform<float, 3>& Transform1, const float Thickness, FCollisionConstraintsArray& NewConstraints);
+		template void UpdateBoxPlaneConstraint<float, 3>(const TAABB<float, 3>& Box, const TRigidTransform<float, 3>& BoxTransform, const TPlane<float, 3>& Plane, const TRigidTransform<float, 3>& PlaneTransform, const float CullDistance, TRigidBodyPointContactConstraint<float, 3>& Constraint);
+		template void UpdateBoxPlaneManifold<float, 3>(TCollisionConstraintBase<float, 3>&  Constraint, const TRigidTransform<float, 3>& ATM, const TRigidTransform<float, 3>& BTM, const float CullDistance);
+		template void ConstructBoxPlaneConstraints<float, 3>(TGeometryParticleHandle<float, 3>* Particle0, TGeometryParticleHandle<float, 3>* Particle1, const FImplicitObject* Implicit0, const FImplicitObject* Implicit1, const TRigidTransform<float, 3>& Transform0, const TRigidTransform<float, 3>& Transform1, const float CullDistance, FCollisionConstraintsArray& NewConstraints);
 
-		template void UpdateSphereSphereConstraint<float, 3>(const TSphere<float, 3>& Sphere1, const TRigidTransform<float, 3>& Sphere1Transform, const TSphere<float, 3>& Sphere2, const TRigidTransform<float, 3>& Sphere2Transform, const float Thickness, TRigidBodyPointContactConstraint<float, 3>& Constraint);
-		template void UpdateSphereSphereManifold<float, 3>(TCollisionConstraintBase<float, 3>&  Constraint, const TRigidTransform<float, 3>& ATM, const TRigidTransform<float, 3>& BTM, const float Thickness);
-		template void ConstructSphereSphereConstraints<float, 3>(TGeometryParticleHandle<float, 3>* Particle0, TGeometryParticleHandle<float, 3>* Particle1, const FImplicitObject* Implicit0, const FImplicitObject* Implicit1, const TRigidTransform<float, 3>& Transform0, const TRigidTransform<float, 3>& Transform1, const float Thickness, FCollisionConstraintsArray& NewConstraints);
+		template void UpdateSphereSphereConstraint<float, 3>(const TSphere<float, 3>& Sphere1, const TRigidTransform<float, 3>& Sphere1Transform, const TSphere<float, 3>& Sphere2, const TRigidTransform<float, 3>& Sphere2Transform, const float CullDistance, TRigidBodyPointContactConstraint<float, 3>& Constraint);
+		template void UpdateSphereSphereManifold<float, 3>(TCollisionConstraintBase<float, 3>&  Constraint, const TRigidTransform<float, 3>& ATM, const TRigidTransform<float, 3>& BTM, const float CullDistance);
+		template void ConstructSphereSphereConstraints<float, 3>(TGeometryParticleHandle<float, 3>* Particle0, TGeometryParticleHandle<float, 3>* Particle1, const FImplicitObject* Implicit0, const FImplicitObject* Implicit1, const TRigidTransform<float, 3>& Transform0, const TRigidTransform<float, 3>& Transform1, const float CullDistance, FCollisionConstraintsArray& NewConstraints);
 
-		template void UpdateSphereHeightFieldConstraint<float, 3>(const TSphere<float, 3>& A, const TRigidTransform<float, 3>& ATransform, const THeightField<float>& B, const TRigidTransform<float, 3>& BTransform, const float Thickness, TRigidBodyPointContactConstraint<float, 3>& Constraint);
-		template void UpdateSphereHeightFieldManifold<float, 3>(TCollisionConstraintBase<float, 3>&  Constraint, const TRigidTransform<float, 3>& ATM, const TRigidTransform<float, 3>& BTM, const float Thickness);
-		template void ConstructSphereHeightFieldConstraints<float, 3>(TGeometryParticleHandle<float, 3>* Particle0, TGeometryParticleHandle<float, 3>* Particle1, const FImplicitObject* Implicit0, const FImplicitObject* Implicit1, const TRigidTransform<float, 3>& Transform0, const TRigidTransform<float, 3>& Transform1, const float Thickness, FCollisionConstraintsArray& NewConstraints);
+		template void UpdateSphereHeightFieldConstraint<float, 3>(const TSphere<float, 3>& A, const TRigidTransform<float, 3>& ATransform, const THeightField<float>& B, const TRigidTransform<float, 3>& BTransform, const float CullDistance, TRigidBodyPointContactConstraint<float, 3>& Constraint);
+		template void UpdateSphereHeightFieldManifold<float, 3>(TCollisionConstraintBase<float, 3>&  Constraint, const TRigidTransform<float, 3>& ATM, const TRigidTransform<float, 3>& BTM, const float CullDistance);
+		template void ConstructSphereHeightFieldConstraints<float, 3>(TGeometryParticleHandle<float, 3>* Particle0, TGeometryParticleHandle<float, 3>* Particle1, const FImplicitObject* Implicit0, const FImplicitObject* Implicit1, const TRigidTransform<float, 3>& Transform0, const TRigidTransform<float, 3>& Transform1, const float CullDistance, FCollisionConstraintsArray& NewConstraints);
 
-		template void UpdateSpherePlaneConstraint<float, 3>(const TSphere<float, 3>& Sphere, const TRigidTransform<float, 3>& SphereTransform, const TPlane<float, 3>& Plane, const TRigidTransform<float, 3>& PlaneTransform, const float Thickness, TRigidBodyPointContactConstraint<float, 3>& Constraint);
-		template void UpdateSpherePlaneManifold<float, 3>(TCollisionConstraintBase<float, 3>&  Constraint, const TRigidTransform<float, 3>& ATM, const TRigidTransform<float, 3>& BTM, const float Thickness);
-		template void ConstructSpherePlaneConstraints<float, 3>(TGeometryParticleHandle<float, 3>* Particle0, TGeometryParticleHandle<float, 3>* Particle1, const FImplicitObject* Implicit0, const FImplicitObject* Implicit1, const TRigidTransform<float, 3>& Transform0, const TRigidTransform<float, 3>& Transform1, const float Thickness, FCollisionConstraintsArray& NewConstraints);
+		template void UpdateSpherePlaneConstraint<float, 3>(const TSphere<float, 3>& Sphere, const TRigidTransform<float, 3>& SphereTransform, const TPlane<float, 3>& Plane, const TRigidTransform<float, 3>& PlaneTransform, const float CullDistance, TRigidBodyPointContactConstraint<float, 3>& Constraint);
+		template void UpdateSpherePlaneManifold<float, 3>(TCollisionConstraintBase<float, 3>&  Constraint, const TRigidTransform<float, 3>& ATM, const TRigidTransform<float, 3>& BTM, const float CullDistance);
+		template void ConstructSpherePlaneConstraints<float, 3>(TGeometryParticleHandle<float, 3>* Particle0, TGeometryParticleHandle<float, 3>* Particle1, const FImplicitObject* Implicit0, const FImplicitObject* Implicit1, const TRigidTransform<float, 3>& Transform0, const TRigidTransform<float, 3>& Transform1, const float CullDistance, FCollisionConstraintsArray& NewConstraints);
 
-		template void UpdateSphereBoxConstraint<float, 3>(const TSphere<float, 3>& Sphere, const TRigidTransform<float, 3>& SphereTransform, const TAABB<float, 3>& Box, const TRigidTransform<float, 3>& BoxTransform, const float Thickness, TRigidBodyPointContactConstraint<float, 3>& Constraint);
-		template void UpdateSphereBoxManifold<float, 3>(TCollisionConstraintBase<float, 3>&  Constraint, const TRigidTransform<float, 3>& ATM, const TRigidTransform<float, 3>& BTM, const float Thickness);
-		template void ConstructSphereBoxConstraints<float, 3>(TGeometryParticleHandle<float, 3>* Particle0, TGeometryParticleHandle<float, 3>* Particle1, const FImplicitObject* Implicit0, const FImplicitObject* Implicit1, const TRigidTransform<float, 3>& Transform0, const TRigidTransform<float, 3>& Transform1, const float Thickness, FCollisionConstraintsArray& NewConstraints);
+		template void UpdateSphereBoxConstraint<float, 3>(const TSphere<float, 3>& Sphere, const TRigidTransform<float, 3>& SphereTransform, const TAABB<float, 3>& Box, const TRigidTransform<float, 3>& BoxTransform, const float CullDistance, TRigidBodyPointContactConstraint<float, 3>& Constraint);
+		template void UpdateSphereBoxManifold<float, 3>(TCollisionConstraintBase<float, 3>&  Constraint, const TRigidTransform<float, 3>& ATM, const TRigidTransform<float, 3>& BTM, const float CullDistance);
+		template void ConstructSphereBoxConstraints<float, 3>(TGeometryParticleHandle<float, 3>* Particle0, TGeometryParticleHandle<float, 3>* Particle1, const FImplicitObject* Implicit0, const FImplicitObject* Implicit1, const TRigidTransform<float, 3>& Transform0, const TRigidTransform<float, 3>& Transform1, const float CullDistance, FCollisionConstraintsArray& NewConstraints);
 
-		template void UpdateSphereCapsuleConstraint<float, 3>(const TSphere<float, 3>& Sphere, const TRigidTransform<float, 3>& SphereTransform, const TCapsule<float>& Box, const TRigidTransform<float, 3>& BoxTransform, const float Thickness, TRigidBodyPointContactConstraint<float, 3>& Constraint);
-		template void UpdateSphereCapsuleManifold<float, 3>(TCollisionConstraintBase<float, 3>&  Constraint, const TRigidTransform<float, 3>& ATM, const TRigidTransform<float, 3>& BTM, const float Thickness);
-		template void ConstructSphereCapsuleConstraints<float, 3>(TGeometryParticleHandle<float, 3>* Particle0, TGeometryParticleHandle<float, 3>* Particle1, const FImplicitObject* Implicit0, const FImplicitObject* Implicit1, const TRigidTransform<float, 3>& Transform0, const TRigidTransform<float, 3>& Transform1, const float Thickness, FCollisionConstraintsArray& NewConstraints);
+		template void UpdateSphereCapsuleConstraint<float, 3>(const TSphere<float, 3>& Sphere, const TRigidTransform<float, 3>& SphereTransform, const TCapsule<float>& Box, const TRigidTransform<float, 3>& BoxTransform, const float CullDistance, TRigidBodyPointContactConstraint<float, 3>& Constraint);
+		template void UpdateSphereCapsuleManifold<float, 3>(TCollisionConstraintBase<float, 3>&  Constraint, const TRigidTransform<float, 3>& ATM, const TRigidTransform<float, 3>& BTM, const float CullDistance);
+		template void ConstructSphereCapsuleConstraints<float, 3>(TGeometryParticleHandle<float, 3>* Particle0, TGeometryParticleHandle<float, 3>* Particle1, const FImplicitObject* Implicit0, const FImplicitObject* Implicit1, const TRigidTransform<float, 3>& Transform0, const TRigidTransform<float, 3>& Transform1, const float CullDistance, FCollisionConstraintsArray& NewConstraints);
 
-		template void UpdateCapsuleCapsuleConstraint<float, 3>(const TCapsule<float>& A, const TRigidTransform<float, 3>& ATransform, const TCapsule<float>& B, const TRigidTransform<float, 3>& BTransform, const float Thickness, TRigidBodyPointContactConstraint<float, 3>& Constraint);
-		template void UpdateCapsuleCapsuleManifold<float, 3>(TCollisionConstraintBase<float, 3>&  Constraint, const TRigidTransform<float, 3>& ATM, const TRigidTransform<float, 3>& BTM, const float Thickness);
-		template void ConstructCapsuleCapsuleConstraints<float, 3>(TGeometryParticleHandle<float, 3>* Particle0, TGeometryParticleHandle<float, 3>* Particle1, const FImplicitObject* Implicit0, const FImplicitObject* Implicit1, const TRigidTransform<float, 3>& Transform0, const TRigidTransform<float, 3>& Transform1, const float Thickness, FCollisionConstraintsArray& NewConstraints);
+		template void UpdateCapsuleCapsuleConstraint<float, 3>(const TCapsule<float>& A, const TRigidTransform<float, 3>& ATransform, const TCapsule<float>& B, const TRigidTransform<float, 3>& BTransform, const float CullDistance, TRigidBodyPointContactConstraint<float, 3>& Constraint);
+		template void UpdateCapsuleCapsuleManifold<float, 3>(TCollisionConstraintBase<float, 3>&  Constraint, const TRigidTransform<float, 3>& ATM, const TRigidTransform<float, 3>& BTM, const float CullDistance);
+		template void ConstructCapsuleCapsuleConstraints<float, 3>(TGeometryParticleHandle<float, 3>* Particle0, TGeometryParticleHandle<float, 3>* Particle1, const FImplicitObject* Implicit0, const FImplicitObject* Implicit1, const TRigidTransform<float, 3>& Transform0, const TRigidTransform<float, 3>& Transform1, const float CullDistance, FCollisionConstraintsArray& NewConstraints);
 
-		template void UpdateCapsuleBoxConstraint<float, 3>(const TCapsule<float>& A, const TRigidTransform<float, 3>& ATransform, const TAABB<float, 3>& B, const TRigidTransform<float, 3>& BTransform, const float Thickness, TRigidBodyPointContactConstraint<float, 3>& Constraint);
-		template void UpdateCapsuleBoxManifold<float, 3>(TCollisionConstraintBase<float, 3>&  Constraint, const TRigidTransform<float, 3>& ATM, const TRigidTransform<float, 3>& BTM, const float Thickness);
-		template void ConstructCapsuleBoxConstraints<float, 3>(TGeometryParticleHandle<float, 3>* Particle0, TGeometryParticleHandle<float, 3>* Particle1, const FImplicitObject* Implicit0, const FImplicitObject* Implicit1, const TRigidTransform<float, 3>& Transform0, const TRigidTransform<float, 3>& Transform1, const float Thickness, FCollisionConstraintsArray& NewConstraints);
+		template void UpdateCapsuleBoxConstraint<float, 3>(const TCapsule<float>& A, const TRigidTransform<float, 3>& ATransform, const TAABB<float, 3>& B, const TRigidTransform<float, 3>& BTransform, const float CullDistance, TRigidBodyPointContactConstraint<float, 3>& Constraint);
+		template void UpdateCapsuleBoxManifold<float, 3>(TCollisionConstraintBase<float, 3>&  Constraint, const TRigidTransform<float, 3>& ATM, const TRigidTransform<float, 3>& BTM, const float CullDistance);
+		template void ConstructCapsuleBoxConstraints<float, 3>(TGeometryParticleHandle<float, 3>* Particle0, TGeometryParticleHandle<float, 3>* Particle1, const FImplicitObject* Implicit0, const FImplicitObject* Implicit1, const TRigidTransform<float, 3>& Transform0, const TRigidTransform<float, 3>& Transform1, const float CullDistance, FCollisionConstraintsArray& NewConstraints);
 
-		template void UpdateCapsuleHeightFieldConstraint<float, 3>(const TCapsule<float>& A, const TRigidTransform<float, 3>& ATransform, const THeightField<float>& B, const TRigidTransform<float, 3>& BTransform, const float Thickness, TRigidBodyPointContactConstraint<float, 3>& Constraint);
-		template void UpdateCapsuleHeightFieldManifold<float, 3>(TCollisionConstraintBase<float, 3>&  Constraint, const TRigidTransform<float, 3>& ATM, const TRigidTransform<float, 3>& BTM, const float Thickness);
-		template void ConstructCapsuleHeightFieldConstraints<float, 3>(TGeometryParticleHandle<float, 3>* Particle0, TGeometryParticleHandle<float, 3>* Particle1, const FImplicitObject* Implicit0, const FImplicitObject* Implicit1, const TRigidTransform<float, 3>& Transform0, const TRigidTransform<float, 3>& Transform1, const float Thickness, FCollisionConstraintsArray& NewConstraints);
+		template void UpdateCapsuleHeightFieldConstraint<float, 3>(const TCapsule<float>& A, const TRigidTransform<float, 3>& ATransform, const THeightField<float>& B, const TRigidTransform<float, 3>& BTransform, const float CullDistance, TRigidBodyPointContactConstraint<float, 3>& Constraint);
+		template void UpdateCapsuleHeightFieldManifold<float, 3>(TCollisionConstraintBase<float, 3>&  Constraint, const TRigidTransform<float, 3>& ATM, const TRigidTransform<float, 3>& BTM, const float CullDistance);
+		template void ConstructCapsuleHeightFieldConstraints<float, 3>(TGeometryParticleHandle<float, 3>* Particle0, TGeometryParticleHandle<float, 3>* Particle1, const FImplicitObject* Implicit0, const FImplicitObject* Implicit1, const TRigidTransform<float, 3>& Transform0, const TRigidTransform<float, 3>& Transform1, const float CullDistance, FCollisionConstraintsArray& NewConstraints);
 
-		template void UpdateConvexConvexConstraint<float, 3>(const FImplicitObject& A, const TRigidTransform<float, 3>& ATM, const FImplicitObject& B, const TRigidTransform<float, 3>& BTM, const float Thickness, TCollisionConstraintBase<float, 3>& Constraint);
-		template void UpdateConvexConvexManifold<float, 3>(TCollisionConstraintBase<float, 3>&  Constraint, const TRigidTransform<float, 3>& ATM, const TRigidTransform<float, 3>& BTM, const float Thickness);
-		template void ConstructConvexConvexConstraints<float, 3>(TGeometryParticleHandle<float, 3>* Particle0, TGeometryParticleHandle<float, 3>* Particle1, const FImplicitObject* Implicit0, const FImplicitObject* Implicit1, const TRigidTransform<float, 3>& Transform0, const TRigidTransform<float, 3>& Transform1, const float Thickness, FCollisionConstraintsArray& NewConstraints);
+		template void UpdateConvexConvexConstraint<float, 3>(const FImplicitObject& A, const TRigidTransform<float, 3>& ATM, const FImplicitObject& B, const TRigidTransform<float, 3>& BTM, const float CullDistance, TCollisionConstraintBase<float, 3>& Constraint);
+		template void UpdateConvexConvexManifold<float, 3>(TCollisionConstraintBase<float, 3>&  Constraint, const TRigidTransform<float, 3>& ATM, const TRigidTransform<float, 3>& BTM, const float CullDistance);
+		template void ConstructConvexConvexConstraints<float, 3>(TGeometryParticleHandle<float, 3>* Particle0, TGeometryParticleHandle<float, 3>* Particle1, const FImplicitObject* Implicit0, const FImplicitObject* Implicit1, const TRigidTransform<float, 3>& Transform0, const TRigidTransform<float, 3>& Transform1, const float CullDistance, FCollisionConstraintsArray& NewConstraints);
 
-		template void UpdateConvexHeightFieldConstraint<float, 3>(const FImplicitObject& A, const TRigidTransform<float, 3>& ATransform, const THeightField<float>& B, const TRigidTransform<float, 3>& BTransform, const float Thickness, TRigidBodyPointContactConstraint<float, 3>& Constraint);
-		template void UpdateConvexHeightFieldManifold<float, 3>(TCollisionConstraintBase<float, 3>&  Constraint, const TRigidTransform<float, 3>& ATM, const TRigidTransform<float, 3>& BTM, const float Thickness);
-		template void ConstructConvexHeightFieldConstraints<float, 3>(TGeometryParticleHandle<float, 3>* Particle0, TGeometryParticleHandle<float, 3>* Particle1, const FImplicitObject* Implicit0, const FImplicitObject* Implicit1, const TRigidTransform<float, 3>& Transform0, const TRigidTransform<float, 3>& Transform1, const float Thickness, FCollisionConstraintsArray& NewConstraints);
+		template void UpdateConvexHeightFieldConstraint<float, 3>(const FImplicitObject& A, const TRigidTransform<float, 3>& ATransform, const THeightField<float>& B, const TRigidTransform<float, 3>& BTransform, const float CullDistance, TRigidBodyPointContactConstraint<float, 3>& Constraint);
+		template void UpdateConvexHeightFieldManifold<float, 3>(TCollisionConstraintBase<float, 3>&  Constraint, const TRigidTransform<float, 3>& ATM, const TRigidTransform<float, 3>& BTM, const float CullDistance);
+		template void ConstructConvexHeightFieldConstraints<float, 3>(TGeometryParticleHandle<float, 3>* Particle0, TGeometryParticleHandle<float, 3>* Particle1, const FImplicitObject* Implicit0, const FImplicitObject* Implicit1, const TRigidTransform<float, 3>& Transform0, const TRigidTransform<float, 3>& Transform1, const float CullDistance, FCollisionConstraintsArray& NewConstraints);
 
-		template void UpdateLevelsetLevelsetConstraint<ECollisionUpdateType::Any, float, 3>(const float Thickness, TRigidBodyPointContactConstraint<float, 3>& Constraint);
-		template void UpdateLevelsetLevelsetConstraint<ECollisionUpdateType::Deepest, float, 3>(const float Thickness, TRigidBodyPointContactConstraint<float, 3>& Constraint);
-		template void UpdateLevelsetLevelsetManifold<float, 3>(TCollisionConstraintBase<float, 3>&  Constraint, const TRigidTransform<float, 3>& ATM, const TRigidTransform<float, 3>& BTM, const float Thickness);
-		template void ConstructLevelsetLevelsetConstraints<float, 3>(TGeometryParticleHandle<float, 3>* Particle0, TGeometryParticleHandle<float, 3>* Particle1, const FImplicitObject* Implicit0, const FImplicitObject* Implicit1, const TRigidTransform<float, 3>& Transform0, const TRigidTransform<float, 3>& Transform1, const float Thickness, FCollisionConstraintsArray& NewConstraints);
+		template void UpdateLevelsetLevelsetConstraint<ECollisionUpdateType::Any, float, 3>(const float CullDistance, TRigidBodyPointContactConstraint<float, 3>& Constraint);
+		template void UpdateLevelsetLevelsetConstraint<ECollisionUpdateType::Deepest, float, 3>(const float CullDistance, TRigidBodyPointContactConstraint<float, 3>& Constraint);
+		template void UpdateLevelsetLevelsetManifold<float, 3>(TCollisionConstraintBase<float, 3>&  Constraint, const TRigidTransform<float, 3>& ATM, const TRigidTransform<float, 3>& BTM, const float CullDistance);
+		template void ConstructLevelsetLevelsetConstraints<float, 3>(TGeometryParticleHandle<float, 3>* Particle0, TGeometryParticleHandle<float, 3>* Particle1, const FImplicitObject* Implicit0, const FImplicitObject* Implicit1, const TRigidTransform<float, 3>& Transform0, const TRigidTransform<float, 3>& Transform1, const float CullDistance, FCollisionConstraintsArray& NewConstraints);
 
-		template void ConstructUnionUnionConstraints<float, 3>(TGeometryParticleHandle<float, 3>* Particle0, TGeometryParticleHandle<float, 3>* Particle1, const FImplicitObject* Implicit0, const FImplicitObject* Implicit1, const TRigidTransform<float, 3>& Transform0, const TRigidTransform<float, 3>& Transform1, const float Thickness, FCollisionConstraintsArray& NewConstraints);
+		template void ConstructUnionUnionConstraints<float, 3>(TGeometryParticleHandle<float, 3>* Particle0, TGeometryParticleHandle<float, 3>* Particle1, const FImplicitObject* Implicit0, const FImplicitObject* Implicit1, const TRigidTransform<float, 3>& Transform0, const TRigidTransform<float, 3>& Transform1, const float CullDistance, FCollisionConstraintsArray& NewConstraints);
 
-		template void UpdateSingleShotManifold<float, 3>(TRigidBodyMultiPointContactConstraint<float, 3>&  Constraint, const TRigidTransform<float, 3>& Transform0, const TRigidTransform<float, 3>& Transform1, const float Thickness);
-		template void UpdateIterativeManifold<float, 3>(TRigidBodyMultiPointContactConstraint<float, 3>&  Constraint, const TRigidTransform<float, 3>& Transform0, const TRigidTransform<float, 3>& Transform1, const float Thickness);
-		template void UpdateManifold<float, 3>(TCollisionConstraintBase<float, 3>& ConstraintBase, const TRigidTransform<float, 3>& ATM, const TRigidTransform<float, 3>& BTM, const float Thickness);
-		template void UpdateConstraint<ECollisionUpdateType::Any, float, 3>(TCollisionConstraintBase<float, 3>& ConstraintBase, const TRigidTransform<float, 3>& Transform0, const TRigidTransform<float, 3>& Transform1, const float Thickness);
-		template void UpdateConstraint<ECollisionUpdateType::Deepest, float, 3>(TCollisionConstraintBase<float, 3>& ConstraintBase, const TRigidTransform<float, 3>& Transform0, const TRigidTransform<float, 3>& Transform1, const float Thickness);
-		template void ConstructConstraints<float, 3>(TGeometryParticleHandle<float, 3>* Particle0, TGeometryParticleHandle<float, 3>* Particle1, const FImplicitObject* Implicit0, const FImplicitObject* Implicit1, const TRigidTransform<float, 3>& Transform0, const TRigidTransform<float, 3>& Transform1, const float Thickness, FCollisionConstraintsArray& NewConstraints);
+		template void UpdateSingleShotManifold<float, 3>(TRigidBodyMultiPointContactConstraint<float, 3>&  Constraint, const TRigidTransform<float, 3>& Transform0, const TRigidTransform<float, 3>& Transform1, const float CullDistance);
+		template void UpdateIterativeManifold<float, 3>(TRigidBodyMultiPointContactConstraint<float, 3>&  Constraint, const TRigidTransform<float, 3>& Transform0, const TRigidTransform<float, 3>& Transform1, const float CullDistance);
+		template void UpdateManifold<float, 3>(TCollisionConstraintBase<float, 3>& ConstraintBase, const TRigidTransform<float, 3>& ATM, const TRigidTransform<float, 3>& BTM, const float CullDistance);
+		template void UpdateConstraint<ECollisionUpdateType::Any, float, 3>(TCollisionConstraintBase<float, 3>& ConstraintBase, const TRigidTransform<float, 3>& Transform0, const TRigidTransform<float, 3>& Transform1, const float CullDistance);
+		template void UpdateConstraint<ECollisionUpdateType::Deepest, float, 3>(TCollisionConstraintBase<float, 3>& ConstraintBase, const TRigidTransform<float, 3>& Transform0, const TRigidTransform<float, 3>& Transform1, const float CullDistance);
+		template void ConstructConstraints<float, 3>(TGeometryParticleHandle<float, 3>* Particle0, TGeometryParticleHandle<float, 3>* Particle1, const FImplicitObject* Implicit0, const FImplicitObject* Implicit1, const TRigidTransform<float, 3>& Transform0, const TRigidTransform<float, 3>& Transform1, const float CullDistance, FCollisionConstraintsArray& NewConstraints);
 
 	} // Collisions
 
diff --git a/Engine/Source/Runtime/Experimental/Chaos/Private/Chaos/Evolution/PBDMinEvolution.cpp b/Engine/Source/Runtime/Experimental/Chaos/Private/Chaos/Evolution/PBDMinEvolution.cpp
--- a/Engine/Source/Runtime/Experimental/Chaos/Private/Chaos/Evolution/PBDMinEvolution.cpp
+++ b/Engine/Source/Runtime/Experimental/Chaos/Private/Chaos/Evolution/PBDMinEvolution.cpp
@@ -30,11 +30,12 @@ namespace Chaos
 	DECLARE_CYCLE_STAT(TEXT("FPBDMinEvolution::DetectCollisions"), STAT_MinEvolution_DetectCollisions, STATGROUP_Chaos);
 	DECLARE_CYCLE_STAT(TEXT("FPBDMinEvolution::UpdatePositions"), STAT_MinEvolution_UpdatePositions, STATGROUP_Chaos);
 
-	FPBDMinEvolution::FPBDMinEvolution(FRigidParticleSOAs& InParticles, FCollisionDetector& InCollisionDetector)
+	FPBDMinEvolution::FPBDMinEvolution(FRigidParticleSOAs& InParticles, FCollisionDetector& InCollisionDetector, const FReal InBoundsExtension)
 		: Particles(InParticles)
 		, CollisionDetector(InCollisionDetector)
 		, NumApplyIterations(0)
 		, NumApplyPushOutIterations(0)
+		, BoundsExtension(InBoundsExtension)
 		, Gravity(FVec3(0))
 	{
 	}
@@ -151,29 +152,18 @@ namespace Chaos
 
 				if (Particle.HasBounds())
 				{
-					const FReal BoundsThickness = 1.0f;
-					const FReal BoundsThicknessVelocityInflation = 2.0f;
 					const TAABB<FReal, 3>& LocalBounds = Particle.LocalBounds();
-					TAABB<FReal, 3> WorldSpaceBounds = LocalBounds.TransformedAABB(FRigidTransform3(Particle.P(), Particle.Q()));
-					WorldSpaceBounds.ThickenSymmetrically(FVec3(BoundsThickness) + BoundsThicknessVelocityInflation * Particle.V().GetAbs() * Dt);
+					
+					TAABB<FReal, 3> NextWorldSpaceBounds = LocalBounds.TransformedAABB(FRigidTransform3(Particle.P(), Particle.Q()));
+					NextWorldSpaceBounds.ThickenSymmetrically(NextWorldSpaceBounds.Extents() * BoundsExtension);
+
+					TAABB<FReal, 3> WorldSpaceBounds = LocalBounds.TransformedAABB(FRigidTransform3(Particle.X(), Particle.R()));
+					WorldSpaceBounds.GrowToInclude(NextWorldSpaceBounds);
+
 					Particle.SetWorldSpaceInflatedBounds(WorldSpaceBounds);
 				}
 			}
 		}
-
-		for (TTransientGeometryParticleHandle<FReal, 3>& Particle : Particles.GetActiveKinematicParticlesView())
-		{
-			if (Particle.HasBounds())
-			{
-				const FReal BoundsThickness = 1.0f;
-				const FReal BoundsThicknessVelocityInflation = 2.0f;
-				const TAABB<FReal, 3>& LocalBounds = Particle.LocalBounds();
-				TAABB<FReal, 3> WorldSpaceBounds = LocalBounds.TransformedAABB(FRigidTransform3(Particle.X(), Particle.R()));
-				const FVec3 VAbs = (Particle.CastToKinematicParticle()) ? Particle.CastToKinematicParticle()->V().GetAbs() : FVec3(0);
-				WorldSpaceBounds.ThickenSymmetrically(FVec3(BoundsThickness) + BoundsThicknessVelocityInflation * VAbs.GetAbs() * Dt);
-				Particle.SetWorldSpaceInflatedBounds(WorldSpaceBounds);
-			}
-		}
 	}
 
 	// @todo(ccaulfield): dedupe (PBDRigidsEvolutionGBF)
@@ -189,6 +179,9 @@ namespace Chaos
 		const FReal MinDt = 1e-6f;
 		for (auto& Particle : Particles.GetActiveKinematicParticlesView())
 		{
+			const FVec3 PrevX = Particle.X();
+			const FRotation3 PrevR = Particle.R();
+
 			TKinematicTarget<FReal, 3>& KinematicTarget = Particle.KinematicTarget();
 			switch (KinematicTarget.GetMode())
 			{
@@ -240,6 +233,20 @@ namespace Chaos
 				break;
 			}
 			}
+
+			// Update world space bouunds
+			if (Particle.HasBounds())
+			{
+				const TAABB<FReal, 3>& LocalBounds = Particle.LocalBounds();
+				
+				TAABB<FReal, 3> NextWorldSpaceBounds = LocalBounds.TransformedAABB(FRigidTransform3(Particle.X(), Particle.R()));
+				NextWorldSpaceBounds.ThickenSymmetrically(NextWorldSpaceBounds.Extents() * BoundsExtension);
+
+				TAABB<FReal, 3> WorldSpaceBounds = LocalBounds.TransformedAABB(FRigidTransform3(PrevX, PrevR));
+				WorldSpaceBounds.GrowToInclude(NextWorldSpaceBounds);
+
+				Particle.SetWorldSpaceInflatedBounds(WorldSpaceBounds);
+			}
 		}
 	}
 
diff --git a/Engine/Source/Runtime/Experimental/Chaos/Private/Chaos/PBDCollisionConstraints.cpp b/Engine/Source/Runtime/Experimental/Chaos/Private/Chaos/PBDCollisionConstraints.cpp
--- a/Engine/Source/Runtime/Experimental/Chaos/Private/Chaos/PBDCollisionConstraints.cpp
+++ b/Engine/Source/Runtime/Experimental/Chaos/Private/Chaos/PBDCollisionConstraints.cpp
@@ -66,13 +66,15 @@ namespace Chaos
 		const TArrayCollectionArray<TSerializablePtr<FChaosPhysicsMaterial>>& InPerParticleMaterials,
 		const int32 InApplyPairIterations /*= 1*/,
 		const int32 InApplyPushOutPairIterations /*= 1*/,
-		const T Thickness /*= (T)0*/)
+		const T CullDistance /*= (T)0*/,
+		const T ShapePadding /*= (T)0*/)
 		: Particles(InParticles)
 		, MCollided(Collided)
 		, MPhysicsMaterials(InPerParticleMaterials)
 		, MApplyPairIterations(InApplyPairIterations)
 		, MApplyPushOutPairIterations(InApplyPushOutPairIterations)
-		, MThickness(Thickness)
+		, MCullDistance(CullDistance)
+		, MShapePadding(ShapePadding)
 		, MAngularFriction(0)
 		, bUseCCD(false)
 		, bEnableCollisions(true)
@@ -294,9 +296,9 @@ namespace Chaos
 		//{
 		//	FConstraintContainerHandle* ConstraintHandle = Handles[ConstraintHandleIndex];
 		//	check(ConstraintHandle != nullptr);
-		//	Collisions::Update<ECollisionUpdateType::Deepest, float, 3>(MThickness, ConstraintHandle->GetContact());
+		//	Collisions::Update<ECollisionUpdateType::Deepest, float, 3>(MCullDistance, MShapePadding, ConstraintHandle->GetContact());
 
-		//	if (ConstraintHandle->GetContact().GetPhi() < MThickness) 
+		//	if (ConstraintHandle->GetContact().GetPhi() < MCullDistance) 
 		//	{
 		//		ConstraintHandle->GetContact().Timestamp = LifespanCounter;
 		//	}
@@ -304,8 +306,8 @@ namespace Chaos
 
 		for (FPointContactConstraint& Contact : PointConstraints)
 		{
-			Collisions::Update<ECollisionUpdateType::Deepest, float, 3>(MThickness, Contact);
-			if (Contact.GetPhi() < MThickness)
+			Collisions::Update<ECollisionUpdateType::Deepest, float, 3>(MCullDistance, Contact);
+			if (Contact.GetPhi() < MCullDistance)
 			{
 				Contact.Timestamp = LifespanCounter;
 			}
@@ -321,13 +323,13 @@ namespace Chaos
 		//{
 		//	FConstraintContainerHandle* ConstraintHandle = Handles[ConstraintHandleIndex];
 		//	check(ConstraintHandle != nullptr);
-		//	Collisions::UpdateManifold<float, 3>(MThickness, ConstraintHandle->GetContact());
+		//	Collisions::UpdateManifold<float, 3>(MCullDistance, MShapePadding, ConstraintHandle->GetContact());
 		//}, bDisableCollisionParallelFor);
 
 		for (FMultiPointContactConstraint& Contact : IterativeConstraints)
 		{
-			Collisions::UpdateManifold<float, 3>(MThickness, Contact);
-			if (Contact.GetPhi() < MThickness)
+			Collisions::UpdateManifold<float, 3>(MCullDistance, Contact);
+			if (Contact.GetPhi() < MCullDistance)
 			{
 				Contact.Timestamp = LifespanCounter;
 			}
@@ -341,7 +343,7 @@ namespace Chaos
 
 		if (MApplyPairIterations > 0)
 		{
-			const Collisions::TContactParticleParameters<T> ParticleParameters = { MThickness, &MCollided };
+			const Collisions::TContactParticleParameters<T> ParticleParameters = { MCullDistance, MShapePadding, &MCollided };
 			const Collisions::TContactIterationParameters<T> IterationParameters = { Dt, Iterations, NumIterations, MApplyPairIterations, nullptr };
 
 			for (FPointContactConstraint& Contact : PointConstraints)
@@ -370,7 +372,7 @@ namespace Chaos
 		bool bNeedsAnotherIteration = false;
 		if (MApplyPushOutPairIterations > 0)
 		{
-			const Collisions::TContactParticleParameters<T> ParticleParameters = { MThickness, &MCollided };
+			const Collisions::TContactParticleParameters<T> ParticleParameters = { MCullDistance, MShapePadding, &MCollided };
 			const Collisions::TContactIterationParameters<T> IterationParameters = { Dt, Iterations, NumIterations, MApplyPushOutPairIterations, &bNeedsAnotherIteration };
 
 			for (FPointContactConstraint& Contact : PointConstraints)
@@ -404,7 +406,7 @@ namespace Chaos
 				FConstraintContainerHandle* ConstraintHandle = InConstraintHandles[ConstraintHandleIndex];
 				check(ConstraintHandle != nullptr);
 
-				Collisions::TContactParticleParameters<T> ParticleParameters = { MThickness, &MCollided };
+				Collisions::TContactParticleParameters<T> ParticleParameters = { MCullDistance, MShapePadding, &MCollided };
 				Collisions::TContactIterationParameters<T> IterationParameters = { Dt, Iterations, NumIterations, MApplyPairIterations, nullptr };
 				Collisions::Apply(ConstraintHandle->GetContact(), IterationParameters, ParticleParameters);
 
@@ -431,7 +433,7 @@ namespace Chaos
 				FConstraintContainerHandle* ConstraintHandle = InConstraintHandles[ConstraintHandleIndex];
 				check(ConstraintHandle != nullptr);
 
-				Collisions::TContactParticleParameters<T> ParticleParameters = { MThickness, &MCollided };
+				Collisions::TContactParticleParameters<T> ParticleParameters = { MCullDistance, MShapePadding, &MCollided };
 				Collisions::TContactIterationParameters<T> IterationParameters = { Dt, Iteration, NumIterations, MApplyPushOutPairIterations, &bNeedsAnotherIteration };
 				Collisions::ApplyPushOut(ConstraintHandle->GetContact(), IsTemporarilyStatic, IterationParameters, ParticleParameters);
 
diff --git a/Engine/Source/Runtime/Experimental/Chaos/Private/Chaos/PBDCollisionConstraintsContact.cpp b/Engine/Source/Runtime/Experimental/Chaos/Private/Chaos/PBDCollisionConstraintsContact.cpp
--- a/Engine/Source/Runtime/Experimental/Chaos/Private/Chaos/PBDCollisionConstraintsContact.cpp
+++ b/Engine/Source/Runtime/Experimental/Chaos/Private/Chaos/PBDCollisionConstraintsContact.cpp
@@ -12,22 +12,22 @@ namespace Chaos
 	namespace Collisions
 	{
 		template<ECollisionUpdateType UpdateType, typename T, int d>
-		void Update(const T Thickness, TCollisionConstraintBase<T, d>& Constraint)
+		void Update(const T CullDistance, TCollisionConstraintBase<T, d>& Constraint)
 		{
-			Constraint.ResetPhi(Thickness);
+			Constraint.ResetPhi(CullDistance);
 			const TRigidTransform<T, d> ParticleTM = GetTransform(Constraint.Particle[0]);
 			const TRigidTransform<T, d> LevelsetTM = GetTransform(Constraint.Particle[1]);
 
-			UpdateConstraint<UpdateType>(Constraint, ParticleTM, LevelsetTM, Thickness);
+			UpdateConstraint<UpdateType>(Constraint, ParticleTM, LevelsetTM, CullDistance);
 		}
 
 		template<typename T, int d>
-		void UpdateManifold(const T Thickness, TCollisionConstraintBase<T, d>& Constraint)
+		void UpdateManifold(const T CullDistance, TCollisionConstraintBase<T, d>& Constraint)
 		{
 			const TRigidTransform<T, d> Transform0 = GetTransform(Constraint.Particle[0]);
 			const TRigidTransform<T, d> Transform1 = GetTransform(Constraint.Particle[1]);
 
-			UpdateManifold(Constraint, Transform0, Transform1, Thickness);
+			UpdateManifold(Constraint, Transform0, Transform1, CullDistance);
 		}
 
 
@@ -209,10 +209,10 @@ namespace Chaos
 				const bool bNeedCollisionUpdate = (PairIt > 0) || (IterationParameters.Iteration > 0);
 				if (bNeedCollisionUpdate)
 				{
-					Collisions::Update<ECollisionUpdateType::Deepest>(ParticleParameters.Thickness, Constraint);
+					Collisions::Update<ECollisionUpdateType::Deepest>(ParticleParameters.CullDistance, Constraint);
 				}
 
-				if (Constraint.GetPhi() >= ParticleParameters.Thickness)
+				if (Constraint.GetPhi() >= ParticleParameters.ShapePadding)
 				{
 					return;
 				}
@@ -260,7 +260,7 @@ namespace Chaos
 			const bool IsTemporarilyStatic0 = IsTemporarilyStatic.Contains(Particle0->GeometryParticleHandle());
 			const bool IsTemporarilyStatic1 = IsTemporarilyStatic.Contains(Particle1->GeometryParticleHandle());
 
-			if (Contact.Phi >= ParticleParameters.Thickness)
+			if (Contact.Phi >= ParticleParameters.ShapePadding)
 			{
 				return AccumulatedImpulse;
 			}
@@ -324,7 +324,7 @@ namespace Chaos
 			}
 
 
-			TVector<T, d> Impulse = PMatrix<T, d, d>(Factor.Inverse()) * ((-Contact.Phi + ParticleParameters.Thickness) * ScalingFactor * Contact.Normal);
+			TVector<T, d> Impulse = PMatrix<T, d, d>(Factor.Inverse()) * ((-Contact.Phi + ParticleParameters.ShapePadding) * ScalingFactor * Contact.Normal);
 			TVector<T, d> AngularImpulse1 = TVector<T, d>::CrossProduct(VectorToPoint1, Impulse);
 			TVector<T, d> AngularImpulse2 = TVector<T, d>::CrossProduct(VectorToPoint2, -Impulse);
 			if (!IsTemporarilyStatic0 && bIsRigidDynamic0)
@@ -355,7 +355,7 @@ namespace Chaos
 
 			for (int32 PairIt = 0; PairIt < IterationParameters.NumPairIterations; ++PairIt)
 			{
-				Update<ECollisionUpdateType::Deepest>(ParticleParameters.Thickness, Constraint);
+				Update<ECollisionUpdateType::Deepest>(ParticleParameters.CullDistance, Constraint);
 
 				Constraint.AccumulatedImpulse += 
 					ApplyPushOutContact(Constraint.Manifold, Particle0, Particle1, IsTemporarilyStatic, IterationParameters, ParticleParameters);
diff --git a/Engine/Source/Runtime/Experimental/Chaos/Public/Chaos/Collision/NarrowPhase.h b/Engine/Source/Runtime/Experimental/Chaos/Public/Chaos/Collision/NarrowPhase.h
--- a/Engine/Source/Runtime/Experimental/Chaos/Public/Chaos/Collision/NarrowPhase.h
+++ b/Engine/Source/Runtime/Experimental/Chaos/Public/Chaos/Collision/NarrowPhase.h
@@ -22,8 +22,10 @@ namespace Chaos
 	{
 	public:
 		// @todo(chaos): COLLISION Transient Handle version
-		// @todo(chaos): COLLISION - Thickness: add shape padding (BoundsThickness is the distance within which we speculatively create constraints)
-		void GenerateCollisions(FCollisionConstraintsArray& NewConstraints, FReal Dt, TGeometryParticleHandle<FReal, 3>* Particle0, TGeometryParticleHandle<FReal, 3>* Particle1, const FReal AABBExpansion, CollisionStats::FStatData& StatData)
+		/**
+		 * /param CullDistance The contact separation at which we ignore the constraint
+		 */
+		void GenerateCollisions(FCollisionConstraintsArray& NewConstraints, FReal Dt, TGeometryParticleHandle<FReal, 3>* Particle0, TGeometryParticleHandle<FReal, 3>* Particle1, const FReal CullDistance, CollisionStats::FStatData& StatData)
 		{
 			SCOPE_CYCLE_COUNTER(STAT_Collisions_NarrowPhase);
 			if (ensure(Particle0 && Particle1))
@@ -35,7 +37,7 @@ namespace Chaos
 				//   determine if the constraint is already defined, and then opt out of 
 				//   the creation process. 
 				//
-				Collisions::ConstructConstraints<FReal, 3>(Particle0, Particle1, Particle0->Geometry().Get(), Particle1->Geometry().Get(), Collisions::GetTransform(Particle0), Collisions::GetTransform(Particle1), AABBExpansion, NewConstraints);
+				Collisions::ConstructConstraints<FReal, 3>(Particle0, Particle1, Particle0->Geometry().Get(), Particle1->Geometry().Get(), Collisions::GetTransform(Particle0), Collisions::GetTransform(Particle1), CullDistance, NewConstraints);
 
 				CHAOS_COLLISION_STAT(if (NewConstraints.Num()) { StatData.IncrementCountNP(NewConstraints.Num()); });
 				CHAOS_COLLISION_STAT(if (!NewConstraints.Num()) { StatData.IncrementRejectedNP(); });
diff --git a/Engine/Source/Runtime/Experimental/Chaos/Public/Chaos/Collision/ParticlePairBroadPhase.h b/Engine/Source/Runtime/Experimental/Chaos/Public/Chaos/Collision/ParticlePairBroadPhase.h
--- a/Engine/Source/Runtime/Experimental/Chaos/Public/Chaos/Collision/ParticlePairBroadPhase.h
+++ b/Engine/Source/Runtime/Experimental/Chaos/Public/Chaos/Collision/ParticlePairBroadPhase.h
@@ -23,11 +23,17 @@ namespace Chaos
 		using FParticlePair = TVector<TGeometryParticleHandle<FReal, 3>*, 2>;
 		using FAABB = TAABB<FReal, 3>;
 
-		FParticlePairBroadPhase(const TArray<FParticlePair>& InParticlePairs)
+		FParticlePairBroadPhase(const TArray<FParticlePair>& InParticlePairs, const FReal InCullDistance)
 			: ParticlePairs(InParticlePairs)
+			, CullDistance(InCullDistance)
 		{
 		}
 
+		void SetCullDustance(const FReal InCullDistance)
+		{
+			CullDistance = InCullDistance;
+		}
+
 		/**
 		 *
 		 */
@@ -38,11 +44,13 @@ namespace Chaos
 		{
 			SCOPE_CYCLE_COUNTER(STAT_Collisions_BroadPhase);
 
-			const FReal AABBExpansion = 1.0f;		// WRONG!!!!
+			OverlappingIndices.Reserve(ParticlePairs.Num());
 
-			FCollisionConstraintsArray NewConstraints;
-			for (const FParticlePair& ParticlePair : ParticlePairs)
+			const int32 NumPairs = ParticlePairs.Num();
+			for (int32 PairIndex = 0; PairIndex < NumPairs; ++PairIndex)
 			{
+				const FParticlePair& ParticlePair = ParticlePairs[PairIndex];
+
 				// Array is const, things in it are not...
 				TGeometryParticleHandle<FReal, 3>* Particle0 = const_cast<TGeometryParticleHandle<FReal, 3>*>(ParticlePair[0]);
 				TGeometryParticleHandle<FReal, 3>* Particle1 = const_cast<TGeometryParticleHandle<FReal, 3>*>(ParticlePair[1]);
@@ -56,15 +64,29 @@ namespace Chaos
 					const TAABB<FReal, 3>& Box1 = Particle1->WorldSpaceInflatedBounds();
 					if (Box0.Intersects(Box1))
 					{
-						NarrowPhase.GenerateCollisions(NewConstraints, Dt,Particle0, Particle1, AABBExpansion, StatData);
-						Receiver.ReceiveCollisions(NewConstraints);
-						NewConstraints.Empty();
+						OverlappingIndices.Add(PairIndex);
 					}
 				}
 			}
+
+			for (int32 PairIndex : OverlappingIndices)
+			{
+				const FParticlePair& ParticlePair = ParticlePairs[PairIndex];
+				TGeometryParticleHandle<FReal, 3>* Particle0 = const_cast<TGeometryParticleHandle<FReal, 3>*>(ParticlePair[0]);
+				TGeometryParticleHandle<FReal, 3>* Particle1 = const_cast<TGeometryParticleHandle<FReal, 3>*>(ParticlePair[1]);
+
+				NarrowPhase.GenerateCollisions(NewConstraints, Dt, Particle0, Particle1, CullDistance, StatData);
+				Receiver.ReceiveCollisions(NewConstraints);
+				NewConstraints.Empty();
+			}
+
+			OverlappingIndices.Reset();
 		}
 
 	private:
 		const TArray<FParticlePair>& ParticlePairs;
+		TArray<int32> OverlappingIndices;
+		FReal CullDistance;
+		FCollisionConstraintsArray NewConstraints;
 	};
 }
diff --git a/Engine/Source/Runtime/Experimental/Chaos/Public/Chaos/CollisionResolution.h b/Engine/Source/Runtime/Experimental/Chaos/Public/Chaos/CollisionResolution.h
--- a/Engine/Source/Runtime/Experimental/Chaos/Public/Chaos/CollisionResolution.h
+++ b/Engine/Source/Runtime/Experimental/Chaos/Public/Chaos/CollisionResolution.h
@@ -42,13 +42,13 @@ namespace Chaos
 		//
 
 		template <typename T, int d>
-		void CHAOS_API UpdateBoxBoxConstraint(const TAABB<T, d>& Box1, const TRigidTransform<T, d>& Box1Transform, const TAABB<T, d>& Box2, const TRigidTransform<T, d>& Box2Transform, const T Thickness, TRigidBodyPointContactConstraint<T, d>& Constraint);
+		void CHAOS_API UpdateBoxBoxConstraint(const TAABB<T, d>& Box1, const TRigidTransform<T, d>& Box1Transform, const TAABB<T, d>& Box2, const TRigidTransform<T, d>& Box2Transform, const T CullDistance, TRigidBodyPointContactConstraint<T, d>& Constraint);
 
 		template <typename T, int d>
-		void CHAOS_API UpdateBoxBoxManifold(TCollisionConstraintBase<T, d>&  Constraint, const TRigidTransform<T, d>& ATM, const TRigidTransform<T, d>& BTM, const T Thickness);
+		void CHAOS_API UpdateBoxBoxManifold(TCollisionConstraintBase<T, d>&  Constraint, const TRigidTransform<T, d>& ATM, const TRigidTransform<T, d>& BTM, const T CullDistance);
 
 		template<typename T, int d>
-		void CHAOS_API ConstructBoxBoxConstraints(TGeometryParticleHandle<T, d>* Particle0, TGeometryParticleHandle<T, d>* Particle1, const FImplicitObject* Implicit0, const FImplicitObject* Implicit1, const TRigidTransform<T, d>& Transform0, const TRigidTransform<T, d>& Transform1, const T Thickness, FCollisionConstraintsArray& NewConstraints);
+		void CHAOS_API ConstructBoxBoxConstraints(TGeometryParticleHandle<T, d>* Particle0, TGeometryParticleHandle<T, d>* Particle1, const FImplicitObject* Implicit0, const FImplicitObject* Implicit1, const TRigidTransform<T, d>& Transform0, const TRigidTransform<T, d>& Transform1, const T CullDistance, FCollisionConstraintsArray& NewConstraints);
 
 
 		//
@@ -56,26 +56,26 @@ namespace Chaos
 		//
 
 		template <typename T, int d>
-		void CHAOS_API UpdateBoxHeightFieldConstraint(const TAABB<T, d>& A, const TRigidTransform<T, d>& ATransform, const THeightField<T>& B, const TRigidTransform<T, d>& BTransform, const T Thickness, TRigidBodyPointContactConstraint<T, d>& Constraint);
+		void CHAOS_API UpdateBoxHeightFieldConstraint(const TAABB<T, d>& A, const TRigidTransform<T, d>& ATransform, const THeightField<T>& B, const TRigidTransform<T, d>& BTransform, const T CullDistance, TRigidBodyPointContactConstraint<T, d>& Constraint);
 
 		template<class T, int d>
-		void CHAOS_API UpdateBoxHeightFieldManifold(TCollisionConstraintBase<T, d>&  Constraint, const TRigidTransform<T, d>& ATM, const TRigidTransform<T, d>& BTM, const T Thickness);
+		void CHAOS_API UpdateBoxHeightFieldManifold(TCollisionConstraintBase<T, d>&  Constraint, const TRigidTransform<T, d>& ATM, const TRigidTransform<T, d>& BTM, const T CullDistance);
 
 		template<typename T, int d>
-		void CHAOS_API ConstructBoxHeightFieldConstraints(TGeometryParticleHandle<T, d>* Particle0, TGeometryParticleHandle<T, d>* Particle1, const FImplicitObject* Implicit0, const FImplicitObject* Implicit1, const TRigidTransform<T, d>& Transform0, const TRigidTransform<T, d>& Transform1, const T Thickness, FCollisionConstraintsArray& NewConstraints);
+		void CHAOS_API ConstructBoxHeightFieldConstraints(TGeometryParticleHandle<T, d>* Particle0, TGeometryParticleHandle<T, d>* Particle1, const FImplicitObject* Implicit0, const FImplicitObject* Implicit1, const TRigidTransform<T, d>& Transform0, const TRigidTransform<T, d>& Transform1, const T CullDistance, FCollisionConstraintsArray& NewConstraints);
 
 		//
 		// Box-Plane
 		//
 
 		template <typename T, int d>
-		void CHAOS_API UpdateBoxPlaneConstraint(const TAABB<T, d>& Box, const TRigidTransform<T, d>& BoxTransform, const TPlane<T, d>& Plane, const TRigidTransform<T, d>& PlaneTransform, const T Thickness, TRigidBodyPointContactConstraint<T, d>& Constraint);
+		void CHAOS_API UpdateBoxPlaneConstraint(const TAABB<T, d>& Box, const TRigidTransform<T, d>& BoxTransform, const TPlane<T, d>& Plane, const TRigidTransform<T, d>& PlaneTransform, const T CullDistance, TRigidBodyPointContactConstraint<T, d>& Constraint);
 
 		template<class T, int d>
-		void CHAOS_API UpdateBoxPlaneManifold(TCollisionConstraintBase<T, d>&  Constraint, const TRigidTransform<T, d>& ATM, const TRigidTransform<T, d>& BTM, const T Thickness);
+		void CHAOS_API UpdateBoxPlaneManifold(TCollisionConstraintBase<T, d>&  Constraint, const TRigidTransform<T, d>& ATM, const TRigidTransform<T, d>& BTM, const T CullDistance);
 
 		template<typename T, int d>
-		void CHAOS_API ConstructBoxPlaneConstraints(TGeometryParticleHandle<T, d>* Particle0, TGeometryParticleHandle<T, d>* Particle1, const FImplicitObject* Implicit0, const FImplicitObject* Implicit1, const TRigidTransform<T, d>& Transform0, const TRigidTransform<T, d>& Transform1, const T Thickness, FCollisionConstraintsArray& NewConstraints);
+		void CHAOS_API ConstructBoxPlaneConstraints(TGeometryParticleHandle<T, d>* Particle0, TGeometryParticleHandle<T, d>* Particle1, const FImplicitObject* Implicit0, const FImplicitObject* Implicit1, const TRigidTransform<T, d>& Transform0, const TRigidTransform<T, d>& Transform1, const T CullDistance, FCollisionConstraintsArray& NewConstraints);
 
 
 
@@ -84,13 +84,13 @@ namespace Chaos
 		//
 
 		template <typename T, int d>
-		void CHAOS_API UpdateBoxTriangleMeshConstraint(const TAABB<T, d>& Box0, const TRigidTransform<T, d>& Transform0, const FTriangleMeshImplicitObject& TriangleMesh1, const TRigidTransform<T, d>& Transformw1, const T Thickness, TRigidBodyPointContactConstraint<T, d>& Constraint);
+		void CHAOS_API UpdateBoxTriangleMeshConstraint(const TAABB<T, d>& Box0, const TRigidTransform<T, d>& Transform0, const FTriangleMeshImplicitObject& TriangleMesh1, const TRigidTransform<T, d>& Transformw1, const T CullDistance, TRigidBodyPointContactConstraint<T, d>& Constraint);
 
 		template <typename T, int d>
-		void CHAOS_API UpdateBoxTriangleMeshManifold(TCollisionConstraintBase<T, d>&  Constraint, const TRigidTransform<T, d>& ATM, const TRigidTransform<T, d>& BTM, const T Thickness);
+		void CHAOS_API UpdateBoxTriangleMeshManifold(TCollisionConstraintBase<T, d>&  Constraint, const TRigidTransform<T, d>& ATM, const TRigidTransform<T, d>& BTM, const T CullDistance);
 
 		template<typename T, int d>
-		void CHAOS_API ConstructBoxTriangleMeshConstraints(TGeometryParticleHandle<T, d>* Particle0, TGeometryParticleHandle<T, d>* Particle1, const FImplicitObject* Implicit0, const FImplicitObject* Implicit1, const TRigidTransform<T, d>& Transform0, const TRigidTransform<T, d>& Transform1, const T Thickness, FCollisionConstraintsArray& NewConstraints);
+		void CHAOS_API ConstructBoxTriangleMeshConstraints(TGeometryParticleHandle<T, d>* Particle0, TGeometryParticleHandle<T, d>* Particle1, const FImplicitObject* Implicit0, const FImplicitObject* Implicit1, const TRigidTransform<T, d>& Transform0, const TRigidTransform<T, d>& Transform1, const T CullDistance, FCollisionConstraintsArray& NewConstraints);
 
 
 		//
@@ -98,167 +98,167 @@ namespace Chaos
 		//
 
 		template <typename T, int d>
-		void CHAOS_API UpdateSphereSphereConstraint(const TSphere<T, d>& Sphere1, const TRigidTransform<T, d>& Sphere1Transform, const TSphere<T, d>& Sphere2, const TRigidTransform<T, d>& Sphere2Transform, const T Thickness, TRigidBodyPointContactConstraint<T, d>& Constraint);
+		void CHAOS_API UpdateSphereSphereConstraint(const TSphere<T, d>& Sphere1, const TRigidTransform<T, d>& Sphere1Transform, const TSphere<T, d>& Sphere2, const TRigidTransform<T, d>& Sphere2Transform, const T CullDistance, TRigidBodyPointContactConstraint<T, d>& Constraint);
 
 		template<class T, int d>
-		void CHAOS_API UpdateSphereSphereManifold(TCollisionConstraintBase<T, d>&  Constraint, const TRigidTransform<T, d>& ATM, const TRigidTransform<T, d>& BTM, const T Thickness);
+		void CHAOS_API UpdateSphereSphereManifold(TCollisionConstraintBase<T, d>&  Constraint, const TRigidTransform<T, d>& ATM, const TRigidTransform<T, d>& BTM, const T CullDistance);
 
 		template<typename T, int d>
-		void CHAOS_API ConstructSphereSphereConstraints(TGeometryParticleHandle<T, d>* Particle0, TGeometryParticleHandle<T, d>* Particle1, const FImplicitObject* Implicit0, const FImplicitObject* Implicit1, const TRigidTransform<T, d>& Transform0, const TRigidTransform<T, d>& Transform1, const T Thickness, FCollisionConstraintsArray& NewConstraints);
+		void CHAOS_API ConstructSphereSphereConstraints(TGeometryParticleHandle<T, d>* Particle0, TGeometryParticleHandle<T, d>* Particle1, const FImplicitObject* Implicit0, const FImplicitObject* Implicit1, const TRigidTransform<T, d>& Transform0, const TRigidTransform<T, d>& Transform1, const T CullDistance, FCollisionConstraintsArray& NewConstraints);
 
 		//
 		// Sphere-HeightField
 		//
 
 		template <typename T, int d>
-		void CHAOS_API UpdateSphereHeightFieldConstraint(const TSphere<T, d>& A, const TRigidTransform<T, d>& ATransform, const THeightField<T>& B, const TRigidTransform<T, d>& BTransform, const T Thickness, TRigidBodyPointContactConstraint<T, d>& Constraint);
+		void CHAOS_API UpdateSphereHeightFieldConstraint(const TSphere<T, d>& A, const TRigidTransform<T, d>& ATransform, const THeightField<T>& B, const TRigidTransform<T, d>& BTransform, const T CullDistance, TRigidBodyPointContactConstraint<T, d>& Constraint);
 
 		template<class T, int d>
-		void CHAOS_API UpdateSphereHeightFieldManifold(TCollisionConstraintBase<T, d>&  Constraint, const TRigidTransform<T, d>& ATM, const TRigidTransform<T, d>& BTM, const T Thickness);
+		void CHAOS_API UpdateSphereHeightFieldManifold(TCollisionConstraintBase<T, d>&  Constraint, const TRigidTransform<T, d>& ATM, const TRigidTransform<T, d>& BTM, const T CullDistance);
 
 		template<typename T, int d>
-		void CHAOS_API ConstructSphereHeightFieldConstraints(TGeometryParticleHandle<T, d>* Particle0, TGeometryParticleHandle<T, d>* Particle1, const FImplicitObject* Implicit0, const FImplicitObject* Implicit1, const TRigidTransform<T, d>& Transform0, const TRigidTransform<T, d>& Transform1, const T Thickness, FCollisionConstraintsArray& NewConstraints);
+		void CHAOS_API ConstructSphereHeightFieldConstraints(TGeometryParticleHandle<T, d>* Particle0, TGeometryParticleHandle<T, d>* Particle1, const FImplicitObject* Implicit0, const FImplicitObject* Implicit1, const TRigidTransform<T, d>& Transform0, const TRigidTransform<T, d>& Transform1, const T CullDistance, FCollisionConstraintsArray& NewConstraints);
 
 		//
 		// Sphere-Plane
 		//
 
 		template <typename T, int d>
-		void CHAOS_API UpdateSpherePlaneConstraint(const TSphere<T, d>& Sphere, const TRigidTransform<T, d>& SphereTransform, const TPlane<T, d>& Plane, const TRigidTransform<T, d>& PlaneTransform, const T Thickness, TRigidBodyPointContactConstraint<T, d>& Constraint);
+		void CHAOS_API UpdateSpherePlaneConstraint(const TSphere<T, d>& Sphere, const TRigidTransform<T, d>& SphereTransform, const TPlane<T, d>& Plane, const TRigidTransform<T, d>& PlaneTransform, const T CullDistance, TRigidBodyPointContactConstraint<T, d>& Constraint);
 
 		template<class T, int d>
-		void CHAOS_API UpdateSpherePlaneManifold(TCollisionConstraintBase<T, d>&  Constraint, const TRigidTransform<T, d>& ATM, const TRigidTransform<T, d>& BTM, const T Thickness);
+		void CHAOS_API UpdateSpherePlaneManifold(TCollisionConstraintBase<T, d>&  Constraint, const TRigidTransform<T, d>& ATM, const TRigidTransform<T, d>& BTM, const T CullDistance);
 
 		template<typename T, int d>
-		void CHAOS_API ConstructSpherePlaneConstraints(TGeometryParticleHandle<T, d>* Particle0, TGeometryParticleHandle<T, d>* Particle1, const FImplicitObject* Implicit0, const FImplicitObject* Implicit1, const TRigidTransform<T, d>& Transform0, const TRigidTransform<T, d>& Transform1, const T Thickness, FCollisionConstraintsArray& NewConstraints);
+		void CHAOS_API ConstructSpherePlaneConstraints(TGeometryParticleHandle<T, d>* Particle0, TGeometryParticleHandle<T, d>* Particle1, const FImplicitObject* Implicit0, const FImplicitObject* Implicit1, const TRigidTransform<T, d>& Transform0, const TRigidTransform<T, d>& Transform1, const T CullDistance, FCollisionConstraintsArray& NewConstraints);
 
 		//
 		// Sphere-Box
 		//
 
 		template <typename T, int d>
-		void CHAOS_API UpdateSphereBoxConstraint(const TSphere<T, d>& Sphere, const TRigidTransform<T, d>& SphereTransform, const TAABB<T, d>& Box, const TRigidTransform<T, d>& BoxTransform, const T Thickness, TRigidBodyPointContactConstraint<T, d>& Constraint);
+		void CHAOS_API UpdateSphereBoxConstraint(const TSphere<T, d>& Sphere, const TRigidTransform<T, d>& SphereTransform, const TAABB<T, d>& Box, const TRigidTransform<T, d>& BoxTransform, const T CullDistance, TRigidBodyPointContactConstraint<T, d>& Constraint);
 
 		template<class T, int d>
-		void CHAOS_API UpdateSphereBoxManifold(TCollisionConstraintBase<T, d>&  Constraint, const TRigidTransform<T, d>& ATM, const TRigidTransform<T, d>& BTM, const T Thickness);
+		void CHAOS_API UpdateSphereBoxManifold(TCollisionConstraintBase<T, d>&  Constraint, const TRigidTransform<T, d>& ATM, const TRigidTransform<T, d>& BTM, const T CullDistance);
 
 		template<typename T, int d>
-		void CHAOS_API ConstructSphereBoxConstraints(TGeometryParticleHandle<T, d>* Particle0, TGeometryParticleHandle<T, d>* Particle1, const FImplicitObject* Implicit0, const FImplicitObject* Implicit1, const TRigidTransform<T, d>& Transform0, const TRigidTransform<T, d>& Transform1, const T Thickness, FCollisionConstraintsArray& NewConstraints);
+		void CHAOS_API ConstructSphereBoxConstraints(TGeometryParticleHandle<T, d>* Particle0, TGeometryParticleHandle<T, d>* Particle1, const FImplicitObject* Implicit0, const FImplicitObject* Implicit1, const TRigidTransform<T, d>& Transform0, const TRigidTransform<T, d>& Transform1, const T CullDistance, FCollisionConstraintsArray& NewConstraints);
 
 		//
 		// Sphere-Capsule
 		//
 
 		template <typename T, int d>
-		void CHAOS_API UpdateSphereCapsuleConstraint(const TSphere<T, d>& Sphere, const TRigidTransform<T, d>& SphereTransform, const TCapsule<T>& Box, const TRigidTransform<T, d>& BoxTransform, const T Thickness, TRigidBodyPointContactConstraint<T, d>& Constraint);
+		void CHAOS_API UpdateSphereCapsuleConstraint(const TSphere<T, d>& Sphere, const TRigidTransform<T, d>& SphereTransform, const TCapsule<T>& Box, const TRigidTransform<T, d>& BoxTransform, const T CullDistance, TRigidBodyPointContactConstraint<T, d>& Constraint);
 
 		template<class T, int d>
-		void CHAOS_API UpdateSphereCapsuleManifold(TCollisionConstraintBase<T, d>&  Constraint, const TRigidTransform<T, d>& ATM, const TRigidTransform<T, d>& BTM, const T Thickness);
+		void CHAOS_API UpdateSphereCapsuleManifold(TCollisionConstraintBase<T, d>&  Constraint, const TRigidTransform<T, d>& ATM, const TRigidTransform<T, d>& BTM, const T CullDistance);
 
 		template<typename T, int d>
-		void CHAOS_API ConstructSphereCapsuleConstraints(TGeometryParticleHandle<T, d>* Particle0, TGeometryParticleHandle<T, d>* Particle1, const FImplicitObject* Implicit0, const FImplicitObject* Implicit1, const TRigidTransform<T, d>& Transform0, const TRigidTransform<T, d>& Transform1, const T Thickness, FCollisionConstraintsArray& NewConstraints);
+		void CHAOS_API ConstructSphereCapsuleConstraints(TGeometryParticleHandle<T, d>* Particle0, TGeometryParticleHandle<T, d>* Particle1, const FImplicitObject* Implicit0, const FImplicitObject* Implicit1, const TRigidTransform<T, d>& Transform0, const TRigidTransform<T, d>& Transform1, const T CullDistance, FCollisionConstraintsArray& NewConstraints);
 
 		//
 		// Capsule-Capsule
 		//
 
 		template <typename T, int d>
-		void CHAOS_API UpdateCapsuleCapsuleConstraint(const TCapsule<T>& A, const TRigidTransform<T, d>& ATransform, const TCapsule<T>& B, const TRigidTransform<T, d>& BTransform, const T Thickness, TRigidBodyPointContactConstraint<T, d>& Constraint);
+		void CHAOS_API UpdateCapsuleCapsuleConstraint(const TCapsule<T>& A, const TRigidTransform<T, d>& ATransform, const TCapsule<T>& B, const TRigidTransform<T, d>& BTransform, const T CullDistance, TRigidBodyPointContactConstraint<T, d>& Constraint);
 
 		template<class T, int d>
-		void CHAOS_API UpdateCapsuleCapsuleManifold(TCollisionConstraintBase<T, d>&  Constraint, const TRigidTransform<T, d>& ATM, const TRigidTransform<T, d>& BTM, const T Thickness);
+		void CHAOS_API UpdateCapsuleCapsuleManifold(TCollisionConstraintBase<T, d>&  Constraint, const TRigidTransform<T, d>& ATM, const TRigidTransform<T, d>& BTM, const T CullDistance);
 
 		template<typename T, int d>
-		void CHAOS_API ConstructCapsuleCapsuleConstraints(TGeometryParticleHandle<T, d>* Particle0, TGeometryParticleHandle<T, d>* Particle1, const FImplicitObject* Implicit0, const FImplicitObject* Implicit1, const TRigidTransform<T, d>& Transform0, const TRigidTransform<T, d>& Transform1, const T Thickness, FCollisionConstraintsArray& NewConstraints);
+		void CHAOS_API ConstructCapsuleCapsuleConstraints(TGeometryParticleHandle<T, d>* Particle0, TGeometryParticleHandle<T, d>* Particle1, const FImplicitObject* Implicit0, const FImplicitObject* Implicit1, const TRigidTransform<T, d>& Transform0, const TRigidTransform<T, d>& Transform1, const T CullDistance, FCollisionConstraintsArray& NewConstraints);
 
 		//
 		// Capsule-Box
 		//
 
 		template <typename T, int d>
-		void CHAOS_API UpdateCapsuleBoxConstraint(const TCapsule<T>& A, const TRigidTransform<T, d>& ATransform, const TAABB<T, d>& B, const TRigidTransform<T, d>& BTransform, const T Thickness, TRigidBodyPointContactConstraint<T, d>& Constraint);
+		void CHAOS_API UpdateCapsuleBoxConstraint(const TCapsule<T>& A, const TRigidTransform<T, d>& ATransform, const TAABB<T, d>& B, const TRigidTransform<T, d>& BTransform, const T CullDistance, TRigidBodyPointContactConstraint<T, d>& Constraint);
 
 		template<class T, int d>
-		void CHAOS_API UpdateCapsuleBoxManifold(TCollisionConstraintBase<T, d>&  Constraint, const TRigidTransform<T, d>& ATM, const TRigidTransform<T, d>& BTM, const T Thickness);
+		void CHAOS_API UpdateCapsuleBoxManifold(TCollisionConstraintBase<T, d>&  Constraint, const TRigidTransform<T, d>& ATM, const TRigidTransform<T, d>& BTM, const T CullDistance);
 
 		template<typename T, int d>
-		void CHAOS_API ConstructCapsuleBoxConstraints(TGeometryParticleHandle<T, d>* Particle0, TGeometryParticleHandle<T, d>* Particle1, const FImplicitObject* Implicit0, const FImplicitObject* Implicit1, const TRigidTransform<T, d>& Transform0, const TRigidTransform<T, d>& Transform1, const T Thickness, FCollisionConstraintsArray& NewConstraints);
+		void CHAOS_API ConstructCapsuleBoxConstraints(TGeometryParticleHandle<T, d>* Particle0, TGeometryParticleHandle<T, d>* Particle1, const FImplicitObject* Implicit0, const FImplicitObject* Implicit1, const TRigidTransform<T, d>& Transform0, const TRigidTransform<T, d>& Transform1, const T CullDistance, FCollisionConstraintsArray& NewConstraints);
 
 		//
 		// Capsule-HeightField
 		//
 
 		template <typename T, int d>
-		void CHAOS_API UpdateCapsuleHeightFieldConstraint(const TCapsule<T>& A, const TRigidTransform<T, d>& ATransform, const THeightField<T>& B, const TRigidTransform<T, d>& BTransform, const T Thickness, TRigidBodyPointContactConstraint<T, d>& Constraint);
+		void CHAOS_API UpdateCapsuleHeightFieldConstraint(const TCapsule<T>& A, const TRigidTransform<T, d>& ATransform, const THeightField<T>& B, const TRigidTransform<T, d>& BTransform, const T CullDistance, TRigidBodyPointContactConstraint<T, d>& Constraint);
 
 		template<class T, int d>
-		void CHAOS_API UpdateCapsuleHeightFieldManifold(TCollisionConstraintBase<T, d>&  Constraint, const TRigidTransform<T, d>& ATM, const TRigidTransform<T, d>& BTM, const T Thickness);
+		void CHAOS_API UpdateCapsuleHeightFieldManifold(TCollisionConstraintBase<T, d>&  Constraint, const TRigidTransform<T, d>& ATM, const TRigidTransform<T, d>& BTM, const T CullDistance);
 
 		template<typename T, int d>
-		void ConstructCapsuleHeightFieldConstraints(TGeometryParticleHandle<T, d>* Particle0, TGeometryParticleHandle<T, d>* Particle1, const FImplicitObject* Implicit0, const FImplicitObject* Implicit1, const TRigidTransform<T, d>& Transform0, const TRigidTransform<T, d>& Transform1, const T Thickness, FCollisionConstraintsArray& NewConstraints);
+		void ConstructCapsuleHeightFieldConstraints(TGeometryParticleHandle<T, d>* Particle0, TGeometryParticleHandle<T, d>* Particle1, const FImplicitObject* Implicit0, const FImplicitObject* Implicit1, const TRigidTransform<T, d>& Transform0, const TRigidTransform<T, d>& Transform1, const T CullDistance, FCollisionConstraintsArray& NewConstraints);
 
 		//
 		// Convex-Convex
 		//
 		template <typename T, int d>
-		void CHAOS_API UpdateConvexConvexConstraint(const FImplicitObject& A, const TRigidTransform<T, d>& ATM, const FImplicitObject& B, const TRigidTransform<T, d>& BTM, const T Thickness, TCollisionConstraintBase<T, d>& Constraint);
+		void CHAOS_API UpdateConvexConvexConstraint(const FImplicitObject& A, const TRigidTransform<T, d>& ATM, const FImplicitObject& B, const TRigidTransform<T, d>& BTM, const T CullDistance, TCollisionConstraintBase<T, d>& Constraint);
 
 		template<class T, int d>
-		void CHAOS_API UpdateConvexConvexManifold(TCollisionConstraintBase<T, d>&  Constraint, const TRigidTransform<T, d>& ATM, const TRigidTransform<T, d>& BTM, const T Thickness);
+		void CHAOS_API UpdateConvexConvexManifold(TCollisionConstraintBase<T, d>&  Constraint, const TRigidTransform<T, d>& ATM, const TRigidTransform<T, d>& BTM, const T CullDistance);
 
 		template <typename T, int d>
-		void CHAOS_API ConstructConvexConvexConstraints(TGeometryParticleHandle<T, d>* Particle0, TGeometryParticleHandle<T, d>* Particle1, const FImplicitObject* Implicit0, const FImplicitObject* Implicit1, const TRigidTransform<T, d>& Transform0, const TRigidTransform<T, d>& Transform1, const T Thickness, FCollisionConstraintsArray& NewConstraints);
+		void CHAOS_API ConstructConvexConvexConstraints(TGeometryParticleHandle<T, d>* Particle0, TGeometryParticleHandle<T, d>* Particle1, const FImplicitObject* Implicit0, const FImplicitObject* Implicit1, const TRigidTransform<T, d>& Transform0, const TRigidTransform<T, d>& Transform1, const T CullDistance, FCollisionConstraintsArray& NewConstraints);
 
 		//
 		// Convex-HeightField
 		//
 
 		template <typename T, int d>
-		void CHAOS_API UpdateConvexHeightFieldConstraint(const FImplicitObject& A, const TRigidTransform<T, d>& ATransform, const THeightField<T>& B, const TRigidTransform<T, d>& BTransform, const T Thickness, TRigidBodyPointContactConstraint<T, d>& Constraint);
+		void CHAOS_API UpdateConvexHeightFieldConstraint(const FImplicitObject& A, const TRigidTransform<T, d>& ATransform, const THeightField<T>& B, const TRigidTransform<T, d>& BTransform, const T CullDistance, TRigidBodyPointContactConstraint<T, d>& Constraint);
 
 		template<class T, int d>
-		void CHAOS_API UpdateConvexHeightFieldManifold(TCollisionConstraintBase<T, d>&  Constraint, const TRigidTransform<T, d>& ATM, const TRigidTransform<T, d>& BTM, const T Thickness);
+		void CHAOS_API UpdateConvexHeightFieldManifold(TCollisionConstraintBase<T, d>&  Constraint, const TRigidTransform<T, d>& ATM, const TRigidTransform<T, d>& BTM, const T CullDistance);
 
 		template<typename T, int d>
-		void CHAOS_API ConstructConvexHeightFieldConstraints(TGeometryParticleHandle<T, d>* Particle0, TGeometryParticleHandle<T, d>* Particle1, const FImplicitObject* Implicit0, const FImplicitObject* Implicit1, const TRigidTransform<T, d>& Transform0, const TRigidTransform<T, d>& Transform1, const T Thickness, FCollisionConstraintsArray& NewConstraints);
+		void CHAOS_API ConstructConvexHeightFieldConstraints(TGeometryParticleHandle<T, d>* Particle0, TGeometryParticleHandle<T, d>* Particle1, const FImplicitObject* Implicit0, const FImplicitObject* Implicit1, const TRigidTransform<T, d>& Transform0, const TRigidTransform<T, d>& Transform1, const T CullDistance, FCollisionConstraintsArray& NewConstraints);
 
 
 		//
 		// Levelset-Levelset
 		//
 		template<ECollisionUpdateType UpdateType, typename T = float, int d = 3>
-		void CHAOS_API UpdateLevelsetLevelsetConstraint(const T Thickness, TRigidBodyPointContactConstraint<T, d>& Constraint);
+		void CHAOS_API UpdateLevelsetLevelsetConstraint(const T CullDistance, TRigidBodyPointContactConstraint<T, d>& Constraint);
 
 		template<class T, int d>
-		void CHAOS_API UpdateLevelsetLevelsetManifold(TCollisionConstraintBase<T, d>&  Constraint, const TRigidTransform<T, d>& ATM, const TRigidTransform<T, d>& BTM, const T Thickness);
+		void CHAOS_API UpdateLevelsetLevelsetManifold(TCollisionConstraintBase<T, d>&  Constraint, const TRigidTransform<T, d>& ATM, const TRigidTransform<T, d>& BTM, const T CullDistance);
 
 		template<typename T, int d>
-		void CHAOS_API ConstructLevelsetLevelsetConstraints(TGeometryParticleHandle<T, d>* Particle0, TGeometryParticleHandle<T, d>* Particle1, const FImplicitObject* Implicit0, const FImplicitObject* Implicit1, const TRigidTransform<T, d>& Transform0, const TRigidTransform<T, d>& Transform1, const T Thickness, FCollisionConstraintsArray& NewConstraints);
+		void CHAOS_API ConstructLevelsetLevelsetConstraints(TGeometryParticleHandle<T, d>* Particle0, TGeometryParticleHandle<T, d>* Particle1, const FImplicitObject* Implicit0, const FImplicitObject* Implicit1, const TRigidTransform<T, d>& Transform0, const TRigidTransform<T, d>& Transform1, const T CullDistance, FCollisionConstraintsArray& NewConstraints);
 
 
 		//
 		// Union-Union
 		//
 		template<typename T, int d>
-		void CHAOS_API ConstructUnionUnionConstraints(TGeometryParticleHandle<T, d>* Particle0, TGeometryParticleHandle<T, d>* Particle1, const FImplicitObject* Implicit0, const FImplicitObject* Implicit1, const TRigidTransform<T, d>& Transform0, const TRigidTransform<T, d>& Transform1, const T Thickness, FCollisionConstraintsArray& NewConstraints);
+		void CHAOS_API ConstructUnionUnionConstraints(TGeometryParticleHandle<T, d>* Particle0, TGeometryParticleHandle<T, d>* Particle1, const FImplicitObject* Implicit0, const FImplicitObject* Implicit1, const TRigidTransform<T, d>& Transform0, const TRigidTransform<T, d>& Transform1, const T CullDistance, FCollisionConstraintsArray& NewConstraints);
 
 		//
 		// Constraint API
 		//
 		template <typename T, int d>
-		void CHAOS_API UpdateSingleShotManifold(TRigidBodyMultiPointContactConstraint<T, d>&  Constraint, const TRigidTransform<T, d>& Transform0, const TRigidTransform<T, d>& Transform1, const T Thickness);
+		void CHAOS_API UpdateSingleShotManifold(TRigidBodyMultiPointContactConstraint<T, d>&  Constraint, const TRigidTransform<T, d>& Transform0, const TRigidTransform<T, d>& Transform1, const T CullDistance);
 
 		template <typename T, int d>
-		void CHAOS_API UpdateIterativeManifold(TRigidBodyMultiPointContactConstraint<T, d>&  Constraint, const TRigidTransform<T, d>& Transform0, const TRigidTransform<T, d>& Transform1, const T Thickness);
+		void CHAOS_API UpdateIterativeManifold(TRigidBodyMultiPointContactConstraint<T, d>&  Constraint, const TRigidTransform<T, d>& Transform0, const TRigidTransform<T, d>& Transform1, const T CullDistance);
 
 		template<typename T, int d>
-		void CHAOS_API UpdateManifold(TCollisionConstraintBase<T, d>& ConstraintBase, const TRigidTransform<T, d>& ATM, const TRigidTransform<T, d>& BTM, const T Thickness);
+		void CHAOS_API UpdateManifold(TCollisionConstraintBase<T, d>& ConstraintBase, const TRigidTransform<T, d>& ATM, const TRigidTransform<T, d>& BTM, const T CullDistance);
 
 		template<ECollisionUpdateType UpdateType, typename T, int d>
-		void CHAOS_API UpdateConstraint(TCollisionConstraintBase<T, d>& ConstraintBase, const TRigidTransform<T, d>& Transform0, const TRigidTransform<T, d>& Transform1, const T Thickness);
+		void CHAOS_API UpdateConstraint(TCollisionConstraintBase<T, d>& ConstraintBase, const TRigidTransform<T, d>& Transform0, const TRigidTransform<T, d>& Transform1, const T CullDistance);
 
 		template<typename T, int d>
-		void CHAOS_API ConstructConstraints(TGeometryParticleHandle<T, d>* Particle0, TGeometryParticleHandle<T, d>* Particle1, const FImplicitObject* Implicit0, const FImplicitObject* Implicit1, const TRigidTransform<T, d>& Transform0, const TRigidTransform<T, d>& Transform1, const T Thickness, FCollisionConstraintsArray& NewConstraints);
+		void CHAOS_API ConstructConstraints(TGeometryParticleHandle<T, d>* Particle0, TGeometryParticleHandle<T, d>* Particle1, const FImplicitObject* Implicit0, const FImplicitObject* Implicit1, const TRigidTransform<T, d>& Transform0, const TRigidTransform<T, d>& Transform1, const T CullDistance, FCollisionConstraintsArray& NewConstraints);
 
 	}
 }
diff --git a/Engine/Source/Runtime/Experimental/Chaos/Public/Chaos/CollisionResolutionTypes.h b/Engine/Source/Runtime/Experimental/Chaos/Public/Chaos/CollisionResolutionTypes.h
--- a/Engine/Source/Runtime/Experimental/Chaos/Public/Chaos/CollisionResolutionTypes.h
+++ b/Engine/Source/Runtime/Experimental/Chaos/Public/Chaos/CollisionResolutionTypes.h
@@ -37,6 +37,14 @@ namespace Chaos
 		Disabled,	/** Collision should be disabled */
 	};
 
+	/** The shape types involved in a contact constraint. Used to look up the collision detection function */
+	enum class CHAOS_API EContactShapesType
+	{
+		Unknown,
+		CapsuleCapsule,
+		CapsuleBox,
+		BoxBox,
+	};
 
 	/*
 	*
@@ -46,7 +54,7 @@ namespace Chaos
 	{
 	public:
 		TCollisionContact(const FImplicitObject* InImplicit0 = nullptr, const FImplicitObject* InImplicit1 = nullptr)
-			: bDisabled(true), Normal(0), Location(0), Phi(FLT_MAX), Friction(0), AngularFriction(0), Restitution(0)
+			: bDisabled(true), Normal(0), Location(0), Phi(FLT_MAX), Friction(0), AngularFriction(0), Restitution(0), ShapesType(EContactShapesType::Unknown)
 		{
 			Implicit[0] = InImplicit0;
 			Implicit[1] = InImplicit1;
@@ -61,6 +69,8 @@ namespace Chaos
 		T AngularFriction;
 		T Restitution;
 
+		EContactShapesType ShapesType;
+
 
 		FString ToString() const
 		{
@@ -101,13 +111,14 @@ namespace Chaos
 		TCollisionConstraintBase(
 			FGeometryParticleHandle* Particle0, const FImplicitObject* Implicit0, const TRigidTransform<T, d>& Transform0,
 			FGeometryParticleHandle* Particle1, const FImplicitObject* Implicit1, const TRigidTransform<T, d>& Transform1,
-			FType InType, int32 InTimestamp = -INT_MAX)
+			FType InType, EContactShapesType ShapesType, int32 InTimestamp = -INT_MAX)
 			: AccumulatedImpulse(0)
 			, Timestamp(InTimestamp)
 			, Type(InType)
 		{
 			ImplicitTransform[0] = Transform0; ImplicitTransform[1] = Transform1;
 			Manifold.Implicit[0] = Implicit0; Manifold.Implicit[1] = Implicit1;
+			Manifold.ShapesType = ShapesType;
 			Particle[0] = Particle0; Particle[1] = Particle1; 
 		}
 
@@ -171,8 +182,9 @@ namespace Chaos
 		TRigidBodyPointContactConstraint() : Base(Base::FType::SinglePoint) {}
 		TRigidBodyPointContactConstraint(
 			FGeometryParticleHandle* Particle0, const FImplicitObject* Implicit0, const TRigidTransform<T, d>& Transform0,
-			FGeometryParticleHandle* Particle1, const FImplicitObject* Implicit1, const TRigidTransform<T, d>& Transform1)
-			: Base(Particle0, Implicit0, Transform0, Particle1, Implicit1, Transform1, Base::FType::SinglePoint) {}
+			FGeometryParticleHandle* Particle1, const FImplicitObject* Implicit1, const TRigidTransform<T, d>& Transform1,
+			EContactShapesType ShapesType = EContactShapesType::Unknown)
+			: Base(Particle0, Implicit0, Transform0, Particle1, Implicit1, Transform1, Base::FType::SinglePoint, ShapesType) {}
 
 		static typename Base::FType StaticType() { return Base::FType::SinglePoint; };
 	};
@@ -195,8 +207,9 @@ namespace Chaos
 		TRigidBodyMultiPointContactConstraint() : Base(Base::FType::MultiPoint), SourceNormalIndex(INDEX_NONE), PlaneNormal(0), PlanePosition(0) {}
 		TRigidBodyMultiPointContactConstraint(
 			FGeometryParticleHandle* Particle0, const FImplicitObject* Implicit0, const TRigidTransform<T, d>& Transform0,
-			FGeometryParticleHandle* Particle1, const FImplicitObject* Implicit1, const TRigidTransform<T, d>& Transform1)
-			: Base(Particle0, Implicit0, Transform0, Particle1, Implicit1, Transform1, Base::FType::MultiPoint)
+			FGeometryParticleHandle* Particle1, const FImplicitObject* Implicit1, const TRigidTransform<T, d>& Transform1,
+			EContactShapesType ShapesType = EContactShapesType::Unknown)
+			: Base(Particle0, Implicit0, Transform0, Particle1, Implicit1, Transform1, Base::FType::MultiPoint, ShapesType)
 			, SourceNormalIndex(INDEX_NONE), PlaneNormal(0), PlanePosition(0) 
 		{}
 
diff --git a/Engine/Source/Runtime/Experimental/Chaos/Public/Chaos/Evolution/PBDMinEvolution.h b/Engine/Source/Runtime/Experimental/Chaos/Public/Chaos/Evolution/PBDMinEvolution.h
--- a/Engine/Source/Runtime/Experimental/Chaos/Public/Chaos/Evolution/PBDMinEvolution.h
+++ b/Engine/Source/Runtime/Experimental/Chaos/Public/Chaos/Evolution/PBDMinEvolution.h
@@ -37,7 +37,7 @@ namespace Chaos
 		using FEvolutionCallback = TFunction<void()>;
 		using FRigidParticleSOAs = TPBDRigidsSOAs<FReal, 3>;
 
-		FPBDMinEvolution(FRigidParticleSOAs& InParticles, FCollisionDetector& InCollisionDetector);
+		FPBDMinEvolution(FRigidParticleSOAs& InParticles, FCollisionDetector& InCollisionDetector, const FReal InBoundsExtension);
 
 		void AddConstraintRule(FSimpleConstraintRule* Rule);
 
@@ -59,6 +59,11 @@ namespace Chaos
 			Gravity = G;
 		}
 
+		void SetBoundsExtension(const FReal InBoundsExtension)
+		{
+			BoundsExtension = InBoundsExtension;
+		}
+
 		void SetPostIntegrateCallback(const FEvolutionCallback& Cb)
 		{
 			PostIntegrateCallback = Cb;
@@ -98,6 +103,7 @@ namespace Chaos
 
 		int32 NumApplyIterations;
 		int32 NumApplyPushOutIterations;
+		FReal BoundsExtension;
 		FVec3 Gravity;
 
 		FEvolutionCallback PostIntegrateCallback;
diff --git a/Engine/Source/Runtime/Experimental/Chaos/Public/Chaos/GJK.h b/Engine/Source/Runtime/Experimental/Chaos/Public/Chaos/GJK.h
--- a/Engine/Source/Runtime/Experimental/Chaos/Public/Chaos/GJK.h
+++ b/Engine/Source/Runtime/Experimental/Chaos/Public/Chaos/GJK.h
@@ -188,7 +188,7 @@ namespace Chaos
 	}
 
 	template <typename T, typename TGeometryA, typename TGeometryB>
-	bool GJKPenetration(const TGeometryA& A, const TGeometryB& B, const TRigidTransform<T, 3>& BToATM, T& OutPenetration, TVec3<T>& OutClosestA, TVec3<T>& OutClosestB, TVec3<T>& OutNormal, const T InThicknessA = 0, const TVector<T, 3>& InitialDir = TVector<T, 3>(1, 0, 0), const T InThicknessB = 0, bool bNegativePenetrationSupport = false)
+	bool GJKPenetration(const TGeometryA& A, const TGeometryB& B, const TRigidTransform<T, 3>& BToATM, T& OutPenetration, TVec3<T>& OutClosestA, TVec3<T>& OutClosestB, TVec3<T>& OutNormal, const T InThicknessA = 0, const TVector<T, 3>& InitialDir = TVector<T, 3>(1, 0, 0), const T InThicknessB = 0, bool bNegativePenetrationSupport = false, int32* OutNumIterations = nullptr)
 	{
 		auto SupportAFunc = [&A](const TVec3<T>& V)
 		{
@@ -267,6 +267,11 @@ namespace Chaos
 
 		} while (!bTerminate);
 
+		if (OutNumIterations != nullptr)
+		{
+			*OutNumIterations = NumIterations;
+		}
+
 		if (PrevDist2 > Eps2)
 		{
 			//generally this happens when shapes are inflated.
diff --git a/Engine/Source/Runtime/Experimental/Chaos/Public/Chaos/PBDCollisionConstraints.h b/Engine/Source/Runtime/Experimental/Chaos/Public/Chaos/PBDCollisionConstraints.h
--- a/Engine/Source/Runtime/Experimental/Chaos/Public/Chaos/PBDCollisionConstraints.h
+++ b/Engine/Source/Runtime/Experimental/Chaos/Public/Chaos/PBDCollisionConstraints.h
@@ -64,7 +64,7 @@ public:
 	TPBDCollisionConstraints(const TPBDRigidsSOAs<T,d>& InParticles, 
 		TArrayCollectionArray<bool>& Collided, 
 		const TArrayCollectionArray<TSerializablePtr<FChaosPhysicsMaterial>>& PerParticleMaterials, 
-		const int32 ApplyPairIterations = 1, const int32 ApplyPushOutPairIterations = 1, const T Thickness = (T)0);
+		const int32 ApplyPairIterations = 1, const int32 ApplyPushOutPairIterations = 1, const T CullDistance = (T)0, const T ShapePadding = (T)0);
 
 	virtual ~TPBDCollisionConstraints() {}
 
@@ -178,9 +178,30 @@ public:
 		return Manifolds.Contains(FConstraintContainerHandle::MakeKey(Base));
 	}
 
-	void SetThickness(T InThickness)
+	// @todo(chaos): remove
+	//void SetThickness(T InThickness)
+	//{
+	//	MCullDistance = InThickness;
+	//}
+
+	void SetCullDistance(T InCullDistance)
+	{
+		MCullDistance = InCullDistance;
+	}
+
+	FReal GetCullDistance() const
+	{
+		return MCullDistance;
+	}
+
+	void SetShapePadding(T InShapePadding)
+	{
+		MShapePadding = InShapePadding;
+	}
+
+	FReal GetShapePadding() const
 	{
-		MThickness = InThickness;
+		return MShapePadding;
 	}
 
 	void SetPairIterations(int32 InPairIterations)
@@ -239,7 +260,8 @@ private:
 	const TArrayCollectionArray<TSerializablePtr<FChaosPhysicsMaterial>>& MPhysicsMaterials;
 	int32 MApplyPairIterations;
 	int32 MApplyPushOutPairIterations;
-	T MThickness;	// @todo(ccaulfield) - COLLISION thickness - this should be used as shape padding (as opposed to broad/narrowphase thickness which is for speculative creation of constraints)
+	T MCullDistance;
+	T MShapePadding;
 	T MAngularFriction;
 	bool bUseCCD;
 	bool bEnableCollisions;
diff --git a/Engine/Source/Runtime/Experimental/Chaos/Public/Chaos/PBDCollisionConstraintsContact.h b/Engine/Source/Runtime/Experimental/Chaos/Public/Chaos/PBDCollisionConstraintsContact.h
--- a/Engine/Source/Runtime/Experimental/Chaos/Public/Chaos/PBDCollisionConstraintsContact.h
+++ b/Engine/Source/Runtime/Experimental/Chaos/Public/Chaos/PBDCollisionConstraintsContact.h
@@ -11,7 +11,8 @@ namespace Chaos
 	{
 		template<class T = float>
 		struct TContactParticleParameters {
-			T Thickness;
+			T CullDistance;
+			T ShapePadding;
 			TArrayCollectionArray<bool>* Collided;
 		};
 
@@ -25,10 +26,10 @@ namespace Chaos
 		};
 
 		template<ECollisionUpdateType UpdateType, typename T, int d>
-		void Update(const T Thickness, TCollisionConstraintBase<T, d>& Constraint);
+		void Update(const T CullDistance, TCollisionConstraintBase<T, d>& Constraint);
 
 		template<typename T, int d>
-		void UpdateManifold(const T Thickness, TCollisionConstraintBase<T, d>& Constraint);
+		void UpdateManifold(const T CullDistance, TCollisionConstraintBase<T, d>& Constraint);
 
 		template<typename T, int d>
 		void Apply(TCollisionConstraintBase<T, d>& Constraint, const TContactIterationParameters<T> & IterationParameters, const TContactParticleParameters<T> & ParticleParameters);
