Collision detection: Small index and type inconsistency fixed

diff --git a/Engine/Source/Runtime/Experimental/Chaos/Public/Chaos/Simplex.h b/Engine/Source/Runtime/Experimental/Chaos/Public/Chaos/Simplex.h
--- a/Engine/Source/Runtime/Experimental/Chaos/Public/Chaos/Simplex.h
+++ b/Engine/Source/Runtime/Experimental/Chaos/Public/Chaos/Simplex.h
@@ -25,7 +25,7 @@ namespace Chaos
 
 		const T X0ToX1Squared = X0ToX1.SizeSquared();
 
-		if (X0ToX1Squared < Dot || X0ToX1Squared <= std::numeric_limits<float>::min())	//if dividing gives 1+ or the line is degenerate
+		if (X0ToX1Squared < Dot || X0ToX1Squared <= std::numeric_limits<T>::min())	//if dividing gives 1+ or the line is degenerate
 		{
 			NumVerts = 1;
 			Idxs[0] = Idxs[1];
@@ -318,7 +318,7 @@ namespace Chaos
 		switch (Idxs.NumVerts)
 		{
 		case 1:
-			OutBarycentric[0] = 1;
+			OutBarycentric[Idxs[0]] = 1;
 			ClosestPoint = Simplex[Idxs[0]]; break;
 		case 2:
 		{
