Added the ability to do a full flush when we have too many dirty elements. This is meant as an emergency flush and is quite expensive

diff --git a/Engine/Source/Programs/HeadlessChaos/Private/HeadlessChaosTestBroadphase.cpp b/Engine/Source/Programs/HeadlessChaos/Private/HeadlessChaosTestBroadphase.cpp
--- a/Engine/Source/Programs/HeadlessChaos/Private/HeadlessChaosTestBroadphase.cpp
+++ b/Engine/Source/Programs/HeadlessChaos/Private/HeadlessChaosTestBroadphase.cpp
@@ -480,6 +480,42 @@ namespace ChaosTest
 
 			SpatialTestHelper(Spatial, Boxes.Get(), Box);
 		}
+
+		{
+			//too many boxes so reoptimize
+			TUniquePtr<TBox<T,3>> Box;
+			auto Boxes = BuildBoxes<T>(Box);
+			TAABBTree<int32,TBoundingVolume<int32,T,3>,T> Spatial(MakeParticleView(Boxes.Get()));
+
+			while(!Spatial.IsAsyncTimeSlicingComplete())
+			{
+				Spatial.ProgressAsyncTimeSlicing(false);
+			}
+
+			EXPECT_EQ(Spatial.NumDirtyElements(),0);
+
+			//fill up until dirty limit
+			int32 Count;
+			for(Count = 0; Count < 10; ++Count)
+			{
+				auto Boxes2 = BuildBoxes<T>(Box);
+				for(uint32 Idx = 0; Idx < Boxes2->Size(); ++Idx)
+				{
+					Spatial.UpdateElement(Idx + Boxes->Size() * Count,Boxes2->WorldSpaceInflatedBounds(Idx),true);
+				}
+
+				EXPECT_EQ(Spatial.NumDirtyElements(), (Count+1)*Boxes->Size());
+			}
+
+			//finally pass dirty limit so reset to 0 and then add the remaining new boxes
+			auto Boxes2 = BuildBoxes<T>(Box);
+			for(uint32 Idx = 0; Idx < Boxes2->Size(); ++Idx)
+			{
+				Spatial.UpdateElement(Idx + Boxes->Size() * Count,Boxes2->WorldSpaceInflatedBounds(Idx),true);
+			}
+
+			EXPECT_EQ(Spatial.NumDirtyElements(),Boxes->Size() - 1);
+		}
 	}
 	template void AABBTreeTest<float>();
 
diff --git a/Engine/Source/Runtime/Experimental/Chaos/Private/Chaos/BoundingVolume.cpp b/Engine/Source/Runtime/Experimental/Chaos/Private/Chaos/BoundingVolume.cpp
--- a/Engine/Source/Runtime/Experimental/Chaos/Private/Chaos/BoundingVolume.cpp
+++ b/Engine/Source/Runtime/Experimental/Chaos/Private/Chaos/BoundingVolume.cpp
@@ -15,6 +15,12 @@ FAutoConsoleVariableRef FBoundingVolumeCVars::CVarFilterFarBodies(
 
 namespace Chaos
 {
+	CHAOS_API int32 MaxDirtyElements = 10000;
+	FAutoConsoleVariableRef CVarMaxDirtyElements(
+		TEXT("p.MaxDirtyElements"),
+		MaxDirtyElements,
+		TEXT("The max number of dirty elements. This forces a flush which is very expensive"));
+
 	template <typename TPayloadType, typename T, int d>
 	ISpatialAcceleration<TPayloadType, T, d>* ISpatialAcceleration<TPayloadType, T, d>::SerializationFactory(FChaosArchive& Ar, ISpatialAcceleration<TPayloadType, T, d>* Accel)
 	{
diff --git a/Engine/Source/Runtime/Experimental/Chaos/Public/Chaos/AABBTree.h b/Engine/Source/Runtime/Experimental/Chaos/Public/Chaos/AABBTree.h
--- a/Engine/Source/Runtime/Experimental/Chaos/Public/Chaos/AABBTree.h
+++ b/Engine/Source/Runtime/Experimental/Chaos/Public/Chaos/AABBTree.h
@@ -79,6 +79,11 @@ struct TAABBTreeLeafArray
 	{
 	}
 
+	void GatherElements(TArray<TPayloadBoundsElement<TPayloadType, T>>& OutElements)
+	{
+		OutElements.Append(Elems);
+	}
+
 	template <typename TSQVisitor, typename TQueryFastData>
 	bool RaycastFast(const TVector<T,3>& Start, TQueryFastData& QueryFastData, TSQVisitor& Visitor) const
 	{
@@ -219,6 +224,8 @@ inline FArchive& operator<<(FArchive& Ar, FAABBTreePayloadInfo& PayloadInfo)
 	return Ar;
 }
 
+extern CHAOS_API int32 MaxDirtyElements;
+
 template <typename TPayloadType, typename TLeafType, typename T, bool bMutable = true>
 class TAABBTree final : public ISpatialAcceleration<TPayloadType, T, 3>
 {
@@ -489,6 +496,16 @@ public:
 				}
 			}
 		}
+
+		if(DirtyElements.Num() > MaxDirtyElements)
+		{
+			ReoptimizeTree();
+		}
+	}
+
+	int32 NumDirtyElements() const
+	{
+		return DirtyElements.Num();
 	}
 
 	const TArray<TPayloadBoundsElement<TPayloadType, T>>& GlobalObjects() const
@@ -536,6 +553,21 @@ private:
 	using FElement = TPayloadBoundsElement<TPayloadType, T>;
 	using FNode = TAABBTreeNode<T>;
 
+	void ReoptimizeTree()
+	{
+		TArray<FElement> AllElements;
+		AllElements.Append(DirtyElements);
+		AllElements.Append(GlobalPayloads);
+
+		for(auto& Leaf : Leaves)
+		{
+			Leaf.GatherElements(AllElements);
+		}
+
+		TAABBTree<TPayloadType,TLeafType,T,bMutable> NewTree(AllElements);
+		*this = NewTree;
+	}
+
 	template <EAABBQueryType Query, typename TQueryFastData, typename SQVisitor>
 	bool QueryImp(const TVector<T, 3>& Start, TQueryFastData& CurData, const TVector<T, 3> QueryHalfExtents, const TAABB<T,3>& QueryBounds, SQVisitor& Visitor) const
 	{
@@ -954,6 +986,28 @@ private:
 
 	}
 
+	TAABBTree<TPayloadType,TLeafType,T,bMutable>& operator=(const TAABBTree<TPayloadType,TLeafType,T,bMutable>& Rhs)
+	{
+		ensure(Rhs.TimeSliceWorkToComplete.IsEmpty());
+		TimeSliceWorkToComplete.Empty();
+		if(this != &Rhs)
+		{
+			FullBounds = Rhs.FullBounds;
+			Nodes = Rhs.Nodes;
+			Leaves = Rhs.Leaves;
+			DirtyElements = Rhs.DirtyElements;
+			GlobalPayloads = Rhs.GlobalPayloads;
+			PayloadToInfo = Rhs.PayloadToInfo;
+			MaxChildrenInLeaf = Rhs.MaxChildrenInLeaf;
+			MaxTreeDepth = Rhs.MaxTreeDepth;
+			MaxPayloadBounds = Rhs.MaxPayloadBounds;
+			MaxNumToProcess = Rhs.MaxNumToProcess;
+			NumProcessedThisSlice = Rhs.NumProcessedThisSlice;
+		}
+
+		return *this;
+	}
+
 	struct FSplitInfo
 	{
 		TAABB<T, 3> SplitBounds;	//Even split of parent bounds
diff --git a/Engine/Source/Runtime/Experimental/Chaos/Public/Chaos/BoundingVolume.h b/Engine/Source/Runtime/Experimental/Chaos/Public/Chaos/BoundingVolume.h
--- a/Engine/Source/Runtime/Experimental/Chaos/Public/Chaos/BoundingVolume.h
+++ b/Engine/Source/Runtime/Experimental/Chaos/Public/Chaos/BoundingVolume.h
@@ -343,6 +343,25 @@ public:
 
 	}
 
+	void GatherElements(TArray<TPayloadBoundsElement<TPayloadType, T>>& OutElements)
+	{
+		OutElements.Append(MGlobalPayloads);
+		OutElements.Reserve(OutElements.Num() + MDirtyElements.Num());
+		for(const FCellElement& Elem : MDirtyElements)
+		{
+			OutElements.Add(TPayloadBoundsElement<TPayloadType,T>{Elem.Payload,Elem.Bounds});
+		}
+
+		for(int32 Idx = 0; Idx < MElements.Num(); ++Idx)
+		{
+			const TArray<FCellElement>& Elems = MElements[Idx];
+			for(const FCellElement& Elem : Elems)
+			{
+				OutElements.Add(TPayloadBoundsElement<TPayloadType,T>{Elem.Payload,Elem.Bounds});
+			}
+		}
+	}
+
 private:
 
 	using FCellElement = TBVCellElement<TPayloadType, T, d>;
