Fix ParticleDirtyFlags when using CHAOS_CHECKED

diff --git a/Engine/Source/Runtime/Experimental/Chaos/Public/Chaos/ParticleDirtyFlags.h b/Engine/Source/Runtime/Experimental/Chaos/Public/Chaos/ParticleDirtyFlags.h
--- a/Engine/Source/Runtime/Experimental/Chaos/Public/Chaos/ParticleDirtyFlags.h
+++ b/Engine/Source/Runtime/Experimental/Chaos/Public/Chaos/ParticleDirtyFlags.h
@@ -28,7 +28,7 @@ namespace Chaos
 		M						= 1 << 16,
 		InvM					= 1 << 17,
 		LinearEtherDrag			= 1 << 18,
-		AngularEtherDrag			= 1 << 19,
+		AngularEtherDrag		= 1 << 19,
 		ObjectState				= 1 << 20,
 		Geometry				= 1 << 21,
 		LinearImpulse			= 1 << 22,
@@ -36,7 +36,7 @@ namespace Chaos
 		GravityEnabled			= 1 << 24,
 		SpatialIdx				= 1 << 25,
 		UniqueIdx				= 1 << 26,
-		ShapeDisableCollision	= 1 << 27,
+		ShapeDisableCollision	= 1 << 27
 #if CHAOS_CHECKED
 		, DebugName				= 1 << 28
 #endif
