Chaos: Add support for bDoubleSided flag (Fixes Tri-meshes that rely on backface culling raycasts)

diff --git a/Engine/Source/Runtime/Engine/Private/PhysicsEngine/Experimental/ChaosInterfaceUtils.cpp b/Engine/Source/Runtime/Engine/Private/PhysicsEngine/Experimental/ChaosInterfaceUtils.cpp
--- a/Engine/Source/Runtime/Engine/Private/PhysicsEngine/Experimental/ChaosInterfaceUtils.cpp
+++ b/Engine/Source/Runtime/Engine/Private/PhysicsEngine/Experimental/ChaosInterfaceUtils.cpp
@@ -305,7 +305,7 @@ namespace ChaosInterface
 
 		if (bMakeComplexGeometry)
 		{
-			for (const auto& ChaosTriMesh : InParams.ChaosTriMeshes)
+			for (auto& ChaosTriMesh : InParams.ChaosTriMeshes)
 			{
 				TUniquePtr<Chaos::FImplicitObject> Implicit;
 				if (Scale == FVector(1))
@@ -317,6 +317,8 @@ namespace ChaosInterface
 					Implicit = TUniquePtr<Chaos::FImplicitObject>(new Chaos::TImplicitObjectScaled<Chaos::FTriangleMeshImplicitObject>(ChaosTriMesh, Scale));
 				}
 
+				ChaosTriMesh->SetCullsBackFaceRaycast(!InParams.bDoubleSided);
+
 				auto NewShape = NewShapeHelper(MakeSerializable(Implicit), nullptr, true);
 				Shapes.Emplace(MoveTemp(NewShape));
 				Geoms.Add(MoveTemp(Implicit));
diff --git a/Engine/Source/Runtime/Experimental/Chaos/Private/Chaos/TriangleMeshImplicitObject.cpp b/Engine/Source/Runtime/Experimental/Chaos/Private/Chaos/TriangleMeshImplicitObject.cpp
--- a/Engine/Source/Runtime/Experimental/Chaos/Private/Chaos/TriangleMeshImplicitObject.cpp
+++ b/Engine/Source/Runtime/Experimental/Chaos/Private/Chaos/TriangleMeshImplicitObject.cpp
@@ -13,13 +13,14 @@ namespace Chaos
 template <typename IdxType>
 struct FTriangleMeshRaycastVisitor
 {
-	FTriangleMeshRaycastVisitor(const FVec3& InStart, const FVec3& InDir, const FReal InThickness, const TParticles<FReal,3>& InParticles, const TArray<TVector<IdxType, 3>>& InElements)
+	FTriangleMeshRaycastVisitor(const FVec3& InStart, const FVec3& InDir, const FReal InThickness, const TParticles<FReal,3>& InParticles, const TArray<TVector<IdxType, 3>>& InElements, bool bInCullsBackFaceRaycast)
 	: Particles(InParticles)
 	, Elements(InElements)
 	, StartPoint(InStart)
 	, Dir(InDir)
 	, Thickness(InThickness)
 	, OutTime(TNumericLimits<FReal>::Max())
+	, bCullsBackFaceRaycast(bInCullsBackFaceRaycast)
 	{
 	}
 
@@ -54,6 +55,12 @@ struct FTriangleMeshRaycastVisitor
 			return true;
 		}
 
+		const bool bBackFace = (FVec3::DotProduct(Dir, TriNormal) > 0.0f);
+		if (bCullsBackFaceRaycast && bBackFace)
+		{
+			return true;
+		}
+
 		const TPlane<FReal, 3> TriPlane{ A, TriNormal };
 		FVec3 RaycastPosition;
 		FVec3 RaycastNormal;
@@ -185,6 +192,7 @@ struct FTriangleMeshRaycastVisitor
 	FVec3 OutPosition;
 	FVec3 OutNormal;
 	int32 OutFaceIndex;
+	bool bCullsBackFaceRaycast;
 };
 
 FReal FTriangleMeshImplicitObject::PhiWithNormal(const FVec3& x, FVec3& Normal) const
@@ -200,7 +208,7 @@ FReal FTriangleMeshImplicitObject::PhiWithNormal(const FVec3& x, FVec3& Normal)
 template <typename IdxType>
 bool FTriangleMeshImplicitObject::RaycastImp(const TArray<TVector<IdxType, 3>>& Elements, const FVec3& StartPoint, const FVec3& Dir, const FReal Length, const FReal Thickness, FReal& OutTime, FVec3& OutPosition, FVec3& OutNormal, int32& OutFaceIndex) const
 {
-	FTriangleMeshRaycastVisitor<IdxType> SQVisitor(StartPoint, Dir, Thickness, MParticles, Elements);
+	FTriangleMeshRaycastVisitor<IdxType> SQVisitor(StartPoint, Dir, Thickness, MParticles, Elements, bCullsBackFaceRaycast);
 
 	if (Thickness > 0)
 	{
@@ -368,6 +376,16 @@ int32 FTriangleMeshImplicitObject::GetExternalFaceIndexFromInternal(int32 Intern
 	return -1;
 }
 
+bool FTriangleMeshImplicitObject::GetCullsBackFaceRaycast() const
+{
+	return bCullsBackFaceRaycast;
+}
+
+void FTriangleMeshImplicitObject::SetCullsBackFaceRaycast(const bool bInCullsBackFace)
+{
+	bCullsBackFaceRaycast = bInCullsBackFace;
+}
+
 template <typename IdxType>
 bool FTriangleMeshImplicitObject::OverlapImp(const TArray<TVec3<IdxType>>& Elements, const FVec3& Point, const FReal Thickness) const
 {
diff --git a/Engine/Source/Runtime/Experimental/Chaos/Public/Chaos/TriangleMeshImplicitObject.h b/Engine/Source/Runtime/Experimental/Chaos/Public/Chaos/TriangleMeshImplicitObject.h
--- a/Engine/Source/Runtime/Experimental/Chaos/Public/Chaos/TriangleMeshImplicitObject.h
+++ b/Engine/Source/Runtime/Experimental/Chaos/Public/Chaos/TriangleMeshImplicitObject.h
@@ -101,13 +101,14 @@ namespace Chaos
 		using FImplicitObject::GetTypeName;
 
 		template <typename IdxType>
-		FTriangleMeshImplicitObject(TParticles<FReal, 3>&& Particles, TArray<TVector<IdxType, 3>>&& Elements, TArray<uint16>&& InMaterialIndices, TUniquePtr<TArray<int32>>&& InExternalFaceIndexMap = nullptr)
+		FTriangleMeshImplicitObject(TParticles<FReal, 3>&& Particles, TArray<TVector<IdxType, 3>>&& Elements, TArray<uint16>&& InMaterialIndices, TUniquePtr<TArray<int32>>&& InExternalFaceIndexMap = nullptr, const bool bInCullsBackFaceRaycast = false)
 		: FImplicitObject(EImplicitObject::HasBoundingBox | EImplicitObject::DisableCollisions, ImplicitObjectType::TriangleMesh)
 		, MParticles(MoveTemp(Particles))
 		, MElements(MoveTemp(Elements))
 		, MLocalBoundingBox(MParticles.X(0), MParticles.X(0))
 		, MaterialIndices(MoveTemp(InMaterialIndices))
 		, ExternalFaceIndexMap(MoveTemp(InExternalFaceIndexMap))
+		, bCullsBackFaceRaycast(bInCullsBackFaceRaycast)
 		{
 			for (uint32 Idx = 1; Idx < MParticles.Size(); ++Idx)
 			{
@@ -156,6 +157,10 @@ namespace Chaos
 		// Returns -1 if InternalFaceIndex is not in map, or map is invalid.
 		int32 GetExternalFaceIndexFromInternal(int32 InternalFaceIndex) const;
 
+		// Does Trimesh cull backfaces in raycast.
+		bool GetCullsBackFaceRaycast() const;
+		void SetCullsBackFaceRaycast(const bool bInCullsBackFace);
+
 
 		virtual int32 FindMostOpposingFace(const FVec3& Position, const FVec3& UnitDir, int32 HintFaceIndex, FReal SearchDistance) const override;
 		virtual FVec3 FindGeometryOpposingNormal(const FVec3& DenormDir, int32 FaceIndex, const FVec3& OriginalNormal) const override;
@@ -255,6 +260,7 @@ namespace Chaos
 		TAABB<FReal, 3> MLocalBoundingBox;
 		TArray<uint16> MaterialIndices;
 		TUniquePtr<TArray<int32>> ExternalFaceIndexMap;
+		bool bCullsBackFaceRaycast;
 
 		//using BVHType = TBoundingVolume<int32, T, 3>;
 		using BVHType = TAABBTree<int32, TAABBTreeLeafArray<int32, FReal>, FReal, /*bMutable=*/false>;
