Handle precision issues for raycast against scaled object

diff --git a/Engine/Source/Runtime/Experimental/Chaos/Public/Chaos/ImplicitObjectScaled.h b/Engine/Source/Runtime/Experimental/Chaos/Public/Chaos/ImplicitObjectScaled.h
--- a/Engine/Source/Runtime/Experimental/Chaos/Public/Chaos/ImplicitObjectScaled.h
+++ b/Engine/Source/Runtime/Experimental/Chaos/Public/Chaos/ImplicitObjectScaled.h
@@ -326,14 +326,14 @@ public:
 
 			if (MObject->Raycast(UnscaledStart, UnscaledDir, UnscaledLength, MInternalThickness + Thickness * MInvScale[0], UnscaledTime, UnscaledPosition, UnscaledNormal, OutFaceIndex))
 			{
-				OutTime = LengthScaleInv * UnscaledTime;
-				if (OutTime != 0) // Normal/Position output may be uninitialized with TOI 0.
+				//We double check that NewTime < Length because of potential precision issues. When that happens we always keep the shortest hit first
+				const T NewTime = LengthScaleInv * UnscaledTime;
+				if (NewTime < Length && NewTime != 0) // Normal/Position output may be uninitialized with TOI 0.
 				{
 					OutPosition = MScale * UnscaledPosition;
 					OutNormal = (MInvScale * UnscaledNormal).GetSafeNormal();
+					return true;
 				}
-				CHAOS_ENSURE(OutTime <= Length);
-				return true;
 			}
 		}
 			
