Fix EPA and GJK when we have a perfect capsule capsule overlap

diff --git a/Engine/Source/Programs/HeadlessChaos/Private/HeadlessChaosTestEPA.cpp b/Engine/Source/Programs/HeadlessChaos/Private/HeadlessChaosTestEPA.cpp
--- a/Engine/Source/Programs/HeadlessChaos/Private/HeadlessChaosTestEPA.cpp
+++ b/Engine/Source/Programs/HeadlessChaos/Private/HeadlessChaosTestEPA.cpp
@@ -96,7 +96,9 @@ namespace ChaosTest
 		{
 			TArray<TVec3<T>> VertsA = { {-1,-1,1}, {-1,-1,-1}, {-1,1,-1}, {1,1,-1} };
 			TArray<TVec3<T>> VertsB = { TVec3<T>(0), TVec3<T>(0), TVec3<T>(0), TVec3<T>(0) };
-			TArray<TEPAEntry<T>> TetFaces = InitializeEPA(VertsA, VertsB, ErrorSupport<T>, ErrorSupport<T>);
+			TArray<TEPAEntry<T>> TetFaces;
+			TVec3<T> TouchingNormal;
+			EXPECT_TRUE(InitializeEPA(VertsA,VertsB,ErrorSupport<T>,ErrorSupport<T>,TetFaces,TouchingNormal));
 
 			EXPECT_EQ(TetFaces.Num(), 4);
 			for (int i = 0; i < TetFaces.Num(); ++i)
@@ -108,7 +110,10 @@ namespace ChaosTest
 		{
 			TArray<TVec3<T>> VertsA = { {-1,-1,-1}, {-1,-1,1}, {-1,1,-1}, {1,1,-1} };
 			TArray<TVec3<T>> VertsB = { TVec3<T>(0), TVec3<T>(0), TVec3<T>(0), TVec3<T>(0) };
-			TArray<TEPAEntry<T>> TetFaces = InitializeEPA(VertsA, VertsB, ErrorSupport<T>, ErrorSupport<T>);
+			TArray<TEPAEntry<T>> TetFaces;
+			TVec3<T> TouchingNormal;
+
+			EXPECT_TRUE(InitializeEPA(VertsA,VertsB,ErrorSupport<T>,ErrorSupport<T>,TetFaces,TouchingNormal));
 
 			EXPECT_EQ(TetFaces.Num(), 4);
 			for (int i = 0; i < TetFaces.Num(); ++i)
@@ -175,7 +180,9 @@ namespace ChaosTest
 				TArray<TVec3<T>> VertsA = { AllVerts[0], AllVerts[1], AllVerts[2] };
 				TArray<TVec3<T>> VertsB = { TVec3<T>(0), TVec3<T>(0), TVec3<T>(0) };
 
-				TArray<TEPAEntry<T>> TetFaces = InitializeEPA(VertsA, VertsB, ASupport, EmptySupport);
+				TArray<TEPAEntry<T>> TetFaces;
+				TVec3<T> TouchingNormal;
+				EXPECT_TRUE(InitializeEPA(VertsA,VertsB,ASupport,EmptySupport,TetFaces,TouchingNormal));
 				EXPECT_VECTOR_NEAR(VertsA[3], AllVerts[3], 1e-4);
 				EXPECT_VECTOR_NEAR(VertsB[3], TVec3<T>(0), 1e-4);
 
@@ -201,7 +208,9 @@ namespace ChaosTest
 				TArray<TVec3<T>> VertsA = { AllVerts[1], AllVerts[0], AllVerts[2] };
 				TArray<TVec3<T>> VertsB = { TVec3<T>(0), TVec3<T>(0), TVec3<T>(0) };
 
-				TArray<TEPAEntry<T>> TetFaces = InitializeEPA(VertsA, VertsB, ASupport, EmptySupport);
+				TArray<TEPAEntry<T>> TetFaces;
+				TVec3<T> TouchingNormal;
+				EXPECT_TRUE(InitializeEPA(VertsA,VertsB,ASupport,EmptySupport,TetFaces,TouchingNormal));
 				EXPECT_VECTOR_NEAR(VertsA[3], AllVerts[3], 1e-4);
 				EXPECT_VECTOR_NEAR(VertsB[3], TVec3<T>(0), 1e-4);
 
@@ -227,9 +236,11 @@ namespace ChaosTest
 				TArray<TVec3<T>> VertsA = { AllVerts[1], AllVerts[0], AllVerts[2] };
 				TArray<TVec3<T>> VertsB = { TVec3<T>(0), TVec3<T>(0), TVec3<T>(0) };
 
-				TArray<TEPAEntry<T>> TetFaces = InitializeEPA(VertsA, VertsB, ASupportNoX, EmptySupport);
-				EXPECT_EQ(TetFaces.Num(), 0);
-
+				TArray<TEPAEntry<T>> TetFaces;
+				TVec3<T> TouchingNormal;
+				EXPECT_FALSE(InitializeEPA(VertsA,VertsB,ASupportNoX,EmptySupport,TetFaces,TouchingNormal));
+				EXPECT_EQ(TouchingNormal.Z,0);
+				EXPECT_EQ(TouchingNormal.Y,0);
 
 				//make sure EPA handles this bad case properly
 				VertsA = { AllVerts[1], AllVerts[0], AllVerts[2] };
@@ -240,7 +251,7 @@ namespace ChaosTest
 				TVec3<T> Dir, WitnessA, WitnessB;
 				EXPECT_EQ(EPA(VertsA, VertsB, ASupportNoX, EmptySupport, Penetration, Dir, WitnessA, WitnessB), EPAResult::BadInitialSimplex);
 				EXPECT_EQ(Penetration, 0);
-				EXPECT_VECTOR_NEAR(Dir, TVec3<T>(0, 0, 1), 1e-7);
+				//EXPECT_VECTOR_NEAR(Dir, TVec3<T>(0, 0, 1), 1e-7);
 				EXPECT_VECTOR_NEAR(WitnessA, TVec3<T>(0), 1e-7);
 				EXPECT_VECTOR_NEAR(WitnessB, TVec3<T>(0), 1e-7);
 			}
@@ -268,7 +279,9 @@ namespace ChaosTest
 				TArray<TVec3<T>> VertsA = { AllVerts[0], AllVerts[2] };
 				TArray<TVec3<T>> VertsB = { TVec3<T>(0), TVec3<T>(0) };
 
-				TArray<TEPAEntry<T>> TetFaces = InitializeEPA(VertsA, VertsB, ASupport, EmptySupport);
+				TArray<TEPAEntry<T>> TetFaces;
+				TVec3<T> TouchingNormal;
+				EXPECT_TRUE(InitializeEPA(VertsA,VertsB,ASupport,EmptySupport,TetFaces,TouchingNormal));
 				EXPECT_VECTOR_NEAR(VertsA[2], AllVerts[1], 1e-4);
 				EXPECT_VECTOR_NEAR(VertsB[2], TVec3<T>(0), 1e-4);
 
@@ -287,7 +300,9 @@ namespace ChaosTest
 				TArray<TVec3<T>> VertsA = { AllVerts[2], AllVerts[0] };
 				TArray<TVec3<T>> VertsB = { TVec3<T>(0), TVec3<T>(0) };
 
-				TArray<TEPAEntry<T>> TetFaces = InitializeEPA(VertsA, VertsB, ASupport, EmptySupport);
+				TArray<TEPAEntry<T>> TetFaces;
+				TVec3<T> TouchingNormal;
+				EXPECT_TRUE(InitializeEPA(VertsA,VertsB,ASupport,EmptySupport,TetFaces,TouchingNormal));
 				EXPECT_VECTOR_NEAR(VertsA[2], AllVerts[1], 1e-4);
 				EXPECT_VECTOR_NEAR(VertsB[2], TVec3<T>(0), 1e-4);
 
@@ -323,8 +338,10 @@ namespace ChaosTest
 				TArray<TVec3<T>> VertsA = { AllVerts[2], AllVerts[0] };
 				TArray<TVec3<T>> VertsB = { TVec3<T>(0), TVec3<T>(0) };
 
-				TArray<TEPAEntry<T>> TetFaces = InitializeEPA(VertsA, VertsB, ASupportNoX, EmptySupport);
-				EXPECT_EQ(TetFaces.Num(), 0);
+				TArray<TEPAEntry<T>> TetFaces;
+				TVec3<T> TouchingNormal;
+				EXPECT_FALSE(InitializeEPA(VertsA,VertsB,ASupportNoX,EmptySupport, TetFaces, TouchingNormal));
+				EXPECT_EQ(TouchingNormal.X,0);
 			}
 
 			//touching line
@@ -350,8 +367,10 @@ namespace ChaosTest
 				TArray<TVec3<T>> VertsA = { AllVerts[2], AllVerts[0] };
 				TArray<TVec3<T>> VertsB = { TVec3<T>(0), TVec3<T>(0) };
 
-				TArray<TEPAEntry<T>> TetFaces = InitializeEPA(VertsA, VertsB, ASupportNoXOrZ, EmptySupport);
-				EXPECT_EQ(TetFaces.Num(), 0);
+				TArray<TEPAEntry<T>> TetFaces;
+				TVec3<T> TouchingNormal;
+				EXPECT_FALSE(InitializeEPA(VertsA,VertsB,ASupportNoXOrZ,EmptySupport,TetFaces,TouchingNormal));
+				EXPECT_EQ(TouchingNormal.X,0);
 			}
 		}
 	}
diff --git a/Engine/Source/Programs/HeadlessChaos/Private/HeadlessChaosTestGJK.cpp b/Engine/Source/Programs/HeadlessChaos/Private/HeadlessChaosTestGJK.cpp
--- a/Engine/Source/Programs/HeadlessChaos/Private/HeadlessChaosTestGJK.cpp
+++ b/Engine/Source/Programs/HeadlessChaos/Private/HeadlessChaosTestGJK.cpp
@@ -1473,6 +1473,26 @@ namespace ChaosTest
 			EXPECT_FALSE((GJKPenetration<false, T>(A,B,NewAToBTM,Penetration,ClosestA,ClosestB,Normal,0,Offset,0)));
 
 		}
+
+		{
+			//capsule perfectly aligned with another capsule but a bit off on the z
+			const TVec3<T> Pt0(0.0,0.0,-45.0);
+			TVec3<T> Pt1 = Pt0;
+			Pt1 += (TVec3<T>(0.0,0.0,1.0) * 90.0);
+
+			const TCapsule<T> A(Pt0,Pt1,34.f);
+			const TCapsule<T> B(Pt0,Pt1,33.8499985f);
+
+			const TRigidTransform<T,3> BToATM(TVec3<T>(0.0f,0.0f,-23.4092140f),TRotation<T,3>::FromElements(0.0,0.0,0.0,1.0));
+
+			EXPECT_TRUE(GJKIntersection<T>(A,B,BToATM,0.0,TVec3<T>(0,0,23.4092140)));
+
+			T Penetration;
+			TVec3<T> ClosestA,ClosestB,Normal;
+			EXPECT_TRUE(GJKPenetration<T>(A,B,BToATM, Penetration, ClosestA, ClosestB, Normal, 0.0,TVec3<T>(0,0,23.4092140)));
+			EXPECT_FLOAT_EQ(Normal.Z,0);
+			EXPECT_FLOAT_EQ(Penetration,A.GetRadius() + B.GetRadius());
+		}
 		
 	}
 
diff --git a/Engine/Source/Runtime/Experimental/Chaos/Public/Chaos/EPA.h b/Engine/Source/Runtime/Experimental/Chaos/Public/Chaos/EPA.h
--- a/Engine/Source/Runtime/Experimental/Chaos/Public/Chaos/EPA.h
+++ b/Engine/Source/Runtime/Experimental/Chaos/Public/Chaos/EPA.h
@@ -122,9 +122,8 @@ struct TEPAEntry
 };
 
 template <typename T, typename SupportALambda, typename SupportBLambda >
-TArray<TEPAEntry<T>> InitializeEPA(TArray<TVec3<T>>& VertsA, TArray<TVec3<T>>& VertsB, const SupportALambda& SupportA, const SupportBLambda& SupportB)
+bool InitializeEPA(TArray<TVec3<T>>& VertsA, TArray<TVec3<T>>& VertsB, const SupportALambda& SupportA, const SupportBLambda& SupportB, TArray<TEPAEntry<T>>& OutEntries, TVec3<T>& OutTouchNormal)
 {
-	TArray<TEPAEntry<T>> Entries;
 	const int32 NumVerts = VertsA.Num();
 	check(VertsB.Num() == NumVerts);
 
@@ -154,7 +153,8 @@ TArray<TEPAEntry<T>> InitializeEPA(TArray<TVec3<T>>& VertsA, TArray<TVec3<T>>& V
 		}
 	};
 
-	Entries.AddUninitialized(4);
+	OutEntries.AddUninitialized(4);
+	OutTouchNormal = TVec3<T>(0,0,1);
 
 	bool bValid = false;
 
@@ -163,7 +163,7 @@ TArray<TEPAEntry<T>> InitializeEPA(TArray<TVec3<T>>& VertsA, TArray<TVec3<T>>& V
 		case 1:
 		{
 			//assuming it's a touching hit at origin
-			break;
+			return false;
 		}
 		case 2:
 		{
@@ -192,35 +192,47 @@ TArray<TEPAEntry<T>> InitializeEPA(TArray<TVec3<T>>& VertsA, TArray<TVec3<T>>& V
 				AddFartherPoint(Orthog);
 				AddFartherPoint(Orthog2);
 
-				bValid = Entries[0].Initialize(VertsA.GetData(), VertsB.GetData(), 1, 2, 3, { 3, 1, 2 }, { 1,1, 1 });
-				bValid &= Entries[1].Initialize(VertsA.GetData(), VertsB.GetData(), 0, 3, 2, { 2,0,3 }, { 2, 1, 0 });
-				bValid &= Entries[2].Initialize(VertsA.GetData(), VertsB.GetData(), 0, 1, 3, { 3,0, 1 }, { 2,2,0 });
-				bValid &= Entries[3].Initialize(VertsA.GetData(), VertsB.GetData(), 0, 2, 1, { 1,0,2 }, { 2,0,0 });
+				bValid = OutEntries[0].Initialize(VertsA.GetData(), VertsB.GetData(), 1, 2, 3, { 3, 1, 2 }, { 1,1, 1 });
+				bValid &= OutEntries[1].Initialize(VertsA.GetData(), VertsB.GetData(), 0, 3, 2, { 2,0,3 }, { 2, 1, 0 });
+				bValid &= OutEntries[2].Initialize(VertsA.GetData(), VertsB.GetData(), 0, 1, 3, { 3,0, 1 }, { 2,2,0 });
+				bValid &= OutEntries[3].Initialize(VertsA.GetData(), VertsB.GetData(), 0, 2, 1, { 1,0,2 }, { 2,0,0 });
+
+				if(!bValid)
+				{
+					OutTouchNormal = Orthog.GetUnsafeNormal();
+					return false;
+				}
 			}
 			break;
 		}
 		case 3:
 		{
 			//triangle, add farthest point along normal
-			bValid = Entries[3].Initialize(VertsA.GetData(), VertsB.GetData(), 0, 2, 1, { 1,0,2 }, { 2,0,0 });	
+			bValid = OutEntries[3].Initialize(VertsA.GetData(), VertsB.GetData(), 0, 2, 1, { 1,0,2 }, { 2,0,0 });	
 			if (ensure(bValid)) //input verts must form a valid triangle
 			{
-				const TEPAEntry<T>& Base = Entries[3];
+				const TEPAEntry<T>& Base = OutEntries[3];
 
 				AddFartherPoint(Base.PlaneNormal);
 
-				bValid = Entries[0].Initialize(VertsA.GetData(), VertsB.GetData(), 1, 2, 3, { 3, 1, 2 }, { 1,1, 1 });
-				bValid &= Entries[1].Initialize(VertsA.GetData(), VertsB.GetData(), 0, 3, 2, { 2,0,3 }, { 2, 1, 0 });
-				bValid &= Entries[2].Initialize(VertsA.GetData(), VertsB.GetData(), 0, 1, 3, { 3,0, 1 }, { 2,2,0 });
+				bValid = OutEntries[0].Initialize(VertsA.GetData(), VertsB.GetData(), 1, 2, 3, { 3, 1, 2 }, { 1,1, 1 });
+				bValid &= OutEntries[1].Initialize(VertsA.GetData(), VertsB.GetData(), 0, 3, 2, { 2,0,3 }, { 2, 1, 0 });
+				bValid &= OutEntries[2].Initialize(VertsA.GetData(), VertsB.GetData(), 0, 1, 3, { 3,0, 1 }, { 2,2,0 });
+
+				if(!bValid)
+				{
+					OutTouchNormal = Base.PlaneNormal;
+					return false;
+				}
 			}
 			break;
 		}
 		case 4:
 		{
-			bValid = Entries[0].Initialize(VertsA.GetData(), VertsB.GetData(), 1, 2, 3, { 3, 1, 2 }, { 1,1, 1 });
-			bValid &= Entries[1].Initialize(VertsA.GetData(), VertsB.GetData(), 0, 3, 2, { 2,0,3 }, { 2, 1, 0 });
-			bValid &= Entries[2].Initialize(VertsA.GetData(), VertsB.GetData(), 0, 1, 3, { 3,0, 1 }, { 2,2,0 });
-			bValid &= Entries[3].Initialize(VertsA.GetData(), VertsB.GetData(), 0, 2, 1, { 1,0,2 }, { 2,0,0 });
+			bValid = OutEntries[0].Initialize(VertsA.GetData(), VertsB.GetData(), 1, 2, 3, { 3, 1, 2 }, { 1,1, 1 });
+			bValid &= OutEntries[1].Initialize(VertsA.GetData(), VertsB.GetData(), 0, 3, 2, { 2,0,3 }, { 2, 1, 0 });
+			bValid &= OutEntries[2].Initialize(VertsA.GetData(), VertsB.GetData(), 0, 1, 3, { 3,0, 1 }, { 2,2,0 });
+			bValid &= OutEntries[3].Initialize(VertsA.GetData(), VertsB.GetData(), 0, 2, 1, { 1,0,2 }, { 2,0,0 });
 			ensure(bValid);	//expect user to give us valid tetrahedron
 			break;
 		}
@@ -231,20 +243,16 @@ TArray<TEPAEntry<T>> InitializeEPA(TArray<TVec3<T>>& VertsA, TArray<TVec3<T>>& V
 	if (bValid)
 	{
 		//make sure normals are pointing out of tetrahedron
-		if (TVec3<T>::DotProduct(Entries[0].PlaneNormal, MinkowskiVert(VertsA.GetData(), VertsB.GetData(), 0)) > 0)
+		if (TVec3<T>::DotProduct(OutEntries[0].PlaneNormal, MinkowskiVert(VertsA.GetData(), VertsB.GetData(), 0)) > 0)
 		{
-			for (TEPAEntry<T>& Entry : Entries)
+			for (TEPAEntry<T>& Entry : OutEntries)
 			{
-				Entry.SwapWinding(Entries.GetData());
+				Entry.SwapWinding(OutEntries.GetData());
 			}
 		}
 	}
-	else
-	{
-		Entries.SetNum(0);
-	}
-
-	return Entries;
+	
+	return bValid;
 }
 
 struct FEPAFloodEntry
@@ -352,19 +360,16 @@ EPAResult EPA(TArray<TVec3<T>>& VertsABuffer, TArray<TVec3<T>>& VertsBBuffer, co
 
 	constexpr T Eps = 1e-2;
 
-	TArray<TEPAEntry<T>> Entries = InitializeEPA(VertsABuffer, VertsBBuffer, SupportA, SupportB);
-
-	if (Entries.Num() < 4)
+	TArray<TEPAEntry<T>> Entries;
+	if(!InitializeEPA(VertsABuffer,VertsBBuffer,SupportA,SupportB, Entries, OutDir))
 	{
 		//either degenerate or a touching hit. Either way return penetration 0
 		OutPenetration = 0;
-		OutDir = TVec3<T>(0, 0, 1);
 		WitnessA = TVec3<T>(0);
 		WitnessB = TVec3<T>(0);
 		return EPAResult::BadInitialSimplex;
 	}
 
-
 #if DEBUG_EPA
 	TArray<TVec3<T>> VertsWBuffer;
 	for (int32 Idx = 0; Idx < 4; ++Idx)
diff --git a/Engine/Source/Runtime/Experimental/Chaos/Public/Chaos/GJK.h b/Engine/Source/Runtime/Experimental/Chaos/Public/Chaos/GJK.h
--- a/Engine/Source/Runtime/Experimental/Chaos/Public/Chaos/GJK.h
+++ b/Engine/Source/Runtime/Experimental/Chaos/Public/Chaos/GJK.h
@@ -235,10 +235,10 @@ namespace Chaos
 				}
 
 				OutPenetration = ThicknessA + ThicknessB;
-				OutNormal = { 0,0,1 };
+				OutNormal = MTD;
 				OutClosestA = ClosestA + OutNormal * ThicknessA;
 				OutClosestB = ClosestBInA - OutNormal * ThicknessB;
-				return false;
+				return OutPenetration > Eps2;
 			}
 		}
 
@@ -631,7 +631,7 @@ namespace Chaos
 					{
 						//assume touching hit
 						OutTime = -Inflation;
-						OutNormal = { 0,0,1 };
+						OutNormal = MTD;
 						OutPosition = As[0] + OutNormal * ThicknessA;
 					}
 				}
