Chaos : Geometry Complexity

- Added support for passing geometry complexity through to the solver.
... If the geometry is marked as UseComplexAsSimple then the TriangleMesh will be used instead of the primitive.
- Enabled TriangleMesh to Boix collisions.

diff --git a/Engine/Source/Runtime/Core/Public/UObject/ExternalPhysicsCustomObjectVersion.h b/Engine/Source/Runtime/Core/Public/UObject/ExternalPhysicsCustomObjectVersion.h
--- a/Engine/Source/Runtime/Core/Public/UObject/ExternalPhysicsCustomObjectVersion.h
+++ b/Engine/Source/Runtime/Core/Public/UObject/ExternalPhysicsCustomObjectVersion.h
@@ -113,7 +113,10 @@ struct CORE_API FExternalPhysicsCustomObjectVersion
 
 		// Acceleration structures use unique payload idx
 		UniquePayloadIdx,
-		
+	
+		// Added serialization for the collision type in the shape
+		SerializeCollisionTraceType,
+
 		// -----<new versions can be added above this line>-------------------------------------------------
 		VersionPlusOne,
 		LatestVersion = VersionPlusOne - 1
diff --git a/Engine/Source/Runtime/Engine/Private/PhysicsEngine/BodyInstance.cpp b/Engine/Source/Runtime/Engine/Private/PhysicsEngine/BodyInstance.cpp
--- a/Engine/Source/Runtime/Engine/Private/PhysicsEngine/BodyInstance.cpp
+++ b/Engine/Source/Runtime/Engine/Private/PhysicsEngine/BodyInstance.cpp
@@ -1169,6 +1169,23 @@ struct FInitBodiesHelper
 		return true;
 	}
 
+#if WITH_CHAOS
+	Chaos::EChaosCollisionTraceFlag ConvertCollisionTraceFlag(ECollisionTraceFlag Flag)
+	{
+		if (Flag == ECollisionTraceFlag::CTF_UseDefault)
+			return Chaos::EChaosCollisionTraceFlag::Chaos_CTF_UseDefault;
+		if (Flag == ECollisionTraceFlag::CTF_UseSimpleAndComplex)
+			return Chaos::EChaosCollisionTraceFlag::Chaos_CTF_UseSimpleAndComplex;
+		if (Flag == ECollisionTraceFlag::CTF_UseSimpleAsComplex)
+			return Chaos::EChaosCollisionTraceFlag::Chaos_CTF_UseSimpleAsComplex;
+		if (Flag == ECollisionTraceFlag::CTF_UseComplexAsSimple)
+			return Chaos::EChaosCollisionTraceFlag::Chaos_CTF_UseComplexAsSimple;
+		if (Flag == ECollisionTraceFlag::CTF_MAX)
+			return Chaos::EChaosCollisionTraceFlag::Chaos_CTF_MAX;
+		ensure(false);
+		return Chaos::EChaosCollisionTraceFlag::Chaos_CTF_UseDefault;
+	}
+#endif
 
 	void InitBodies()
 	{
@@ -1209,16 +1226,25 @@ struct FInitBodiesHelper
 							ActorHandles.Add(ActorHandle);
 
 #if WITH_CHAOS
+							const int32 NumShapes = FPhysicsInterface::GetNumShapes(ActorHandle);
+
 							// If this shape shouldn't collide in the sim we disable it here until we support
 							// a separation of unions for these shapes
 							if(BI->GetCollisionEnabled() == ECollisionEnabled::QueryOnly || BI->GetCollisionEnabled() == ECollisionEnabled::NoCollision)
 							{
-								const int32 NumShapes = FPhysicsInterface::GetNumShapes(ActorHandle);
 								for(int32 ShapeIndex = 0; ShapeIndex < NumShapes; ++ShapeIndex)
 								{
 									ActorHandle->SetShapeCollisionDisable(ShapeIndex, true);
 								}
 							}
+							if (BI->BodySetup.IsValid())
+							{
+								for (int32 ShapeIndex = 0; ShapeIndex < NumShapes; ++ShapeIndex)
+								{
+									ActorHandle->SetShapeCollisionTraceType(ShapeIndex, ConvertCollisionTraceFlag(BI->BodySetup->CollisionTraceFlag)) ;
+								}
+							}
+
 #endif
 /*
 							With the implementation of AddActorsToScene_AssumesLocked,
diff --git a/Engine/Source/Runtime/Engine/Private/PhysicsEngine/Experimental/ChaosInterfaceUtils.cpp b/Engine/Source/Runtime/Engine/Private/PhysicsEngine/Experimental/ChaosInterfaceUtils.cpp
--- a/Engine/Source/Runtime/Engine/Private/PhysicsEngine/Experimental/ChaosInterfaceUtils.cpp
+++ b/Engine/Source/Runtime/Engine/Private/PhysicsEngine/Experimental/ChaosInterfaceUtils.cpp
@@ -134,6 +134,21 @@ namespace ChaosInterface
 
 #endif
 
+	Chaos::EChaosCollisionTraceFlag ConvertCollisionTraceFlag(ECollisionTraceFlag Flag)
+	{
+		if (Flag == ECollisionTraceFlag::CTF_UseDefault)
+			return Chaos::EChaosCollisionTraceFlag::Chaos_CTF_UseDefault;
+		if (Flag == ECollisionTraceFlag::CTF_UseSimpleAndComplex)
+			return Chaos::EChaosCollisionTraceFlag::Chaos_CTF_UseSimpleAndComplex;
+		if (Flag == ECollisionTraceFlag::CTF_UseSimpleAsComplex)
+			return Chaos::EChaosCollisionTraceFlag::Chaos_CTF_UseSimpleAsComplex;
+		if (Flag == ECollisionTraceFlag::CTF_UseComplexAsSimple)
+			return Chaos::EChaosCollisionTraceFlag::Chaos_CTF_UseComplexAsSimple;
+		if (Flag == ECollisionTraceFlag::CTF_MAX)
+			return Chaos::EChaosCollisionTraceFlag::Chaos_CTF_MAX;
+		ensure(false);
+		return Chaos::EChaosCollisionTraceFlag::Chaos_CTF_UseDefault;
+	}
 
 	void CreateGeometry(const FGeometryAddParams& InParams, TArray<TUniquePtr<Chaos::FImplicitObject>>& OutGeoms, Chaos::TShapesArray<float, 3>& OutShapes)
 	{
@@ -162,12 +177,13 @@ namespace ChaosInterface
 
 		ensure(bMakeComplexGeometry || bMakeSimpleGeometry);
 
-		auto NewShapeHelper = [&InParams](Chaos::TSerializablePtr<Chaos::FImplicitObject> InGeom, void* UserData, bool bComplexShape = false)
+		auto NewShapeHelper = [&InParams, &CollisionTraceType](Chaos::TSerializablePtr<Chaos::FImplicitObject> InGeom, void* UserData, bool bComplexShape = false)
 		{
 			auto NewShape = Chaos::TPerShapeData<float, 3>::CreatePerShapeData();
 			NewShape->Geometry = InGeom;
 			NewShape->QueryData = bComplexShape ? InParams.CollisionData.CollisionFilterData.QueryComplexFilter : InParams.CollisionData.CollisionFilterData.QuerySimpleFilter;
 			NewShape->SimData = InParams.CollisionData.CollisionFilterData.SimFilter;
+			NewShape->CollisionTraceType = ConvertCollisionTraceFlag(CollisionTraceType);
 			NewShape->UpdateShapeBounds(InParams.WorldTransform);
 			NewShape->UserData = UserData;
 			NewShape->bSimulate = bComplexShape ? InParams.CollisionData.CollisionFlags.bEnableSimCollisionComplex : InParams.CollisionData.CollisionFlags.bEnableSimCollisionSimple;
diff --git a/Engine/Source/Runtime/Experimental/Chaos/Private/Chaos/CollisionResolution.cpp b/Engine/Source/Runtime/Experimental/Chaos/Private/Chaos/CollisionResolution.cpp
--- a/Engine/Source/Runtime/Experimental/Chaos/Private/Chaos/CollisionResolution.cpp
+++ b/Engine/Source/Runtime/Experimental/Chaos/Private/Chaos/CollisionResolution.cpp
@@ -1094,7 +1094,7 @@ namespace Chaos
 
 			//
 			// @todo(chaos): Collision Constraints (CollisionMap)
-			//    Modify Construct() and Update() use a CollisionMap indexed on 
+			//    Modify Construct() and Update() to use a CollisionMap indexed on 
 			//    EImplicitObjectType, instead of the if/else chain. Also, remove 
 			//    the blocks with the ensure(false), they are just for validation 
 			//    after the recent change. 
@@ -1175,7 +1175,18 @@ namespace Chaos
 			{
 				//UpdateCapsuleSphereManifold(ConstraintBase, Transform0, Transform1, CullDistance);
 			}
+			else if (Implicit0Type == TBox<T, d>::StaticType() && Implicit1Type == FTriangleMeshImplicitObject::StaticType())
+			{
+				//UpdateBoxTriangleMeshManifold(ConstraintBase, Transform0, Transform1, CullDistance);
+			}
+			else if (Implicit0Type == FTriangleMeshImplicitObject::StaticType() && Implicit1Type == TBox<T, d>::StaticType())
+			{
+				ensure(false);
+			}
+			//
 			// the generic convex bodies are last
+			//
+
 			else if (Implicit0Type == THeightField<T>::StaticType() && Implicit1.IsConvex())
 			{
 				ensure(false);
@@ -1193,17 +1204,6 @@ namespace Chaos
 				UpdateLevelsetLevelsetManifold(ConstraintBase, Transform0, Transform1, CullDistance);
 			}
 
-			/*
-			else if (Implicit0Type == TBox<T, d>::StaticType() && Implicit1Type == FTriangleMeshImplicitObject::StaticType())
-			{
-				UpdateBoxTriangleMeshManifold(ConstraintBase, Transform0, Transform1, CullDistance);
-			}
-			else if (Implicit0Type == FTriangleMeshImplicitObject::StaticType() && Implicit1Type == TBox<T, d>::StaticType())
-			{
-				ensure(false);
-			}
-			*/
-
 		}
 
 		template<ECollisionUpdateType UpdateType, typename T, int d>
@@ -1376,7 +1376,22 @@ namespace Chaos
 				//     second.
 				ensure(false);
 			}
+			else if (Implicit0Type == TBox<T, d>::StaticType() && Implicit1Type == FTriangleMeshImplicitObject::StaticType())
+			{
+				UpdateBoxTriangleMeshConstraint(Implicit0.template GetObject<TBox<T, d>>()->GetAABB(), Transform0, *GetInnerObject<FTriangleMeshImplicitObject>(Implicit1), Transform1, CullDistance, *ConstraintBase.template As<TRigidBodyPointContactConstraint<T, d>>());
+			}
+			else if (Implicit0Type == TBox<T, d>::StaticType() && Implicit1Type == FTriangleMeshImplicitObject::StaticType())
+			{
+				//     This case should not be necessary. The triangle mesh
+				//     will only ever be collided against, so ideally will never
+				//     be in index[0] position of the constraint, also the construction
+				//     of the constraint will just switch the index position so its always
+				//     second.
+				ensure(false);
+			}
+			//
 			// the generic convex bodies are last
+			//
 			else if (Implicit0Type == THeightField<T>::StaticType() && Implicit1.IsConvex())
 			{
 				//     This case should not be necessary. The height fields 
@@ -1398,22 +1413,6 @@ namespace Chaos
 			{
 				UpdateLevelsetLevelsetConstraint<UpdateType>(CullDistance, *ConstraintBase.template As<TRigidBodyPointContactConstraint<T, d> >());
 			}
-
-			/*
-			else if (Implicit0Type == TBox<T, d>::StaticType() && Implicit1Type == FTriangleMeshImplicitObject::StaticType())
-			{
-			UpdateBoxTriangleMeshConstraint(Implicit0.template GetObject<TBox<T, d>>()->GetAABB(), Transform0, *GetUnscaledInner<FTriangleMeshImplicitObject>(Implicit1), Transform1, CullDistance, *ConstraintBase.template As<TRigidBodyPointContactConstraint<T, d>>());
-			}
-			else if (Implicit0Type == TBox<T, d>::StaticType() && Implicit1Type == FTriangleMeshImplicitObject::StaticType())
-			{
-				//     This case should not be necessary. The triangle mesh
-				//     will only ever be collided against, so ideally will never
-				//     be in index[0] position of the constraint, also the construction
-				//     of the constraint will just switch the index position so its always
-				//     second.
-				ensure(false);
-			}
-			*/
 		}
 
 		template<ECollisionUpdateType UpdateType, typename T, int d>
@@ -1445,18 +1444,43 @@ namespace Chaos
 		template<typename T, int d>
 		void ConstructConstraints(TGeometryParticleHandle<T, d>* Particle0, TGeometryParticleHandle<T, d>* Particle1, const FImplicitObject* Implicit0, const FImplicitObject* Implicit1, const TRigidTransform<T, d>& Transform0, const TRigidTransform<T, d>& Transform1, const T CullDistance, FCollisionConstraintsArray& NewConstraints)
 		{
-			//  TriMesh Culling
-			if (Implicit0 && Implicit0->GetType() == ImplicitObjectType::TriangleMesh) return;
-			if (Implicit1 && Implicit1->GetType() == ImplicitObjectType::TriangleMesh) return;
+			EImplicitObjectType Implicit0Type = Implicit0 ? GetInnerType(Implicit0->GetType()) : ImplicitObjectType::Unknown;
+			EImplicitObjectType Implicit1Type = Implicit1 ? GetInnerType(Implicit1->GetType()) : ImplicitObjectType::Unknown;
 
 			// If either shape is disabled for collision bail without constructing a constraint
-			const TPerShapeData<T, d>* Shape0 = Particle0->GetImplicitShape(Implicit0);
-			const TPerShapeData<T, d>* Shape1 = Particle1->GetImplicitShape(Implicit1);
-			if ((Shape0 && Shape0->bDisable) || (Shape1 && Shape1->bDisable))
+			if (const TPerShapeData<T, d>* Shape0 = Particle0->GetImplicitShape(Implicit0))
+			{
+				if (Shape0->bDisable)
+				{
+					return;
+				}
+				if (Implicit0Type == ImplicitObjectType::TriangleMesh && Shape0->CollisionTraceType != Chaos_CTF_UseComplexAsSimple)
+				{
+					return;
+				}
+			} 
+			else if (Implicit0Type == ImplicitObjectType::TriangleMesh) 
 			{
 				return;
 			}
 
+			if (const TPerShapeData<T, d>* Shape1 = Particle1->GetImplicitShape(Implicit1))
+			{
+				if (Shape1->bDisable)
+				{
+					return;
+				}
+				if (Implicit1Type == ImplicitObjectType::TriangleMesh && Shape1->CollisionTraceType != Chaos_CTF_UseComplexAsSimple)
+				{
+					return;
+				}
+			}
+			else if (Implicit1Type == ImplicitObjectType::TriangleMesh)
+			{
+				return;
+			}
+
+
 			if (!Implicit0 || !Implicit1)
 			{
 				ConstructLevelsetLevelsetConstraints(Particle0, Particle1, Implicit0, Implicit1, Transform0, Transform1, CullDistance, NewConstraints);
@@ -1580,9 +1604,6 @@ namespace Chaos
 				return;
 			}
 
-			EImplicitObjectType Implicit0Type = GetInnerType(Implicit0->GetType());
-			EImplicitObjectType Implicit1Type = GetInnerType(Implicit1->GetType());
-
 			if (Implicit0Type == TBox<T, d>::StaticType() && Implicit1Type == TBox<T, d>::StaticType())
 			{
 				ConstructBoxBoxConstraints(Particle0, Particle1, Implicit0, Implicit1, Transform0, Transform1, CullDistance, NewConstraints);
@@ -1659,7 +1680,19 @@ namespace Chaos
 			{
 				ConstructCapsuleHeightFieldConstraints(Particle1, Particle0, Implicit1, Implicit0, Transform1, Transform0, CullDistance, NewConstraints);
 			}
+			else if (Implicit0Type == TBox<T, d>::StaticType() && Implicit1Type == FTriangleMeshImplicitObject::StaticType())
+			{
+				ConstructBoxTriangleMeshConstraints(Particle0, Particle1, Implicit0, Implicit1, Transform0, Transform1, CullDistance, NewConstraints);
+			}
+			else if (Implicit0Type == FTriangleMeshImplicitObject::StaticType() && Implicit1Type == TBox<T, d>::StaticType())
+			{
+				ConstructBoxTriangleMeshConstraints(Particle1, Particle0, Implicit1, Implicit0, Transform1, Transform0, CullDistance, NewConstraints);
+			}
+
+			//
 			// the generic convex bodies are last
+			//
+
 			else if (Implicit0->IsConvex() && Implicit1Type == THeightField<T>::StaticType())
 			{
 				ConstructConvexHeightFieldConstraints(Particle0, Particle1, Implicit0, Implicit1, Transform0, Transform1, CullDistance, NewConstraints);
@@ -1676,19 +1709,6 @@ namespace Chaos
 			{
 				ConstructLevelsetLevelsetConstraints(Particle0, Particle1, Implicit0, Implicit1, Transform0, Transform1, CullDistance, NewConstraints);
 			}
-
-
-			/*
-			if (Implicit0Type == TBox<T, d>::StaticType() && Implicit1Type == FTriangleMeshImplicitObject::StaticType())
-			{
-				ConstructBoxTriangleMeshConstraints(Particle0, Particle1, Implicit0, Implicit1, Transform0, Transform1, CullDistance, NewConstraints);
-			}
-			else if (Implicit0Type == FTriangleMeshImplicitObject::StaticType() && Implicit1Type == TBox<T, d>::StaticType())
-			{
-				ConstructBoxTriangleMeshConstraints(Particle1, Particle0, Implicit1, Implicit0, Transform1, Transform0, CullDistance, NewConstraints);
-			}
-			*/
-
 		}
 
 		template void UpdateBoxBoxConstraint<float, 3>(const TAABB<float, 3>& Box1, const TRigidTransform<float, 3>& Box1Transform, const TAABB<float, 3>& Box2, const TRigidTransform<float, 3>& Box2Transform, const float CullDistance, TRigidBodyPointContactConstraint<float, 3>& Constraint);
diff --git a/Engine/Source/Runtime/Experimental/Chaos/Private/Chaos/GeometryParticles.cpp b/Engine/Source/Runtime/Experimental/Chaos/Private/Chaos/GeometryParticles.cpp
--- a/Engine/Source/Runtime/Experimental/Chaos/Private/Chaos/GeometryParticles.cpp
+++ b/Engine/Source/Runtime/Experimental/Chaos/Private/Chaos/GeometryParticles.cpp
@@ -60,6 +60,7 @@ namespace Chaos
 		, Materials()
 		, bDisable(false)
 		, bSimulate(true)
+		, CollisionTraceType(EChaosCollisionTraceFlag::Chaos_CTF_UseDefault)
 	{
 	}
 
@@ -122,6 +123,13 @@ namespace Chaos
 		{
 			Ar << bSimulate;
 		}
+
+		if (Ar.CustomVer(FExternalPhysicsCustomObjectVersion::GUID) >= FExternalPhysicsCustomObjectVersion::SerializeCollisionTraceType)
+		{
+			int32 Data = (int32)CollisionTraceType;
+			Ar << Data;
+			CollisionTraceType = (EChaosCollisionTraceFlag)Data;
+		}
 	}
 
 
diff --git a/Engine/Source/Runtime/Experimental/Chaos/Public/Chaos/CollisionFilterData.h b/Engine/Source/Runtime/Experimental/Chaos/Public/Chaos/CollisionFilterData.h
--- a/Engine/Source/Runtime/Experimental/Chaos/Public/Chaos/CollisionFilterData.h
+++ b/Engine/Source/Runtime/Experimental/Chaos/Public/Chaos/CollisionFilterData.h
@@ -2,6 +2,7 @@
 #pragma once
 
 #include "ChaosArchive.h"
+#include "UObject/ExternalPhysicsCustomObjectVersion.h"
 
 struct FCollisionFilterData
 {
diff --git a/Engine/Source/Runtime/Experimental/Chaos/Public/Chaos/GeometryParticles.h b/Engine/Source/Runtime/Experimental/Chaos/Public/Chaos/GeometryParticles.h
--- a/Engine/Source/Runtime/Experimental/Chaos/Public/Chaos/GeometryParticles.h
+++ b/Engine/Source/Runtime/Experimental/Chaos/Public/Chaos/GeometryParticles.h
@@ -19,11 +19,28 @@
 
 namespace Chaos
 {
+
+	enum EChaosCollisionTraceFlag
+	{
+		/** Use project physics settings (DefaultShapeComplexity) */
+		Chaos_CTF_UseDefault,
+		/** Create both simple and complex shapes. Simple shapes are used for regular scene queries and collision tests. Complex shape (per poly) is used for complex scene queries.*/
+		Chaos_CTF_UseSimpleAndComplex,
+		/** Create only simple shapes. Use simple shapes for all scene queries and collision tests.*/
+		Chaos_CTF_UseSimpleAsComplex,
+		/** Create only complex shapes (per poly). Use complex shapes for all scene queries and collision tests. Can be used in simulation for static shapes only (i.e can be collided against but not moved through forces or velocity.) */
+		Chaos_CTF_UseComplexAsSimple,
+		/** */
+		Chaos_CTF_MAX,
+	};
+
+
 	/** Data that is associated with geometry. If a union is used an entry is created per internal geometry */
 	template <typename T, int d>
 	class CHAOS_API TPerShapeData
 	{
 	public:
+
 		static constexpr bool AlwaysSerializable = true;
 		static TUniquePtr<TPerShapeData<T, d>> CreatePerShapeData();
 
@@ -39,6 +56,7 @@ namespace Chaos
 		// TODO: Bitfields?
 		bool bDisable;
 		bool bSimulate;
+		EChaosCollisionTraceFlag CollisionTraceType;
 
 		void UpdateShapeBounds(const TRigidTransform<T, d>& WorldTM);
 
diff --git a/Engine/Source/Runtime/Experimental/Chaos/Public/Chaos/ParticleDirtyFlags.h b/Engine/Source/Runtime/Experimental/Chaos/Public/Chaos/ParticleDirtyFlags.h
--- a/Engine/Source/Runtime/Experimental/Chaos/Public/Chaos/ParticleDirtyFlags.h
+++ b/Engine/Source/Runtime/Experimental/Chaos/Public/Chaos/ParticleDirtyFlags.h
@@ -36,9 +36,10 @@ namespace Chaos
 		GravityEnabled			= 1 << 24,
 		SpatialIdx				= 1 << 25,
 		UniqueIdx				= 1 << 26,
-		ShapeDisableCollision	= 1 << 27
+		ShapeDisableCollision	= 1 << 27,
+		CollisionTraceType		= 1 << 28
 #if CHAOS_CHECKED
-		, DebugName				= 1 << 28
+		, DebugName				= 1 << 29
 #endif
 	};
 
diff --git a/Engine/Source/Runtime/Experimental/Chaos/Public/Chaos/ParticleHandle.h b/Engine/Source/Runtime/Experimental/Chaos/Public/Chaos/ParticleHandle.h
--- a/Engine/Source/Runtime/Experimental/Chaos/Public/Chaos/ParticleHandle.h
+++ b/Engine/Source/Runtime/Experimental/Chaos/Public/Chaos/ParticleHandle.h
@@ -1261,6 +1261,7 @@ public:
 	static constexpr bool AlwaysSerializable = true;
 
 protected:
+
 	TGeometryParticle(const TGeometryParticleParameters<T, d>& StaticParams = TGeometryParticleParameters<T, d>())
 		: MUserData(nullptr)
 	{
@@ -1427,6 +1428,16 @@ public:
 		}
 	}
 
+	void SetShapeCollisionTraceType(int32 InShapeIndex, EChaosCollisionTraceFlag TraceType)
+	{
+		const EChaosCollisionTraceFlag Current = MShapesArray[InShapeIndex]->CollisionTraceType;
+		if (Current != TraceType)
+		{
+			MShapesArray[InShapeIndex]->CollisionTraceType = TraceType;
+			MarkDirty(EParticleFlags::CollisionTraceType);
+		}
+	}
+
 	FParticleData* NewData() const { return new TGeometryParticleData<T, d>( *this ); }
 
 	bool IsDirty() const
@@ -1559,9 +1570,11 @@ public:
 	{
 		const TShapesArray<T, d>& Shapes = InParticle.ShapesArray();
 		ShapeCollisionDisableFlags.Empty(Shapes.Num());
+		CollisionTraceType.Empty(Shapes.Num());
 		for (const TUniquePtr<TPerShapeData<T, d>>& ShapePtr : Shapes)
 		{
 			ShapeCollisionDisableFlags.Add(ShapePtr->bDisable);
+			CollisionTraceType.Add(ShapePtr->CollisionTraceType);
 		}
 
 	}
@@ -1576,6 +1589,7 @@ public:
 		UniqueIdx = FUniqueIdx();
 		DirtyFlags.Clear();
 		ShapeCollisionDisableFlags.Reset();
+		CollisionTraceType.Reset();
 #if CHAOS_CHECKED
 		DebugName = NAME_None;
 #endif
@@ -1588,6 +1602,7 @@ public:
 	FUniqueIdx UniqueIdx;
 	FParticleDirtyFlags DirtyFlags;
 	TBitArray<> ShapeCollisionDisableFlags;
+	TArray<EChaosCollisionTraceFlag> CollisionTraceType;
 #if CHAOS_CHECKED
 	FName DebugName;
 #endif
diff --git a/Engine/Source/Runtime/Experimental/ChaosSolvers/Private/PhysicsProxy/SingleParticlePhysicsProxy.cpp b/Engine/Source/Runtime/Experimental/ChaosSolvers/Private/PhysicsProxy/SingleParticlePhysicsProxy.cpp
--- a/Engine/Source/Runtime/Experimental/ChaosSolvers/Private/PhysicsProxy/SingleParticlePhysicsProxy.cpp
+++ b/Engine/Source/Runtime/Experimental/ChaosSolvers/Private/PhysicsProxy/SingleParticlePhysicsProxy.cpp
@@ -73,6 +73,16 @@ void FSingleParticlePhysicsProxy<Chaos::TGeometryParticle<float, 3>>::PushToPhys
 				Shape->bDisable = Data->ShapeCollisionDisableFlags[CurrShape++];
 			}
 		}
+
+		// @todo(chaos) : Fix - The initial update is not flagging the dirty flag, so this state is not updated. 
+		//if (Data->DirtyFlags.IsDirty(Chaos::EParticleFlags::CollisionTraceType))
+		{
+			int32 CurrShape = 0;
+			for (const TUniquePtr<Chaos::TPerShapeData<Chaos::FReal, 3>>& Shape : RigidHandle->ShapesArray())
+			{
+				Shape->CollisionTraceType = Data->CollisionTraceType[CurrShape++];
+			}
+		}
 	}
 }
 
@@ -164,6 +174,14 @@ void FSingleParticlePhysicsProxy<Chaos::TKinematicGeometryParticle<float, 3>>::P
 				Shape->bDisable = Data->ShapeCollisionDisableFlags[CurrShape++];
 			}
 		}
+		if (Data->DirtyFlags.IsDirty(Chaos::EParticleFlags::CollisionTraceType))
+		{
+			int32 CurrShape = 0;
+			for (const TUniquePtr<Chaos::TPerShapeData<Chaos::FReal, 3>>& Shape : RigidHandle->ShapesArray())
+			{
+				Shape->CollisionTraceType = Data->CollisionTraceType[CurrShape++];
+			}
+		}
 
 	}
 }
@@ -323,6 +341,14 @@ void FSingleParticlePhysicsProxy<Chaos::TPBDRigidParticle<float, 3>>::PushToPhys
 				Shape->bDisable = Data->ShapeCollisionDisableFlags[CurrShape++];
 			}
 		}
+		if (Data->DirtyFlags.IsDirty(Chaos::EParticleFlags::CollisionTraceType))
+		{
+			int32 CurrShape = 0;
+			for (const TUniquePtr<Chaos::TPerShapeData<Chaos::FReal, 3>>& Shape : RigidHandle->ShapesArray())
+			{
+				Shape->CollisionTraceType = Data->CollisionTraceType[CurrShape++];
+			}
+		}
 
 		if (bDynamicPropertyUpdated)
 		{
