Chaos : TriMesh Collisions

- Support for TriMesh to Sphere.

diff --git a/Engine/Source/Runtime/Experimental/Chaos/Private/Chaos/CollisionResolution.cpp b/Engine/Source/Runtime/Experimental/Chaos/Private/Chaos/CollisionResolution.cpp
--- a/Engine/Source/Runtime/Experimental/Chaos/Private/Chaos/CollisionResolution.cpp
+++ b/Engine/Source/Runtime/Experimental/Chaos/Private/Chaos/CollisionResolution.cpp
@@ -234,6 +234,17 @@ namespace Chaos
 			for (int32 Index = 0; Index < Constraint.NumSamples(); Index++) { Constraint[Index].Delta = (Center - Constraint[Index].X).SizeSquared(); }
 		}
 
+		template <typename TPGeometryClass>
+		const TPGeometryClass* GetInnerObject(const FImplicitObject& Geometry)
+		{
+			if (const TImplicitObjectScaled<TPGeometryClass>* ScaledConvexImplicit = Geometry.template GetObject<const TImplicitObjectScaled<TPGeometryClass> >())
+				return (Geometry.template GetObject<const TImplicitObjectScaled<TPGeometryClass> >())->GetUnscaledObject();
+			else if (const TImplicitObjectInstanced<TPGeometryClass>* InstancedImplicit = Geometry.template GetObject<const TImplicitObjectInstanced<TPGeometryClass> >())
+				return (Geometry.template GetObject<const TImplicitObjectInstanced<TPGeometryClass> >())->GetInstancedObject();
+			else if (const TPGeometryClass* ConvexImplicit = Geometry.template GetObject<const TPGeometryClass>())
+				return Geometry.template GetObject<const TPGeometryClass>();
+			return nullptr;
+		}
 
 
 		//
@@ -423,18 +434,6 @@ namespace Chaos
 
 		}
 
-		template <typename TPGeometryClass>
-		const TPGeometryClass* GetInnerObject(const FImplicitObject& Geometry)
-		{
-			if (const TImplicitObjectScaled<TPGeometryClass>* ScaledConvexImplicit = Geometry.template GetObject<const TImplicitObjectScaled<TPGeometryClass> >())
-				return (Geometry.template GetObject<const TImplicitObjectScaled<TPGeometryClass> >())->GetUnscaledObject();
-			else if (const TImplicitObjectInstanced<TPGeometryClass>* InstancedImplicit = Geometry.template GetObject<const TImplicitObjectInstanced<TPGeometryClass> >())
-				return (Geometry.template GetObject<const TImplicitObjectInstanced<TPGeometryClass> >())->GetInstancedObject();
-			else if (const TPGeometryClass* ConvexImplicit = Geometry.template GetObject<const TPGeometryClass>())
-				return Geometry.template GetObject<const TPGeometryClass>();
-			return nullptr;
-		}
-
 		template<typename T, int d>
 		void ConstructBoxTriangleMeshConstraints(TGeometryParticleHandle<T, d>* Particle0, TGeometryParticleHandle<T, d>* Particle1, const FImplicitObject* Implicit0, const FImplicitObject* Implicit1, const TRigidTransform<T, d>& Transform0, const TRigidTransform<T, d>& Transform1, const T CullDistance, FCollisionConstraintsArray& NewConstraints)
 		{
@@ -694,6 +693,45 @@ namespace Chaos
 			}
 		}
 
+		//
+		// Sphere-TriangleMesh
+		//
+
+		template <typename T, int d>
+		TContactPoint<T> SphereTriangleMeshContactPoint(const TSphere<T, d>& A, const TRigidTransform<T, d>& ATransform, const FTriangleMeshImplicitObject& B, const TRigidTransform<T, d>& BTransform, const T CullDistance)
+		{
+			return GJKImplicitContactPoint< TSphere<float, 3> >(TSphere<float, 3>(A), ATransform, B, BTransform, CullDistance);
+		}
+
+		template <typename T, int d>
+		void UpdateSphereTriangleMeshConstraint(const TSphere<T, d>& Sphere0, const TRigidTransform<T, d>& Transform0, const FTriangleMeshImplicitObject& TriangleMesh1, const TRigidTransform<T, d>& Transform1, const T CullDistance, TRigidBodyPointContactConstraint<T, d>& Constraint)
+		{
+			UpdateContactPoint(Constraint.Manifold, SphereTriangleMeshContactPoint(Sphere0, Transform0, TriangleMesh1, Transform1, CullDistance));
+		}
+
+
+		template <typename T, int d>
+		void UpdateSphereTriangleMeshManifold(TCollisionConstraintBase<T, d>&  Constraint, const TRigidTransform<T, d>& ATM, const TRigidTransform<T, d>& BTM, const T CullDistance)
+		{
+
+		}
+
+		template<typename T, int d>
+		void ConstructSphereTriangleMeshConstraints(TGeometryParticleHandle<T, d>* Particle0, TGeometryParticleHandle<T, d>* Particle1, const FImplicitObject* Implicit0, const FImplicitObject* Implicit1, const TRigidTransform<T, d>& Transform0, const TRigidTransform<T, d>& Transform1, const T CullDistance, FCollisionConstraintsArray& NewConstraints)
+		{
+			const TSphere<T, d> * Object0 = Implicit0->template GetObject<const TSphere<T, d> >();
+			const FTriangleMeshImplicitObject * Object1 = GetInnerObject<FTriangleMeshImplicitObject>(*Implicit1);
+			if (ensure(Object0 && Object1))
+			{
+				TRigidTransform<T, d> ParticleImplicit0TM = Transform0.GetRelativeTransform(Collisions::GetTransform(Particle0));
+				TRigidTransform<T, d> ParticleImplicit1TM = Transform1.GetRelativeTransform(Collisions::GetTransform(Particle1));
+				FRigidBodyPointContactConstraint Constraint = FRigidBodyPointContactConstraint(Particle0, Implicit0, ParticleImplicit0TM, Particle1, Implicit1, ParticleImplicit1TM);
+				UpdateSphereTriangleMeshConstraint(*Object0, Transform0, *Object1, Transform1, CullDistance, Constraint);
+				NewConstraints.TryAdd(CullDistance, Constraint);
+			}
+		}
+
+
 
 		//
 		// Capsule-Capsule
@@ -1183,6 +1221,14 @@ namespace Chaos
 			{
 				ensure(false);
 			}
+			else if (Implicit0Type == TSphere<T, d>::StaticType() && Implicit1Type == FTriangleMeshImplicitObject::StaticType())
+			{
+				//UpdateSphereTriangleMeshManifold(ConstraintBase, Transform0, Transform1, CullDistance);
+			}
+			else if (Implicit0Type == FTriangleMeshImplicitObject::StaticType() && Implicit1Type == TSphere<T, d>::StaticType())
+			{
+				ensure(false);
+			}
 			//
 			// the generic convex bodies are last
 			//
@@ -1389,9 +1435,26 @@ namespace Chaos
 				//     second.
 				ensure(false);
 			}
+			else if (Implicit0Type == TSphere<T, d>::StaticType() && Implicit1Type == FTriangleMeshImplicitObject::StaticType())
+			{
+				UpdateSphereTriangleMeshConstraint(*Implicit0.template GetObject<TSphere<T, d>>(), Transform0, *GetInnerObject<FTriangleMeshImplicitObject>(Implicit1), Transform1, CullDistance, *ConstraintBase.template As<TRigidBodyPointContactConstraint<T, d>>());
+			}
+			else if (Implicit0Type == TSphere<T, d>::StaticType() && Implicit1Type == FTriangleMeshImplicitObject::StaticType())
+			{
+				//     This case should not be necessary. The triangle mesh
+				//     will only ever be collided against, so ideally will never
+				//     be in index[0] position of the constraint, also the construction
+				//     of the constraint will just switch the index position so its always
+				//     second.
+				ensure(false);
+			}
+
+
 			//
 			// the generic convex bodies are last
 			//
+
+
 			else if (Implicit0Type == THeightField<T>::StaticType() && Implicit1.IsConvex())
 			{
 				//     This case should not be necessary. The height fields 
@@ -1688,6 +1751,14 @@ namespace Chaos
 			{
 				ConstructBoxTriangleMeshConstraints(Particle1, Particle0, Implicit1, Implicit0, Transform1, Transform0, CullDistance, NewConstraints);
 			}
+			else if (Implicit0Type == TSphere<T, d>::StaticType() && Implicit1Type == FTriangleMeshImplicitObject::StaticType())
+			{
+				ConstructSphereTriangleMeshConstraints(Particle0, Particle1, Implicit0, Implicit1, Transform0, Transform1, CullDistance, NewConstraints);
+			}
+			else if (Implicit0Type == FTriangleMeshImplicitObject::StaticType() && Implicit1Type == TSphere<T, d>::StaticType())
+			{
+				ConstructSphereTriangleMeshConstraints(Particle1, Particle0, Implicit1, Implicit0, Transform1, Transform0, CullDistance, NewConstraints);
+			}
 
 			//
 			// the generic convex bodies are last
