Chaos : Collisions

- Added TriMesh to Capsule collisions

diff --git a/Engine/Source/Runtime/Experimental/Chaos/Private/Chaos/CollisionResolution.cpp b/Engine/Source/Runtime/Experimental/Chaos/Private/Chaos/CollisionResolution.cpp
--- a/Engine/Source/Runtime/Experimental/Chaos/Private/Chaos/CollisionResolution.cpp
+++ b/Engine/Source/Runtime/Experimental/Chaos/Private/Chaos/CollisionResolution.cpp
@@ -848,7 +848,7 @@ namespace Chaos
 		template <typename T, int d>
 		TContactPoint<T> CapsuleHeightFieldContactPoint(const TCapsule<T>& A, const TRigidTransform<T, d>& ATransform, const THeightField<T>& B, const TRigidTransform<T, d>& BTransform, const T CullDistance)
 		{
-			return GJKImplicitContactPoint< TCapsule<float> >(TCapsule<float>(A), ATransform, B, BTransform, CullDistance);
+			return GJKImplicitContactPoint< TCapsule<float> >(A, ATransform, B, BTransform, CullDistance);
 		}
 
 		template <typename T, int d>
@@ -880,6 +880,45 @@ namespace Chaos
 			}
 		}
 
+		//
+		// Capsule-TriangleMesh
+		//
+
+		template <typename T, int d>
+		TContactPoint<T> CapsuleTriangleMeshContactPoint(const TCapsule<T>& A, const TRigidTransform<T, d>& ATransform, const FTriangleMeshImplicitObject& B, const TRigidTransform<T, d>& BTransform, const T CullDistance)
+		{
+			return GJKImplicitContactPoint< TCapsule<T> >(A, ATransform, B, BTransform, CullDistance);
+		}
+
+		template <typename T, int d>
+		void UpdateCapsuleTriangleMeshConstraint(const TCapsule<T>& Capsule0, const TRigidTransform<T, d>& Transform0, const FTriangleMeshImplicitObject& TriangleMesh1, const TRigidTransform<T, d>& Transform1, const T CullDistance, TRigidBodyPointContactConstraint<T, d>& Constraint)
+		{
+			UpdateContactPoint(Constraint.Manifold, CapsuleTriangleMeshContactPoint(Capsule0, Transform0, TriangleMesh1, Transform1, CullDistance));
+		}
+
+
+		template <typename T, int d>
+		void UpdateCapsuleTriangleMeshManifold(TCollisionConstraintBase<T, d>&  Constraint, const TRigidTransform<T, d>& ATM, const TRigidTransform<T, d>& BTM, const T CullDistance)
+		{
+
+		}
+
+		template<typename T, int d>
+		void ConstructCapsuleTriangleMeshConstraints(TGeometryParticleHandle<T, d>* Particle0, TGeometryParticleHandle<T, d>* Particle1, const FImplicitObject* Implicit0, const FImplicitObject* Implicit1, const TRigidTransform<T, d>& Transform0, const TRigidTransform<T, d>& Transform1, const T CullDistance, FCollisionConstraintsArray& NewConstraints)
+		{
+			const TCapsule<T> * Object0 = Implicit0->template GetObject<const TCapsule<T> >();
+			const FTriangleMeshImplicitObject * Object1 = GetInnerObject<FTriangleMeshImplicitObject>(*Implicit1);
+			if (ensure(Object0 && Object1))
+			{
+				TRigidTransform<T, d> ParticleImplicit0TM = Transform0.GetRelativeTransform(Collisions::GetTransform(Particle0));
+				TRigidTransform<T, d> ParticleImplicit1TM = Transform1.GetRelativeTransform(Collisions::GetTransform(Particle1));
+				FRigidBodyPointContactConstraint Constraint = FRigidBodyPointContactConstraint(Particle0, Implicit0, ParticleImplicit0TM, Particle1, Implicit1, ParticleImplicit1TM);
+				UpdateCapsuleTriangleMeshConstraint(*Object0, Transform0, *Object1, Transform1, CullDistance, Constraint);
+				NewConstraints.TryAdd(CullDistance, Constraint);
+			}
+		}
+
+
 		//
 		// Convex - Convex
 		//
@@ -1229,6 +1268,14 @@ namespace Chaos
 			{
 				ensure(false);
 			}
+			else if (Implicit0Type == TCapsule<T>::StaticType() && Implicit1Type == FTriangleMeshImplicitObject::StaticType())
+			{
+				//UpdateCapsuleTriangleMeshManifold(ConstraintBase, Transform0, Transform1, CullDistance);
+			}
+			else if (Implicit0Type == FTriangleMeshImplicitObject::StaticType() && Implicit1Type == TCapsule<T>::StaticType())
+			{
+				ensure(false);
+			}
 			//
 			// the generic convex bodies are last
 			//
@@ -1426,7 +1473,7 @@ namespace Chaos
 			{
 				UpdateBoxTriangleMeshConstraint(Implicit0.template GetObject<TBox<T, d>>()->GetAABB(), Transform0, *GetInnerObject<FTriangleMeshImplicitObject>(Implicit1), Transform1, CullDistance, *ConstraintBase.template As<TRigidBodyPointContactConstraint<T, d>>());
 			}
-			else if (Implicit0Type == TBox<T, d>::StaticType() && Implicit1Type == FTriangleMeshImplicitObject::StaticType())
+			else if (Implicit0Type == FTriangleMeshImplicitObject::StaticType() && Implicit1Type == TBox<T, d>::StaticType())
 			{
 				//     This case should not be necessary. The triangle mesh
 				//     will only ever be collided against, so ideally will never
@@ -1439,7 +1486,20 @@ namespace Chaos
 			{
 				UpdateSphereTriangleMeshConstraint(*Implicit0.template GetObject<TSphere<T, d>>(), Transform0, *GetInnerObject<FTriangleMeshImplicitObject>(Implicit1), Transform1, CullDistance, *ConstraintBase.template As<TRigidBodyPointContactConstraint<T, d>>());
 			}
-			else if (Implicit0Type == TSphere<T, d>::StaticType() && Implicit1Type == FTriangleMeshImplicitObject::StaticType())
+			else if (Implicit0Type == FTriangleMeshImplicitObject::StaticType() && Implicit1Type == TSphere<T, d>::StaticType())
+			{
+				//     This case should not be necessary. The triangle mesh
+				//     will only ever be collided against, so ideally will never
+				//     be in index[0] position of the constraint, also the construction
+				//     of the constraint will just switch the index position so its always
+				//     second.
+				ensure(false);
+			}
+			else if (Implicit0Type == TCapsule<T>::StaticType() && Implicit1Type == FTriangleMeshImplicitObject::StaticType())
+			{
+				UpdateCapsuleTriangleMeshConstraint(*Implicit0.template GetObject<TCapsule<T>>(), Transform0, *GetInnerObject<FTriangleMeshImplicitObject>(Implicit1), Transform1, CullDistance, *ConstraintBase.template As<TRigidBodyPointContactConstraint<T, d>>());
+			}
+			else if (Implicit0Type == FTriangleMeshImplicitObject::StaticType() && Implicit1Type == TCapsule<T>::StaticType() )
 			{
 				//     This case should not be necessary. The triangle mesh
 				//     will only ever be collided against, so ideally will never
@@ -1759,6 +1819,14 @@ namespace Chaos
 			{
 				ConstructSphereTriangleMeshConstraints(Particle1, Particle0, Implicit1, Implicit0, Transform1, Transform0, CullDistance, NewConstraints);
 			}
+			else if (Implicit0Type == TCapsule<T>::StaticType() && Implicit1Type == FTriangleMeshImplicitObject::StaticType())
+			{
+				ConstructCapsuleTriangleMeshConstraints(Particle0, Particle1, Implicit0, Implicit1, Transform0, Transform1, CullDistance, NewConstraints);
+			}
+			else if (Implicit0Type == FTriangleMeshImplicitObject::StaticType() && Implicit1Type == TCapsule<T>::StaticType())
+			{
+				ConstructCapsuleTriangleMeshConstraints(Particle1, Particle0, Implicit1, Implicit0, Transform1, Transform0, CullDistance, NewConstraints);
+			}
 
 			//
 			// the generic convex bodies are last
