Chaos : Collisions

- Removed sanity if ensures for shipped builds.

diff --git a/Engine/Source/Runtime/Experimental/Chaos/Private/Chaos/CollisionResolution.cpp b/Engine/Source/Runtime/Experimental/Chaos/Private/Chaos/CollisionResolution.cpp
--- a/Engine/Source/Runtime/Experimental/Chaos/Private/Chaos/CollisionResolution.cpp
+++ b/Engine/Source/Runtime/Experimental/Chaos/Private/Chaos/CollisionResolution.cpp
@@ -1224,6 +1224,7 @@
 			{
 				UpdateCapsuleHeightFieldManifold(ConstraintBase, Transform0, Transform1, CullDistance);
 			}
+#if !UE_BUILD_SHIPPING
 			else if (Implicit0Type == THeightField<T>::StaticType() && Implicit1Type == TBox<T, d>::StaticType())
 			{
 				ensure(false);
@@ -1284,6 +1285,7 @@
 			{
 				ensure(false);
 			}
+#endif
 			else if (Implicit0.IsConvex() && Implicit1Type == THeightField<T>::StaticType())
 			{
 				UpdateConvexHeightFieldManifold(ConstraintBase, Transform0, Transform1, CullDistance);
@@ -1387,15 +1389,6 @@
 			{
 				UpdateCapsuleHeightFieldConstraint(*Implicit0.template GetObject<TCapsule<T>>(), Transform0, *Implicit1.template GetObject< THeightField<T> >(), Transform1, CullDistance, *ConstraintBase.template As<TRigidBodyPointContactConstraint<T, d>>());
 			}
-			else if (Implicit0Type == THeightField<T>::StaticType() && Implicit1Type == TBox<T, d>::StaticType())
-			{
-				//     This case should not be necessary. The height fields 
-				//     will only ever be collided against, so ideally will never 
-				//     be in index[0] position of the constraint, also the construction
-				//     of the constraint will just switch the index position so its always 
-				//     second.
-				ensure(false);
-			}
 			else if (Implicit0Type == TPlane<T, d>::StaticType() && Implicit1Type == TBox<T, d>::StaticType())
 			{
 				TRigidBodyPointContactConstraint<T, d>& Constraint = *ConstraintBase.template As<TRigidBodyPointContactConstraint<T, d> >();
@@ -1407,15 +1400,6 @@
 					Constraint.SetNormal(-Constraint.GetNormal());
 				}
 			}
-			else if (Implicit0Type == THeightField<T>::StaticType() && Implicit1Type == TSphere<T, d>::StaticType())
-			{
-				//     This case should not be necessary. The height fields 
-				//     will only ever be collided against, so ideally will never 
-				//     be in index[0] position of the constraint, also the construction
-				//     of the constraint will just switch the index position so its always 
-				//     second.
-				ensure(false);
-			}
 			else if (Implicit0Type == TPlane<T, d>::StaticType() && Implicit1Type == TSphere<T, d>::StaticType())
 			{
 				TRigidBodyPointContactConstraint<T, d>& Constraint = *ConstraintBase.template As<TRigidBodyPointContactConstraint<T, d> >();
@@ -1460,6 +1444,28 @@
 					Constraint.SetNormal(-Constraint.GetNormal());
 				}
 			}
+			else if (Implicit0Type == TBox<T, d>::StaticType() && Implicit1Type == FTriangleMeshImplicitObject::StaticType())
+			{
+				UpdateBoxTriangleMeshConstraint(Implicit0.template GetObject<TBox<T, d>>()->GetAABB(), Transform0, *GetInnerObject<FTriangleMeshImplicitObject>(Implicit1), Transform1, CullDistance, *ConstraintBase.template As<TRigidBodyPointContactConstraint<T, d>>());
+			}
+			else if (Implicit0Type == TSphere<T, d>::StaticType() && Implicit1Type == FTriangleMeshImplicitObject::StaticType())
+			{
+				UpdateSphereTriangleMeshConstraint(*Implicit0.template GetObject<TSphere<T, d>>(), Transform0, *GetInnerObject<FTriangleMeshImplicitObject>(Implicit1), Transform1, CullDistance, *ConstraintBase.template As<TRigidBodyPointContactConstraint<T, d>>());
+			}
+			else if (Implicit0Type == TCapsule<T>::StaticType() && Implicit1Type == FTriangleMeshImplicitObject::StaticType())
+			{
+				UpdateCapsuleTriangleMeshConstraint(*Implicit0.template GetObject<TCapsule<T>>(), Transform0, *GetInnerObject<FTriangleMeshImplicitObject>(Implicit1), Transform1, CullDistance, *ConstraintBase.template As<TRigidBodyPointContactConstraint<T, d>>());
+			}
+#if !UE_BUILD_SHIPPING
+			else if (Implicit0Type == THeightField<T>::StaticType() && Implicit1Type == TBox<T, d>::StaticType())
+			{
+				//     This case should not be necessary. The height fields 
+				//     will only ever be collided against, so ideally will never 
+				//     be in index[0] position of the constraint, also the construction
+				//     of the constraint will just switch the index position so its always 
+				//     second.
+				ensure(false);
+			}
 			else if (Implicit0Type == THeightField<T>::StaticType() && Implicit1Type == TSphere<T, d>::StaticType())
 			{
 				//     This case should not be necessary. The height fields 
@@ -1469,9 +1475,14 @@
 				//     second.
 				ensure(false);
 			}
-			else if (Implicit0Type == TBox<T, d>::StaticType() && Implicit1Type == FTriangleMeshImplicitObject::StaticType())
+			else if (Implicit0Type == THeightField<T>::StaticType() && Implicit1Type == TSphere<T, d>::StaticType())
 			{
-				UpdateBoxTriangleMeshConstraint(Implicit0.template GetObject<TBox<T, d>>()->GetAABB(), Transform0, *GetInnerObject<FTriangleMeshImplicitObject>(Implicit1), Transform1, CullDistance, *ConstraintBase.template As<TRigidBodyPointContactConstraint<T, d>>());
+				//     This case should not be necessary. The height fields 
+				//     will only ever be collided against, so ideally will never 
+				//     be in index[0] position of the constraint, also the construction
+				//     of the constraint will just switch the index position so its always 
+				//     second.
+				ensure(false);
 			}
 			else if (Implicit0Type == FTriangleMeshImplicitObject::StaticType() && Implicit1Type == TBox<T, d>::StaticType())
 			{
@@ -1482,10 +1493,6 @@
 				//     second.
 				ensure(false);
 			}
-			else if (Implicit0Type == TSphere<T, d>::StaticType() && Implicit1Type == FTriangleMeshImplicitObject::StaticType())
-			{
-				UpdateSphereTriangleMeshConstraint(*Implicit0.template GetObject<TSphere<T, d>>(), Transform0, *GetInnerObject<FTriangleMeshImplicitObject>(Implicit1), Transform1, CullDistance, *ConstraintBase.template As<TRigidBodyPointContactConstraint<T, d>>());
-			}
 			else if (Implicit0Type == FTriangleMeshImplicitObject::StaticType() && Implicit1Type == TSphere<T, d>::StaticType())
 			{
 				//     This case should not be necessary. The triangle mesh
@@ -1495,10 +1502,6 @@
 				//     second.
 				ensure(false);
 			}
-			else if (Implicit0Type == TCapsule<T>::StaticType() && Implicit1Type == FTriangleMeshImplicitObject::StaticType())
-			{
-				UpdateCapsuleTriangleMeshConstraint(*Implicit0.template GetObject<TCapsule<T>>(), Transform0, *GetInnerObject<FTriangleMeshImplicitObject>(Implicit1), Transform1, CullDistance, *ConstraintBase.template As<TRigidBodyPointContactConstraint<T, d>>());
-			}
 			else if (Implicit0Type == FTriangleMeshImplicitObject::StaticType() && Implicit1Type == TCapsule<T>::StaticType() )
 			{
 				//     This case should not be necessary. The triangle mesh
@@ -1524,6 +1527,7 @@
 				//     second.
 				ensure(false);
 			}
+#endif
 			else if (Implicit0.IsConvex() && Implicit1Type == THeightField<T>::StaticType())
 			{
 				UpdateConvexHeightFieldConstraint(Implicit0, Transform0, *Implicit1.template GetObject< THeightField<T> >(), Transform1, CullDistance, *ConstraintBase.template As<TRigidBodyPointContactConstraint<T, d>>());
