Chaos: fix for BufferPhysicsResults, FlipBuffers crashes.

diff --git a/Engine/Source/Runtime/Engine/Private/PhysicsEngine/Experimental/PhysScene_Chaos.cpp b/Engine/Source/Runtime/Engine/Private/PhysicsEngine/Experimental/PhysScene_Chaos.cpp
--- a/Engine/Source/Runtime/Engine/Private/PhysicsEngine/Experimental/PhysScene_Chaos.cpp
+++ b/Engine/Source/Runtime/Engine/Private/PhysicsEngine/Experimental/PhysScene_Chaos.cpp
@@ -1735,6 +1735,15 @@ void FPhysScene_ChaosInterface::CompleteSceneSimulation(ENamedThreads::Type Curr
 		{
 			if(Solver->HasActiveParticles())
 			{
+				// We need to process the Command Q here to prevent corruption in BufferPhysicsResults & FlipBuffers
+				// needs further investigation.
+				TQueue<TFunction<void(FPhysicsSolver*)>, EQueueMode::Mpsc>& Queue = Solver->CommandQueue;
+				TFunction<void(FPhysicsSolver*)> Command;
+				while (Queue.Dequeue(Command))
+				{
+					Command(Solver);
+				}
+
 				ActiveSolvers.Add(Solver);
 			}
 		}
diff --git a/Engine/Source/Runtime/Experimental/ChaosSolvers/Private/PBDRigidsSolver.cpp b/Engine/Source/Runtime/Experimental/ChaosSolvers/Private/PBDRigidsSolver.cpp
--- a/Engine/Source/Runtime/Experimental/ChaosSolvers/Private/PBDRigidsSolver.cpp
+++ b/Engine/Source/Runtime/Experimental/ChaosSolvers/Private/PBDRigidsSolver.cpp
@@ -618,22 +618,25 @@ namespace Chaos
 		TParticleView<TPBDRigidParticles<float, 3>>& ActiveParticles = GetParticles().GetActiveParticlesView();
 		for (Chaos::TPBDRigidParticleHandleImp<float, 3, false>& ActiveObject : ActiveParticles)
 		{
-			switch (ActiveObject.GetParticleType())
+			if (ensure(ActiveObject.GTGeometryParticle()->Proxy))
 			{
-			case Chaos::EParticleType::Rigid:
-				((FRigidParticlePhysicsProxy*)(ActiveObject.GTGeometryParticle()->Proxy))->FlipBuffer();
-				break;
-			case Chaos::EParticleType::Kinematic:
-				((FKinematicGeometryParticlePhysicsProxy*)(ActiveObject.GTGeometryParticle()->Proxy))->FlipBuffer();
-				break;
-			case Chaos::EParticleType::Static:
-				((FGeometryParticlePhysicsProxy*)(ActiveObject.GTGeometryParticle()->Proxy))->FlipBuffer();
-				break;
-			case Chaos::EParticleType::GeometryCollection:
-				ActiveGC.AddUnique((FGeometryCollectionPhysicsProxy*)(ActiveObject.GTGeometryParticle()->Proxy));
-				break;
-			default:
-				check(false);
+				switch (ActiveObject.GetParticleType())
+				{
+				case Chaos::EParticleType::Rigid:
+					((FRigidParticlePhysicsProxy*)(ActiveObject.GTGeometryParticle()->Proxy))->FlipBuffer();
+					break;
+				case Chaos::EParticleType::Kinematic:
+					((FKinematicGeometryParticlePhysicsProxy*)(ActiveObject.GTGeometryParticle()->Proxy))->FlipBuffer();
+					break;
+				case Chaos::EParticleType::Static:
+					((FGeometryParticlePhysicsProxy*)(ActiveObject.GTGeometryParticle()->Proxy))->FlipBuffer();
+					break;
+				case Chaos::EParticleType::GeometryCollection:
+					ActiveGC.AddUnique((FGeometryCollectionPhysicsProxy*)(ActiveObject.GTGeometryParticle()->Proxy));
+					break;
+				default:
+					check(false);
+				}
 			}
 		}
 
@@ -652,22 +655,25 @@ namespace Chaos
 		TParticleView<TPBDRigidParticles<float, 3>>& ActiveParticles = GetParticles().GetActiveParticlesView();
 		for (Chaos::TPBDRigidParticleHandleImp<float, 3, false>& ActiveObject : ActiveParticles)
 		{
-			switch (ActiveObject.GetParticleType())
+			if (ensure(ActiveObject.GTGeometryParticle()->Proxy))
 			{
-			case Chaos::EParticleType::Rigid:
-				((FRigidParticlePhysicsProxy*)(ActiveObject.GTGeometryParticle()->Proxy))->PullFromPhysicsState();
-				break;
-			case Chaos::EParticleType::Kinematic:
-				((FKinematicGeometryParticlePhysicsProxy*)(ActiveObject.GTGeometryParticle()->Proxy))->PullFromPhysicsState();
-				break;
-			case Chaos::EParticleType::Static:
-				((FGeometryParticlePhysicsProxy*)(ActiveObject.GTGeometryParticle()->Proxy))->PullFromPhysicsState();
-				break;
-			case Chaos::EParticleType::GeometryCollection:
-				ActiveGC.AddUnique((FGeometryCollectionPhysicsProxy*)(ActiveObject.GTGeometryParticle()->Proxy));
-				break;
-			default:
-				check(false);
+				switch (ActiveObject.GetParticleType())
+				{
+				case Chaos::EParticleType::Rigid:
+					((FRigidParticlePhysicsProxy*)(ActiveObject.GTGeometryParticle()->Proxy))->PullFromPhysicsState();
+					break;
+				case Chaos::EParticleType::Kinematic:
+					((FKinematicGeometryParticlePhysicsProxy*)(ActiveObject.GTGeometryParticle()->Proxy))->PullFromPhysicsState();
+					break;
+				case Chaos::EParticleType::Static:
+					((FGeometryParticlePhysicsProxy*)(ActiveObject.GTGeometryParticle()->Proxy))->PullFromPhysicsState();
+					break;
+				case Chaos::EParticleType::GeometryCollection:
+					ActiveGC.AddUnique((FGeometryCollectionPhysicsProxy*)(ActiveObject.GTGeometryParticle()->Proxy));
+					break;
+				default:
+					check(false);
+				}
 			}
 		}
 
