Fixed memory leak in Chaos material sync (leak in underlying handle array copy).
Added more Chaos LLM tags to track material allocations

diff --git a/Engine/Source/Runtime/Experimental/Chaos/Private/Chaos/PhysicalMaterials.cpp b/Engine/Source/Runtime/Experimental/Chaos/Private/Chaos/PhysicalMaterials.cpp
--- a/Engine/Source/Runtime/Experimental/Chaos/Private/Chaos/PhysicalMaterials.cpp
+++ b/Engine/Source/Runtime/Experimental/Chaos/Private/Chaos/PhysicalMaterials.cpp
@@ -1,6 +1,7 @@
 // Copyright Epic Games, Inc. All Rights Reserved.
 
 #include "Chaos/PhysicalMaterials.h"
+#include "HAL/LowLevelMemTracker.h"
 
 namespace Chaos
 {
@@ -58,6 +59,8 @@ namespace Chaos
 
 	FMaterialHandle FPhysicalMaterialManager::Create()
 	{
+		LLM_SCOPE(ELLMTag::Chaos);
+
 		check(IsInGameThread());
 		FMaterialHandle OutHandle;
 		OutHandle.InnerHandle = Materials.Create();
@@ -70,6 +73,8 @@ namespace Chaos
 
 	void FPhysicalMaterialManager::Destroy(FMaterialHandle InHandle)
 	{
+		LLM_SCOPE(ELLMTag::Chaos);
+
 		check(IsInGameThread());
 		if(InHandle.InnerHandle.IsValid())
 		{
diff --git a/Engine/Source/Runtime/Experimental/Chaos/Public/Chaos/Framework/Handles.h b/Engine/Source/Runtime/Experimental/Chaos/Public/Chaos/Framework/Handles.h
--- a/Engine/Source/Runtime/Experimental/Chaos/Public/Chaos/Framework/Handles.h
+++ b/Engine/Source/Runtime/Experimental/Chaos/Public/Chaos/Framework/Handles.h
@@ -211,18 +211,7 @@ namespace Chaos
 
 		~THandleArray()
 		{
-			if(Data)
-			{
-				for(int32 Index = 0; Index < NumData; ++Index)
-				{
-					if(Validity[Index])
-					{
-						DestructItem(Data + Index);
-					}
-				}
-
-				FMemory::Free(Data);
-			}
+			DestroyItems();
 		}
 
 		ElementType* Get(FHandle InHandle) const
@@ -405,8 +394,34 @@ namespace Chaos
 
 	private:
 
+		void DestroyItems()
+		{
+			if(Data)
+			{
+				for(int32 Index = 0; Index < NumData; ++Index)
+				{
+					if(Validity[Index])
+					{
+						DestructItem(Data + Index);
+					}
+				}
+
+				FMemory::Free(Data);
+			}
+
+			FreeList = InvalidFreeIndex;
+			Capacity = 0;
+			NumData = 0;
+			NumActive = 0;
+			Data = nullptr;
+			Validity.Reset();
+			HandleEntries.Reset();
+		}
+
 		void CopyFrom(const THandleArray& Other)
 		{
+			DestroyItems();
+
 			Capacity = Other.Capacity;
 			HandleEntries = Other.HandleEntries;
 			Validity = Other.Validity;
@@ -431,6 +446,8 @@ namespace Chaos
 
 		void MoveFrom(THandleArray& Other)
 		{
+			DestroyItems();
+
 			Capacity = Other.Capacity;
 
 			if(Capacity > 0)
@@ -707,6 +724,8 @@ namespace Chaos
 
 		void CopyFrom(const THandleHeap<ElementType, IndexWidth, GenerationWidth>& Other)
 		{
+			Empty();
+
 			FreeList = Other.FreeList;
 			NumActive = Other.NumActive;
 			Validity = Other.Validity;
@@ -724,6 +743,8 @@ namespace Chaos
 
 		void MoveFrom(THandleHeap<ElementType, IndexWidth, GenerationWidth>& Other)
 		{
+			Empty();
+
 			FreeList = Other.FreeList;
 			NumActive = Other.NumActive;
 
diff --git a/Engine/Source/Runtime/Experimental/ChaosSolvers/Private/ChaosSolversModule.cpp b/Engine/Source/Runtime/Experimental/ChaosSolvers/Private/ChaosSolversModule.cpp
--- a/Engine/Source/Runtime/Experimental/ChaosSolvers/Private/ChaosSolversModule.cpp
+++ b/Engine/Source/Runtime/Experimental/ChaosSolvers/Private/ChaosSolversModule.cpp
@@ -464,6 +464,8 @@ Chaos::FPhysicsSolver* FChaosSolversModule::CreateSolver(bool bStandalone /*= fa
 #endif
 )
 {
+	LLM_SCOPE(ELLMTag::Chaos);
+
 	FChaosScopeSolverLock SolverScopeLock;
 	
 	Chaos::EMultiBufferMode SolverBufferMode = Chaos::EMultiBufferMode::Single;
@@ -495,6 +497,7 @@ Chaos::FPhysicsSolver* FChaosSolversModule::CreateSolver(bool bStandalone /*= fa
 		// Need to let the thread know there's a new solver to care about
 		Dispatcher->EnqueueCommandImmediate([NewSolver](Chaos::FPersistentPhysicsTask* PhysThread)
 		{
+			LLM_SCOPE(ELLMTag::Chaos);
 			PhysThread->AddSolver(NewSolver);
 		});
 	}
@@ -528,6 +531,8 @@ void FChaosSolversModule::SetDedicatedThreadTickMode(EChaosSolverTickMode InTick
 
 void FChaosSolversModule::DestroySolver(Chaos::FPhysicsSolver* InSolver)
 {
+	LLM_SCOPE(ELLMTag::Chaos);
+
 	FChaosScopeSolverLock SolverScopeLock;
 
 	if(Solvers.Remove(InSolver) > 0)
@@ -536,6 +541,7 @@ void FChaosSolversModule::DestroySolver(Chaos::FPhysicsSolver* InSolver)
 		{
 			Dispatcher->EnqueueCommandImmediate([InSolver](Chaos::FPersistentPhysicsTask* PhysThread)
 			{
+				LLM_SCOPE(ELLMTag::Chaos);
 				if(PhysThread)
 				{
 					PhysThread->RemoveSolver(InSolver);
