Chaos : Collisions

- TriangleMesh to Convex collisions

diff --git a/Engine/Source/Runtime/Experimental/Chaos/Private/Chaos/CollisionResolution.cpp b/Engine/Source/Runtime/Experimental/Chaos/Private/Chaos/CollisionResolution.cpp
--- a/Engine/Source/Runtime/Experimental/Chaos/Private/Chaos/CollisionResolution.cpp
+++ b/Engine/Source/Runtime/Experimental/Chaos/Private/Chaos/CollisionResolution.cpp
@@ -1027,6 +1027,44 @@ namespace Chaos
 			}
 		}
 
+		//
+		// Convex-TriangleMesh
+		//
+
+		template <typename T, int d>
+		TContactPoint<T> ConvexTriangleMeshContactPoint(const FImplicitObject& A, const TRigidTransform<T, d>& ATransform, const FTriangleMeshImplicitObject& B, const TRigidTransform<T, d>& BTransform, const T CullDistance)
+		{
+			return GJKImplicitContactPoint< FConvex >(A, ATransform, B, BTransform, CullDistance);
+		}
+
+		template <typename T, int d>
+		void UpdateConvexTriangleMeshConstraint(const FImplicitObject& Convex0, const TRigidTransform<T, d>& Transform0, const FTriangleMeshImplicitObject& TriangleMesh1, const TRigidTransform<T, d>& Transform1, const T CullDistance, TRigidBodyPointContactConstraint<T, d>& Constraint)
+		{
+			UpdateContactPoint(Constraint.Manifold, ConvexTriangleMeshContactPoint(Convex0, Transform0, TriangleMesh1, Transform1, CullDistance));
+		}
+
+
+		template <typename T, int d>
+		void UpdateConvexTriangleMeshManifold(TCollisionConstraintBase<T, d>&  Constraint, const TRigidTransform<T, d>& ATM, const TRigidTransform<T, d>& BTM, const T CullDistance)
+		{
+
+		}
+
+		template<typename T, int d>
+		void ConstructConvexTriangleMeshConstraints(TGeometryParticleHandle<T, d>* Particle0, TGeometryParticleHandle<T, d>* Particle1, const FImplicitObject* Implicit0, const FImplicitObject* Implicit1, const TRigidTransform<T, d>& Transform0, const TRigidTransform<T, d>& Transform1, const T CullDistance, FCollisionConstraintsArray& NewConstraints)
+		{
+			const FTriangleMeshImplicitObject * Object1 = GetInnerObject<FTriangleMeshImplicitObject>(*Implicit1);
+			if (ensure(Implicit0->IsConvex() && Object1) )
+			{
+				TRigidTransform<T, d> ParticleImplicit0TM = Transform0.GetRelativeTransform(Collisions::GetTransform(Particle0));
+				TRigidTransform<T, d> ParticleImplicit1TM = Transform1.GetRelativeTransform(Collisions::GetTransform(Particle1));
+				FRigidBodyPointContactConstraint Constraint = FRigidBodyPointContactConstraint(Particle0, Implicit0, ParticleImplicit0TM, Particle1, Implicit1, ParticleImplicit1TM);
+				UpdateConvexTriangleMeshConstraint(*Implicit0, Transform0, *Object1, Transform1, CullDistance, Constraint);
+				NewConstraints.TryAdd(CullDistance, Constraint);
+			}
+		}
+
+
 
 
 		//
@@ -1277,6 +1315,14 @@ namespace Chaos
 			{
 				ensure(false);
 			}
+			else if (Implicit0Type == FConvex::StaticType() && Implicit1Type == FTriangleMeshImplicitObject::StaticType())
+			{
+				//UpdateConvexTriangleMeshManifold(ConstraintBase, Transform0, Transform1, CullDistance);
+			}
+			else if (Implicit0Type == FTriangleMeshImplicitObject::StaticType() && Implicit1Type == FConvex::StaticType())
+			{
+				ensure(false);
+			}
 			//
 			// the generic convex bodies are last
 			//
@@ -1512,7 +1558,6 @@ namespace Chaos
 				ensure(false);
 			}
 
-
 			//
 			// the generic convex bodies are last
 			//
@@ -1527,11 +1572,24 @@ namespace Chaos
 				//     second.
 				ensure(false);
 			}
+			else if (Implicit0Type == FTriangleMeshImplicitObject::StaticType() && Implicit1.IsConvex())
+			{
+				//     This case should not be necessary. The triangle mesh
+				//     will only ever be collided against, so ideally will never
+				//     be in index[0] position of the constraint, also the construction
+				//     of the constraint will just switch the index position so its always
+				//     second.
+				ensure(false);
+			}
 #endif
 			else if (Implicit0.IsConvex() && Implicit1Type == THeightField<T>::StaticType())
 			{
 				UpdateConvexHeightFieldConstraint(Implicit0, Transform0, *Implicit1.template GetObject< THeightField<T> >(), Transform1, CullDistance, *ConstraintBase.template As<TRigidBodyPointContactConstraint<T, d>>());
 			}
+			else if (Implicit0Type == FConvex::StaticType() && Implicit1Type == FTriangleMeshImplicitObject::StaticType())
+			{
+				UpdateConvexTriangleMeshConstraint(Implicit0, Transform0, *GetInnerObject<FTriangleMeshImplicitObject>(Implicit1), Transform1, CullDistance, *ConstraintBase.template As<TRigidBodyPointContactConstraint<T, d>>());
+			}
 			else if (Implicit0.IsConvex() && Implicit1.IsConvex())
 			{
 				UpdateConvexConvexConstraint(Implicit0, Transform0, Implicit1, Transform1, CullDistance, ConstraintBase);
@@ -1844,6 +1902,14 @@ namespace Chaos
 			{
 				ConstructConvexHeightFieldConstraints(Particle1, Particle0, Implicit1, Implicit0, Transform1, Transform0, CullDistance, NewConstraints);
 			}
+			else if (Implicit0->IsConvex() && Implicit1Type == FTriangleMeshImplicitObject::StaticType())
+			{
+				ConstructConvexTriangleMeshConstraints(Particle0, Particle1, Implicit0, Implicit1, Transform0, Transform1, CullDistance, NewConstraints);
+			}
+			else if (Implicit0Type == FTriangleMeshImplicitObject::StaticType() && Implicit1->IsConvex())
+			{
+				ConstructConvexTriangleMeshConstraints(Particle1, Particle0, Implicit1, Implicit0, Transform1, Transform0, CullDistance, NewConstraints);
+			}
 			else if (Implicit0->IsConvex() && Implicit1->IsConvex())
 			{
 				ConstructConvexConvexConstraints(Particle0, Particle1, Implicit0, Implicit1, Transform0, Transform1, CullDistance, NewConstraints);
