Improve precision for heightfield sweep and GJK in general

diff --git a/Engine/Source/Runtime/Experimental/Chaos/Private/Chaos/HeightField.cpp b/Engine/Source/Runtime/Experimental/Chaos/Private/Chaos/HeightField.cpp
--- a/Engine/Source/Runtime/Experimental/Chaos/Private/Chaos/HeightField.cpp
+++ b/Engine/Source/Runtime/Experimental/Chaos/Private/Chaos/HeightField.cpp
@@ -229,12 +229,15 @@ namespace Chaos
 					return false;
 				}
 
-				TTriangle<T> Convex(A, B, C);
+				//Convert into local space of A to get better precision
+
+				TTriangle<T> Triangle(TVec3<T>(0), B-A, C-A);
 
 				T Time;
-				TVector<T, 3> HitPosition;
+				TVector<T, 3> LocalHitPosition;
 				TVector<T, 3> HitNormal;
-				if(GJKRaycast2<T>(Convex, OtherGeom, StartTM, Dir, CurrentLength, Time, HitPosition, HitNormal, Thickness, bComputeMTD))
+				const TRigidTransform<T,3> LocalStartTM(StartTM.GetTranslation() - A,StartTM.GetRotation());
+				if(GJKRaycast2<T>(Triangle, OtherGeom, LocalStartTM, Dir, CurrentLength, Time, LocalHitPosition, HitNormal, Thickness, bComputeMTD))
 				{
 					if(Time < OutTime)
 					{
@@ -249,7 +252,7 @@ namespace Chaos
 						if(!bHole)
 						{
 							OutNormal = HitNormal;
-							OutPosition = HitPosition;
+							OutPosition = LocalHitPosition + A;
 							OutTime = Time;
 							OutFaceIndex = FaceIndex;
 
