Chaos : Collisions

- Reject simple collisions when AsComplex is being used.

diff --git a/Engine/Source/Runtime/Experimental/Chaos/Private/Chaos/CollisionResolution.cpp b/Engine/Source/Runtime/Experimental/Chaos/Private/Chaos/CollisionResolution.cpp
--- a/Engine/Source/Runtime/Experimental/Chaos/Private/Chaos/CollisionResolution.cpp
+++ b/Engine/Source/Runtime/Experimental/Chaos/Private/Chaos/CollisionResolution.cpp
@@ -1643,6 +1643,10 @@ namespace Chaos
 				{
 					return;
 				}
+				else if (Shape0->CollisionTraceType == Chaos_CTF_UseComplexAsSimple && Implicit0Type != ImplicitObjectType::TriangleMesh)
+				{
+					return;
+				}
 			} 
 			else if (Implicit0Type == ImplicitObjectType::TriangleMesh) 
 			{
@@ -1659,6 +1663,11 @@ namespace Chaos
 				{
 					return;
 				}
+				else if (Shape1->CollisionTraceType == Chaos_CTF_UseComplexAsSimple && Implicit1Type != ImplicitObjectType::TriangleMesh)
+				{
+					return;
+				}
+
 			}
 			else if (Implicit1Type == ImplicitObjectType::TriangleMesh)
 			{
