Fixed Chaos handle arrays freeing garbage ptrs on move construction into a temporary.

diff --git a/Engine/Source/Runtime/Experimental/Chaos/Public/Chaos/Framework/Handles.h b/Engine/Source/Runtime/Experimental/Chaos/Public/Chaos/Framework/Handles.h
--- a/Engine/Source/Runtime/Experimental/Chaos/Public/Chaos/Framework/Handles.h
+++ b/Engine/Source/Runtime/Experimental/Chaos/Public/Chaos/Framework/Handles.h
@@ -165,6 +165,7 @@ namespace Chaos
 		}
 
 		THandleArray(THandleArray<ElementType, IndexWidth, GenerationWidth>&& Other)
+			: THandleArray(0)
 		{
 			MoveFrom(Other);
 		}
