Get Field_LinearVelocity working with particle handles.

diff --git a/Engine/Source/Runtime/Experimental/ChaosSolvers/Private/PhysicsProxy/FieldSystemPhysicsProxy.cpp b/Engine/Source/Runtime/Experimental/ChaosSolvers/Private/PhysicsProxy/FieldSystemPhysicsProxy.cpp
--- a/Engine/Source/Runtime/Experimental/ChaosSolvers/Private/PhysicsProxy/FieldSystemPhysicsProxy.cpp
+++ b/Engine/Source/Runtime/Experimental/ChaosSolvers/Private/PhysicsProxy/FieldSystemPhysicsProxy.cpp
@@ -420,23 +420,31 @@ void FFieldSystemPhysicsProxy::FieldParameterUpdateCallback(
 				if (ensureMsgf(Command.RootNode->Type() == FFieldNode<FVector>::StaticType(),
 					TEXT("Field based evaluation of the simulations 'LinearVelocity' parameter expects FVector field inputs.")))
 				{
-					FFieldSystemPhysicsProxy::ContiguousIndices(IndicesArray, CurrentSolver, ResolutionType, IndicesArray.Num() != Particles.Size());
-					if (IndicesArray.Num())
+					if (Handles.Num())
 					{
-						TArrayView<ContextIndex> IndexView(&(IndicesArray[0]), IndicesArray.Num());
-
-						FVector * tptr = &(Particles.X(0));
-						TArrayView<FVector> SamplesView(tptr, int32(Particles.Size()));
+						TArrayView<FVector> SamplePointsView(&(SamplePoints[0]), SamplePoints.Num());
+						TArrayView<ContextIndex> SampleIndicesView(&(SampleIndices[0]), SampleIndices.Num());
+						FFieldContext Context(
+							SampleIndicesView,
+							SamplePointsView,
+							Command.MetaData);
 
-						FFieldContext Context{
-							IndexView,
-							SamplesView,
-							Command.MetaData
-						};
+						TArray<FVector> LocalResults;
+						LocalResults.AddUninitialized(Handles.Num());					
+						TArrayView<FVector> ResultsView(&(LocalResults[0]), LocalResults.Num());
+						static_cast<const FFieldNode<FVector> *>(Command.RootNode.Get())->Evaluate(
+							Context, ResultsView);
 
-						FVector * vptr = &(Particles.V(0));
-						TArrayView<FVector> ResultsView(vptr, Particles.Size());
-						static_cast<const FFieldNode<FVector> *>(Command.RootNode.Get())->Evaluate(Context, ResultsView);
+						int32 i = 0;
+						for (Chaos::TGeometryParticleHandle<float, 3>* Handle : Handles)
+						{
+							Chaos::TPBDRigidParticleHandle<float, 3>* RigidHandle = Handle->CastToRigidParticle();
+							if(RigidHandle && RigidHandle->ObjectState() == Chaos::EObjectStateType::Dynamic)
+							{
+								RigidHandle->V() += ResultsView[i];
+							}
+							++i;
+						}
 					}
 				}
 				CommandsToRemove.Add(CommandIndex);
@@ -983,7 +991,6 @@ void FFieldSystemPhysicsProxy::FieldForcesUpdateCallback(
 
 						TArrayView<FVector> SamplePointsView(&(SamplePoints[0]), SamplePoints.Num());
 						TArrayView<ContextIndex> SampleIndicesView(&(SampleIndices[0]), SampleIndices.Num());
-
 						FFieldContext Context(
 							SampleIndicesView,
 							SamplePointsView,
