Avoid accessing sleep events when GT particles are deleted

diff --git a/Engine/Source/Runtime/Experimental/ChaosSolvers/Private/EventDefaults.cpp b/Engine/Source/Runtime/Experimental/ChaosSolvers/Private/EventDefaults.cpp
--- a/Engine/Source/Runtime/Experimental/ChaosSolvers/Private/EventDefaults.cpp
+++ b/Engine/Source/Runtime/Experimental/ChaosSolvers/Private/EventDefaults.cpp
@@ -367,12 +367,12 @@ namespace Chaos
 
 			NonConstSolver->Particles.GetDynamicParticles().GetSleepDataLock().ReadLock();
 			auto& SolverSleepingData = NonConstSolver->Particles.GetDynamicParticles().GetSleepData();
-			for (const TSleepData<float, 3>& SleepData : SolverSleepingData)
+			for(const TSleepData<float, 3>& SleepData : SolverSleepingData)
 			{
-				if (SleepData.Particle)
+				if(SleepData.Particle)
 				{
 					TGeometryParticle<float, 3>* Particle = SleepData.Particle->GTGeometryParticle();
-					if (Particle->Proxy != nullptr)
+					if(Particle && Particle->Proxy != nullptr)
 					{
 						int32 NewIdx = EventSleepDataArray.Add(TSleepingData<float, 3>());
 						TSleepingData<float, 3>& SleepingDataArrayItem = EventSleepDataArray[NewIdx];
