Chaos: Remove possible duplicates from sleeping buffer when particle is removed.

diff --git a/Engine/Source/Runtime/Experimental/Chaos/Public/Chaos/PBDRigidsSOAs.h b/Engine/Source/Runtime/Experimental/Chaos/Public/Chaos/PBDRigidsSOAs.h
--- a/Engine/Source/Runtime/Experimental/Chaos/Public/Chaos/PBDRigidsSOAs.h
+++ b/Engine/Source/Runtime/Experimental/Chaos/Public/Chaos/PBDRigidsSOAs.h
@@ -202,14 +202,11 @@ public:
 			// TODO think about this case more
 			GetDynamicParticles().GetSleepDataLock().WriteLock();
 			auto& SleepData = GetDynamicParticles().GetSleepData();
-			for (int32 Idx = 0; Idx < SleepData.Num(); ++Idx)
-			{
-				if (SleepData[Idx].Particle == Particle)
-				{
-					SleepData.RemoveAtSwap(Idx);
-					break;
-				}
-			}
+
+			SleepData.RemoveAllSwap([Particle](TSleepData<T, d>& Entry) {
+				return Entry.Particle == Particle;
+			});
+
 			GetDynamicParticles().GetSleepDataLock().WriteUnlock();
 		}
 
