Fix edge cases with ClipLine and also avoid allocating memory. Unit test that exposed this will be submitted in a follow up once other bugs it exposed are fixed

diff --git a/Engine/Source/Runtime/Experimental/Chaos/Public/Chaos/HeightField.h b/Engine/Source/Runtime/Experimental/Chaos/Public/Chaos/HeightField.h
--- a/Engine/Source/Runtime/Experimental/Chaos/Public/Chaos/HeightField.h
+++ b/Engine/Source/Runtime/Experimental/Chaos/Public/Chaos/HeightField.h
@@ -394,68 +394,114 @@ namespace Chaos
 
 			bool ClipLine(TVector<T, 2>& InOutStart, TVector<T, 2>& InOutEnd) const
 			{
+				
 				// Test we don't need to clip at all, quite likely with a heightfield so optimize for it.
-				if(IsInside(InOutStart) && IsInside(InOutEnd))
+				const bool bStartInside = IsInside(InOutStart);
+				const bool bEndInside = IsInside(InOutEnd);
+				if(bStartInside && bEndInside)
 				{
 					return true;
 				}
 
-				TArray<T> HitTimes;
-				const TVector<T, 2> Extent = GetExtent();
-				float TA, TB;
-				if(Utilities::IntersectLineSegments2D(InOutStart, InOutEnd, Min, TVector<T, 2>(Min[0] + Extent[0], Min[1]), TA, TB))
+				const TVector<T,2> Dir = InOutEnd - InOutStart;
+
+				// Tiny ray not inside so must be outside
+				if(Dir.SizeSquared() < 1e-4)
 				{
-					HitTimes.Add(TA);
+					return false;
 				}
 
-				if(Utilities::IntersectLineSegments2D(InOutStart, InOutEnd, Min, TVector<T, 2>(Min[0], Min[1] + Extent[1]), TA, TB))
+				bool bPerpendicular[2];
+				TVector<T,2> InvDir;
+				for(int Axis = 0; Axis < 2; ++Axis)
 				{
-					HitTimes.Add(TA);
+					bPerpendicular[Axis] = Dir[Axis] == 0;
+					InvDir[Axis] = bPerpendicular[Axis] ? 0 : 1 / Dir[Axis];
 				}
 
-				if(Utilities::IntersectLineSegments2D(InOutStart, InOutEnd, Max, TVector<T, 2>(Max[0] - Extent[0], Max[1]), TA, TB))
+				
+
+				if(bStartInside)
 				{
-					HitTimes.Add(TA);
+					const T TimeToExit = ComputeTimeToExit(InOutStart,InvDir);
+					InOutEnd = InOutStart + Dir * TimeToExit;
+					return true;
 				}
 
-				if(Utilities::IntersectLineSegments2D(InOutStart, InOutEnd, Max, TVector<T, 2>(Max[0], Max[1] - Extent[1]), TA, TB))
+				if(bEndInside)
 				{
-					HitTimes.Add(TA);
+					const T TimeToExit = ComputeTimeToExit(InOutEnd,-InvDir);
+					InOutStart = InOutEnd - Dir * TimeToExit;
+					return true;
 				}
 
-				const int32 NumTimes = HitTimes.Num();
-				if(NumTimes > 0)
+				//start and end outside, need to see if we even intersect
+				T TimesToEnter[2] = {TNumericLimits<T>::Max(),TNumericLimits<T>::Max()};
+				T TimesToExit[2] = {TNumericLimits<T>::Max(),TNumericLimits<T>::Max()};
+				
+				for(int Axis = 0; Axis < 2; ++Axis)
 				{
-					// Can only ever be 0, 1 or 2 entries in here. First check if we're starting inside the box
-					// so we correctly set the clip extents
-					HitTimes.Sort();
-					const TVector<T, 2> TempStart = InOutStart;
-
-					if(IsInside(InOutStart) && ensure(NumTimes == 1))
+					if(bPerpendicular[Axis])
 					{
-						// we begin somewhere inside the box - just clip the end
-						InOutEnd = TempStart + HitTimes[0] * (InOutEnd - TempStart);
+						if(InOutStart[Axis] >= Min[Axis] && InOutStart[Axis] <= Max[Axis])
+						{
+							TimesToEnter[Axis] = 0;
+						}
 					}
 					else
 					{
-						// Clip the start and if necessary, the end (might not need to if it's already inside)
-						InOutStart = TempStart + HitTimes[0] * (InOutEnd - TempStart);
-
-						if(HitTimes.IsValidIndex(1))
+						if(Dir[Axis] > 0)
+						{
+							if(InOutStart[Axis] <= Max[Axis])
+							{
+								TimesToEnter[Axis] = FMath::Max<T>(Min[Axis] - InOutStart[Axis], 0) * InvDir[Axis];
+								TimesToExit[Axis] = (Max[Axis] - InOutStart[Axis])  * InvDir[Axis];
+							}
+						}
+						else if(Dir[Axis] < 0)
 						{
-							InOutEnd = TempStart + HitTimes[1] * (InOutEnd - TempStart);
+							if(InOutStart[Axis] >= Min[Axis])
+							{
+								TimesToEnter[Axis] = FMath::Max<T>(InOutStart[Axis] - Max[Axis],0) * InvDir[Axis];
+								TimesToExit[Axis] = (InOutStart[Axis] - Min[Axis]) * InvDir[Axis];
+							}
 						}
 					}
+				}
 
-					// Long rays can leave us barely outside the bounds, clamp the final results to avoid this
-					InOutStart = Clamp(InOutStart);
-					InOutEnd = Clamp(InOutEnd);
+				const T TimeToEnter = FMath::Max(TimesToEnter[0],TimesToEnter[1]);
+				const T TimeToExit = FMath::Min(TimesToExit[0],TimesToExit[1]);
 
-					// We must hit the bound in some way
-					return true;
+				if(TimeToExit < TimeToEnter)
+				{
+					//no intersection
+					return false;
 				}
 
-				return false;
+				InOutEnd = InOutStart + Dir * TimeToExit;
+				InOutStart = InOutStart + Dir * TimeToEnter;
+				return true;
+			}
+
+		private:
+			//This helper assumes Start is inside the min/max box and uses InvDir to compute how long it takes to exit
+			T ComputeTimeToExit(const TVector<T,2>& Start,const TVector<T,2>& InvDir) const
+			{
+				T Times[2] ={TNumericLimits<T>::Max(),TNumericLimits<T>::Max()};
+				for(int Axis = 0; Axis < 2; ++Axis)
+				{
+					if(InvDir[Axis] > 0)
+					{
+						Times[Axis] = (Max[Axis] - Start[Axis]) * InvDir[Axis];
+					}
+					else if(InvDir[Axis] < 0)
+					{
+						Times[Axis] = (Start[Axis] - Min[Axis]) * InvDir[Axis];
+					}
+				}
+
+				const T MinTime = FMath::Min(Times[0],Times[1]);
+				return MinTime;
 			}
 		};
 
