Chaos - Capsule-box manifold

diff --git a/Engine/Source/Runtime/Experimental/Chaos/Private/Chaos/ChaosDebugDraw.cpp b/Engine/Source/Runtime/Experimental/Chaos/Private/Chaos/ChaosDebugDraw.cpp
--- a/Engine/Source/Runtime/Experimental/Chaos/Private/Chaos/ChaosDebugDraw.cpp
+++ b/Engine/Source/Runtime/Experimental/Chaos/Private/Chaos/ChaosDebugDraw.cpp
@@ -14,6 +14,8 @@
 #include "Chaos/Sphere.h"
 #include "Chaos/Utilities.h"
 
+//#pragma optimize ("", off)
+
 namespace Chaos
 {
 	namespace DebugDraw
@@ -160,44 +162,52 @@ namespace Chaos
 
 		void DrawCollisionImpl(const FRigidTransform3& SpaceTransform, const TPBDCollisionConstraintHandle<float, 3>* ConstraintHandle, float ColorScale)
 		{
-			if (ConstraintHandle->GetType() == TCollisionConstraintBase<float,3>::FType::SinglePoint)
-			{
-				const TRigidBodyPointContactConstraint<FReal, 3>& Contact = ConstraintHandle->GetPointContact();
-				if (Contact.GetPhi() > 0)
-				{
-					ColorScale = ColorScale * (FReal)0.1;
-				}
+			const FCollisionConstraintBase& Contact = ConstraintHandle->GetContact();
+			const FReal ActiveColorScale = (Contact.GetPhi() > 0)? ColorScale * 0.1f : ColorScale;
 
-				FVec3 Location = SpaceTransform.TransformPosition(Contact.GetLocation());
-				FVec3 Normal = SpaceTransform.TransformVector(Contact.GetNormal());
+			FVec3 Location = SpaceTransform.TransformPosition(Contact.GetLocation());
+			FVec3 Normal = SpaceTransform.TransformVector(Contact.GetNormal());
 
-				if (ContactWidth > 0)
-				{
-					FColor C0 = (ColorScale * FColor(128, 0, 0)).ToFColor(false);
-					FMatrix Axes = FRotationMatrix::MakeFromX(Normal);
-					FDebugDrawQueue::GetInstance().DrawDebugCircle(Location, DrawScale * ContactWidth, 12, C0, false, KINDA_SMALL_NUMBER, DrawPriority, LineThickness, Axes.GetUnitAxis(EAxis::Y), Axes.GetUnitAxis(EAxis::Z), false);
-				}
-				if (ContactLen > 0)
-				{
-					FColor C1 = (ColorScale * FColor(255, 0, 0)).ToFColor(false);
-					FDebugDrawQueue::GetInstance().DrawDebugLine(Location, Location + DrawScale * ContactLen * Normal, C1, false, KINDA_SMALL_NUMBER, DrawPriority, LineThickness);
-				}
-				if (ContactPhiWidth > 0 && Contact.GetPhi() < FLT_MAX)
-				{
-					FColor C2 = (ColorScale * FColor(128, 128, 0)).ToFColor(false);
-					FMatrix Axes = FRotationMatrix::MakeFromX(Normal);
-					FDebugDrawQueue::GetInstance().DrawDebugCircle(Location - Contact.GetPhi() * Normal, DrawScale * ContactPhiWidth, 12, C2, false, KINDA_SMALL_NUMBER, DrawPriority, LineThickness, Axes.GetUnitAxis(EAxis::Y), Axes.GetUnitAxis(EAxis::Z), false);
-				}
-				if (ContactOwnerWidth > 0)
+			if (ContactWidth > 0)
+			{
+				bool bIsManifold = (ConstraintHandle->GetType() == FCollisionConstraintBase::FType::MultiPoint);
+				FColor C0 = bIsManifold? (ColorScale * FColor(0, 0, 200)).ToFColor(false) : (ColorScale * FColor(200, 0, 0)).ToFColor(false);
+				FMatrix Axes = FRotationMatrix::MakeFromX(Normal);
+				FDebugDrawQueue::GetInstance().DrawDebugCircle(Location, DrawScale * ContactWidth, 12, C0, false, KINDA_SMALL_NUMBER, DrawPriority, LineThickness, Axes.GetUnitAxis(EAxis::Y), Axes.GetUnitAxis(EAxis::Z), false);
+
+				if (bIsManifold)
 				{
-					FColor C3 = (ColorScale * FColor(128, 128, 128)).ToFColor(false);
-					FMatrix Axes = FRotationMatrix::MakeFromX(Normal);
-					FVec3 P0 = SpaceTransform.TransformPosition(ConstraintHandle->GetConstrainedParticles()[0]->X());
-					FVec3 P1 = SpaceTransform.TransformPosition(ConstraintHandle->GetConstrainedParticles()[1]->X());
-					FDebugDrawQueue::GetInstance().DrawDebugLine(Location, P0, C3, false, KINDA_SMALL_NUMBER, DrawPriority, LineThickness);
-					FDebugDrawQueue::GetInstance().DrawDebugLine(Location, P1, C3, false, KINDA_SMALL_NUMBER, DrawPriority, LineThickness);
+					const FRigidBodyMultiPointContactConstraint& MultiPointConstraint = *ConstraintHandle->GetContact().As<FRigidBodyMultiPointContactConstraint>();
+					TConstGenericParticleHandle<FReal, 3> PointsParticle = MultiPointConstraint.PointsParticleHandle();
+					FRigidTransform3 PointsTransform = FParticleUtilities::GetActorWorldTransform(PointsParticle) * SpaceTransform;
+					for (int32 SampleIndex = 1; SampleIndex < MultiPointConstraint.NumManifoldPoints(); ++SampleIndex)
+					{
+						FVec3 S0 = PointsTransform.TransformPosition(MultiPointConstraint.GetManifoldPoint(SampleIndex - 1));
+						FVec3 S1 = PointsTransform.TransformPosition(MultiPointConstraint.GetManifoldPoint(SampleIndex));
+						FDebugDrawQueue::GetInstance().DrawDebugLine(S0, S1, FColor::Orange, false, KINDA_SMALL_NUMBER, DrawPriority, LineThickness);
+					}
 				}
 			}
+			if (ContactLen > 0)
+			{
+				FColor C1 = (ActiveColorScale * FColor(255, 0, 0)).ToFColor(false);
+				FDebugDrawQueue::GetInstance().DrawDebugLine(Location, Location + DrawScale * ContactLen * Normal, C1, false, KINDA_SMALL_NUMBER, DrawPriority, LineThickness);
+			}
+			if (ContactPhiWidth > 0 && Contact.GetPhi() < FLT_MAX)
+			{
+				FColor C2 = (ActiveColorScale * FColor(128, 128, 0)).ToFColor(false);
+				FMatrix Axes = FRotationMatrix::MakeFromX(Normal);
+				FDebugDrawQueue::GetInstance().DrawDebugCircle(Location - Contact.GetPhi() * Normal, DrawScale * ContactPhiWidth, 12, C2, false, KINDA_SMALL_NUMBER, DrawPriority, LineThickness, Axes.GetUnitAxis(EAxis::Y), Axes.GetUnitAxis(EAxis::Z), false);
+			}
+			if (ContactOwnerWidth > 0)
+			{
+				FColor C3 = (ColorScale * FColor(128, 128, 128)).ToFColor(false);
+				FMatrix Axes = FRotationMatrix::MakeFromX(Normal);
+				FVec3 P0 = SpaceTransform.TransformPosition(ConstraintHandle->GetConstrainedParticles()[0]->X());
+				FVec3 P1 = SpaceTransform.TransformPosition(ConstraintHandle->GetConstrainedParticles()[1]->X());
+				FDebugDrawQueue::GetInstance().DrawDebugLine(Location, P0, C3, false, KINDA_SMALL_NUMBER, DrawPriority, LineThickness * 0.5f);
+				FDebugDrawQueue::GetInstance().DrawDebugLine(Location, P1, C3, false, KINDA_SMALL_NUMBER, DrawPriority, LineThickness * 0.5f);
+			}
 		}
 
 		void DrawJointConstraintImpl(const FRigidTransform3& SpaceTransform, const FVec3& InPa, const FVec3& InCa, const FVec3& InXa, const FMatrix33& Ra, const FVec3& InPb, const FVec3& InCb, const FVec3& InXb, const FMatrix33& Rb, int32 Level, int32 Index, FReal ColorScale, uint32 FeatureMask)
diff --git a/Engine/Source/Runtime/Experimental/Chaos/Private/Chaos/CollisionResolution.cpp b/Engine/Source/Runtime/Experimental/Chaos/Private/Chaos/CollisionResolution.cpp
--- a/Engine/Source/Runtime/Experimental/Chaos/Private/Chaos/CollisionResolution.cpp
+++ b/Engine/Source/Runtime/Experimental/Chaos/Private/Chaos/CollisionResolution.cpp
@@ -24,6 +24,13 @@ DECLARE_CYCLE_STAT(TEXT("Collisions::GJK"), STAT_Collisions_GJK, STATGROUP_Chaos
 
 //#pragma optimize("", off)
 
+bool Chaos_Collision_UseManifolds = true;
+FAutoConsoleVariableRef CVarChaosCollisionUseManifolds(TEXT("p.Chaos.Collision.UseManifolds"), Chaos_Collision_UseManifolds, TEXT("Enable/Disable use of manifoldes in collision."));
+
+bool Chaos_Collision_UseManifolds_Test = false;
+FAutoConsoleVariableRef CVarChaosCollisionUseManifoldsTest(TEXT("p.Chaos.Collision.UseManifoldsTest"), Chaos_Collision_UseManifolds_Test, TEXT("Enable/Disable use of manifoldes in collision."));
+
+
 namespace Chaos
 {
 	namespace Collisions
@@ -53,6 +60,53 @@ namespace Chaos
 			}
 		}
 
+		/**
+		 * Select the best manifold point as the new contact point and update the constraint data.
+		 */
+		template<class T, int d>
+		void UpdateConstraintFromManifold(TRigidBodyMultiPointContactConstraint<T, d>& Constraint, const TRigidTransform<T, d>& ParticleTransform0, const TRigidTransform<T, d>& ParticleTransform1, const T CullDistance)
+		{
+			// TEMP: Test code - run single-point algorithm instead. Allows us to debug draw the manifold without using it
+			if (Chaos_Collision_UseManifolds_Test)
+			{
+				TRigidBodyPointContactConstraint<T, d> PointConstraint(Constraint.Particle[0], Constraint.Manifold.Implicit[0], Constraint.ImplicitTransform[0], Constraint.Particle[1], Constraint.Manifold.Implicit[1], Constraint.ImplicitTransform[1], Constraint.Manifold.ShapesType);
+				UpdateConstraint<ECollisionUpdateType::Deepest>(PointConstraint, ParticleTransform0, ParticleTransform1, CullDistance);
+				Constraint.Manifold = PointConstraint.Manifold;
+				return;
+			}
+
+			// Get the plane and point transforms (depends which body owns the plane)
+			const FRigidTransform3 Transform0 = Constraint.ImplicitTransform[0] * ParticleTransform0;
+			const FRigidTransform3 Transform1 = Constraint.ImplicitTransform[1] * ParticleTransform1;
+			const FRigidTransform3& PlaneTransform = (Constraint.GetManifoldPlaneOwnerIndex() == 0) ? Transform0 : Transform1;
+			const FRigidTransform3& PointsTransform = (Constraint.GetManifoldPlaneOwnerIndex() == 0) ? Transform1 : Transform0;
+
+			// World-space manifold plane
+			FVec3 PlaneNormal = PlaneTransform.TransformVector(Constraint.GetManifoldPlaneNormal());
+			FVec3 PlanePos = PlaneTransform.TransformPosition(Constraint.GetManifoldPlanePosition());
+
+			// Select the best manifold point
+			const int32 NumPoints = Constraint.NumManifoldPoints();
+			for (int32 PointIndex = 0; PointIndex < NumPoints; ++PointIndex)
+			{
+				// World-space manifold point and distance to manifold plane
+				FVec3 PointPos = PointsTransform.TransformPosition(Constraint.GetManifoldPoint(PointIndex));
+				FReal PointDistance = FVec3::DotProduct(PointPos - PlanePos, PlaneNormal);
+
+				// If this is the deepest hit, select it
+				if (PointDistance < Constraint.Manifold.Phi)
+				{
+					// @todo(chaos): Consider using average of plane and point positions for contact location
+					FVec3 ContactPos = PointPos - PointDistance * PlaneNormal;
+					FVec3 ContactNormal = (Constraint.GetManifoldPlaneOwnerIndex() == 0) ? -PlaneNormal : PlaneNormal;
+					Constraint.Manifold.Phi = PointDistance;
+					Constraint.Manifold.Location = ContactPos;
+					Constraint.Manifold.Normal = ContactNormal;
+				}
+			}
+		}
+
+
 		template <typename T, int d, typename GeometryA, typename GeometryB>
 		TContactPoint<T> GJKContactPoint(const GeometryA& A, const TRigidTransform<T, d>& ATM, const GeometryB& B, const TRigidTransform<T, d>& BTM, const TVector<T, 3>& InitialDir)
 		{
@@ -150,26 +204,22 @@ namespace Chaos
 			//
 			int32 FaceIndex = Constraint.Manifold.Implicit[0]->FindClosestFaceAndVertices(Transform0.InverseTransformPosition(ContactPoint.Location), CollisionSamples, 1.f);
 
-			if (!ContactPoint.Normal.Equals(Constraint.PlaneNormal) || !Constraint.NumSamples())
-			{
-				Constraint.PlaneNormal = Transform1.InverseTransformVector(ContactPoint.Normal);
-				Constraint.PlanePosition = Transform1.InverseTransformPosition(ContactPoint.Location - ContactPoint.Phi*ContactPoint.Normal);
-			}
-
-
-			if (FaceIndex != Constraint.SourceNormalIndex || !Constraint.NumSamples())
+			bool bNewManifold = (FaceIndex != Constraint.GetManifoldPlaneFaceIndex()) || (Constraint.NumManifoldPoints() == 0);
+			if (bNewManifold)
 			{
-				Constraint.ResetSamples(CollisionSamples.Num());
-				Constraint.SourceNormalIndex = FaceIndex;
+				const FVec3 PlaneNormal = Transform1.InverseTransformVector(ContactPoint.Normal);
+				const FVec3 PlanePos = Transform1.InverseTransformPosition(ContactPoint.Location - ContactPoint.Phi*ContactPoint.Normal);
+				Constraint.SetManifoldPlane(1, FaceIndex, PlaneNormal, PlanePos);
 
 				//
 				// @todo(chaos) : Collision Manifold
 				//   Only save the four best samples and hard-code the size of Constraint.Samples to [len:4].
 				//   Currently this just grabs all points and uses the deepest point for resolution. 
 				//
+				Constraint.ResetManifoldPoints(CollisionSamples.Num());
 				for (FVec3 Sample : CollisionSamples)
 				{
-					Constraint.AddSample({ Sample,0.f });
+					Constraint.AddManifoldPoint(Sample);
 				}
 			}
 		}
@@ -180,9 +230,9 @@ namespace Chaos
 			auto SumSampleData = [&](TRigidBodyMultiPointContactConstraint<T, d>& LambdaConstraint) -> TVector<float, 3>
 			{
 				TVector<float, 3> ReturnValue(0);
-				for (int i = 0; i < LambdaConstraint.NumSamples(); i++)
+				for (int i = 0; i < LambdaConstraint.NumManifoldPoints(); i++)
 				{
-					ReturnValue += LambdaConstraint[i].X;
+					ReturnValue += LambdaConstraint.GetManifoldPoint(i);
 				}
 				return ReturnValue;
 			};
@@ -190,21 +240,22 @@ namespace Chaos
 			// iterative manifolds for non TConvex implicit objects that require sampling 
 			TContactPoint<T> ContactPoint = ConvexConvexContactPoint(*Constraint.Manifold.Implicit[0], Transform0, *Constraint.Manifold.Implicit[1], Transform1, CullDistance);
 
-			if (!ContactPoint.Normal.Equals(Constraint.PlaneNormal) || !Constraint.NumSamples())
+			if (!ContactPoint.Normal.Equals(Constraint.GetManifoldPlaneNormal()) || !Constraint.NumManifoldPoints())
 			{
-				Constraint.ResetSamples();
-				Constraint.PlaneNormal = Transform1.InverseTransformVector(ContactPoint.Normal);
-				Constraint.PlanePosition = Transform1.InverseTransformPosition(ContactPoint.Location - ContactPoint.Phi*ContactPoint.Normal);
+				Constraint.ResetManifoldPoints();
+				FVec3 PlaneNormal = Transform1.InverseTransformVector(ContactPoint.Normal);
+				FVec3 PlanePosition = Transform1.InverseTransformPosition(ContactPoint.Location - ContactPoint.Phi*ContactPoint.Normal);
+				Constraint.SetManifoldPlane(1, INDEX_NONE, PlaneNormal, PlanePosition);
 			}
 
 			TVector<T, d> SurfaceSample = Transform0.InverseTransformPosition(ContactPoint.Location);
-			if (Constraint.NumSamples() < 4)
+			if (Constraint.NumManifoldPoints() < 4)
 			{
-				Constraint.AddSample({ SurfaceSample,0.f });
+				Constraint.AddManifoldPoint(SurfaceSample);
 			}
-			else if (Constraint.NumSamples() == 4)
+			else if (Constraint.NumManifoldPoints() == 4)
 			{
-				TVector<T, d> Center = SumSampleData(Constraint) / Constraint.NumSamples();
+				TVector<T, d> Center = SumSampleData(Constraint) / Constraint.NumManifoldPoints();
 				T Delta = (Center - SurfaceSample).SizeSquared();
 
 				//
@@ -214,24 +265,25 @@ namespace Chaos
 				//
 				T SmallestDelta = FLT_MAX;
 				int32 SmallestIndex = 0;
-				for (int32 idx = 0; idx < Constraint.NumSamples(); idx++)
-					if (Constraint[idx].Delta < SmallestDelta) {
-						SmallestDelta = Constraint[idx].Delta;
+				for (int32 idx = 0; idx < Constraint.NumManifoldPoints(); idx++)
+				{
+					T IdxDelta = (Center - Constraint.GetManifoldPoint(idx)).SizeSquared();
+					if (IdxDelta < SmallestDelta)
+					{
+						SmallestDelta = IdxDelta;
 						SmallestIndex = idx;
 					}
+				}
 
-				if (Delta > SmallestDelta) {
-					Constraint[SmallestIndex] = { SurfaceSample,Delta };
+				if (Delta > SmallestDelta) 
+				{
+					Constraint.SetManifoldPoint(SmallestIndex, SurfaceSample);
 				}
 			}
 			else
 			{
 				ensure(false); // max of 4 points
 			}
-
-			typedef FRigidBodyMultiPointContactConstraint::FSampleData FSampleData;
-			TVector<T, d> Center = SumSampleData(Constraint) / Constraint.NumSamples();
-			for (int32 Index = 0; Index < Constraint.NumSamples(); Index++) { Constraint[Index].Delta = (Center - Constraint[Index].X).SizeSquared(); }
 		}
 
 		template <typename TPGeometryClass>
@@ -779,7 +831,6 @@ namespace Chaos
 		template<typename T, int d>
 		void ConstructCapsuleCapsuleConstraints(TGeometryParticleHandle<T, d>* Particle0, TGeometryParticleHandle<T, d>* Particle1, const FImplicitObject* Implicit0, const FImplicitObject* Implicit1, const TRigidTransform<T, d>& Transform0, const TRigidTransform<T, d>& Transform1, const T CullDistance, FCollisionConstraintsArray& NewConstraints)
 		{
-
 			const TCapsule<T> * Object0 = Implicit0->template GetObject<const TCapsule<T> >();
 			const TCapsule<T> * Object1 = Implicit1->template GetObject<const TCapsule<T> >();
 			if (ensure(Object0 && Object1))
@@ -809,11 +860,202 @@ namespace Chaos
 			UpdateContactPoint(Constraint.Manifold, CapsuleBoxContactPoint(A, ATransform, B, BTransform, InitialDir, CullDistance));
 		}
 
+		/**
+		 * Generate a Capsule-Box Manifold.
+		 *
+		 * This starts by finding the closest feature using GJK/EPA. Then it adds points to the manifold that depend on
+		 * the closest feature types on the box (face, edge, vertex) and capsule (edge, vertex).
+		 *
+		 * For box vertex collisions, only the single near point/plane is used in the manifold.
+		 * For box face collisions, the box owns the manifold plane, and the capsule edge nearest the box face is used to generate manifold points.
+		 * For box edge collision with the capsule edge, the capsule owns the manifold plane, and the capsule edge is projected onto the box plane to generate the manifold points.
+		 * For box edge with capsule vertex collisions, only the single near point/plane is used in the manifold.
+		 * 
+		 */
 		template<class T, int d>
-		void UpdateCapsuleBoxManifold(TCollisionConstraintBase<T, d>&  Constraint, const TRigidTransform<T, d>& ATM, const TRigidTransform<T, d>& BTM, const T CullDistance)
+		void UpdateCapsuleBoxManifold(const TCapsule<T>& Capsule, const TRigidTransform<T, d>& CapsuleTM, const TAABB<T, d>& Box, const TRigidTransform<T, d>& BoxTM, const T CullDistance, FRigidBodyMultiPointContactConstraint& Constraint)
 		{
-			// @todo(chaos) : Stub Update Manifold
-			//   Stub function for updating the manifold prior to the Apply and ApplyPushOut
+			Constraint.ResetManifoldPoints();
+
+			// Find the nearest points on the capsule and box
+			// Note: We flip the order for GJK so we get the normal in box space. This makes it easier to build the face-capsule manifold.
+			const TRigidTransform<T, d> CapsuleToBoxTM = CapsuleTM.GetRelativeTransform(BoxTM);
+
+			// NOTE: All GJK results in box-space
+			// @todo(ccaulfield): use center-to-center direction for InitialDir
+			FVec3 InitialDir = FVec3(1, 0, 0);
+			T Penetration;
+			FVec3 CapsuleClosestBoxSpace, BoxClosestBoxSpace, NormalBoxSpace;
+			{
+				SCOPE_CYCLE_COUNTER(STAT_Collisions_GJK);
+				if (!ensure(GJKPenetration<true>(Box, Capsule, CapsuleToBoxTM, Penetration, BoxClosestBoxSpace, CapsuleClosestBoxSpace, NormalBoxSpace, (T)0, InitialDir, (T)0)))
+				{
+					return;
+				}
+			}
+
+			// Cache the closest point so we don't need to re-iterate over the manifold on the first iteration
+			Constraint.Manifold.Location = BoxTM.TransformPosition(BoxClosestBoxSpace);
+			Constraint.Manifold.Normal = BoxTM.TransformVector(NormalBoxSpace);
+			Constraint.Manifold.Phi = -Penetration;
+
+			// Find the box feature that the near point is on
+			// Face, Edge, or Vertex can be determined from number of non-zero elements in the box-space normal.
+			const FReal ComponentEpsilon = 0.09f; // ~Sin(5deg) = angle tolerance
+			int32 NumNonZeroNormalComponents = 0;
+			int32 MaxComponentIndex = INDEX_NONE;
+			FReal MaxComponentValue = 0;
+			for (int32 ComponentIndex = 0; ComponentIndex < 3; ++ComponentIndex)
+			{
+				FReal AbsComponentValue = FMath::Abs(NormalBoxSpace[ComponentIndex]);
+				if (AbsComponentValue > ComponentEpsilon)
+				{
+					++NumNonZeroNormalComponents;
+				}
+				if (AbsComponentValue > MaxComponentValue)
+				{
+					MaxComponentValue = AbsComponentValue;
+					MaxComponentIndex = ComponentIndex;
+				}
+			}
+
+			// Make sure we actually have a feature to use
+			if (!ensure(MaxComponentIndex != INDEX_NONE))
+			{
+				return;
+			}
+
+			FVec3 CapsuleAxis = CapsuleToBoxTM.TransformVector(Capsule.GetAxis()); // Box space capsule axis
+			FReal CapsuleAxisNormal = FVec3::DotProduct(CapsuleAxis, NormalBoxSpace);
+			const bool bIsBoxFaceContact = (NumNonZeroNormalComponents == 1);
+			const bool bIsBoxEdgeContact = (NumNonZeroNormalComponents == 2);
+			const bool bIsBoxVertexContact = (NumNonZeroNormalComponents == 3);
+			const bool bIsCapsuleEdgeContact = (CapsuleAxisNormal < ComponentEpsilon);
+			const bool bIsCapsuleVertexContact = !bIsCapsuleEdgeContact;
+
+			// We just use the nearest point for these combinations, with the box as the plane owner
+			if (bIsBoxVertexContact || (bIsBoxEdgeContact && bIsCapsuleVertexContact))
+			{
+				Constraint.SetManifoldPlane(1, INDEX_NONE, NormalBoxSpace, BoxClosestBoxSpace);
+				Constraint.AddManifoldPoint(CapsuleToBoxTM.InverseTransformPosition(CapsuleClosestBoxSpace));
+				return;
+			}
+
+			// Get the radial direction from the capsule axis to the box face
+			FVec3 CapsuleRadiusDir = FVec3::CrossProduct(FVec3::CrossProduct(NormalBoxSpace, CapsuleAxis), CapsuleAxis);
+			if (!Utilities::NormalizeSafe(CapsuleRadiusDir))
+			{
+				// Capsule is perpendicular to the face, just leave the one-point manifold
+				Constraint.SetManifoldPlane(1, INDEX_NONE, NormalBoxSpace, BoxClosestBoxSpace);
+				Constraint.AddManifoldPoint(CapsuleToBoxTM.InverseTransformPosition(CapsuleClosestBoxSpace));
+				return;
+			}
+
+			// This is the axis of the box face we will be using below
+			int32 FaceAxisIndex = MaxComponentIndex;
+
+			// If this is a box face contact, the box is the manifold plane owner, regardless of the capsule feature involved. 
+			// We take the edge of the capsule nearest the face, clip it to the face, and use the clipped points as the manifold points.
+			// NOTE: The plane normal is not exactly the face normal, so this is not strictly the correct manifold. In particular
+			// it will over-limit the rotation for this frame. We could fix this by reverse-projecting the edge points based on the angle...
+			if (bIsBoxFaceContact)
+			{
+				// Initialize the manifold with the nearest plane and nearest point
+				Constraint.SetManifoldPlane(1, INDEX_NONE, NormalBoxSpace, BoxClosestBoxSpace);
+				Constraint.AddManifoldPoint(CapsuleToBoxTM.InverseTransformPosition(CapsuleClosestBoxSpace));
+
+				// Calculate the capsule edge nearest the face
+				FVec3 CapsuleVert0 = CapsuleToBoxTM.TransformPosition(Capsule.GetX1()) + Capsule.GetMargin() * CapsuleRadiusDir;
+				FVec3 CapsuleVert1 = CapsuleToBoxTM.TransformPosition(Capsule.GetX2()) + Capsule.GetMargin() * CapsuleRadiusDir;
+
+				// Clip the capsule edge to the axis-aligned box face
+				for (int32 ClipAxisIndex = 0; ClipAxisIndex < 3; ++ClipAxisIndex)
+				{
+					if (ClipAxisIndex != FaceAxisIndex)
+					{
+						bool bAcceptedPos = Utilities::ClipLineSegmentToAxisAlignedPlane(CapsuleVert0, CapsuleVert1, ClipAxisIndex, 1.0f, Box.Max()[ClipAxisIndex]);
+						bool bAcceptedNeg = Utilities::ClipLineSegmentToAxisAlignedPlane(CapsuleVert0, CapsuleVert1, ClipAxisIndex, -1.0f, Box.Min()[ClipAxisIndex]);
+						if (!bAcceptedPos || !bAcceptedNeg)
+						{
+							// Capsule edge is outside a face - stick with the single point we have
+							return;
+						}
+					}
+				}
+
+				// Add the clipped points to the manifold if they are not too close to each other, or the near point calculated above.
+				// Note: verts are in box space - need to be in capsule space
+				// @todo(ccaulfield): manifold point distance tolerance should be a per-solver or per object setting
+				const FReal DistanceToleranceSq = (0.1f * Capsule.GetHeight()) * (0.1f * Capsule.GetHeight());
+				bool bUserVert0 = ((CapsuleVert1 - CapsuleVert0).SizeSquared() > DistanceToleranceSq) && ((CapsuleVert0 - CapsuleClosestBoxSpace).SizeSquared() > DistanceToleranceSq);
+				bool bUseVert1 = ((CapsuleVert1 - CapsuleClosestBoxSpace).SizeSquared() > DistanceToleranceSq);
+				if (bUserVert0)
+				{ 
+					Constraint.AddManifoldPoint(CapsuleToBoxTM.InverseTransformPosition(CapsuleVert0));
+				}
+				if (bUseVert1)
+				{
+					Constraint.AddManifoldPoint(CapsuleToBoxTM.InverseTransformPosition(CapsuleVert1));
+				}
+
+				return;
+			}
+
+			// If the box edge and the capsule edge are the closest features, we treat the capsule as the
+			// plane owner. Then select the most-parallel box face as the source of manifold points. The
+			// manifold points are found by projecting the capsule verts onto the face and clipping to it.
+			// We have a capsule edge if the normal is perpendicular to the capsule axis
+			if (bIsBoxEdgeContact && bIsCapsuleEdgeContact)
+			{
+				// Initialize the manifold with the plane and point on the capsule (move to capsule space)
+				FVec3 NormalCapsuleSpace = CapsuleToBoxTM.InverseTransformVector(NormalBoxSpace);
+				FVec3 CapsuleClosestCapsuleSpace = CapsuleToBoxTM.InverseTransformPosition(CapsuleClosestBoxSpace);
+				Constraint.SetManifoldPlane(0, INDEX_NONE, -NormalCapsuleSpace, CapsuleClosestCapsuleSpace);
+				Constraint.AddManifoldPoint(BoxClosestBoxSpace);
+
+				// Project the capsule edge onto the box face
+				FReal FaceAxisSign = FMath::Sign(NormalBoxSpace[FaceAxisIndex]);
+				FReal FacePos = (FaceAxisSign >= 0.0f)? Box.Max()[FaceAxisIndex] : Box.Min()[FaceAxisIndex];
+				FVec3 CapsuleVert0 = CapsuleToBoxTM.TransformPosition(Capsule.GetX1());
+				FVec3 CapsuleVert1 = CapsuleToBoxTM.TransformPosition(Capsule.GetX2());
+				Utilities::ProjectPointOntoAxisAlignedPlane(CapsuleVert0, CapsuleRadiusDir, FaceAxisIndex, FaceAxisSign, FacePos);
+				Utilities::ProjectPointOntoAxisAlignedPlane(CapsuleVert1, CapsuleRadiusDir, FaceAxisIndex, FaceAxisSign, FacePos);
+
+				// Clip the capsule edge to the axis-aligned box face
+				for (int32 ClipAxisIndex = 0; ClipAxisIndex < 3; ++ClipAxisIndex)
+				{
+					if (ClipAxisIndex != FaceAxisIndex)
+					{
+						bool bAcceptedPos = Utilities::ClipLineSegmentToAxisAlignedPlane(CapsuleVert0, CapsuleVert1, ClipAxisIndex, 1.0f, Box.Max()[ClipAxisIndex]);
+						bool bAcceptedNeg = Utilities::ClipLineSegmentToAxisAlignedPlane(CapsuleVert0, CapsuleVert1, ClipAxisIndex, -1.0f, Box.Min()[ClipAxisIndex]);
+						if (!bAcceptedPos || !bAcceptedNeg)
+						{
+							// Capsule edge is outside a face - stick with the single point we have
+							return;
+						}
+					}
+				}
+
+				// Add the clipped points to the manifold if they are not too close to each other, or the near point calculated above.
+				// Note: verts are in box space - need to be in capsule space
+				// @todo(ccaulfield): manifold point distance tolerance should be a per-solver or per object setting
+				const FReal DistanceToleranceSq = (0.1f * Capsule.GetHeight()) * (0.1f * Capsule.GetHeight());
+				bool bUserVert0 = ((CapsuleVert1 - CapsuleVert0).SizeSquared() > DistanceToleranceSq) && ((CapsuleVert0 - BoxClosestBoxSpace).SizeSquared() > DistanceToleranceSq);
+				bool bUseVert1 = ((CapsuleVert1 - BoxClosestBoxSpace).SizeSquared() > DistanceToleranceSq);
+				if (bUserVert0)
+				{
+					Constraint.AddManifoldPoint(CapsuleVert0);
+				}
+				if (bUseVert1)
+				{
+					Constraint.AddManifoldPoint(CapsuleVert1);
+				}
+
+				return;
+			}
+
+			// All feature combinations should be covered above
+			ensure(false);
 		}
 
 		template<typename T, int d>
@@ -834,9 +1076,19 @@ namespace Chaos
 				{
 					TRigidTransform<T, d> ParticleImplicit0TM = Transform0.GetRelativeTransform(Collisions::GetTransform(Particle0));
 					TRigidTransform<T, d> ParticleImplicit1TM = Transform1.GetRelativeTransform(Collisions::GetTransform(Particle1));
-					FRigidBodyPointContactConstraint Constraint = FRigidBodyPointContactConstraint(Particle0, Implicit0, ParticleImplicit0TM, Particle1, Implicit1, ParticleImplicit1TM, EContactShapesType::CapsuleBox);
-					UpdateCapsuleBoxConstraint(*Object0, Transform0, Object1->BoundingBox(), Transform1, CullDistance, Constraint);
-					NewConstraints.TryAdd(CullDistance, Constraint);
+
+					if (Chaos_Collision_UseManifolds)
+					{ 
+						FRigidBodyMultiPointContactConstraint Constraint = FRigidBodyMultiPointContactConstraint(Particle0, Implicit0, ParticleImplicit0TM, Particle1, Implicit1, ParticleImplicit1TM, EContactShapesType::CapsuleBox);
+						UpdateCapsuleBoxManifold(*Object0, Transform0, Object1->BoundingBox(), Transform1, CullDistance, Constraint);
+						NewConstraints.TryAdd(CullDistance, Constraint);
+					}
+					else
+					{
+						FRigidBodyPointContactConstraint Constraint = FRigidBodyPointContactConstraint(Particle0, Implicit0, ParticleImplicit0TM, Particle1, Implicit1, ParticleImplicit1TM, EContactShapesType::CapsuleBox);
+						UpdateCapsuleBoxConstraint(*Object0, Transform0, Object1->BoundingBox(), Transform1, CullDistance, Constraint);
+						NewConstraints.TryAdd(CullDistance, Constraint);
+					}
 				}
 			}
 		}
@@ -939,21 +1191,21 @@ namespace Chaos
 
 				const TRigidTransform<T, d> AToBTM = Transform0.GetRelativeTransform(Transform1);
 
-				TPlane<T, d> CollisionPlane(Constraint.PlanePosition, Constraint.PlaneNormal);
+				TPlane<T, d> CollisionPlane(Constraint.GetManifoldPlanePosition(), Constraint.GetManifoldPlaneNormal());
 
 				// re-sample the constraint based on the distance from the collision plane.
-				for (int32 Idx = 0; Idx < Constraint.NumSamples(); Idx++)
+				for (int32 Idx = 0; Idx < Constraint.NumManifoldPoints(); Idx++)
 				{
-					Constraint[Idx].Manifold.Phi = CollisionPlane.PhiWithNormal(AToBTM.TransformPosition(Constraint[Idx].X), Constraint[Idx].Manifold.Normal);
-					Constraint[Idx].Manifold.Normal = Transform1.TransformVector(Constraint.PlaneNormal);
-					Constraint[Idx].Manifold.Location = Transform0.TransformPosition(Constraint[Idx].X);
+					FVec3 Location = Constraint.GetManifoldPoint(Idx);
+					FVec3 Normal;
+					FReal Phi = CollisionPlane.PhiWithNormal(AToBTM.TransformPosition(Location), Normal);
 
 					// save the best point for collision processing	
-					if (ContactPoint.Phi > Constraint[Idx].Manifold.Phi)
+					if (ContactPoint.Phi > Phi)
 					{
-						ContactPoint.Phi = Constraint[Idx].Manifold.Phi;
-						ContactPoint.Normal = Constraint[Idx].Manifold.Normal;
-						ContactPoint.Location = Constraint[Idx].Manifold.Location;
+						ContactPoint.Phi = Phi;
+						ContactPoint.Normal = Transform1.TransformVector(Constraint.GetManifoldPlaneNormal());
+						ContactPoint.Location = Transform0.TransformPosition(Location);
 					}
 				}
 			}
@@ -1170,6 +1422,12 @@ namespace Chaos
 		template<typename T, int d>
 		void UpdateManifold(TCollisionConstraintBase<T, d>& ConstraintBase, const TRigidTransform<T, d>& ATM, const TRigidTransform<T, d>& BTM, const T CullDistance)
 		{
+			if (ConstraintBase.GetType() != TCollisionConstraintBase<T, d>::FType::MultiPoint)
+			{
+				return;
+			}
+			TRigidBodyMultiPointContactConstraint<T, d>& Constraint = *ConstraintBase.template As<TRigidBodyMultiPointContactConstraint<T, d>>();
+
 			const FImplicitObject& Implicit0 = *ConstraintBase.Manifold.Implicit[0];
 			const FImplicitObject& Implicit1 = *ConstraintBase.Manifold.Implicit[1];
 
@@ -1214,7 +1472,8 @@ namespace Chaos
 			//    the blocks with the ensure(false), they are just for validation 
 			//    after the recent change. 
 			//
-
+			// @todo(chaos): use constraint ShapesType to index manifold update function
+			//
 			EImplicitObjectType Implicit0Type = GetInnerType(Implicit0.GetType());
 			EImplicitObjectType Implicit1Type = GetInnerType(Implicit1.GetType()); 
 			
@@ -1256,7 +1515,7 @@ namespace Chaos
 			}
 			else if (Implicit0Type == TCapsule<T>::StaticType() && Implicit1Type == TBox<T, d>::StaticType())
 			{
-				UpdateCapsuleBoxManifold(ConstraintBase, Transform0, Transform1, CullDistance);
+				UpdateCapsuleBoxManifold<T, d>(*Constraint.Manifold.Implicit[0]->template GetObject<const TCapsule<T>>(), Transform0, Constraint.Manifold.Implicit[1]->template GetObject<const TBox<T, d>>()->GetAABB(), Transform1, CullDistance, Constraint);
 			}
 			else if (Implicit0Type == TCapsule<T>::StaticType() && Implicit1Type == THeightField<T>::StaticType())
 			{
@@ -1348,7 +1607,7 @@ namespace Chaos
 		}
 
 		template<ECollisionUpdateType UpdateType, typename T, int d>
-		void UpdateConstraintAny(TCollisionConstraintBase<T, d>& ConstraintBase, const FImplicitObject& Implicit0, const FImplicitObject& Implicit1, const TRigidTransform<T, d>& Transform0, const TRigidTransform<T, d>& Transform1, const T CullDistance)
+		void UpdateAnyConstraint(TCollisionConstraintBase<T, d>& ConstraintBase, const FImplicitObject& Implicit0, const FImplicitObject& Implicit1, const TRigidTransform<T, d>& Transform0, const TRigidTransform<T, d>& Transform1, const T CullDistance)
 		{
 #if !UE_BUILD_SHIPPING
 			EImplicitObjectType Implicit0OuterType = Implicit0.GetType();
@@ -1603,26 +1862,37 @@ namespace Chaos
 		template<ECollisionUpdateType UpdateType, typename T, int d>
 		void UpdateConstraint(TCollisionConstraintBase<T, d>& ConstraintBase, const TRigidTransform<T, d>& ParticleTransform0, const TRigidTransform<T, d>& ParticleTransform1, const T CullDistance)
 		{
-			const FImplicitObject& Implicit0 = *ConstraintBase.Manifold.Implicit[0];
-			const FImplicitObject& Implicit1 = *ConstraintBase.Manifold.Implicit[1];
-
-			const TRigidTransform<T, d>& Transform0 = ConstraintBase.ImplicitTransform[0] * ParticleTransform0;
-			const TRigidTransform<T, d>& Transform1 = ConstraintBase.ImplicitTransform[1] * ParticleTransform1;
-
-			switch (ConstraintBase.Manifold.ShapesType)
+			if (ConstraintBase.GetType() == TCollisionConstraintBase<T, d>::FType::MultiPoint)
 			{
-			case EContactShapesType::CapsuleCapsule:
-				UpdateCapsuleCapsuleConstraint(*Implicit0.template GetObject<TCapsule<T>>(), Transform0, *Implicit1.template GetObject<TCapsule<T>>(), Transform1, CullDistance, *ConstraintBase.template As<TRigidBodyPointContactConstraint<T, d>>());
-				break;
-			case EContactShapesType::CapsuleBox:
-				UpdateCapsuleBoxConstraint(*Implicit0.template GetObject<TCapsule<T>>(), Transform0, Implicit1.template GetObject<TBox<T, d>>()->GetAABB(), Transform1, CullDistance, *ConstraintBase.template As<TRigidBodyPointContactConstraint<T, d>>());
-				break;
-			case EContactShapesType::BoxBox:
-				UpdateBoxBoxConstraint(Implicit0.template GetObject<TBox<T, d>>()->GetAABB(), Transform0, Implicit1.template GetObject<TBox<T, d>>()->GetAABB(), Transform1, CullDistance, *ConstraintBase.template As<TRigidBodyPointContactConstraint<T, d>>());
-				break;
-			default:
-				UpdateConstraintAny<UpdateType, T, d>(ConstraintBase, Implicit0, Implicit1, Transform0, Transform1, CullDistance);
-				break;
+				// Select the best point from the available manifold points as the current contact
+				TRigidBodyMultiPointContactConstraint<T, d>& MultiPointConstraint = *ConstraintBase.template As<TRigidBodyMultiPointContactConstraint<T, d>>();
+				UpdateConstraintFromManifold(MultiPointConstraint, ParticleTransform0, ParticleTransform1, CullDistance);
+			}
+			else if (ConstraintBase.GetType() == TCollisionConstraintBase<T, d>::FType::SinglePoint)
+			{
+				const FImplicitObject& Implicit0 = *ConstraintBase.Manifold.Implicit[0];
+				const FImplicitObject& Implicit1 = *ConstraintBase.Manifold.Implicit[1];
+				const TRigidTransform<T, d> Transform0 = ConstraintBase.ImplicitTransform[0] * ParticleTransform0;
+				const TRigidTransform<T, d> Transform1 = ConstraintBase.ImplicitTransform[1] * ParticleTransform1;
+
+				// Re-run collision detection to generate a new contact
+				// @todo(chaos): all constraints should have a ShapesType set in ConstructConstraints. Then we can use a function table for all.
+				TRigidBodyPointContactConstraint<T, d>& PointConstraint = *ConstraintBase.template As<TRigidBodyPointContactConstraint<T, d>>();
+				switch (ConstraintBase.Manifold.ShapesType)
+				{
+				case EContactShapesType::CapsuleCapsule:
+					UpdateCapsuleCapsuleConstraint(*Implicit0.template GetObject<TCapsule<T>>(), Transform0, *Implicit1.template GetObject<TCapsule<T>>(), Transform1, CullDistance, PointConstraint);
+					break;
+				case EContactShapesType::CapsuleBox:
+					UpdateCapsuleBoxConstraint(*Implicit0.template GetObject<TCapsule<T>>(), Transform0, Implicit1.template GetObject<TBox<T, d>>()->GetAABB(), Transform1, CullDistance, PointConstraint);
+					break;
+				case EContactShapesType::BoxBox:
+					UpdateBoxBoxConstraint(Implicit0.template GetObject<TBox<T, d>>()->GetAABB(), Transform0, Implicit1.template GetObject<TBox<T, d>>()->GetAABB(), Transform1, CullDistance, PointConstraint);
+					break;
+				default:
+					UpdateAnyConstraint<UpdateType, T, d>(ConstraintBase, Implicit0, Implicit1, Transform0, Transform1, CullDistance);
+					break;
+				}
 			}
 		}
 
@@ -1703,7 +1973,6 @@ namespace Chaos
 				return;
 			}
 
-
 			else if (Implicit0OuterType == TImplicitObjectInstanced<FConvex>::StaticType())
 			{
 				const TImplicitObjectInstanced<FConvex>* TransformedImplicit0 = Implicit0->template GetObject<const TImplicitObjectInstanced<FConvex>>();
@@ -1966,7 +2235,7 @@ namespace Chaos
 		template void ConstructCapsuleCapsuleConstraints<float, 3>(TGeometryParticleHandle<float, 3>* Particle0, TGeometryParticleHandle<float, 3>* Particle1, const FImplicitObject* Implicit0, const FImplicitObject* Implicit1, const TRigidTransform<float, 3>& Transform0, const TRigidTransform<float, 3>& Transform1, const float CullDistance, FCollisionConstraintsArray& NewConstraints);
 
 		template void UpdateCapsuleBoxConstraint<float, 3>(const TCapsule<float>& A, const TRigidTransform<float, 3>& ATransform, const TAABB<float, 3>& B, const TRigidTransform<float, 3>& BTransform, const float CullDistance, TRigidBodyPointContactConstraint<float, 3>& Constraint);
-		template void UpdateCapsuleBoxManifold<float, 3>(TCollisionConstraintBase<float, 3>&  Constraint, const TRigidTransform<float, 3>& ATM, const TRigidTransform<float, 3>& BTM, const float CullDistance);
+		template void UpdateCapsuleBoxManifold<float, 3>(const TCapsule<float>& Capsule, const TRigidTransform<float, 3>& CapsuleTM, const TAABB<float, 3>& Box, const TRigidTransform<float, 3>& BoxTM, const float CullDistance, FRigidBodyMultiPointContactConstraint& Constraint);
 		template void ConstructCapsuleBoxConstraints<float, 3>(TGeometryParticleHandle<float, 3>* Particle0, TGeometryParticleHandle<float, 3>* Particle1, const FImplicitObject* Implicit0, const FImplicitObject* Implicit1, const TRigidTransform<float, 3>& Transform0, const TRigidTransform<float, 3>& Transform1, const float CullDistance, FCollisionConstraintsArray& NewConstraints);
 
 		template void UpdateCapsuleHeightFieldConstraint<float, 3>(const TCapsule<float>& A, const TRigidTransform<float, 3>& ATransform, const THeightField<float>& B, const TRigidTransform<float, 3>& BTransform, const float CullDistance, TRigidBodyPointContactConstraint<float, 3>& Constraint);
@@ -1993,6 +2262,7 @@ namespace Chaos
 		template void UpdateManifold<float, 3>(TCollisionConstraintBase<float, 3>& ConstraintBase, const TRigidTransform<float, 3>& ATM, const TRigidTransform<float, 3>& BTM, const float CullDistance);
 		template void UpdateConstraint<ECollisionUpdateType::Any, float, 3>(TCollisionConstraintBase<float, 3>& ConstraintBase, const TRigidTransform<float, 3>& Transform0, const TRigidTransform<float, 3>& Transform1, const float CullDistance);
 		template void UpdateConstraint<ECollisionUpdateType::Deepest, float, 3>(TCollisionConstraintBase<float, 3>& ConstraintBase, const TRigidTransform<float, 3>& Transform0, const TRigidTransform<float, 3>& Transform1, const float CullDistance);
+		template void UpdateConstraintFromManifold(TRigidBodyMultiPointContactConstraint<float, 3>& Constraint, const TRigidTransform<float, 3>& Transform0, const TRigidTransform<float, 3>& Transform1, const float CullDistance);
 		template void ConstructConstraints<float, 3>(TGeometryParticleHandle<float, 3>* Particle0, TGeometryParticleHandle<float, 3>* Particle1, const FImplicitObject* Implicit0, const FImplicitObject* Implicit1, const TRigidTransform<float, 3>& Transform0, const TRigidTransform<float, 3>& Transform1, const float CullDistance, FCollisionConstraintsArray& NewConstraints);
 
 	} // Collisions
diff --git a/Engine/Source/Runtime/Experimental/Chaos/Private/Chaos/PBDCollisionConstraints.cpp b/Engine/Source/Runtime/Experimental/Chaos/Private/Chaos/PBDCollisionConstraints.cpp
--- a/Engine/Source/Runtime/Experimental/Chaos/Private/Chaos/PBDCollisionConstraints.cpp
+++ b/Engine/Source/Runtime/Experimental/Chaos/Private/Chaos/PBDCollisionConstraints.cpp
@@ -154,7 +154,11 @@ namespace Chaos
 
 		check(Handle != nullptr);
 		Handles.Add(Handle);
+
+#if CHAOS_COLLISION_PERSISTENCE_ENABLED
+		check(!Manifolds.Contains(Handle->GetKey()));
 		Manifolds.Add(Handle->GetKey(), Handle);
+#endif
 
 		UpdateConstraintMaterialProperties(PointConstraints[Idx]);
 	}
@@ -168,7 +172,11 @@ namespace Chaos
 
 		check(Handle != nullptr);
 		Handles.Add(Handle);
+
+#if CHAOS_COLLISION_PERSISTENCE_ENABLED
+		check(!Manifolds.Contains(Handle->GetKey()));
 		Manifolds.Add(Handle->GetKey(), Handle);
+#endif
 
 		UpdateConstraintMaterialProperties(IterativeConstraints[Idx]);
 	}
@@ -186,8 +194,8 @@ namespace Chaos
 	{
 		SCOPE_CYCLE_COUNTER(STAT_Collisions_Reset);
 
+#if CHAOS_COLLISION_PERSISTENCE_ENABLED
 		TArray<FConstraintContainerHandle*> CopyOfHandles = Handles;
-
 		int32 LifespanWindow = LifespanCounter - 1;
 		for (FConstraintContainerHandle* ContactHandle : CopyOfHandles)
 		{
@@ -196,7 +204,15 @@ namespace Chaos
 				RemoveConstraint(ContactHandle);
 			}
 		}
-
+#else
+		for (FConstraintContainerHandle* Handle : Handles)
+		{
+			HandleAllocator.FreeHandle(Handle);
+		}
+		PointConstraints.Reset();
+		IterativeConstraints.Reset();
+		Handles.Reset();
+#endif
 
 		MAngularFriction = 0;
 		bUseCCD = false;
@@ -245,23 +261,27 @@ namespace Chaos
 
 		if (ConstraintType == FCollisionConstraintBase::FType::SinglePoint)
 		{
+#if CHAOS_COLLISION_PERSISTENCE_ENABLED
 			if (Idx < PointConstraints.Num() - 1)
 			{
 				// update the handle
 				FConstraintContainerHandleKey Key = FPBDCollisionConstraintHandle::MakeKey(&PointConstraints.Last());
 				Manifolds[Key]->SetConstraintIndex(Idx, ConstraintType);
 			}
+#endif
 			PointConstraints.RemoveAtSwap(Idx);
 
 		}
 		else if (ConstraintType == FCollisionConstraintBase::FType::MultiPoint)
 		{
+#if CHAOS_COLLISION_PERSISTENCE_ENABLED
 			if (Idx < IterativeConstraints.Num() - 1)
 			{
 				// update the handle
 				FConstraintContainerHandleKey Key = FPBDCollisionConstraintHandle::MakeKey(&IterativeConstraints.Last());
 				Manifolds[Key]->SetConstraintIndex(Idx, ConstraintType);
 			}
+#endif
 			IterativeConstraints.RemoveAtSwap(Idx);
 		}
 		else 
@@ -272,7 +292,9 @@ namespace Chaos
 		// @todo(chaos): Collision Manifold
 		//   Add an index to the handle in the Manifold.Value 
 		//   to prevent the search in Handles when removed.
-		Manifolds.Remove(KeyToRemove);  
+#if CHAOS_COLLISION_PERSISTENCE_ENABLED
+		Manifolds.Remove(KeyToRemove);
+#endif
 		Handles.Remove(Handle);
 
 		ensure(Handles.Num() == PointConstraints.Num() + IterativeConstraints.Num());
@@ -287,16 +309,19 @@ namespace Chaos
 		// Clustering uses update constraints to force a re-evaluation. 
 	}
 
+	// Called once per frame to update persistent constraints (reruns collision detection, or selects the best manifold point)
 	template<typename T, int d>
 	void TPBDCollisionConstraints<T, d>::UpdateConstraints(T Dt)
 	{
 		SCOPE_CYCLE_COUNTER(STAT_Collisions_UpdatePointConstraints);
 
+		// @todo(chaos): parallelism needs to be optional
+
 		//PhysicsParallelFor(Handles.Num(), [&](int32 ConstraintHandleIndex)
 		//{
 		//	FConstraintContainerHandle* ConstraintHandle = Handles[ConstraintHandleIndex];
 		//	check(ConstraintHandle != nullptr);
-		//	Collisions::Update<ECollisionUpdateType::Deepest, float, 3>(MCullDistance, MShapePadding, ConstraintHandle->GetContact());
+		//	Collisions::Update(MCullDistance, MShapePadding, ConstraintHandle->GetContact());
 
 		//	if (ConstraintHandle->GetContact().GetPhi() < MCullDistance) 
 		//	{
@@ -306,7 +331,7 @@ namespace Chaos
 
 		for (FPointContactConstraint& Contact : PointConstraints)
 		{
-			Collisions::Update<ECollisionUpdateType::Deepest, float, 3>(MCullDistance, Contact);
+			Collisions::Update(Contact, MCullDistance);
 			if (Contact.GetPhi() < MCullDistance)
 			{
 				Contact.Timestamp = LifespanCounter;
@@ -314,21 +339,24 @@ namespace Chaos
 		}
 	}
 
+	// Called once per tick to update/regenerate persistent manifold planes and points
 	template<typename T, int d>
 	void TPBDCollisionConstraints<T, d>::UpdateManifolds(T Dt)
 	{
 		SCOPE_CYCLE_COUNTER(STAT_Collisions_UpdateManifoldConstraints);
 
+		// @todo(chaos): parallelism needs to be optional
+
 		//PhysicsParallelFor(Handles.Num(), [&](int32 ConstraintHandleIndex)
 		//{
 		//	FConstraintContainerHandle* ConstraintHandle = Handles[ConstraintHandleIndex];
 		//	check(ConstraintHandle != nullptr);
-		//	Collisions::UpdateManifold<float, 3>(MCullDistance, MShapePadding, ConstraintHandle->GetContact());
+		//	Collisions::Update(MCullDistance, MShapePadding, ConstraintHandle->GetContact());
 		//}, bDisableCollisionParallelFor);
 
 		for (FMultiPointContactConstraint& Contact : IterativeConstraints)
 		{
-			Collisions::UpdateManifold<float, 3>(MCullDistance, Contact);
+			Collisions::UpdateManifold(Contact, MCullDistance);
 			if (Contact.GetPhi() < MCullDistance)
 			{
 				Contact.Timestamp = LifespanCounter;
diff --git a/Engine/Source/Runtime/Experimental/Chaos/Private/Chaos/PBDCollisionConstraintsContact.cpp b/Engine/Source/Runtime/Experimental/Chaos/Private/Chaos/PBDCollisionConstraintsContact.cpp
--- a/Engine/Source/Runtime/Experimental/Chaos/Private/Chaos/PBDCollisionConstraintsContact.cpp
+++ b/Engine/Source/Runtime/Experimental/Chaos/Private/Chaos/PBDCollisionConstraintsContact.cpp
@@ -7,155 +7,163 @@
 #include "Chaos/Particle/ParticleUtilities.h"
 #include "Chaos/Utilities.h"
 
+//#pragma optimize("", off)
+
 namespace Chaos
 {
 	namespace Collisions
 	{
-		template<ECollisionUpdateType UpdateType, typename T, int d>
-		void Update(const T CullDistance, TCollisionConstraintBase<T, d>& Constraint)
+		extern void UpdateManifold(FRigidBodyMultiPointContactConstraint& Constraint, const FReal CullDistance)
 		{
-			Constraint.ResetPhi(CullDistance);
-			const TRigidTransform<T, d> ParticleTM = GetTransform(Constraint.Particle[0]);
-			const TRigidTransform<T, d> LevelsetTM = GetTransform(Constraint.Particle[1]);
+			const FRigidTransform3 Transform0 = GetTransform(Constraint.Particle[0]);
+			const FRigidTransform3 Transform1 = GetTransform(Constraint.Particle[1]);
 
-			UpdateConstraint<UpdateType>(Constraint, ParticleTM, LevelsetTM, CullDistance);
+			UpdateManifold(Constraint, Transform0, Transform1, CullDistance);
 		}
 
-		template<typename T, int d>
-		void UpdateManifold(const T CullDistance, TCollisionConstraintBase<T, d>& Constraint)
+		void Update(FRigidBodyPointContactConstraint& Constraint, const FReal CullDistance)
 		{
-			const TRigidTransform<T, d> Transform0 = GetTransform(Constraint.Particle[0]);
-			const TRigidTransform<T, d> Transform1 = GetTransform(Constraint.Particle[1]);
+			const FRigidTransform3 Transform0 = GetTransform(Constraint.Particle[0]);
+			const FRigidTransform3 Transform1 = GetTransform(Constraint.Particle[1]);
 
-			UpdateManifold(Constraint, Transform0, Transform1, CullDistance);
+			Constraint.ResetPhi(CullDistance);
+			UpdateConstraint<ECollisionUpdateType::Deepest>(Constraint, Transform0, Transform1, CullDistance);
 		}
 
+		void Update(FRigidBodyMultiPointContactConstraint& Constraint, const FReal CullDistance)
+		{
+			const FRigidTransform3 Transform0 = GetTransform(Constraint.Particle[0]);
+			const FRigidTransform3 Transform1 = GetTransform(Constraint.Particle[1]);
+
+			Constraint.ResetPhi(CullDistance);
+			UpdateConstraintFromManifold(Constraint, Transform0, Transform1, CullDistance);
+		}
 
-		template<typename T, int d>
-		TVector<T,d> ApplyContact(TCollisionContact<T, d>& Contact,
-			TGenericParticleHandle<T, d> Particle0, 
-			TGenericParticleHandle<T, d> Particle1,
-			const TContactIterationParameters<T> & IterationParameters, const TContactParticleParameters<T> & ParticleParameters)
+		FVec3 ApplyContact(FCollisionContact& Contact,
+			TGenericParticleHandle<FReal, 3> Particle0, 
+			TGenericParticleHandle<FReal, 3> Particle1,
+			const TContactIterationParameters<FReal> & IterationParameters,
+			const TContactParticleParameters<FReal> & ParticleParameters)
 		{
-			TVector<T, d> AccumulatedImpulse(0);
+			FVec3 AccumulatedImpulse(0);
 
-			TPBDRigidParticleHandle<T, d>* PBDRigid0 = Particle0->CastToRigidParticle();
-			TPBDRigidParticleHandle<T, d>* PBDRigid1 = Particle1->CastToRigidParticle();
+			TPBDRigidParticleHandle<FReal, 3>* PBDRigid0 = Particle0->CastToRigidParticle();
+			TPBDRigidParticleHandle<FReal, 3>* PBDRigid1 = Particle1->CastToRigidParticle();
 
 			bool bIsRigidDynamic0 = PBDRigid0 && PBDRigid0->ObjectState() == EObjectStateType::Dynamic;
 			bool bIsRigidDynamic1 = PBDRigid1 && PBDRigid1->ObjectState() == EObjectStateType::Dynamic;
 
-			const TVector<T, d> ZeroVector = TVector<T, d>(0);
-			TVector<T, d> P0 = FParticleUtilities::GetCoMWorldPosition(Particle0);
-			TVector<T, d> P1 = FParticleUtilities::GetCoMWorldPosition(Particle1);
-			TRotation<T, d> Q0 = FParticleUtilities::GetCoMWorldRotation(Particle0);
-			TRotation<T, d> Q1 = FParticleUtilities::GetCoMWorldRotation(Particle1);
+			const FVec3 ZeroVector = FVec3(0);
+			FVec3 P0 = FParticleUtilities::GetCoMWorldPosition(Particle0);
+			FVec3 P1 = FParticleUtilities::GetCoMWorldPosition(Particle1);
+			FRotation3 Q0 = FParticleUtilities::GetCoMWorldRotation(Particle0);
+			FRotation3 Q1 = FParticleUtilities::GetCoMWorldRotation(Particle1);
 
-			TVector<T, d> VectorToPoint1 = Contact.Location - P0;
-			TVector<T, d> VectorToPoint2 = Contact.Location - P1;
-			TVector<T, d> Body1Velocity = FParticleUtilities::GetVelocityAtCoMRelativePosition(Particle0, VectorToPoint1);
-			TVector<T, d> Body2Velocity = FParticleUtilities::GetVelocityAtCoMRelativePosition(Particle1, VectorToPoint2);
-			TVector<T, d> RelativeVelocity = Body1Velocity - Body2Velocity;
-			T RelativeNormalVelocity = TVector<T, d>::DotProduct(RelativeVelocity, Contact.Normal);
+			FVec3 VectorToPoint1 = Contact.Location - P0;
+			FVec3 VectorToPoint2 = Contact.Location - P1;
+			FVec3 Body1Velocity = FParticleUtilities::GetVelocityAtCoMRelativePosition(Particle0, VectorToPoint1);
+			FVec3 Body2Velocity = FParticleUtilities::GetVelocityAtCoMRelativePosition(Particle1, VectorToPoint2);
+			FVec3 RelativeVelocity = Body1Velocity - Body2Velocity;
+			FReal RelativeNormalVelocity = FVec3::DotProduct(RelativeVelocity, Contact.Normal);
 
 			if (RelativeNormalVelocity < 0) // ignore separating constraints
 			{
-				PMatrix<T, d, d> WorldSpaceInvI1 = bIsRigidDynamic0 ? Utilities::ComputeWorldSpaceInertia(Q0, PBDRigid0->InvI()) : PMatrix<T, d, d>(0);
-				PMatrix<T, d, d> WorldSpaceInvI2 = bIsRigidDynamic1 ? Utilities::ComputeWorldSpaceInertia(Q1, PBDRigid1->InvI()) : PMatrix<T, d, d>(0);
-				PMatrix<T, d, d> Factor =
-					(bIsRigidDynamic0 ? ComputeFactorMatrix3(VectorToPoint1, WorldSpaceInvI1, PBDRigid0->InvM()) : PMatrix<T, d, d>(0)) +
-					(bIsRigidDynamic1 ? ComputeFactorMatrix3(VectorToPoint2, WorldSpaceInvI2, PBDRigid1->InvM()) : PMatrix<T, d, d>(0));
-				TVector<T, d> Impulse;
-				TVector<T, d> AngularImpulse(0);
+				FMatrix33 WorldSpaceInvI1 = bIsRigidDynamic0 ? Utilities::ComputeWorldSpaceInertia(Q0, PBDRigid0->InvI()) : FMatrix33(0);
+				FMatrix33 WorldSpaceInvI2 = bIsRigidDynamic1 ? Utilities::ComputeWorldSpaceInertia(Q1, PBDRigid1->InvI()) : FMatrix33(0);
+				FMatrix33 Factor =
+					(bIsRigidDynamic0 ? ComputeFactorMatrix3(VectorToPoint1, WorldSpaceInvI1, PBDRigid0->InvM()) : FMatrix33(0)) +
+					(bIsRigidDynamic1 ? ComputeFactorMatrix3(VectorToPoint2, WorldSpaceInvI2, PBDRigid1->InvM()) : FMatrix33(0));
+				FVec3 Impulse;
+				FVec3 AngularImpulse(0);
 
 				// Resting contact if very close to the surface
 				bool bApplyRestitution = (RelativeVelocity.Size() > (2 * 980 * IterationParameters.Dt));
-				T Restitution = (bApplyRestitution) ? Contact.Restitution : (T)0;
-				T Friction = Contact.Friction;
-				T AngularFriction = Contact.AngularFriction;
+				FReal Restitution = (bApplyRestitution) ? Contact.Restitution : (FReal)0;
+				FReal Friction = Contact.Friction;
+				FReal AngularFriction = Contact.AngularFriction;
 
 				if (Friction > 0)
 				{
-					TVector<T, d> VelocityChange = -(Restitution * RelativeNormalVelocity * Contact.Normal + RelativeVelocity);
-					T NormalVelocityChange = TVector<T, d>::DotProduct(VelocityChange, Contact.Normal);
-					PMatrix<T, d, d> FactorInverse = Factor.Inverse();
-					TVector<T, d> MinimalImpulse = FactorInverse * VelocityChange;
-					const T MinimalImpulseDotNormal = TVector<T, d>::DotProduct(MinimalImpulse, Contact.Normal);
-					const T TangentialSize = (MinimalImpulse - MinimalImpulseDotNormal * Contact.Normal).Size();
+					FVec3 VelocityChange = -(Restitution * RelativeNormalVelocity * Contact.Normal + RelativeVelocity);
+					FReal NormalVelocityChange = FVec3::DotProduct(VelocityChange, Contact.Normal);
+					FMatrix33 FactorInverse = Factor.Inverse();
+					FVec3 MinimalImpulse = FactorInverse * VelocityChange;
+					const FReal MinimalImpulseDotNormal = FVec3::DotProduct(MinimalImpulse, Contact.Normal);
+					const FReal TangentialSize = (MinimalImpulse - MinimalImpulseDotNormal * Contact.Normal).Size();
 					if (TangentialSize <= Friction * MinimalImpulseDotNormal)
 					{
 						//within friction cone so just solve for static friction stopping the object
 						Impulse = MinimalImpulse;
 						if (AngularFriction)
 						{
-							TVector<T, d> RelativeAngularVelocity = Particle0->W() - Particle1->W();
-							T AngularNormal = TVector<T, d>::DotProduct(RelativeAngularVelocity, Contact.Normal);
-							TVector<T, d> AngularTangent = RelativeAngularVelocity - AngularNormal * Contact.Normal;
-							TVector<T, d> FinalAngularVelocity = FMath::Sign(AngularNormal) * FMath::Max((T)0, FMath::Abs(AngularNormal) - AngularFriction * NormalVelocityChange) * Contact.Normal + FMath::Max((T)0, AngularTangent.Size() - AngularFriction * NormalVelocityChange) * AngularTangent.GetSafeNormal();
-							TVector<T, d> Delta = FinalAngularVelocity - RelativeAngularVelocity;
+							FVec3 RelativeAngularVelocity = Particle0->W() - Particle1->W();
+							FReal AngularNormal = FVec3::DotProduct(RelativeAngularVelocity, Contact.Normal);
+							FVec3 AngularTangent = RelativeAngularVelocity - AngularNormal * Contact.Normal;
+							FVec3 FinalAngularVelocity = FMath::Sign(AngularNormal) * FMath::Max((FReal)0, FMath::Abs(AngularNormal) - AngularFriction * NormalVelocityChange) * Contact.Normal + FMath::Max((FReal)0, AngularTangent.Size() - AngularFriction * NormalVelocityChange) * AngularTangent.GetSafeNormal();
+							FVec3 Delta = FinalAngularVelocity - RelativeAngularVelocity;
 							if (!bIsRigidDynamic0 && bIsRigidDynamic1)
 							{
-								PMatrix<T, d, d> WorldSpaceI2 = (Q1 * FMatrix::Identity) * PBDRigid1->I() * (Q1 * FMatrix::Identity).GetTransposed();
-								TVector<T, d> ImpulseDelta = PBDRigid1->M() * TVector<T, d>::CrossProduct(VectorToPoint2, Delta);
+								FMatrix33 WorldSpaceI2 = (Q1 * FMatrix::Identity) * PBDRigid1->I() * (Q1 * FMatrix::Identity).GetTransposed();
+								FVec3 ImpulseDelta = PBDRigid1->M() * FVec3::CrossProduct(VectorToPoint2, Delta);
 								Impulse += ImpulseDelta;
-								AngularImpulse += WorldSpaceI2 * Delta - TVector<T, d>::CrossProduct(VectorToPoint2, ImpulseDelta);
+								AngularImpulse += WorldSpaceI2 * Delta - FVec3::CrossProduct(VectorToPoint2, ImpulseDelta);
 							}
 							else if (bIsRigidDynamic0 && !bIsRigidDynamic1)
 							{
-								PMatrix<T, d, d> WorldSpaceI1 = (Q0 * FMatrix::Identity) * PBDRigid0->I() * (Q0 * FMatrix::Identity).GetTransposed();
-								TVector<T, d> ImpulseDelta = PBDRigid0->M() * TVector<T, d>::CrossProduct(VectorToPoint1, Delta);
+								FMatrix33 WorldSpaceI1 = (Q0 * FMatrix::Identity) * PBDRigid0->I() * (Q0 * FMatrix::Identity).GetTransposed();
+								FVec3 ImpulseDelta = PBDRigid0->M() * FVec3::CrossProduct(VectorToPoint1, Delta);
 								Impulse += ImpulseDelta;
-								AngularImpulse += WorldSpaceI1 * Delta - TVector<T, d>::CrossProduct(VectorToPoint1, ImpulseDelta);
+								AngularImpulse += WorldSpaceI1 * Delta - FVec3::CrossProduct(VectorToPoint1, ImpulseDelta);
 							}
 							else if (bIsRigidDynamic0 && bIsRigidDynamic1)
 							{
-								PMatrix<T, d, d> Cross1(0, VectorToPoint1.Z, -VectorToPoint1.Y, -VectorToPoint1.Z, 0, VectorToPoint1.X, VectorToPoint1.Y, -VectorToPoint1.X, 0);
-								PMatrix<T, d, d> Cross2(0, VectorToPoint2.Z, -VectorToPoint2.Y, -VectorToPoint2.Z, 0, VectorToPoint2.X, VectorToPoint2.Y, -VectorToPoint2.X, 0);
-								PMatrix<T, d, d> CrossI1 = Cross1 * WorldSpaceInvI1;
-								PMatrix<T, d, d> CrossI2 = Cross2 * WorldSpaceInvI2;
-								PMatrix<T, d, d> Diag1 = CrossI1 * Cross1.GetTransposed() + CrossI2 * Cross2.GetTransposed();
+								FMatrix33 Cross1(0, VectorToPoint1.Z, -VectorToPoint1.Y, -VectorToPoint1.Z, 0, VectorToPoint1.X, VectorToPoint1.Y, -VectorToPoint1.X, 0);
+								FMatrix33 Cross2(0, VectorToPoint2.Z, -VectorToPoint2.Y, -VectorToPoint2.Z, 0, VectorToPoint2.X, VectorToPoint2.Y, -VectorToPoint2.X, 0);
+								FMatrix33 CrossI1 = Cross1 * WorldSpaceInvI1;
+								FMatrix33 CrossI2 = Cross2 * WorldSpaceInvI2;
+								FMatrix33 Diag1 = CrossI1 * Cross1.GetTransposed() + CrossI2 * Cross2.GetTransposed();
 								Diag1.M[0][0] += PBDRigid0->InvM() + PBDRigid1->InvM();
 								Diag1.M[1][1] += PBDRigid0->InvM() + PBDRigid1->InvM();
 								Diag1.M[2][2] += PBDRigid0->InvM() + PBDRigid1->InvM();
-								PMatrix<T, d, d> OffDiag1 = (CrossI1 + CrossI2) * -1;
-								PMatrix<T, d, d> Diag2 = (WorldSpaceInvI1 + WorldSpaceInvI2).Inverse();
-								PMatrix<T, d, d> OffDiag1Diag2 = OffDiag1 * Diag2;
-								TVector<T, d> ImpulseDelta = PMatrix<T, d, d>((Diag1 - OffDiag1Diag2 * OffDiag1.GetTransposed()).Inverse())* ((OffDiag1Diag2 * -1) * Delta);
+								FMatrix33 OffDiag1 = (CrossI1 + CrossI2) * -1;
+								FMatrix33 Diag2 = (WorldSpaceInvI1 + WorldSpaceInvI2).Inverse();
+								FMatrix33 OffDiag1Diag2 = OffDiag1 * Diag2;
+								FVec3 ImpulseDelta = FMatrix33((Diag1 - OffDiag1Diag2 * OffDiag1.GetTransposed()).Inverse())* ((OffDiag1Diag2 * -1) * Delta);
 								Impulse += ImpulseDelta;
-								AngularImpulse += Diag2 * (Delta - PMatrix<T, d, d>(OffDiag1.GetTransposed()) * ImpulseDelta);
+								AngularImpulse += Diag2 * (Delta - FMatrix33(OffDiag1.GetTransposed()) * ImpulseDelta);
 							}
 						}
 					}
 					else
 					{
 						//outside friction cone, solve for normal relative velocity and keep tangent at cone edge
-						TVector<T, d> Tangent = (RelativeVelocity - TVector<T, d>::DotProduct(RelativeVelocity, Contact.Normal) * Contact.Normal).GetSafeNormal();
-						TVector<T, d> DirectionalFactor = Factor * (Contact.Normal - Friction * Tangent);
-						T ImpulseDenominator = TVector<T, d>::DotProduct(Contact.Normal, DirectionalFactor);
+						FVec3 Tangent = (RelativeVelocity - FVec3::DotProduct(RelativeVelocity, Contact.Normal) * Contact.Normal).GetSafeNormal();
+						FVec3 DirectionalFactor = Factor * (Contact.Normal - Friction * Tangent);
+						FReal ImpulseDenominator = FVec3::DotProduct(Contact.Normal, DirectionalFactor);
 						if (!ensureMsgf(FMath::Abs(ImpulseDenominator) > SMALL_NUMBER, TEXT("Contact:%s\n\nParticle:%s\n\nLevelset:%s\n\nDirectionalFactor:%s, ImpulseDenominator:%f"),
 							*Contact.ToString(),
 							*Particle0->ToString(),
 							*Particle1->ToString(),
 							*DirectionalFactor.ToString(), ImpulseDenominator))
 						{
-							ImpulseDenominator = (T)1;
+							ImpulseDenominator = (FReal)1;
 						}
 
-						const T ImpulseMag = -(1 + Restitution) * RelativeNormalVelocity / ImpulseDenominator;
+						const FReal ImpulseMag = -(1 + Restitution) * RelativeNormalVelocity / ImpulseDenominator;
 						Impulse = ImpulseMag * (Contact.Normal - Friction * Tangent);
 					}
 				}
 				else
 				{
-					T ImpulseDenominator = TVector<T, d>::DotProduct(Contact.Normal, Factor * Contact.Normal);
-					TVector<T, d> ImpulseNumerator = -(1 + Restitution) * TVector<T, d>::DotProduct(RelativeVelocity, Contact.Normal)* Contact.Normal;
+					FReal ImpulseDenominator = FVec3::DotProduct(Contact.Normal, Factor * Contact.Normal);
+					FVec3 ImpulseNumerator = -(1 + Restitution) * FVec3::DotProduct(RelativeVelocity, Contact.Normal)* Contact.Normal;
 					if (!ensureMsgf(FMath::Abs(ImpulseDenominator) > SMALL_NUMBER, TEXT("Contact:%s\n\nParticle:%s\n\nLevelset:%s\n\nFactor*Constraint.Normal:%s, ImpulseDenominator:%f"),
 						*Contact.ToString(),
 						*Particle0->ToString(),
 						*Particle1->ToString(),
 						*(Factor * Contact.Normal).ToString(), ImpulseDenominator))
 					{
-						ImpulseDenominator = (T)1;
+						ImpulseDenominator = (FReal)1;
 					}
 					Impulse = ImpulseNumerator / ImpulseDenominator;
 				}
@@ -173,7 +181,7 @@ namespace Chaos
 					PBDRigid0->W() += DW;
 					// Position update as part of pbd
 					P0 += (DV * IterationParameters.Dt);
-					Q0 += FRotation3::FromElements(DW, 0.f) * Q0 * IterationParameters.Dt * T(0.5);
+					Q0 += FRotation3::FromElements(DW, 0.f) * Q0 * IterationParameters.Dt * FReal(0.5);
 					Q0.Normalize();
 					FParticleUtilities::SetCoMWorldTransform(PBDRigid0, P0, Q0);
 				}
@@ -187,7 +195,7 @@ namespace Chaos
 					PBDRigid1->W() += DW;
 					// Position update as part of pbd
 					P1 += (DV * IterationParameters.Dt);
-					Q1 += FRotation3::FromElements(DW, 0.f) * Q1 * IterationParameters.Dt * T(0.5);
+					Q1 += FRotation3::FromElements(DW, 0.f) * Q1 * IterationParameters.Dt * FReal(0.5);
 					Q1.Normalize();
 					FParticleUtilities::SetCoMWorldTransform(PBDRigid1, P1, Q1);
 				}
@@ -195,11 +203,11 @@ namespace Chaos
 			return AccumulatedImpulse;
 		}
 
-		template<typename T, int d>
-		void Apply(TCollisionConstraintBase<T, d>& Constraint, const TContactIterationParameters<T> & IterationParameters, const TContactParticleParameters<T> & ParticleParameters)
+		template<typename T_CONSTRAINT>
+		void ApplyImpl(T_CONSTRAINT& Constraint, const FContactIterationParameters & IterationParameters, const FContactParticleParameters & ParticleParameters)
 		{
-			TGenericParticleHandle<T, d> Particle0 = TGenericParticleHandle<T, d>(Constraint.Particle[0]);
-			TGenericParticleHandle<T, d> Particle1 = TGenericParticleHandle<T, d>(Constraint.Particle[1]);
+			TGenericParticleHandle<FReal, 3> Particle0 = TGenericParticleHandle<FReal, 3>(Constraint.Particle[0]);
+			TGenericParticleHandle<FReal, 3> Particle1 = TGenericParticleHandle<FReal, 3>(Constraint.Particle[1]);
 
 			for (int32 PairIt = 0; PairIt < IterationParameters.NumPairIterations; ++PairIt)
 			{
@@ -209,7 +217,7 @@ namespace Chaos
 				const bool bNeedCollisionUpdate = (PairIt > 0) || (IterationParameters.Iteration > 0);
 				if (bNeedCollisionUpdate)
 				{
-					Collisions::Update<ECollisionUpdateType::Deepest>(ParticleParameters.CullDistance, Constraint);
+					Collisions::Update(Constraint, ParticleParameters.CullDistance);
 				}
 
 				if (Constraint.GetPhi() >= ParticleParameters.ShapePadding)
@@ -236,27 +244,48 @@ namespace Chaos
 			}
 		}
 
-		template<typename T, int d> 
-		TVector<T, d>
-		ApplyPushOutContact(
-			TCollisionContact<T,d>& Contact,
-			TGenericParticleHandle<T, d> Particle0, 
-			TGenericParticleHandle<T, d> Particle1,
-			const TSet<const TGeometryParticleHandle<T, d>*>& IsTemporarilyStatic,
-			const TContactIterationParameters<T> & IterationParameters, const TContactParticleParameters<T> & ParticleParameters)
+		void Apply(FCollisionConstraintBase& Constraint, const FContactIterationParameters & IterationParameters, const FContactParticleParameters & ParticleParameters)
+		{
+			if (Constraint.GetType() == FCollisionConstraintBase::FType::SinglePoint)
+			{
+				ApplyImpl(*Constraint.As<FRigidBodyPointContactConstraint>(), IterationParameters, ParticleParameters);
+			}
+			else if (Constraint.GetType() == FCollisionConstraintBase::FType::MultiPoint)
+			{
+				ApplyImpl(*Constraint.As<FRigidBodyMultiPointContactConstraint>(), IterationParameters, ParticleParameters);
+			}
+		}
+
+		void Apply(FRigidBodyPointContactConstraint& Constraint, const FContactIterationParameters & IterationParameters, const FContactParticleParameters & ParticleParameters)
+		{
+			ApplyImpl(Constraint, IterationParameters, ParticleParameters);
+		}
+
+		void Apply(FRigidBodyMultiPointContactConstraint& Constraint, const FContactIterationParameters & IterationParameters, const FContactParticleParameters & ParticleParameters)
 		{
-			TVector<T, d> AccumulatedImpulse(0);
+			ApplyImpl(Constraint, IterationParameters, ParticleParameters);
+		}
 
-			TPBDRigidParticleHandle<T, d>* PBDRigid0 = Particle0->CastToRigidParticle();
-			TPBDRigidParticleHandle<T, d>* PBDRigid1 = Particle1->CastToRigidParticle();
+
+		FVec3 ApplyPushOutContact(
+			FCollisionContact& Contact,
+			TGenericParticleHandle<FReal, 3> Particle0, 
+			TGenericParticleHandle<FReal, 3> Particle1,
+			const TSet<const TGeometryParticleHandle<FReal, 3>*>& IsTemporarilyStatic,
+			const FContactIterationParameters & IterationParameters, const FContactParticleParameters & ParticleParameters)
+		{
+			FVec3 AccumulatedImpulse(0);
+
+			TPBDRigidParticleHandle<FReal, 3>* PBDRigid0 = Particle0->CastToRigidParticle();
+			TPBDRigidParticleHandle<FReal, 3>* PBDRigid1 = Particle1->CastToRigidParticle();
 			const bool bIsRigidDynamic0 = PBDRigid0 && PBDRigid0->ObjectState() == EObjectStateType::Dynamic;
 			const bool bIsRigidDynamic1 = PBDRigid1 && PBDRigid1->ObjectState() == EObjectStateType::Dynamic;
 
-			const TVector<T, d> ZeroVector = TVector<T, d>(0);
-			TVector<T, d> P0 = FParticleUtilities::GetCoMWorldPosition(Particle0);
-			TVector<T, d> P1 = FParticleUtilities::GetCoMWorldPosition(Particle1);
-			TRotation<T, d> Q0 = FParticleUtilities::GetCoMWorldRotation(Particle0);
-			TRotation<T, d> Q1 = FParticleUtilities::GetCoMWorldRotation(Particle1);
+			const FVec3 ZeroVector = FVec3(0);
+			FVec3 P0 = FParticleUtilities::GetCoMWorldPosition(Particle0);
+			FVec3 P1 = FParticleUtilities::GetCoMWorldPosition(Particle1);
+			FRotation3 Q0 = FParticleUtilities::GetCoMWorldRotation(Particle0);
+			FRotation3 Q1 = FParticleUtilities::GetCoMWorldRotation(Particle1);
 			const bool IsTemporarilyStatic0 = IsTemporarilyStatic.Contains(Particle0->GeometryParticleHandle());
 			const bool IsTemporarilyStatic1 = IsTemporarilyStatic.Contains(Particle1->GeometryParticleHandle());
 
@@ -275,26 +304,26 @@ namespace Chaos
 				*IterationParameters.NeedsAnotherIteration = true;
 			}
 
-			PMatrix<T, d, d> WorldSpaceInvI1 = bIsRigidDynamic0 ? Utilities::ComputeWorldSpaceInertia(Q0, PBDRigid0->InvI()) : PMatrix<T, d, d>(0);
-			PMatrix<T, d, d> WorldSpaceInvI2 = bIsRigidDynamic1 ? Utilities::ComputeWorldSpaceInertia(Q1, PBDRigid1->InvI()) : PMatrix<T, d, d>(0);
-			TVector<T, d> VectorToPoint1 = Contact.Location - P0;
-			TVector<T, d> VectorToPoint2 = Contact.Location - P1;
-			PMatrix<T, d, d> Factor =
-				(bIsRigidDynamic0 ? ComputeFactorMatrix3(VectorToPoint1, WorldSpaceInvI1, PBDRigid0->InvM()) : PMatrix<T, d, d>(0)) +
-				(bIsRigidDynamic1 ? ComputeFactorMatrix3(VectorToPoint2, WorldSpaceInvI2, PBDRigid1->InvM()) : PMatrix<T, d, d>(0));
-			T Numerator = FMath::Min((T)(IterationParameters.Iteration + 2), (T)IterationParameters.NumIterations);
-			T ScalingFactor = Numerator / (T)IterationParameters.NumIterations;
+			FMatrix33 WorldSpaceInvI1 = bIsRigidDynamic0 ? Utilities::ComputeWorldSpaceInertia(Q0, PBDRigid0->InvI()) : FMatrix33(0);
+			FMatrix33 WorldSpaceInvI2 = bIsRigidDynamic1 ? Utilities::ComputeWorldSpaceInertia(Q1, PBDRigid1->InvI()) : FMatrix33(0);
+			FVec3 VectorToPoint1 = Contact.Location - P0;
+			FVec3 VectorToPoint2 = Contact.Location - P1;
+			FMatrix33 Factor =
+				(bIsRigidDynamic0 ? ComputeFactorMatrix3(VectorToPoint1, WorldSpaceInvI1, PBDRigid0->InvM()) : FMatrix33(0)) +
+				(bIsRigidDynamic1 ? ComputeFactorMatrix3(VectorToPoint2, WorldSpaceInvI2, PBDRigid1->InvM()) : FMatrix33(0));
+			FReal Numerator = FMath::Min((FReal)(IterationParameters.Iteration + 2), (FReal)IterationParameters.NumIterations);
+			FReal ScalingFactor = Numerator / (FReal)IterationParameters.NumIterations;
 
 			//if pushout is needed we better fix relative velocity along normal. Treat it as if 0 restitution
-			TVector<T, d> Body1Velocity = FParticleUtilities::GetVelocityAtCoMRelativePosition(Particle0, VectorToPoint1);
-			TVector<T, d> Body2Velocity = FParticleUtilities::GetVelocityAtCoMRelativePosition(Particle1, VectorToPoint2);
-			TVector<T, d> RelativeVelocity = Body1Velocity - Body2Velocity;
-			const T RelativeVelocityDotNormal = TVector<T, d>::DotProduct(RelativeVelocity, Contact.Normal);
+			FVec3 Body1Velocity = FParticleUtilities::GetVelocityAtCoMRelativePosition(Particle0, VectorToPoint1);
+			FVec3 Body2Velocity = FParticleUtilities::GetVelocityAtCoMRelativePosition(Particle1, VectorToPoint2);
+			FVec3 RelativeVelocity = Body1Velocity - Body2Velocity;
+			const FReal RelativeVelocityDotNormal = FVec3::DotProduct(RelativeVelocity, Contact.Normal);
 			if (RelativeVelocityDotNormal < 0)
 			{
-				const TVector<T, d> ImpulseNumerator = -TVector<T, d>::DotProduct(RelativeVelocity, Contact.Normal) * Contact.Normal * ScalingFactor;
-				const TVector<T, d> FactorContactNormal = Factor * Contact.Normal;
-				T ImpulseDenominator = TVector<T, d>::DotProduct(Contact.Normal, FactorContactNormal);
+				const FVec3 ImpulseNumerator = -FVec3::DotProduct(RelativeVelocity, Contact.Normal) * Contact.Normal * ScalingFactor;
+				const FVec3 FactorContactNormal = Factor * Contact.Normal;
+				FReal ImpulseDenominator = FVec3::DotProduct(Contact.Normal, FactorContactNormal);
 				if (!ensureMsgf(FMath::Abs(ImpulseDenominator) > SMALL_NUMBER, 
 					TEXT("ApplyPushout Contact:%s\n\nParticle:%s\n\nLevelset:%s\n\nFactor*Contact.Normal:%s, ImpulseDenominator:%f"),
 					*Contact.ToString(),
@@ -302,15 +331,15 @@ namespace Chaos
 					*Particle1->ToString(),
 					*FactorContactNormal.ToString(), ImpulseDenominator))
 				{
-					ImpulseDenominator = (T)1;
+					ImpulseDenominator = (FReal)1;
 				}
 
-				TVector<T, d> VelocityFixImpulse = ImpulseNumerator / ImpulseDenominator;
+				FVec3 VelocityFixImpulse = ImpulseNumerator / ImpulseDenominator;
 				VelocityFixImpulse = GetEnergyClampedImpulse(Particle0->CastToRigidParticle(), Particle1->CastToRigidParticle(), VelocityFixImpulse, VectorToPoint1, VectorToPoint2, Body1Velocity, Body2Velocity);
 				AccumulatedImpulse += VelocityFixImpulse;	//question: should we track this?
 				if (!IsTemporarilyStatic0 && bIsRigidDynamic0)
 				{
-					TVector<T, d> AngularImpulse = TVector<T, d>::CrossProduct(VectorToPoint1, VelocityFixImpulse);
+					FVec3 AngularImpulse = FVec3::CrossProduct(VectorToPoint1, VelocityFixImpulse);
 					PBDRigid0->V() += PBDRigid0->InvM() * VelocityFixImpulse;
 					PBDRigid0->W() += WorldSpaceInvI1 * AngularImpulse;
 
@@ -318,7 +347,7 @@ namespace Chaos
 
 				if (!IsTemporarilyStatic1 && bIsRigidDynamic1)
 				{
-					TVector<T, d> AngularImpulse = TVector<T, d>::CrossProduct(VectorToPoint2, -VelocityFixImpulse);
+					FVec3 AngularImpulse = FVec3::CrossProduct(VectorToPoint2, -VelocityFixImpulse);
 					PBDRigid1->V() -= PBDRigid1->InvM() * VelocityFixImpulse;
 					PBDRigid1->W() += WorldSpaceInvI2 * AngularImpulse;
 				}
@@ -326,20 +355,20 @@ namespace Chaos
 			}
 
 
-			TVector<T, d> Impulse = PMatrix<T, d, d>(Factor.Inverse()) * ((-Contact.Phi + ParticleParameters.ShapePadding) * ScalingFactor * Contact.Normal);
-			TVector<T, d> AngularImpulse1 = TVector<T, d>::CrossProduct(VectorToPoint1, Impulse);
-			TVector<T, d> AngularImpulse2 = TVector<T, d>::CrossProduct(VectorToPoint2, -Impulse);
+			FVec3 Impulse = FMatrix33(Factor.Inverse()) * ((-Contact.Phi + ParticleParameters.ShapePadding) * ScalingFactor * Contact.Normal);
+			FVec3 AngularImpulse1 = FVec3::CrossProduct(VectorToPoint1, Impulse);
+			FVec3 AngularImpulse2 = FVec3::CrossProduct(VectorToPoint2, -Impulse);
 			if (!IsTemporarilyStatic0 && bIsRigidDynamic0)
 			{
 				P0 += PBDRigid0->InvM() * Impulse;
-				Q0 = TRotation<T, d>::FromVector(WorldSpaceInvI1 * AngularImpulse1) * Q0;
+				Q0 = FRotation3::FromVector(WorldSpaceInvI1 * AngularImpulse1) * Q0;
 				Q0.Normalize();
 				FParticleUtilities::SetCoMWorldTransform(Particle0, P0, Q0);
 			}
 			if (!IsTemporarilyStatic1 && bIsRigidDynamic1)
 			{
 				P1 -= PBDRigid1->InvM() * Impulse;
-				Q1 = TRotation<T, d>::FromVector(WorldSpaceInvI2 * AngularImpulse2) * Q1;
+				Q1 = FRotation3::FromVector(WorldSpaceInvI2 * AngularImpulse2) * Q1;
 				Q1.Normalize();
 				FParticleUtilities::SetCoMWorldTransform(Particle1, P1, Q1);
 			}
@@ -348,30 +377,46 @@ namespace Chaos
 		}
 
 
-		template<typename T, int d>
-		void ApplyPushOut(TCollisionConstraintBase<T, d>& Constraint, const TSet<const TGeometryParticleHandle<T, d>*>& IsTemporarilyStatic,
-			const TContactIterationParameters<T> & IterationParameters, const TContactParticleParameters<T> & ParticleParameters)
+		template<typename T_CONSTRAINT>
+		void ApplyPushOutImpl(T_CONSTRAINT& Constraint, const TSet<const TGeometryParticleHandle<FReal, 3>*>& IsTemporarilyStatic,
+			const FContactIterationParameters & IterationParameters, const FContactParticleParameters & ParticleParameters)
 		{
-			TGenericParticleHandle<T, d> Particle0 = TGenericParticleHandle<T, d>(Constraint.Particle[0]);
-			TGenericParticleHandle<T, d> Particle1 = TGenericParticleHandle<T, d>(Constraint.Particle[1]);
+			TGenericParticleHandle<FReal, 3> Particle0 = TGenericParticleHandle<FReal, 3>(Constraint.Particle[0]);
+			TGenericParticleHandle<FReal, 3> Particle1 = TGenericParticleHandle<FReal, 3>(Constraint.Particle[1]);
 
 			for (int32 PairIt = 0; PairIt < IterationParameters.NumPairIterations; ++PairIt)
 			{
-				Update<ECollisionUpdateType::Deepest>(ParticleParameters.CullDistance, Constraint);
+				Update(Constraint, ParticleParameters.CullDistance);
 
 				Constraint.AccumulatedImpulse += 
 					ApplyPushOutContact(Constraint.Manifold, Particle0, Particle1, IsTemporarilyStatic, IterationParameters, ParticleParameters);
 			}
 		}
 
+		void ApplyPushOut(FCollisionConstraintBase& Constraint, const TSet<const TGeometryParticleHandle<FReal, 3>*>& IsTemporarilyStatic, 
+			const FContactIterationParameters & IterationParameters, const FContactParticleParameters & ParticleParameters)
+		{
+			if (Constraint.GetType() == FCollisionConstraintBase::FType::SinglePoint)
+			{
+				ApplyPushOutImpl(*Constraint.As<FRigidBodyPointContactConstraint>(), IsTemporarilyStatic, IterationParameters, ParticleParameters);
+			}
+			else if (Constraint.GetType() == FCollisionConstraintBase::FType::MultiPoint)
+			{
+				ApplyPushOutImpl(*Constraint.As<FRigidBodyMultiPointContactConstraint>(), IsTemporarilyStatic, IterationParameters, ParticleParameters);
+			}
+		}
+
+		void ApplyPushOut(FRigidBodyPointContactConstraint& Constraint, const TSet<const TGeometryParticleHandle<FReal, 3>*>& IsTemporarilyStatic,
+			const FContactIterationParameters & IterationParameters, const FContactParticleParameters & ParticleParameters)
+		{
+			ApplyPushOutImpl(Constraint, IsTemporarilyStatic, IterationParameters, ParticleParameters);
+		}
 
-		template void Update<ECollisionUpdateType::Any, float, 3>(const float, TCollisionConstraintBase<float, 3>&);
-		template void Update<ECollisionUpdateType::Deepest, float, 3>(const float, TCollisionConstraintBase<float, 3>&);
-		template void UpdateManifold<float, 3>(const float, TCollisionConstraintBase<float, 3>&);
-		//
-		template void Apply<float, 3>(TCollisionConstraintBase<float, 3>&, const TContactIterationParameters<float> &, const TContactParticleParameters<float> &);
-		template void ApplyPushOut<float, 3>(TCollisionConstraintBase<float,3>& , const TSet<const TGeometryParticleHandle<float,3>*>&,
-			const TContactIterationParameters<float> & IterationParameters, const TContactParticleParameters<float> & ParticleParameters);
+		void ApplyPushOut(FRigidBodyMultiPointContactConstraint& Constraint, const TSet<const TGeometryParticleHandle<FReal, 3>*>& IsTemporarilyStatic,
+			const FContactIterationParameters & IterationParameters, const FContactParticleParameters & ParticleParameters)
+		{
+			ApplyPushOutImpl(Constraint, IsTemporarilyStatic, IterationParameters, ParticleParameters);
+		}
 
 	} // Collisions
 
diff --git a/Engine/Source/Runtime/Experimental/Chaos/Public/Chaos/CollisionResolution.h b/Engine/Source/Runtime/Experimental/Chaos/Public/Chaos/CollisionResolution.h
--- a/Engine/Source/Runtime/Experimental/Chaos/Public/Chaos/CollisionResolution.h
+++ b/Engine/Source/Runtime/Experimental/Chaos/Public/Chaos/CollisionResolution.h
@@ -257,6 +257,9 @@ namespace Chaos
 		template<ECollisionUpdateType UpdateType, typename T, int d>
 		void CHAOS_API UpdateConstraint(TCollisionConstraintBase<T, d>& ConstraintBase, const TRigidTransform<T, d>& Transform0, const TRigidTransform<T, d>& Transform1, const T CullDistance);
 
+		template<class T, int d>
+		void UpdateConstraintFromManifold(TRigidBodyMultiPointContactConstraint<T, d>& Constraint, const TRigidTransform<T, d>& Transform0, const TRigidTransform<T, d>& Transform1, const T CullDistance);
+
 		template<typename T, int d>
 		void CHAOS_API ConstructConstraints(TGeometryParticleHandle<T, d>* Particle0, TGeometryParticleHandle<T, d>* Particle1, const FImplicitObject* Implicit0, const FImplicitObject* Implicit1, const TRigidTransform<T, d>& Transform0, const TRigidTransform<T, d>& Transform1, const T CullDistance, FCollisionConstraintsArray& NewConstraints);
 
diff --git a/Engine/Source/Runtime/Experimental/Chaos/Public/Chaos/CollisionResolutionTypes.h b/Engine/Source/Runtime/Experimental/Chaos/Public/Chaos/CollisionResolutionTypes.h
--- a/Engine/Source/Runtime/Experimental/Chaos/Public/Chaos/CollisionResolutionTypes.h
+++ b/Engine/Source/Runtime/Experimental/Chaos/Public/Chaos/CollisionResolutionTypes.h
@@ -125,7 +125,7 @@ namespace Chaos
 		FType GetType() const { return Type; }
 
 		template<class AS_T> AS_T * As() { return static_cast<AS_T*>(this); }
-		template<class AS_T> const AS_T * As() const { return static_cast<AS_T*>(this); }
+		template<class AS_T> const AS_T * As() const { return static_cast<const AS_T*>(this); }
 
 		bool ContainsManifold(const FImplicitObject* A, const FImplicitObject* B) const { return A == Manifold.Implicit[0] && B == Manifold.Implicit[1]; }
 		void SetManifold(const FImplicitObject* A, const FImplicitObject* B) { Manifold.Implicit[0] = A; Manifold.Implicit[1] = B; }
@@ -202,32 +202,52 @@ namespace Chaos
 		using FGeometryParticleHandle = TGeometryParticleHandle<T, d>;
 		using FManifold = TCollisionContact<T, d>;
 		using Base::Particle;
-		struct FSampleData { TVector<T,d> X; float Delta; FManifold Manifold; };
 
-		TRigidBodyMultiPointContactConstraint() : Base(Base::FType::MultiPoint), SourceNormalIndex(INDEX_NONE), PlaneNormal(0), PlanePosition(0) {}
+		TRigidBodyMultiPointContactConstraint() : Base(Base::FType::MultiPoint), PlaneOwnerIndex(INDEX_NONE), PlaneFaceIndex(INDEX_NONE), PlaneNormal(0), PlanePosition(0) {}
 		TRigidBodyMultiPointContactConstraint(
 			FGeometryParticleHandle* Particle0, const FImplicitObject* Implicit0, const TRigidTransform<T, d>& Transform0,
 			FGeometryParticleHandle* Particle1, const FImplicitObject* Implicit1, const TRigidTransform<T, d>& Transform1,
 			EContactShapesType ShapesType = EContactShapesType::Unknown)
 			: Base(Particle0, Implicit0, Transform0, Particle1, Implicit1, Transform1, Base::FType::MultiPoint, ShapesType)
-			, SourceNormalIndex(INDEX_NONE), PlaneNormal(0), PlanePosition(0) 
+			, PlaneOwnerIndex(INDEX_NONE), PlaneFaceIndex(INDEX_NONE), PlaneNormal(0), PlanePosition(0)
 		{}
 
 		static typename Base::FType StaticType() { return Base::FType::MultiPoint; };
 
-		int SourceNormalIndex; // index of normal on particle[0] body;
-		TVector<T, d> PlaneNormal; // local space contact normal on Particle1
-		TVector<T, d> PlanePosition; // local space surface position on Particle1
+		// Get the particle that owns the plane
+		FGeometryParticleHandle* PlaneParticleHandle() const { check(PlaneOwnerIndex >= 0 && PlaneOwnerIndex < 2); return Particle[PlaneOwnerIndex]; }
 
-		// Samples
-		int32               NumSamples()                   { return Samples.Num(); }
-		void                AddSample(FSampleData && Data) { Samples.Add(Data); };
-		void                ResetSamples(int32 NewSize=0)  { Samples.Reset(NewSize); }
-		FSampleData &       operator[](int32 Index) { ensure(0 <= Index && Index < NumSamples()); return Samples[Index]; }
-		const FSampleData & operator[](int32 Index) const  { ensure(0 <= Index && Index < NumSamples()); return Samples[Index]; }
+		// Get the particle that owns the manifold sample points
+		FGeometryParticleHandle* PointsParticleHandle() const { check(PlaneOwnerIndex >= 0 && PlaneOwnerIndex < 2); return Particle[1 - PlaneOwnerIndex]; }
+
+		int32 GetManifoldPlaneOwnerIndex() const { return PlaneOwnerIndex; }
+		int32 GetManifoldPlaneFaceIndex() const { return PlaneFaceIndex; }
+		const FVec3& GetManifoldPlaneNormal() const { return PlaneNormal; }
+		const FVec3& GetManifoldPlanePosition() const { return PlanePosition; }
+		void SetManifoldPlane(int32 OwnerIndex, int32 FaceIndex, const FVec3& Normal, const FVec3& Pos)
+		{
+			PlaneOwnerIndex = OwnerIndex;
+			PlaneFaceIndex = FaceIndex;
+			PlaneNormal = Normal;
+			PlanePosition = Pos;
+		}
+
+		int32               NumManifoldPoints() const { return Points.Num(); }
+		const FVec3&		GetManifoldPoint(int32 Index) const { check(0 <= Index && Index < NumManifoldPoints()); return Points[Index]; }
+		void				SetManifoldPoint(int32 Index, const FVec3& Point) { check(0 <= Index && Index < NumManifoldPoints()); Points[Index] = Point; }
+		void                AddManifoldPoint(const FVec3& Point) { Points.Add(Point); };
+		void                ResetManifoldPoints(int32 NewSize=0)  { Points.Reset(NewSize); }
 
 	private:
-		TArray<FSampleData> Samples; // iterative samples
+		// Manifold plane data
+		int PlaneOwnerIndex; // index of particle which owns the plane. The other owns the sample positions
+		int PlaneFaceIndex; // index of face used as the manifold plane on plane-owner body
+		TVector<T, d> PlaneNormal; // local space contact normal on plane-owner
+		TVector<T, d> PlanePosition; // local space surface position on plane-owner
+
+		// Manifold point data
+		static const int32 MaxPoints = 4;
+		TArray<FVec3, TFixedAllocator<MaxPoints>> Points;
 	};
 	typedef TRigidBodyMultiPointContactConstraint<float, 3> FRigidBodyMultiPointContactConstraint;
 
@@ -362,8 +382,8 @@ namespace Chaos
 	template<int T_MAXCONSTRAINTS>
 	struct TCollisionConstraintsStore
 	{
-		TArray<TRigidBodyPointContactConstraint<FReal, 3>, TInlineAllocator<T_MAXCONSTRAINTS>> SinglePointConstraints;
-		TArray<TRigidBodyMultiPointContactConstraint<FReal, 3>, TInlineAllocator<T_MAXCONSTRAINTS>> MultiPointConstraints;
+		TArray<TRigidBodyPointContactConstraint<FReal, 3>, TFixedAllocator<T_MAXCONSTRAINTS>> SinglePointConstraints;
+		TArray<TRigidBodyMultiPointContactConstraint<FReal, 3>, TFixedAllocator<T_MAXCONSTRAINTS>> MultiPointConstraints;
 
 		int32 Num() const { return SinglePointConstraints.Num() + MultiPointConstraints.Num(); }
 
diff --git a/Engine/Source/Runtime/Experimental/Chaos/Public/Chaos/PBDCollisionConstraints.h b/Engine/Source/Runtime/Experimental/Chaos/Public/Chaos/PBDCollisionConstraints.h
--- a/Engine/Source/Runtime/Experimental/Chaos/Public/Chaos/PBDCollisionConstraints.h
+++ b/Engine/Source/Runtime/Experimental/Chaos/Public/Chaos/PBDCollisionConstraints.h
@@ -12,6 +12,9 @@
 #include "BoundingVolume.h"
 #include "AABBTree.h"
 
+// @todo(chaos): optimize and re-enable persistent constraints if we want it
+#define CHAOS_COLLISION_PERSISTENCE_ENABLED 0
+
 namespace Chaos
 {
 template<typename T, int d>
@@ -175,7 +178,11 @@ public:
 
 	bool Contains(const FConstraintBase* Base) const 
 	{
+#if CHAOS_COLLISION_PERSISTENCE_ENABLED
 		return Manifolds.Contains(FConstraintContainerHandle::MakeKey(Base));
+#else
+		return false;
+#endif
 	}
 
 	// @todo(chaos): remove
@@ -252,7 +259,9 @@ private:
 	TArray<FPointContactConstraint> PointConstraints;
 	TArray<FMultiPointContactConstraint> IterativeConstraints;
 
+#if CHAOS_COLLISION_PERSISTENCE_ENABLED
 	TMap< FConstraintContainerHandleKey, FConstraintContainerHandle* > Manifolds;
+#endif
 	TArray<FConstraintContainerHandle*> Handles;
 	FConstraintHandleAllocator HandleAllocator;
 
diff --git a/Engine/Source/Runtime/Experimental/Chaos/Public/Chaos/PBDCollisionConstraintsContact.h b/Engine/Source/Runtime/Experimental/Chaos/Public/Chaos/PBDCollisionConstraintsContact.h
--- a/Engine/Source/Runtime/Experimental/Chaos/Public/Chaos/PBDCollisionConstraintsContact.h
+++ b/Engine/Source/Runtime/Experimental/Chaos/Public/Chaos/PBDCollisionConstraintsContact.h
@@ -15,6 +15,7 @@ namespace Chaos
 			T ShapePadding;
 			TArrayCollectionArray<bool>* Collided;
 		};
+		using FContactParticleParameters = TContactParticleParameters<FReal>;
 
 		template<class T = float>
 		struct TContactIterationParameters {
@@ -24,19 +25,27 @@ namespace Chaos
 			const int32 NumPairIterations;
 			bool* NeedsAnotherIteration;
 		};
+		using FContactIterationParameters = TContactIterationParameters<FReal>;
 
-		template<ECollisionUpdateType UpdateType, typename T, int d>
-		void Update(const T CullDistance, TCollisionConstraintBase<T, d>& Constraint);
+		// Regenerate (one-shot or incremental) the manifold plane and points
+		extern void UpdateManifold(FRigidBodyMultiPointContactConstraint& Constraint, const FReal CullDistance);
 
-		template<typename T, int d>
-		void UpdateManifold(const T CullDistance, TCollisionConstraintBase<T, d>& Constraint);
+		// Update the constraint (re-runs collision detection for this contact)
+		extern void Update(FRigidBodyPointContactConstraint& Constraint, const FReal CullDistance);
 
-		template<typename T, int d>
-		void Apply(TCollisionConstraintBase<T, d>& Constraint, const TContactIterationParameters<T> & IterationParameters, const TContactParticleParameters<T> & ParticleParameters);
+		// Update the constraint (select best point from the manifold)
+		extern void Update(FRigidBodyMultiPointContactConstraint& Constraint, const FReal CullDistance);
 
-		template<typename T, int d>
-		void ApplyPushOut(TCollisionConstraintBase<T, d>& Constraint, const TSet<const TGeometryParticleHandle<T, d>*>& IsTemporarilyStatic,
-			const TContactIterationParameters<T> & IterationParameters, const TContactParticleParameters<T> & ParticleParameters);
+		extern void Apply(FCollisionConstraintBase& Constraint, const FContactIterationParameters& IterationParameters, const FContactParticleParameters& ParticleParameters);
+		extern void Apply(FRigidBodyPointContactConstraint& Constraint, const FContactIterationParameters& IterationParameters, const FContactParticleParameters& ParticleParameters);
+		extern void Apply(FRigidBodyMultiPointContactConstraint& Constraint, const FContactIterationParameters& IterationParameters, const FContactParticleParameters& ParticleParameters);
+
+		void ApplyPushOut(FCollisionConstraintBase& Constraint, const TSet<const TGeometryParticleHandle<FReal, 3>*>& IsTemporarilyStatic,
+			const FContactIterationParameters& IterationParameters, const FContactParticleParameters& ParticleParameters);
+		void ApplyPushOut(FRigidBodyPointContactConstraint& Constraint, const TSet<const TGeometryParticleHandle<FReal, 3>*>& IsTemporarilyStatic,
+			const FContactIterationParameters& IterationParameters, const FContactParticleParameters& ParticleParameters);
+		void ApplyPushOut(FRigidBodyPointContactConstraint& Constraint, const TSet<const TGeometryParticleHandle<FReal, 3>*>& IsTemporarilyStatic,
+			const FContactIterationParameters& IterationParameters, const FContactParticleParameters& ParticleParameters);
 
 
 	}
diff --git a/Engine/Source/Runtime/Experimental/Chaos/Public/Chaos/Utilities.h b/Engine/Source/Runtime/Experimental/Chaos/Public/Chaos/Utilities.h
--- a/Engine/Source/Runtime/Experimental/Chaos/Public/Chaos/Utilities.h
+++ b/Engine/Source/Runtime/Experimental/Chaos/Public/Chaos/Utilities.h
@@ -286,6 +286,107 @@ namespace Chaos
 			return OutTA >= 0 && OutTA <= 1 && OutTB > 0 && OutTB < 1;
 		}
 
+		/**
+		 * Clip a line segment to inside a plane (plane normal pointing outwards).
+		 * @return false if the line is completely outside the plane, true otherwise.
+		 */
+		inline bool ClipLineSegmentToPlane(FVec3& V0, FVec3& V1, const FVec3& PlaneNormal, const FVec3& PlanePos)
+		{
+			FReal Dist0 = FVec3::DotProduct(V0 - PlanePos, PlaneNormal);
+			FReal Dist1 = FVec3::DotProduct(V1 - PlanePos, PlaneNormal);
+			if ((Dist0 > 0.0f) && (Dist1 > 0.0f))
+			{
+				// Whole line segment is outside of face - reject it
+				return false;
+			}
+			
+			if ((Dist0 > 0.0f) && (Dist1 < 0.0f))
+			{
+				// We must move vert 0 to the plane
+				FReal ClippedT = -Dist1 / (Dist0 - Dist1);
+				V0 = FMath::Lerp(V1, V0, ClippedT);
+			}
+			else if ((Dist1 > 0.0f) && (Dist0 < 0.0f))
+			{
+				// We must move vert 1 to the plane
+				FReal ClippedT = -Dist0 / (Dist1 - Dist0);
+				V1 = FMath::Lerp(V0, V1, ClippedT);
+			}
+
+			return true;
+		}
+
+		/**
+		 * Clip a line segment to the inside of an axis aligned plane (normal pointing outwards).
+		 */
+		inline bool ClipLineSegmentToAxisAlignedPlane(FVec3& V0, FVec3& V1, const int32 AxisIndex, const FReal PlaneDir, const FReal PlanePos)
+		{
+			FReal Dist0 = (V0[AxisIndex] - PlanePos) * PlaneDir;
+			FReal Dist1 = (V1[AxisIndex] - PlanePos) * PlaneDir;
+			if ((Dist0 > 0.0f) && (Dist1 > 0.0f))
+			{
+				// Whole line segment is outside of face - reject it
+				return false;
+			}
+
+			if ((Dist0 > 0.0f) && (Dist1 < 0.0f))
+			{
+				// We must move vert 0 to the plane
+				FReal ClippedT = -Dist1 / (Dist0 - Dist1);
+				V0 = FMath::Lerp(V1, V0, ClippedT);
+			}
+			else if ((Dist1 > 0.0f) && (Dist0 < 0.0f))
+			{
+				// We must move vert 1 to the plane
+				FReal ClippedT = -Dist0 / (Dist1 - Dist0);
+				V1 = FMath::Lerp(V0, V1, ClippedT);
+			}
+
+			return true;
+		}
+
+		/**
+		 * Project a point V along direction Dir onto an axis aligned plane.
+		 * /note Does not check for division by zero (Dir parallel to plane).
+		 */
+		inline void ProjectPointOntoAxisAlignedPlane(FVec3& V, const FVec3& Dir, int32 AxisIndex, FReal PlaneDir, FReal PlanePos)
+		{
+			// V -> V + ((PlanePos - V) | PlaneNormal) / (Dir | PlaneNormal)
+			FReal Denominator = Dir[AxisIndex] * PlaneDir;
+			FReal Numerator = (PlanePos - V[AxisIndex]) * PlaneDir;
+			FReal F = Numerator / Denominator;
+			V = V + F * Dir;
+		}
+
+		/**
+		 * Project a point V along direction Dir onto an axis aligned plane.
+		 * /return true if the point was successfully projected onto the plane; false if the direction is parallel to the plane.
+		 */
+		inline bool ProjectPointOntoAxisAlignedPlaneSafe(FVec3& V, const FVec3& Dir, int32 AxisIndex, FReal PlaneDir, FReal PlanePos, FReal Epsilon)
+		{
+			// V -> V + ((PlanePos - V) | PlaneNormal) / (Dir | PlaneNormal)
+			FReal Denominator = Dir[AxisIndex] * PlaneDir;
+			if (Denominator > Epsilon)
+			{
+				FReal Numerator = (PlanePos - V[AxisIndex]) * PlaneDir;
+				FReal F = Numerator / Denominator;
+				V = V + F * Dir;
+				return true;
+			}
+			return false;
+		}
+
+		inline bool NormalizeSafe(FVec3& V, FReal EpsilonSq = SMALL_NUMBER)
+		{
+			FReal VLenSq = V.SizeSquared();
+			if (VLenSq > EpsilonSq)
+			{
+				V = V * FMath::InvSqrt(VLenSq);
+				return true;
+			}
+			return false;
+		}
+
 		/**
 		 * Given the local-space inertia for an unscaled object, return an inertia as if generated from a non-uniformly scaled shape with the specified scale.
 		 * If bScaleMass is true, it also takes into account the fact that the mass would have changed by the increase in volume.
