Chaos - fix convex-convex manifold (temp until we optimize the manifold points)

diff --git a/Engine/Source/Runtime/Experimental/Chaos/Public/Chaos/CollisionResolutionTypes.h b/Engine/Source/Runtime/Experimental/Chaos/Public/Chaos/CollisionResolutionTypes.h
--- a/Engine/Source/Runtime/Experimental/Chaos/Public/Chaos/CollisionResolutionTypes.h
+++ b/Engine/Source/Runtime/Experimental/Chaos/Public/Chaos/CollisionResolutionTypes.h
@@ -247,7 +247,7 @@ namespace Chaos
 
 		// Manifold point data
 		static const int32 MaxPoints = 4;
-		TArray<FVec3, TFixedAllocator<MaxPoints>> Points;
+		TArray<FVec3, TInlineAllocator<MaxPoints>> Points; // @todo(chaos): use TFixedAllocator when we handle building the best manifold properly
 	};
 	typedef TRigidBodyMultiPointContactConstraint<float, 3> FRigidBodyMultiPointContactConstraint;
 
