Chaos - joint swing-twist lock fix (mostly - still a very rare glitch)

diff --git a/Engine/Source/Runtime/Experimental/Chaos/Private/Chaos/Joint/PBDJointSolverGaussSeidel.cpp b/Engine/Source/Runtime/Experimental/Chaos/Private/Chaos/Joint/PBDJointSolverGaussSeidel.cpp
--- a/Engine/Source/Runtime/Experimental/Chaos/Private/Chaos/Joint/PBDJointSolverGaussSeidel.cpp
+++ b/Engine/Source/Runtime/Experimental/Chaos/Private/Chaos/Joint/PBDJointSolverGaussSeidel.cpp
@@ -86,27 +86,34 @@ namespace Chaos
 	}
 
 
-	void GetTwistAxisAngleLocal(
+	FReal GetTwistAngle(const FRotation3& InTwist)
+	{
+		FRotation3 Twist = InTwist.GetNormalized();
+		ensure(FMath::Abs(Twist.W) <= 1.0f);
+		FReal Angle = Twist.GetAngle();
+		if (Angle > PI)
+		{
+			Angle = Angle - (FReal)2 * PI;
+		}
+		if (Twist.X < 0.0f)
+		{
+			Angle = -Angle;
+		}
+		return Angle;
+	}
+
+
+	void GetTwistAxisAngle(
 		const FRotation3& R0,
 		const FRotation3& R1,
-		FVec3& TwistAxisLocal,
-		FReal& TwistAngle)
+		FVec3& Axis,
+		FReal& Angle)
 	{
-		// Decompose rotation of body 1 relative to body 0 into swing and twist rotations, assuming twist is X axis
 		FRotation3 R01Twist, R01Swing;
 		FPBDJointUtilities::DecomposeSwingTwistLocal(R0, R1, R01Swing, R01Twist);
 
-		TwistAxisLocal = FJointConstants::TwistAxis();
-
-		TwistAngle = R01Twist.GetAngle();
-		if (TwistAngle > PI)
-		{
-			TwistAngle = TwistAngle - (FReal)2 * PI;
-		}
-		if (R01Twist.X < 0)
-		{
-			TwistAngle = -TwistAngle;
-		}
+		Axis = R0 * FJointConstants::TwistAxis();
+		Angle = GetTwistAngle(R01Twist);
 	}
 
 
@@ -207,15 +214,6 @@ namespace Chaos
 		FRotation3 R01Twist, R01Swing;
 		FPBDJointUtilities::DecomposeSwingTwistLocal(R0, R1, R01Swing, R01Twist);
 		FVec3 TwistAxis01 = FJointConstants::TwistAxis();
-		FReal TwistAngle = R01Twist.GetAngle();
-		if (TwistAngle > PI)
-		{
-			TwistAngle = TwistAngle - (FReal)2 * PI;
-		}
-		if (R01Twist.X < 0)
-		{
-			TwistAngle = -TwistAngle;
-		}
 		const FVec3 TwistAxis = R0 * TwistAxis01;
 
 		const FRotation3 R1NoTwist = R1 * R01Twist.Inverse();
@@ -424,7 +422,6 @@ namespace Chaos
 			}
 			else
 			{
-				
 				if ((Swing1Motion == EJointMotionType::Limited) && bSwingSoft)
 				{
 					ApplySwingConstraintSoft(Dt, SolverSettings, JointSettings, EJointAngularConstraintIndex::Swing1, EJointAngularAxisIndex::Swing1);
@@ -825,25 +822,24 @@ namespace Chaos
 	
 	void FJointSolverGaussSeidel::ApplyRotationConstraint(
 		const FReal Stiffness,
-		const FVec3& Axis0,
-		const FVec3& Axis1,
+		const FVec3& Axis,
 		const FReal Angle)
 	{
 		// World-space inverse mass
 		const FMatrix33 InvI0 = Utilities::ComputeWorldSpaceInertia(Qs[0], InvILs[0]);
 		const FMatrix33 InvI1 = Utilities::ComputeWorldSpaceInertia(Qs[1], InvILs[1]);
-		const FVec3 IA0 = Utilities::Multiply(InvI0, Axis0);
-		const FVec3 IA1 = Utilities::Multiply(InvI1, Axis1);
+		const FVec3 IA0 = Utilities::Multiply(InvI0, Axis);
+		const FVec3 IA1 = Utilities::Multiply(InvI1, Axis);
 
 		// Joint-space inverse mass
-		const FReal II0 = FVec3::DotProduct(Axis0, IA0);
-		const FReal II1 = FVec3::DotProduct(Axis1, IA1);
+		const FReal II0 = FVec3::DotProduct(Axis, IA0);
+		const FReal II1 = FVec3::DotProduct(Axis, IA1);
 
 		const FVec3 DR0 = IA0 * (Angle / (II0 + II1));
 		const FVec3 DR1 = IA1 * -(Angle / (II0 + II1));
-		//const FVec3 DR0 = Axis0 * (Angle * II0 / (II0 + II1));
-		//const FVec3 DR1 = Axis1 * -(Angle * II1 / (II0 + II1));
-		
+		//const FVec3 DR0 = Axis * (Angle * II0 / (II0 + II1));
+		//const FVec3 DR1 = Axis * -(Angle * II1 / (II0 + II1));
+
 		ApplyRotationDelta(Stiffness, DR0, DR1);
 
 		// Correct the positional error that was introduced by the rotation correction. This is
@@ -873,15 +869,14 @@ namespace Chaos
 		const FReal Stiffness,
 		const FReal Damping,
 		const bool bAccelerationMode,
-		const FVec3& Axis0,
-		const FVec3& Axis1,
+		const FVec3& Axis,
 		const FReal Angle,
 		FReal& Lambda)
 	{
 		if (bChaos_Joint_ISPC_Enabled)
 		{
 #if INTEL_ISPC
-			ispc::ApplyRotationConstraintSoft((ispc::FJointSolverGaussSeidel*)this, Dt, Stiffness, Damping, bAccelerationMode, (ispc::FVector&) Axis0, (ispc::FVector&)Axis1, Angle, Lambda);
+			ispc::ApplyRotationConstraintSoft((ispc::FJointSolverGaussSeidel*)this, Dt, Stiffness, Damping, bAccelerationMode, (ispc::FVector&) Axis, Angle, Lambda);
 #endif
 		}
 		else
@@ -889,12 +884,12 @@ namespace Chaos
 			// World-space inverse mass
 			const FMatrix33 InvI0 = Utilities::ComputeWorldSpaceInertia(Qs[0], InvILs[0]);
 			const FMatrix33 InvI1 = Utilities::ComputeWorldSpaceInertia(Qs[1], InvILs[1]);
-			const FVec3 IA0 = Utilities::Multiply(InvI0, Axis0);
-			const FVec3 IA1 = Utilities::Multiply(InvI1, Axis1);
+			const FVec3 IA0 = Utilities::Multiply(InvI0, Axis);
+			const FVec3 IA1 = Utilities::Multiply(InvI1, Axis);
 
 			// Joint-space inverse mass
-			FReal II0 = FVec3::DotProduct(Axis0, IA0);
-			FReal II1 = FVec3::DotProduct(Axis1, IA1);
+			FReal II0 = FVec3::DotProduct(Axis, IA0);
+			FReal II1 = FVec3::DotProduct(Axis, IA1);
 
 			// If we are correcting the position, we need to adjust the constraint effective mass using parallel axis theorem
 			// @todo(ccaulfield): the IMs here are constant per constraint...pre-build and cache it
@@ -902,7 +897,7 @@ namespace Chaos
 			{
 				if (InvMs[0] > 0)
 				{
-					const FVec3 LinearAxis0 = FVec3::CrossProduct(Xs[0] - Ps[0], Axis0);
+					const FVec3 LinearAxis0 = FVec3::CrossProduct(Xs[0] - Ps[0], Axis);
 					const FReal LinearAxisLen0Sq = LinearAxis0.SizeSquared();
 					if (LinearAxisLen0Sq > KINDA_SMALL_NUMBER)
 					{
@@ -912,7 +907,7 @@ namespace Chaos
 				}
 				if (InvMs[1] > 0)
 				{
-					const FVec3 LinearAxis1 = FVec3::CrossProduct(Xs[1] - Ps[1], Axis1);
+					const FVec3 LinearAxis1 = FVec3::CrossProduct(Xs[1] - Ps[1], Axis);
 					const FReal LinearAxisLen1Sq = LinearAxis1.SizeSquared();
 					if (LinearAxisLen1Sq > KINDA_SMALL_NUMBER)
 					{
@@ -929,7 +924,7 @@ namespace Chaos
 			{
 				const FVec3 W0 = FRotation3::CalculateAngularVelocity(PrevQs[0], Qs[0], 1.0f);
 				const FVec3 W1 = FRotation3::CalculateAngularVelocity(PrevQs[1], Qs[1], 1.0f);
-				AngVelDt = FVec3::DotProduct(Axis0, W0) - FVec3::DotProduct(Axis1, W1);
+				AngVelDt = FVec3::DotProduct(Axis, W0) - FVec3::DotProduct(Axis, W1);
 			}
 
 			const FReal SpringMassScale = (bAccelerationMode) ? 1.0f / II : 1.0f;
@@ -940,8 +935,8 @@ namespace Chaos
 
 			//const FVec3 DR0 = IA0 * DLambda;
 			//const FVec3 DR1 = IA1 * -DLambda;
-			const FVec3 DR0 = Axis0 * (DLambda * II0);
-			const FVec3 DR1 = Axis1 * -(DLambda * II1);
+			const FVec3 DR0 = Axis * (DLambda * II0);
+			const FVec3 DR1 = Axis * -(DLambda * II1);
 
 			Lambda += DLambda;
 			ApplyRotationDelta(1.0f, DR0, DR1);
@@ -980,25 +975,13 @@ namespace Chaos
 		const FPBDJointSolverSettings& SolverSettings,
 		const FPBDJointSettings& JointSettings)
 	{
-		FVec3 TwistAxisLocal;
+		FVec3 TwistAxis;
 		FReal TwistAngle;
-		GetTwistAxisAngleLocal(Rs[0], Rs[1], TwistAxisLocal, TwistAngle);
-		
-		const FVec3 TwistAxis0 = Rs[0] * TwistAxisLocal;
-		const FVec3 TwistAxis1 = Rs[1] * TwistAxisLocal;
-
-		FReal TwistAngleMax = FLT_MAX;
-		if (JointSettings.AngularMotionTypes[(int32)EJointAngularConstraintIndex::Twist] == EJointMotionType::Limited)
-		{
-			TwistAngleMax = JointSettings.AngularLimits[(int32)EJointAngularConstraintIndex::Twist];
-		}
-		else if (JointSettings.AngularMotionTypes[(int32)EJointAngularConstraintIndex::Twist] == EJointMotionType::Locked)
-		{
-			TwistAngleMax = 0;
-		}
+		GetTwistAxisAngle(Rs[0], Rs[1], TwistAxis, TwistAngle);
 
 		// Calculate the twist correction to apply to each body
 		FReal DTwistAngle = 0;
+		FReal TwistAngleMax = JointSettings.AngularLimits[(int32)EJointAngularConstraintIndex::Twist];
 		if (TwistAngle > TwistAngleMax)
 		{
 			DTwistAngle = TwistAngle - TwistAngleMax;
@@ -1012,7 +995,7 @@ namespace Chaos
 		if (DTwistAngle != 0)
 		{
 			FReal TwistStiffness = FPBDJointUtilities::GetTwistStiffness(SolverSettings, JointSettings);
-			ApplyRotationConstraint(TwistStiffness, TwistAxis0, TwistAxis1, DTwistAngle);
+			ApplyRotationConstraint(TwistStiffness, TwistAxis, DTwistAngle);
 		}
 	}
 
@@ -1022,25 +1005,13 @@ namespace Chaos
 		const FPBDJointSolverSettings& SolverSettings,
 		const FPBDJointSettings& JointSettings)
 	{
-		FVec3 TwistAxisLocal;
+		FVec3 TwistAxis;
 		FReal TwistAngle;
-		GetTwistAxisAngleLocal(Rs[0], Rs[1], TwistAxisLocal, TwistAngle);
-
-		const FVec3 TwistAxis0 = Rs[0] * TwistAxisLocal;
-		const FVec3 TwistAxis1 = Rs[1] * TwistAxisLocal;
-
-		FReal TwistAngleMax = FLT_MAX;
-		if (JointSettings.AngularMotionTypes[(int32)EJointAngularConstraintIndex::Twist] == EJointMotionType::Limited)
-		{
-			TwistAngleMax = JointSettings.AngularLimits[(int32)EJointAngularConstraintIndex::Twist];
-		}
-		else if (JointSettings.AngularMotionTypes[(int32)EJointAngularConstraintIndex::Twist] == EJointMotionType::Locked)
-		{
-			TwistAngleMax = 0;
-		}
+		GetTwistAxisAngle(Rs[0], Rs[1], TwistAxis, TwistAngle);
 
 		// Calculate the twist correction to apply to each body
 		FReal DTwistAngle = 0;
+		FReal TwistAngleMax = JointSettings.AngularLimits[(int32)EJointAngularConstraintIndex::Twist];
 		if (TwistAngle > TwistAngleMax)
 		{
 			DTwistAngle = TwistAngle - TwistAngleMax;
@@ -1056,7 +1027,7 @@ namespace Chaos
 			const FReal TwistStiffness = FPBDJointUtilities::GetSoftTwistStiffness(SolverSettings, JointSettings);
 			const FReal TwistDamping = FPBDJointUtilities::GetSoftTwistDamping(SolverSettings, JointSettings);
 			const bool bAccelerationMode = FPBDJointUtilities::GetAngularSoftAccelerationMode(SolverSettings, JointSettings);
-			ApplyRotationConstraintSoft(Dt, TwistStiffness, TwistDamping, bAccelerationMode, TwistAxis0, TwistAxis1, DTwistAngle, TwistSoftLambda);
+			ApplyRotationConstraintSoft(Dt, TwistStiffness, TwistDamping, bAccelerationMode, TwistAxis, DTwistAngle, TwistSoftLambda);
 		}
 	}
 
@@ -1066,23 +1037,10 @@ namespace Chaos
 		const FPBDJointSolverSettings& SolverSettings,
 		const FPBDJointSettings& JointSettings)
 	{
-		// Decompose rotation of body 1 relative to body 0 into swing and twist rotations, assuming twist is X axis
-		FRotation3 R01Twist, R01Swing;
-		FPBDJointUtilities::DecomposeSwingTwistLocal(Rs[0], Rs[1], R01Swing, R01Twist);
-
-		FVec3 TwistAxis01 = FJointConstants::TwistAxis();
-		FReal TwistAngle = R01Twist.GetAngle();
-		if (TwistAngle > PI)
-		{
-			TwistAngle = TwistAngle - (FReal)2 * PI;
-		}
-		if (R01Twist.X < 0)
-		{
-			TwistAngle = -TwistAngle;
-		}
+		FVec3 TwistAxis;
+		FReal TwistAngle;
+		GetTwistAxisAngle(Rs[0], Rs[1], TwistAxis, TwistAngle);
 
-		const FVec3 TwistAxis0 = Rs[0] * TwistAxis01;
-		const FVec3 TwistAxis1 = Rs[1] * TwistAxis01;
 		const FReal TwistAngleTarget = JointSettings.AngularDriveTargetAngles[(int32)EJointAngularConstraintIndex::Twist];
 		const FReal DTwistAngle = TwistAngle - TwistAngleTarget;
 
@@ -1090,7 +1048,7 @@ namespace Chaos
 		const FReal AngularDriveStiffness = FPBDJointUtilities::GetAngularTwistDriveStiffness(SolverSettings, JointSettings);
 		const FReal AngularDriveDamping = FPBDJointUtilities::GetAngularTwistDriveDamping(SolverSettings, JointSettings);
 		const bool bAccelerationMode = FPBDJointUtilities::GetDriveAccelerationMode(SolverSettings, JointSettings);
-		ApplyRotationConstraintSoft(Dt, AngularDriveStiffness, AngularDriveDamping, bAccelerationMode, TwistAxis0, TwistAxis1, DTwistAngle, TwistDriveLambda);
+		ApplyRotationConstraintSoft(Dt, AngularDriveStiffness, AngularDriveDamping, bAccelerationMode, TwistAxis, DTwistAngle, TwistDriveLambda);
 	}
 
 
@@ -1121,7 +1079,7 @@ namespace Chaos
 		if (DSwingAngle != 0)
 		{
 			FReal SwingStiffness = FPBDJointUtilities::GetSwingStiffness(SolverSettings, JointSettings);
-			ApplyRotationConstraint(SwingStiffness, SwingAxis, SwingAxis, DSwingAngle);
+			ApplyRotationConstraint(SwingStiffness, SwingAxis, DSwingAngle);
 		}
 	}
 
@@ -1146,7 +1104,7 @@ namespace Chaos
 			const FReal SoftSwingStiffness = FPBDJointUtilities::GetSoftSwingStiffness(SolverSettings, JointSettings);
 			const FReal SoftSwingDamping = FPBDJointUtilities::GetSoftSwingDamping(SolverSettings, JointSettings);
 			const bool bAccelerationMode = FPBDJointUtilities::GetAngularSoftAccelerationMode(SolverSettings, JointSettings);
-			ApplyRotationConstraintSoft(Dt, SoftSwingStiffness, SoftSwingDamping, bAccelerationMode, SwingAxis, SwingAxis, DSwingAngle, SwingSoftLambda);
+			ApplyRotationConstraintSoft(Dt, SoftSwingStiffness, SoftSwingDamping, bAccelerationMode, SwingAxis, DSwingAngle, SwingSoftLambda);
 		}
 	}
 
@@ -1173,7 +1131,7 @@ namespace Chaos
 		const FReal AngularDriveStiffness = FPBDJointUtilities::GetAngularSwingDriveStiffness(SolverSettings, JointSettings);
 		const FReal AngularDriveDamping = FPBDJointUtilities::GetAngularSwingDriveDamping(SolverSettings, JointSettings);
 		const bool bAccelerationMode = FPBDJointUtilities::GetDriveAccelerationMode(SolverSettings, JointSettings);
-		ApplyRotationConstraintSoft(Dt, AngularDriveStiffness, AngularDriveDamping, bAccelerationMode, SwingAxis, SwingAxis, DSwingAngle, SwingDriveLambda);
+		ApplyRotationConstraintSoft(Dt, AngularDriveStiffness, AngularDriveDamping, bAccelerationMode, SwingAxis, DSwingAngle, SwingDriveLambda);
 	}
 
 
@@ -1241,11 +1199,9 @@ namespace Chaos
 		FReal SwingAngle;
 		GetSwingAxisAngle(Rs[0], Rs[1], SolverSettings.SwingTwistAngleTolerance, SwingConstraintIndex, SwingAxisIndex, SwingAxis, SwingAngle);
 
-		const bool isLocked = (JointSettings.AngularMotionTypes[(int32)SwingConstraintIndex] == EJointMotionType::Locked);
-		const FReal SwingAngleMax = isLocked ? 0 : JointSettings.AngularLimits[(int32)SwingConstraintIndex];
-
 		// Calculate swing error we need to correct
 		FReal DSwingAngle = 0;
+		const FReal SwingAngleMax = JointSettings.AngularLimits[(int32)SwingConstraintIndex];
 		if (SwingAngle > SwingAngleMax)
 		{
 			DSwingAngle = SwingAngle - SwingAngleMax;
@@ -1259,7 +1215,7 @@ namespace Chaos
 		if (DSwingAngle != 0)
 		{
 			const FReal SwingStiffness = FPBDJointUtilities::GetSwingStiffness(SolverSettings, JointSettings);
-			ApplyRotationConstraint(SwingStiffness, SwingAxis, SwingAxis, DSwingAngle);
+			ApplyRotationConstraint(SwingStiffness, SwingAxis, DSwingAngle);
 		}
 	}
 
@@ -1294,7 +1250,7 @@ namespace Chaos
 			const FReal SoftSwingStiffness = FPBDJointUtilities::GetSoftSwingStiffness(SolverSettings, JointSettings);
 			const FReal SoftSwingDamping = FPBDJointUtilities::GetSoftSwingDamping(SolverSettings, JointSettings);
 			const bool bAccelerationMode = FPBDJointUtilities::GetAngularSoftAccelerationMode(SolverSettings, JointSettings);
-			ApplyRotationConstraintSoft(Dt, SoftSwingStiffness, SoftSwingDamping, bAccelerationMode, SwingAxis, SwingAxis, DSwingAngle, SwingSoftLambda);
+			ApplyRotationConstraintSoft(Dt, SoftSwingStiffness, SoftSwingDamping, bAccelerationMode, SwingAxis, DSwingAngle, SwingSoftLambda);
 		}
 	}
 
@@ -1309,7 +1265,7 @@ namespace Chaos
 		FVec3 SwingAxis;
 		FReal SwingAngle;
 		GetSwingAxisAngle(Rs[0], Rs[1], SolverSettings.SwingTwistAngleTolerance, SwingConstraintIndex, SwingAxisIndex, SwingAxis, SwingAngle);
-	
+
 		const FReal SwingAngleTarget = JointSettings.AngularDriveTargetAngles[(int32)SwingConstraintIndex];
 		const FReal DSwingAngle = SwingAngle - SwingAngleTarget;
 
@@ -1317,7 +1273,7 @@ namespace Chaos
 		const FReal AngularDriveStiffness = FPBDJointUtilities::GetAngularSwingDriveStiffness(SolverSettings, JointSettings);
 		const FReal AngularDriveDamping = FPBDJointUtilities::GetAngularSwingDriveDamping(SolverSettings, JointSettings);
 		const bool bAccelerationMode = FPBDJointUtilities::GetDriveAccelerationMode(SolverSettings, JointSettings);
-		ApplyRotationConstraintSoft(Dt, AngularDriveStiffness, AngularDriveDamping, bAccelerationMode, SwingAxis, SwingAxis, DSwingAngle, SwingDriveLambda);
+		ApplyRotationConstraintSoft(Dt, AngularDriveStiffness, AngularDriveDamping, bAccelerationMode, SwingAxis, DSwingAngle, SwingDriveLambda);
 	}
 
 
@@ -1353,7 +1309,7 @@ namespace Chaos
 			const FReal AngularDriveStiffness = FPBDJointUtilities::GetAngularSLerpDriveStiffness(SolverSettings, JointSettings);
 			const FReal AngularDriveDamping = FPBDJointUtilities::GetAngularSLerpDriveDamping(SolverSettings, JointSettings);
 			const bool bAccelerationMode = FPBDJointUtilities::GetDriveAccelerationMode(SolverSettings, JointSettings);
-			ApplyRotationConstraintSoft(Dt, AngularDriveStiffness, AngularDriveDamping, bAccelerationMode, SLerpAxis, SLerpAxis, -SLerpAngle, SwingDriveLambda);
+			ApplyRotationConstraintSoft(Dt, AngularDriveStiffness, AngularDriveDamping, bAccelerationMode, SLerpAxis, -SLerpAngle, SwingDriveLambda);
 		}
 	}
 
diff --git a/Engine/Source/Runtime/Experimental/Chaos/Private/Chaos/PBDJointConstraintUtilities.cpp b/Engine/Source/Runtime/Experimental/Chaos/Private/Chaos/PBDJointConstraintUtilities.cpp
--- a/Engine/Source/Runtime/Experimental/Chaos/Private/Chaos/PBDJointConstraintUtilities.cpp
+++ b/Engine/Source/Runtime/Experimental/Chaos/Private/Chaos/PBDJointConstraintUtilities.cpp
@@ -8,11 +8,27 @@
 
 namespace Chaos
 {
+	void ToSwingTwistX(const FRotation3& Q, FRotation3& OutSwing, FRotation3& OutTwist)
+	{
+		FReal S = FMath::Sqrt(Q.X * Q.X + Q.W * Q.W);
+		if (S > KINDA_SMALL_NUMBER)
+		{
+			FReal SInv = 1.0f / S;
+			OutTwist = FRotation3::FromElements(Q.X * SInv, 0, 0, Q.W * SInv);
+			OutSwing = FRotation3::FromElements(0, (Q.W * Q.Y + Q.X * Q.Z) * SInv, (Q.W * Q.Z - Q.X * Q.Y) * SInv, S);
+		}
+		else
+		{
+			OutTwist = FRotation3::FromIdentity();
+			OutSwing = Q;
+		}
+	}
+
 
 	void FPBDJointUtilities::DecomposeSwingTwistLocal(const FRotation3& R0, const FRotation3& R1, FRotation3& R01Swing, FRotation3& R01Twist)
 	{
 		const FRotation3 R01 = R0.Inverse() * R1;
-		R01.ToSwingTwistX(R01Swing, R01Twist);
+		ToSwingTwistX(R01, R01Swing, R01Twist);
 	}
 
 	// @todo(ccaulfield): separate linear soft and stiff
diff --git a/Engine/Source/Runtime/Experimental/Chaos/Public/Chaos/Joint/PBDJointSolverGaussSeidel.h b/Engine/Source/Runtime/Experimental/Chaos/Public/Chaos/Joint/PBDJointSolverGaussSeidel.h
--- a/Engine/Source/Runtime/Experimental/Chaos/Public/Chaos/Joint/PBDJointSolverGaussSeidel.h
+++ b/Engine/Source/Runtime/Experimental/Chaos/Public/Chaos/Joint/PBDJointSolverGaussSeidel.h
@@ -171,8 +171,7 @@ namespace Chaos
 
 		void ApplyRotationConstraint(
 			const FReal Stiffness,
-			const FVec3& Axis0,
-			const FVec3& Axis1,
+			const FVec3& Axis,
 			const FReal Angle);
 
 		void ApplyRotationConstraintSoft(
@@ -180,8 +179,7 @@ namespace Chaos
 			const FReal Stiffness,
 			const FReal Damping,
 			const bool bAccelerationMode,
-			const FVec3& Axis0,
-			const FVec3& Axis1,
+			const FVec3& Axis,
 			const FReal Angle,
 			FReal& Lambda);
 
