Chaos - missing file (joint ispc)

diff --git a/Engine/Source/Runtime/Experimental/Chaos/Private/Chaos/Joint/PBDJointSolverGaussSeidel.ispc b/Engine/Source/Runtime/Experimental/Chaos/Private/Chaos/Joint/PBDJointSolverGaussSeidel.ispc
--- a/Engine/Source/Runtime/Experimental/Chaos/Private/Chaos/Joint/PBDJointSolverGaussSeidel.ispc
+++ b/Engine/Source/Runtime/Experimental/Chaos/Private/Chaos/Joint/PBDJointSolverGaussSeidel.ispc
@@ -217,20 +217,19 @@ export void ApplyRotationConstraintSoft(
 	const uniform float Stiffness,
 	const uniform float Damping,
 	const uniform bool bAccelerationMode,
-	const uniform FVector& Axis0,
-	const uniform FVector& Axis1,
+	const uniform FVector& Axis,
 	const uniform float Angle,
 	uniform float& Lambda)
 {
 	// World-space inverse mass
 	const uniform FMatrix33 InvI0 = ComputeWorldSpaceInertia(M->Qs[0], M->InvILs[0]);
 	const uniform FMatrix33 InvI1 = ComputeWorldSpaceInertia(M->Qs[1], M->InvILs[1]);
-	const uniform FVector IA0 = Multiply(InvI0, Axis0);
-	const uniform FVector IA1 = Multiply(InvI1, Axis1);
+	const uniform FVector IA0 = Multiply(InvI0, Axis);
+	const uniform FVector IA1 = Multiply(InvI1, Axis);
 
 	// Joint-space inverse mass
-	uniform float II0 = VectorDot(Axis0, IA0);
-	uniform float II1 = VectorDot(Axis1, IA1);
+	uniform float II0 = VectorDot(Axis, IA0);
+	uniform float II1 = VectorDot(Axis, IA1);
 
 	// If we are correcting the position, we need to adjust the constraint effective mass using parallel axis theorem
 	// @todo(ccaulfield): the IMs here are constant per constraint...pre-build and cache it
@@ -238,7 +237,7 @@ export void ApplyRotationConstraintSoft(
 	{
 		if (M->InvMs[0] > 0)
 		{
-			const uniform FVector LinearAxis0 = VectorCross(M->Xs[0] - M->Ps[0], Axis0);
+			const uniform FVector LinearAxis0 = VectorCross(M->Xs[0] - M->Ps[0], Axis);
 			const uniform float LinearAxisLen0Sq = VectorSizeSquared(LinearAxis0);
 			if (LinearAxisLen0Sq > KINDA_SMALL_NUMBER)
 			{
@@ -248,7 +247,7 @@ export void ApplyRotationConstraintSoft(
 		}
 		if (M->InvMs[1] > 0)
 		{
-			const uniform FVector LinearAxis1 = VectorCross(M->Xs[1] - M->Ps[1], Axis1);
+			const uniform FVector LinearAxis1 = VectorCross(M->Xs[1] - M->Ps[1], Axis);
 			const uniform float LinearAxisLen1Sq = VectorSizeSquared(LinearAxis1);
 			if (LinearAxisLen1Sq > KINDA_SMALL_NUMBER)
 			{
@@ -265,7 +264,7 @@ export void ApplyRotationConstraintSoft(
 	{
 		const uniform FVector W0 = CalculateAngularVelocity(M->PrevQs[0], M->Qs[0], 1.0f);
 		const uniform FVector W1 = CalculateAngularVelocity(M->PrevQs[1], M->Qs[1], 1.0f);
-		AngVelDt = VectorDot(Axis0, W0) - VectorDot(Axis1, W1);
+		AngVelDt = VectorDot(Axis, W0) - VectorDot(Axis, W1);
 	}
 
 	const uniform float SpringMassScale = (bAccelerationMode) ? 1.0f / II : 1.0f;
@@ -276,8 +275,8 @@ export void ApplyRotationConstraintSoft(
 
 	//const FVec3 DR0 = IA0 * DLambda;
 	//const FVec3 DR1 = IA1 * -DLambda;
-	const uniform FVector DR0 = Axis0 * (DLambda * II0);
-	const uniform FVector DR1 = Axis1 * -(DLambda * II1);
+	const uniform FVector DR0 = Axis * (DLambda * II0);
+	const uniform FVector DR1 = Axis * -(DLambda * II1);
 
 	Lambda += DLambda;
 	ApplyRotationDelta(M, 1.0f, DR0, DR1);
