Chaos - increase default debug draw line thickness for shapes

diff --git a/Engine/Source/Runtime/Experimental/Chaos/Private/Chaos/ChaosDebugDraw.cpp b/Engine/Source/Runtime/Experimental/Chaos/Private/Chaos/ChaosDebugDraw.cpp
--- a/Engine/Source/Runtime/Experimental/Chaos/Private/Chaos/ChaosDebugDraw.cpp
+++ b/Engine/Source/Runtime/Experimental/Chaos/Private/Chaos/ChaosDebugDraw.cpp
@@ -32,7 +32,7 @@ namespace Chaos
 		float DrawScale = 1.0f;
 		float FontHeight = 10.0f;
 		float FontScale = 1.5f;
-		float ShapeThicknesScale = 0.2f;
+		float ShapeThicknesScale = 1.0f;
 		float PointSize = 2.0f;
 		int DrawPriority = 10.0f;
 
