Chaos - fix collision RemoveConstraint (for clusters)

diff --git a/Engine/Source/Runtime/Experimental/Chaos/Private/Chaos/PBDCollisionConstraints.cpp b/Engine/Source/Runtime/Experimental/Chaos/Private/Chaos/PBDCollisionConstraints.cpp
--- a/Engine/Source/Runtime/Experimental/Chaos/Private/Chaos/PBDCollisionConstraints.cpp
+++ b/Engine/Source/Runtime/Experimental/Chaos/Private/Chaos/PBDCollisionConstraints.cpp
@@ -152,6 +152,8 @@ namespace Chaos
 		FConstraintContainerHandle* Handle = HandleAllocator.template AllocHandle< TRigidBodyPointContactConstraint<T, d> >(this, Idx);
 		Handle->GetContact().Timestamp = -INT_MAX; // force point constraints to be deleted.
 
+		PointConstraints[Idx].ConstraintHandle = Handle;
+
 		check(Handle != nullptr);
 		Handles.Add(Handle);
 
@@ -170,6 +172,8 @@ namespace Chaos
 		FConstraintContainerHandle* Handle = HandleAllocator.template AllocHandle< TRigidBodyMultiPointContactConstraint<T, d> >(this, Idx);
 		Handle->GetContact().Timestamp = LifespanCounter;
 
+		IterativeConstraints[Idx].ConstraintHandle = Handle;
+
 		check(Handle != nullptr);
 		Handles.Add(Handle);
 
@@ -270,6 +274,10 @@ namespace Chaos
 			}
 #endif
 			PointConstraints.RemoveAtSwap(Idx);
+			if (Idx < PointConstraints.Num())
+			{
+				PointConstraints[Idx].ConstraintHandle->SetConstraintIndex(Idx, FCollisionConstraintBase::FType::SinglePoint);
+			}
 
 		}
 		else if (ConstraintType == FCollisionConstraintBase::FType::MultiPoint)
@@ -283,6 +291,10 @@ namespace Chaos
 			}
 #endif
 			IterativeConstraints.RemoveAtSwap(Idx);
+			if (Idx < IterativeConstraints.Num())
+			{
+				IterativeConstraints[Idx].ConstraintHandle->SetConstraintIndex(Idx, FCollisionConstraintBase::FType::MultiPoint);
+			}
 		}
 		else 
 		{
diff --git a/Engine/Source/Runtime/Experimental/Chaos/Public/Chaos/CollisionResolutionTypes.h b/Engine/Source/Runtime/Experimental/Chaos/Public/Chaos/CollisionResolutionTypes.h
--- a/Engine/Source/Runtime/Experimental/Chaos/Public/Chaos/CollisionResolutionTypes.h
+++ b/Engine/Source/Runtime/Experimental/Chaos/Public/Chaos/CollisionResolutionTypes.h
@@ -101,6 +101,7 @@ namespace Chaos
 		TCollisionConstraintBase(FType InType = FType::None)
 			: AccumulatedImpulse(0)
 			, Timestamp(-INT_MAX)
+			, ConstraintHandle(nullptr)
 			, Type(InType)
 		{ 
 			ImplicitTransform[0] = TRigidTransform<T, d>::Identity; ImplicitTransform[1] = TRigidTransform<T,d>::Identity;
@@ -114,6 +115,7 @@ namespace Chaos
 			FType InType, EContactShapesType ShapesType, int32 InTimestamp = -INT_MAX)
 			: AccumulatedImpulse(0)
 			, Timestamp(InTimestamp)
+			, ConstraintHandle(nullptr)
 			, Type(InType)
 		{
 			ImplicitTransform[0] = Transform0; ImplicitTransform[1] = Transform1;
@@ -160,6 +162,7 @@ namespace Chaos
 		TVector<T, d> AccumulatedImpulse;
 		FManifold Manifold;
 		int32 Timestamp;
+		TPBDCollisionConstraintHandle<T, d>* ConstraintHandle;
 
 	private:
 
