Chaos: Fix threading issue with accesing proxy through GTParticle, reverted other temp fix where commands were De-Qd again.

diff --git a/Engine/Source/Runtime/Engine/Private/PhysicsEngine/Experimental/PhysScene_Chaos.cpp b/Engine/Source/Runtime/Engine/Private/PhysicsEngine/Experimental/PhysScene_Chaos.cpp
--- a/Engine/Source/Runtime/Engine/Private/PhysicsEngine/Experimental/PhysScene_Chaos.cpp
+++ b/Engine/Source/Runtime/Engine/Private/PhysicsEngine/Experimental/PhysScene_Chaos.cpp
@@ -1743,15 +1743,6 @@ void FPhysScene_ChaosInterface::CompleteSceneSimulation(ENamedThreads::Type Curr
 		{
 			if(Solver->HasActiveParticles())
 			{
-				// We need to process the Command Q here to prevent corruption in BufferPhysicsResults & FlipBuffers
-				// needs further investigation.
-				TQueue<TFunction<void(Chaos::FPhysicsSolver*)>, EQueueMode::Mpsc>& Queue = Solver->GetCommandQueue();
-				TFunction<void(FPhysicsSolver*)> Command;
-				while (Queue.Dequeue(Command))
-				{
-					Command(Solver);
-				}
-
 				ActiveSolvers.Add(Solver);
 			}
 		}
diff --git a/Engine/Source/Runtime/Experimental/ChaosSolvers/Private/EventDefaults.cpp b/Engine/Source/Runtime/Experimental/ChaosSolvers/Private/EventDefaults.cpp
--- a/Engine/Source/Runtime/Experimental/ChaosSolvers/Private/EventDefaults.cpp
+++ b/Engine/Source/Runtime/Experimental/ChaosSolvers/Private/EventDefaults.cpp
@@ -72,7 +72,8 @@ namespace Chaos
 
 						// Since Clustered GCs can be unioned the particleIndex representing the union 
 						// is not associated with a PhysicsProxy
-						if (Constraint.Particle[0]->Handle()->GTGeometryParticle()->Proxy != nullptr)
+						IPhysicsProxyBase* Proxy = Solver->GetProxy(Constraint.Particle[0]->Handle());
+						if (Proxy != nullptr)
 
 						{
 							if (ensure(!Constraint.AccumulatedImpulse.ContainsNaN() && FMath::IsFinite(Constraint.GetPhi())))
@@ -371,15 +372,16 @@ namespace Chaos
 			{
 				if(SleepData.Particle)
 				{
+					IPhysicsProxyBase* PhysicsProxy = NonConstSolver->GetProxy(SleepData.Particle);
+
 					TGeometryParticle<float, 3>* Particle = SleepData.Particle->GTGeometryParticle();
-					if(Particle && Particle->Proxy != nullptr)
+					if(Particle != nullptr && PhysicsProxy != nullptr)
 					{
 						int32 NewIdx = EventSleepDataArray.Add(TSleepingData<float, 3>());
 						TSleepingData<float, 3>& SleepingDataArrayItem = EventSleepDataArray[NewIdx];
 						SleepingDataArrayItem.Particle = Particle;
 						SleepingDataArrayItem.Sleeping = SleepData.Sleeping;
 
-						IPhysicsProxyBase* PhysicsProxy = Particle->Proxy;
 						AllSleepIndicesByPhysicsProxy.FindOrAdd(PhysicsProxy).Add(NewIdx);
 					}
 				}
diff --git a/Engine/Source/Runtime/Experimental/ChaosSolvers/Private/PBDRigidsSolver.cpp b/Engine/Source/Runtime/Experimental/ChaosSolvers/Private/PBDRigidsSolver.cpp
--- a/Engine/Source/Runtime/Experimental/ChaosSolvers/Private/PBDRigidsSolver.cpp
+++ b/Engine/Source/Runtime/Experimental/ChaosSolvers/Private/PBDRigidsSolver.cpp
@@ -260,6 +260,7 @@ namespace Chaos
 			}
 
 			Handle->GTGeometryParticle() = GTParticle;
+			Solver->MParticleToProxy.Add(Handle, ProxyBase);
 			Solver->GetEvolution()->CreateParticle(Handle);
 
 			delete ProxyData;
@@ -332,6 +333,8 @@ namespace Chaos
 				delete Proxy;
 			}
 
+			Solver->MParticleToProxy.Remove(Handle);
+
 			// Use the handle to destroy the particle data
 			Solver->GetEvolution()->DestroyParticle(Handle);
 		});
@@ -581,18 +584,19 @@ namespace Chaos
 		TParticleView<TPBDRigidParticles<float, 3>>& ActiveParticles = GetParticles().GetActiveParticlesView();
 		for (Chaos::TPBDRigidParticleHandleImp<float, 3, false>& ActiveObject : ActiveParticles)
 		{
-			if (ensure(ActiveObject.GTGeometryParticle()->Proxy))
+			IPhysicsProxyBase* Proxy = GetProxy(ActiveObject.Handle());
+			if (ensure(Proxy))
 			{
 				switch (ActiveObject.GetParticleType())
 				{
 				case Chaos::EParticleType::Rigid:
-					((FRigidParticlePhysicsProxy*)(ActiveObject.GTGeometryParticle()->Proxy))->BufferPhysicsResults();
+					((FRigidParticlePhysicsProxy*)(Proxy))->BufferPhysicsResults();
 					break;
 				case Chaos::EParticleType::Kinematic:
-					((FKinematicGeometryParticlePhysicsProxy*)(ActiveObject.GTGeometryParticle()->Proxy))->BufferPhysicsResults();
+					((FKinematicGeometryParticlePhysicsProxy*)(Proxy))->BufferPhysicsResults();
 					break;
 				case Chaos::EParticleType::Static:
-					((FGeometryParticlePhysicsProxy*)(ActiveObject.GTGeometryParticle()->Proxy))->BufferPhysicsResults();
+					((FGeometryParticlePhysicsProxy*)(Proxy))->BufferPhysicsResults();
 					break;
 				case Chaos::EParticleType::GeometryCollection:
 					ActiveGC.AddUnique((FGeometryCollectionPhysicsProxy*)(ActiveObject.GTGeometryParticle()->Proxy));
@@ -618,18 +622,19 @@ namespace Chaos
 		TParticleView<TPBDRigidParticles<float, 3>>& ActiveParticles = GetParticles().GetActiveParticlesView();
 		for (Chaos::TPBDRigidParticleHandleImp<float, 3, false>& ActiveObject : ActiveParticles)
 		{
-			if (ensure(ActiveObject.GTGeometryParticle()->Proxy))
+			IPhysicsProxyBase* Proxy = GetProxy(ActiveObject.Handle());
+			if (ensure(Proxy))
 			{
 				switch (ActiveObject.GetParticleType())
 				{
 				case Chaos::EParticleType::Rigid:
-					((FRigidParticlePhysicsProxy*)(ActiveObject.GTGeometryParticle()->Proxy))->FlipBuffer();
+					((FRigidParticlePhysicsProxy*)(Proxy))->FlipBuffer();
 					break;
 				case Chaos::EParticleType::Kinematic:
-					((FKinematicGeometryParticlePhysicsProxy*)(ActiveObject.GTGeometryParticle()->Proxy))->FlipBuffer();
+					((FKinematicGeometryParticlePhysicsProxy*)(Proxy))->FlipBuffer();
 					break;
 				case Chaos::EParticleType::Static:
-					((FGeometryParticlePhysicsProxy*)(ActiveObject.GTGeometryParticle()->Proxy))->FlipBuffer();
+					((FGeometryParticlePhysicsProxy*)(Proxy))->FlipBuffer();
 					break;
 				case Chaos::EParticleType::GeometryCollection:
 					ActiveGC.AddUnique((FGeometryCollectionPhysicsProxy*)(ActiveObject.GTGeometryParticle()->Proxy));
@@ -655,18 +660,19 @@ namespace Chaos
 		TParticleView<TPBDRigidParticles<float, 3>>& ActiveParticles = GetParticles().GetActiveParticlesView();
 		for (Chaos::TPBDRigidParticleHandleImp<float, 3, false>& ActiveObject : ActiveParticles)
 		{
-			if (ensure(ActiveObject.GTGeometryParticle()->Proxy))
+			IPhysicsProxyBase* Proxy = GetProxy(ActiveObject.Handle());
+			if (ensure(Proxy))
 			{
 				switch (ActiveObject.GetParticleType())
 				{
 				case Chaos::EParticleType::Rigid:
-					((FRigidParticlePhysicsProxy*)(ActiveObject.GTGeometryParticle()->Proxy))->PullFromPhysicsState();
+					((FRigidParticlePhysicsProxy*)(Proxy))->PullFromPhysicsState();
 					break;
 				case Chaos::EParticleType::Kinematic:
-					((FKinematicGeometryParticlePhysicsProxy*)(ActiveObject.GTGeometryParticle()->Proxy))->PullFromPhysicsState();
+					((FKinematicGeometryParticlePhysicsProxy*)(Proxy))->PullFromPhysicsState();
 					break;
 				case Chaos::EParticleType::Static:
-					((FGeometryParticlePhysicsProxy*)(ActiveObject.GTGeometryParticle()->Proxy))->PullFromPhysicsState();
+					((FGeometryParticlePhysicsProxy*)(Proxy))->PullFromPhysicsState();
 					break;
 				case Chaos::EParticleType::GeometryCollection:
 					ActiveGC.AddUnique((FGeometryCollectionPhysicsProxy*)(ActiveObject.GTGeometryParticle()->Proxy));
diff --git a/Engine/Source/Runtime/Experimental/ChaosSolvers/Public/PBDRigidsSolver.h b/Engine/Source/Runtime/Experimental/ChaosSolvers/Public/PBDRigidsSolver.h
--- a/Engine/Source/Runtime/Experimental/ChaosSolvers/Public/PBDRigidsSolver.h
+++ b/Engine/Source/Runtime/Experimental/ChaosSolvers/Public/PBDRigidsSolver.h
@@ -298,21 +298,28 @@ namespace Chaos
 		template<typename ParticleType>
 		void FlipBuffer(Chaos::TGeometryParticleHandle<float, 3>* Handle)
 		{
-			((ParticleType*)(Handle->GTGeometryParticle()->Proxy))->FlipBuffer();
+			((ParticleType*)(GetProxy(Handle)))->FlipBuffer();
 		}
 
 		template<typename ParticleType>
 		void PullFromPhysicsState(Chaos::TGeometryParticleHandle<float, 3>* Handle)
 		{
-			((ParticleType*)(Handle->GTGeometryParticle()->Proxy))->PullFromPhysicsState();
+			((ParticleType*)(GetProxy(Handle)))->PullFromPhysicsState();
 		}
 
 		template<typename ParticleType>
 		void BufferPhysicsResults(Chaos::TGeometryParticleHandle<float, 3>* Handle)
 		{
-			((ParticleType*)(Handle->GTGeometryParticle()->Proxy))->BufferPhysicsResults();
+			((ParticleType*)(GetProxy(Handle)))->BufferPhysicsResults();
 		}
 
+		IPhysicsProxyBase* GetProxy(const Chaos::TGeometryParticleHandle<float, 3>* Handle) const
+		{
+			IPhysicsProxyBase* const* PhysicsProxyPtr = MParticleToProxy.Find(Handle);
+			return PhysicsProxyPtr ? *PhysicsProxyPtr : nullptr;
+		}
+		
+
 		//
 		// Solver Data
 		//
@@ -332,6 +339,7 @@ namespace Chaos
 		TUniquePtr<FEventManager> MEventManager;
 		TUniquePtr<FSolverEventFilters> MSolverEventFilters;
 		TUniquePtr<FActiveParticlesBuffer> MActiveParticlesBuffer;
+		TMap<const Chaos::TGeometryParticleHandle<float, 3>*, IPhysicsProxyBase*> MParticleToProxy;
 
 		//
 		// Commands
