Fix clang compile error.

diff --git a/Engine/Source/Runtime/Experimental/Chaos/Public/Chaos/GeometryParticles.h b/Engine/Source/Runtime/Experimental/Chaos/Public/Chaos/GeometryParticles.h
--- a/Engine/Source/Runtime/Experimental/Chaos/Public/Chaos/GeometryParticles.h
+++ b/Engine/Source/Runtime/Experimental/Chaos/Public/Chaos/GeometryParticles.h
@@ -219,8 +219,8 @@ namespace Chaos
 		CHAOS_API const TRotation<T, d>& R(const int32 Index) const { return MR[Index]; }
 		CHAOS_API TRotation<T, d>& R(const int32 Index) { return MR[Index]; }
 
-		FUniqueIdx UniqueIdx(const int32 Index) const { return MUniqueIdx[Index]; }
-		FUniqueIdx& UniqueIdx(const int32 Index) { return MUniqueIdx[Index]; }
+		CHAOS_API FUniqueIdx UniqueIdx(const int32 Index) const { return MUniqueIdx[Index]; }
+		CHAOS_API FUniqueIdx& UniqueIdx(const int32 Index) { return MUniqueIdx[Index]; }
 
 		CHAOS_API TSerializablePtr<FImplicitObject> Geometry(const int32 Index) const { return MGeometry[Index]; }
 
