Chaos Cloth - Fix long range attachments crashes on multi-island meshes and on multi-cloth assets.
- Optimize geodesic queue creation.
- Replace auto from code.

diff --git a/Engine/Source/Runtime/Experimental/Chaos/Private/Chaos/PBDLongRangeConstraintsBase.cpp b/Engine/Source/Runtime/Experimental/Chaos/Private/Chaos/PBDLongRangeConstraintsBase.cpp
--- a/Engine/Source/Runtime/Experimental/Chaos/Private/Chaos/PBDLongRangeConstraintsBase.cpp
+++ b/Engine/Source/Runtime/Experimental/Chaos/Private/Chaos/PBDLongRangeConstraintsBase.cpp
@@ -194,9 +194,12 @@ void TPBDLongRangeConstraintsBase<T, d>::ComputeGeodesicConstraints(
     const TMap<int32, TSet<uint32>>& PointToNeighbors,
     const int32 NumberOfAttachments)
 {
+	TArray<int32> UsedIndices;
+	PointToNeighbors.GenerateKeyArray(UsedIndices);
+
 	// TODO(mlentine): Support changing which particles are kinematic during simulation
 	TArray<uint32> KinematicParticles;
-	for (uint32 i = 0; i < InParticles.Size(); ++i)
+	for (const uint32 i : UsedIndices)
 	{
 		if (InParticles.InvM(i) == 0)
 		{
@@ -206,7 +209,7 @@ void TPBDLongRangeConstraintsBase<T, d>::ComputeGeodesicConstraints(
 	TArray<TArray<uint32>> IslandElements = ComputeIslands(InParticles, PointToNeighbors, KinematicParticles);
 	// Store distances for all adjacent vertices
 	TMap<TVector<uint32, 2>, T> Distances;
-	for (uint32 i = 0; i < InParticles.Size(); ++i)
+	for (const uint32 i : UsedIndices)
 	{
 		auto Neighbors = PointToNeighbors[i];
 		for (auto Neighbor : Neighbors)
@@ -217,10 +220,10 @@ void TPBDLongRangeConstraintsBase<T, d>::ComputeGeodesicConstraints(
 	// Start and End Points to path and geodesic distance
 	TMap<TVector<uint32, 2>, Pair<T, TArray<uint32>>> GeodesicPaths;
 	// Dijkstra for each Kinematic Particle (assume a small number of kinematic points) - note this is N^2 log N with N kinematic points
-	for (auto Element : KinematicParticles)
+	for (const uint32 Element : KinematicParticles)
 	{
 		GeodesicPaths.Add(TVector<uint32, 2>(Element, Element), {0, {Element}});
-		for (uint32 i = 0; i < InParticles.Size(); ++i)
+		for (const uint32 i : UsedIndices)
 		{
 			if (i != Element)
 			{
@@ -228,35 +231,33 @@ void TPBDLongRangeConstraintsBase<T, d>::ComputeGeodesicConstraints(
 			}
 		}
 	}
-	PhysicsParallelFor(KinematicParticles.Num(), [&](int32 Index) {
-		auto Element = KinematicParticles[Index];
-		std::priority_queue<Pair<T, uint32>, std::vector<Pair<T, uint32>>, std::greater<Pair<T, uint32>>> q;
-		for (uint32 i = 0; i < InParticles.Size(); ++i)
-		{
-			q.push(MakePair(GeodesicPaths[TVector<uint32, 2>(Element, i)].First, i));
-		}
+	PhysicsParallelFor(KinematicParticles.Num(), [&](int32 Index)
+	{
+		const uint32 Element = KinematicParticles[Index];
+		std::priority_queue<Pair<T, uint32>, std::vector<Pair<T, uint32>>, std::greater<Pair<T, uint32>>> q;  // TODO(Kriss.Gossart): Remove use of std container
+		q.push(MakePair((T)0., Element));
 		TSet<uint32> Visited;
 		while (!q.empty())
 		{
-			auto PairElem = q.top();
+			const Pair<T, uint32> PairElem = q.top();
 			q.pop();
 			if (Visited.Contains(PairElem.Second))
 				continue;
 			Visited.Add(PairElem.Second);
-			auto CurrentStartEnd = TVector<uint32, 2>(Element, PairElem.Second);
-			auto Neighbors = PointToNeighbors[PairElem.Second];
-			for (auto Neighbor : Neighbors)
+			const TVector<uint32, 2> CurrentStartEnd(Element, PairElem.Second);
+			const TSet<uint32>& Neighbors = PointToNeighbors[PairElem.Second];
+			for (const uint32 Neighbor : Neighbors)
 			{
 				check(Neighbor != PairElem.Second);
-				auto NeighborStartEnd = TVector<uint32, 2>(Element, Neighbor);
-				auto NeighborDistancePath = GeodesicPaths[NeighborStartEnd];
+				const TVector<uint32, 2> NeighborStartEnd(Element, Neighbor);
+				const Pair<T, TArray<uint32>>& NeighborDistancePath = GeodesicPaths[NeighborStartEnd];
 				// Compute a possible distance for NeighborStartEnd
-				T NewDist = PairElem.First + Distances[TVector<uint32, 2>(PairElem.Second, Neighbor)];
+				const T NewDist = PairElem.First + Distances[TVector<uint32, 2>(PairElem.Second, Neighbor)];
 				if (NewDist < NeighborDistancePath.First)
 				{
-					auto NewPath = GeodesicPaths[CurrentStartEnd].Second;
-					check(NewPath.Num() > 0 && NewPath[NewPath.Num() - 1] != Neighbor)
-					    NewPath.Add(Neighbor);
+					TArray<uint32> NewPath = GeodesicPaths[CurrentStartEnd].Second;
+					check(NewPath.Num() > 0 && NewPath[NewPath.Num() - 1] != Neighbor);
+					NewPath.Add(Neighbor);
 					GeodesicPaths[NeighborStartEnd] = {NewDist, NewPath};
 					q.push(MakePair(GeodesicPaths[NeighborStartEnd].First, Neighbor));
 				}
@@ -264,27 +265,33 @@ void TPBDLongRangeConstraintsBase<T, d>::ComputeGeodesicConstraints(
 		}
 	});
 	FCriticalSection CriticalSection;
-	PhysicsParallelFor(InParticles.Size(), [&](int32 i) {
+	PhysicsParallelFor(UsedIndices.Num(), [&](uint32 UsedIndex) {
+		const uint32 i = UsedIndices[UsedIndex];
 		if (InParticles.InvM(i) == 0)
 			return;
 		TArray<Pair<T, int32>> ClosestElements;
-		for (auto Elements : IslandElements)
+		for (const TArray<uint32>& Elements : IslandElements)
 		{
-			int32 ClosestElement = -1;
-			for (auto Element : Elements)
+			if (!Elements.Num()) { continue; }  // Empty island 
+
+			int32 ClosestElement = Elements[0];
+			T ClosestDistance = GeodesicPaths[TVector<uint32, 2>(ClosestElement, i)].First;
+			if (ClosestDistance == FLT_MAX) { continue; }  // Not on this island
+
+			for (const uint32 Element : Elements)
 			{
-				if (ClosestElement < 0 || GeodesicPaths[TVector<uint32, 2>(ClosestElement, i)].First > GeodesicPaths[TVector<uint32, 2>(Element, i)].First)
+				const T Distance = GeodesicPaths[TVector<uint32, 2>(Element, i)].First;
+				if (Distance < ClosestDistance)
 				{
+					ClosestDistance = Distance;
 					ClosestElement = Element;
 				}
 			}
-			// Empty Island
-			if (ClosestElement < 0)
-				continue;
-			TVector<uint32, 2> Index(ClosestElement, i);
+
+			const TVector<uint32, 2> Index(ClosestElement, i);
 			check(GeodesicPaths[Index].First != FLT_MAX);
 			check(GeodesicPaths[Index].Second.Num() > 1);
-			ClosestElements.Add(MakePair(GeodesicPaths[Index].First, ClosestElement));
+			ClosestElements.Add(MakePair(ClosestDistance, ClosestElement));
 		}
 		// How to sort based on smalled first value of pair....
 		ClosestElements.Sort();
@@ -292,9 +299,9 @@ void TPBDLongRangeConstraintsBase<T, d>::ComputeGeodesicConstraints(
 		{
 			ClosestElements.SetNum(NumberOfAttachments);
 		}
-		for (auto Element : ClosestElements)
+		for (const Pair<T, int32>& Element : ClosestElements)
 		{
-			TVector<uint32, 2> Index(Element.Second, i);
+			const TVector<uint32, 2> Index(Element.Second, i);
 			check(GeodesicPaths[Index].First == Element.First);
 			check(FGenericPlatformMath::Abs(Element.First - ComputeGeodesicDistance(InParticles, GeodesicPaths[Index].Second)) < 1e-4);
 			CriticalSection.Lock();
