Chaos: Fix ComputePenetration overlaps not returning MTD for HeightFields and Trimesh.
Fix some queries setting bComputeMTD w/ uninitialized flag.
Fixes FindTeleportSpot failures, and vehicles exploding on dismount.

diff --git a/Engine/Source/Runtime/Experimental/Chaos/Private/Chaos/HeightField.cpp b/Engine/Source/Runtime/Experimental/Chaos/Private/Chaos/HeightField.cpp
--- a/Engine/Source/Runtime/Experimental/Chaos/Private/Chaos/HeightField.cpp
+++ b/Engine/Source/Runtime/Experimental/Chaos/Private/Chaos/HeightField.cpp
@@ -5,6 +5,7 @@
 #include "Chaos/GJK.h"
 #include "Chaos/ImplicitObjectScaled.h"
 #include "Chaos/Capsule.h"
+#include "Chaos/GeometryQueries.h"
 #include "HAL/IConsoleManager.h"
 #include "HAL/PlatformMisc.h"
 #include "Chaos/Triangle.h"
@@ -1328,8 +1329,11 @@ namespace Chaos
 
 	template <typename T>
 	template <typename QueryGeomType>
-	bool THeightField<T>::OverlapGeomImp(const QueryGeomType& QueryGeom, const TRigidTransform<T, 3>& QueryTM, const T Thickness) const
+	bool THeightField<T>::OverlapGeomImp(const QueryGeomType& QueryGeom, const TRigidTransform<T, 3>& QueryTM, const T Thickness, FVec3& OutNormal, FReal& OutPenetration, bool bComputeMTD) const
 	{
+		OutNormal = FVec3(0.0f);
+		OutPenetration = 0.0f;
+
 		auto OverlapTriangle = [&](const TVector<T, 3>& A, const TVector<T, 3>& B, const TVector<T, 3>& C) -> bool
 		{
 			const TVector<T, 3> AB = B - A;
@@ -1340,7 +1344,29 @@ namespace Chaos
 			const TVector<T, 3> Offset = TVector<T, 3>::CrossProduct(AB, AC);
 
 			TTriangle<T> TriangleConvex(A, B, C);
-			return GJKIntersection(TriangleConvex, QueryGeom, QueryTM, Thickness, Offset);
+			if (bComputeMTD)
+			{
+				FVec3 TriangleNormal(0.0f);
+				FReal Penetration = 0.0f;
+				FVec3 ClosestA(0.0f);
+				FVec3 ClosestB(0.0f);
+				if (GJKPenetration(TriangleConvex, QueryGeom, QueryTM, Penetration, ClosestA, ClosestB, TriangleNormal, Thickness))
+				{
+					// Use Deepest MTD.
+					if (Penetration > OutPenetration)
+					{
+						OutPenetration = Penetration;
+						OutNormal = TriangleNormal;
+					}
+					return true;
+				}
+
+				return false;
+			}
+			else
+			{
+				return GJKIntersection(TriangleConvex, QueryGeom, QueryTM, Thickness, Offset);
+			}
 		};
 
 		bool bResult = false;
@@ -1357,6 +1383,7 @@ namespace Chaos
 
 		GetGridIntersections(FlatQueryBounds, Intersections);
 
+		bool bOverlaps = false;
 		for(const TVector<int32, 2>& Cell : Intersections)
 		{
 			const int32 SingleIndex = Cell[1] * (GeomData.NumCols) + Cell[0];
@@ -1364,64 +1391,184 @@ namespace Chaos
 
 			if(OverlapTriangle(Points[0], Points[1], Points[3]))
 			{
-				return true;
+				bOverlaps = true;
+				if (!bComputeMTD)
+				{
+					return true;
+				}
 			}
 
 			if(OverlapTriangle(Points[0], Points[3], Points[2]))
 			{
-				return true;
+				bOverlaps = true;
+				if (!bComputeMTD)
+				{
+					return true;
+				}
 			}
 		}
 
-		return false;
+		return bOverlaps;
 	}
 
 	template <typename T>
-	bool THeightField<T>::OverlapGeom(const TSphere<T, 3>& QueryGeom, const TRigidTransform<T, 3>& QueryTM, const T Thickness) const
+	bool THeightField<T>::OverlapGeom(const TSphere<T, 3>& QueryGeom, const TRigidTransform<T, 3>& QueryTM, const T Thickness, FMTDInfo* OutMTD) const
 	{
-		return OverlapGeomImp(QueryGeom, QueryTM, Thickness);
+		FVec3 Normal;
+		FReal Penetration;
+		if (OutMTD)
+		{
+			// Computes MTD
+			if (OverlapGeomImp(QueryGeom, QueryTM, Thickness, Normal, Penetration, /*bComputeMTD=*/true))
+			{
+				OutMTD->Normal = Normal;
+				OutMTD->Penetration = Penetration;
+				return true;
+			}
+		}
+
+		// No MTD
+		return OverlapGeomImp(QueryGeom, QueryTM, Thickness, Normal, Penetration, /*bComputeMTD=*/false);
 	}
 
 	template <typename T>
-	bool THeightField<T>::OverlapGeom(const TBox<T, 3>& QueryGeom, const TRigidTransform<T, 3>& QueryTM, const T Thickness) const
+	bool THeightField<T>::OverlapGeom(const TBox<T, 3>& QueryGeom, const TRigidTransform<T, 3>& QueryTM, const T Thickness, FMTDInfo* OutMTD) const
 	{
-		return OverlapGeomImp(QueryGeom, QueryTM, Thickness);
+		FVec3 Normal;
+		FReal Penetration;
+		if (OutMTD)
+		{
+			// Computes MTD
+			if (OverlapGeomImp(QueryGeom, QueryTM, Thickness, Normal, Penetration, /*bComputeMTD=*/true))
+			{
+				OutMTD->Normal = Normal;
+				OutMTD->Penetration = Penetration;
+				return true;
+			}
+		}
+
+		// No MTD
+		return OverlapGeomImp(QueryGeom, QueryTM, Thickness, Normal, Penetration, /*bComputeMTD=*/false);
 	}
 
 	template <typename T>
-	bool THeightField<T>::OverlapGeom(const TCapsule<T>& QueryGeom, const TRigidTransform<T, 3>& QueryTM, const T Thickness) const
+	bool THeightField<T>::OverlapGeom(const TCapsule<T>& QueryGeom, const TRigidTransform<T, 3>& QueryTM, const T Thickness, FMTDInfo* OutMTD) const
 	{
-		return OverlapGeomImp(QueryGeom, QueryTM, Thickness);
+		FVec3 Normal;
+		FReal Penetration;
+		if (OutMTD)
+		{
+			// Computes MTD
+			if (OverlapGeomImp(QueryGeom, QueryTM, Thickness, Normal, Penetration, /*bComputeMTD=*/true))
+			{
+				OutMTD->Normal = Normal;
+				OutMTD->Penetration = Penetration;
+				return true;
+			}
+		}
+
+		// No MTD
+		return OverlapGeomImp(QueryGeom, QueryTM, Thickness, Normal, Penetration, /*bComputeMTD=*/false);
 	}
 
 	template <typename T>
-	bool THeightField<T>::OverlapGeom(const FConvex& QueryGeom, const TRigidTransform<T, 3>& QueryTM, const T Thickness) const
+	bool THeightField<T>::OverlapGeom(const FConvex& QueryGeom, const TRigidTransform<T, 3>& QueryTM, const T Thickness, FMTDInfo* OutMTD) const
 	{
-		return OverlapGeomImp(QueryGeom, QueryTM, Thickness);
+		FVec3 Normal;
+		FReal Penetration;
+		if (OutMTD)
+		{
+			// Computes MTD
+			if (OverlapGeomImp(QueryGeom, QueryTM, Thickness, Normal, Penetration, /*bComputeMTD=*/true))
+			{
+				OutMTD->Normal = Normal;
+				OutMTD->Penetration = Penetration;
+				return true;
+			}
+		}
+
+		// No MTD
+		return OverlapGeomImp(QueryGeom, QueryTM, Thickness, Normal, Penetration, /*bComputeMTD=*/false);
 	}
 
 	template <typename T>
-	bool THeightField<T>::OverlapGeom(const TImplicitObjectScaled<TSphere<T, 3>>& QueryGeom, const TRigidTransform<T, 3>& QueryTM, const T Thickness) const
+	bool THeightField<T>::OverlapGeom(const TImplicitObjectScaled<TSphere<T, 3>>& QueryGeom, const TRigidTransform<T, 3>& QueryTM, const T Thickness, FMTDInfo* OutMTD) const
 	{
-		return OverlapGeomImp(QueryGeom, QueryTM, Thickness);
+		FVec3 Normal;
+		FReal Penetration;
+		if (OutMTD)
+		{
+			// Computes MTD
+			if (OverlapGeomImp(QueryGeom, QueryTM, Thickness, Normal, Penetration, /*bComputeMTD=*/true))
+			{
+				OutMTD->Normal = Normal;
+				OutMTD->Penetration = Penetration;
+				return true;
+			}
+		}
+
+		// No MTD
+		return OverlapGeomImp(QueryGeom, QueryTM, Thickness, Normal, Penetration, /*bComputeMTD=*/false);
 	}
 
 	template <typename T>
-	bool THeightField<T>::OverlapGeom(const TImplicitObjectScaled<TBox<T, 3>>& QueryGeom, const TRigidTransform<T, 3>& QueryTM, const T Thickness) const
+	bool THeightField<T>::OverlapGeom(const TImplicitObjectScaled<TBox<T, 3>>& QueryGeom, const TRigidTransform<T, 3>& QueryTM, const T Thickness, FMTDInfo* OutMTD) const
 	{
-		return OverlapGeomImp(QueryGeom, QueryTM, Thickness);
+		FVec3 Normal;
+		FReal Penetration;
+		if (OutMTD)
+		{
+			// Computes MTD
+			if (OverlapGeomImp(QueryGeom, QueryTM, Thickness, Normal, Penetration, /*bComputeMTD=*/true))
+			{
+				OutMTD->Normal = Normal;
+				OutMTD->Penetration = Penetration;
+				return true;
+			}
+		}
+
+		// No MTD
+		return OverlapGeomImp(QueryGeom, QueryTM, Thickness, Normal, Penetration, /*bComputeMTD=*/false);
 	}
 
 	template <typename T>
-	bool THeightField<T>::OverlapGeom(const TImplicitObjectScaled<TCapsule<T>>& QueryGeom, const TRigidTransform<T, 3>& QueryTM, const T Thickness) const
+	bool THeightField<T>::OverlapGeom(const TImplicitObjectScaled<TCapsule<T>>& QueryGeom, const TRigidTransform<T, 3>& QueryTM, const T Thickness, FMTDInfo* OutMTD) const
 	{
-		return OverlapGeomImp(QueryGeom, QueryTM, Thickness);
+		FVec3 Normal;
+		FReal Penetration;
+		if (OutMTD)
+		{
+			// Computes MTD
+			if (OverlapGeomImp(QueryGeom, QueryTM, Thickness, Normal, Penetration, /*bComputeMTD=*/true))
+			{
+				OutMTD->Normal = Normal;
+				OutMTD->Penetration = Penetration;
+				return true;
+			}
+		}
+
+		// No MTD
+		return OverlapGeomImp(QueryGeom, QueryTM, Thickness, Normal, Penetration, /*bComputeMTD=*/false);
 	}
 
 	template <typename T>
-	bool THeightField<T>::OverlapGeom(const TImplicitObjectScaled<FConvex>& QueryGeom, const TRigidTransform<T, 3>& QueryTM, const T Thickness) const
+	bool THeightField<T>::OverlapGeom(const TImplicitObjectScaled<FConvex>& QueryGeom, const TRigidTransform<T, 3>& QueryTM, const T Thickness, FMTDInfo* OutMTD) const
 	{
-		return OverlapGeomImp(QueryGeom, QueryTM, Thickness);
+		FVec3 Normal;
+		FReal Penetration;
+		if (OutMTD)
+		{
+			// Computes MTD
+			if (OverlapGeomImp(QueryGeom, QueryTM, Thickness, Normal, Penetration, /*bComputeMTD=*/true))
+			{
+				OutMTD->Normal = Normal;
+				OutMTD->Penetration = Penetration;
+				return true;
+			}
+		}
+
+		// No MTD
+		return OverlapGeomImp(QueryGeom, QueryTM, Thickness, Normal, Penetration, /*bComputeMTD=*/false);
 	}
 
 	template <typename T>
diff --git a/Engine/Source/Runtime/Experimental/Chaos/Private/Chaos/TriangleMeshImplicitObject.cpp b/Engine/Source/Runtime/Experimental/Chaos/Private/Chaos/TriangleMeshImplicitObject.cpp
--- a/Engine/Source/Runtime/Experimental/Chaos/Private/Chaos/TriangleMeshImplicitObject.cpp
+++ b/Engine/Source/Runtime/Experimental/Chaos/Private/Chaos/TriangleMeshImplicitObject.cpp
@@ -5,6 +5,7 @@
 #include "Chaos/Triangle.h"
 #include "Chaos/Convex.h"
 #include "Chaos/ImplicitObjectScaled.h"
+#include "Chaos/GeometryQueries.h"
 
 namespace Chaos
 {
@@ -501,9 +502,30 @@ void TransformOverlapInputsHelper(const TImplicitObjectScaled<QueryGeomType>& Qu
 	OutScaledQueryTM = TRigidTransform<FReal, 3>(QueryTM.GetLocation() * InvScale, QueryTM.GetRotation());
 }
 
+template <typename QueryGeomType>
+void TransformOverlapOutputsHelper(const QueryGeomType& QueryGeom, const TVec3<FReal>& Normal, const FReal& Penetration, TVec3<FReal>& OutNormal, FReal& OutPenetration)
+{
+	OutNormal = Normal;
+	OutPenetration = Penetration;
+}
+
+template <typename QueryGeomType>
+void TransformOverlapOutputsHelper(const TImplicitObjectScaled<QueryGeomType>& QueryGeom, const TVec3<FReal>& Normal, const FReal& Penetration, TVec3<FReal>& OutNormal, FReal& OutPenetration)
+{
+	const TVec3<FReal> InvScale = QueryGeom.GetInvScale();
+
+	const TVec3<FReal> UnscaledNormalDenorm = (InvScale * Normal);
+	const FReal UnscaledNormalLength = UnscaledNormalDenorm.Size();
+	if (UnscaledNormalLength > 0.0f)
+	{
+		OutNormal = UnscaledNormalDenorm / UnscaledNormalLength;
+	}
+
+	OutPenetration = Penetration * UnscaledNormalLength;
+}
 
 template <typename QueryGeomType>
-bool FTriangleMeshImplicitObject::OverlapGeomImp(const QueryGeomType& QueryGeom, const FRigidTransform3& QueryTM, const FReal Thickness) const
+bool FTriangleMeshImplicitObject::OverlapGeomImp(const QueryGeomType& QueryGeom, const FRigidTransform3& QueryTM, const FReal Thickness, FVec3& OutNormal, FReal& OutPenetration, const bool bComputeMTD) const
 {
 	bool bResult = false;
 	TAABB<FReal, 3> QueryBounds = QueryGeom.BoundingBox();
@@ -511,6 +533,12 @@ bool FTriangleMeshImplicitObject::OverlapGeomImp(const QueryGeomType& QueryGeom,
 	QueryBounds = QueryBounds.TransformedAABB(QueryTM);
 	const TArray<int32> PotentialIntersections = BVH.FindAllIntersections(QueryBounds);
 
+	OutNormal = FVec3(0.0f);
+	OutPenetration = 0.0f;
+
+	FVec3 DeepestNormal = FVec3(0.0f);
+	FReal DeepestPenetration = 0.0f;
+
 	const auto& InnerQueryGeom = GetGeomHelper(QueryGeom);
 
 	TRigidTransform<FReal, 3> TransformedQueryTM;
@@ -518,25 +546,57 @@ bool FTriangleMeshImplicitObject::OverlapGeomImp(const QueryGeomType& QueryGeom,
 
 	auto LambdaHelper = [&](const auto& Elements)
 	{
-		for(int32 TriIdx : PotentialIntersections)
+		if (bComputeMTD)
 		{
-			TVec3<FReal> A,B,C;
-			TransformVertsHelper(QueryGeom,TriIdx,MParticles,Elements,A,B,C);
+			bool bOverlap = false;
+			for (int32 TriIdx : PotentialIntersections)
+			{
+				TVec3<FReal> A, B, C;
+				TransformVertsHelper(QueryGeom, TriIdx, MParticles, Elements, A, B, C);
+
+				FVec3 TriangleNormal(0.0f);
+				FReal Penetration = 0.0f;
+				FVec3 ClosestA(0.0f);
+				FVec3 ClosestB(0.0f);
+				if (GJKPenetration(TTriangle<FReal>(A, B, C), InnerQueryGeom, TransformedQueryTM, Penetration, ClosestA, ClosestB, TriangleNormal, Thickness))
+				{
+					bOverlap = true;
 
-			const FVec3 AB = B - A;
-			const FVec3 AC = C - A;
+					// Use Deepest MTD.
+					if (Penetration > DeepestPenetration)
+					{
+						DeepestPenetration = Penetration;
+						DeepestNormal = TriangleNormal;
+					}
+				}
+			}
 
-			//It's most likely that the query object is in front of the triangle since queries tend to be on the outside.
-			//However, maybe we should check if it's behind the triangle plane. Also, we should enforce this winding in some way
-			const FVec3 Offset = FVec3::CrossProduct(AB,AC);
+			TransformOverlapOutputsHelper(QueryGeom, DeepestNormal, DeepestPenetration, OutNormal, OutPenetration);
 
-			if(GJKIntersection(TTriangle<FReal>(A,B,C),InnerQueryGeom,TransformedQueryTM,Thickness,Offset))
+			return bOverlap;
+		}
+		else
+		{
+			for (int32 TriIdx : PotentialIntersections)
 			{
-				return true;
+				TVec3<FReal> A, B, C;
+				TransformVertsHelper(QueryGeom, TriIdx, MParticles, Elements, A, B, C);
+
+				const FVec3 AB = B - A;
+				const FVec3 AC = C - A;
+
+				//It's most likely that the query object is in front of the triangle since queries tend to be on the outside.
+				//However, maybe we should check if it's behind the triangle plane. Also, we should enforce this winding in some way
+				const FVec3 Offset = FVec3::CrossProduct(AB, AC);
+
+				if (GJKIntersection(TTriangle<FReal>(A, B, C), InnerQueryGeom, TransformedQueryTM, Thickness, Offset))
+				{
+					return true;
+				}
 			}
-		}
 
-		return false;
+			return false;
+		}
 	};
 
 	if(MElements.RequiresLargeIndices())
@@ -549,49 +609,175 @@ bool FTriangleMeshImplicitObject::OverlapGeomImp(const QueryGeomType& QueryGeom,
 	}
 }
 
-bool FTriangleMeshImplicitObject::OverlapGeom(const TSphere<FReal, 3>& QueryGeom, const FRigidTransform3& QueryTM, const FReal Thickness) const
+bool FTriangleMeshImplicitObject::OverlapGeom(const TSphere<FReal, 3>& QueryGeom, const FRigidTransform3& QueryTM, const FReal Thickness, FMTDInfo* OutMTD) const
 {
-	return OverlapGeomImp(QueryGeom, QueryTM, Thickness);
+	FVec3 Normal;
+	FReal Penetration;
+	if (OutMTD)
+	{
+		// Computes MTD
+		if (OverlapGeomImp(QueryGeom, QueryTM, Thickness, Normal, Penetration, /*bComputeMTD=*/true))
+		{
+			OutMTD->Normal = Normal;
+			OutMTD->Penetration = Penetration;
+			return true;
+		}
+	}
+
+	// No MTD
+	return OverlapGeomImp(QueryGeom, QueryTM, Thickness, Normal, Penetration, /*bComputeMTD=*/false);
 }
 
-bool FTriangleMeshImplicitObject::OverlapGeom(const TBox<FReal, 3>& QueryGeom, const FRigidTransform3& QueryTM, const FReal Thickness) const
+bool FTriangleMeshImplicitObject::OverlapGeom(const TBox<FReal, 3>& QueryGeom, const FRigidTransform3& QueryTM, const FReal Thickness, FMTDInfo* OutMTD) const
 {
-	return OverlapGeomImp(QueryGeom, QueryTM, Thickness);
+	FVec3 Normal;
+	FReal Penetration;
+	if (OutMTD)
+	{
+		// Computes MTD
+		if (OverlapGeomImp(QueryGeom, QueryTM, Thickness, Normal, Penetration, /*bComputeMTD=*/true))
+		{
+			OutMTD->Normal = Normal;
+			OutMTD->Penetration = Penetration;
+			return true;
+		}
+	}
+
+	// No MTD
+	return OverlapGeomImp(QueryGeom, QueryTM, Thickness, Normal, Penetration, /*bComputeMTD=*/false);
 }
 
-bool FTriangleMeshImplicitObject::OverlapGeom(const TCapsule<FReal>& QueryGeom, const FRigidTransform3& QueryTM, const FReal Thickness) const
+bool FTriangleMeshImplicitObject::OverlapGeom(const TCapsule<FReal>& QueryGeom, const FRigidTransform3& QueryTM, const FReal Thickness, FMTDInfo* OutMTD) const
 {
-	return OverlapGeomImp(QueryGeom, QueryTM, Thickness);
+	FVec3 Normal;
+	FReal Penetration;
+	if (OutMTD)
+	{
+		// Computes MTD
+		if (OverlapGeomImp(QueryGeom, QueryTM, Thickness, Normal, Penetration, /*bComputeMTD=*/true))
+		{
+			OutMTD->Normal = Normal;
+			OutMTD->Penetration = Penetration;
+			return true;
+		}
+	}
+
+	// No MTD
+	return OverlapGeomImp(QueryGeom, QueryTM, Thickness, Normal, Penetration, /*bComputeMTD=*/false);
 }
 
-bool FTriangleMeshImplicitObject::OverlapGeom(const FConvex& QueryGeom, const FRigidTransform3& QueryTM, const FReal Thickness) const
+bool FTriangleMeshImplicitObject::OverlapGeom(const FConvex& QueryGeom, const FRigidTransform3& QueryTM, const FReal Thickness, FMTDInfo* OutMTD) const
 {
-	return OverlapGeomImp(QueryGeom, QueryTM, Thickness);
+	FVec3 Normal;
+	FReal Penetration;
+	if (OutMTD)
+	{
+		// Computes MTD
+		if (OverlapGeomImp(QueryGeom, QueryTM, Thickness, Normal, Penetration, /*bComputeMTD=*/true))
+		{
+			OutMTD->Normal = Normal;
+			OutMTD->Penetration = Penetration;
+			return true;
+		}
+	}
+
+	// No MTD
+	return OverlapGeomImp(QueryGeom, QueryTM, Thickness, Normal, Penetration, /*bComputeMTD=*/false);
 }
 
-bool FTriangleMeshImplicitObject::OverlapGeom(const TImplicitObjectScaled<TSphere<FReal, 3>>& QueryGeom, const FRigidTransform3& QueryTM, const FReal Thickness) const
+bool FTriangleMeshImplicitObject::OverlapGeom(const TImplicitObjectScaled<TSphere<FReal, 3>>& QueryGeom, const FRigidTransform3& QueryTM, const FReal Thickness, FMTDInfo* OutMTD) const
 {
-	return OverlapGeomImp(QueryGeom, QueryTM, Thickness);
+	FVec3 Normal;
+	FReal Penetration;
+	if (OutMTD)
+	{
+		// Computes MTD
+		if (OverlapGeomImp(QueryGeom, QueryTM, Thickness, Normal, Penetration, /*bComputeMTD=*/true))
+		{
+			OutMTD->Normal = Normal;
+			OutMTD->Penetration = Penetration;
+			return true;
+		}
+	}
+
+	// No MTD
+	return OverlapGeomImp(QueryGeom, QueryTM, Thickness, Normal, Penetration, /*bComputeMTD=*/false);
 }
 
-bool FTriangleMeshImplicitObject::OverlapGeom(const TImplicitObjectScaled<TBox<FReal, 3>>& QueryGeom, const FRigidTransform3& QueryTM, const FReal Thickness) const
+bool FTriangleMeshImplicitObject::OverlapGeom(const TImplicitObjectScaled<TBox<FReal, 3>>& QueryGeom, const FRigidTransform3& QueryTM, const FReal Thickness, FMTDInfo* OutMTD) const
 {
-	return OverlapGeomImp(QueryGeom, QueryTM, Thickness);
+	FVec3 Normal;
+	FReal Penetration;
+	if (OutMTD)
+	{
+		// Computes MTD
+		if (OverlapGeomImp(QueryGeom, QueryTM, Thickness, Normal, Penetration, /*bComputeMTD=*/true))
+		{
+			OutMTD->Normal = Normal;
+			OutMTD->Penetration = Penetration;
+			return true;
+		}
+	}
+
+	// No MTD
+	return OverlapGeomImp(QueryGeom, QueryTM, Thickness, Normal, Penetration, /*bComputeMTD=*/false);
 }
 
-bool FTriangleMeshImplicitObject::OverlapGeom(const TImplicitObjectScaled<TCapsule<FReal>>& QueryGeom, const FRigidTransform3& QueryTM, const FReal Thickness) const
+bool FTriangleMeshImplicitObject::OverlapGeom(const TImplicitObjectScaled<TCapsule<FReal>>& QueryGeom, const FRigidTransform3& QueryTM, const FReal Thickness, FMTDInfo* OutMTD) const
 {
-	return OverlapGeomImp(QueryGeom, QueryTM, Thickness);
+	FVec3 Normal;
+	FReal Penetration;
+	if (OutMTD)
+	{
+		// Computes MTD
+		if (OverlapGeomImp(QueryGeom, QueryTM, Thickness, Normal, Penetration, /*bComputeMTD=*/true))
+		{
+			OutMTD->Normal = Normal;
+			OutMTD->Penetration = Penetration;
+			return true;
+		}
+	}
+
+	// No MTD
+	return OverlapGeomImp(QueryGeom, QueryTM, Thickness, Normal, Penetration, /*bComputeMTD=*/false);
 }
 
-bool FTriangleMeshImplicitObject::OverlapGeom(const TImplicitObjectScaled<FConvex>& QueryGeom, const FRigidTransform3& QueryTM, const FReal Thickness) const
+bool FTriangleMeshImplicitObject::OverlapGeom(const TImplicitObjectScaled<FConvex>& QueryGeom, const FRigidTransform3& QueryTM, const FReal Thickness, FMTDInfo* OutMTD) const
 {
-	return OverlapGeomImp(QueryGeom, QueryTM, Thickness);
+	FVec3 Normal;
+	FReal Penetration;
+	if (OutMTD)
+	{
+		// Computes MTD
+		if (OverlapGeomImp(QueryGeom, QueryTM, Thickness, Normal, Penetration, /*bComputeMTD=*/true))
+		{
+			OutMTD->Normal = Normal;
+			OutMTD->Penetration = Penetration;
+			return true;
+		}
+	}
+
+	// No MTD
+	return OverlapGeomImp(QueryGeom, QueryTM, Thickness, Normal, Penetration, /*bComputeMTD=*/false);
 }
 
-bool FTriangleMeshImplicitObject::OverlapGeom(const TImplicitObjectScaled<TImplicitObjectScaled<FConvex>>& QueryGeom, const FRigidTransform3& QueryTM, const FReal Thickness) const
+bool FTriangleMeshImplicitObject::OverlapGeom(const TImplicitObjectScaled<TImplicitObjectScaled<FConvex>>& QueryGeom, const FRigidTransform3& QueryTM, const FReal Thickness, FMTDInfo* OutMTD) const
 {
-	return OverlapGeomImp(QueryGeom, QueryTM, Thickness);
+	FVec3 Normal;
+	FReal Penetration;
+	if (OutMTD)
+	{
+		// Computes MTD
+		if (OverlapGeomImp(QueryGeom, QueryTM, Thickness, Normal, Penetration, /*bComputeMTD=*/true))
+		{
+			OutMTD->Normal = Normal;
+			OutMTD->Penetration = Penetration;
+			return true;
+		}
+	}
+
+	// No MTD
+	return OverlapGeomImp(QueryGeom, QueryTM, Thickness, Normal, Penetration, /*bComputeMTD=*/false);
 }
 
 
diff --git a/Engine/Source/Runtime/Experimental/Chaos/Public/Chaos/GeometryQueries.h b/Engine/Source/Runtime/Experimental/Chaos/Public/Chaos/GeometryQueries.h
--- a/Engine/Source/Runtime/Experimental/Chaos/Public/Chaos/GeometryQueries.h
+++ b/Engine/Source/Runtime/Experimental/Chaos/Public/Chaos/GeometryQueries.h
@@ -118,23 +118,23 @@ namespace Chaos
 			case ImplicitObjectType::HeightField:
 			{
 				const THeightField<FReal>& AHeightField = static_cast<const THeightField<FReal>&>(A);
-				return AHeightField.OverlapGeom(B, BToATM, Thickness);
+				return AHeightField.OverlapGeom(B, BToATM, Thickness, OutMTD);
 			}
 			case ImplicitObjectType::TriangleMesh:
 			{
 				const FTriangleMeshImplicitObject& ATriangleMesh = static_cast<const FTriangleMeshImplicitObject&>(A);
-				return ATriangleMesh.OverlapGeom(B, BToATM, Thickness);
+				return ATriangleMesh.OverlapGeom(B, BToATM, Thickness, OutMTD);
 			}
 			default:
 				if (IsScaled(AType))
 				{
 					const auto& AScaled = TImplicitObjectScaled<FTriangleMeshImplicitObject>::AsScaledChecked(A);
-					return AScaled.LowLevelOverlapGeom(B, BToATM, Thickness);
+					return AScaled.LowLevelOverlapGeom(B, BToATM, Thickness, OutMTD);
 				}
 				else if(IsInstanced(AType))
 				{
 					const auto& AInstanced = TImplicitObjectInstanced<FTriangleMeshImplicitObject>::AsInstancedChecked(A);
-					return AInstanced.LowLevelOverlapGeom(B,BToATM,Thickness);
+					return AInstanced.LowLevelOverlapGeom(B,BToATM,Thickness, OutMTD);
 				}
 				else
 				{
diff --git a/Engine/Source/Runtime/Experimental/Chaos/Public/Chaos/HeightField.h b/Engine/Source/Runtime/Experimental/Chaos/Public/Chaos/HeightField.h
--- a/Engine/Source/Runtime/Experimental/Chaos/Public/Chaos/HeightField.h
+++ b/Engine/Source/Runtime/Experimental/Chaos/Public/Chaos/HeightField.h
@@ -18,6 +18,7 @@ namespace Chaos
 	template<typename T>
 	class THeightfieldRaycastVisitor;
 	class FConvex;
+	struct FMTDInfo;
 }
 
 namespace Chaos
@@ -53,14 +54,14 @@ namespace Chaos
 		virtual bool Raycast(const TVector<T, 3>& StartPoint, const TVector<T, 3>& Dir, const T Length, const T Thickness, T& OutTime, TVector<T,3>& OutPosition, TVector<T,3>& OutNormal, int32& OutFaceIndex) const override;
 		virtual bool Overlap(const TVector<T, 3>& Point, const T Thickness) const override;
 		
-		bool OverlapGeom(const TSphere<T, 3>& QueryGeom, const TRigidTransform<T, 3>& QueryTM, const T Thickness) const;
-		bool OverlapGeom(const TBox<T, 3>& QueryGeom, const TRigidTransform<T, 3>& QueryTM, const T Thickness) const;
-		bool OverlapGeom(const TCapsule<T>& QueryGeom, const TRigidTransform<T, 3>& QueryTM, const T Thickness) const;
-		bool OverlapGeom(const FConvex& QueryGeom, const TRigidTransform<T, 3>& QueryTM, const T Thickness) const;
-		bool OverlapGeom(const TImplicitObjectScaled<TSphere<T, 3>>& QueryGeom, const TRigidTransform<T, 3>& QueryTM, const T Thickness) const;
-		bool OverlapGeom(const TImplicitObjectScaled<TBox<T, 3>>& QueryGeom, const TRigidTransform<T, 3>& QueryTM, const T Thickness) const;
-		bool OverlapGeom(const TImplicitObjectScaled<TCapsule<T>>& QueryGeom, const TRigidTransform<T, 3>& QueryTM, const T Thickness) const;
-		bool OverlapGeom(const TImplicitObjectScaled<FConvex>& QueryGeom, const TRigidTransform<T, 3>& QueryTM, const T Thickness) const;
+		bool OverlapGeom(const TSphere<T, 3>& QueryGeom, const TRigidTransform<T, 3>& QueryTM, const T Thickness, FMTDInfo* OutMTD = nullptr) const;
+		bool OverlapGeom(const TBox<T, 3>& QueryGeom, const TRigidTransform<T, 3>& QueryTM, const T Thickness, FMTDInfo* OutMTD = nullptr) const;
+		bool OverlapGeom(const TCapsule<T>& QueryGeom, const TRigidTransform<T, 3>& QueryTM, const T Thickness, FMTDInfo* OutMTD = nullptr) const;
+		bool OverlapGeom(const FConvex& QueryGeom, const TRigidTransform<T, 3>& QueryTM, const T Thickness, FMTDInfo* OutMTD = nullptr) const;
+		bool OverlapGeom(const TImplicitObjectScaled<TSphere<T, 3>>& QueryGeom, const TRigidTransform<T, 3>& QueryTM, const T Thickness, FMTDInfo* OutMTD = nullptr) const;
+		bool OverlapGeom(const TImplicitObjectScaled<TBox<T, 3>>& QueryGeom, const TRigidTransform<T, 3>& QueryTM, const T Thickness, FMTDInfo* OutMTD = nullptr) const;
+		bool OverlapGeom(const TImplicitObjectScaled<TCapsule<T>>& QueryGeom, const TRigidTransform<T, 3>& QueryTM, const T Thickness, FMTDInfo* OutMTD = nullptr) const;
+		bool OverlapGeom(const TImplicitObjectScaled<FConvex>& QueryGeom, const TRigidTransform<T, 3>& QueryTM, const T Thickness, FMTDInfo* OutMTD = nullptr) const;
 
 		bool SweepGeom(const TSphere<T, 3>& QueryGeom, const TRigidTransform<T, 3>& StartTM, const TVector<T, 3>& Dir, const T Length, T& OutTime, TVector<T, 3>& OutPosition, TVector<T, 3>& OutNormal, int32& OutFaceIndex, const T Thickness = 0, bool bComputeMTD = false) const;
 		bool SweepGeom(const TBox<T, 3>& QueryGeom, const TRigidTransform<T, 3>& StartTM, const TVector<T, 3>& Dir, const T Length, T& OutTime, TVector<T, 3>& OutPosition, TVector<T, 3>& OutNormal, int32& OutFaceIndex, const T Thickness = 0, bool bComputeMTD = false) const;
@@ -542,7 +543,7 @@ namespace Chaos
 		friend FImplicitObject;
 
 		template <typename QueryGeomType>
-		bool OverlapGeomImp(const QueryGeomType& QueryGeom, const TRigidTransform<T, 3>& QueryTM, const T Thickness) const;
+		bool OverlapGeomImp(const QueryGeomType& QueryGeom, const TRigidTransform<T, 3>& QueryTM, const T Thickness, FVec3& OutNormal, FReal& OutPenetration, const bool bComputeMTD) const;
 
 		template <typename QueryGeomType>
 		bool SweepGeomImp(const QueryGeomType& QueryGeom, const TRigidTransform<T, 3>& StartTM, const TVector<T, 3>& Dir, const T Length, T& OutTime, TVector<T, 3>& OutPosition, TVector<T, 3>& OutNormal, int32& OutFaceIndex, const T Thickness, bool bComputeMTD) const;
diff --git a/Engine/Source/Runtime/Experimental/Chaos/Public/Chaos/ImplicitObjectScaled.h b/Engine/Source/Runtime/Experimental/Chaos/Public/Chaos/ImplicitObjectScaled.h
--- a/Engine/Source/Runtime/Experimental/Chaos/Public/Chaos/ImplicitObjectScaled.h
+++ b/Engine/Source/Runtime/Experimental/Chaos/Public/Chaos/ImplicitObjectScaled.h
@@ -13,6 +13,8 @@
 namespace Chaos
 {
 
+struct FMTDInfo;
+
 template <typename TConcrete>
 class TImplicitObjectInstanced final : public FImplicitObject
 {
@@ -131,9 +133,9 @@ public:
 
 	/** This is a low level function and assumes the internal object has a OverlapGeom function. Should not be called directly. See GeometryQueries.h : OverlapQuery */
 	template <typename QueryGeomType>
-	bool LowLevelOverlapGeom(const QueryGeomType& B,const TRigidTransform<T,d>& BToATM,T Thickness = 0) const
+	bool LowLevelOverlapGeom(const QueryGeomType& B,const TRigidTransform<T,d>& BToATM,T Thickness = 0, FMTDInfo* OutMTD = nullptr) const
 	{
-		return MObject->OverlapGeom(B,BToATM,Thickness);
+		return MObject->OverlapGeom(B,BToATM,Thickness, OutMTD);
 	}
 
 protected:
@@ -384,13 +386,13 @@ public:
 
 	/** This is a low level function and assumes the internal object has a OverlapGeom function. Should not be called directly. See GeometryQueries.h : OverlapQuery */
 	template <typename QueryGeomType>
-	bool LowLevelOverlapGeom(const QueryGeomType& B, const TRigidTransform<T, d>& BToATM, T Thickness = 0) const
+	bool LowLevelOverlapGeom(const QueryGeomType& B, const TRigidTransform<T, d>& BToATM, T Thickness = 0, FMTDInfo* OutMTD = nullptr) const
 	{
 		ensure(Thickness == 0 || (FMath::IsNearlyEqual(MScale[0], MScale[1]) && FMath::IsNearlyEqual(MScale[0], MScale[2])));
 
 		auto ScaledB = MakeScaledHelper(B, MInvScale);
 		TRigidTransform<T, d> BToATMNoScale(BToATM.GetLocation() * MInvScale, BToATM.GetRotation());
-		return MObject->OverlapGeom(ScaledB, BToATMNoScale, MInternalThickness + Thickness);
+		return MObject->OverlapGeom(ScaledB, BToATMNoScale, MInternalThickness + Thickness, OutMTD);
 	}
 
 	virtual int32 FindMostOpposingFace(const TVector<T, d>& Position, const TVector<T, d>& UnitDir, int32 HintFaceIndex, T SearchDist) const override
diff --git a/Engine/Source/Runtime/Experimental/Chaos/Public/Chaos/TriangleMeshImplicitObject.h b/Engine/Source/Runtime/Experimental/Chaos/Public/Chaos/TriangleMeshImplicitObject.h
--- a/Engine/Source/Runtime/Experimental/Chaos/Public/Chaos/TriangleMeshImplicitObject.h
+++ b/Engine/Source/Runtime/Experimental/Chaos/Public/Chaos/TriangleMeshImplicitObject.h
@@ -20,6 +20,7 @@ namespace Chaos
 	class TCapsule;
 
 	class FConvex;
+	struct FMTDInfo;
 
 	class CHAOS_API FTrimeshIndexBuffer
 	{
@@ -126,15 +127,15 @@ namespace Chaos
 		virtual bool Raycast(const FVec3& StartPoint, const FVec3& Dir, const FReal Length, const FReal Thickness, FReal& OutTime, FVec3& OutPosition, FVec3& OutNormal, int32& OutFaceIndex) const override;
 		virtual bool Overlap(const FVec3& Point, const FReal Thickness) const override;
 
-		bool OverlapGeom(const TSphere<FReal, 3>& QueryGeom, const FRigidTransform3& QueryTM, const FReal Thickness) const;
-		bool OverlapGeom(const TBox<FReal, 3>& QueryGeom, const FRigidTransform3& QueryTM, const FReal Thickness) const;
-		bool OverlapGeom(const TCapsule<FReal>& QueryGeom, const FRigidTransform3& QueryTM, const FReal Thickness) const;
-		bool OverlapGeom(const FConvex& QueryGeom, const FRigidTransform3& QueryTM, const FReal Thickness) const;
-		bool OverlapGeom(const TImplicitObjectScaled<TSphere<FReal, 3>>& QueryGeom, const FRigidTransform3& QueryTM, const FReal Thickness) const;
-		bool OverlapGeom(const TImplicitObjectScaled<TBox<FReal, 3>>& QueryGeom, const FRigidTransform3& QueryTM, const FReal Thickness) const;
-		bool OverlapGeom(const TImplicitObjectScaled<TCapsule<FReal>>& QueryGeom, const FRigidTransform3& QueryTM, const FReal Thickness) const;
-		bool OverlapGeom(const TImplicitObjectScaled<FConvex>& QueryGeom, const FRigidTransform3& QueryTM, const FReal Thickness) const;
-		bool OverlapGeom(const TImplicitObjectScaled<TImplicitObjectScaled<FConvex>>& QueryGeom, const FRigidTransform3& QueryTM, const FReal Thickness) const;
+		bool OverlapGeom(const TSphere<FReal, 3>& QueryGeom, const FRigidTransform3& QueryTM, const FReal Thickness, FMTDInfo* OutMTD = nullptr) const;
+		bool OverlapGeom(const TBox<FReal, 3>& QueryGeom, const FRigidTransform3& QueryTM, const FReal Thickness, FMTDInfo* OutMTD = nullptr) const;
+		bool OverlapGeom(const TCapsule<FReal>& QueryGeom, const FRigidTransform3& QueryTM, const FReal Thickness, FMTDInfo* OutMTD = nullptr) const;
+		bool OverlapGeom(const FConvex& QueryGeom, const FRigidTransform3& QueryTM, const FReal Thickness, FMTDInfo* OutMTD = nullptr) const;
+		bool OverlapGeom(const TImplicitObjectScaled<TSphere<FReal, 3>>& QueryGeom, const FRigidTransform3& QueryTM, const FReal Thickness, FMTDInfo* OutMTD = nullptr) const;
+		bool OverlapGeom(const TImplicitObjectScaled<TBox<FReal, 3>>& QueryGeom, const FRigidTransform3& QueryTM, const FReal Thickness, FMTDInfo* OutMTD = nullptr) const;
+		bool OverlapGeom(const TImplicitObjectScaled<TCapsule<FReal>>& QueryGeom, const FRigidTransform3& QueryTM, const FReal Thickness, FMTDInfo* OutMTD = nullptr) const;
+		bool OverlapGeom(const TImplicitObjectScaled<FConvex>& QueryGeom, const FRigidTransform3& QueryTM, const FReal Thickness, FMTDInfo* OutMTD = nullptr) const;
+		bool OverlapGeom(const TImplicitObjectScaled<TImplicitObjectScaled<FConvex>>& QueryGeom, const FRigidTransform3& QueryTM, const FReal Thickness, FMTDInfo* OutMTD = nullptr) const;
 
 		bool SweepGeom(const TSphere<FReal, 3>& QueryGeom, const FRigidTransform3& StartTM, const FVec3& Dir, const FReal Length, FReal& OutTime, FVec3& OutPosition, FVec3& OutNormal, int32& OutFaceIndex, const FReal Thickness = 0, const bool bComputeMTD = false) const;
 		bool SweepGeom(const TBox<FReal, 3>& QueryGeom, const FRigidTransform3& StartTM, const FVec3& Dir, const FReal Length, FReal& OutTime, FVec3& OutPosition, FVec3& OutNormal, int32& OutFaceIndex, const FReal Thickness = 0, const bool bComputeMTD = false) const;
@@ -325,7 +326,7 @@ namespace Chaos
 		bool GJKContactPointImp(const QueryGeomType& QueryGeom, const FRigidTransform3& QueryTM, const FReal Thickness, FVec3& Location, FVec3& Normal, FReal& Penetration) const;
 
 		template<typename QueryGeomType>
-		bool OverlapGeomImp(const QueryGeomType& QueryGeom, const FRigidTransform3& QueryTM, const FReal Thickness) const;
+		bool OverlapGeomImp(const QueryGeomType& QueryGeom, const FRigidTransform3& QueryTM, const FReal Thickness, FVec3& OutNormal, FReal& OutPenetration, const bool bComputeMTD) const;
 
 		template<typename QueryGeomType>
 		bool SweepGeomImp(const QueryGeomType& QueryGeom, const FRigidTransform3& StartTM, const FVec3& Dir, const FReal Length, FReal& OutTime, FVec3& OutPosition, FVec3& OutNormal, int32& OutFaceIndex, const FReal Thickness, const bool bComputeMTD) const;
diff --git a/Engine/Source/Runtime/PhysicsSQ/Private/SQAccelerator.cpp b/Engine/Source/Runtime/PhysicsSQ/Private/SQAccelerator.cpp
--- a/Engine/Source/Runtime/PhysicsSQ/Private/SQAccelerator.cpp
+++ b/Engine/Source/Runtime/PhysicsSQ/Private/SQAccelerator.cpp
@@ -240,7 +240,7 @@ private:
 				TVector<float, 3> WorldPosition, WorldNormal;
 				float Distance = 0;	//not needed but fixes compiler warning for overlap
 				int32 FaceIdx = INDEX_NONE;	//not needed but fixes compiler warning for overlap
-				const bool bComputeMTD = !!((uint16)(OutputFlags & EHitFlags::MTD));
+				const bool bComputeMTD = !!((uint16)(OutputFlags.HitFlags & EHitFlags::MTD));
 
 				if (SQ == ESQType::Raycast)
 				{
@@ -262,7 +262,19 @@ private:
 				}
 				else if (SQ == ESQType::Overlap || (SQ == ESQType::Sweep && CurData->CurrentLength == 0))
 				{
-					bHit = OverlapQuery(*Geom, ActorTM, *QueryGeom, StartTM, /*Thickness=*/0);
+					if (bComputeMTD)
+					{
+						FMTDInfo MTDInfo;
+						bHit = OverlapQuery(*Geom, ActorTM, *QueryGeom, StartTM, /*Thickness=*/0, &MTDInfo);
+						if (bHit)
+						{
+							WorldNormal = MTDInfo.Normal * MTDInfo.Penetration;
+						}
+					}
+					else
+					{
+						bHit = OverlapQuery(*Geom, ActorTM, *QueryGeom, StartTM, /*Thickness=*/0);
+					}
 				}
 
 				if(bHit)
@@ -322,7 +334,7 @@ private:
 	const FVector StartPoint;
 	const FVector Dir;
 	const FVector HalfExtents;
-	EHitFlags OutputFlags;
+	FHitFlags OutputFlags;
 	bool bAnyHit;
 	const FQueryDebugParams DebugParams;
 	ChaosInterface::FSQHitBuffer<THitType>& HitBuffer;
