Chaos : Collision Filtering

- Added support for copying filter data into the solver. Filtering collisions using SimData is still WIP.

diff --git a/Engine/Source/Runtime/Experimental/Chaos/Private/Chaos/CollisionResolution.cpp b/Engine/Source/Runtime/Experimental/Chaos/Private/Chaos/CollisionResolution.cpp
--- a/Engine/Source/Runtime/Experimental/Chaos/Private/Chaos/CollisionResolution.cpp
+++ b/Engine/Source/Runtime/Experimental/Chaos/Private/Chaos/CollisionResolution.cpp
@@ -1896,50 +1896,105 @@ namespace Chaos
 			}
 		}
 
-		template<typename T, int d>
-		void ConstructConstraints(TGeometryParticleHandle<T, d>* Particle0, TGeometryParticleHandle<T, d>* Particle1, const FImplicitObject* Implicit0, const FImplicitObject* Implicit1, const TRigidTransform<T, d>& Transform0, const TRigidTransform<T, d>& Transform1, const T CullDistance, FCollisionConstraintsArray& NewConstraints)
+		typedef uint8 FMaskFilter;
+		enum { NumExtraFilterBits = 6 };
+		enum { NumCollisionChannelBits = 5 };
+
+		inline bool IsValid(const FCollisionFilterData& Filter)
 		{
-			EImplicitObjectType Implicit0Type = Implicit0 ? GetInnerType(Implicit0->GetType()) : ImplicitObjectType::Unknown;
-			EImplicitObjectType Implicit1Type = Implicit1 ? GetInnerType(Implicit1->GetType()) : ImplicitObjectType::Unknown;
+			return Filter.Word0 || Filter.Word1 || Filter.Word2 || Filter.Word3;
+		}
+
+		inline uint32 GetChaosCollisionChannel(uint32 Word3)
+		{
+			uint32 ChannelMask = (Word3 << NumExtraFilterBits) >> (32 - NumCollisionChannelBits);
+			return (uint32)ChannelMask;
+		}
+
+		inline uint32 GetChaosCollisionChannelAndExtraFilter(uint32 Word3, FMaskFilter& OutMaskFilter)
+		{
+			uint32 ChannelMask = GetChaosCollisionChannel(Word3);
+			OutMaskFilter = Word3 >> (32 - NumExtraFilterBits);
+			return (uint32)ChannelMask;
+		}
 
-			// If either shape is disabled for collision bail without constructing a constraint
-			if (const TPerShapeData<T, d>* Shape0 = Particle0->GetImplicitShape(Implicit0))
+		template<class T, int d>
+		bool DoCollide(EImplicitObjectType Implicit0Type, const TPerShapeData<T, d>* Shape0, EImplicitObjectType Implicit1Type, const TPerShapeData<T, d>* Shape1)
+		{
+			/*
+			if (Shape0 && Shape1)
+			{
+				if (!IsValid(Shape0->SimData) && !IsValid(Shape1->SimData))
+				{
+					return true;
+				}
+
+				FMaskFilter Filter0Mask, Filter1Mask;
+				const uint32 Filter0Channel = GetChaosCollisionChannelAndExtraFilter(Shape0->SimData.Word3, Filter0Mask);
+				const uint32 Filter1Channel = GetChaosCollisionChannelAndExtraFilter(Shape1->SimData.Word3, Filter1Mask);
+
+				if ((Filter0Mask & Filter1Mask) != 0)
+				{
+					return false;
+				}
+
+				const uint32 Filter1Bit = 1 << (Filter1Channel); // SIMDATA_TO_BITFIELD
+				uint32 const Filter0Bit = 1 << (Filter0Channel); // SIMDATA_TO_BITFIELD
+				return (Filter0Bit & Shape1->SimData.Word1) && (Filter1Bit & Shape0->SimData.Word1);
+			}
+			*/
+
+			if (Shape0)
 			{
 				if (Shape0->bDisable)
 				{
-					return;
+					return false;
 				}
 				if (Implicit0Type == ImplicitObjectType::TriangleMesh && Shape0->CollisionTraceType != Chaos_CTF_UseComplexAsSimple)
 				{
-					return;
+					return false;
 				}
 				else if (Shape0->CollisionTraceType == Chaos_CTF_UseComplexAsSimple && Implicit0Type != ImplicitObjectType::TriangleMesh)
 				{
-					return;
+					return false;
 				}
-			} 
-			else if (Implicit0Type == ImplicitObjectType::TriangleMesh) 
+			}
+			else if (Implicit0Type == ImplicitObjectType::TriangleMesh)
 			{
-				return;
+				return false;
 			}
 
-			if (const TPerShapeData<T, d>* Shape1 = Particle1->GetImplicitShape(Implicit1))
+			if (Shape1)
 			{
 				if (Shape1->bDisable)
 				{
-					return;
+					return false;
 				}
 				if (Implicit1Type == ImplicitObjectType::TriangleMesh && Shape1->CollisionTraceType != Chaos_CTF_UseComplexAsSimple)
 				{
-					return;
+					return false;
 				}
 				else if (Shape1->CollisionTraceType == Chaos_CTF_UseComplexAsSimple && Implicit1Type != ImplicitObjectType::TriangleMesh)
 				{
-					return;
+					return false;
 				}
 
 			}
 			else if (Implicit1Type == ImplicitObjectType::TriangleMesh)
+			{
+				return false;
+			}
+
+			return true;
+		}
+
+		template<typename T, int d>
+		void ConstructConstraints(TGeometryParticleHandle<T, d>* Particle0, TGeometryParticleHandle<T, d>* Particle1, const FImplicitObject* Implicit0, const FImplicitObject* Implicit1, const TRigidTransform<T, d>& Transform0, const TRigidTransform<T, d>& Transform1, const T CullDistance, FCollisionConstraintsArray& NewConstraints)
+		{
+			EImplicitObjectType Implicit0Type = Implicit0 ? GetInnerType(Implicit0->GetType()) : ImplicitObjectType::Unknown;
+			EImplicitObjectType Implicit1Type = Implicit1 ? GetInnerType(Implicit1->GetType()) : ImplicitObjectType::Unknown;
+
+			if (!DoCollide(Implicit0Type, Particle0->GetImplicitShape(Implicit0), Implicit1Type, Particle1->GetImplicitShape(Implicit1)))
 			{
 				return;
 			}
@@ -1986,7 +2041,6 @@ namespace Chaos
 				return;
 			}
 
-
 			else if (Implicit0OuterType == TImplicitObjectInstanced<TBox<FReal,3>>::StaticType())
 			{
 				const TImplicitObjectInstanced<TBox<FReal, 3>>* TransformedImplicit0 = Implicit0->template GetObject<const TImplicitObjectInstanced<TBox<FReal, 3>>>();
@@ -2067,6 +2121,7 @@ namespace Chaos
 				return;
 			}
 
+
 			if (Implicit0Type == TBox<T, d>::StaticType() && Implicit1Type == TBox<T, d>::StaticType())
 			{
 				ConstructBoxBoxConstraints(Particle0, Particle1, Implicit0, Implicit1, Transform0, Transform1, CullDistance, NewConstraints);
diff --git a/Engine/Source/Runtime/Experimental/Chaos/Public/Chaos/CollisionFilterData.h b/Engine/Source/Runtime/Experimental/Chaos/Public/Chaos/CollisionFilterData.h
--- a/Engine/Source/Runtime/Experimental/Chaos/Public/Chaos/CollisionFilterData.h
+++ b/Engine/Source/Runtime/Experimental/Chaos/Public/Chaos/CollisionFilterData.h
@@ -22,3 +22,8 @@ inline Chaos::FChaosArchive& operator<<(Chaos::FChaosArchive& Ar, FCollisionFilt
 	Ar << Filter.Word0 << Filter.Word1 << Filter.Word2 << Filter.Word3;
 	return Ar;
 }
+
+inline bool operator!=(const FCollisionFilterData& A, const FCollisionFilterData& B)
+{
+	return A.Word0!=B.Word0 || A.Word1!=B.Word1 || A.Word2!=B.Word2 || A.Word3!=B.Word3;
+}
\ No newline at end of file
diff --git a/Engine/Source/Runtime/Experimental/Chaos/Public/Chaos/ParticleDirtyFlags.h b/Engine/Source/Runtime/Experimental/Chaos/Public/Chaos/ParticleDirtyFlags.h
--- a/Engine/Source/Runtime/Experimental/Chaos/Public/Chaos/ParticleDirtyFlags.h
+++ b/Engine/Source/Runtime/Experimental/Chaos/Public/Chaos/ParticleDirtyFlags.h
@@ -37,9 +37,10 @@ namespace Chaos
 		SpatialIdx				= 1 << 25,
 		UniqueIdx				= 1 << 26,
 		ShapeDisableCollision	= 1 << 27,
-		CollisionTraceType		= 1 << 28
+		CollisionTraceType		= 1 << 28,
+		ShapeSimData            = 1 << 29
 #if CHAOS_CHECKED
-		, DebugName				= 1 << 29
+		, DebugName				= 1 << 30
 #endif
 	};
 
diff --git a/Engine/Source/Runtime/Experimental/Chaos/Public/Chaos/ParticleHandle.h b/Engine/Source/Runtime/Experimental/Chaos/Public/Chaos/ParticleHandle.h
--- a/Engine/Source/Runtime/Experimental/Chaos/Public/Chaos/ParticleHandle.h
+++ b/Engine/Source/Runtime/Experimental/Chaos/Public/Chaos/ParticleHandle.h
@@ -1438,6 +1438,16 @@ public:
 		}
 	}
 
+	void SetShapeSimData(int32 InShapeIndex, const FCollisionFilterData& SimData)
+	{
+		const FCollisionFilterData& Current = MShapesArray[InShapeIndex]->SimData;
+		if (Current != SimData)
+		{
+			MShapesArray[InShapeIndex]->SimData = SimData;
+			MarkDirty(EParticleFlags::ShapeSimData);
+		}
+	}
+
 	FParticleData* NewData() const { return new TGeometryParticleData<T, d>( *this ); }
 
 	bool IsDirty() const
@@ -1571,12 +1581,13 @@ public:
 		const TShapesArray<T, d>& Shapes = InParticle.ShapesArray();
 		ShapeCollisionDisableFlags.Empty(Shapes.Num());
 		CollisionTraceType.Empty(Shapes.Num());
+		ShapeSimData.Empty(Shapes.Num());
 		for (const TUniquePtr<TPerShapeData<T, d>>& ShapePtr : Shapes)
 		{
 			ShapeCollisionDisableFlags.Add(ShapePtr->bDisable);
 			CollisionTraceType.Add(ShapePtr->CollisionTraceType);
+			ShapeSimData.Add(ShapePtr->SimData);
 		}
-
 	}
 
 	void Reset() 
@@ -1590,6 +1601,7 @@ public:
 		DirtyFlags.Clear();
 		ShapeCollisionDisableFlags.Reset();
 		CollisionTraceType.Reset();
+		ShapeSimData.Reset();
 #if CHAOS_CHECKED
 		DebugName = NAME_None;
 #endif
@@ -1603,6 +1615,7 @@ public:
 	FParticleDirtyFlags DirtyFlags;
 	TBitArray<> ShapeCollisionDisableFlags;
 	TArray<EChaosCollisionTraceFlag> CollisionTraceType;
+	TArray < FCollisionFilterData > ShapeSimData;
 #if CHAOS_CHECKED
 	FName DebugName;
 #endif
diff --git a/Engine/Source/Runtime/Experimental/ChaosSolvers/Private/PhysicsProxy/SingleParticlePhysicsProxy.cpp b/Engine/Source/Runtime/Experimental/ChaosSolvers/Private/PhysicsProxy/SingleParticlePhysicsProxy.cpp
--- a/Engine/Source/Runtime/Experimental/ChaosSolvers/Private/PhysicsProxy/SingleParticlePhysicsProxy.cpp
+++ b/Engine/Source/Runtime/Experimental/ChaosSolvers/Private/PhysicsProxy/SingleParticlePhysicsProxy.cpp
@@ -83,6 +83,17 @@ void FSingleParticlePhysicsProxy<Chaos::TGeometryParticle<float, 3>>::PushToPhys
 				Shape->CollisionTraceType = Data->CollisionTraceType[CurrShape++];
 			}
 		}
+
+		// @todo(chaos) : Fix - The initial update is not flagging the dirty flag, so this state is not updated. 
+		//if (Data->DirtyFlags.IsDirty(Chaos::EParticleFlags::ShapeSimData))
+		{
+			int32 CurrShape = 0;
+			for (const TUniquePtr<Chaos::TPerShapeData<Chaos::FReal, 3>>& Shape : RigidHandle->ShapesArray())
+			{
+				Shape->SimData = Data->ShapeSimData[CurrShape++];
+			}
+		}
+
 	}
 }
 
@@ -182,6 +193,15 @@ void FSingleParticlePhysicsProxy<Chaos::TKinematicGeometryParticle<float, 3>>::P
 				Shape->CollisionTraceType = Data->CollisionTraceType[CurrShape++];
 			}
 		}
+		//if (Data->DirtyFlags.IsDirty(Chaos::EParticleFlags::ShapeSimData))
+		{
+			int32 CurrShape = 0;
+			for (const TUniquePtr<Chaos::TPerShapeData<Chaos::FReal, 3>>& Shape : RigidHandle->ShapesArray())
+			{
+				Shape->SimData = Data->ShapeSimData[CurrShape++];
+			}
+		}
+
 
 	}
 }
@@ -349,6 +369,14 @@ void FSingleParticlePhysicsProxy<Chaos::TPBDRigidParticle<float, 3>>::PushToPhys
 				Shape->CollisionTraceType = Data->CollisionTraceType[CurrShape++];
 			}
 		}
+		//if (Data->DirtyFlags.IsDirty(Chaos::EParticleFlags::ShapeSimData))
+		{
+			int32 CurrShape = 0;
+			for (const TUniquePtr<Chaos::TPerShapeData<Chaos::FReal, 3>>& Shape : RigidHandle->ShapesArray())
+			{
+				Shape->SimData = Data->ShapeSimData[CurrShape++];
+			}
+		}
 
 		if (bDynamicPropertyUpdated)
 		{
