Chaos : Collision Filtering

- Basic filtering in CollisionResolution.

diff --git a/Engine/Source/Runtime/Experimental/Chaos/Private/Chaos/CollisionResolution.cpp b/Engine/Source/Runtime/Experimental/Chaos/Private/Chaos/CollisionResolution.cpp
--- a/Engine/Source/Runtime/Experimental/Chaos/Private/Chaos/CollisionResolution.cpp
+++ b/Engine/Source/Runtime/Experimental/Chaos/Private/Chaos/CollisionResolution.cpp
@@ -1921,35 +1921,17 @@ namespace Chaos
 		template<class T, int d>
 		bool DoCollide(EImplicitObjectType Implicit0Type, const TPerShapeData<T, d>* Shape0, EImplicitObjectType Implicit1Type, const TPerShapeData<T, d>* Shape1)
 		{
-			/*
-			if (Shape0 && Shape1)
-			{
-				if (!IsValid(Shape0->SimData) && !IsValid(Shape1->SimData))
-				{
-					return true;
-				}
-
-				FMaskFilter Filter0Mask, Filter1Mask;
-				const uint32 Filter0Channel = GetChaosCollisionChannelAndExtraFilter(Shape0->SimData.Word3, Filter0Mask);
-				const uint32 Filter1Channel = GetChaosCollisionChannelAndExtraFilter(Shape1->SimData.Word3, Filter1Mask);
-
-				if ((Filter0Mask & Filter1Mask) != 0)
-				{
-					return false;
-				}
-
-				const uint32 Filter1Bit = 1 << (Filter1Channel); // SIMDATA_TO_BITFIELD
-				uint32 const Filter0Bit = 1 << (Filter0Channel); // SIMDATA_TO_BITFIELD
-				return (Filter0Bit & Shape1->SimData.Word1) && (Filter1Bit & Shape0->SimData.Word1);
-			}
-			*/
+			//
+			// Disabled shapes do not collide
+			//
+			if (Shape0 && Shape0->bDisable) return false;
+			if (Shape1 && Shape1->bDisable) return false;
 
+			//
+			// Triangle Mesh geometry is only used if the shape specifies UseComplexAsSimple
+			//
 			if (Shape0)
 			{
-				if (Shape0->bDisable)
-				{
-					return false;
-				}
 				if (Implicit0Type == ImplicitObjectType::TriangleMesh && Shape0->CollisionTraceType != Chaos_CTF_UseComplexAsSimple)
 				{
 					return false;
@@ -1966,10 +1948,6 @@ namespace Chaos
 
 			if (Shape1)
 			{
-				if (Shape1->bDisable)
-				{
-					return false;
-				}
 				if (Implicit1Type == ImplicitObjectType::TriangleMesh && Shape1->CollisionTraceType != Chaos_CTF_UseComplexAsSimple)
 				{
 					return false;
@@ -1978,13 +1956,36 @@ namespace Chaos
 				{
 					return false;
 				}
-
 			}
 			else if (Implicit1Type == ImplicitObjectType::TriangleMesh)
 			{
 				return false;
 			}
 
+			//
+			// Shape Filtering
+			//
+			if (Shape0 && Shape1)
+			{
+
+				if (IsValid(Shape0->SimData) && IsValid(Shape1->SimData))
+				{
+					FMaskFilter Filter0Mask, Filter1Mask;
+					const uint32 Filter0Channel = GetChaosCollisionChannelAndExtraFilter(Shape0->SimData.Word3, Filter0Mask);
+					const uint32 Filter1Channel = GetChaosCollisionChannelAndExtraFilter(Shape1->SimData.Word3, Filter1Mask);
+
+					if ((Filter0Mask & Filter1Mask) != 0)
+					{
+						return false;
+					}
+
+					const uint32 Filter1Bit = 1 << (Filter1Channel); // SIMDATA_TO_BITFIELD
+					uint32 const Filter0Bit = 1 << (Filter0Channel); // SIMDATA_TO_BITFIELD
+					return (Filter0Bit & Shape1->SimData.Word1) && (Filter1Bit & Shape0->SimData.Word1);
+				}
+			}
+
+
 			return true;
 		}
 
