Fix robomerge

Sync userdata with pt

diff --git a/Engine/Source/Runtime/Experimental/Chaos/Public/Chaos/GeometryParticles.h b/Engine/Source/Runtime/Experimental/Chaos/Public/Chaos/GeometryParticles.h
--- a/Engine/Source/Runtime/Experimental/Chaos/Public/Chaos/GeometryParticles.h
+++ b/Engine/Source/Runtime/Experimental/Chaos/Public/Chaos/GeometryParticles.h
@@ -125,6 +125,7 @@ namespace Chaos
 			TArrayCollection::AddArray(&MWorldSpaceInflatedBounds);
 			TArrayCollection::AddArray(&MHasBounds);
 			TArrayCollection::AddArray(&MSpatialIdx);
+			TArrayCollection::AddArray(&MUserData);
 #if CHAOS_CHECKED
 			TArrayCollection::AddArray(&MDebugName);
 #endif
@@ -152,6 +153,7 @@ namespace Chaos
 			, MWorldSpaceInflatedBounds(MoveTemp(Other.MWorldSpaceInflatedBounds))
 			, MHasBounds(MoveTemp(Other.MHasBounds))
 			, MSpatialIdx(MoveTemp(Other.MSpatialIdx))
+			, MUserData(MoveTemp(Other.MUserData))
 #if CHAOS_DETERMINISTIC
 			, MParticleIDs(MoveTemp(Other.MParticleIDs))
 #endif
@@ -168,6 +170,7 @@ namespace Chaos
 			TArrayCollection::AddArray(&MWorldSpaceInflatedBounds);
 			TArrayCollection::AddArray(&MHasBounds);
 			TArrayCollection::AddArray(&MSpatialIdx);
+			TArrayCollection::AddArray(&MUserData);
 #if CHAOS_DETERMINISTIC
 			TArrayCollection::AddArray(&MParticleIDs);
 #endif
@@ -199,6 +202,7 @@ namespace Chaos
 			TArrayCollection::AddArray(&MWorldSpaceInflatedBounds);
 			TArrayCollection::AddArray(&MHasBounds);
 			TArrayCollection::AddArray(&MSpatialIdx);
+			TArrayCollection::AddArray(&MUserData);
 #if CHAOS_DETERMINISTIC
 			TArrayCollection::AddArray(&MParticleIDs);
 #endif
@@ -222,6 +226,9 @@ namespace Chaos
 		CHAOS_API FUniqueIdx UniqueIdx(const int32 Index) const { return MUniqueIdx[Index]; }
 		CHAOS_API FUniqueIdx& UniqueIdx(const int32 Index) { return MUniqueIdx[Index]; }
 
+		void*& UserData(const int32 Index) { return MUserData[Index]; }
+		const void* UserData(const int32 Index) const { return MUserData[Index]; }
+
 		CHAOS_API TSerializablePtr<FImplicitObject> Geometry(const int32 Index) const { return MGeometry[Index]; }
 
 		CHAOS_API const TUniquePtr<FImplicitObject>& DynamicGeometry(const int32 Index) const { return MDynamicGeometry[Index]; }
@@ -446,6 +453,7 @@ namespace Chaos
 		TArrayCollectionArray<TAABB<T, d>> MWorldSpaceInflatedBounds;
 		TArrayCollectionArray<bool> MHasBounds;
 		TArrayCollectionArray<FSpatialAccelerationIdx> MSpatialIdx;
+		TArrayCollectionArray<void*> MUserData;
 
 		void UpdateShapesArray(const int32 Index)
 		{
diff --git a/Engine/Source/Runtime/Experimental/Chaos/Public/Chaos/ParticleDirtyFlags.h b/Engine/Source/Runtime/Experimental/Chaos/Public/Chaos/ParticleDirtyFlags.h
--- a/Engine/Source/Runtime/Experimental/Chaos/Public/Chaos/ParticleDirtyFlags.h
+++ b/Engine/Source/Runtime/Experimental/Chaos/Public/Chaos/ParticleDirtyFlags.h
@@ -38,9 +38,10 @@ namespace Chaos
 		UniqueIdx				= 1 << 26,
 		ShapeDisableCollision	= 1 << 27,
 		CollisionTraceType		= 1 << 28,
-		ShapeSimData            = 1 << 29
+		ShapeSimData            = 1 << 29,
+		UserData				= 1 << 30
 #if CHAOS_CHECKED
-		, DebugName				= 1 << 30
+		, DebugName				= 1 << 31
 #endif
 	};
 
diff --git a/Engine/Source/Runtime/Experimental/Chaos/Public/Chaos/ParticleHandle.h b/Engine/Source/Runtime/Experimental/Chaos/Public/Chaos/ParticleHandle.h
--- a/Engine/Source/Runtime/Experimental/Chaos/Public/Chaos/ParticleHandle.h
+++ b/Engine/Source/Runtime/Experimental/Chaos/Public/Chaos/ParticleHandle.h
@@ -50,6 +50,7 @@ void GeometryParticleDefaultConstruct(FConcrete& Concrete, const TGeometryPartic
 	Concrete.SetX(TVector<T, d>(0));
 	Concrete.SetR(TRotation<T, d>::Identity);
 	Concrete.SetSpatialIdx(FSpatialAccelerationIdx{ 0,0 });
+	Concrete.SetUserData(nullptr);
 }
 
 template <typename T, int d, typename FConcrete>
@@ -361,6 +362,9 @@ public:
 	FUniqueIdx UniqueIdx() const { return GeometryParticles->UniqueIdx(ParticleIdx); }
 	void SetUniqueIdx(const FUniqueIdx UniqueIdx) const { GeometryParticles->UniqueIdx(ParticleIdx) = UniqueIdx; }
 
+	void* UserData() const { return GeometryParticles->UserData(ParticleIdx); }
+	void SetUserData(void* InUserData) { GeometryParticles->UserData(ParticleIdx) = InUserData; }
+
 	const TRotation<T, d>& R() const { return GeometryParticles->R(ParticleIdx); }
 	TRotation<T, d>& R() { return GeometryParticles->R(ParticleIdx); }
 	void SetR(const TRotation<T, d>& InR) { GeometryParticles->R(ParticleIdx) = InR; }
@@ -1361,6 +1365,7 @@ public:
 	void* UserData() const { return this->MUserData; }
 	void SetUserData(void* InUserData)
 	{
+		this->MarkDirty(EParticleFlags::UserData);
 		this->MUserData = InUserData;
 	}
 
@@ -1560,6 +1565,7 @@ public:
 		, X(TVector<T, d>(0))
 		, R(TRotation<T, d>())
 		, SpatialIdx(FSpatialAccelerationIdx{ 0,0 })
+		, UserData(nullptr)
 		, DirtyFlags()
 #if CHAOS_CHECKED
 		, DebugName(NAME_None)
@@ -1573,6 +1579,7 @@ public:
 		, Geometry(InParticle.GeometrySharedLowLevel())
 		, SpatialIdx(InParticle.SpatialIdx())
 		, UniqueIdx(InParticle.UniqueIdx())
+		, UserData(InParticle.UserData())
 		, DirtyFlags(InParticle.DirtyFlags())
 #if CHAOS_CHECKED
 		, DebugName(InParticle.DebugName())
@@ -1598,6 +1605,7 @@ public:
 		Geometry = TSharedPtr<FImplicitObjectUnion, ESPMode::ThreadSafe>();
 		SpatialIdx = FSpatialAccelerationIdx{ 0,0 };
 		UniqueIdx = FUniqueIdx();
+		UserData = nullptr;
 		DirtyFlags.Clear();
 		ShapeCollisionDisableFlags.Reset();
 		CollisionTraceType.Reset();
@@ -1612,6 +1620,7 @@ public:
 	TSharedPtr<FImplicitObject, ESPMode::ThreadSafe> Geometry;
 	FSpatialAccelerationIdx SpatialIdx;
 	FUniqueIdx UniqueIdx;
+	void* UserData;
 	FParticleDirtyFlags DirtyFlags;
 	TBitArray<> ShapeCollisionDisableFlags;
 	TArray<EChaosCollisionTraceFlag> CollisionTraceType;
@@ -2299,7 +2308,6 @@ FChaosArchive& operator<<(FChaosArchive& Ar, TAccelerationStructureHandle<T, d>&
 	AccelerationHandle.Serialize(Ar);
 	return Ar;
 }
-
 #if PLATFORM_MAC || PLATFORM_LINUX
 extern template class CHAOS_API TGeometryParticle<float, 3>;
 extern template class CHAOS_API TKinematicGeometryParticle<float, 3>;
diff --git a/Engine/Source/Runtime/Experimental/ChaosSolvers/Private/PhysicsProxy/SingleParticlePhysicsProxy.cpp b/Engine/Source/Runtime/Experimental/ChaosSolvers/Private/PhysicsProxy/SingleParticlePhysicsProxy.cpp
--- a/Engine/Source/Runtime/Experimental/ChaosSolvers/Private/PhysicsProxy/SingleParticlePhysicsProxy.cpp
+++ b/Engine/Source/Runtime/Experimental/ChaosSolvers/Private/PhysicsProxy/SingleParticlePhysicsProxy.cpp
@@ -61,6 +61,7 @@ void FSingleParticlePhysicsProxy<Chaos::TGeometryParticle<float, 3>>::PushToPhys
 		}
 		RigidHandle->SetSpatialIdx(Data->SpatialIdx);	//todo: this needs to only happen once during initialization
 		RigidHandle->SetUniqueIdx(Data->UniqueIdx);
+		RigidHandle->SetUserData(Data->UserData);
 #if CHAOS_CHECKED
 		RigidHandle->SetDebugName(Data->DebugName);
 #endif
@@ -164,6 +165,7 @@ void FSingleParticlePhysicsProxy<Chaos::TKinematicGeometryParticle<float, 3>>::P
 		RigidHandle->SetRotationOfMass(Data->MRotationOfMass);
 		RigidHandle->SetSpatialIdx(Data->SpatialIdx);	//todo: this needs to only happen once during initialization
 		RigidHandle->SetUniqueIdx(Data->UniqueIdx);
+		RigidHandle->SetUserData(Data->UserData);
 #if CHAOS_CHECKED
 		RigidHandle->SetDebugName(Data->DebugName);
 #endif
@@ -273,6 +275,7 @@ void FSingleParticlePhysicsProxy<Chaos::TPBDRigidParticle<float, 3>>::PushToPhys
 		// This should be moved to initialization.
 		RigidHandle->SetSpatialIdx(Data->SpatialIdx);	//todo: this needs to only happen once during initialization
 		RigidHandle->SetUniqueIdx(Data->UniqueIdx);	//todo: this needs to only happen once during initialization
+		RigidHandle->SetUserData(Data->UserData);
 #if CHAOS_CHECKED
 		RigidHandle->SetDebugName(Data->DebugName);
 #endif
