Fix Scaled TriMesh Overlap w/ MTD.
Fix Heightfield OverlapGeomImp wrong numeric types used.

diff --git a/Engine/Source/Runtime/Experimental/Chaos/Private/Chaos/HeightField.cpp b/Engine/Source/Runtime/Experimental/Chaos/Private/Chaos/HeightField.cpp
--- a/Engine/Source/Runtime/Experimental/Chaos/Private/Chaos/HeightField.cpp
+++ b/Engine/Source/Runtime/Experimental/Chaos/Private/Chaos/HeightField.cpp
@@ -1329,12 +1329,15 @@ namespace Chaos
 
 	template <typename T>
 	template <typename QueryGeomType>
-	bool THeightField<T>::OverlapGeomImp(const QueryGeomType& QueryGeom, const TRigidTransform<T, 3>& QueryTM, const T Thickness, FVec3& OutNormal, FReal& OutPenetration, bool bComputeMTD) const
+	bool THeightField<T>::OverlapGeomImp(const QueryGeomType& QueryGeom, const TRigidTransform<T, 3>& QueryTM, const T Thickness, FMTDInfo* OutMTD) const
 	{
-		OutNormal = FVec3(0.0f);
-		OutPenetration = 0.0f;
+		if (OutMTD)
+		{
+			OutMTD->Normal = TVec3<T>(0);
+			OutMTD->Penetration = TNumericLimits<T>::Lowest();
+		}
 
-		auto OverlapTriangle = [&](const TVector<T, 3>& A, const TVector<T, 3>& B, const TVector<T, 3>& C) -> bool
+		auto OverlapTriangle = [&](const TVector<T, 3>& A, const TVector<T, 3>& B, const TVector<T, 3>& C, FMTDInfo* InnerMTD) -> bool
 		{
 			const TVector<T, 3> AB = B - A;
 			const TVector<T, 3> AC = C - A;
@@ -1344,19 +1347,19 @@ namespace Chaos
 			const TVector<T, 3> Offset = TVector<T, 3>::CrossProduct(AB, AC);
 
 			TTriangle<T> TriangleConvex(A, B, C);
-			if (bComputeMTD)
+			if (InnerMTD)
 			{
-				FVec3 TriangleNormal(0.0f);
-				FReal Penetration = 0.0f;
-				FVec3 ClosestA(0.0f);
-				FVec3 ClosestB(0.0f);
+				TVec3<T> TriangleNormal(0);
+				T Penetration = 0;
+				TVec3<T> ClosestA(0);
+				TVec3<T> ClosestB(0);
 				if (GJKPenetration(TriangleConvex, QueryGeom, QueryTM, Penetration, ClosestA, ClosestB, TriangleNormal, Thickness))
 				{
 					// Use Deepest MTD.
-					if (Penetration > OutPenetration)
+					if (Penetration > InnerMTD->Penetration)
 					{
-						OutPenetration = Penetration;
-						OutNormal = TriangleNormal;
+						InnerMTD->Penetration = Penetration;
+						InnerMTD->Normal = TriangleNormal;
 					}
 					return true;
 				}
@@ -1389,19 +1392,19 @@ namespace Chaos
 			const int32 SingleIndex = Cell[1] * (GeomData.NumCols) + Cell[0];
 			GeomData.GetPointsScaled(SingleIndex, Points);
 
-			if(OverlapTriangle(Points[0], Points[1], Points[3]))
+			if(OverlapTriangle(Points[0], Points[1], Points[3], OutMTD))
 			{
 				bOverlaps = true;
-				if (!bComputeMTD)
+				if (!OutMTD)
 				{
 					return true;
 				}
 			}
 
-			if(OverlapTriangle(Points[0], Points[3], Points[2]))
+			if(OverlapTriangle(Points[0], Points[3], Points[2], OutMTD))
 			{
 				bOverlaps = true;
-				if (!bComputeMTD)
+				if (!OutMTD)
 				{
 					return true;
 				}
@@ -1414,161 +1417,49 @@ namespace Chaos
 	template <typename T>
 	bool THeightField<T>::OverlapGeom(const TSphere<T, 3>& QueryGeom, const TRigidTransform<T, 3>& QueryTM, const T Thickness, FMTDInfo* OutMTD) const
 	{
-		FVec3 Normal;
-		FReal Penetration;
-		if (OutMTD)
-		{
-			// Computes MTD
-			if (OverlapGeomImp(QueryGeom, QueryTM, Thickness, Normal, Penetration, /*bComputeMTD=*/true))
-			{
-				OutMTD->Normal = Normal;
-				OutMTD->Penetration = Penetration;
-				return true;
-			}
-		}
-
-		// No MTD
-		return OverlapGeomImp(QueryGeom, QueryTM, Thickness, Normal, Penetration, /*bComputeMTD=*/false);
+		return OverlapGeomImp(QueryGeom, QueryTM, Thickness, OutMTD);
 	}
 
 	template <typename T>
 	bool THeightField<T>::OverlapGeom(const TBox<T, 3>& QueryGeom, const TRigidTransform<T, 3>& QueryTM, const T Thickness, FMTDInfo* OutMTD) const
 	{
-		FVec3 Normal;
-		FReal Penetration;
-		if (OutMTD)
-		{
-			// Computes MTD
-			if (OverlapGeomImp(QueryGeom, QueryTM, Thickness, Normal, Penetration, /*bComputeMTD=*/true))
-			{
-				OutMTD->Normal = Normal;
-				OutMTD->Penetration = Penetration;
-				return true;
-			}
-		}
-
-		// No MTD
-		return OverlapGeomImp(QueryGeom, QueryTM, Thickness, Normal, Penetration, /*bComputeMTD=*/false);
+		return OverlapGeomImp(QueryGeom, QueryTM, Thickness, OutMTD);
 	}
 
 	template <typename T>
 	bool THeightField<T>::OverlapGeom(const TCapsule<T>& QueryGeom, const TRigidTransform<T, 3>& QueryTM, const T Thickness, FMTDInfo* OutMTD) const
 	{
-		FVec3 Normal;
-		FReal Penetration;
-		if (OutMTD)
-		{
-			// Computes MTD
-			if (OverlapGeomImp(QueryGeom, QueryTM, Thickness, Normal, Penetration, /*bComputeMTD=*/true))
-			{
-				OutMTD->Normal = Normal;
-				OutMTD->Penetration = Penetration;
-				return true;
-			}
-		}
-
-		// No MTD
-		return OverlapGeomImp(QueryGeom, QueryTM, Thickness, Normal, Penetration, /*bComputeMTD=*/false);
+		return OverlapGeomImp(QueryGeom, QueryTM, Thickness, OutMTD);
 	}
 
 	template <typename T>
 	bool THeightField<T>::OverlapGeom(const FConvex& QueryGeom, const TRigidTransform<T, 3>& QueryTM, const T Thickness, FMTDInfo* OutMTD) const
 	{
-		FVec3 Normal;
-		FReal Penetration;
-		if (OutMTD)
-		{
-			// Computes MTD
-			if (OverlapGeomImp(QueryGeom, QueryTM, Thickness, Normal, Penetration, /*bComputeMTD=*/true))
-			{
-				OutMTD->Normal = Normal;
-				OutMTD->Penetration = Penetration;
-				return true;
-			}
-		}
-
-		// No MTD
-		return OverlapGeomImp(QueryGeom, QueryTM, Thickness, Normal, Penetration, /*bComputeMTD=*/false);
+		return OverlapGeomImp(QueryGeom, QueryTM, Thickness, OutMTD);
 	}
 
 	template <typename T>
 	bool THeightField<T>::OverlapGeom(const TImplicitObjectScaled<TSphere<T, 3>>& QueryGeom, const TRigidTransform<T, 3>& QueryTM, const T Thickness, FMTDInfo* OutMTD) const
 	{
-		FVec3 Normal;
-		FReal Penetration;
-		if (OutMTD)
-		{
-			// Computes MTD
-			if (OverlapGeomImp(QueryGeom, QueryTM, Thickness, Normal, Penetration, /*bComputeMTD=*/true))
-			{
-				OutMTD->Normal = Normal;
-				OutMTD->Penetration = Penetration;
-				return true;
-			}
-		}
-
-		// No MTD
-		return OverlapGeomImp(QueryGeom, QueryTM, Thickness, Normal, Penetration, /*bComputeMTD=*/false);
+		return OverlapGeomImp(QueryGeom, QueryTM, Thickness, OutMTD);
 	}
 
 	template <typename T>
 	bool THeightField<T>::OverlapGeom(const TImplicitObjectScaled<TBox<T, 3>>& QueryGeom, const TRigidTransform<T, 3>& QueryTM, const T Thickness, FMTDInfo* OutMTD) const
 	{
-		FVec3 Normal;
-		FReal Penetration;
-		if (OutMTD)
-		{
-			// Computes MTD
-			if (OverlapGeomImp(QueryGeom, QueryTM, Thickness, Normal, Penetration, /*bComputeMTD=*/true))
-			{
-				OutMTD->Normal = Normal;
-				OutMTD->Penetration = Penetration;
-				return true;
-			}
-		}
-
-		// No MTD
-		return OverlapGeomImp(QueryGeom, QueryTM, Thickness, Normal, Penetration, /*bComputeMTD=*/false);
+		return OverlapGeomImp(QueryGeom, QueryTM, Thickness, OutMTD);
 	}
 
 	template <typename T>
 	bool THeightField<T>::OverlapGeom(const TImplicitObjectScaled<TCapsule<T>>& QueryGeom, const TRigidTransform<T, 3>& QueryTM, const T Thickness, FMTDInfo* OutMTD) const
 	{
-		FVec3 Normal;
-		FReal Penetration;
-		if (OutMTD)
-		{
-			// Computes MTD
-			if (OverlapGeomImp(QueryGeom, QueryTM, Thickness, Normal, Penetration, /*bComputeMTD=*/true))
-			{
-				OutMTD->Normal = Normal;
-				OutMTD->Penetration = Penetration;
-				return true;
-			}
-		}
-
-		// No MTD
-		return OverlapGeomImp(QueryGeom, QueryTM, Thickness, Normal, Penetration, /*bComputeMTD=*/false);
+		return OverlapGeomImp(QueryGeom, QueryTM, Thickness, OutMTD);
 	}
 
 	template <typename T>
 	bool THeightField<T>::OverlapGeom(const TImplicitObjectScaled<FConvex>& QueryGeom, const TRigidTransform<T, 3>& QueryTM, const T Thickness, FMTDInfo* OutMTD) const
 	{
-		FVec3 Normal;
-		FReal Penetration;
-		if (OutMTD)
-		{
-			// Computes MTD
-			if (OverlapGeomImp(QueryGeom, QueryTM, Thickness, Normal, Penetration, /*bComputeMTD=*/true))
-			{
-				OutMTD->Normal = Normal;
-				OutMTD->Penetration = Penetration;
-				return true;
-			}
-		}
-
-		// No MTD
-		return OverlapGeomImp(QueryGeom, QueryTM, Thickness, Normal, Penetration, /*bComputeMTD=*/false);
+		return OverlapGeomImp(QueryGeom, QueryTM, Thickness, OutMTD);
 	}
 
 	template <typename T>
diff --git a/Engine/Source/Runtime/Experimental/Chaos/Private/Chaos/TriangleMeshImplicitObject.cpp b/Engine/Source/Runtime/Experimental/Chaos/Private/Chaos/TriangleMeshImplicitObject.cpp
--- a/Engine/Source/Runtime/Experimental/Chaos/Private/Chaos/TriangleMeshImplicitObject.cpp
+++ b/Engine/Source/Runtime/Experimental/Chaos/Private/Chaos/TriangleMeshImplicitObject.cpp
@@ -503,29 +503,7 @@ void TransformOverlapInputsHelper(const TImplicitObjectScaled<QueryGeomType>& Qu
 }
 
 template <typename QueryGeomType>
-void TransformOverlapOutputsHelper(const QueryGeomType& QueryGeom, const TVec3<FReal>& Normal, const FReal& Penetration, TVec3<FReal>& OutNormal, FReal& OutPenetration)
-{
-	OutNormal = Normal;
-	OutPenetration = Penetration;
-}
-
-template <typename QueryGeomType>
-void TransformOverlapOutputsHelper(const TImplicitObjectScaled<QueryGeomType>& QueryGeom, const TVec3<FReal>& Normal, const FReal& Penetration, TVec3<FReal>& OutNormal, FReal& OutPenetration)
-{
-	const TVec3<FReal> InvScale = QueryGeom.GetInvScale();
-
-	const TVec3<FReal> UnscaledNormalDenorm = (InvScale * Normal);
-	const FReal UnscaledNormalLength = UnscaledNormalDenorm.Size();
-	if (UnscaledNormalLength > 0.0f)
-	{
-		OutNormal = UnscaledNormalDenorm / UnscaledNormalLength;
-	}
-
-	OutPenetration = Penetration * UnscaledNormalLength;
-}
-
-template <typename QueryGeomType>
-bool FTriangleMeshImplicitObject::OverlapGeomImp(const QueryGeomType& QueryGeom, const FRigidTransform3& QueryTM, const FReal Thickness, FVec3& OutNormal, FReal& OutPenetration, const bool bComputeMTD) const
+bool FTriangleMeshImplicitObject::OverlapGeomImp(const QueryGeomType& QueryGeom, const FRigidTransform3& QueryTM, const FReal Thickness, FMTDInfo* OutMTD) const
 {
 	bool bResult = false;
 	TAABB<FReal, 3> QueryBounds = QueryGeom.BoundingBox();
@@ -533,20 +511,20 @@ bool FTriangleMeshImplicitObject::OverlapGeomImp(const QueryGeomType& QueryGeom,
 	QueryBounds = QueryBounds.TransformedAABB(QueryTM);
 	const TArray<int32> PotentialIntersections = BVH.FindAllIntersections(QueryBounds);
 
-	OutNormal = FVec3(0.0f);
-	OutPenetration = 0.0f;
-
-	FVec3 DeepestNormal = FVec3(0.0f);
-	FReal DeepestPenetration = 0.0f;
+	if (OutMTD)
+	{
+		OutMTD->Normal = FVec3(0.0);
+		OutMTD->Penetration = TNumericLimits<FReal>::Lowest();
+	}
 
 	const auto& InnerQueryGeom = GetGeomHelper(QueryGeom);
 
 	TRigidTransform<FReal, 3> TransformedQueryTM;
 	TransformOverlapInputsHelper(QueryGeom, QueryTM, TransformedQueryTM);
 
-	auto LambdaHelper = [&](const auto& Elements)
+	auto LambdaHelper = [&](const auto& Elements, FMTDInfo* InnerMTD)
 	{
-		if (bComputeMTD)
+		if (InnerMTD)
 		{
 			bool bOverlap = false;
 			for (int32 TriIdx : PotentialIntersections)
@@ -554,25 +532,23 @@ bool FTriangleMeshImplicitObject::OverlapGeomImp(const QueryGeomType& QueryGeom,
 				TVec3<FReal> A, B, C;
 				TransformVertsHelper(QueryGeom, TriIdx, MParticles, Elements, A, B, C);
 
-				FVec3 TriangleNormal(0.0f);
-				FReal Penetration = 0.0f;
-				FVec3 ClosestA(0.0f);
-				FVec3 ClosestB(0.0f);
+				FVec3 TriangleNormal(0.0);
+				FReal Penetration = 0.0;
+				FVec3 ClosestA(0.0);
+				FVec3 ClosestB(0.0);
 				if (GJKPenetration(TTriangle<FReal>(A, B, C), InnerQueryGeom, TransformedQueryTM, Penetration, ClosestA, ClosestB, TriangleNormal, Thickness))
 				{
 					bOverlap = true;
 
 					// Use Deepest MTD.
-					if (Penetration > DeepestPenetration)
+					if (Penetration > InnerMTD->Penetration)
 					{
-						DeepestPenetration = Penetration;
-						DeepestNormal = TriangleNormal;
+						InnerMTD->Penetration = Penetration;
+						InnerMTD->Normal = TriangleNormal;
 					}
 				}
 			}
 
-			TransformOverlapOutputsHelper(QueryGeom, DeepestNormal, DeepestPenetration, OutNormal, OutPenetration);
-
 			return bOverlap;
 		}
 		else
@@ -601,183 +577,57 @@ bool FTriangleMeshImplicitObject::OverlapGeomImp(const QueryGeomType& QueryGeom,
 
 	if(MElements.RequiresLargeIndices())
 	{
-		return LambdaHelper(MElements.GetLargeIndexBuffer());
+		return LambdaHelper(MElements.GetLargeIndexBuffer(), OutMTD);
 	}
 	else
 	{
-		return LambdaHelper(MElements.GetSmallIndexBuffer());
+		return LambdaHelper(MElements.GetSmallIndexBuffer(), OutMTD);
 	}
 }
 
 bool FTriangleMeshImplicitObject::OverlapGeom(const TSphere<FReal, 3>& QueryGeom, const FRigidTransform3& QueryTM, const FReal Thickness, FMTDInfo* OutMTD) const
 {
-	FVec3 Normal;
-	FReal Penetration;
-	if (OutMTD)
-	{
-		// Computes MTD
-		if (OverlapGeomImp(QueryGeom, QueryTM, Thickness, Normal, Penetration, /*bComputeMTD=*/true))
-		{
-			OutMTD->Normal = Normal;
-			OutMTD->Penetration = Penetration;
-			return true;
-		}
-	}
-
-	// No MTD
-	return OverlapGeomImp(QueryGeom, QueryTM, Thickness, Normal, Penetration, /*bComputeMTD=*/false);
+	return OverlapGeomImp(QueryGeom, QueryTM, Thickness, OutMTD);
 }
 
 bool FTriangleMeshImplicitObject::OverlapGeom(const TBox<FReal, 3>& QueryGeom, const FRigidTransform3& QueryTM, const FReal Thickness, FMTDInfo* OutMTD) const
 {
-	FVec3 Normal;
-	FReal Penetration;
-	if (OutMTD)
-	{
-		// Computes MTD
-		if (OverlapGeomImp(QueryGeom, QueryTM, Thickness, Normal, Penetration, /*bComputeMTD=*/true))
-		{
-			OutMTD->Normal = Normal;
-			OutMTD->Penetration = Penetration;
-			return true;
-		}
-	}
-
-	// No MTD
-	return OverlapGeomImp(QueryGeom, QueryTM, Thickness, Normal, Penetration, /*bComputeMTD=*/false);
+	return OverlapGeomImp(QueryGeom, QueryTM, Thickness, OutMTD);
 }
 
 bool FTriangleMeshImplicitObject::OverlapGeom(const TCapsule<FReal>& QueryGeom, const FRigidTransform3& QueryTM, const FReal Thickness, FMTDInfo* OutMTD) const
 {
-	FVec3 Normal;
-	FReal Penetration;
-	if (OutMTD)
-	{
-		// Computes MTD
-		if (OverlapGeomImp(QueryGeom, QueryTM, Thickness, Normal, Penetration, /*bComputeMTD=*/true))
-		{
-			OutMTD->Normal = Normal;
-			OutMTD->Penetration = Penetration;
-			return true;
-		}
-	}
-
-	// No MTD
-	return OverlapGeomImp(QueryGeom, QueryTM, Thickness, Normal, Penetration, /*bComputeMTD=*/false);
+	return OverlapGeomImp(QueryGeom, QueryTM, Thickness, OutMTD);
 }
 
 bool FTriangleMeshImplicitObject::OverlapGeom(const FConvex& QueryGeom, const FRigidTransform3& QueryTM, const FReal Thickness, FMTDInfo* OutMTD) const
 {
-	FVec3 Normal;
-	FReal Penetration;
-	if (OutMTD)
-	{
-		// Computes MTD
-		if (OverlapGeomImp(QueryGeom, QueryTM, Thickness, Normal, Penetration, /*bComputeMTD=*/true))
-		{
-			OutMTD->Normal = Normal;
-			OutMTD->Penetration = Penetration;
-			return true;
-		}
-	}
-
-	// No MTD
-	return OverlapGeomImp(QueryGeom, QueryTM, Thickness, Normal, Penetration, /*bComputeMTD=*/false);
+	return OverlapGeomImp(QueryGeom, QueryTM, Thickness, OutMTD);
 }
 
 bool FTriangleMeshImplicitObject::OverlapGeom(const TImplicitObjectScaled<TSphere<FReal, 3>>& QueryGeom, const FRigidTransform3& QueryTM, const FReal Thickness, FMTDInfo* OutMTD) const
 {
-	FVec3 Normal;
-	FReal Penetration;
-	if (OutMTD)
-	{
-		// Computes MTD
-		if (OverlapGeomImp(QueryGeom, QueryTM, Thickness, Normal, Penetration, /*bComputeMTD=*/true))
-		{
-			OutMTD->Normal = Normal;
-			OutMTD->Penetration = Penetration;
-			return true;
-		}
-	}
-
-	// No MTD
-	return OverlapGeomImp(QueryGeom, QueryTM, Thickness, Normal, Penetration, /*bComputeMTD=*/false);
+	return OverlapGeomImp(QueryGeom, QueryTM, Thickness, OutMTD);
 }
 
 bool FTriangleMeshImplicitObject::OverlapGeom(const TImplicitObjectScaled<TBox<FReal, 3>>& QueryGeom, const FRigidTransform3& QueryTM, const FReal Thickness, FMTDInfo* OutMTD) const
 {
-	FVec3 Normal;
-	FReal Penetration;
-	if (OutMTD)
-	{
-		// Computes MTD
-		if (OverlapGeomImp(QueryGeom, QueryTM, Thickness, Normal, Penetration, /*bComputeMTD=*/true))
-		{
-			OutMTD->Normal = Normal;
-			OutMTD->Penetration = Penetration;
-			return true;
-		}
-	}
-
-	// No MTD
-	return OverlapGeomImp(QueryGeom, QueryTM, Thickness, Normal, Penetration, /*bComputeMTD=*/false);
+	return OverlapGeomImp(QueryGeom, QueryTM, Thickness, OutMTD);
 }
 
 bool FTriangleMeshImplicitObject::OverlapGeom(const TImplicitObjectScaled<TCapsule<FReal>>& QueryGeom, const FRigidTransform3& QueryTM, const FReal Thickness, FMTDInfo* OutMTD) const
 {
-	FVec3 Normal;
-	FReal Penetration;
-	if (OutMTD)
-	{
-		// Computes MTD
-		if (OverlapGeomImp(QueryGeom, QueryTM, Thickness, Normal, Penetration, /*bComputeMTD=*/true))
-		{
-			OutMTD->Normal = Normal;
-			OutMTD->Penetration = Penetration;
-			return true;
-		}
-	}
-
-	// No MTD
-	return OverlapGeomImp(QueryGeom, QueryTM, Thickness, Normal, Penetration, /*bComputeMTD=*/false);
+	return OverlapGeomImp(QueryGeom, QueryTM, Thickness, OutMTD);
 }
 
 bool FTriangleMeshImplicitObject::OverlapGeom(const TImplicitObjectScaled<FConvex>& QueryGeom, const FRigidTransform3& QueryTM, const FReal Thickness, FMTDInfo* OutMTD) const
 {
-	FVec3 Normal;
-	FReal Penetration;
-	if (OutMTD)
-	{
-		// Computes MTD
-		if (OverlapGeomImp(QueryGeom, QueryTM, Thickness, Normal, Penetration, /*bComputeMTD=*/true))
-		{
-			OutMTD->Normal = Normal;
-			OutMTD->Penetration = Penetration;
-			return true;
-		}
-	}
-
-	// No MTD
-	return OverlapGeomImp(QueryGeom, QueryTM, Thickness, Normal, Penetration, /*bComputeMTD=*/false);
+	return OverlapGeomImp(QueryGeom, QueryTM, Thickness, OutMTD);
 }
 
 bool FTriangleMeshImplicitObject::OverlapGeom(const TImplicitObjectScaled<TImplicitObjectScaled<FConvex>>& QueryGeom, const FRigidTransform3& QueryTM, const FReal Thickness, FMTDInfo* OutMTD) const
 {
-	FVec3 Normal;
-	FReal Penetration;
-	if (OutMTD)
-	{
-		// Computes MTD
-		if (OverlapGeomImp(QueryGeom, QueryTM, Thickness, Normal, Penetration, /*bComputeMTD=*/true))
-		{
-			OutMTD->Normal = Normal;
-			OutMTD->Penetration = Penetration;
-			return true;
-		}
-	}
-
-	// No MTD
-	return OverlapGeomImp(QueryGeom, QueryTM, Thickness, Normal, Penetration, /*bComputeMTD=*/false);
+	return OverlapGeomImp(QueryGeom, QueryTM, Thickness, OutMTD);
 }
 
 
diff --git a/Engine/Source/Runtime/Experimental/Chaos/Public/Chaos/HeightField.h b/Engine/Source/Runtime/Experimental/Chaos/Public/Chaos/HeightField.h
--- a/Engine/Source/Runtime/Experimental/Chaos/Public/Chaos/HeightField.h
+++ b/Engine/Source/Runtime/Experimental/Chaos/Public/Chaos/HeightField.h
@@ -543,7 +543,7 @@ namespace Chaos
 		friend FImplicitObject;
 
 		template <typename QueryGeomType>
-		bool OverlapGeomImp(const QueryGeomType& QueryGeom, const TRigidTransform<T, 3>& QueryTM, const T Thickness, FVec3& OutNormal, FReal& OutPenetration, const bool bComputeMTD) const;
+		bool OverlapGeomImp(const QueryGeomType& QueryGeom, const TRigidTransform<T, 3>& QueryTM, const T Thickness, FMTDInfo* OutMTD = nullptr) const;
 
 		template <typename QueryGeomType>
 		bool SweepGeomImp(const QueryGeomType& QueryGeom, const TRigidTransform<T, 3>& StartTM, const TVector<T, 3>& Dir, const T Length, T& OutTime, TVector<T, 3>& OutPosition, TVector<T, 3>& OutNormal, int32& OutFaceIndex, const T Thickness, bool bComputeMTD) const;
diff --git a/Engine/Source/Runtime/Experimental/Chaos/Public/Chaos/TriangleMeshImplicitObject.h b/Engine/Source/Runtime/Experimental/Chaos/Public/Chaos/TriangleMeshImplicitObject.h
--- a/Engine/Source/Runtime/Experimental/Chaos/Public/Chaos/TriangleMeshImplicitObject.h
+++ b/Engine/Source/Runtime/Experimental/Chaos/Public/Chaos/TriangleMeshImplicitObject.h
@@ -326,7 +326,7 @@ namespace Chaos
 		bool GJKContactPointImp(const QueryGeomType& QueryGeom, const FRigidTransform3& QueryTM, const FReal Thickness, FVec3& Location, FVec3& Normal, FReal& Penetration) const;
 
 		template<typename QueryGeomType>
-		bool OverlapGeomImp(const QueryGeomType& QueryGeom, const FRigidTransform3& QueryTM, const FReal Thickness, FVec3& OutNormal, FReal& OutPenetration, const bool bComputeMTD) const;
+		bool OverlapGeomImp(const QueryGeomType& QueryGeom, const FRigidTransform3& QueryTM, const FReal Thickness, FMTDInfo* OutMTD = nullptr) const;
 
 		template<typename QueryGeomType>
 		bool SweepGeomImp(const QueryGeomType& QueryGeom, const FRigidTransform3& StartTM, const FVec3& Dir, const FReal Length, FReal& OutTime, FVec3& OutPosition, FVec3& OutNormal, int32& OutFaceIndex, const FReal Thickness, const bool bComputeMTD) const;
