Improve cloth collision issues when characters are crouching and skydiving by adding an extra collision pass

diff --git a/Engine/Source/Runtime/Experimental/Chaos/Private/Chaos/PBDEvolution.cpp b/Engine/Source/Runtime/Experimental/Chaos/Private/Chaos/PBDEvolution.cpp
--- a/Engine/Source/Runtime/Experimental/Chaos/Private/Chaos/PBDEvolution.cpp
+++ b/Engine/Source/Runtime/Experimental/Chaos/Private/Chaos/PBDEvolution.cpp
@@ -126,6 +126,13 @@ void TPBDEvolution<T, d>::AdvanceOneTimeStep(const T Dt)
 		InitConstraintRule();  // Clear XPBD's Lambdas
 	}
 
+
+	// Do one extra collision pass at the start to decrease likelyhood of cloth penetrating- TODO: Add option for more collision passed interleaved between constraints
+	{
+		SCOPE_CYCLE_COUNTER(STAT_ChaosPBDCollisionRule);
+		CollisionRule.ApplyPerParticle(MParticles, Dt);
+	}
+
 	for (int i = 0; i < MNumIterations; ++i)
 	{
 		for (TFunction<void(TPBDParticles<T, d>&, const T)>& ConstraintRule : MConstraintRules)
