Contact modification callback

diff --git a/Engine/Source/Runtime/Engine/Private/PhysicsEngine/Experimental/PhysScene_Chaos.cpp b/Engine/Source/Runtime/Engine/Private/PhysicsEngine/Experimental/PhysScene_Chaos.cpp
--- a/Engine/Source/Runtime/Engine/Private/PhysicsEngine/Experimental/PhysScene_Chaos.cpp
+++ b/Engine/Source/Runtime/Engine/Private/PhysicsEngine/Experimental/PhysScene_Chaos.cpp
@@ -62,6 +62,8 @@ TAutoConsoleVariable<int32> CVar_ChaosSimulationEnable(TEXT("P.Chaos.Simulation.
 
 DEFINE_LOG_CATEGORY_STATIC(LogFPhysScene_ChaosSolver, Log, All);
 
+Chaos::TCollisionModifierCallback<float, 3> FPhysScene_ChaosInterface::CollisionModifierCallback;
+
 void DumpHierarchyStats(const TArray<FString>& Args)
 {
 #if !UE_BUILD_SHIPPING
@@ -1141,8 +1143,8 @@ FPhysScene_ChaosInterface::FPhysScene_ChaosInterface(const AWorldSettings* InSet
 	Scene.SetPhysicsReplication(PhysicsReplication);
 
 	Scene.GetSolver()->PhysSceneHack = this;
-	
 
+	Scene.GetSolver()->GetEvolution()->SetCollisionModifierCallback(CollisionModifierCallback);
 }
 
 void FPhysScene_ChaosInterface::OnWorldBeginPlay()
diff --git a/Engine/Source/Runtime/Engine/Public/Physics/Experimental/PhysScene_Chaos.h b/Engine/Source/Runtime/Engine/Public/Physics/Experimental/PhysScene_Chaos.h
--- a/Engine/Source/Runtime/Engine/Public/Physics/Experimental/PhysScene_Chaos.h
+++ b/Engine/Source/Runtime/Engine/Public/Physics/Experimental/PhysScene_Chaos.h
@@ -16,6 +16,7 @@
 #include "Chaos/Declares.h"
 #include "PhysicsProxy/SingleParticlePhysicsProxyFwd.h"
 #include "Framework/Threading.h"
+#include "Chaos/PBDCollisionConstraints.h"
 
 #ifndef CHAOS_WITH_PAUSABLE_SOLVER
 #define CHAOS_WITH_PAUSABLE_SOLVER 1
@@ -414,6 +415,8 @@ public:
 	ENGINE_API bool ExecPxVis(uint32 SceneType, const TCHAR* Cmd, FOutputDevice* Ar);
 	ENGINE_API bool ExecApexVis(uint32 SceneType, const TCHAR* Cmd, FOutputDevice* Ar);
 
+	ENGINE_API static Chaos::TCollisionModifierCallback<float, 3> CollisionModifierCallback;
+
 #if XGE_FIXED
 	template<typename PayloadType>
 	void RegisterEvent(const Chaos::EEventType& EventID, TFunction<void(const Chaos::FPBDRigidsSolver* Solver, PayloadType& EventData)> InLambda)
diff --git a/Engine/Source/Runtime/Experimental/Chaos/Private/Chaos/PBDCollisionConstraints.cpp b/Engine/Source/Runtime/Experimental/Chaos/Private/Chaos/PBDCollisionConstraints.cpp
--- a/Engine/Source/Runtime/Experimental/Chaos/Private/Chaos/PBDCollisionConstraints.cpp
+++ b/Engine/Source/Runtime/Experimental/Chaos/Private/Chaos/PBDCollisionConstraints.cpp
@@ -225,14 +225,17 @@ namespace Chaos
 	template<typename T, int d>
 	void TPBDCollisionConstraints<T, d>::ApplyCollisionModifier(const TFunction<ECollisionModifierResult(const FConstraintContainerHandle* Handle)>& CollisionModifier)
 	{
-		TArray<FConstraintContainerHandle*> CopyOfHandles = Handles;
-
-		for (FConstraintContainerHandle* ContactHandle : CopyOfHandles)
+		if (CollisionModifier)
 		{
-			ECollisionModifierResult Result = CollisionModifier(ContactHandle);
-			if (Result == ECollisionModifierResult::Disabled)
+			TArray<FConstraintContainerHandle*> CopyOfHandles = Handles;
+
+			for (FConstraintContainerHandle* ContactHandle : CopyOfHandles)
 			{
-				RemoveConstraint(ContactHandle);
+				ECollisionModifierResult Result = CollisionModifier(ContactHandle);
+				if (Result == ECollisionModifierResult::Disabled)
+				{
+					RemoveConstraint(ContactHandle);
+				}
 			}
 		}
 	}
diff --git a/Engine/Source/Runtime/Experimental/Chaos/Private/Chaos/PBDRigidsEvolutionGBF.cpp b/Engine/Source/Runtime/Experimental/Chaos/Private/Chaos/PBDRigidsEvolutionGBF.cpp
--- a/Engine/Source/Runtime/Experimental/Chaos/Private/Chaos/PBDRigidsEvolutionGBF.cpp
+++ b/Engine/Source/Runtime/Experimental/Chaos/Private/Chaos/PBDRigidsEvolutionGBF.cpp
@@ -174,6 +174,11 @@ void TPBDRigidsEvolutionGBF<T, d>::AdvanceOneTimeStep(const T Dt, const T StepFr
 		CHAOS_COLLISION_STAT(StatData.Print());
 	}
 
+	if (CollisionModifierCallback)
+	{
+		CollisionConstraints.ApplyCollisionModifier(CollisionModifierCallback);
+	}
+
 	if (PostDetectCollisionsCallback != nullptr)
 	{
 		PostDetectCollisionsCallback();
diff --git a/Engine/Source/Runtime/Experimental/Chaos/Public/Chaos/CollisionResolutionTypes.h b/Engine/Source/Runtime/Experimental/Chaos/Public/Chaos/CollisionResolutionTypes.h
--- a/Engine/Source/Runtime/Experimental/Chaos/Public/Chaos/CollisionResolutionTypes.h
+++ b/Engine/Source/Runtime/Experimental/Chaos/Public/Chaos/CollisionResolutionTypes.h
@@ -50,7 +50,7 @@ namespace Chaos
 	*
 	*/
 	template<class T, int d>
-	class CHAOS_API TCollisionContact
+	class TCollisionContact
 	{
 	public:
 		TCollisionContact(const FImplicitObject* InImplicit0 = nullptr, const FImplicitObject* InImplicit1 = nullptr)
@@ -85,7 +85,7 @@ namespace Chaos
 	*
 	*/
 	template<class T = float, int d = 3>
-	class CHAOS_API TCollisionConstraintBase
+	class TCollisionConstraintBase
 	{
 	public:
 		using FGeometryParticleHandle = TGeometryParticleHandle<T, d>;
@@ -131,6 +131,7 @@ namespace Chaos
 
 		bool ContainsManifold(const FImplicitObject* A, const FImplicitObject* B) const { return A == Manifold.Implicit[0] && B == Manifold.Implicit[1]; }
 		void SetManifold(const FImplicitObject* A, const FImplicitObject* B) { Manifold.Implicit[0] = A; Manifold.Implicit[1] = B; }
+		const FManifold& GetManifold() const { return Manifold; }
 
 		//
 		// API
@@ -175,7 +176,7 @@ namespace Chaos
 	*
 	*/
 	template<class T, int d>
-	class CHAOS_API TRigidBodyPointContactConstraint : public TCollisionConstraintBase<T, d>
+	class TRigidBodyPointContactConstraint : public TCollisionConstraintBase<T, d>
 	{
 	public:
 		using Base = TCollisionConstraintBase<T, d>;
@@ -198,7 +199,7 @@ namespace Chaos
 	*
 	*/
 	template<class T, int d>
-	class CHAOS_API TRigidBodyMultiPointContactConstraint : public TCollisionConstraintBase<T, d>
+	class TRigidBodyMultiPointContactConstraint : public TCollisionConstraintBase<T, d>
 	{
 	public:
 		using Base = TCollisionConstraintBase<T, d>;
@@ -273,7 +274,7 @@ namespace Chaos
 
 
 	template<class T, int d>
-	class CHAOS_API TPBDCollisionConstraintHandle : public TContainerConstraintHandle<TPBDCollisionConstraints<T, d>>
+	class TPBDCollisionConstraintHandle : public TContainerConstraintHandle<TPBDCollisionConstraints<T, d>>
 	{
 	public:
 		using Base = TContainerConstraintHandle<TPBDCollisionConstraints<T, d>>;
@@ -418,4 +419,19 @@ namespace Chaos
 	};
 
 	using FCollisionConstraintsArray = TCollisionConstraintsStore<8>;
+
+
+#if PLATFORM_MAC || PLATFORM_LINUX
+	extern template class CHAOS_API TCollisionContact<float, 3>;
+	extern template class CHAOS_API TCollisionConstraintBase<float, 3>;
+	extern template class CHAOS_API TRigidBodyPointContactConstraint<float, 3>;
+	extern template class CHAOS_API TRigidBodyMultiPointContactConstraint<float, 3>;
+	extern template class CHAOS_API TPBDCollisionConstraintHandle<float, 3>;
+#else
+	extern template class TCollisionContact<float, 3>;
+	extern template class TCollisionConstraintBase<float, 3>;
+	extern template class TRigidBodyPointContactConstraint<float, 3>;
+	extern template class TRigidBodyMultiPointContactConstraint<float, 3>;
+	extern template class TPBDCollisionConstraintHandle<float, 3>;
+#endif
 }
diff --git a/Engine/Source/Runtime/Experimental/Chaos/Public/Chaos/PBDCollisionConstraints.h b/Engine/Source/Runtime/Experimental/Chaos/Public/Chaos/PBDCollisionConstraints.h
--- a/Engine/Source/Runtime/Experimental/Chaos/Public/Chaos/PBDCollisionConstraints.h
+++ b/Engine/Source/Runtime/Experimental/Chaos/Public/Chaos/PBDCollisionConstraints.h
@@ -4,6 +4,7 @@
 #include "Chaos/ConstraintHandle.h"
 #include "Chaos/CollisionResolutionTypes.h"
 #include "Chaos/PBDConstraintContainer.h"
+#include "Chaos/CollisionResolutionTypes.h"
 #include "Framework/BufferedData.h"
 
 #include <memory>
@@ -20,6 +21,18 @@ namespace Chaos
 template<typename T, int d>
 class TPBDCollisionConstraints;
 
+template <typename T, int d>
+class TPBDCollisionConstraintHandle;
+
+template <typename T, int d>
+class TCollisionConstraintBase;
+
+template <typename T, int d>
+class TRigidBodyPointContactConstraint;
+
+template <typename T, int d>
+class TRigidBodyMultiPointContactConstraint;
+
 template <typename T, int d>
 class TRigidTransform;
 
@@ -46,6 +59,9 @@ using TRigidBodyContactConstraintsPostApplyCallback = TFunction<void(const T Dt,
 template<typename T, int d>
 using TRigidBodyContactConstraintsPostApplyPushOutCallback = TFunction<void(const T Dt, const TArray<TPBDCollisionConstraintHandle<T, d>*>&, bool)>;
 
+template <typename T, int d>
+using TCollisionModifierCallback = TFunction<ECollisionModifierResult(const TPBDCollisionConstraintHandle<T, d>*)>;
+
 /**
  * A container and solver for collision constraints.
  */
@@ -63,6 +79,7 @@ public:
 	using FMultiPointContactConstraint = TRigidBodyMultiPointContactConstraint<T, d>;
 	using FConstraintHandleAllocator = TConstraintHandleAllocator<TPBDCollisionConstraints<T, d>>;
 	using FConstraintContainerHandleKey = typename TPBDCollisionConstraintHandle<T, d>::FHandleKey;
+	using FCollisionModifier = TCollisionModifierCallback<T, d>;
 
 	TPBDCollisionConstraints(const TPBDRigidsSOAs<T,d>& InParticles, 
 		TArrayCollectionArray<bool>& Collided, 
@@ -93,7 +110,7 @@ public:
 	 * You would probably call this in the PostComputeCallback. Prefer this to calling RemoveConstraints in a loop,
 	 * so you don't have to worry about constraint iterator/indices changing.
 	 */
-	void ApplyCollisionModifier(const TFunction<ECollisionModifierResult(const FConstraintContainerHandle* Handle)>& CollisionModifier);
+	void ApplyCollisionModifier(const FCollisionModifier& CollisionModifier);
 
 
 	/**
diff --git a/Engine/Source/Runtime/Experimental/Chaos/Public/Chaos/PBDRigidsEvolutionGBF.h b/Engine/Source/Runtime/Experimental/Chaos/Public/Chaos/PBDRigidsEvolutionGBF.h
--- a/Engine/Source/Runtime/Experimental/Chaos/Public/Chaos/PBDRigidsEvolutionGBF.h
+++ b/Engine/Source/Runtime/Experimental/Chaos/Public/Chaos/PBDRigidsEvolutionGBF.h
@@ -75,6 +75,11 @@ namespace Chaos
 			PostDetectCollisionsCallback = Cb;
 		}
 
+		void SetCollisionModifierCallback(const TCollisionModifierCallback<T, d>& Cb)
+		{
+			CollisionModifierCallback = Cb;
+		}
+
 		void SetPreApplyCallback(const TPBDRigidsEvolutionCallback<T, d>& Cb)
 		{
 			PreApplyCallback = Cb;
@@ -207,6 +212,7 @@ namespace Chaos
 
 		TPBDRigidsEvolutionCallback<T, d> PostIntegrateCallback;
 		TPBDRigidsEvolutionCallback<T, d> PostDetectCollisionsCallback;
+		TCollisionModifierCallback<T, d> CollisionModifierCallback;
 		TPBDRigidsEvolutionCallback<T, d> PreApplyCallback;
 		TPBDRigidsEvolutionIslandCallback<T, d> PostApplyCallback;
 		TPBDRigidsEvolutionIslandCallback<T, d> PostApplyPushOutCallback;
