Fix TCollisionContact::ToString for Linux, attempt #2

diff --git a/Engine/Source/Runtime/Experimental/Chaos/Private/Chaos/CollisionResolutionTypes.cpp b/Engine/Source/Runtime/Experimental/Chaos/Private/Chaos/CollisionResolutionTypes.cpp
--- a/Engine/Source/Runtime/Experimental/Chaos/Private/Chaos/CollisionResolutionTypes.cpp
+++ b/Engine/Source/Runtime/Experimental/Chaos/Private/Chaos/CollisionResolutionTypes.cpp
@@ -6,4 +6,9 @@
 namespace Chaos
 {
 	template class CHAOS_API TPBDCollisionConstraintHandle<float, 3>; 
+
+	FString TCollisionContact<float, 3>::ToString() const
+	{
+		return FString::Printf(TEXT("Location:%s, Normal:%s, Phi:%f"), *Location.ToString(), *Normal.ToString(), Phi);
+	}
 };
diff --git a/Engine/Source/Runtime/Experimental/Chaos/Public/Chaos/CollisionResolutionTypes.h b/Engine/Source/Runtime/Experimental/Chaos/Public/Chaos/CollisionResolutionTypes.h
--- a/Engine/Source/Runtime/Experimental/Chaos/Public/Chaos/CollisionResolutionTypes.h
+++ b/Engine/Source/Runtime/Experimental/Chaos/Public/Chaos/CollisionResolutionTypes.h
@@ -71,11 +71,7 @@ namespace Chaos
 
 		EContactShapesType ShapesType;
 
-
-		FString CHAOS_API ToString() const
-		{
-			return FString::Printf(TEXT("Location:%s, Normal:%s, Phi:%f"), *Location.ToString(), *Normal.ToString(), Phi);
-		}
+		FString ToString() const;
 
 		const FImplicitObject* Implicit[2]; // {Of Particle[0], Of Particle[1]}
 	};
