Fix HeightField GetHeightAt expecting world scale position instead of local scale position.
Fixes strange water VFX behaviour.

diff --git a/Engine/Source/Runtime/Experimental/Chaos/Private/Chaos/HeightField.cpp b/Engine/Source/Runtime/Experimental/Chaos/Private/Chaos/HeightField.cpp
--- a/Engine/Source/Runtime/Experimental/Chaos/Private/Chaos/HeightField.cpp
+++ b/Engine/Source/Runtime/Experimental/Chaos/Private/Chaos/HeightField.cpp
@@ -557,21 +557,18 @@ namespace Chaos
 	}
 
 	template<typename T>
-	T Chaos::THeightField<T>::GetHeightAt(const TVector<T, 2>& InGridLocation) const
+	T Chaos::THeightField<T>::GetHeightAt(const TVector<T, 2>& InGridLocationLocal) const
 	{
-		const TVector<T, 2> Scale2D(GeomData.Scale[0], GeomData.Scale[1]);
-		const TVector<T, 2> UnscaledLocation = InGridLocation / Scale2D;
-
-		if(UnscaledLocation == FlatGrid.Clamp(UnscaledLocation))
+		if(CHAOS_ENSURE(InGridLocationLocal == FlatGrid.Clamp(InGridLocationLocal)))
 		{
-			TVector<int32, 2> CellCoord = FlatGrid.Cell(UnscaledLocation);
+			TVector<int32, 2> CellCoord = FlatGrid.Cell(InGridLocationLocal);
 
 			const int32 SingleIndex = CellCoord[1] * (GeomData.NumCols) + CellCoord[0];
 			TVector<T, 3> Pts[4];
 			GeomData.GetPoints(SingleIndex, Pts);
 
-			float FractionX = FMath::Frac(UnscaledLocation[0]);
-			float FractionY = FMath::Frac(UnscaledLocation[1]);
+			float FractionX = FMath::Frac(InGridLocationLocal[0]);
+			float FractionY = FMath::Frac(InGridLocationLocal[1]);
 
 			if(FractionX > FractionY)
 			{
diff --git a/Engine/Source/Runtime/Experimental/Chaos/Public/Chaos/HeightField.h b/Engine/Source/Runtime/Experimental/Chaos/Public/Chaos/HeightField.h
--- a/Engine/Source/Runtime/Experimental/Chaos/Public/Chaos/HeightField.h
+++ b/Engine/Source/Runtime/Experimental/Chaos/Public/Chaos/HeightField.h
@@ -41,7 +41,8 @@ namespace Chaos
 		void EditHeights(TArrayView<const uint16> InHeights, int32 InBeginRow, int32 InBeginCol, int32 InNumRows, int32 InNumCols);
 		T GetHeight(int32 InIndex) const;
 		T GetHeight(int32 InX, int32 InY) const;
-		T GetHeightAt(const TVector<T, 2>& InGridLocation) const;
+		T GetHeightAt(const TVector<T, 2>& InGridLocationLocal) const;
+
 		int32 GetNumRows() const { return GeomData.NumRows; }
 		int32 GetNumCols() const { return GeomData.NumCols; }
 
