#UE4 CIS fix

diff --git a/Engine/Source/Runtime/Experimental/Chaos/Private/Chaos/CollisionResolutionTypes.cpp b/Engine/Source/Runtime/Experimental/Chaos/Private/Chaos/CollisionResolutionTypes.cpp
--- a/Engine/Source/Runtime/Experimental/Chaos/Private/Chaos/CollisionResolutionTypes.cpp
+++ b/Engine/Source/Runtime/Experimental/Chaos/Private/Chaos/CollisionResolutionTypes.cpp
@@ -6,10 +6,5 @@
 namespace Chaos
 {
 	template class CHAOS_API TPBDCollisionConstraintHandle<float, 3>; 
-
-	template<>
-	FString TCollisionContact<float, 3>::ToString() const
-	{
-		return FString::Printf(TEXT("Location:%s, Normal:%s, Phi:%f"), *Location.ToString(), *Normal.ToString(), Phi);
-	}
+	template class CHAOS_API TCollisionContact<float, 3>;
 };
diff --git a/Engine/Source/Runtime/Experimental/Chaos/Public/Chaos/CollisionResolutionTypes.h b/Engine/Source/Runtime/Experimental/Chaos/Public/Chaos/CollisionResolutionTypes.h
--- a/Engine/Source/Runtime/Experimental/Chaos/Public/Chaos/CollisionResolutionTypes.h
+++ b/Engine/Source/Runtime/Experimental/Chaos/Public/Chaos/CollisionResolutionTypes.h
@@ -71,7 +71,11 @@ namespace Chaos
 
 		EContactShapesType ShapesType;
 
-		FString ToString() const;
+
+		FString ToString() const
+		{
+			return FString::Printf(TEXT("Location:%s, Normal:%s, Phi:%f"), *Location.ToString(), *Normal.ToString(), Phi);
+		}
 
 		const FImplicitObject* Implicit[2]; // {Of Particle[0], Of Particle[1]}
 	};
