More linux linker fixes. Remove bad explicit instantiation externs. Fix return type.

diff --git a/Engine/Source/Runtime/Experimental/Chaos/Public/Chaos/CollisionResolutionTypes.h b/Engine/Source/Runtime/Experimental/Chaos/Public/Chaos/CollisionResolutionTypes.h
--- a/Engine/Source/Runtime/Experimental/Chaos/Public/Chaos/CollisionResolutionTypes.h
+++ b/Engine/Source/Runtime/Experimental/Chaos/Public/Chaos/CollisionResolutionTypes.h
@@ -143,7 +143,7 @@ namespace Chaos
 		T GetPhi() const { return Manifold.Phi; }
 
 		void SetDisabled(bool bInDisabled) { Manifold.bDisabled = bInDisabled; }
-		TVector<T, d> GetDisabled() const { return Manifold.bDisabled; }
+		bool GetDisabled() const { return Manifold.bDisabled; }
 
 		void SetNormal(const TVector<T, d> & InNormal) { Manifold.Normal = InNormal; }
 		TVector<T, d> GetNormal() const { return Manifold.Normal; }
@@ -423,15 +423,9 @@ namespace Chaos
 
 #if PLATFORM_MAC || PLATFORM_LINUX
 	extern template class CHAOS_API TCollisionContact<float, 3>;
-	extern template class CHAOS_API TCollisionConstraintBase<float, 3>;
-	extern template class CHAOS_API TRigidBodyPointContactConstraint<float, 3>;
-	extern template class CHAOS_API TRigidBodyMultiPointContactConstraint<float, 3>;
 	extern template class CHAOS_API TPBDCollisionConstraintHandle<float, 3>;
 #else
 	extern template class TCollisionContact<float, 3>;
-	extern template class TCollisionConstraintBase<float, 3>;
-	extern template class TRigidBodyPointContactConstraint<float, 3>;
-	extern template class TRigidBodyMultiPointContactConstraint<float, 3>;
 	extern template class TPBDCollisionConstraintHandle<float, 3>;
 #endif
 }
