#UE4 Fix copyright

diff --git a/Engine/Source/Runtime/Experimental/Chaos/Public/Chaos/PerGroupDampVelocity.h b/Engine/Source/Runtime/Experimental/Chaos/Public/Chaos/PerGroupDampVelocity.h
--- a/Engine/Source/Runtime/Experimental/Chaos/Public/Chaos/PerGroupDampVelocity.h
+++ b/Engine/Source/Runtime/Experimental/Chaos/Public/Chaos/PerGroupDampVelocity.h
@@ -1,4 +1,4 @@
-// Copyright 1998-2019 Epic Games, Inc. All Rights Reserved.
+// Copyright Epic Games, Inc. All Rights Reserved.
 #pragma once
 
 #include "Chaos/PerParticleDampVelocity.h"
