Chaos : Shape Management

- Append child and parent implicit objects to the shape map
- Centralize  the casting utilities.

diff --git a/Engine/Source/Runtime/Engine/Private/PhysicsEngine/Experimental/PhysInterface_Chaos.cpp b/Engine/Source/Runtime/Engine/Private/PhysicsEngine/Experimental/PhysInterface_Chaos.cpp
--- a/Engine/Source/Runtime/Engine/Private/PhysicsEngine/Experimental/PhysInterface_Chaos.cpp
+++ b/Engine/Source/Runtime/Engine/Private/PhysicsEngine/Experimental/PhysInterface_Chaos.cpp
@@ -1599,7 +1599,7 @@ bool CalculateMassPropertiesOfImplicitType(
 		else
 		{
 			//question: is LocalTM enough to merge these two branches?
-			Chaos::CastHelper(*ImplicitObject, FTransform::Identity, [&OutMassProperties, InDensityKGPerCM](const auto& Object, const auto& LocalTM)
+			Chaos::Utilities::CastHelper(*ImplicitObject, FTransform::Identity, [&OutMassProperties, InDensityKGPerCM](const auto& Object, const auto& LocalTM)
 			{
 				OutMassProperties.Volume = Object.GetVolume();
 				OutMassProperties.Mass = OutMassProperties.Volume * InDensityKGPerCM;
@@ -1821,7 +1821,7 @@ bool FPhysInterface_Chaos::Sweep_Geom(FHitResult& OutHit, const FBodyInstance* I
 							Chaos::TVector<float, 3> WorldPosition;
 							Chaos::TVector<float, 3> WorldNormal;
 							int32 FaceIdx;
-							if (Chaos::CastHelper(ShapeAdapter.GetGeometry(), ActorTM, [&](const auto& Downcast, const auto& FullActorTM) { return Chaos::SweepQuery(*Shape->Geometry, FullActorTM, Downcast, StartTM, Dir, DeltaMag, Hit.Distance, WorldPosition, WorldNormal, FaceIdx, 0.f, false); }))
+							if (Chaos::Utilities::CastHelper(ShapeAdapter.GetGeometry(), ActorTM, [&](const auto& Downcast, const auto& FullActorTM) { return Chaos::SweepQuery(*Shape->Geometry, FullActorTM, Downcast, StartTM, Dir, DeltaMag, Hit.Distance, WorldPosition, WorldNormal, FaceIdx, 0.f, false); }))
 							{
 								// we just like to make sure if the hit is made
 								FCollisionFilterData QueryFilter;
@@ -1881,7 +1881,7 @@ bool Overlap_GeomInternal(const FBodyInstance* InInstance, const Chaos::FImplici
 			if (OutOptResult)
 			{
 				Chaos::FMTDInfo MTDInfo;
-				if (Chaos::CastHelper(InGeom, ActorTM, [&](const auto& Downcast, const auto& FullActorTM) { return Chaos::OverlapQuery(*Shape->Geometry, FullActorTM, Downcast, GeomTransform, /*Thickness=*/0, &MTDInfo); }))
+				if (Chaos::Utilities::CastHelper(InGeom, ActorTM, [&](const auto& Downcast, const auto& FullActorTM) { return Chaos::OverlapQuery(*Shape->Geometry, FullActorTM, Downcast, GeomTransform, /*Thickness=*/0, &MTDInfo); }))
 				{
 					OutOptResult->Distance = MTDInfo.Penetration;
 					OutOptResult->Direction = MTDInfo.Normal;
@@ -1890,7 +1890,7 @@ bool Overlap_GeomInternal(const FBodyInstance* InInstance, const Chaos::FImplici
 			}
 			else	//question: why do we even allow user to not pass in MTD info?
 			{
-				if (Chaos::CastHelper(InGeom, ActorTM, [&](const auto& Downcast, const auto& FullActorTM) { return Chaos::OverlapQuery(*Shape->Geometry, FullActorTM, Downcast, GeomTransform); }))
+				if (Chaos::Utilities::CastHelper(InGeom, ActorTM, [&](const auto& Downcast, const auto& FullActorTM) { return Chaos::OverlapQuery(*Shape->Geometry, FullActorTM, Downcast, GeomTransform); }))
 				{
 					return true;
 				}
diff --git a/Engine/Source/Runtime/Experimental/Chaos/Private/Chaos/CollisionResolution.cpp b/Engine/Source/Runtime/Experimental/Chaos/Private/Chaos/CollisionResolution.cpp
--- a/Engine/Source/Runtime/Experimental/Chaos/Private/Chaos/CollisionResolution.cpp
+++ b/Engine/Source/Runtime/Experimental/Chaos/Private/Chaos/CollisionResolution.cpp
@@ -2,6 +2,7 @@
 
 #include "Chaos/CollisionResolution.h"
 
+#include "Chaos/CastingUtilities.h"
 #include "Chaos/ChaosPerfTest.h"
 #include "Chaos/Capsule.h"
 #include "Chaos/CollisionResolutionTypes.h"
@@ -184,9 +185,9 @@ namespace Chaos
 		template<class T, int d>
 		TContactPoint<T> ConvexConvexContactPoint(const FImplicitObject& A, const TRigidTransform<T, d>& ATM, const FImplicitObject& B, const TRigidTransform<T, d>& BTM, const T CullDistance)
 		{
-			return CastHelper(A, ATM, [&](const auto& ADowncast, const TRigidTransform<T,d>& AFullTM)
+			return Utilities::CastHelper(A, ATM, [&](const auto& ADowncast, const TRigidTransform<T,d>& AFullTM)
 			{
-				return CastHelper(B, BTM, [&](const auto& BDowncast, const TRigidTransform<T,d>& BFullTM)
+				return Utilities::CastHelper(B, BTM, [&](const auto& BDowncast, const TRigidTransform<T,d>& BFullTM)
 				{
 					return GJKContactPoint(ADowncast, AFullTM, BDowncast, BFullTM, TVector<T, d>(1, 0, 0));
 				});
diff --git a/Engine/Source/Runtime/Experimental/Chaos/Private/Chaos/GeometryParticles.cpp b/Engine/Source/Runtime/Experimental/Chaos/Private/Chaos/GeometryParticles.cpp
--- a/Engine/Source/Runtime/Experimental/Chaos/Private/Chaos/GeometryParticles.cpp
+++ b/Engine/Source/Runtime/Experimental/Chaos/Private/Chaos/GeometryParticles.cpp
@@ -1,6 +1,8 @@
 // Copyright Epic Games, Inc. All Rights Reserved.
 
 #include "Chaos/GeometryParticles.h"
+
+#include "Chaos/CastingUtilities.h"
 #include "Chaos/ImplicitObject.h"
 #include "Chaos/ImplicitObjectUnion.h"
 #include "Chaos/ParticleHandle.h"
@@ -176,6 +178,78 @@ namespace Chaos
 		Ar << SerializableGeometryParticles;
 	}
 
+	template <typename T, int d, EGeometryParticlesSimType SimType>
+	void TGeometryParticlesImp<T, d, SimType>::MapImplicitShapes()
+	{
+		int32 NumShapeArrays = MShapesArray.Num();
+		ImplicitShapeMap.Resize(NumShapeArrays);
+		for (int32 Index = 0; Index < NumShapeArrays; ++Index)
+		{
+			MapImplicitShapes(Index);
+		}
+	}
+
+	template <typename T, int d, EGeometryParticlesSimType SimType>
+	void TGeometryParticlesImp<T, d, SimType>::MapImplicitShapes(int32 Index)
+	{
+		checkSlow(Index >= 0 && Index < ImplicitShapeMap.Num());
+
+		TMap<const FImplicitObject*, int32>& Mapping = ImplicitShapeMap[Index];
+		TShapesArray<T, d>& ShapeArray = MShapesArray[Index];
+		Mapping.Reset();
+
+		int32 ShapeIndex = 0;
+		for (TUniquePtr<TPerShapeData<T, d>>& Shape : ShapeArray)
+		{
+			Mapping.Add(Shape->Geometry.Get(), ShapeIndex++);
+		}
+
+		if (MGeometry[Index])
+		{
+			int32 CurrentShapeIndex = INDEX_NONE;
+
+			if (const auto* Union = MGeometry[Index]->template GetObject<FImplicitObjectUnion>())
+			{
+				for (const TUniquePtr<FImplicitObject>& ImplicitObject : Union->GetObjects())
+				{
+					if (ImplicitObject.Get())
+					{
+						if (const FImplicitObject* ImplicitChildObject = Utilities::ImplicitChildHelper(ImplicitObject.Get()))
+						{
+							if (ImplicitShapeMap[Index].Contains(ImplicitObject.Get()))
+							{
+								ImplicitShapeMap[Index].Add(ImplicitChildObject, ImplicitShapeMap[Index][ImplicitObject.Get()]);
+							}
+							else if (ImplicitShapeMap[Index].Contains(ImplicitChildObject))
+							{
+								ImplicitShapeMap[Index].Add(ImplicitObject.Get(), ImplicitShapeMap[Index][ImplicitChildObject]);
+							}
+							
+						}
+					}
+				}
+			}
+			else
+			{
+				if (const FImplicitObject* ImplicitChildObject = Utilities::ImplicitChildHelper(MGeometry[Index].Get()))
+				{
+					if (ImplicitShapeMap[Index].Contains(MGeometry[Index].Get()))
+					{
+						ImplicitShapeMap[Index].Add(ImplicitChildObject, ImplicitShapeMap[Index][MGeometry[Index].Get()]);
+					}
+					else if (ImplicitShapeMap[Index].Contains(ImplicitChildObject))
+					{
+						ImplicitShapeMap[Index].Add(MGeometry[Index].Get(), ImplicitShapeMap[Index][ImplicitChildObject]);
+					}
+					
+				}
+			}
+
+		}
+	}
+
+
+
 	
 	template class TGeometryParticlesImp<float, 3, EGeometryParticlesSimType::RigidBodySim>;
 	template class TGeometryParticlesImp<float, 3, EGeometryParticlesSimType::Other>;
diff --git a/Engine/Source/Runtime/Experimental/Chaos/Public/Chaos/CastingUtilities.h b/Engine/Source/Runtime/Experimental/Chaos/Public/Chaos/CastingUtilities.h
new file mode 100644
--- /dev/null
+++ b/Engine/Source/Runtime/Experimental/Chaos/Public/Chaos/CastingUtilities.h
@@ -0,0 +1,114 @@
+// Copyright 1998-2019 Epic Games, Inc. All Rights Reserved.
+#pragma once
+
+#include "Chaos/ImplicitObject.h"
+
+#include "Chaos/ImplicitObjectScaled.h"
+#include "Chaos/ImplicitObjectTransformed.h"
+#include "Chaos/Sphere.h"
+#include "Chaos/Capsule.h"
+#include "Chaos/Convex.h"
+//#include "Chaos/TriangleMeshImplicitObject.h"
+#include "Chaos/Box.h"
+
+namespace Chaos
+{
+	namespace Utilities
+	{
+		template <typename Lambda>
+		FORCEINLINE_DEBUGGABLE auto CastHelper(const FImplicitObject& Geom, const FRigidTransform3& TM, const Lambda& Func)
+		{
+			const EImplicitObjectType Type = Geom.GetType(true);
+			switch (Type)
+			{
+			case ImplicitObjectType::Sphere: return Func(Geom.template GetObjectChecked<TSphere<FReal, 3>>(), TM);
+			case ImplicitObjectType::Box: return Func(Geom.template GetObjectChecked<TBox<FReal, 3>>(), TM);
+			case ImplicitObjectType::Capsule: return Func(Geom.template GetObjectChecked<TCapsule<FReal>>(), TM);
+			case ImplicitObjectType::Convex: return Func(Geom.template GetObjectChecked<FConvex>(), TM);
+			case ImplicitObjectType::IsScaled | ImplicitObjectType::Sphere: return Func(Geom.template GetObjectChecked<TImplicitObjectScaled<TSphere<FReal, 3>>>(), TM);
+			case ImplicitObjectType::IsScaled | ImplicitObjectType::Box: return Func(Geom.template GetObjectChecked< TImplicitObjectScaled<TBox<FReal, 3>>>(), TM);
+			case ImplicitObjectType::IsScaled | ImplicitObjectType::Capsule: return Func(Geom.template GetObjectChecked< TImplicitObjectScaled<TCapsule<FReal>>>(), TM);
+			case ImplicitObjectType::IsScaled | ImplicitObjectType::Convex: return Func(Geom.template GetObjectChecked< TImplicitObjectScaled<FConvex>>(), TM);
+			case ImplicitObjectType::IsInstanced | ImplicitObjectType::Sphere:
+			{
+				return Func(Geom.template GetObjectChecked< TImplicitObjectInstanced<TSphere<FReal, 3>>>().GetInstancedObject()->template GetObjectChecked<TSphere<FReal, 3>>(), TM);
+			}
+			case ImplicitObjectType::IsInstanced | ImplicitObjectType::Box:
+			{
+				return Func(Geom.template GetObjectChecked< TImplicitObjectInstanced<TBox<FReal, 3>>>().GetInstancedObject()->template GetObjectChecked<TBox<FReal, 3>>(), TM);
+			}
+			case ImplicitObjectType::IsInstanced | ImplicitObjectType::Capsule:
+			{
+				return Func(Geom.template GetObjectChecked< TImplicitObjectInstanced<TCapsule<FReal>>>().GetInstancedObject()->template GetObjectChecked<TCapsule<FReal>>(), TM);
+			}
+			case ImplicitObjectType::IsInstanced | ImplicitObjectType::Convex:
+			{
+				return Func(Geom.template GetObjectChecked< TImplicitObjectInstanced<FConvex>>().GetInstancedObject()->template GetObjectChecked<FConvex>(), TM);
+			}
+			case ImplicitObjectType::Transformed:
+			{
+				const auto& ImplicitObjectTransformed = (Geom.template GetObjectChecked<TImplicitObjectTransformed<FReal, 3>>());
+				return CastHelper(*ImplicitObjectTransformed.GetTransformedObject(), ImplicitObjectTransformed.GetTransform() * TM, Func);
+			}
+
+			default: check(false);
+			}
+			return Func(Geom.template GetObjectChecked<TSphere<FReal, 3>>(), TM);	//needed for return type
+		}
+
+		inline
+		const FImplicitObject* ImplicitChildHelper(const FImplicitObject* ImplicitObject)
+		{
+			EImplicitObjectType ImplicitType = ImplicitObject->GetType();
+			if (ImplicitType == TImplicitObjectTransformed<FReal, 3>::StaticType())
+			{
+				return ImplicitObject->template GetObject<const TImplicitObjectTransformed<FReal, 3>>()->GetTransformedObject();
+			}
+
+			else if (ImplicitType == TImplicitObjectScaled<FConvex>::StaticType())
+			{
+				return ImplicitObject->template GetObject<const TImplicitObjectScaled<FConvex>>()->GetUnscaledObject();
+			}
+			else if (ImplicitType == TImplicitObjectScaled<TBox<FReal, 3>>::StaticType())
+			{
+				return ImplicitObject->template GetObject<const TImplicitObjectScaled<TBox<FReal, 3>>>()->GetUnscaledObject();
+			}
+			else if (ImplicitType == TImplicitObjectScaled<TCapsule<FReal>>::StaticType())
+			{
+				return ImplicitObject->template GetObject<const TImplicitObjectScaled<TCapsule<FReal>>>()->GetUnscaledObject();
+			}
+			else if (ImplicitType == TImplicitObjectScaled<TSphere<FReal, 3>>::StaticType())
+			{
+				return ImplicitObject->template GetObject<const TImplicitObjectScaled<TSphere<FReal, 3>>>()->GetUnscaledObject();
+			}
+			//else if (ImplicitType == TImplicitObjectScaled<FTriangleMeshImplicitObject>::StaticType())
+			//{
+			//	return ImplicitObject->template GetObject<const TImplicitObjectScaled<FTriangleMeshImplicitObject>>()->GetUnscaledObject();
+			//}
+
+			else if (ImplicitType == TImplicitObjectInstanced<FConvex>::StaticType())
+			{
+				return ImplicitObject->template GetObject<const TImplicitObjectInstanced<FConvex>>()->GetInstancedObject();
+			}
+			else if (ImplicitType == TImplicitObjectInstanced<TBox<FReal, 3>>::StaticType())
+			{
+				return ImplicitObject->template GetObject<const TImplicitObjectInstanced<TBox<FReal, 3>>>()->GetInstancedObject();
+			}
+			else if (ImplicitType == TImplicitObjectInstanced<TCapsule<FReal>>::StaticType())
+			{
+				return ImplicitObject->template GetObject<const TImplicitObjectInstanced<TCapsule<FReal>>>()->GetInstancedObject();
+			}
+			else if (ImplicitType == TImplicitObjectInstanced<TSphere<FReal, 3>>::StaticType())
+			{
+				return ImplicitObject->template GetObject<const TImplicitObjectInstanced<TSphere<FReal, 3>>>()->GetInstancedObject();
+			}
+			//else if (ImplicitType == TImplicitObjectInstanced<FTriangleMeshImplicitObject>::StaticType())
+			//{
+			//	return ImplicitObject->template GetObject<const TImplicitObjectInstanced<FTriangleMeshImplicitObject>>()->GetInstancedObject();
+			//}
+			return nullptr;
+		}
+
+	} // namespace Utilities
+
+} // namespace Chaos
diff --git a/Engine/Source/Runtime/Experimental/Chaos/Public/Chaos/GeometryParticles.h b/Engine/Source/Runtime/Experimental/Chaos/Public/Chaos/GeometryParticles.h
--- a/Engine/Source/Runtime/Experimental/Chaos/Public/Chaos/GeometryParticles.h
+++ b/Engine/Source/Runtime/Experimental/Chaos/Public/Chaos/GeometryParticles.h
@@ -461,30 +461,9 @@ namespace Chaos
 			MapImplicitShapes(Index);
 		}
 
-		void MapImplicitShapes()
-		{
-			int32 NumShapeArrays = MShapesArray.Num();
-			ImplicitShapeMap.Resize(NumShapeArrays);
-			for(int32 Index = 0; Index < NumShapeArrays; ++Index)
-			{
-				MapImplicitShapes(Index);
-			}
-		}
-
-		void MapImplicitShapes(int32 Index)
-		{
-			checkSlow(Index >= 0 && Index < ImplicitShapeMap.Num());
+		void MapImplicitShapes();
 
-			TMap<const FImplicitObject*, int32>& Mapping = ImplicitShapeMap[Index];
-			TShapesArray<T, d>& ShapeArray = MShapesArray[Index];
-			Mapping.Reset();
-
-			int32 ShapeIndex = 0;
-			for(TUniquePtr<TPerShapeData<T, d>>& Shape : ShapeArray)
-			{
-				Mapping.Add(Shape->Geometry.Get(), ShapeIndex++);
-			}
-		}
+		void MapImplicitShapes(int32 Index);
 
 		template <typename T2, int d2, EGeometryParticlesSimType SimType2>
 		friend class TGeometryParticlesImp;
diff --git a/Engine/Source/Runtime/Experimental/Chaos/Public/Chaos/GeometryQueries.h b/Engine/Source/Runtime/Experimental/Chaos/Public/Chaos/GeometryQueries.h
--- a/Engine/Source/Runtime/Experimental/Chaos/Public/Chaos/GeometryQueries.h
+++ b/Engine/Source/Runtime/Experimental/Chaos/Public/Chaos/GeometryQueries.h
@@ -1,6 +1,7 @@
 // Copyright Epic Games, Inc. All Rights Reserved.
 #pragma once
 
+#include "Chaos/CastingUtilities.h"
 #include "Chaos/ImplicitObject.h"
 #include "Chaos/Plane.h"
 #include "Chaos/Transform.h"
@@ -19,46 +20,6 @@
 
 namespace Chaos
 {
-	template <typename Lambda>
-	FORCEINLINE_DEBUGGABLE auto CastHelper(const FImplicitObject& Geom, const FRigidTransform3& TM, const Lambda& Func)
-	{
-		const EImplicitObjectType Type = Geom.GetType(true);
-		switch (Type)
-		{
-		case ImplicitObjectType::Sphere: return Func(Geom.template GetObjectChecked<TSphere<FReal, 3>>(), TM);
-		case ImplicitObjectType::Box: return Func(Geom.template GetObjectChecked<TBox<FReal, 3>>(), TM);
-		case ImplicitObjectType::Capsule: return Func(Geom.template GetObjectChecked<TCapsule<FReal>>(), TM);
-		case ImplicitObjectType::Convex: return Func(Geom.template GetObjectChecked<FConvex>(), TM);
-		case ImplicitObjectType::IsScaled | ImplicitObjectType::Sphere: return Func(Geom.template GetObjectChecked<TImplicitObjectScaled<TSphere<FReal, 3>>>(), TM);
-		case ImplicitObjectType::IsScaled | ImplicitObjectType::Box: return Func(Geom.template GetObjectChecked< TImplicitObjectScaled<TBox<FReal, 3>>>(), TM);
-		case ImplicitObjectType::IsScaled | ImplicitObjectType::Capsule: return Func(Geom.template GetObjectChecked< TImplicitObjectScaled<TCapsule<FReal>>>(), TM);
-		case ImplicitObjectType::IsScaled | ImplicitObjectType::Convex: return Func(Geom.template GetObjectChecked< TImplicitObjectScaled<FConvex>>(), TM);
-		case ImplicitObjectType::IsInstanced | ImplicitObjectType::Sphere:
-		{
-			return Func(Geom.template GetObjectChecked< TImplicitObjectInstanced<TSphere<FReal,3>>>().GetInstancedObject()->template GetObjectChecked<TSphere<FReal,3>>(),TM);
-		}
-		case ImplicitObjectType::IsInstanced | ImplicitObjectType::Box:
-		{
-			return Func(Geom.template GetObjectChecked< TImplicitObjectInstanced<TBox<FReal,3>>>().GetInstancedObject()->template GetObjectChecked<TBox<FReal,3>>(),TM);
-		}
-		case ImplicitObjectType::IsInstanced | ImplicitObjectType::Capsule:
-		{
-			return Func(Geom.template GetObjectChecked< TImplicitObjectInstanced<TCapsule<FReal>>>().GetInstancedObject()->template GetObjectChecked<TCapsule<FReal>>(),TM);
-		}
-		case ImplicitObjectType::IsInstanced | ImplicitObjectType::Convex:
-		{
-			return Func(Geom.template GetObjectChecked< TImplicitObjectInstanced<FConvex>>().GetInstancedObject()->template GetObjectChecked<FConvex>(),TM);
-		}
-		case ImplicitObjectType::Transformed:
-		{
-			const auto& ImplicitObjectTransformed = (Geom.template GetObjectChecked<TImplicitObjectTransformed<FReal,3>>());
-			return CastHelper(*ImplicitObjectTransformed.GetTransformedObject(), ImplicitObjectTransformed.GetTransform() * TM, Func);
-		}
-
-		default: check(false);
-		}
-		return Func(Geom.template GetObjectChecked<TSphere<FReal, 3>>(), TM);	//needed for return type
-	}
 
 	struct FMTDInfo
 	{
@@ -94,7 +55,7 @@ namespace Chaos
 			const FVec3 Offset = ATM.GetLocation() - BTM.GetLocation();
 			if (OutMTD)
 			{
-				return CastHelper(A, BToATM, [&](const auto& AConcrete, const auto& BToAFullTM)
+				return Utilities::CastHelper(A, BToATM, [&](const auto& AConcrete, const auto& BToAFullTM)
 				{
 					FVec3 LocalA,LocalB,LocalNormal;
 					if(GJKPenetration<false, FReal>(AConcrete,B,BToAFullTM,OutMTD->Penetration,LocalA,LocalB,LocalNormal,Thickness,Offset.SizeSquared() < 1e-4 ? FVec3(1,0,0) : Offset))
@@ -108,7 +69,7 @@ namespace Chaos
 			}
 			else
 			{
-				return CastHelper(A, BToATM, [&](const auto& AConcrete, const auto& BToAFullTM) { return GJKIntersection<FReal>(AConcrete, B, BToAFullTM, Thickness, Offset.SizeSquared() < 1e-4 ? FVec3(1, 0, 0) : Offset); });
+				return Utilities::CastHelper(A, BToATM, [&](const auto& AConcrete, const auto& BToAFullTM) { return GJKIntersection<FReal>(AConcrete, B, BToAFullTM, Thickness, Offset.SizeSquared() < 1e-4 ? FVec3(1, 0, 0) : Offset); });
 			}
 		}
 		else
@@ -206,7 +167,7 @@ namespace Chaos
 			}
 
 			const FVec3 Offset = ATM.GetLocation() - BTM.GetLocation();
-			bResult = CastHelper(A, BToATM, [&](const auto& ADowncast, const auto& BToAFullTM){ return GJKRaycast2(ADowncast, B, BToAFullTM, LocalDir, Length, OutTime, LocalPosition, LocalNormal, Thickness, bComputeMTD, Offset, Thickness); });
+			bResult = Utilities::CastHelper(A, BToATM, [&](const auto& ADowncast, const auto& BToAFullTM){ return GJKRaycast2(ADowncast, B, BToAFullTM, LocalDir, Length, OutTime, LocalPosition, LocalNormal, Thickness, bComputeMTD, Offset, Thickness); });
 			if (AType == ImplicitObjectType::Convex)
 			{
 				//todo: find face index
diff --git a/Engine/Source/Runtime/Experimental/Chaos/Public/Chaos/ParticleHandle.h b/Engine/Source/Runtime/Experimental/Chaos/Public/Chaos/ParticleHandle.h
--- a/Engine/Source/Runtime/Experimental/Chaos/Public/Chaos/ParticleHandle.h
+++ b/Engine/Source/Runtime/Experimental/Chaos/Public/Chaos/ParticleHandle.h
@@ -2,6 +2,7 @@
 #pragma once
 
 #include "Templates/ChooseClass.h"
+#include "Chaos/CastingUtilities.h"
 #include "Chaos/PBDRigidClusteredParticles.h"
 #include "Chaos/PBDGeometryCollectionParticles.h"
 #include "Chaos/ParticleHandleFwd.h"
@@ -10,6 +11,8 @@
 #include "Chaos/Framework/PhysicsProxyBase.h"
 #include "ChaosCheck.h"
 
+
+
 class IPhysicsProxyBase;
 
 namespace Chaos
@@ -1536,14 +1539,54 @@ protected:
 		MapImplicitShapes();
 	}
 
+	//TGeometryParticle::
 	void MapImplicitShapes()
 	{
 		ImplicitShapeMap.Reset();
 		int32 ShapeIndex = 0;
-		for(TUniquePtr<TPerShapeData<T, d>>& ShapeData : MShapesArray)
+		for (TUniquePtr<TPerShapeData<T, d>>& ShapeData : MShapesArray)
 		{
 			ImplicitShapeMap.Add(ShapeData->Geometry.Get(), ShapeIndex++);
 		}
+
+		if (MGeometry)
+		{
+			int32 CurrentShapeIndex = INDEX_NONE;
+			if (const auto* Union = MGeometry->template GetObject<FImplicitObjectUnion>())
+			{
+				for (const TUniquePtr<FImplicitObject>& ImplicitObject : Union->GetObjects())
+				{
+					if (ImplicitObject.Get())
+					{
+						if (const FImplicitObject* ImplicitChildObject = Utilities::ImplicitChildHelper(ImplicitObject.Get()))
+						{
+							if (ImplicitShapeMap.Contains(ImplicitObject.Get()))
+							{
+								ImplicitShapeMap.Add(ImplicitChildObject, ImplicitShapeMap[ImplicitObject.Get()]);
+							}
+							else if (ImplicitShapeMap.Contains(ImplicitChildObject))
+							{
+								ImplicitShapeMap.Add(ImplicitObject.Get(), ImplicitShapeMap[ImplicitChildObject]);
+							}
+						}
+					}
+				}
+			}
+			else 
+			{
+				if (const FImplicitObject* ImplicitChildObject = Utilities::ImplicitChildHelper(MGeometry.Get()))
+				{
+					if (ImplicitShapeMap.Contains(MGeometry.Get()))
+					{
+						ImplicitShapeMap.Add(ImplicitChildObject, ImplicitShapeMap[MGeometry.Get()]);
+					}
+					else if (ImplicitShapeMap.Contains(ImplicitChildObject))
+					{
+						ImplicitShapeMap.Add(MGeometry.Get(), ImplicitShapeMap[ImplicitChildObject]);
+					}
+				}
+			}
+		}
 	}
 };
 
diff --git a/Engine/Source/Runtime/PhysicsSQ/Private/SQAccelerator.cpp b/Engine/Source/Runtime/PhysicsSQ/Private/SQAccelerator.cpp
--- a/Engine/Source/Runtime/PhysicsSQ/Private/SQAccelerator.cpp
+++ b/Engine/Source/Runtime/PhysicsSQ/Private/SQAccelerator.cpp
@@ -13,6 +13,7 @@
 #endif
 
 #include "ChaosInterfaceWrapperCore.h"
+#include "Chaos/CastingUtilities.h"
 #include "Chaos/ISpatialAcceleration.h"
 #include "Chaos/PBDCollisionConstraints.h"
 #include "Chaos/GeometryQueries.h"
@@ -382,7 +383,7 @@ void SweepHelper(const QueryGeomType& QueryGeom,const Chaos::ISpatialAcceleratio
 
 void FChaosSQAccelerator::Sweep(const Chaos::FImplicitObject& QueryGeom, const FTransform& StartTM, const FVector& Dir, const float DeltaMagnitude, ChaosInterface::FSQHitBuffer<ChaosInterface::FSweepHit>& HitBuffer, EHitFlags OutputFlags, const FQueryFilterData& QueryFilterData, ICollisionQueryFilterCallbackBase& QueryCallback, const FQueryDebugParams& DebugParams) const
 {
-	return Chaos::CastHelper(QueryGeom, StartTM, [&](const auto& Downcast, const FTransform& StartFullTM) { return SweepHelper(Downcast, SpatialAcceleration, StartFullTM, Dir, DeltaMagnitude, HitBuffer, OutputFlags, QueryFilterData, QueryCallback, DebugParams); });
+	return Chaos::Utilities::CastHelper(QueryGeom, StartTM, [&](const auto& Downcast, const FTransform& StartFullTM) { return SweepHelper(Downcast, SpatialAcceleration, StartFullTM, Dir, DeltaMagnitude, HitBuffer, OutputFlags, QueryFilterData, QueryCallback, DebugParams); });
 }
 
 template <typename QueryGeomType>
@@ -400,7 +401,7 @@ void OverlapHelper(const QueryGeomType& QueryGeom, const Chaos::ISpatialAccelera
 
 void FChaosSQAccelerator::Overlap(const Chaos::FImplicitObject& QueryGeom, const FTransform& GeomPose, ChaosInterface::FSQHitBuffer<ChaosInterface::FOverlapHit>& HitBuffer, const FQueryFilterData& QueryFilterData, ICollisionQueryFilterCallbackBase& QueryCallback, const FQueryDebugParams& DebugParams) const
 {
-	return Chaos::CastHelper(QueryGeom, GeomPose, [&](const auto& Downcast, const FTransform& GeomFullPose) { return OverlapHelper(Downcast, SpatialAcceleration, GeomFullPose, HitBuffer, QueryFilterData, QueryCallback, DebugParams); });
+	return Chaos::Utilities::CastHelper(QueryGeom, GeomPose, [&](const auto& Downcast, const FTransform& GeomFullPose) { return OverlapHelper(Downcast, SpatialAcceleration, GeomFullPose, HitBuffer, QueryFilterData, QueryCallback, DebugParams); });
 }
 
 #if WITH_PHYSX
